Slate: Refactored core Slate implementation into SlateCore module in preparation for UMG.

diff --git a/Engine/Source/Runtime/Slate/Public/Widgets/Views/STreeView.h b/Engine/Source/Runtime/Slate/Public/Widgets/Views/STreeView.h
new file mode 100644
--- /dev/null
+++ b/Engine/Source/Runtime/Slate/Public/Widgets/Views/STreeView.h
@@ -0,0 +1,613 @@
+// Copyright 1998-2014 Epic Games, Inc. All Rights Reserved.
+ 
+#pragma once
+
+/** Info needed by a (relatively) small fraction of the tree items; some of them may not be visible. */
+struct FSparseItemInfo
+{
+	/**
+	 * Construct a new FTreeItem.
+	 *
+	 * @param InItemToVisualize   The DateItem pointer being wrapped by this FTreeItem
+	 * @param InHasChildren       Does this item have children? True if yes.
+	 */
+	FSparseItemInfo( bool InIsExpanded, bool InHasExpandedChildren )
+	: bIsExpanded( InIsExpanded )
+	, bHasExpandedChildren( InHasExpandedChildren )
+	{	
+	}
+
+	/** Is this tree item expanded? */
+	bool bIsExpanded;
+
+	/** Does this tree item have any expanded children? */
+	bool bHasExpandedChildren;
+};
+
+/** Info needed by every visible item in the tree */
+struct FItemInfo
+{
+	FItemInfo()
+	{		
+	}
+
+	FItemInfo( int32 InNestingLevel, bool InHasChildren )
+	: NestingLevel( InNestingLevel )
+	, bHasChildren( InHasChildren )
+	{
+	}
+
+	/** Nesting level within the tree: 0 is root-level, 1 is children of root, etc. */
+	int32 NestingLevel;
+
+	/** Does this tree item have children? */
+	bool bHasChildren;
+};
+
+
+
+
+
+
+/**
+ * This assumes you are familiar with SListView; see SListView.
+ *
+ * TreeView setup is virtually identical to that of ListView.
+ * Additionally, TreeView introduces a new delegate: OnGetChildren().
+ * OnGetChildren() takes some DataItem being observed by the tree
+ * and returns that item's children. Like ListView, TreeView operates
+ * exclusively with pointers to DataItems.
+ * 
+ */
+template<typename ItemType>
+class STreeView : public SListView< ItemType >
+{
+
+	public:
+		typedef typename TListTypeTraits< ItemType >::NullableType NullableItemType;
+
+		typedef typename TSlateDelegates< ItemType >::FOnGetChildren FOnGetChildren;
+		typedef typename TSlateDelegates< ItemType >::FOnGenerateRow FOnGenerateRow;
+		typedef typename TSlateDelegates< ItemType >::FOnSetExpansionRecursive FOnSetExpansionRecursive;
+		typedef typename TSlateDelegates< ItemType >::FOnItemScrolledIntoView FOnItemScrolledIntoView;
+		typedef typename TSlateDelegates< NullableItemType >::FOnSelectionChanged FOnSelectionChanged;
+		typedef typename TSlateDelegates< ItemType >::FOnMouseButtonDoubleClick FOnMouseButtonDoubleClick;
+		typedef typename TSlateDelegates< ItemType >::FOnExpansionChanged FOnExpansionChanged;
+	public:
+	
+		SLATE_BEGIN_ARGS( STreeView<ItemType> )
+			: _OnGenerateRow()
+			, _OnGetChildren()
+			, _OnSetExpansionRecursive()
+			, _TreeItemsSource( static_cast< TArray<ItemType>* >(NULL) ) //@todo Slate Syntax: Initializing from NULL without a cast
+			, _ItemHeight(16)
+			, _OnContextMenuOpening()
+			, _OnMouseButtonDoubleClick()
+			, _OnSelectionChanged()
+			, _OnExpansionChanged()
+			, _SelectionMode(ESelectionMode::Multi)
+			, _ClearSelectionOnClick(true)
+			, _ExternalScrollbar()
+			{}
+
+			SLATE_EVENT( FOnGenerateRow, OnGenerateRow )
+
+			SLATE_EVENT( FOnItemScrolledIntoView, OnItemScrolledIntoView )
+
+			SLATE_EVENT( FOnGetChildren, OnGetChildren )
+
+			SLATE_EVENT( FOnSetExpansionRecursive, OnSetExpansionRecursive )
+
+			SLATE_ARGUMENT( TArray<ItemType>* , TreeItemsSource )
+
+			SLATE_ATTRIBUTE( float, ItemHeight )
+
+			SLATE_EVENT( FOnContextMenuOpening, OnContextMenuOpening )
+
+			SLATE_EVENT( FOnMouseButtonDoubleClick, OnMouseButtonDoubleClick )
+
+			SLATE_EVENT( FOnSelectionChanged, OnSelectionChanged )
+
+			SLATE_EVENT( FOnExpansionChanged, OnExpansionChanged )
+
+			SLATE_ATTRIBUTE( ESelectionMode::Type, SelectionMode )
+
+			SLATE_ARGUMENT( TSharedPtr<SHeaderRow>, HeaderRow )
+
+			SLATE_ARGUMENT ( bool, ClearSelectionOnClick )
+
+			SLATE_ARGUMENT( TSharedPtr<SScrollBar>, ExternalScrollbar )
+
+		SLATE_END_ARGS()
+
+		
+		/**
+		 * Construct this widget
+		 *
+		 * @param	InArgs	The declaration data for this widget
+		 */
+		void Construct( const FArguments& InArgs )
+		{
+			this->OnGenerateRow = InArgs._OnGenerateRow;
+			this->OnItemScrolledIntoView = InArgs._OnItemScrolledIntoView;
+			this->OnGetChildren = InArgs._OnGetChildren;
+			this->OnSetExpansionRecursive = InArgs._OnSetExpansionRecursive;
+			this->TreeItemsSource = InArgs._TreeItemsSource;
+
+			this->OnContextMenuOpening = InArgs._OnContextMenuOpening;
+			this->OnDoubleClick = InArgs._OnMouseButtonDoubleClick;
+			this->OnSelectionChanged = InArgs._OnSelectionChanged;
+			this->OnExpansionChanged = InArgs._OnExpansionChanged;
+			this->SelectionMode = InArgs._SelectionMode;
+
+			this->bClearSelectionOnClick = InArgs._ClearSelectionOnClick;
+
+			// Check for any parameters that the coder forgot to specify.
+			FString ErrorString;
+			{
+				if ( !this->OnGenerateRow.IsBound() )
+				{
+					ErrorString += TEXT("Please specify an OnGenerateRow. \n");
+				}
+
+				if ( this->TreeItemsSource == NULL )
+				{
+					ErrorString += TEXT("Please specify a TreeItemsSource. \n");
+				}
+
+				if ( !this->OnGetChildren.IsBound() )
+				{
+					ErrorString += TEXT("Please specify an OnGetChildren. \n");
+				}
+			}
+
+			if (ErrorString.Len() > 0)
+			{
+				// Let the coder know what they forgot
+				this->ChildSlot
+				.HAlign(HAlign_Center)
+				.VAlign(VAlign_Center)
+				[
+					SNew(STextBlock)
+					.Text(ErrorString)
+				];
+			}
+			else
+			{
+				// Make the TableView
+				this->ConstructChildren( 0, InArgs._ItemHeight, InArgs._HeaderRow, InArgs._ExternalScrollbar );
+			}
+		}
+
+		STreeView()
+			: SListView< ItemType >( ETableViewMode::Tree )
+			, TreeItemsSource( NULL )
+			, bTreeItemsAreDirty( true )
+		{
+			this->ItemsSource = &LinearizedItems;
+		}
+
+public:
+	// Inherited from SWidget
+	virtual FReply OnKeyDown( const FGeometry& MyGeometry, const FKeyboardEvent& InKeyboardEvent ) OVERRIDE
+	{
+		// Check for selection/expansion toggling keys (Left, Right)
+		// SelectorItem represents the keyboard selection. If it isn't valid then we don't know what to expand.
+		// Don't respond to key-presses containing "Alt" as a modifier
+		if ( TListTypeTraits<ItemType>::IsPtrValid(this->SelectorItem) && !InKeyboardEvent.IsAltDown() )
+		{
+			if ( InKeyboardEvent.GetKey() == EKeys::Left )
+			{
+				if( TListTypeTraits<ItemType>::IsPtrValid(this->SelectorItem) )
+				{
+					ItemType RangeSelectionEndItem = TListTypeTraits<ItemType>::NullableItemTypeConvertToItemType( this->SelectorItem );
+					int32 SelectionIndex = this->LinearizedItems.Find( RangeSelectionEndItem );
+
+					if ( Private_DoesItemHaveChildren(SelectionIndex) && Private_IsItemExpanded( RangeSelectionEndItem ) )
+					{
+						// Collapse the selected item
+						Private_SetItemExpansion(RangeSelectionEndItem, false);
+					}
+					else
+					{
+						// Select the parent, who should be a previous item in the list whose nesting level is less than the selected one
+						int32 SelectedNestingDepth = Private_GetNestingDepth(SelectionIndex);
+						for (SelectionIndex--; SelectionIndex >= 0; --SelectionIndex)
+						{
+							if ( Private_GetNestingDepth(SelectionIndex) < SelectedNestingDepth )
+							{
+								// Found the parent
+								this->KeyboardSelect(this->LinearizedItems[SelectionIndex], InKeyboardEvent);
+								break;
+							}
+						}
+					}
+				}
+
+				return FReply::Handled();
+			}
+			else if ( InKeyboardEvent.GetKey() == EKeys::Right )
+			{
+				if( TListTypeTraits<ItemType>::IsPtrValid(this->SelectorItem) )
+				{
+					ItemType RangeSelectionEndItem = TListTypeTraits<ItemType>::NullableItemTypeConvertToItemType( this->SelectorItem );
+					int32 SelectionIndex = this->LinearizedItems.Find( RangeSelectionEndItem );
+
+					// Right only applies to items with children
+					if ( Private_DoesItemHaveChildren(SelectionIndex) )
+					{
+						if ( !Private_IsItemExpanded(RangeSelectionEndItem) )
+						{
+							// Expand the selected item
+							Private_SetItemExpansion(RangeSelectionEndItem, true);
+						}
+						else
+						{
+							// Select the first child, who should be the next item in the list						
+							// Make sure we aren't the last item on the list
+							if ( SelectionIndex < this->LinearizedItems.Num() - 1 )
+							{
+								this->KeyboardSelect(this->LinearizedItems[SelectionIndex + 1], InKeyboardEvent);
+							}
+						}
+					}
+				}
+
+				return FReply::Handled();
+			}
+		}
+
+		return SListView<ItemType>::OnKeyDown(MyGeometry, InKeyboardEvent);
+	}
+	
+private:
+	//
+	// Private Interface
+	//
+	// Tree View adds the ability to expand/collapse items.
+	// All the selection functionality is inherited from ListView.
+
+	virtual bool Private_IsItemExpanded( const ItemType& TheItem ) const OVERRIDE
+	{
+		const FSparseItemInfo* ItemInfo = SparseItemInfos.Find(TheItem);
+		return ItemInfo != NULL && ItemInfo->bIsExpanded;
+	}
+
+	virtual void Private_SetItemExpansion( ItemType TheItem, bool bShouldBeExpanded ) OVERRIDE
+	{
+		const FSparseItemInfo* const SparseItemInfo = SparseItemInfos.Find(TheItem);
+		bool bWasExpanded = false;
+
+		if(SparseItemInfo)
+		{
+			bWasExpanded = SparseItemInfo->bIsExpanded;
+			SparseItemInfos.Add(TheItem, FSparseItemInfo(bShouldBeExpanded, SparseItemInfo->bHasExpandedChildren));
+		}
+		else if(bShouldBeExpanded)
+		{
+			SparseItemInfos.Add(TheItem, FSparseItemInfo(bShouldBeExpanded, false));
+		}
+		
+		if(bWasExpanded != bShouldBeExpanded)
+		{
+			OnExpansionChanged.ExecuteIfBound(TheItem, bShouldBeExpanded);
+
+			// We must rebuild the linearized version of the tree because
+			// either some children became visible or some got removed.
+			RequestTreeRefresh();
+		}
+	}
+
+	virtual void Private_OnExpanderArrowShiftClicked( ItemType TheItem, bool bShouldBeExpanded ) OVERRIDE
+	{
+		if(OnSetExpansionRecursive.IsBound())
+		{
+			OnSetExpansionRecursive.Execute(TheItem, bShouldBeExpanded);
+
+			// We must rebuild the linearized version of the tree because
+			// either some children became visible or some got removed.
+			RequestTreeRefresh();
+		}
+	}
+
+	virtual bool Private_DoesItemHaveChildren( int32 ItemIndexInList ) const OVERRIDE
+	{
+		return DenseItemInfos[ItemIndexInList].bHasChildren;
+	}
+
+	virtual int32 Private_GetNestingDepth( int32 ItemIndexInList ) const OVERRIDE
+	{
+		return DenseItemInfos[ItemIndexInList].NestingLevel;
+	}
+
+
+
+
+public:
+		/**
+		 * See SWidget::Tick()
+		 *
+		 * @param  AllottedGeometry The space allotted for this widget
+		 * @param  InCurrentTime  Current absolute real time
+		 * @param  InDeltaTime  Real time passed since last tick
+		 */
+		virtual void Tick( const FGeometry& AllottedGeometry, const double InCurrentTime, const float InDeltaTime )
+		{
+			if ( bTreeItemsAreDirty )
+			{
+				// Check that ItemsPanel was made; we never make it if the user failed to specify all the parameters necessary to make the tree work.
+				if ( this->ItemsPanel.IsValid() )
+				{
+					// We are about to repopulate linearized items; the ListView that TreeView is built on top of will also need to refresh.
+					bTreeItemsAreDirty = false;
+
+					if ( OnGetChildren.IsBound() && TreeItemsSource != NULL )								
+					{
+						// We make copies of the old expansion and selection sets so that we can remove
+						// any items that are no longer seen by the tree.
+						TSet< ItemType > TempSelectedItemsMap;
+						TMap< ItemType, FSparseItemInfo > TempSparseItemInfo;
+						TArray<FItemInfo> TempDenseItemInfos;
+					
+						// Rebuild the linearized view of the tree data.
+						LinearizedItems.Empty();
+						PopulateLinearizedItems( *TreeItemsSource, LinearizedItems, TempDenseItemInfos, 0, TempSelectedItemsMap, TempSparseItemInfo, true );
+						SparseItemInfos = TempSparseItemInfo;
+						DenseItemInfos = TempDenseItemInfos;
+
+						if(	this->SelectedItems.Num() != TempSelectedItemsMap.Num() || 
+							this->SelectedItems.Difference(TempSelectedItemsMap).Num() > 0 || 
+							TempSelectedItemsMap.Difference(this->SelectedItems).Num() > 0 )
+						{
+							this->SelectedItems = TempSelectedItemsMap;
+
+							if ( !TListTypeTraits<ItemType>::IsPtrValid( this->RangeSelectionStart ) ||
+								 !this->SelectedItems.Contains( TListTypeTraits<ItemType>::NullableItemTypeConvertToItemType( this->RangeSelectionStart ) ))
+							{
+									TListTypeTraits< ItemType >::ResetPtr( this->RangeSelectionStart );
+									TListTypeTraits< ItemType >::ResetPtr( this->SelectorItem );
+							}	
+							else if ( !TListTypeTraits<ItemType>::IsPtrValid( this->SelectorItem ) || 
+									  !this->SelectedItems.Contains( TListTypeTraits<ItemType>::NullableItemTypeConvertToItemType( this->SelectorItem ) ) )
+							{
+								this->SelectorItem = this->RangeSelectionStart;
+							}
+
+							this->Private_SignalSelectionChanged(ESelectInfo::Direct);
+						}
+					}
+				}
+			}
+			
+			// Tick the TreeView superclass so that it can refresh.
+			// This may be due to TreeView requesting a refresh or because new items became visible due to resizing or scrolling.
+			SListView< ItemType >::Tick(AllottedGeometry, InCurrentTime, InDeltaTime);
+		}
+		
+		/** 
+		 * Given: an array of items (ItemsSource) each of which potentially has a child.
+		 * Task: populate the LinearezedItems array with a flattened version of the visible data items.
+		 *       In the process, remove any items that are not visible while maintaining any collapsed
+		 *       items that may have expanded children.
+		 *
+		 * @param ItemsSource          An array of data items each of which may have 0 or more children.
+		 * @param LinearizedItems      Array to populate with items based on expanded/collapsed state
+		 * @param NewDenseItemInfos    Array representing how nested each item in the Linearized items is, and whether it has children
+		 * @param TreeLevel            The current level of indentation
+		 * @param OutNewSelectedItems  Selected items minus any items that are no longer observed by the list
+		 * @param NewSparseItemInfo    Expanded items and items that have expanded children minus any items that are no longer observed by the list
+		 * @param bAddingItems         Are we adding encountered items to the linearized items list or just testing them for existence
+		 *
+		 * @return true if we encountered expanded children; false otherwise.
+		 */
+		bool PopulateLinearizedItems(
+			const TArray<ItemType>& ItemsSource,
+			TArray< ItemType >& InLinearizedItems,
+			TArray< FItemInfo >& NewDenseItemInfos,
+			int32 TreeLevel,
+			TSet< ItemType >& OutNewSelectedItems,
+			TMap< ItemType, FSparseItemInfo >& NewSparseItemInfo,
+			bool bAddingItems )
+		{
+			bool bSawExpandedItems = false;
+			for ( int32 ItemIndex = 0; ItemIndex < ItemsSource.Num(); ++ItemIndex )
+			{
+				const ItemType& CurItem = ItemsSource[ItemIndex];
+
+				// Find this items children.
+				TArray<ItemType> ChildItems;
+				OnGetChildren.Execute( ItemsSource[ItemIndex], ChildItems );
+
+				const bool bHasChildren = ChildItems.Num() > 0;
+
+				// Is this item expanded, does it have expanded children?
+				const FSparseItemInfo* CurItemInfo = SparseItemInfos.Find( CurItem );
+				const bool bIsExpanded = (CurItemInfo == NULL) ? false : CurItemInfo->bIsExpanded;
+				bool bHasExpandedChildren = (CurItemInfo == NULL) ? false : CurItemInfo->bHasExpandedChildren;
+								
+				// Add this item to the linearized list and update the selection set.
+				if (bAddingItems)
+				{
+					InLinearizedItems.Add( CurItem );
+
+					NewDenseItemInfos.Add( FItemInfo(TreeLevel, bHasChildren) );
+
+					const bool bIsSelected = this->IsItemSelected( CurItem );
+					if (bIsSelected)
+					{
+						OutNewSelectedItems.Add( CurItem );
+					}
+				}				
+				
+				if ( bIsExpanded || bHasExpandedChildren )
+				{ 
+					// If this item is expanded, we should process all of its children at the next indentation level.
+					// If it is collapsed, process its children but do not add them to the linearized list.
+					const bool bAddChildItems = bAddingItems && bIsExpanded;
+					bHasExpandedChildren = PopulateLinearizedItems( ChildItems, InLinearizedItems, NewDenseItemInfos, TreeLevel+1, OutNewSelectedItems, NewSparseItemInfo, bAddChildItems );
+				}
+
+				if ( bIsExpanded || bHasExpandedChildren )
+				{
+					// Update the item info for this tree item.
+					NewSparseItemInfo.Add( CurItem, FSparseItemInfo( bIsExpanded, bHasExpandedChildren) );
+				}
+
+
+				// If we encountered any expanded nodes along the way, the parent has expanded children.
+				bSawExpandedItems = bSawExpandedItems || bIsExpanded || bHasExpandedChildren;				
+			}
+
+			return bSawExpandedItems;
+		}
+		
+		/**
+		 * Given a TreeItem, create a Widget to represent it in the tree view.
+		 *
+		 * @param InItem   Item to visualize
+		 *
+		 * @return A widget that represents this item.
+		 */
+		virtual TSharedRef<ITableRow> GenerateNewWidget( ItemType InItem )
+		{
+			if ( this->OnGenerateRow.IsBound() )
+			{
+				return this->OnGenerateRow.Execute( InItem, this->SharedThis(this) );
+			}
+			else
+			{
+				TSharedRef< STreeView<ItemType> > This = StaticCastSharedRef< STreeView<ItemType> >(this->AsShared());
+
+				// The programmer did not provide an OnGenerateRow() handler; let them know.
+				TSharedRef< STableRow<ItemType> > NewTreeItemWidget = 
+					SNew( STableRow<ItemType>, This )
+					.Content()
+					[
+						SNew(STextBlock) .Text( NSLOCTEXT("STreeView", "BrokenSetupMessage", "OnGenerateWidget() not assigned.") )
+					];
+
+				return NewTreeItemWidget;
+			}
+		}
+		
+		/**
+		 * Queue up a regeneration of the linearized items on the next tick.
+		 */
+		void RequestTreeRefresh()
+		{
+			bTreeItemsAreDirty = true;
+			RequestListRefresh();
+		}
+		
+		/**
+		 * Set whether some data item is expanded or not.
+		 * 
+		 * @param InItem         The item whose expansion state to control
+		 * @param InExpandItem   If true the item should be expanded; otherwise collapsed
+		 */
+		void SetItemExpansion( const ItemType& InItem, bool InShouldExpandItem )
+		{
+			Private_SetItemExpansion(InItem, InShouldExpandItem);
+		}
+
+		/** Collapse all the items in the tree and expand InItem */
+		void SetSingleExpandedItem( const ItemType& InItem )
+		{
+			// Will we have to do any work?
+			const bool bItemAlreadyLoneExpanded = (this->SparseItemInfos.Num() == 1) && this->IsItemExpanded(InItem);
+
+			if (!bItemAlreadyLoneExpanded)
+			{
+				this->SparseItemInfos.Empty();
+				Private_SetItemExpansion(InItem, true);
+			}
+		}
+		
+		/** 
+		 * @param InItem   The data item whose expansion state to query
+		 *
+		 * @return true if the item is expanded; false otherwise.
+		 */
+		bool IsItemExpanded( const ItemType& InItem ) const
+		{
+			return Private_IsItemExpanded( InItem );
+		}
+
+		
+		/**
+		 * Set the TreeItemsSource. The Tree will generate widgets to represent these items.
+		 *
+		 * @param InItemsSource  A pointer to the array of items that should be observed by this TreeView
+		 */
+		void SetTreeItemsSource( TArray<ItemType>* InItemsSource)
+		{
+			TreeItemsSource = InItemsSource;
+			RequestTreeRefresh();
+		}
+
+		/**
+		 * Generates a set of items that are currently expanded
+		 *
+		 * @param ExpandedItems	The generated set of expanded items.
+		 */
+		void GetExpandedItems( TSet< ItemType >& ExpandedItems ) const
+		{
+			for( typename TMap<ItemType, FSparseItemInfo>::TConstIterator InfoIterator(SparseItemInfos); InfoIterator; ++InfoIterator )
+			{
+				if ( InfoIterator.Value().bIsExpanded )
+				{
+					ExpandedItems.Add( InfoIterator.Key() );
+				}
+			}			
+		}
+
+		/**
+		 * Clears the entire set of expanded items
+		 */
+		void ClearExpandedItems()
+		{
+			SparseItemInfos.Empty();
+			RequestTreeRefresh();
+		}
+
+	private:
+		/** Hide RequestListRefresh() for the tree widget's interface.  You should always call RequestTreeRefresh()
+		    for trees, not RequestListRefresh() */
+		void RequestListRefresh()
+		{
+			SListView< ItemType >::RequestListRefresh();
+		}
+
+	protected:
+
+		
+		/** The delegate that is invoked whenever we need to gather an item's children */
+		FOnGetChildren OnGetChildren;
+
+		/** The delegate that is invoked to recursively expand/collapse a tree items children */
+		FOnSetExpansionRecursive OnSetExpansionRecursive;
+
+		/** A pointer to the items being observed by the tree view. */
+		TArray<ItemType>* TreeItemsSource;		
+		
+		/** Info needed by a small fraction of tree items; some of these are not visible to the user */
+		TMap<ItemType, FSparseItemInfo> SparseItemInfos;
+
+		/** Info needed by every item in the linearized version of the tree. */
+		TArray<FItemInfo> DenseItemInfos;
+
+		/**
+		 * A linearized version of the items being observed by the tree view.
+		 * Note that we inherit from a ListView, which we point at this linearized version of the tree.
+		 */
+		TArray< ItemType > LinearizedItems;
+
+		/** The delegate that is invoked whnever an item in the tree is expanded or collapsed */
+		FOnExpansionChanged OnExpansionChanged;
+
+	private:		
+		/** true when the LinearizedItems need to be regenerated. */
+		bool bTreeItemsAreDirty;
+
+};
+
diff --git a/Engine/Source/Runtime/Slate/Slate.Build.cs b/Engine/Source/Runtime/Slate/Slate.Build.cs
--- a/Engine/Source/Runtime/Slate/Slate.Build.cs
+++ b/Engine/Source/Runtime/Slate/Slate.Build.cs
@@ -13,15 +13,35 @@ public class Slate : ModuleRules
 				"Core",
 				"CoreUObject",
 				"InputCore",
+				"SlateCore",
 			}
-			);
+		);
 
         if (!UEBuildConfiguration.bBuildDedicatedServer)
         {
 			AddThirdPartyPrivateStaticDependencies(Target, "FreeType2");
         }
 
-        PrivateDependencyModuleNames.Add( "Sockets" );
+		PrivateDependencyModuleNames.Add("Sockets");
+
+		PrivateIncludePaths.AddRange(
+			new string[] {
+				"Runtime/Slate/Private",
+				"Runtime/Slate/Private/Docking",
+				"Runtime/Slate/Private/DockingFramework",
+				"Runtime/Slate/Private/Framework",
+				"Runtime/Slate/Private/Framework/Animation",
+				"Runtime/Slate/Private/Framework/Application",
+				"Runtime/Slate/Private/Framework/Commands",
+				"Runtime/Slate/Private/Framework/Styling",
+				"Runtime/Slate/Private/Framework/Widgets",
+				"Runtime/Slate/Private/IOS",
+				"Runtime/Slate/Private/MultiBox",
+				"Runtime/Slate/Private/TableView",
+				"Runtime/Slate/Private/Testing",
+				"Runtime/Slate/Private/Text",
+			}
+		);
 
 		if ((Target.Platform == UnrealTargetPlatform.Win64) ||
 			(Target.Platform == UnrealTargetPlatform.Win32))
diff --git a/Engine/Source/Runtime/SlateCore/Private/Animation/CurveSequence.cpp b/Engine/Source/Runtime/SlateCore/Private/Animation/CurveSequence.cpp
new file mode 100644
--- /dev/null
+++ b/Engine/Source/Runtime/SlateCore/Private/Animation/CurveSequence.cpp
@@ -0,0 +1,179 @@
+// Copyright 1998-2014 Epic Games, Inc. All Rights Reserved.
+
+/*=============================================================================
+	CurveSequence.cpp: Implements the FCurveSequence structure.
+=============================================================================*/
+
+#include "SlateCorePrivatePCH.h"
+
+
+/* FCurveSequence structors
+ *****************************************************************************/
+
+FCurveSequence::FCurveSequence( )
+	: StartTime(0)
+	, TotalDuration(0)
+	, bInReverse(true)
+{ }
+
+
+/* FCurveSequence interface
+ *****************************************************************************/
+
+FCurveSequence::FCurveSequence( const float InStartTimeSeconds, const float InDurationSeconds, const ECurveEaseFunction::Type InEaseFunction )
+	: StartTime(0)
+	, TotalDuration(0)
+	, bInReverse(true)
+{
+	const FCurveHandle IgnoredCurveHandle = AddCurve( InStartTimeSeconds, InDurationSeconds, InEaseFunction );
+}
+
+
+FCurveHandle FCurveSequence::AddCurve( const float InStartTimeSeconds, const float InDurationSeconds, const ECurveEaseFunction::Type InEaseFunction )
+{
+	// Keep track of how long this sequence is
+	TotalDuration = FMath::Max(TotalDuration, InStartTimeSeconds + InDurationSeconds);
+	// The initial state is to be at the end of the animation.
+	StartTime = TotalDuration;
+
+	// Actually make this curve and return a handle to it.
+	Curves.Add( FSlateCurve(InStartTimeSeconds, InDurationSeconds, InEaseFunction) );
+	return FCurveHandle(this, Curves.Num() - 1);
+}
+
+
+FCurveHandle FCurveSequence::AddCurveRelative( const float InOffset, const float InDurationSecond, const ECurveEaseFunction::Type InEaseFunction )
+{
+	const float CurveStartTime = TotalDuration + InOffset;
+	return AddCurve(CurveStartTime, InDurationSecond, InEaseFunction);
+}
+
+
+void FCurveSequence::Play( const float StartAtTime )
+{
+	// Playing forward
+	bInReverse = false;
+
+	// We start playing NOW.
+	SetStartTime(FSlateApplicationBase::Get().GetCurrentTime() - StartAtTime);
+}
+
+
+void FCurveSequence::Reverse( )
+{
+	// We're going the other way now.
+	bInReverse = !bInReverse;	
+
+	// CurTime is now; we cannot change that, so everything happens relative to CurTime.
+	const double CurTime = FSlateApplicationBase::Get().GetCurrentTime();
+	
+	// We've played this far into the animation.
+	const float FractionCompleted = FMath::Clamp((CurTime - StartTime) / TotalDuration, 0.0, 1.0);
+	
+	// Assume CurTime is constant (now).
+	// Figure out when the animation would need to have started in order to keep
+	// its place if playing in reverse.
+	const double NewStartTime = CurTime - TotalDuration * (1 - FractionCompleted);
+	SetStartTime(NewStartTime);
+}
+
+
+void FCurveSequence::PlayReverse( const float StartAtTime )
+{
+	bInReverse = true;
+
+	// We start reversing NOW.
+	SetStartTime(FSlateApplicationBase::Get().GetCurrentTime() - StartAtTime);
+}
+
+
+bool FCurveSequence::IsPlaying( ) const
+{
+	return (FSlateApplicationBase::Get().GetCurrentTime() - StartTime) <= TotalDuration;
+}
+
+
+void FCurveSequence::SetStartTime( double InStartTime )
+{
+	StartTime = InStartTime;
+}
+
+
+float FCurveSequence::GetSequenceTime( ) const
+{
+	const double CurrentTime = FSlateApplicationBase::Get().GetCurrentTime();
+
+	return IsInReverse()
+		? TotalDuration - (CurrentTime - StartTime)
+		: CurrentTime - StartTime;
+}
+
+
+float FCurveSequence::GetSequenceTimeLooping( ) const
+{
+	return FMath::Fmod(GetSequenceTime(), TotalDuration);
+}
+
+
+bool FCurveSequence::IsInReverse( ) const
+{
+	return bInReverse;
+}
+
+
+bool FCurveSequence::IsForward( ) const
+{
+	return !bInReverse;
+}
+
+
+void FCurveSequence::JumpToStart( )
+{
+	bInReverse = true;
+	SetStartTime(FSlateApplicationBase::Get().GetCurrentTime() - TotalDuration);
+}
+
+
+void FCurveSequence::JumpToEnd( )
+{
+	bInReverse = false;
+	SetStartTime(FSlateApplicationBase::Get().GetCurrentTime() - TotalDuration);
+}
+
+
+bool FCurveSequence::IsAtStart( ) const
+{
+	return (IsInReverse() == true && IsPlaying() == false);
+}
+
+
+bool FCurveSequence::IsAtEnd( ) const
+{
+	return (IsForward() == true && IsPlaying() == false);
+}
+
+
+float FCurveSequence::GetLerp( ) const
+{
+	// Only supported for sequences with a single curve.  If you have multiple curves, use your FCurveHandle to compute
+	// interpolation alpha values.
+	checkSlow(Curves.Num() == 1);
+
+	return FCurveHandle(this, 0).GetLerp();
+}
+
+
+float FCurveSequence::GetLerpLooping( ) const
+{
+	// Only supported for sequences with a single curve.  If you have multiple curves, use your FCurveHandle to compute
+	// interpolation alpha values.
+	checkSlow(Curves.Num() == 1);
+
+	return FCurveHandle(this, 0).GetLerpLooping();
+}
+
+
+const FCurveSequence::FSlateCurve& FCurveSequence::GetCurve( int32 CurveIndex ) const
+{
+	return Curves[CurveIndex];
+}
diff --git a/Engine/Source/Runtime/SlateCore/Private/Application/SlateApplicationBase.cpp b/Engine/Source/Runtime/SlateCore/Private/Application/SlateApplicationBase.cpp
new file mode 100644
--- /dev/null
+++ b/Engine/Source/Runtime/SlateCore/Private/Application/SlateApplicationBase.cpp
@@ -0,0 +1,13 @@
+// Copyright 1998-2014 Epic Games, Inc. All Rights Reserved.
+
+/*=============================================================================
+	SlateApplicationBase.cpp: Implements the FSlateApplicationBase class.
+=============================================================================*/
+
+#include "SlateCorePrivatePCH.h"
+
+
+/* Static initialization
+ *****************************************************************************/
+
+TSharedPtr<FSlateApplicationBase> FSlateApplicationBase::CurrentBaseApplication = nullptr;
diff --git a/Engine/Source/Runtime/SlateCore/Private/Application/ThrottleManager.cpp b/Engine/Source/Runtime/SlateCore/Private/Application/ThrottleManager.cpp
new file mode 100644
--- /dev/null
+++ b/Engine/Source/Runtime/SlateCore/Private/Application/ThrottleManager.cpp
@@ -0,0 +1,70 @@
+// Copyright 1998-2014 Epic Games, Inc. All Rights Reserved.
+
+/*=============================================================================
+	ThrottleManager.cpp: Implements the FSlateThrottleManager class.
+=============================================================================*/
+
+#include "SlateCorePrivatePCH.h"
+
+
+/* FSlateThrottleManager structors
+ *****************************************************************************/
+
+FSlateThrottleManager::FSlateThrottleManager( )
+	: CVarAllowThrottle(TEXT("Slate.bAllowThrottling"), bShouldThrottle, TEXT("Allow Slate to throttle parts of the engine to ensure the UI is responsive") )
+	, bShouldThrottle(1)
+	, ThrottleCount(0)
+{ }
+
+
+/* FSlateThrottleManager interface
+ *****************************************************************************/
+
+FThrottleRequest FSlateThrottleManager::EnterResponsiveMode( )
+{
+	// Increase the number of active throttle requests
+	++ThrottleCount;
+
+	// Create a new handle for the request and return it to the user so they can close it later
+	FThrottleRequest NewHandle;
+	NewHandle.Index = ThrottleCount;
+
+	return NewHandle;
+}
+
+
+bool FSlateThrottleManager::IsAllowingExpensiveTasks( ) const
+{
+	// Expensive tasks are allowed if the number of active throttle requests is zero
+	// or we always always allow it due to a CVar
+	return (ThrottleCount == 0) || !bShouldThrottle;
+}
+
+
+void FSlateThrottleManager::LeaveResponsiveMode( FThrottleRequest& InHandle )
+{
+	if( InHandle.IsValid() )
+	{
+		// Decrement throttle count. If it becomes zero we are no longer throttling
+		--ThrottleCount;
+		check(ThrottleCount >= 0);
+
+		InHandle.Index = INDEX_NONE;
+	}
+}
+
+
+/* FSlateThrottleManager static functions
+ *****************************************************************************/
+
+FSlateThrottleManager& FSlateThrottleManager::Get( )
+{
+	static FSlateThrottleManager* Instance = nullptr;
+
+	if (Instance == nullptr)
+	{
+		Instance = new FSlateThrottleManager;
+	}
+
+	return *Instance;
+}
diff --git a/Engine/Source/Runtime/SlateCore/Private/Brushes/SlateDynamicImageBrush.cpp b/Engine/Source/Runtime/SlateCore/Private/Brushes/SlateDynamicImageBrush.cpp
new file mode 100644
--- /dev/null
+++ b/Engine/Source/Runtime/SlateCore/Private/Brushes/SlateDynamicImageBrush.cpp
@@ -0,0 +1,24 @@
+// Copyright 1998-2014 Epic Games, Inc. All Rights Reserved.
+
+/*=============================================================================
+	SlateDynamicImageBrush.cpp: Implements the FSlateDynamicImageBrush structure.
+=============================================================================*/
+
+#include "SlateCorePrivatePCH.h"
+
+
+/* FSlateDynamicImageBrush structors
+ *****************************************************************************/
+
+FSlateDynamicImageBrush::~FSlateDynamicImageBrush( )
+{
+	if (FSlateApplicationBase::IsInitialized())
+	{
+		TSharedPtr<FSlateRenderer> Renderer = FSlateApplicationBase::Get().GetRenderer();
+
+		if (Renderer.IsValid())
+		{
+			FSlateApplicationBase::Get().GetRenderer()->ReleaseDynamicResource(*this);
+		}
+	}
+}
diff --git a/Engine/Source/Runtime/SlateCore/Private/Fonts/FontCache.cpp b/Engine/Source/Runtime/SlateCore/Private/Fonts/FontCache.cpp
new file mode 100644
--- /dev/null
+++ b/Engine/Source/Runtime/SlateCore/Private/Fonts/FontCache.cpp
@@ -0,0 +1,764 @@
+// Copyright 1998-2014 Epic Games, Inc. All Rights Reserved.
+
+#include "SlateCorePrivatePCH.h"
+
+
+#ifndef WITH_FREETYPE
+	#define WITH_FREETYPE	0
+#endif // WITH_FREETYPE
+
+#if PLATFORM_COMPILER_HAS_GENERIC_KEYWORD
+	#define generic __identifier(generic)
+#endif	//PLATFORM_COMPILER_HAS_GENERIC_KEYWORD
+
+#if WITH_FREETYPE
+	#include "ft2build.h"
+
+	// Freetype style include
+	#include FT_FREETYPE_H
+	#include FT_GLYPH_H
+	#include FT_MODULE_H
+
+#endif // WITH_FREETYPE
+
+namespace FontCacheConstants
+{
+	/** The horizontal dpi we render at */
+	const uint32 HorizontalDPI = 96;
+	/** The vertical dpi we render at */
+	const uint32 VerticalDPI = 96;
+
+	/** Number of characters that can be indexed directly in the cache */
+	const int32 DirectAccessSize = 256;
+
+	/** Number of possible elements in each measurement cache */
+	const uint32 MeasureCacheSize = 500;
+}
+
+#if WITH_FREETYPE
+const uint32 GlyphFlags = FT_LOAD_TARGET_NORMAL | FT_LOAD_NO_BITMAP;
+
+/**
+ * Memory allocation functions to be used only by freetype
+ */
+static void* FreetypeAlloc( FT_Memory Memory, long size )
+{
+	return FMemory::Malloc( size );
+}
+
+static void* FreetypeRealloc( FT_Memory Memory, long CurSize, long NewSize, void* Block )
+{
+	return FMemory::Realloc( Block, NewSize );
+}
+
+static void FreetypeFree( FT_Memory Memory, void* Block )
+{
+	return FMemory::Free( Block );
+}
+
+#endif // WITH_FREETYPE
+
+/**
+ * An interface to the freetype API.                     
+ */
+class FFreeTypeInterface
+{
+public:
+	FFreeTypeInterface()
+	{
+#if WITH_FREETYPE
+		CustomMemory = (FT_Memory)FMemory::Malloc( sizeof(*CustomMemory) );
+		// Init freetype
+		CustomMemory->alloc = FreetypeAlloc;
+		CustomMemory->realloc = FreetypeRealloc;
+		CustomMemory->free = FreetypeFree;
+		CustomMemory->user = nullptr;
+
+		int32 Error = FT_New_Library( CustomMemory, &FTLibrary );
+		
+		if ( Error )
+		{
+			checkf(0, TEXT("Could not init Freetype"));
+		}
+		
+		FT_Add_Default_Modules( FTLibrary );
+#endif // WITH_FREETYPE
+	}
+
+	~FFreeTypeInterface()
+	{
+#if WITH_FREETYPE
+		// Clear before releasing freetype resources
+		Flush();
+
+		FT_Done_Library( FTLibrary );
+		FMemory::Free( CustomMemory );
+#endif // WITH_FREETYPE
+	}
+
+	/**
+	 * Flushes stored data.
+	 */ 
+	void Flush()
+	{
+#if WITH_FREETYPE
+		FontToKerningPairMap.Empty();
+		// toss memory
+		for (TMap<FName, FFontFaceAndMemory>::TIterator It(FontFaceMap); It; ++It)
+		{
+			FMemory::Free(It.Value().Memory);
+		}
+		FontFaceMap.Empty();
+#endif // WITH_FREETYPE
+	}
+
+	/** 
+	 * Creates render data for a specific character 
+	 * 
+	 * @param InFontInfo	Information about the font to render the character with
+	 * @param Char			The character to render
+	 * @param OutCharInfo	Will contain the created render data
+	 */
+	void GetRenderData( const FSlateFontInfo& InFontInfo, TCHAR Char, FCharacterRenderData& OutRenderData, float Scale )
+	{
+#if WITH_FREETYPE
+		// Find or load the face if needed
+		FT_UInt GlyphIndex = 0;
+		FT_Face FontFace = GetFontFace( InFontInfo.FontName );
+
+		if ( FontFace != nullptr ) 
+		{
+			// Get the index to the glyph in the font face
+			GlyphIndex = FT_Get_Char_Index( FontFace, Char );
+		}
+
+		uint32 LocalGlyphFlags = GlyphFlags;
+
+		// If the requested glyph doesn't exist, use the localization fallback font.
+		if( GlyphIndex == 0 )
+		{
+			static FName FallbackFontName( NAME_None );
+			if( FallbackFontName == NAME_None )
+			{
+				FText FallbackFontFilename;
+				FallbackFontName = FName( *( FPaths::EngineContentDir() / TEXT("Slate/Fonts/") / ( NSLOCTEXT("Slate", "FallbackFont", "DroidSansFallback").ToString() + TEXT(".ttf") ) ) );
+			}
+			if( FallbackFontName.IsValid() && FallbackFontName != NAME_None )
+			{
+				FontFace = GetFontFace( FallbackFontName );
+				if (FontFace != nullptr)
+				{					
+					GlyphIndex = FT_Get_Char_Index( FontFace, Char );
+					LocalGlyphFlags |= FT_LOAD_FORCE_AUTOHINT;
+				}
+				else
+				{
+					FallbackFontName = FName(NAME_Inactive);
+				}
+			}
+		}
+
+		// If the requested glyph doesn't exist, use the last resort fallback font.
+		if( GlyphIndex == 0 )
+		{
+			static FName LastResortFontName( NAME_None );
+			if( LastResortFontName == NAME_None )
+			{
+				LastResortFontName = FName( *( FPaths::EngineContentDir() / TEXT("Slate/Fonts/LastResort.ttf") ) );
+			}
+			if( LastResortFontName.IsValid() && LastResortFontName != NAME_None )
+			{
+				FontFace = GetFontFace( LastResortFontName );
+				check( FontFace );
+				GlyphIndex = FT_Get_Char_Index( FontFace, Char );
+				LocalGlyphFlags |= FT_LOAD_FORCE_AUTOHINT;
+			}
+			else
+			{
+				LastResortFontName = FName(NAME_Inactive);
+			}
+		}
+
+		// Set the character size to render at (needs to be in 1/64 of a "point")
+		FT_Error Error = FT_Set_Char_Size( FontFace, 0, InFontInfo.Size*64, FontCacheConstants::HorizontalDPI, FontCacheConstants::VerticalDPI );
+		check(Error==0);
+
+		if( Scale != 1.0f )
+		{
+			FT_Matrix ScaleMatrix;
+			ScaleMatrix.xy = 0;
+			ScaleMatrix.xx = (FT_Fixed)(Scale * 65536);
+			ScaleMatrix.yy = (FT_Fixed)(Scale * 65536);
+			ScaleMatrix.yx = 0;
+			FT_Set_Transform( FontFace, &ScaleMatrix, nullptr );
+		}
+		else
+		{
+			FT_Set_Transform( FontFace, nullptr, nullptr );
+		}
+
+
+		// Load the glyph.  Force using the freetype hinter because not all true type fonts have their own hinting
+		Error = FT_Load_Glyph( FontFace, GlyphIndex, LocalGlyphFlags );
+		check(Error==0);
+
+		// Get the slot for the glyph.  This contains measurement info
+		FT_GlyphSlot Slot = FontFace->glyph;
+
+		FT_Render_Glyph( Slot, FT_RENDER_MODE_NORMAL );
+
+		// one byte per pixel 
+		const uint32 GlyphPixelSize = 1;
+
+		FT_Bitmap Bitmap = Slot->bitmap;
+
+		OutRenderData.RawPixels.Reset();
+		OutRenderData.RawPixels.AddUninitialized( Bitmap.rows * Bitmap.width );
+
+		if( Bitmap.pixel_mode == FT_PIXEL_MODE_MONO )
+		{
+			uint32 BytesPerLine = Bitmap.pitch;
+
+			// For each texel generate a corresponding color for a 32 bit image
+			for( int32 Height = 0; Height < Bitmap.rows; ++Height )
+			{
+				for( uint32 ByteIdx = 0; ByteIdx < BytesPerLine; ++ByteIdx )
+				{
+					// The current byte
+					uint8 Byte = Bitmap.buffer[ Height*BytesPerLine + ByteIdx ];
+
+					// We may not need to use all the bits in a byte.  If there are more than 8 pixels left to check, check all 8
+					// otherwise check whats left
+					uint8 NumBitsToCheck = FMath::Min<uint8>(Bitmap.width-ByteIdx*8, 8);
+
+					// Iterate through each bit in this byte.
+					// Since each byte has up to 8 pixels, we need to generate an uint8 for each bit
+					for( int32 Bit = 0; Bit < NumBitsToCheck; ++Bit )
+					{
+						// Most significant bit is left most bit
+						if( ( Byte & ( 1 << (7-Bit) ) ) == 0 )
+						{
+							// Bit not set, pixel is black
+							OutRenderData.RawPixels[ Height * Bitmap.width + Bit + 8*ByteIdx ] = 0;
+						}
+						else
+						{
+							// Bit set, pixel is white
+							OutRenderData.RawPixels[ Height * Bitmap.width + Bit + 8*ByteIdx ] = 255;
+						}
+					}
+				}
+			}
+		}
+		else
+		{
+			// Copy the rendered bitmap to our raw pixels array
+			for( int32 Row = 0; Row < Bitmap.rows; ++Row )
+			{
+				// Copy a single row. Note Bitmap.pitch contains the offset (in bytes) between rows.  Not always equal to Bitmap.width!
+				FMemory::Memcpy( &OutRenderData.RawPixels[Row*Bitmap.width], &Bitmap.buffer[Row*Bitmap.pitch], Bitmap.width*GlyphPixelSize );
+			}
+		}
+
+		FT_BBox GlyphBox;
+		FT_Glyph Glyph;
+		FT_Get_Glyph( Slot, &Glyph );
+		FT_Glyph_Get_CBox( Glyph, FT_GLYPH_BBOX_PIXELS, &GlyphBox );
+
+		int32 Height = (FT_MulFix( FontFace->height, FontFace->size->metrics.y_scale ) / 64) * Scale;
+
+		// Set measurement info for this character
+		OutRenderData.Char = Char;
+		OutRenderData.MeasureInfo.SizeX = Bitmap.width;
+		OutRenderData.MeasureInfo.SizeY = Bitmap.rows;
+		OutRenderData.MaxHeight = Height;
+
+		// Need to divide by 64 to get pixels;
+		// Descender is not scaled by freetype.  Scale it now. 
+		OutRenderData.MeasureInfo.GlobalDescender = (FontFace->size->metrics.descender / 64) * Scale;
+		// Note we use Slot->advance instead of Slot->metrics.horiAdvance because Slot->Advance contains transformed position (needed if we scale)
+		OutRenderData.MeasureInfo.XAdvance =  Slot->advance.x / 64;
+		OutRenderData.MeasureInfo.HorizontalOffset = Slot->bitmap_left;
+		OutRenderData.MeasureInfo.VerticalOffset = Slot->bitmap_top;
+
+		FT_Done_Glyph( Glyph );
+#endif // WITH_FREETYPE
+	}
+
+	/**
+	 * @param Whether or not the font has kerning
+	 */
+	bool HasKerning( const FSlateFontInfo& InFontInfo )
+	{
+#if WITH_FREETYPE
+		FT_Face FontFace = GetFontFace( InFontInfo.FontName );
+
+		if ( FontFace == nullptr )
+		{
+			return false;
+		}
+
+		return FT_HAS_KERNING( FontFace ) != 0;
+#else
+		return false;
+#endif // WITH_FREETYPE
+	}
+
+	/**
+	 * Calculates the kerning amount for a pair of characters
+	 *
+	 * @param First			The first character in the pair
+	 * @param Second		The second character in the pair
+ 	 * @param InFontInfo	Information about the font that used to draw the string with the first and second characters
+	 * @return The kerning amount, 0 if no kerning
+	 */
+	int8 GetKerning( TCHAR First, TCHAR Second, const FSlateFontKey& FontKey )
+	{
+#if WITH_FREETYPE
+		int32 Kerning = 0;
+		int32* FoundKerning = nullptr;
+
+		FT_Face FontFace = GetFontFace( FontKey.FontInfo.FontName );
+
+		// Check if this font has kerning.  Not all fonts do.
+		if( FontFace != nullptr && FT_HAS_KERNING( FontFace ) )
+		{
+			FT_Error Error = FT_Set_Char_Size( FontFace, 0, FontKey.FontInfo.Size*64, FontCacheConstants::HorizontalDPI, FontCacheConstants::VerticalDPI  );
+
+			if( FontKey.Scale != 1.0f )
+			{
+				FT_Matrix ScaleMatrix;
+				ScaleMatrix.xy = 0;
+				ScaleMatrix.xx = (FT_Fixed)(FontKey.Scale * 65536);
+				ScaleMatrix.yy = (FT_Fixed)(FontKey.Scale * 65536);
+				ScaleMatrix.yx = 0;
+				FT_Set_Transform( FontFace, &ScaleMatrix, nullptr );
+			}
+			else
+			{
+				FT_Set_Transform( FontFace, nullptr, nullptr );
+			}
+
+			check(Error==0);
+		
+			int32 KernValue = 0;
+		
+			FT_UInt FirstIndex = FT_Get_Char_Index( FontFace, First );
+			FT_UInt SecondIndex = FT_Get_Char_Index( FontFace, Second );
+
+			FT_Vector KerningVec;
+			FT_Get_Kerning( FontFace, FirstIndex, SecondIndex, FT_KERNING_DEFAULT, &KerningVec );
+
+			// Return pixel sizes
+			int32 NewKerning = KerningVec.x / 64;
+
+			// Set this pair's kerning value
+			//KeringPairMap->Add( Pair, NewKerning );
+
+			Kerning = NewKerning;
+		}
+
+		return Kerning;
+#else
+		return 0;
+#endif // WITH_FREETYPE
+	}
+
+private:
+
+#if WITH_FREETYPE
+	/**
+	 * Gets or loads a freetype font face
+	 *
+	 * @param The name of the font to load
+	 */
+	FT_Face GetFontFace( const FName& FontName )
+	{
+		FFontFaceAndMemory* FaceAndMemory = FontFaceMap.Find( FontName );
+		if (!FaceAndMemory)
+		{
+			// make a new entry
+			FaceAndMemory = &FontFaceMap.Add(FontName, FFontFaceAndMemory());
+
+			// load the font via UE4 methods, so that it will route through all proper file management (network file loading, etc)
+			FString FontPath = FontName.ToString();
+
+			// default to error condition
+			bool Error = true;
+
+			int64 FileSize = IFileManager::Get().FileSize(*FontPath);
+			if ( FileSize > 0 )
+			{
+				// allocate space for the font
+				FaceAndMemory->Memory = (uint8*)FMemory::Malloc((uint32)FileSize);
+				FArchive* Ar = IFileManager::Get().CreateFileReader(*FontPath);
+				if (Ar != nullptr)
+				{
+					// read in the file
+					Ar->Serialize( FaceAndMemory->Memory, FileSize );
+					delete Ar;
+
+					// initialize the font, setting the error code
+					Error = FT_New_Memory_Face( FTLibrary, FaceAndMemory->Memory, (FT_Long)FileSize, 0, &FaceAndMemory->Face ) != 0;
+				}
+
+				// if it failed, we don't want to keep the memory around
+				if ( Error )
+				{
+					FMemory::Free(FaceAndMemory->Memory);
+					FontFaceMap.Remove(FontName);
+					FaceAndMemory = nullptr;
+				}
+			}
+
+			if ( Error )
+			{
+				UE_LOG( LogSlate, Warning, TEXT("GetFontFace failed to load or process '%s'"), *FontPath);
+				FontFaceMap.Remove(FontName);
+				FaceAndMemory = nullptr;
+			}
+		}
+
+		if ( FaceAndMemory == nullptr )
+		{
+			return nullptr;
+		}
+
+		return FaceAndMemory->Face;
+	}
+
+private:
+	struct FFontFaceAndMemory
+	{
+		// the FT2 object
+		FT_Face Face;
+
+		// the memory for the face (can't be a TArray as the FontFaceMap could be reallocated and copy it's members around)
+		uint8* Memory;
+
+		FFontFaceAndMemory()
+			: Face(nullptr)
+			, Memory(nullptr)
+		{
+		}
+	};
+
+	/** Mapping of font names to freetype faces */
+	TMap<FName,FFontFaceAndMemory> FontFaceMap;
+	/** Mapping of fonts to maps of kerning pairs */
+	TMap<FSlateFontInfo, TMap<FKerningPair,int32> > FontToKerningPairMap; 
+	/** Free type library interface */
+	FT_Library FTLibrary;
+	FT_Memory CustomMemory;
+#endif // WITH_FREETYPE
+};
+
+FKerningTable::FKerningTable( const FSlateFontKey& InFont, const FSlateFontCache& InFontCache )
+	: DirectAccessTable( nullptr )
+	, FontCache( InFontCache )
+	, FontKey( InFont )
+{
+	bHasKerning = InFontCache.HasKerning( InFont.FontInfo );
+}
+
+FKerningTable::~FKerningTable()
+{
+	if( DirectAccessTable )
+	{
+		const uint32 DirectAccessSizeBytes = FontCacheConstants::DirectAccessSize*FontCacheConstants::DirectAccessSize*sizeof(int8);
+
+		FMemory::Free( DirectAccessTable );
+
+		DEC_MEMORY_STAT_BY( STAT_SlateFontKerningTableMemory, DirectAccessSizeBytes );
+	}
+	DEC_MEMORY_STAT_BY( STAT_SlateFontKerningTableMemory, MappedKerningPairs.GetAllocatedSize() );
+}
+
+int8 FKerningTable::GetKerning( TCHAR FirstChar, TCHAR SecondChar )
+{
+	int8 OutKerning = 0;
+
+	if( bHasKerning && FirstChar < FontCacheConstants::DirectAccessSize && SecondChar < FontCacheConstants::DirectAccessSize )
+	{
+		// This character can be directly indexed
+
+		if( !DirectAccessTable )
+		{
+			// Create the table now
+			CreateDirectTable();
+		}
+
+		// Determine the index into the kerning table
+		const uint32 Index = FirstChar*FontCacheConstants::DirectAccessSize+SecondChar;
+
+		OutKerning = DirectAccessTable[Index];
+		// If the kerning value hasn't been accessed yet, get the value from the cache now 
+		if( OutKerning == MAX_int8 )
+		{
+			OutKerning = FontCache.GetKerning( FirstChar, SecondChar, FontKey );
+			DirectAccessTable[Index] = OutKerning;
+		}
+	}
+	else if( bHasKerning )
+	{
+		// Kerning is mapped
+		FKerningPair KerningPair( FirstChar, SecondChar );
+
+		int8* FoundKerning = MappedKerningPairs.Find( KerningPair );
+		if( !FoundKerning )
+		{
+			OutKerning = FontCache.GetKerning( FirstChar, SecondChar, FontKey );
+
+#if STATS
+			const uint32 CurrentMemoryUsage = MappedKerningPairs.GetAllocatedSize();
+#endif
+			MappedKerningPairs.Add( KerningPair, OutKerning );
+		
+#if STATS
+			uint32 NewMemoryUsage = MappedKerningPairs.GetAllocatedSize();
+			if( NewMemoryUsage > CurrentMemoryUsage )
+			{
+				INC_MEMORY_STAT_BY( STAT_SlateFontKerningTableMemory, NewMemoryUsage-CurrentMemoryUsage );
+			}
+#endif
+		}
+	}
+
+	return OutKerning;
+}
+
+void FKerningTable::CreateDirectTable()
+{
+	const uint32 DirectAccessSizeBytes = FontCacheConstants::DirectAccessSize*FontCacheConstants::DirectAccessSize*sizeof(int8);
+
+	check( !DirectAccessTable )
+	DirectAccessTable = (int8*)FMemory::Malloc(DirectAccessSizeBytes);
+	// Invalidate all entries so we know when to get a new kerning value from the implementation
+	FMemory::Memset( DirectAccessTable, MAX_int8, DirectAccessSizeBytes );
+	INC_MEMORY_STAT_BY( STAT_SlateFontKerningTableMemory, DirectAccessSizeBytes );
+}
+
+FCharacterList::FCharacterList( const FSlateFontKey& InFontKey, const FSlateFontCache& InFontCache )
+	: KerningTable( InFontKey, InFontCache )
+	, FontKey( InFontKey )
+	, FontCache( InFontCache )
+	, MaxDirectIndexedEntries( FontCacheConstants::DirectAccessSize )
+	, MaxHeight( 0 )
+	, Baseline( 0 )
+{
+
+
+}
+
+int8 FCharacterList::GetKerning( TCHAR First, TCHAR Second )
+{
+	return KerningTable.GetKerning( First, Second );
+}
+
+uint16 FCharacterList::GetMaxHeight() const
+{
+	if( MaxHeight == 0 )
+	{
+		MaxHeight = FontCache.GetMaxCharacterHeight( FontKey.FontInfo, FontKey.Scale );
+	}
+
+	return MaxHeight;
+}
+
+int16 FCharacterList::GetBaseline() const
+{
+	if( Baseline == 0 )
+	{
+		Baseline = FontCache.GetBaseline( FontKey.FontInfo, FontKey.Scale );
+	}
+
+	return Baseline;
+}
+
+FCharacterEntry& FCharacterList::CacheCharacter( TCHAR Character )
+{
+	SCOPE_CYCLE_COUNTER( STAT_SlateFontCachingTime );
+
+	FCharacterEntry NewEntry;
+	bool bSuccess = FontCache.AddNewEntry( Character, FontKey, NewEntry );
+
+	check( bSuccess );
+
+	if( Character < MaxDirectIndexedEntries && NewEntry.IsValidEntry() )
+	{
+		DirectIndexEntries[ Character ] = NewEntry;
+		return DirectIndexEntries[ Character ];
+	}
+	else
+	{
+		return MappedEntries.Add( Character, NewEntry );
+	}
+}
+
+FSlateFontCache::FSlateFontCache( TSharedRef<ISlateFontAtlasFactory> InFontAtlasFactory )
+	: FTInterface( new FFreeTypeInterface )
+	, FontAtlasFactory( InFontAtlasFactory )
+	, bFlushRequested( false )
+{
+
+}
+
+FSlateFontCache::~FSlateFontCache()
+{	
+
+}
+
+bool FSlateFontCache::AddNewEntry( TCHAR Character, const FSlateFontKey& InKey, FCharacterEntry& OutCharacterEntry ) const
+{
+	FCharacterRenderData RenderData;
+	// Render the character 
+	FTInterface->GetRenderData( InKey.FontInfo, Character, RenderData, InKey.Scale );	
+
+	// the index of the atlas where the character is stored
+	int32 AtlasIndex = 0;
+	const FAtlasedTextureSlot* NewSlot = nullptr;
+
+	for( AtlasIndex = 0; AtlasIndex < FontAtlases.Num(); ++AtlasIndex ) 
+	{
+		// Add the character to the texture
+		NewSlot = FontAtlases[AtlasIndex]->AddCharacter( RenderData );
+		if( NewSlot )
+		{
+			break;
+		}
+	}
+
+	if( !NewSlot )
+	{
+		TSharedRef<FSlateFontAtlas> FontAtlas = FontAtlasFactory->CreateFontAtlas();
+
+		// Add the character to the texture
+		NewSlot = FontAtlas->AddCharacter( RenderData );
+
+		AtlasIndex = FontAtlases.Add( FontAtlas );
+
+		INC_DWORD_STAT_BY( STAT_SlateNumFontAtlases, 1 );
+
+		if( FontAtlases.Num() > 1 )
+		{
+			// There is more than one font atlas which means there is a lot of font data being cached
+			// try to shrink it next time
+			bFlushRequested = true;
+		}
+	}
+
+	// If the new slot is null and we are not just measuring, then the atlas is full
+	bool bSuccess = NewSlot != nullptr;
+
+	if( bSuccess )
+	{
+		OutCharacterEntry.StartU = NewSlot->X;
+		OutCharacterEntry.StartV = NewSlot->Y;
+		OutCharacterEntry.USize = NewSlot->Width;
+		OutCharacterEntry.VSize = NewSlot->Height;
+		OutCharacterEntry.TextureIndex = 0;
+		OutCharacterEntry.XAdvance = RenderData.MeasureInfo.XAdvance;
+		OutCharacterEntry.VerticalOffset = RenderData.MeasureInfo.VerticalOffset;
+		OutCharacterEntry.GlobalDescender = RenderData.MeasureInfo.GlobalDescender;
+		OutCharacterEntry.HorizontalOffset = RenderData.MeasureInfo.HorizontalOffset;
+		OutCharacterEntry.TextureIndex = AtlasIndex;
+		OutCharacterEntry.Valid = 1;
+	}
+
+	return bSuccess;
+}
+
+FCharacterList& FSlateFontCache::GetCharacterList( const FSlateFontInfo &InFontInfo, float FontScale ) const
+{
+	// Create a key for looking up each character
+	FSlateFontKey FontKey( InFontInfo, FontScale );
+
+	TSharedRef< class FCharacterList >* CachedCharacterList = FontToCharacterListCache.Find( FontKey );
+
+	// Additional setup
+	if( CachedCharacterList == nullptr )
+	{
+		return FontToCharacterListCache.Add( FontKey, MakeShareable( new FCharacterList( FontKey, *this ) ) ).Get();
+	}
+
+	return CachedCharacterList->Get();
+}
+
+uint16 FSlateFontCache::GetMaxCharacterHeight( const FSlateFontInfo& InFontInfo, float FontScale ) const
+{
+	FCharacterRenderData NewRenderData;
+
+	// Just render the null character 
+	TCHAR Char = 0;
+	// Render the character 
+	FTInterface->GetRenderData( InFontInfo, Char, NewRenderData, FontScale );	
+
+	return NewRenderData.MaxHeight;
+}
+
+uint16 FSlateFontCache::GetBaseline( const FSlateFontInfo& InFontInfo, float FontScale ) const
+{
+	FCharacterRenderData NewRenderData;
+
+	// Just render the null character 
+	TCHAR Char = 0;
+	// Render the character 
+	FTInterface->GetRenderData( InFontInfo, Char, NewRenderData, FontScale );	
+
+	return NewRenderData.MeasureInfo.GlobalDescender;
+}
+
+int8 FSlateFontCache::GetKerning( TCHAR First, TCHAR Second, const FSlateFontKey& FontKey ) const 
+{
+	return FTInterface->GetKerning( First, Second, FontKey );
+}
+
+bool FSlateFontCache::HasKerning( const FSlateFontInfo& InFontInfo ) const
+{
+	return FTInterface->HasKerning( InFontInfo );
+}
+
+void FSlateFontCache::ConditionalFlushCache()
+{
+	if( bFlushRequested )
+	{
+		FlushCache();
+		bFlushRequested = false;
+	}
+}
+
+void FSlateFontCache::UpdateCache()
+{
+	for( int32 AtlasIndex = 0; AtlasIndex < FontAtlases.Num(); ++AtlasIndex ) 
+	{
+		FontAtlases[AtlasIndex]->ConditionalUpdateTexture();
+	}
+}
+
+void FSlateFontCache::ReleaseResources()
+{
+	for( int32 AtlasIndex = 0; AtlasIndex < FontAtlases.Num(); ++AtlasIndex ) 
+	{
+		FontAtlases[AtlasIndex]->ReleaseResources();
+	}
+}
+
+void FSlateFontCache::FlushCache() const
+{
+	FontToCharacterListCache.Empty();
+	FTInterface->Flush();
+
+	for( int32 AtlasIndex = 0; AtlasIndex < FontAtlases.Num(); ++AtlasIndex ) 
+	{
+		FontAtlases[AtlasIndex]->ReleaseResources();
+	}
+
+	// hack
+	FSlateApplicationBase::Get().GetRenderer()->FlushCommands();
+
+	SET_DWORD_STAT( STAT_SlateNumFontAtlases, 0 );
+
+	FontAtlases.Empty();
+
+	UE_LOG( LogSlate, Verbose, TEXT("Slate font cache was flushed") );
+}
diff --git a/Engine/Source/Runtime/SlateCore/Private/Fonts/FontCacheUtils.h b/Engine/Source/Runtime/SlateCore/Private/Fonts/FontCacheUtils.h
new file mode 100644
--- /dev/null
+++ b/Engine/Source/Runtime/SlateCore/Private/Fonts/FontCacheUtils.h
@@ -0,0 +1,230 @@
+// Copyright 1998-2014 Epic Games, Inc. All Rights Reserved.
+
+/*=============================================================================
+	FontCacheUtils.h: Declares the FLRUStringCache class.
+=============================================================================*/
+
+#pragma once
+
+
+typedef FString KeyType;
+typedef FVector2D ValueType;
+
+
+/**
+ * Basic Least Recently Used (LRU) cache
+ */
+class FLRUStringCache
+{
+public:
+	FLRUStringCache( int32 InMaxNumElements )
+		: LookupSet()
+		, MostRecent(nullptr)
+		, LeastRecent(nullptr)
+		, MaxNumElements(InMaxNumElements)
+	{ }
+
+	~FLRUStringCache()
+	{
+		Empty();
+	}
+
+	/**
+	 * Accesses an item in the cache.  
+	 */
+	FORCEINLINE const ValueType* AccessItem(const KeyType& Key)
+	{
+		CacheEntry** Entry = LookupSet.Find( Key );
+		if( Entry )
+		{
+			MarkAsRecent( *Entry );
+			return &((*Entry)->Value);
+		}
+		
+		return nullptr;
+	}
+
+	void Add( const KeyType& Key, const ValueType& Value )
+	{
+		CacheEntry** Entry = LookupSet.Find( Key );
+	
+		// Make a new link
+		if( !Entry )
+		{
+			if( LookupSet.Num() == MaxNumElements )
+			{
+				Eject();
+				checkf( LookupSet.Num() < MaxNumElements, TEXT("Could not eject item from the LRU: (%d of %d), %s"), LookupSet.Num(), MaxNumElements, *LeastRecent->Key );
+			}
+
+			CacheEntry* NewEntry = new CacheEntry( Key, Value );
+
+			// Link before the most recent so that we become the most recent
+			NewEntry->Link(MostRecent);
+			MostRecent = NewEntry;
+
+			if( LeastRecent == nullptr )
+			{
+				LeastRecent = NewEntry;
+			}
+
+			STAT( uint32 CurrentMemUsage = LookupSet.GetAllocatedSize() );
+			LookupSet.Add( NewEntry );
+			STAT( uint32 NewMemUsage = LookupSet.GetAllocatedSize() );
+			INC_MEMORY_STAT_BY( STAT_SlateFontMeasureCacheMemory,  NewMemUsage - CurrentMemUsage );
+		}
+		else
+		{
+			// Trying to add an existing value 
+			CacheEntry* EntryPtr = *Entry;
+			
+			// Update the value
+			EntryPtr->Value = Value;
+			checkSlow( EntryPtr->Key == Key );
+			// Mark as the most recent
+			MarkAsRecent( EntryPtr );
+		}
+	}
+
+	void Empty()
+	{
+		DEC_MEMORY_STAT_BY( STAT_SlateFontMeasureCacheMemory, LookupSet.GetAllocatedSize() );
+
+		for( TSet<CacheEntry*, FCaseSensitiveStringKeyFuncs >::TIterator It(LookupSet); It; ++It )
+		{
+			CacheEntry* Entry = *It;
+			// Note no need to unlink anything here. we are emptying the entire list
+			delete Entry;
+		}
+		LookupSet.Empty();
+
+		MostRecent = LeastRecent = nullptr;
+	}
+private:
+	struct CacheEntry
+	{
+		KeyType Key;
+		ValueType Value;
+		CacheEntry* Next;
+		CacheEntry* Prev;
+
+		CacheEntry( const KeyType& InKey, const ValueType& InValue )
+			: Key( InKey )
+			, Value( InValue )
+			, Next( nullptr )
+			, Prev( nullptr )
+		{
+			INC_MEMORY_STAT_BY( STAT_SlateFontMeasureCacheMemory, Key.GetAllocatedSize()+sizeof(Value)+sizeof(Next)+sizeof(Prev) );
+		}
+
+		~CacheEntry()
+		{
+			DEC_MEMORY_STAT_BY( STAT_SlateFontMeasureCacheMemory, Key.GetAllocatedSize()+sizeof(Value)+sizeof(Next)+sizeof(Prev) );
+		}
+
+		FORCEINLINE void Link( CacheEntry* Before )
+		{
+			Next = Before;
+
+			if( Before )
+			{
+				Before->Prev = this;
+			}
+		}
+
+		FORCEINLINE void Unlink()
+		{
+			if( Prev )
+			{
+				Prev->Next = Next;
+			}
+
+			if( Next )
+			{
+				Next->Prev = Prev;
+			}
+
+			Prev = nullptr;
+			Next = nullptr;
+		}
+	};
+
+	struct FCaseSensitiveStringKeyFuncs : BaseKeyFuncs<CacheEntry*, FString>
+	{
+		FORCEINLINE static const KeyType& GetSetKey( const CacheEntry* Entry )
+		{
+			return Entry->Key;
+		}
+
+		FORCEINLINE static bool Matches(const KeyType& A,const KeyType& B)
+		{
+			return A.Equals( B, ESearchCase::CaseSensitive );
+		}
+
+		FORCEINLINE static uint32 GetKeyHash(const KeyType& Identifier)
+		{
+			return FCrc::StrCrc_DEPRECATED( *Identifier );
+		}
+	};
+
+	
+
+	/**
+	 * Marks the link as the the most recent
+	 *
+	 * @param Link	The link to mark as most recent
+	 */
+	FORCEINLINE void MarkAsRecent( CacheEntry* Entry )
+	{
+		checkSlow( LeastRecent && MostRecent );
+
+		// If we are the least recent entry we are no longer the least recent
+		// The previous least recent item is now the most recent.  If it is nullptr then this entry is the only item in the list
+		if( Entry == LeastRecent && LeastRecent->Prev != nullptr )
+		{
+			LeastRecent = LeastRecent->Prev;
+		}
+
+		// No need to relink if we happen to already be accessing the most recent value
+		if( Entry != MostRecent )
+		{
+			// Unlink from its current spot
+			Entry->Unlink();
+
+			// Relink before the most recent so that we become the most recent
+			Entry->Link(MostRecent);
+			MostRecent = Entry;
+		}
+	}
+
+	/**
+	 * Removes the least recent item from the cache
+	 */
+	FORCEINLINE void Eject()
+	{
+		CacheEntry* EntryToRemove = LeastRecent;
+		// Eject the least recent, no more space
+		check( EntryToRemove );
+		STAT( uint32 CurrentMemUsage = LookupSet.GetAllocatedSize() );
+		LookupSet.Remove( EntryToRemove->Key );
+		STAT( uint32 NewMemUsage = LookupSet.GetAllocatedSize() );
+		DEC_MEMORY_STAT_BY( STAT_SlateFontMeasureCacheMemory,  CurrentMemUsage - NewMemUsage );
+
+		LeastRecent = LeastRecent->Prev;
+
+		// Unlink the LRU
+		EntryToRemove->Unlink();
+
+		delete EntryToRemove;
+	}
+
+private:
+	TSet< CacheEntry*, FCaseSensitiveStringKeyFuncs > LookupSet;
+	/** Most recent item in the cache */
+	CacheEntry* MostRecent;
+	/** Least recent item in the cache */
+	CacheEntry* LeastRecent;
+	/** The maximum number of elements in the cache */
+	int32 MaxNumElements;
+};
+
diff --git a/Engine/Source/Runtime/SlateCore/Private/Fonts/FontMeasure.cpp b/Engine/Source/Runtime/SlateCore/Private/Fonts/FontMeasure.cpp
new file mode 100644
--- /dev/null
+++ b/Engine/Source/Runtime/SlateCore/Private/Fonts/FontMeasure.cpp
@@ -0,0 +1,280 @@
+// Copyright 1998-2014 Epic Games, Inc. All Rights Reserved.
+
+#include "SlateCorePrivatePCH.h"
+
+
+namespace FontMeasureConstants
+{
+	/** Number of possible elements in each measurement cache */
+	const uint32 MeasureCacheSize = 500;
+}
+
+TSharedRef< FSlateFontMeasure > FSlateFontMeasure::Create( const TSharedRef<class FSlateFontCache>& FontCache )
+{
+	return MakeShareable( new FSlateFontMeasure( FontCache ) );
+}
+
+FSlateFontMeasure::FSlateFontMeasure( const TSharedRef<class FSlateFontCache>& InFontCache )
+	: FontCache( InFontCache )
+{
+
+}
+
+FVector2D FSlateFontMeasure::Measure( const FString& Text, const FSlateFontInfo& InFontInfo, float FontScale ) const
+{
+	int32 DummyLastCharacterIndex;
+	return MeasureStringInternal( Text, 0, Text.Len(), InFontInfo, false, FontScale, INDEX_NONE, ELastCharacterIndexFormat::Unused, DummyLastCharacterIndex);
+}
+
+FVector2D FSlateFontMeasure::Measure( const FText& Text, const FSlateFontInfo& InFontInfo, float FontScale ) const
+{
+	int32 DummyLastCharacterIndex;
+	const FString String = Text.ToString();
+	return MeasureStringInternal( String, 0, String.Len(), InFontInfo, false, FontScale, INDEX_NONE, ELastCharacterIndexFormat::Unused, DummyLastCharacterIndex);
+}
+
+FVector2D FSlateFontMeasure::Measure( const FString& Text, int32 StartIndex, int32 EndIndex, const FSlateFontInfo& InFontInfo, bool IncludeKerningWithPrecedingChar, float FontScale ) const
+{
+	int32 DummyLastCharacterIndex;
+	return MeasureStringInternal( Text, StartIndex, EndIndex, InFontInfo, IncludeKerningWithPrecedingChar, FontScale, INDEX_NONE, ELastCharacterIndexFormat::Unused, DummyLastCharacterIndex);
+}
+
+int32 FSlateFontMeasure::FindLastWholeCharacterIndexBeforeOffset( const FString& Text, const FSlateFontInfo& InFontInfo, const int32 HorizontalOffset, float FontScale ) const
+{
+	return FindLastWholeCharacterIndexBeforeOffset( Text, 0, Text.Len(), InFontInfo, HorizontalOffset, false, FontScale );
+}
+
+int32 FSlateFontMeasure::FindLastWholeCharacterIndexBeforeOffset( const FText& Text, const FSlateFontInfo& InFontInfo, const int32 HorizontalOffset, float FontScale ) const
+{
+	return FindLastWholeCharacterIndexBeforeOffset( Text.ToString(), InFontInfo, HorizontalOffset, FontScale );
+}
+
+
+int32 FSlateFontMeasure::FindLastWholeCharacterIndexBeforeOffset( const FString& Text, int32 StartIndex, int32 EndIndex, const FSlateFontInfo& InFontInfo, const int32 HorizontalOffset, bool IncludeKerningWithPrecedingChar, float FontScale ) const
+{
+	int32 FoundLastCharacterIndex = INDEX_NONE;
+	MeasureStringInternal( Text, StartIndex, EndIndex, InFontInfo, IncludeKerningWithPrecedingChar, FontScale, HorizontalOffset, ELastCharacterIndexFormat::LastWholeCharacterBeforeOffset, FoundLastCharacterIndex );
+	return FoundLastCharacterIndex;
+}
+
+
+int32 FSlateFontMeasure::FindFirstWholeCharacterIndexAfterOffset( const FText& Text, const FSlateFontInfo& InFontInfo, const int32 HorizontalOffset, float FontScale ) const
+{
+	return FindFirstWholeCharacterIndexAfterOffset( Text.ToString(), InFontInfo, HorizontalOffset, FontScale );
+}
+
+int32 FSlateFontMeasure::FindFirstWholeCharacterIndexAfterOffset( const FString& Text, const FSlateFontInfo& InFontInfo, const int32 HorizontalOffset, float FontScale ) const
+{
+	return FindFirstWholeCharacterIndexAfterOffset( Text, 0, Text.Len(), InFontInfo, HorizontalOffset, false, FontScale );
+}
+
+int32 FSlateFontMeasure::FindFirstWholeCharacterIndexAfterOffset( const FString& Text, int32 StartIndex, int32 EndIndex, const FSlateFontInfo& InFontInfo, const int32 HorizontalOffset, bool IncludeKerningWithPrecedingChar, float FontScale ) const
+{
+	int32 FoundLastCharacterIndex = FindCharacterIndexAtOffset( Text, StartIndex, EndIndex, InFontInfo, HorizontalOffset, IncludeKerningWithPrecedingChar, FontScale );
+
+	float TextWidth = Measure( Text, StartIndex, EndIndex, InFontInfo, IncludeKerningWithPrecedingChar, FontScale ).X;
+	float AvailableWidth = TextWidth - HorizontalOffset;
+
+	float RightStringWidth = Measure( Text, FoundLastCharacterIndex, EndIndex, InFontInfo, IncludeKerningWithPrecedingChar, FontScale ).X;
+	if ( AvailableWidth < RightStringWidth )
+	{
+		++FoundLastCharacterIndex;
+	}
+
+	return FoundLastCharacterIndex;
+}
+
+int32 FSlateFontMeasure::FindCharacterIndexAtOffset( const FString& Text, const FSlateFontInfo& InFontInfo, const int32 HorizontalOffset, float FontScale ) const
+{
+	return FindCharacterIndexAtOffset( Text, 0, Text.Len(), InFontInfo, HorizontalOffset, false, FontScale );
+}
+
+int32 FSlateFontMeasure::FindCharacterIndexAtOffset( const FText& Text, const FSlateFontInfo& InFontInfo, const int32 HorizontalOffset, float FontScale ) const
+{
+	return FindCharacterIndexAtOffset( Text.ToString(), InFontInfo, HorizontalOffset, FontScale );
+}
+
+int32 FSlateFontMeasure::FindCharacterIndexAtOffset( const FString& Text, int32 StartIndex, int32 EndIndex, const FSlateFontInfo& InFontInfo, const int32 HorizontalOffset, bool IncludeKerningWithPrecedingChar, float FontScale ) const
+{
+	int32 FoundLastCharacterIndex = 0;
+	MeasureStringInternal( Text, StartIndex, EndIndex, InFontInfo, IncludeKerningWithPrecedingChar, FontScale, HorizontalOffset, ELastCharacterIndexFormat::CharacterAtOffset, FoundLastCharacterIndex );
+	return FoundLastCharacterIndex;
+}
+
+FVector2D FSlateFontMeasure::MeasureStringInternal( const FString& Text, int32 StartIndex, int32 EndIndex, const FSlateFontInfo& InFontInfo, bool IncludeKerningWithPrecedingChar, float FontScale, int32 StopAfterHorizontalOffset, ELastCharacterIndexFormat CharIndexFormat, int32& OutLastCharacterIndex ) const
+{
+	SCOPE_CYCLE_COUNTER(STAT_SlateMeasureStringTime);
+	FCharacterList& CharacterList = FontCache->GetCharacterList( InFontInfo, FontScale );
+	const uint16 MaxHeight = CharacterList.GetMaxHeight();
+
+	const bool DoesStartAtBeginning = StartIndex == 0;
+	const bool DoesFinishAtEnd = EndIndex == Text.Len();
+	const int32 TextRangeLength = EndIndex - StartIndex;
+	if ( EndIndex - StartIndex <= 0 || EndIndex <= 0 || StartIndex < 0 || EndIndex <= StartIndex )
+	{
+		return FVector2D( 0, MaxHeight );
+	}
+
+#define USE_MEASURE_CACHING 1
+#if USE_MEASURE_CACHING
+	TSharedPtr< FMeasureCache > CurrentMeasureCache = nullptr;
+	// Do not cache strings which have small sizes or which have complicated measure requirements
+	if( DoesStartAtBeginning && DoesFinishAtEnd && !IncludeKerningWithPrecedingChar && TextRangeLength > 5 && StopAfterHorizontalOffset == INDEX_NONE )
+	{
+		FSlateFontKey FontKey(InFontInfo,FontScale);
+		TSharedPtr< FMeasureCache > FoundMeasureCache = FontToMeasureCache.FindRef( FontKey );
+
+		if ( FoundMeasureCache.IsValid() )
+		{
+			CurrentMeasureCache = FoundMeasureCache;
+			const FVector2D* CachedMeasurement = CurrentMeasureCache->AccessItem( Text );
+			if( CachedMeasurement )
+			{
+				return *CachedMeasurement;
+			}
+		}
+		else
+		{
+			CurrentMeasureCache = MakeShareable( new FMeasureCache( FontMeasureConstants::MeasureCacheSize ) );
+			FontToMeasureCache.Add( FontKey, CurrentMeasureCache );
+		}
+	}
+#endif
+
+	// The size of the string
+	FVector2D Size(0,0);
+	// Widest line encountered while drawing this text.
+	int32 MaxLineWidth = 0;
+	// The width of the current line so far.
+	int32 CurrentX = 0;
+	// Accumulated height of this block of text
+	int32 StringSizeY = MaxHeight;
+	// The previous char (for kerning)
+	TCHAR PreviousChar = 0;
+
+	//If we are measuring a range then we should take into account the kerning with the character before the start of the range
+	if ( !DoesStartAtBeginning && IncludeKerningWithPrecedingChar )
+	{
+		PreviousChar = Text[ StartIndex - 1 ];
+	}
+
+	int32 FinalPosX = 0;
+
+	int32 CharIndex;
+	for( CharIndex = StartIndex; CharIndex < EndIndex; ++CharIndex )
+	{
+		TCHAR CurrentChar = Text[CharIndex];
+
+		const bool IsNewline = (CurrentChar == '\n');
+
+		if (IsNewline)
+		{
+			// New line means
+			// 1) we accumulate total height
+			StringSizeY += MaxHeight;
+			// 2) update the longest line we've encountered
+			MaxLineWidth = FMath::Max(CurrentX, MaxLineWidth);
+			// 3) the next line starts at the beginning
+			CurrentX = 0;
+		}
+		else
+		{
+			const FCharacterEntry& Entry = CharacterList[CurrentChar];
+
+			int32 Kerning = 0;
+			if( CharIndex > 0 )
+			{
+				Kerning = CharacterList.GetKerning( PreviousChar, CurrentChar );
+			}
+
+			const int32 TotalCharSpacing = 
+				Kerning + Entry.HorizontalOffset +		// Width is any kerning plus how much to advance the position when drawing a new character
+				Entry.XAdvance;	// How far we advance 
+
+			PreviousChar = CurrentChar;
+
+			CurrentX += Kerning + Entry.XAdvance;
+
+			// Were we asked to stop measuring after the specified horizontal offset in pixels?
+			if( StopAfterHorizontalOffset != INDEX_NONE )
+			{
+				if( CharIndexFormat == ELastCharacterIndexFormat::CharacterAtOffset )
+				{
+					// Round our test toward the character's center position
+					if( StopAfterHorizontalOffset < CurrentX - TotalCharSpacing / 2 )
+					{
+						// We've reached the stopping point, so bail
+						break;
+					}
+				}
+				else if( CharIndexFormat == ELastCharacterIndexFormat::LastWholeCharacterBeforeOffset )
+				{
+					if( StopAfterHorizontalOffset < CurrentX )
+					{
+						//LastWholeCharacterBeforeOffset is an inclusive index
+						--CharIndex;
+						if ( CharIndex < StartIndex )
+						{
+							CharIndex = INDEX_NONE;
+						}
+
+						// We've reached the stopping point, so bail
+						break;
+					}
+				}
+			}
+		}
+	}
+
+	// We just finished a line, so need to update the longest line encountered.
+	MaxLineWidth = FMath::Max(CurrentX, MaxLineWidth);
+
+	Size.X = MaxLineWidth;
+	Size.Y = StringSizeY;
+
+	if (CharIndexFormat == ELastCharacterIndexFormat::LastWholeCharacterBeforeOffset)
+	{
+		if (CharIndex >= EndIndex)
+		{
+			//LastWholeCharacterBeforeOffset is an inclusive index
+			CharIndex = EndIndex - 1;
+		}
+	}
+
+	OutLastCharacterIndex = CharIndex;
+
+#if USE_MEASURE_CACHING
+	if( StopAfterHorizontalOffset == INDEX_NONE && CurrentMeasureCache.IsValid() )
+	{
+		CurrentMeasureCache->Add( Text, Size );
+	}
+#endif
+	return Size;
+}
+
+uint16 FSlateFontMeasure::GetMaxCharacterHeight( const FSlateFontInfo& InFontInfo, float FontScale ) const
+{
+	FCharacterList& CharacterList = FontCache->GetCharacterList( InFontInfo, FontScale );
+
+	const uint16 MaxHeight = CharacterList.GetMaxHeight();
+	return MaxHeight;
+}
+
+uint8 FSlateFontMeasure::GetKerning( const FSlateFontInfo& InFontInfo, float FontScale, TCHAR PreviousCharacter, TCHAR CurrentCharacter ) const
+{
+	FCharacterList& CharacterList = FontCache->GetCharacterList( InFontInfo, FontScale );
+	return CharacterList.GetKerning( PreviousCharacter, CurrentCharacter );
+}
+
+uint16 FSlateFontMeasure::GetBaseline( const FSlateFontInfo& InFontInfo, float FontScale ) const
+{
+	FCharacterList& CharacterList = FontCache->GetCharacterList( InFontInfo, FontScale );
+
+	const uint16 Baseline = CharacterList.GetBaseline();
+	return Baseline;
+}
+
+void FSlateFontMeasure::FlushCache()
+{
+	FontToMeasureCache.Empty();
+}
diff --git a/Engine/Source/Runtime/SlateCore/Private/Fonts/FontTypes.cpp b/Engine/Source/Runtime/SlateCore/Private/Fonts/FontTypes.cpp
new file mode 100644
--- /dev/null
+++ b/Engine/Source/Runtime/SlateCore/Private/Fonts/FontTypes.cpp
@@ -0,0 +1,36 @@
+// Copyright 1998-2014 Epic Games, Inc. All Rights Reserved.
+
+#include "SlateCorePrivatePCH.h"
+
+
+FSlateFontAtlas::FSlateFontAtlas( uint32 InWidth, uint32 InHeight )
+	: FSlateTextureAtlas( InWidth, InHeight, sizeof(uint8), 0 )
+{
+}
+
+FSlateFontAtlas::~FSlateFontAtlas()
+{
+
+}	
+
+/** 
+ * Adds a character to the texture.
+ *
+ * @param CharInfo	Information about the size of the character
+ */
+const FAtlasedTextureSlot* FSlateFontAtlas::AddCharacter( const FCharacterRenderData& RenderData )
+{
+	return AddTexture( RenderData.MeasureInfo.SizeX, RenderData.MeasureInfo.SizeY, RenderData.RawPixels );
+}
+
+void FSlateFontAtlas::Flush()
+{
+	// Empty the atlas
+	Empty();
+
+	// Recreate the data
+	InitAtlasData();
+
+	bNeedsUpdate = true;
+	ConditionalUpdateTexture();
+}
\ No newline at end of file
diff --git a/Engine/Source/Runtime/SlateCore/Private/Fonts/SlateFontInfo.cpp b/Engine/Source/Runtime/SlateCore/Private/Fonts/SlateFontInfo.cpp
new file mode 100644
--- /dev/null
+++ b/Engine/Source/Runtime/SlateCore/Private/Fonts/SlateFontInfo.cpp
@@ -0,0 +1,52 @@
+// Copyright 1998-2014 Epic Games, Inc. All Rights Reserved.
+
+/*=============================================================================
+	SlateFontInfo.cpp: Implements the FSlateFontInfo structure.
+=============================================================================*/
+
+#include "SlateCorePrivatePCH.h"
+
+
+/* FSlateFontInfo structors
+ *****************************************************************************/
+
+FSlateFontInfo::FSlateFontInfo( )
+	: FontName(NAME_None)
+	, Size(0)
+{ }
+
+
+FSlateFontInfo::FSlateFontInfo( const FString& InFontName, uint16 InSize )
+	: FontName(*InFontName)
+	, Size(InSize)
+{
+	//Useful for debugging style breakages
+	//check( FPaths::FileExists( FontName.ToString() ) );
+}
+
+
+FSlateFontInfo::FSlateFontInfo( const FName& InFontName, uint16 InSize )
+	: FontName(InFontName)
+	, Size(InSize)
+{
+	//Useful for debugging style breakages
+	//check( FPaths::FileExists( FontName.ToString() ) );
+}
+
+
+FSlateFontInfo::FSlateFontInfo( const ANSICHAR* InFontName, uint16 InSize )
+	: FontName(InFontName)
+	, Size(InSize)
+{
+	//Useful for debugging style breakages
+	//check( FPaths::FileExists( FontName.ToString() ) );
+}
+
+
+FSlateFontInfo::FSlateFontInfo( const WIDECHAR* InFontName, uint16 InSize )
+	: FontName(InFontName)
+	, Size(InSize)
+{
+	//Useful for debugging style breakages
+	//check( FPaths::FileExists( FontName.ToString() ) );
+}
diff --git a/Engine/Source/Runtime/SlateCore/Private/Layout/ArrangedChildren.cpp b/Engine/Source/Runtime/SlateCore/Private/Layout/ArrangedChildren.cpp
new file mode 100644
--- /dev/null
+++ b/Engine/Source/Runtime/SlateCore/Private/Layout/ArrangedChildren.cpp
@@ -0,0 +1,25 @@
+// Copyright 1998-2014 Epic Games, Inc. All Rights Reserved.
+
+/*=============================================================================
+	ArrangedChildren.cpp: Implements the FArrangedChildren class.
+=============================================================================*/
+
+#include "SlateCorePrivatePCH.h"
+
+
+/* FArrangedChildren interface
+ *****************************************************************************/
+
+void FArrangedChildren::AddWidget( const FArrangedWidget& InWidgetGeometry )
+{
+	AddWidget(InWidgetGeometry.Widget->GetVisibility(), InWidgetGeometry);
+}
+
+
+void FArrangedChildren::AddWidget( EVisibility VisibilityOverride, const FArrangedWidget& InWidgetGeometry )
+{
+	if (this->Accepts(VisibilityOverride))
+	{
+		Array.Add(InWidgetGeometry);
+	}
+}
diff --git a/Engine/Source/Runtime/SlateCore/Private/Layout/ArrangedWidget.cpp b/Engine/Source/Runtime/SlateCore/Private/Layout/ArrangedWidget.cpp
new file mode 100644
--- /dev/null
+++ b/Engine/Source/Runtime/SlateCore/Private/Layout/ArrangedWidget.cpp
@@ -0,0 +1,16 @@
+// Copyright 1998-2014 Epic Games, Inc. All Rights Reserved.
+
+/*=============================================================================
+	ArrangedWidget.cpp: Implements the FArrangedWidget class.
+=============================================================================*/
+
+#include "SlateCorePrivatePCH.h"
+
+
+/* FArrangedWidget interface
+ *****************************************************************************/
+
+FString FArrangedWidget::ToString( ) const
+{
+	return FString::Printf(TEXT("%s @ %s"), *Widget->ToString(), *Geometry.ToString());
+}
diff --git a/Engine/Source/Runtime/SlateCore/Private/Layout/Geometry.cpp b/Engine/Source/Runtime/SlateCore/Private/Layout/Geometry.cpp
new file mode 100644
--- /dev/null
+++ b/Engine/Source/Runtime/SlateCore/Private/Layout/Geometry.cpp
@@ -0,0 +1,60 @@
+// Copyright 1998-2014 Epic Games, Inc. All Rights Reserved.
+
+/*=============================================================================
+	Geometry.cpp: Implements the FGeometry class.
+=============================================================================*/
+
+#include "SlateCorePrivatePCH.h"
+
+
+/* FArrangedChildren interface
+ *****************************************************************************/
+
+FArrangedWidget FGeometry::MakeChild( const TSharedRef<SWidget>& InWidget, const FVector2D& ChildOffset, const FVector2D& ChildSize, float InScale ) const
+{
+	return FArrangedWidget( InWidget, this->MakeChild(ChildOffset, ChildSize, InScale) );
+}
+
+FPaintGeometry FGeometry::ToPaintGeometry( FVector2D InOffset, FVector2D InSize, float InScale ) const
+{
+	const float CombinedScale = this->Scale*InScale;
+	return FPaintGeometry( this->AbsolutePosition + InOffset*CombinedScale, InSize * CombinedScale, CombinedScale );
+}
+
+FPaintGeometry FGeometry::ToInflatedPaintGeometry( const FVector2D& InflateAmount ) const
+{
+	return FPaintGeometry( this->AbsolutePosition - InflateAmount*this->Scale, (this->Size+InflateAmount*2)*this->Scale, this->Scale );
+}
+
+FPaintGeometry FGeometry::CenteredPaintGeometryOnLeft( const FVector2D& SizeBeingAligned, float InScale ) const
+{
+	const float CombinedScale = this->Scale*InScale;
+	
+	return FPaintGeometry(
+		this->AbsolutePosition + FVector2D(-SizeBeingAligned.X, this->Size.Y/2 - SizeBeingAligned.Y/2) * InScale,
+		SizeBeingAligned * CombinedScale,
+		CombinedScale
+	);
+}
+
+FPaintGeometry FGeometry::CenteredPaintGeometryOnRight( const FVector2D& SizeBeingAligned, float InScale ) const
+{
+	const float CombinedScale = this->Scale*InScale;
+	
+	return FPaintGeometry(
+		this->AbsolutePosition + FVector2D(this->Size.X, this->Size.Y/2 - SizeBeingAligned.Y/2) * InScale,
+		SizeBeingAligned * CombinedScale,
+		CombinedScale
+	);
+}
+
+FPaintGeometry FGeometry::CenteredPaintGeometryBelow( const FVector2D& SizeBeingAligned, float InScale ) const
+{
+	const float CombinedScale = this->Scale*InScale;
+	
+	return FPaintGeometry(
+		this->AbsolutePosition + FVector2D(this->Size.X/2 - SizeBeingAligned.X/2, this->Size.Y) * InScale,
+		SizeBeingAligned * CombinedScale,
+		CombinedScale
+	);
+}
diff --git a/Engine/Source/Runtime/SlateCore/Private/Layout/SlateRect.cpp b/Engine/Source/Runtime/SlateCore/Private/Layout/SlateRect.cpp
new file mode 100644
--- /dev/null
+++ b/Engine/Source/Runtime/SlateCore/Private/Layout/SlateRect.cpp
@@ -0,0 +1,16 @@
+// Copyright 1998-2014 Epic Games, Inc. All Rights Reserved.
+
+/*=============================================================================
+	SlateRect.cpp: Implements the FSlateRect structure.
+=============================================================================*/
+
+#include "SlateCorePrivatePCH.h"
+
+
+/* FSlateRect interface
+ *****************************************************************************/
+
+FSlateRect FSlateRect::InsetBy( const FMargin& InsetAmount ) const
+{
+	return FSlateRect(Left + InsetAmount.Left, Top + InsetAmount.Top, Right - InsetAmount.Right, Bottom - InsetAmount.Bottom);
+}
diff --git a/Engine/Source/Runtime/SlateCore/Private/Layout/Visibility.cpp b/Engine/Source/Runtime/SlateCore/Private/Layout/Visibility.cpp
new file mode 100644
--- /dev/null
+++ b/Engine/Source/Runtime/SlateCore/Private/Layout/Visibility.cpp
@@ -0,0 +1,58 @@
+// Copyright 1998-2014 Epic Games, Inc. All Rights Reserved.
+
+/*=============================================================================
+	Visibility.cpp: Implements the FVisibility class.
+=============================================================================*/
+
+#include "SlateCorePrivatePCH.h"
+
+
+/* Static initialization
+ *****************************************************************************/
+
+const EVisibility EVisibility::Visible(EVisibility::VIS_Visible);
+const EVisibility EVisibility::Collapsed(EVisibility::VIS_Collapsed);
+const EVisibility EVisibility::Hidden(EVisibility::VIS_Hidden);
+const EVisibility EVisibility::HitTestInvisible(EVisibility::VIS_HitTestInvisible);
+const EVisibility EVisibility::SelfHitTestInvisible(EVisibility::VIS_SelfHitTestInvisible);
+const EVisibility EVisibility::All(EVisibility::VIS_All);
+
+
+/* FVisibility interface
+ *****************************************************************************/
+
+FString EVisibility::ToString( ) const
+{
+	static const FString VisibleString("Visible");
+	static const FString HiddenString("Hidden");
+	static const FString CollapsedString("Collapsed");
+	static const FString HitTestInvisibleString("HitTestInvisible");
+	static const FString SelfHitTestInvisibleString("SelfHitTestInvisible");
+
+	if ((Value & VISPRIVATE_ChildrenHitTestVisible) == 0)
+	{
+		return HitTestInvisibleString;
+	}
+
+	if ((Value & VISPRIVATE_SelfHitTestVisible) == 0)
+	{
+		return SelfHitTestInvisibleString;
+	}
+
+	if ((Value & VISPRIVATE_Visible) != 0)
+	{
+		return VisibleString;
+	}
+
+	if ((Value & VISPRIVATE_Hidden) != 0)
+	{
+		return HiddenString;
+	}
+
+	if ((Value & VISPRIVATE_Collapsed) != 0)
+	{
+		return CollapsedString;
+	}
+
+	return FString();
+}
diff --git a/Engine/Source/Runtime/SlateCore/Private/Layout/WidgetPath.cpp b/Engine/Source/Runtime/SlateCore/Private/Layout/WidgetPath.cpp
new file mode 100644
--- /dev/null
+++ b/Engine/Source/Runtime/SlateCore/Private/Layout/WidgetPath.cpp
@@ -0,0 +1,218 @@
+// Copyright 1998-2014 Epic Games, Inc. All Rights Reserved.
+
+#include "SlateCorePrivatePCH.h"
+
+
+FArrangedWidget FWidgetPath::FindArrangedWidget( TSharedRef<const SWidget> WidgetToFind ) const
+{
+	for(int32 WidgetIndex = 0; WidgetIndex < Widgets.Num(); ++WidgetIndex)
+	{
+		if ( Widgets(WidgetIndex).Widget == WidgetToFind )
+		{
+			return Widgets(WidgetIndex);
+		}
+	}
+
+	return FArrangedWidget( SNullWidget::NullWidget, FGeometry() );
+}
+
+
+/** Matches any widget that is focusable */
+struct FFocusableWidgetMatcher
+{
+	bool IsMatch( const TSharedRef<const SWidget>& InWidget ) const
+	{
+		return InWidget->IsEnabled() && InWidget->SupportsKeyboardFocus();
+	}
+};
+
+/**
+ * Move focus either forward on backward in the path level specified by PathLevel.
+ * That is, this movement of focus will modify the subtree under Widgets(PathLevel).
+ *
+ * @param PathLevel       The level in this WidgetPath whose focus to move.
+ * @param MoveDirectin    Move focus forward or backward?
+ *
+ * @return true if the focus moved successfully, false if we were unable to move focus
+ */
+bool FWidgetPath::MoveFocus(int32 PathLevel, EFocusMoveDirection::Type MoveDirection)
+{
+	const int32 MoveDirectionAsInt = (MoveDirection == EFocusMoveDirection::Next)
+		? +1
+		: -1;
+
+
+	if ( PathLevel == Widgets.Num()-1 )
+	{
+		// We are the currently focused widget because we are at the very bottom of focus path.
+		if (MoveDirection == EFocusMoveDirection::Next)
+		{
+			// EFocusMoveDirection::Next implies descend, so try to find a focusable descendant.
+			return ExtendPathTo( FFocusableWidgetMatcher() );
+		}
+		else
+		{
+			// EFocusMoveDirection::Previous implies move focus up a level.
+			return false;
+		}
+		
+	}
+	else if ( Widgets.Num() > 1 )
+	{
+		// We are not the last widget in the path.
+		// GOAL: Look for a focusable descendant to the left or right of the currently focused path.
+	
+		// Arrange the children so we can iterate through them regardless of widget type.
+		FArrangedChildren ArrangedChildren(EVisibility::Visible);
+		Widgets(PathLevel).Widget->ArrangeChildren( Widgets(PathLevel).Geometry, ArrangedChildren );
+
+		// Find the currently focused child among the children.
+		int32 FocusedChildIndex = ArrangedChildren.FindItemIndex( Widgets(PathLevel+1) );
+		FocusedChildIndex = (FocusedChildIndex) % ArrangedChildren.Num() + MoveDirectionAsInt;
+
+		// Now actually search for the widget.
+		for( ; FocusedChildIndex < ArrangedChildren.Num() && FocusedChildIndex >= 0; FocusedChildIndex += MoveDirectionAsInt )
+		{
+			// Neither disabled widgets nor their children can be focused.
+			if ( ArrangedChildren(FocusedChildIndex).Widget->IsEnabled() )
+			{
+				// Look for a focusable descendant.
+				FArrangedChildren PathToFocusableChild = GeneratePathToWidget( FFocusableWidgetMatcher(), ArrangedChildren(FocusedChildIndex), MoveDirection );
+				// Either we found a focusable descendant, or an immediate child that is focusable.
+				const bool bFoundNextFocusable = ( PathToFocusableChild.Num() > 0 ) || ArrangedChildren(FocusedChildIndex).Widget->SupportsKeyboardFocus();
+				if ( bFoundNextFocusable )
+				{
+					// We found the next focusable widget, so make this path point at the new widget by:
+					// First, truncating the FocusPath up to the current level (i.e. PathLevel).
+					Widgets.Remove( PathLevel+1, Widgets.Num()-PathLevel-1 );
+					// Second, add the immediate child that is focus or whose descendant is focused.
+					Widgets.AddWidget( ArrangedChildren(FocusedChildIndex) );
+					// Add path to focused descendants if any.
+					Widgets.Append( PathToFocusableChild );
+					// We successfully moved focus!
+					return true;
+				}
+			}
+		}		
+		
+	}
+
+	return false;
+}
+
+/** Construct a weak widget path from a widget path. Defaults to an invalid path. */
+FWeakWidgetPath::FWeakWidgetPath( const FWidgetPath& InWidgetPath )
+: Window( InWidgetPath.TopLevelWindow )
+{
+	for( int32 WidgetIndex = 0; WidgetIndex < InWidgetPath.Widgets.Num(); ++WidgetIndex )
+	{
+		Widgets.Add( TWeakPtr<SWidget>( InWidgetPath.Widgets(WidgetIndex).Widget ) );
+	}
+}
+
+/** Make a non-weak WidgetPath out of this WeakWidgetPath. Do this by computing all the relevant geometries and converting the weak pointers to TSharedPtr. */	
+FWidgetPath FWeakWidgetPath::ToWidgetPath(EInterruptedPathHandling::Type InterruptedPathHandling) const
+{
+	FWidgetPath WidgetPath;
+	ToWidgetPath( WidgetPath, InterruptedPathHandling );
+	return WidgetPath;
+}
+
+FWeakWidgetPath::EPathResolutionResult::Result FWeakWidgetPath::ToWidgetPath( FWidgetPath& WidgetPath, EInterruptedPathHandling::Type InterruptedPathHandling ) const
+{
+	FArrangedChildren PathWithGeometries(EVisibility::Visible);
+	TArray< TSharedPtr<SWidget> > WidgetPtrs;
+	
+	// Convert the weak pointers into shared pointers because we are about to do something with this path instead of just observe it.
+	TSharedPtr<SWindow> TopLevelWindowPtr = Window.Pin();
+	for( TArray< TWeakPtr<SWidget> >::TConstIterator SomeWeakWidgetPtr( Widgets ); SomeWeakWidgetPtr; ++SomeWeakWidgetPtr )
+	{
+		WidgetPtrs.Add( SomeWeakWidgetPtr->Pin() );
+	}	
+	
+	// The path can get interrupted if some subtree of widgets disappeared, but we still maintain weak references to it.
+	bool bPathUninterrupted = false;
+
+	// For each widget in the path compute the geometry. We are able to do this starting with the top-level window because it knows its own geometry.
+	if ( TopLevelWindowPtr.IsValid() )
+	{
+		bPathUninterrupted = true;
+
+		FGeometry ParentGeometry = TopLevelWindowPtr->GetWindowGeometryInScreen();
+		PathWithGeometries.AddWidget( FArrangedWidget( TopLevelWindowPtr.ToSharedRef(), ParentGeometry ) );
+		
+		FArrangedChildren ArrangedChildren(EVisibility::Visible);
+		
+		// For every widget in the vertical slice...
+		for( int32 WidgetIndex = 0; bPathUninterrupted && WidgetIndex < WidgetPtrs.Num()-1; ++WidgetIndex )
+		{
+			TSharedPtr<SWidget> CurWidget = WidgetPtrs[WidgetIndex];
+
+			bool bFoundChild = false;
+			if ( CurWidget.IsValid() )
+			{
+				// Arrange the widget's children to find their geometries.
+				ArrangedChildren.Empty();
+				CurWidget->ArrangeChildren(ParentGeometry, ArrangedChildren);
+				
+				// Find the next widget in the path among the arranged children.
+				for( int32 SearchIndex = 0; !bFoundChild && SearchIndex < ArrangedChildren.Num(); ++SearchIndex )
+				{						
+					if ( ArrangedChildren(SearchIndex).Widget == WidgetPtrs[WidgetIndex+1] )
+					{
+						bFoundChild = true;
+						// Remember the widget and the associated geometry.
+						PathWithGeometries.AddWidget( ArrangedChildren(SearchIndex) );
+						// The next child in the vertical slice will be arranged with respect to its parent's geometry.
+						ParentGeometry = ArrangedChildren(SearchIndex).Geometry;
+					}
+				}
+			}
+
+			bPathUninterrupted = bFoundChild;
+			if (!bFoundChild && InterruptedPathHandling == EInterruptedPathHandling::ReturnInvalid )
+			{
+				return EPathResolutionResult::Truncated;
+			}
+		}			
+	}
+	
+	WidgetPath = FWidgetPath( TopLevelWindowPtr, PathWithGeometries );
+	return bPathUninterrupted ? EPathResolutionResult::Live : EPathResolutionResult::Truncated;
+}
+
+bool FWeakWidgetPath::ContainsWidget( const TSharedRef< const SWidget >& SomeWidget ) const
+{
+	for ( int32 WidgetIndex=0; WidgetIndex<Widgets.Num(); ++WidgetIndex )
+	{
+		if (Widgets[WidgetIndex].Pin() == SomeWidget)
+		{
+			return true;
+		}
+	}
+
+	return false;
+}
+
+/**
+ * @param MoveDirection      Direction in which to move the focus.
+ * 
+ * @return The new focus path.
+ */
+FWidgetPath FWeakWidgetPath::ToNextFocusedPath(EFocusMoveDirection::Type MoveDirection)
+{
+	// Make a copy of the focus path. We will mutate it until it meets the necessary requirements.
+	FWidgetPath NewFocusPath = this->ToWidgetPath();
+	TSharedPtr<SWidget> CurrentlyFocusedWidget = this->Widgets.Last().Pin();
+
+	bool bMovedFocus = false;
+	// Attempt to move the focus starting at the leafmost widget and bubbling up to the root (i.e. the window)
+	for (int32 FocusNodeIndex=NewFocusPath.Widgets.Num()-1; !bMovedFocus && FocusNodeIndex >= 0; --FocusNodeIndex)
+	{
+		bMovedFocus = NewFocusPath.MoveFocus( FocusNodeIndex, MoveDirection );
+	}
+
+	return NewFocusPath;
+}
+
+
diff --git a/Engine/Source/Runtime/SlateCore/Private/Logging/EventLogger.cpp b/Engine/Source/Runtime/SlateCore/Private/Logging/EventLogger.cpp
new file mode 100644
--- /dev/null
+++ b/Engine/Source/Runtime/SlateCore/Private/Logging/EventLogger.cpp
@@ -0,0 +1,135 @@
+// Copyright 1998-2014 Epic Games, Inc. All Rights Reserved.
+
+/*=============================================================================
+	EventLogger.cpp: Implements various Slate event logger implementations.
+=============================================================================*/
+
+#include "SlateCorePrivatePCH.h"
+
+
+/* Local helper functions
+ *****************************************************************************/
+
+static FString PrettyPrint(EEventLog::Type Event, const FString& AdditionalContent = FString(), TSharedPtr<SWidget> Widget = TSharedPtr<SWidget>())
+{
+	FString WidgetInfo;
+	FString Content;
+
+	if (Widget.IsValid())
+	{
+		WidgetInfo = FString::Printf(TEXT(", %s %s"), *Widget->GetTypeAsString(), *Widget->GetReadableLocation());
+	}
+
+	if (AdditionalContent.Len())
+	{
+		Content = FString::Printf(TEXT(": %s"), *AdditionalContent);
+	}
+
+	return FString::Printf(TEXT("[%07.2f] %s%s%s"), FPlatformTime::Seconds() - GStartTime, *EEventLog::ToString(Event), *Content, *WidgetInfo);
+}
+
+
+/* FFileEventLogger interface
+ *****************************************************************************/
+
+FString FFileEventLogger::GetLog( ) const
+{
+	FString OutputLog = LINE_TERMINATOR LINE_TERMINATOR;
+	
+	for (int32 LoggedEventIndex = 0; LoggedEventIndex < LoggedEvents.Num(); ++LoggedEventIndex)
+	{
+		OutputLog += LoggedEvents[LoggedEventIndex] + LINE_TERMINATOR;
+	}
+
+	return OutputLog;
+}
+
+
+void FFileEventLogger::Log(EEventLog::Type Event, const FString& AdditionalContent, TSharedPtr<SWidget> Widget)
+{
+	if (Widget.IsValid())
+	{
+		LoggedEvents.Add( PrettyPrint(Event, AdditionalContent, Widget) );
+	}
+}
+
+void FFileEventLogger::SaveToFile()
+{
+	FString LogFilePath;
+	{
+		const TCHAR LogFileName[] = TEXT("EventLog");
+		LogFilePath = FPaths::CreateTempFilename(*FPaths::GameLogDir(), LogFileName, TEXT(".log"));
+	}
+
+	FOutputDeviceFile EventLogFile(*LogFilePath);
+
+	for (int32 LoggedEventIndex = 0; LoggedEventIndex < LoggedEvents.Num(); ++LoggedEventIndex)
+	{
+		EventLogFile.Serialize(*LoggedEvents[LoggedEventIndex], ELogVerbosity::Log, NAME_None);
+	}
+
+	EventLogFile.Flush();
+	EventLogFile.TearDown();
+}
+
+
+/* FConsoleEventLogger interface
+ *****************************************************************************/
+
+FString FConsoleEventLogger::GetLog( ) const
+{
+	return FString();
+}
+
+void FConsoleEventLogger::Log( EEventLog::Type Event, const FString& AdditionalContent, TSharedPtr<SWidget> Widget )
+{
+	UE_LOG(LogSlate, Log, TEXT("%s"), *PrettyPrint(Event, AdditionalContent, Widget));
+}
+
+
+/* FStabilityEventLogger interface
+ *****************************************************************************/
+
+/** Limit of how many items we should have in our stability log */
+static const int32 STABILITY_LOG_MAX_SIZE = 100;
+
+FString FStabilityEventLogger::GetLog() const
+{
+	FString OutputLog = LINE_TERMINATOR LINE_TERMINATOR;
+
+	for (int32 LoggedEventIndex = 0; LoggedEventIndex < LoggedEvents.Num(); ++LoggedEventIndex)
+	{
+		OutputLog += LoggedEvents[LoggedEventIndex] + LINE_TERMINATOR;
+	}
+
+	return OutputLog;
+}
+
+
+void FStabilityEventLogger::Log(EEventLog::Type Event, const FString& AdditionalContent, TSharedPtr<SWidget> Widget)
+{
+	// filter out events that happen a lot
+	if (Event == EEventLog::MouseMove ||
+		Event == EEventLog::MouseEnter ||
+		Event == EEventLog::MouseLeave ||
+		Event == EEventLog::DragEnter ||
+		Event == EEventLog::DragLeave ||
+		Event == EEventLog::DragOver)
+	{
+		return;
+	}
+
+	if (Widget.IsValid())
+	{
+		LoggedEvents.Add(PrettyPrint(Event, AdditionalContent, Widget));
+	}
+	else
+	{
+		LoggedEvents.Add(PrettyPrint(Event, AdditionalContent));
+	}
+
+	if (LoggedEvents.Num() > STABILITY_LOG_MAX_SIZE)
+	{
+		LoggedEvents.RemoveAt(0, LoggedEvents.Num() - STABILITY_LOG_MAX_SIZE);
+	}
+}
