Slate: Refactored core Slate implementation into SlateCore module in preparation for UMG.

diff --git a/Engine/Source/Runtime/SlateCore/Private/Rendering/DrawElements.cpp b/Engine/Source/Runtime/SlateCore/Private/Rendering/DrawElements.cpp
new file mode 100644
--- /dev/null
+++ b/Engine/Source/Runtime/SlateCore/Private/Rendering/DrawElements.cpp
@@ -0,0 +1,163 @@
+// Copyright 1998-2014 Epic Games, Inc. All Rights Reserved.
+
+#include "SlateCorePrivatePCH.h"
+
+
+void FSlateDrawElement::MakeDebugQuad( FSlateWindowElementList& ElementList, uint32 InLayer, const FPaintGeometry& PaintGeometry, const FSlateRect& InClippingRect)
+{
+	FSlateDrawElement& DrawElt = ElementList.AddUninitialized();
+	DrawElt.ElementType = ET_DebugQuad;
+	DrawElt.Position = PaintGeometry.DrawPosition;
+	DrawElt.Size = PaintGeometry.DrawSize;
+	DrawElt.ClippingRect = InClippingRect;
+	DrawElt.Layer = InLayer;
+	DrawElt.DrawEffects = ESlateDrawEffect::None;
+	DrawElt.Scale = PaintGeometry.DrawScale;
+}
+
+void FSlateDrawElement::MakeBox( 
+	FSlateWindowElementList& ElementList,
+	uint32 InLayer, 
+	const FPaintGeometry& PaintGeometry, 
+	const FSlateBrush* InBrush, 
+	const FSlateRect& InClippingRect, 
+	ESlateDrawEffect::Type InDrawEffects, 
+	const FLinearColor& InTint )
+{
+	FSlateDrawElement& DrawElt = ElementList.AddUninitialized();
+	DrawElt.ElementType = (InBrush->DrawAs == ESlateBrushDrawType::Border) ? ET_Border : ET_Box;
+	DrawElt.Position = PaintGeometry.DrawPosition;
+	DrawElt.Size = PaintGeometry.DrawSize;
+	DrawElt.ClippingRect = InClippingRect;
+	DrawElt.DataPayload.SetBoxPayloadProperties( InBrush, InTint );
+	DrawElt.Layer = InLayer;
+	DrawElt.DrawEffects = InDrawEffects;
+	DrawElt.Scale = PaintGeometry.DrawScale;
+}
+
+void FSlateDrawElement::MakeRotatedBox( 
+	FSlateWindowElementList& ElementList,
+	uint32 InLayer, 
+	const FPaintGeometry& PaintGeometry, 
+	const FSlateBrush* InBrush, 
+	const FSlateRect& InClippingRect, 
+	ESlateDrawEffect::Type InDrawEffects, 
+	float Angle,
+	TOptional<FVector2D> InRotationPoint,
+	ERotationSpace RotationSpace,
+	const FLinearColor& InTint )
+{
+	FSlateDrawElement& DrawElt = ElementList.AddUninitialized();
+	DrawElt.ElementType = (InBrush->DrawAs == ESlateBrushDrawType::Border) ? ET_Border : ET_Box;
+	DrawElt.Position = PaintGeometry.DrawPosition;
+	DrawElt.Size = PaintGeometry.DrawSize;
+	DrawElt.ClippingRect = InClippingRect;
+	DrawElt.Layer = InLayer;
+	DrawElt.DrawEffects = InDrawEffects;
+	DrawElt.Scale = PaintGeometry.DrawScale;
+
+	FVector2D RotationPoint = GetRotationPoint( PaintGeometry, InRotationPoint, RotationSpace );
+	DrawElt.DataPayload.SetRotatedBoxPayloadProperties( InBrush, Angle, RotationPoint, InTint );
+}
+
+void FSlateDrawElement::MakeText( FSlateWindowElementList& ElementList, uint32 InLayer, const FPaintGeometry& PaintGeometry, const FString& InText, const int32 StartIndex, const int32 EndIndex, const FSlateFontInfo& InFontInfo, const FSlateRect& InClippingRect,ESlateDrawEffect::Type InDrawEffects, const FLinearColor& InTint )
+{
+	FSlateDrawElement& DrawElt = ElementList.AddUninitialized();
+	DrawElt.ElementType = ET_Text;
+	DrawElt.Position = PaintGeometry.DrawPosition;
+	DrawElt.Size = PaintGeometry.DrawSize;
+	DrawElt.ClippingRect = InClippingRect;
+	DrawElt.DataPayload.SetTextPayloadProperties( FString( EndIndex - StartIndex, *InText + StartIndex ), InFontInfo, InTint );
+	DrawElt.Layer = InLayer;
+	DrawElt.DrawEffects = InDrawEffects;
+	DrawElt.Scale = PaintGeometry.DrawScale;
+}
+
+void FSlateDrawElement::MakeText( FSlateWindowElementList& ElementList, uint32 InLayer, const FPaintGeometry& PaintGeometry, const FString& InText, const FSlateFontInfo& InFontInfo, const FSlateRect& InClippingRect,ESlateDrawEffect::Type InDrawEffects, const FLinearColor& InTint )
+{
+	FSlateDrawElement& DrawElt = ElementList.AddUninitialized();
+	DrawElt.ElementType = ET_Text;
+	DrawElt.Position = PaintGeometry.DrawPosition;
+	DrawElt.Size = PaintGeometry.DrawSize;
+	DrawElt.ClippingRect = InClippingRect;
+	DrawElt.DataPayload.SetTextPayloadProperties( InText, InFontInfo, InTint );
+	DrawElt.Layer = InLayer;
+	DrawElt.DrawEffects = InDrawEffects;
+	DrawElt.Scale = PaintGeometry.DrawScale;
+}
+
+void FSlateDrawElement::MakeText( FSlateWindowElementList& ElementList, uint32 InLayer, const FPaintGeometry& PaintGeometry, const FText& InText, const FSlateFontInfo& InFontInfo, const FSlateRect& InClippingRect,ESlateDrawEffect::Type InDrawEffects, const FLinearColor& InTint )
+{
+	FSlateDrawElement& DrawElt = ElementList.AddUninitialized();
+	DrawElt.ElementType = ET_Text;
+	DrawElt.Position = PaintGeometry.DrawPosition;
+	DrawElt.Size = PaintGeometry.DrawSize;
+	DrawElt.ClippingRect = InClippingRect;
+	DrawElt.DataPayload.SetTextPayloadProperties( InText.ToString(), InFontInfo, InTint );
+	DrawElt.Layer = InLayer;
+	DrawElt.DrawEffects = InDrawEffects;
+	DrawElt.Scale = PaintGeometry.DrawScale;
+}
+
+void FSlateDrawElement::MakeGradient( FSlateWindowElementList& ElementList, uint32 InLayer, const FPaintGeometry& InPaintGeometry, TArray<FSlateGradientStop> InGradientStops, EOrientation InGradientType, const FSlateRect& InClippingRect, ESlateDrawEffect::Type InDrawEffects, bool bGammaCorrect )
+{
+	FSlateDrawElement& DrawElt = ElementList.AddUninitialized();
+	DrawElt.ElementType = ET_Gradient;
+	DrawElt.Position = InPaintGeometry.DrawPosition;
+	DrawElt.Size = InPaintGeometry.DrawSize;
+	DrawElt.ClippingRect = InClippingRect;
+	DrawElt.DataPayload.SetGradientPayloadProperties( InGradientStops, InGradientType, bGammaCorrect );
+	DrawElt.Layer = InLayer;
+	DrawElt.DrawEffects = InDrawEffects;
+	DrawElt.Scale = InPaintGeometry.DrawScale;
+}
+
+void FSlateDrawElement::MakeSpline( FSlateWindowElementList& ElementList, uint32 InLayer, const FPaintGeometry& PaintGeometry, const FVector2D& InStart, const FVector2D& InStartDir, const FVector2D& InEnd, const FVector2D& InEndDir, const FSlateRect InClippingRect, float InThickness, ESlateDrawEffect::Type InDrawEffects, const FLinearColor& InTint )
+{
+	FSlateDrawElement& DrawElt = ElementList.AddUninitialized();
+	DrawElt.ElementType = ET_Spline;
+	DrawElt.Position = PaintGeometry.DrawPosition;
+	DrawElt.ClippingRect = InClippingRect;
+	DrawElt.DataPayload.SetSplinePayloadProperties( InStart, InStartDir, InEnd, InEndDir, InThickness, InTint );
+	DrawElt.Layer = InLayer;
+	DrawElt.DrawEffects = InDrawEffects;
+	DrawElt.Scale = PaintGeometry.DrawScale;
+}
+
+void FSlateDrawElement::MakeDrawSpaceSpline( FSlateWindowElementList& ElementList, uint32 InLayer, const FVector2D& InStart, const FVector2D& InStartDir, const FVector2D& InEnd, const FVector2D& InEndDir, const FSlateRect InClippingRect, float InThickness, ESlateDrawEffect::Type InDrawEffects, const FColor& InTint )
+{
+	MakeSpline( ElementList, InLayer, FPaintGeometry::Identity(), InStart, InStartDir, InEnd, InEndDir, InClippingRect, InThickness, InDrawEffects, InTint );
+}
+
+void FSlateDrawElement::MakeLines( FSlateWindowElementList& ElementList, uint32 InLayer, const FPaintGeometry& PaintGeometry, const TArray<FVector2D>& Points, const FSlateRect InClippingRect, ESlateDrawEffect::Type InDrawEffects, const FLinearColor& InTint, bool bAntialias )
+{
+	FSlateDrawElement& DrawElt = ElementList.AddUninitialized();
+	DrawElt.ElementType = ET_Line;
+	DrawElt.Position = PaintGeometry.DrawPosition;
+	DrawElt.ClippingRect = InClippingRect;
+	DrawElt.DataPayload.SetLinesPayloadProperties( Points, InTint, bAntialias, ESlateLineJoinType::Sharp );
+	DrawElt.Layer = InLayer;
+	DrawElt.DrawEffects = InDrawEffects;
+	DrawElt.Scale = PaintGeometry.DrawScale;
+}
+
+void FSlateDrawElement::MakeViewport( FSlateWindowElementList& ElementList, uint32 InLayer, const FPaintGeometry& PaintGeometry, TSharedPtr<const ISlateViewport> Viewport, const FSlateRect& InClippingRect, bool bGammaCorrect, bool bAllowBlending, ESlateDrawEffect::Type InDrawEffects, const FLinearColor& InTint )
+{
+	FSlateDrawElement& DrawElt = ElementList.AddUninitialized();
+	DrawElt.ElementType = ET_Viewport;
+	DrawElt.Position = PaintGeometry.DrawPosition;
+	DrawElt.Size = PaintGeometry.DrawSize;
+	DrawElt.ClippingRect = InClippingRect;
+	DrawElt.DataPayload.SetViewportPayloadProperties( Viewport, InTint, bGammaCorrect, bAllowBlending );
+	DrawElt.Layer = InLayer;
+	DrawElt.DrawEffects = InDrawEffects;
+	DrawElt.Scale = PaintGeometry.DrawScale;
+}
+void FSlateDrawElement::MakeCustom( FSlateWindowElementList& ElementList, uint32 InLayer, TSharedPtr<ICustomSlateElement, ESPMode::ThreadSafe> CustomDrawer )
+{
+	FSlateDrawElement& DrawElt = ElementList.AddUninitialized();
+	DrawElt.ElementType = ET_Custom;
+	DrawElt.Layer = InLayer;
+	DrawElt.DataPayload.SetCustomDrawerPayloadProperties( CustomDrawer );
+	DrawElt.ClippingRect = FSlateRect(1,1,1,1);
+}
\ No newline at end of file
diff --git a/Engine/Source/Runtime/SlateCore/Private/Rendering/ElementBatcher.cpp b/Engine/Source/Runtime/SlateCore/Private/Rendering/ElementBatcher.cpp
new file mode 100644
--- /dev/null
+++ b/Engine/Source/Runtime/SlateCore/Private/Rendering/ElementBatcher.cpp
@@ -0,0 +1,1476 @@
+// Copyright 1998-2014 Epic Games, Inc. All Rights Reserved.
+
+/*=============================================================================
+	ElementBatcher.cpp: Implements the FSlateElementBatcher class.
+=============================================================================*/
+
+#include "SlateCorePrivatePCH.h"
+
+
+static const FSlateImageBrush SplineFilterTable( FPaths::EngineContentDir() / TEXT("Slate/SplineFilterTable.png"), FVector2D(32,1) );
+
+
+/**
+ * Computes the element tint color based in the user specified color and brush being used.
+ * Note: The color could be in RGB or HSV
+ * @return The final color to be passed per vertex
+ */
+static FColor GetElementColor( const FLinearColor& InColor, const FSlateBrush* InBrush )
+{
+	// Pass the color through
+	return InColor.ToFColor(false);
+}
+
+
+/**
+ * Rotates a point around another point 
+ *
+ * @param InPoint	The point to rotate
+ * @param AboutPoint	The point to rotate about
+ * @param Radians	The angle of rotation in radians
+ * @return The rotated point
+ */
+static FVector2D RotatePoint( FVector2D InPoint, const FVector2D& AboutPoint, float Radians )
+{
+	if( Radians != 0.0f )
+	{
+		float CosAngle = FMath::Cos(Radians);
+		float SinAngle = FMath::Sin(Radians);
+
+		InPoint.X -= AboutPoint.X;
+		InPoint.Y -= AboutPoint.Y;
+
+		float X = (InPoint.X * CosAngle) - (InPoint.Y * SinAngle);
+		float Y = (InPoint.X * SinAngle) + (InPoint.Y * CosAngle);
+
+		return FVector2D( X + AboutPoint.X, Y + AboutPoint.Y );
+	}
+
+	return InPoint;
+}
+
+
+FSlateElementBatcher::FSlateElementBatcher( TSharedRef<FSlateRenderingPolicy> InRenderingPolicy )
+	: RenderingPolicy( InRenderingPolicy )
+{
+
+}
+
+
+FSlateElementBatcher::~FSlateElementBatcher()
+{
+}
+
+
+/** 
+ * Adds a new element to be rendered 
+ * 
+ * @param InElement	The element to render 
+ */
+void FSlateElementBatcher::AddElements( const TArray<FSlateDrawElement>& DrawElements )
+{
+	SCOPE_CYCLE_COUNTER(STAT_SlateBuildElements);
+
+	for( int32 DrawElementIndex = 0; DrawElementIndex < DrawElements.Num(); ++DrawElementIndex )
+	{
+		const FSlateDrawElement& DrawElement = DrawElements[DrawElementIndex];
+
+		const FSlateRect& InClippingRect = DrawElement.GetClippingRect();
+	
+		// A zero or negatively sized clipping rect means the geometry will not be displayed
+		const bool bIsFullyClipped = (!InClippingRect.IsValid() || (InClippingRect.Top == 0.0f && InClippingRect.Left == 0.0f && InClippingRect.Bottom == 0.0f && InClippingRect.Right == 0.0f));
+
+		if ( !bIsFullyClipped )
+		{
+			// Determine what type of element to add
+			switch( DrawElement.GetElementType() )
+			{
+			case FSlateDrawElement::ET_Box:
+				AddBoxElement( DrawElement.GetPosition(), DrawElement.GetSize(), DrawElement.GetScale(), DrawElement.GetDataPayload(), DrawElement.GetClippingRect(), DrawElement.GetDrawEffects(), DrawElement.GetLayer() );
+				break;
+			case FSlateDrawElement::ET_DebugQuad:
+				AddQuadElement( DrawElement.GetPosition(), DrawElement.GetSize(), DrawElement.GetScale(), DrawElement.GetClippingRect(), DrawElement.GetLayer() );
+				break;
+			case FSlateDrawElement::ET_Text:
+				AddTextElement( DrawElement.GetPosition(), DrawElement.GetSize(), DrawElement.GetScale(), DrawElement.GetDataPayload(), DrawElement.GetClippingRect(), DrawElement.GetDrawEffects(), DrawElement.GetLayer() );
+				break;
+			case FSlateDrawElement::ET_Spline:
+				AddSplineElement( DrawElement.GetPosition(), DrawElement.GetScale(), DrawElement.GetDataPayload(), DrawElement.GetClippingRect(), DrawElement.GetDrawEffects(), DrawElement.GetLayer() );
+				break;
+			case FSlateDrawElement::ET_Line:
+				AddLineElement( DrawElement.GetPosition(), DrawElement.GetSize(), DrawElement.GetScale(), DrawElement.GetDataPayload(), DrawElement.GetClippingRect(), DrawElement.GetDrawEffects(), DrawElement.GetLayer() );
+				break;
+			case FSlateDrawElement::ET_Gradient:
+				AddGradientElement( DrawElement.GetPosition(), DrawElement.GetSize(), DrawElement.GetScale(), DrawElement.GetDataPayload(), DrawElement.GetClippingRect(), DrawElement.GetDrawEffects(), DrawElement.GetLayer() );
+				break;
+			case FSlateDrawElement::ET_Viewport:
+				AddViewportElement( DrawElement.GetPosition(), DrawElement.GetSize(), DrawElement.GetScale(), DrawElement.GetDataPayload(), DrawElement.GetClippingRect(), DrawElement.GetDrawEffects(), DrawElement.GetLayer() );
+				break;
+			case FSlateDrawElement::ET_Border:
+				AddBorderElement( DrawElement.GetPosition(), DrawElement.GetSize(), DrawElement.GetScale(), DrawElement.GetDataPayload(), DrawElement.GetClippingRect(), DrawElement.GetDrawEffects(), DrawElement.GetLayer() ); 
+				break;
+			case FSlateDrawElement::ET_Custom:
+				AddCustomElement( DrawElement.GetPosition(), DrawElement.GetSize(), DrawElement.GetScale(), DrawElement.GetDataPayload(), DrawElement.GetClippingRect(), DrawElement.GetDrawEffects(), DrawElement.GetLayer() ); 
+				break;
+			default:
+				checkf(0, TEXT("Invalid element type"));
+				break;
+			}
+		}
+	}
+}
+
+void FSlateElementBatcher::AddQuadElement( const FVector2D& Position, const FVector2D& Size, float Scale, const FSlateRect& InClippingRect, uint32 Layer, FColor Color )
+{
+	FSlateElementBatch& ElementBatch = FindBatchForElement( Layer, FShaderParams(), nullptr, ESlateDrawPrimitive::TriangleList, ESlateShader::Default, ESlateDrawEffect::None, ESlateBatchDrawFlag::Wireframe|ESlateBatchDrawFlag::NoBlending );
+	TArray<FSlateVertex>& BatchVertices = BatchVertexArrays[ElementBatch.VertexArrayIndex];
+	TArray<SlateIndex>& BatchIndices = BatchIndexArrays[ElementBatch.IndexArrayIndex];
+
+	// Determine the four corners of the quad
+	FVector2D TopLeft = Position;
+	FVector2D TopRight = FVector2D( Position.X+Size.X, Position.Y);
+	FVector2D BotLeft = FVector2D( Position.X, Position.Y+Size.Y);
+	FVector2D BotRight = FVector2D( Position.X + Size.X, Position.Y+Size.Y);
+
+	// The start index of these vertices in the index buffer
+	uint32 IndexStart = BatchVertices.Num();
+
+	// Add four vertices to the list of verts to be added to the vertex buffer
+	BatchVertices.Add( FSlateVertex(TopLeft, FVector2D(0.0f,0.0f), Color, InClippingRect ) );
+	BatchVertices.Add( FSlateVertex(TopRight, FVector2D(1.0f,0.0f), Color, InClippingRect ) );
+	BatchVertices.Add( FSlateVertex(BotLeft, FVector2D(0.0f,1.0f), Color, InClippingRect ) );
+	BatchVertices.Add( FSlateVertex(BotRight, FVector2D(1.0f,1.0f),Color, InClippingRect ) );
+
+	// The offset into the index buffer where this quads indices start
+	uint32 IndexOffsetStart = BatchIndices.Num();
+	// Add 6 indices to the vertex buffer.  (2 tri's per quad, 3 indices per tri)
+	BatchIndices.Add( IndexStart + 0 );
+	BatchIndices.Add( IndexStart + 1 );
+	BatchIndices.Add( IndexStart + 2 );
+
+	BatchIndices.Add( IndexStart + 2 );
+	BatchIndices.Add( IndexStart + 1 );
+	BatchIndices.Add( IndexStart + 3 );
+}
+
+/** 
+ * Creates vertices necessary to draw a box element.  
+ *
+ * @param Position	The top left screen space position of the element
+ * @param Size		The size of the element
+ * @param InPayload	The margin and texture data necessary to draw the quad
+ */
+void FSlateElementBatcher::AddBoxElement( const FVector2D& InPosition, const FVector2D& Size, float Scale, const FSlateDataPayload& InPayload, const FSlateRect& InClippingRect, ESlateDrawEffect::Type InDrawEffects, uint32 Layer )
+{
+	check( InPayload.BrushResource );
+
+	SCOPE_CYCLE_COUNTER(STAT_SlateBuildBoxes);
+
+	if(InPayload.BrushResource->DrawAs != ESlateBrushDrawType::NoDrawType)
+	{
+		FVector2D Position( FMath::Trunc(InPosition.X), FMath::Trunc(InPosition.Y) );
+
+		const float Angle = InPayload.Angle;
+		const FVector2D RotationPoint = InPayload.RotationPoint;
+
+		uint32 TextureWidth = 1;
+		uint32 TextureHeight = 1;
+
+		// Get the default start and end UV.  If the texture is atlased this value will be a subset of this
+		FVector2D StartUV = FVector2D(0.0f,0.0f);
+		FVector2D EndUV = FVector2D(1.0f,1.0f);
+		FVector2D SizeUV;
+		const float PixelCenterOffset = RenderingPolicy->GetPixelCenterOffset();
+		FVector2D HalfTexel;
+
+		FSlateShaderResourceProxy* ResourceProxy = RenderingPolicy->GetTextureResource( *InPayload.BrushResource );
+		FSlateShaderResource* Resource = nullptr;
+		if( ResourceProxy )
+		{
+			// The actual texture for rendering.  If the texture is atlased this is the atlas
+			Resource = ResourceProxy->Resource;
+			// The width and height of the texture (non-atlased size)
+			TextureWidth = ResourceProxy->ActualSize.X != 0 ? ResourceProxy->ActualSize.X : 1;
+			TextureHeight = ResourceProxy->ActualSize.Y != 0 ? ResourceProxy->ActualSize.Y : 1;
+
+			// Texel offset
+			HalfTexel = FVector2D( PixelCenterOffset/TextureWidth, PixelCenterOffset/TextureHeight );
+
+			SizeUV = ResourceProxy->SizeUV;
+			StartUV = ResourceProxy->StartUV + HalfTexel;
+			EndUV = StartUV + ResourceProxy->SizeUV;
+		}
+		else
+		{
+			// nullptr texture
+			SizeUV = FVector2D(1.0f,1.0f);
+			HalfTexel = FVector2D( PixelCenterOffset, PixelCenterOffset );
+		}
+
+
+		FColor Tint = GetElementColor( InPayload.Tint, InPayload.BrushResource );
+
+		ESlateBrushTileType::Type TilingRule = InPayload.BrushResource->Tiling;
+		const bool bTileHorizontal = (TilingRule == ESlateBrushTileType::Both || TilingRule == ESlateBrushTileType::Horizontal);
+		const bool bTileVertical = (TilingRule == ESlateBrushTileType::Both || TilingRule == ESlateBrushTileType::Vertical);
+
+		// Pass the tiling information as a flag so we can pick the correct texture addressing mode
+		ESlateBatchDrawFlag::Type DrawFlags = ( ( bTileHorizontal ? ESlateBatchDrawFlag::TileU : 0 ) | ( bTileVertical ? ESlateBatchDrawFlag::TileV : 0 ) );
+
+		FShaderParams RotationParams = Angle != 0.0f ? FShaderParams::MakeVertexShaderParams( FVector4( Angle, RotationPoint.X, RotationPoint.Y, 0 ) ) : FShaderParams();
+
+		FSlateElementBatch& ElementBatch = FindBatchForElement( Layer, RotationParams, Resource, ESlateDrawPrimitive::TriangleList, ESlateShader::Default, InDrawEffects, DrawFlags );
+		TArray<FSlateVertex>& BatchVertices = BatchVertexArrays[ElementBatch.VertexArrayIndex];
+		TArray<SlateIndex>& BatchIndices = BatchIndexArrays[ElementBatch.IndexArrayIndex];
+
+		float HorizontalTiling = bTileHorizontal ? Size.X/Scale/TextureWidth : 1.0f;
+		float VerticalTiling = bTileVertical ? Size.Y/Scale/TextureHeight : 1.0f;
+
+		const FVector2D Tiling( HorizontalTiling,VerticalTiling );
+
+		// The start index of these vertices in the index buffer
+		uint32 IndexStart = BatchVertices.Num();
+		// The offset into the index buffer where this elements indices start
+		uint32 IndexOffsetStart = BatchIndices.Num();
+
+		if( InPayload.BrushResource->Margin.Left!= 0.0f || 
+			InPayload.BrushResource->Margin.Top != 0.0f ||
+			InPayload.BrushResource->Margin.Right != 0.0f ||
+			InPayload.BrushResource->Margin.Bottom != 0.0f )
+		{
+			// Create 9 quads for the box element based on the following diagram
+			//     ___LeftMargin    ___RightMargin
+			//    /                /
+			//  +--+-------------+--+
+			//  |  |c1           |c2| ___TopMargin
+			//  +--o-------------o--+
+			//  |  |             |  |
+			//  |  |c3           |c4|
+			//  +--o-------------o--+
+			//  |  |             |  | ___BottomMargin
+			//  +--+-------------+--+
+
+			const FMargin& Margin = InPayload.BrushResource->Margin;
+
+
+			// Determine the texture coordinates for each quad
+			// These are not scaled.
+			const float LeftMarginU = (Margin.Left > 0.0f)
+				? StartUV.X + Margin.Left * SizeUV.X + HalfTexel.X
+				: StartUV.X;
+			const float TopMarginV = (Margin.Top > 0.0f)
+				? StartUV.Y + Margin.Top * SizeUV.Y + HalfTexel.Y
+				: StartUV.Y;
+			const float RightMarginU = (Margin.Right > 0.0f)
+				? EndUV.X - Margin.Right * SizeUV.X + HalfTexel.X
+				: EndUV.X;
+			const float BottomMarginV = (Margin.Bottom > 0.0f)
+				? EndUV.Y - Margin.Bottom * SizeUV.Y + HalfTexel.Y
+				: EndUV.Y;
+
+			// Determine the margins for each quad
+			const FMargin ScaledMargin = Margin * Scale;
+
+			float LeftMarginX = TextureWidth * ScaledMargin.Left;
+			float TopMarginY = TextureHeight * ScaledMargin.Top;
+			float RightMarginX = Size.X - TextureWidth * ScaledMargin.Right;
+			float BottomMarginY = Size.Y - TextureHeight * ScaledMargin.Bottom;
+
+			// If the margins are overlapping the margins are too big or the button is too small
+			// so clamp margins to half of the box size
+			if( RightMarginX < LeftMarginX )
+			{
+				LeftMarginX = Size.X / 2;
+				RightMarginX = LeftMarginX;
+			}
+
+			if( BottomMarginY < TopMarginY )
+			{
+				TopMarginY = Size.Y / 2;
+				BottomMarginY = TopMarginY;
+			}
+
+			LeftMarginX += Position.X;
+			TopMarginY += Position.Y;
+			RightMarginX += Position.X;
+			BottomMarginY += Position.Y;
+
+			FVector2D EndPos = FVector2D(Position.X, Position.Y) + Size;
+
+			EndPos.X = FMath::Trunc( EndPos.X );
+			EndPos.Y = FMath::Trunc( EndPos.Y );
+
+			BatchVertices.Add( FSlateVertex( Position,									StartUV,									Tiling,	Tint, InClippingRect ) ); //0
+			BatchVertices.Add( FSlateVertex( FVector2D( Position.X, TopMarginY ),		FVector2D( StartUV.X, TopMarginV ),			Tiling,	Tint, InClippingRect ) ); //1
+			BatchVertices.Add( FSlateVertex( FVector2D( LeftMarginX, Position.Y ),		FVector2D( LeftMarginU, StartUV.Y ),		Tiling,	Tint, InClippingRect ) ); //2
+			BatchVertices.Add( FSlateVertex( FVector2D( LeftMarginX, TopMarginY ),		FVector2D( LeftMarginU, TopMarginV ),		Tiling,	Tint, InClippingRect ) ); //3
+			BatchVertices.Add( FSlateVertex( FVector2D( RightMarginX, Position.Y ),		FVector2D( RightMarginU, StartUV.Y ),		Tiling,	Tint, InClippingRect ) ); //4
+			BatchVertices.Add( FSlateVertex( FVector2D( RightMarginX, TopMarginY ),		FVector2D( RightMarginU,TopMarginV),		Tiling,	Tint, InClippingRect ) ); //5
+			BatchVertices.Add( FSlateVertex( FVector2D( EndPos.X, Position.Y ),			FVector2D( EndUV.X, StartUV.Y ),			Tiling,	Tint, InClippingRect ) ); //6
+			BatchVertices.Add( FSlateVertex( FVector2D( EndPos.X, TopMarginY ),			FVector2D( EndUV.X, TopMarginV),			Tiling,	Tint, InClippingRect ) ); //7
+
+			BatchVertices.Add( FSlateVertex( FVector2D( Position.X, BottomMarginY ),	FVector2D( StartUV.X, BottomMarginV ),		Tiling,	Tint, InClippingRect ) ); //8
+			BatchVertices.Add( FSlateVertex( FVector2D( LeftMarginX, BottomMarginY ),	FVector2D( LeftMarginU, BottomMarginV ),	Tiling,	Tint, InClippingRect ) ); //9
+			BatchVertices.Add( FSlateVertex( FVector2D( RightMarginX, BottomMarginY ),	FVector2D( RightMarginU, BottomMarginV ),	Tiling,	Tint, InClippingRect ) ); //10
+			BatchVertices.Add( FSlateVertex( FVector2D( EndPos.X, BottomMarginY ),		FVector2D( EndUV.X, BottomMarginV ),		Tiling,	Tint, InClippingRect ) ); //11
+			BatchVertices.Add( FSlateVertex( FVector2D( Position.X, EndPos.Y ),			FVector2D( StartUV.X, EndUV.Y ),			Tiling,	Tint, InClippingRect ) ); //12
+			BatchVertices.Add( FSlateVertex( FVector2D( LeftMarginX, EndPos.Y ),		FVector2D( LeftMarginU, EndUV.Y ),			Tiling,	Tint, InClippingRect ) ); //13
+			BatchVertices.Add( FSlateVertex( FVector2D( RightMarginX, EndPos.Y ),		FVector2D( RightMarginU, EndUV.Y ),			Tiling,	Tint, InClippingRect ) ); //14
+			BatchVertices.Add( FSlateVertex( FVector2D( EndPos.X, EndPos.Y ),			EndUV,										Tiling,	Tint, InClippingRect ) ); //15
+
+			// Top
+			BatchIndices.Add( IndexStart + 0 );
+			BatchIndices.Add( IndexStart + 1 );
+			BatchIndices.Add( IndexStart + 2 );
+			BatchIndices.Add( IndexStart + 2 );
+			BatchIndices.Add( IndexStart + 1 );
+			BatchIndices.Add( IndexStart + 3 );
+
+			BatchIndices.Add( IndexStart + 2 );
+			BatchIndices.Add( IndexStart + 3 );
+			BatchIndices.Add( IndexStart + 4 );
+			BatchIndices.Add( IndexStart + 4 );
+			BatchIndices.Add( IndexStart + 3 );
+			BatchIndices.Add( IndexStart + 5 );
+
+			BatchIndices.Add( IndexStart + 4 );
+			BatchIndices.Add( IndexStart + 5 );
+			BatchIndices.Add( IndexStart + 6 );
+			BatchIndices.Add( IndexStart + 6 );
+			BatchIndices.Add( IndexStart + 5 );
+			BatchIndices.Add( IndexStart + 7 );
+
+			// Middle
+			BatchIndices.Add( IndexStart + 1 );
+			BatchIndices.Add( IndexStart + 8 );
+			BatchIndices.Add( IndexStart + 3 );
+			BatchIndices.Add( IndexStart + 3 );
+			BatchIndices.Add( IndexStart + 8 );
+			BatchIndices.Add( IndexStart + 9 );
+
+			BatchIndices.Add( IndexStart + 3 );
+			BatchIndices.Add( IndexStart + 9 );
+			BatchIndices.Add( IndexStart + 5 );
+			BatchIndices.Add( IndexStart + 5 );
+			BatchIndices.Add( IndexStart + 9 );
+			BatchIndices.Add( IndexStart + 10 );
+
+			BatchIndices.Add( IndexStart + 5 );
+			BatchIndices.Add( IndexStart + 10 );
+			BatchIndices.Add( IndexStart + 7 );
+			BatchIndices.Add( IndexStart + 7 );
+			BatchIndices.Add( IndexStart + 10 );
+			BatchIndices.Add( IndexStart + 11 );
+
+			// Bottom
+			BatchIndices.Add( IndexStart + 8 );
+			BatchIndices.Add( IndexStart + 12 );
+			BatchIndices.Add( IndexStart + 9 );
+			BatchIndices.Add( IndexStart + 9 );
+			BatchIndices.Add( IndexStart + 12 );
+			BatchIndices.Add( IndexStart + 13 );
+
+			BatchIndices.Add( IndexStart + 9 );
+			BatchIndices.Add( IndexStart + 13 );
+			BatchIndices.Add( IndexStart + 10 );
+			BatchIndices.Add( IndexStart + 10 );
+			BatchIndices.Add( IndexStart + 13 );
+			BatchIndices.Add( IndexStart + 14 );
+
+			BatchIndices.Add( IndexStart + 10 );
+			BatchIndices.Add( IndexStart + 14 );
+			BatchIndices.Add( IndexStart + 11 );
+			BatchIndices.Add( IndexStart + 11 );
+			BatchIndices.Add( IndexStart + 14 );
+			BatchIndices.Add( IndexStart + 15 );
+		}
+		else
+		{
+			FVector2D TopLeft =	Position;
+			FVector2D TopRight = FVector2D( Position.X+Size.X, Position.Y);
+			FVector2D BotLeft =	 FVector2D( Position.X, Position.Y+Size.Y);
+			FVector2D BotRight = FVector2D( Position.X + Size.X, Position.Y+Size.Y);
+
+			// Add four vertices to the list of verts to be added to the vertex buffer
+			BatchVertices.Add( FSlateVertex(TopLeft,	StartUV,						Tiling,	Tint, InClippingRect ) );
+			BatchVertices.Add( FSlateVertex(TopRight,	FVector2D(EndUV.X,StartUV.Y),	Tiling,	Tint, InClippingRect ) );
+			BatchVertices.Add( FSlateVertex(BotLeft,	FVector2D(StartUV.X,EndUV.Y),	Tiling,	Tint, InClippingRect ) );
+			BatchVertices.Add( FSlateVertex(BotRight,	EndUV,							Tiling,	Tint, InClippingRect ) );
+
+			BatchIndices.Add( IndexStart + 0 );
+			BatchIndices.Add( IndexStart + 1 );
+			BatchIndices.Add( IndexStart + 2 );
+
+			BatchIndices.Add( IndexStart + 2 );
+			BatchIndices.Add( IndexStart + 1 );
+			BatchIndices.Add( IndexStart + 3 );
+		}
+	}
+}
+
+/** 
+ * Creates vertices necessary to draw a string (one quad per character)
+ *
+ * @param Position	The top left screen space position of the element
+ * @param InPayload	The data payload containing font,text, and size information for the string
+ */
+void FSlateElementBatcher::AddTextElement( const FVector2D& Position, const FVector2D& Size, float InScale, const FSlateDataPayload& InPayload, const FSlateRect& InClippingRect, ESlateDrawEffect::Type InDrawEffects, uint32 Layer )
+{
+	SCOPE_CYCLE_COUNTER(STAT_SlateBuildText);
+
+	const FString& Text = InPayload.Text;
+	// Nothing to do if no text
+	if( Text.Len() == 0 )
+	{
+		return;
+	}
+
+	FSlateFontCache& FontCache = *RenderingPolicy->GetFontCache().Get();
+
+	const float FontScale = InScale;
+
+	FCharacterList& CharacterList = FontCache.GetCharacterList( InPayload.FontInfo, FontScale );
+
+	float MaxHeight = CharacterList.GetMaxHeight();
+
+	uint32 FontTextureIndex = 0;
+	FSlateShaderResource* FontTexture = nullptr;
+
+	FSlateElementBatch* ElementBatch = nullptr;
+	TArray<FSlateVertex>* BatchVertices = nullptr;
+	TArray<SlateIndex>* BatchIndices = nullptr;
+
+	uint32 VertexOffset = 0;
+	uint32 IndexOffset = 0;
+
+	float InvTextureSizeX = 0;
+	float InvTextureSizeY = 0;
+
+	int32 LineX = 0;
+
+	TCHAR PreviousChar = 0;
+
+	int32 Kerning = 0;
+
+	int32 PosX = FMath::Trunc(Position.X);
+	int32 PosY = FMath::Trunc(Position.Y);
+
+	LineX = PosX;
+	
+	FColor FinalColor = GetElementColor( InPayload.Tint, nullptr );
+
+	for( int32 CharIndex = 0; CharIndex < Text.Len(); ++CharIndex )
+	{
+		TCHAR CurrentChar = Text[ CharIndex ];
+
+		const bool IsNewline = (CurrentChar == '\n');
+
+		if (IsNewline)
+		{
+			// Move down: we are drawing the next line.
+			PosY += MaxHeight;
+			// Carriage return 
+			LineX = PosX;
+		}
+		else
+		{
+			const FCharacterEntry& Entry = CharacterList[ CurrentChar ];
+
+			if( FontTexture == nullptr || Entry.TextureIndex != FontTextureIndex )
+			{
+				// Font has a new texture for this glyph. Refresh the batch we use and the index we are currently using
+				FontTextureIndex = Entry.TextureIndex;
+
+				FontTexture = FontCache.GetTextureResource( FontTextureIndex );
+				ElementBatch = &FindBatchForElement( Layer, FShaderParams(), FontTexture, ESlateDrawPrimitive::TriangleList, ESlateShader::Font, InDrawEffects );
+
+				BatchVertices = &BatchVertexArrays[ElementBatch->VertexArrayIndex];
+				BatchIndices = &BatchIndexArrays[ElementBatch->IndexArrayIndex];
+
+				VertexOffset = BatchVertices->Num();
+				IndexOffset = BatchIndices->Num();
+				
+				InvTextureSizeX = 1.0f/FontTexture->GetWidth();
+				InvTextureSizeY = 1.0f/FontTexture->GetHeight();
+			}
+
+			const bool bIsWhitespace = FChar::IsWhitespace(CurrentChar);
+
+			if( !bIsWhitespace && CharIndex > 0 )
+			{
+				Kerning = CharacterList.GetKerning( PreviousChar, CurrentChar );
+			}
+			else
+			{
+				Kerning = 0;
+			}
+
+			LineX += Kerning;
+			PreviousChar = CurrentChar;
+
+			if( !bIsWhitespace )
+			{
+				const float X = LineX + Entry.HorizontalOffset;
+				// Note PosX,PosY is the upper left corner of the bounding box representing the string.  This computes the Y position of the baseline where text will sit
+
+				const float Y = PosY - Entry.VerticalOffset+MaxHeight+Entry.GlobalDescender;
+				const float U = Entry.StartU * InvTextureSizeX;
+				const float V = Entry.StartV * InvTextureSizeY;
+				const float SizeX = Entry.USize;
+				const float SizeY = Entry.VSize;
+				const float SizeU = Entry.USize * InvTextureSizeX;
+				const float SizeV = Entry.VSize * InvTextureSizeY;
+
+				FSlateRect CharRect( X, Y, X+SizeX, Y+SizeY );
+
+				if( FSlateRect::DoRectanglesIntersect( InClippingRect, CharRect )	)
+				{
+					TArray<FSlateVertex>& BatchVerticesRef = *BatchVertices;
+					TArray<SlateIndex>& BatchIndicesRef = *BatchIndices;
+
+					FVector2D UpperLeft( X, Y );
+					FVector2D UpperRight( X+SizeX, Y );
+					FVector2D LowerLeft( X, Y+SizeY );
+					FVector2D LowerRight( X+SizeX, Y+SizeY );
+
+					// Add four vertices for this quad
+					BatchVerticesRef.AddUninitialized( 4 );
+					// Add six indices for this quad
+					BatchIndicesRef.AddUninitialized( 6 );
+
+					// The start index of these vertices in the index buffer
+					uint32 IndexStart = VertexOffset;
+
+					// Add four vertices to the list of verts to be added to the vertex buffer
+					BatchVerticesRef[ VertexOffset++ ] = FSlateVertex( UpperLeft,								FVector2D(U,V),				FinalColor, InClippingRect );
+					BatchVerticesRef[ VertexOffset++ ] = FSlateVertex( FVector2D(LowerRight.X,UpperLeft.Y),		FVector2D(U+SizeU, V),		FinalColor, InClippingRect );
+					BatchVerticesRef[ VertexOffset++ ] = FSlateVertex( FVector2D(UpperLeft.X,LowerRight.Y),		FVector2D(U, V+SizeV),		FinalColor, InClippingRect );
+					BatchVerticesRef[ VertexOffset++ ] = FSlateVertex( LowerRight,								FVector2D(U+SizeU, V+SizeV),FinalColor, InClippingRect );
+
+					BatchIndicesRef[IndexOffset++] = IndexStart + 0;
+					BatchIndicesRef[IndexOffset++] = IndexStart + 1;
+					BatchIndicesRef[IndexOffset++] = IndexStart + 2;
+					BatchIndicesRef[IndexOffset++] = IndexStart + 1;
+					BatchIndicesRef[IndexOffset++] = IndexStart + 3;
+					BatchIndicesRef[IndexOffset++] = IndexStart + 2;
+				}
+			}
+
+			LineX += Entry.XAdvance;
+		}
+	}
+}
+
+/** 
+ * Creates vertices necessary to draw a gradient box (horizontal or vertical)
+ *
+ * @param Position	The top left screen space position of the element
+ * @param Size		The size of the element
+ * @param InPayload	The data payload containing gradient stops
+ */
+void FSlateElementBatcher::AddGradientElement( const FVector2D& Position, const FVector2D& Size, float Scale, const class FSlateDataPayload& InPayload, const FSlateRect& InClippingRect, ESlateDrawEffect::Type InDrawEffects, uint32 Layer )
+{
+	SCOPE_CYCLE_COUNTER(STAT_SlateBuildGradients);
+
+	// There must be at least one gradient stop
+	check( InPayload.GradientStops.Num() > 0 );
+
+	FSlateElementBatch& ElementBatch = FindBatchForElement( Layer, FShaderParams(), nullptr, ESlateDrawPrimitive::TriangleList, ESlateShader::Default, InDrawEffects, InPayload.bGammaCorrect == false ? ESlateBatchDrawFlag::NoGamma : ESlateBatchDrawFlag::None );
+	TArray<FSlateVertex>& BatchVertices = BatchVertexArrays[ElementBatch.VertexArrayIndex];
+	TArray<SlateIndex>& BatchIndices = BatchIndexArrays[ElementBatch.IndexArrayIndex];
+
+	// Determine the four corners of the quad containing the gradient
+	FVector2D TopLeft = Position;
+	FVector2D TopRight = FVector2D(Position.X + Size.X, Position.Y);
+	FVector2D BotLeft = FVector2D(Position.X, Position.Y + Size.Y);
+	FVector2D BotRight = FVector2D(Position.X + Size.X, Position.Y + Size.Y);
+
+	// Copy the gradient stops.. We may need to add more
+	TArray<FSlateGradientStop> GradientStops = InPayload.GradientStops;
+
+	const FSlateGradientStop& FirstStop = InPayload.GradientStops[0];
+	const FSlateGradientStop& LastStop = InPayload.GradientStops[ InPayload.GradientStops.Num() - 1 ];
+		
+	// Determine if the first and last stops are not at the start and end of the quad
+	// If they are not add a gradient stop with the same color as the first and/or last stop
+	if( InPayload.GradientType == Orient_Vertical )
+	{
+		if( TopLeft.X < FirstStop.Position.X + TopLeft.X )
+		{
+			// The first stop is after the left side of the quad.  Add a stop at the left side of the quad using the same color as the first stop
+			GradientStops.Insert( FSlateGradientStop( FVector2D(0.0, 0.0), FirstStop.Color ), 0 );
+		}
+
+		if( TopRight.X > LastStop.Position.X + TopLeft.X )
+		{
+			// The last stop is before the right side of the quad.  Add a stop at the right side of the quad using the same color as the last stop
+			GradientStops.Add( FSlateGradientStop( BotRight, LastStop.Color ) ); 
+		}
+	}
+	else
+	{
+
+		if( TopLeft.Y < FirstStop.Position.Y + TopLeft.Y )
+		{
+			// The first stop is after the top side of the quad.  Add a stop at the top side of the quad using the same color as the first stop
+			GradientStops.Insert( FSlateGradientStop( FVector2D(0.0, 0.0), FirstStop.Color ), 0 );
+		}
+
+		if( BotLeft.Y > LastStop.Position.Y + TopLeft.Y )
+		{
+			// The last stop is before the bottom side of the quad.  Add a stop at the bottom side of the quad using the same color as the last stop
+			GradientStops.Add( FSlateGradientStop( BotRight, LastStop.Color ) ); 
+		}
+	}
+
+	uint32 IndexOffsetStart = BatchIndices.Num();
+
+	// Add a pair of vertices for each gradient stop. Connecting them to the previous stop if necessary
+	// Assumes gradient stops are sorted by position left to right or top to bottom
+	for( int32 StopIndex = 0; StopIndex < GradientStops.Num(); ++StopIndex )
+	{
+		uint32 IndexStart = BatchVertices.Num();
+
+		const FSlateGradientStop& CurStop = GradientStops[StopIndex];
+
+		// The start vertex at this stop
+		FVector2D StartPt;
+		// The end vertex at this stop
+		FVector2D EndPt;
+
+		if( InPayload.GradientType == Orient_Vertical )
+		{
+			// Gradient stop is vertical so gradients to left to right
+			StartPt = TopLeft;
+			EndPt = BotLeft;
+			StartPt.X += CurStop.Position.X * Scale;
+			EndPt.X += CurStop.Position.X  * Scale;	
+		}
+		else
+		{
+			// Gradient stop is horizontal so gradients to top to bottom
+			StartPt = TopLeft;
+			EndPt = TopRight;
+			StartPt.Y += CurStop.Position.Y * Scale;
+			EndPt.Y += CurStop.Position.Y * Scale;
+		}
+
+		if( StopIndex == 0 )
+		{
+			// First stop does not have a full quad yet so do not create indices
+			BatchVertices.Add( FSlateVertex(StartPt, FVector2D::ZeroVector, FVector2D::ZeroVector, CurStop.Color.ToFColor(false), InClippingRect ) );
+			BatchVertices.Add( FSlateVertex(EndPt, FVector2D::ZeroVector, FVector2D::ZeroVector, CurStop.Color.ToFColor(false), InClippingRect ) );
+		}
+		else
+		{
+			// All stops after the first have indices and generate quads
+			BatchVertices.Add( FSlateVertex(StartPt, FVector2D::ZeroVector, FVector2D::ZeroVector, CurStop.Color.ToFColor(false), InClippingRect ) );
+			BatchVertices.Add( FSlateVertex(EndPt, FVector2D::ZeroVector, FVector2D::ZeroVector, CurStop.Color.ToFColor(false), InClippingRect ) );
+
+			// Connect the indices to the previous vertices
+			BatchIndices.Add( IndexStart - 2 );
+			BatchIndices.Add( IndexStart - 1 );
+			BatchIndices.Add( IndexStart + 0 );
+
+			BatchIndices.Add( IndexStart + 0 );
+			BatchIndices.Add( IndexStart - 1 );
+			BatchIndices.Add( IndexStart + 1 );
+		}
+	}
+}
+
+void FSlateElementBatcher::AddSplineElement( const FVector2D& Position, float Scale, const FSlateDataPayload& InPayload, const FSlateRect& InClippingRect, ESlateDrawEffect::Type InDrawEffects, uint32 Layer)
+{
+	SCOPE_CYCLE_COUNTER(STAT_SlateBuildSplines);
+
+	//@todo SLATE: Merge with AddLineElement?
+
+	const float DirectLength = (InPayload.EndPt - InPayload.StartPt).Size();
+	const float HandleLength = ((InPayload.EndPt - InPayload.EndDir) - (InPayload.StartPt + InPayload.StartDir)).Size();
+	float NumSteps = FMath::Clamp<float>(FMath::Ceil(FMath::Max(DirectLength,HandleLength)/15.0f), 1, 256);
+
+	// 1 is the minimum thickness we support
+	float InThickness = FMath::Max( 1.0f, InPayload.Thickness );
+
+	// The radius to use when checking the distance of pixels to the actual line.  Arbitrary value based on what looks the best
+	const float Radius = 1.5f;
+
+	// Compute the actual size of the line we need based on thickness.  Need to ensure pixels that are at least Thickness/2 + Sample radius are generated so that we have enough pixels to blend.
+	// The idea for the spline anti-alising technique is based on the fast prefiltered lines technique published in GPU Gems 2 
+	const float LineThickness = FMath::Ceil( (2.0f * Radius + InThickness ) * FMath::Sqrt(2.0f) );
+
+	// The amount we increase each side of the line to generate enough pixels
+	const float HalfThickness = LineThickness * .5f + Radius;
+
+	// Currently splines are not atlased because they are tiled.  So we just assume the texture proxy holds the actual texture
+	FSlateShaderResourceProxy* ResourceProxy = RenderingPolicy->GetTextureResource( SplineFilterTable );
+	check(ResourceProxy && ResourceProxy->Resource);
+
+	// Find a batch for the element
+	FSlateElementBatch& ElementBatch = FindBatchForElement( Layer, FShaderParams::MakePixelShaderParams( FVector4( InPayload.Thickness,Radius,0,0) ), ResourceProxy->Resource, ESlateDrawPrimitive::TriangleList, ESlateShader::LineSegment, InDrawEffects );
+	TArray<FSlateVertex>& BatchVertices = BatchVertexArrays[ElementBatch.VertexArrayIndex];
+	TArray<SlateIndex>& BatchIndices = BatchIndexArrays[ElementBatch.IndexArrayIndex];
+
+
+	const FVector2D StartPt = Position + InPayload.StartPt * Scale;
+	const FVector2D StartDir = InPayload.StartDir * Scale;
+	const FVector2D EndPt = Position + InPayload.EndPt * Scale;
+	const FVector2D EndDir = InPayload.EndDir * Scale;
+	
+	// Compute the normal to the line
+	FVector2D Normal = FVector2D( StartPt.Y - EndPt.Y, EndPt.X - StartPt.X ).SafeNormal();
+
+	FVector2D Up = Normal * HalfThickness;
+
+	// Generate the first segment
+	const float Alpha = 1.0f/NumSteps;
+	FVector2D StartPos = StartPt;
+	FVector2D EndPos = FVector2D( FMath::CubicInterp( StartPt, StartDir, EndPt, EndDir, Alpha ) );
+
+	FColor FinalColor = GetElementColor( InPayload.Tint, nullptr );
+
+	BatchVertices.Add( FSlateVertex( StartPos + Up, StartPos, EndPos, FinalColor, InClippingRect ) );
+	BatchVertices.Add( FSlateVertex( StartPos - Up, StartPos, EndPos, FinalColor, InClippingRect ) );
+
+	// Generate the rest of the segments
+	for( int32 Step = 0; Step < NumSteps; ++Step )
+	{
+		// Skip the first step as it was already generated
+		if( Step > 0 )
+		{
+			const float StepAlpha = (Step+1.0f)/NumSteps;
+			EndPos = FVector2D( FMath::CubicInterp( StartPt, StartDir, EndPt, EndDir, StepAlpha ) );
+		}
+
+		int32 IndexStart = BatchVertices.Num();
+
+		// Compute the normal to the line
+		FVector2D SegmentNormal = FVector2D( StartPos.Y - EndPos.Y, EndPos.X - StartPos.X ).SafeNormal();
+
+		// Create the new vertices for the thick line segment
+		Up = SegmentNormal * HalfThickness;
+
+		BatchVertices.Add( FSlateVertex( EndPos + Up, StartPos, EndPos, FinalColor, InClippingRect ) );
+		BatchVertices.Add( FSlateVertex( EndPos - Up, StartPos, EndPos, FinalColor, InClippingRect ) );
+
+		BatchIndices.Add( IndexStart - 2 );
+		BatchIndices.Add( IndexStart - 1 );
+		BatchIndices.Add( IndexStart + 0 );
+
+		BatchIndices.Add( IndexStart + 0 );
+		BatchIndices.Add( IndexStart + 1 );
+		BatchIndices.Add( IndexStart - 1 );
+
+		StartPos = EndPos;
+	}
+}
+
+/**
+ * Calculates the intersection of two line segments P1->P2, P3->P4
+ * The tolerance setting is used when the lines arent currently intersecting but will intersect in the future  
+ * The higher the tolerance the greater the distance that the intersection point can be.
+ *
+ * @return true if the line intersects.  Populates Intersection
+ */
+static bool LineIntersect( const FVector2D& P1, const FVector2D& P2, const FVector2D& P3, const FVector2D& P4, FVector2D& Intersect, float Tolerance = .1f)
+{
+	float NumA = ( (P4.X-P3.X)*(P1.Y-P3.Y) - (P4.Y-P3.Y)*(P1.X-P3.X) );
+	float NumB =  ( (P2.X-P1.X)*(P1.Y-P3.Y) - (P2.Y-P1.Y)*(P1.X-P3.X) );
+
+	float Denom = (P4.Y-P3.Y)*(P2.X-P1.X) - (P4.X-P3.X)*(P2.Y-P1.Y);
+
+	if( FMath::IsNearlyZero( NumA ) && FMath::IsNearlyZero( NumB )  )
+	{
+		// Lines are the same
+		Intersect = (P1 + P2) / 2;
+		return true;
+	}
+
+	if( FMath::IsNearlyZero(Denom) )
+	{
+		// Lines are parallel
+		return false;
+	}
+	 
+	float B = NumB / Denom;
+	float A = NumA / Denom;
+
+	// Note that this is a "tweaked" intersection test for the purpose of joining line segments.  We don't just want to know if the line segements
+	// Intersect, but where they would if they don't currently. Except that we dont care in the case that where the segments 
+	// intersection is so far away that its infeasible to use the intersection point later.
+	if( A >= -Tolerance && A <= (1.0f + Tolerance ) && B >= -Tolerance && B <= (1.0f + Tolerance) )
+	{
+		Intersect = P1+A*(P2-P1);
+		return true;
+	}
+
+	return false;
+}
+	
+
+void FSlateElementBatcher::AddLineElement( const FVector2D& Position, const FVector2D& Size, float Scale, const FSlateDataPayload& InPayload, const FSlateRect& InClippingRect, ESlateDrawEffect::Type DrawEffects, uint32 Layer )
+{
+	SCOPE_CYCLE_COUNTER(STAT_SlateBuildLines);
+
+	if( InPayload.Points.Num() < 2 )
+	{
+		return;
+	}
+
+	FColor FinalColor = GetElementColor( InPayload.Tint, nullptr );
+
+	if( InPayload.bAntialias )
+	{
+		// Currently splines are not atlased because they are tiled.  So we just assume the texture proxy holds the actual texture
+		FSlateShaderResourceProxy* ResourceProxy = RenderingPolicy->GetTextureResource( SplineFilterTable );
+		check(ResourceProxy && ResourceProxy->Resource);
+
+		// The radius to use when checking the distance of pixels to the actual line.  Arbitrary value based on what looks the best
+		const float Radius = 1.5f;
+
+		float RequestedThickness = FMath::Max( 1.0f, InPayload.Thickness );
+
+		// Compute the actual size of the line we need based on thickness.  Need to ensure pixels that are at least Thickness/2 + Sample radius are generated so that we have enough pixels to blend.
+		// The idea for the anti-aliasing technique is based on the fast prefiltered lines technique published in GPU Gems 2 
+		const float LineThickness = FMath::Ceil( (2.0f * Radius + RequestedThickness ) * FMath::Sqrt(2.0f) );
+
+		// The amount we increase each side of the line to generate enough pixels
+		const float HalfThickness = LineThickness * .5f + Radius;
+
+		// Find a batch for the element
+		FSlateElementBatch& ElementBatch = FindBatchForElement( Layer, FShaderParams::MakePixelShaderParams( FVector4(RequestedThickness,Radius,0,0) ), ResourceProxy->Resource, ESlateDrawPrimitive::TriangleList, ESlateShader::LineSegment, DrawEffects );
+		TArray<FSlateVertex>& BatchVertices = BatchVertexArrays[ElementBatch.VertexArrayIndex];
+		TArray<SlateIndex>& BatchIndices = BatchIndexArrays[ElementBatch.IndexArrayIndex];
+
+		const TArray<FVector2D>& Points = InPayload.Points;
+
+		FVector2D StartPos = Position + Points[0]*Scale;
+		FVector2D EndPos = Position + Points[1]*Scale;
+
+		FVector2D Normal = FVector2D( StartPos.Y - EndPos.Y, EndPos.X - StartPos.X ).SafeNormal();
+
+		FVector2D Up = Normal * HalfThickness;
+
+		BatchVertices.Add( FSlateVertex( StartPos + Up, StartPos, EndPos, FinalColor, InClippingRect ) );
+		BatchVertices.Add( FSlateVertex( StartPos - Up, StartPos, EndPos, FinalColor, InClippingRect ) );
+
+		// Generate the rest of the segments
+		for( int32 Point = 1; Point < Points.Num(); ++Point )
+		{
+			EndPos = Position + Points[Point]*Scale;
+			// Determine if we should check the intersection point with the next line segment.
+			// We will adjust were this line ends to the intersection
+			bool bCheckIntersection = Points.IsValidIndex( Point+1 ) ;
+			uint32 IndexStart = BatchVertices.Num();
+
+			// Compute the normal to the line
+			Normal = FVector2D( StartPos.Y - EndPos.Y, EndPos.X - StartPos.X ).SafeNormal();
+
+			// Create the new vertices for the thick line segment
+			Up = Normal * HalfThickness;
+
+			FVector2D IntersectUpper = EndPos + Up;
+			FVector2D IntersectLower = EndPos - Up;
+			FVector2D IntersectCenter = EndPos;
+
+			if( bCheckIntersection )
+			{
+				// The end point of the next segment
+				const FVector2D NextEndPos = Position + Points[Point+1]*Scale;
+
+				// The normal of the next segment
+				const FVector2D NextNormal = FVector2D( EndPos.Y - NextEndPos.Y, NextEndPos.X - EndPos.X ).SafeNormal();
+
+				// The next amount to adjust the vertices by 
+				FVector2D NextUp = NextNormal * HalfThickness;
+
+				FVector2D IntersectionPoint;
+				if( LineIntersect( StartPos + Up, EndPos + Up, EndPos + NextUp, NextEndPos + NextUp, IntersectionPoint ) )
+				{
+					// If the lines intersect adjust where the line starts
+					IntersectUpper = IntersectionPoint;
+
+					// visualizes the intersection
+					//AddQuadElement( IntersectUpper-FVector2D(1,1), FVector2D(2,2), 1, InClippingRect, Layer+1, FColor(255,0,0));
+
+				}
+
+				if( LineIntersect( StartPos - Up, EndPos - Up, EndPos - NextUp, NextEndPos - NextUp, IntersectionPoint ) )
+				{
+					// If the lines intersect adjust where the line starts
+					IntersectLower = IntersectionPoint;
+
+					// visualizes the intersection
+					//AddQuadElement( IntersectLower-FVector2D(1,1), FVector2D(2,2), 1, InClippingRect, Layer+1, FColor(255,255,0));
+				}
+				// the midpoint of the intersection.  Used as the new end to the line segment (not adjusted for anti-aliasing)
+				IntersectCenter = (IntersectUpper+IntersectLower) * .5f;
+			}
+
+			if( Point > 1 )
+			{
+				// Make a copy of the last two vertices and update their start and end position to reflect the new line segment
+				FSlateVertex StartV1 = BatchVertices[IndexStart-1];
+				FSlateVertex StartV2 = BatchVertices[IndexStart-2];
+
+				StartV1.TexCoords.X = StartPos.X;
+				StartV1.TexCoords.Y = StartPos.Y;
+				StartV1.TexCoords.Z = IntersectCenter.X;
+				StartV1.TexCoords.W = IntersectCenter.Y;
+
+				StartV2.TexCoords.X = StartPos.X;
+				StartV2.TexCoords.Y = StartPos.Y;
+				StartV2.TexCoords.Z = IntersectCenter.X;
+				StartV2.TexCoords.W = IntersectCenter.Y;
+
+				IndexStart += 2;
+				BatchVertices.Add( StartV2 );
+				BatchVertices.Add( StartV1 );
+			}
+
+			BatchVertices.Add( FSlateVertex( IntersectUpper, StartPos, IntersectCenter, FinalColor, InClippingRect ) );
+			BatchVertices.Add( FSlateVertex( IntersectLower, StartPos, IntersectCenter, FinalColor, InClippingRect ) );
+
+			BatchIndices.Add( IndexStart - 1 );
+			BatchIndices.Add( IndexStart - 2 );
+			BatchIndices.Add( IndexStart + 0 );
+
+			BatchIndices.Add( IndexStart + 0 );
+			BatchIndices.Add( IndexStart + 1 );
+			BatchIndices.Add( IndexStart - 1 );
+
+			StartPos = EndPos;
+		}
+	}
+	else
+	{
+		// Find a batch for the element
+		FSlateElementBatch& ElementBatch = FindBatchForElement( Layer, FShaderParams(), nullptr, ESlateDrawPrimitive::LineList, ESlateShader::Default, DrawEffects );
+		TArray<FSlateVertex>& BatchVertices = BatchVertexArrays[ElementBatch.VertexArrayIndex];
+		TArray<SlateIndex>& BatchIndices = BatchIndexArrays[ElementBatch.IndexArrayIndex];
+
+		// Generate the rest of the segments
+		for( int32 Point = 0; Point < InPayload.Points.Num()-1; ++Point )
+		{
+			uint32 IndexStart = BatchVertices.Num();
+			FVector2D StartPos = Position + InPayload.Points[Point]*Scale;
+			FVector2D EndPos = Position + InPayload.Points[Point+1]*Scale;
+
+
+			BatchVertices.Add( FSlateVertex( StartPos, FVector2D::ZeroVector, FinalColor, InClippingRect ) );
+			BatchVertices.Add( FSlateVertex( EndPos, FVector2D::ZeroVector, FinalColor, InClippingRect ) );
+
+			BatchIndices.Add( IndexStart );
+			BatchIndices.Add( IndexStart + 1 );
+		}
+	}
+}
+
+/** 
+ * Creates vertices necessary to draw viewport
+ *
+ * @param Position	The top left screen space position of the element
+ * @param Size		The size of the element
+ * @param InPayload	The data payload containing gradient stops
+ * @param InClippingRect	The clipping rect of the element
+ */
+void FSlateElementBatcher::AddViewportElement( const FVector2D& InPosition, const FVector2D& InSize, float Scale, const FSlateDataPayload& InPayload, const FSlateRect& InClippingRect, ESlateDrawEffect::Type InDrawEffects, uint32 Layer )
+{
+	SCOPE_CYCLE_COUNTER(STAT_SlateBuildViewports);
+
+	const FColor FinalColor = GetElementColor( InPayload.Tint, nullptr );
+
+	ESlateBatchDrawFlag::Type DrawFlags = ESlateBatchDrawFlag::None;
+	
+	if( !InPayload.bAllowBlending )
+	{
+		DrawFlags |= ESlateBatchDrawFlag::NoBlending;
+	}
+
+	if( !InPayload.bGammaCorrect )
+	{
+		// Don't apply gamma correction
+		DrawFlags |= ESlateBatchDrawFlag::NoGamma;
+	}
+
+	TSharedPtr<const ISlateViewport> ViewportPin = InPayload.Viewport.Pin();
+
+	FSlateElementBatch& ElementBatch = FindBatchForElement( Layer, FShaderParams(), RenderingPolicy->GetViewportResource( ViewportPin.Get() ), ESlateDrawPrimitive::TriangleList, ESlateShader::Default, InDrawEffects, DrawFlags );
+	TArray<FSlateVertex>& BatchVertices = BatchVertexArrays[ElementBatch.VertexArrayIndex];
+	TArray<SlateIndex>& BatchIndices = BatchIndexArrays[ElementBatch.IndexArrayIndex];
+
+	// Tag this batch as requiring vsync if the viewport requires it.
+	if( ViewportPin.IsValid() )
+	{
+		bRequiresVsync |= ViewportPin->RequiresVsync();
+	}
+
+	FVector2D Position = FVector2D( FMath::Trunc( InPosition.X ), FMath::Trunc( InPosition.Y ) );
+	FVector2D Size = FVector2D( FMath::Trunc( InSize.X ), FMath::Trunc( InSize.Y ) );
+
+	// Determine the four corners of the quad
+	FVector2D TopLeft = Position;
+	FVector2D TopRight = FVector2D(Position.X+Size.X, Position.Y);
+	FVector2D BotLeft = FVector2D(Position.X, Position.Y+Size.Y);
+	FVector2D BotRight = FVector2D(Position.X + Size.X, Position.Y+Size.Y);
+
+	// The start index of these vertices in the index buffer
+	uint32 IndexStart = BatchVertices.Num();
+
+	// Add four vertices to the list of verts to be added to the vertex buffer
+	BatchVertices.Add( FSlateVertex( TopLeft,	FVector2D(0.0f,0.0f),	FinalColor, InClippingRect ) );
+	BatchVertices.Add( FSlateVertex( TopRight,	FVector2D(1.0f,0.0f),	FinalColor, InClippingRect ) );
+	BatchVertices.Add( FSlateVertex( BotLeft,	FVector2D(0.0f,1.0f),	FinalColor, InClippingRect ) );
+	BatchVertices.Add( FSlateVertex( BotRight,	FVector2D(1.0f,1.0f),	FinalColor, InClippingRect ) );
+
+	// The offset into the index buffer where this quads indices start
+	uint32 IndexOffsetStart = BatchIndices.Num();
+
+	// Add 6 indices to the vertex buffer.  (2 tri's per quad, 3 indices per tri)
+	BatchIndices.Add( IndexStart + 0 );
+	BatchIndices.Add( IndexStart + 1 );
+	BatchIndices.Add( IndexStart + 2 );
+
+	BatchIndices.Add( IndexStart + 2 );
+	BatchIndices.Add( IndexStart + 1 );
+	BatchIndices.Add( IndexStart + 3 );
+
+}
+
+/** 
+ * Creates vertices necessary to draw a border element
+ *
+ * @param Position	The top left screen space position of the element
+ * @param Size		The size of the element
+ * @param InPayload	The data payload containing gradient stops
+ * @param InClippingRect	The clipping rect of the element
+ */
+void FSlateElementBatcher::AddBorderElement( const FVector2D& InPosition, const FVector2D& Size, float Scale, const FSlateDataPayload& InPayload, const FSlateRect& InClippingRect, ESlateDrawEffect::Type InDrawEffects, uint32 Layer )
+{
+	SCOPE_CYCLE_COUNTER(STAT_SlateBuildBorders);
+
+	check( InPayload.BrushResource );
+
+	const FVector2D Position( FMath::Trunc(InPosition.X), FMath::Trunc(InPosition.Y) );
+
+	uint32 TextureWidth = 1;
+	uint32 TextureHeight = 1;
+
+	// Currently borders are not atlased because they are tiled.  So we just assume the texture proxy holds the actual texture
+	FSlateShaderResourceProxy* ResourceProxy = RenderingPolicy->GetTextureResource( *InPayload.BrushResource );
+	FSlateShaderResource* Resource = ResourceProxy ? ResourceProxy->Resource : nullptr;
+	if( Resource )
+	{
+		TextureWidth = Resource->GetWidth();
+		TextureHeight = Resource->GetHeight();
+	}
+ 
+	// Texel offset
+	const float PixelCenterOffset = RenderingPolicy->GetPixelCenterOffset();
+	const FVector2D HalfTexel( PixelCenterOffset/TextureWidth, PixelCenterOffset/TextureHeight );
+
+	const FVector2D StartUV = HalfTexel;
+	const FVector2D EndUV = FVector2D( 1.0f, 1.0f ) + HalfTexel;
+
+	const FMargin& Margin = InPayload.BrushResource->Margin;
+
+	const FMargin ScaledMargin = Margin * Scale;
+
+	// Determine the margins for each quad
+	float LeftMarginX = TextureWidth * ScaledMargin.Left;
+	float TopMarginY = TextureHeight * ScaledMargin.Top;
+	float RightMarginX = Size.X - TextureWidth * ScaledMargin.Right;
+	float BottomMarginY = Size.Y - TextureHeight * ScaledMargin.Bottom;
+
+	// If the margins are overlapping the margins are too big or the button is too small
+	// so clamp margins to half of the box size
+	if( RightMarginX < LeftMarginX )
+	{
+		LeftMarginX = Size.X / 2;
+		RightMarginX = LeftMarginX;
+	}
+
+	if( BottomMarginY < TopMarginY )
+	{
+		TopMarginY = Size.Y / 2;
+		BottomMarginY = TopMarginY;
+	}
+
+	// Determine the texture coordinates for each quad
+	float LeftMarginU = (Margin.Left > 0.0f) ? Margin.Left : 0.0f;
+	float TopMarginV = (Margin.Top > 0.0f) ? Margin.Top : 0.0f;
+	float RightMarginU = (Margin.Right > 0.0f) ? 1.0f - Margin.Right : 1.0f;
+	float BottomMarginV = (Margin.Bottom > 0.0f) ? 1.0f - Margin.Bottom : 1.0f;
+
+	LeftMarginU += HalfTexel.X;
+	TopMarginV += HalfTexel.Y;
+	BottomMarginV += HalfTexel.Y;
+	RightMarginU += HalfTexel.X;
+
+	// Determine the amount of tiling needed for the texture in this element.  The formula is number of pixels covered by the tiling portion of the texture / the number number of texels corresponding to the tiled portion of the texture.
+	float TopTiling = (RightMarginX-LeftMarginX)/(TextureWidth * ( 1 - Margin.GetTotalSpaceAlong<Orient_Horizontal>() ));
+	float LeftTiling = (BottomMarginY-TopMarginY)/(TextureHeight * (1 - Margin.GetTotalSpaceAlong<Orient_Vertical>() ));
+	
+	FShaderParams ShaderParams = FShaderParams::MakePixelShaderParams( FVector4(LeftMarginU,RightMarginU,TopMarginV,BottomMarginV) );
+
+	// The tint color applies to all brushes and is passed per vertex
+	FColor Tint = GetElementColor( InPayload.Tint, InPayload.BrushResource );
+
+	// Pass the tiling information as a flag so we can pick the correct texture addressing mode
+	ESlateBatchDrawFlag::Type DrawFlags = (ESlateBatchDrawFlag::TileU|ESlateBatchDrawFlag::TileV);
+
+	FSlateElementBatch& ElementBatch = FindBatchForElement( Layer, ShaderParams, Resource, ESlateDrawPrimitive::TriangleList, ESlateShader::Border, InDrawEffects, DrawFlags );
+	TArray<FSlateVertex>& BatchVertices = BatchVertexArrays[ ElementBatch.VertexArrayIndex ];
+	TArray<SlateIndex>& BatchIndices = BatchIndexArrays[ ElementBatch.IndexArrayIndex ];
+
+	// Ensure tiling of at least 1.  
+	TopTiling = TopTiling >= 1.0f ? TopTiling : 1.0f;
+	LeftTiling = LeftTiling >= 1.0f ? LeftTiling : 1.0f;
+	float RightTiling = LeftTiling;
+	float BottomTiling = TopTiling;
+
+	LeftMarginX += Position.X;
+	TopMarginY += Position.Y;
+	RightMarginX += Position.X;
+	BottomMarginY += Position.Y;
+
+	FVector2D EndPos = FVector2D(Position.X, Position.Y) + Size;
+
+	EndPos.X = FMath::Trunc(EndPos.X);
+	EndPos.Y = FMath::Trunc(EndPos.Y);
+
+	// The start index of these vertices in the index buffer
+	uint32 IndexStart = BatchVertices.Num();
+
+	// Zero for second UV indicates no tiling and to just pass the UV though (for the corner sections)
+	FVector2D Zero(0,0);
+
+	// Add all the vertices needed for this element.  Vertices are duplicated so that we can have some sections with no tiling and some with tiling.
+	BatchVertices.Add( FSlateVertex( Position,									FVector2D( StartUV.X, StartUV.Y ),			Zero,							Tint, InClippingRect ) ); //0
+	BatchVertices.Add( FSlateVertex( FVector2D( Position.X, TopMarginY ),		FVector2D( StartUV.X, TopMarginV ),			Zero,							Tint, InClippingRect ) ); //1
+	BatchVertices.Add( FSlateVertex( FVector2D( LeftMarginX, Position.Y ),		FVector2D( LeftMarginU, StartUV.Y ),		Zero,							Tint, InClippingRect ) ); //2
+	BatchVertices.Add( FSlateVertex( FVector2D( LeftMarginX, TopMarginY ),		FVector2D( LeftMarginU, TopMarginV ),		Zero,							Tint, InClippingRect ) ); //3
+
+	BatchVertices.Add( FSlateVertex( FVector2D( LeftMarginX, Position.Y ),		FVector2D( StartUV.X, StartUV.Y ),			FVector2D(TopTiling, 0.0f),		Tint, InClippingRect ) ); //4
+	BatchVertices.Add( FSlateVertex( FVector2D( LeftMarginX, TopMarginY ),		FVector2D( StartUV.X, TopMarginV ),			FVector2D(TopTiling, 0.0f),		Tint, InClippingRect ) ); //5
+	BatchVertices.Add( FSlateVertex( FVector2D( RightMarginX, Position.Y ),		FVector2D( EndUV.X, StartUV.Y ),			FVector2D(TopTiling, 0.0f),		Tint, InClippingRect ) ); //6
+	BatchVertices.Add( FSlateVertex( FVector2D( RightMarginX, TopMarginY ),		FVector2D( EndUV.X, TopMarginV),			FVector2D(TopTiling, 0.0f),		Tint, InClippingRect ) ); //7
+
+	BatchVertices.Add( FSlateVertex( FVector2D( RightMarginX, Position.Y ),		FVector2D( RightMarginU, StartUV.Y ),		Zero,							Tint, InClippingRect ) ); //8
+	BatchVertices.Add( FSlateVertex( FVector2D( RightMarginX, TopMarginY ),		FVector2D( RightMarginU, TopMarginV),		Zero,							Tint, InClippingRect ) ); //9
+	BatchVertices.Add( FSlateVertex( FVector2D( EndPos.X, Position.Y ),			FVector2D( EndUV.X, StartUV.Y ),			Zero,							Tint, InClippingRect ) ); //10
+	BatchVertices.Add( FSlateVertex( FVector2D( EndPos.X, TopMarginY ),			FVector2D( EndUV.X, TopMarginV),			Zero,							Tint, InClippingRect ) ); //11
+
+	BatchVertices.Add( FSlateVertex( FVector2D( Position.X, TopMarginY ),		FVector2D( StartUV.X, StartUV.Y ),			FVector2D(0.0f, LeftTiling),	Tint, InClippingRect ) ); //12
+	BatchVertices.Add( FSlateVertex( FVector2D( Position.X, BottomMarginY ),	FVector2D( StartUV.X, EndUV.Y ),			FVector2D(0.0f, LeftTiling),	Tint, InClippingRect ) ); //13
+	BatchVertices.Add( FSlateVertex( FVector2D( LeftMarginX, TopMarginY ),		FVector2D( LeftMarginU, StartUV.Y ),		FVector2D(0.0f, LeftTiling),	Tint, InClippingRect ) ); //14
+	BatchVertices.Add( FSlateVertex( FVector2D( LeftMarginX, BottomMarginY ),	FVector2D( LeftMarginU, EndUV.Y ),			FVector2D(0.0f, LeftTiling),	Tint, InClippingRect ) ); //15
+
+	BatchVertices.Add( FSlateVertex( FVector2D( RightMarginX, TopMarginY ),		FVector2D( RightMarginU, StartUV.Y ),		FVector2D(0.0f, RightTiling),	Tint, InClippingRect ) ); //16
+	BatchVertices.Add( FSlateVertex( FVector2D( RightMarginX, BottomMarginY ),	FVector2D( RightMarginU, EndUV.Y ),			FVector2D(0.0f, RightTiling),	Tint, InClippingRect ) ); //17
+	BatchVertices.Add( FSlateVertex( FVector2D( EndPos.X, TopMarginY ),			FVector2D( EndUV.X, StartUV.Y ),			FVector2D(0.0f, RightTiling),	Tint, InClippingRect ) ); //18
+	BatchVertices.Add( FSlateVertex( FVector2D( EndPos.X, BottomMarginY ),		FVector2D( EndUV.X, EndUV.Y ),				FVector2D(0.0f, RightTiling),	Tint, InClippingRect ) ); //19
+
+	BatchVertices.Add( FSlateVertex( FVector2D( Position.X, BottomMarginY ),	FVector2D( StartUV.X, BottomMarginV ),		Zero,							Tint, InClippingRect ) ); //20
+	BatchVertices.Add( FSlateVertex( FVector2D( Position.X, EndPos.Y ),			FVector2D( StartUV.X, EndUV.Y ),			Zero,							Tint, InClippingRect ) ); //21
+	BatchVertices.Add( FSlateVertex( FVector2D( LeftMarginX, BottomMarginY ),	FVector2D( LeftMarginU, BottomMarginV ),	Zero,							Tint, InClippingRect ) ); //22
+	BatchVertices.Add( FSlateVertex( FVector2D( LeftMarginX, EndPos.Y ),		FVector2D( LeftMarginU, EndUV.Y ),			Zero,							Tint, InClippingRect ) ); //23
+	
+	BatchVertices.Add( FSlateVertex( FVector2D( LeftMarginX, BottomMarginY ),	FVector2D( StartUV.X, BottomMarginV ),		FVector2D(BottomTiling, 0.0f),	Tint, InClippingRect ) ); //24
+	BatchVertices.Add( FSlateVertex( FVector2D( LeftMarginX, EndPos.Y ),		FVector2D( StartUV.X, EndUV.Y ),			FVector2D(BottomTiling, 0.0f),	Tint, InClippingRect ) ); //25
+	BatchVertices.Add( FSlateVertex( FVector2D( RightMarginX,BottomMarginY ),	FVector2D( EndUV.X, BottomMarginV ),		FVector2D(BottomTiling, 0.0f),	Tint, InClippingRect ) ); //26
+	BatchVertices.Add( FSlateVertex( FVector2D( RightMarginX, EndPos.Y ),		FVector2D( EndUV.X, EndUV.Y ),				FVector2D(BottomTiling, 0.0f),	Tint, InClippingRect ) ); //27
+
+	BatchVertices.Add( FSlateVertex( FVector2D( RightMarginX, BottomMarginY ),	FVector2D( RightMarginU, BottomMarginV ),	Zero,							Tint, InClippingRect ) ); //29
+	BatchVertices.Add( FSlateVertex( FVector2D( RightMarginX, EndPos.Y ),		FVector2D( RightMarginU, EndUV.Y ),			Zero,							Tint, InClippingRect ) ); //30
+	BatchVertices.Add( FSlateVertex( FVector2D( EndPos.X, BottomMarginY ),		FVector2D( EndUV.X, BottomMarginV ),		Zero,							Tint, InClippingRect ) ); //31
+	BatchVertices.Add( FSlateVertex( FVector2D( EndPos.X, EndPos.Y ),			FVector2D( EndUV.X, EndUV.Y ),				Zero,							Tint, InClippingRect ) ); //32
+
+
+	// The offset into the index buffer where this elements indices start
+	uint32 IndexOffsetStart = BatchIndices.Num();
+
+	// Top
+	BatchIndices.Add( IndexStart + 0 );
+	BatchIndices.Add( IndexStart + 1 );
+	BatchIndices.Add( IndexStart + 2 );
+	BatchIndices.Add( IndexStart + 2 );
+	BatchIndices.Add( IndexStart + 1 );
+	BatchIndices.Add( IndexStart + 3 );
+
+	BatchIndices.Add( IndexStart + 4 );
+	BatchIndices.Add( IndexStart + 5 );
+	BatchIndices.Add( IndexStart + 6 );
+	BatchIndices.Add( IndexStart + 6 );
+	BatchIndices.Add( IndexStart + 5 );
+	BatchIndices.Add( IndexStart + 7 );
+
+	BatchIndices.Add( IndexStart + 8 );
+	BatchIndices.Add( IndexStart + 9 );
+	BatchIndices.Add( IndexStart + 10 );
+	BatchIndices.Add( IndexStart + 10 );
+	BatchIndices.Add( IndexStart + 9 );
+	BatchIndices.Add( IndexStart + 11 );
+
+	// Middle
+	BatchIndices.Add( IndexStart + 12 );
+	BatchIndices.Add( IndexStart + 13 );
+	BatchIndices.Add( IndexStart + 14 );
+	BatchIndices.Add( IndexStart + 14 );
+	BatchIndices.Add( IndexStart + 13 );
+	BatchIndices.Add( IndexStart + 15 );
+
+	BatchIndices.Add( IndexStart + 16 );
+	BatchIndices.Add( IndexStart + 17 );
+	BatchIndices.Add( IndexStart + 18 );
+	BatchIndices.Add( IndexStart + 18 );
+	BatchIndices.Add( IndexStart + 17 );
+	BatchIndices.Add( IndexStart + 19 );
+
+	// Bottom
+	BatchIndices.Add( IndexStart + 20 );
+	BatchIndices.Add( IndexStart + 21 );
+	BatchIndices.Add( IndexStart + 22 );
+	BatchIndices.Add( IndexStart + 22 );
+	BatchIndices.Add( IndexStart + 21 );
+	BatchIndices.Add( IndexStart + 23 );
+
+	BatchIndices.Add( IndexStart + 24 );
+	BatchIndices.Add( IndexStart + 25 );
+	BatchIndices.Add( IndexStart + 26 );
+	BatchIndices.Add( IndexStart + 26 );
+	BatchIndices.Add( IndexStart + 25 );
+	BatchIndices.Add( IndexStart + 27 );
+
+	BatchIndices.Add( IndexStart + 28 );
+	BatchIndices.Add( IndexStart + 29 );
+	BatchIndices.Add( IndexStart + 30 );
+	BatchIndices.Add( IndexStart + 30 );
+	BatchIndices.Add( IndexStart + 29 );
+	BatchIndices.Add( IndexStart + 31 );
+}
+
+void FSlateElementBatcher::AddCustomElement( const FVector2D& Position, const FVector2D& Size, float Scale, const FSlateDataPayload& InPayload, const FSlateRect& InClippingRect, ESlateDrawEffect::Type DrawEffects, uint32 Layer )
+{
+	if( InPayload.CustomDrawer.IsValid() )
+	{
+		// See if the layer already exists.
+		TSet<FSlateElementBatch>* ElementBatches = LayerToElementBatches.Find( Layer );
+		if( !ElementBatches )
+		{
+			// The layer doesn't exist so make it now
+			ElementBatches = &LayerToElementBatches.Add( Layer, TSet<FSlateElementBatch>() );
+		}
+		check( ElementBatches );
+
+
+		// Custom elements are not batched together 
+		ElementBatches->Add( FSlateElementBatch( InPayload.CustomDrawer ) );
+	}
+}
+
+FSlateElementBatch& FSlateElementBatcher::FindBatchForElement( 
+	uint32 Layer, 
+	const FShaderParams& ShaderParams, 
+	const FSlateShaderResource* InTexture, 
+	ESlateDrawPrimitive::Type PrimitiveType,
+	ESlateShader::Type ShaderType, 
+	ESlateDrawEffect::Type DrawEffects, 
+	ESlateBatchDrawFlag::Type DrawFlags)
+{
+	SCOPE_CYCLE_COUNTER(STAT_SlateFindBatchTime);
+
+	// See if the layer already exists.
+	TSet<FSlateElementBatch>* ElementBatches = LayerToElementBatches.Find( Layer );
+	if( !ElementBatches )
+	{
+		// The layer doesn't exist so make it now
+		ElementBatches = &LayerToElementBatches.Add( Layer, TSet<FSlateElementBatch>() );
+	}
+	check( ElementBatches );
+
+	// Create a temp batch so we can use it as our key to find if the same batch already exists
+	FSlateElementBatch TempBatch( InTexture, ShaderParams, ShaderType, PrimitiveType, DrawEffects, DrawFlags );
+
+	FSlateElementBatch* ElementBatch = ElementBatches->Find( TempBatch );
+	if( !ElementBatch )
+	{
+		// No batch with the specified parameter exists.  Create it from the temp batch.
+		FSetElementId ID = ElementBatches->Add( TempBatch );
+		ElementBatch = &(*ElementBatches)[ID];
+
+		// Get a free vertex array
+		if( VertexArrayFreeList.Num() > 0 )
+		{
+			ElementBatch->VertexArrayIndex = VertexArrayFreeList.Pop();
+			BatchVertexArrays[ElementBatch->VertexArrayIndex].Reserve(200);		
+		}
+		else
+		{
+			// There are no free vertex arrays so we must add one		
+			uint32 NewIndex = BatchVertexArrays.Add( TArray<FSlateVertex>() );
+			BatchVertexArrays[NewIndex].Reserve(200);		
+
+			ElementBatch->VertexArrayIndex = NewIndex;
+		}
+
+		// Get a free index array
+		if( IndexArrayFreeList.Num() > 0 )
+		{
+			ElementBatch->IndexArrayIndex = IndexArrayFreeList.Pop();
+			BatchIndexArrays[ElementBatch->IndexArrayIndex].Reserve(200);		
+		}
+		else
+		{
+			// There are no free index arrays so we must add one
+			uint32 NewIndex = BatchIndexArrays.Add( TArray<SlateIndex>() );
+			BatchIndexArrays[NewIndex].Reserve(500);	
+
+			ElementBatch->IndexArrayIndex = NewIndex;
+
+			check( BatchIndexArrays.IsValidIndex( ElementBatch->IndexArrayIndex ) );
+		}
+		
+	}
+	check( ElementBatch );
+
+	// Increment the number of elements in the batch.
+	++ElementBatch->NumElementsInBatch;
+	return *ElementBatch;
+}
+
+void FSlateElementBatcher::AddBasicVertices( TArray<FSlateVertex>& OutVertices, FSlateElementBatch& ElementBatch, const TArray<FSlateVertex>& VertexBatch )
+{
+	uint32 FirstIndex = OutVertices.Num();
+	OutVertices.AddUninitialized( VertexBatch.Num() );
+
+	uint32 RequiredSize =  VertexBatch.Num() * VertexBatch.GetTypeSize();
+
+	FMemory::Memcpy( &OutVertices[FirstIndex], VertexBatch.GetData(), RequiredSize );
+
+	RequiredVertexMemory += RequiredSize;
+	TotalVertexMemory += VertexBatch.GetAllocatedSize();
+	NumVertices += VertexBatch.Num();
+
+	ElementBatch.VertexOffset = FirstIndex;
+	ElementBatch.NumVertices = VertexBatch.Num();
+}
+
+void FSlateElementBatcher::AddIndices( TArray<SlateIndex>& OutIndices, FSlateElementBatch& ElementBatch, const TArray<SlateIndex>& IndexBatch )
+{
+	uint32 FirstIndex = OutIndices.Num();
+	OutIndices.AddUninitialized( IndexBatch.Num() );
+
+	uint32 RequiredSize =  IndexBatch.Num() * IndexBatch.GetTypeSize();
+
+	FMemory::Memcpy( &OutIndices[FirstIndex], IndexBatch.GetData(), RequiredSize );
+
+	RequiredIndexMemory += RequiredSize;
+	TotalIndexMemory += IndexBatch.GetAllocatedSize();
+
+	ElementBatch.IndexOffset = FirstIndex;
+	ElementBatch.NumIndices = IndexBatch.Num();
+}
+
+/**
+ * Updates rendering buffers
+ */
+void FSlateElementBatcher::FillBatchBuffers( FSlateWindowElementList& WindowElementList, bool& bRequiresStencilTest )
+{
+	SCOPE_CYCLE_COUNTER( STAT_SlateUpdateBufferGTTime );
+
+	TArray<FSlateRenderBatch>& OutRenderBatches = WindowElementList.RenderBatches;
+	TArray<FSlateVertex>& OutBatchedVertices = WindowElementList.BatchedVertices;
+	TArray<SlateIndex>& OutBatchedIndices = WindowElementList.BatchedIndices;
+
+	check( OutRenderBatches.Num() == 0 && OutBatchedVertices.Num() == 0 && OutBatchedIndices.Num() == 0 );
+
+	// Sort by layer
+	LayerToElementBatches.KeySort( TLess<uint32>() );
+
+
+#if STATS
+	NumLayers = FMath::Max<uint32>(NumLayers, LayerToElementBatches.Num() );
+#endif
+
+	bRequiresStencilTest = false;
+	// For each element batch add its vertices and indices to the bulk lists.
+	for( TMap< uint32, TSet<FSlateElementBatch> >::TIterator It( LayerToElementBatches ); It; ++It )
+	{
+		TSet<FSlateElementBatch>& ElementBatches = It.Value();
+
+#if STATS
+		NumBatches += ElementBatches.Num();
+#endif
+		for( TSet<FSlateElementBatch>::TIterator BatchIt(ElementBatches); BatchIt; ++BatchIt )
+		{
+			FSlateElementBatch& ElementBatch = *BatchIt;
+
+			bool bIsMaterial = ElementBatch.GetShaderResource() && ElementBatch.GetShaderResource()->GetType() == ESlateShaderResource::Material;
+
+			if( !ElementBatch.GetCustomDrawer().IsValid() )
+			{
+				if( ElementBatch.GetShaderType() == ESlateShader::LineSegment )
+				{
+					bRequiresStencilTest = true;
+				}
+
+				TArray<FSlateVertex>& BatchVertices = BatchVertexArrays[ ElementBatch.VertexArrayIndex ];
+				VertexArrayFreeList.Add( ElementBatch.VertexArrayIndex );
+
+				TArray<SlateIndex>& BatchIndices = BatchIndexArrays[ ElementBatch.IndexArrayIndex ];
+				IndexArrayFreeList.Add( ElementBatch.IndexArrayIndex );
+
+				// We should have at least some vertices and indices in the batch or none at all
+				check( BatchVertices.Num() > 0 && BatchIndices.Num() > 0  || BatchVertices.Num() == 0 && BatchIndices.Num() == 0 );
+
+				if( BatchVertices.Num() > 0 && BatchIndices.Num() > 0  )
+				{
+					
+					AddBasicVertices( OutBatchedVertices, ElementBatch, BatchVertices );
+					AddIndices( OutBatchedIndices, ElementBatch, BatchIndices );
+
+					OutRenderBatches.Add( FSlateRenderBatch( ElementBatch ) );
+
+					// Done with the batch
+					BatchVertices.Empty();
+					BatchIndices.Empty();
+				}
+
+			}
+			else
+			{
+				OutRenderBatches.Add( FSlateRenderBatch( ElementBatch ) );
+			}
+		}
+	}
+
+}
+
+/** 
+ * Resets all stored data accumulated during the batching process
+ */
+void FSlateElementBatcher::ResetBatches()
+{
+	LayerToElementBatches.Reset();
+	bRequiresVsync = false;
+}
+
+
+void FSlateElementBatcher::ResetStats()
+{
+	SET_DWORD_STAT( STAT_SlateNumLayers, NumLayers );
+	SET_DWORD_STAT( STAT_SlateNumBatches, NumBatches );
+	SET_DWORD_STAT( STAT_SlateVertexCount, NumVertices );
+	SET_MEMORY_STAT( STAT_SlateVertexBatchMemory, TotalVertexMemory );
+	SET_MEMORY_STAT( STAT_SlateIndexBatchMemory, TotalIndexMemory );
+
+	NumLayers = 0;
+	NumBatches = 0;
+	NumVertices = 0;
+	RequiredIndexMemory = 0;
+	RequiredVertexMemory = 0;
+	TotalVertexMemory = 0;
+	TotalIndexMemory = 0;
+}
+
diff --git a/Engine/Source/Runtime/SlateCore/Private/Rendering/SlateDrawBuffer.cpp b/Engine/Source/Runtime/SlateCore/Private/Rendering/SlateDrawBuffer.cpp
new file mode 100644
--- /dev/null
+++ b/Engine/Source/Runtime/SlateCore/Private/Rendering/SlateDrawBuffer.cpp
@@ -0,0 +1,64 @@
+// Copyright 1998-2014 Epic Games, Inc. All Rights Reserved.
+
+/*=============================================================================
+	SlateDrawBuffer.cpp: Implements the FSlateDrawBuffer class.
+=============================================================================*/
+
+#include "SlateCorePrivatePCH.h"
+
+
+/* Local constants
+ *****************************************************************************/
+
+uint32 GDrawBuffers = 0;
+uint32 MaxDrawBuffers = 0;
+
+
+/* FSlateDrawBuffer structors
+ *****************************************************************************/
+
+FSlateDrawBuffer::FSlateDrawBuffer( ) 
+	: Locked(0)
+{
+	++GDrawBuffers;
+
+	if (GDrawBuffers > MaxDrawBuffers)
+	{
+		MaxDrawBuffers = GDrawBuffers;
+	}
+}
+
+
+FSlateDrawBuffer::~FSlateDrawBuffer( )
+{
+	--GDrawBuffers;
+}
+
+
+/* FSlateDrawBuffer interface
+ *****************************************************************************/
+
+void FSlateDrawBuffer::ClearBuffer( )
+{
+	WindowElementLists.Empty();
+}
+
+
+FSlateWindowElementList& FSlateDrawBuffer::AddWindowElementList( TSharedRef<SWindow> ForWindow )
+{
+	int32 Index = WindowElementLists.Add(FSlateWindowElementList(ForWindow));
+
+	return WindowElementLists[Index];
+}
+
+
+bool FSlateDrawBuffer::Lock( )
+{
+	return FPlatformAtomics::InterlockedCompareExchange(&Locked, 1, 0) == 0;
+}
+
+
+void FSlateDrawBuffer::Unlock( )
+{
+	FPlatformAtomics::InterlockedExchange(&Locked, 0);
+}
\ No newline at end of file
diff --git a/Engine/Source/Runtime/SlateCore/Private/Rendering/SlateRenderer.cpp b/Engine/Source/Runtime/SlateCore/Private/Rendering/SlateRenderer.cpp
new file mode 100644
--- /dev/null
+++ b/Engine/Source/Runtime/SlateCore/Private/Rendering/SlateRenderer.cpp
@@ -0,0 +1,62 @@
+// Copyright 1998-2014 Epic Games, Inc. All Rights Reserved.
+
+/*=============================================================================
+	Renderer.cpp: Implements the SNullWidget class.
+=============================================================================*/
+
+#include "SlateCorePrivatePCH.h"
+
+
+/* FSlateRenderer structors
+ *****************************************************************************/
+
+FSlateRenderer::~FSlateRenderer()
+{
+	// debug check to make sure we are properly cleaning up draw buffers
+	checkSlow( GDrawBuffers == 0 );
+}
+
+
+/* FSlateRenderer interface
+ *****************************************************************************/
+
+void FSlateRenderer::FlushFontCache( )
+{
+	FontCache->FlushCache();
+	FontMeasure->FlushCache();
+}
+
+
+bool FSlateRenderer::IsViewportFullscreen( const SWindow& Window ) const
+{
+	check( IsThreadSafeForSlateRendering() );
+
+	bool bFullscreen = false;
+
+	if (FPlatformProperties::SupportsWindowedMode())
+	{
+		if( GIsEditor)
+		{
+			bFullscreen = false;
+		}
+		else
+		{
+			bFullscreen = Window.GetWindowMode() == EWindowMode::Fullscreen;
+		}
+	}
+	else
+	{
+		bFullscreen = true;
+	}
+
+	return bFullscreen;
+}
+
+
+/* Global functions
+ *****************************************************************************/
+
+bool IsThreadSafeForSlateRendering( )
+{
+	return ((GSlateLoadingThreadId != 0) || IsInGameThread());
+}
diff --git a/Engine/Source/Runtime/SlateCore/Private/SlateCoreClasses.cpp b/Engine/Source/Runtime/SlateCore/Private/SlateCoreClasses.cpp
new file mode 100644
--- /dev/null
+++ b/Engine/Source/Runtime/SlateCore/Private/SlateCoreClasses.cpp
@@ -0,0 +1,11 @@
+// Copyright 1998-2014 Epic Games, Inc. All Rights Reserved.
+
+/*=============================================================================
+	SlateCoreClasses.cpp: Implements the module's script classes.
+=============================================================================*/
+
+#include "SlateCorePrivatePCH.h"
+
+
+FSlateWidgetStyle::FSlateWidgetStyle( )
+{ }
diff --git a/Engine/Source/Runtime/SlateCore/Private/SlateCoreModule.cpp b/Engine/Source/Runtime/SlateCore/Private/SlateCoreModule.cpp
new file mode 100644
--- /dev/null
+++ b/Engine/Source/Runtime/SlateCore/Private/SlateCoreModule.cpp
@@ -0,0 +1,64 @@
+// Copyright 1998-2014 Epic Games, Inc. All Rights Reserved.
+
+/*=============================================================================
+	SlateCoreModule.cpp: Implements the FSlateCoreModule class.
+=============================================================================*/
+
+#include "SlateCorePrivatePCH.h"
+#include "ModuleManager.h"
+
+
+DEFINE_LOG_CATEGORY(LogSlate);
+DEFINE_LOG_CATEGORY(LogSlateStyles);
+
+DEFINE_STAT(STAT_SlateMeasureStringTime);
+DEFINE_STAT(STAT_SlateRenderingRTTime);
+DEFINE_STAT(STAT_SlateUpdateBufferRTTime);
+DEFINE_STAT(STAT_SlateDrawTime);
+DEFINE_STAT(STAT_SlateUpdateBufferGTTime);
+DEFINE_STAT(STAT_SlateBuildLines);
+DEFINE_STAT(STAT_SlateBuildBoxes);
+DEFINE_STAT(STAT_SlateBuildBorders);
+DEFINE_STAT(STAT_SlateBuildText);
+DEFINE_STAT(STAT_SlateBuildGradients);
+DEFINE_STAT(STAT_SlateBuildSplines);
+DEFINE_STAT(STAT_SlateBuildElements);
+DEFINE_STAT(STAT_SlateFindBatchTime);
+DEFINE_STAT(STAT_SlateFontCachingTime);
+DEFINE_STAT(STAT_SlateRenderingGTTime);
+DEFINE_STAT(STAT_SlateBuildViewports);
+
+DEFINE_STAT(STAT_SlateNumTextureAtlases);
+DEFINE_STAT(STAT_SlateNumFontAtlases);
+DEFINE_STAT(STAT_SlateNumNonAtlasedTextures);
+DEFINE_STAT(STAT_SlateNumDynamicTextures);
+
+DEFINE_STAT(STAT_SlateNumLayers);
+DEFINE_STAT(STAT_SlateNumBatches);
+DEFINE_STAT(STAT_SlateVertexCount);
+
+DEFINE_STAT(STAT_SlateTextureDataMemory);
+DEFINE_STAT(STAT_SlateTextureAtlasMemory);
+DEFINE_STAT(STAT_SlateFontKerningTableMemory);
+DEFINE_STAT(STAT_SlateFontMeasureCacheMemory);
+DEFINE_STAT(STAT_SlateVertexBatchMemory);
+DEFINE_STAT(STAT_SlateIndexBatchMemory);
+DEFINE_STAT(STAT_SlateVertexBufferMemory);
+DEFINE_STAT(STAT_SlateIndexBufferMemory);
+
+
+const float FOptionalSize::Unspecified = -1.0f;
+
+
+/**
+ * Implements the SlateCore module.
+ */
+class FSlateCoreModule
+	: public IModuleInterface
+{
+public:
+
+};
+
+
+IMPLEMENT_MODULE(FSlateCoreModule, SlateCore);
diff --git a/Engine/Source/Runtime/SlateCore/Private/SlateCorePrivatePCH.h b/Engine/Source/Runtime/SlateCore/Private/SlateCorePrivatePCH.h
new file mode 100644
--- /dev/null
+++ b/Engine/Source/Runtime/SlateCore/Private/SlateCorePrivatePCH.h
@@ -0,0 +1,23 @@
+// Copyright 1998-2014 Epic Games, Inc. All Rights Reserved.
+
+/*=============================================================================
+	SlateCorePrivatePCH.h: Pre-compiled header file for the SlateCore module.
+=============================================================================*/
+
+#pragma once
+
+#include "SlateCore.h"
+
+
+/* Dependencies
+ *****************************************************************************/
+
+#include "Json.h"
+#include "PropertyTag.h"
+#include "RemoteConfigIni.h"
+
+
+/* Private includes
+ *****************************************************************************/
+
+#include "FontCacheUtils.h"
diff --git a/Engine/Source/Runtime/SlateCore/Private/Styling/SlateBrush.cpp b/Engine/Source/Runtime/SlateCore/Private/Styling/SlateBrush.cpp
new file mode 100644
--- /dev/null
+++ b/Engine/Source/Runtime/SlateCore/Private/Styling/SlateBrush.cpp
@@ -0,0 +1,82 @@
+// Copyright 1998-2014 Epic Games, Inc. All Rights Reserved.
+
+/*=============================================================================
+	SlateBrush.cpp: Implements the FSlateBrush structure.
+=============================================================================*/
+
+#include "SlateCorePrivatePCH.h"
+
+
+/* FSlateBrush structors
+ *****************************************************************************/
+
+FSlateBrush::FSlateBrush( ESlateBrushDrawType::Type InDrawType, const FName InResourceName, const FMargin& InMargin, ESlateBrushTileType::Type InTiling, ESlateBrushImageType::Type InImageType, const FVector2D& InImageSize, const FLinearColor& InTint, UObject* InObjectResource, bool bInDynamicallyLoaded )
+	: ImageSize( InImageSize )
+	, DrawAs( InDrawType )
+	, Margin( InMargin )
+	, TintColor( InTint )
+	, Tiling( InTiling )
+	, ImageType( InImageType )
+	, ResourceObject( InObjectResource )
+	, ResourceName( InResourceName )
+	, bIsDynamicallyLoaded( bInDynamicallyLoaded )
+{
+	bHasUObject_DEPRECATED = (InObjectResource != nullptr) || InResourceName.ToString().StartsWith(FSlateBrush::UTextureIdentifier());
+
+	//Useful for debugging style breakages
+	//if ( !bHasUObject_DEPRECATED && InResourceName.IsValid() && InResourceName != NAME_None )
+	//{
+	//	checkf( FPaths::FileExists( InResourceName.ToString() ), *FPaths::ConvertRelativePathToFull( InResourceName.ToString() ) );
+	//}
+}
+
+
+FSlateBrush::FSlateBrush( ESlateBrushDrawType::Type InDrawType, const FName InResourceName, const FMargin& InMargin, ESlateBrushTileType::Type InTiling, ESlateBrushImageType::Type InImageType, const FVector2D& InImageSize, const TSharedRef< FLinearColor >& InTint, UObject* InObjectResource, bool bInDynamicallyLoaded )
+	: ImageSize( InImageSize )
+	, DrawAs( InDrawType )
+	, Margin( InMargin )
+	, TintColor( InTint )
+	, Tiling( InTiling )
+	, ImageType( InImageType )
+	, ResourceObject( InObjectResource )
+	, ResourceName( InResourceName )
+	, bIsDynamicallyLoaded( bInDynamicallyLoaded )
+{
+	bHasUObject_DEPRECATED = (InObjectResource != nullptr) || InResourceName.ToString().StartsWith(FSlateBrush::UTextureIdentifier());
+
+	//Useful for debugging style breakages
+	//if ( !bHasUObject_DEPRECATED && InResourceName.IsValid() && InResourceName != NAME_None )
+	//{
+	//	checkf( FPaths::FileExists( InResourceName.ToString() ), *FPaths::ConvertRelativePathToFull( InResourceName.ToString() ) );
+	//}
+}
+
+
+FSlateBrush::FSlateBrush( ESlateBrushDrawType::Type InDrawType, const FName InResourceName, const FMargin& InMargin, ESlateBrushTileType::Type InTiling, ESlateBrushImageType::Type InImageType, const FVector2D& InImageSize, const FSlateColor& InTint, UObject* InObjectResource, bool bInDynamicallyLoaded )
+	: ImageSize(InImageSize)
+	, DrawAs(InDrawType)
+	, Margin(InMargin)
+	, TintColor(InTint)
+	, Tiling(InTiling)
+	, ImageType(InImageType)
+	, ResourceObject(InObjectResource)
+	, ResourceName(InResourceName)
+	, bIsDynamicallyLoaded(bInDynamicallyLoaded)
+{
+	bHasUObject_DEPRECATED = (InObjectResource != nullptr) || InResourceName.ToString().StartsWith(FSlateBrush::UTextureIdentifier());
+
+	//Useful for debugging style breakages
+	//if ( !bHasUObject_DEPRECATED && InResourceName.IsValid() && InResourceName != NAME_None )
+	//{
+	//	checkf( FPaths::FileExists( InResourceName.ToString() ), *FPaths::ConvertRelativePathToFull( InResourceName.ToString() ) );
+	//}
+}
+
+
+/* FSlateBrush static functions
+ *****************************************************************************/
+
+const FString FSlateBrush::UTextureIdentifier( )
+{
+	return FString(TEXT("texture:/"));
+}
diff --git a/Engine/Source/Runtime/SlateCore/Private/Styling/SlateStyleRegistry.cpp b/Engine/Source/Runtime/SlateCore/Private/Styling/SlateStyleRegistry.cpp
new file mode 100644
--- /dev/null
+++ b/Engine/Source/Runtime/SlateCore/Private/Styling/SlateStyleRegistry.cpp
@@ -0,0 +1,64 @@
+// Copyright 1998-2014 Epic Games, Inc. All Rights Reserved.
+
+/*=============================================================================
+	SlateStyleRegistry.cpp: Implements the FSlateStyleRegistry class.
+=============================================================================*/
+
+#include "SlateCorePrivatePCH.h"
+
+
+/* Static initialization
+ *****************************************************************************/
+
+TMap<FName, const ISlateStyle*> FSlateStyleRegistry::SlateStyleRepository;
+
+
+/* Static functions
+ *****************************************************************************/
+
+void FSlateStyleRegistry::RegisterSlateStyle( const ISlateStyle& InSlateStyle )
+{
+	const FName& SlateStyleName = InSlateStyle.GetStyleSetName();
+	check( SlateStyleName.IsValid() );
+	check( !FindSlateStyle( SlateStyleName ) );
+
+	SlateStyleRepository.Add( SlateStyleName, &InSlateStyle );
+
+	if ( FSlateApplicationBase::IsInitialized() && FSlateApplicationBase::Get().GetRenderer().IsValid() )
+	{
+		FSlateApplicationBase::Get().GetRenderer()->LoadStyleResources( InSlateStyle );
+	}
+}
+
+
+void FSlateStyleRegistry::UnRegisterSlateStyle( const ISlateStyle& InSlateStyle )
+{
+	const FName& SlateStyleName = InSlateStyle.GetStyleSetName();
+	check( SlateStyleName.IsValid() );
+
+	SlateStyleRepository.Remove( SlateStyleName );
+}
+
+
+void FSlateStyleRegistry::UnRegisterSlateStyle(const FName StyleSetName)
+{
+	check(StyleSetName.IsValid());
+	SlateStyleRepository.Remove(StyleSetName);
+}
+
+
+const ISlateStyle* FSlateStyleRegistry::FindSlateStyle( const FName& InSlateStyleName )
+{
+	const ISlateStyle** SlateStylePtr = SlateStyleRepository.Find( InSlateStyleName );
+	return (SlateStylePtr) ? *SlateStylePtr : nullptr;
+}
+
+
+void FSlateStyleRegistry::GetAllResources( TArray< const FSlateBrush* >& OutResources )
+{
+	/* Iterate the style chunks and collect their resources */
+	for( auto It = SlateStyleRepository.CreateConstIterator(); It; ++It )
+	{
+		It->Value->GetResources( OutResources );
+	}
+}
diff --git a/Engine/Source/Runtime/SlateCore/Private/Styling/SlateStyleSet.cpp b/Engine/Source/Runtime/SlateCore/Private/Styling/SlateStyleSet.cpp
new file mode 100644
--- /dev/null
+++ b/Engine/Source/Runtime/SlateCore/Private/Styling/SlateStyleSet.cpp
@@ -0,0 +1,410 @@
+// Copyright 1998-2014 Epic Games, Inc. All Rights Reserved.
+
+#include "SlateCorePrivatePCH.h"
+
+
+FSlateStyleSet::FSlateStyleSet(const FName& InStyleSetName)
+: StyleSetName(InStyleSetName)
+, DefaultBrush(new FSlateImageBrush(FPaths::EngineContentDir() / TEXT("Slate/Checkerboard.png"), FVector2D(16, 16), FLinearColor::White, ESlateBrushTileType::Both))
+{
+	// Add a mapping so that this resource will be discovered by GetStyleResources.
+	Set(TEXT("Default"), GetDefaultBrush());
+}
+
+FSlateStyleSet::~FSlateStyleSet()
+{
+	// Delete all allocated brush resources.
+	for ( TMap< FName, FSlateBrush* >::TIterator It(BrushResources); It; ++It )
+	{
+		delete It.Value();
+	}
+}
+
+
+const FName& FSlateStyleSet::GetStyleSetName() const
+{
+	return StyleSetName;
+}
+
+void FSlateStyleSet::GetResources(TArray< const FSlateBrush* >& OutResources) const
+{
+	// Collection for this style's brush resources.
+	TArray< const FSlateBrush* > SlateBrushResources;
+	for ( TMap< FName, FSlateBrush* >::TConstIterator It(BrushResources); It; ++It )
+	{
+		SlateBrushResources.Add(It.Value());
+	}
+
+	//Gather resources from our definitions
+	for ( TMap< FName, TSharedRef< struct FSlateWidgetStyle > >::TConstIterator It(WidgetStyleValues); It; ++It )
+	{
+		It.Value()->GetResources(SlateBrushResources);
+	}
+
+	// Append this style's resources to OutResources.
+	OutResources.Append(SlateBrushResources);
+}
+
+void FSlateStyleSet::SetContentRoot(const FString& InContentRootDir)
+{
+	ContentRootDir = InContentRootDir;
+}
+
+FString FSlateStyleSet::RootToContentDir(const ANSICHAR* RelativePath, const TCHAR* Extension)
+{
+	return ( ContentRootDir / RelativePath ) + Extension;
+}
+
+FString FSlateStyleSet::RootToContentDir(const WIDECHAR* RelativePath, const TCHAR* Extension)
+{
+	return ( ContentRootDir / RelativePath ) + Extension;
+}
+
+FString FSlateStyleSet::RootToContentDir(const FString& RelativePath, const TCHAR* Extension)
+{
+	return ( ContentRootDir / RelativePath ) + Extension;
+}
+
+FString FSlateStyleSet::RootToContentDir(const ANSICHAR* RelativePath)
+{
+	return ( ContentRootDir / RelativePath );
+}
+
+FString FSlateStyleSet::RootToContentDir(const WIDECHAR* RelativePath)
+{
+	return ( ContentRootDir / RelativePath );
+}
+
+FString FSlateStyleSet::RootToContentDir(const FString& RelativePath)
+{
+	return ( ContentRootDir / RelativePath );
+}
+
+void FSlateStyleSet::SetCoreContentRoot(const FString& InCoreContentRootDir)
+{
+	CoreContentRootDir = InCoreContentRootDir;
+}
+
+FString FSlateStyleSet::RootToCoreContentDir(const ANSICHAR* RelativePath, const TCHAR* Extension)
+{
+	return ( CoreContentRootDir / RelativePath ) + Extension;
+}
+
+FString FSlateStyleSet::RootToCoreContentDir(const WIDECHAR* RelativePath, const TCHAR* Extension)
+{
+	return ( CoreContentRootDir / RelativePath ) + Extension;
+}
+
+FString FSlateStyleSet::RootToCoreContentDir(const FString& RelativePath, const TCHAR* Extension)
+{
+	return ( CoreContentRootDir / RelativePath ) + Extension;
+}
+
+FString FSlateStyleSet::RootToCoreContentDir(const ANSICHAR* RelativePath)
+{
+	return ( CoreContentRootDir / RelativePath );
+}
+
+FString FSlateStyleSet::RootToCoreContentDir(const WIDECHAR* RelativePath)
+{
+	return ( CoreContentRootDir / RelativePath );
+}
+
+FString FSlateStyleSet::RootToCoreContentDir(const FString& RelativePath)
+{
+	return ( CoreContentRootDir / RelativePath );
+}
+
+float FSlateStyleSet::GetFloat(const FName PropertyName, const ANSICHAR* Specifier) const
+{
+	const float* Result = FloatValues.Find(Join(PropertyName, Specifier));
+
+	return Result ? *Result : FStyleDefaults::GetFloat();
+}
+
+FVector2D FSlateStyleSet::GetVector(const FName PropertyName, const ANSICHAR* Specifier) const
+{
+	const FVector2D* const Result = Vector2DValues.Find(Join(PropertyName, Specifier));
+
+	return Result ? *Result : FStyleDefaults::GetVector2D();
+}
+
+const FLinearColor& FSlateStyleSet::GetColor(const FName PropertyName, const ANSICHAR* Specifier) const
+{
+	const FName LookupName(Join(PropertyName, Specifier));
+	const FLinearColor* Result = ColorValues.Find(LookupName);
+
+#if DO_GUARD_SLOW
+	if ( Result == nullptr && !MissingResources.Contains(LookupName) )
+	{
+		MissingResources.Add(LookupName);
+		Log(ISlateStyle::Warning, FText::Format(NSLOCTEXT("SlateStyleSet", "UknownColor", "Unable to find Color '{0}'."), FText::FromName(LookupName)));
+	}
+#endif
+
+	return Result ? *Result : FStyleDefaults::GetColor();
+}
+
+const FSlateColor FSlateStyleSet::GetSlateColor(const FName PropertyName, const ANSICHAR* Specifier) const
+{
+	static FSlateColor UseForegroundStatic = FSlateColor::UseForeground();
+
+	const FName StyleName(Join(PropertyName, Specifier));
+	const FSlateColor* Result = SlateColorValues.Find(StyleName);
+
+	if ( Result == nullptr )
+	{
+		const FLinearColor* LinearColorLookup = ColorValues.Find(StyleName);
+		if ( LinearColorLookup == nullptr )
+		{
+			return UseForegroundStatic;
+		}
+
+		return *LinearColorLookup;
+	}
+
+#if DO_GUARD_SLOW
+	if ( Result == nullptr && !MissingResources.Contains(StyleName) )
+	{
+		MissingResources.Add(StyleName);
+		Log(ISlateStyle::Warning, FText::Format(NSLOCTEXT("SlateStyleSet", "UknownSlateColor", "Unable to find SlateColor '{0}'."), FText::FromName(StyleName)));
+	}
+#endif
+
+	return *Result;
+}
+
+const FMargin& FSlateStyleSet::GetMargin(const FName PropertyName, const ANSICHAR* Specifier) const
+{
+	const FName StyleName(Join(PropertyName, Specifier));
+	const FMargin* const Result = MarginValues.Find(StyleName);
+
+#if DO_GUARD_SLOW
+	if ( Result == nullptr && !MissingResources.Contains(StyleName) )
+	{
+		MissingResources.Add(StyleName);
+		Log(ISlateStyle::Warning, FText::Format(NSLOCTEXT("SlateStyleSet", "UknownMargin", "Unable to find Margin '{0}'."), FText::FromName(StyleName)));
+	}
+#endif
+
+	return Result ? *Result : FStyleDefaults::GetMargin();
+}
+
+const FSlateBrush* FSlateStyleSet::GetBrush(const FName PropertyName, const ANSICHAR* Specifier) const
+{
+	const FName StyleName = Join(PropertyName, Specifier);
+	const FSlateBrush* Result = BrushResources.FindRef(StyleName);
+
+	if ( Result == nullptr )
+	{
+		TWeakPtr< FSlateDynamicImageBrush > WeakImageBrush = DynamicBrushes.FindRef(StyleName);
+		TSharedPtr< FSlateDynamicImageBrush > ImageBrush = WeakImageBrush.Pin();
+
+		if ( ImageBrush.IsValid() )
+		{
+			Result = ImageBrush.Get();
+		}
+	}
+
+#if DO_GUARD_SLOW
+	if ( Result == nullptr && !MissingResources.Contains(StyleName))
+	{
+		MissingResources.Add(StyleName);
+		Log(ISlateStyle::Warning, FText::Format(NSLOCTEXT("SlateStyleSet", "UknownBrush", "Unable to find Brush '{0}'."), FText::FromName(StyleName)));
+	}
+#endif
+
+	return Result ? Result : GetDefaultBrush();
+}
+
+const FSlateBrush* FSlateStyleSet::GetOptionalBrush(const FName PropertyName, const ANSICHAR* Specifier, const FSlateBrush* const InDefaultBrush) const
+{
+	const FName StyleName = Join(PropertyName, Specifier);
+	const FSlateBrush* Result = BrushResources.FindRef(StyleName);
+
+	if ( Result == nullptr )
+	{
+		TWeakPtr< FSlateDynamicImageBrush > WeakImageBrush = DynamicBrushes.FindRef(StyleName);
+		TSharedPtr< FSlateDynamicImageBrush > ImageBrush = WeakImageBrush.Pin();
+
+		if ( ImageBrush.IsValid() )
+		{
+			Result = ImageBrush.Get();
+		}
+	}
+
+	return Result ? Result : InDefaultBrush;
+}
+
+const TSharedPtr< FSlateDynamicImageBrush > FSlateStyleSet::GetDynamicImageBrush(const FName BrushTemplate, const FName TextureName, const ANSICHAR* Specifier)
+{
+	return GetDynamicImageBrush(BrushTemplate, Specifier, nullptr, TextureName);
+}
+
+const TSharedPtr< FSlateDynamicImageBrush > FSlateStyleSet::GetDynamicImageBrush(const FName BrushTemplate, const ANSICHAR* Specifier, UTexture2D* TextureResource, const FName TextureName)
+{
+	return GetDynamicImageBrush(Join(BrushTemplate, Specifier), TextureResource, TextureName);
+}
+
+const TSharedPtr< FSlateDynamicImageBrush > FSlateStyleSet::GetDynamicImageBrush(const FName BrushTemplate, UTexture2D* TextureResource, const FName TextureName)
+{
+	//create a resource name
+	FName ResourceName;
+	ResourceName = TextureName == NAME_None ? BrushTemplate : FName(*( BrushTemplate.ToString() + TextureName.ToString() ));
+
+	//see if we already have that brush
+	TWeakPtr< FSlateDynamicImageBrush > WeakImageBrush = DynamicBrushes.FindRef(ResourceName);
+
+	//if we don't have the image brush, then make it
+	TSharedPtr< FSlateDynamicImageBrush > ReturnBrush = WeakImageBrush.Pin();
+
+	if ( !ReturnBrush.IsValid() )
+	{
+		const FSlateBrush* Result = BrushResources.FindRef(Join(BrushTemplate, nullptr));
+
+		if ( Result == nullptr )
+		{
+			Result = GetDefaultBrush();
+		}
+
+		//create the new brush
+		ReturnBrush = MakeShareable(new FSlateDynamicImageBrush(TextureResource, Result->ImageSize, ResourceName));
+
+		//add it to the dynamic brush list
+		DynamicBrushes.Add(ResourceName, ReturnBrush);
+	}
+
+	return ReturnBrush;
+}
+
+FSlateBrush* FSlateStyleSet::GetDefaultBrush() const
+{
+	return DefaultBrush;
+}
+
+const FSlateSound& FSlateStyleSet::GetSound(const FName PropertyName, const ANSICHAR* Specifier) const
+{
+	const FName StyleName = Join(PropertyName, Specifier);
+	const FSlateSound* Result = Sounds.Find(StyleName);
+
+#if DO_GUARD_SLOW
+	if ( Result == nullptr && !MissingResources.Contains(StyleName) )
+	{
+		MissingResources.Add(StyleName);
+		Log(ISlateStyle::Warning, FText::Format(NSLOCTEXT("SlateStyleSet", "UknownSound", "Unable to find Sound '{0}'."), FText::FromName(StyleName)));
+	}
+#endif
+
+	return Result ? *Result : FStyleDefaults::GetSound();
+}
+
+FSlateFontInfo FSlateStyleSet::GetFontStyle(const FName PropertyName, const ANSICHAR* Specifier) const
+{
+	const FSlateFontInfo* Result = FontInfoResources.Find(Join(PropertyName, Specifier));
+
+	return Result ? *Result : FStyleDefaults::GetFontInfo();
+}
+
+const FSlateWidgetStyle* FSlateStyleSet::GetWidgetStyleInternal(const FName DesiredTypeName, const FName StyleName) const
+{
+	const TSharedRef< struct FSlateWidgetStyle >* StylePtr = WidgetStyleValues.Find(StyleName);
+
+	if ( StylePtr == nullptr )
+	{
+		Log(ISlateStyle::Warning, FText::Format(NSLOCTEXT("SlateStyleSet", "UnknownWidgetStyle", "Unable to find Slate Widget Style '{0}'. Using {1} defaults instead."), FText::FromName(StyleName), FText::FromName(DesiredTypeName)));
+		return nullptr;
+	}
+
+	const TSharedRef< struct FSlateWidgetStyle > Style = *StylePtr;
+
+	if ( Style->GetTypeName() != DesiredTypeName )
+	{
+		Log(ISlateStyle::Error, FText::Format(NSLOCTEXT("SlateStyleSet", "WrongWidgetStyleType", "The Slate Widget Style '{0}' is not of the desired type. Desired: '{1}', Actual: '{2}'"), FText::FromName(StyleName), FText::FromName(DesiredTypeName), FText::FromName(Style->GetTypeName())));
+		return nullptr;
+	}
+
+	return &Style.Get();
+}
+
+void FSlateStyleSet::Log(ISlateStyle::EStyleMessageSeverity Severity, const FText& Message) const
+{
+	if ( Severity == ISlateStyle::Error )
+	{
+		UE_LOG(LogSlateStyle, Error, TEXT("%s"), *Message.ToString());
+	}
+	else if ( Severity == ISlateStyle::PerformanceWarning )
+	{
+		UE_LOG(LogSlateStyle, Warning, TEXT("%s"), *Message.ToString());
+	}
+	else if ( Severity == ISlateStyle::Warning )
+	{
+		UE_LOG(LogSlateStyle, Warning, TEXT("%s"), *Message.ToString());
+	}
+	else if ( Severity == ISlateStyle::Info )
+	{
+		UE_LOG(LogSlateStyle, Log, TEXT("%s"), *Message.ToString());
+	}
+	else
+	{
+		UE_LOG(LogSlateStyle, Fatal, TEXT("%s"), *Message.ToString());
+	}
+}
+
+void FSlateStyleSet::LogUnusedBrushResources()
+{
+	TArray<FString> Filenames;
+	IFileManager::Get().FindFilesRecursive(Filenames, *ContentRootDir, *FString("*.png"), true, false, false);
+
+	for ( FString& FilePath : Filenames )
+	{
+		bool IsUsed = false;
+		for ( auto& Entry : BrushResources )
+		{
+			if ( IsBrushFromFile(FilePath, Entry.Value) )
+			{
+				IsUsed = true;
+				break;
+			}
+		}
+
+		if ( !IsUsed )
+		{
+			for ( auto& Entry : WidgetStyleValues )
+			{
+				TArray< const FSlateBrush* > WidgetBrushes;
+				Entry.Value->GetResources(WidgetBrushes);
+
+				for ( const FSlateBrush* Brush : WidgetBrushes )
+				{
+					if ( IsBrushFromFile(FilePath, Brush) )
+					{
+						IsUsed = true;
+						break;
+					}
+				}
+
+				if ( IsUsed )
+				{
+					break;
+				}
+			}
+		}
+
+		if ( !IsUsed )
+		{
+			Log(ISlateStyle::Warning, FText::FromString(FilePath));
+		}
+	}
+}
+
+bool FSlateStyleSet::IsBrushFromFile(const FString& FilePath, const FSlateBrush* Brush)
+{
+	FString Path = Brush->GetResourceName().ToString();
+	FPaths::MakeStandardFilename(Path);
+	if ( Path.Compare(FilePath, ESearchCase::IgnoreCase) == 0 )
+	{
+		return true;
+	}
+
+	return false;
+}
