Slate: Refactored core Slate implementation into SlateCore module in preparation for UMG.

diff --git a/Engine/Source/Runtime/SlateCore/Private/Styling/SlateWidgetStyleAsset.cpp b/Engine/Source/Runtime/SlateCore/Private/Styling/SlateWidgetStyleAsset.cpp
new file mode 100644
--- /dev/null
+++ b/Engine/Source/Runtime/SlateCore/Private/Styling/SlateWidgetStyleAsset.cpp
@@ -0,0 +1,10 @@
+// Copyright 1998-2014 Epic Games, Inc. All Rights Reserved.
+
+#include "SlateCorePrivatePCH.h"
+
+
+USlateWidgetStyleAsset::USlateWidgetStyleAsset( const class FPostConstructInitializeProperties& PCIP )
+	: Super(PCIP)
+{
+	
+}
diff --git a/Engine/Source/Runtime/SlateCore/Private/Styling/SlateWidgetStyleContainerBase.cpp b/Engine/Source/Runtime/SlateCore/Private/Styling/SlateWidgetStyleContainerBase.cpp
new file mode 100644
--- /dev/null
+++ b/Engine/Source/Runtime/SlateCore/Private/Styling/SlateWidgetStyleContainerBase.cpp
@@ -0,0 +1,22 @@
+// Copyright 1998-2014 Epic Games, Inc. All Rights Reserved.
+
+/*=============================================================================
+	SlateWidgetStyleContainerBase.cpp: Implements the USlateWidgetStyleContainerBase class.
+=============================================================================*/
+
+#include "SlateCorePrivatePCH.h"
+#include "SlateWidgetStyleContainerBase.h"
+
+
+DEFINE_LOG_CATEGORY(LogSlateStyle);
+
+
+USlateWidgetStyleContainerBase::USlateWidgetStyleContainerBase( const class FPostConstructInitializeProperties& PCIP )
+	: Super(PCIP)
+{ }
+
+
+const struct FSlateWidgetStyle* const USlateWidgetStyleContainerBase::GetStyle() const
+{
+	return nullptr;
+}
diff --git a/Engine/Source/Runtime/SlateCore/Private/Styling/SlateWidgetStyleContainerInterface.cpp b/Engine/Source/Runtime/SlateCore/Private/Styling/SlateWidgetStyleContainerInterface.cpp
new file mode 100644
--- /dev/null
+++ b/Engine/Source/Runtime/SlateCore/Private/Styling/SlateWidgetStyleContainerInterface.cpp
@@ -0,0 +1,11 @@
+// Copyright 1998-2014 Epic Games, Inc. All Rights Reserved.
+
+#include "SlateCorePrivatePCH.h"
+#include "SlateWidgetStyleContainerInterface.h"
+
+USlateWidgetStyleContainerInterface::USlateWidgetStyleContainerInterface( const class FPostConstructInitializeProperties& PCIP )
+	: Super(PCIP)
+{
+	
+}
+
diff --git a/Engine/Source/Runtime/SlateCore/Private/Styling/StyleDefaults.cpp b/Engine/Source/Runtime/SlateCore/Private/Styling/StyleDefaults.cpp
new file mode 100644
--- /dev/null
+++ b/Engine/Source/Runtime/SlateCore/Private/Styling/StyleDefaults.cpp
@@ -0,0 +1,10 @@
+// Copyright 1998-2014 Epic Games, Inc. All Rights Reserved.
+
+#include "SlateCorePrivatePCH.h"
+
+
+float FStyleDefaults::DefaultFloat;
+FVector2D FStyleDefaults::DefaultFVector2D = FVector2D( 66.0f, 66.0f );
+FLinearColor FStyleDefaults::DefaultColor = FLinearColor( 1, 1, 1 );
+FMargin FStyleDefaults::DefaultMargin = 66.0f;
+FSlateSound FStyleDefaults::DefaultSound;
diff --git a/Engine/Source/Runtime/SlateCore/Private/Textures/TextureAtlas.cpp b/Engine/Source/Runtime/SlateCore/Private/Textures/TextureAtlas.cpp
new file mode 100644
--- /dev/null
+++ b/Engine/Source/Runtime/SlateCore/Private/Textures/TextureAtlas.cpp
@@ -0,0 +1,230 @@
+// Copyright 1998-2014 Epic Games, Inc. All Rights Reserved.
+
+/*=============================================================================
+	TextureAtlas.cpp: Implements the FSlateTextureAtlas class.
+=============================================================================*/
+
+#include "SlateCorePrivatePCH.h"
+
+
+/* FSlateTextureAtlas helper class
+ *****************************************************************************/
+
+FSlateTextureAtlas::~FSlateTextureAtlas()
+{
+	Empty();
+}
+
+
+/* FSlateTextureAtlas interface
+ *****************************************************************************/
+
+void FSlateTextureAtlas::Empty()
+{
+	// Remove all nodes
+	DestroyNodes( RootNode );
+	RootNode = nullptr;
+
+	STAT(uint32 MemoryBefore = AtlasData.GetAllocatedSize());
+
+	// Clear all raw data
+	AtlasData.Empty();
+
+	STAT(uint32 MemoryAfter = AtlasData.GetAllocatedSize());
+	DEC_MEMORY_STAT_BY(STAT_SlateTextureAtlasMemory, MemoryBefore-MemoryAfter);
+}
+
+
+const FAtlasedTextureSlot* FSlateTextureAtlas::AddTexture( uint32 TextureWidth, uint32 TextureHeight, const TArray<uint8>& Data )
+{
+	// Find a spot for the character in the texture
+	const FAtlasedTextureSlot* NewSlot = FindSlotForTexture(TextureWidth, TextureHeight);
+
+	if (NewSlot)
+	{
+		CopyDataIntoSlot(NewSlot, Data);
+		MarkTextureDirty();
+	}
+
+	return NewSlot;
+}
+	
+
+const FAtlasedTextureSlot* FSlateTextureAtlas::FindSlotForTexture( uint32 InWidth, uint32 InHeight )
+{
+	return FindSlotForTexture(*RootNode, InWidth, InHeight);
+}
+
+
+void FSlateTextureAtlas::DestroyNodes( FAtlasedTextureSlot* StartNode )
+{
+	if (StartNode->Left)
+	{
+		DestroyNodes(StartNode->Left);
+	}
+
+	if (StartNode->Right)
+	{
+		DestroyNodes(StartNode->Right);
+	}
+
+	delete StartNode;
+}
+
+
+void FSlateTextureAtlas::InitAtlasData()
+{
+	check(RootNode == nullptr && AtlasData.Num() == 0);
+
+	RootNode = new FAtlasedTextureSlot(0, 0, AtlasWidth, AtlasHeight, Padding);
+	AtlasData.AddZeroed(AtlasWidth * AtlasHeight * Stride);
+
+	INC_MEMORY_STAT_BY(STAT_SlateTextureAtlasMemory, AtlasData.GetAllocatedSize());
+}
+
+
+void FSlateTextureAtlas::CopyRow( FCopyRowData& CopyRowData )
+{
+	const uint8* Data = CopyRowData.SrcData;
+	uint8* Start = CopyRowData.DestData;
+	const uint32 SourceWidth = CopyRowData.SrcTextureWidth;
+	const uint32 DestWidth = CopyRowData.DestTextureWidth;
+	const uint32 CopyStride = CopyRowData.DataStride;
+	const uint32 CopyPadding = CopyRowData.Padding;
+	const uint32 SrcRow = CopyRowData.SrcRow;
+	const uint32 DestRow = CopyRowData.DestRow;
+
+	const uint8* SourceDataAddr = &Data[SrcRow * SourceWidth * CopyStride]; 
+	uint8* DestDataAddr = &Start[DestRow * DestWidth * CopyStride + CopyPadding * CopyStride]; 
+	FMemory::Memcpy(DestDataAddr, SourceDataAddr, SourceWidth * CopyStride); 
+
+	if (CopyRowData.bDialateToPadding)
+	{ 
+		const uint8* FirstPixel = SourceDataAddr; 
+		const uint8* LastPixel = SourceDataAddr + ((SourceWidth - 1) * CopyStride); 
+		uint8* DestDataAddrNoPadding = &Start[DestRow * DestWidth * CopyStride]; 
+		FMemory::Memcpy(DestDataAddrNoPadding, FirstPixel, CopyStride ); 
+		DestDataAddrNoPadding = &Start[DestRow * DestWidth * CopyStride + (CopyRowData.RowWidth - 1) * CopyStride]; 
+		FMemory::Memcpy(DestDataAddrNoPadding, LastPixel, CopyStride); 
+	} 
+}
+
+
+void FSlateTextureAtlas::CopyDataIntoSlot( const FAtlasedTextureSlot* SlotToCopyTo, const TArray<uint8>& Data )
+{
+	// Copy pixel data to the texture
+	uint8* Start = &AtlasData[SlotToCopyTo->Y*AtlasWidth*Stride + SlotToCopyTo->X*Stride];
+	
+	// Account for same padding on each sides
+	const uint32 AllPadding = Padding*2;
+	// The width of the source texture without padding (actual width)
+	const uint32 SourceWidth = SlotToCopyTo->Width-AllPadding; 
+
+	FCopyRowData CopyRowData;
+	CopyRowData.bDialateToPadding = Padding != 0;
+	CopyRowData.DataStride = Stride;
+	CopyRowData.DestData = Start;
+	CopyRowData.SrcData = Data.GetData();
+	CopyRowData.DestTextureWidth = AtlasWidth;
+	CopyRowData.SrcTextureWidth = SourceWidth;
+	CopyRowData.Padding = Padding;
+	CopyRowData.RowWidth = SlotToCopyTo->Width;
+
+	// Copy color data into padding for bilinear filtering. 
+	// Not used if no padding (assumes sampling outside boundaries of the sub texture is not possible)
+	if (Padding)
+	{
+		// Copy first color row into padding.  
+		CopyRowData.SrcRow = 0;
+		CopyRowData.DestRow = 0;
+
+		CopyRow(CopyRowData);	
+	}
+
+	// Copy each row of the texture
+	for (uint32 Row = Padding; Row < SlotToCopyTo->Height-Padding; ++Row)
+	{
+		CopyRowData.SrcRow = Row-Padding;
+		CopyRowData.DestRow = Row;
+
+		CopyRow(CopyRowData);
+	}
+
+	if (Padding)
+	{
+		// Copy last color row into padding row for bilinear filtering
+		CopyRowData.SrcRow = SlotToCopyTo->Height - AllPadding - 1;
+		CopyRowData.DestRow = SlotToCopyTo->Height - Padding;
+
+		CopyRow(CopyRowData);
+	}
+}
+
+
+const FAtlasedTextureSlot* FSlateTextureAtlas::FindSlotForTexture( FAtlasedTextureSlot& Start, uint32 InWidth, uint32 InHeight )
+{
+	// If there are left and right slots there are empty regions around this slot.  
+	// It also means this slot is occupied by a texture
+	if (Start.Left || Start.Right)
+	{
+		// Recursively search left for the smallest empty slot that can fit the texture
+		if (Start.Left)
+		{
+			const FAtlasedTextureSlot* NewSlot = FindSlotForTexture(*Start.Left, InWidth, InHeight);
+			if (NewSlot)
+			{
+				return NewSlot;
+			}
+		}
+
+		// Recursively search left for the smallest empty slot that can fit the texture
+		if (Start.Right)
+		{
+			const FAtlasedTextureSlot* NewSlot = FindSlotForTexture(*Start.Right, InWidth, InHeight);
+			if(NewSlot)
+			{
+				return NewSlot;
+			}
+		}
+
+		// Not enough space
+		return nullptr;
+	}
+
+	// Account for padding on both sides
+	uint32 TotalPadding = Padding * 2;
+
+	// This slot can't fit the character
+	if (InWidth+TotalPadding > (Start.Width) || InHeight+TotalPadding > (Start.Height))
+	{
+		// not enough space
+		return nullptr;
+	}
+	
+	uint32 PaddedWidth = InWidth+TotalPadding;
+	uint32 PaddedHeight = InHeight+TotalPadding;
+	// The width and height of the new child node
+	uint32 RemainingWidth =  FMath::Max<int32>(0,Start.Width - PaddedWidth);
+	uint32 RemainingHeight = FMath::Max<int32>(0,Start.Height - PaddedHeight);
+
+
+	// Split the remaining area around this slot into two children.
+	if (RemainingHeight <= RemainingWidth)
+	{
+		// Split vertically
+		Start.Left = new FAtlasedTextureSlot(Start.X, Start.Y + PaddedHeight, PaddedWidth, RemainingHeight, Padding);
+		Start.Right = new FAtlasedTextureSlot(Start.X + PaddedWidth, Start.Y, RemainingWidth, Start.Height, Padding);
+	}
+	else
+	{
+		// Split horizontally
+		Start.Left = new FAtlasedTextureSlot(Start.X + PaddedWidth, Start.Y, RemainingWidth, PaddedHeight, Padding);
+		Start.Right = new FAtlasedTextureSlot(Start.X, Start.Y + PaddedHeight, Start.Width, RemainingHeight, Padding);
+	}
+
+	// Shrink the slot to the remaining area.
+	Start.Width = InWidth+TotalPadding;
+	Start.Height = InHeight+TotalPadding;
+
+	return &Start;
+}
diff --git a/Engine/Source/Runtime/SlateCore/Private/Types/WidgetStyle.cpp b/Engine/Source/Runtime/SlateCore/Private/Types/WidgetStyle.cpp
new file mode 100644
--- /dev/null
+++ b/Engine/Source/Runtime/SlateCore/Private/Types/WidgetStyle.cpp
@@ -0,0 +1,27 @@
+// Copyright 1998-2014 Epic Games, Inc. All Rights Reserved.
+
+/*=============================================================================
+	WidgetStyle.cpp: Implements the FWidgetStyle structure.
+=============================================================================*/
+
+#include "SlateCorePrivatePCH.h"
+
+
+/* Static initialization
+ *****************************************************************************/
+
+const float FWidgetStyle::SubdueAmount = 0.6f;
+
+
+/* FWidgetStyle interface
+ *****************************************************************************/
+
+FWidgetStyle& FWidgetStyle::SetForegroundColor( const TAttribute<struct FSlateColor>& InForeground )
+{
+	ForegroundColor = InForeground.Get().GetColor(*this);
+
+	SubduedForeground = ForegroundColor;
+	SubduedForeground.A *= SubdueAmount;
+
+	return *this;
+}
diff --git a/Engine/Source/Runtime/SlateCore/Private/Widgets/SBoxPanel.cpp b/Engine/Source/Runtime/SlateCore/Private/Widgets/SBoxPanel.cpp
new file mode 100644
--- /dev/null
+++ b/Engine/Source/Runtime/SlateCore/Private/Widgets/SBoxPanel.cpp
@@ -0,0 +1,327 @@
+// Copyright 1998-2014 Epic Games, Inc. All Rights Reserved.
+
+#include "SlateCorePrivatePCH.h"
+
+
+/**
+ * Construct this widget
+ *
+ * @param	InArgs	The declaration data for this widget
+ */
+void SHorizontalBox::Construct( const SHorizontalBox::FArguments& InArgs )
+{
+	const int32 NumSlots = InArgs.Slots.Num();
+	for ( int32 SlotIndex = 0; SlotIndex < NumSlots; ++SlotIndex )
+	{
+		Children.Add( InArgs.Slots[SlotIndex] );
+	}
+}
+
+
+/**
+ * Construct this widget
+ *
+ * @param	InArgs	The declaration data for this widget
+ */
+void SVerticalBox::Construct( const SVerticalBox::FArguments& InArgs )
+{
+	const int32 NumSlots = InArgs.Slots.Num();
+	for ( int32 SlotIndex = 0; SlotIndex < NumSlots; ++SlotIndex )
+	{
+		Children.Add( InArgs.Slots[SlotIndex] );
+	}
+}
+
+template<EOrientation Orientation>
+static void ArrangeChildrenAlong( const TPanelChildren<SBoxPanel::FSlot>& Children, const FGeometry& AllottedGeometry, FArrangedChildren& ArrangedChildren )
+{
+	// Allotted space will be given to fixed-size children first.
+	// Remaining space will be proportionately divided between stretch children (SizeRule_Stretch)
+	// based on their stretch coefficient
+
+	if (Children.Num() > 0)
+	{
+		float StretchCoefficientTotal = 0;
+		float FixedTotal = 0;
+
+		// Compute the sum of stretch coefficients (SizeRule_Stretch) and space required by fixed-size widgets
+		// (SizeRule_Auto).
+		for( int32 ChildIndex=0; ChildIndex < Children.Num(); ++ChildIndex )
+		{
+			const SBoxPanel::FSlot& CurChild = Children[ChildIndex];
+
+			if ( CurChild.Widget->GetVisibility() != EVisibility::Collapsed )
+			{
+				// All widgets contribute their margin to the fixed space requirement
+				FixedTotal += CurChild.SlotPadding.Get().GetTotalSpaceAlong<Orientation>();
+
+				if ( CurChild.SizeParam.SizeRule == FSizeParam::SizeRule_Stretch )
+				{
+					// for stretch children we save sum up the stretch coefficients
+					StretchCoefficientTotal += CurChild.SizeParam.Value.Get();
+				}
+				else if (CurChild.SizeParam.SizeRule == FSizeParam::SizeRule_AspectRatio)
+				{
+					// Aspect ratio children contribute size such that they maintain
+					// the aspect ratio of 1 (i.e. they're square).
+					float ChildSize = (Orientation == Orient_Vertical)
+						? AllottedGeometry.Size.X
+						: AllottedGeometry.Size.Y;
+
+					// Clamp to the max size if it was specified
+					float MaxSize = CurChild.MaxSize.Get();
+					FixedTotal += MaxSize > 0 ? FMath::Min( MaxSize, ChildSize ) : ChildSize ;
+				
+				}
+				else
+				{
+					// Auto-sized children contribute their desired size to the fixed space requirement
+					float ChildSize = (Orientation == Orient_Vertical)
+					? CurChild.Widget->GetDesiredSize().Y
+					: CurChild.Widget->GetDesiredSize().X;
+
+					// Clamp to the max size if it was specified
+					float MaxSize = CurChild.MaxSize.Get();
+					FixedTotal += MaxSize > 0 ? FMath::Min( MaxSize, ChildSize ) : ChildSize ;
+				}
+			}
+
+		}
+
+		// The space available for SizeRule_Stretch widgets is any space that wasn't taken up by fixed-sized widgets.
+		const float NonFixedSpace = FMath::Max( 0.0f, (Orientation == Orient_Vertical)
+			? AllottedGeometry.Size.Y - FixedTotal
+			: AllottedGeometry.Size.X - FixedTotal );
+
+		float PositionSoFar = 0;
+
+		// Now that we have the total fixed-space requirement and the total stretch coefficients we can
+		// arrange widgets top-to-bottom or left-to-right (depending on the orientation).
+		for( int32 ChildIndex=0; ChildIndex < Children.Num(); ++ChildIndex )
+		{
+			const SBoxPanel::FSlot& CurChild = Children[ChildIndex];
+			const EVisibility ChildVisibility = CurChild.Widget->GetVisibility();
+
+			FVector2D ChildDesiredSize = CurChild.Widget->GetDesiredSize();
+			
+			// Figure out the area allocated to the child in the direction of BoxPanel
+			// The area allocated to the slot is ChildSize + the associated margin.
+			float ChildSize = 0;
+			if ( ChildVisibility != EVisibility::Collapsed )
+			{
+				// The size of the widget depends on its size type
+				if ( CurChild.SizeParam.SizeRule == FSizeParam::SizeRule_Stretch )
+				{
+					// Stretch widgets get a fraction of the space remaining after all the fixed-space requirements are met
+					ChildSize = NonFixedSpace * CurChild.SizeParam.Value.Get() / StretchCoefficientTotal;
+				}
+				else if (CurChild.SizeParam.SizeRule == FSizeParam::SizeRule_AspectRatio)
+				{
+					// Aspect ratio children contribute size such that they maintain
+					// the aspect ratio of 1.
+					ChildSize = (Orientation == Orient_Vertical)
+						? AllottedGeometry.Size.X
+						: AllottedGeometry.Size.Y;
+				}
+				else
+				{
+					// Auto-sized widgets get their desired-size value
+					ChildSize = (Orientation == Orient_Vertical)
+						? ChildDesiredSize.Y
+						: ChildDesiredSize.X;
+				}
+
+				// Clamp to the max size if it was specified
+				float MaxSize = CurChild.MaxSize.Get();
+				if( MaxSize > 0 )
+				{
+					ChildSize = FMath::Min( MaxSize, ChildSize );
+				}
+			}
+
+			const FMargin SlotPadding(CurChild.SlotPadding.Get());
+
+			FVector2D SlotSize = (Orientation == Orient_Vertical)
+				? FVector2D( AllottedGeometry.Size.X, ChildSize + SlotPadding.GetTotalSpaceAlong<Orient_Vertical>() )
+				: FVector2D( ChildSize + SlotPadding.GetTotalSpaceAlong<Orient_Horizontal>(), AllottedGeometry.Size.Y );
+
+			// Figure out the size and local position of the child within the slot			
+			AlignmentArrangeResult XAlignmentResult = AlignChild<Orient_Horizontal>( SlotSize.X, CurChild, SlotPadding );					
+			AlignmentArrangeResult YAlignmentResult = AlignChild<Orient_Vertical>( SlotSize.Y, CurChild, SlotPadding );
+
+			// Add the information about this child to the output list (ArrangedChildren)
+			ArrangedChildren.AddWidget( ChildVisibility, AllottedGeometry.MakeChild(
+				// The child widget being arranged
+				CurChild.Widget,
+				// Child's local position (i.e. position within parent)
+				(Orientation == Orient_Vertical)
+					? FVector2D( XAlignmentResult.Offset, PositionSoFar + YAlignmentResult.Offset )
+					: FVector2D( PositionSoFar + XAlignmentResult.Offset, YAlignmentResult.Offset ),
+				// Child's size
+				FVector2D( XAlignmentResult.Size, YAlignmentResult.Size )
+				));
+
+			if ( ChildVisibility != EVisibility::Collapsed )
+			{
+				// Offset the next child by the size of the current child and any post-child (bottom/right) margin
+				PositionSoFar += (Orientation == Orient_Vertical) ? SlotSize.Y : SlotSize.X;
+			}
+		}
+	}
+}
+
+/**
+ * Panels arrange their children in a space described by the AllottedGeometry parameter. The results of the arrangement
+ * should be returned by appending a FArrangedWidget pair for every child widget. See StackPanel for an example
+ *
+ * @param AllottedGeometry    The geometry allotted for this widget by its parent.
+ * @param ArrangedChildren    The array to which to add the WidgetGeometries that represent the arranged children.
+ */
+void SBoxPanel::ArrangeChildren( const FGeometry& AllottedGeometry, FArrangedChildren& ArrangedChildren ) const
+{
+	if ( this->Orientation == Orient_Horizontal )
+	{
+		ArrangeChildrenAlong<Orient_Horizontal>(this->Children, AllottedGeometry, ArrangedChildren );
+	}
+	else
+	{
+		ArrangeChildrenAlong<Orient_Vertical>(this->Children, AllottedGeometry, ArrangedChildren );
+	}
+}
+
+/**
+ * Helper to ComputeDesiredSize().
+ *
+ * @param Orientation   Template parameters that controls the orientation in which the children are layed out
+ * @param Children      The children whose size we want to assess in a horizontal or vertical arrangement.
+ *
+ * @return The size desired by the children given an orientation.
+ */
+template<EOrientation Orientation>
+static FVector2D ComputeDesiredSizeForBox( const TPanelChildren<SBoxPanel::FSlot>& Children ) 
+{
+	// The desired size of this panel is the total size desired by its children plus any margins specified in this panel.
+	// The layout along the panel's axis is describe dy the SizeParam, while the perpendicular layout is described by the
+	// alignment property.
+	FVector2D MyDesiredSize(0,0);
+	int32 NumAspectRatioChildren = 0;
+	for( int32 ChildIndex=0; ChildIndex < Children.Num(); ++ChildIndex )
+	{
+		const SBoxPanel::FSlot& CurChild = Children[ChildIndex];
+		const FVector2D& CurChildDesiredSize = CurChild.Widget->GetDesiredSize();
+		if ( CurChild.Widget->GetVisibility() != EVisibility::Collapsed )
+		{
+			if (Orientation == Orient_Vertical)
+			{
+				// For a vertical panel, we want to find the maximum desired width (including margin).
+				// That will be the desired width of the whole panel.
+				MyDesiredSize.X = FMath::Max(MyDesiredSize.X, CurChildDesiredSize.X + CurChild.SlotPadding.Get().GetTotalSpaceAlong<Orient_Horizontal>());
+
+				if ( CurChild.SizeParam.SizeRule == FSizeParam::SizeRule_AspectRatio )
+				{
+					// Omit aspect ratio children for now. We do not know how big they need to be yet.
+					NumAspectRatioChildren++;
+				}
+				else
+				{
+					// Clamp to the max size if it was specified
+					float FinalChildDesiredSize = CurChildDesiredSize.Y;
+					float MaxSize = CurChild.MaxSize.Get();
+					if( MaxSize > 0 )
+					{
+						FinalChildDesiredSize = FMath::Min( MaxSize, FinalChildDesiredSize );
+					}
+
+					MyDesiredSize.Y += FinalChildDesiredSize + CurChild.SlotPadding.Get().GetTotalSpaceAlong<Orient_Vertical>();
+				}
+						
+			}
+			else
+			{
+				// A horizontal panel is just a sideways vertical panel: the axes are swapped.
+
+				MyDesiredSize.Y = FMath::Max(MyDesiredSize.Y, CurChildDesiredSize.Y + CurChild.SlotPadding.Get().GetTotalSpaceAlong<Orient_Vertical>());
+
+				if ( CurChild.SizeParam.SizeRule == FSizeParam::SizeRule_AspectRatio )
+				{
+					// Omit aspect ratio children for now. We do not know how big they need to be yet.
+					NumAspectRatioChildren++;
+				}
+				else
+				{
+					// Clamp to the max size if it was specified
+					float FinalChildDesiredSize = CurChildDesiredSize.X;
+					float MaxSize = CurChild.MaxSize.Get();
+					if( MaxSize > 0 )
+					{
+						FinalChildDesiredSize = FMath::Min( MaxSize, FinalChildDesiredSize );
+					}
+
+					MyDesiredSize.X += FinalChildDesiredSize + CurChild.SlotPadding.Get().GetTotalSpaceAlong<Orient_Horizontal>();
+				}
+			}
+		}
+	}
+
+	// Aspect ratio children need room proportional to the height/width of the entire box
+	if ( Orientation == Orient_Vertical )
+	{
+		MyDesiredSize.Y += MyDesiredSize.X * NumAspectRatioChildren;
+	}
+	else
+	{
+		MyDesiredSize.X += MyDesiredSize.Y * NumAspectRatioChildren;
+	}
+
+
+	return MyDesiredSize;
+}
+
+/**
+ * A Panel's desired size in the space required to arrange of its children on the screen while respecting all of
+ * the children's desired sizes and any layout-related options specified by the user. See StackPanel for an example.
+ *
+ * @return The desired size.
+ */
+FVector2D SBoxPanel::ComputeDesiredSize() const
+{
+	return (Orientation == Orient_Horizontal)
+		? ComputeDesiredSizeForBox<Orient_Horizontal>(this->Children)
+		: ComputeDesiredSizeForBox<Orient_Vertical>(this->Children);
+}
+
+/** @return  The children of a panel in a slot-agnostic way. */
+FChildren* SBoxPanel::GetChildren()
+{
+	return &Children;
+}
+
+int32 SBoxPanel::RemoveSlot( const TSharedRef<SWidget>& SlotWidget )
+{
+	for (int32 SlotIdx = 0; SlotIdx < Children.Num(); ++SlotIdx)
+	{
+		if ( SlotWidget == Children[SlotIdx].Widget )
+		{
+			Children.RemoveAt(SlotIdx);
+			return SlotIdx;
+		}
+	}
+
+	return -1;
+}
+
+void SBoxPanel::ClearChildren()
+{
+	Children.Empty();
+}
+
+/**
+ * A Box Panel's orientation cannot be changed once it is constructed..
+ *
+ * @param InOrientation   The orientation of the Box Panel
+ */
+SBoxPanel::SBoxPanel( EOrientation InOrientation )
+: Orientation(InOrientation)
+{
+
+}
diff --git a/Engine/Source/Runtime/SlateCore/Private/Widgets/SCompoundWidget.cpp b/Engine/Source/Runtime/SlateCore/Private/Widgets/SCompoundWidget.cpp
new file mode 100644
--- /dev/null
+++ b/Engine/Source/Runtime/SlateCore/Private/Widgets/SCompoundWidget.cpp
@@ -0,0 +1,76 @@
+// Copyright 1998-2014 Epic Games, Inc. All Rights Reserved.
+
+/*=============================================================================
+	SCompoundWidget.cpp: Implements the SCompoundWidget class.
+=============================================================================*/
+
+#include "SlateCorePrivatePCH.h"
+
+
+/**
+ * The widget should respond by populating the OutDrawElements array with FDrawElements 
+ * that represent it and any of its children.
+ *
+ * @param AllottedGeometry  The FGeometry that describes an area in which the widget should appear.
+ * @param MyClippingRect    The clipping rectangle allocated for this widget and its children.
+ * @param OutDrawElements   A list of FDrawElements to populate with the output.
+ * @param LayerId           The Layer onto which this widget should be rendered.
+ * @param InColorAndOpacity Color and Opacity to be applied to all the descendants of the widget being painted
+ * @param bParentEnabled	True if the parent of this widget is enabled.
+ *
+ * @return The maximum layer ID attained by this widget or any of its children.
+ */
+int32 SCompoundWidget::OnPaint( const FGeometry& AllottedGeometry, const FSlateRect& MyClippingRect, FSlateWindowElementList& OutDrawElements, int32 LayerId, const FWidgetStyle& InWidgetStyle, bool bParentEnabled ) const
+{
+	// A CompoundWidget just draws its children
+	FArrangedChildren ArrangedChildren(EVisibility::Visible);
+	{
+#if SLATE_HD_STATS
+		SCOPE_CYCLE_COUNTER( STAT_SlateOnPaint_SCompoundWidget );
+#endif
+		this->ArrangeChildren(AllottedGeometry, ArrangedChildren);
+	}
+
+	// There may be zero elements in this array if our child collapsed/hidden
+	if( ArrangedChildren.Num() > 0 )
+	{
+		check( ArrangedChildren.Num() == 1 );
+		FArrangedWidget& TheChild = ArrangedChildren(0);
+
+		const FSlateRect ChildClippingRect = AllottedGeometry.GetClippingRect().InsetBy( ChildSlot.SlotPadding.Get() * AllottedGeometry.Scale ).IntersectionWith(MyClippingRect);
+
+		FWidgetStyle CompoundedWidgetStyle = FWidgetStyle(InWidgetStyle)
+			.BlendColorAndOpacityTint(ColorAndOpacity.Get())
+			.SetForegroundColor( ForegroundColor );
+
+		return TheChild.Widget->OnPaint( TheChild.Geometry, ChildClippingRect, OutDrawElements, LayerId + 1, CompoundedWidgetStyle, ShouldBeEnabled( bParentEnabled ) );
+	}
+	return LayerId;
+}
+
+FChildren* SCompoundWidget::GetChildren()
+{
+	return &ChildSlot;
+}
+
+
+FVector2D SCompoundWidget::ComputeDesiredSize() const
+{
+	EVisibility ChildVisibility = ChildSlot.Widget->GetVisibility();
+	if ( ChildVisibility != EVisibility::Collapsed )
+	{
+		return ChildSlot.Widget->GetDesiredSize() + ChildSlot.SlotPadding.Get().GetDesiredSize();
+	}
+	
+	return FVector2D::ZeroVector;
+}
+
+void SCompoundWidget::ArrangeChildren( const FGeometry& AllottedGeometry, FArrangedChildren& ArrangedChildren ) const
+{
+	ArrangeSingleChild( AllottedGeometry, ArrangedChildren, ChildSlot, ContentScale );
+}
+
+FSlateColor SCompoundWidget::GetForegroundColor() const
+{
+	return ForegroundColor.Get();
+}
\ No newline at end of file
diff --git a/Engine/Source/Runtime/SlateCore/Private/Widgets/SLeafWidget.cpp b/Engine/Source/Runtime/SlateCore/Private/Widgets/SLeafWidget.cpp
new file mode 100644
--- /dev/null
+++ b/Engine/Source/Runtime/SlateCore/Private/Widgets/SLeafWidget.cpp
@@ -0,0 +1,28 @@
+//// Copyright 1998-2014 Epic Games, Inc. All Rights Reserved.
+
+/*=============================================================================
+	SLeafWidget.cpp: Implements the SLeafWidget class.
+=============================================================================*/
+
+#include "SlateCorePrivatePCH.h"
+
+
+/* Static initialization
+ *****************************************************************************/
+
+FNoChildren SLeafWidget::NoChildrenInstance;
+
+
+/* SLeafWidget interface
+ *****************************************************************************/
+
+FChildren* SLeafWidget::GetChildren( )
+{
+	return &SLeafWidget::NoChildrenInstance;
+}
+
+
+void SLeafWidget::ArrangeChildren( const FGeometry& AllottedGeometry, FArrangedChildren& ArrangedChildren ) const
+{
+	// Nothing to arrange; Leaf Widgets do not have children.
+}
diff --git a/Engine/Source/Runtime/SlateCore/Private/Widgets/SNullWidget.cpp b/Engine/Source/Runtime/SlateCore/Private/Widgets/SNullWidget.cpp
new file mode 100644
--- /dev/null
+++ b/Engine/Source/Runtime/SlateCore/Private/Widgets/SNullWidget.cpp
@@ -0,0 +1,51 @@
+// Copyright 1998-2014 Epic Games, Inc. All Rights Reserved.
+
+/*=============================================================================
+	SNullWidget.cpp: Implements the SNullWidget class.
+=============================================================================*/
+ 
+#include "SlateCorePrivatePCH.h"
+
+
+class SLATECORE_API SNullWidgetContent
+	: public SLeafWidget
+{
+public:
+
+	SLATE_BEGIN_ARGS(SNullWidgetContent)
+	{
+		_Visibility = EVisibility::SelfHitTestInvisible;
+	}
+	SLATE_END_ARGS()
+
+	void Construct( const FArguments& InArgs )
+	{
+		SWidget::Construct( 
+			InArgs._ToolTipText,
+			InArgs._ToolTip,
+			InArgs._Cursor, 
+			InArgs._IsEnabled,
+			InArgs._Visibility,
+			InArgs._Tag
+		);
+	}
+
+public:
+
+	// SWidget interface
+
+	virtual int32 OnPaint( const FGeometry& AllottedGeometry, const FSlateRect& MyClippingRect, FSlateWindowElementList& OutDrawElements, int32 LayerId, const FWidgetStyle& InWidgetStyle, bool bParentEnabled ) const OVERRIDE
+	{
+		return LayerId;
+	}
+
+	virtual FVector2D ComputeDesiredSize( ) const OVERRIDE
+	{
+		return FVector2D(0.0f, 0.0f);
+	}
+
+	// End of SWidget interface
+};
+
+
+SLATECORE_API TSharedRef<SWidget> SNullWidget::NullWidget = SNew(SNullWidgetContent).Visibility(EVisibility::Hidden);
diff --git a/Engine/Source/Runtime/SlateCore/Private/Widgets/SOverlay.cpp b/Engine/Source/Runtime/SlateCore/Private/Widgets/SOverlay.cpp
new file mode 100644
--- /dev/null
+++ b/Engine/Source/Runtime/SlateCore/Private/Widgets/SOverlay.cpp
@@ -0,0 +1,199 @@
+// Copyright 1998-2014 Epic Games, Inc. All Rights Reserved.
+
+ #include "SlateCorePrivatePCH.h"
+
+
+/**
+ * Construct this widget
+ *
+ * @param	InArgs	The declaration data for this widget
+ */
+void SOverlay::Construct( const SOverlay::FArguments& InArgs )
+{
+	const int32 NumSlots = InArgs.Slots.Num();
+	for ( int32 SlotIndex = 0; SlotIndex < NumSlots; ++SlotIndex )
+	{
+		Children.Add( InArgs.Slots[SlotIndex] );
+	}
+}
+
+
+void SOverlay::ArrangeChildren( const FGeometry& AllottedGeometry, FArrangedChildren& ArrangedChildren ) const
+{
+	for ( int32 ChildIndex = 0; ChildIndex < Children.Num(); ++ChildIndex )
+	{
+		const FOverlaySlot& CurChild = Children[ChildIndex];
+		const EVisibility ChildVisibility = CurChild.Widget->GetVisibility();
+		if ( ArrangedChildren.Accepts(ChildVisibility) )
+		{
+			const FMargin SlotPadding(CurChild.SlotPadding.Get());
+			AlignmentArrangeResult XResult = AlignChild<Orient_Horizontal>(AllottedGeometry.Size.X, CurChild, SlotPadding);
+			AlignmentArrangeResult YResult = AlignChild<Orient_Vertical>(AllottedGeometry.Size.Y, CurChild, SlotPadding);
+
+			ArrangedChildren.AddWidget( ChildVisibility, AllottedGeometry.MakeChild(
+				CurChild.Widget,
+				FVector2D(XResult.Offset,YResult.Offset),
+				FVector2D(XResult.Size, YResult.Size)
+			) );
+		}
+	}
+}
+
+
+/**
+ * A Panel's desired size in the space required to arrange of its children on the screen while respecting all of
+ * the children's desired sizes and any layout-related options specified by the user. See StackPanel for an example.
+ *
+ * @return The desired size.
+ */
+FVector2D SOverlay::ComputeDesiredSize() const
+{
+	FVector2D MaxSize(0,0);
+	for ( int32 ChildIndex=0; ChildIndex < Children.Num(); ++ChildIndex )
+	{
+		const FOverlaySlot& CurSlot = Children[ChildIndex];
+		const EVisibility ChildVisibilty = CurSlot.Widget->GetVisibility();
+		if ( ChildVisibilty != EVisibility::Collapsed )
+		{
+			FVector2D ChildDesiredSize = CurSlot.Widget->GetDesiredSize() + CurSlot.SlotPadding.Get().GetDesiredSize();
+			MaxSize.X = FMath::Max( MaxSize.X, ChildDesiredSize.X );
+			MaxSize.Y = FMath::Max( MaxSize.Y, ChildDesiredSize.Y );
+		}
+	}
+
+	return MaxSize;
+}
+
+/** @return  The children of a panel in a slot-agnostic way. */
+FChildren* SOverlay::GetChildren()
+{
+	return &Children;
+}
+
+/**
+ * The widget should respond by populating the OutDrawElements array with FDrawElements 
+ * that represent it and any of its children.
+ *
+ * @param AllottedGeometry  The FGeometry that describes an area in which the widget should appear.
+ * @param MyClippingRect    The clipping rectangle allocated for this widget and its children.
+ * @param OutDrawElements   A list of FDrawElements to populate with the output.
+ * @param LayerId           The Layer onto which this widget should be rendered.
+ * @param InColorAndOpacity Color and Opacity to be applied to all the descendants of the widget being painted
+ * @param bParentEnabled	True if the parent of this widget is enabled.
+ *
+ * @return The maximum layer ID attained by this widget or any of its children.
+ */
+int32 SOverlay::OnPaint( const FGeometry& AllottedGeometry, const FSlateRect& MyClippingRect, FSlateWindowElementList& OutDrawElements, int32 LayerId, const FWidgetStyle& InWidgetStyle, bool bParentEnabled ) const
+{
+	FArrangedChildren ArrangedChildren(EVisibility::Visible);
+	{
+#if SLATE_HD_STATS
+		SCOPE_CYCLE_COUNTER( STAT_SlateOnPaint_SOverlay );
+#endif
+		// The box panel has no visualization of its own; it just visualizes its children.
+		this->ArrangeChildren(AllottedGeometry, ArrangedChildren);
+	}
+
+	// Because we paint multiple children, we must track the maximum layer id that they produced in case one of our parents
+	// wants to an overlay for all of its contents.
+	int32 MaxLayerId = LayerId;
+
+	for (int32 ChildIndex = 0; ChildIndex < ArrangedChildren.Num(); ++ChildIndex)
+	{
+		FArrangedWidget& CurWidget = ArrangedChildren(ChildIndex);
+		FSlateRect ChildClipRect = MyClippingRect.IntersectionWith( CurWidget.Geometry.GetClippingRect() );
+		const int32 CurWidgetsMaxLayerId = CurWidget.Widget->OnPaint( CurWidget.Geometry, ChildClipRect, OutDrawElements, MaxLayerId + 1, InWidgetStyle, ShouldBeEnabled( bParentEnabled ) );
+
+		MaxLayerId = FMath::Max( MaxLayerId, CurWidgetsMaxLayerId );
+	}
+
+	return MaxLayerId;
+}
+
+SOverlay::FOverlaySlot& SOverlay::AddSlot( int32 ZOrder )
+{
+	FOverlaySlot& NewSlot = SOverlay::Slot();
+	if ( ZOrder == INDEX_NONE )
+	{
+		// No ZOrder was specified; just add to the end of the list.
+		// Use a ZOrder index one after the last elements.
+		NewSlot.ZOrder = (Children.Num() == 0)
+			? 0
+			: ( Children[ Children.Num()-1 ].ZOrder + 1 );
+
+		this->Children.Add( &NewSlot );
+	}
+	else
+	{
+		// Figure out where to add the widget based on ZOrder
+		bool bFoundSlot = false;
+		int32 CurSlotIndex = 0;
+		for( ; CurSlotIndex < Children.Num(); ++CurSlotIndex )
+		{
+			const FOverlaySlot& CurSlot = Children[ CurSlotIndex ];
+			if( ZOrder < CurSlot.ZOrder )
+			{
+				// Insert before
+				bFoundSlot = true;
+				break;
+			}
+		}
+
+		// Add a slot at the desired location
+		this->Children.Insert( &NewSlot, CurSlotIndex );
+	}
+
+	NewSlot.ZOrder = ZOrder;
+	return NewSlot;
+}
+
+void SOverlay::RemoveSlot( int32 ZOrder )
+{
+	if (ZOrder != INDEX_NONE)
+	{
+		for( int32 ChildIndex=0; ChildIndex < Children.Num(); ++ChildIndex )
+		{
+			if ( Children[ChildIndex].ZOrder == ZOrder )
+			{
+				Children.RemoveAt( ChildIndex );
+				return;
+			}
+		}
+
+		ensureMsgf(false, TEXT("Could not remove slot. There are no children with ZOrder %d."));
+	}
+	else if (Children.Num() > 0)
+	{
+		Children.RemoveAt( Children.Num() - 1 );
+	}
+	else
+	{
+		ensureMsgf(false, TEXT("Could not remove slot. There are no slots left."));
+	}
+}
+
+void SOverlay::ClearChildren()
+{
+	Children.Empty();
+}
+
+/** Returns the number of child widgets */
+int32 SOverlay::GetNumWidgets() const
+{
+	return Children.Num();
+}
+
+
+void SOverlay::RemoveSlot( TSharedRef< SWidget > Widget )
+{
+	// Search and remove
+	for( int32 CurSlotIndex = 0; CurSlotIndex < Children.Num(); ++CurSlotIndex )
+	{
+		const FOverlaySlot& CurSlot = Children[ CurSlotIndex ];
+		if( CurSlot.Widget == Widget )
+		{
+			Children.RemoveAt( CurSlotIndex );
+			break;
+		}
+	}
+}
diff --git a/Engine/Source/Runtime/SlateCore/Private/Widgets/SPanel.cpp b/Engine/Source/Runtime/SlateCore/Private/Widgets/SPanel.cpp
new file mode 100644
--- /dev/null
+++ b/Engine/Source/Runtime/SlateCore/Private/Widgets/SPanel.cpp
@@ -0,0 +1,54 @@
+// Copyright 1998-2014 Epic Games, Inc. All Rights Reserved.
+
+/*=============================================================================
+	SPanel.cpp: Implements the SPanel class.
+=============================================================================*/
+
+#include "SlateCorePrivatePCH.h"
+
+
+/* SWidget overrides
+ *****************************************************************************/
+
+int32 SPanel::OnPaint( const FGeometry& AllottedGeometry, const FSlateRect& MyClippingRect, FSlateWindowElementList& OutDrawElements, int32 LayerId, const FWidgetStyle& InWidgetStyle, bool bParentEnabled ) const
+{
+	// REMOVE THIS: Currently adding a debug quad so that BoxPanels are visible during development
+	// OutDrawElements.AddItem( FSlateDrawElement::MakeQuad( AllottedGeometry.AbsolutePosition, AllottedGeometry.Size ) );	
+	FArrangedChildren ArrangedChildren(EVisibility::Visible);
+	{
+#if SLATE_HD_STATS
+		SCOPE_CYCLE_COUNTER( STAT_SlateOnPaint_SPanel );
+#endif
+		// The box panel has no visualization of its own; it just visualizes its children.
+		
+		this->ArrangeChildren(AllottedGeometry, ArrangedChildren);
+	}
+
+	return PaintArrangedChildren(ArrangedChildren, MyClippingRect, OutDrawElements, LayerId, InWidgetStyle, bParentEnabled);
+}
+
+
+/* SPanel implementation
+ *****************************************************************************/
+
+int32 SPanel::PaintArrangedChildren( const FArrangedChildren& ArrangedChildren, const FSlateRect& MyClippingRect, FSlateWindowElementList& OutDrawElements, int32 LayerId, const FWidgetStyle& InWidgetStyle, bool bParentEnabled  ) const
+{
+	// Because we paint multiple children, we must track the maximum layer id that they produced in case one of our parents
+	// wants to an overlay for all of its contents.
+	int32 MaxLayerId = LayerId;
+
+	for (int32 ChildIndex = 0; ChildIndex < ArrangedChildren.Num(); ++ChildIndex)
+	{
+		const FArrangedWidget& CurWidget = ArrangedChildren(ChildIndex);
+		FSlateRect ChildClipRect = MyClippingRect.IntersectionWith(CurWidget.Geometry.GetClippingRect());
+		
+		if (ChildClipRect.GetSize().Size() > 0)
+		{
+			const int32 CurWidgetsMaxLayerId = CurWidget.Widget->OnPaint(CurWidget.Geometry, ChildClipRect, OutDrawElements, LayerId, InWidgetStyle, ShouldBeEnabled(bParentEnabled));
+			MaxLayerId = FMath::Max(MaxLayerId, CurWidgetsMaxLayerId);
+		}		
+		
+	}
+	
+	return MaxLayerId;
+}
diff --git a/Engine/Source/Runtime/SlateCore/Private/Widgets/SWidget.cpp b/Engine/Source/Runtime/SlateCore/Private/Widgets/SWidget.cpp
new file mode 100644
--- /dev/null
+++ b/Engine/Source/Runtime/SlateCore/Private/Widgets/SWidget.cpp
@@ -0,0 +1,474 @@
+// Copyright 1998-2014 Epic Games, Inc. All Rights Reserved.
+
+#include "SlateCorePrivatePCH.h"
+
+
+SWidget::SWidget()
+	: CreatedInFile( TEXT("") )
+	, CreatedOnLine( -1 )
+	, Cursor( TOptional<EMouseCursor::Type>() )
+	, EnabledState( true )
+	, Visibility( EVisibility::Visible )
+	, DesiredSize(FVector2D::ZeroVector)
+	, ToolTip()
+	, bToolTipForceFieldEnabled( false )
+	, bIsHovered( false )
+{
+
+}
+
+void SWidget::Construct(
+	const TAttribute<FString> & InToolTipText ,
+	const TSharedPtr<IToolTip> & InToolTip ,
+	const TAttribute< TOptional<EMouseCursor::Type> > & InCursor ,
+	const TAttribute<bool> & InEnabledState ,
+	const TAttribute<EVisibility> & InVisibility,
+	const FName& InTag
+)
+{
+	if ( InToolTip.IsValid() )
+	{
+		// If someone specified a fancy widget tooltip, use it.
+		ToolTip = InToolTip;
+	}
+	else if ( InToolTipText.IsBound() || !(InToolTipText.Get().IsEmpty()) )
+	{
+		// If someone specified a text binding, make a tooltip out of it
+		ToolTip = FSlateApplicationBase::Get().MakeToolTip(InToolTipText);
+	}
+	else if( !ToolTip.IsValid() || (ToolTip.IsValid() && ToolTip->IsEmpty()) )
+	{	
+		// We don't have a tooltip.
+		ToolTip.Reset();
+	}
+
+	Cursor = InCursor;
+	EnabledState = InEnabledState;
+	Visibility = InVisibility;
+}
+
+
+FReply SWidget::OnKeyboardFocusReceived( const FGeometry& MyGeometry, const FKeyboardFocusEvent& InKeyboardFocusEvent )
+{
+	return FReply::Unhandled();
+}
+
+
+FReply SWidget::OnKeyChar( const FGeometry& MyGeometry, const FCharacterEvent& InCharacterEvent )
+{
+	return FReply::Unhandled();
+}
+
+
+FReply SWidget::OnPreviewKeyDown( const FGeometry& MyGeometry, const FKeyboardEvent& InKeyboardEvent )
+{
+	return FReply::Unhandled();
+}
+
+
+FReply SWidget::OnKeyDown( const FGeometry& MyGeometry, const FKeyboardEvent& InKeyboardEvent )
+{
+	return FReply::Unhandled();
+}
+
+
+FReply SWidget::OnKeyUp( const FGeometry& MyGeometry, const FKeyboardEvent& InKeyboardEvent )
+{
+	return FReply::Unhandled();
+}
+
+
+FReply SWidget::OnMouseButtonDown( const FGeometry& MyGeometry, const FPointerEvent& MouseEvent )
+{
+	return FReply::Unhandled();
+}
+
+
+FReply SWidget::OnPreviewMouseButtonDown( const FGeometry& MyGeometry, const FPointerEvent& MouseEvent )
+{
+	return FReply::Unhandled();
+}
+
+
+FReply SWidget::OnMouseButtonUp( const FGeometry& MyGeometry, const FPointerEvent& MouseEvent )
+{
+	return FReply::Unhandled();
+}
+
+
+FReply SWidget::OnMouseMove( const FGeometry& MyGeometry, const FPointerEvent& MouseEvent )
+{
+	return FReply::Unhandled();
+}
+
+
+void SWidget::OnMouseEnter( const FGeometry& MyGeometry, const FPointerEvent& MouseEvent )
+{
+	bIsHovered = true;
+}
+
+
+void SWidget::OnMouseLeave( const FPointerEvent& MouseEvent )
+{
+	bIsHovered = false;
+}
+
+
+FReply SWidget::OnMouseWheel( const FGeometry& MyGeometry, const FPointerEvent& MouseEvent )
+{
+	return FReply::Unhandled();
+}
+
+
+FCursorReply SWidget::OnCursorQuery( const FGeometry& MyGeometry, const FPointerEvent& CursorEvent ) const
+{
+	TOptional<EMouseCursor::Type> TheCursor = Cursor.Get();
+	return ( TheCursor.IsSet() )
+		? FCursorReply::Cursor( TheCursor.GetValue() )
+		: FCursorReply::Unhandled();
+}
+
+
+FReply SWidget::OnMouseButtonDoubleClick( const FGeometry& InMyGeometry, const FPointerEvent& InMouseEvent )
+{
+	return FReply::Unhandled();
+}
+
+
+bool SWidget::OnVisualizeTooltip( const TSharedPtr<SWidget>& TooltipContent )
+{
+	return false;
+}
+
+
+FReply SWidget::OnDragDetected( const FGeometry& MyGeometry, const FPointerEvent& MouseEvent )
+{
+	return FReply::Unhandled();
+}
+
+
+void SWidget::OnDragEnter( const FGeometry& MyGeometry, const FDragDropEvent& DragDropEvent )
+{
+}
+
+
+void SWidget::OnDragLeave( const FDragDropEvent& DragDropEvent )
+{
+}
+
+
+FReply SWidget::OnDragOver( const FGeometry& MyGeometry, const FDragDropEvent& DragDropEvent )
+{
+	return FReply::Unhandled();
+}
+
+
+FReply SWidget::OnDrop( const FGeometry& MyGeometry, const FDragDropEvent& DragDropEvent )
+{
+	return FReply::Unhandled();
+}
+
+
+FReply SWidget::OnControllerButtonPressed( const FGeometry& MyGeometry, const FControllerEvent& ControllerEvent )
+{
+	return FReply::Unhandled();
+}
+
+
+FReply SWidget::OnControllerButtonReleased( const FGeometry& MyGeometry, const FControllerEvent& ControllerEvent )
+{
+	return FReply::Unhandled();
+}
+
+
+FReply SWidget::OnControllerAnalogValueChanged( const FGeometry& MyGeometry, const FControllerEvent& ControllerEvent )
+{
+	return FReply::Unhandled();
+}
+
+
+FReply SWidget::OnTouchGesture( const FGeometry& MyGeometry, const FPointerEvent& GestureEvent )
+{
+	return FReply::Unhandled();
+}
+
+
+FReply SWidget::OnTouchStarted( const FGeometry& MyGeometry, const FPointerEvent& TouchEvent )
+{
+	return FReply::Unhandled();
+}
+
+
+FReply SWidget::OnTouchMoved( const FGeometry& MyGeometry, const FPointerEvent& TouchEvent )
+{
+	return FReply::Unhandled();
+}
+
+
+FReply SWidget::OnTouchEnded( const FGeometry& MyGeometry, const FPointerEvent& TouchEvent )
+{
+	return FReply::Unhandled();
+}
+
+
+FReply SWidget::OnMotionDetected( const FGeometry& MyGeometry, const FMotionEvent& MotionEvent )
+{
+	return FReply::Unhandled();
+}
+
+
+EWindowZone::Type SWidget::GetWindowZoneOverride() const
+{
+	// No special behavior.  Override this in derived widgets, if needed.
+	return EWindowZone::Unspecified;
+}
+
+
+bool SWidget::OnHitTest( const FGeometry& MyGeometry, FVector2D InAbsoluteCursorPosition )
+{
+	return true;
+}
+
+
+void SWidget::TickWidgetsRecursively( const FGeometry& AllottedGeometry, const double InCurrentTime, const float InDeltaTime )
+{
+	// Tick this widget
+	Tick( AllottedGeometry, InCurrentTime, InDeltaTime );
+
+	// Gather all children, whether they're visible or not.  We need to allow invisible widgets to
+	// consider whether they should still be invisible in their tick functions, as well as maintain
+	// other state when hidden,
+	FArrangedChildren ArrangedChildren(EVisibility::All);
+	ArrangeChildren(AllottedGeometry, ArrangedChildren);
+
+	// Recurse!
+	for(int32 ChildIndex=0; ChildIndex < ArrangedChildren.Num(); ++ChildIndex)
+	{
+		FArrangedWidget& SomeChild = ArrangedChildren(ChildIndex);
+		SomeChild.Widget->TickWidgetsRecursively( SomeChild.Geometry, InCurrentTime, InDeltaTime );
+	}
+}
+
+
+void SWidget::SlatePrepass()	
+{
+	// Cache child desired sizes first. This widget's desired size is
+	// a function of its children's sizes.
+	FChildren* MyChildren = this->GetChildren();
+	int32 NumChildren = MyChildren->Num();
+	for(int32 ChildIndex=0; ChildIndex < NumChildren; ++ChildIndex)
+	{
+		const TSharedRef<SWidget>& Child = MyChildren->GetChildAt(ChildIndex);
+		if ( Child->Visibility.Get() != EVisibility::Collapsed )
+		{
+			// Recur: Descend down the widget tree.
+			Child->SlatePrepass();
+		}
+	}
+
+	// Cache this widget's desired size.
+	CacheDesiredSize();
+}
+
+
+void SWidget::CacheDesiredSize()
+{
+	// Cache this widget's desired size.
+	this->Advanced_SetDesiredSize( this->ComputeDesiredSize() );
+}
+
+
+const FVector2D& SWidget::GetDesiredSize() const
+{
+	return DesiredSize;
+}
+
+
+bool SWidget::SupportsKeyboardFocus() const
+{
+	return false;
+}
+
+
+bool SWidget::HasKeyboardFocus() const
+{
+	return (FSlateApplicationBase::Get().GetKeyboardFocusedWidget().Get() == this);
+}
+
+
+bool SWidget::HasFocusedDescendants() const
+{
+	return FSlateApplicationBase::Get().HasFocusedDescendants(SharedThis(this));
+}
+
+
+bool SWidget::HasMouseCapture() const
+{
+	return (FSlateApplicationBase::Get().GetMouseCaptor().Get() == this);
+}
+
+
+bool SWidget::FindChildGeometries( const FGeometry& MyGeometry, const TSet< TSharedRef<SWidget> >& WidgetsToFind, TMap<TSharedRef<SWidget>, FArrangedWidget>& OutResult ) const
+{
+	FindChildGeometries_Helper(MyGeometry, WidgetsToFind, OutResult);
+	return OutResult.Num() == WidgetsToFind.Num();
+}
+
+
+void SWidget::FindChildGeometries_Helper( const FGeometry& MyGeometry, const TSet< TSharedRef<SWidget> >& WidgetsToFind, TMap<TSharedRef<SWidget>, FArrangedWidget>& OutResult ) const
+{
+	// Perform a breadth first search!
+
+	FArrangedChildren ArrangedChildren(EVisibility::Visible);
+	this->ArrangeChildren( MyGeometry, ArrangedChildren );
+	const int32 NumChildren = ArrangedChildren.Num();
+
+	// See if we found any of the widgets on this level.
+	for(int32 ChildIndex=0; ChildIndex < NumChildren; ++ChildIndex )
+	{
+		const FArrangedWidget& CurChild = ArrangedChildren( ChildIndex );
+		
+		if ( WidgetsToFind.Contains(CurChild.Widget) )
+		{
+			// We found one of the widgets for which we need geometry!
+			OutResult.Add( CurChild.Widget, CurChild );
+		}
+	}
+
+	// If we have not found all the widgets that we were looking for, descend.
+	if ( OutResult.Num() != WidgetsToFind.Num() )
+	{
+		// Look for widgets among the children.
+		for( int32 ChildIndex=0; ChildIndex < NumChildren; ++ChildIndex )
+		{
+			const FArrangedWidget& CurChild = ArrangedChildren( ChildIndex );
+			CurChild.Widget->FindChildGeometries_Helper( CurChild.Geometry, WidgetsToFind, OutResult );
+		}	
+	}	
+}
+
+
+FGeometry SWidget::FindChildGeometry( const FGeometry& MyGeometry, TSharedRef<SWidget> WidgetToFind ) const
+{
+	// We just need to find the one WidgetToFind among our descendants.
+	TSet< TSharedRef<SWidget> > WidgetsToFind;
+	{
+		WidgetsToFind.Add( WidgetToFind );
+	}
+	TMap<TSharedRef<SWidget>, FArrangedWidget> Result;
+
+	FindChildGeometries( MyGeometry, WidgetsToFind, Result );
+
+	return Result.FindChecked( WidgetToFind ).Geometry;
+}
+
+
+int32 SWidget::FindChildUnderMouse( const FArrangedChildren& Children, const FPointerEvent& MouseEvent )
+{
+	const FVector2D& AbsoluteCursorLocation = MouseEvent.GetScreenSpacePosition();
+	const int32 NumChildren = Children.Num();
+	for( int32 ChildIndex=NumChildren-1; ChildIndex >= 0; --ChildIndex )
+	{
+		const FArrangedWidget& Candidate = Children(ChildIndex);
+		const bool bCandidateUnderCursor = 
+			// Candidate is physically under the cursor
+			Candidate.Geometry.IsUnderLocation( AbsoluteCursorLocation ) && 
+			// Candidate actually considers itself hit by this test
+			(Candidate.Widget->OnHitTest( Candidate.Geometry, AbsoluteCursorLocation ));
+
+		if (bCandidateUnderCursor)
+		{
+			return ChildIndex;
+		}
+	}
+
+	return INDEX_NONE;
+}
+
+
+FString SWidget::ToString() const
+{
+	return FString::Printf(TEXT("%s [%s(%d)]"), *this->TypeOfWidget.ToString(), *this->CreatedInFile.ToString(), this->CreatedOnLine );
+}
+
+
+FString SWidget::GetTypeAsString() const
+{
+	return this->TypeOfWidget.ToString();
+}
+
+
+FName SWidget::GetType() const
+{
+	return this->TypeOfWidget;
+}
+
+
+FString SWidget::GetReadableLocation() const
+{
+	return FString::Printf(TEXT("%s(%d)"), *this->CreatedInFile.ToString(), this->CreatedOnLine );
+}
+
+
+FString SWidget::GetCreatedInFile() const
+{
+	return this->CreatedInFileFullPath.ToString();
+}
+
+
+int32 SWidget::GetCreatedInLineNumber() const
+{
+	return this->CreatedOnLine;
+}
+
+
+FSlateColor SWidget::GetForegroundColor() const
+{
+	static FSlateColor NoColor = FSlateColor::UseForeground();
+	return NoColor;
+}
+
+
+void SWidget::SetToolTipText( const TAttribute<FString>& ToolTipString )
+{
+	ToolTip = FSlateApplicationBase::Get().MakeToolTip(ToolTipString);
+}
+
+
+void SWidget::SetToolTipText( const FText& ToolTipText )
+{
+	ToolTip = FSlateApplicationBase::Get().MakeToolTip(ToolTipText);
+}
+
+
+void SWidget::SetToolTip( const TSharedPtr<IToolTip> & InToolTip )
+{
+	ToolTip = InToolTip;
+}
+
+
+TSharedPtr<IToolTip> SWidget::GetToolTip()
+{
+	return ToolTip;
+}
+
+
+void SWidget::EnableToolTipForceField( const bool bEnableForceField )
+{
+	bToolTipForceFieldEnabled = bEnableForceField;
+}
+
+
+void SWidget::SetCursor( const TAttribute< TOptional<EMouseCursor::Type> >& InCursor )
+{
+	Cursor = InCursor;
+}
+
+
+void SWidget::SetDebugInfo( const ANSICHAR* InType, const ANSICHAR* InFile, int32 OnLine )
+{
+	this->TypeOfWidget = InType;
+	this->CreatedInFileFullPath = FName( InFile );
+	this->CreatedInFile = FName( *FPaths::GetCleanFilename(InFile) );
+	this->CreatedOnLine = OnLine;
+}
diff --git a/Engine/Source/Runtime/SlateCore/Private/Widgets/SWindow.cpp b/Engine/Source/Runtime/SlateCore/Private/Widgets/SWindow.cpp
new file mode 100644
--- /dev/null
+++ b/Engine/Source/Runtime/SlateCore/Private/Widgets/SWindow.cpp
@@ -0,0 +1,1573 @@
+// Copyright 1998-2014 Epic Games, Inc. All Rights Reserved.
+
+/*=============================================================================
+	SWindow.cpp: Implements the SWindow class.
+=============================================================================*/
+
+#include "SlateCorePrivatePCH.h"
+
+
+namespace SWindowDefs
+{
+	/** Height of a Slate window title bar, in pixels */
+	static const float DefaultTitleBarSize = 24.0f;
+
+	/** Size of the hit result border for the window borders */
+	static const FSlateRect HitResultBorderSize(10,10,10,10);
+
+	/** Actual size of the window borders */
+	static const FMargin WindowBorderSize(5,5,5,5);
+
+	/** Size of the corner rounding radius.  Used for regular, non-maximized windows only (not tool-tips or decorators.) */
+	static const int32 CornerRadius = 6;
+}
+
+
+/**
+ * An internal overlay used by Slate to support in-window pop ups and tooltips.
+ * The overlay ignores DPI scaling when it does its own arrangement, but otherwise
+ * passes all DPI scale values through.
+ */
+class SPopupLayer : public SPanel
+{
+public:
+	SLATE_BEGIN_ARGS( SPopupLayer )
+		{
+			_Visibility = EVisibility::SelfHitTestInvisible;
+		}
+
+		SLATE_SUPPORTS_SLOT( FPopupLayerSlot )
+
+	SLATE_END_ARGS()
+
+	void Construct( const FArguments& InArgs, const TSharedRef<SWindow>& InWindow )
+	{
+		OwnerWindow = InWindow;
+
+		const int32 NumSlots = InArgs.Slots.Num();
+		for ( int32 SlotIndex = 0; SlotIndex < NumSlots; ++SlotIndex )
+		{
+			Children.Add( InArgs.Slots[SlotIndex] );
+		}
+	}
+
+	/** Make a new ListPanel::Slot  */
+	FPopupLayerSlot& Slot()
+	{
+		return *(new FPopupLayerSlot());
+	}
+	
+	/** Add a slot to the ListPanel */
+	FPopupLayerSlot& AddSlot(int32 InsertAtIndex = INDEX_NONE)
+	{
+		FPopupLayerSlot& NewSlot = SPopupLayer::Slot();
+		if (InsertAtIndex == INDEX_NONE)
+		{
+			this->Children.Add( &NewSlot );
+		}
+		else
+		{
+			this->Children.Insert( &NewSlot, InsertAtIndex );
+		}
+	
+		return NewSlot;
+	}
+
+	void RemoveSlot(const TSharedRef<SWidget>& WidgetToRemove)
+	{
+		for( int32 CurSlotIndex = 0; CurSlotIndex < Children.Num(); ++CurSlotIndex )
+		{
+			const FPopupLayerSlot& CurSlot = Children[ CurSlotIndex ];
+			if( CurSlot.Widget == WidgetToRemove )
+			{
+				Children.RemoveAt( CurSlotIndex );
+				return;
+			}
+		}
+	}
+
+private:
+
+	virtual void ArrangeChildren( const FGeometry& AllottedGeometry, FArrangedChildren& ArrangedChildren ) const OVERRIDE
+	{
+		const FVector2D WindowDesktopPosition = (ensure(OwnerWindow.IsValid()))
+			? OwnerWindow.Pin()->GetPositionInScreen()
+			: FVector2D::ZeroVector;
+
+		const FVector2D WindowSize = (ensure(OwnerWindow.IsValid()))
+			? OwnerWindow.Pin()->GetClientSizeInScreen()
+			: FVector2D::ZeroVector;
+		
+		for ( int32 ChildIndex = 0; ChildIndex < Children.Num(); ++ChildIndex )
+		{
+			const FPopupLayerSlot& CurChild = Children[ChildIndex];
+			const EVisibility ChildVisibility = CurChild.Widget->GetVisibility();
+			if ( ArrangedChildren.Accepts(ChildVisibility) )
+			{
+				const FVector2D WidgetDesiredSize = CurChild.Widget->GetDesiredSize();
+				const float ChildScale = CurChild.Scale_Attribute.Get();
+				const bool bClampToWindow = CurChild.Clamp_Attribute.Get();
+				const FVector2D ChildSize = WidgetDesiredSize * ChildScale;
+				FVector2D ChildLocalPosition = CurChild.DesktopPosition_Attribute.Get() - WindowDesktopPosition;
+				
+				if(bClampToWindow)
+				{
+					FVector2D ClampBufferReducedWindowSize = WindowSize - CurChild.ClampBuffer_Attribute.Get();
+					ChildLocalPosition.X = FMath::Clamp(ChildLocalPosition.X - FMath::Max(0.0f, (ChildLocalPosition.X  + WidgetDesiredSize.X ) - ClampBufferReducedWindowSize.X), 0.0f, MAX_flt);
+					ChildLocalPosition.Y = FMath::Clamp(ChildLocalPosition.Y - FMath::Max(0.0f, (ChildLocalPosition.Y  + WidgetDesiredSize.Y ) - ClampBufferReducedWindowSize.Y), 0.0f, MAX_flt);
+				}
+
+				// The position is explicitly in desktop pixels.
+				// The size and DPI scale come from the widget that is using
+				// this overlay to "punch" through the UI.
+				ArrangedChildren.AddWidget( ChildVisibility, FArrangedWidget( CurChild.Widget,
+					FGeometry(
+						ChildLocalPosition / ChildScale,
+						AllottedGeometry.AbsolutePosition,
+						FVector2D(
+							// Notice that override sizes get divided by ChildScale because any explicit size overrides are in pixel units and are not affected by DPI scaling.
+							CurChild.WidthOverride_Attribute.IsSet() ? CurChild.WidthOverride_Attribute.Get() / ChildScale : WidgetDesiredSize.X,
+							CurChild.HeightOverride_Attribute.IsSet() ? CurChild.HeightOverride_Attribute.Get() / ChildScale : WidgetDesiredSize.Y
+						),
+						ChildScale
+				) ) );
+			}
+		}
+	}
+
+	virtual FVector2D ComputeDesiredSize() const OVERRIDE
+	{
+		return FVector2D(100,100);
+	}
+
+	/**
+	 * All widgets must provide a way to access their children in a layout-agnostic way.
+	 * Panels store their children in Slots, which creates a dilemma. Most panels
+	 * can store their children in a TPanelChildren<Slot>, where the Slot class
+	 * provides layout information about the child it stores. In that case
+	 * GetChildren should simply return the TPanelChildren<Slot>. See StackPanel for an example.
+	 */
+	virtual FChildren* GetChildren() OVERRIDE
+	{
+		return &Children;
+	}
+
+	TPanelChildren<FPopupLayerSlot> Children;
+	TWeakPtr<SWindow> OwnerWindow;
+};
+
+
+void SWindow::Construct(const FArguments& InArgs)
+{
+	check(InArgs._Style);
+	this->Style = InArgs._Style;
+	this->WindowBackground = &InArgs._Style->BackgroundBrush;
+
+	this->Title = InArgs._Title;
+	this->bDragAnywhere = InArgs._bDragAnywhere;
+	this->bIsTransparent = InArgs._SupportsTransparency;
+	this->Opacity = InArgs._InitialOpacity;
+	this->bInitiallyMaximized = InArgs._IsInitiallyMaximized;
+	this->SizingRule = InArgs._SizingRule;
+	this->bIsPopupWindow = InArgs._IsPopupWindow;
+	this->bFocusWhenFirstShown = InArgs._FocusWhenFirstShown;
+	this->bActivateWhenFirstShown = InArgs._ActivateWhenFirstShown;
+	this->bHasOSWindowBorder = InArgs._UseOSWindowBorder;
+	this->bHasMinimizeButton = InArgs._SupportsMinimize;
+	this->bHasMaximizeButton = InArgs._SupportsMaximize;
+	this->bHasSizingFrame = !InArgs._IsPopupWindow && InArgs._SizingRule == ESizingRule::UserSized;
+	
+	// calculate window size from client size
+	const bool bCreateTitleBar = InArgs._CreateTitleBar && !bIsPopupWindow && !bIsCursorDecoratorWindow && !bHasOSWindowBorder;
+	FVector2D WindowSize = InArgs._ClientSize;
+
+	// Do not adjust the client size if we have an OS border.  
+	if (!HasOSWindowBorder())
+	{
+		const FMargin BorderSize = GetWindowBorderSize();
+
+		WindowSize.X += BorderSize.Left + BorderSize.Right;
+		WindowSize.Y += BorderSize.Bottom + BorderSize.Top;
+
+		if (bCreateTitleBar)
+		{
+			WindowSize.Y += SWindowDefs::DefaultTitleBarSize;
+		}
+	}
+
+	// calculate initial window position
+	FVector2D WindowPosition = InArgs._ScreenPosition;
+
+	AutoCenterRule = InArgs._AutoCenter;
+
+	// Get desktop metrics
+	FDisplayMetrics DisplayMetrics;
+	FSlateApplicationBase::Get().GetDisplayMetrics( DisplayMetrics );
+	const FPlatformRect& VirtualDisplayRect = DisplayMetrics.VirtualDisplayRect;
+	const FPlatformRect& PrimaryDisplayRect = DisplayMetrics.PrimaryDisplayWorkAreaRect;
+
+	// If we're manually positioning the window we need to check if it's outside
+	// of the virtual bounds of the current displays or too large.
+	if ( AutoCenterRule == EAutoCenter::None && InArgs._SaneWindowPlacement )
+	{
+		// Check to see if the upper left corner of the window is outside the virtual
+		// bounds of the display, if so reset to preferred work area
+		if (WindowPosition.X < VirtualDisplayRect.Left ||
+			WindowPosition.X >= VirtualDisplayRect.Right ||
+			WindowPosition.Y < VirtualDisplayRect.Top ||
+			WindowPosition.Y >= VirtualDisplayRect.Bottom)
+		{
+			AutoCenterRule = EAutoCenter::PreferredWorkArea;
+		}
+
+		float PrimaryWidthPadding = DisplayMetrics.PrimaryDisplayWidth - 
+			(PrimaryDisplayRect.Right - 1 - PrimaryDisplayRect.Left);
+		float PrimaryHeightPadding = DisplayMetrics.PrimaryDisplayHeight - 
+			(PrimaryDisplayRect.Bottom - 1 - PrimaryDisplayRect.Top);
+
+		float VirtualWidth = (VirtualDisplayRect.Right - 1 - VirtualDisplayRect.Left);
+		float VirtualHeight = (VirtualDisplayRect.Bottom - 1 - VirtualDisplayRect.Top);
+
+		// Make sure that the window size is no larger than the virtual display area.
+		WindowSize.X = FMath::Clamp(WindowSize.X, 0.0f, VirtualWidth - PrimaryWidthPadding);
+		WindowSize.Y = FMath::Clamp(WindowSize.Y, 0.0f, VirtualHeight - PrimaryHeightPadding);
+	}
+
+	if( AutoCenterRule != EAutoCenter::None )
+	{
+		FSlateRect AutoCenterRect;
+
+		switch( AutoCenterRule )
+		{
+		default:
+		case EAutoCenter::PrimaryWorkArea:
+			AutoCenterRect = FSlateRect(
+				(float)PrimaryDisplayRect.Left, 
+				(float)PrimaryDisplayRect.Top,
+				(float)PrimaryDisplayRect.Right,
+				(float)PrimaryDisplayRect.Bottom );		
+			break;
+		case EAutoCenter::PreferredWorkArea:
+			AutoCenterRect = FSlateApplicationBase::Get().GetPreferredWorkArea();
+			break;
+		}
+
+		// Setup a position and size for the main frame window that's centered in the desktop work area
+		const FVector2D DisplayTopLeft( AutoCenterRect.Left, AutoCenterRect.Top );
+		const FVector2D DisplaySize( AutoCenterRect.Right - AutoCenterRect.Left, AutoCenterRect.Bottom - AutoCenterRect.Top );
+		WindowPosition = DisplayTopLeft + ( DisplaySize - WindowSize ) * 0.5f;
+	}
+
+#if PLATFORM_HTML5 
+	// UE expects mouse coordinates in screen space. SDL/HTML5 canvas provides in client space. 
+	// Anchor the window at the top/left corner to make sure client space coordinates and screen space coordinates match up. 
+	WindowPosition.X =  WindowPosition.Y = 0; 
+#endif 
+	this->InitialDesiredScreenPosition = WindowPosition;
+	this->InitialDesiredSize = WindowSize;
+
+	this->ConstructWindowInternals( bCreateTitleBar );
+	this->SetContent( InArgs._Content.Widget );
+}
+
+
+TSharedRef<SWindow> SWindow::MakeNotificationWindow()
+{
+	TSharedRef<SWindow> NewWindow =
+		SNew(SWindow)
+		.SupportsMaximize( false )
+		.SupportsMinimize( false )
+		.IsPopupWindow( true )
+		.CreateTitleBar( false )
+		.SizingRule( ESizingRule::Autosized )
+		.SupportsTransparency( true )
+		.InitialOpacity( 0.0f )
+		.FocusWhenFirstShown( false )
+		.ActivateWhenFirstShown( false );
+
+	// Notification windows slide open so we'll mark them as resized frequently
+	NewWindow->bSizeWillChangeOften = true;
+	NewWindow->ExpectedMaxWidth = 1024;
+	NewWindow->ExpectedMaxHeight = 256;
+
+	return NewWindow;
+}
+
+
+TSharedRef<SWindow> SWindow::MakeToolTipWindow()
+{
+	TSharedRef<SWindow> NewWindow = SNew( SWindow )
+		.IsPopupWindow( true )
+		.SizingRule( ESizingRule::Autosized )
+		.FocusWhenFirstShown( false )
+		.ActivateWhenFirstShown( false );
+	NewWindow->bIsToolTipWindow = true;
+	NewWindow->bIsTopmostWindow = true;
+	NewWindow->bIsTransparent = true;
+	NewWindow->Opacity = 0.0f;
+
+	// NOTE: These sizes are tweaked for SToolTip widgets (text wrap width of around 400 px)
+	NewWindow->bSizeWillChangeOften = true;
+	NewWindow->ExpectedMaxWidth = 512;
+	NewWindow->ExpectedMaxHeight = 256;
+
+	return NewWindow;
+}
+
+
+TSharedRef<SWindow> SWindow::MakeCursorDecorator()
+{
+	TSharedRef<SWindow> NewWindow = SNew( SWindow )
+		.IsPopupWindow( true )
+		.SizingRule( ESizingRule::Autosized )
+		.FocusWhenFirstShown( false )
+		.ActivateWhenFirstShown( false );
+	NewWindow->bIsToolTipWindow = true;
+	NewWindow->bIsTopmostWindow = true;
+	NewWindow->bIsCursorDecoratorWindow = true;
+	NewWindow->bIsTransparent = true;
+	NewWindow->Opacity = 1.0f;
+
+	return NewWindow;
+}
+
+FVector2D SWindow::ComputeWindowSizeForContent( FVector2D ContentSize )
+{
+	// @todo mainframe: This code should be updated to handle the case where we're spawning a window that doesn't have 
+	//                  a traditional title bar, such as a window that contains a primary SDockingArea.  Currently, the
+	//                  size reported here will be too large!
+	return ContentSize + FVector2D(0, SWindowDefs::DefaultTitleBarSize) + SWindowDefs::WindowBorderSize.GetDesiredSize();
+}
+
+
+void SWindow::ConstructWindowInternals( const bool bCreateTitleBar )
+{
+	ForegroundColor = FCoreStyle::Get().GetSlateColor("DefaultForeground");
+
+	// Setup widget that represents the main area of the window.  That is, everything inside the window's border.
+	TSharedRef< SVerticalBox > MainWindowArea = 
+		SNew( SVerticalBox )
+			.Visibility( EVisibility::SelfHitTestInvisible );
+
+	if (bCreateTitleBar)
+	{
+		// @todo mainframe: Should be measured from actual title bar content widgets.  Don't use a hard-coded size!
+		TitleBarSize = SWindowDefs::DefaultTitleBarSize;
+
+		MainWindowArea->AddSlot()
+			.AutoHeight()
+			[
+				FSlateApplicationBase::Get().MakeWindowTitleBar(SharedThis(this), nullptr, TitleBar)
+			];
+	}
+	else
+	{
+		TitleBarSize = 0;
+	}
+
+	// create window content slot
+	MainWindowArea->AddSlot()
+		.FillHeight(1.0f)
+		.Expose(ContentSlot)
+		[
+			SNullWidget::NullWidget
+		];
+
+	// create window
+	if (!bIsToolTipWindow && !bIsPopupWindow && !bHasOSWindowBorder)
+	{
+		TAttribute<EVisibility> WindowContentVisibility(this, &SWindow::GetWindowContentVisibility);
+		TAttribute<const FSlateBrush*> WindowBackgroundAttr(this, &SWindow::GetWindowBackground);
+		TAttribute<const FSlateBrush*> WindowOutlineAttr(this, &SWindow::GetWindowOutline);
+		TAttribute<FSlateColor> WindowOutlineColorAttr(this, &SWindow::GetWindowOutlineColor);
+
+		this->ChildSlot
+		[
+			SAssignNew(WindowOverlay, SOverlay)
+				.Visibility(EVisibility::SelfHitTestInvisible)
+
+			// window background
+			+ SOverlay::Slot()
+				[
+					FSlateApplicationBase::Get().MakeImage(
+						WindowBackgroundAttr,
+						FLinearColor::White,
+						WindowContentVisibility
+					)
+				]
+
+			// window border
+			+ SOverlay::Slot()
+				[
+					FSlateApplicationBase::Get().MakeImage(
+						&Style->BorderBrush,
+						FLinearColor::White,
+						WindowContentVisibility
+					)
+				]
+
+			// main area
+			+ SOverlay::Slot()
+				[
+					SNew(SVerticalBox)
+						.Visibility(WindowContentVisibility)
+
+					+ SVerticalBox::Slot()
+					
+						.Padding(TAttribute<FMargin>( this, &SWindow::GetWindowBorderSize ))
+						[
+							MainWindowArea
+						]
+				]
+
+			// pop-up layer
+			+ SOverlay::Slot()
+				[
+					SAssignNew(PopupLayer, SPopupLayer, SharedThis(this))
+				]
+
+			// window outline
+			+ SOverlay::Slot()
+				[
+					FSlateApplicationBase::Get().MakeImage(
+						WindowOutlineAttr,
+						WindowOutlineColorAttr,
+						WindowContentVisibility
+					)
+				]
+		];
+	}
+	else if( bHasOSWindowBorder )
+	{
+		this->ChildSlot
+		[
+			SAssignNew(WindowOverlay, SOverlay)
+			+ SOverlay::Slot()
+			[
+					MainWindowArea
+			]
+			+ SOverlay::Slot()
+			[
+				SAssignNew(PopupLayer, SPopupLayer, SharedThis(this))
+			]
+		];
+	}
+}
+
+/** Are any of our child windows active? */
+bool SWindow::HasActiveChildren() const
+{
+	for (int32 i = 0; i < ChildWindows.Num(); ++i)
+	{
+		if ( ChildWindows[i] == FSlateApplicationBase::Get().GetActiveTopLevelWindow() || ChildWindows[i]->HasActiveChildren() )
+		{
+			return true;
+		}
+	}
+
+	return false;
+}
+
+
+void SWindow::Tick( const FGeometry& AllottedGeometry, const double InCurrentTime, const float InDeltaTime )
+{
+	if( Morpher.bIsActive && bHasEverBeenDrawn )
+	{
+		if(Morpher.bIsPendingPlay)
+		{
+			Morpher.Sequence.Play();
+			Morpher.bIsPendingPlay = false;
+		}
+		if ( Morpher.Sequence.IsPlaying() )
+		{
+			const float InterpAlpha = Morpher.Sequence.GetLerp();
+
+			if( Morpher.bIsAnimatingWindowSize )
+			{
+				FSlateRect WindowRect = FMath::Lerp( Morpher.StartingMorphShape, Morpher.TargetMorphShape, InterpAlpha );
+				if( WindowRect != GetRectInScreen() )
+				{
+					check( SizingRule != ESizingRule::Autosized );
+					this->ReshapeWindow( WindowRect );
+				}
+			}
+			else
+			{
+				const FVector2D StartPosition( Morpher.StartingMorphShape.Left, Morpher.StartingMorphShape.Top );
+				const FVector2D TargetPosition( Morpher.TargetMorphShape.Left, Morpher.TargetMorphShape.Top );
+				const FVector2D NewPosition( FMath::Lerp( StartPosition, TargetPosition, InterpAlpha ) );
+				if( NewPosition != this->GetPositionInScreen() )
+				{
+					this->MoveWindowTo( NewPosition );
+				}
+			}
+
+			const float NewOpacity = FMath::Lerp( Morpher.StartingOpacity, Morpher.TargetOpacity, InterpAlpha );
+			this->SetOpacity( NewOpacity );
+		}
+		else
+		{
+			if( Morpher.bIsAnimatingWindowSize )
+			{
+				if( Morpher.TargetMorphShape != GetRectInScreen() )
+				{
+					check( SizingRule != ESizingRule::Autosized );
+					this->ReshapeWindow( Morpher.TargetMorphShape );
+				}
+			}
+			else
+			{
+				const FVector2D TargetPosition( Morpher.TargetMorphShape.Left, Morpher.TargetMorphShape.Top );
+				if( TargetPosition != this->GetPositionInScreen() )
+				{
+					this->MoveWindowTo( TargetPosition );
+				}
+			}
+
+			this->SetOpacity( Morpher.TargetOpacity );
+			Morpher.bIsActive = false;
+		}
+	}
+}
+
+FVector2D SWindow::GetInitialDesiredSizeInScreen() const
+{
+	return InitialDesiredSize;
+}
+
+FVector2D SWindow::GetInitialDesiredPositionInScreen() const
+{
+	return InitialDesiredScreenPosition;
+}
+
+FGeometry SWindow::GetWindowGeometryInScreen() const
+{
+	const float AppScale = FSlateApplicationBase::Get().GetApplicationScale();
+	return FGeometry( ScreenPosition/AppScale, FVector2D::ZeroVector, Size/AppScale, AppScale );
+}
+
+FGeometry SWindow::GetWindowGeometryInWindow() const
+{
+	const float AppScale = FSlateApplicationBase::Get().GetApplicationScale();
+	return FGeometry( FVector2D::ZeroVector, FVector2D::ZeroVector, Size/AppScale, AppScale );
+}
+
+FVector2D SWindow::GetPositionInScreen() const
+{
+	return ScreenPosition;
+}
+
+FVector2D SWindow::GetSizeInScreen() const
+{
+	return Size;
+}
+
+FSlateRect SWindow::GetNonMaximizedRectInScreen() const
+{
+	int X = 0;
+	int Y = 0;
+	int Width = 0;
+	int Height = 0;
+	
+	if ( NativeWindow->GetRestoredDimensions(X, Y, Width, Height) )
+	{
+		return FSlateRect( X, Y, X+Width, Y+Height );
+	}
+	else
+	{
+		return GetRectInScreen();
+	}
+}
+
+FSlateRect SWindow::GetRectInScreen() const
+{ 
+	return FSlateRect( ScreenPosition.X, ScreenPosition.Y, ScreenPosition.X + Size.X, ScreenPosition.Y + Size.Y );
+}
+
+FVector2D SWindow::GetClientSizeInScreen() const
+{
+	if (HasOSWindowBorder())
+	{
+		return Size;
+	}
+
+	FVector2D ClientSize = Size;
+	FMargin BorderSize = GetWindowBorderSize();
+
+	ClientSize.X -= (BorderSize.Left + BorderSize.Right);
+	ClientSize.Y -= (BorderSize.Top + BorderSize.Bottom + TitleBarSize);
+
+	return ClientSize;
+}
+
+FSlateRect SWindow::GetClippingRectangleInWindow() const
+{
+	return FSlateRect( 0, 0, Size.X, Size.Y );
+}
+
+
+FMargin SWindow::GetWindowBorderSize() const
+{
+// Mac didn't want a window border, and consoles don't either, so only do this in Windows
+#if PLATFORM_WINDOWS
+	if (NativeWindow.IsValid() && NativeWindow->IsMaximized())
+	{
+		int32 OSWindowBorderSize = NativeWindow->GetWindowBorderSize();
+		return FMargin(OSWindowBorderSize);
+	}
+
+	return SWindowDefs::WindowBorderSize;
+#else
+	return FMargin();
+#endif
+}
+
+
+void SWindow::MoveWindowTo( FVector2D NewPosition )
+{
+	if (NativeWindow.IsValid())
+	{
+		NativeWindow->MoveWindowTo( FMath::Trunc(NewPosition.X), FMath::Trunc(NewPosition.Y) );
+	}
+	else
+	{
+		InitialDesiredScreenPosition = NewPosition;
+	}
+}
+
+void SWindow::ReshapeWindow( FVector2D NewPosition, FVector2D NewSize )
+{
+	if (NativeWindow.IsValid())
+	{
+		NativeWindow->ReshapeWindow( FMath::Trunc(NewPosition.X), FMath::Trunc(NewPosition.Y), FMath::Trunc(NewSize.X), FMath::Trunc(NewSize.Y) );
+	}
+	else
+	{
+		InitialDesiredScreenPosition = NewPosition;
+		InitialDesiredSize = NewSize;
+	}
+
+	SetCachedSize( NewSize );
+}
+
+void SWindow::ReshapeWindow( const FSlateRect& InNewShape )
+{
+	ReshapeWindow( FVector2D(InNewShape.Left, InNewShape.Top), FVector2D( InNewShape.Right - InNewShape.Left,  InNewShape.Bottom - InNewShape.Top) );
+}
+
+void SWindow::Resize( FVector2D NewSize )
+{
+	Morpher.Sequence.JumpToEnd();
+	if ( Size != NewSize )
+	{
+		if (NativeWindow.IsValid())
+		{
+			NativeWindow->ReshapeWindow( FMath::Trunc(ScreenPosition.X), FMath::Trunc(ScreenPosition.Y), FMath::Trunc(NewSize.X), FMath::Trunc(NewSize.Y) );
+		}
+		else
+		{
+			InitialDesiredSize = NewSize;
+		}
+	}
+}
+
+FSlateRect SWindow::GetFullScreenInfo() const
+{
+	if (NativeWindow.IsValid())
+	{
+		int32 X;
+		int32 Y;
+		int32 Width;
+		int32 Height;
+
+		if ( NativeWindow->GetFullScreenInfo( X, Y, Width, Height ) )
+		{
+			return FSlateRect( X, Y, X + Width, Y + Height );
+		}
+	}
+
+	return FSlateRect();
+}
+
+void SWindow::SetCachedScreenPosition(FVector2D NewPosition)
+{
+	ScreenPosition = NewPosition;
+	OnWindowMoved.ExecuteIfBound( SharedThis( this ) );
+}
+
+void SWindow::SetCachedSize( FVector2D NewSize )
+{
+	if( NativeWindow.IsValid() )
+	{
+		NativeWindow->AdjustCachedSize( NewSize );
+	}
+	Size = NewSize;
+}
+
+bool SWindow::IsMorphing() const
+{
+	return Morpher.bIsActive && Morpher.Sequence.IsPlaying();
+}
+
+bool SWindow::IsMorphingSize() const
+{
+	return IsMorphing() && Morpher.bIsAnimatingWindowSize;
+}
+
+
+void SWindow::MorphToPosition( const FCurveSequence& Sequence, const float TargetOpacity, const FVector2D& TargetPosition )
+{
+	Morpher.bIsAnimatingWindowSize = false;
+	Morpher.Sequence = Sequence;
+	Morpher.TargetOpacity = TargetOpacity;
+	UpdateMorphTargetPosition( TargetPosition );
+	StartMorph();
+}
+
+
+void SWindow::MorphToShape( const FCurveSequence& Sequence, const float TargetOpacity, const FSlateRect& TargetShape )
+{
+	Morpher.bIsAnimatingWindowSize = true;
+	Morpher.Sequence = Sequence;
+	Morpher.TargetOpacity = TargetOpacity;
+	UpdateMorphTargetShape(TargetShape);
+	StartMorph();
+}
+
+void SWindow::StartMorph()
+{
+	Morpher.StartingOpacity = GetOpacity();
+	Morpher.StartingMorphShape = FSlateRect( this->ScreenPosition.X, this->ScreenPosition.Y, this->ScreenPosition.X + this->Size.X, this->ScreenPosition.Y + this->Size.Y );
+	Morpher.bIsPendingPlay = true;
+	Morpher.bIsActive = true;
+	Morpher.Sequence.JumpToStart();
+}
+
+const FSlateBrush* SWindow::GetWindowBackground() const
+{
+	return WindowBackground;
+}
+
+const FSlateBrush* SWindow::GetWindowOutline() const
+{
+	return &Style->OutlineBrush;
+}
+
+FSlateColor SWindow::GetWindowOutlineColor() const
+{
+	return Style->OutlineColor;
+}
+
+void SWindow::UpdateMorphTargetShape( const FSlateRect& TargetShape )
+{
+	Morpher.TargetMorphShape = TargetShape;
+}
+
+void SWindow::UpdateMorphTargetPosition( const FVector2D& TargetPosition )
+{
+	Morpher.TargetMorphShape.Left = Morpher.TargetMorphShape.Right = TargetPosition.X;
+	Morpher.TargetMorphShape.Top = Morpher.TargetMorphShape.Bottom = TargetPosition.Y;
+}
+
+FVector2D SWindow::GetMorphTargetPosition() const
+{
+	return FVector2D( Morpher.TargetMorphShape.Left, Morpher.TargetMorphShape.Top );
+}
+
+
+FSlateRect SWindow::GetMorphTargetShape() const
+{
+	return Morpher.TargetMorphShape;
+}
+
+void SWindow::FlashWindow()
+{
+	if (TitleBar.IsValid())
+	{
+		TitleBar->Flash();
+	}
+}
+
+void SWindow::BringToFront( bool bForce )
+{
+	if (NativeWindow.IsValid())
+	{
+		NativeWindow->BringToFront( bForce );
+	}
+}
+
+void SWindow::HACK_ForceToFront()
+{
+	if (NativeWindow.IsValid())
+	{
+		NativeWindow->HACK_ForceToFront();
+	}
+}
+
+TSharedPtr<FGenericWindow> SWindow::GetNativeWindow()
+{
+	return NativeWindow;
+}
+
+TSharedPtr<const FGenericWindow> SWindow::GetNativeWindow() const
+{
+	return NativeWindow;
+} 
+
+bool SWindow::IsDescendantOf( const TSharedPtr<SWindow>& ParentWindow ) const
+{
+	TSharedPtr<SWindow> CandidateToCheck = this->GetParentWindow();
+	
+	// Keep checking our parent until we get to the root of the tree or find the window we were looking for.
+	while (CandidateToCheck.IsValid())
+	{
+		if (CandidateToCheck == ParentWindow)
+		{
+			// One of our ancestor windows is the ParentWindow we were looking for!
+			return true;
+		}
+
+		// Consider the next ancestor
+		CandidateToCheck = CandidateToCheck->GetParentWindow();
+	}
+
+	return false;
+}
+
+void SWindow::SetNativeWindow( TSharedRef<FGenericWindow> InNativeWindow )
+{
+	check( ! NativeWindow.IsValid() );
+	NativeWindow = InNativeWindow;
+}
+
+void SWindow::SetContent( TSharedRef<SWidget> InContent )
+{
+	if ( bIsPopupWindow || bIsCursorDecoratorWindow )
+	{
+		this->ChildSlot.operator[]( InContent );
+	}
+	else
+	{
+		this->ContentSlot->operator[]( InContent );
+	}
+}
+
+TSharedRef<const SWidget> SWindow::GetContent() const
+{
+	if ( bIsPopupWindow || bIsCursorDecoratorWindow )
+	{
+		return this->ChildSlot.GetChildAt(0);
+	}
+	else
+	{
+		return this->ContentSlot->Widget;
+	}
+}
+
+SOverlay::FOverlaySlot& SWindow::AddOverlaySlot( const int32 ZOrder )
+{
+	if(!WindowOverlay.IsValid())
+	{
+		ensureMsg( false, TEXT("This window does not support overlays. The added slot will not be visible!") );
+		WindowOverlay = SNew(SOverlay).Visibility( EVisibility::HitTestInvisible );
+	}
+
+	return WindowOverlay->AddSlot(ZOrder);
+}
+
+void SWindow::RemoveOverlaySlot( const TSharedRef<SWidget>& InContent )
+{
+	if(WindowOverlay.IsValid())
+	{
+		WindowOverlay->RemoveSlot( InContent );
+	}
+}
+
+/** Return a new slot in the popup layer. Assumes that the window has a popup layer. */
+FPopupLayerSlot& SWindow::AddPopupLayerSlot()
+{
+	ensure( PopupLayer.IsValid() );
+	return PopupLayer->AddSlot();
+}
+
+/** Counterpart to AddPopupLayerSlot */
+void SWindow::RemovePopupLayerSlot( const TSharedRef<SWidget>& WidgetToRemove )
+{
+	PopupLayer->RemoveSlot( WidgetToRemove );
+}
+
+/** @return should this window show up in the taskbar */
+bool SWindow::AppearsInTaskbar() const
+{
+	return !bIsPopupWindow && !bIsToolTipWindow && !bIsCursorDecoratorWindow;
+}
+
+void SWindow::SetOnWindowDeactivated( const FOnWindowDeactivated& InDelegate )
+{
+	OnWindowDeactivated = InDelegate;
+}
+
+/** Sets the delegate to execute right before the window is closed */
+void SWindow::SetOnWindowClosed( const FOnWindowClosed& InDelegate )
+{
+	OnWindowClosed = InDelegate;
+}
+
+/** Sets the delegate to execute right after the window has been moved */
+void SWindow::SetOnWindowMoved( const FOnWindowMoved& InDelegate)
+{
+	OnWindowMoved = InDelegate;
+}
+
+/** Sets the delegate to override RequestDestroyWindow */
+void SWindow::SetRequestDestroyWindowOverride( const FRequestDestroyWindowOverride& InDelegate )
+{
+	RequestDestroyWindowOverride = InDelegate;
+}
+
+/** Request that this window be destroyed. The window is not destroyed immediately. Instead it is placed in a queue for destruction on next Tick */
+void SWindow::RequestDestroyWindow()
+{
+	if( RequestDestroyWindowOverride.IsBound() )
+	{
+		RequestDestroyWindowOverride.Execute( SharedThis(this) );
+	}
+	else
+	{
+		FSlateApplicationBase::Get().RequestDestroyWindow( SharedThis(this) );
+	}
+}
+
+/** Warning: use Request Destroy Window whenever possible!  This method destroys the window immediately! */
+void SWindow::DestroyWindowImmediately()
+{
+	check( NativeWindow.IsValid() );
+
+	// Destroy the native window
+	NativeWindow->Destroy();
+}
+
+/** Calls the OnWindowClosed delegate when this window is about to be closed */
+void SWindow::NotifyWindowBeingDestroyed()
+{
+	OnWindowClosed.ExecuteIfBound( SharedThis( this ) );
+}
+
+/** Make the window visible */
+void SWindow::ShowWindow()
+{
+	// Make sure the viewport is setup for this window
+	if( !bHasEverBeenShown )
+	{
+		if( ensure( NativeWindow.IsValid() ) )
+		{
+			// We can only create a viewport after the window has been shown (otherwise the swap chain creation may fail)
+			FSlateApplicationBase::Get().GetRenderer()->CreateViewport( SharedThis( this ) );
+		}
+
+		// Auto sized windows don't know their size until after their position is set.
+		// Repositioning the window on show with the new size solves this.
+		if ( SizingRule == ESizingRule::Autosized && AutoCenterRule != EAutoCenter::None )
+		{
+			SlatePrepass();
+			ReshapeWindow( InitialDesiredScreenPosition - (GetDesiredSize() * 0.5f), GetDesiredSize() );
+		}
+
+		// Set the window to be maximized if we need to.  Note that this won't actually show the window if its not
+		// already shown.
+		InitialMaximize();
+	}
+
+	bHasEverBeenShown = true;
+
+	if (NativeWindow.IsValid())
+	{
+		NativeWindow->Show();
+
+		// If this is a tompost window (like a tooltip), make sure that its always rendered top most
+		if( IsTopmostWindow() )
+		{
+			NativeWindow->BringToFront();
+		}
+	}
+}
+
+/** Make the window invisible */
+void SWindow::HideWindow()
+{
+	if (NativeWindow.IsValid())
+	{
+		NativeWindow->Hide();
+	}
+}
+
+void SWindow::EnableWindow( bool bEnable )
+{
+	NativeWindow->Enable( bEnable );
+
+	for( int32 ChildIndex = 0; ChildIndex < ChildWindows.Num(); ++ChildIndex )
+	{
+		ChildWindows[ChildIndex]->EnableWindow( bEnable );
+	}
+}
+
+
+/** @return true if the window is visible, false otherwise*/
+bool SWindow::IsVisible() const
+{
+	return NativeWindow.IsValid() && NativeWindow->IsVisible();
+}
+
+bool SWindow::IsWindowMaximized() const
+{
+	return NativeWindow->IsMaximized();
+}
+
+/** Maximize the window if bInitiallyMaximized is set */
+void SWindow::InitialMaximize()
+{
+	if (NativeWindow.IsValid() && bInitiallyMaximized)
+	{
+		NativeWindow->Maximize();
+	}
+}
+
+/**
+ * Sets the opacity of this window
+ *
+ * @param	InOpacity	The new window opacity represented as a floating point scalar
+ */
+void SWindow::SetOpacity( const float InOpacity )
+{
+	if( Opacity != InOpacity )
+	{
+		check( NativeWindow.IsValid() );
+		Opacity = InOpacity;
+		NativeWindow->SetOpacity( Opacity );
+	}
+}
+
+
+/** @return the window's current opacity */
+float SWindow::GetOpacity() const
+{
+	return Opacity;
+}
+
+bool SWindow::SupportsTransparency() const
+{
+	return bIsTransparent;
+}
+
+
+/** @return A String representation of the widget */
+FString SWindow::ToString() const
+{
+	return FString::Printf( *NSLOCTEXT("SWindow", "Window_Title", " Window : %s ").ToString(), *GetTitle().ToString() );
+}
+
+/** @return true if the window should be activated when first shown */
+bool SWindow::ActivateWhenFirstShown() const
+{
+	return bActivateWhenFirstShown;
+}
+
+/** @return true if the window accepts input; false if the window is non-interactive */
+bool SWindow::AcceptsInput() const
+{
+	return !bIsCursorDecoratorWindow && !bIsToolTipWindow;
+}
+
+/** @return true if the user decides the size of the window; false if the content determines the size of the window */
+bool SWindow::IsUserSized() const
+{
+	return SizingRule == ESizingRule::UserSized;
+}
+
+bool SWindow::IsAutosized() const
+{
+	return SizingRule == ESizingRule::Autosized;
+}
+
+void SWindow::SetSizingRule( ESizingRule::Type InSizingRule )
+{
+	SizingRule = InSizingRule;
+}
+
+/** @return true if this is a vanilla window, or one being used for some special purpose: e.g. tooltip or menu */
+bool SWindow::IsRegularWindow() const
+{
+	return !bIsPopupWindow && !bIsToolTipWindow && !bIsCursorDecoratorWindow;
+}
+
+/** @return true if the window should be on top of all other windows; false otherwise */
+bool SWindow::IsTopmostWindow() const
+{
+	return bIsTopmostWindow;
+}
+
+/** @return true if mouse coordinates is within this window */
+bool SWindow::IsScreenspaceMouseWithin(FVector2D ScreenspaceMouseCoordinate) const
+{
+	const FVector2D LocalMouseCoordinate = GetWindowGeometryInScreen().AbsoluteToLocal(ScreenspaceMouseCoordinate);
+	return NativeWindow->IsPointInWindow(FMath::Trunc(LocalMouseCoordinate.X), FMath::Trunc(LocalMouseCoordinate.Y));
+}
+
+/** @return true if this is a user-sized window with a thick edge */
+bool SWindow::HasSizingFrame() const
+{
+	return bHasSizingFrame;
+}
+
+/** @return true if this window has a maximize button/box on the titlebar area */
+bool SWindow::HasMaximizeBox() const
+{
+	return bHasMaximizeButton;
+}
+
+/** @return true if this window has a minimize button/box on the titlebar area */
+bool SWindow::HasMinimizeBox() const
+{
+	return bHasMinimizeButton;
+}
+
+FCursorReply SWindow::OnCursorQuery( const FGeometry& MyGeometry, const FPointerEvent& CursorEvent ) const
+{
+#if !PLATFORM_MAC // On Mac we depend on system's window resizing
+	if (bHasSizingFrame)
+	{
+		if (WindowZone == EWindowZone::TopLeftBorder || WindowZone == EWindowZone::BottomRightBorder)
+		{
+			return FCursorReply::Cursor(EMouseCursor::ResizeSouthEast);
+		}
+		else if (WindowZone == EWindowZone::BottomLeftBorder || WindowZone == EWindowZone::TopRightBorder)
+		{
+			return FCursorReply::Cursor(EMouseCursor::ResizeSouthWest);
+		}
+		else if (WindowZone == EWindowZone::TopBorder || WindowZone == EWindowZone::BottomBorder)
+		{
+			return FCursorReply::Cursor(EMouseCursor::ResizeUpDown);
+		}
+		else if (WindowZone == EWindowZone::LeftBorder || WindowZone == EWindowZone::RightBorder)
+		{
+			return FCursorReply::Cursor(EMouseCursor::ResizeLeftRight);
+		}
+	}
+#endif
+	return FCursorReply::Unhandled();
+}
+
+bool SWindow::OnIsActiveChanged( const FWindowActivateEvent& ActivateEvent )
+{
+	const bool bWasDeactivated = ( ActivateEvent.GetActivationType() == FWindowActivateEvent::EA_Deactivate );
+	if ( bWasDeactivated )
+	{
+		OnWindowDeactivated.ExecuteIfBound();
+
+		const EWindowMode::Type WindowMode = GetWindowMode();
+		// If the window is not fullscreen, we do not want to automatically recapture the mouse unless an external UI such as Steam is open. Fullscreen windows we do.
+		if( WindowMode != EWindowMode::Fullscreen && WidgetToFocusOnActivate.IsValid() && WidgetToFocusOnActivate.Pin()->HasMouseCapture() && !FSlateApplicationBase::Get().IsExternalUIOpened())
+		{
+			//For a windowed application with an OS border, if the user is giving focus back to the application by clicking on the close/(X) button, then we must clear 
+			//the weak pointer to WidgetToFocus--so that the application's main viewport does not steal focus immediately (thus canceling the close attempt).
+			
+			//This change introduces a different bug where slate context is lost when closing popup menus.  However, this issue is negated by a 
+			//change to FMenuStack::PushMenu, where we ReleaseMouseCapture when immediately shifting focus.
+			WidgetToFocusOnActivate.Reset();
+		}
+	}
+
+	return true;
+}
+
+
+void SWindow::Maximize()
+{
+	if (NativeWindow.IsValid())
+	{
+		NativeWindow->Maximize();
+	}
+}
+
+void SWindow::Restore()
+{
+	if (NativeWindow.IsValid())
+	{
+		NativeWindow->Restore();
+	}
+}
+
+void SWindow::Minimize()
+{
+	if (NativeWindow.IsValid())
+	{
+		NativeWindow->Minimize();
+	}
+}
+
+int32 SWindow::GetCornerRadius()
+{
+	return IsRegularWindow() ? SWindowDefs::CornerRadius : 0;
+}
+
+bool SWindow::SupportsKeyboardFocus() const
+{
+	return !bIsToolTipWindow && !bIsCursorDecoratorWindow;
+}
+
+FReply SWindow::OnKeyboardFocusReceived( const FGeometry& MyGeometry, const FKeyboardFocusEvent& InKeyboardFocusEvent )
+{
+	// If we're becoming active and we were set to restore keyboard focus to a specific widget
+	// after reactivating, then do so now
+	TSharedPtr< SWidget > PinnedWidgetToFocusOnActivate( WidgetToFocusOnActivate.Pin() );
+
+	if( PinnedWidgetToFocusOnActivate.IsValid() && 
+		( InKeyboardFocusEvent.GetCause() == EKeyboardFocusCause::WindowActivate  || FSlateApplicationBase::Get().IsExternalUIOpened() ) )
+	{
+		TArray< TSharedRef<SWindow> > JustThisWindow;
+		{
+			JustThisWindow.Add( SharedThis(this) );
+		}
+
+		FWidgetPath WidgetToFocusPath;
+		if( FSlateApplicationBase::Get().FindPathToWidgetVirtual( JustThisWindow, PinnedWidgetToFocusOnActivate.ToSharedRef(), WidgetToFocusPath ) )
+		{
+			FSlateApplicationBase::Get().SetKeyboardFocus( WidgetToFocusPath, EKeyboardFocusCause::SetDirectly );
+		}
+	}
+
+	return FReply::Handled();
+}
+
+FReply SWindow::OnMouseButtonDown( const FGeometry& MyGeometry, const FPointerEvent& MouseEvent )
+{
+	if (bDragAnywhere && MouseEvent.GetEffectingButton() == EKeys::LeftMouseButton)
+	{
+		return FReply::Handled().CaptureMouse(SharedThis(this));
+	}
+	else
+	{
+		return FReply::Unhandled();
+	}
+}
+
+FReply SWindow::OnMouseButtonUp( const FGeometry& MyGeometry, const FPointerEvent& MouseEvent )
+{
+	if (bDragAnywhere &&  this->HasMouseCapture() && MouseEvent.GetEffectingButton() == EKeys::LeftMouseButton)
+	{
+		return FReply::Handled().ReleaseMouseCapture();
+	}
+	else
+	{
+		return FReply::Unhandled();
+	}
+}
+
+FReply SWindow::OnMouseMove( const FGeometry& MyGeometry, const FPointerEvent& MouseEvent )
+{
+	if ( bDragAnywhere && this->HasMouseCapture() && MouseEvent.IsMouseButtonDown(EKeys::LeftMouseButton) )
+	{
+		this->MoveWindowTo( ScreenPosition + MouseEvent.GetCursorDelta() );
+		return FReply::Handled();
+	}
+	else
+	{
+		return FReply::Unhandled();
+	}
+
+}
+
+FVector2D SWindow::ComputeDesiredSize() const
+{
+	const float Scale = FSlateApplicationBase::Get().GetApplicationScale();
+	return SCompoundWidget::ComputeDesiredSize() * Scale;
+}
+
+const TArray< TSharedRef<SWindow> >& SWindow::GetChildWindows() const
+{
+	return ChildWindows;
+}
+
+TArray< TSharedRef<SWindow> >& SWindow::GetChildWindows()
+{
+	return ChildWindows;
+}
+
+void SWindow::AddChildWindow( const TSharedRef<SWindow>& ChildWindow )
+{
+	TSharedPtr<SWindow> PreviousParent = ChildWindow->ParentWindowPtr.Pin();
+	if (PreviousParent.IsValid())
+	{
+		// This child already had a parent, so we are actually re-parenting it
+		const bool bRemovedSuccessfully = PreviousParent->RemoveDescendantWindow(ChildWindow);
+		check(bRemovedSuccessfully);
+	}
+
+	ChildWindow->ParentWindowPtr = SharedThis(this);
+	ChildWindow->WindowBackground = &Style->ChildBackgroundBrush;
+
+	FSlateApplicationBase::Get().ArrangeWindowToFrontVirtual( ChildWindows, ChildWindow );
+}
+
+TSharedPtr<SWindow> SWindow::GetParentWindow() const
+{
+	return ParentWindowPtr.Pin();
+}
+
+
+TSharedPtr<SWindow> SWindow::GetTopmostAncestor()
+{
+	TSharedPtr<SWindow> TopmostParentSoFar = SharedThis(this);
+	while ( TopmostParentSoFar->ParentWindowPtr.IsValid() )
+	{
+		TopmostParentSoFar = TopmostParentSoFar->ParentWindowPtr.Pin();
+	}
+
+	return TopmostParentSoFar;
+}
+
+bool SWindow::RemoveDescendantWindow( const TSharedRef<SWindow>& DescendantToRemove )
+{
+	const bool bRemoved = 0 != ChildWindows.Remove(DescendantToRemove);
+
+	for ( int32 ChildIndex=0; ChildIndex < ChildWindows.Num(); ++ChildIndex )
+	{
+		TSharedRef<SWindow>& ChildWindow = ChildWindows[ChildIndex];
+		if ( ChildWindow->RemoveDescendantWindow( DescendantToRemove ))
+		{
+			// Reset to the non-child background style
+			ChildWindow->WindowBackground = &Style->BackgroundBrush;
+			return true;
+		}
+	}
+
+	return false;
+}
+
+void SWindow::SetOnWorldSwitchHack( FOnSwitchWorldHack& InOnSwitchWorldHack )
+{
+	OnWorldSwitchHack = InOnSwitchWorldHack;
+}
+
+int32 SWindow::SwitchWorlds( int32 WorldId ) const
+{
+	return OnWorldSwitchHack.IsBound() ? OnWorldSwitchHack.Execute( WorldId ) : false;
+}
+
+bool PointWithinSlateRect(const FVector2D& Point, const FSlateRect& Rect)
+{
+	return Point.X >= Rect.Left && Point.X < Rect.Right &&
+		Point.Y >= Rect.Top && Point.Y < Rect.Bottom;
+}
+
+EWindowZone::Type SWindow::GetCurrentWindowZone(FVector2D LocalMousePosition)
+{
+	// Don't allow position/resizing of window while in fullscreen mode by ignoring Title Bar/Border Zones
+	if ( GetWindowMode() == EWindowMode::WindowedFullscreen || GetWindowMode() == EWindowMode::Fullscreen )
+	{
+		return EWindowZone::ClientArea;
+	}
+	else if(LocalMousePosition.X >= 0 && LocalMousePosition.X < Size.X &&
+			LocalMousePosition.Y >= 0 && LocalMousePosition.Y < Size.Y)
+	{
+		int32 Row = 1;
+		int32 Col = 1;
+		if (SizingRule == ESizingRule::UserSized)
+		{
+			if (LocalMousePosition.X < SWindowDefs::HitResultBorderSize.Left)
+			{
+				Col = 0;
+			}
+			else if (LocalMousePosition.X >= Size.X - SWindowDefs::HitResultBorderSize.Right)
+			{
+				Col = 2;
+			}
+			if (LocalMousePosition.Y < SWindowDefs::HitResultBorderSize.Top)
+			{
+				Row = 0;
+			}
+			else if (LocalMousePosition.Y >= Size.Y - SWindowDefs::HitResultBorderSize.Bottom)
+			{
+				Row = 2;
+			}
+
+			// The actual border is smaller than the hit result zones
+			// This grants larger corner areas to grab onto
+			bool bInBorder =	LocalMousePosition.X < SWindowDefs::WindowBorderSize.Left ||
+								LocalMousePosition.X >= Size.X - SWindowDefs::WindowBorderSize.Right ||
+								LocalMousePosition.Y < SWindowDefs::WindowBorderSize.Top ||
+								LocalMousePosition.Y >= Size.Y - SWindowDefs::WindowBorderSize.Bottom;
+
+			if (!bInBorder)
+			{
+				Row = 1;
+				Col = 1;
+			}
+		}
+
+		static const EWindowZone::Type TypeZones[3][3] = 
+		{
+			{EWindowZone::TopLeftBorder,		EWindowZone::TopBorder,		EWindowZone::TopRightBorder},
+			{EWindowZone::LeftBorder,			EWindowZone::ClientArea,	EWindowZone::RightBorder},
+			{EWindowZone::BottomLeftBorder,		EWindowZone::BottomBorder,	EWindowZone::BottomRightBorder},
+		};
+
+		EWindowZone::Type InZone = TypeZones[Row][Col];
+		if (InZone == EWindowZone::ClientArea)
+		{
+			// Hittest to see if the widget under the mouse should be treated as a title bar (i.e. should move the window)
+			TArray< TSharedRef<SWindow> > ThisWindow;
+			{
+				ThisWindow.Add( SharedThis(this) );
+			}
+
+			FWidgetPath HitTestResults = FSlateApplicationBase::Get().LocateWindowUnderMouse(FSlateApplicationBase::Get().GetCursorPos(), ThisWindow);
+			if( HitTestResults.Widgets.Num() > 0 )
+			{
+				const EWindowZone::Type ZoneOverride = HitTestResults.Widgets.Last().Widget->GetWindowZoneOverride();
+				if( ZoneOverride != EWindowZone::Unspecified )
+				{
+					// The widget overrode the window zone
+					InZone = ZoneOverride;
+				}
+				else if( HitTestResults.Widgets.Last().Widget == AsShared() )
+				{
+					// The window itself was hit, so check for a traditional title bar
+					if( (LocalMousePosition.Y - SWindowDefs::WindowBorderSize.Top) < TitleBarSize )
+					{
+						InZone = EWindowZone::TitleBar;
+					}
+				}
+			}
+		}
+
+		WindowZone = InZone;
+	}
+	else
+	{
+		WindowZone = EWindowZone::NotInWindow;
+	}
+	return WindowZone;
+}
+
+/**
+ * Default constructor. Protected because SWindows must always be used via TSharedPtr. Instead, use FSlateApplication::MakeWindow()
+ */
+SWindow::SWindow()
+	: Opacity( 1.0f )
+	, SizingRule( ESizingRule::UserSized )
+	, bIsTransparent( false )
+	, bIsPopupWindow( false )
+	, bIsToolTipWindow( false )
+	, bIsTopmostWindow( false )
+	, bSizeWillChangeOften( false )
+	, bIsCursorDecoratorWindow( false )
+	, bInitiallyMaximized( false )
+	, bHasEverBeenShown( false )
+	, bHasEverBeenDrawn( false )
+	, bFocusWhenFirstShown(true)
+	, bActivateWhenFirstShown(true)
+	, bHasOSWindowBorder( false )
+	, bHasMinimizeButton( false )
+	, bHasMaximizeButton( false )
+	, bHasSizingFrame( false )
+	, InitialDesiredScreenPosition( FVector2D::ZeroVector )
+	, InitialDesiredSize( FVector2D::ZeroVector )
+	, ScreenPosition( FVector2D::ZeroVector )
+	, PreFullscreenPosition( FVector2D::ZeroVector )
+	, Size( FVector2D::ZeroVector )
+	, TitleBarSize( SWindowDefs::DefaultTitleBarSize )
+	, ContentSlot(nullptr)
+	, Style( &FCoreStyle::Get().GetWidgetStyle<FWindowStyle>("Window") )
+	, WindowBackground( &Style->BackgroundBrush )
+	, bShouldShowWindowContentDuringOverlay( false )
+	, ExpectedMaxWidth( INDEX_NONE )
+	, ExpectedMaxHeight( INDEX_NONE )
+{
+}
+
+
+FOptionalSize SWindow::GetTitleBarSize() const
+{
+	return TitleBarSize;
+}
+
+
+void SWindow::SetFullWindowOverlayContent( TSharedPtr<SWidget> InContent )
+{
+	if( FullWindowOverlayWidget.IsValid() )
+	{
+		// Remove the last slot
+		WindowOverlay->RemoveSlot( FullWindowOverlayWidget.ToSharedRef() );
+		FullWindowOverlayWidget.Reset();
+	}
+
+	if( InContent.IsValid() )
+	{
+		FullWindowOverlayWidget = InContent;
+
+		// Create a slot in our overlay to hold the content
+		WindowOverlay->AddSlot( 1 )
+		[
+			InContent.ToSharedRef()
+		];
+	}
+}
+
+/** Toggle window between fullscreen and normal mode */
+void SWindow::SetWindowMode( EWindowMode::Type NewWindowMode )
+{
+	EWindowMode::Type CurrentWindowMode = NativeWindow->GetWindowMode();
+
+	if( CurrentWindowMode != NewWindowMode )
+	{
+		bool bFullscreen = NewWindowMode != EWindowMode::Windowed;
+
+		bool bWasFullscreen = CurrentWindowMode != EWindowMode::Windowed;
+
+		// We need to store off the screen position when entering fullscreen so that we can move the window back to its original position after leaving fullscreen
+		if( bFullscreen )
+		{
+			PreFullscreenPosition = ScreenPosition;
+		}
+
+		NativeWindow->SetWindowMode( NewWindowMode );
+	
+		FSlateApplicationBase::Get().GetRenderer()->UpdateFullscreenState( SharedThis(this), Size.X, Size.Y );
+
+		if( TitleArea.IsValid() )
+		{
+			// Collapse the Window title bar when switching to Fullscreen
+			TitleArea->SetVisibility( (NewWindowMode == EWindowMode::Fullscreen || NewWindowMode == EWindowMode::WindowedFullscreen ) ? EVisibility::Collapsed : EVisibility::Visible );
+		}
+
+		if( bWasFullscreen )
+		{
+			// If we left fullscreen, reset the screen position;
+			MoveWindowTo(PreFullscreenPosition);
+		}
+	}
+
+}
+
+bool SWindow::HasFullWindowOverlayContent() const
+{
+	return FullWindowOverlayWidget.IsValid();
+}
+
+void SWindow::BeginFullWindowOverlayTransition()
+{
+	bShouldShowWindowContentDuringOverlay = true; 
+}
+
+void SWindow::EndFullWindowOverlayTransition()
+{
+	bShouldShowWindowContentDuringOverlay = false;
+}
+
+EVisibility SWindow::GetWindowContentVisibility() const
+{
+	// The content of the window should be visible unless we have a full window overlay content
+	// in which case the full window overlay content is visible but nothing under it
+	return (bShouldShowWindowContentDuringOverlay == true || !FullWindowOverlayWidget.IsValid()) ? EVisibility::SelfHitTestInvisible : EVisibility::Hidden;
+};
diff --git a/Engine/Source/Runtime/SlateCore/Public/Animation/CurveHandle.h b/Engine/Source/Runtime/SlateCore/Public/Animation/CurveHandle.h
new file mode 100644
--- /dev/null
+++ b/Engine/Source/Runtime/SlateCore/Public/Animation/CurveHandle.h
@@ -0,0 +1,98 @@
+// Copyright 1998-2014 Epic Games, Inc. All Rights Reserved.
+
+/*=============================================================================
+	CurveHandle.h: Declares the FCurveHandle structure.
+=============================================================================*/
+
+#pragma once
+
+
+/**
+ * Types of easing functions for Slate animation curves.  These are used to smooth out animations.
+ */
+namespace ECurveEaseFunction
+{
+	enum Type
+	{
+		/** Linear interpolation, with no easing */
+		Linear,
+
+		/** Quadratic ease in */
+		QuadIn,
+		
+		/** Quadratic ease out */
+		QuadOut,
+
+		/** Quadratic ease in, quadratic ease out */
+		QuadInOut,
+
+		/** Cubic ease in */
+		CubicIn,
+		
+		/** Cubic ease out */
+		CubicOut,
+
+		/** Cubic ease in, cubic ease out */
+		CubicInOut,
+	};
+}
+
+
+/**
+ * A handle to curve within a curve sequence.
+ */
+struct SLATECORE_API FCurveHandle
+{
+	/**
+	 * Creates and initializes a curve handle.
+	 *
+	 * @param InOwnerSequence The curve sequence that owns this handle.
+	 * @param InCurveIndex The index of this handle.
+	 */
+	FCurveHandle( const struct FCurveSequence* InOwnerSequence = nullptr, int32 InCurveIndex = 0 );
+
+public:
+
+	/**
+	 * Gets the linearly interpolated value between 0 and 1 for this curve.
+	 *
+	 * @return Lerp value.
+	 *
+	 * @see GetLerpLooping
+	 */
+	float GetLerp( ) const;
+
+	/**
+	 * Just like GetLerp, but loops forever.
+	 *
+	 * @return Looped lerp value.
+	 *
+	 * @see GetLerp
+	 */
+	float GetLerpLooping( ) const;
+
+	/**
+	 * Checks whether this handle is initialized.
+	 *
+	 * A curve handle is considered initialized if it has an owner sequence.
+	 *
+	 * @return true if initialized, false otherwise.
+	 */
+	bool IsInitialized( ) const
+	{
+		return (OwnerSequence != nullptr);
+	}
+
+public:
+
+	/** Applies animation easing to lerp value */
+	static float ApplyEasing( float Time, ECurveEaseFunction::Type EaseType );
+
+private:
+
+	/** The sequence containing this curve */
+	const struct FCurveSequence* OwnerSequence;
+
+	/** The index of the curve in the Curves array */
+	int32 CurveIndex;
+};
