Slate: Refactored core Slate implementation into SlateCore module in preparation for UMG.

diff --git a/Engine/Source/Runtime/SlateCore/Public/Animation/CurveSequence.h b/Engine/Source/Runtime/SlateCore/Public/Animation/CurveSequence.h
new file mode 100644
--- /dev/null
+++ b/Engine/Source/Runtime/SlateCore/Public/Animation/CurveSequence.h
@@ -0,0 +1,177 @@
+// Copyright 1998-2014 Epic Games, Inc. All Rights Reserved.
+
+/*=============================================================================
+	CurveSequence.h: Declares the FCurveSequence structure.
+=============================================================================*/
+
+#pragma once
+
+
+/**
+ * A sequence of curves that can be used to drive animations.
+ * Each curve within the sequence has a time offset and a duration.
+ * This makes FCurveSequence convenient for crating staggered animations.
+ * e.g.
+ *   // We want to zoom in a widget, and then fade in its contents.
+ *   FCurveHandle ZoomCurve = Sequence.AddCurve( 0, 0.15f );
+ *   FCurveHandle FadeCurve = Sequence.AddCurve( 0.15f, 0.1f );
+ *	 Sequence.Play();
+ */
+struct SLATECORE_API FCurveSequence
+{
+public:
+
+	/** A curve has a time offset and duration.*/
+	struct FSlateCurve
+	{
+		/** Constructor */
+		FSlateCurve( float InStartTime, float InDurationSeconds, const ECurveEaseFunction::Type InEaseFunction )
+			: DurationSeconds(InDurationSeconds)
+			, StartTime(InStartTime)
+			, EaseFunction(InEaseFunction)
+		{
+		}
+
+		/** Length of this animation in seconds */
+		float DurationSeconds;
+
+		/** Start time for this animation */
+		float StartTime;
+
+		/**
+		 * Type of easing function to use for this curve.
+		 * Could be passed it at call site.
+		 */
+		ECurveEaseFunction::Type EaseFunction;
+	};
+
+	/** Default constructor */
+	FCurveSequence( );
+
+	/**
+	 * Construct by adding a single animation curve to this sequence.  Does not provide access to the curve though.
+	 *
+	 * @param InStartTimeSeconds   When to start this curve.
+	 * @param InDurationSeconds    How long this curve lasts.
+	 * @param InEaseFunction       Easing function to use for this curve.  Defaults to Linear.  Use this to smooth out your animation transitions.
+	 *
+	 * @return A FCurveHandle that can be used to get the value of this curve after the animation starts playing.
+	 */
+	FCurveSequence( const float InStartTimeSeconds, const float InDurationSeconds, const ECurveEaseFunction::Type InEaseFunction = ECurveEaseFunction::Linear  );
+
+	/**
+	 * Add a new curve at a given time and offset.
+	 *
+	 * @param InStartTimeSeconds   When to start this curve.
+	 * @param InDurationSeconds    How long this curve lasts.
+	 * @param InEaseFunction       Easing function to use for this curve.  Defaults to Linear.  Use this to smooth out your animation transitions.
+	 *
+	 * @return A FCurveHandle that can be used to get the value of this curve after the animation starts playing.
+	 */
+	FCurveHandle AddCurve( const float InStartTimeSeconds, const float InDurationSeconds, const ECurveEaseFunction::Type InEaseFunction = ECurveEaseFunction::Linear );
+
+	/**
+	 * Add a new curve relative to the current end of the sequence. Makes stacking easier.
+	 * e.g. doing 
+	 *     AddCurveRelative(0,5);
+	 *     AddCurveRelative(0,3);
+	 * Is equivalent to
+	 *     AddCurve(0,5);
+	 *     AddCurve(5,3)
+	 *
+	 * @param InOffset             Offset from the last curve in the sequence.
+	 * @param InDurationSecond     How long this curve lasts.
+	 * @param InEaseFunction       Easing function to use for this curve.  Defaults to Linear.  Use this to smooth out your animation transitions.
+	 */
+	FCurveHandle AddCurveRelative( const float InOffset, const float InDurationSecond, const ECurveEaseFunction::Type InEaseFunction = ECurveEaseFunction::Linear );
+
+	/**
+	 * Start playing this curve sequence
+	 *
+	 * @param	StartAtTime		Specifies a time offset relative to the animation to start at.  Defaults to zero (the actual start of the sequence.)
+	 */
+	void Play( const float StartAtTime = 0.0f );
+
+	/** Reverse the direction of an in-progress animation */
+	void Reverse( );
+
+	/**
+	 * Start playing this curve sequence in reverse
+	 *
+	 * @param	StartAtTime		Specifies a time offset relative to the animation to start at.  Defaults to zero (the actual start of the sequence.)
+	 */
+	void PlayReverse( const float StartAtTime = 0.0f );
+
+	/**
+	 * Checks whether the sequence is currently playing.
+	 *
+	 * @return true if playing, false otherwise.
+	 */
+	bool IsPlaying( ) const;
+
+	/** @return the current time relative to the beginning of the sequence. */
+	float GetSequenceTime( ) const;
+
+	/** @return the current time relative to the beginning of the sequence as if the animation were a looping one. */
+	float GetSequenceTimeLooping( ) const;
+
+	/** @return true if the animation is in reverse */
+	bool IsInReverse( ) const;
+
+	/** @return true if the animation is in forward gear */
+	bool IsForward( ) const;
+
+	/** Jumps immediately to the beginning of the animation sequence */
+	void JumpToStart( );
+
+	/** Jumps immediately to the end of the animation sequence */
+	void JumpToEnd( );
+
+	/** Is the sequence at the start? */
+	bool IsAtStart( ) const;
+
+	/** Is the sequence at the end? */
+	bool IsAtEnd( ) const;
+
+	/**
+	 * For single-curve animations, returns the interpolation alpha for the animation.  If you call this function
+	 * on a sequence with multiple curves, an assertion will trigger.
+	 *
+	 * @return A linearly interpolated value between 0 and 1 for this curve.
+	 */
+	float GetLerp( ) const;
+
+	/**
+	 * For single-curve animations, returns the looping interpolation alpha for the animation.  If you call this
+	 * function on a sequence with multiple curves, an assertion will trigger.
+	 *
+	 * @return A linearly interpolated value between 0 and 1 for this curve.
+	 */
+	float GetLerpLooping( ) const;
+
+	/**
+	 * @param CurveIndex  Index of a curve in the curves array.
+	 *
+	 * @return A curve given the index into the curves array
+	 */
+	const FCurveSequence::FSlateCurve& GetCurve( int32 CurveIndex ) const;
+
+protected:
+
+	/** @param InStartTime  when this curve sequence started playing */
+	void SetStartTime( double InStartTime );
+
+private:
+
+	/** All the curves in this sequence. */
+	TArray<FSlateCurve> Curves;
+
+	/** When the curve started playing. */
+	double StartTime;
+
+	/** How long the entire sequence lasts. */
+	float TotalDuration;
+
+	/** Are we playing the animation in reverse */
+	bool bInReverse;
+};
diff --git a/Engine/Source/Runtime/SlateCore/Public/Application/SlateApplicationBase.h b/Engine/Source/Runtime/SlateCore/Public/Application/SlateApplicationBase.h
new file mode 100644
--- /dev/null
+++ b/Engine/Source/Runtime/SlateCore/Public/Application/SlateApplicationBase.h
@@ -0,0 +1,259 @@
+// Copyright 1998-2014 Epic Games, Inc. All Rights Reserved.
+
+/*=============================================================================
+	SlateApplicationBase.h: Declares the FSlateApplicationBase class.
+=============================================================================*/
+
+#pragma once
+
+
+/**
+ * Interface for window title bars.
+ */
+class IWindowTitleBar
+{
+public:
+
+	virtual void Flash( ) = 0;
+};
+
+
+/**
+ * Base class for Slate applications.
+ *
+ * This class currently serves a temporary workaround for solving SlateCore dependencies to FSlateApplication.
+ * It should probably be removed once FSlateApplication has been refactored into SlateCore.
+ */
+class SLATECORE_API FSlateApplicationBase
+{
+public:
+
+	/**
+	 * Gets the renderer being used to draw this application.
+	 *
+	 * @return The Slate renderer.
+	 */
+	TSharedPtr<FSlateRenderer> GetRenderer( ) const
+	{
+		return Renderer;
+	}
+
+public:
+
+	/**
+	 * Associates a top level Slate Window with a native window and ensures that it is tracked properly by the application.
+	 * Calling this method will cause the window to be displayed (unless specified otherwise), so be sure to associate
+	 * content with the window object you're passing in first!
+	 *
+	 * @param InSlateWindow A SlateWindow to which to add a native window.
+	 * @param bShowImmediately true to show the window, false if you're going to call ShowWindow() yourself later.
+	 *
+	 * @return a reference to the SWindow that was just added.
+	 */
+	virtual TSharedRef<SWindow> AddWindow( TSharedRef<SWindow> InSlateWindow, const bool bShowImmediately = true ) = 0;
+
+	/**
+	 * Reorders an array of windows so the specified window is "brought to the front"
+	 */
+	virtual void ArrangeWindowToFrontVirtual( TArray<TSharedRef<SWindow>>& Windows, const TSharedRef<SWindow>& WindowToBringToFront ) = 0;
+
+	/**
+	 * @todo slate: Remove this method or make it private.
+	 * Searches for the specified widget and generates a full path to it.  Note that this is
+	 * a relatively slow operation!
+	 * 
+	 * @param  InWidget       Widget to generate a path to
+	 * @param  OutWidgetPath  The generated widget path
+	 * @param  VisibilityFilter	Widgets must have this type of visibility to be included the path
+	 *
+	 * @return	True if the widget path was found
+	 */
+	virtual bool FindPathToWidgetVirtual( const TArray<TSharedRef<SWindow>> WindowsToSearch, TSharedRef<const SWidget> InWidget, FWidgetPath& OutWidgetPath, EVisibility VisibilityFilter = EVisibility::Visible ) = 0;
+
+	/**
+	 * Gets the active top-level window.
+	 *
+	 * @return The top level window, or nullptr if no Slate windows are currently active.
+	 */
+	virtual TSharedPtr<SWindow> GetActiveTopLevelWindow() const = 0;
+
+	/**
+	 * Gets the global application icon.
+	 *
+	 * @return The icon.
+	 */
+	virtual const FSlateBrush* GetAppIcon( ) const = 0;
+
+	/**
+	 * Gets the ratio SlateUnit / ScreenPixel.
+	 *
+	 * @return Application scale.
+	 */
+	virtual float GetApplicationScale( ) const = 0;
+
+	/**
+	 * Gets Slate's current cached real time.
+	 *
+	 * This time value is updated every frame right before we tick widgets and is the recommended time value to use
+	 * for UI animations and transitions, as opposed to calling FPlatformTime::Seconds() (which is generally slower)
+	 *
+	 * @return  The current Slate real time in seconds
+	 */
+	virtual const double GetCurrentTime( ) const = 0;
+
+	/**
+	 * Gets the last known position of the cursor.
+	 *
+	 * @return Cursor position.
+	 */
+	virtual FVector2D GetCursorPos( ) const = 0;
+
+	/**
+	 * Gets the size of the cursor..
+	 *
+	 * @return Cursor size.
+	 */
+	virtual FVector2D GetCursorSize( ) const = 0;
+
+	/**
+	 * Gets the application's display metrics.
+	 *
+	 * @param OutDisplayMetrics Will contain the display metrics.
+	 */
+	virtual void GetDisplayMetrics( FDisplayMetrics& OutDisplayMetrics ) const = 0;
+
+	/**
+	 * Gets the widget that currently has keyboard focus, if any.
+	 *
+	 * @return The focused widget, or nullptr if no widget has focus.
+	 */
+	virtual TSharedPtr< SWidget > GetKeyboardFocusedWidget( ) const = 0;
+
+	/** Gets the Widget that currently captures the mouse.
+	 *
+	 * @return The captor widget, or nullptr if no widget captured the mouse.
+	 */
+	virtual TSharedPtr< SWidget > GetMouseCaptor( ) const = 0;
+
+	/**
+	 * Gets the rectangle of the current preferred work area.
+	 *
+	 * @return Area rectangle.
+	 */
+	virtual FSlateRect GetPreferredWorkArea( ) const = 0;
+
+	/**
+	 * Checks whether the specified widget has any descendants which are currently focused.
+	 *
+	 * @param Widget The widget to check.
+	 *
+	 * @return true if any descendants are focused, false otherwise.
+	 */
+	virtual bool HasFocusedDescendants( const TSharedRef< const SWidget >& Widget ) const = 0;
+
+	/**
+	 * Checks whether an UI for external services such as Steam is open.
+	 *
+	 * @return true if an external UI is open, false otherwise.
+	 */
+	virtual bool IsExternalUIOpened( ) = 0;
+
+	/** 
+	 * Given the screen-space coordinate of the mouse cursor, searches for a string of widgets that are under the mouse.
+	 *
+	 * The widgets will be returned with the associated geometry. The first item will always
+	 * be the top-level window while the last item will be the leaf-most widget.
+	 *
+	 * @return The path to the widget.
+	 */
+	virtual FWidgetPath LocateWindowUnderMouse( FVector2D ScreenspaceMouseCoordinate, const TArray< TSharedRef<SWindow > >& Windows, bool bIgnoreEnabledStatus = false ) = 0;
+
+	/**
+	 * Creates an image widget.
+	 *
+	 * @return The new image widget.
+	 */
+	virtual TSharedRef<SWidget> MakeImage( const TAttribute<const FSlateBrush*>& Image, const TAttribute<FSlateColor>& Color, const TAttribute<EVisibility>& Visibility ) const = 0;
+
+	/**
+	 * Creates a tool tip with the specified string.
+	 *
+	 * @param ToolTipString The string attribute to assign to the tool tip.
+	 *
+	 * @return The tool tip.
+	 */
+	virtual TSharedRef<IToolTip> MakeToolTip( const TAttribute<FString>& ToolTipString ) = 0;
+
+	/**
+	 * Creates a tool tip with the specified text.
+	 *
+	 * @param ToolTipText The text to assign to the tool tip.
+	 *
+	 * @return The tool tip.
+	 */
+	virtual TSharedRef<IToolTip> MakeToolTip( const FText& ToolTipText ) = 0;
+
+	/**
+	 * Creates a title bar for the specified window.
+	 *
+	 * @param Window The window to create the title bar for.
+	 * @param CenterContent Optional content for the title bar's center (will override window title).
+	 * @param OutTitleBar Will hold a pointer to the title bar's interface.
+	 *
+	 * @return The new title bar widget.
+	 */
+	virtual TSharedRef<SWidget> MakeWindowTitleBar( const TSharedRef<SWindow>& Window, const TSharedPtr<SWidget>& CenterContent, TSharedPtr<IWindowTitleBar>& OutTitleBar ) const = 0;
+
+	/**
+	 * Destroying windows has implications on some OSs (e.g. destroying Win32 HWNDs can cause events to be lost).
+	 *
+	 * Slate strictly controls when windows are destroyed.
+	 *
+	 * @param WindowToDestroy The window to queue for destruction.
+	 */
+	virtual void RequestDestroyWindow( TSharedRef<SWindow> WindowToDestroy ) = 0;
+
+	/**
+	 * Sets keyboard focus to the specified widget.  The widget must be allowed to receive keyboard focus.
+	 *
+	 * @param  InWidget WidgetPath to the Widget to being focused
+	 * @param InCause The reason that keyboard focus is changing
+	 *
+	 * @return true if the widget is now focused, false otherwise.
+	 */
+	virtual bool SetKeyboardFocus( const FWidgetPath& InFocusPath, const EKeyboardFocusCause::Type InCause ) = 0;
+
+public:
+
+	/**
+	 * Returns the current instance of the application. The application should have been initialized before
+	 * this method is called
+	 *
+	 * @return  Reference to the application
+	 */
+	static FSlateApplicationBase& Get( )
+	{
+		check(IsInGameThread());
+		return *CurrentBaseApplication;
+	}
+
+	/**
+	 * Returns true if a Slate application instance is currently initialized and ready
+	 *
+	 * @return  True if Slate application is initialized
+	 */
+	static bool IsInitialized( )
+	{
+		return CurrentBaseApplication.IsValid();
+	}
+
+protected:
+
+	// Holds the Slate renderer used to render this application.
+	TSharedPtr<FSlateRenderer> Renderer;
+
+protected:
+
+	// Holds a pointer to the current application.
+	static TSharedPtr<FSlateApplicationBase> CurrentBaseApplication;
+};
diff --git a/Engine/Source/Runtime/SlateCore/Public/Application/ThrottleManager.h b/Engine/Source/Runtime/SlateCore/Public/Application/ThrottleManager.h
new file mode 100644
--- /dev/null
+++ b/Engine/Source/Runtime/SlateCore/Public/Application/ThrottleManager.h
@@ -0,0 +1,96 @@
+// Copyright 1998-2014 Epic Games, Inc. All Rights Reserved.
+
+/*=============================================================================
+	ThrottleManager.h: Declares the FSlateThrottleManager class.
+=============================================================================*/
+
+#pragma once
+
+
+/**
+ * A handle to a throttle request made to the throttle manager
+ * Throttles can only be ended by passing back a request
+ */
+class FThrottleRequest
+{
+	// Only the manager is allowed to change values
+	friend class FSlateThrottleManager;
+
+public:
+
+	FThrottleRequest()
+		: Index(INDEX_NONE)
+	{ }
+
+public:
+
+	/**
+	 * Whether or not the handle is valid
+	 *
+	 * @return true if the throttle is valid
+	 */
+	bool IsValid() const
+	{
+		return Index != INDEX_NONE;
+	}
+
+private:
+
+	int32 Index;
+};
+
+
+/**
+ * A class which manages requests to throttle parts of the engine to ensure Slate UI performance                   
+ */
+class SLATECORE_API FSlateThrottleManager
+{
+public:
+
+	/**
+	 * Constructor                   
+	 */
+	FSlateThrottleManager();
+
+public:
+
+	/**
+	 * Requests that we enter responsive mode.  I.E throttle slow parts of the engine   
+	 *
+	 * @return A handle to the request to enter responsive mode.  Can only be ended with this request
+	 */
+	FThrottleRequest EnterResponsiveMode();
+
+	/**
+	 * Request to leave responsive mode.  
+	 * Note that this may not end responsive mode in the case that multiple EnterResponsiveMode requests were made                    
+	 *
+	 * @param InHandle	The handle that was created with EnterResponsiveMode
+	 */
+	void LeaveResponsiveMode( FThrottleRequest& InHandle );
+
+	/**
+	 * Whether or not we allow expensive tasks which could hurt performance to occur
+	 *
+	 * @return true if we allow expensive tasks, false otherwise
+	 */
+	bool IsAllowingExpensiveTasks() const;
+
+public:
+
+	/**
+	 * Gets the instance of this manager                   
+	 */
+	static FSlateThrottleManager& Get( );
+
+private:
+
+	/** CVar allowing us to toggle throttling ability */
+	FAutoConsoleVariableRef CVarAllowThrottle;
+
+	/** CVar variable to check if we allow throttling (int32 for compatibility) */
+	int32 bShouldThrottle;
+
+	/** Number of active throttle requests */
+	uint32 ThrottleCount;
+};
diff --git a/Engine/Source/Runtime/SlateCore/Public/Brushes/SlateBorderBrush.h b/Engine/Source/Runtime/SlateCore/Public/Brushes/SlateBorderBrush.h
new file mode 100644
--- /dev/null
+++ b/Engine/Source/Runtime/SlateCore/Public/Brushes/SlateBorderBrush.h
@@ -0,0 +1,106 @@
+// Copyright 1998-2014 Epic Games, Inc. All Rights Reserved.
+
+/*=============================================================================
+	SlateBoxBrush.h: Declares the FSlateBoxBrush structure.
+=============================================================================*/
+
+#pragma once
+
+
+/**
+ * Similar to FSlateBoxBrush but has no middle and the sides tile instead of stretching.
+ * The margin is applied exactly as in FSlateBoxBrush.
+ */
+struct SLATECORE_API FSlateBorderBrush
+	: public FSlateBrush
+{
+	/**
+	 * @param InImageName The name of the texture to draw
+	 * @param InMargin Determines the sides and corner sizes; see FSlateBoxBrush.
+	 * @param InColorAndOpacity	Color and opacity scale.
+	 * @param InImageType The type of image this this is
+	 */
+	FORCENOINLINE FSlateBorderBrush( const FName& InImageName, const FMargin& InMargin, const FLinearColor& InColorAndOpacity = FLinearColor(1,1,1,1), ESlateBrushImageType::Type InImageType = ESlateBrushImageType::FullColor )
+		: FSlateBrush(ESlateBrushDrawType::Border, InImageName, InMargin, ESlateBrushTileType::Both, InImageType, FVector2D::ZeroVector, InColorAndOpacity)
+	{
+		// A border with no margin will not show up at all.
+		check(InMargin.GetDesiredSize().SizeSquared() > 0);
+	}
+
+	FORCENOINLINE FSlateBorderBrush( const FString& InImageName, const FMargin& InMargin, const FLinearColor& InColorAndOpacity = FLinearColor(1,1,1,1), ESlateBrushImageType::Type InImageType = ESlateBrushImageType::FullColor )
+		: FSlateBrush(ESlateBrushDrawType::Border, *InImageName, InMargin, ESlateBrushTileType::Both, InImageType, FVector2D::ZeroVector, InColorAndOpacity)
+	{
+		// A border with no margin will not show up at all.
+		check(InMargin.GetDesiredSize().SizeSquared() > 0);
+	}
+
+	FORCENOINLINE FSlateBorderBrush( const ANSICHAR* InImageName, const FMargin& InMargin, const FLinearColor& InColorAndOpacity = FLinearColor(1,1,1,1), ESlateBrushImageType::Type InImageType = ESlateBrushImageType::FullColor )
+		: FSlateBrush(ESlateBrushDrawType::Border, InImageName, InMargin, ESlateBrushTileType::Both, InImageType, FVector2D::ZeroVector, InColorAndOpacity)
+	{
+		// A border with no margin will not show up at all.
+		check(InMargin.GetDesiredSize().SizeSquared() > 0);
+	}
+
+	FORCENOINLINE FSlateBorderBrush( const WIDECHAR* InImageName, const FMargin& InMargin, const FLinearColor& InColorAndOpacity = FLinearColor(1,1,1,1), ESlateBrushImageType::Type InImageType = ESlateBrushImageType::FullColor )
+		: FSlateBrush(ESlateBrushDrawType::Border, InImageName, InMargin, ESlateBrushTileType::Both, InImageType, FVector2D::ZeroVector, InColorAndOpacity)
+	{
+		// A border with no margin will not show up at all.
+		check(InMargin.GetDesiredSize().SizeSquared() > 0);
+	}
+
+	FORCENOINLINE FSlateBorderBrush( const FName& InImageName, const FMargin& InMargin, const TSharedRef< FLinearColor >& InColorAndOpacity, ESlateBrushImageType::Type InImageType = ESlateBrushImageType::FullColor )
+		: FSlateBrush(ESlateBrushDrawType::Border, InImageName, InMargin, ESlateBrushTileType::Both, InImageType, FVector2D::ZeroVector, InColorAndOpacity)
+	{
+		// A border with no margin will not show up at all.
+		check(InMargin.GetDesiredSize().SizeSquared() > 0);
+	}
+
+	FORCENOINLINE FSlateBorderBrush( const FString& InImageName, const FMargin& InMargin, const TSharedRef< FLinearColor >& InColorAndOpacity, ESlateBrushImageType::Type InImageType = ESlateBrushImageType::FullColor )
+		: FSlateBrush(ESlateBrushDrawType::Border, *InImageName, InMargin, ESlateBrushTileType::Both, InImageType, FVector2D::ZeroVector, InColorAndOpacity)
+	{
+		// A border with no margin will not show up at all.
+		check(InMargin.GetDesiredSize().SizeSquared() > 0);
+	}
+
+	FORCENOINLINE FSlateBorderBrush( const ANSICHAR* InImageName, const FMargin& InMargin, const TSharedRef< FLinearColor >& InColorAndOpacity, ESlateBrushImageType::Type InImageType = ESlateBrushImageType::FullColor )
+		: FSlateBrush(ESlateBrushDrawType::Border, InImageName, InMargin, ESlateBrushTileType::Both, InImageType, FVector2D::ZeroVector, InColorAndOpacity)
+	{
+		// A border with no margin will not show up at all.
+		check(InMargin.GetDesiredSize().SizeSquared() > 0);
+	}
+
+	FORCENOINLINE FSlateBorderBrush( const TCHAR* InImageName, const FMargin& InMargin, const TSharedRef< FLinearColor >& InColorAndOpacity, ESlateBrushImageType::Type InImageType = ESlateBrushImageType::FullColor )
+		: FSlateBrush(ESlateBrushDrawType::Border, InImageName, InMargin, ESlateBrushTileType::Both, InImageType, FVector2D::ZeroVector, InColorAndOpacity)
+	{
+		// A border with no margin will not show up at all.
+		check(InMargin.GetDesiredSize().SizeSquared() > 0);
+	}
+
+	FORCENOINLINE FSlateBorderBrush( const FName& InImageName, const FMargin& InMargin, const FSlateColor& InColorAndOpacity, ESlateBrushImageType::Type InImageType = ESlateBrushImageType::FullColor )
+		: FSlateBrush(ESlateBrushDrawType::Border, InImageName, InMargin, ESlateBrushTileType::Both, InImageType, FVector2D::ZeroVector, InColorAndOpacity)
+	{
+		// A border with no margin will not show up at all.
+		check(InMargin.GetDesiredSize().SizeSquared() > 0);
+	}
+
+	FORCENOINLINE FSlateBorderBrush( const FString& InImageName, const FMargin& InMargin, const FSlateColor& InColorAndOpacity, ESlateBrushImageType::Type InImageType = ESlateBrushImageType::FullColor )
+		: FSlateBrush(ESlateBrushDrawType::Border, *InImageName, InMargin, ESlateBrushTileType::Both, InImageType, FVector2D::ZeroVector, InColorAndOpacity)
+	{
+		// A border with no margin will not show up at all.
+		check(InMargin.GetDesiredSize().SizeSquared() > 0);
+	}
+
+	FORCENOINLINE FSlateBorderBrush( const ANSICHAR* InImageName, const FMargin& InMargin, const FSlateColor& InColorAndOpacity, ESlateBrushImageType::Type InImageType = ESlateBrushImageType::FullColor )
+		: FSlateBrush(ESlateBrushDrawType::Border, InImageName, InMargin, ESlateBrushTileType::Both, InImageType, FVector2D::ZeroVector, InColorAndOpacity)
+	{
+		// A border with no margin will not show up at all.
+		check(InMargin.GetDesiredSize().SizeSquared() > 0);
+	}
+
+	FORCENOINLINE FSlateBorderBrush( const TCHAR* InImageName, const FMargin& InMargin, const FSlateColor& InColorAndOpacity, ESlateBrushImageType::Type InImageType = ESlateBrushImageType::FullColor )
+		: FSlateBrush(ESlateBrushDrawType::Border, InImageName, InMargin, ESlateBrushTileType::Both, InImageType, FVector2D::ZeroVector, InColorAndOpacity)
+	{
+		// A border with no margin will not show up at all.
+		check(InMargin.GetDesiredSize().SizeSquared() > 0);
+	}
+};
diff --git a/Engine/Source/Runtime/SlateCore/Public/Brushes/SlateBoxBrush.h b/Engine/Source/Runtime/SlateCore/Public/Brushes/SlateBoxBrush.h
new file mode 100644
--- /dev/null
+++ b/Engine/Source/Runtime/SlateCore/Public/Brushes/SlateBoxBrush.h
@@ -0,0 +1,138 @@
+// Copyright 1998-2014 Epic Games, Inc. All Rights Reserved.
+
+/*=============================================================================
+	SlateBoxBrush.h: Declares the FSlateBoxBrush structure.
+=============================================================================*/
+
+#pragma once
+
+
+/**
+ * A 3x3 box where the sides stretch horizontally and vertically and the middle stretches to fill.
+ * The corners will not be stretched. The size of the sides and corners is determined by the
+ * Margin as follows:
+ *
+ *                 _____________________
+ *                |  | Margin.Top    |  |
+ *                |__|_______________|__|   Margin.Right
+ *                |  |               |  |  /
+ *              +--> |               | <--+
+ *             /  |__|_______________|__|
+ *  Margin.Left   |  | Margin.Bottom |  |
+ *                |__|_______________|__|
+ *
+ */
+struct SLATECORE_API FSlateBoxBrush
+	: public FSlateBrush
+{
+	/**
+	 * Make a 3x3 box that stretches the texture.
+	 *
+	 * @param InImageName		The name of image to make into a box
+	 * @param InMargin			The size of corners and sides in normalized texture UV space.
+	 * @param InColorAndOpacity	Color and opacity scale.
+	 */
+	FORCENOINLINE FSlateBoxBrush( const FName& InImageName, const FMargin& InMargin, const FLinearColor& InColorAndOpacity = FLinearColor(1,1,1,1), ESlateBrushImageType::Type InImageType = ESlateBrushImageType::FullColor )
+		: FSlateBrush(ESlateBrushDrawType::Box, InImageName, InMargin, ESlateBrushTileType::NoTile, InImageType, FVector2D::ZeroVector, InColorAndOpacity)
+	{ }
+
+	FORCENOINLINE FSlateBoxBrush( const FString& InImageName, const FMargin& InMargin, const FLinearColor& InColorAndOpacity = FLinearColor(1,1,1,1), ESlateBrushImageType::Type InImageType = ESlateBrushImageType::FullColor )
+		: FSlateBrush(ESlateBrushDrawType::Box, *InImageName, InMargin, ESlateBrushTileType::NoTile, InImageType, FVector2D::ZeroVector, InColorAndOpacity)
+	{ }
+
+	FORCENOINLINE FSlateBoxBrush( const ANSICHAR* InImageName, const FMargin& InMargin, const FLinearColor& InColorAndOpacity = FLinearColor(1,1,1,1), ESlateBrushImageType::Type InImageType = ESlateBrushImageType::FullColor )
+		: FSlateBrush(ESlateBrushDrawType::Box, InImageName, InMargin, ESlateBrushTileType::NoTile, InImageType, FVector2D::ZeroVector, InColorAndOpacity)
+	{ } 
+
+	FORCENOINLINE FSlateBoxBrush( const WIDECHAR* InImageName, const FMargin& InMargin, const FLinearColor& InColorAndOpacity = FLinearColor(1,1,1,1), ESlateBrushImageType::Type InImageType = ESlateBrushImageType::FullColor )
+		: FSlateBrush(ESlateBrushDrawType::Box, InImageName, InMargin, ESlateBrushTileType::NoTile, InImageType, FVector2D::ZeroVector, InColorAndOpacity)
+	{ }
+
+	FORCENOINLINE FSlateBoxBrush( const FName& InImageName, const FMargin& InMargin, const TSharedRef< FLinearColor >& InColorAndOpacity, ESlateBrushImageType::Type InImageType = ESlateBrushImageType::FullColor )
+		: FSlateBrush(ESlateBrushDrawType::Box, InImageName, InMargin, ESlateBrushTileType::NoTile, InImageType, FVector2D::ZeroVector, InColorAndOpacity)
+	{ }
+
+	FORCENOINLINE FSlateBoxBrush( const FString& InImageName, const FMargin& InMargin, const TSharedRef< FLinearColor >& InColorAndOpacity, ESlateBrushImageType::Type InImageType = ESlateBrushImageType::FullColor )
+		: FSlateBrush(ESlateBrushDrawType::Box, *InImageName, InMargin, ESlateBrushTileType::NoTile, InImageType, FVector2D::ZeroVector, InColorAndOpacity)
+	{ }
+
+	FORCENOINLINE FSlateBoxBrush( const ANSICHAR* InImageName, const FMargin& InMargin, const TSharedRef< FLinearColor >& InColorAndOpacity, ESlateBrushImageType::Type InImageType = ESlateBrushImageType::FullColor )
+		: FSlateBrush(ESlateBrushDrawType::Box, InImageName, InMargin, ESlateBrushTileType::NoTile, InImageType, FVector2D::ZeroVector, InColorAndOpacity)
+	{ }
+
+	FORCENOINLINE FSlateBoxBrush( const TCHAR* InImageName, const FMargin& InMargin, const TSharedRef< FLinearColor >& InColorAndOpacity, ESlateBrushImageType::Type InImageType = ESlateBrushImageType::FullColor )
+		: FSlateBrush(ESlateBrushDrawType::Box, InImageName, InMargin, ESlateBrushTileType::NoTile, InImageType, FVector2D::ZeroVector, InColorAndOpacity)
+	{ }
+
+	FORCENOINLINE FSlateBoxBrush( const FName& InImageName, const FMargin& InMargin, const FSlateColor& InColorAndOpacity, ESlateBrushImageType::Type InImageType = ESlateBrushImageType::FullColor )
+		: FSlateBrush(ESlateBrushDrawType::Box, InImageName, InMargin, ESlateBrushTileType::NoTile, InImageType, FVector2D::ZeroVector, InColorAndOpacity)
+	{ }
+
+	FORCENOINLINE FSlateBoxBrush( const FString& InImageName, const FMargin& InMargin, const FSlateColor& InColorAndOpacity, ESlateBrushImageType::Type InImageType = ESlateBrushImageType::FullColor )
+		: FSlateBrush(ESlateBrushDrawType::Box, *InImageName, InMargin, ESlateBrushTileType::NoTile, InImageType, FVector2D::ZeroVector, InColorAndOpacity)
+	{ }
+
+	FORCENOINLINE FSlateBoxBrush( const ANSICHAR* InImageName, const FMargin& InMargin, const FSlateColor& InColorAndOpacity, ESlateBrushImageType::Type InImageType = ESlateBrushImageType::FullColor )
+		: FSlateBrush(ESlateBrushDrawType::Box, InImageName, InMargin, ESlateBrushTileType::NoTile, InImageType, FVector2D::ZeroVector, InColorAndOpacity)
+	{ }
+
+	FORCENOINLINE FSlateBoxBrush( const TCHAR* InImageName, const FMargin& InMargin, const FSlateColor& InColorAndOpacity, ESlateBrushImageType::Type InImageType = ESlateBrushImageType::FullColor )
+		: FSlateBrush(ESlateBrushDrawType::Box, InImageName, InMargin, ESlateBrushTileType::NoTile, InImageType, FVector2D::ZeroVector, InColorAndOpacity)
+	{ }
+
+	/**
+	 * Make a 3x3 box that stretches the texture.
+	 *
+	 * @param InImageName       The name of image to make into a box
+	 * @param ImageSize         The size of the resource as we want it to appear in slate units.
+	 * @param InMargin          The size of corners and sides in texture space.
+	 * @param InColorAndOpacity	Color and opacity scale. Note of the image type is ImageType_TintMask, this value should be in HSV
+	 */
+	FORCENOINLINE FSlateBoxBrush( const FName& InImageName, const FVector2D& ImageSize, const FMargin& InMargin, const FLinearColor& InColorAndOpacity = FLinearColor(1,1,1,1), ESlateBrushImageType::Type InImageType = ESlateBrushImageType::FullColor )
+		: FSlateBrush(ESlateBrushDrawType::Box, InImageName, InMargin, ESlateBrushTileType::NoTile, InImageType, ImageSize, InColorAndOpacity)
+	{ }
+
+	FORCENOINLINE FSlateBoxBrush( const FString& InImageName, const FVector2D& ImageSize, const FMargin& InMargin, const FLinearColor& InColorAndOpacity = FLinearColor(1,1,1,1), ESlateBrushImageType::Type InImageType = ESlateBrushImageType::FullColor )
+		: FSlateBrush(ESlateBrushDrawType::Box, *InImageName, InMargin, ESlateBrushTileType::NoTile, InImageType, ImageSize, InColorAndOpacity)
+	{ }
+
+	FORCENOINLINE FSlateBoxBrush( const ANSICHAR* InImageName, const FVector2D& ImageSize, const FMargin& InMargin, const FLinearColor& InColorAndOpacity = FLinearColor(1,1,1,1), ESlateBrushImageType::Type InImageType = ESlateBrushImageType::FullColor )
+		: FSlateBrush(ESlateBrushDrawType::Box, InImageName, InMargin, ESlateBrushTileType::NoTile, InImageType, ImageSize, InColorAndOpacity)
+	{ }
+
+	FORCENOINLINE FSlateBoxBrush( const WIDECHAR* InImageName, const FVector2D& ImageSize, const FMargin& InMargin, const FLinearColor& InColorAndOpacity = FLinearColor(1,1,1,1), ESlateBrushImageType::Type InImageType = ESlateBrushImageType::FullColor )
+		: FSlateBrush(ESlateBrushDrawType::Box, InImageName, InMargin, ESlateBrushTileType::NoTile, InImageType, ImageSize, InColorAndOpacity)
+	{ }
+
+	FORCENOINLINE FSlateBoxBrush( const FName& InImageName, const FVector2D& ImageSize, const FMargin& InMargin, const TSharedRef< FLinearColor >& InColorAndOpacity, ESlateBrushImageType::Type InImageType = ESlateBrushImageType::FullColor )
+		: FSlateBrush(ESlateBrushDrawType::Box, InImageName, InMargin, ESlateBrushTileType::NoTile, InImageType, ImageSize, InColorAndOpacity)
+	{ }
+
+	FORCENOINLINE FSlateBoxBrush( const FString& InImageName, const FVector2D& ImageSize, const FMargin& InMargin, const TSharedRef< FLinearColor >& InColorAndOpacity, ESlateBrushImageType::Type InImageType = ESlateBrushImageType::FullColor )
+		: FSlateBrush(ESlateBrushDrawType::Box, *InImageName, InMargin, ESlateBrushTileType::NoTile, InImageType, ImageSize, InColorAndOpacity)
+	{ }
+
+	FORCENOINLINE FSlateBoxBrush( const ANSICHAR* InImageName, const FVector2D& ImageSize, const FMargin& InMargin, const TSharedRef< FLinearColor >& InColorAndOpacity, ESlateBrushImageType::Type InImageType = ESlateBrushImageType::FullColor )
+		: FSlateBrush(ESlateBrushDrawType::Box, InImageName, InMargin, ESlateBrushTileType::NoTile, InImageType, ImageSize, InColorAndOpacity)
+	{ }
+
+	FORCENOINLINE FSlateBoxBrush( const TCHAR* InImageName, const FVector2D& ImageSize, const FMargin& InMargin, const TSharedRef< FLinearColor >& InColorAndOpacity, ESlateBrushImageType::Type InImageType = ESlateBrushImageType::FullColor )
+		: FSlateBrush(ESlateBrushDrawType::Box, InImageName, InMargin, ESlateBrushTileType::NoTile, InImageType, ImageSize, InColorAndOpacity)
+	{ }
+
+	FORCENOINLINE FSlateBoxBrush( const FName& InImageName, const FVector2D& ImageSize, const FMargin& InMargin, const FSlateColor& InColorAndOpacity, ESlateBrushImageType::Type InImageType = ESlateBrushImageType::FullColor )
+		: FSlateBrush(ESlateBrushDrawType::Box, InImageName, InMargin, ESlateBrushTileType::NoTile, InImageType, ImageSize, InColorAndOpacity)
+	{ }
+
+	FORCENOINLINE FSlateBoxBrush( const FString& InImageName, const FVector2D& ImageSize, const FMargin& InMargin, const FSlateColor& InColorAndOpacity, ESlateBrushImageType::Type InImageType = ESlateBrushImageType::FullColor )
+		: FSlateBrush(ESlateBrushDrawType::Box, *InImageName, InMargin, ESlateBrushTileType::NoTile, InImageType, ImageSize, InColorAndOpacity)
+	{ }
+
+	FORCENOINLINE FSlateBoxBrush( const ANSICHAR* InImageName, const FVector2D& ImageSize, const FMargin& InMargin, const FSlateColor& InColorAndOpacity, ESlateBrushImageType::Type InImageType = ESlateBrushImageType::FullColor )
+		: FSlateBrush(ESlateBrushDrawType::Box, InImageName, InMargin, ESlateBrushTileType::NoTile, InImageType, ImageSize, InColorAndOpacity)
+	{ }
+
+	FORCENOINLINE FSlateBoxBrush( const TCHAR* InImageName, const FVector2D& ImageSize, const FMargin& InMargin, const FSlateColor& InColorAndOpacity, ESlateBrushImageType::Type InImageType = ESlateBrushImageType::FullColor )
+		: FSlateBrush(ESlateBrushDrawType::Box, InImageName, InMargin, ESlateBrushTileType::NoTile, InImageType, ImageSize, InColorAndOpacity)
+	{ }
+};
diff --git a/Engine/Source/Runtime/SlateCore/Public/Brushes/SlateColorBrush.h b/Engine/Source/Runtime/SlateCore/Public/Brushes/SlateColorBrush.h
new file mode 100644
--- /dev/null
+++ b/Engine/Source/Runtime/SlateCore/Public/Brushes/SlateColorBrush.h
@@ -0,0 +1,42 @@
+// Copyright 1998-2014 Epic Games, Inc. All Rights Reserved.
+
+/*=============================================================================
+	SlateColorBrush.h: Declares the FSlateColorBrush structure.
+=============================================================================*/
+
+#pragma once
+
+
+/**
+ * A color brush.  Draws a box with no margins that only has a color applied to it                   
+ */
+struct SLATECORE_API FSlateColorBrush
+	: public FSlateBrush
+{
+	/** 
+	 * Creates and initializes a new instance with the specified linear color.
+	 *
+	 * @param InColor The linear color to use
+	 */
+	FORCENOINLINE FSlateColorBrush( const FLinearColor& InColor )
+		: FSlateBrush(ESlateBrushDrawType::Box, NAME_None, FMargin(0.0f), ESlateBrushTileType::NoTile, ESlateBrushImageType::NoImage, FVector2D::ZeroVector, InColor)
+	{ }
+
+	/** 
+	 * Creates and initializes a new instance with the specified color value.
+	 *
+	 * @param InColor The color value to use
+	 */
+	FORCENOINLINE FSlateColorBrush( const FColor& InColor )
+		: FSlateBrush(ESlateBrushDrawType::Box, NAME_None, FMargin(0.0f), ESlateBrushTileType::NoTile, ESlateBrushImageType::NoImage, FVector2D::ZeroVector, InColor.ReinterpretAsLinear())
+	{ }
+
+	/** 
+	 * Creates and initializes a new instance with the specified shared color.
+	 *
+	 * @param InColor The shared color to use
+	 */
+	FORCENOINLINE FSlateColorBrush( const TSharedRef< FLinearColor >& InColor )
+		: FSlateBrush(ESlateBrushDrawType::Box, NAME_None, FMargin(0.0f), ESlateBrushTileType::NoTile, ESlateBrushImageType::NoImage, FVector2D::ZeroVector, InColor)
+	{ }
+};
\ No newline at end of file
diff --git a/Engine/Source/Runtime/SlateCore/Public/Brushes/SlateDynamicImageBrush.h b/Engine/Source/Runtime/SlateCore/Public/Brushes/SlateDynamicImageBrush.h
new file mode 100644
--- /dev/null
+++ b/Engine/Source/Runtime/SlateCore/Public/Brushes/SlateDynamicImageBrush.h
@@ -0,0 +1,65 @@
+// Copyright 1998-2014 Epic Games, Inc. All Rights Reserved.
+
+/*=============================================================================
+	SlateDynamicImageBrush.h: Declares the FSlateDynamicImageBrush structure.
+=============================================================================*/
+
+#pragma once
+
+
+/**
+ * Ignores the Margin. Just renders the image. Can tile the image instead of stretching.
+ */
+struct SLATECORE_API FSlateDynamicImageBrush
+	: public FSlateBrush
+{
+	/**
+	 * @param InTexture		The UTexture2D being used for this brush.
+	 * @param InImageSize		How large should the image be (not necessarily the image size on disk)
+	 * @param InTint		The tint of the image
+	 * @param InTiling		How do we tile if at all?
+	 * @param InImageType		The type of image this this is
+	 */
+	FORCENOINLINE FSlateDynamicImageBrush( 
+		class UTexture2D* InTexture, 
+		const FVector2D& InImageSize,
+		const FName InTextureName,
+		const FLinearColor& InTint = FLinearColor(1.0f, 1.0f, 1.0f, 1.0f), 
+		ESlateBrushTileType::Type InTiling = ESlateBrushTileType::NoTile, 
+		ESlateBrushImageType::Type InImageType = ESlateBrushImageType::FullColor
+	)
+		: FSlateBrush(ESlateBrushDrawType::Image, FName(TEXT("None")), FMargin(0.0f), InTiling, InImageType, InImageSize, InTint, (UObject*)InTexture)
+	{
+		bIsDynamicallyLoaded = true;
+
+		// if we have a texture, make a unique name
+		if (GetResourceObject() != nullptr)
+		{
+			ResourceName = InTextureName;
+		}
+	}
+
+	/**
+	 * @param InTextureName		The name of the texture to load.
+	 * @param InImageSize		How large should the image be (not necessarily the image size on disk)
+	 * @param InTint		The tint of the image.
+	 * @param InTiling		How do we tile if at all?
+	 * @param InImageType		The type of image this this is
+	 */
+	FSlateDynamicImageBrush( 
+		const FName InTextureName,
+		const FVector2D& InImageSize,
+		const FLinearColor& InTint = FLinearColor(1.0f, 1.0f, 1.0f, 1.0f), 
+		ESlateBrushTileType::Type InTiling = ESlateBrushTileType::NoTile, 
+		ESlateBrushImageType::Type InImageType = ESlateBrushImageType::FullColor
+	)
+		: FSlateBrush(ESlateBrushDrawType::Image, InTextureName, FMargin(0.0f), InTiling, InImageType, InImageSize, InTint, nullptr, true)
+	{
+		bIsDynamicallyLoaded = true;
+	}
+
+	/**
+	 * Destructor.
+	 */
+	virtual ~FSlateDynamicImageBrush();
+};
diff --git a/Engine/Source/Runtime/SlateCore/Public/Brushes/SlateImageBrush.h b/Engine/Source/Runtime/SlateCore/Public/Brushes/SlateImageBrush.h
new file mode 100644
--- /dev/null
+++ b/Engine/Source/Runtime/SlateCore/Public/Brushes/SlateImageBrush.h
@@ -0,0 +1,70 @@
+// Copyright 1998-2014 Epic Games, Inc. All Rights Reserved.
+
+/*=============================================================================
+	SlateImageBrush.h: Declares the FSlateImageBrush structure.
+=============================================================================*/
+
+#pragma once
+
+
+/**
+ * Ignores the Margin. Just renders the image. Can tile the image instead of stretching.
+ */
+struct SLATECORE_API FSlateImageBrush
+	: public FSlateBrush
+{
+	/**
+	 * @param InImageName	The name of the texture to draw
+	 * @param InImageSize	How large should the image be (not necessarily the image size on disk)
+	 * @param InTint		The tint of the image.
+	 * @param InTiling		How do we tile if at all?
+	 * @param InImageType	The type of image this this is
+	 */
+	FORCENOINLINE FSlateImageBrush( const FName& InImageName, const FVector2D& InImageSize, const FLinearColor& InTint = FLinearColor(1,1,1,1), ESlateBrushTileType::Type InTiling = ESlateBrushTileType::NoTile, ESlateBrushImageType::Type InImageType = ESlateBrushImageType::FullColor )
+		: FSlateBrush(ESlateBrushDrawType::Image, InImageName, FMargin(0), InTiling, InImageType, InImageSize, InTint)
+	{ }
+
+	FORCENOINLINE FSlateImageBrush( const FString& InImageName, const FVector2D& InImageSize, const FLinearColor& InTint = FLinearColor(1,1,1,1), ESlateBrushTileType::Type InTiling = ESlateBrushTileType::NoTile, ESlateBrushImageType::Type InImageType = ESlateBrushImageType::FullColor )
+		: FSlateBrush(ESlateBrushDrawType::Image, *InImageName, FMargin(0), InTiling, InImageType, InImageSize, InTint)
+	{ }
+
+	FORCENOINLINE FSlateImageBrush( const ANSICHAR* InImageName, const FVector2D& InImageSize, const FLinearColor& InTint = FLinearColor(1,1,1,1), ESlateBrushTileType::Type InTiling = ESlateBrushTileType::NoTile, ESlateBrushImageType::Type InImageType = ESlateBrushImageType::FullColor )
+		: FSlateBrush(ESlateBrushDrawType::Image, InImageName, FMargin(0), InTiling, InImageType, InImageSize, InTint)
+	{ }
+
+	FORCENOINLINE FSlateImageBrush( const WIDECHAR* InImageName, const FVector2D& InImageSize, const FLinearColor& InTint = FLinearColor(1,1,1,1), ESlateBrushTileType::Type InTiling = ESlateBrushTileType::NoTile, ESlateBrushImageType::Type InImageType = ESlateBrushImageType::FullColor )
+		: FSlateBrush(ESlateBrushDrawType::Image, InImageName, FMargin(0), InTiling, InImageType, InImageSize, InTint)
+	{ }
+
+	FORCENOINLINE FSlateImageBrush( const FName& InImageName, const FVector2D& InImageSize, const TSharedRef< FLinearColor >& InTint, ESlateBrushTileType::Type InTiling = ESlateBrushTileType::NoTile, ESlateBrushImageType::Type InImageType = ESlateBrushImageType::FullColor )
+		: FSlateBrush(ESlateBrushDrawType::Image, InImageName, FMargin(0), InTiling, InImageType, InImageSize, InTint)
+	{ }
+
+	FORCENOINLINE FSlateImageBrush( const FString& InImageName, const FVector2D& InImageSize, const TSharedRef< FLinearColor >& InTint, ESlateBrushTileType::Type InTiling = ESlateBrushTileType::NoTile, ESlateBrushImageType::Type InImageType = ESlateBrushImageType::FullColor )
+		: FSlateBrush(ESlateBrushDrawType::Image, *InImageName, FMargin(0), InTiling, InImageType, InImageSize, InTint)
+	{ }
+
+	FORCENOINLINE FSlateImageBrush( const ANSICHAR* InImageName, const FVector2D& InImageSize, const TSharedRef< FLinearColor >& InTint, ESlateBrushTileType::Type InTiling = ESlateBrushTileType::NoTile, ESlateBrushImageType::Type InImageType = ESlateBrushImageType::FullColor )
+		: FSlateBrush(ESlateBrushDrawType::Image, InImageName, FMargin(0), InTiling, InImageType, InImageSize, InTint)
+	{ }
+
+	FORCENOINLINE FSlateImageBrush( const TCHAR* InImageName, const FVector2D& InImageSize, const TSharedRef< FLinearColor >& InTint, ESlateBrushTileType::Type InTiling = ESlateBrushTileType::NoTile, ESlateBrushImageType::Type InImageType = ESlateBrushImageType::FullColor )
+		: FSlateBrush(ESlateBrushDrawType::Image, InImageName, FMargin(0), InTiling, InImageType, InImageSize, InTint)
+	{ }
+
+	FORCENOINLINE FSlateImageBrush( const FName& InImageName, const FVector2D& InImageSize, const FSlateColor& InTint, ESlateBrushTileType::Type InTiling = ESlateBrushTileType::NoTile, ESlateBrushImageType::Type InImageType = ESlateBrushImageType::FullColor )
+		: FSlateBrush(ESlateBrushDrawType::Image, InImageName, FMargin(0), InTiling, InImageType, InImageSize, InTint)
+	{ }
+
+	FORCENOINLINE FSlateImageBrush( const FString& InImageName, const FVector2D& InImageSize, const FSlateColor& InTint, ESlateBrushTileType::Type InTiling = ESlateBrushTileType::NoTile, ESlateBrushImageType::Type InImageType = ESlateBrushImageType::FullColor )
+		: FSlateBrush(ESlateBrushDrawType::Image, *InImageName, FMargin(0), InTiling, InImageType, InImageSize, InTint)
+	{ }
+
+	FORCENOINLINE FSlateImageBrush( const ANSICHAR* InImageName, const FVector2D& InImageSize, const FSlateColor& InTint, ESlateBrushTileType::Type InTiling = ESlateBrushTileType::NoTile, ESlateBrushImageType::Type InImageType = ESlateBrushImageType::FullColor )
+		: FSlateBrush(ESlateBrushDrawType::Image, InImageName, FMargin(0), InTiling, InImageType, InImageSize, InTint)
+	{ }
+
+	FORCENOINLINE FSlateImageBrush( const TCHAR* InImageName, const FVector2D& InImageSize, const FSlateColor& InTint, ESlateBrushTileType::Type InTiling = ESlateBrushTileType::NoTile, ESlateBrushImageType::Type InImageType = ESlateBrushImageType::FullColor )
+		: FSlateBrush(ESlateBrushDrawType::Image, InImageName, FMargin(0), InTiling, InImageType, InImageSize, InTint)
+	{ }
+};
diff --git a/Engine/Source/Runtime/SlateCore/Public/Brushes/SlateNoResource.h b/Engine/Source/Runtime/SlateCore/Public/Brushes/SlateNoResource.h
new file mode 100644
--- /dev/null
+++ b/Engine/Source/Runtime/SlateCore/Public/Brushes/SlateNoResource.h
@@ -0,0 +1,24 @@
+// Copyright 1998-2014 Epic Games, Inc. All Rights Reserved.
+
+/*=============================================================================
+	SlateNoResource.h: Declares the FSlateNoResource structure.
+=============================================================================*/
+
+#pragma once
+
+
+/**
+ * A resource that has no appearance
+ */
+struct SLATECORE_API FSlateNoResource
+	: public FSlateBrush
+{
+	/**
+	 * Creates and initializes a new instance.
+	 *
+	 * @param InImageSize An optional image size (default is zero).
+	 */
+	FORCENOINLINE FSlateNoResource( const FVector2D& InImageSize = FVector2D::ZeroVector )
+		: FSlateBrush(ESlateBrushDrawType::NoDrawType, FName(NAME_None), FMargin(0), ESlateBrushTileType::NoTile, ESlateBrushImageType::NoImage, InImageSize)
+	{ }
+};
diff --git a/Engine/Source/Runtime/SlateCore/Public/Fonts/FontCache.h b/Engine/Source/Runtime/SlateCore/Public/Fonts/FontCache.h
new file mode 100644
--- /dev/null
+++ b/Engine/Source/Runtime/SlateCore/Public/Fonts/FontCache.h
@@ -0,0 +1,336 @@
+// Copyright 1998-2014 Epic Games, Inc. All Rights Reserved.
+
+#pragma once
+
+struct FCharacterRenderData;
+class FFreeTypeInterface;
+
+
+/** Information for rendering one character */
+struct FCharacterEntry
+{
+	/** Start X location of the character in the texture */
+	float StartU;
+	/** Start Y location of the character in the texture */
+	float StartV;
+	/** X Size of the character in the texture */
+	float USize;
+	/** Y Size of the character in the texture */
+	float VSize;
+	/** The vertical distance from the baseline to the topmost border of the character */
+	int16 VerticalOffset;
+	/** The vertical distance from the origin to the left most border of the character */
+	int16 HorizontalOffset;
+	/** The largest vertical distance below the baseline for any character in the font */
+	int16 GlobalDescender;
+	/** The amount to advance in X before drawing the next character in a string */
+	int16 XAdvance;
+	/** Index to a specific texture in the font cache. */
+	uint8 TextureIndex;
+	/** 1 if this entry is valid, 0 otherwize. */
+	uint8 Valid;
+
+	FCharacterEntry()
+	{
+		FMemory::Memzero( this, sizeof(FCharacterEntry) );
+	}
+
+	/**
+	 * Checks if the character entry is valid (I.E has been cached)
+	 * 
+	 * @param InEntry The entry to check
+	 */
+	FORCEINLINE bool IsValidEntry() const
+	{
+		return Valid == 1;
+	}
+};
+
+struct FKerningPair
+{
+	TCHAR First;
+	TCHAR Second;
+
+	FKerningPair( TCHAR InFirst, TCHAR InSecond )
+		: First(InFirst),Second(InSecond)
+	{
+	}
+
+	bool operator==(const FKerningPair& Other ) const
+	{
+		return First == Other.First && Second == Other.Second;
+	}
+
+	friend inline uint32 GetTypeHash( const FKerningPair& Key )
+	{
+		return FCrc::MemCrc32( &Key, sizeof(Key) );
+	}
+};
+
+
+class FSlateFontCache;
+
+/**
+ * A Kerning table for a single font key
+ */
+class FKerningTable
+{
+public:
+	FKerningTable( const FSlateFontKey& InFont, const FSlateFontCache& InFontCache );
+	~FKerningTable();
+
+	/**
+	 * Gets a kerning value for a pair of characters
+	 *
+	 * @param FirstChar		The first character in the pair
+	 * @param SecondChar	The second character in the pair
+	 * @return The kerning value
+	 */
+	int8 GetKerning( TCHAR FirstChar, TCHAR SecondChar );
+private:
+	/**
+	 * Allocated memory for the directly indexed kerning table
+	 */
+	void CreateDirectTable();
+private:
+	/** Extended kerning Pairs which are mapped to save memory */
+	TMap<FKerningPair,int8> MappedKerningPairs;
+	/** Direct access kerning table for ascii chars.  Note its very important that this stays small.  Kerni */
+	int8* DirectAccessTable;
+	/** Interface to freetype for accessing new kerning values */
+	const class FSlateFontCache& FontCache;
+	/** Font this kerning table is for */
+	FSlateFontKey FontKey;
+	/** Whether the font has kerning at all */
+	bool bHasKerning;
+};
+
+
+/**
+ * Manages a potentially large list of font characters
+ * Uses a directly indexed by TCHAR array until space runs out and then maps the rest to conserve memory
+ * Every character indexed by TCHAR could potentially cost a lot of memory of a lot of empty entries are created
+ * because characters being used are far apart
+ */
+class FCharacterList
+{
+public:
+	FCharacterList( const FSlateFontKey& InFontKey, const FSlateFontCache& InFontCache );
+
+	/* @return Is the character in this list */
+	bool IsValidIndex( TCHAR Character ) const
+	{
+		return DirectIndexEntries.IsValidIndex( Character ) || ( Character >= MaxDirectIndexedEntries && MappedEntries.Contains( Character ) );
+	}
+
+	/**
+	 * Gets data about how to render and measure a character 
+	 * Caching and atlasing it if needed
+	 *
+	 * @param Character	The character to get
+	 * @return Data about the character
+	 */
+	const FCharacterEntry& operator[]( TCHAR Character )
+	{
+		if( Character < MaxDirectIndexedEntries )
+		{
+			// The character can be indexed directly
+
+			int32 NumToAdd = (Character - DirectIndexEntries.Num()) + 1;
+			if( NumToAdd > 0 )
+			{
+				// The character doesnt exist yet in the index and there is not enough space
+				// Resize the array now
+				DirectIndexEntries.AddZeroed( NumToAdd );
+			}
+			
+			FCharacterEntry& CharacterEntry = DirectIndexEntries[ Character ];
+			if( !CharacterEntry.IsValidEntry() )
+			{
+				// Character has not been cached yet
+				CharacterEntry = CacheCharacter( Character );
+			}
+
+			return CharacterEntry;
+		}
+		else
+		{
+			FCharacterEntry& CharacterEntry = MappedEntries.FindOrAdd( Character );
+			if( !CharacterEntry.IsValidEntry() )
+			{
+				// Character has not been cached yet
+				CharacterEntry = CacheCharacter( Character );
+			}
+
+			return CharacterEntry;
+		}
+	}
+
+	/**
+	 * Gets a kerning value for a pair of characters
+	 *
+	 * @param FirstChar		The first character in the pair
+	 * @param SecondChar	The second character in the pair
+	 * @return The kerning value
+	 */
+	int8 GetKerning( TCHAR First, TCHAR Second );
+
+	/**
+	 * @return The global max height for any character in this font
+	 */
+	uint16 GetMaxHeight() const;
+
+	/** 
+	 * Returns the baseline for the font used by this character 
+	 *
+	 * @return The offset from the bottom of the max character height to the baseline. Be aware that the value will be negative.
+	 */
+	int16 GetBaseline() const;
+
+private:
+
+	/**
+	 * Caches a new character
+	 * 
+	 * @param Character	The character to cache
+	 */
+	FCharacterEntry& CacheCharacter( TCHAR Character );
+
+
+private:
+
+	/** Entries for larger character sets to conserve memory */
+	TMap<TCHAR, FCharacterEntry> MappedEntries; 
+	/** Directly indexed entries for fast lookup */
+	TArray<FCharacterEntry> DirectIndexEntries;
+	/** Table of kerning values for this font */
+	FKerningTable KerningTable;
+	/** Font for this character list */
+	FSlateFontKey FontKey;
+	/** Reference to the font cache for accessing new unseen characters */
+	const class FSlateFontCache& FontCache;
+	/** Number of directly indexed entries */
+	int32 MaxDirectIndexedEntries;
+	/** The global max height for any character in this font */
+	mutable uint16 MaxHeight;
+	/** The offset from the bottom of the max character height to the baseline. */
+	mutable uint16 Baseline;
+};
+
+/*
+ * Font caching implementation
+ * Caches characters into textures as needed
+ */
+class SLATECORE_API FSlateFontCache
+{
+public:
+	/**
+	 * Constructor
+	 *
+	 * @param InTextureSize The size of the atlas texture
+	 * @param InFontAlas	Platform specific font atlas resource
+	 */
+	FSlateFontCache( TSharedRef<ISlateFontAtlasFactory> InFontAtlasFactory );
+	virtual ~FSlateFontCache();
+
+	/** 
+	 * Gets information for how to draw all characters in the specified string. Caches characters as they are found
+	 * 
+	 * @param Text			The string to get character information for
+	 * @param InFontInfo	Information about the font that the string is drawn with
+	 * @param FontScale	The scale to apply to the font
+	 * @param OutCharacterEntries	Populated array of character entries. Indices of characters in Text match indices in this array
+	 */
+	class FCharacterList& GetCharacterList( const FSlateFontInfo &InFontInfo, float FontScale ) const;
+
+	/** 
+	 * Add a new entries into a cache atlas
+	 *
+	 * @param InFontInfo	Information about the font being used for the characters
+	 * @param Characters	The characters to cache
+	 * @param FontScale		The font scale to use
+	 * @return true if the characters could be cached. false if the cache is full
+	 */
+	bool AddNewEntry( TCHAR Character, const FSlateFontKey& InKey, FCharacterEntry& OutCharacterEntry ) const;
+
+	/** 
+	 * Flush the cache if needed
+	 */
+	void ConditionalFlushCache();
+
+	/**
+	 * Updates the texture used for rendering
+	 */
+	void UpdateCache();
+
+	/**
+	 * Releases rendering resources
+	 */
+	void ReleaseResources();
+
+	/**
+	 * Get the texture resource for a font atlas at a given index
+	 * 
+	 * @param Index	The index of the texture 
+	 * @return Handle to the texture resource
+	 */
+	FSlateShaderResource* GetTextureResource( uint32 Index ) { return FontAtlases[Index]->GetTexture(); }
+
+	/**
+	 * Calculates the kerning amount for a pair of characters
+	 *
+	 * @param First		The first character in the pair
+	 * @param Second	The second character in the pair
+	 * @oaran FontKey	The font to get kerning from
+	 * @return The kerning amount, 0 if no kerning
+	 */
+	int8 GetKerning( TCHAR First, TCHAR Second, const FSlateFontKey& FontKey ) const;
+
+	/**
+	 * @return Whether or not a font has kerning information
+	 */
+	bool HasKerning( const FSlateFontInfo& InFontInfo ) const;
+
+	/**
+	 * Returns the height of the largest character in the font. 
+	 *
+	 * @param InFontInfo	A descriptor of the font to get character size for 
+	 * @param FontScale		The scale to apply to the font
+	 * 
+	 * @return The largest character height
+	 */
+	uint16 GetMaxCharacterHeight( const FSlateFontInfo& InFontInfo, float FontScale ) const;
+
+	/**
+	 * Returns the baseline for the specified font.
+	 *
+	 * @param InFontInfo	A descriptor of the font to get character size for 
+	 * @param FontScale		The scale to apply to the font
+	 * 
+	 * @return The offset from the bottom of the max character height to the baseline.
+	 */
+	uint16 GetBaseline( const FSlateFontInfo& InFontInfo, float FontScale ) const;
+
+	/**
+	 * Clears all cached data from the cache
+	 */
+	void FlushCache() const;
+
+
+private:
+
+	/** Mapping Font keys to cached data */
+	mutable TMap<FSlateFontKey, TSharedRef< class FCharacterList > > FontToCharacterListCache;
+
+	/** Interface to the freetype library */
+	mutable TSharedRef<FFreeTypeInterface> FTInterface;
+
+	/** Array of all font atlases */
+	mutable TArray< TSharedRef<FSlateFontAtlas> > FontAtlases;
+
+	/** Factory for creating new font atlases */
+	TSharedRef<ISlateFontAtlasFactory> FontAtlasFactory;
+
+	/** Whether or not we have a pending request to flush the cache when it is safe to do so */
+	mutable bool bFlushRequested;
+};
diff --git a/Engine/Source/Runtime/SlateCore/Public/Fonts/FontTypes.h b/Engine/Source/Runtime/SlateCore/Public/Fonts/FontTypes.h
new file mode 100644
--- /dev/null
+++ b/Engine/Source/Runtime/SlateCore/Public/Fonts/FontTypes.h
@@ -0,0 +1,97 @@
+// Copyright 1998-2014 Epic Games, Inc. All Rights Reserved.
+#pragma once
+
+struct SLATECORE_API FSlateFontKey
+{
+	FSlateFontInfo FontInfo;
+	float Scale;
+
+	FSlateFontKey( const FSlateFontInfo& InInfo, float InScale )
+		: FontInfo( InInfo )
+		, Scale( InScale )
+	{
+
+	}
+
+	bool operator==(const FSlateFontKey& Other ) const
+	{
+		return FontInfo == Other.FontInfo && Scale == Other.Scale;
+	}
+
+	friend inline uint32 GetTypeHash( const FSlateFontKey& Key )
+	{
+		return GetTypeHash(Key.FontInfo) + FCrc::MemCrc32( &Key.Scale, sizeof(float) );
+	}
+};
+
+/** Measurement details for a specific character */
+struct FCharacterMeasureInfo
+{
+	/** Width of the character in pixels */
+	int16 SizeX;
+	/** Height of the character in pixels */
+	int16 SizeY;
+	/** How many pixels to advance in X after drawing this character (when its part of a string)*/
+	int16 XAdvance;
+	/** The largest vertical distance below the baseline for any character in the font */
+	int16 GlobalDescender;
+	/** The vertical distance from the baseline to the topmost border of the glyph bitmap */
+	int16 VerticalOffset;
+	/** The horizontal distance from the origin to the leftmost border of the character */
+	int16 HorizontalOffset;
+	FCharacterMeasureInfo( int16 InSizeX = 0, int16 InSizeY = 0, int16 InXAdvance = 0, int16 InVerticalOffset = 0, int16 InHorizontalOffset = 0 )
+		: SizeX(InSizeX),SizeY(InSizeY),XAdvance(InXAdvance),VerticalOffset(InVerticalOffset), HorizontalOffset(InHorizontalOffset)
+	{
+	}
+};
+
+/** Contains pixel data for a character rendered from freetype as well as measurement info */
+struct FCharacterRenderData
+{
+	FCharacterMeasureInfo MeasureInfo;
+	/** Measurement data for this character */
+	/** Raw pixels created by freetype */
+	TArray<uint8> RawPixels;
+	/** @todo Doesnt belong here. */
+	uint16 MaxHeight;
+	/** The character that was rendered */
+	TCHAR Char;	
+};
+
+/** 
+ * Representation of a texture for fonts in which characters are packed tightly based on their bounding rectangle 
+ */
+class SLATECORE_API FSlateFontAtlas : public FSlateTextureAtlas
+{
+public:
+	FSlateFontAtlas( uint32 InWidth, uint32 InHeight );
+	virtual ~FSlateFontAtlas();
+
+	/**
+	 * Returns the texture resource used by slate
+	 */
+	virtual class FSlateShaderResource* GetTexture() = 0;
+	
+	/**
+	 * Releases rendering resources of this cache
+	 */
+	virtual void ReleaseResources(){}
+
+	/**
+	 * Flushes all cached data.
+	 */
+	void Flush();
+
+	/** 
+	 * Adds a character to the texture.
+	 *
+	 * @param CharInfo	Information about the size of the character
+	 */
+	const struct FAtlasedTextureSlot* AddCharacter( const FCharacterRenderData& CharInfo );
+};
+
+class SLATECORE_API ISlateFontAtlasFactory
+{
+public:
+	virtual TSharedRef<class FSlateFontAtlas> CreateFontAtlas() const = 0;
+};
\ No newline at end of file
diff --git a/Engine/Source/Runtime/SlateCore/Public/Fonts/SlateFontInfo.h b/Engine/Source/Runtime/SlateCore/Public/Fonts/SlateFontInfo.h
new file mode 100644
--- /dev/null
+++ b/Engine/Source/Runtime/SlateCore/Public/Fonts/SlateFontInfo.h
@@ -0,0 +1,94 @@
+// Copyright 1998-2014 Epic Games, Inc. All Rights Reserved.
+
+/*=============================================================================
+	SlateFontInfo.h: Declares the FSlateFontInfo structure.
+=============================================================================*/
+
+#pragma once
+
+#include "SlateFontInfo.generated.h"
+
+
+/**
+ * A representation of a font in Slate.
+ */
+USTRUCT()
+struct SLATECORE_API FSlateFontInfo
+{
+	GENERATED_USTRUCT_BODY()
+
+	/** The name of the font */
+	UPROPERTY(EditAnywhere, Category=SlateStyleRules)
+	FName FontName;
+
+	/** The size of the font */
+	UPROPERTY(EditAnywhere, Category=SlateStyleRules)
+	int32 Size;
+
+public:
+
+	/**
+	 * Default constructor.
+	 */
+	FSlateFontInfo();
+
+	/**
+	 * Creates and initializes a new instance with the specified font name and size.
+	 *
+	 * @param InFontName The name of the font.
+	 * @param InSize The size of the font.
+	 */
+	FSlateFontInfo( const FString& InFontName, uint16 InSize );
+
+	/**
+	 * Creates and initializes a new instance with the specified font name and size.
+	 *
+	 * @param InFontName The name of the font.
+	 * @param InSize The size of the font.
+	 */
+	FSlateFontInfo( const FName& InFontName, uint16 InSize );
+
+	/**
+	 * Creates and initializes a new instance with the specified font name and size.
+	 *
+	 * @param InFontName The name of the font.
+	 * @param InSize The size of the font.
+	 */
+	FSlateFontInfo( const ANSICHAR* InFontName, uint16 InSize );
+
+	/**
+	 * Creates and initializes a new instance with the specified font name and size.
+	 *
+	 * @param InFontName The name of the font.
+	 * @param InSize The size of the font.
+	 */
+	FSlateFontInfo( const WIDECHAR* InFontName, uint16 InSize );
+
+public:
+
+	/**
+	 * Compares this font info with another for equality.
+	 *
+	 * @param Other The other font info.
+	 *
+	 * @return true if the two font infos are equal, false otherwise.
+	 */
+	bool operator==( const FSlateFontInfo& Other ) const 
+	{
+		return FontName == Other.FontName && Size == Other.Size;
+	}
+
+	/**
+	 * Calculates a type hash value for a font info.
+	 *
+	 * Type hashes are used in certain collection types, such as TMap.
+	 *
+	 * @param FontInfo The font info to calculate the hash for.
+	 *
+	 * @return The hash value.
+	 */
+	friend inline uint32 GetTypeHash( const FSlateFontInfo& FontInfo )
+	{
+		return GetTypeHash(FontInfo.FontName) + (FontInfo.Size << 17);
+	}
+};
diff --git a/Engine/Source/Runtime/SlateCore/Public/Input/CursorReply.h b/Engine/Source/Runtime/SlateCore/Public/Input/CursorReply.h
new file mode 100644
--- /dev/null
+++ b/Engine/Source/Runtime/SlateCore/Public/Input/CursorReply.h
@@ -0,0 +1,54 @@
+// Copyright 1998-2014 Epic Games, Inc. All Rights Reserved.
+
+/*=============================================================================
+	CursorReply.h: Declares the FCursorReply class.
+=============================================================================*/
+
+#pragma once
+
+
+/**
+ * A reply to the OnQueryCursor event.
+ */
+class FCursorReply
+{
+public:
+	/** Makes a NULL response; this response will not change the previous value. */
+	static FCursorReply Unhandled()
+	{
+		return FCursorReply();
+	}
+		
+	/** Respond with a cursor type. */
+	static FCursorReply Cursor( EMouseCursor::Type InCursor )
+	{
+		return FCursorReply( InCursor );
+	}
+		
+	/** @return true if this this reply is a result of the event being handled; false otherwise. */
+	bool IsEventHandled(){ return bIsHandled; }
+		
+	/** @return The requested MouseCursor when the event was handled. Undefined otherwise. */
+	EMouseCursor::Type GetCursor() { return MouseCursor; }
+		
+private:
+	/** Internal constructor - makes a NULL result. */
+	FCursorReply()
+	: bIsHandled(false)
+	, MouseCursor( EMouseCursor::Default )
+	{
+	}
+		
+	/** Internal constructor - makes a non-NULL result. */
+	FCursorReply( EMouseCursor::Type InCursorType )
+	: bIsHandled(true)
+	, MouseCursor( InCursorType )
+	{		
+	}
+		
+	/** Does this reply have any meaning */ 
+	bool bIsHandled;
+		
+	/** The value of the reply, if bHasValue is true. */
+	EMouseCursor::Type MouseCursor;
+};
diff --git a/Engine/Source/Runtime/SlateCore/Public/Input/Events.h b/Engine/Source/Runtime/SlateCore/Public/Input/Events.h
new file mode 100644
--- /dev/null
+++ b/Engine/Source/Runtime/SlateCore/Public/Input/Events.h
@@ -0,0 +1,603 @@
+// Copyright 1998-2014 Epic Games, Inc. All Rights Reserved.
+
+#pragma once
+
+class SWindow;
+
+
+/**
+ * Context for keyboard focus change
+ */
+namespace EKeyboardFocusCause
+{
+	enum Type
+	{
+		/** Keyboard focus was changed because of a mouse action */
+		Mouse,
+
+		/** Keyboard focus was changed in response to a keystroke, such as the TAB key. */
+		Keyboard,
+
+		/** Keyboard focus was changed because someone asked the application to change it. */
+		SetDirectly,
+
+		/** Keyboard focus was explicitly cleared via the escape key or other similar action */
+		Cleared,
+
+		/** Keyboard focus was changed because another widget lost focus, and focus moved to a new widget */
+		OtherWidgetLostFocus,
+
+		/** Keyboard focus was set in response to the owning window being activated */
+		WindowActivate,
+	};
+};
+
+
+/**
+ * FKeyboardFocusEvent is used when notifying widgets about keyboard focus changes
+ * It is passed to event handlers dealing with keyboard focus
+ */
+class FKeyboardFocusEvent
+{
+
+public:
+
+	/**
+	 * Constructor.  Events are immutable once constructed.
+	 *
+	 * @param  InCause  The cause of the focus event
+	 */
+	FKeyboardFocusEvent( const EKeyboardFocusCause::Type InCause )
+		: Cause( InCause )
+	{ }
+
+
+	/**
+	 * Queries the reason for the focus change
+	 *
+	 * @return  The cause of the keyboard focus change
+	 */
+	EKeyboardFocusCause::Type GetCause() const
+	{
+		return Cause;
+	}
+
+
+private:
+
+	/** The cause of the focus change */
+	EKeyboardFocusCause::Type Cause;
+};
+
+
+/**
+ * Base class for all mouse and keyboard events.
+ */
+class FInputEvent
+{
+public:
+
+	/**
+	 * Constructor.  Events are immutable once constructed.
+	 *
+	 * @param  InModifierKeys  Modifier key state for this event
+	 * @param  bInIsRepeat  True if this key is an auto-repeated keystroke
+	 */
+	FInputEvent( const FModifierKeysState& InModifierKeys, const bool bInIsRepeat )
+		: ModifierKeys(InModifierKeys)
+		, bIsRepeat(bInIsRepeat)
+		, EventPath(nullptr)
+	{ }
+
+	/**
+	 * Virtual destructor.
+	 */
+	virtual ~FInputEvent( ) { }
+
+	/**
+	 * Returns whether or not this character is an auto-repeated keystroke
+	 *
+	 * @return  True if this character is a repeat
+	 */
+	bool IsRepeat() const
+	{
+		return bIsRepeat;
+	}
+
+	/**
+	 * Returns true if either shift key was down when this event occurred
+	 *
+	 * @return  True if shift is pressed
+	 */
+	bool IsShiftDown() const
+	{
+		return ModifierKeys.IsShiftDown();
+	}
+
+	/**
+	 * Returns true if left shift key was down when this event occurred
+	 *
+	 * @return True if left shift is pressed.
+	 */
+	bool IsLeftShiftDown() const
+	{
+		return ModifierKeys.IsLeftShiftDown();
+	}
+
+	/**
+	 * Returns true if right shift key was down when this event occurred
+	 *
+	 * @return True if right shift is pressed.
+	 */
+	bool IsRightShiftDown() const
+	{
+		return ModifierKeys.IsRightShiftDown();
+	}
+
+	/**
+	 * Returns true if either control key was down when this event occurred
+	 *
+	 * @return  True if control is pressed
+	 */
+	bool IsControlDown() const
+	{
+		return ModifierKeys.IsControlDown();
+	}
+
+	/**
+	 * Returns true if left control key was down when this event occurred
+	 *
+	 * @return  True if left control is pressed
+	 */
+	bool IsLeftControlDown() const
+	{
+		return ModifierKeys.IsLeftControlDown();
+	}
+
+	/**
+	 * Returns true if right control key was down when this event occurred
+	 *
+	 * @return  True if right control is pressed
+	 */
+	bool IsRightControlDown() const
+	{
+		return ModifierKeys.IsRightControlDown();
+	}
+
+	/**
+	 * Returns true if either alt key was down when this event occurred
+	 *
+	 * @return  True if alt is pressed
+	 */
+	bool IsAltDown() const
+	{
+		return ModifierKeys.IsAltDown();
+	}
+
+	/**
+	 * Returns true if left alt key was down when this event occurred
+	 *
+	 * @return  True if left alt is pressed
+	 */
+	bool IsLeftAltDown() const
+	{
+		return ModifierKeys.IsLeftAltDown();
+	}
+
+	/**
+	 * Returns true if right alt key was down when this event occurred
+	 *
+	 * @return  True if right alt is pressed
+	 */
+	bool IsRightAltDown() const
+	{
+		return ModifierKeys.IsRightAltDown();
+	}
+
+	/** The event path provides additional context for handling */
+	FGeometry FindGeometry( const TSharedRef<SWidget>& WidgetToFind ) const
+	{
+		return EventPath->FindArrangedWidget(WidgetToFind).Geometry;
+	}
+
+	TSharedRef<SWindow> GetWindow() const
+	{
+		return EventPath->GetWindow();
+	}
+
+	/** Set the widget path along which this event will be routed */
+	void SetEventPath( const FWidgetPath& InEventPath )
+	{
+		EventPath = &InEventPath;
+	}
+
+protected:
+
+	/** State of modifier keys when this event happened. */
+	FModifierKeysState ModifierKeys;
+
+	/** True if this key was auto-repeated */
+	bool bIsRepeat;
+
+	/** Events are sent along paths. See GetEventPath() */
+	const FWidgetPath* EventPath;
+};
+
+
+/**
+ * FKeyboardEvent describes a keyboard action (key pressed or released.)
+ * It is passed to event handlers dealing with keyboard input.
+ */
+class FKeyboardEvent : public FInputEvent
+{
+public:
+
+	/**
+	 * Constructor.  Events are immutable once constructed.
+	 *
+	 * @param  InKeyName  Character name
+	 * @param  InModifierKeys  Modifier key state for this event
+	 * @param  bInIsRepeat  True if this key is an auto-repeated keystroke
+	 */
+	FKeyboardEvent( const FKey InKey,
+					const FModifierKeysState& InModifierKeys, 
+					const bool bInIsRepeat,
+					const uint32 InCharacterCode )
+		: FInputEvent( InModifierKeys, bInIsRepeat )
+		, Key( InKey )
+		, CharacterCode( InCharacterCode )
+	{ }
+	
+	/**
+	 * Returns the name of the key for this event
+	 *
+	 * @return  Key name
+	 */
+	FKey GetKey() const
+	{
+		return Key;
+	}
+
+	/**
+	 * Returns the character code for this event
+	 *
+	 * @return  Character code or 0 if this event was not a character key press
+	 */
+	uint32 GetCharacter() const
+	{
+		return CharacterCode;
+	}
+
+private:
+	/** Name of the key that was pressed */
+	FKey Key;
+	/** The character code of the key that was pressed.  Only applicable to typed character keys, 0 otherwise*/
+	uint32 CharacterCode;
+};
+
+
+/**
+ * FCharacterEvent describes a keyboard action where the utf-16 code is given.  Used for OnKeyChar messages
+ */
+class FCharacterEvent : public FInputEvent
+{
+public:
+	FCharacterEvent( const TCHAR InCharacter, const FModifierKeysState& InModifierKeys, const bool bInIsRepeat )
+		: FInputEvent( InModifierKeys, bInIsRepeat )
+		, Character( InCharacter )
+	{ }
+
+	/**
+	 * Returns the character for this event
+	 *
+	 * @return  Character
+	 */
+	TCHAR GetCharacter() const
+	{
+		return Character;
+	}
+
+private:
+
+	/** The character that was pressed */
+	TCHAR Character;
+};
+
+
+/**
+ * Helper class to auto-populate a set with the set of "keys" a touch represents
+ */
+class FTouchKeySet : public TSet<FKey>
+{
+public:
+
+	FTouchKeySet(FKey Key)
+	{
+		this->Add(Key);
+	}
+};
+
+
+/**
+ * FPointerEvent describes a mouse or touch action (e.g. Press, Release, Move, etc).
+ * It is passed to event handlers dealing with pointer-based input.
+ */
+class FPointerEvent : public FInputEvent
+{
+public:
+
+	/** Events are immutable once constructed. */
+	FPointerEvent(
+		const FVector2D& InScreenSpacePosition,
+		const FVector2D& InLastScreenSpacePosition,
+		const TSet<FKey>& InPressedButtons,
+		FKey InEffectingButton,
+		float InWheelDelta,
+		const FModifierKeysState& InModifierKeys
+		)
+		: FInputEvent( InModifierKeys, false )
+		, ScreenSpacePosition( InScreenSpacePosition )
+		, LastScreenSpacePosition( InLastScreenSpacePosition )
+		, CursorDelta( InScreenSpacePosition - InLastScreenSpacePosition )
+		, PressedButtons( InPressedButtons )
+		, EffectingButton (InEffectingButton )
+		, UserIndex( 0 )
+		, PointerIndex( 0 )
+		, TouchpadIndex( 0 )
+		, bIsTouchEvent( false )
+		, GestureType( EGestureEvent::None )
+		, WheelOrGestureDelta( 0, InWheelDelta )
+	{ }
+
+	/** A constructor for raw mouse events */
+	FPointerEvent(
+		const FVector2D& InScreenSpacePosition,
+		const FVector2D& InLastScreenSpacePosition,
+		const FVector2D& InDelta,
+		const TSet<FKey>& InPressedButtons,
+		const FModifierKeysState& InModifierKeys
+		)
+		: FInputEvent( InModifierKeys, false )
+		, ScreenSpacePosition( InScreenSpacePosition )
+		, LastScreenSpacePosition( InLastScreenSpacePosition )
+		, CursorDelta( InDelta )
+		, PressedButtons( InPressedButtons )
+		, UserIndex( 0 )
+		, PointerIndex( 0 )
+		, TouchpadIndex( 0 )
+		, bIsTouchEvent( false )
+		, GestureType( EGestureEvent::None )
+		, WheelOrGestureDelta( 0, 0 )
+	{ }
+
+	/** A constructor for touch events */
+	FPointerEvent(
+		uint32 InUserIndex,
+		uint32 InPointerIndex,
+		const FVector2D& InScreenSpacePosition,
+		const FVector2D& InLastScreenSpacePosition,
+		bool bPressLeftMouseButton,
+		const FModifierKeysState& InModifierKeys = FModifierKeysState(false, false, false, false, false, false),
+		uint32 InTouchpadIndex=0
+		)
+		: FInputEvent( InModifierKeys, false )
+		, ScreenSpacePosition( InScreenSpacePosition )
+		, LastScreenSpacePosition( InLastScreenSpacePosition )
+		, CursorDelta( LastScreenSpacePosition - ScreenSpacePosition )
+		, PressedButtons( bPressLeftMouseButton ? StandardTouchKeySet : EmptyTouchKeySet )
+		, EffectingButton ( EKeys::LeftMouseButton )
+		, UserIndex( InUserIndex )
+		, PointerIndex( InPointerIndex )
+		, TouchpadIndex( InTouchpadIndex )
+		, bIsTouchEvent( true )
+		, GestureType( EGestureEvent::None )
+		, WheelOrGestureDelta( 0, 0 )
+	{ }
+
+	/** A constructor for gesture events */
+	FPointerEvent(
+		const FVector2D& InScreenSpacePosition,
+		const FVector2D& InLastScreenSpacePosition,
+		const TSet<FKey>& InPressedButtons,
+		const FModifierKeysState& InModifierKeys,
+		EGestureEvent::Type InGestureType,
+		const FVector2D& InGestureDelta
+		)
+		: FInputEvent( InModifierKeys, false )
+		, ScreenSpacePosition( InScreenSpacePosition )
+		, LastScreenSpacePosition( InLastScreenSpacePosition )
+		, CursorDelta( LastScreenSpacePosition - ScreenSpacePosition )
+		, PressedButtons( InPressedButtons )
+		, UserIndex( 0 )
+		, PointerIndex( 0 )
+		, bIsTouchEvent( false )
+		, GestureType( InGestureType )
+		, WheelOrGestureDelta( InGestureDelta )
+	{ }
+	
+	/** @return The position of the cursor in screen space */
+	const FVector2D& GetScreenSpacePosition() const { return ScreenSpacePosition; }
+
+	/** @return The position of the cursor in screen space last time we handled an input event */
+	const FVector2D& GetLastScreenSpacePosition() const { return LastScreenSpacePosition; }
+
+	/** @return the distance the mouse traveled since the last event was handled. */
+	FVector2D GetCursorDelta() const { return CursorDelta; }
+
+	/** Mouse buttons that are currently pressed */
+	bool IsMouseButtonDown( FKey MouseButton ) const { return PressedButtons.Contains( MouseButton ); }
+
+	/** Mouse button that caused this event to be raised (possibly EB_None) */
+	FKey GetEffectingButton() const { return EffectingButton; }
+	
+	/** How much did the mouse wheel turn since the last mouse event */
+	float GetWheelDelta() const { return WheelOrGestureDelta.Y; }
+
+	/** @return The index of the user that caused the event */
+	int32 GetUserIndex() const { return UserIndex; }
+
+	/** @return The unique identifier of the pointer (e.g., finger index) */
+	uint32 GetPointerIndex() const { return PointerIndex; }
+
+	/** @return The index of the touchpad that generated this event (for platforms with multiple touchpads per user) */
+	uint32 GetTouchpadIndex() const { return TouchpadIndex; }
+
+	/** @return Is this event a result from a touch (as opposed to a mouse) */
+	uint32 IsTouchEvent() const { return bIsTouchEvent; }
+
+	/** @return The type of touch gesture */
+	EGestureEvent::Type GetGestureType() const { return GestureType; }
+
+	/** @return The change in gesture value since the last gesture event of the same type. */
+	const FVector2D& GetGestureDelta() const { return WheelOrGestureDelta; }
+
+private:
+
+	FVector2D ScreenSpacePosition;
+	FVector2D LastScreenSpacePosition;
+	FVector2D CursorDelta;
+	const TSet<FKey>& PressedButtons;
+	FKey EffectingButton;
+	uint32 UserIndex;
+	uint32 PointerIndex;
+	uint32 TouchpadIndex;
+	bool bIsTouchEvent;
+	EGestureEvent::Type GestureType;
+	FVector2D WheelOrGestureDelta;
+
+	/** A reusable set of keys for touch events (since this class keeps a reference to a TSet) */
+	static FTouchKeySet StandardTouchKeySet;
+	static FTouchKeySet EmptyTouchKeySet;
+};
+
+
+/**
+ * FControllerEvent describes a controller action (e.g. Button Press, Release, Analog stick move, etc).
+ * It is passed to event handlers dealing with controller input.
+ */
+class FControllerEvent : public FInputEvent
+{
+public:
+
+	FControllerEvent(
+		FKey InEffectingButton,
+		int32 InUserIndex,
+		float InAnalogValue,
+		bool bIsRepeat )
+		: FInputEvent( FModifierKeysState(false,false,false,false,false,false), bIsRepeat )
+		, EffectingButton( InEffectingButton )
+		, UserIndex( InUserIndex )
+		, AnalogValue( InAnalogValue )
+	{ }
+
+	/** @return The controller button that caused this event */
+	FKey GetEffectingButton() const { return EffectingButton; }
+
+	/** @return The index of the user that caused the event */
+	uint32 GetUserIndex() const { return UserIndex; }
+
+	/** @return Analog value between 0 and 1.  1 being fully pressed, 0 being not pressed at all */
+	float GetAnalogValue() const { return AnalogValue; }
+
+private:
+
+	/** The controller button that caused this event */
+	FKey EffectingButton;
+	/** The index of the user that caused the event */
+	uint32 UserIndex;
+	/**  Analog value between 0 and 1.  1 being fully pressed, 0 being not pressed at all */
+	float AnalogValue;
+};
+
+
+/**
+ * FMotionEvent describes a touchpad action (press, move, lift)
+ * It is passed to event handlers dealing with touch input.
+ */
+class FMotionEvent : public FInputEvent
+{
+public:
+	FMotionEvent(
+		uint32 InUserIndex,
+		const FVector& InTilt, 
+		const FVector& InRotationRate, 
+		const FVector& InGravity, 
+		const FVector& InAcceleration
+		)
+		: FInputEvent( FModifierKeysState(false,false,false,false,false,false), false )
+		, UserIndex( InUserIndex )
+		, Tilt( InTilt )
+		, RotationRate( InRotationRate )
+		, Gravity( InGravity )
+		, Acceleration( InAcceleration )
+	{ }
+
+	/** @return The index of the user that caused the event */
+	uint32 GetUserIndex() const { return UserIndex; }
+
+	/** @return Current tilt of the device/controller */
+	const FVector& GetTilt() const { return Tilt; }
+
+	/** @return Rotation speed */
+	const FVector& GetRotationRate() const { return RotationRate; }
+
+	/** @return Gravity vector (pointing down into the ground) */
+	const FVector& GetGravity() const { return Gravity; }
+
+	/** @return 3D acceleration of the device */
+	const FVector& GetAcceleration() const { return Acceleration; }
+
+private:
+
+	/** The index of the user that caused the event */
+	uint32 UserIndex;
+	/** The current tilt of the device/controller */
+	FVector Tilt;
+	/** The rotation speed */
+	FVector RotationRate;
+	/** The gravity vector (pointing down into the ground) */
+	FVector Gravity;
+	/** The 3D acceleration of the device */
+	FVector Acceleration;
+};
+
+
+/**
+ * FWindowActivateEvent describes a window being activated or deactivated.
+ * (i.e. brought to the foreground or moved to the background)
+ * This event is only passed to top level windows; most widgets are incapable
+ * of receiving this event.
+ */
+class FWindowActivateEvent
+{
+public:
+
+	enum EActivationType
+	{
+		EA_Activate,
+		EA_ActivateByMouse,
+		EA_Deactivate
+	};
+
+	FWindowActivateEvent( EActivationType InActivationType, TSharedRef<SWindow> InAffectedWindow )
+		: ActivationType(InActivationType)
+		, AffectedWindow(InAffectedWindow)
+	{ }
+
+	/** Decscribes what actually happened to the window (e.g. Activated, Deactivated, Activated by a mouse click) */
+	EActivationType GetActivationType() const
+	{
+		return ActivationType;
+	}
+
+	/** The window that this activation/deactivation happened to */
+	TSharedRef<SWindow> GetAffectedWindow() const
+	{
+		return AffectedWindow;
+	}
+
+private:
+
+	EActivationType ActivationType;
+	TSharedRef<SWindow> AffectedWindow;
+};
diff --git a/Engine/Source/Runtime/SlateCore/Public/Input/Reply.h b/Engine/Source/Runtime/SlateCore/Public/Input/Reply.h
new file mode 100644
--- /dev/null
+++ b/Engine/Source/Runtime/SlateCore/Public/Input/Reply.h
@@ -0,0 +1,275 @@
+// Copyright 1998-2014 Epic Games, Inc. All Rights Reserved.
+
+#pragma once
+
+
+// @todo Slate: this is highly sketchy, as FDragDropOperation is declared in Slate,
+// but cannot currently be brought into SlateCore due to dependencies to SWindow.
+class FDragDropOperation;
+class SWidget;
+
+
+/**
+ * A Reply is something that a Slate event returns to the system to notify it about certain aspect of how an event was handled.
+ * For example, a widget may handle an OnMouseDown event by asking the system to give mouse capture to a specific Widget.
+ * To do this, return FReply::CaptureMouse( NewMouseCapture ).
+ */
+class FReply
+{
+public:
+		
+	/** An event should return a FReply::Handled().CaptureMouse( SomeWidget ) as a means of asking the system to forward all mouse events to SomeWidget */
+	FReply& CaptureMouse( TSharedRef<SWidget> InMouseCaptor )
+	{
+		this->MouseCaptor = InMouseCaptor;
+		return *this;
+	}
+
+	FReply& CaptureJoystick( TSharedRef<SWidget> InJoystickCaptor, bool bInAllJoysticks = false )
+	{
+		this->JoystickCaptor = InJoystickCaptor;
+		this->bAllJoysticks = bInAllJoysticks;
+		return *this;
+	}
+
+	/**
+	 * Enables the use of high precision (raw input) mouse movement, for more accurate mouse movement without mouse ballistics
+	 * NOTE: This implies mouse capture and hidden mouse movement.  Releasing mouse capture deactivates this mode.
+	 */
+	FReply& UseHighPrecisionMouseMovement( TSharedRef<SWidget> InMouseCaptor )
+	{
+		this->MouseCaptor = InMouseCaptor;
+		this->bUseHighPrecisionMouse = true;
+		return *this;
+	}
+
+	/**
+	 * An even should return FReply::Handled().SetMousePos to ask Slate to move the mouse cursor to a different location
+	 */
+	FReply& SetMousePos( const FIntPoint& NewMousePos )
+	{
+		this->RequestedMousePos = NewMousePos;
+		return *this;
+	}
+
+	/** An event should return FReply::Handled().SetKeyboardFocus( SomeWidget ) as a means of asking the system to set keyboard focus to the provided widget*/
+	FReply& SetKeyboardFocus( TSharedRef<SWidget> GiveMeFocus, EKeyboardFocusCause::Type ReasonFocusIsChanging )
+	{
+		this->FocusRecipient = GiveMeFocus;
+		this->FocusChangeReason = ReasonFocusIsChanging;
+		return *this;
+	}
+		
+	/**
+	 * An event should return FReply::Handled().LockMouseToWidget( SomeWidget ) as a means of asking the system to 
+	 * Lock the mouse so it cannot move out of the bounds of the widget.
+	 */
+	FReply& LockMouseToWidget( TSharedRef<SWidget> InWidget )
+	{
+		this->MouseLockWidget = InWidget;
+		this->bShouldReleaseMouseLock = false;
+		return *this;
+	}
+
+	/** 
+	 * An event should return a FReply::Handled().ReleaseMouseLock() to ask the system to release mouse lock on a widget
+	 */
+	FReply& ReleaseMouseLock()
+	{
+		this->bShouldReleaseMouseLock = true;
+		MouseLockWidget.Reset();
+		return *this;
+	}
+
+	/** 
+	 * An event should return a FReply::Handled().ReleaseMouse() to ask the system to release mouse capture
+	 * NOTE: Deactivates high precision mouse movement if activated.
+	 */
+	FReply& ReleaseMouseCapture()
+	{
+		this->bReleaseMouseCapture = true;
+		this->bUseHighPrecisionMouse = false;
+		return *this;		
+	}
+
+	/**
+	 * An event should return a FReply::Handled().ReleaseJoystickCapture() to ask the system to release joystick capture
+	 */
+	FReply& ReleaseJoystickCapture(bool bInAllJoysticks = false)
+	{
+		this->bReleaseJoystickCapture = true;
+		this->bAllJoysticks = bInAllJoysticks;
+		return *this;		
+	}
+
+	/**
+	 * Ask Slate to detect if a user started dragging in this widget.
+	 * If a drag is detected, Slate will send an OnDragDetected event.
+	 *
+	 * @param DetectDragInMe  Detect dragging in this widget
+	 * @param MouseButton     This button should be pressed to detect the drag
+	 */
+	FReply& DetectDrag( const TSharedRef<SWidget>& DetectDragInMe, FKey MouseButton )
+	{
+		this->DetectDragForWidget = DetectDragInMe;
+		this->DetectDragForMouseButton = MouseButton;
+		return *this;
+	}
+
+	/**
+	 * An event should return FReply::Handled().BeginDragDrop( Content ) to initiate a drag and drop operation.
+	 *
+	 * @param InDragDropContent  The content that is being dragged. This could be a widget, or some arbitrary data
+	 *
+	 * @return reference back to the FReply so that this call can be chained.
+	 */
+	FReply& BeginDragDrop(TSharedRef<FDragDropOperation> InDragDropContent)
+	{
+		this->DragDropContent = InDragDropContent;
+		return *this;
+	}
+
+	/** An event should return FReply::Handled().EndDragDrop() to request that the current drag/drop operation be terminated. */
+	FReply& EndDragDrop()
+	{
+		this->bEndDragDrop = true;
+		return *this;
+	}
+
+	/** Set the widget that handled the event; undefined if never handled. This method is to be used by SlateApplication only! */
+	FReply& SetHandler( const TSharedRef<SWidget>& InHandler )
+	{
+		this->EventHandler = InHandler;
+		return *this;
+	}
+		
+	/** Ensures throttling for Slate UI responsiveness is not done on mouse down */
+	FReply& PreventThrottling()
+	{
+		this->bPreventThrottling = true;
+		return *this;
+	}
+
+public:
+
+	/** True if this reply indicated that the event was handled */
+	bool IsEventHandled() const { return bIsHandled; }
+
+	/** The widget that ultimately handled the event */
+	const TSharedPtr<SWidget> GetHandler() const { return EventHandler; }
+
+	/** True if this reply indicated that we should release mouse capture as a result of the event being handled */
+	bool ShouldReleaseMouse() const { return bReleaseMouseCapture; }
+
+	/** true if this reply indicated that we should release joystick capture as a result of the event being handled */
+	bool ShouldReleaseJoystick() const { return bReleaseJoystickCapture; }
+
+	/** If the event replied with a request to capture or release the joystick whether it should do it for all joysticks or just the one tied to the UserIndex */
+	bool AffectsAllJoysticks() const { return bAllJoysticks; }
+
+	/** True if this reply indicated that we should use high precision mouse movement */
+	bool ShouldUseHighPrecisionMouse() const { return bUseHighPrecisionMouse; }
+
+	/** True if the reply indicated that we should release mouse lock */
+	bool ShouldReleaseMouseLock() const { return bShouldReleaseMouseLock; }
+
+	/** Whether or not we should throttle on mouse down */
+	bool ShouldThrottle() const { return !bPreventThrottling; }
+
+	/** Returns the widget that the mouse should be locked to (if any) */
+	const TSharedPtr<SWidget>& GetMouseLockWidget() const { return MouseLockWidget; }
+
+	/** When not nullptr, keyboard focus has been requested to be set on the FocusRecipient. */
+	const TSharedPtr<SWidget>& GetFocusRecepient() const { return FocusRecipient; }
+
+	/** Get the reason that a focus change is being requested. */
+	EKeyboardFocusCause::Type GetFocusCause() const { return FocusChangeReason; }
+
+	/** If the event replied with a request to capture the mouse, this returns the desired mouse captor. Otherwise returns an invalid pointer. */
+	const TSharedPtr<SWidget>& GetMouseCaptor() const { return MouseCaptor; }	
+
+	/** If the event replied with a request to capture the joystick, this returns the desired joystick captor. Otherwise returns an invalid pointer. */
+	const TSharedPtr<SWidget>& GetJoystickCaptor() const { return JoystickCaptor; }
+
+	/** @return the Content that we should use for the Drag and Drop operation; Invalid SharedPtr if a drag and drop operation is not requested*/
+	const TSharedPtr<FDragDropOperation>& GetDragDropContent() const { return DragDropContent; }
+
+	/** @return true if the user has asked us to terminate the ongoing drag/drop operation */
+	bool ShouldEndDragDrop() const { return bEndDragDrop; }
+
+	/** @return a widget for which to detect a drag; Invalid SharedPtr if no drag detection requested */
+	const TSharedPtr<SWidget>& GetDetectDragRequest() const { return DetectDragForWidget; }  
+
+	/** @return the mouse button for which we are detecting a drag */
+	FKey GetDetectDragRequestButton() const { return DetectDragForMouseButton; }
+
+	/** @return The coordinates of the requested mouse position */
+	const TOptional<FIntPoint>& GetRequestedMousePos() const { return RequestedMousePos; }
+		
+public:
+
+	/**
+	 * An event should return a FReply::Handled() to let the system know that an event was handled.
+	 */
+	static FReply Handled( )
+	{
+		FReply NewReply;
+		NewReply.bIsHandled = true;
+		return NewReply;
+	}
+
+	/**
+	 * An event should return a FReply::Unhandled() to let the system know that an event was unhandled.
+	 */
+	static FReply Unhandled( )
+	{
+		FReply NewReply;
+		NewReply.bIsHandled = false;
+		return NewReply;
+	}
+
+private:
+
+	/**
+	 * Hidden default constructor.
+	 */
+	FReply( )
+		: RequestedMousePos()
+		, EventHandler(nullptr)
+		, MouseCaptor(nullptr)
+		, JoystickCaptor(nullptr)
+		, FocusRecipient(nullptr)
+		, MouseLockWidget(nullptr)
+		, DragDropContent(nullptr)
+		, FocusChangeReason(EKeyboardFocusCause::SetDirectly)
+		, bIsHandled(false)
+		, bReleaseMouseCapture(false)
+		, bReleaseJoystickCapture(false)
+		, bAllJoysticks(false)
+		, bShouldReleaseMouseLock(false)
+		, bUseHighPrecisionMouse(false)
+		, bPreventThrottling(false)
+		, bEndDragDrop(false)
+	{ }
+		
+private:
+
+	TOptional<FIntPoint> RequestedMousePos;
+	TSharedPtr<SWidget> EventHandler;
+	TSharedPtr<SWidget> MouseCaptor;
+	TSharedPtr<SWidget> JoystickCaptor;
+	TSharedPtr<SWidget> FocusRecipient;
+	TSharedPtr<SWidget> MouseLockWidget;
+	TSharedPtr<SWidget> DetectDragForWidget;
+	FKey DetectDragForMouseButton;
+	TSharedPtr<FDragDropOperation> DragDropContent;
+	EKeyboardFocusCause::Type FocusChangeReason;
+	uint32 bIsHandled:1;
+	uint32 bReleaseMouseCapture:1;
+	uint32 bReleaseJoystickCapture:1;
+	uint32 bAllJoysticks:1;
+	uint32 bShouldReleaseMouseLock:1;
+	uint32 bUseHighPrecisionMouse:1;
+	uint32 bPreventThrottling:1;
+	uint32 bEndDragDrop:1;
+};
diff --git a/Engine/Source/Runtime/SlateCore/Public/Layout/ArrangedChildren.h b/Engine/Source/Runtime/SlateCore/Public/Layout/ArrangedChildren.h
new file mode 100644
--- /dev/null
+++ b/Engine/Source/Runtime/SlateCore/Public/Layout/ArrangedChildren.h
@@ -0,0 +1,136 @@
+// Copyright 1998-2014 Epic Games, Inc. All Rights Reserved.
+
+/*=============================================================================
+	ArrangedChildren.h: Declares the FArrangedChildren class.
+=============================================================================*/
+
+#pragma once
+
+
+/**
+ * The results of an ArrangeChildren are always returned as an FArrangedChildren.
+ * FArrangedChildren supports a filter that is useful for excluding widgets with unwanted
+ * visibilities.
+ */
+class SLATECORE_API FArrangedChildren
+{
+	private:
+	
+	EVisibility VisibilityFilter;
+	
+	public:
+
+	/**
+	 * Construct a new container for arranged children that only accepts children that match the VisibilityFilter.
+	 * e.g.
+	 *  FArrangedChildren ArrangedChildren( VIS_All ); // Children will be included regardless of visibility
+	 *  FArrangedChildren ArrangedChildren( EVisibility::Visible ); // Only visible children will be included
+	 *  FArrangedChildren ArrangedChildren( EVisibility::Collapsed | EVisibility::Hidden ); // Only hidden and collapsed children will be included.
+	 */
+	FArrangedChildren( EVisibility InVisibilityFilter )
+	: VisibilityFilter( InVisibilityFilter )
+	{
+	}
+
+	/** Reverse the order of the arranged children */
+	void Reverse()
+	{
+		int32 LastElementIndex = Array.Num() - 1;
+		for (int32 WidgetIndex = 0; WidgetIndex < Array.Num()/2; ++WidgetIndex )
+		{
+			Array.Swap( WidgetIndex, LastElementIndex - WidgetIndex );
+		}
+	}
+
+	/**
+	 * Add an arranged widget (i.e. widget and its resulting geometry) to the list of Arranged children.
+	 *
+	 * @param VisibilityOverride   The arrange function may override the visibility of the widget for the purposes
+	 *                             of layout or performance (i.e. prevent redundant call to Widget->GetVisibility())
+	 * @param InWidgetGeometry     The arranged widget (i.e. widget and its geometry)
+	 */
+	void AddWidget( EVisibility VisibilityOverride, const FArrangedWidget& InWidgetGeometry );
+	
+	/**
+	 * Add an arranged widget (i.e. widget and its resulting geometry) to the list of Arranged children
+	 * based on the the visibility filter and the arranged widget's visibility
+	 */
+	void AddWidget( const FArrangedWidget& InWidgetGeometry );
+
+	EVisibility GetFilter() const
+	{
+		return VisibilityFilter;
+	}
+
+	bool Accepts( EVisibility InVisibility ) const
+	{
+		return 0 != (InVisibility.Value & VisibilityFilter.Value);
+	}
+
+	// We duplicate parts of TArray's interface here!
+	// Inheriting from "public TArray<FArrangedWidget>"
+	// saves us this boilerplate, but causes instantiation of
+	// many template combinations that we do not want.
+
+	private:
+	/** Internal representation of the array widgets */
+	TArray<FArrangedWidget> Array;
+
+	public:
+
+	TArray<FArrangedWidget>& GetInternalArray()
+	{
+		return Array;
+	}
+
+	const TArray<FArrangedWidget>& GetInternalArray() const
+	{
+		return Array;
+	}
+
+	int32 Num() const
+	{
+		return Array.Num();
+	}
+
+	const FArrangedWidget& operator()(int32 Index) const
+	{
+		return Array[Index];
+	}
+
+	FArrangedWidget& operator()(int32 Index)
+	{
+		return Array[Index];
+	}
+
+	const FArrangedWidget& Last() const
+	{
+		return Array.Last();
+	}
+
+	FArrangedWidget& Last()
+	{
+		return Array.Last();
+	}
+
+	int32 FindItemIndex(const FArrangedWidget& ItemToFind ) const
+	{
+		return Array.Find(ItemToFind);
+	}
+
+	void Remove( int32 Index, int32 Count=1 )
+	{
+		Array.RemoveAt(Index, Count);
+	}
+
+	void Append( const FArrangedChildren& Source )
+	{
+		Array.Append( Source.Array );
+	}
+
+	void Empty()
+	{
+		Array.Empty();
+	}
+
+};
