Slate: Refactored core Slate implementation into SlateCore module in preparation for UMG.

diff --git a/Engine/Source/Runtime/SlateCore/Public/Layout/Children.h b/Engine/Source/Runtime/SlateCore/Public/Layout/Children.h
new file mode 100644
--- /dev/null
+++ b/Engine/Source/Runtime/SlateCore/Public/Layout/Children.h
@@ -0,0 +1,226 @@
+// Copyright 1998-2014 Epic Games, Inc. All Rights Reserved.
+
+#pragma once
+
+class SWidget;
+
+/**
+ * FChildren is an interface that must be implemented by all child containers.
+ * It allows iteration over a list of any Widget's children regardless of how
+ * the underlying Widget happens to store its children.
+ * 
+ * FChildren is intended to be returned by the GetChildren() method.
+ * 
+ */
+class SLATECORE_API FChildren
+{
+public:
+	/** @return the number of children */
+	virtual int32 Num() const = 0;
+	/** @return pointer to the Widget at the specified Index. */
+	virtual TSharedRef<SWidget> GetChildAt( int32 Index ) = 0;
+	/** @return const pointer to the Widget at the specified Index. */
+	virtual TSharedRef<const SWidget> GetChildAt( int32 Index ) const = 0;
+
+protected:
+	virtual ~FChildren(){}
+};
+
+
+/**
+ * Widgets with no Children can return an instance of FNoChildren.
+ * For convenience a shared instance SWidget::NoChildrenInstance can be used.
+ */
+class SLATECORE_API FNoChildren : public FChildren
+{
+public:
+	virtual int32 Num() const { return 0; }
+	
+	virtual TSharedRef<SWidget> GetChildAt( int32 )
+	{
+		// Nobody should be getting a child when there aren't any children.
+		// We expect this to crash!
+		check( false );
+		return TSharedPtr<SWidget>(nullptr).ToSharedRef();
+	}
+	
+	virtual TSharedRef<const SWidget> GetChildAt( int32 ) const
+	{
+		// Nobody should be getting a child when there aren't any children.
+		// We expect this to crash!
+		check( false );
+		return TSharedPtr<const SWidget>(nullptr).ToSharedRef();
+	}
+};
+
+/**
+ * Widgets that will only have one child can return an instance of FOneChild.
+ */
+template <typename ChildType, typename MixedIntoType>
+class TSupportsOneChildMixin : public FChildren
+{
+public:
+	TSupportsOneChildMixin()
+		: Widget( SNullWidget::NullWidget )
+	{ }
+
+	virtual int32 Num() const { return 1; }
+	virtual TSharedRef<SWidget> GetChildAt( int32 ChildIndex ) { check(ChildIndex == 0); return Widget; }
+	virtual TSharedRef<const SWidget> GetChildAt( int32 ChildIndex ) const { check(ChildIndex == 0); return Widget; }
+	
+	MixedIntoType& operator[]( TSharedRef<SWidget> InWidget )
+	{
+		Widget = InWidget;
+		return *(static_cast<MixedIntoType*>(this));
+	}
+
+	MixedIntoType& Expose( MixedIntoType*& OutVarToInit )
+	{
+		OutVarToInit = static_cast<MixedIntoType*>(this);
+		return *(static_cast<MixedIntoType*>(this));
+	}
+
+	TSharedRef<ChildType> Widget;
+};
+
+/**
+ * For widgets that do not own their content, but are responsible for presenting someone else's content.
+ * e.g. Tooltips are just presented by the owner window; not actually owned by it. They can go away at any time
+ *      and then they'll just stop being shown.
+ */
+template <typename ChildType>
+class TWeakChild : public FChildren
+{
+public:
+	TWeakChild()
+		: Widget( SNullWidget::NullWidget )
+	{ }
+
+	virtual int32 Num() const { return Widget.IsValid() ? 1 : 0 ; }
+	virtual TSharedRef<SWidget> GetChildAt( int32 ChildIndex ) { check(ChildIndex == 0); return Widget.Pin().ToSharedRef(); }
+	virtual TSharedRef<const SWidget> GetChildAt( int32 ChildIndex ) const { check(ChildIndex == 0); return Widget.Pin().ToSharedRef(); }
+	
+	TWeakPtr<ChildType> Widget;
+};
+
+template <typename MixedIntoType>
+class TSupportsContentAlignmentMixin
+{
+public:
+	TSupportsContentAlignmentMixin(const EHorizontalAlignment InHAlign, const EVerticalAlignment InVAlign)
+	: HAlignment( InHAlign )
+	, VAlignment( InVAlign )
+	{
+		
+	}
+
+	MixedIntoType& HAlign( EHorizontalAlignment InHAlignment )
+	{
+		HAlignment = InHAlignment;
+		return *(static_cast<MixedIntoType*>(this));
+	}
+
+	MixedIntoType& VAlign( EVerticalAlignment InVAlignment )
+	{
+		VAlignment = InVAlignment;
+		return *(static_cast<MixedIntoType*>(this));
+	}
+	
+	EHorizontalAlignment HAlignment;
+	EVerticalAlignment VAlignment;
+};
+
+template <typename MixedIntoType>
+class TSupportsContentPaddingMixin
+{
+public:
+	MixedIntoType& Padding( const TAttribute<FMargin> InPadding )
+	{
+		SlotPadding = InPadding;
+		return *(static_cast<MixedIntoType*>(this));
+	}
+
+	MixedIntoType& Padding( float Uniform )
+	{
+		SlotPadding = FMargin(Uniform);
+		return *(static_cast<MixedIntoType*>(this));
+	}
+
+	MixedIntoType& Padding( float Horizontal, float Vertical )
+	{
+		SlotPadding = FMargin(Horizontal, Vertical);
+		return *(static_cast<MixedIntoType*>(this));
+	}
+
+	MixedIntoType& Padding( float Left, float Top, float Right, float Bottom )
+	{
+		SlotPadding = FMargin(Left, Top, Right, Bottom);
+		return *(static_cast<MixedIntoType*>(this));
+	}
+
+	TAttribute< FMargin > SlotPadding;
+};
+
+/** A slot that support alignment of content and padding */
+class SLATECORE_API FSimpleSlot : public TSupportsOneChildMixin<SWidget, FSimpleSlot>, public TSupportsContentAlignmentMixin<FSimpleSlot>, public TSupportsContentPaddingMixin<FSimpleSlot>
+{
+public:
+	FSimpleSlot()
+	: TSupportsContentAlignmentMixin<FSimpleSlot>(HAlign_Fill, VAlign_Fill)
+	{
+	}
+};
+
+
+/**
+ * A generic FChildren that stores children along with layout-related information.
+ * The type containing Widget* and layout info is specified by ChildType.
+ * ChildType must have a public member SWidget* Widget;
+ */
+template<typename SlotType>
+class TPanelChildren : public FChildren, public TIndirectArray< SlotType >
+{
+public:
+	virtual int32 Num() const
+	{
+		return TIndirectArray<SlotType>::Num();
+	}
+
+	virtual TSharedRef<SWidget> GetChildAt( int32 Index )
+	{
+		return this->operator[](Index).Widget;
+	}
+
+	virtual TSharedRef<const SWidget> GetChildAt( int32 Index ) const
+	{
+		return this->operator[](Index).Widget;
+	}
+};
+
+/**
+ * Some advanced widgets contain no layout information, and do not require slots.
+ * Those widgets may wish to store a specialized type of child widget.
+ * In those cases, using TSlotlessChildren is convenient.
+ *
+ * TSlotlessChildren should not be used for general-purpose widgets.
+ */
+template<typename ChildType>
+class TSlotlessChildren : public FChildren, public TArray< TSharedRef<ChildType> >
+{
+public:
+	/** @return the number of children */
+	virtual int32 Num() const
+	{
+		return TArray< TSharedRef<ChildType> >::Num();
+	}
+	/** @return pointer to the Widget at the specified Index. */
+	virtual TSharedRef<SWidget> GetChildAt( int32 Index )
+	{
+		return this->operator[](Index);
+	}
+	/** @return const pointer to the Widget at the specified Index. */
+	virtual TSharedRef<const SWidget> GetChildAt( int32 Index ) const
+	{
+		return this->operator[](Index);
+	}
+};
diff --git a/Engine/Source/Runtime/SlateCore/Public/Layout/Geometry.h b/Engine/Source/Runtime/SlateCore/Public/Layout/Geometry.h
new file mode 100644
--- /dev/null
+++ b/Engine/Source/Runtime/SlateCore/Public/Layout/Geometry.h
@@ -0,0 +1,227 @@
+// Copyright 1998-2014 Epic Games, Inc. All Rights Reserved.
+
+/*=============================================================================
+	Geometry.h: Declares the FGeometry class.
+=============================================================================*/
+
+#pragma once
+
+
+class FArrangedWidget;
+class SWidget;
+
+
+/**
+ * Represents the position, size, and absolute position of a Widget in Slate.
+ * The absolute location of a geometry is usually screen space or
+ * window space depending on where the geometry originated.
+ * Geometries are usually paired with a SWidget pointer in order
+ * to provide information about a specific widget (see FArrangedWidget).
+ * A Geometry's parent is generally thought to be the Geometry of the
+ * the corresponding parent widget.
+ */
+class SLATECORE_API FGeometry
+{
+
+public:
+
+	/**
+	 * Default constructor.
+	 */
+	FGeometry( )
+		: Position(0.0f, 0.0f)
+		, AbsolutePosition(0.0f, 0.0f)
+		, Size(0.0f, 0.0f)
+		, Scale(1.0f)
+	{ }
+
+	/**
+	 * Construct a new geometry given the following parameters:
+	 * 
+	 * @param OffsetFromParent         Local position of this geometry within its parent geometry.
+	 * @param ParentAbsolutePosition   The absolute position of the parent geometry containing this geometry.
+	 * @param InSize                   The size of this geometry.
+	 * @param InScale                  The scale of this geometry with respect to Normal Slate Coordinates.
+	 */
+	FGeometry( const FVector2D& OffsetFromParent, const FVector2D& ParentAbsolutePosition, const FVector2D& InSize, float InScale )
+		: Position(OffsetFromParent)
+		, AbsolutePosition(ParentAbsolutePosition + (OffsetFromParent * InScale))
+		, Size(InSize)	
+		, Scale(InScale)
+	{ }
+
+public:
+
+	/**
+	 * Test geometries for strict equality (i.e. exact float equality).
+	 *
+	 * @param Other A geometry to compare to.
+	 *
+	 * @return true if this is equal to other, false otherwise.
+	 */
+	bool operator==( const FGeometry& Other ) const
+	{
+		return
+			this->Position == Other.Position &&
+			this->AbsolutePosition == Other.AbsolutePosition &&
+			this->Size == Other.Size &&
+			this->Scale == Other.Scale;
+	}
+	
+	/**
+	 * Test geometries for strict inequality (i.e. exact float equality negated).
+	 *
+	 * @param Other A geometry to compare to.
+	 *
+	 * @return false if this is equal to other, true otherwise.
+	 */
+	bool operator!=( const FGeometry& Other ) const
+	{
+		return !( this->operator==(Other) );
+	}
+
+public:
+	
+	/**
+	 * Create a child geometry; i.e. a geometry relative to this geometry.
+	 *
+	 * @param ChildOffset  The offset of the child from this Geometry's position.
+	 * @param ChildSize    The size of the child geometry in screen units.
+	 * @param InScale      How much to scale this child with respect to its parent.
+	 *
+	 * @return A Geometry that is offset from this geometry by ChildOffset
+	 */
+	FGeometry MakeChild( const FVector2D& ChildOffset, const FVector2D& ChildSize, float InScale = 1.0f ) const
+	{
+		return FGeometry( ChildOffset, this->AbsolutePosition, ChildSize, this->Scale * InScale );
+	}
+	
+	/**
+	 * Create a child widget geometry; i.e. a geometry relative to this geometry that
+	 * is also associated with a Widget.
+	 *
+	 * @param InWidget     The widget to associate with the newly created child geometry
+	 * @param ChildOffset  The offset of the child from this Geometry's position.
+	 * @param ChildSize    The size of the child geometry in screen units.
+	 * @param InScale      How much to scale this child with respect to its parent.
+	 *
+	 * @return A WidgetGeometry that is offset from this geometry by ChildOffset
+	 */
+	FArrangedWidget MakeChild( const TSharedRef<SWidget>& InWidget, const FVector2D& ChildOffset, const FVector2D& ChildSize, float InScale = 1.0f ) const;
+
+	/**
+	 * Make an FPaintGeometry in this FGeometry at the specified Offset.
+	 *
+	 * @param InOffset   Offset in SlateUnits from this FGeometry's origin.
+	 * @param InSize     The size of this paint element in SlateUnits.
+	 * @param InScale    Additional scaling to apply to the DrawElements.
+	 *
+	 * @return a FPaintGeometry derived this FGeometry.
+	 */
+	FPaintGeometry ToPaintGeometry( FVector2D InOffset, FVector2D InSize, float InScale ) const;
+
+	/** A PaintGeometry that is not scaled additionally with respect to its originating FGeometry */
+	FPaintGeometry ToPaintGeometry( FVector2D InOffset, FVector2D InSize ) const
+	{
+		return ToPaintGeometry( InOffset, InSize, 1.0f );
+	}
+
+	/** Convert the FGeometry to an FPaintGeometry with no changes */
+	FPaintGeometry ToPaintGeometry() const
+	{
+		return ToPaintGeometry( FVector2D(0,0), this->Size, 1.0f );
+	}
+
+	/** An FPaintGeometry that is offset into the FGeometry from which it originates and occupies all the available room */
+	FPaintGeometry ToOffsetPaintGeometry( FVector2D InOffset ) const
+	{
+		return ToPaintGeometry( InOffset, this->Size, 1.0f );
+	}
+
+	/**
+	 * Make an FPaintGeometry from this FGeometry by inflating the FGeometry by InflateAmount.
+	 */
+	FPaintGeometry ToInflatedPaintGeometry( const FVector2D& InflateAmount ) const;
+
+	FPaintGeometry CenteredPaintGeometryOnLeft( const FVector2D& SizeBeingAligned, float InScale ) const;
+
+	FPaintGeometry CenteredPaintGeometryOnRight( const FVector2D& SizeBeingAligned, float InScale ) const;
+
+	FPaintGeometry CenteredPaintGeometryBelow( const FVector2D& SizeBeingAligned, float InScale ) const;
+
+
+	/** @return true if the provided location is within the bounds of this geometry. */
+	const bool IsUnderLocation( const FVector2D& AbsoluteCoordinate ) const
+	{
+		return
+			( AbsolutePosition.X <= AbsoluteCoordinate.X ) && ( ( AbsolutePosition.X + Size.X * Scale ) > AbsoluteCoordinate.X ) &&
+			( AbsolutePosition.Y <= AbsoluteCoordinate.Y ) && ( ( AbsolutePosition.Y + Size.Y * Scale ) > AbsoluteCoordinate.Y );
+	}
+
+	/** @return Translates AbsoluteCoordinate into the space defined by this Geometry. */
+	const FVector2D AbsoluteToLocal( FVector2D AbsoluteCoordinate ) const
+	{
+		return (AbsoluteCoordinate - this->AbsolutePosition) / Scale;
+	}
+
+	/**
+	 * Translates local coordinates into absolute coordinates
+	 *
+	 * @return  Absolute coordinates
+	 */
+	const FVector2D LocalToAbsolute( FVector2D LocalCoordinate ) const
+	{
+		return (LocalCoordinate * Scale) + this->AbsolutePosition;
+	}
+	
+	/**
+	 * Returns the allocated geometry's relative position and size as a rect
+	 *
+	 * @return  Allotted geometry relative position and size rectangle
+	 */
+	FSlateRect GetRect( ) const
+	{
+		return FSlateRect( Position.X, Position.Y, Position.X + Size.X, Position.Y + Size.Y );
+	}
+
+	/**
+	 * Returns a clipping rectangle corresponding to the allocated geometry's absolute position and size.
+	 * Note that the clipping rectangle starts 1 pixel above and left of the geometry because clipping is not
+	 * inclusive on the lower bound.
+	 *
+	 * @return  Allotted geometry absolute position and size rectangle
+	 */
+	FSlateRect GetClippingRect( ) const
+	{
+		return FSlateRect( AbsolutePosition.X, AbsolutePosition.Y, AbsolutePosition.X + Size.X * Scale, AbsolutePosition.Y + Size.Y * Scale );
+	}
+	
+	/** @return A String representation of this Geometry */
+	FString ToString( ) const
+	{
+		return FString::Printf(TEXT("[Pos=%s, Abs=%s, Size=%s]"), *Position.ToString(), *AbsolutePosition.ToString(), *Size.ToString() );
+	}
+
+	/** @return the size of the geometry in screen space */
+	FVector2D GetDrawSize() const
+	{
+		return Size * Scale;
+	}
+
+public:
+
+	/** Position relative to the parent */
+	FVector2D Position;
+
+	/** Position in screen space or (window space during rendering) */
+	FVector2D AbsolutePosition;	
+
+	/** The dimensions */
+	FVector2D Size;
+
+	/** How much this geometry is scaled: WidgetCoodinates / SlateCoordinates */
+	float Scale;
+};
+
+
+template <> struct TIsPODType<FGeometry> { enum { Value = true }; };
diff --git a/Engine/Source/Runtime/SlateCore/Public/Layout/LayoutUtils.h b/Engine/Source/Runtime/SlateCore/Public/Layout/LayoutUtils.h
new file mode 100644
--- /dev/null
+++ b/Engine/Source/Runtime/SlateCore/Public/Layout/LayoutUtils.h
@@ -0,0 +1,118 @@
+// Copyright 1998-2014 Epic Games, Inc. All Rights Reserved.
+
+#pragma once
+
+struct AlignmentArrangeResult
+{
+	AlignmentArrangeResult( float InOffset, float InSize )
+	: Offset(InOffset)
+	, Size(InSize)
+	{
+	}
+	
+	float Offset;
+	float Size;
+};
+
+namespace ArrangeUtils
+{
+	/** Gets the alignment of an axis-agnostic int32 so that we can do aligment on an axis without caring about its orientation */
+	template<EOrientation Orientation>
+	struct GetChildAlignment
+	{
+		template<typename SlotType>
+		static int32 AsInt( const SlotType& InSlot );
+	};
+
+	template<>
+	struct GetChildAlignment<Orient_Horizontal>
+	{
+		template<typename SlotType>
+		static int32 AsInt( const SlotType& InSlot )
+		{
+			return static_cast<int32>(InSlot.HAlignment);
+		}
+	};
+
+	template<>
+	struct GetChildAlignment<Orient_Vertical>
+	{
+		template<typename SlotType>
+		static int32 AsInt( const SlotType& InSlot )
+		{
+			return static_cast<int32>(InSlot.VAlignment);
+		}
+	};
+
+}
+
+
+/**
+ * Helper method to BoxPanel::ArrangeChildren.
+ * 
+ * @param AllottedSize         The size available to arrange the widget along the given orientation
+ * @param ChildToArrange       The widget and associated layout information
+ * 
+ * @return  Offset and Size of widget
+ */
+template<EOrientation Orientation, typename SlotType>
+static AlignmentArrangeResult AlignChild( float AllottedSize, const SlotType& ChildToArrange, const FMargin& SlotPadding, const float& ContentScale = 1.0f )
+{
+	float ChildDesiredSize = (Orientation == Orient_Horizontal)
+		? ( ChildToArrange.Widget->GetDesiredSize().X * ContentScale )
+		: ( ChildToArrange.Widget->GetDesiredSize().Y * ContentScale );
+
+	const FMargin& Margin = SlotPadding;
+	const float TotalMargin = Margin.GetTotalSpaceAlong<Orientation>();	
+	const float MarginPre = (Orientation == Orient_Horizontal) ? Margin.Left : Margin.Top;
+	const float MarginPost = (Orientation == Orient_Horizontal) ? Margin.Right : Margin.Bottom;
+	
+	const int32 Alignment = ArrangeUtils::GetChildAlignment<Orientation>::AsInt(ChildToArrange);
+		
+	const float ChildSize = FMath::Min(ChildDesiredSize, AllottedSize-TotalMargin);
+		
+	switch( Alignment )	
+	{
+		default:
+		case HAlign_Fill:
+			return AlignmentArrangeResult( MarginPre, (AllottedSize - TotalMargin)*ContentScale );
+			break;
+		
+		case HAlign_Left: // same as Align_Top
+			return AlignmentArrangeResult( MarginPre, ChildSize );
+			break;
+		
+		case HAlign_Center:
+			return AlignmentArrangeResult( (AllottedSize - ChildSize ) / 2.0f + MarginPre - MarginPost, ChildSize );
+			break;
+		
+		case HAlign_Right: // same as Align_Bottom		
+			return AlignmentArrangeResult( AllottedSize - ChildSize - MarginPost, ChildSize );
+			break;
+	}
+}
+
+
+/**
+ * Arrange a ChildSlot within the AllottedGeometry and populate ArrangedChildren with the arranged result.
+ * The code makes certain assumptions about the type of ChildSlot.
+ */
+template<typename SlotType>
+void ArrangeSingleChild( const FGeometry& AllottedGeometry, FArrangedChildren& ArrangedChildren, const SlotType& ChildSlot, const TAttribute<FVector2D>& ContentScale )
+{
+	const EVisibility ChildVisibility = ChildSlot.Widget->GetVisibility();
+	if ( ArrangedChildren.Accepts(ChildVisibility) )
+	{
+		const FVector2D ThisContentScale = ContentScale.Get();
+		const FMargin SlotPadding(ChildSlot.SlotPadding.Get());
+		AlignmentArrangeResult XResult = AlignChild<Orient_Horizontal>(AllottedGeometry.Size.X, ChildSlot, SlotPadding, ThisContentScale.X);
+		AlignmentArrangeResult YResult = AlignChild<Orient_Vertical>(AllottedGeometry.Size.Y, ChildSlot, SlotPadding, ThisContentScale.Y);
+
+		ArrangedChildren.AddWidget( ChildVisibility, AllottedGeometry.MakeChild(
+				ChildSlot.Widget,
+				FVector2D(XResult.Offset,YResult.Offset),
+				FVector2D(XResult.Size, YResult.Size)
+		) );
+	}
+
+}
diff --git a/Engine/Source/Runtime/SlateCore/Public/Layout/PaintGeometry.h b/Engine/Source/Runtime/SlateCore/Public/Layout/PaintGeometry.h
new file mode 100644
--- /dev/null
+++ b/Engine/Source/Runtime/SlateCore/Public/Layout/PaintGeometry.h
@@ -0,0 +1,74 @@
+// Copyright 1998-2014 Epic Games, Inc. All Rights Reserved.
+
+/*=============================================================================
+	PaintGeometry.h: Declares the FPaintGeometry structure.
+=============================================================================*/
+
+#pragma once
+
+
+/**
+ * A Paint geometry is a draw-space rectangle.
+ * The DrawPosition and DrawSize are already positioned and
+ * scaled for immediate consumption by the draw code.
+ *
+ * The DrawScale is only applied to the internal aspects of the draw primitives. e.g. Line thickness, 3x3 grid margins, etc.
+ */
+struct SLATECORE_API FPaintGeometry
+{
+	/** Render-space position at which we will draw */
+	FVector2D DrawPosition;
+
+	/**
+	 * The size in render-space that we want to make this element.
+	 * This field is ignored by some elements (e.g. spline, lines)
+	 */
+	FVector2D DrawSize;
+
+	/** Only affects the draw aspects of individual draw elements. e.g. line thickness, 3x3 grid margins */
+	float DrawScale;
+
+public:
+
+	/**
+	 * Creates and initializes a new instance.
+	 *
+	 * @param InDrawPosition The draw position.
+	 * @param InDrawSize The draw size.
+	 * @param InDrawScale The draw scale.
+	 */
+	FPaintGeometry( FVector2D InDrawPosition, FVector2D InDrawSize, float InDrawScale )
+		: DrawPosition(InDrawPosition)
+		, DrawSize (InDrawSize)
+		, DrawScale(InDrawScale)
+	{ }
+
+public:
+
+	/**
+	 * Converts this paint geometry to a Slate rectangle.
+	 *
+	 * @return The Slate rectangle.
+	 */
+	FSlateRect ToSlateRect( ) const
+	{
+		return FSlateRect(DrawPosition.X, DrawPosition.Y, DrawPosition.X + DrawSize.X, DrawPosition.Y + DrawSize.Y);
+	}
+
+public:
+
+	static const FPaintGeometry& Identity( )
+	{
+		static FPaintGeometry IdentityPaintGeometry
+		(
+			FVector2D(0.0f, 0.0f),	// Position: Draw at the origin of draw space
+			FVector2D(0.0f, 0.0f),	// Size    : Size is ignored by splines
+			1.0f					// Scale   : Do not do any scaling.
+		);
+
+		return IdentityPaintGeometry;
+	}
+};
+
+
+template<> struct TIsPODType<FPaintGeometry> { enum { Value = true }; };
diff --git a/Engine/Source/Runtime/SlateCore/Public/Layout/SlateRect.h b/Engine/Source/Runtime/SlateCore/Public/Layout/SlateRect.h
new file mode 100644
--- /dev/null
+++ b/Engine/Source/Runtime/SlateCore/Public/Layout/SlateRect.h
@@ -0,0 +1,168 @@
+// Copyright 1998-2014 Epic Games, Inc. All Rights Reserved.
+
+/*=============================================================================
+	SlateRect.h: Declares the FSlateRect structure.
+=============================================================================*/
+
+#pragma once
+
+
+/** 
+ * A rectangle defined by upper-left and lower-right corners
+ */
+class SLATECORE_API FSlateRect
+{
+public:
+
+	float Left;
+	float Top;
+	float Right;
+	float Bottom;
+
+	FSlateRect( float InLeft = -1, float InTop = -1, float InRight = -1, float InBottom = -1 )
+		: Left(InLeft)
+		, Top(InTop)
+		, Right(InRight)
+		, Bottom(InBottom)
+	{ }
+
+	FSlateRect( const FVector2D& InStartPos, const FVector2D& InEndPos )
+		: Left(InStartPos.X)
+		, Top(InStartPos.Y)
+		, Right(InEndPos.X)
+		, Bottom(InEndPos.Y)
+	{ }
+	
+public:
+
+	bool IsValid() const
+	{
+		return Left >= 0 && Right >= 0 && Bottom >=0 && Top >= 0;
+	}
+
+	/**
+	 * Returns the size of the rectangle.
+	 *
+	 * @return The size as a vector.
+	 */
+	FVector2D GetSize() const
+	{
+		return FVector2D( Right - Left, Bottom - Top );
+	}
+
+	/**
+	 * Returns the center of the rectangle
+	 * 
+	 * @return The center point.
+	 */
+	FVector2D GetCenter() const
+	{
+		return FVector2D( Left, Top ) + GetSize() * 0.5f;
+	}
+
+	/**
+	 * Return a rectangle that is contracted on each side by the amount specified in each margin.
+	 *
+	 * @param FMargin The amount to contract the geometry.
+	 *
+	 * @return An inset rectangle.
+	 */
+	FSlateRect InsetBy( const struct FMargin& InsetAmount ) const;
+
+	/**
+	 * Returns the rect that encompasses both rectangles
+	 * 
+	 * @param	Other	The other rectangle
+	 *
+	 * @return	Rectangle that is big enough to fit both rectangles
+	 */
+	FSlateRect Expand( const FSlateRect& Other ) const
+	{
+		return FSlateRect( FMath::Min( Left, Other.Left ), FMath::Min( Top, Other.Top ), FMath::Max( Right, Other.Right ), FMath::Max( Bottom, Other.Bottom ) );
+	}
+
+	
+	FSlateRect IntersectionWith( const FSlateRect& Other ) const 
+	{
+		FSlateRect Intersected( FMath::Max( this->Left, Other.Left ), FMath::Max(this->Top, Other.Top), FMath::Min( this->Right, Other.Right ), FMath::Min( this->Bottom, Other.Bottom ) );
+		if ( (Intersected.Bottom < Intersected.Top) || (Intersected.Right < Intersected.Left) )
+		{
+			// The intersection has 0 area and should not be rendered at all.
+			return FSlateRect(0,0,0,0);
+		}
+		else
+		{
+			return Intersected;
+		}
+	}
+
+	/**
+	 * Returns whether or not a point is inside the rectangle
+	 * Note: The lower right and bottom points of the rect are not inside the rectangle due to rendering API clipping rules.
+	 * 
+	 * @param Point	The point to check
+	 * @return True if the point is inside the rectangle
+	 */
+	FORCEINLINE bool ContainsPoint( const FVector2D& Point ) const
+	{
+		return Point.X >= Left && Point.X < Right && Point.Y >= Top && Point.Y < Bottom;
+	}
+
+	bool operator==( const FSlateRect& Other ) const
+	{
+		return
+			Left == Other.Left &&
+			Top == Other.Top &&
+			Right == Other.Right &&
+			Bottom == Other.Bottom;
+	}
+
+	bool operator!=( const FSlateRect& Other ) const
+	{
+		return Left != Other.Left || Top != Other.Top || Right != Other.Right || Bottom != Other.Bottom;
+	}
+
+	friend FSlateRect operator+( const FSlateRect& A, const FSlateRect& B )
+	{
+		return FSlateRect( A.Left + B.Left, A.Top + B.Top, A.Right + B.Right, A.Bottom + B.Bottom );
+	}
+
+	friend FSlateRect operator-( const FSlateRect& A, const FSlateRect& B )
+	{
+		return FSlateRect( A.Left - B.Left, A.Top - B.Top, A.Right - B.Right, A.Bottom - B.Bottom );
+	}
+
+	friend FSlateRect operator*( float Scalar, const FSlateRect& Rect )
+	{
+		return FSlateRect( Rect.Left * Scalar, Rect.Top * Scalar, Rect.Right * Scalar, Rect.Bottom * Scalar );
+	}
+
+	/** Do rectangles A and B intersect? */
+	static bool DoRectanglesIntersect( const FSlateRect& A, const FSlateRect& B )
+	{
+		//  Segments A and B do not intersect when:
+		//
+		//       (left)   A     (right)
+		//         o-------------o
+		//  o---o        OR         o---o
+		//    B                       B
+		//
+		//
+		// We assume the A and B are well-formed rectangles.
+		// i.e. (Top,Left) is above and to the left of (Bottom,Right)
+		const bool bDoNotOverlap =
+			B.Right < A.Left || A.Right < B.Left ||
+			B.Bottom < A.Top || A.Bottom < B.Top;
+
+		return ! bDoNotOverlap;
+	}
+
+	/** Is rectangle B contained within rectangle A? */
+	static bool IsRectangleContained( const FSlateRect& A, const FSlateRect& B )
+	{
+		return (A.Left <= B.Left) && (A.Right >= B.Right) && (A.Top <= B.Top) && (A.Bottom >= B.Bottom);
+	}
+};
+
+
+template<> struct TIsPODType<FSlateRect> { enum { Value = true }; };
diff --git a/Engine/Source/Runtime/SlateCore/Public/Layout/WidgetPath.h b/Engine/Source/Runtime/SlateCore/Public/Layout/WidgetPath.h
new file mode 100644
--- /dev/null
+++ b/Engine/Source/Runtime/SlateCore/Public/Layout/WidgetPath.h
@@ -0,0 +1,285 @@
+// Copyright 1998-2014 Epic Games, Inc. All Rights Reserved.
+
+#pragma once
+
+
+class SWindow;
+
+
+/** Matches widgets against InWidget */
+struct FWidgetMatcher
+{
+	FWidgetMatcher( const TSharedRef<const SWidget> InWidget )
+		: WidgetToFind( InWidget )
+	{}
+
+	bool IsMatch( const TSharedRef<const SWidget>& InWidget ) const
+	{
+		return WidgetToFind == InWidget;
+	}
+
+	TSharedRef<const SWidget> WidgetToFind;
+};
+
+/**
+ * A widget path is a vertical slice through the tree.
+ * The canonical form for widget paths is "leafmost last". The top-level window always resides at index 0.
+ * A widget path also contains a reference to a top-level SWindow that contains all the widgets in the path.
+ * The window is needed for its ability to determine its own geometry, from which the geometries of the rest
+ * of the widget can be determined.
+ */
+class FWidgetPath
+{
+public:
+	/** Constructor */
+	FWidgetPath()
+	: Widgets( EVisibility::Visible )
+	{
+	}
+
+	FWidgetPath( TSharedPtr<SWindow> InTopLevelWindow, const FArrangedChildren& InWidgetPath )
+	: Widgets( InWidgetPath )
+	, TopLevelWindow(InTopLevelWindow)
+	{
+	
+	}
+
+	/**
+	 * @param MarkerWidget Copy the path up to and including this widget 
+	 *
+	 * @return a copy of the widget path down to and including the MarkerWidget. If the MarkerWidget is not found in the path, return an invalid path.
+	 */
+	FWidgetPath GetPathDownTo( TSharedRef<const SWidget> MarkerWidget ) const
+	{
+		FArrangedChildren ClippedPath(EVisibility::Visible);
+		bool bCopiedMarker = false;
+		for( int32 WidgetIndex = 0; !bCopiedMarker && WidgetIndex < Widgets.Num(); ++WidgetIndex )
+		{
+			ClippedPath.AddWidget( Widgets(WidgetIndex) );
+			bCopiedMarker = (Widgets(WidgetIndex).Widget == MarkerWidget);
+		}
+		
+		if ( bCopiedMarker )
+		{
+			// We found the MarkerWidget and copied the path down to (and including) it.
+			return FWidgetPath( TopLevelWindow, ClippedPath );
+		}
+		else
+		{
+			// The MarkerWidget was not in the widget path. We failed.
+			return FWidgetPath( nullptr, FArrangedChildren(EVisibility::Visible) );		
+		}	
+	}
+	
+	/** @return true if the WidgetToFind is in this WidgetPath, false otherwise. */
+	bool ContainsWidget( TSharedRef<const SWidget> WidgetToFind ) const
+	{
+		for(int32 WidgetIndex = 0; WidgetIndex < Widgets.Num(); ++WidgetIndex)
+		{
+			if ( Widgets(WidgetIndex).Widget == WidgetToFind )
+			{
+				return true;
+			}
+		}
+		
+		return false;
+	}
+
+	SLATECORE_API FArrangedWidget FindArrangedWidget( TSharedRef<const SWidget> WidgetToFind ) const;
+	
+	/**
+	 * Get the first (top-most) widget in this path, which is always a window; assumes path is valid
+	 *
+	 * @return Window at the top of this path
+	 */
+	TSharedRef<SWindow> GetWindow()
+	{
+		check(IsValid());
+
+		TSharedRef<SWindow> FirstWidgetWindow = StaticCastSharedRef<SWindow>(Widgets(0).Widget);
+		return FirstWidgetWindow;
+	}
+
+	/**
+	 * Get the first (top-most) widget in this path, which is always a window; assumes path is valid
+	 *
+	 * @return Window at the top of this path
+	 */
+	TSharedRef<SWindow> GetWindow() const
+	{
+		check(IsValid());
+
+		TSharedRef<SWindow> FirstWidgetWindow = StaticCastSharedRef<SWindow>(Widgets(0).Widget);
+		return FirstWidgetWindow;
+	}
+	
+	/** A valid path has at least one widget in it */
+	bool IsValid() const { return Widgets.Num() > 0; }
+	
+	FString ToString() const
+	{
+		FString StringBuffer;
+		for( int32 WidgetIndex = Widgets.Num()-1; WidgetIndex >= 0; --WidgetIndex )
+		{
+			StringBuffer += Widgets(WidgetIndex).ToString();
+			StringBuffer += TEXT("\n");
+		}
+		return StringBuffer;
+	}
+
+	/**
+	 * Extend the current path such that it reaches some widget that qualifies as a Match
+	 * The widget to match must be a descendant of the last widget currently in the path.
+	 *
+	 * @param Matcher         Some struct that has a "bool IsMatch( const TSharedRef<const SWidget>& InWidget ) const" method
+	 * @param  VisibilityFilter	Widgets must have this type of visibility to be included the path
+	 *
+	 * @return true if successful; false otherwise.
+	 */
+	template<typename MatcherType>
+	bool ExtendPathTo( const MatcherType& Matcher, EVisibility VisibilityFilter = EVisibility::Visible )
+	{
+		const FArrangedWidget& LastWidget = Widgets.Last();
+		
+		FArrangedChildren Extension = GeneratePathToWidget( Matcher, LastWidget, EFocusMoveDirection::Next, VisibilityFilter );
+
+		for( int32 WidgetIndex=0; WidgetIndex < Extension.Num(); ++WidgetIndex )
+		{
+			this->Widgets.AddWidget( Extension(WidgetIndex) );
+		}
+
+		return Extension.Num() > 0;
+	}
+
+	/**
+	 * Generate a path from FromWidget to WidgetToFind. The path will not include FromWidget.
+	 *
+	 * @param Matcher         Some struct that has a "bool IsMatch( const TSharedRef<const SWidget>& InWidget ) const" method
+	 * @param FromWidget      Widget from which we a building a path.AddItem*
+	 * @param  VisibilityFilter	Widgets must have this type of visibility to be included the path
+	 * 
+	 * @return A path from FromWidget to WidgetToFind; will not include FromWidget.
+	 */
+	template<typename MatcherType>
+	FArrangedChildren GeneratePathToWidget( const MatcherType& Matcher, const FArrangedWidget& FromWidget, EFocusMoveDirection::Type Direction = EFocusMoveDirection::Next, EVisibility VisibilityFilter = EVisibility::Visible )
+	{
+		FArrangedChildren PathResult(VisibilityFilter);
+
+		if ( Direction == EFocusMoveDirection::Next )
+		{
+			SearchForWidgetRecursively( Matcher, FromWidget, PathResult, VisibilityFilter );
+		}
+		else
+		{
+			SearchForWidgetRecursively_Reverse( Matcher, FromWidget, PathResult, VisibilityFilter );
+		}
+		
+
+		// Reverse the list of widgets we found; canonical form is leafmost last.
+		PathResult.Reverse();
+
+		return PathResult;
+	}
+	
+	/**
+	 * Move focus either forward on backward in the path level specified by PathLevel.
+	 * That is, this movement of focus will modify the subtree under Widgets(PathLevel).
+	 *
+	 * @param PathLevel       The level in this WidgetPath whose focus to move.
+	 * @param MoveDirectin    Move focus forward or backward?
+	 *
+	 * @return true if the focus moved successfully, false if we were unable to move focus
+	 */
+	bool MoveFocus(int32 PathLevel, EFocusMoveDirection::Type MoveDirection);
+
+
+	FArrangedChildren Widgets;
+	TSharedPtr< SWindow > TopLevelWindow;
+
+private:
+
+	/**
+	 * Utility function to search recursively through a widget hierarchy for a specific widget
+	 *
+	 * @param Matcher         Some struct that has a "bool IsMatch( const TSharedRef<const SWidget>& InWidget ) const" method
+	 * @param  InCandidate      The current widget-geometry pair we're testing
+	 * @param  OutReversedPath  The resulting path in reversed order (canonical order is Windows @ index 0, Leafmost widget is last.)
+	 * @param  VisibilityFilter	Widgets must have this type of visibility to be included the path
+	 *
+	 * @return  true if the child widget was found; false otherwise
+	 */
+	template<typename MatchRuleType>
+	static bool SearchForWidgetRecursively( const MatchRuleType& MatchRule, const FArrangedWidget& InCandidate, FArrangedChildren& OutReversedPath, EVisibility VisibilityFilter = EVisibility::Visible );
+
+	/** Identical to SearchForWidgetRecursively, but iterates in reverse order */
+	template<typename MatchRuleType>
+	static bool SearchForWidgetRecursively_Reverse( const MatchRuleType& MatchRule, const FArrangedWidget& InCandidate, FArrangedChildren& OutReversedPath, EVisibility VisibilityFilter = EVisibility::Visible );
+
+};
+
+
+/**
+ * Just like a WidgetPath, but uses weak pointers and does not store geometry.
+ */
+class SLATECORE_API FWeakWidgetPath
+{
+public:
+	/** Construct a weak widget path from a widget path. Defaults to an invalid path. */
+	FWeakWidgetPath( const FWidgetPath& InWidgetPath = FWidgetPath( TSharedPtr<SWindow>(nullptr), FArrangedChildren(EVisibility::Visible) ) );
+
+	/** Should interrupted paths truncate or return an invalid path? */
+	struct EInterruptedPathHandling
+	{
+		enum Type
+		{
+			Truncate,
+			ReturnInvalid
+		};
+	};
+
+	/**
+	 * Make a non-weak WidgetPath out of this WeakWidgetPath. Do this by computing all the relevant geometries and converting the weak pointers to TSharedPtr.
+	 *
+	 * @param InterruptedPathHandling  Should interrupted paths result in a truncated path or an invalid path
+	 */
+	FWidgetPath ToWidgetPath( EInterruptedPathHandling::Type InterruptedPathHandling = EInterruptedPathHandling::Truncate ) const;
+
+	struct EPathResolutionResult
+	{
+		enum Result
+		{
+			Live,
+			Truncated
+		};
+	};
+	/**
+	 * Make a non-weak WidgetPath out of this WeakWidgetPath. Do this by computing all the relevant geometries and converting the weak pointers to TSharedPtr.
+	 *
+	 * @param WidgetPath				The non-weak path is returned via this.
+	 * @param InterruptedPathHandling	Should interrupted paths result in a truncated path or an invalid path.
+	 * @return Whether the path is truncated or live - a live path refers to a widget that is currently active and visible, a widget with a truncated path is not.
+	 */
+	EPathResolutionResult::Result ToWidgetPath( FWidgetPath& WidgetPath, EInterruptedPathHandling::Type InterruptedPathHandling = EInterruptedPathHandling::Truncate ) const;
+
+	bool ContainsWidget( const TSharedRef< const SWidget >& SomeWidget ) const;
+
+	/**
+	 * @param MoveDirection      Direction in which to move the focus.
+	 * 
+	 * @return The new focus path.
+	 */
+	FWidgetPath ToNextFocusedPath(EFocusMoveDirection::Type MoveDirection);
+	
+	/** Get the last (leaf-most) widget in this path; assumes path is valid */
+	TWeakPtr< SWidget > GetLastWidget() const
+	{
+		check( IsValid() );
+		return Widgets[Widgets.Num()-1];
+	}
+	
+	/** A valid path has at least one widget in it */
+	bool IsValid() const { return Widgets.Num() > 0; }
+
+	TArray< TWeakPtr<SWidget> > Widgets;
+	TWeakPtr< SWindow > Window;
+};
diff --git a/Engine/Source/Runtime/SlateCore/Public/Layout/WidgetPath.inl b/Engine/Source/Runtime/SlateCore/Public/Layout/WidgetPath.inl
new file mode 100644
--- /dev/null
+++ b/Engine/Source/Runtime/SlateCore/Public/Layout/WidgetPath.inl
@@ -0,0 +1,64 @@
+// Copyright 1998-2014 Epic Games, Inc. All Rights Reserved.
+
+#pragma once
+
+// FWidgetPath functions that depend on classes defined after WidgetPath.h is included in Slate.h
+
+/**
+ * Utility function to search recursively through a widget hierarchy for a specific widget
+ *
+ * @param Matcher         Some struct that has a "bool IsMatch( const TSharedRef<const SWidget>& InWidget ) const" method
+ * @param  InCandidate      The current widget-geometry pair we're testing
+ * @param  OutReversedPath  The resulting path in reversed order (canonical order is Windows @ index 0, Leafmost widget is last.)
+ * @param  VisibilityFilter	Widgets must have this type of visibility to be included the path
+ *
+ * @return  true if the child widget was found; false otherwise
+ */
+template<typename MatchRuleType>
+bool FWidgetPath::SearchForWidgetRecursively( const MatchRuleType& MatchRule, const FArrangedWidget& InCandidate, FArrangedChildren& OutReversedPath, EVisibility VisibilityFilter )
+{
+	FArrangedChildren ArrangedChildren(VisibilityFilter);
+	InCandidate.Widget->ArrangeChildren( InCandidate.Geometry, ArrangedChildren );
+
+	for( int32 ChildIndex = 0; ChildIndex < ArrangedChildren.Num(); ++ChildIndex )
+	{
+		const FArrangedWidget& SomeChild = ArrangedChildren(ChildIndex);
+		if ( MatchRule.IsMatch(SomeChild.Widget) )
+		{
+			OutReversedPath.AddWidget( SomeChild );
+			return true;
+		}
+		else if ( SearchForWidgetRecursively( MatchRule, SomeChild, OutReversedPath, VisibilityFilter ) )
+		{
+			OutReversedPath.AddWidget( SomeChild );
+			return true;
+		}
+	}
+
+	return false;
+}
+
+/** Identical to SearchForWidgetRecursively, but iterates in reverse order */
+template<typename MatchRuleType>
+bool FWidgetPath::SearchForWidgetRecursively_Reverse( const MatchRuleType& MatchRule, const FArrangedWidget& InCandidate, FArrangedChildren& OutReversedPath, EVisibility VisibilityFilter )
+{
+	FArrangedChildren ArrangedChildren(VisibilityFilter);
+	InCandidate.Widget->ArrangeChildren( InCandidate.Geometry, ArrangedChildren );
+
+	for( int32 ChildIndex = ArrangedChildren.Num()-1; ChildIndex >=0 ; --ChildIndex )
+	{
+		const FArrangedWidget& SomeChild = ArrangedChildren(ChildIndex);
+		if ( MatchRule.IsMatch(SomeChild.Widget) )
+		{
+			OutReversedPath.AddWidget( SomeChild );
+			return true;
+		}
+		else if ( SearchForWidgetRecursively( MatchRule, SomeChild, OutReversedPath, VisibilityFilter) )
+		{
+			OutReversedPath.AddWidget( SomeChild );
+			return true;
+		}
+	}
+
+	return false;
+}
diff --git a/Engine/Source/Runtime/SlateCore/Public/Logging/EventLogger.h b/Engine/Source/Runtime/SlateCore/Public/Logging/EventLogger.h
new file mode 100644
--- /dev/null
+++ b/Engine/Source/Runtime/SlateCore/Public/Logging/EventLogger.h
@@ -0,0 +1,126 @@
+// Copyright 1998-2014 Epic Games, Inc. All Rights Reserved.
+
+/*=============================================================================
+	EventLogger.h: Declares various Slate event logger implementations.
+=============================================================================*/
+
+#pragma once
+
+
+/**
+ * Implements an event logger that forwards events to two other event loggers.
+ */
+
+class FMultiTargetLogger
+	: public IEventLogger
+{
+public:
+
+	/**
+	 * Creates and initializes a new instance.
+	 *
+	 * @param InLoggerA The first logger to forward events to.
+	 * @param InLoggerB The second logger to forward events to.
+	 */
+	FMultiTargetLogger( const TSharedRef<IEventLogger>& InLoggerA, const TSharedRef<IEventLogger>& InLoggerB )
+		: LoggerA(InLoggerA)
+		, LoggerB(InLoggerB)
+	{ }
+
+public:
+
+	// Begin IEventLogger interface
+
+	virtual FString GetLog( ) const OVERRIDE
+	{
+		return LoggerA->GetLog() + TEXT("\n") + LoggerB->GetLog();
+	}
+
+	virtual void Log(EEventLog::Type Event, const FString& AdditionalContent, TSharedPtr<SWidget> Widget) OVERRIDE
+	{
+		LoggerA->Log(Event, AdditionalContent, Widget);
+		LoggerB->Log(Event, AdditionalContent, Widget);
+	}
+
+	virtual void SaveToFile( ) OVERRIDE
+	{
+		LoggerA->SaveToFile();
+		LoggerB->SaveToFile();
+	}
+
+	// EndIEventLogger interface
+
+private:
+
+	TSharedRef<IEventLogger> LoggerA;
+	TSharedRef<IEventLogger> LoggerB;
+};
+
+
+/**
+ * Implements an event logger that logs to a file.
+ */
+class SLATECORE_API FFileEventLogger
+	: public IEventLogger
+{
+public:
+
+	// Begin IEventLogger interface
+
+	virtual FString GetLog( ) const OVERRIDE;
+
+	virtual void Log( EEventLog::Type Event, const FString& AdditionalContent, TSharedPtr<SWidget> Widget ) OVERRIDE;
+	
+	virtual void SaveToFile( ) OVERRIDE;
+
+	// EndIEventLogger interface
+
+private:
+
+	// Holds the collection of logged events.
+	TArray<FString> LoggedEvents;
+};
+
+
+class SLATECORE_API FConsoleEventLogger
+	: public IEventLogger
+{
+public:
+
+	// Begin IEventLogger interface
+
+	virtual FString GetLog( ) const OVERRIDE;
+
+	virtual void Log( EEventLog::Type Event, const FString& AdditionalContent, TSharedPtr<SWidget> Widget ) OVERRIDE;
+	
+	virtual void SaveToFile( ) OVERRIDE { }
+	
+	// End IEventLogger interface
+
+private:
+};
+
+
+/**
+ * Implements an event logger that logs events related to assertions and ensures.
+ */
+class SLATECORE_API FStabilityEventLogger
+	: public IEventLogger
+{
+public:
+
+	// Begin IEventLogger interface
+	
+	virtual FString GetLog( ) const OVERRIDE;
+
+	virtual void Log( EEventLog::Type Event, const FString& AdditionalContent, TSharedPtr<SWidget> Widget ) OVERRIDE;
+	
+	virtual void SaveToFile( ) OVERRIDE { }
+	
+	// End IEventLogger interface
+
+private:
+
+	// Holds the collection of logged events.
+	TArray<FString> LoggedEvents;
+};
diff --git a/Engine/Source/Runtime/SlateCore/Public/Logging/IEventLogger.h b/Engine/Source/Runtime/SlateCore/Public/Logging/IEventLogger.h
new file mode 100644
--- /dev/null
+++ b/Engine/Source/Runtime/SlateCore/Public/Logging/IEventLogger.h
@@ -0,0 +1,114 @@
+// Copyright 1998-2014 Epic Games, Inc. All Rights Reserved.
+
+/*=============================================================================
+	IEventLogger.h: Declares the IEventLogger interface.
+=============================================================================*/
+
+#pragma once
+
+
+namespace EEventLog
+{
+	/**
+	 * Enumerates Slate events that can be logged.
+	 */
+	enum Type
+	{
+		MouseMove,
+		MouseEnter,
+		MouseLeave,
+		MouseButtonDown,
+		MouseButtonUp,
+		MouseButtonDoubleClick,
+		MouseWheel,
+		DragDetected,
+		DragEnter,
+		DragLeave,
+		DragOver,
+		DragDrop,
+		DropMessage,
+		KeyDown,
+		KeyUp,
+		KeyChar,
+		UICommand,
+		BeginTransaction,
+		EndTransaction,
+		TouchGesture,
+		Other,
+	};
+
+	/**
+	 * Converts a log event to a string.
+	 *
+	 * @param Event The event to convert.
+	 *
+	 * @return The event's string representation.
+	 */
+	inline FString ToString(EEventLog::Type Event)
+	{
+		switch (Event)
+		{
+		case MouseMove: return FString("MouseMove");
+		case MouseEnter: return FString("MouseEnter");
+		case MouseLeave: return FString("MouseLeave");
+		case MouseButtonDown: return FString("MouseButtonDown");
+		case MouseButtonUp: return FString("MouseButtonUp");
+		case MouseButtonDoubleClick: return FString("MouseButtonDoubleClick");
+		case MouseWheel: return FString("MouseWheel");
+		case DragDetected: return FString("DragDetected");
+		case DragEnter: return FString("DragEnter");
+		case DragLeave: return FString("DragLeave");
+		case DragOver: return FString("DragOver");
+		case DragDrop: return FString("DragDrop");
+		case DropMessage: return FString("DropMessage");
+		case KeyDown: return FString("KeyDown");
+		case KeyUp: return FString("KeyUp");
+		case KeyChar: return FString("KeyChar");
+		case UICommand: return FString("UICommand");
+		case BeginTransaction: return FString("BeginTransaction");
+		case EndTransaction: return FString("EndTransaction");
+		case TouchGesture: return FString("TouchGesture");
+		case Other: return FString("Other");
+
+		default:
+			return FString("Unknown");
+		}
+	}
+}
+
+
+/**
+ * Interface for Slate event loggers.
+ */
+class IEventLogger
+{
+public:
+
+	/**
+	 * Gets all collected events in a string.
+	 *
+	 * @return A string containing the logged events.
+	 */
+	virtual FString GetLog( ) const = 0;
+
+	/**
+	 * Logs an event.
+	 *
+	 * @param Event The event to log.
+	 * @param AdditionalContent An optional string associated with the event.
+	 * @param Widget The widget that generated the event.
+	 */
+	virtual void Log( EEventLog::Type Event, const FString& AdditionalContent, TSharedPtr<SWidget> Widget ) = 0;
+
+	/**
+	 * Saves all collected events to a file.
+	 */
+	virtual void SaveToFile( ) = 0;
+
+public:
+
+	/**
+	 * Virtual destructor.
+	 */
+	virtual ~IEventLogger( ) { }
+};
diff --git a/Engine/Source/Runtime/SlateCore/Public/Rendering/DrawElements.h b/Engine/Source/Runtime/SlateCore/Public/Rendering/DrawElements.h
new file mode 100644
--- /dev/null
+++ b/Engine/Source/Runtime/SlateCore/Public/Rendering/DrawElements.h
@@ -0,0 +1,657 @@
+// Copyright 1998-2014 Epic Games, Inc. All Rights Reserved.
+
+#pragma once
+
+
+class SWindow;
+class FSlateViewportInterface;
+
+
+struct FSlateGradientStop
+{
+	FVector2D Position;
+	FLinearColor Color;
+
+	FSlateGradientStop( const FVector2D& InPosition, const FLinearColor& InColor )
+		: Position(InPosition)
+		, Color(InColor)
+	{
+
+	}
+};
+
+template <> struct TIsPODType<FSlateGradientStop> { enum { Value = true }; };
+
+
+class FSlateDataPayload
+{
+public:
+	// Element tint
+	FLinearColor Tint;
+
+	// Spline Data
+	FVector2D StartPt;
+	FVector2D StartDir;
+	FVector2D EndPt;
+	FVector2D EndDir;
+
+	// Brush data
+	const FSlateBrush* BrushResource;
+
+	// Box Data
+	FVector2D RotationPoint;
+	float Angle;
+
+	// Spline/Line Data
+	float Thickness;
+
+	// Font data
+	FSlateFontInfo FontInfo;
+	FString Text;
+	
+	// Gradient data
+	TArray<FSlateGradientStop> GradientStops;
+	EOrientation GradientType;
+
+	// Line data
+	TArray<FVector2D> Points;
+
+	// Viewport data (intentionally weak to allow the source element to be destructed after this element has been added for drawing)
+	TWeakPtr<const ISlateViewport> Viewport;
+
+	// Misc data
+	bool bGammaCorrect;
+	bool bAllowBlending;
+
+	// Custom drawer data
+	TWeakPtr<ICustomSlateElement, ESPMode::ThreadSafe> CustomDrawer;
+
+	// Line data
+	ESlateLineJoinType::Type SegmentJoinType;
+	bool bAntialias;
+
+	FSlateDataPayload()
+		: Tint(FLinearColor::White)
+		, BrushResource(nullptr)
+		, Viewport(nullptr)
+	{ }
+
+	void SetBoxPayloadProperties( const FSlateBrush* InBrush, const FLinearColor& InTint )
+	{
+		Tint = InTint;
+		BrushResource = InBrush;
+		Angle = 0.0f;
+	}
+
+	void SetRotatedBoxPayloadProperties( const FSlateBrush* InBrush, float InAngle, const FVector2D& LocalRotationPoint, const FLinearColor& InTint )
+	{
+		Tint = InTint;
+		BrushResource = InBrush;
+		RotationPoint = LocalRotationPoint;
+		Angle = InAngle;
+	}
+
+	void SetTextPayloadProperties( const FString& InText, const FSlateFontInfo& InFontInfo, const FLinearColor& InTint )
+	{
+		Tint = InTint;
+		FontInfo = InFontInfo;
+		Text = InText;
+	}
+
+	void SetGradientPayloadProperties( const TArray<FSlateGradientStop>& InGradientStops, EOrientation InGradientType, bool bInGammaCorrect )
+	{
+		GradientStops = InGradientStops;
+		GradientType = InGradientType;
+		bGammaCorrect = bInGammaCorrect;
+	}
+
+	void SetSplinePayloadProperties( const FVector2D& InStart, const FVector2D& InStartDir, const FVector2D& InEnd, const FVector2D& InEndDir, float InThickness, const FLinearColor& InTint )
+	{
+		Tint = InTint;
+		StartPt = InStart;
+		StartDir = InStartDir;
+		EndPt = InEnd;
+		EndDir = InEndDir;
+		BrushResource = nullptr;
+		Thickness = InThickness;
+	}
+
+	void SetLinesPayloadProperties( const TArray<FVector2D>& InPoints, const FLinearColor& InTint, bool bInAntialias, ESlateLineJoinType::Type InJoinType )
+	{
+		Tint = InTint;
+		Thickness = 0.0f;
+		Points = InPoints;
+		SegmentJoinType = InJoinType;
+		bAntialias = bInAntialias;
+	}
+
+	void SetViewportPayloadProperties( const TSharedPtr<const ISlateViewport>& InViewport, const FLinearColor& InTint, bool bInGammaCorrect, bool bInAllowBlending )
+	{
+		Tint = InTint;
+		Viewport = InViewport;
+		bGammaCorrect = bInGammaCorrect;
+		bAllowBlending = bInAllowBlending;
+	}
+
+	void SetCustomDrawerPayloadProperties( const TSharedPtr<ICustomSlateElement, ESPMode::ThreadSafe>& InCustomDrawer )
+	{
+		CustomDrawer = InCustomDrawer;
+	}
+};
+
+/**
+ * FSlateDrawElement is the building block for Slate's rendering interface.
+ * Slate describes its visual output as an ordered list of FSlateDrawElement s
+ */
+class FSlateDrawElement
+{
+public:
+	enum EElementType
+	{
+		ET_Box,
+		ET_DebugQuad,
+		ET_Text,
+		ET_Spline,
+		ET_Line,
+		ET_Gradient,
+		ET_Viewport,
+		ET_Border,
+		ET_Custom,
+	};
+
+	enum ERotationSpace
+	{
+		/** Relative to the element.  (0,0) is the upper left corner of the element */
+		RelativeToElement,
+		/** Relative to the alloted paint geometry.  (0,0) is the upper left corner of the paint geometry */
+		RelativeToWorld,
+	};
+
+	/**
+	 * Creates a wireframe quad for debug purposes
+	 *
+	 * @param ElementList			The list in which to add elements
+	 * @param InLayer				The layer to draw the element on
+	 * @param PaintGeometry         DrawSpace position and dimensions; see FPaintGeometry
+	 * @param InClippingRect		Parts of the element are clipped if it falls outside of this rectangle
+	 */
+	SLATECORE_API static void MakeDebugQuad( FSlateWindowElementList& ElementList, uint32 InLayer, const FPaintGeometry& PaintGeometry, const FSlateRect& InClippingRect);
+
+	/**
+	 * Creates a box element based on the following diagram.  Allows for this element to be resized while maintain the border of the image
+	 * If there are no margins the resulting box is simply a quad
+	 *     ___LeftMargin    ___RightMargin
+	 *    /                /
+	 *  +--+-------------+--+
+	 *  |  |c1           |c2| ___TopMargin
+	 *  +--o-------------o--+
+	 *  |  |             |  |
+	 *  |  |c3           |c4|
+	 *  +--o-------------o--+
+	 *  |  |             |  | ___BottomMargin
+	 *  +--+-------------+--+
+	 *
+	 * @param ElementList			The list in which to add elements
+	 * @param InLayer               The layer to draw the element on
+	 * @param PaintGeometry         DrawSpace position and dimensions; see FPaintGeometry
+	 * @param InBrush               Brush to apply to this element
+	 * @param InClippingRect        Parts of the element are clipped if it falls outside of this rectangle
+	 * @param InDrawEffects         Optional draw effects to apply
+	 * @param InTint                Color to tint the element
+	 */
+	SLATECORE_API static void MakeBox( 
+		FSlateWindowElementList& ElementList,
+		uint32 InLayer, 
+		const FPaintGeometry& PaintGeometry, 
+		const FSlateBrush* InBrush, 
+		const FSlateRect& InClippingRect, 
+		ESlateDrawEffect::Type InDrawEffects = ESlateDrawEffect::None, 
+		const FLinearColor& InTint = FLinearColor::White );
+
+	SLATECORE_API static void MakeRotatedBox(
+		FSlateWindowElementList& ElementList,
+		uint32 InLayer, 
+		const FPaintGeometry& PaintGeometry, 
+		const FSlateBrush* InBrush, 
+		const FSlateRect& InClippingRect, 
+		ESlateDrawEffect::Type InDrawEffects = ESlateDrawEffect::None, 
+		float Angle = 0.0f,
+		TOptional<FVector2D> InRotationPoint = TOptional<FVector2D>(),
+		ERotationSpace RotationSpace = RelativeToElement,
+		const FLinearColor& InTint = FLinearColor::White );
+
+	/**
+	 * Creates a text element which displays a string of a rendered in a certain font on the screen
+	 *
+	 * @param ElementList			The list in which to add elements
+	 * @param InLayer               The layer to draw the element on
+	 * @param PaintGeometry         DrawSpace position and dimensions; see FPaintGeometry
+	 * @param InText                The string to draw
+	 * @param StartIndex            Inclusive index to start rendering from on the specified text
+	 * @param EndIndex				Exclusive index to stop rendering on the specified text
+	 * @param InFontInfo            The font to draw the string with
+	 * @param InClippingRect        Parts of the element are clipped if it falls outside of this rectangle
+	 * @param InDrawEffects         Optional draw effects to apply
+	 * @param InTint                Color to tint the element
+	 */
+	SLATECORE_API static void MakeText( FSlateWindowElementList& ElementList, uint32 InLayer, const FPaintGeometry& PaintGeometry, const FString& InText, const int32 StartIndex, const int32 EndIndex, const FSlateFontInfo& InFontInfo, const FSlateRect& InClippingRect,ESlateDrawEffect::Type InDrawEffects = ESlateDrawEffect::None, const FLinearColor& InTint =FLinearColor::White );
+	
+	SLATECORE_API static void MakeText( FSlateWindowElementList& ElementList, uint32 InLayer, const FPaintGeometry& PaintGeometry, const FString& InText, const FSlateFontInfo& InFontInfo, const FSlateRect& InClippingRect,ESlateDrawEffect::Type InDrawEffects = ESlateDrawEffect::None, const FLinearColor& InTint =FLinearColor::White );
+
+	SLATECORE_API static void MakeText( FSlateWindowElementList& ElementList, uint32 InLayer, const FPaintGeometry& PaintGeometry, const FText& InText, const FSlateFontInfo& InFontInfo, const FSlateRect& InClippingRect,ESlateDrawEffect::Type InDrawEffects = ESlateDrawEffect::None, const FLinearColor& InTint =FLinearColor::White );
+
+	/**
+	 * Creates a gradient element
+	 *
+	 * @param ElementList			   The list in which to add elements
+	 * @param InLayer                  The layer to draw the element on
+	 * @param PaintGeometry            DrawSpace position and dimensions; see FPaintGeometry
+	 * @param InGradientStops          List of gradient stops which define the element
+	 * @param InGradientType           The type of gradient (I.E Horizontal, vertical)
+	 * @param InClippingRect           Parts of the element are clipped if it falls outside of this rectangle
+	 * @param InDrawEffects            Optional draw effects to apply
+	 */
+	SLATECORE_API static void MakeGradient( FSlateWindowElementList& ElementList, uint32 InLayer, const FPaintGeometry& InPaintGeometry, TArray<FSlateGradientStop> InGradientStops, EOrientation InGradientType, const FSlateRect& InClippingRect, ESlateDrawEffect::Type InDrawEffects = ESlateDrawEffect::None, bool bGammaCorrect = true );
+
+	/**
+	 * Creates a spline element
+	 *
+	 * @param ElementList			The list in which to add elements
+	 * @param InLayer               The layer to draw the element on
+	 * @param PaintGeometry         DrawSpace position and dimensions; see FPaintGeometry
+	 * @param InStart               The start point of the spline (local space)
+	 * @param InStartDir            The direction of the spline from the start point
+	 * @param InEnd                 The end point of the spline (local space)
+	 * @param InEndDir              The direction of the spline to the end point
+	 * @param InClippingRect        Parts of the element are clipped if it falls outside of this rectangle
+	 * @param InDrawEffects         Optional draw effects to apply
+	 * @param InTint                Color to tint the element
+	 */
+	SLATECORE_API static void MakeSpline( FSlateWindowElementList& ElementList, uint32 InLayer, const FPaintGeometry& PaintGeometry, const FVector2D& InStart, const FVector2D& InStartDir, const FVector2D& InEnd, const FVector2D& InEndDir, const FSlateRect InClippingRect, float InThickness = 0.0f, ESlateDrawEffect::Type InDrawEffects = ESlateDrawEffect::None, const FLinearColor& InTint=FLinearColor::White );
+
+	/** Just like MakeSpline but in draw-space coordinates. This is useful for connecting already-transformed widgets together. */
+	SLATECORE_API static void MakeDrawSpaceSpline( FSlateWindowElementList& ElementList, uint32 InLayer, const FVector2D& InStart, const FVector2D& InStartDir, const FVector2D& InEnd, const FVector2D& InEndDir, const FSlateRect InClippingRect, float InThickness = 0.0f, ESlateDrawEffect::Type InDrawEffects = ESlateDrawEffect::None, const FColor& InTint=FColor(255,255,255) );
+
+	/**
+	 * Creates a line defined by the provided points
+	 *
+	 * @param ElementList			   The list in which to add elements
+	 * @param InLayer                  The layer to draw the element on
+	 * @param PaintGeometry            DrawSpace position and dimensions; see FPaintGeometry
+	 * @param Points                   Points that make up the lines.  The points are joined together. I.E if Points has A,B,C there the line is A-B-C.  To draw non-joining line segments call MakeLines multiple times
+	 * @param InClippingRect           Parts of the element are clipped if it falls outside of this rectangle
+	 * @param InDrawEffects            Optional draw effects to apply
+	 * @param InTint                   Color to tint the element
+	 */
+	SLATECORE_API static void MakeLines( FSlateWindowElementList& ElementList, uint32 InLayer, const FPaintGeometry& PaintGeometry, const TArray<FVector2D>& Points, const FSlateRect InClippingRect, ESlateDrawEffect::Type InDrawEffects = ESlateDrawEffect::None, const FLinearColor& InTint=FLinearColor::White, bool bAntialias = true );
+
+	/**
+	 * Creates a viewport element which is useful for rendering custom data in a texture into Slate
+	 *
+	 * @param ElementList		   The list in which to add elements
+	 * @param InLayer                  The layer to draw the element on
+	 * @param PaintGeometry            DrawSpace position and dimensions; see FPaintGeometry
+	 * @param Viewport                 Interface for drawing the viewport
+	 * @param InClippingRect           Parts of the element are clipped if it falls outside of this rectangle
+	 * @param InScale                  Draw scale to apply to the entire element
+	 * @param InDrawEffects            Optional draw effects to apply
+	 * @param InTint                   Color to tint the element
+	 */
+	SLATECORE_API static void MakeViewport( FSlateWindowElementList& ElementList, uint32 InLayer, const FPaintGeometry& PaintGeometry, TSharedPtr<const ISlateViewport> Viewport, const FSlateRect& InClippingRect, bool bGammaCorrect = true, bool bAllowBlending = true, ESlateDrawEffect::Type InDrawEffects = ESlateDrawEffect::None, const FLinearColor& InTint=FLinearColor::White );
+
+	/**
+	 * Creates a custom element which can be used to manually draw into the Slate render target with graphics API calls rather than Slate elements
+	 *
+	 * @param ElementList		   The list in which to add elements
+	 * @param InLayer                  The layer to draw the element on
+	 * @param PaintGeometry            DrawSpace position and dimensions; see FPaintGeometry
+	 * @param CustomDrawer		   Interface to a drawer which will be called when Slate renders this element
+	 */
+	SLATECORE_API static void MakeCustom( FSlateWindowElementList& ElementList, uint32 InLayer, TSharedPtr<ICustomSlateElement, ESPMode::ThreadSafe> CustomDrawer );
+
+
+	EElementType GetElementType() const { return ElementType; }
+	uint32 GetLayer() const { return Layer; }
+	const FVector2D& GetPosition() const { return Position; }
+	const FVector2D& GetSize() const { return Size; }
+	float GetScale() const { return Scale; }
+	const FSlateRect& GetClippingRect() const { return ClippingRect; }
+	const FSlateDataPayload& GetDataPayload() const { return DataPayload; }
+	uint32 GetDrawEffects() const { return DrawEffects; }
+
+private:
+
+	static FVector2D GetRotationPoint( const FPaintGeometry& PaintGeometry, const TOptional<FVector2D>& UserRotationPoint, ERotationSpace RotationSpace )
+	{
+		FVector2D RotationPoint(0,0);
+	
+		const FVector2D& Position = PaintGeometry.DrawPosition;
+		const FVector2D& Size = PaintGeometry.DrawSize;
+
+		switch( RotationSpace )
+		{
+		case RelativeToElement:
+			{
+				if( !UserRotationPoint.IsSet() )
+				{
+					// If the user did not specify a rotation point, we rotate about the center of the element
+					RotationPoint = ( Position + ( Position + Size ) ) / 2.0f;
+				}
+				else
+				{
+					// Rotate relative to the position of the element
+					RotationPoint = Position + UserRotationPoint.GetValue();
+				}
+			}
+			break;
+		case RelativeToWorld:
+			{
+				// Passthrough, its in world space
+				RotationPoint =  UserRotationPoint.Get( FVector2D::ZeroVector );
+			}
+			break;
+		default:
+			check(0);
+			break;
+		}
+
+		return RotationPoint;
+	}
+
+private:
+	FSlateDataPayload DataPayload;
+	FSlateRect ClippingRect;
+	FVector2D Position;
+	FVector2D Size;
+	float Scale;
+	uint32 Layer;
+	uint32 DrawEffects;
+	EElementType ElementType;
+};
+
+/**
+ * Shader parameters for slate
+ */
+struct FShaderParams
+{
+	/** Vertex shader parameters */
+	FVector4 VertexParams;
+	/** Pixel shader parameters */
+	FVector4 PixelParams;
+
+	FShaderParams()
+		: VertexParams( 0,0,0,0 )
+		, PixelParams( 0,0,0,0 )
+	{}
+
+	FShaderParams( const FVector4& InVertexParams, const FVector4& InPixelParams )
+		: VertexParams( InVertexParams )
+		, PixelParams( InPixelParams )
+	{}
+
+	bool operator==( const FShaderParams& Other ) const
+	{
+		return VertexParams == Other.VertexParams && PixelParams == Other.PixelParams;
+	}
+
+	static FShaderParams MakePixelShaderParams( const FVector4& PixelShaderParams )
+	{
+		return FShaderParams( FVector4(0,0,0,0), PixelShaderParams );
+	}
+
+	static FShaderParams MakeVertexShaderParams( const FVector4& VertexShaderParams )
+	{
+		return FShaderParams( VertexShaderParams, FVector4(0,0,0,0) );
+	}
+};
+
+/** 
+ * Represents an element batch for rendering. 
+ */
+class FSlateElementBatch
+{
+public: 
+	FSlateElementBatch( const FSlateShaderResource* InShaderResource, const FShaderParams& InShaderParams, ESlateShader::Type ShaderType, ESlateDrawPrimitive::Type PrimitiveType, ESlateDrawEffect::Type DrawEffects, ESlateBatchDrawFlag::Type DrawFlags )
+		: BatchKey( InShaderParams, ShaderType, PrimitiveType, DrawEffects, DrawFlags )
+		, ShaderResource(InShaderResource)
+		, VertexOffset(0)
+		, IndexOffset(0)
+		, NumVertices(0)
+		, NumIndices(0)
+		, NumElementsInBatch(0)
+		, VertexArrayIndex(INDEX_NONE)
+		, IndexArrayIndex(INDEX_NONE)
+	{}
+
+	FSlateElementBatch( TWeakPtr<ICustomSlateElement, ESPMode::ThreadSafe> InCustomDrawer )
+		: BatchKey( InCustomDrawer )
+		, ShaderResource( nullptr )
+		, VertexOffset(0)
+		, IndexOffset(0)
+		, NumVertices(0)
+		, NumIndices(0)
+		, NumElementsInBatch(0)
+		, VertexArrayIndex(INDEX_NONE)
+		, IndexArrayIndex(INDEX_NONE)
+	{}
+
+
+	bool operator==( const FSlateElementBatch& Other ) const
+	{	
+		return BatchKey == Other.BatchKey && ShaderResource == Other.ShaderResource;
+	}
+
+	friend uint32 GetTypeHash( const FSlateElementBatch& ElementBatch )
+	{
+		return PointerHash(ElementBatch.ShaderResource, GetTypeHash(ElementBatch.BatchKey));
+	}
+
+	const FSlateShaderResource* GetShaderResource() const { return ShaderResource; }
+	const FShaderParams& GetShaderParams() const { return BatchKey.ShaderParams; }
+	ESlateBatchDrawFlag::Type GetDrawFlags() const { return BatchKey.DrawFlags; }
+	ESlateDrawPrimitive::Type GetPrimitiveType() const { return BatchKey.DrawPrimitiveType; }
+	ESlateShader::Type GetShaderType() const { return BatchKey.ShaderType; } 
+	ESlateDrawEffect::Type GetDrawEffects() const { return BatchKey.DrawEffects; }
+	const TWeakPtr<ICustomSlateElement, ESPMode::ThreadSafe> GetCustomDrawer() const { return BatchKey.CustomDrawer; }
+private:
+	struct FBatchKey
+	{
+		const FShaderParams ShaderParams;
+		const TWeakPtr<ICustomSlateElement, ESPMode::ThreadSafe> CustomDrawer;
+		const ESlateBatchDrawFlag::Type DrawFlags;	
+		const ESlateShader::Type ShaderType;
+		const ESlateDrawPrimitive::Type DrawPrimitiveType;
+		const ESlateDrawEffect::Type DrawEffects;
+
+		FBatchKey( const FShaderParams& InShaderParams, ESlateShader::Type InShaderType, ESlateDrawPrimitive::Type InDrawPrimitiveType, ESlateDrawEffect::Type InDrawEffects, ESlateBatchDrawFlag::Type InDrawFlags )
+			: ShaderParams( InShaderParams )
+			, DrawFlags( InDrawFlags )
+			, ShaderType( InShaderType )
+			, DrawPrimitiveType( InDrawPrimitiveType )
+			, DrawEffects( InDrawEffects )
+		{
+		}
+
+		FBatchKey( TWeakPtr<ICustomSlateElement, ESPMode::ThreadSafe> InCustomDrawer )
+			: ShaderParams()
+			, CustomDrawer( InCustomDrawer )
+			, DrawFlags( ESlateBatchDrawFlag::None )
+			, ShaderType( ESlateShader::Default )
+			, DrawPrimitiveType( ESlateDrawPrimitive::TriangleList )
+			, DrawEffects( ESlateDrawEffect::None )
+		{}
+
+		FORCEINLINE bool operator==( const FBatchKey& Other ) const
+		{
+			return DrawFlags == Other.DrawFlags
+				&& ShaderType == Other.ShaderType
+				&& DrawPrimitiveType == Other.DrawPrimitiveType
+				&& DrawEffects == Other.DrawEffects
+			    && ShaderParams == Other.ShaderParams
+			   	&& CustomDrawer == Other.CustomDrawer;
+		}
+
+		friend uint32 GetTypeHash( const FBatchKey& InBatchKey )
+		{
+			return FCrc::MemCrc32(&InBatchKey.ShaderParams, sizeof(FShaderParams)) ^ ((InBatchKey.ShaderType << 16) | (InBatchKey.DrawFlags+InBatchKey.ShaderType+InBatchKey.DrawPrimitiveType+InBatchKey.DrawEffects));
+		}
+	};
+
+	/** A secondary key which represents elements needed to make a batch unique */
+	FBatchKey BatchKey;
+	/** Shader resource to use with this batch.  Used as a primary key.  No batch can have multiple textures */
+	const FSlateShaderResource* ShaderResource;
+
+public:
+	/** How far into the vertex buffer is this batch*/
+	uint32 VertexOffset;
+	/** How far into the index buffer this batch is*/	
+	uint32 IndexOffset;
+	/** Number of vertices in the batch */
+	uint32 NumVertices;
+	/** Number of indices in the batch */
+	uint32 NumIndices;
+	/** Number of elements in the batch */
+	uint32 NumElementsInBatch;
+	/** Index into an array of vertex arrays where this batches vertices are found (before submitting to the vertex buffer)*/
+	int32 VertexArrayIndex;
+	/** Index into an array of index arrays where this batches indices are found (before submitting to the index buffer) */
+	int32 IndexArrayIndex;
+};
+
+class FSlateRenderBatch
+{
+public:
+	FSlateRenderBatch( const FSlateElementBatch& InBatch )
+		: ShaderParams( InBatch.GetShaderParams() )
+		, Texture( InBatch.GetShaderResource() )
+		, CustomDrawer( InBatch.GetCustomDrawer() )
+		, DrawFlags( InBatch.GetDrawFlags() )
+		, ShaderType( InBatch.GetShaderType() )
+		, DrawPrimitiveType( InBatch.GetPrimitiveType() )
+		, DrawEffects( InBatch.GetDrawEffects() )
+		, VertexOffset( InBatch.VertexOffset )
+		, IndexOffset( InBatch.IndexOffset )
+		, NumVertices( InBatch.NumVertices )
+		, NumIndices( InBatch.NumIndices )
+	{}
+
+public:
+	const FShaderParams ShaderParams;
+
+	/** Texture to use with this batch.  */
+	const FSlateShaderResource* Texture;
+
+	const TWeakPtr<ICustomSlateElement, ESPMode::ThreadSafe> CustomDrawer;
+
+	const ESlateBatchDrawFlag::Type DrawFlags;	
+
+	const ESlateShader::Type ShaderType;
+
+	const ESlateDrawPrimitive::Type DrawPrimitiveType;
+
+	const ESlateDrawEffect::Type DrawEffects;
+
+	/** How far into the vertex buffer is this batch*/
+	const uint32 VertexOffset;
+	/** How far into the index buffer this batch is*/	
+	const uint32 IndexOffset;
+	/** Number of vertices in the batch */
+	const uint32 NumVertices;
+	/** Number of indices in the batch */
+	const uint32 NumIndices;
+};
+
+/**
+ * Represents a top level window and its draw elements.
+ */
+class FSlateWindowElementList
+{
+	friend class FSlateElementBatcher;
+public:
+	/** 
+	 * Construct a new list of elements with which to paint a window.
+	 *
+	 * @param InWindow   The window that we will be painting
+	 */
+	explicit FSlateWindowElementList( TSharedPtr<SWindow> InWindow = TSharedPtr<SWindow>() )
+		: TopLevelWindow( InWindow )
+	{		
+	}
+	
+	/** @return Get the window that we will be painting */
+	const TSharedPtr<SWindow> GetWindow() const
+	{
+		return TopLevelWindow.Pin();
+	}
+	
+	/** @return Get the window that we will be painting */
+	TSharedPtr<SWindow> GetWindow()
+	{
+		return TopLevelWindow.Pin();
+	}
+	
+	/** @return Get the draw elements that we want to render into this window */
+	const TArray<FSlateDrawElement>& GetDrawElements() const
+	{
+		return DrawElements;
+	}
+	
+	/**
+	 * Add a draw element to the list
+	 *
+	 * @param InDrawElement  The draw element to add
+	 */
+	void AddItem( const FSlateDrawElement& InDrawElement )
+	{
+		DrawElements.Add( InDrawElement );
+	}
+
+	FSlateDrawElement& AddUninitialized()
+	{
+		const int32 InsertIdx = DrawElements.AddUninitialized();
+		return *(new(DrawElements.GetTypedData() + InsertIdx) FSlateDrawElement());
+	}
+	
+	/**
+	 * Remove all the elements from this draw list.
+	 */
+	void Empty()
+	{
+		DrawElements.Empty();
+		RenderBatches.Empty();
+		BatchedVertices.Empty();
+		BatchedIndices.Empty();
+	}
+	
+	/**
+	 * Returns a list of element batches for this window
+	 */
+	const TArray<FSlateRenderBatch>& GetRenderBatches() const { return RenderBatches; }
+
+	/**
+	 * Returns all the vertices for every draw element in this window
+	 */
+	const TArray<FSlateVertex>& GetBatchedVertices() const { return BatchedVertices; }
+
+	/**
+	 * Returns all the indices for every draw element in this window
+	 */
+	const TArray<SlateIndex>& GetBatchedIndices() const { return BatchedIndices; }
+
+protected:
+	/** The top level window which these elements are being drawn on */
+	TWeakPtr<SWindow> TopLevelWindow;
+	
+	/** List of draw elements for the window */
+	TArray<FSlateDrawElement> DrawElements;
+
+	/** List of element batches sorted by later for use in rendering (for threaded renderers, can only be accessed from the render thread)*/
+	TArray<FSlateRenderBatch> RenderBatches;
+
+	/** All the vertices for every batch. */
+	TArray<FSlateVertex> BatchedVertices;
+
+	/** All the indices for every batch. */
+	TArray<SlateIndex> BatchedIndices;
+};
diff --git a/Engine/Source/Runtime/SlateCore/Public/Rendering/ElementBatcher.h b/Engine/Source/Runtime/SlateCore/Public/Rendering/ElementBatcher.h
new file mode 100644
--- /dev/null
+++ b/Engine/Source/Runtime/SlateCore/Public/Rendering/ElementBatcher.h
@@ -0,0 +1,217 @@
+// Copyright 1998-2014 Epic Games, Inc. All Rights Reserved.
+
+/*=============================================================================
+	ElementBatcher.h: Declares the FSlateElementBatcher class.
+=============================================================================*/
+
+#pragma once
+
+class FSlateShaderResource;
+
+/**
+ * A class which batches Slate elements for rendering
+ */
+class SLATECORE_API FSlateElementBatcher
+{
+public:
+
+	FSlateElementBatcher( TSharedRef<FSlateRenderingPolicy> InRenderingPolicy );
+	~FSlateElementBatcher();
+
+	/** 
+	 * Batches elements to be rendered 
+	 * 
+	 * @param DrawElements	The elements to batch
+	 */
+	void AddElements( const TArray<FSlateDrawElement>& DrawElements );
+
+	/**
+	 * Populates the window element list with batched data for use in rendering
+	 *
+	 * @param WindowElementList		The window element list to update
+	 * @param bRequiresStencilTest	Will be set to true if stencil testing should be enabled when drawing this viewport (currently for anti-aliased lines)
+	 */
+	void FillBatchBuffers(FSlateWindowElementList& WindowElementList, bool& bRequiresStencilTest);
+
+	/**
+	 * Returns true if the elements in this batcher require v-sync.
+	 */
+	bool RequiresVsync() const { return bRequiresVsync; }
+
+	/** 
+	 * Resets all stored data accumulated during the batching process
+	 */
+	void ResetBatches();
+
+	void ResetStats();
+
+private:
+
+	/** 
+	 * Creates vertices necessary to draw a Quad element 
+	 *
+	 * @param Position			The top left screen space position of the element
+	 * @param Size				The size of the element
+	 * @param Scale				The amount to scale the element by
+	 * @param InClippingRect	The clipping rectangle for the element
+	 * @param Layer				The layer to draw this element in
+	 * @param DrawEffects		DrawEffects to apply
+	 * @param Color				The color of the quad
+	 */
+	void AddQuadElement( const FVector2D& Position, const FVector2D& Size, float Scale, const FSlateRect& InClippingRect, uint32 Layer, FColor Color = FColor(255,255,255));
+
+	/** 
+	 * Creates vertices necessary to draw a 3x3 element
+	 *
+	 * @param Position			The top left screen space position of the element
+	 * @param Size				The size of the element
+	 * @param Scale				The amount to scale the element by
+	 * @param InPayload			The data payload for this element
+	 * @param InClippingRect	The clipping rectangle for the element
+	 * @param DrawEffects		DrawEffects to apply
+	 * @param Layer				The layer to draw this element in
+	 */
+	void AddBoxElement( const FVector2D& Position, const FVector2D& Size, float Scale, const FSlateDataPayload& InPayload, const FSlateRect& InClippingRect, ESlateDrawEffect::Type DrawEffects, uint32 Layer );
+
+	/** 
+	 * Creates vertices necessary to draw a string (one quad per character)
+	 *
+	 * @param Position			The top left screen space position of the element
+	 * @param Size				The size of the element
+	 * @param Scale				The amount to scale the element by
+	 * @param InPayload			The data payload for this element
+	 * @param InClippingRect	The clipping rectangle for the element
+	 * @param DrawEffects		DrawEffects to apply
+	 * @param Layer				The layer to draw this element in
+	 */
+	void AddTextElement( const FVector2D& Position, const FVector2D& Size, float Scale, const FSlateDataPayload& InPayload, const FSlateRect& InClippingRect, ESlateDrawEffect::Type DrawEffects, uint32 Layer );
+
+	/** 
+	 * Creates vertices necessary to draw a gradient box (horizontal or vertical)
+	 *
+	 * @param Position			The top left screen space position of the element
+	 * @param Size				The size of the element
+	 * @param Scale				The amount to scale the element by
+	 * @param InPayload			The data payload for this element
+	 * @param InClippingRect	The clipping rectangle for the element
+	 * @param DrawEffects		DrawEffects to apply
+	 * @param Layer				The layer to draw this element in
+	 */
+	void AddGradientElement( const FVector2D& Position, const FVector2D& Size, float Scale, const FSlateDataPayload& InPayload, const FSlateRect& InClippingRect, ESlateDrawEffect::Type DrawEffects, uint32 Layer );
+
+	/** 
+	 * Creates vertices necessary to draw a spline (Bezier curve)
+	 *
+	 * @param Position			The top left screen space position of the element
+	 * @param Scale				The amount to scale the element by
+	 * @param InPayload			The data payload for this element
+	 * @param InClippingRect	The clipping rectangle for the element
+	 * @param DrawEffects		DrawEffects to apply
+	 * @param Layer				The layer to draw this element in
+	 */
+	void AddSplineElement( const FVector2D& Position, float Scale, const FSlateDataPayload& InPayload, const FSlateRect& InClippingRect, ESlateDrawEffect::Type DrawEffects, uint32 Layer );
+
+	/** 
+	 * Creates vertices necessary to draw a series of attached line segments
+	 *
+	 * @param Position			The top left screen space position of the element
+	 * @param Size				The size of the element (not used)
+	 * @param Scale				The amount to scale the element by (not used)
+	 * @param InPayload			The data payload for this element
+	 * @param InClippingRect	The clipping rectangle for the element
+	 * @param DrawEffects		DrawEffects to apply
+	 * @param Layer				The layer to draw this element in
+	 */
+	void AddLineElement( const FVector2D& Position, const FVector2D& Size, float Scale, const FSlateDataPayload& InPayload, const FSlateRect& InClippingRect, ESlateDrawEffect::Type DrawEffects, uint32 Layer );
+	
+	/** 
+	 * Creates vertices necessary to draw a viewport (just a textured quad)
+	 *
+	 * @param Position			The top left screen space position of the element
+	 * @param Scale				The amount to scale the element by
+	 * @param Size				The size of the element
+	 * @param InPayload			The data payload for this element
+	 * @param InClippingRect	The clipping rectangle for the element
+	 * @param Layer				The layer to draw this element in
+	 */
+	void AddViewportElement( const FVector2D& Position, const FVector2D& Size, float Scale, const FSlateDataPayload& InPayload, const FSlateRect& InClippingRect, ESlateDrawEffect::Type DrawEffects, uint32 Layer );
+
+	/** 
+	 * Creates vertices necessary to draw a border element
+	 *
+	 * @param Position		The top left screen space position of the element
+	 * @param Size			The size of the element
+	 * @param Scale			The amount to scale the element by
+	 * @param InPayload		The data payload for this element
+	 * @param DrawEffects	DrawEffects to apply
+	 * @param Layer			The layer to draw this element in
+	 */
+	void AddBorderElement( const FVector2D& Position, const FVector2D& Size, float Scale, const FSlateDataPayload& InPayload, const FSlateRect& InClippingRect, ESlateDrawEffect::Type DrawEffects, uint32 Layer );
+
+	/**
+	 * Batches a custom slate drawing element
+	 *
+	 * @param Position		The top left screen space position of the element
+	 * @param Size			The size of the element
+	 * @param Scale			The amount to scale the element by
+	 * @param InPayload		The data payload for this element
+	 * @param DrawEffects	DrawEffects to apply
+	 * @param Layer			The layer to draw this element in
+	 */
+	void AddCustomElement( const FVector2D& Position, const FVector2D& Size, float Scale, const FSlateDataPayload& InPayload, const FSlateRect& InClippingRect, ESlateDrawEffect::Type DrawEffects, uint32 Layer );
+
+	/** 
+	 * Finds an batch for an element based on the passed in parameters
+	 * Elements with common parameters and layers will batched together.
+	 *
+	 * @param Layer			The layer where this element should be drawn (signifies draw order)
+	 * @param ShaderParams	The shader params for this element
+	 * @param InTexture		The texture to use in the batch
+	 * @param PrimitiveType	The primitive type( triangles, lines ) to use when drawing the batch
+	 * @param ShaderType	The shader to use when rendering this batch
+	 * @param DrawFlags		Any optional draw flags for this batch
+	 */	
+	FSlateElementBatch& FindBatchForElement( uint32 Layer, 
+											 const FShaderParams& ShaderParams, 
+											 const FSlateShaderResource* InTexture, 
+											 ESlateDrawPrimitive::Type PrimitiveType, 
+											 ESlateShader::Type ShaderType, 
+											 ESlateDrawEffect::Type DrawEffects, 
+											 ESlateBatchDrawFlag::Type DrawFlags = ESlateBatchDrawFlag::None);
+
+	void AddBasicVertices( TArray<FSlateVertex>& OutVertices, FSlateElementBatch& ElementBatch, const TArray<FSlateVertex>& VertexBatch );
+
+	void AddIndices( TArray<SlateIndex>& OutIndices, FSlateElementBatch& ElementBatch, const TArray<SlateIndex>& IndexBatch );
+
+private:
+
+	const FSlateBrush *SplineBrush;
+
+	// Element batch maps sorted by layer.
+	TMap<uint32, TSet<FSlateElementBatch>> LayerToElementBatches;
+
+	// Array of vertex lists that are currently free (have no elements in them).
+	TArray<uint32> VertexArrayFreeList;
+
+	// Array of index lists that are currently free (have no elements in them).
+	TArray<uint32> IndexArrayFreeList;
+
+	// Array of vertex lists for batching vertices. We use this method for quickly resetting the arrays without deleting memory.
+	TArray<TArray<FSlateVertex>> BatchVertexArrays;
+
+	// Array of vertex lists for batching indices. We use this method for quickly resetting the arrays without deleting memory.
+	TArray<TArray<SlateIndex>> BatchIndexArrays;
+
+	TSharedRef<FSlateRenderingPolicy> RenderingPolicy;
+
+	// true if any element in the batch requires vsync.
+	bool bRequiresVsync;
+
+	uint32 NumVertices;
+	uint32 NumBatches;
+	uint32 NumLayers;
+	uint32 TotalVertexMemory;
+	uint32 RequiredVertexMemory;
+	uint32 TotalIndexMemory;
+	uint32 RequiredIndexMemory;
+};
diff --git a/Engine/Source/Runtime/SlateCore/Public/Rendering/RenderingCommon.h b/Engine/Source/Runtime/SlateCore/Public/Rendering/RenderingCommon.h
new file mode 100644
--- /dev/null
+++ b/Engine/Source/Runtime/SlateCore/Public/Rendering/RenderingCommon.h
@@ -0,0 +1,465 @@
+// Copyright 1998-2014 Epic Games, Inc. All Rights Reserved.
+
+#pragma once
+
+struct FVector2D;
+class FSlateRect;
+
+#define SLATE_USE_32BIT_INDICES !PLATFORM_USES_ES2
+
+#if SLATE_USE_32BIT_INDICES
+typedef uint32 SlateIndex;
+#else
+typedef uint16 SlateIndex;
+#endif
+
+/**
+ * Draw primitive types                   
+ */
+namespace ESlateDrawPrimitive
+{
+	typedef uint8 Type;
+
+	const Type LineList = 0;
+	const Type TriangleList = 1;
+};
+
+/**
+ * Shader types. NOTE: mirrored in the shader file   
+ * If you add a type here you must also implement the proper shader type (TSlateElementPS).  See SlateShaders.h
+ */
+namespace ESlateShader
+{
+	typedef uint8 Type;
+	/** The default shader type.  Simple texture lookup */
+	const Type Default = 0;
+	/** Border shader */
+	const Type Border = 1;
+	/** Font shader, same as default except uses an alpha only texture */
+	const Type Font = 2;
+	/** Line segment shader. For drawing anti-aliased lines **/
+	const Type LineSegment = 3;
+};
+
+/**
+ * Effects that can be applied to elements when rendered.
+ * Note: New effects added should be in bit mask form
+ * * If you add a type here you must also implement the proper shader type (TSlateElementPS).  See SlateShaders.h
+ */
+namespace ESlateDrawEffect
+{
+	typedef uint8 Type;
+	/** No effect applied */
+	const Type None = 0;
+	/** Draw the element with a disabled effect */
+	const Type DisabledEffect = 1<<0;
+	/** Don't read from texture alpha channel */
+	const Type IgnoreTextureAlpha = 1<<2;
+};
+
+
+/** Flags for drawing a batch */
+namespace ESlateBatchDrawFlag
+{
+	typedef uint8 Type;
+	/** No draw flags */
+	const Type None					=		0;
+	/** Draw the element with no blending */
+	const Type NoBlending			=		0x01;
+	/** Draw the element as wireframe */
+	const Type Wireframe			=		0x02;
+	/** The element should be tiled horizontally */
+	const Type TileU				=		0x04;
+	/** The element should be tiled vertically */
+	const Type TileV				=		0x08;
+	/** No gamma correction should be done */
+	const Type NoGamma				=		0x10;
+};
+
+namespace ESlateLineJoinType
+{
+	enum Type
+	{
+		// Joins line segments with a sharp edge (miter)
+		Sharp =	0,
+		// Simply stitches together line segments
+		Simple = 1,
+	};
+};
+
+
+/** 
+ * A struct which defines a basic vertex seen by the Slate vertex buffers and shaders
+ */
+struct FSlateVertex
+{
+	/** Texture coordinates.  The first 2 are in xy and the 2nd are in zw */
+	FVector4 TexCoords; 
+	/** Clipping coordinates (left,top,right,bottom)*/
+	uint16 ClipCoords[4];
+	/** Position of the vertex in world space */
+	uint16 Position[2];
+	/** Vertex color */
+	FColor Color;
+	
+	FSlateVertex() {}
+
+	FORCEINLINE FSlateVertex( const FVector2D& InPosition, const FVector2D& InTexCoord, const FVector2D& InTexCoord2, const FColor& InColor, const FSlateRect& InClipCoords) 
+		: TexCoords( InTexCoord.X, InTexCoord.Y, InTexCoord2.X, InTexCoord2.Y )
+		, Color( InColor )
+	{
+
+		Position[0] = FMath::Trunc(InPosition.X);
+		Position[1] = FMath::Trunc(InPosition.Y);
+		ClipCoords[0] = FMath::Trunc(InClipCoords.Left);
+		ClipCoords[1] = FMath::Trunc(InClipCoords.Top);
+		ClipCoords[2] = FMath::Trunc(InClipCoords.Right);
+		ClipCoords[3] = FMath::Trunc(InClipCoords.Bottom);
+	}
+
+	FORCEINLINE FSlateVertex( const FVector2D& InPosition, const FVector2D& InTexCoord, const FVector2D& InTexCoord2, const FColor& InColor, const FVector4& InClipCoords) 
+		: TexCoords( InTexCoord.X, InTexCoord.Y, InTexCoord2.X, InTexCoord2.Y )
+		, Color( InColor )
+	{
+
+		Position[0] = FMath::Trunc(InPosition.X);
+		Position[1] = FMath::Trunc(InPosition.Y);
+		ClipCoords[0] = FMath::Trunc(InClipCoords.X);
+		ClipCoords[1] = FMath::Trunc(InClipCoords.Y);
+		ClipCoords[2] = FMath::Trunc(InClipCoords.Z);
+		ClipCoords[3] = FMath::Trunc(InClipCoords.W);
+	}
+
+	FORCEINLINE FSlateVertex( const FVector2D& InPosition, const FVector2D& InTexCoord, const FColor& InColor, const FSlateRect& InClipCoords ) 
+		: TexCoords( InTexCoord.X, InTexCoord.Y, 1.0f, 1.0f )
+		, Color( InColor )
+	{
+
+		Position[0] = FMath::Trunc(InPosition.X);
+		Position[1] = FMath::Trunc(InPosition.Y);
+		ClipCoords[0] = FMath::Trunc(InClipCoords.Left);
+		ClipCoords[1] = FMath::Trunc(InClipCoords.Top);
+		ClipCoords[2] = FMath::Trunc(InClipCoords.Right);
+		ClipCoords[3] = FMath::Trunc(InClipCoords.Bottom);
+	}
+
+	FORCEINLINE FSlateVertex( const FVector2D& InPosition, const FVector2D& InTexCoord, const uint32 InDWORD ) 
+		: TexCoords( InTexCoord.X, InTexCoord.Y, 1.0f, 1.0f )
+		, Color( InDWORD )
+	{
+
+		Position[0] = FMath::Trunc(InPosition.X);
+		Position[1] = FMath::Trunc(InPosition.Y);
+		ClipCoords[0] = 0;
+		ClipCoords[1] = 0;
+		ClipCoords[2] = 0;
+		ClipCoords[3] = 0;
+	}
+};
+
+template<> struct TIsPODType<FSlateVertex> { enum { Value = true }; };
+
+/**
+ * Viewport implementation interface that is used by SViewport when it needs to draw and processes input.                   
+ */
+class ISlateViewport
+{
+public:
+	virtual ~ISlateViewport() {}
+
+	/**
+	 * Called by Slate when the viewport widget is drawn
+	 * Implementers of this interface can use this method to perform custom
+	 * per draw functionality.  This is only called if the widget is visible
+	 *
+	 * @param AllottedGeometry	The geometry of the viewport widget
+	 */
+	virtual void OnDrawViewport( const FGeometry& AllottedGeometry, const FSlateRect& MyClippingRect, class FSlateWindowElementList& OutDrawElements, int32 LayerId, const FWidgetStyle& InWidgetStyle, bool bParentEnabled ) { }
+	
+	/**
+	 * Returns the size of the viewport                   
+	 */
+	virtual FIntPoint GetSize() const = 0;
+
+	/**
+	 * Returns a slate texture used to draw the rendered viewport in Slate.                   
+	 */
+	virtual class FSlateShaderResource* GetViewportRenderTargetTexture() const = 0;
+
+	/**
+	 * Performs any ticking necessary by this handle                   
+	 */
+	virtual void Tick( float DeltaTime )
+	{
+	}
+
+	/**
+	 * Returns true if the viewport should be vsynced.
+	 */
+	virtual bool RequiresVsync() const = 0;
+
+	/**
+	 * Called when Slate needs to know what the mouse cursor should be.
+	 * 
+	 * @return FCursorReply::Unhandled() if the event is not handled; FCursorReply::Cursor() otherwise.
+	 */
+	virtual FCursorReply OnCursorQuery( const FGeometry& MyGeometry, const FPointerEvent& CursorEvent )
+	{
+		return FCursorReply::Unhandled();
+	}
+
+	/**
+	 *	Returns whether the software cursor is currently visible
+	 */
+	virtual bool IsSoftwareCursorVisible() const
+	{
+		return false;
+	}
+
+	/**
+	 *	Returns the current position of the software cursor
+	 */
+	virtual FVector2D GetSoftwareCursorPosition() const
+	{
+		return FVector2D::ZeroVector;
+	}
+
+	/**
+	 * Called by Slate when a mouse button is pressed inside the viewport
+	 *
+	 * @param MyGeometry	Information about the location and size of the viewport
+	 * @param MouseEvent	Information about the mouse event
+	 *
+	 * @return Whether the event was handled along with possible requests for the system to take action.
+	 */
+	virtual FReply OnMouseButtonDown( const FGeometry& MyGeometry, const FPointerEvent& MouseEvent )
+	{
+		return FReply::Unhandled();
+	}
+
+	/**
+	 * Called by Slate when a mouse button is released inside the viewport
+	 *
+	 * @param MyGeometry	Information about the location and size of the viewport
+	 * @param MouseEvent	Information about the mouse event
+	 *
+	 * @return Whether the event was handled along with possible requests for the system to take action.
+	 */
+	virtual FReply OnMouseButtonUp( const FGeometry& MyGeometry, const FPointerEvent& MouseEvent )
+	{
+		return FReply::Unhandled();
+	}
+
+
+	virtual void OnMouseEnter( const FGeometry& MyGeometry, const FPointerEvent& MouseEvent )
+	{
+
+	}
+
+	virtual void OnMouseLeave( const FPointerEvent& MouseEvent )
+	{
+		
+	}
+
+	/**
+	 * Called by Slate when a mouse button is released inside the viewport
+	 *
+	 * @param MyGeometry	Information about the location and size of the viewport
+	 * @param MouseEvent	Information about the mouse event
+	 *
+	 * @return Whether the event was handled along with possible requests for the system to take action.
+	 */
+	virtual FReply OnMouseMove( const FGeometry& MyGeometry, const FPointerEvent& MouseEvent )
+	{
+		return FReply::Unhandled();
+	}
+
+	/**
+	 * Called by Slate when the mouse wheel is used inside the viewport
+	 *
+	 * @param MyGeometry	Information about the location and size of the viewport
+	 * @param MouseEvent	Information about the mouse event
+	 *
+	 * @return Whether the event was handled along with possible requests for the system to take action.
+	 */
+	virtual FReply OnMouseWheel( const FGeometry& MyGeometry, const FPointerEvent& MouseEvent )
+	{
+		return FReply::Unhandled();
+	}
+
+	/**
+	 * Called by Slate when the mouse wheel is used inside the viewport
+	 *
+	 * @param MyGeometry	Information about the location and size of the viewport
+	 * @param MouseEvent	Information about the mouse event
+	 *
+	 * @return Whether the event was handled along with possible requests for the system to take action.
+	 */
+	virtual FReply OnMouseButtonDoubleClick( const FGeometry& InMyGeometry, const FPointerEvent& InMouseEvent )
+	{
+		return FReply::Unhandled();
+	}
+
+	/**
+	 * Called by Slate when a key is pressed inside the viewport
+	 *
+	 * @param MyGeometry	Information about the location and size of the viewport
+	 * @param MouseEvent	Information about the keyboard event
+	 *
+	 * @return Whether the event was handled along with possible requests for the system to take action.
+	 */
+	virtual FReply OnKeyDown( const FGeometry& MyGeometry, const FKeyboardEvent& InKeyboardEvent )
+	{
+		return FReply::Unhandled();
+	}
+
+	/**
+	 * Called by Slate when a key is released inside the viewport
+	 *
+	 * @param MyGeometry	Information about the location and size of the viewport
+	 * @param MouseEvent	Information about the keyboard event
+	 *
+	 * @return Whether the event was handled along with possible requests for the system to take action.
+	 */
+	virtual FReply OnKeyUp( const FGeometry& MyGeometry, const FKeyboardEvent& InKeyboardEvent )
+	{
+		return FReply::Unhandled();
+	}
+
+	/**
+	 * Called by Slate when a character key is pressed while the viewport has keyboard focus
+	 *
+	 * @param MyGeometry	Information about the location and size of the viewport
+	 * @param MouseEvent	Information about the character that was pressed
+	 *
+	 * @return Whether the event was handled along with possible requests for the system to take action.
+	 */
+	virtual FReply OnKeyChar( const FGeometry& MyGeometry, const FCharacterEvent& InCharacterEvent )
+	{
+		return FReply::Unhandled();
+	}
+
+	/**
+	 * Called when the viewport gains keyboard focus.  
+	 *
+	 * @param InKeyboardFocusEvent	Information about what caused the viewport to gain focus
+	 */
+	virtual FReply OnKeyboardFocusReceived( const FKeyboardFocusEvent& InKeyboardFocusEvent )
+	{
+		return FReply::Unhandled();
+	}
+
+	/**
+	 * Called when a controller button is pressed
+	 * 
+	 * @param ControllerEvent	The controller event generated
+	 */
+	virtual FReply OnControllerButtonPressed( const FGeometry& MyGeometry, const FControllerEvent& ControllerEvent )
+	{
+		return FReply::Unhandled();
+	}
+
+	/**
+	 * Called when a controller button is released
+	 * 
+	 * @param ControllerEvent	The controller event generated
+	 */
+	virtual FReply OnControllerButtonReleased( const FGeometry& MyGeometry, const FControllerEvent& ControllerEvent )
+	{
+		return FReply::Unhandled();
+	}
+
+	/**
+	 * Called when an analog value on a controller changes
+	 * 
+	 * @param ControllerEvent	The controller event generated
+	 */
+	virtual FReply OnControllerAnalogValueChanged( const FGeometry& MyGeometry, const FControllerEvent& ControllerEvent )
+	{
+		return FReply::Unhandled();
+	}
+
+	/**
+	 * Called when a touchpad touch is started (finger down)
+	 * 
+	 * @param ControllerEvent	The controller event generated
+	 */
+	virtual FReply OnTouchStarted( const FGeometry& MyGeometry, const FPointerEvent& InTouchEvent )
+	{
+		return FReply::Unhandled();
+	}
+
+	/**
+	 * Called when a touchpad touch is moved  (finger moved)
+	 * 
+	 * @param ControllerEvent	The controller event generated
+	 */
+	virtual FReply OnTouchMoved( const FGeometry& MyGeometry, const FPointerEvent& InTouchEvent )
+	{
+		return FReply::Unhandled();
+	}
+
+	/**
+	 * Called when a touchpad touch is ended (finger lifted)
+	 * 
+	 * @param ControllerEvent	The controller event generated
+	 */
+	virtual FReply OnTouchEnded( const FGeometry& MyGeometry, const FPointerEvent& InTouchEvent )
+	{
+		return FReply::Unhandled();
+	}
+
+	/**
+	 * Called on a touchpad gesture event
+	 *
+	 * @param InGestureEvent	The touch event generated
+	 */
+	virtual FReply OnTouchGesture( const FGeometry& MyGeometry, const FPointerEvent& InGestureEvent )
+	{
+		return FReply::Unhandled();
+	}
+	
+	/**
+	 * Called when motion is detected (controller or device)
+	 * 
+	 * @param ControllerEvent	The controller event generated
+	 */
+	virtual FReply OnMotionDetected( const FGeometry& MyGeometry, const FMotionEvent& MotionEvent )
+	{
+		return FReply::Unhandled();
+	}
+
+	/**
+	 * Called when the viewport loses keyboard focus.  
+	 *
+	 * @param InKeyboardFocusEvent	Information about what caused the viewport to lose focus
+	 */
+	virtual void OnKeyboardFocusLost( const FKeyboardFocusEvent& InKeyboardFocusEvent )
+	{
+	}
+
+	/**
+	 * Called when the viewports top level window is being closed
+	 */
+	virtual void OnViewportClosed()
+	{
+	}
+};
+
+/**
+ * An interface for a custom slate drawing element
+ * Implementers of this interface are expected to handle destroying this interface properly when a separate 
+ * rendering thread may have access to it. (I.E this cannot be destroyed from a different thread if the rendering thread is using it)
+ */
+class ICustomSlateElement
+{
+public:
+	virtual ~ICustomSlateElement() {}
+
+	/** 
+	 * Called from the rendering thread when it is time to render the element
+	 *
+	 * @param RenderTarget	handle to the platform specific render target implementation.  Note this is already bound by Slate initially 
+	 */
+	virtual void DrawRenderThread( const void* RenderTarget ) = 0;
+};
\ No newline at end of file
diff --git a/Engine/Source/Runtime/SlateCore/Public/Rendering/RenderingPolicy.h b/Engine/Source/Runtime/SlateCore/Public/Rendering/RenderingPolicy.h
new file mode 100644
--- /dev/null
+++ b/Engine/Source/Runtime/SlateCore/Public/Rendering/RenderingPolicy.h
@@ -0,0 +1,54 @@
+// Copyright 1998-2014 Epic Games, Inc. All Rights Reserved.
+
+/*=============================================================================
+	RenderingPolicy.h: Declares the FRenderingPolicy class.
+=============================================================================*/
+
+#pragma once
+
+
+class FSlateElementBatch;
+class FSlateShaderResource;
+struct FSlateVertex;
+
+
+/**
+ * Abstract base class for Slate rendering policies.
+ */
+class FSlateRenderingPolicy
+{
+public:
+
+	/**
+	 * Creates and initializes a new instance.
+	 *
+	 * @param InPixelCenterOffset
+	 */
+	FSlateRenderingPolicy( float InPixelCenterOffset )
+		: PixelCenterOffset( InPixelCenterOffset )
+	{ }
+
+	/**
+	 * Virtual constructor.
+	 */
+	virtual ~FSlateRenderingPolicy( ) { }
+
+public:
+
+	virtual void UpdateBuffers( const FSlateWindowElementList& WindowElementList ) = 0;
+
+	virtual FSlateShaderResource* GetViewportResource( const ISlateViewport* InViewportInterface ) = 0;
+
+	virtual FSlateShaderResourceProxy* GetTextureResource( const FSlateBrush& Brush ) = 0;
+
+	virtual TSharedPtr<class FSlateFontCache>& GetFontCache( ) = 0;
+
+	float GetPixelCenterOffset( ) const
+	{
+		return PixelCenterOffset;
+	}
+
+private:
+
+	float PixelCenterOffset;
+};
diff --git a/Engine/Source/Runtime/SlateCore/Public/Rendering/SlateDrawBuffer.h b/Engine/Source/Runtime/SlateCore/Public/Rendering/SlateDrawBuffer.h
new file mode 100644
--- /dev/null
+++ b/Engine/Source/Runtime/SlateCore/Public/Rendering/SlateDrawBuffer.h
@@ -0,0 +1,65 @@
+// Copyright 1998-2014 Epic Games, Inc. All Rights Reserved.
+
+/*=============================================================================
+	SlateDrawBuffer.h: Declares the FSlateDrawBuffer class.
+=============================================================================*/
+
+#pragma once
+
+
+class SLATECORE_API FSlateDrawBuffer
+{
+public:
+
+	/**
+	 * Default constructor.
+	 */
+	explicit FSlateDrawBuffer( );
+
+	/**
+	 * Destructor.
+	 */
+	~FSlateDrawBuffer( );
+
+public:
+
+	/**
+	 * Removes all data from the buffer.
+	 */
+	void ClearBuffer( );
+
+	/**
+	 * Creates a new FSlateWindowElementList and returns a reference to it so it can have draw elements added to it
+	 *
+	 * @param ForWindow    The window for which we are creating a list of paint elements.
+	 */
+	FSlateWindowElementList& AddWindowElementList( TSharedRef<SWindow> ForWindow );
+
+	/**
+	 * Gets all window element lists in this buffer.
+	 */
+	TArray<FSlateWindowElementList>& GetWindowElementLists( )
+	{
+		return WindowElementLists;
+	}
+
+	/** 
+	 * Locks the draw buffer.  Indicates that the viewport is in use
+	 *
+	 * @return true if the viewport could be locked.  False otherwise
+	 */
+	bool Lock( );
+
+	/**
+	 * Unlocks the buffer.  Indicates that the buffer is free                   
+	 */
+	void Unlock( );
+
+protected:
+
+	// List of window element lists.
+	TArray<FSlateWindowElementList> WindowElementLists;
+
+	// 1 if this buffer is locked, 0 otherwise.
+	int32 Locked;
+};
