Slate: Refactored core Slate implementation into SlateCore module in preparation for UMG.

diff --git a/Engine/Source/Runtime/SlateCore/Public/Rendering/SlateRenderer.h b/Engine/Source/Runtime/SlateCore/Public/Rendering/SlateRenderer.h
new file mode 100644
--- /dev/null
+++ b/Engine/Source/Runtime/SlateCore/Public/Rendering/SlateRenderer.h
@@ -0,0 +1,183 @@
+// Copyright 1998-2014 Epic Games, Inc. All Rights Reserved.
+
+/*=============================================================================
+	FSlateRenderer.h: Declares the FSlateRenderer class.
+=============================================================================*/
+
+#pragma once
+
+
+class SWindow;
+
+
+/**
+ * Abstract base class for Slate renderers.
+ */
+class FSlateRenderer
+{
+public:
+	virtual SLATECORE_API ~FSlateRenderer();
+
+	/** Returns a draw buffer that can be used by Slate windows to draw window elements */
+	virtual FSlateDrawBuffer& GetDrawBuffer() = 0;
+
+	virtual void Initialize() = 0;
+
+	virtual void Destroy() = 0;
+
+	/**
+	 * Creates a rendering viewport
+	 *
+	 * @param InWindow	The window to create the viewport for
+	 */ 
+	virtual void CreateViewport( const TSharedRef<SWindow> InWindow ) = 0;
+
+	/**
+	 * Requests that a rendering viewport be resized
+	 *
+	 * @param Window		The window to resize
+	 * @param Width			The new width of the window
+	 * @param Height		The new height of the window
+	 */
+	virtual void RequestResize( const TSharedPtr<SWindow>& InWindow, uint32 NewSizeX, uint32 NewSizeY ) = 0;
+
+	/**
+	 * Sets fullscreen state on the window's rendering viewport
+	 *
+	 * @param	InWindow		The window to update fullscreen state on
+	 * @param	OverrideResX	0 if no override
+	 * @param	OverrideResY	0 if no override
+	 */
+	virtual void UpdateFullscreenState( const TSharedRef<SWindow> InWindow, uint32 OverrideResX = 0, uint32 OverrideResY = 0 ) = 0;
+
+	/** 
+	 * Creates necessary resources to render a window and sends draw commands to the rendering thread
+	 *
+	 * @param WindowDrawBuffer	The buffer containing elements to draw 
+	 */
+	virtual void DrawWindows( FSlateDrawBuffer& InWindowDrawBuffer ) = 0;
+	
+	/** 
+	 * Renders a window using resources stored from a previous call to DrawWindows, if
+	 * the previous call did store that data. Optional implementation.
+	 */
+	virtual void DrawWindows() {}
+	
+	/**
+	 * Copies all slate windows out to a buffer at half resolution with debug information
+	 * like the mouse cursor and any keypresses.
+	 */
+	virtual void CopyWindowsToDrawBuffer(const TArray<FString>& KeypressBuffer) {}
+	
+	/** Allows and disallows access to the crash tracker buffer data on the CPU */
+	virtual void MapCrashTrackerBuffer(void** OutImageData, int32* OutWidth, int32* OutHeight) {}
+	virtual void UnmapCrashTrackerBuffer() {}
+	
+	/** 
+	 * Sets which color vision filter to use
+	 */
+	virtual void SetColorVisionDeficiencyType( uint32 Type ) {} 
+
+	/** 
+	 * Creates a dynamic image resource and returns its size
+	 *
+	 * @param InTextureName The name of the texture resource
+	 * @return The size of the loaded texture
+	 */
+	virtual FIntPoint GenerateDynamicImageResource(const FName InTextureName) {check(0); return FIntPoint( 0, 0 );}
+
+	/**
+	 * Creates a dynamic image resource
+	 *
+	 * @param ResourceName		The name of the texture resource
+	 * @param Width				The width of the resource
+	 * @param Height			The height of the image
+	 * @param Bytes				The payload for the resource
+	 * @return					true if the resource was successfully generated, otherwise false
+	 */
+	virtual bool GenerateDynamicImageResource( FName ResourceName, uint32 Width, uint32 Height, const TArray< uint8 >& Bytes ) { return false; }
+
+	/** Called when a window is destroyed to give the renderer a chance to free resources */
+	virtual void OnWindowDestroyed( const TSharedRef<SWindow>& InWindow ) = 0;
+	
+	/**
+	 * Returns the viewport rendering resource (backbuffer) for the provided window
+	 *
+	 * @param Window	The window to get the viewport from 
+	 */
+	virtual void* GetViewportResource( const SWindow& Window ) { return nullptr; }
+	
+	TSharedRef< class FSlateFontMeasure > GetFontMeasureService() const 
+	{
+		return FontMeasure.ToSharedRef();
+	}
+
+	/**
+	 * Flushes all cached data from the font cache
+	 */
+	void SLATECORE_API FlushFontCache();
+
+	/**
+	 * Gives the renderer a chance to wait for any render commands to be completed before returning/
+	 */
+	virtual void FlushCommands() const {};
+
+	/**
+	 * Gives the renderer a chance to synchronize with another thread in the event that the renderer runs 
+	 * in a multi-threaded environment.  This function does not return until the sync is complete
+	 */
+	virtual void Sync() const {};
+
+	/**
+	 * Reloads all texture resources from disk
+	 */
+	virtual void ReloadTextureResources() {}
+
+	/**
+	 * Loads all the resources used by the specified SlateStyle
+	 */
+	virtual void LoadStyleResources( const ISlateStyle& Style ) {}
+
+	/** Creates a window to visualize the texture atlases */
+	virtual void DisplayTextureAtlases() {}
+
+	/** Releases a specific resource */
+	virtual void ReleaseDynamicResource( const FSlateBrush& Brush ) = 0;
+
+	/**
+	 * Returns whether or not a viewport should be in  fullscreen
+	 *
+	 * @Window	The window to check for fullscreen
+	 * @return true if the window's viewport should be fullscreen
+	 */
+	bool SLATECORE_API IsViewportFullscreen( const SWindow& Window ) const;
+
+	/** Returns whether shaders that Slate depends on have been compiled. */
+	virtual bool AreShadersInitialized() const { return true; }
+
+	/** 
+	 * Removes references to FViewportRHI's.  
+	 * This has to be done explicitly instead of using the FRenderResource mechanism because FViewportRHI's are managed by the game thread.
+	 * This is needed before destroying the RHI device. 
+	 */
+	virtual void InvalidateAllViewports() {}
+
+	/** 
+	 * Prepares the renderer to take a screenshot of the UI.  The Rect is portion of the rendered output
+	 * that will be stored into the TArray of FColors.
+	 */
+	virtual void PrepareToTakeScreenshot(const FIntRect& Rect, TArray<FColor>* OutColorData) {}
+
+protected:
+
+	TSharedPtr<class FSlateFontCache> FontCache;
+	TSharedPtr<class FSlateFontMeasure> FontMeasure;
+};
+
+
+/**
+ * Is this thread valid for sending out rendering commands?
+ * If the slate loading thread exists, then yes, it is always safe
+ * Otherwise, we have to be on the game thread
+ */
+bool SLATECORE_API IsThreadSafeForSlateRendering( );
diff --git a/Engine/Source/Runtime/SlateCore/Public/SlateCore.h b/Engine/Source/Runtime/SlateCore/Public/SlateCore.h
new file mode 100644
--- /dev/null
+++ b/Engine/Source/Runtime/SlateCore/Public/SlateCore.h
@@ -0,0 +1,159 @@
+// Copyright 1998-2014 Epic Games, Inc. All Rights Reserved.
+
+/*=============================================================================
+	SlateCore.h: SlateCore module public header file.
+=============================================================================*/
+
+#pragma once
+
+
+/* Dependencies
+ *****************************************************************************/
+
+#include "CoreUObject.h"
+#include "InputCore.h"
+#include "ModuleManager.h"
+
+
+/* Globals
+ *****************************************************************************/
+
+SLATECORE_API DECLARE_LOG_CATEGORY_EXTERN(LogSlate, Log, All);
+SLATECORE_API DECLARE_LOG_CATEGORY_EXTERN(LogSlateStyles, Log, All);
+
+DECLARE_CYCLE_STAT_EXTERN(TEXT("Measure String"), STAT_SlateMeasureStringTime, STATGROUP_Slate,);
+DECLARE_CYCLE_STAT_EXTERN(TEXT("Slate Rendering RT Time"), STAT_SlateRenderingRTTime, STATGROUP_Slate , SLATECORE_API);
+DECLARE_CYCLE_STAT_EXTERN(TEXT("Update Buffers RT"), STAT_SlateUpdateBufferRTTime, STATGROUP_Slate , SLATECORE_API);
+DECLARE_CYCLE_STAT_EXTERN(TEXT("Draw Time"), STAT_SlateDrawTime, STATGROUP_Slate, SLATECORE_API);
+DECLARE_CYCLE_STAT_EXTERN(TEXT("Update Buffers GT"), STAT_SlateUpdateBufferGTTime, STATGROUP_Slate,);
+DECLARE_CYCLE_STAT_EXTERN(TEXT("Font Caching"), STAT_SlateFontCachingTime, STATGROUP_Slate,);
+DECLARE_CYCLE_STAT_EXTERN(TEXT("Build Elements"), STAT_SlateBuildElements, STATGROUP_Slate,);
+DECLARE_CYCLE_STAT_EXTERN(TEXT("Find Batch"), STAT_SlateFindBatchTime, STATGROUP_Slate,);
+DECLARE_CYCLE_STAT_EXTERN(TEXT("Build Lines"), STAT_SlateBuildLines, STATGROUP_Slate,);
+DECLARE_CYCLE_STAT_EXTERN(TEXT("Build Splines"), STAT_SlateBuildSplines, STATGROUP_Slate,);
+DECLARE_CYCLE_STAT_EXTERN(TEXT("Build Gradients"), STAT_SlateBuildGradients, STATGROUP_Slate,);
+DECLARE_CYCLE_STAT_EXTERN(TEXT("Build Text"), STAT_SlateBuildText, STATGROUP_Slate,);
+DECLARE_CYCLE_STAT_EXTERN(TEXT("Build Border"), STAT_SlateBuildBorders, STATGROUP_Slate,);
+DECLARE_CYCLE_STAT_EXTERN(TEXT("Build Boxes"), STAT_SlateBuildBoxes, STATGROUP_Slate,);
+DECLARE_CYCLE_STAT_EXTERN(TEXT("Slate Rendering GT Time"), STAT_SlateRenderingGTTime, STATGROUP_Slate, SLATECORE_API);
+DECLARE_CYCLE_STAT_EXTERN(TEXT("Build Viewports"), STAT_SlateBuildViewports, STATGROUP_Slate,);
+
+DECLARE_DWORD_ACCUMULATOR_STAT_EXTERN(TEXT("Num Texture Atlases"), STAT_SlateNumTextureAtlases, STATGROUP_SlateMemory, SLATECORE_API);
+DECLARE_DWORD_ACCUMULATOR_STAT_EXTERN(TEXT("Num Font Atlases"), STAT_SlateNumFontAtlases, STATGROUP_SlateMemory,);
+DECLARE_DWORD_ACCUMULATOR_STAT_EXTERN(TEXT("Num Non-Atlased Textures"), STAT_SlateNumNonAtlasedTextures, STATGROUP_SlateMemory,SLATECORE_API);
+DECLARE_DWORD_ACCUMULATOR_STAT_EXTERN(TEXT("Num Dynamic Textures"), STAT_SlateNumDynamicTextures, STATGROUP_SlateMemory, SLATECORE_API);
+
+DECLARE_DWORD_COUNTER_STAT_EXTERN(TEXT("Num Layers"), STAT_SlateNumLayers, STATGROUP_Slate,);
+DECLARE_DWORD_COUNTER_STAT_EXTERN(TEXT("Num Batches"), STAT_SlateNumBatches, STATGROUP_Slate,);
+DECLARE_DWORD_COUNTER_STAT_EXTERN(TEXT("Num Vertices"), STAT_SlateVertexCount, STATGROUP_Slate,);
+
+DECLARE_MEMORY_STAT_EXTERN(TEXT("Texture Data CPU Memory"), STAT_SlateTextureDataMemory, STATGROUP_SlateMemory, SLATECORE_API);
+DECLARE_MEMORY_STAT_EXTERN(TEXT("Texture Atlas Memory"), STAT_SlateTextureAtlasMemory, STATGROUP_SlateMemory,);
+DECLARE_MEMORY_STAT_EXTERN(TEXT("Font Kerning Table Memory"), STAT_SlateFontKerningTableMemory, STATGROUP_SlateMemory,);
+DECLARE_MEMORY_STAT_EXTERN(TEXT("Font Measure Memory"), STAT_SlateFontMeasureCacheMemory, STATGROUP_SlateMemory,);
+DECLARE_MEMORY_STAT_EXTERN(TEXT("Batch Vertex Memory"), STAT_SlateVertexBatchMemory, STATGROUP_SlateMemory,);
+DECLARE_MEMORY_STAT_EXTERN(TEXT("Batch Index Memory"), STAT_SlateIndexBatchMemory, STATGROUP_SlateMemory,);
+DECLARE_MEMORY_STAT_EXTERN(TEXT("Vertex Buffer Memory"), STAT_SlateVertexBufferMemory, STATGROUP_SlateMemory, SLATECORE_API);
+DECLARE_MEMORY_STAT_EXTERN(TEXT("Index Buffer Memory"), STAT_SlateIndexBufferMemory, STATGROUP_SlateMemory, SLATECORE_API);
+
+
+/* Includes
+ *****************************************************************************/
+
+// Types
+#include "SlateConstants.h"
+#include "SlateEnums.h"
+#include "SlateStructs.h"
+
+// Layout
+#include "Margin.h"
+#include "SlateRect.h"
+#include "PaintGeometry.h"
+#include "Geometry.h"
+#include "Visibility.h"
+#include "ArrangedWidget.h"
+#include "ArrangedChildren.h"
+#include "SNullWidget.h"
+#include "Children.h"
+#include "WidgetPath.h"
+#include "LayoutUtils.h"
+
+// Animation
+#include "CurveHandle.h"
+#include "CurveSequence.h"
+#include "SlateSprings.h"
+
+// Sound
+#include "SlateSound.h"
+#include "ISlateSoundDevice.h"
+
+// Styling
+#include "WidgetStyle.h"
+#include "SlateBrush.h"
+#include "SlateColor.h"
+#include "SlateWidgetStyle.h"
+#include "SlateWidgetStyleAsset.h"
+
+// Textures
+#include "SlateShaderResource.h"
+#include "SlateTextureData.h"
+#include "TextureAtlas.h"
+#include "TextureManager.h"
+
+// Fonts
+#include "SlateFontInfo.h"
+#include "FontTypes.h"
+#include "FontCache.h"
+#include "FontMeasure.h"
+
+// Brushes
+#include "SlateBorderBrush.h"
+#include "SlateBoxBrush.h"
+#include "SlateColorBrush.h"
+#include "SlateDynamicImageBrush.h"
+#include "SlateImageBrush.h"
+#include "SlateNoResource.h"
+
+// Styling (continued)
+#include "StyleDefaults.h"
+#include "ISlateStyle.h"
+#include "SlateStyleRegistry.h"
+#include "SlateStyle.h"
+
+// Input
+#include "CursorReply.h"
+#include "Events.h"
+#include "DragAndDrop.h"
+#include "Reply.h"
+
+// Rendering
+#include "RenderingCommon.h"
+#include "DrawElements.h"
+#include "RenderingPolicy.h"
+#include "ElementBatcher.h"
+#include "SlateDrawBuffer.h"
+#include "SlateRenderer.h"
+
+// Application
+#include "IToolTip.h"
+#include "SlateApplicationBase.h"
+#include "ThrottleManager.h"
+
+// Widgets
+#include "DeclarativeSyntaxSupport.h"
+#include "SWidget.h"
+#include "SCompoundWidget.h"
+#include "SLeafWidget.h"
+#include "SPanel.h"
+#include "WidgetPath.inl"
+////
+#include "SBoxPanel.h"
+#include "SOverlay.h"
+#include "SlateTypes.h"
+#include "CoreStyle.h"
+////
+#include "SWindow.h"
+
+//Logging
+#include "IEventLogger.h"
+#include "EventLogger.h"
diff --git a/Engine/Source/Runtime/SlateCore/Public/Sound/ISlateSoundDevice.h b/Engine/Source/Runtime/SlateCore/Public/Sound/ISlateSoundDevice.h
new file mode 100644
--- /dev/null
+++ b/Engine/Source/Runtime/SlateCore/Public/Sound/ISlateSoundDevice.h
@@ -0,0 +1,33 @@
+// Copyright 1998-2014 Epic Games, Inc. All Rights Reserved.
+#pragma once
+
+
+/** Interface that Slate uses to play sounds. */
+class SLATECORE_API ISlateSoundDevice
+{
+public:
+	/** Play the sound resource within the given Slate sound object */
+	virtual void PlaySound(const FSlateSound& Sound, int32 UserIndex = 0) const = 0;
+
+	/** @return The duration of the given sound resource */
+	virtual float GetSoundDuration(const FSlateSound& Sound) const = 0;
+
+	/** Virtual destructor: because it is an interface / pure virtual base class. */
+	virtual ~ISlateSoundDevice(){}
+};
+
+/** Silent implementation of ISlateSoundDevice; it plays nothing. */
+class SLATECORE_API FNullSlateSoundDevice : public ISlateSoundDevice
+{
+public:
+	virtual void PlaySound(const FSlateSound&, int32) const OVERRIDE
+	{
+	}
+
+	virtual float GetSoundDuration(const FSlateSound& Sound) const OVERRIDE
+	{
+		return 0.0f;
+	}
+
+	virtual ~FNullSlateSoundDevice(){}
+};
diff --git a/Engine/Source/Runtime/SlateCore/Public/Sound/SlateSound.h b/Engine/Source/Runtime/SlateCore/Public/Sound/SlateSound.h
new file mode 100644
--- /dev/null
+++ b/Engine/Source/Runtime/SlateCore/Public/Sound/SlateSound.h
@@ -0,0 +1,86 @@
+// Copyright 1998-2014 Epic Games, Inc. All Rights Reserved.
+
+/*=============================================================================
+	SlateSound.h: Declares the FSlateSound structure.
+=============================================================================*/
+
+#pragma once
+
+#include "SlateSound.generated.h"
+
+
+/**
+ * An intermediary to make UBaseSound available for Slate to play sounds
+ */
+USTRUCT()
+struct SLATECORE_API FSlateSound
+{
+	GENERATED_USTRUCT_BODY()
+
+public:
+
+	/**
+	 * Default constructor.
+	 */
+	FSlateSound( )
+		: ResourceObject(nullptr)
+	{ }
+
+public:
+
+	/** Get the resource object associated with this sound
+	 *
+	 * @note Ensure that you only access the resource as a USoundBase 
+	 *
+	 * @return The resource object, or null if this sound resource is empty
+	 */
+	UObject* GetResourceObject( ) const;
+
+	/**
+	 * Used when updating the ResourceObject within an FSlateSound from the editor to clear out any legacy data that may be set.
+	 */
+	void StripLegacyData_DEPRECATED( );
+
+	/**
+	 * Used to upgrade a FName property to a FSlateSound property
+	 */
+	bool SerializeFromMismatchedTag( struct FPropertyTag const& Tag, FArchive& Ar );
+
+public:
+
+	/**
+	 * Construct a FSlateSound from a FName
+	 * @note This functionality is deprecated and should only be used for upgrading old data
+	 *
+	 * @param SoundName The name of the sound to load, eg) SoundCue'/Game/QA_Assets/Sound/TEST_Music_Ambient.TEST_Music_Ambient'
+	 *
+	 * @return An FSlateSound object with a ResourceObject pointer that's either set the given sound, or points to null if the sound could not be found
+	 */
+	static FSlateSound FromName_DEPRECATED( const FName& SoundName );
+
+protected:
+
+	/**
+	 * Pointer to the USoundBase. Holding onto it as a UObject because USoundBase is not available in Slate core.
+	 * Edited via FSlateSoundStructCustomization to ensure you can only set USoundBase assets on it.
+	 */
+	UPROPERTY(EditAnywhere, Category=Sound)
+	UObject* ResourceObject;
+
+	/** The legacy resource name; only used by sounds that have been set-up in code, or otherwise upgraded from old FName properties, set to NAME_None in non-legacy instances */
+	FName LegacyResourceName_DEPRECATED;
+
+	/** A weak pointer the the resource loaded from the legacy resource name; may be null if the resource needs (re)loading */
+	mutable TWeakObjectPtr<UObject> LegacyResourceObject_DEPRECATED;
+};
+
+
+template<>
+struct TStructOpsTypeTraits<FSlateSound>
+	: public TStructOpsTypeTraitsBase
+{
+	enum 
+	{
+		WithSerializeFromMismatchedTag = true,
+	};
+};
diff --git a/Engine/Source/Runtime/SlateCore/Public/Styling/CoreStyle.h b/Engine/Source/Runtime/SlateCore/Public/Styling/CoreStyle.h
new file mode 100644
--- /dev/null
+++ b/Engine/Source/Runtime/SlateCore/Public/Styling/CoreStyle.h
@@ -0,0 +1,46 @@
+// Copyright 1998-2014 Epic Games, Inc. All Rights Reserved.
+
+/*=============================================================================
+	CoreStyle.h: Declares the FCoreStyle class.
+=============================================================================*/
+
+#pragma once
+
+
+/**
+ * Core slate style
+ */
+class SLATECORE_API FCoreStyle 
+{
+public:
+
+	static TSharedRef<class ISlateStyle> Create( const FName& InStyleSetName = "CoreStyle" );
+
+	/** @return the singleton instance. */
+	static const ISlateStyle& Get( )
+	{
+		return *(Instance.Get());
+	}
+
+	static void ResetToDefault( );
+
+	/** Used to override the default selection colors */
+	static void SetSelectorColor( const FLinearColor& NewColor );
+	static void SetSelectionColor( const FLinearColor& NewColor );
+	static void SetInactiveSelectionColor( const FLinearColor& NewColor );
+	static void SetPressedSelectionColor( const FLinearColor& NewColor );
+
+	// todo: jdale - These are only here because of UTouchInterface::Activate and the fact that GetDynamicImageBrush is non-const
+	static const TSharedPtr<FSlateDynamicImageBrush> GetDynamicImageBrush( FName BrushTemplate, FName TextureName, const ANSICHAR* Specifier = nullptr );
+	static const TSharedPtr<FSlateDynamicImageBrush> GetDynamicImageBrush( FName BrushTemplate, const ANSICHAR* Specifier, class UTexture2D* TextureResource, FName TextureName );
+	static const TSharedPtr<FSlateDynamicImageBrush> GetDynamicImageBrush( FName BrushTemplate, class UTexture2D* TextureResource, FName TextureName );
+
+private:
+
+	static void SetStyle( const TSharedRef< class ISlateStyle >& NewStyle );
+
+private:
+
+	/** Singleton instances of this style. */
+	static TSharedPtr< class ISlateStyle > Instance;
+};
diff --git a/Engine/Source/Runtime/SlateCore/Public/Styling/ISlateStyle.h b/Engine/Source/Runtime/SlateCore/Public/Styling/ISlateStyle.h
new file mode 100644
--- /dev/null
+++ b/Engine/Source/Runtime/SlateCore/Public/Styling/ISlateStyle.h
@@ -0,0 +1,167 @@
+// Copyright 1998-2014 Epic Games, Inc. All Rights Reserved.
+#pragma once
+
+/**
+ * 
+ */
+class SLATECORE_API ISlateStyle
+{
+public:
+	/** Constructor. */
+	ISlateStyle()
+	{}
+
+	/** Destructor. */
+	virtual ~ISlateStyle()
+	{}
+
+	/** @return The name used to identity this style set */
+	virtual const FName& GetStyleSetName() const = 0;
+
+	/**
+	 * Populate an array of FSlateBrush with resources consumed by this style.
+	 * @param OutResources - the array to populate.
+	 */
+	virtual void GetResources( TArray< const FSlateBrush* >& OutResources ) const = 0;
+
+	template< typename WidgetStyleType >            
+	const WidgetStyleType& GetWidgetStyle( FName PropertyName, const ANSICHAR* Specifier = nullptr ) const 
+	{
+		const FSlateWidgetStyle* ExistingStyle = GetWidgetStyleInternal( WidgetStyleType::TypeName, Join( PropertyName, Specifier ) );
+
+		if ( ExistingStyle == nullptr )
+		{
+			return WidgetStyleType::GetDefault();
+		}
+
+		return *static_cast< const WidgetStyleType* >( ExistingStyle );
+	}
+
+	template< typename WidgetStyleType >            
+	bool HasWidgetStyle( FName PropertyName, const ANSICHAR* Specifier = nullptr ) const
+	{
+		return GetWidgetStyleInternal( WidgetStyleType::TypeName, Join( PropertyName, Specifier ) ) != nullptr;
+	}
+
+	/**
+	 * @param PropertyName   Name of the property to get
+	 * @param Specifier      An optional string to append to the property name
+	 *
+	 * @return a float property.
+	 */
+	virtual float GetFloat( const FName PropertyName, const ANSICHAR* Specifier = nullptr ) const = 0;
+
+	/**
+	 * @param PropertyName   Name of the property to get
+	 * @param Specifier      An optional string to append to the property name
+	 *
+	 * @return a FVector2D property.
+	 */
+	virtual FVector2D GetVector( const FName PropertyName, const ANSICHAR* Specifier = nullptr ) const = 0;
+
+	/**
+	 * @param PropertyName   Name of the property to get
+	 * @param Specifier      An optional string to append to the property name
+	 *
+	 * @return a FLinearColor property.
+	 */
+	virtual const FLinearColor& GetColor( const FName PropertyName, const ANSICHAR* Specifier = nullptr ) const = 0;
+
+	/**
+	 * @param PropertyName   Name of the property to get
+	 * @param Specifier      An optional string to append to the property name
+	 *
+	 * @return a FLinearColor property.
+	 */
+	virtual const FSlateColor GetSlateColor( const FName PropertyName, const ANSICHAR* Specifier = nullptr ) const = 0;
+
+	/**
+	 * @param PropertyName   Name of the property to get
+	 * @param Specifier      An optional string to append to the property name
+	 *
+	 * @return a FMargin property.
+	 */
+	virtual const FMargin& GetMargin( const FName PropertyName, const ANSICHAR* Specifier = nullptr ) const = 0;
+
+	/**
+	 * @param PropertyName   Name of the property to get
+	 * @param Specifier      An optional string to append to the property name
+	 *
+	 * @return an FSlateBrush property.
+	 */		
+	virtual const FSlateBrush* GetBrush( const FName PropertyName, const ANSICHAR* Specifier = nullptr ) const = 0;
+
+	/**
+	 * Just like GetBrush, but returns DefaultBrush instead of the
+	 * "missing brush" image when the resource is not found.
+	 * 
+	 * @param PropertyName   Name of the property to get
+	 * @param Specifier      An optional string to append to the property name
+	 *
+	 * @return an FSlateBrush property.
+	 */		
+	virtual const FSlateBrush* GetOptionalBrush( const FName PropertyName, const ANSICHAR* Specifier = nullptr, const FSlateBrush* const DefaultBrush = FStyleDefaults::GetNoBrush() ) const = 0;
+
+	virtual const TSharedPtr< FSlateDynamicImageBrush > GetDynamicImageBrush( const FName BrushTemplate, const FName TextureName, const ANSICHAR* Specifier = nullptr ) = 0;
+
+	virtual const TSharedPtr< FSlateDynamicImageBrush > GetDynamicImageBrush( const FName BrushTemplate, const ANSICHAR* Specifier, UTexture2D* TextureResource, const FName TextureName ) = 0;
+
+	virtual const TSharedPtr< FSlateDynamicImageBrush > GetDynamicImageBrush( const FName BrushTemplate, UTexture2D* TextureResource, const FName TextureName ) = 0;
+
+	/**
+	 * Get default FSlateBrush.
+	 * @return -	Default Slate brush value.
+	 */
+	virtual FSlateBrush* GetDefaultBrush() const = 0;
+
+	/** Look up a sound property specified by PropertyName and optional Specifier. */
+	virtual const FSlateSound& GetSound( const FName PropertyName, const ANSICHAR* Specifier = nullptr ) const = 0;
+
+	/**
+	 * @param Propertyname   Name of the property to get
+	 * @param InFontInfo     The value to set
+	 *
+	 * @return a InFontInfo property.
+	 */		
+	virtual FSlateFontInfo GetFontStyle( const FName PropertyName, const ANSICHAR* Specifier = nullptr ) const = 0;
+
+	static FName Join( FName A, const ANSICHAR* B )
+	{
+		if( B == nullptr )
+		{
+			return A;
+		}
+		else
+		{
+			return FName( *( A.ToString() + B ) );
+		}
+	}
+
+protected:
+
+	/** 
+	 * The severity of the message type 
+	 * Ordered according to their severity
+	 */
+	enum EStyleMessageSeverity
+	{
+		CriticalError		= 0,
+		Error				= 1,
+		PerformanceWarning	= 2,
+		Warning				= 3,
+		Info				= 4,	// Should be last
+	};
+
+
+protected:
+
+	/**
+	 * 
+	 */
+	virtual const FSlateWidgetStyle* GetWidgetStyleInternal( const FName DesiredTypeName, const FName StyleName ) const = 0;
+
+	/**
+	 * 
+	 */
+	virtual void Log( EStyleMessageSeverity Severity, const FText& Message ) const = 0;
+};
\ No newline at end of file
diff --git a/Engine/Source/Runtime/SlateCore/Public/Styling/SlateBrush.h b/Engine/Source/Runtime/SlateCore/Public/Styling/SlateBrush.h
new file mode 100644
--- /dev/null
+++ b/Engine/Source/Runtime/SlateCore/Public/Styling/SlateBrush.h
@@ -0,0 +1,236 @@
+// Copyright 1998-2014 Epic Games, Inc. All Rights Reserved.
+
+/*=============================================================================
+	SlateBrush.h: Declares the FSlateBrush structure.
+=============================================================================*/
+
+#pragma once
+
+#include "Margin.h"
+#include "SlateColor.h"
+
+#include "SlateBrush.generated.h"
+
+
+/**
+ * Enumerates ways in which an image can be drawn.
+ */
+UENUM()
+namespace ESlateBrushDrawType
+{
+	enum Type
+	{
+		/** Don't do anything */
+		NoDrawType,
+
+		/** Draw a 3x3 box, where the sides and the middle stretch based on the Margin */
+		Box,
+
+		/** Draw a 3x3 border where the sides tile and the middle is empty */
+		Border,
+
+		/** Draw an image; margin is ignored */
+		Image
+	};
+}
+
+
+/**
+ * Enumerates tiling options for image drawing.
+ */
+UENUM()
+namespace ESlateBrushTileType
+{
+	enum Type
+	{
+		/** Just stretch */
+		NoTile,
+
+		/** Tile the image horizontally */
+		Horizontal,
+
+		/** Tile the image vertically */
+		Vertical,
+
+		/** Tile in both directions */
+		Both
+	};
+}
+
+
+/*
+ * Enumerates brush image types.
+ */
+UENUM()
+namespace ESlateBrushImageType
+{
+	enum Type
+	{
+		/** No image is loaded.  Color only brushes, transparent brushes etc */
+		NoImage,
+
+		/** The image to be loaded is in full color */
+		FullColor,
+
+		/** The image is a special texture in linear space (usually a rendering resource such as a lookup table)*/
+		Linear,
+	};
+}
+
+
+/**
+ * An brush which contains information about how to draw a Slate element
+ */
+USTRUCT()
+struct SLATECORE_API FSlateBrush
+{
+	GENERATED_USTRUCT_BODY()
+
+	/** Size of the resource in Slate Units */
+	UPROPERTY(EditAnywhere, Category=Brush)
+	FVector2D ImageSize;
+
+	/** How to draw the image */
+	UPROPERTY(EditAnywhere, Category=Brush)
+	TEnumAsByte<enum ESlateBrushDrawType::Type > DrawAs;
+
+	/** The margin to use in Box and Border modes */
+	UPROPERTY(EditAnywhere, Category=Brush, meta=(UVSpace="true"))
+	FMargin Margin;
+
+	/** Tinting applied to the image. */
+	UPROPERTY()
+	FLinearColor Tint_DEPRECATED;
+
+	/** Tinting applied to the image. */
+	UPROPERTY(EditAnywhere, Category=Brush, meta=(DisplayName="Tint"))
+	FSlateColor TintColor;
+
+	/** How to tile the image in Image mode */
+	UPROPERTY(EditAnywhere, Category=Brush)
+	TEnumAsByte<enum ESlateBrushTileType::Type> Tiling;
+
+	/** The type of image */
+	UPROPERTY()
+	TEnumAsByte<enum ESlateBrushImageType::Type> ImageType;
+
+public:
+
+	/**
+	 * Default constructor.
+	 */
+	FSlateBrush()
+		: ImageSize(32.0f, 32.0f)
+		, DrawAs(ESlateBrushDrawType::Image)
+		, Margin(0.0f)
+		, Tint_DEPRECATED(FLinearColor::White)
+		, TintColor(FLinearColor::White)
+		, Tiling(ESlateBrushTileType::NoTile)
+		, ImageType(ESlateBrushImageType::NoImage)
+		, ResourceObject(nullptr)
+		, ResourceName(NAME_None)
+		, bIsDynamicallyLoaded(false)
+		, bHasUObject_DEPRECATED(false)
+	{ }
+
+public:
+
+	/**
+	 * Gets the name of the resource object, if any.
+	 *
+	 * @return Resource name, or NAME_None if the resource object is not set.
+	 */
+	const FName GetResourceName( ) const
+	{
+		return ((ResourceName != NAME_None) || (ResourceObject == nullptr))
+			? ResourceName
+			: ResourceObject->GetFName();
+	}
+
+	/**
+	 * Gets the UObject that represents the brush resource, if any.
+	 *
+	 * Ensure that we only access the TextureObject as a UTexture2D.
+	 *
+	 * @return The resource object, or nullptr if it is not set.
+	 */
+	class UObject* GetResourceObject( ) const
+	{
+		return ResourceObject;
+	}
+
+	/**
+	 * Gets the brush's tint color.
+	 *
+	 * @param InWidgetStyle The widget style to get the tint for.
+	 *
+	 * @return Tint color.
+	 */
+	FLinearColor GetTint( const FWidgetStyle& InWidgetStyle ) const
+	{
+		return TintColor.GetColor(InWidgetStyle);
+	}
+
+	/**
+	 * Checks whether this brush has a UTexture object
+	 *
+	 * @return true if it has a UTexture object, false otherwise.
+	 */
+	bool HasUObject( ) const
+	{
+		return (ResourceObject != nullptr) || (bHasUObject_DEPRECATED);
+	}
+
+	/**
+	 * Checks whether the brush resource is loaded dynamically.
+	 *
+	 * @return true if loaded dynamically, false otherwise.
+	 */
+	bool IsDynamicallyLoaded( ) const
+	{
+		return bIsDynamicallyLoaded;
+	}
+
+public:
+
+	/**
+	 * Gets the identifier for UObject based texture paths.
+	 *
+	 * @return Texture identifier string.
+	 */
+	static const FString UTextureIdentifier( );
+
+protected:
+
+	/** Pointer to the UTexture2D. Holding onto it as a UObject because UTexture2D is not available in Slate core. */
+	UPROPERTY(EditAnywhere, Category=Brush, meta=(DisplayThumbnail="true", DisplayName="Texture"))
+	UObject* ResourceObject;
+
+	/** The name of the rendering resource to use */
+	UPROPERTY()
+	FName ResourceName;
+
+	/** Whether or not the brush path is a path to a UObject */
+	bool bIsDynamicallyLoaded;
+
+	/** Whether or not the brush has a UTexture resource */
+	bool bHasUObject_DEPRECATED;
+
+protected:
+
+	/** 
+	 * This constructor is protected; use one of the deriving classes instead.
+	 *
+	 * @param InDrawType      How to draw the texture
+	 * @param InResourceName  The name of the resource
+	 * @param InMargin        Margin to use in border and box modes
+	 * @param InTiling        Tile horizontally/vertically or both? (only in image mode)
+	 * @param InImageType	  The type of image
+	 * @param InTint		  Tint to apply to the element.
+	 */
+	 FORCENOINLINE FSlateBrush( ESlateBrushDrawType::Type InDrawType, const FName InResourceName, const FMargin& InMargin, ESlateBrushTileType::Type InTiling, ESlateBrushImageType::Type InImageType, const FVector2D& InImageSize, const FLinearColor& InTint = FLinearColor::White, UObject* InObjectResource = nullptr, bool bInDynamicallyLoaded = false );
+
+	 FORCENOINLINE FSlateBrush( ESlateBrushDrawType::Type InDrawType, const FName InResourceName, const FMargin& InMargin, ESlateBrushTileType::Type InTiling, ESlateBrushImageType::Type InImageType, const FVector2D& InImageSize, const TSharedRef< FLinearColor >& InTint, UObject* InObjectResource = nullptr, bool bInDynamicallyLoaded = false );
+
+	 FORCENOINLINE FSlateBrush( ESlateBrushDrawType::Type InDrawType, const FName InResourceName, const FMargin& InMargin, ESlateBrushTileType::Type InTiling, ESlateBrushImageType::Type InImageType, const FVector2D& InImageSize, const FSlateColor& InTint, UObject* InObjectResource = nullptr, bool bInDynamicallyLoaded = false );
+};
diff --git a/Engine/Source/Runtime/SlateCore/Public/Styling/SlateColor.h b/Engine/Source/Runtime/SlateCore/Public/Styling/SlateColor.h
new file mode 100644
--- /dev/null
+++ b/Engine/Source/Runtime/SlateCore/Public/Styling/SlateColor.h
@@ -0,0 +1,179 @@
+// Copyright 1998-2014 Epic Games, Inc. All Rights Reserved.
+
+/*=============================================================================
+	SlateColor.h: Declares the FSlateColor structure.
+=============================================================================*/
+
+#pragma once
+
+#include "SlateColor.generated.h"
+
+
+/**
+ * Enumerates types of color values that can be held by Slate color.
+ *
+ * Should we use the specified color? If not, then which color from the style should we use.
+ */
+UENUM()
+namespace ESlateColorStylingMode
+{
+	enum Type
+	{
+		/** Color value is stored in this Slate color. */
+		UseColor_Specified,
+
+		/** Color value is stored in the linked color. */
+		UseColor_Specified_Link,
+
+		/** Use the widget's foreground color. */
+		UseColor_Foreground,
+
+		/** Use the widget's subdued color. */
+		UseColor_Foreground_Subdued
+	};
+}
+
+
+/**
+ * A Slate color can be a directly specified value, or the color can be pulled from a WidgetStyle.
+ */
+USTRUCT()
+struct FSlateColor
+{
+	GENERATED_USTRUCT_BODY()
+
+public:
+
+	/**
+	 * Default constructor.
+	 *
+	 * Uninitialized Slate colors are Fuchsia by default.
+	 */
+	FSlateColor()
+		: SpecifiedColor(FLinearColor(1.0f, 0.0f, 1.0f))
+		, ColorUseRule(ESlateColorStylingMode::UseColor_Specified)
+	{ }
+
+	/**
+	 * Creates a new Slate color with the specified values.
+	 *
+	 * @param InColor The color value to assign.
+	 */
+	FSlateColor( const FLinearColor& InColor )
+		: SpecifiedColor(InColor)
+		, ColorUseRule(ESlateColorStylingMode::UseColor_Specified)
+	{ }
+
+	/**
+	 * Creates a new Slate color that is linked to the given values.
+	 *
+	 * @param InColor The color value to assign.
+	 */
+	FSlateColor( const TSharedRef< FLinearColor >& InColor )
+		: SpecifiedColor(FLinearColor(1.0f, 0.0f, 1.0f))
+		, ColorUseRule(ESlateColorStylingMode::UseColor_Specified_Link)
+		, LinkedSpecifiedColor(InColor)
+	{ }
+
+public:
+
+	/**
+	 * Gets the color value represented by this Slate color.
+	 *
+	 * @param InWidgetStyle The widget style to use when this color represents a foreground or subdued color.
+	 *
+	 * @return The color value.
+	 *
+	 * @see GetSpecifiedColor
+	 */
+	const FLinearColor& GetColor( const FWidgetStyle& InWidgetStyle ) const
+	{
+		switch(ColorUseRule)
+		{
+		default:
+		case ESlateColorStylingMode::UseColor_Foreground:
+			return InWidgetStyle.GetForegroundColor();
+			break;
+
+		case ESlateColorStylingMode::UseColor_Specified:
+			return SpecifiedColor;
+			break;
+
+		case ESlateColorStylingMode::UseColor_Specified_Link:
+			return *LinkedSpecifiedColor;
+			break;
+
+		case ESlateColorStylingMode::UseColor_Foreground_Subdued:
+			return InWidgetStyle.GetSubduedForegroundColor();
+			break;
+		};
+	}
+
+	/**
+	 * Gets the specified color value.
+	 *
+	 * @return The specified color value.
+	 *
+	 * @see GetColor
+	 * @see IsColorSpecified
+	 */
+	FLinearColor GetSpecifiedColor( ) const
+	{
+		if (LinkedSpecifiedColor.IsValid())
+		{
+			return *LinkedSpecifiedColor;
+		}
+
+		return SpecifiedColor;
+	}
+
+	/**
+	 * Checks whether the values for this color have been specified.
+	 *
+	 * @return true if specified, false otherwise.
+	 *
+	 * @see GetSpecifiedColor
+	 */
+	bool IsColorSpecified( ) const
+	{
+		return (ColorUseRule == ESlateColorStylingMode::UseColor_Specified) || (ESlateColorStylingMode::UseColor_Specified_Link);
+	}
+
+public:
+
+	/**
+	 * @returns an FSlateColor that is the widget's foreground. */
+	static FSlateColor UseForeground( )
+	{
+		return FSlateColor( ESlateColorStylingMode::UseColor_Foreground );
+	}
+
+	/** @returns an FSlateColor that is the subdued version of the widget's foreground. */
+	static FSlateColor UseSubduedForeground( )
+	{
+		return FSlateColor( ESlateColorStylingMode::UseColor_Foreground_Subdued );
+	}
+
+protected:
+	
+	// Private constructor to prevent construction of invalid FSlateColors
+	FSlateColor( ESlateColorStylingMode::Type InColorUseRule )
+		: ColorUseRule(InColorUseRule)
+		, LinkedSpecifiedColor()
+	{ }
+
+protected:
+
+	// The current specified color; only meaningful when ColorToUse == UseColor_Specified.
+	UPROPERTY(EditAnywhere, Category=Color)
+	FLinearColor SpecifiedColor;
+
+	// The rule for which color to pick.
+	UPROPERTY(EditAnywhere, Category=Color)
+	TEnumAsByte<ESlateColorStylingMode::Type> ColorUseRule;
+
+private:
+
+	// A shared pointer to the linked color value, if any.
+	TSharedPtr<FLinearColor> LinkedSpecifiedColor;
+};
diff --git a/Engine/Source/Runtime/SlateCore/Public/Styling/SlateStyle.h b/Engine/Source/Runtime/SlateCore/Public/Styling/SlateStyle.h
new file mode 100644
--- /dev/null
+++ b/Engine/Source/Runtime/SlateCore/Public/Styling/SlateStyle.h
@@ -0,0 +1,246 @@
+// Copyright 1998-2014 Epic Games, Inc. All Rights Reserved.
+
+#pragma once
+
+
+/**
+ * A slate style chunk that contains a collection of named properties that guide the appearance of Slate.
+ * At the moment, basically FEditorStyle.
+ */
+class SLATECORE_API FSlateStyleSet : public ISlateStyle
+{
+public:
+	/**
+	 * Construct a style chunk.
+	 * @param InStyleSetName The name used to identity this style set
+	 */
+	FSlateStyleSet(const FName& InStyleSetName);
+	
+	/** Destructor. */
+	virtual ~FSlateStyleSet();
+
+	virtual const FName& GetStyleSetName() const OVERRIDE;
+	virtual void GetResources(TArray< const FSlateBrush* >& OutResources) const OVERRIDE;
+	virtual void SetContentRoot(const FString& InContentRootDir);
+	virtual FString RootToContentDir(const ANSICHAR* RelativePath, const TCHAR* Extension);
+	virtual FString RootToContentDir(const WIDECHAR* RelativePath, const TCHAR* Extension);
+	virtual FString RootToContentDir(const FString& RelativePath, const TCHAR* Extension);
+	virtual FString RootToContentDir(const ANSICHAR* RelativePath);
+	virtual FString RootToContentDir(const WIDECHAR* RelativePath);
+	virtual FString RootToContentDir(const FString& RelativePath);
+
+	virtual void SetCoreContentRoot(const FString& InCoreContentRootDir);
+
+	virtual FString RootToCoreContentDir(const ANSICHAR* RelativePath, const TCHAR* Extension);
+	virtual FString RootToCoreContentDir(const WIDECHAR* RelativePath, const TCHAR* Extension);
+	virtual FString RootToCoreContentDir(const FString& RelativePath, const TCHAR* Extension);
+	virtual FString RootToCoreContentDir(const ANSICHAR* RelativePath);
+	virtual FString RootToCoreContentDir(const WIDECHAR* RelativePath);
+	virtual FString RootToCoreContentDir(const FString& RelativePath);
+
+	virtual float GetFloat(const FName PropertyName, const ANSICHAR* Specifier = nullptr) const OVERRIDE;
+	virtual FVector2D GetVector(const FName PropertyName, const ANSICHAR* Specifier = nullptr) const OVERRIDE;
+	virtual const FLinearColor& GetColor(const FName PropertyName, const ANSICHAR* Specifier = nullptr) const OVERRIDE;
+	virtual const FSlateColor GetSlateColor(const FName PropertyName, const ANSICHAR* Specifier = nullptr) const OVERRIDE;
+	virtual const FMargin& GetMargin(const FName PropertyName, const ANSICHAR* Specifier = nullptr) const OVERRIDE;
+
+	virtual const FSlateBrush* GetBrush(const FName PropertyName, const ANSICHAR* Specifier = nullptr) const OVERRIDE;
+	virtual const FSlateBrush* GetOptionalBrush(const FName PropertyName, const ANSICHAR* Specifier = nullptr, const FSlateBrush* const DefaultBrush = FStyleDefaults::GetNoBrush()) const OVERRIDE;
+
+	virtual const TSharedPtr< FSlateDynamicImageBrush > GetDynamicImageBrush(const FName BrushTemplate, const FName TextureName, const ANSICHAR* Specifier = nullptr) OVERRIDE;
+	virtual const TSharedPtr< FSlateDynamicImageBrush > GetDynamicImageBrush(const FName BrushTemplate, const ANSICHAR* Specifier, UTexture2D* TextureResource, const FName TextureName) OVERRIDE;
+	virtual const TSharedPtr< FSlateDynamicImageBrush > GetDynamicImageBrush(const FName BrushTemplate, UTexture2D* TextureResource, const FName TextureName) OVERRIDE;
+
+	virtual FSlateBrush* GetDefaultBrush() const OVERRIDE;
+
+	virtual const FSlateSound& GetSound(const FName PropertyName, const ANSICHAR* Specifier = nullptr) const OVERRIDE;
+	virtual FSlateFontInfo GetFontStyle(const FName PropertyName, const ANSICHAR* Specifier = nullptr) const OVERRIDE;
+
+public:
+
+	template< typename DefinitionType >            
+	FORCENOINLINE void Set( const FName PropertyName, const DefinitionType& InStyleDefintion )
+	{
+		const TSharedRef< struct FSlateWidgetStyle >* DefinitionPtr = WidgetStyleValues.Find( PropertyName );
+
+		if ( DefinitionPtr == nullptr )
+		{
+			WidgetStyleValues.Add( PropertyName, MakeShareable( new DefinitionType( InStyleDefintion ) ) );
+		}
+		else
+		{
+			WidgetStyleValues.Add( PropertyName, MakeShareable( new DefinitionType( InStyleDefintion ) ) );
+		}
+	}
+
+	/**
+	 * Set float properties
+	 * @param PropertyName - Name of the property to set.
+	 * @param InFloat - The value to set.
+	 */
+	FORCENOINLINE void Set( const FName PropertyName, const float InFloat )
+	{
+		FloatValues.Add( PropertyName, InFloat );
+	}
+
+	/**
+	 * Add a FVector2D property to this style's collection.
+	 * @param PropertyName - Name of the property to set.
+	 * @param InVector - The value to set.
+	 */
+	FORCENOINLINE void Set( const FName PropertyName, const FVector2D InVector )
+	{
+		Vector2DValues.Add( PropertyName, InVector );
+	}
+
+	/**
+	 * Set FLinearColor property.
+	 * @param PropertyName - Name of the property to set.
+	 * @param InColor - The value to set.
+	 */
+	FORCENOINLINE void Set( const FName PropertyName, const FLinearColor& InColor )
+	{
+		ColorValues.Add( PropertyName, InColor );
+	}
+
+	FORCENOINLINE void Set( const FName PropertyName, const FColor& InColor )
+	{
+		ColorValues.Add( PropertyName, InColor );
+	}
+
+	/**
+	 * Add a FSlateLinearColor property to this style's collection.
+	 * @param PropertyName - Name of the property to add.
+	 * @param InColor - The value to add.
+	 */
+	FORCENOINLINE void Set( const FName PropertyName, const FSlateColor& InColor )
+	{
+		SlateColorValues.Add( PropertyName, InColor );
+	}
+
+	/**
+	 * Add a FMargin property to this style's collection.
+	 * @param PropertyName - Name of the property to add.
+	 * @param InMargin - The value to add.
+	 */
+	FORCENOINLINE void Set( const FName PropertyName, const FMargin& InMargin )
+	{
+		MarginValues.Add( PropertyName, InMargin );
+	}
+
+	/**
+	 * Add a FSlateBrush property to this style's collection
+	 * @param PropertyName - Name of the property to set.
+	 * @param InBrush - The brush to set.
+	 */
+	FORCENOINLINE void Set( const FName PropertyName, FSlateBrush* InBrush )
+	{
+		BrushResources.Add( PropertyName, InBrush );
+	}
+
+	FORCENOINLINE void Set( const FName PropertyName, FSlateNoResource* InBrush )
+	{
+		BrushResources.Add( PropertyName, InBrush );
+	}
+
+	FORCENOINLINE void Set( const FName PropertyName, FSlateBoxBrush* InBrush )
+	{
+		BrushResources.Add( PropertyName, InBrush );
+	}
+
+	FORCENOINLINE void Set( const FName PropertyName, FSlateBorderBrush* InBrush )
+	{
+		BrushResources.Add( PropertyName, InBrush );
+	}
+
+	FORCENOINLINE void Set( const FName PropertyName, FSlateImageBrush* InBrush )
+	{
+		BrushResources.Add( PropertyName, InBrush );
+	}
+
+	FORCENOINLINE void Set( const FName PropertyName, FSlateDynamicImageBrush* InBrush )
+	{
+		BrushResources.Add( PropertyName, InBrush );
+	}
+
+	FORCENOINLINE void Set( const FName PropertyName, FSlateColorBrush* InBrush )
+	{
+		BrushResources.Add( PropertyName, InBrush );
+	}
+
+	/**
+	 * Set FSlateSound properties
+	 *
+	 * @param PropertyName  Name of the property to set
+	 * @param InSound		Sound to set
+	 */
+	FORCENOINLINE void Set( FName PropertyName, const FSlateSound& InSound )
+	{
+		Sounds.Add( PropertyName, InSound );
+	}
+		
+	/**
+	 * Set FSlateFontInfo properties
+	 *
+	 * @param PropertyName  Name of the property to set
+	 * @param InFontStyle   The value to set
+	 */
+	FORCENOINLINE void Set( FName PropertyName, const FSlateFontInfo& InFontInfo )
+	{
+		FontInfoResources.Add( PropertyName, InFontInfo );
+	}
+
+protected:
+
+	virtual const FSlateWidgetStyle* GetWidgetStyleInternal(const FName DesiredTypeName, const FName StyleName) const OVERRIDE;
+
+	virtual void Log(ISlateStyle::EStyleMessageSeverity Severity, const FText& Message) const OVERRIDE;
+
+	virtual void LogUnusedBrushResources();
+
+protected:
+
+	bool IsBrushFromFile(const FString& FilePath, const FSlateBrush* Brush);
+
+	/** The name used to identity this style set */
+	FName StyleSetName;
+
+	/** This dir is Engine/Editor/Slate folder **/
+	FString ContentRootDir;
+
+	/** This dir is Engine/Slate folder to share the items **/
+	FString CoreContentRootDir;
+
+	TMap< FName, TSharedRef< struct FSlateWidgetStyle > > WidgetStyleValues;
+
+	/** float property storage. */
+	TMap< FName, float > FloatValues;
+
+	/** FVector2D property storage. */
+	TMap< FName, FVector2D > Vector2DValues;
+
+	/** Color property storage. */
+	TMap< FName, FLinearColor > ColorValues;
+	
+	/** FSlateColor property storage. */
+	TMap< FName, FSlateColor > SlateColorValues;
+
+	/** FMargin property storage. */
+	TMap< FName, FMargin > MarginValues;
+
+	/* FSlateBrush property storage */
+	FSlateBrush* DefaultBrush;
+	TMap< FName, FSlateBrush* > BrushResources;
+
+	/** SlateSound property storage */
+	TMap< FName, FSlateSound > Sounds;
+	
+	/** FSlateFontInfo property storage. */
+	TMap< FName, FSlateFontInfo > FontInfoResources;
+
+	/** A list of dynamic brushes */
+	TMap< FName, TWeakPtr< FSlateDynamicImageBrush > > DynamicBrushes;
+
+	/** A set of resources that were requested, but not found. */
+	mutable TSet< FName > MissingResources;
+};
diff --git a/Engine/Source/Runtime/SlateCore/Public/Styling/SlateStyleRegistry.h b/Engine/Source/Runtime/SlateCore/Public/Styling/SlateStyleRegistry.h
new file mode 100644
--- /dev/null
+++ b/Engine/Source/Runtime/SlateCore/Public/Styling/SlateStyleRegistry.h
@@ -0,0 +1,53 @@
+// Copyright 1998-2014 Epic Games, Inc. All Rights Reserved.
+
+/*=============================================================================
+	SlateStyleRegistry.h: Declares the FSlateStyleRegistry class.
+=============================================================================*/
+
+#pragma once
+
+
+/**
+* A central repository that can be used to track and manage chunks of slate style data.
+*/
+class SLATECORE_API FSlateStyleRegistry
+{
+public:
+	/**
+	 * Add a slate style to the repository.
+	 * @param ISlateStyle - interface to a user-definable style to add.
+	 */
+	static void RegisterSlateStyle( const ISlateStyle& InSlateStyle );
+
+	/**
+	 * Remove a slate style from the repository.
+	 *
+	 * @param ISlateStyle Interface to the style to remove.
+	 */
+	static void UnRegisterSlateStyle( const ISlateStyle& InSlateStyle );
+
+	/**
+	 * Removes a slate style from the repository (by name).
+	 *
+	 * @param StyleSetName The name of the style set to remove.
+	 */
+	static void UnRegisterSlateStyle( const FName StyleSetName );
+
+	/**
+	 * Find a slate style in the repository.
+	 * @param InSlateStyleName - The name of the slate style to find.
+	 * @return The slate style, or null if it couldn't be found
+	 */
+	static const ISlateStyle* FindSlateStyle( const FName& InSlateStyleName );
+
+	/**
+	 * Populate an array of slate brushes with all of the resources used by the registered styles.
+	 * @param OutResources - array of slate brushes to populate.
+	 */
+	static void GetAllResources( TArray<const FSlateBrush*>& OutResources );
+
+private:
+
+	/** Repository is just a collection of shared style pointers. */
+	static TMap<FName, const ISlateStyle*> SlateStyleRepository;
+};
diff --git a/Engine/Source/Runtime/SlateCore/Public/Styling/SlateWidgetStyle.h b/Engine/Source/Runtime/SlateCore/Public/Styling/SlateWidgetStyle.h
new file mode 100644
--- /dev/null
+++ b/Engine/Source/Runtime/SlateCore/Public/Styling/SlateWidgetStyle.h
@@ -0,0 +1,54 @@
+// Copyright 1998-2014 Epic Games, Inc. All Rights Reserved.
+
+/*=============================================================================
+	SlateWidgetStyle.h: Declares the FSlateWidgetStyle structure.
+=============================================================================*/
+
+#pragma once
+
+#include "SlateWidgetStyle.generated.h"
+
+
+/**
+ * Base structure for widget styles.
+ */
+USTRUCT()
+struct SLATECORE_API FSlateWidgetStyle
+{
+	GENERATED_USTRUCT_BODY()
+
+public:
+
+	/**
+	 * Default constructor.
+	 */
+	FSlateWidgetStyle( );
+
+	/**
+	 * Virtual Destructor.
+	 */
+	virtual ~FSlateWidgetStyle() { }
+
+public:
+
+	/**
+	 * Gets the brush resources associated with this style.
+	 *
+	 * This method must be implemented by inherited structures.
+	 *
+	 * @param OutBrushes The brush resources.
+	 */
+	virtual void GetResources( TArray<const FSlateBrush*> & OutBrushes ) const { }
+
+	/**
+	 * Gets the name of this style.
+	 *
+	 * This method must be implemented by inherited structures.
+	 *
+	 * @return Widget style name.
+	 */
+	virtual const FName GetTypeName() const
+	{
+		return FName( NAME_None );
+	} 
+};
diff --git a/Engine/Source/Runtime/SlateCore/Public/Styling/SlateWidgetStyleAsset.h b/Engine/Source/Runtime/SlateCore/Public/Styling/SlateWidgetStyleAsset.h
new file mode 100644
--- /dev/null
+++ b/Engine/Source/Runtime/SlateCore/Public/Styling/SlateWidgetStyleAsset.h
@@ -0,0 +1,73 @@
+// Copyright 1998-2014 Epic Games, Inc. All Rights Reserved.
+#pragma once
+
+#include "SlateWidgetStyleContainerBase.h"
+#include "SlateWidgetStyleAsset.generated.h"
+
+/**
+ * Just a wrapper for the struct with real data in it.
+ */
+UCLASS(hidecategories=Object)
+class SLATECORE_API USlateWidgetStyleAsset : public UObject
+{
+	GENERATED_UCLASS_BODY()
+		  
+public:  
+	/**  */
+	UPROPERTY(Category=Appearance, EditAnywhere, EditInline)
+	USlateWidgetStyleContainerBase* CustomStyle;
+
+	template< class WidgetStyleType >            
+	const WidgetStyleType* GetStyle() const 
+	{
+		return static_cast< const WidgetStyleType* >( GetStyle( WidgetStyleType::TypeName ) );
+	}
+
+	template< class WidgetStyleType >            
+	const WidgetStyleType* GetStyleChecked() const 
+	{
+		return static_cast< const WidgetStyleType* >( GetStyleChecked( WidgetStyleType::TypeName ) );
+	}
+
+	const FSlateWidgetStyle* GetStyle( const FName DesiredTypeName ) const 
+	{
+		if ( CustomStyle == nullptr )
+		{
+			return nullptr;
+		}
+
+		const FSlateWidgetStyle* Style = CustomStyle->GetStyle();
+
+		if ( (Style == nullptr) || (Style->GetTypeName() != DesiredTypeName) )
+		{
+			return nullptr;
+		}
+
+		return Style;
+	}
+
+	const FSlateWidgetStyle* GetStyleChecked( const FName DesiredTypeName ) const 
+	{
+		if ( CustomStyle == nullptr )
+		{
+			UE_LOG( LogSlateStyle, Error, TEXT("USlateWidgetStyleAsset::GetStyle : No custom style set for '%s'."), *GetPathName() );
+			return nullptr;
+		}
+
+		const FSlateWidgetStyle* Style = CustomStyle->GetStyle();
+
+		if ( Style == nullptr )
+		{
+			UE_LOG( LogSlateStyle, Error, TEXT("USlateWidgetStyleAsset::GetStyle : No style found in custom style set for '%s'."), *GetPathName() );
+			return nullptr;
+		}
+
+		if ( Style->GetTypeName() != DesiredTypeName )
+		{
+			UE_LOG( LogSlateStyle, Error, TEXT("USlateWidgetStyleAsset::GetStyle : The custom style is not of the desired type. Desired: '%s', Actual: '%s'"), *DesiredTypeName.ToString(), *Style->GetTypeName().ToString() );
+			return nullptr;
+		}
+
+		return Style;
+	}
+};
diff --git a/Engine/Source/Runtime/SlateCore/Public/Styling/SlateWidgetStyleContainerBase.h b/Engine/Source/Runtime/SlateCore/Public/Styling/SlateWidgetStyleContainerBase.h
new file mode 100644
--- /dev/null
+++ b/Engine/Source/Runtime/SlateCore/Public/Styling/SlateWidgetStyleContainerBase.h
@@ -0,0 +1,21 @@
+// Copyright 1998-2014 Epic Games, Inc. All Rights Reserved.
+
+#pragma once
+
+#include "SlateWidgetStyleContainerInterface.h"
+#include "SlateWidgetStyleContainerBase.generated.h"
+
+SLATECORE_API DECLARE_LOG_CATEGORY_EXTERN(LogSlateStyle, Log, All);
+
+/**
+ * Just a wrapper for the struct with real data in it.
+ */
+UCLASS(hidecategories=Object)
+class SLATECORE_API USlateWidgetStyleContainerBase : public UObject, public ISlateWidgetStyleContainerInterface
+{
+	GENERATED_UCLASS_BODY()
+
+public:
+
+	virtual const struct FSlateWidgetStyle* const GetStyle() const OVERRIDE;
+};
diff --git a/Engine/Source/Runtime/SlateCore/Public/Styling/SlateWidgetStyleContainerInterface.h b/Engine/Source/Runtime/SlateCore/Public/Styling/SlateWidgetStyleContainerInterface.h
new file mode 100644
--- /dev/null
+++ b/Engine/Source/Runtime/SlateCore/Public/Styling/SlateWidgetStyleContainerInterface.h
@@ -0,0 +1,21 @@
+// Copyright 1998-2014 Epic Games, Inc. All Rights Reserved.
+
+#pragma once
+
+#include "SlateWidgetStyleContainerInterface.generated.h"
+
+
+UINTERFACE()
+class SLATECORE_API USlateWidgetStyleContainerInterface : public UInterface
+{
+	GENERATED_UINTERFACE_BODY()
+};
+
+class ISlateWidgetStyleContainerInterface
+{
+	GENERATED_IINTERFACE_BODY()
+
+public:
+
+	virtual const struct FSlateWidgetStyle* const GetStyle() const = 0;
+};
\ No newline at end of file
diff --git a/Engine/Source/Runtime/SlateCore/Public/Styling/StyleDefaults.h b/Engine/Source/Runtime/SlateCore/Public/Styling/StyleDefaults.h
new file mode 100644
--- /dev/null
+++ b/Engine/Source/Runtime/SlateCore/Public/Styling/StyleDefaults.h
@@ -0,0 +1,75 @@
+// Copyright 1998-2014 Epic Games, Inc. All Rights Reserved.
+#pragma once
+
+class SLATECORE_API FStyleDefaults
+{
+public:
+
+	/**
+	 * Return the static default float value.
+	 * @return - The static default float value.
+	 */
+	static float GetFloat()
+	{
+		return DefaultFloat;
+	}
+	
+	/**
+	 * Get default FVector2D value.
+	 * @return - The default FVector2D.
+	 */
+	const static FVector2D GetVector2D()
+	{
+		return DefaultFVector2D;
+	}
+	
+	/**
+	 * Get default FLinearColor.
+	 * @return - The default FLinearColor.
+	 */
+	static const FLinearColor& GetColor()
+	{
+		return DefaultColor;
+	}
+	
+	/**
+	 * Get default FMargin.
+	 * @return - The default FMargin value.
+	 */
+	static const FMargin& GetMargin()
+	{
+		return DefaultMargin;
+	}
+
+	/**
+	 * @return - Returns no brush.
+	 */
+	static const FSlateBrush* GetNoBrush()
+	{
+		static const FSlateNoResource NoBrush;
+		return &NoBrush;
+	}
+
+	static const FSlateFontInfo GetFontInfo(uint16 Size = 10)
+	{
+		static const FSlateFontInfo DefaultFontInfo = FSlateFontInfo( FPaths::EngineContentDir() / TEXT( "Slate/Fonts/Roboto-Regular.ttf" ), Size );
+		return DefaultFontInfo;
+	}
+
+	/**
+	 * Return the static default sound value.
+	 * @return - The static default sound value.
+	 */
+	static const FSlateSound& GetSound()
+	{
+		return DefaultSound;
+	}
+
+private:
+
+	static float DefaultFloat;
+	static FVector2D DefaultFVector2D;
+	static FLinearColor DefaultColor;
+	static FMargin DefaultMargin;
+	static FSlateSound DefaultSound;
+};
diff --git a/Engine/Source/Runtime/SlateCore/Public/Textures/SlateIcon.h b/Engine/Source/Runtime/SlateCore/Public/Textures/SlateIcon.h
new file mode 100644
--- /dev/null
+++ b/Engine/Source/Runtime/SlateCore/Public/Textures/SlateIcon.h
@@ -0,0 +1,133 @@
+// Copyright 1998-2014 Epic Games, Inc. All Rights Reserved.
+
+/*=============================================================================
+	SlateIcon.h: Declares the FSlateIcon structure.
+=============================================================================*/
+
+#pragma once
+
+
+/**
+ * Struct used to represent an icon in Slate
+ */
+struct SLATE_API FSlateIcon
+{
+	/**
+	 * Default constructor (empty icon).
+	 */
+	FSlateIcon( );
+
+	/**
+	 * Creates and initializes a new icon from a style set and style name
+	 *
+	 * @param InStyleSetName The name of the style set the icon can be found in.
+	 * @param StyleName The name of the style for the icon (assumes there may also be a ".Small" variant)
+	 */
+	FSlateIcon( const FName& InStyleSetName, const FName& InStyleName );
+
+	/**
+	 * Creates and initializes a new icon from a style set and style name
+	 *
+	 * @param InStyleSetName The name of the style set the icon can be found in.
+	 * @param StyleName The name of the style for the icon
+	 * @param InSmallStyleName The name of the style for the small icon
+	 */
+	FSlateIcon( const FName& InStyleSetName, const FName& InStyleName, const FName& InSmallStyleName );
+
+public:
+
+	/**
+	 * Gets the resolved icon.
+	 *
+	 * @return Icon brush, or FStyleDefaults::GetNoBrush() if the icon wasn't found.
+	 *
+	 * @see GetSmallIcon
+	 */
+	const FSlateBrush* GetIcon( ) const;
+
+	/**
+	 * Gets the resolved small icon.
+	 *
+	 * @return Icon brush, or FStyleDefaults::GetNoBrush() if the icon wasn't found.
+	 *
+	 * @see GetIcon
+	 */
+	const FSlateBrush* GetSmallIcon( ) const;
+
+	/**
+	 * Gets the name of the style for the icon.
+	 *
+	 * @return Style name.
+	 *
+	 * @see GetStyleName
+	 * @see GetStyleSet
+	 * @see GetStyleSetName
+	 */
+	const FName& GetSmallStyleName( ) const
+	{
+		return SmallStyleName;
+	}
+
+	/**
+	 * Gets the name of the style for the icon.
+	 *
+	 * @return Style name.
+	 *
+	 * @see GetSmallStyleName
+	 * @see GetStyleSet
+	 * @see GetStyleSetName
+	 */
+	const FName& GetStyleName( ) const
+	{
+		return StyleName;
+	}
+
+	/**
+	 * Gets the resolved style set.
+	 *
+	 * @return Style set, or nullptr if the style set wasn't found.
+	 *
+	 * @see GetSmallStyleName
+	 * @see GetStyleName
+	 * @see GetStyleSetName
+	 */
+	const class ISlateStyle* GetStyleSet( ) const;
+
+	/**
+	 * Gets the name of the style set the icon can be found in.
+	 *
+	 * @return Style name.
+	 *
+	 * @see GetSmallStyleName
+	 * @see GetStyleName
+	 * @see GetStyleSet
+	 */
+	const FName& GetStyleSetName( ) const
+	{
+		return StyleSetName;
+	}
+
+	/**
+	 * Checks whether the icon is set to something.
+	 *
+	 * @return true if the icon is set, false otherwise.
+	 */
+	const bool IsSet( ) const
+	{
+		return bIsSet;
+	}
+
+private:
+
+	// The name of the style set the icon can be found in.
+	FName StyleSetName;
+
+	// The name of the style for the icon.
+	FName StyleName;
+
+	// The name of the style for the small icon.
+	FName SmallStyleName;
+
+	// Whether this icon has been explicitly set-up.
+	bool bIsSet;
+};
diff --git a/Engine/Source/Runtime/SlateCore/Public/Textures/SlateTextureData.h b/Engine/Source/Runtime/SlateCore/Public/Textures/SlateTextureData.h
new file mode 100644
--- /dev/null
+++ b/Engine/Source/Runtime/SlateCore/Public/Textures/SlateTextureData.h
@@ -0,0 +1,108 @@
+// Copyright 1998-2014 Epic Games, Inc. All Rights Reserved.
+
+/*=============================================================================
+	FSlateTextureData.h: Declares the FFSlateTextureData structure.
+=============================================================================*/
+
+#pragma once
+
+
+/**
+ * Holds texture data for upload to a rendering resource
+ */
+struct SLATECORE_API FSlateTextureData
+{
+	FSlateTextureData( uint32 InWidth = 0, uint32 InHeight = 0, uint32 InStride = 0, const TArray<uint8>& InBytes = TArray<uint8>() )
+		: Bytes(InBytes)
+		, Width(InWidth)
+		, Height(InHeight)
+		, StrideBytes(InStride)
+	{
+		INC_MEMORY_STAT_BY( STAT_SlateTextureDataMemory, Bytes.GetAllocatedSize() );
+	}
+
+	FSlateTextureData( const FSlateTextureData &Other )
+		: Bytes( Other.Bytes )
+		, Width( Other.Width )
+		, Height( Other.Height )
+		, StrideBytes( Other.StrideBytes )
+	{
+		INC_MEMORY_STAT_BY( STAT_SlateTextureDataMemory, Bytes.GetAllocatedSize() );
+	}
+
+	FSlateTextureData& operator=( const FSlateTextureData& Other )
+	{
+		if( this != &Other )
+		{
+			SetRawData( Other.Width, Other.Height, Other.StrideBytes, Other.Bytes );
+		}
+		return *this;
+	}
+
+	~FSlateTextureData()
+	{
+		DEC_MEMORY_STAT_BY( STAT_SlateTextureDataMemory, Bytes.GetAllocatedSize() );
+	}
+
+	void SetRawData( uint32 InWidth, uint32 InHeight, uint32 InStride, const TArray<uint8>& InBytes )
+	{
+		DEC_MEMORY_STAT_BY( STAT_SlateTextureDataMemory, Bytes.GetAllocatedSize() );
+
+		Width = InWidth;
+		Height = InHeight;
+		StrideBytes = InStride;
+		Bytes = InBytes;
+
+		INC_MEMORY_STAT_BY( STAT_SlateTextureDataMemory, Bytes.GetAllocatedSize() );
+	}
+	
+	void Empty()
+	{
+		DEC_MEMORY_STAT_BY( STAT_SlateTextureDataMemory, Bytes.GetAllocatedSize() );
+		Bytes.Empty();
+	}
+
+	uint32 GetWidth() const
+	{
+		return Width;
+	}
+	
+	uint32 GetHeight() const
+	{
+		return Height;
+	}
+
+	uint32 GetStride() const
+	{
+		return StrideBytes;
+	}
+
+	const TArray<uint8>& GetRawBytes() const
+	{
+		return Bytes;
+	}
+
+	/** Accesses the raw bytes of already sized texture data */
+	uint8* GetRawBytesPtr( )
+	{
+		return Bytes.GetTypedData();
+	}
+
+private:
+
+	/** Raw uncompressed texture data */
+	TArray<uint8> Bytes;
+
+	/** Width of the texture */
+	uint32 Width;
+
+	/** Height of the texture */
+	uint32 Height;
+
+	/** The number of bytes of each pixel */
+	uint32 StrideBytes;
+};
+
+
+typedef TSharedPtr<FSlateTextureData, ESPMode::ThreadSafe> FSlateTextureDataPtr;
+typedef TSharedRef<FSlateTextureData, ESPMode::ThreadSafe> FSlateTextureDataRef;
diff --git a/Engine/Source/Runtime/SlateCore/Public/Textures/TextureAtlas.h b/Engine/Source/Runtime/SlateCore/Public/Textures/TextureAtlas.h
new file mode 100644
--- /dev/null
+++ b/Engine/Source/Runtime/SlateCore/Public/Textures/TextureAtlas.h
@@ -0,0 +1,174 @@
+// Copyright 1998-2014 Epic Games, Inc. All Rights Reserved.
+
+#pragma once
+
+/**
+ * Structure holding information about where a texture is located in the atlas
+ * Pointers to left and right children build a tree of texture rectangles so we can easily find optimal slots for textures
+ */
+struct FAtlasedTextureSlot
+{
+	/** Left child slot. If nullptr there is no texture data here */
+	FAtlasedTextureSlot* Left;
+	/** Right child slot. If nullptr there is no texture data here  */
+	FAtlasedTextureSlot* Right;
+	/** The X position of the character in the texture */
+	uint32 X;
+	/** The Y position of the character in the texture */
+	uint32 Y;
+	/** The width of the character */
+	uint32 Width;
+	/** The height of the character */
+	uint32 Height;
+	uint8 Padding;
+	FAtlasedTextureSlot( uint32 InX, uint32 InY, uint32 InWidth, uint32 InHeight, uint8 InPadding )
+		: Left(nullptr)
+		, Right(nullptr)
+		, X(InX)
+		, Y(InY)
+		, Width(InWidth)
+		, Height(InHeight)
+		, Padding(InPadding)
+	
+	{
+
+	}
+};
+
+
+/**
+ * Base class texture atlases in Slate
+ */
+class SLATECORE_API FSlateTextureAtlas
+{
+public:
+	FSlateTextureAtlas( uint32 InWidth, uint32 InHeight, uint32 StrideBytes, uint32 InPadding )
+		: AtlasData()
+		, RootNode( nullptr )
+		, AtlasWidth( InWidth )
+		, AtlasHeight( InHeight )
+		, Stride( StrideBytes )
+		, Padding( InPadding )
+		, bNeedsUpdate( false )
+	{
+		InitAtlasData();
+	}
+
+	virtual ~FSlateTextureAtlas();
+
+	/**
+	 * Clears atlas data
+	 */
+	void Empty();
+
+	/**
+	 * Adds a texture to the atlas
+	 *
+	 * @param TextureWidth	Width of the texture
+	 * @param TextureHeight	Height of the texture
+	 * @param Data			Raw texture data
+	 */
+	const FAtlasedTextureSlot* AddTexture( uint32 TextureWidth, uint32 TextureHeight, const TArray<uint8>& Data );
+
+	/** @return the width of the atlas */
+	uint32 GetWidth() const { return AtlasWidth; }
+	/** @return the height of the atlas */
+	uint32 GetHeight() const { return AtlasHeight; }
+
+	/** Marks the texture as dirty and needing its rendering resources updated */
+	void MarkTextureDirty() { check( IsInGameThread() ); bNeedsUpdate = true; }
+	
+	/**
+	 * Updates the texture used for rendering if needed
+	 */
+	virtual void ConditionalUpdateTexture() = 0;
+	
+protected:
+	/** 
+	 * Finds the optimal slot for a texture in the atlas starting the search from the root
+	 * 
+	 * @param Width The width of the texture we are adding
+	 * @param Height The height of the texture we are adding
+	 */
+	const FAtlasedTextureSlot* FindSlotForTexture( uint32 InWidth, uint32 InHeight );
+
+	/** 
+	 * Destroys everything in the atlas starting at the provided node
+	 *
+	 * @param StartNode
+	 */
+	void DestroyNodes( FAtlasedTextureSlot* StartNode );
+	
+	/**
+	 * Creates enough space for a single texture the width and height of the atlas
+	 */
+	void InitAtlasData();
+
+	struct FCopyRowData
+	{
+		/** Source data to copy */
+		const uint8* SrcData;
+		/** Place to copy data to */
+		uint8* DestData;
+		/** The row number to copy */
+		uint32 SrcRow;
+		/** The row number to copy to */
+		uint32 DestRow;
+		/** The width of a source row */
+		uint32 RowWidth;
+		/** The width of the source texture */
+		uint32 SrcTextureWidth;
+		/** The width of the dest texture */
+		uint32 DestTextureWidth;
+		/** The stride of each pixel */
+		uint32 DataStride;
+		/** Padding in the destination */
+		uint32 Padding;
+		/** Whether to dialate the texture and put the extra data in the padding area */
+		bool bDialateToPadding;
+	};
+
+	/**
+	 * Copies a single row from a source texture to a dest texture
+	 * Optionally dilating the texture by the padding amount
+	 *
+	 * @param CopyRowData	Information for how to copy a row
+	 */
+	void CopyRow( FCopyRowData& CopyRowData );
+
+	/** 
+	 * Copies texture data into the atlas at a given slot
+	 *
+	 * @param SlotToCopyTo	The occupied slot in the atlas where texture data should be copied to
+	 * @param Data			The data to copy into the atlas
+	 */
+	void CopyDataIntoSlot( const FAtlasedTextureSlot* SlotToCopyTo, const TArray<uint8>& Data );
+
+private:
+	/** 
+	 * Finds the optimal slot for a texture in the atlas.  
+	 * Does this by doing a DFS over all existing slots to see if the texture can fit in 
+	 * the empty area next to an existing slot
+	 * 
+	 * @param Start	The start slot to check
+	 * @param Width The width of the texture we are adding
+	 * @param Height The height of the texture we are adding
+	 */
+	const FAtlasedTextureSlot* FindSlotForTexture( FAtlasedTextureSlot& Start, uint32 InWidth, uint32 InHeight );
+
+protected:
+	/** Actual texture data contained in the atlas */
+	TArray<uint8> AtlasData;
+	/** Root node for the tree of data.  */
+	FAtlasedTextureSlot* RootNode;
+	/** Width of the atlas */
+	uint32 AtlasWidth;
+	/** Height of the atlas */
+	uint32 AtlasHeight;
+	/** Stride of the atlas in bytes */
+	uint32 Stride;
+	/** Padding on all sides of each texture in the atlas */
+	uint32 Padding;
+	/** True if this texture needs to have its rendering resources updated */
+	bool bNeedsUpdate;
+};
diff --git a/Engine/Source/Runtime/SlateCore/Public/Textures/TextureManager.h b/Engine/Source/Runtime/SlateCore/Public/Textures/TextureManager.h
new file mode 100644
--- /dev/null
+++ b/Engine/Source/Runtime/SlateCore/Public/Textures/TextureManager.h
@@ -0,0 +1,71 @@
+// Copyright 1998-2014 Epic Games, Inc. All Rights Reserved.
+
+#pragma once
+
+/**
+ * Represents a new texture that has been loaded but no resource created for it
+ */
+struct FNewTextureInfo
+{
+	/** Raw data */
+	FSlateTextureDataPtr TextureData;
+	/** Whether or not the texture should be atlased */
+	bool bShouldAtlas;
+	/** Whether or not the texture is in srgb space */
+	bool bSrgb;
+	FNewTextureInfo()
+		: bShouldAtlas(true)
+		, bSrgb(true)
+	{
+
+	}
+};
+
+struct FCompareFNewTextureInfoByTextureSize
+{
+	FORCEINLINE bool operator()( const FNewTextureInfo& A, const FNewTextureInfo& B ) const
+	{
+		return (B.TextureData->GetWidth()+B.TextureData->GetHeight()) < (A.TextureData->GetWidth()+A.TextureData->GetHeight());
+	}
+};
+
+
+/** 
+ * Base texture manager class used by a Slate renderer to manage texture resources
+ */
+class FSlateShaderResourceManager
+{
+public:
+	FSlateShaderResourceManager() {};
+	virtual ~FSlateShaderResourceManager() 
+	{
+		ClearTextureMap();
+	}
+
+
+	/** 
+	 * Returns a texture associated with the passed in name.  Should return nullptr if not found 
+	 */
+	virtual FSlateShaderResourceProxy* GetTexture( const FSlateBrush& InBrush ) = 0;
+
+protected:
+
+	void ClearTextureMap()
+	{
+		// delete all allocated textures
+		for( TMap<FName,FSlateShaderResourceProxy*>::TIterator It(ResourceMap); It; ++It )
+		{
+			delete It.Value();
+		}
+		ResourceMap.Empty();
+	}
+
+	FString GetResourcePath( const FSlateBrush& InBrush )
+	{
+		// assume the brush name contains the whole path
+		return InBrush.GetResourceName().ToString();
+	}
+
+	/** Mapping of names to texture pointers */
+	TMap<FName,FSlateShaderResourceProxy*> ResourceMap;
+};
diff --git a/Engine/Source/Runtime/SlateCore/Public/Types/SlateConstants.h b/Engine/Source/Runtime/SlateCore/Public/Types/SlateConstants.h
new file mode 100644
--- /dev/null
+++ b/Engine/Source/Runtime/SlateCore/Public/Types/SlateConstants.h
@@ -0,0 +1,17 @@
+// Copyright 1998-2014 Epic Games, Inc. All Rights Reserved.
+
+/*=============================================================================
+	SlateConstants.h: Declares various constants for Slate.
+=============================================================================*/
+
+#pragma once
+
+
+/** How far the user has to drag the mouse before we consider the action dragging rather than a click */
+const float SlateDragStartDistance = 5.0f;
+
+/** How far the user has to drag the mouse before we consider the action panning rather than a click */
+const float SlatePanTriggerDistance = 1.0f;
+
+/** How much to scroll for each click of the mouse wheel (in Slate Screen Units). */
+const float WheelScrollAmount = 32.0f;
diff --git a/Engine/Source/Runtime/SlateCore/Public/Types/SlateEnums.h b/Engine/Source/Runtime/SlateCore/Public/Types/SlateEnums.h
new file mode 100644
--- /dev/null
+++ b/Engine/Source/Runtime/SlateCore/Public/Types/SlateEnums.h
@@ -0,0 +1,139 @@
+// Copyright 1998-2014 Epic Games, Inc. All Rights Reserved.
+
+/*=============================================================================
+	SlateEnums.h: Declares various enumeration types for Slate.
+=============================================================================*/
+
+#pragma once
+
+#include "SlateEnums.generated.h"
+
+
+namespace EButtonClickMethod
+{
+	/**
+	 * Enumerates different methods that a button click can be triggered. Normally, DownAndUp is appropriate.
+	 */
+	enum Type
+	{
+		/**
+		 * User must press the button, then release while over the button to trigger the click.
+		 * This is the most common type of button.
+		 */
+		DownAndUp,
+
+		/**
+		 * Click will be triggered immediately on mouse down, and mouse will not be captured.
+		 */
+		MouseDown,
+
+		/**
+		 * Click will always be triggered when mouse button is released over the button,
+		 * even if the button wasn't pressed down over it.
+		 */
+		MouseUp,
+	};
+}
+
+
+namespace EFocusMoveDirection
+{
+	/**
+	 * Enumerates directions in which the focus can move.
+	 */
+	enum Type
+	{
+		/** Move to next widget. */
+		Next,
+
+		/** Move focus to previous widget. */
+		Previous
+	};
+}
+
+
+/**
+ * Enumerates horizontal alignment options, i.e. for widget slots.
+ */
+UENUM()
+enum EHorizontalAlignment
+{
+	/** Fill the entire width. */
+	HAlign_Fill UMETA(DisplayName="Fill"),
+
+	/** Left-align. */
+	HAlign_Left UMETA(DisplayName="Left"),
+
+	/** Center-align. */
+	HAlign_Center UMETA(DisplayName="Center"),
+
+	/** Right-align. */
+	HAlign_Right UMETA(DisplayName="Right"),	
+};
+
+
+/**
+ * Enumerates vertical alignment options, i.e. for widget slots.
+ */
+UENUM()
+enum EVerticalAlignment
+{
+	/** Fill the entire height. */
+	VAlign_Fill UMETA(DisplayName="Fill"),
+
+	/** Top-align. */
+	VAlign_Top UMETA(DisplayName="Top"),
+
+	/** Center-align. */
+	VAlign_Center UMETA(DisplayName="Center"),
+
+	/** Bottom-align. */
+	VAlign_Bottom UMETA(DisplayName="Bottom"),
+};
+
+
+/**
+ * Enumerates possible placements for pop-up menus.
+ */
+UENUM()
+enum EMenuPlacement
+{
+	/** Place the menu immediately below the anchor */
+	MenuPlacement_BelowAnchor UMETA(DisplayName="Below Anchor"),
+
+	/** Place the menu immediately below the anchor and match is width to the anchor's content */
+	MenuPlacement_ComboBox UMETA(DisplayName="Combo Box"),
+
+	/** Place the menu to the right of the anchor */
+	MenuPlacement_MenuRight UMETA(DisplayName="Menu Right"),
+
+	/** Place the menu immediately above the anchor, not transition effect */
+	MenuPlacement_AboveAnchor UMETA(DisplayName="Above Anchor"),
+};
+
+
+/**
+ * Enumerates widget orientations.
+ */
+UENUM()
+enum EOrientation
+{
+	/** Orient horizontally, i.e. left to right. */
+	Orient_Horizontal UMETA(DisplayName="Horizontal"),
+
+	/** Orient vertically, i.e. top to bottom. */
+	Orient_Vertical UMETA(DisplayName="Vertical"),
+};
+
+
+/**
+ * Enumerates scroll directions.
+ */
+enum EScrollDirection
+{
+	/** Scroll down. */
+	Scroll_Down UMETA(DisplayName="Down"),
+
+	/** Scroll up. */
+	Scroll_Up UMETA(DisplayName="Up"),
+};
diff --git a/Engine/Source/Runtime/SlateCore/Public/Types/SlateStructs.h b/Engine/Source/Runtime/SlateCore/Public/Types/SlateStructs.h
new file mode 100644
--- /dev/null
+++ b/Engine/Source/Runtime/SlateCore/Public/Types/SlateStructs.h
@@ -0,0 +1,164 @@
+// Copyright 1998-2014 Epic Games, Inc. All Rights Reserved.
+
+/*=============================================================================
+	SlateStructs.h: Declares various structured types for Slate.
+=============================================================================*/
+
+#pragma once
+
+
+/**
+ * Structure for optional floating point sizes.
+ */
+struct FOptionalSize 
+{
+	/**
+	 * Creates an unspecified size.
+	 */
+	FOptionalSize( )
+		: Size(Unspecified)
+	{ }
+
+	/**
+	 * Creates a size with the specified value.
+	 *
+	 * @param SpecifiedSize The size to set.
+	 */
+	FOptionalSize( const float SpecifiedSize )
+		: Size(SpecifiedSize)
+	{ }
+
+public:
+
+	/**
+	 * Checks whether the size is set.
+	 *
+	 * @return true if the size is set, false if it is unespecified.
+	 *
+	 * @see Get
+	 */
+	bool IsSet( ) const
+	{
+		return Size != Unspecified;
+	}
+
+	/**
+	 * Gets the value of the size.
+	 *
+	 * Before calling this method, check with IsSet() whether the size is actually specified.
+	 * Unspecified sizes a value of -1.0f will be returned.
+	 *
+	 * @see IsSet
+	 */
+	float Get( ) const
+	{
+		return Size;
+	}
+
+private:
+
+	// constant for unspecified sizes.
+	SLATECORE_API static const float Unspecified;
+
+	// Holds the size, if specified.
+	float Size;
+};
+
+
+/**
+ * Base structure for size parameters.
+ *
+ * Describes a way in which a parent widget allocates available space to its child widgets.
+ *
+ * When SizeRule is SizeRule_Auto, the widget's DesiredSize will be used as the space required.
+ * When SizeRule is SizeRule_AspectRatio, the widget will attempt to maintain the specified aspect ratio.
+ * When SizeRule is SizeRule_Stretch, the available space will be distributed proportionately between
+ * peer Widgets depending on the Value property. Available space is space remaining after all the
+ * peers' SizeRule_Auto requirements have been satisfied.
+ *
+ * FSizeParam cannot be constructed directly - see FStretch, FAuto, and FAspectRatio
+ */
+struct FSizeParam
+{
+	enum ESizeRule
+	{
+		SizeRule_Auto,
+		SizeRule_Stretch,
+		SizeRule_AspectRatio
+	};
+	
+	/** The sizing rule to use. */
+	ESizeRule SizeRule;
+
+	/**
+	 * The actual value this size parameter stores.
+	 *
+	 * This value can be driven by a delegate. It is only used for the Stretch mode.
+	 */
+	TAttribute<float> Value;
+	
+protected:
+
+	/**
+	 * Hidden constructor.
+	 *
+	 * Use FAspectRatio, FAuto, FStretch to instantiate size parameters.
+	 *
+	 * @see FAspectRatio
+	 * @see FAuto
+	 * @see FStretch
+	 */
+	FSizeParam( ESizeRule InTypeOfSize, const TAttribute<float>& InValue )
+		: SizeRule(InTypeOfSize)
+		, Value(InValue)
+	{ }
+};
+
+
+/**
+ * Structure for size parameters with SizeRule = SizeRule_Stretch.
+ *
+ * @see FAspectRatio
+ * @see FAuto
+ */
+struct FStretch
+	: public FSizeParam
+{
+	FStretch( const TAttribute<float>& StretchAmount )
+		: FSizeParam(SizeRule_Stretch, StretchAmount)
+	{ }
+
+	FStretch( )
+		: FSizeParam(SizeRule_Stretch, 1.0f)
+	{ }
+};
+
+
+/**
+ * Structure for size parameters with SizeRule = SizeRule_Auto.
+ *
+ * @see FAspectRatio
+ * @see FStretch
+ */
+struct FAuto
+	: public FSizeParam
+{
+	FAuto()
+		: FSizeParam(SizeRule_Auto, 0.0f)
+	{ }
+};
+
+
+/**
+ * Structure for size parameters with SizeRule == SizeRule_AspectRatio.
+ *
+ * @see FAuto
+ * @see FStretch
+ */
+struct FAspectRatio
+	: public FSizeParam
+{
+	FAspectRatio( )
+		: FSizeParam(SizeRule_AspectRatio, 1.0f)
+	{ }
+};
