Slate: Refactored core Slate implementation into SlateCore module in preparation for UMG.

diff --git a/Engine/Source/Runtime/SlateCore/Public/Widgets/DeclarativeSyntaxSupport.h b/Engine/Source/Runtime/SlateCore/Public/Widgets/DeclarativeSyntaxSupport.h
new file mode 100644
--- /dev/null
+++ b/Engine/Source/Runtime/SlateCore/Public/Widgets/DeclarativeSyntaxSupport.h
@@ -0,0 +1,1170 @@
+// Copyright 1998-2014 Epic Games, Inc. All Rights Reserved.
+
+/*=============================================================================
+	DeclarativeSyntaxSupport.h: Defines various macros and types for Slate's declarative syntax.
+=============================================================================*/
+
+#pragma once
+
+
+/**
+ * Slate widgets are constructed through SNew and SAssignNew.
+ * e.g.
+ *      
+ *     TSharedRef<SButton> MyButton = SNew(SButton);
+ *        or
+ *     TSharedPtr<SButton> MyButton;
+ *     SAssignNew( MyButton, SButton );
+ *
+ * Using SNew and SAssignNew ensures that widgets are populated
+ */
+
+#define SNew( WidgetType, ... ) \
+	MakeTDecl<WidgetType>( #WidgetType, __FILE__, __LINE__, RequiredArgs::MakeRequiredArgs(__VA_ARGS__) ) <<= TYPENAME_OUTSIDE_TEMPLATE WidgetType::FArguments()
+
+
+#define SAssignNew( ExposeAs, WidgetType, ... ) \
+	MakeTDecl<WidgetType>( #WidgetType, __FILE__, __LINE__, RequiredArgs::MakeRequiredArgs(__VA_ARGS__) ) . Expose( ExposeAs ) <<= TYPENAME_OUTSIDE_TEMPLATE WidgetType::FArguments()
+
+
+/**
+ * Widget authors can use SLATE_BEGIN_ARGS and SLATE_END_ARS to add support
+ * for widget construction via SNew and SAssignNew.
+ * e.g.
+ * 
+ *    SLATE_BEGIN_ARGS( SMyWidget )
+ *         , _PreferredWidth( 150.0f )
+ *         , _ForegroundColor( FLinearColor::White )
+ *         {}
+ *
+ *         SLATE_ATTRIBUTE(float, PreferredWidth)
+ *         SLATE_ATTRIBUTE(FSlateColor, ForegroundColor)
+ *    SLATE_END_ARGS()
+ */
+
+#define SLATE_BEGIN_ARGS( WidgetType ) \
+	public: \
+	struct FArguments : public TSlateBaseNamedArgs<WidgetType> \
+	{ \
+		typedef FArguments WidgetArgsType; \
+		FORCENOINLINE FArguments()
+		
+
+#define SLATE_END_ARGS() \
+	};
+
+
+#define INTERNAL_SLATE_DECL_ATTRIBUTE( Var, AttrType, AttrName ) \
+		WidgetArgsType& AttrName( const AttrType& InAttribute ) \
+		{ \
+			Var = InAttribute; \
+			return this->Me(); \
+		} \
+	\
+		/* Bind attribute with delegate to a global function */ \
+		/* NOTE: We use a template here to avoid 'typename' issues when hosting attributes inside templated classes */ \
+		template< typename StaticFuncPtr > \
+		WidgetArgsType& AttrName##_Static( StaticFuncPtr InFunc )	\
+		{ \
+			Var = AttrType::Create( AttrType::FGetter::CreateStatic( InFunc ) ); \
+			return this->Me(); \
+		} \
+		template< typename Var1Type > \
+		WidgetArgsType& AttrName##_Static( typename AttrType::FGetter::template TStaticDelegate_OneVar< Var1Type >::FFuncPtr InFunc, Var1Type Var1 )	\
+		{ \
+			Var = AttrType::Create( AttrType::FGetter::CreateStatic( InFunc, Var1 ) ); \
+			return this->Me(); \
+		} \
+		template< typename Var1Type, typename Var2Type > \
+		WidgetArgsType& AttrName##_Static( typename AttrType::FGetter::template TStaticDelegate_TwoVars< Var1Type, Var2Type >::FFuncPtr InFunc, Var1Type Var1, Var2Type Var2 )	\
+		{ \
+			Var = AttrType::Create( AttrType::FGetter::CreateStatic( InFunc, Var1, Var2 ) ); \
+			return this->Me(); \
+		} \
+		template< typename Var1Type, typename Var2Type, typename Var3Type > \
+		WidgetArgsType& AttrName##_Static( typename AttrType::FGetter::template TStaticDelegate_ThreeVars< Var1Type, Var2Type, Var3Type >::FFuncPtr InFunc, Var1Type Var1, Var2Type Var2, Var3Type Var3 )	\
+		{ \
+			Var = AttrType::Create( AttrType::FGetter::CreateStatic( InFunc, Var1, Var2, Var3 ) ); \
+			return this->Me(); \
+		} \
+		template< typename Var1Type, typename Var2Type, typename Var3Type, typename Var4Type > \
+		WidgetArgsType& AttrName##_Static( typename AttrType::FGetter::template TStaticDelegate_FourVars< Var1Type, Var2Type, Var3Type, Var4Type >::FFuncPtr InFunc, Var1Type Var1, Var2Type Var2, Var3Type Var3, Var4Type Var4 )	\
+		{ \
+			Var = AttrType::Create( AttrType::FGetter::CreateStatic( InFunc, Var1, Var2, Var3, Var4 ) ); \
+			return this->Me(); \
+		} \
+	\
+		/* Bind attribute with delegate to a raw C++ class method */ \
+		template< class UserClass >	\
+		WidgetArgsType& AttrName##_Raw( UserClass* InUserObject, typename AttrType::FGetter::template TRawMethodDelegate_Const< UserClass >::FMethodPtr InFunc )	\
+		{ \
+			Var = AttrType::Create( AttrType::FGetter::CreateRaw( InUserObject, InFunc ) ); \
+			return this->Me(); \
+		} \
+		template< class UserClass, typename Var1Type >	\
+		WidgetArgsType& AttrName##_Raw( UserClass* InUserObject, typename AttrType::FGetter::template TRawMethodDelegate_OneVar_Const< UserClass, Var1Type  >::FMethodPtr InFunc, Var1Type Var1 )	\
+		{ \
+			Var = AttrType::Create( AttrType::FGetter::CreateRaw( InUserObject, InFunc, Var1 ) ); \
+			return this->Me(); \
+		} \
+		template< class UserClass, typename Var1Type, typename Var2Type >	\
+		WidgetArgsType& AttrName##_Raw( UserClass* InUserObject, typename AttrType::FGetter::template TRawMethodDelegate_TwoVars_Const< UserClass, Var1Type, Var2Type >::FMethodPtr InFunc, Var1Type Var1, Var2Type Var2 )	\
+		{ \
+			Var = AttrType::Create( AttrType::FGetter::CreateRaw( InUserObject, InFunc, Var1, Var2 ) ); \
+			return this->Me(); \
+		} \
+		template< class UserClass, typename Var1Type, typename Var2Type, typename Var3Type >	\
+		WidgetArgsType& AttrName##_Raw( UserClass* InUserObject, typename AttrType::FGetter::template TRawMethodDelegate_ThreeVars_Const< UserClass, Var1Type, Var2Type, Var3Type >::FMethodPtr InFunc, Var1Type Var1, Var2Type Var2, Var3Type Var3 )	\
+		{ \
+			Var = AttrType::Create( AttrType::FGetter::CreateRaw( InUserObject, InFunc, Var1, Var2, Var3 ) ); \
+			return this->Me(); \
+		} \
+		template< class UserClass, typename Var1Type, typename Var2Type, typename Var3Type, typename Var4Type >	\
+		WidgetArgsType& AttrName##_Raw( UserClass* InUserObject, typename AttrType::FGetter::template TRawMethodDelegate_FourVars_Const< UserClass, Var1Type, Var2Type, Var3Type, Var4Type >::FMethodPtr InFunc, Var1Type Var1, Var2Type Var2, Var3Type Var3, Var4Type Var4 )	\
+		{ \
+			Var = AttrType::Create( AttrType::FGetter::CreateRaw( InUserObject, InFunc, Var1, Var2, Var3, Var4 ) ); \
+			return this->Me(); \
+		} \
+	\
+		/* Bind attribute with delegate to a shared pointer-based class method.  Slate mostly uses shared pointers so we use an overload for this type of binding. */ \
+		template< class UserClass >	\
+		WidgetArgsType& AttrName( TSharedRef< UserClass > InUserObjectRef, typename AttrType::FGetter::template TSPMethodDelegate_Const< UserClass >::FMethodPtr InFunc )	\
+		{ \
+			Var = AttrType( InUserObjectRef, InFunc ); \
+			return this->Me(); \
+		} \
+		template< class UserClass, typename Var1Type >	\
+		WidgetArgsType& AttrName( TSharedRef< UserClass > InUserObjectRef, typename AttrType::FGetter::template TSPMethodDelegate_OneVar_Const< UserClass, Var1Type >::FMethodPtr InFunc, Var1Type Var1 )	\
+		{ \
+			Var = AttrType::Create( AttrType::FGetter::CreateSP( InUserObjectRef, InFunc, Var1 ) ); \
+			return this->Me(); \
+		} \
+		template< class UserClass, typename Var1Type, typename Var2Type >	\
+		WidgetArgsType& AttrName( TSharedRef< UserClass > InUserObjectRef, typename AttrType::FGetter::template TSPMethodDelegate_TwoVars_Const< UserClass, Var1Type, Var2Type >::FMethodPtr InFunc, Var1Type Var1, Var2Type Var2 )	\
+		{ \
+			Var = AttrType::Create( AttrType::FGetter::CreateSP( InUserObjectRef, InFunc, Var1, Var2 ) ); \
+			return this->Me(); \
+		} \
+		template< class UserClass, typename Var1Type, typename Var2Type, typename Var3Type >	\
+		WidgetArgsType& AttrName( TSharedRef< UserClass > InUserObjectRef, typename AttrType::FGetter::template TSPMethodDelegate_ThreeVars_Const< UserClass, Var1Type, Var2Type, Var3Type >::FMethodPtr InFunc, Var1Type Var1, Var2Type Var2, Var3Type Var3 )	\
+		{ \
+			Var = AttrType::Create( AttrType::FGetter::CreateSP( InUserObjectRef, InFunc, Var1, Var2, Var3 ) ); \
+			return this->Me(); \
+		} \
+		template< class UserClass, typename Var1Type, typename Var2Type, typename Var3Type, typename Var4Type >	\
+		WidgetArgsType& AttrName( TSharedRef< UserClass > InUserObjectRef, typename AttrType::FGetter::template TSPMethodDelegate_FourVars_Const< UserClass, Var1Type, Var2Type, Var3Type, Var4Type >::FMethodPtr InFunc, Var1Type Var1, Var2Type Var2, Var3Type Var3, Var4Type Var4 )	\
+		{ \
+			Var = AttrType::Create( AttrType::FGetter::CreateSP( InUserObjectRef, InFunc, Var1, Var2, Var3, Var4 ) ); \
+			return this->Me(); \
+		} \
+	\
+		/* Bind attribute with delegate to a shared pointer-based class method.  Slate mostly uses shared pointers so we use an overload for this type of binding. */ \
+		template< class UserClass >	\
+		WidgetArgsType& AttrName( UserClass* InUserObject, typename AttrType::FGetter::template TSPMethodDelegate_Const< UserClass >::FMethodPtr InFunc )	\
+		{ \
+			Var = AttrType::Create( AttrType::FGetter::CreateSP( InUserObject, InFunc ) ); \
+			return this->Me(); \
+		} \
+		template< class UserClass, typename Var1Type >	\
+		WidgetArgsType& AttrName( UserClass* InUserObject, typename AttrType::FGetter::template TSPMethodDelegate_OneVar_Const< UserClass, Var1Type >::FMethodPtr InFunc, Var1Type Var1 )	\
+		{ \
+			Var = AttrType::Create( AttrType::FGetter::CreateSP( InUserObject, InFunc, Var1 ) ); \
+			return this->Me(); \
+		} \
+		template< class UserClass, typename Var1Type, typename Var2Type >	\
+		WidgetArgsType& AttrName( UserClass* InUserObject, typename AttrType::FGetter::template TSPMethodDelegate_TwoVars_Const< UserClass, Var1Type, Var2Type >::FMethodPtr InFunc, Var1Type Var1, Var2Type Var2 )	\
+		{ \
+			Var = AttrType::Create( AttrType::FGetter::CreateSP( InUserObject, InFunc, Var1, Var2 ) ); \
+			return this->Me(); \
+		} \
+		template< class UserClass, typename Var1Type, typename Var2Type, typename Var3Type >	\
+		WidgetArgsType& AttrName( UserClass* InUserObject, typename AttrType::FGetter::template TSPMethodDelegate_ThreeVars_Const< UserClass, Var1Type, Var2Type, Var3Type >::FMethodPtr InFunc, Var1Type Var1, Var2Type Var2, Var3Type Var3 )	\
+		{ \
+			Var = AttrType::Create( AttrType::FGetter::CreateSP( InUserObject, InFunc, Var1, Var2, Var3 ) ); \
+			return this->Me(); \
+		} \
+		template< class UserClass, typename Var1Type, typename Var2Type, typename Var3Type, typename Var4Type >	\
+		WidgetArgsType& AttrName( UserClass* InUserObject, typename AttrType::FGetter::template TSPMethodDelegate_FourVars_Const< UserClass, Var1Type, Var2Type, Var3Type, Var4Type >::FMethodPtr InFunc, Var1Type Var1, Var2Type Var2, Var3Type Var3, Var4Type Var4 )	\
+		{ \
+			Var = AttrType::Create( AttrType::FGetter::CreateSP( InUserObject, InFunc, Var1, Var2, Var3, Var4 ) ); \
+			return this->Me(); \
+		} \
+	\
+		/* Bind attribute with delegate to a UObject-based class method */ \
+		template< class UserClass >	\
+		WidgetArgsType& AttrName##_UObject( UserClass* InUserObject, typename AttrType::FGetter::template TUObjectMethodDelegate_Const< UserClass >::FMethodPtr InFunc )	\
+		{ \
+			Var = AttrType::Create( AttrType::FGetter::CreateUObject( InUserObject, InFunc ) ); \
+			return this->Me(); \
+		} \
+		template< class UserClass, typename Var1Type >	\
+		WidgetArgsType& AttrName##_UObject( UserClass* InUserObject, typename AttrType::FGetter::template TUObjectMethodDelegate_OneVar_Const< UserClass, Var1Type >::FMethodPtr InFunc, Var1Type Var1 )	\
+		{ \
+			Var = AttrType::Create( AttrType::FGetter::CreateUObject( InUserObject, InFunc, Var1 ) ); \
+			return this->Me(); \
+		} \
+		template< class UserClass, typename Var1Type, typename Var2Type  >	\
+		WidgetArgsType& AttrName##_UObject( UserClass* InUserObject, typename AttrType::FGetter::template TUObjectMethodDelegate_TwoVars_Const< UserClass, Var1Type, Var2Type >::FMethodPtr InFunc, Var1Type Var1, Var2Type Var2 )	\
+		{ \
+			Var = AttrType::Create( AttrType::FGetter::CreateUObject( InUserObject, InFunc, Var1, Var2 ) ); \
+			return this->Me(); \
+		} \
+		template< class UserClass, typename Var1Type, typename Var2Type, typename Var3Type >	\
+		WidgetArgsType& AttrName##_UObject( UserClass* InUserObject, typename AttrType::FGetter::template TUObjectMethodDelegate_ThreeVars_Const< UserClass, Var1Type, Var2Type, Var3Type >::FMethodPtr InFunc, Var1Type Var1, Var2Type Var2, Var3Type Var3 )	\
+		{ \
+			Var = AttrType::Create( AttrType::FGetter::CreateUObject( InUserObject, InFunc, Var1, Var2, Var3 ) ); \
+			return this->Me(); \
+		} \
+		template< class UserClass, typename Var1Type, typename Var2Type, typename Var3Type, typename Var4Type >	\
+		WidgetArgsType& AttrName##_UObject( UserClass* InUserObject, typename AttrType::FGetter::template TUObjectMethodDelegate_FourVars_Const< UserClass, Var1Type, Var2Type, Var3Type, Var4Type >::FMethodPtr InFunc, Var1Type Var1, Var2Type Var2, Var3Type Var3, Var4Type Var4 )	\
+		{ \
+			Var = AttrType::Create( AttrType::FGetter::CreateUObject( InUserObject, InFunc, Var1, Var2, Var3, Var4 ) ); \
+			return this->Me(); \
+		}
+
+
+/**
+ * Internal macro to declare overloads for setting attributes
+ *
+ * This macro expects 'WidgetArgsType' to be defined in the scope
+ * where it is used.
+ */
+#define INTERNAL_SLATE_DECL_ATTRIBUTE_TEXT_PASSTHROUGH( Var, AttrType, AttrName ) \
+		struct AttrName##_Local\
+		{\
+			static FString PassThroughDelegate( AttrType::FGetter TextDelegate )\
+			{\
+				return TextDelegate.Execute().ToString();\
+			}\
+			\
+			static FString PassThroughAttribute( AttrType TextAttribute )\
+			{\
+				return TextAttribute.Get( FText::GetEmpty() ).ToString();\
+			}\
+		};\
+		\
+		WidgetArgsType& AttrName( const FString& String ) \
+		{ \
+			Var = TAttribute< FString >( String ); \
+			return this->Me(); \
+		} \
+		WidgetArgsType& AttrName( const FText& Text ) \
+		{ \
+			Var = TAttribute< FString >( Text.ToString() ); \
+			return this->Me(); \
+		} \
+		WidgetArgsType& AttrName( const AttrType& InAttribute ) \
+		{ \
+			if ( InAttribute.IsBound() ) {\
+				Var = TAttribute< FString >::Create( TAttribute< FString >::FGetter::CreateStatic( &AttrName##_Local::PassThroughAttribute, InAttribute ) ); \
+			} else {\
+				Var = InAttribute.Get().ToString();\
+			}\
+			return this->Me(); \
+		} \
+	\
+		/* Bind attribute with delegate to a global function */ \
+		/* NOTE: We use a template here to avoid 'typename' issues when hosting attributes inside templated classes */ \
+		WidgetArgsType& AttrName##_Static( AttrType::FGetter::FStaticDelegate::FFuncPtr InFunc )	\
+		{ \
+			Var = TAttribute< FString >::Create( TAttribute< FString >::FGetter::CreateStatic( &AttrName##_Local::PassThroughDelegate, AttrType::FGetter::CreateStatic( InFunc ) ) ); \
+			return this->Me(); \
+		} \
+		template< typename Var1Type > \
+		WidgetArgsType& AttrName##_Static( typename AttrType::FGetter::template TStaticDelegate_OneVar< Var1Type >::FFuncPtr InFunc, Var1Type Var1 )	\
+		{ \
+			Var = TAttribute< FString >::Create( TAttribute< FString >::FGetter::CreateStatic( &AttrName##_Local::PassThroughDelegate, AttrType::FGetter::CreateStatic( InFunc, Var1 ) ) ); \
+			return this->Me(); \
+		} \
+		template< typename Var1Type, typename Var2Type > \
+		WidgetArgsType& AttrName##_Static( typename AttrType::FGetter::template TStaticDelegate_TwoVars< Var1Type, Var2Type >::FFuncPtr InFunc, Var1Type Var1, Var2Type Var2 )	\
+		{ \
+			Var = TAttribute< FString >::Create( TAttribute< FString >::FGetter::CreateStatic( &AttrName##_Local::PassThroughDelegate, AttrType::FGetter::CreateStatic( InFunc, Var1, Var2 ) ) ); \
+			return this->Me(); \
+		} \
+		template< typename Var1Type, typename Var2Type, typename Var3Type > \
+		WidgetArgsType& AttrName##_Static( typename AttrType::FGetter::template TStaticDelegate_ThreeVars< Var1Type, Var2Type, Var3Type >::FFuncPtr InFunc, Var1Type Var1, Var2Type Var2, Var3Type Var3 )	\
+		{ \
+			Var = TAttribute< FString >::Create( TAttribute< FString >::FGetter::CreateStatic( &AttrName##_Local::PassThroughDelegate, AttrType::FGetter::CreateStatic( InFunc, Var1, Var2, Var3 ) ) ); \
+			return this->Me(); \
+		} \
+		template< typename Var1Type, typename Var2Type, typename Var3Type, typename Var4Type > \
+		WidgetArgsType& AttrName##_Static( typename AttrType::FGetter::template TStaticDelegate_FourVars< Var1Type, Var2Type, Var3Type, Var4Type >::FFuncPtr InFunc, Var1Type Var1, Var2Type Var2, Var3Type Var3, Var4Type Var4 )	\
+		{ \
+			Var = TAttribute< FString >::Create( TAttribute< FString >::FGetter::CreateStatic( &AttrName##_Local::PassThroughDelegate, AttrType::FGetter::CreateStatic( InFunc, Var1, Var2, Var3, Var4 ) ) ); \
+			return this->Me(); \
+		} \
+	\
+		/* Bind attribute with delegate to a raw C++ class method */ \
+		template< class UserClass >	\
+		WidgetArgsType& AttrName##_Raw( UserClass* InUserObject, typename AttrType::FGetter::template TRawMethodDelegate_Const< UserClass >::FMethodPtr InFunc )	\
+		{ \
+			Var = TAttribute< FString >::Create( TAttribute< FString >::FGetter::CreateStatic( &AttrName##_Local::PassThroughDelegate, AttrType::FGetter::CreateRaw( InUserObject, InFunc ) ) ); \
+			return this->Me(); \
+		} \
+		template< class UserClass, typename Var1Type >	\
+		WidgetArgsType& AttrName##_Raw( UserClass* InUserObject, typename AttrType::FGetter::template TRawMethodDelegate_OneVar_Const< UserClass, Var1Type  >::FMethodPtr InFunc, Var1Type Var1 )	\
+		{ \
+			Var = TAttribute< FString >::Create( TAttribute< FString >::FGetter::CreateStatic( &AttrName##_Local::PassThroughDelegate, AttrType::FGetter::CreateRaw( InUserObject, InFunc, Var1 ) ) ); \
+			return this->Me(); \
+		} \
+		template< class UserClass, typename Var1Type, typename Var2Type >	\
+		WidgetArgsType& AttrName##_Raw( UserClass* InUserObject, typename AttrType::FGetter::template TRawMethodDelegate_TwoVars_Const< UserClass, Var1Type, Var2Type >::FMethodPtr InFunc, Var1Type Var1, Var2Type Var2 )	\
+		{ \
+			Var = TAttribute< FString >::Create( TAttribute< FString >::FGetter::CreateStatic( &AttrName##_Local::PassThroughDelegate, AttrType::FGetter::CreateRaw( InUserObject, InFunc, Var1, Var2 ) ) ); \
+			return this->Me(); \
+		} \
+		template< class UserClass, typename Var1Type, typename Var2Type, typename Var3Type >	\
+		WidgetArgsType& AttrName##_Raw( UserClass* InUserObject, typename AttrType::FGetter::template TRawMethodDelegate_ThreeVars_Const< UserClass, Var1Type, Var2Type, Var3Type >::FMethodPtr InFunc, Var1Type Var1, Var2Type Var2, Var3Type Var3 )	\
+		{ \
+			Var = TAttribute< FString >::Create( TAttribute< FString >::FGetter::CreateStatic( &AttrName##_Local::PassThroughDelegate, AttrType::FGetter::CreateRaw( InUserObject, InFunc, Var1, Var2, Var3 ) ) ); \
+			return this->Me(); \
+		} \
+		template< class UserClass, typename Var1Type, typename Var2Type, typename Var3Type, typename Var4Type >	\
+		WidgetArgsType& AttrName##_Raw( UserClass* InUserObject, typename AttrType::FGetter::template TRawMethodDelegate_FourVars_Const< UserClass, Var1Type, Var2Type, Var3Type, Var4Type >::FMethodPtr InFunc, Var1Type Var1, Var2Type Var2, Var3Type Var3, Var4Type Var4 )	\
+		{ \
+			Var = TAttribute< FString >::Create( TAttribute< FString >::FGetter::CreateStatic( &AttrName##_Local::PassThroughDelegate, AttrType::FGetter::CreateRaw( InUserObject, InFunc, Var1, Var2, Var3, Var4 ) ) ); \
+			return this->Me(); \
+		} \
+	\
+		/* Bind attribute with delegate to a shared pointer-based class method.  Slate mostly uses shared pointers so we use an overload for this type of binding. */ \
+		template< class UserClass >	\
+		WidgetArgsType& AttrName( TSharedRef< UserClass > InUserObjectRef, typename AttrType::FGetter::template TSPMethodDelegate_Const< UserClass >::FMethodPtr InFunc )	\
+		{ \
+			Var = TAttribute< FString >::Create( TAttribute< FString >::FGetter::CreateStatic( &AttrName##_Local::PassThroughDelegate, AttrType::FGetter::CreateSP( InUserObjectRef, InFunc ) ) ); \
+			return this->Me(); \
+		} \
+		template< class UserClass, typename Var1Type >	\
+		WidgetArgsType& AttrName( TSharedRef< UserClass > InUserObjectRef, typename AttrType::FGetter::template TSPMethodDelegate_OneVar_Const< UserClass, Var1Type >::FMethodPtr InFunc, Var1Type Var1 )	\
+		{ \
+			Var = TAttribute< FString >::Create( TAttribute< FString >::FGetter::CreateStatic( &AttrName##_Local::PassThroughDelegate, AttrType::FGetter::CreateSP( InUserObjectRef, InFunc, Var1 ) ) ); \
+			return this->Me(); \
+		} \
+		template< class UserClass, typename Var1Type, typename Var2Type >	\
+		WidgetArgsType& AttrName( TSharedRef< UserClass > InUserObjectRef, typename AttrType::FGetter::template TSPMethodDelegate_TwoVars_Const< UserClass, Var1Type, Var2Type >::FMethodPtr InFunc, Var1Type Var1, Var2Type Var2 )	\
+		{ \
+			Var = TAttribute< FString >::Create( TAttribute< FString >::FGetter::CreateStatic( &AttrName##_Local::PassThroughDelegate, AttrType::FGetter::CreateSP( InUserObjectRef, InFunc, Var1, Var2 ) ) ); \
+			return this->Me(); \
+		} \
+		template< class UserClass, typename Var1Type, typename Var2Type, typename Var3Type >	\
+		WidgetArgsType& AttrName( TSharedRef< UserClass > InUserObjectRef, typename AttrType::FGetter::template TSPMethodDelegate_ThreeVars_Const< UserClass, Var1Type, Var2Type, Var3Type >::FMethodPtr InFunc, Var1Type Var1, Var2Type Var2, Var3Type Var3 )	\
+		{ \
+			Var = TAttribute< FString >::Create( TAttribute< FString >::FGetter::CreateStatic( &AttrName##_Local::PassThroughDelegate, AttrType::FGetter::CreateSP( InUserObjectRef, InFunc, Var1, Var2, Var3 ) ) ); \
+			return this->Me(); \
+		} \
+		template< class UserClass, typename Var1Type, typename Var2Type, typename Var3Type, typename Var4Type >	\
+		WidgetArgsType& AttrName( TSharedRef< UserClass > InUserObjectRef, typename AttrType::FGetter::template TSPMethodDelegate_FourVars_Const< UserClass, Var1Type, Var2Type, Var3Type, Var4Type >::FMethodPtr InFunc, Var1Type Var1, Var2Type Var2, Var3Type Var3, Var4Type Var4 )	\
+		{ \
+			Var = TAttribute< FString >::Create( TAttribute< FString >::FGetter::CreateStatic( &AttrName##_Local::PassThroughDelegate, AttrType::FGetter::CreateSP( InUserObjectRef, InFunc, Var1, Var2, Var3, Var4 ) ) ); \
+			return this->Me(); \
+		} \
+	\
+		/* Bind attribute with delegate to a shared pointer-based class method.  Slate mostly uses shared pointers so we use an overload for this type of binding. */ \
+		template< class UserClass >	\
+		WidgetArgsType& AttrName( UserClass* InUserObject, typename AttrType::FGetter::template TSPMethodDelegate_Const< UserClass >::FMethodPtr InFunc )	\
+		{ \
+			Var = TAttribute< FString >::Create( TAttribute< FString >::FGetter::CreateStatic( &AttrName##_Local::PassThroughDelegate, AttrType::FGetter::CreateSP( InUserObject, InFunc ) ) ); \
+			return this->Me(); \
+		} \
+		template< class UserClass, typename Var1Type >	\
+		WidgetArgsType& AttrName( UserClass* InUserObject, typename AttrType::FGetter::template TSPMethodDelegate_OneVar_Const< UserClass, Var1Type >::FMethodPtr InFunc, Var1Type Var1 )	\
+		{ \
+			Var = TAttribute< FString >::Create( TAttribute< FString >::FGetter::CreateStatic( &AttrName##_Local::PassThroughDelegate, AttrType::FGetter::CreateSP( InUserObject, InFunc, Var1 ) ) ); \
+			return this->Me(); \
+		} \
+		template< class UserClass, typename Var1Type, typename Var2Type >	\
+		WidgetArgsType& AttrName( UserClass* InUserObject, typename AttrType::FGetter::template TSPMethodDelegate_TwoVars_Const< UserClass, Var1Type, Var2Type >::FMethodPtr InFunc, Var1Type Var1, Var2Type Var2 )	\
+		{ \
+			Var = TAttribute< FString >::Create( TAttribute< FString >::FGetter::CreateStatic( &AttrName##_Local::PassThroughDelegate, AttrType::FGetter::CreateSP( InUserObject, InFunc, Var1, Var2 ) ) ); \
+			return this->Me(); \
+		} \
+		template< class UserClass, typename Var1Type, typename Var2Type, typename Var3Type >	\
+		WidgetArgsType& AttrName( UserClass* InUserObject, typename AttrType::FGetter::template TSPMethodDelegate_ThreeVars_Const< UserClass, Var1Type, Var2Type, Var3Type >::FMethodPtr InFunc, Var1Type Var1, Var2Type Var2, Var3Type Var3 )	\
+		{ \
+			Var = TAttribute< FString >::Create( TAttribute< FString >::FGetter::CreateStatic( &AttrName##_Local::PassThroughDelegate, AttrType::FGetter::CreateSP( InUserObject, InFunc, Var1, Var2, Var3 ) ) ); \
+			return this->Me(); \
+		} \
+		template< class UserClass, typename Var1Type, typename Var2Type, typename Var3Type, typename Var4Type >	\
+		WidgetArgsType& AttrName( UserClass* InUserObject, typename AttrType::FGetter::template TSPMethodDelegate_FourVars_Const< UserClass, Var1Type, Var2Type, Var3Type, Var4Type >::FMethodPtr InFunc, Var1Type Var1, Var2Type Var2, Var3Type Var3, Var4Type Var4 )	\
+		{ \
+			Var = TAttribute< FString >::Create( TAttribute< FString >::FGetter::CreateStatic( &AttrName##_Local::PassThroughDelegate, AttrType::FGetter::CreateSP( InUserObject, InFunc, Var1, Var2, Var3, Var4 ) ) ); \
+			return this->Me(); \
+		} \
+	\
+		/* Bind attribute with delegate to a UObject-based class method */ \
+		template< class UserClass >	\
+		WidgetArgsType& AttrName##_UObject( UserClass* InUserObject, typename AttrType::FGetter::template TUObjectMethodDelegate_Const< UserClass >::FMethodPtr InFunc )	\
+		{ \
+			Var = TAttribute< FString >::Create( TAttribute< FString >::FGetter::CreateStatic( &AttrName##_Local::PassThroughDelegate, AttrType::FGetter::CreateUObject( InUserObject, InFunc ) ) ); \
+			return this->Me(); \
+		} \
+		template< class UserClass, typename Var1Type >	\
+		WidgetArgsType& AttrName##_UObject( UserClass* InUserObject, typename AttrType::FGetter::template TUObjectMethodDelegate_OneVar_Const< UserClass, Var1Type >::FMethodPtr InFunc, Var1Type Var1 )	\
+		{ \
+			Var = TAttribute< FString >::Create( TAttribute< FString >::FGetter::CreateStatic( &AttrName##_Local::PassThroughDelegate, AttrType::FGetter::CreateUObject( InUserObject, InFunc, Var1 ) ) ); \
+			return this->Me(); \
+		} \
+		template< class UserClass, typename Var1Type, typename Var2Type  >	\
+		WidgetArgsType& AttrName##_UObject( UserClass* InUserObject, typename AttrType::FGetter::template TUObjectMethodDelegate_TwoVars_Const< UserClass, Var1Type, Var2Type >::FMethodPtr InFunc, Var1Type Var1, Var2Type Var2 )	\
+		{ \
+			Var = TAttribute< FString >::Create( TAttribute< FString >::FGetter::CreateStatic( &AttrName##_Local::PassThroughDelegate, AttrType::FGetter::CreateUObject( InUserObject, InFunc, Var1, Var2 ) ) ); \
+			return this->Me(); \
+		} \
+		template< class UserClass, typename Var1Type, typename Var2Type, typename Var3Type >	\
+		WidgetArgsType& AttrName##_UObject( UserClass* InUserObject, typename AttrType::FGetter::template TUObjectMethodDelegate_ThreeVars_Const< UserClass, Var1Type, Var2Type, Var3Type >::FMethodPtr InFunc, Var1Type Var1, Var2Type Var2, Var3Type Var3 )	\
+		{ \
+			Var = TAttribute< FString >::Create( TAttribute< FString >::FGetter::CreateStatic( &AttrName##_Local::PassThroughDelegate, AttrType::FGetter::CreateUObject( InUserObject, InFunc, Var1, Var2, Var3 ) ) ); \
+			return this->Me(); \
+		} \
+		template< class UserClass, typename Var1Type, typename Var2Type, typename Var3Type, typename Var4Type >	\
+		WidgetArgsType& AttrName##_UObject( UserClass* InUserObject, typename AttrType::FGetter::template TUObjectMethodDelegate_FourVars_Const< UserClass, Var1Type, Var2Type, Var3Type, Var4Type >::FMethodPtr InFunc, Var1Type Var1, Var2Type Var2, Var3Type Var3, Var4Type Var4 )	\
+		{ \
+			Var = TAttribute< FString >::Create( TAttribute< FString >::FGetter::CreateStatic( &AttrName##_Local::PassThroughDelegate, AttrType::FGetter::CreateUObject( InUserObject, InFunc, Var1, Var2, Var3, Var4 ) ) ); \
+			return this->Me(); \
+		}
+
+
+/**
+ * Use this macro to add a attribute to the declaration of your widget.
+ * An attribute can be a value or a function.
+ */
+#define SLATE_ATTRIBUTE( AttrType, AttrName ) \
+	TAttribute< AttrType > _##AttrName; \
+	INTERNAL_SLATE_DECL_ATTRIBUTE( _##AttrName, TAttribute< AttrType >, AttrName )
+
+
+#define SLATE_TEXT_ATTRIBUTE( AttrName ) \
+	TAttribute< FString > _##AttrName; \
+	INTERNAL_SLATE_DECL_ATTRIBUTE_TEXT_PASSTHROUGH( _##AttrName, TAttribute< FText >, AttrName ) \
+	INTERNAL_SLATE_DECL_ATTRIBUTE( _##AttrName, TAttribute< FString >, AttrName )
+
+
+/**
+ * Use this macro to declar a slate argument.
+ * Arguments differ from attributes in that they can only be values
+ */
+#define SLATE_ARGUMENT( ArgType, ArgName ) \
+		ArgType _##ArgName; \
+		WidgetArgsType& ArgName( ArgType InArg ) \
+		{ \
+			_##ArgName = InArg; \
+			return this->Me(); \
+		}
+
+#define SLATE_TEXT_ARGUMENT( ArgName ) \
+	FString _##ArgName; \
+	WidgetArgsType& ArgName( FString InArg ) \
+	{ \
+		_##ArgName = InArg; \
+		return this->Me(); \
+	}\
+	\
+	WidgetArgsType& ArgName( FText InArg ) \
+	{ \
+		_##ArgName = InArg.ToString(); \
+		return this->Me(); \
+	}
+
+
+/**
+ * Use this macro to declare a slate argument.
+ * Arguments differ from attributes in that they can only be values
+ */
+#define SLATE_STYLE_ARGUMENT( ArgType, ArgName ) \
+		const ArgType* _##ArgName; \
+		WidgetArgsType& ArgName( const ArgType* InArg ) \
+		{ \
+			_##ArgName = InArg; \
+			return this->Me(); \
+		}\
+		\
+		WidgetArgsType& ArgName( const class USlateWidgetStyleAsset* const InSlateStyleAsset ) \
+		{ \
+			_##ArgName = InSlateStyleAsset->GetStyleChecked< ArgType >(); \
+			return this->Me(); \
+		}\
+		\
+		WidgetArgsType& ArgName( const TWeakObjectPtr< const class USlateWidgetStyleAsset >& InSlateStyleAsset ) \
+		{ \
+			_##ArgName = InSlateStyleAsset->GetStyleChecked< ArgType >(); \
+			return this->Me(); \
+		}\
+		\
+		WidgetArgsType& ArgName( const class ISlateStyle* InSlateStyle, const FName& StyleName, const ANSICHAR* Specifier = nullptr ) \
+		{ \
+			check( InSlateStyle != nullptr ); \
+			_##ArgName = &InSlateStyle->GetWidgetStyle< ArgType >( StyleName, Specifier ); \
+			return this->Me(); \
+		}\
+		\
+		WidgetArgsType& ArgName( const class ISlateStyle& InSlateStyle, const FName& StyleName, const ANSICHAR* Specifier = nullptr ) \
+		{ \
+			_##ArgName = &InSlateStyle.GetWidgetStyle< ArgType >( StyleName, Specifier ); \
+			return this->Me(); \
+		}\
+		\
+		WidgetArgsType& ArgName( const TWeakObjectPtr< const class ISlateStyle >& InSlateStyle, const FName& StyleName, const ANSICHAR* Specifier = nullptr ) \
+		{ \
+			_##ArgName = &InSlateStyle->GetWidgetStyle< ArgType >( StyleName, Specifier ); \
+			return this->Me(); \
+		} \
+		\
+		WidgetArgsType& ArgName( const TSharedPtr< const class ISlateStyle >& InSlateStyle, const FName& StyleName, const ANSICHAR* Specifier = nullptr ) \
+		{ \
+			_##ArgName = &InSlateStyle->GetWidgetStyle< ArgType >( StyleName, Specifier ); \
+			return this->Me(); \
+		} 
+
+
+/**
+ * Use this macro between SLATE_BEGIN_ARGS and SLATE_END_ARGS
+ * in order to add support for slots.
+ */
+#define SLATE_SUPPORTS_SLOT( SlotType ) \
+		TArray< SlotType* > Slots; \
+		WidgetArgsType& operator + (SlotType& SlotToAdd) \
+		{ \
+			Slots.Add( &SlotToAdd ); \
+			return *this; \
+		}
+
+
+ /**
+ * Use this macro between SLATE_BEGIN_ARGS and SLATE_END_ARGS
+ * in order to add support for slots that have slate args.
+ */
+#define SLATE_SUPPORTS_SLOT_WITH_ARGS( SlotType ) \
+	TArray< SlotType* > Slots; \
+	WidgetArgsType& operator + (const SlotType::FArguments& ArgumentsForNewSlot) \
+		{ \
+			Slots.Add( new SlotType( ArgumentsForNewSlot ) ); \
+			return *this; \
+		}
+
+
+/** A widget reference that is always a valid pointer; defaults to SNullWidget */
+struct TAlwaysValidWidget
+{
+	TAlwaysValidWidget()
+	: Widget(SNullWidget::NullWidget)
+	{
+	}
+
+	TSharedRef<SWidget> Widget;
+};
+
+
+/**
+ * We want to be able to do:
+ * SNew( ContainerWidget )
+ * .SomeContentArea()
+ * [
+ *   // Child widgets go here
+ * ]
+ *
+ * NamedSlotProperty is a helper that will be returned by SomeContentArea().
+ */
+template<class DeclarationType>
+struct NamedSlotProperty
+{
+	NamedSlotProperty( DeclarationType& InOwnerDeclaration, TAlwaysValidWidget& ContentToSet )
+		: OwnerDeclaration( InOwnerDeclaration )
+		, SlotContent(ContentToSet)
+	{}
+
+	DeclarationType & operator[]( const TSharedRef<SWidget>& InChild )
+	{
+		SlotContent.Widget = InChild;
+		return OwnerDeclaration;
+	}
+
+	DeclarationType & OwnerDeclaration;
+	TAlwaysValidWidget & SlotContent;
+};
+
+
+/**
+ * Use this macro to add support for named slot properties such as Content and Header. See NamedSlotProperty for more details.
+ *
+ * NOTE: If you're using this within a widget class that is templated, then you might have to specify a full name for the declaration.
+ *       For example: SLATE_NAMED_SLOT( typename SSuperWidget<T>::Declaration, Content )
+ */
+#define SLATE_NAMED_SLOT( DeclarationType, SlotName ) \
+		NamedSlotProperty< DeclarationType > SlotName() \
+		{ \
+			return NamedSlotProperty< DeclarationType >( *this, _##SlotName ); \
+		} \
+		TAlwaysValidWidget _##SlotName; \
+
+#define SLATE_DEFAULT_SLOT( DeclarationType, SlotName ) \
+		SLATE_NAMED_SLOT(DeclarationType, SlotName) ; \
+		DeclarationType & operator[]( const TSharedRef<SWidget> InChild ) \
+		{ \
+			_##SlotName.Widget = InChild; \
+			return *this; \
+		}
+
+
+/**
+ * Use this macro to add event handler support to the declarative syntax of your widget.
+ * It is expected that the widget has a delegate called of type 'EventDelegateType' that is
+ * named 'EventName'.
+ */	
+#define SLATE_EVENT( DelegateName, EventName ) \
+		WidgetArgsType& EventName( const DelegateName& InDelegate ) \
+		{ \
+			_##EventName = InDelegate; \
+			return *this; \
+		} \
+		\
+		/* Set event delegate to a global function */ \
+		/* NOTE: We use a template here to avoid 'typename' issues when hosting attributes inside templated classes */ \
+		template< typename StaticFuncPtr > \
+		WidgetArgsType& EventName##_Static( StaticFuncPtr InFunc )	\
+		{ \
+			_##EventName = DelegateName::CreateStatic( InFunc ); \
+			return *this; \
+		} \
+		template< typename StaticFuncPtr, typename Var1Type > \
+		WidgetArgsType& EventName##_Static( StaticFuncPtr InFunc, Var1Type Var1 )	\
+		{ \
+			_##EventName = DelegateName::CreateStatic( InFunc, Var1 ); \
+			return *this; \
+		} \
+		template< typename StaticFuncPtr, typename Var1Type, typename Var2Type > \
+		WidgetArgsType& EventName##_Static( StaticFuncPtr InFunc, Var1Type Var1, Var2Type Var2 )	\
+		{ \
+			_##EventName = DelegateName::CreateStatic( InFunc, Var1, Var2 ); \
+			return *this; \
+		} \
+		template< typename StaticFuncPtr, typename Var1Type, typename Var2Type, typename Var3Type > \
+		WidgetArgsType& EventName##_Static( StaticFuncPtr InFunc, Var1Type Var1, Var2Type Var2, Var3Type Var3 )	\
+		{ \
+			_##EventName = DelegateName::CreateStatic( InFunc, Var1, Var2, Var3 ); \
+			return *this; \
+		} \
+		template< typename StaticFuncPtr, typename Var1Type, typename Var2Type, typename Var3Type, typename Var4Type > \
+		WidgetArgsType& EventName##_Static( StaticFuncPtr InFunc, Var1Type Var1, Var2Type Var2, Var3Type Var3, Var4Type Var4 )	\
+		{ \
+			_##EventName = DelegateName::CreateStatic( InFunc, Var1, Var2, Var3, Var4 ); \
+			return *this; \
+		} \
+		\
+		/* Set event delegate to a raw C++ class method */ \
+		template< class UserClass >	\
+		WidgetArgsType& EventName##_Raw( UserClass* InUserObject, typename DelegateName::template TRawMethodDelegate< UserClass >::FMethodPtr InFunc )	\
+		{ \
+			_##EventName = DelegateName::CreateRaw( InUserObject, InFunc ); \
+			return *this; \
+		} \
+		template< class UserClass >	\
+		WidgetArgsType& EventName##_Raw( UserClass* InUserObject, typename DelegateName::template TRawMethodDelegate_Const< UserClass >::FMethodPtr InFunc )	\
+		{ \
+			_##EventName = DelegateName::CreateRaw( InUserObject, InFunc ); \
+			return *this; \
+		} \
+		template< class UserClass, typename Var1Type >	\
+		WidgetArgsType& EventName##_Raw( UserClass* InUserObject, typename DelegateName::template TRawMethodDelegate_OneVar< UserClass, Var1Type >::FMethodPtr InFunc, Var1Type Var1 )	\
+		{ \
+			_##EventName = DelegateName::CreateRaw( InUserObject, InFunc, Var1 ); \
+			return *this; \
+		} \
+		template< class UserClass, typename Var1Type >	\
+		WidgetArgsType& EventName##_Raw( UserClass* InUserObject, typename DelegateName::template TRawMethodDelegate_OneVar_Const< UserClass, Var1Type >::FMethodPtr InFunc, Var1Type Var1 )	\
+		{ \
+			_##EventName = DelegateName::CreateRaw( InUserObject, InFunc, Var1 ); \
+			return *this; \
+		} \
+		template< class UserClass, typename Var1Type, typename Var2Type >	\
+		WidgetArgsType& EventName##_Raw( UserClass* InUserObject, typename DelegateName::template TRawMethodDelegate_TwoVars< UserClass, Var1Type, Var2Type >::FMethodPtr InFunc, Var1Type Var1, Var2Type Var2 )	\
+		{ \
+			_##EventName = DelegateName::CreateRaw( InUserObject, InFunc, Var1, Var2 ); \
+			return *this; \
+		} \
+		template< class UserClass, typename Var1Type, typename Var2Type >	\
+		WidgetArgsType& EventName##_Raw( UserClass* InUserObject, typename DelegateName::template TRawMethodDelegate_TwoVars_Const< UserClass, Var1Type, Var2Type >::FMethodPtr InFunc, Var1Type Var1, Var2Type Var2 )	\
+		{ \
+			_##EventName = DelegateName::CreateRaw( InUserObject, InFunc, Var1, Var2 ); \
+			return *this; \
+		} \
+		template< class UserClass, typename Var1Type, typename Var2Type, typename Var3Type >	\
+		WidgetArgsType& EventName##_Raw( UserClass* InUserObject, typename DelegateName::template TRawMethodDelegate_ThreeVars< UserClass, Var1Type, Var2Type, Var3Type >::FMethodPtr InFunc, Var1Type Var1, Var2Type Var2, Var3Type Var3 )	\
+		{ \
+			_##EventName = DelegateName::CreateRaw( InUserObject, InFunc, Var1, Var2, Var3 ); \
+			return *this; \
+		} \
+		template< class UserClass, typename Var1Type, typename Var2Type, typename Var3Type >	\
+		WidgetArgsType& EventName##_Raw( UserClass* InUserObject, typename DelegateName::template TRawMethodDelegate_ThreeVars_Const< UserClass, Var1Type, Var2Type, Var3Type >::FMethodPtr InFunc, Var1Type Var1, Var2Type Var2, Var3Type Var3 )	\
+		{ \
+			_##EventName = DelegateName::CreateRaw( InUserObject, InFunc, Var1, Var2, Var3 ); \
+			return *this; \
+		} \
+		template< class UserClass, typename Var1Type, typename Var2Type, typename Var3Type, typename Var4Type >	\
+		WidgetArgsType& EventName##_Raw( UserClass* InUserObject, typename DelegateName::template TRawMethodDelegate_FourVars< UserClass, Var1Type, Var2Type, Var3Type, Var4Type >::FMethodPtr InFunc, Var1Type Var1, Var2Type Var2, Var3Type Var3, Var4Type Var4 )	\
+		{ \
+			_##EventName = DelegateName::CreateRaw( InUserObject, InFunc, Var1, Var2, Var3, Var4 ); \
+			return *this; \
+		} \
+		template< class UserClass, typename Var1Type, typename Var2Type, typename Var3Type, typename Var4Type >	\
+		WidgetArgsType& EventName##_Raw( UserClass* InUserObject, typename DelegateName::template TRawMethodDelegate_FourVars_Const< UserClass, Var1Type, Var2Type, Var3Type, Var4Type >::FMethodPtr InFunc, Var1Type Var1, Var2Type Var2, Var3Type Var3, Var4Type Var4 )	\
+		{ \
+			_##EventName = DelegateName::CreateRaw( InUserObject, InFunc, Var1, Var2, Var3, Var4 ); \
+			return *this; \
+		} \
+		\
+		/* Set event delegate to a shared pointer-based class method.  Slate mostly uses shared pointers so we use an overload for this type of binding. */ \
+		template< class UserClass >	\
+		WidgetArgsType& EventName( TSharedRef< UserClass > InUserObjectRef, typename DelegateName::template TSPMethodDelegate< UserClass >::FMethodPtr InFunc )	\
+		{ \
+			_##EventName = DelegateName::CreateSP( InUserObjectRef, InFunc ); \
+			return *this; \
+		} \
+		template< class UserClass >	\
+		WidgetArgsType& EventName( TSharedRef< UserClass > InUserObjectRef, typename DelegateName::template TSPMethodDelegate_Const< UserClass >::FMethodPtr InFunc )	\
+		{ \
+			_##EventName = DelegateName::CreateSP( InUserObjectRef, InFunc ); \
+			return *this; \
+		} \
+		template< class UserClass, typename Var1Type >	\
+		WidgetArgsType& EventName( TSharedRef< UserClass > InUserObjectRef, typename DelegateName::template TSPMethodDelegate_OneVar< UserClass, Var1Type >::FMethodPtr InFunc, Var1Type Var1 )	\
+		{ \
+			_##EventName = DelegateName::CreateSP( InUserObjectRef, InFunc, Var1 ); \
+			return *this; \
+		} \
+		template< class UserClass, typename Var1Type >	\
+		WidgetArgsType& EventName( TSharedRef< UserClass > InUserObjectRef, typename DelegateName::template TSPMethodDelegate_OneVar_Const< UserClass, Var1Type >::FMethodPtr InFunc, Var1Type Var1 )	\
+		{ \
+			_##EventName = DelegateName::CreateSP( InUserObjectRef, InFunc, Var1 ); \
+			return *this; \
+		} \
+		template< class UserClass, typename Var1Type, typename Var2Type >	\
+		WidgetArgsType& EventName( TSharedRef< UserClass > InUserObjectRef, typename DelegateName::template TSPMethodDelegate_TwoVars< UserClass, Var1Type, Var2Type >::FMethodPtr InFunc, Var1Type Var1, Var2Type Var2 )	\
+		{ \
+			_##EventName = DelegateName::CreateSP( InUserObjectRef, InFunc, Var1, Var2 ); \
+			return *this; \
+		} \
+		template< class UserClass, typename Var1Type, typename Var2Type >	\
+		WidgetArgsType& EventName( TSharedRef< UserClass > InUserObjectRef, typename DelegateName::template TSPMethodDelegate_TwoVars_Const< UserClass, Var1Type, Var2Type >::FMethodPtr InFunc, Var1Type Var1, Var2Type Var2 )	\
+		{ \
+			_##EventName = DelegateName::CreateSP( InUserObjectRef, InFunc, Var1, Var2 ); \
+			return *this; \
+		} \
+		template< class UserClass, typename Var1Type, typename Var2Type, typename Var3Type >	\
+		WidgetArgsType& EventName( TSharedRef< UserClass > InUserObjectRef, typename DelegateName::template TSPMethodDelegate_ThreeVars< UserClass, Var1Type, Var2Type, Var3Type >::FMethodPtr InFunc, Var1Type Var1, Var2Type Var2, Var3Type Var3 )	\
+		{ \
+			_##EventName = DelegateName::CreateSP( InUserObjectRef, InFunc, Var1, Var2, Var3 ); \
+			return *this; \
+		} \
+		template< class UserClass, typename Var1Type, typename Var2Type, typename Var3Type >	\
+		WidgetArgsType& EventName( TSharedRef< UserClass > InUserObjectRef, typename DelegateName::template TSPMethodDelegate_ThreeVars_Const< UserClass, Var1Type, Var2Type, Var3Type >::FMethodPtr InFunc, Var1Type Var1, Var2Type Var2, Var3Type Var3 )	\
+		{ \
+			_##EventName = DelegateName::CreateSP( InUserObjectRef, InFunc, Var1, Var2, Var3 ); \
+			return *this; \
+		} \
+		template< class UserClass, typename Var1Type, typename Var2Type, typename Var3Type, typename Var4Type >	\
+		WidgetArgsType& EventName( TSharedRef< UserClass > InUserObjectRef, typename DelegateName::template TSPMethodDelegate_FourVars< UserClass, Var1Type, Var2Type, Var3Type, Var4Type >::FMethodPtr InFunc, Var1Type Var1, Var2Type Var2, Var3Type Var3, Var4Type Var4 )	\
+		{ \
+			_##EventName = DelegateName::CreateSP( InUserObjectRef, InFunc, Var1, Var2, Var3, Var4 ); \
+			return *this; \
+		} \
+		template< class UserClass, typename Var1Type, typename Var2Type, typename Var3Type, typename Var4Type >	\
+		WidgetArgsType& EventName( TSharedRef< UserClass > InUserObjectRef, typename DelegateName::template TSPMethodDelegate_FourVars_Const< UserClass, Var1Type, Var2Type, Var3Type, Var4Type >::FMethodPtr InFunc, Var1Type Var1, Var2Type Var2, Var3Type Var3, Var4Type Var4 )	\
+		{ \
+			_##EventName = DelegateName::CreateSP( InUserObjectRef, InFunc, Var1, Var2, Var3, Var4 ); \
+			return *this; \
+		} \
+		\
+		/* Set event delegate to a shared pointer-based class method.  Slate mostly uses shared pointers so we use an overload for this type of binding. */ \
+		template< class UserClass >	\
+		WidgetArgsType& EventName( UserClass* InUserObject, typename DelegateName::template TSPMethodDelegate< UserClass >::FMethodPtr InFunc )	\
+		{ \
+			_##EventName = DelegateName::CreateSP( InUserObject, InFunc ); \
+			return *this; \
+		} \
+		template< class UserClass >	\
+		WidgetArgsType& EventName( UserClass* InUserObject, typename DelegateName::template TSPMethodDelegate_Const< UserClass >::FMethodPtr InFunc )	\
+		{ \
+			_##EventName = DelegateName::CreateSP( InUserObject, InFunc ); \
+			return *this; \
+		} \
+		template< class UserClass, typename Var1Type >	\
+		WidgetArgsType& EventName( UserClass* InUserObject, typename DelegateName::template TSPMethodDelegate_OneVar< UserClass, Var1Type >::FMethodPtr InFunc, Var1Type Var1 )	\
+		{ \
+			_##EventName = DelegateName::CreateSP( InUserObject, InFunc, Var1 ); \
+			return *this; \
+		} \
+		template< class UserClass, typename Var1Type >	\
+		WidgetArgsType& EventName( UserClass* InUserObject, typename DelegateName::template TSPMethodDelegate_OneVar_Const< UserClass, Var1Type >::FMethodPtr InFunc, Var1Type Var1 )	\
+		{ \
+			_##EventName = DelegateName::CreateSP( InUserObject, InFunc, Var1 ); \
+			return *this; \
+		} \
+		template< class UserClass, typename Var1Type, typename Var2Type >	\
+		WidgetArgsType& EventName( UserClass* InUserObject, typename DelegateName::template TSPMethodDelegate_TwoVars< UserClass, Var1Type, Var2Type >::FMethodPtr InFunc, Var1Type Var1, Var2Type Var2 )	\
+		{ \
+			_##EventName = DelegateName::CreateSP( InUserObject, InFunc, Var1, Var2 ); \
+			return *this; \
+		} \
+		template< class UserClass, typename Var1Type, typename Var2Type >	\
+		WidgetArgsType& EventName( UserClass* InUserObject, typename DelegateName::template TSPMethodDelegate_TwoVars_Const< UserClass, Var1Type, Var2Type >::FMethodPtr InFunc, Var1Type Var1, Var2Type Var2 )	\
+		{ \
+			_##EventName = DelegateName::CreateSP( InUserObject, InFunc, Var1, Var2 ); \
+			return *this; \
+		} \
+		template< class UserClass, typename Var1Type, typename Var2Type, typename Var3Type >	\
+		WidgetArgsType& EventName( UserClass* InUserObject, typename DelegateName::template TSPMethodDelegate_ThreeVars< UserClass, Var1Type, Var2Type, Var3Type >::FMethodPtr InFunc, Var1Type Var1, Var2Type Var2, Var3Type Var3 )	\
+		{ \
+			_##EventName = DelegateName::CreateSP( InUserObject, InFunc, Var1, Var2, Var3 ); \
+			return *this; \
+		} \
+		template< class UserClass, typename Var1Type, typename Var2Type, typename Var3Type >	\
+		WidgetArgsType& EventName( UserClass* InUserObject, typename DelegateName::template TSPMethodDelegate_ThreeVars_Const< UserClass, Var1Type, Var2Type, Var3Type >::FMethodPtr InFunc, Var1Type Var1, Var2Type Var2, Var3Type Var3 )	\
+		{ \
+			_##EventName = DelegateName::CreateSP( InUserObject, InFunc, Var1, Var2, Var3 ); \
+			return *this; \
+		} \
+		template< class UserClass, typename Var1Type, typename Var2Type, typename Var3Type, typename Var4Type >	\
+		WidgetArgsType& EventName( UserClass* InUserObject, typename DelegateName::template TSPMethodDelegate_FourVars< UserClass, Var1Type, Var2Type, Var3Type, Var4Type >::FMethodPtr InFunc, Var1Type Var1, Var2Type Var2, Var3Type Var3, Var4Type Var4 )	\
+		{ \
+			_##EventName = DelegateName::CreateSP( InUserObject, InFunc, Var1, Var2, Var3, Var4 ); \
+			return *this; \
+		} \
+		template< class UserClass, typename Var1Type, typename Var2Type, typename Var3Type, typename Var4Type >	\
+		WidgetArgsType& EventName( UserClass* InUserObject, typename DelegateName::template TSPMethodDelegate_FourVars_Const< UserClass, Var1Type, Var2Type, Var3Type, Var4Type >::FMethodPtr InFunc, Var1Type Var1, Var2Type Var2, Var3Type Var3, Var4Type Var4 )	\
+		{ \
+			_##EventName = DelegateName::CreateSP( InUserObject, InFunc, Var1, Var2, Var3, Var4 ); \
+			return *this; \
+		} \
+		\
+		/* Set event delegate to a UObject-based class method */ \
+		template< class UserClass >	\
+		WidgetArgsType& EventName##_UObject( UserClass* InUserObject, typename DelegateName::template TUObjectMethodDelegate< UserClass >::FMethodPtr InFunc )	\
+		{ \
+			_##EventName = DelegateName::CreateUObject( InUserObject, InFunc ); \
+			return *this; \
+		} \
+		template< class UserClass >	\
+		WidgetArgsType& EventName##_UObject( UserClass* InUserObject, typename DelegateName::template TUObjectMethodDelegate_Const< UserClass >::FMethodPtr InFunc )	\
+		{ \
+			_##EventName = DelegateName::CreateUObject( InUserObject, InFunc ); \
+			return *this; \
+		} \
+		template< class UserClass, typename Var1Type >	\
+		WidgetArgsType& EventName##_UObject( UserClass* InUserObject, typename DelegateName::template TUObjectMethodDelegate_OneVar< UserClass, Var1Type >::FMethodPtr InFunc, Var1Type Var1 )	\
+		{ \
+			_##EventName = DelegateName::CreateUObject( InUserObject, InFunc, Var1 ); \
+			return *this; \
+		} \
+		template< class UserClass, typename Var1Type >	\
+		WidgetArgsType& EventName##_UObject( UserClass* InUserObject, typename DelegateName::template TUObjectMethodDelegate_OneVar_Const< UserClass, Var1Type >::FMethodPtr InFunc, Var1Type Var1 )	\
+		{ \
+			_##EventName = DelegateName::CreateUObject( InUserObject, InFunc, Var1 ); \
+			return *this; \
+		} \
+		template< class UserClass, typename Var1Type, typename Var2Type >	\
+		WidgetArgsType& EventName##_UObject( UserClass* InUserObject, typename DelegateName::template TUObjectMethodDelegate_TwoVars< UserClass, Var1Type, Var2Type >::FMethodPtr InFunc, Var1Type Var1, Var2Type Var2 )	\
+		{ \
+			_##EventName = DelegateName::CreateUObject( InUserObject, InFunc, Var1, Var2 ); \
+			return *this; \
+		} \
+		template< class UserClass, typename Var1Type, typename Var2Type >	\
+		WidgetArgsType& EventName##_UObject( UserClass* InUserObject, typename DelegateName::template TUObjectMethodDelegate_TwoVars_Const< UserClass, Var1Type, Var2Type >::FMethodPtr InFunc, Var1Type Var1, Var2Type Var2 )	\
+		{ \
+			_##EventName = DelegateName::CreateUObject( InUserObject, InFunc, Var1, Var2 ); \
+			return *this; \
+		} \
+		template< class UserClass, typename Var1Type, typename Var2Type, typename Var3Type >	\
+		WidgetArgsType& EventName##_UObject( UserClass* InUserObject, typename DelegateName::template TUObjectMethodDelegate_ThreeVars< UserClass, Var1Type, Var2Type, Var3Type >::FMethodPtr InFunc, Var1Type Var1, Var2Type Var2, Var3Type Var3 )	\
+		{ \
+			_##EventName = DelegateName::CreateUObject( InUserObject, InFunc, Var1, Var2, Var3 ); \
+			return *this; \
+		} \
+		template< class UserClass, typename Var1Type, typename Var2Type, typename Var3Type >	\
+		WidgetArgsType& EventName##_UObject( UserClass* InUserObject, typename DelegateName::template TUObjectMethodDelegate_ThreeVars_Const< UserClass, Var1Type, Var2Type, Var3Type >::FMethodPtr InFunc, Var1Type Var1, Var2Type Var2, Var3Type Var3 )	\
+		{ \
+			_##EventName = DelegateName::CreateUObject( InUserObject, InFunc, Var1, Var2, Var3 ); \
+			return *this; \
+		} \
+		template< class UserClass, typename Var1Type, typename Var2Type, typename Var3Type, typename Var4Type >	\
+		WidgetArgsType& EventName##_UObject( UserClass* InUserObject, typename DelegateName::template TUObjectMethodDelegate_FourVars< UserClass, Var1Type, Var2Type, Var3Type, Var4Type >::FMethodPtr InFunc, Var1Type Var1, Var2Type Var2, Var3Type Var3, Var4Type Var4 )	\
+		{ \
+			_##EventName = DelegateName::CreateUObject( InUserObject, InFunc, Var1, Var2, Var3, Var4 ); \
+			return *this; \
+		} \
+		template< class UserClass, typename Var1Type, typename Var2Type, typename Var3Type, typename Var4Type >	\
+		WidgetArgsType& EventName##_UObject( UserClass* InUserObject, typename DelegateName::template TUObjectMethodDelegate_FourVars_Const< UserClass, Var1Type, Var2Type, Var3Type, Var4Type >::FMethodPtr InFunc, Var1Type Var1, Var2Type Var2, Var3Type Var3, Var4Type Var4 )	\
+		{ \
+			_##EventName = DelegateName::CreateUObject( InUserObject, InFunc, Var1, Var2, Var3, Var4 ); \
+			return *this; \
+		} \
+		\
+		DelegateName _##EventName; \
+
+
+/** Base class for named arguments. Provides settings necessary for all widgets. */
+template<typename WidgetType>
+struct TSlateBaseNamedArgs
+{
+	typedef typename WidgetType::FArguments WidgetArgsType;
+	
+	TSlateBaseNamedArgs()
+	: _ToolTipText()
+	, _ToolTip()
+	, _Cursor( TOptional<EMouseCursor::Type>() )
+	, _IsEnabled( true )
+	, _Visibility( EVisibility::Visible )
+	{
+	}
+
+	/** Used by the named argument pattern as a safe way to 'return *this' for call-chaining purposes. */
+	WidgetArgsType& Me()
+	{
+		return *(static_cast<WidgetArgsType*>(this));
+	}
+
+	SLATE_TEXT_ATTRIBUTE( ToolTipText )
+	SLATE_ARGUMENT( TSharedPtr<IToolTip>, ToolTip )
+	SLATE_ATTRIBUTE( TOptional<EMouseCursor::Type>, Cursor )
+	SLATE_ATTRIBUTE( bool, IsEnabled )
+	SLATE_ATTRIBUTE( EVisibility, Visibility )
+	SLATE_ARGUMENT( FName, Tag )
+};
+
+
+namespace RequiredArgs
+{
+	struct T0RequiredArgs
+	{
+		T0RequiredArgs()
+		{
+		}
+
+		template<class WidgetType, typename NamedArgsType>
+		void CallConstruct( const TSharedRef<WidgetType>& OnWidget, const NamedArgsType& WithNamedArgs ) const
+		{
+			// YOUR WIDGET MUST IMPLEMENT void Construct( const FArguments& InArgs )
+			OnWidget->Construct( WithNamedArgs );
+		}
+	};
+
+	template<typename Arg0Type>
+	struct T1RequiredArgs
+	{
+		T1RequiredArgs( Arg0Type InArg0 )
+		: Arg0(InArg0)
+		{
+		}
+
+		template<class WidgetType, typename NamedArgsType>
+		void CallConstruct( const TSharedRef<WidgetType>& OnWidget, const NamedArgsType& WithNamedArgs ) const
+		{
+			// YOUR WIDGET MUST IMPLEMENT void Construct( const FArguments& InArgs )
+			OnWidget->Construct( WithNamedArgs, Arg0 );
+		}
+
+		Arg0Type Arg0;
+
+	};
+
+	template<typename Arg0Type, typename Arg1Type>
+	struct T2RequiredArgs
+	{
+		T2RequiredArgs( Arg0Type InArg0, Arg1Type InArg1 )
+		: Arg0(InArg0)
+		, Arg1(InArg1)
+		{
+		}
+
+		template<class WidgetType, typename NamedArgsType>
+		void CallConstruct( const TSharedRef<WidgetType>& OnWidget, const NamedArgsType& WithNamedArgs ) const
+		{
+			// YOUR WIDGET MUST IMPLEMENT void Construct( const FArguments& InArgs )
+			OnWidget->Construct( WithNamedArgs, Arg0, Arg1 );
+		}
+
+		Arg0Type Arg0;
+		Arg1Type Arg1;
+	};
+
+	template<typename Arg0Type, typename Arg1Type, typename Arg2Type>
+	struct T3RequiredArgs
+	{
+		T3RequiredArgs( Arg0Type InArg0, Arg1Type InArg1, Arg2Type InArg2 )
+		: Arg0(InArg0)
+		, Arg1(InArg1)
+		, Arg2(InArg2)
+		{
+		}
+
+		template<class WidgetType, typename NamedArgsType>
+		void CallConstruct( const TSharedRef<WidgetType>& OnWidget, const NamedArgsType& WithNamedArgs ) const
+		{
+			// YOUR WIDGET MUST IMPLEMENT void Construct( const FArguments& InArgs )
+			OnWidget->Construct( WithNamedArgs, Arg0, Arg1, Arg2 );
+		}
+
+		Arg0Type Arg0;
+		Arg1Type Arg1;
+		Arg2Type Arg2;
+	};
+
+	template<typename Arg0Type, typename Arg1Type, typename Arg2Type, typename Arg3Type>
+	struct T4RequiredArgs
+	{
+		T4RequiredArgs( Arg0Type InArg0, Arg1Type InArg1, Arg2Type InArg2, Arg3Type InArg3 )
+		: Arg0(InArg0)
+		, Arg1(InArg1)
+		, Arg2(InArg2)
+		, Arg3(InArg3)
+		{
+		}
+
+		template<class WidgetType, typename NamedArgsType>
+		void CallConstruct( const TSharedRef<WidgetType>& OnWidget, const NamedArgsType& WithNamedArgs ) const
+		{
+			// YOUR WIDGET MUST IMPLEMENT void Construct( const FArguments& InArgs )
+			OnWidget->Construct( WithNamedArgs, Arg0, Arg1, Arg2, Arg3 );
+		}
+
+		Arg0Type Arg0;
+		Arg1Type Arg1;
+		Arg2Type Arg2;
+		Arg3Type Arg3;
+	};
+
+	template<typename Arg0Type, typename Arg1Type, typename Arg2Type, typename Arg3Type, typename Arg4Type>
+	struct T5RequiredArgs
+	{
+		T5RequiredArgs( Arg0Type InArg0, Arg1Type InArg1, Arg2Type InArg2, Arg3Type InArg3, Arg4Type InArg4 )
+			: Arg0(InArg0)
+			, Arg1(InArg1)
+			, Arg2(InArg2)
+			, Arg3(InArg3)
+			, Arg4(InArg4)
+		{
+		}
+
+		template<class WidgetType, typename NamedArgsType>
+		void CallConstruct( const TSharedRef<WidgetType>& OnWidget, const NamedArgsType& WithNamedArgs ) const
+		{
+			// YOUR WIDGET MUST IMPLEMENT void Construct( const FArguments& InArgs )
+			OnWidget->Construct( WithNamedArgs, Arg0, Arg1, Arg2, Arg3, Arg4 );
+		}
+
+		Arg0Type Arg0;
+		Arg1Type Arg1;
+		Arg2Type Arg2;
+		Arg3Type Arg3;
+		Arg4Type Arg4;
+	};
+
+	FORCEINLINE T0RequiredArgs MakeRequiredArgs()
+	{
+		return T0RequiredArgs();
+	}
+
+	template<typename Arg0Type>
+	T1RequiredArgs<Arg0Type> MakeRequiredArgs(Arg0Type InArg0)
+	{
+		return T1RequiredArgs<Arg0Type>( InArg0 );
+	}
+
+	template<typename Arg0Type, typename Arg1Type>
+	T2RequiredArgs<Arg0Type, Arg1Type> MakeRequiredArgs(Arg0Type InArg0, Arg1Type InArg1)
+	{
+		return T2RequiredArgs<Arg0Type, Arg1Type>( InArg0, InArg1 );
+	}
+	
+	template<typename Arg0Type, typename Arg1Type, typename Arg2Type>
+	T3RequiredArgs<Arg0Type, Arg1Type, Arg2Type> MakeRequiredArgs(Arg0Type InArg0, Arg1Type InArg1, Arg2Type InArg2)
+	{
+		return T3RequiredArgs<Arg0Type, Arg1Type, Arg2Type>( InArg0, InArg1, InArg2 );
+	}
+
+	template<typename Arg0Type, typename Arg1Type, typename Arg2Type, typename Arg3Type>
+	T4RequiredArgs<Arg0Type, Arg1Type, Arg2Type, Arg3Type> MakeRequiredArgs(Arg0Type InArg0, Arg1Type InArg1, Arg2Type InArg2, Arg3Type InArg3)
+	{
+		return T4RequiredArgs<Arg0Type, Arg1Type, Arg2Type, Arg3Type>( InArg0, InArg1, InArg2, InArg3 );
+	}
+
+	template<typename Arg0Type, typename Arg1Type, typename Arg2Type, typename Arg3Type, typename Arg4Type>
+	T5RequiredArgs<Arg0Type, Arg1Type, Arg2Type, Arg3Type, Arg4Type> MakeRequiredArgs(Arg0Type InArg0, Arg1Type InArg1, Arg2Type InArg2, Arg3Type InArg3, Arg4Type InArg4)
+	{
+		return T5RequiredArgs<Arg0Type, Arg1Type, Arg2Type, Arg3Type, Arg4Type>( InArg0, InArg1, InArg2, InArg3, InArg4 );
+	}
+}
+
+
+/**
+ * Utility class used during widget instantiation.
+ * Performs widget allocation and construction.
+ * Ensures that debug info is set correctly.
+ * Returns TSharedRef to widget.
+ *
+ * @see SNew
+ * @see SAssignNew
+ */
+template<class WidgetType, typename RequiredArgsPayloadType>
+struct TDecl
+{
+	TDecl( const ANSICHAR* InType, const ANSICHAR* InFile, int32 OnLine, const RequiredArgsPayloadType& InRequiredArgs )
+		: _Widget( new WidgetType() )
+		, _RequiredArgs( InRequiredArgs )
+	{
+		_Widget->SetDebugInfo( InType, InFile, OnLine );
+	}
+
+	/**
+	 * Initialize OutVarToInit with the widget that is being constructed.
+	 * @see SAssignNew
+	 */
+	template<class ExposeAsWidgetType>
+	TDecl& Expose( TSharedPtr<ExposeAsWidgetType>& OutVarToInit )
+	{
+		OutVarToInit = _Widget;
+		return *this;
+	}
+
+	/**
+	 * Initialize a WEAK OutVarToInit with the widget that is being constructed.
+	 * @see SAssignNew
+	 */
+	template<class ExposeAsWidgetType>
+	TDecl& Expose( TWeakPtr<ExposeAsWidgetType>& OutVarToInit )
+	{
+		OutVarToInit = _Widget;
+		return *this;
+	}
+
+	/**
+	 * Complete widget construction from InArgs.
+	 *
+	 * @param InArgs  NamedArguments from which to construct the widget.
+	 *
+	 * @return A reference to the widget that we constructed.
+	 */
+	TSharedRef<WidgetType> operator<<=( const typename WidgetType::FArguments& InArgs ) const
+	{
+		_Widget->SWidget::Construct(
+			InArgs._ToolTipText,
+			InArgs._ToolTip ,
+			InArgs._Cursor ,
+			InArgs._IsEnabled ,
+			InArgs._Visibility,
+			InArgs._Tag );
+
+		_RequiredArgs.CallConstruct( _Widget, InArgs );
+
+		return _Widget;
+	}
+
+	const TSharedRef<WidgetType> _Widget;
+	const RequiredArgsPayloadType _RequiredArgs;
+};
+
+
+template<typename WidgetType, typename RequiredArgsPayloadType>
+TDecl<WidgetType, RequiredArgsPayloadType> MakeTDecl( const ANSICHAR* InType, const ANSICHAR* InFile, int32 OnLine, const RequiredArgsPayloadType& InRequiredArgs )
+{
+	return TDecl<WidgetType, RequiredArgsPayloadType>(InType, InFile, OnLine, InRequiredArgs);
+}
diff --git a/Engine/Source/Runtime/SlateCore/Public/Widgets/IToolTip.h b/Engine/Source/Runtime/SlateCore/Public/Widgets/IToolTip.h
new file mode 100644
--- /dev/null
+++ b/Engine/Source/Runtime/SlateCore/Public/Widgets/IToolTip.h
@@ -0,0 +1,51 @@
+// Copyright 1998-2014 Epic Games, Inc. All Rights Reserved.
+
+/*=============================================================================
+	IToolTip.h: Declares the IToolTip interface.
+=============================================================================*/
+
+#pragma once
+
+
+/**
+ * Interface for tool tips.
+ */
+class IToolTip
+{
+public:
+
+	/**
+	 * Gets the widget that this tool tip represents.
+	 *
+	 * @return The tool tip widget.
+	 */
+	virtual TSharedRef<class SWidget> AsWidget( ) = 0;
+
+	/**
+	 * Gets the tool tip's content widget.
+	 *
+	 * @return The content widget.
+	 */
+	virtual const TSharedRef<SWidget>& GetContentWidget( ) = 0;
+
+	/**
+	 * Checks whether this tool tip has no content to display right now.
+	 *
+	 * @return true if the tool tip has no content to display, false otherwise.
+	 */
+	virtual bool IsEmpty( ) const = 0;
+
+	/**
+	 * Checks whether this tool tip can be made interactive by the user (by holding Ctrl).
+	 *
+	 * @return true if it is an interactive tool tip, false otherwise.
+	 */
+	virtual bool IsInteractive( ) const = 0;
+
+public:
+
+	/**
+	 * Virtual destructor.
+	 */
+	virtual ~IToolTip( ) { }
+};
diff --git a/Engine/Source/Runtime/SlateCore/Public/Widgets/SBoxPanel.h b/Engine/Source/Runtime/SlateCore/Public/Widgets/SBoxPanel.h
new file mode 100644
--- /dev/null
+++ b/Engine/Source/Runtime/SlateCore/Public/Widgets/SBoxPanel.h
@@ -0,0 +1,340 @@
+// Copyright 1998-2014 Epic Games, Inc. All Rights Reserved.
+
+/*=============================================================================
+	SBoxPanel.h: Declares the SBoxPanel, SHorizontalBox and SVerticalBox classes.
+=============================================================================*/
+
+#pragma once
+
+
+/**
+ * Implements a box panel widget.
+ *
+ * A BoxPanel contains one BoxPanel child and describes how that child should be arranged on the screen.
+ */
+class SLATECORE_API SBoxPanel
+	: public SPanel
+{
+public:
+
+	class FSlot
+	{
+	public:		
+		/** Horizontal and Vertical Boxes inherit from FSlot */
+		virtual ~FSlot() { }
+
+		/** The child widget contained in this slot. */
+		TSharedRef<SWidget> Widget;
+		/**
+		 * How much space this slot should occupy along panel's direction.
+		 *   When SizeRule is SizeRule_Auto, the widget's DesiredSize will be used as the space required.
+		 *   When SizeRule is SizeRule_AspectRatio, the widget will attempt to maintain the specified aspect ratio.
+		 *   When SizeRule is SizeRule_Stretch, the available space will be distributed proportionately between
+		 *   peer Widgets depending on the Value property. Available space is space remaining after all the
+		 *   peers' SizeRule_Auto requirements have been satisfied.
+		 */
+		FSizeParam SizeParam;
+		
+		/** Horizontal positioning of child within the allocated slot */
+		EHorizontalAlignment HAlignment;
+		
+		/** Vertical positioning of child within the allocated slot */
+		EVerticalAlignment VAlignment;
+		
+		/** The padding to add around the child. */
+		TAttribute<FMargin> SlotPadding;
+		
+		/** The max size that this slot can be (0 if no max) */
+		TAttribute<float> MaxSize;
+			
+	protected:
+
+		/**
+		 * Hidden default constructor.
+		 */
+		FSlot()
+			: Widget( SNullWidget::NullWidget )
+			, SizeParam( FStretch(1) )
+			, HAlignment( HAlign_Fill )
+			, VAlignment( VAlign_Fill )
+			, SlotPadding( FMargin(0) )
+			, MaxSize( 0.0f )
+		{ }
+	};
+
+public:
+
+	/** Removes a slot from this box panel which contains the specified SWidget
+	 *
+	 * @param SlotWidget The widget to match when searching through the slots
+	 * @returns The index in the children array where the slot was removed and -1 if no slot was found matching the widget
+	 */
+	int32 RemoveSlot( const TSharedRef<SWidget>& SlotWidget );
+
+	/**
+	 * Removes all children from the box.
+	 */
+	void ClearChildren();
+
+public:
+
+	// Begin SWidget overrides.
+
+	virtual void ArrangeChildren( const FGeometry& AllottedGeometry, FArrangedChildren& ArrangedChildren ) const;
+	
+	virtual FVector2D ComputeDesiredSize() const;
+
+	virtual FChildren* GetChildren();
+
+	// End SWidget overrides.
+
+protected:
+	/**
+	 * A Box Panel's orientation cannot be changed once it is constructed..
+	 *
+	 * @param InOrientation   The orientation of the Box Panel
+	 */
+	SBoxPanel( EOrientation InOrientation );
+
+	/** The Box Panel's orientation; determined at construct time. */
+	const EOrientation Orientation;
+	/** The Box Panel's children. */
+	TPanelChildren<FSlot> Children;
+};
+
+
+/** A Horizontal Box Panel. See SBoxPanel for more info. */
+class SLATECORE_API SHorizontalBox : public SBoxPanel
+{
+public:
+	class FSlot : public SBoxPanel::FSlot
+	{
+		public:
+		FSlot& AutoWidth()
+		{
+			SizeParam = FAuto();
+			return *this;
+		}
+
+		FSlot& MaxWidth( const TAttribute< float >& InMaxWidth )
+		{
+			MaxSize = InMaxWidth;
+			return *this;
+		}
+
+		FSlot& FillWidth( const TAttribute< float >& StretchCoefficient )
+		{
+			SizeParam = FStretch( StretchCoefficient );
+			return *this;
+		}
+
+		FSlot& AspectRatio()
+		{
+			SizeParam = FAspectRatio();
+			return *this;
+		}
+		
+		FSlot& Padding( float Uniform )
+		{
+			SlotPadding = FMargin(Uniform);
+			return *this;
+		}
+
+		FSlot& Padding( float Horizontal, float Vertical )
+		{
+			SlotPadding = FMargin(Horizontal, Vertical);
+			return *this;
+		}
+
+		FSlot& Padding( float Left, float Top, float Right, float Bottom )
+		{
+			SlotPadding = FMargin(Left, Top, Right, Bottom);
+			return *this;
+		}
+		
+		FSlot& Padding( TAttribute<FMargin> InPadding )
+		{
+			SlotPadding = InPadding;
+			return *this;
+		}
+
+		FSlot& HAlign( EHorizontalAlignment InHAlignment )
+		{
+			HAlignment = InHAlignment;
+			return *this;
+		}
+
+		FSlot& VAlign( EVerticalAlignment InVAlignment )
+		{
+			VAlignment = InVAlignment;
+			return *this;
+		}
+
+		FSlot& operator[]( TSharedRef<SWidget> InWidget )
+		{
+			Widget = InWidget;
+			return *this;
+		}
+
+		FSlot& Expose( FSlot*& OutVarToInit )
+		{
+			OutVarToInit = this;
+			return *this;
+		}
+	};
+
+	static FSlot& Slot()
+	{
+		return *(new FSlot());
+	}
+
+	SLATE_BEGIN_ARGS( SHorizontalBox )
+	{
+		_Visibility = EVisibility::SelfHitTestInvisible;
+	}
+
+		SLATE_SUPPORTS_SLOT(SHorizontalBox::FSlot)
+
+	SLATE_END_ARGS()
+
+	FSlot& AddSlot()
+	{
+		SHorizontalBox::FSlot& NewSlot = SHorizontalBox::Slot();
+		this->Children.Add( &NewSlot );
+		return NewSlot;
+	}
+
+	FORCENOINLINE SHorizontalBox()
+	: SBoxPanel( Orient_Horizontal )
+	{
+	}
+
+	/**
+	 * Construct this widget
+	 *
+	 * @param	InArgs	The declaration data for this widget
+	 */
+	void Construct( const FArguments& InArgs );
+};
+
+/** A Vertical Box Panel. See SBoxPanel for more info. */
+class SLATECORE_API SVerticalBox : public SBoxPanel
+{
+public:
+	class FSlot : public SBoxPanel::FSlot
+	{
+		public:
+		
+		FSlot& AutoHeight()
+		{
+			SizeParam = FAuto();
+			return *this;
+		}
+
+		FSlot& MaxHeight( const TAttribute< float >& InMaxHeight )
+		{
+			MaxSize = InMaxHeight;
+			return *this;
+		}
+
+		FSlot& FillHeight( const TAttribute< float >& StretchCoefficient )
+		{
+			SizeParam = FStretch( StretchCoefficient );
+			return *this;
+		}
+
+		FSlot& AspectRatio()
+		{
+			SizeParam = FAspectRatio();
+			return *this;
+		}
+
+		FSlot& Padding( float Uniform )
+		{
+			SlotPadding = FMargin(Uniform);
+			return *this;
+		}
+
+		FSlot& Padding( float Horizontal, float Vertical )
+		{
+			SlotPadding = FMargin(Horizontal, Vertical);
+			return *this;
+		}
+
+		FSlot& Padding( float Left, float Top, float Right, float Bottom )
+		{
+			SlotPadding = FMargin(Left, Top, Right, Bottom);
+			return *this;
+		}
+
+		FSlot& Padding( const TAttribute<FMargin>::FGetter& InDelegate )
+		{
+			SlotPadding.Bind( InDelegate );
+			return *this;
+		}
+
+		FSlot& HAlign( EHorizontalAlignment InHAlignment )
+		{
+			HAlignment = InHAlignment;
+			return *this;
+		}
+
+		FSlot& VAlign( EVerticalAlignment InVAlignment )
+		{
+			VAlignment = InVAlignment;
+			return *this;
+		}
+
+		FSlot& Padding( TAttribute<FMargin> InPadding )
+		{
+			SlotPadding = InPadding;
+			return *this;
+		}
+
+		FSlot& operator[]( TSharedRef<SWidget> InWidget )
+		{
+			Widget = InWidget;
+			return *this;
+		}
+
+		FSlot& Expose( FSlot*& OutVarToInit )
+		{
+			OutVarToInit = this;
+			return *this;
+		}
+	};
+
+	static FSlot& Slot()
+	{
+		return *(new FSlot());
+	}
+
+
+	SLATE_BEGIN_ARGS( SVerticalBox )
+	{
+		_Visibility = EVisibility::SelfHitTestInvisible;
+	}
+
+		SLATE_SUPPORTS_SLOT(SVerticalBox::FSlot)
+
+	SLATE_END_ARGS()
+
+	FSlot& AddSlot()
+	{
+		SVerticalBox::FSlot& NewSlot = SVerticalBox::Slot();
+		this->Children.Add( &NewSlot );
+		return NewSlot;
+	}
+
+	FORCENOINLINE SVerticalBox()
+	: SBoxPanel( Orient_Vertical )
+	{
+	}
+
+	/**
+	 * Construct this widget
+	 *
+	 * @param	InArgs	The declaration data for this widget
+	 */
+	void Construct( const FArguments& InArgs );
+};
diff --git a/Engine/Source/Runtime/SlateCore/Public/Widgets/SCompoundWidget.h b/Engine/Source/Runtime/SlateCore/Public/Widgets/SCompoundWidget.h
new file mode 100644
--- /dev/null
+++ b/Engine/Source/Runtime/SlateCore/Public/Widgets/SCompoundWidget.h
@@ -0,0 +1,101 @@
+// Copyright 1998-2014 Epic Games, Inc. All Rights Reserved.
+
+/*=============================================================================
+	SCompoundWidget.h: Declares the SCompoundWidget class.
+=============================================================================*/
+
+#pragma once
+
+
+/**
+ * Implements a compound widget.
+ *
+ * A CompoundWidget is the base from which most non-primitive widgets should be built.
+ * CompoundWidgets have a protected member named ChildSlot.
+ */
+class SLATECORE_API SCompoundWidget
+	: public SWidget
+{
+public:
+
+	/**
+	 * Returns the size scaling factor for this widget.
+	 *
+	 * @return Size scaling factor.
+	 */
+	const FVector2D GetContentScale() const
+	{
+		return ContentScale.Get();
+	}
+
+	/**
+	 * Sets the content scale for this widget.
+	 *
+	 * @param InContentScale Content scaling factor.
+	 */
+	void SetContentScale( const TAttribute< FVector2D >& InContentScale )
+	{
+		ContentScale = InContentScale;
+	}
+
+	/**
+	 * Sets the widget's color.
+	 *
+	 * @param InColorAndOpacity The ColorAndOpacity to be applied to this widget and all its contents.
+	 */
+	void SetColorAndOpacity( const TAttribute<FLinearColor>& InColorAndOpacity )
+	{
+		ColorAndOpacity = InColorAndOpacity;
+	}
+
+	/**
+	 * Sets the widget's foreground color.
+	 *
+	 * @param InColor The color to set.
+	 */
+	void SetForegroundColor( const TAttribute<FSlateColor>& InForegroundColor )
+	{
+		ForegroundColor = InForegroundColor;
+	}
+
+public:
+
+	// Begin SWidgetOverrides
+
+	virtual int32 OnPaint( const FGeometry& AllottedGeometry, const FSlateRect& MyClippingRect, FSlateWindowElementList& OutDrawElements, int32 LayerId, const FWidgetStyle& InWidgetStyle, bool bParentEnabled ) const OVERRIDE;
+	
+	virtual FChildren* GetChildren() OVERRIDE;
+
+	virtual FVector2D ComputeDesiredSize() const OVERRIDE;
+
+	virtual void ArrangeChildren( const FGeometry& AllottedGeometry, FArrangedChildren& ArrangedChildren ) const OVERRIDE;
+
+	virtual FSlateColor GetForegroundColor() const OVERRIDE;
+
+	// Begin SWidgetOverrides
+
+protected:
+
+	/**
+	 * Hidden constructor.
+	 */
+	FORCENOINLINE SCompoundWidget()
+		: ContentScale(FVector2D(1.0f, 1.0f))
+		, ColorAndOpacity(FLinearColor::White)
+		, ForegroundColor(FSlateColor::UseForeground())
+	{ }
+
+protected:
+
+	/** The slot that contains this widget's descendants.*/
+	FSimpleSlot ChildSlot;
+
+	/** The layout scale to apply to this widget's contents; useful for animation. */
+	TAttribute<FVector2D> ContentScale;
+
+	/** The color and opacity to apply to this widget and all its descendants. */
+	TAttribute<FLinearColor> ColorAndOpacity;
+
+	/** Optional foreground color that will be inherited by all of this widget's contents */
+	TAttribute<FSlateColor> ForegroundColor;
+};
diff --git a/Engine/Source/Runtime/SlateCore/Public/Widgets/SLeafWidget.h b/Engine/Source/Runtime/SlateCore/Public/Widgets/SLeafWidget.h
new file mode 100644
--- /dev/null
+++ b/Engine/Source/Runtime/SlateCore/Public/Widgets/SLeafWidget.h
@@ -0,0 +1,57 @@
+// Copyright 1998-2014 Epic Games, Inc. All Rights Reserved.
+
+/*=============================================================================
+	SLeafWidget.h: Declares the SLeafWidget class.
+=============================================================================*/
+
+#pragma once
+
+
+/**
+ * Implements a leaf widget.
+ *
+ * A LeafWidget is a Widget that has no slots for children.
+ * LeafWidgets are usually intended as building blocks for aggregate widgets.
+ */
+class SLATECORE_API SLeafWidget
+	: public SWidget
+{
+public:
+
+	// Begin SWidget overrides
+
+	/**
+	 * Overwritten from SWidget.
+	 *
+	 * LeafWidgets provide a visual representation of themselves. They do so by adding DrawElement(s)
+	 * to the OutDrawElements. DrawElements should have their positions set to absolute coordinates in
+	 * Window space; for this purpose the Slate system provides the AllottedGeometry parameter.
+	 * AllottedGeometry describes the space allocated for the visualization of this widget.
+	 *
+	 * Whenever possible, LeafWidgets should avoid dealing with layout properties. See TextBlock for an example.
+	 */
+	virtual int32 OnPaint( const FGeometry& AllottedGeometry, const FSlateRect& MyClippingRect, FSlateWindowElementList& OutDrawElements, int32 LayerId, const FWidgetStyle& InWidgetStyle, bool bParentEnabled ) const = 0;
+	
+	/**
+	 * Overwritten from SWidget.
+	 *
+	 * LeafWidgets should compute their DesiredSize based solely on their visual representation. There is no need to
+	 * take child widgets into account as LeafWidgets have none by definition. For example, the TextBlock widget simply
+	 * measures the area necessary to display its text with the given font and font size.
+	 */
+	virtual FVector2D ComputeDesiredSize( ) const = 0;
+	
+	/**
+	 * Overwritten from SWidget.
+	 *
+	 * Leaf widgets never have children.
+	 */
+	virtual FChildren* GetChildren( );
+
+	virtual void ArrangeChildren( const FGeometry& AllottedGeometry, FArrangedChildren& ArrangedChildren ) const;
+
+private:
+
+	// Shared instance of FNoChildren for all widgets with no children.
+	static FNoChildren NoChildrenInstance;
+};
diff --git a/Engine/Source/Runtime/SlateCore/Public/Widgets/SNullWidget.h b/Engine/Source/Runtime/SlateCore/Public/Widgets/SNullWidget.h
new file mode 100644
--- /dev/null
+++ b/Engine/Source/Runtime/SlateCore/Public/Widgets/SNullWidget.h
@@ -0,0 +1,26 @@
+// Copyright 1998-2014 Epic Games, Inc. All Rights Reserved.
+
+/*=============================================================================
+	SNullWidget.h: Declares the SNullWidget class.
+=============================================================================*/
+
+#pragma once
+
+
+/**
+ * Implements a widget that can be used as a placeholder.
+ *
+ * Widgets that support slots, such as SOverlay and SHorizontalBox should initialize
+ * their slots' child widgets to SNullWidget if no user defined widget was provided.
+ */
+class SLATECORE_API SNullWidget
+{
+public:
+
+	/**
+	 * Returns a placeholder widget.
+	 *
+	 * @return The widget.
+	 */
+	static TSharedRef<class SWidget> NullWidget;
+};
diff --git a/Engine/Source/Runtime/SlateCore/Public/Widgets/SOverlay.h b/Engine/Source/Runtime/SlateCore/Public/Widgets/SOverlay.h
new file mode 100644
--- /dev/null
+++ b/Engine/Source/Runtime/SlateCore/Public/Widgets/SOverlay.h
@@ -0,0 +1,131 @@
+// Copyright 1998-2014 Epic Games, Inc. All Rights Reserved.
+
+/*=============================================================================
+	SOverlay.h: Declares the SOverlay class.
+=============================================================================*/
+
+#pragma once
+
+
+/**
+ * Implements an overlay widget.
+ *
+ * Overlay widgets allow for layering several widgets on top of each other.
+ * Each slot of an overlay represents a layer that can contain one widget.
+ * The slots will be rendered on top of each other in the order they are declared in code.
+ *
+ * Usage:
+ *		SNew(SOverlay)
+ *		+ SOverlay::Slot(SNew(SMyWidget1))
+ *		+ SOverlay::Slot(SNew(SMyWidget2))
+ *		+ SOverlay::Slot(SNew(SMyWidget3))
+ *
+ *		Note that SWidget3 will be drawn on top of SWidget2 and SWidget1.
+ */
+class SLATECORE_API SOverlay
+	: public SPanel
+{
+public:	
+
+	/** A slot that support alignment of content and padding and z-order */
+	class SLATECORE_API FOverlaySlot
+		: public TSupportsOneChildMixin<SWidget, FOverlaySlot>
+	{
+	public:
+		FOverlaySlot()
+			: HAlignment(HAlign_Fill)
+			, VAlignment(VAlign_Fill)
+			, SlotPadding(0.0f)
+ 			, ZOrder(0)
+		{ }
+
+		FOverlaySlot& HAlign( EHorizontalAlignment InHAlignment )
+		{
+			HAlignment = InHAlignment;
+			return *this;
+		}
+
+		FOverlaySlot& VAlign( EVerticalAlignment InVAlignment )
+		{
+			VAlignment = InVAlignment;
+			return *this;
+		}
+
+		FOverlaySlot& Padding( const TAttribute<FMargin> InPadding )
+		{
+			SlotPadding = InPadding;
+			return *this;
+		}
+
+
+		/** The child widget contained in this slot. */
+		EHorizontalAlignment HAlignment;
+		EVerticalAlignment VAlignment;
+		TAttribute< FMargin > SlotPadding;
+
+		/** Slots with larger ZOrder values will draw above slots with smaller ZOrder values.  Slots
+		    with the same ZOrder will simply draw in the order they were added.  Currently this only
+			works for overlay slots that are added dynamically with AddWidget() and RemoveWidget() */
+		int32 ZOrder;
+	};
+
+
+	SLATE_BEGIN_ARGS( SOverlay )
+		{
+			_Visibility = EVisibility::SelfHitTestInvisible;
+		}
+
+		SLATE_SUPPORTS_SLOT( SOverlay::FOverlaySlot )
+
+	SLATE_END_ARGS()
+
+	/**
+	 * Virtual destructor.
+	 */
+	virtual ~SOverlay( ) { }
+
+public:
+
+	/**
+	 * Construct this widget.
+	 *
+	 * @param	InArgs	The declaration data for this widget
+	 */
+	void Construct( const FArguments& InArgs );
+
+	/** Returns the number of child widgets */
+	int32 GetNumWidgets() const;
+
+	/**
+	 * Removes a widget from this overlay.
+	 *
+	 * @param	Widget	The widget content to remove
+	 */
+	void RemoveSlot( TSharedRef< SWidget > Widget );
+
+	/** Adds a slot at the specified location (ignores Z-order) */
+	FOverlaySlot& AddSlot(int32 ZOrder=INDEX_NONE);
+
+	/** Removes a slot at the specified location */
+	void RemoveSlot(int32 ZOrder=INDEX_NONE);
+
+	/** Removes all children from the overlay */
+	void ClearChildren();
+
+	/** @return a new slot. Slots contain children for SOverlay */
+	static FOverlaySlot& Slot()
+	{
+		return *(new FOverlaySlot());
+	}
+
+	// SWidget interface
+	virtual void ArrangeChildren( const FGeometry& AllottedGeometry, FArrangedChildren& ArrangedChildren ) const OVERRIDE;
+	virtual FVector2D ComputeDesiredSize() const OVERRIDE;
+	virtual FChildren* GetChildren() OVERRIDE;
+	virtual int32 OnPaint( const FGeometry& AllottedGeometry, const FSlateRect& MyClippingRect, FSlateWindowElementList& OutDrawElements, int32 LayerId, const FWidgetStyle& InWidgetStyle, bool bParentEnabled ) const OVERRIDE;
+	// End of SWidget interface
+
+protected:
+	/** The SOverlay's slots; each slot contains a child widget. */
+	TPanelChildren<FOverlaySlot> Children;
+};
diff --git a/Engine/Source/Runtime/SlateCore/Public/Widgets/SPanel.h b/Engine/Source/Runtime/SlateCore/Public/Widgets/SPanel.h
new file mode 100644
--- /dev/null
+++ b/Engine/Source/Runtime/SlateCore/Public/Widgets/SPanel.h
@@ -0,0 +1,76 @@
+// Copyright 1998-2014 Epic Games, Inc. All Rights Reserved.
+
+/*=============================================================================
+	SPanel.h: Declares the SPanel class.
+=============================================================================*/
+
+#pragma once
+
+
+/**
+ * A Panel arranges its child widgets on the screen.
+ *
+ * Each child widget should be stored in a Slot. The Slot describes how the individual child should be arranged with
+ * respect to its parent (i.e. the Panel) and its peers Widgets (i.e. the Panel's other children.)
+ * For a simple example see StackPanel.
+ */
+class SLATECORE_API SPanel
+	: public SWidget
+{
+public:	
+
+	/**
+	 * Panels arrange their children in a space described by the AllottedGeometry parameter. The results of the arrangement
+	 * should be returned by appending a FArrangedWidget pair for every child widget. See StackPanel for an example
+	 *
+	 * @param AllottedGeometry    The geometry allotted for this widget by its parent.
+	 * @param ArrangedChildren    The array to which to add the WidgetGeometries that represent the arranged children.
+	 */
+	virtual void ArrangeChildren( const FGeometry& AllottedGeometry, FArrangedChildren& ArrangedChildren ) const = 0;
+	/**
+	 * A Panel's desired size in the space required to arrange of its children on the screen while respecting all of
+	 * the children's desired sizes and any layout-related options specified by the user. See StackPanel for an example.
+	 *
+	 * @return The desired size.
+	 */
+	virtual FVector2D ComputeDesiredSize( ) const = 0;
+
+	/**
+	 * All widgets must provide a way to access their children in a layout-agnostic way.
+	 * Panels store their children in Slots, which creates a dilemma. Most panels
+	 * can store their children in a TPanelChildren<Slot>, where the Slot class
+	 * provides layout information about the child it stores. In that case
+	 * GetChildren should simply return the TPanelChildren<Slot>. See StackPanel for an example.
+	 */
+	virtual FChildren* GetChildren( ) = 0;
+
+public:
+
+	/**
+	 * Most panels do not create widgets as part of their implementation, so
+	 * they do not need to implement a Construct()
+	 */
+	void Construct( ) { }
+
+public:
+
+	// Begin SWidget overrides
+
+	virtual int32 OnPaint( const FGeometry& AllottedGeometry, const FSlateRect& MyClippingRect, FSlateWindowElementList& OutDrawElements, int32 LayerId, const FWidgetStyle& InWidgetStyle, bool bParentEnabled ) const OVERRIDE;
+
+	// End SWidget overrides
+
+protected:
+
+	/**
+	 * Just like OnPaint, but takes already arranged children. Can be handy for writing custom SPanels.
+	 */
+	int32 PaintArrangedChildren( const FArrangedChildren& ArrangedChildren, const FSlateRect& MyClippingRect, FSlateWindowElementList& OutDrawElements, int32 LayerId, const FWidgetStyle& InWidgetStyle, bool bParentEnabled  ) const;
+	
+protected:
+
+	/**
+	 * Hidden default constructor.
+	 */
+	SPanel( ) { }
+};
