Slate: Refactored core Slate implementation into SlateCore module in preparation for UMG.

diff --git a/Engine/Source/Runtime/SlateCore/Public/Widgets/SWidget.h b/Engine/Source/Runtime/SlateCore/Public/Widgets/SWidget.h
new file mode 100644
--- /dev/null
+++ b/Engine/Source/Runtime/SlateCore/Public/Widgets/SWidget.h
@@ -0,0 +1,715 @@
+// Copyright 1998-2014 Epic Games, Inc. All Rights Reserved.
+
+/*=============================================================================
+	SWidget.h: Declares the SWidget class.
+=============================================================================*/
+
+#pragma once
+
+
+class SLATECORE_API FSlateControlledConstruction
+{
+public:
+	FSlateControlledConstruction(){}
+	virtual ~FSlateControlledConstruction(){}
+	
+private:
+	/** UI objects cannot be copy-constructed */
+	FSlateControlledConstruction(const FSlateControlledConstruction& Other){}
+	
+	/** UI objects cannot be copied. */
+	void operator= (const FSlateControlledConstruction& Other){}
+
+	/** Widgets should only ever be constructed via SNew or SAssignNew */
+	void* operator new (const size_t InSize )
+	{
+		return FMemory::Malloc(InSize);
+	}
+
+	template<class WidgetType, typename RequiredArgsPayloadType>
+	friend struct TDecl;
+
+public:
+	void operator delete(void * mem)
+	{
+		FMemory::Free(mem);
+	}
+};
+
+
+/**
+ * Abstract base class for Slate widgets.
+ *
+ * STOP. DO NOT INHERIT DIRECTLY FROM WIDGET!
+ *
+ * Inheritance:
+ *   Widget is not meant to be directly inherited. Instead consider inheriting from LeafWidget or Panel,
+ *   which represent intended use cases and provide a succinct set of methods which to override.
+ *
+ *   SWidget is the base class for all interactive Slate entities. SWidget's public interface describes
+ *   everything that a Widget can do and is fairly complex as a result.
+ * 
+ * Events:
+ *   Events in Slate are implemented as virtual functions that the Slate system will call
+ *   on a Widget in order to notify the Widget about an important occurrence (e.g. a key press)
+ *   or querying the Widget regarding some information (e.g. what mouse cursor should be displayed).
+ *
+ *   Widget provides a default implementation for most events; the default implementation does nothing
+ *   and does not handle the event.
+ *
+ *   Some events are able to reply to the system by returning an FReply, FCursorReply, or similar
+ *   object. 
+ */
+class SLATECORE_API SWidget
+	: public FSlateControlledConstruction,
+	  public TSharedFromThis<SWidget>		// Enables 'this->AsShared()'
+{
+public:
+	
+	/**
+	 * Construct a SWidget based on initial parameters.
+	 */
+	void Construct(
+		const TAttribute<FString> & InToolTipText ,
+		const TSharedPtr<IToolTip> & InToolTip ,
+		const TAttribute< TOptional<EMouseCursor::Type> > & InCursor ,
+		const TAttribute<bool> & InEnabledState ,
+		const TAttribute<EVisibility> & InVisibility,
+		const FName& InTag );
+
+	//
+	// GENERAL EVENTS
+	//
+
+	/**
+	 * The widget should respond by populating the OutDrawElements array with FDrawElements 
+	 * that represent it and any of its children.
+	 *
+	 * @param AllottedGeometry  The FGeometry that describes an area in which the widget should appear.
+	 * @param MyClippingRect    The clipping rectangle allocated for this widget and its children.
+	 * @param OutDrawElements   A list of FDrawElements to populate with the output.
+	 * @param LayerId           The Layer onto which this widget should be rendered.
+	 * @param InColorAndOpacity Color and Opacity to be applied to all the descendants of the widget being painted
+ 	 * @param bParentEnabled	True if the parent of this widget is enabled.
+	 *
+	 * @return The maximum layer ID attained by this widget or any of its children.
+	 */
+	virtual int32 OnPaint( const FGeometry& AllottedGeometry, const FSlateRect& MyClippingRect, FSlateWindowElementList& OutDrawElements, int32 LayerId, const FWidgetStyle& InWidgetStyle, bool bParentEnabled ) const = 0;
+
+	/**
+	 * Ticks this widget.  Override in derived classes, but always call the parent implementation.
+	 *
+	 * @param  AllottedGeometry The space allotted for this widget
+	 * @param  InCurrentTime  Current absolute real time
+	 * @param  InDeltaTime  Real time passed since last tick
+	 */
+	virtual void Tick( const FGeometry& AllottedGeometry, const double InCurrentTime, const float InDeltaTime ) { }
+
+	/**
+	 * Called when a widget is being hit tested and has passed the bounding box test to determine
+	 * if the precise widget geometry is under the cursor.
+	 *
+	 * @param MyGeometry The Geometry of the widget receiving the event
+	 * @param InAbsoluteCursorPosition	The absolute cursor position
+	 *
+	 * @return true if this widget considers itself hit by the passed in cursor position.
+	 */
+	virtual bool OnHitTest( const FGeometry& MyGeometry, FVector2D InAbsoluteCursorPosition );
+
+	//
+	// KEYBOARD INPUT
+	//
+
+	/**
+	 * Called when keyboard focus is given to this widget.  This event does not bubble.
+	 *
+	 * @param MyGeometry The Geometry of the widget receiving the event
+	 * @param  InKeyboardFocusEvent  KeyboardFocusEvent
+	 *
+	 * @return  Returns whether the event was handled, along with other possible actions
+	 */
+	virtual FReply OnKeyboardFocusReceived( const FGeometry& MyGeometry, const FKeyboardFocusEvent& InKeyboardFocusEvent );
+
+	/**
+	 * Called when this widget loses the keyboard focus.  This event does not bubble.
+	 *
+	 * @param  InKeyboardFocusEvent  KeyboardFocusEvent
+	 */
+	virtual void OnKeyboardFocusLost( const FKeyboardFocusEvent& InKeyboardFocusEvent ) { }
+
+	/** Called whenever a focus path is changing on all the widgets within the old and new focus paths */
+	virtual void OnKeyboardFocusChanging( const FWeakWidgetPath& PreviousFocusPath, const FWidgetPath& NewWidgetPath ) { }
+
+	/**
+	 * Called after a character is entered while this widget has keyboard focus
+	 *
+	 * @param MyGeometry The Geometry of the widget receiving the event
+	 * @param  InCharacterEvent  Character event
+	 *
+	 * @return  Returns whether the event was handled, along with other possible actions
+	 */
+	virtual FReply OnKeyChar( const FGeometry& MyGeometry,const FCharacterEvent& InCharacterEvent );
+
+	/**
+	 * Called after a key is pressed when this widget or a child of this widget has keyboard focus
+	 * If a widget handles this event, OnKeyDown will *not* be passed to the focused widget.
+	 *
+	 * This event is primarily to allow parent widgets to consume an event before a child widget processes
+	 * it and it should be used only when there is no better design alternative.
+	 *
+	 * @param MyGeometry The Geometry of the widget receiving the event
+	 * @param  InKeyboardEvent  Keyboard event
+	 *
+	 * @return  Returns whether the event was handled, along with other possible actions
+	 */
+	virtual FReply OnPreviewKeyDown( const FGeometry& MyGeometry, const FKeyboardEvent& InKeyboardEvent );
+
+	/**
+	 * Called after a key is pressed when this widget has keyboard focus (this event bubbles if not handled)
+	 *
+	 * @param MyGeometry The Geometry of the widget receiving the event
+	 * @param  InKeyboardEvent  Keyboard event
+	 *
+	 * @return  Returns whether the event was handled, along with other possible actions
+	 */
+	virtual FReply OnKeyDown( const FGeometry& MyGeometry, const FKeyboardEvent& InKeyboardEvent );
+
+	/**
+	 * Called after a key is released when this widget has keyboard focus
+	 *
+	 * @param MyGeometry The Geometry of the widget receiving the event
+	 * @param  InKeyboardEvent  Keyboard event
+	 *
+	 * @return  Returns whether the event was handled, along with other possible actions
+	 */
+	virtual FReply OnKeyUp( const FGeometry& MyGeometry, const FKeyboardEvent& InKeyboardEvent );
+
+	//
+	// MOUSE INPUT
+	//
+
+	/**
+	 * The system calls this method to notify the widget that a mouse button was pressed within it. This event is bubbled.
+	 *
+	 * @param MyGeometry The Geometry of the widget receiving the event
+	 * @param MouseEvent Information about the input event
+	 *
+	 * @return Whether the event was handled along with possible requests for the system to take action.
+	 */
+	virtual FReply OnMouseButtonDown( const FGeometry& MyGeometry, const FPointerEvent& MouseEvent );
+
+	/**
+	 * Just like OnMouseButtonDown, but tunnels instead of bubbling.
+	 * If this even is handled, OnMouseButtonDown will not be sent.
+	 * 
+	 * Use this event sparingly as preview events generally make UIs more
+	 * difficult to reason about.
+	 */
+	virtual FReply OnPreviewMouseButtonDown( const FGeometry& MyGeometry, const FPointerEvent& MouseEvent );
+	
+	/**
+	 * The system calls this method to notify the widget that a mouse button was release within it. This event is bubbled.
+	 *
+	 * @param MyGeometry The Geometry of the widget receiving the event
+	 * @param MouseEvent Information about the input event
+	 *
+	 * @return Whether the event was handled along with possible requests for the system to take action.
+	 */
+	virtual FReply OnMouseButtonUp( const FGeometry& MyGeometry, const FPointerEvent& MouseEvent );
+	
+	/**
+	 * The system calls this method to notify the widget that a mouse moved within it. This event is bubbled.
+	 *
+	 * @param MyGeometry The Geometry of the widget receiving the event
+	 * @param MouseEvent Information about the input event
+	 *
+	 * @return Whether the event was handled along with possible requests for the system to take action.
+	 */
+	virtual FReply OnMouseMove( const FGeometry& MyGeometry, const FPointerEvent& MouseEvent );
+	
+	/**
+	 * The system will use this event to notify a widget that the cursor has entered it. This event is NOT bubbled.
+	 *
+	 * @param MyGeometry The Geometry of the widget receiving the event
+	 * @param MouseEvent Information about the input event
+	 */
+	virtual void OnMouseEnter( const FGeometry& MyGeometry, const FPointerEvent& MouseEvent );
+	
+	/**
+	 * The system will use this event to notify a widget that the cursor has left it. This event is NOT bubbled.
+	 *
+	 * @param MouseEvent Information about the input event
+	 */
+	virtual void OnMouseLeave( const FPointerEvent& MouseEvent );
+	
+	/**
+	 * Called when the mouse wheel is spun. This event is bubbled.
+	 *
+	 * @param  MouseEvent  Mouse event
+	 *
+	 * @return  Returns whether the event was handled, along with other possible actions
+	 */
+	virtual FReply OnMouseWheel( const FGeometry& MyGeometry, const FPointerEvent& MouseEvent );
+	
+	/**
+	 * The system asks each widget under the mouse to provide a cursor. This event is bubbled.
+	 * 
+	 * @return FCursorReply::Unhandled() if the event is not handled; return FCursorReply::Cursor() otherwise.
+	 */
+	virtual FCursorReply OnCursorQuery( const FGeometry& MyGeometry, const FPointerEvent& CursorEvent ) const;
+
+	/**
+	 * Called when a mouse button is double clicked.  Override this in derived classes.
+	 *
+	 * @param  InMyGeometry  Widget geometry
+	 * @param  InMouseEvent  Mouse button event
+	 *
+	 * @return  Returns whether the event was handled, along with other possible actions
+	 */
+	virtual FReply OnMouseButtonDoubleClick( const FGeometry& InMyGeometry, const FPointerEvent& InMouseEvent );
+
+	/**
+	 * Called when Slate wants to visualize tooltip.
+	 * If nobody handles this event, Slate will use default tooltip visualization.
+	 * If you override this event, you should probably return true.
+	 * 
+	 * @param  TooltipContent    The TooltipContent that I may want to visualize.
+	 *
+	 * @return true if this widget visualized the tooltip content; i.e., the event is handled.
+	 */
+	virtual bool OnVisualizeTooltip( const TSharedPtr<SWidget>& TooltipContent );
+
+	/**
+	 * Called when Slate detects that a widget started to be dragged.
+	 * Usage:
+	 * A widget can ask Slate to detect a drag.
+	 * OnMouseDown() reply with FReply::Handled().DetectDrag( SharedThis(this) ).
+	 * Slate will either send an OnDragDetected() event or do nothing.
+	 * If the user releases a mouse button or leaves the widget before
+	 * a drag is triggered (maybe user started at the very edge) then no event will be
+	 * sent.
+	 *
+	 * @param  InMyGeometry  Widget geometry
+	 * @param  InMouseEvent  MouseMove that triggered the drag
+	 * 
+	 */
+	virtual FReply OnDragDetected( const FGeometry& MyGeometry, const FPointerEvent& MouseEvent );
+
+	//
+	// DRAG AND DROP (DragDrop)
+	//
+
+	/**
+	 * Called during drag and drop when the drag enters a widget.
+	 *
+	 * Enter/Leave events in slate are meant as lightweight notifications.
+	 * So we do not want to capture mouse or set focus in response to these.
+	 * However, OnDragEnter must also support external APIs (e.g. OLE Drag/Drop)
+	 * Those require that we let them know whether we can handle the content
+	 * being dragged OnDragEnter.
+	 *
+	 * The concession is to return a can_handled/cannot_handle
+	 * boolean rather than a full FReply.
+	 *
+	 * @param MyGeometry      The geometry of the widget receiving the event.
+	 * @param DragDropEvent   The drag and drop event.
+	 *
+	 * @return A reply that indicated whether the contents of the DragDropEvent can potentially be processed by this widget.
+	 */
+	virtual void OnDragEnter( const FGeometry& MyGeometry, const FDragDropEvent& DragDropEvent );
+	
+	/**
+	 * Called during drag and drop when the drag leaves a widget.
+	 *
+	 * @param DragDropEvent   The drag and drop event.
+	 */
+	virtual void OnDragLeave( const FDragDropEvent& DragDropEvent );
+
+	/**
+	 * Called during drag and drop when the the mouse is being dragged over a widget.
+	 *
+	 * @param MyGeometry      The geometry of the widget receiving the event.
+	 * @param DragDropEvent   The drag and drop event.
+	 *
+	 * @return A reply that indicated whether this event was handled.
+	 */
+	virtual FReply OnDragOver( const FGeometry& MyGeometry, const FDragDropEvent& DragDropEvent );
+	
+	/**
+	 * Called when the user is dropping something onto a widget; terminates drag and drop.
+	 *
+	 * @param MyGeometry      The geometry of the widget receiving the event.
+	 * @param DragDropEvent   The drag and drop event.
+	 *
+	 * @return A reply that indicated whether this event was handled.
+	 */
+	virtual FReply OnDrop( const FGeometry& MyGeometry, const FDragDropEvent& DragDropEvent );
+	
+	/**
+	 * Called when a controller button is pressed
+	 * 
+	 * @param ControllerEvent	The controller event generated
+	 */
+	virtual FReply OnControllerButtonPressed( const FGeometry& MyGeometry, const FControllerEvent& ControllerEvent );
+
+	/**
+	 * Called when a controller button is released
+	 * 
+	 * @param ControllerEvent	The controller event generated
+	 */
+	virtual FReply OnControllerButtonReleased( const FGeometry& MyGeometry, const FControllerEvent& ControllerEvent );
+
+	/**
+	 * Called when an analog value on a controller changes
+	 * 
+	 * @param ControllerEvent	The controller event generated
+	 */
+	virtual FReply OnControllerAnalogValueChanged( const FGeometry& MyGeometry, const FControllerEvent& ControllerEvent );
+
+	/**
+	 * Called when the user performs a gesture on trackpad. This event is bubbled.
+	 *
+	 * @param  GestureEvent  gesture event
+	 *
+	 * @return  Returns whether the event was handled, along with other possible actions
+	 */
+	virtual FReply OnTouchGesture( const FGeometry& MyGeometry, const FPointerEvent& GestureEvent );
+	
+	/**
+	 * Called when a touchpad touch is started (finger down)
+	 * 
+	 * @param TouchEvent	The touch event generated
+	 */
+	virtual FReply OnTouchStarted( const FGeometry& MyGeometry, const FPointerEvent& InTouchEvent );
+
+	/**
+	 * Called when a touchpad touch is moved  (finger moved)
+	 * 
+	 * @param TouchEvent	The touch event generated
+	 */
+	virtual FReply OnTouchMoved( const FGeometry& MyGeometry, const FPointerEvent& InTouchEvent );
+
+	/**
+	 * Called when a touchpad touch is ended (finger lifted)
+	 * 
+	 * @param TouchEvent	The touch event generated
+	 */
+	virtual FReply OnTouchEnded( const FGeometry& MyGeometry, const FPointerEvent& InTouchEvent );
+
+	/**
+	 * Called when motion is detected (controller or device)
+	 * 
+	 * @param MotionEvent	The motion event generated
+	 */
+	virtual FReply OnMotionDetected( const FGeometry& MyGeometry, const FMotionEvent& InMotionEvent );
+
+	/**
+	 * Called when the mouse is moved over the widget's window, to determine if we should report whether
+	 * OS-specific features should be active at this location (such as a title bar grip, system menu, etc.)
+	 * Usually you should not need to override this function.
+	 *
+	 * @return	The window "zone" the cursor is over, or EWindowZone::Unspecified if no special behavior is needed
+	 */
+	virtual EWindowZone::Type GetWindowZoneOverride() const;
+
+	//
+	// LAYOUT
+	//
+
+	/**
+	 * Applies styles as it descends through the widget hierarchy.
+	 * Gathers desired sizes on the way up.
+	 * i.e. Caches the desired size of all of this widget's children recursively, then caches desired size for itself.
+	 */
+	void SlatePrepass();
+
+	/** @return the DesiredSize that was computed the last time CacheDesiredSize() was called. */
+	const FVector2D& GetDesiredSize() const;
+
+	/**
+	 * The system calls this method. It performs a breadth-first traversal of every visible widget and asks
+	 * each widget to cache how big it needs to be in order to present all of its content.
+	 */
+	protected: virtual void CacheDesiredSize();
+	
+	/**
+	 * Explicitly set the desired size. This is highly advanced functionality that is meant
+	 * to be used in conjunction with overriding CacheDesiredSize. Use ComputeDesiredSize() instead.
+	 */
+	protected: void Advanced_SetDesiredSize(const FVector2D& InDesiredSize) { DesiredSize = InDesiredSize; }
+
+	/**
+	 * Compute the ideal size necessary to display this widget. For aggregate widgets (e.g. panels) this size should include the
+	 * size necessary to show all of its children. CacheDesiredSize() guarantees that the size of descendants is computed and cached
+	 * before that of the parents, so it is safe to call GetDesiredSize() for any children while implementing this method.
+	 *
+	 * Note that ComputeDesiredSize() is meant as an aide to the developer. It is NOT meant to be very robust in many cases. If your
+	 * widget is simulating a bouncing ball, you should just return a reasonable size; e.g. 160x160. Let the programmer set up a reasonable
+	 * rule of resizing the bouncy ball simulation.
+	 *
+	 * @return The desired size.
+	 */
+	protected: virtual FVector2D ComputeDesiredSize() const = 0;
+
+	public:
+
+	/**
+	 * Compute the Geometry of all the children and add populate the ArrangedChildren list with their values.
+	 * Each type of Layout panel should arrange children based on desired behavior.
+	 *
+	 * @param AllottedGeometry    The geometry allotted for this widget by its parent.
+	 * @param ArrangedChildren    The array to which to add the WidgetGeometries that represent the arranged children.
+	 */
+	virtual void ArrangeChildren( const FGeometry& AllottedGeometry, FArrangedChildren& ArrangedChildren ) const = 0;
+
+	/**
+	 * Every widget that has children must implement this method. This allows for iteration over the Widget's
+	 * children regardless of how they are actually stored.
+	 */
+	// @todo Slate: Consider renaming to GetVisibleChildren  (not ALL children will be returned in all cases)
+	virtual FChildren* GetChildren() = 0;
+
+	/**
+	 * Checks to see if this widget supports keyboard focus.  Override this in derived classes.
+	 *
+	 * @return  True if this widget can take keyboard focus
+	 */
+	virtual bool SupportsKeyboardFocus() const;
+
+
+	/**
+	 * Checks to see if this widget currently has the keyboard focus
+	 *
+	 * @return  True if this widget has keyboard focus
+	 */
+	virtual bool HasKeyboardFocus() const;
+
+	/**
+	 * @return Whether this widget has any descendants with keyboard focus
+	 */
+	bool HasFocusedDescendants() const;
+
+	/**
+	 * Checks to see if this widget is the current mouse captor
+	 *
+	 * @return  True if this widget has captured the mouse
+	 */
+	bool HasMouseCapture() const;
+
+	/**
+	 * Called when this widget had captured the mouse, but that capture has been revoked for some reason.
+	 */
+	virtual void OnMouseCaptureLost() { }
+
+	/**
+	 * Ticks this widget and all of it's child widgets.  Should not be called directly.
+	 *
+	 * @param  AllottedGeometry The space allotted for this widget
+	 * @param  InCurrentTime  Current absolute real time
+	 * @param  InDeltaTime  Real time passed since last tick
+	 */
+	void TickWidgetsRecursively( const FGeometry& AllottedGeometry, const double InCurrentTime, const float InDeltaTime );
+
+	/**
+	 * Sets the enabled state of this widget
+	 *
+	 * @param InEnabledState	An attribute containing the enabled state or a delegate to call to get the enabled state.
+	 */
+	void SetEnabled( const TAttribute<bool>& InEnabledState )
+	{
+		EnabledState = InEnabledState;
+	}
+
+	/** @return Whether or not this widget is enabled */
+	bool IsEnabled() const 
+	{
+		return EnabledState.Get();
+	}
+
+	/** @return The tool tip associated with this widget; Invalid reference if there is not one */
+	virtual TSharedPtr<IToolTip> GetToolTip();
+
+	/** Called when a tooltip displayed from this widget is being closed */
+	virtual void OnToolTipClosing() { }
+
+	/**
+	 * Sets whether this widget is a "tool tip force field".  That is, tool-tips should never spawn over the area
+	 * occupied by this widget, and will instead be repelled to an outside edge
+	 *
+	 * @param	bEnableForceField	True to enable tool tip force field for this widget
+	 */
+	void EnableToolTipForceField( const bool bEnableForceField );
+
+	/** @return True if a tool tip force field is active on this widget */
+	bool HasToolTipForceField() const
+	{
+		return bToolTipForceFieldEnabled;
+	}
+
+	/** @return True if this widget hovered */
+	virtual bool IsHovered() const
+	{
+		return bIsHovered;
+	}
+
+	/** @return is this widget visible, hidden or collapsed */
+	EVisibility GetVisibility() const
+	{
+		return Visibility.Get();
+	}
+
+	/** @param InVisibility  should this widget be */
+	void SetVisibility( TAttribute<EVisibility> InVisibility )
+	{
+		Visibility = InVisibility;
+	}
+
+	/**
+	 * Set the tool tip that should appear when this widget is hovered.
+	 *
+	 * @param InToolTipText  the text that should appear in the tool tip
+	 */
+	void SetToolTipText( const TAttribute<FString>& InToolTipText );
+
+	void SetToolTipText( const FText& InToolTipText );
+
+	/**
+	 * Set the tool tip that should appear when this widget is hovered.
+	 *
+	 * @param InToolTip  the widget that should appear in the tool tip
+	 */
+	void SetToolTip( const TSharedPtr<IToolTip>& InToolTip );
+
+	/**
+	 * Set the cursor that should appear when this widget is hovered
+	 */
+	void SetCursor( const TAttribute< TOptional<EMouseCursor::Type> >& InCursor );
+
+	/**
+	 * Used by Slate to set the runtime debug info about this widget.
+	 */
+	void SetDebugInfo( const ANSICHAR* InType, const ANSICHAR* InFile, int32 OnLine );
+
+public:
+	//
+	// Widget Inspector and debugging methods
+	//
+
+	/** @return A String representation of the widget */
+	virtual FString ToString() const;
+
+	/** @return A String of the widget's type */
+	FString GetTypeAsString() const;
+
+	/** @return The widget's type as an FName ID */
+	FName GetType() const;
+
+	/** @return A String of the widget's code location in readable format */
+	virtual FString GetReadableLocation() const;
+
+	/** @return A String of the widget's code location */
+	virtual FString GetCreatedInFile() const;
+
+	/** @return The line number of the widgets location */
+	virtual int32 GetCreatedInLineNumber() const;
+
+	/** @return the Foreground color that this widget sets; unset options if the widget does not set a foreground color */
+	virtual FSlateColor GetForegroundColor() const;
+
+protected:
+
+	/**
+	 * Hidden default constructor.
+	 *
+	 * Use SNew(WidgetClassName) to instantiate new widgets.
+	 *
+	 * @see SNew
+	 */
+	SWidget();
+
+	/** 
+	 * Find the geometry of a descendant widget. This method assumes that WidgetsToFind are a descendants of this widget.
+	 * Note that not all widgets are guaranteed to be found; OutResult will contain null entries for missing widgets.
+	 *
+	 * @param MyGeometry      The geometry of this widget.
+	 * @param WidgetsToFind   The widgets whose geometries we wish to discover.
+	 * @param OutResult       A map of widget references to their respective geometries.
+	 *
+	 * @return True if all the WidgetGeometries were found. False otherwise.
+	 */
+	bool FindChildGeometries( const FGeometry& MyGeometry, const TSet< TSharedRef<SWidget> >& WidgetsToFind, TMap<TSharedRef<SWidget>, FArrangedWidget>& OutResult ) const;
+
+	/**
+	 * Actual implementation of FindChildGeometries.
+	 *
+	 * @param MyGeometry      The geometry of this widget.
+	 * @param WidgetsToFind   The widgets whose geometries we wish to discover.
+	 * @param OutResult       A map of widget references to their respective geometries.
+	 */
+	void FindChildGeometries_Helper( const FGeometry& MyGeometry, const TSet< TSharedRef<SWidget> >& WidgetsToFind, TMap<TSharedRef<SWidget>, FArrangedWidget>& OutResult ) const;
+
+	/** 
+	 * Find the geometry of a descendant widget. This method assumes that WidgetToFind is a descendant of this widget.
+	 *
+	 * @param MyGeometry   The geometry of this widget.
+	 * @param WidgetToFind The widget whose geometry we wish to discover.
+	 *
+	 * @return the geometry of WidgetToFind.
+	 */
+	FGeometry FindChildGeometry( const FGeometry& MyGeometry, TSharedRef<SWidget> WidgetToFind ) const;
+
+	/** @return The index of the child that the mouse is currently hovering */
+	static int32 FindChildUnderMouse( const FArrangedChildren& Children, const FPointerEvent& MouseEvent );
+
+	/** 
+	 * Determines if this widget should be enabled.
+	 * 
+	 * @param InParentEnabled	true if the parent of this widget is enabled
+	 * @return true if the widget is enabled
+	 */
+	bool ShouldBeEnabled( bool InParentEnabled ) const
+	{
+		// This widget should be enabled if its parent is enabled and it is enabled
+		return IsEnabled() && InParentEnabled;
+	}
+
+protected:
+	
+	//	DEBUG INFORMATION
+	// @todo Slate: Should compile out in final release builds?
+	FName TypeOfWidget;
+	/** Full file path in which this widget was created */
+	FName CreatedInFileFullPath;
+	/** Filename in which this widget was created */
+	FName CreatedInFile;
+	/** Line number on which this widget was created */
+	int32 CreatedOnLine;
+
+	/** The cursor to show when the mouse is hovering over this widget. */
+	TAttribute< TOptional<EMouseCursor::Type> > Cursor;
+
+	/** Whether or not this widget is enabled */
+	TAttribute< bool > EnabledState;
+
+	/** Is this widget visible, hidden or collapsed */
+	TAttribute< EVisibility > Visibility;
+
+private:
+
+	/**
+	 * Stores the ideal size this widget wants to be.
+	 * This member is intentionally private, because only the very base class (Widget) can write DesiredSize.
+	 * See CacheDesiredSize(), ComputeDesiredSize()
+	 */
+	FVector2D DesiredSize;
+
+	/** Tool tip content for this widget */
+	TSharedPtr<IToolTip> ToolTip;
+
+	// Whether this widget is a "tool tip force field".  That is, tool-tips should never spawn over the area
+	// occupied by this widget, and will instead be repelled to an outside edge
+ 	bool bToolTipForceFieldEnabled;
+
+	/** Is this widget hovered? */
+	bool bIsHovered;
+};
diff --git a/Engine/Source/Runtime/SlateCore/Public/Widgets/SWindow.h b/Engine/Source/Runtime/SlateCore/Public/Widgets/SWindow.h
new file mode 100644
--- /dev/null
+++ b/Engine/Source/Runtime/SlateCore/Public/Widgets/SWindow.h
@@ -0,0 +1,984 @@
+// Copyright 1998-2014 Epic Games, Inc. All Rights Reserved.
+
+/*=============================================================================
+	SWindow.h: Declares the SWindow class.
+=============================================================================*/
+
+#pragma once
+
+
+/** Notification that a window has been deactivated */
+DECLARE_DELEGATE( FOnWindowDeactivated );
+
+/** Notification that a window is about to be closed */
+DECLARE_DELEGATE_OneParam( FOnWindowClosed, const TSharedRef<SWindow>& );
+
+/** Notification that a window has been moved */
+DECLARE_DELEGATE_OneParam( FOnWindowMoved, const TSharedRef<SWindow>& );
+
+/** Override delegate for RequestDestroyWindow */
+DECLARE_DELEGATE_OneParam( FRequestDestroyWindowOverride, const TSharedRef<SWindow>& );
+
+/** Called when we need to switch game worlds for a window */
+DECLARE_DELEGATE_RetVal_OneParam( int32, FOnSwitchWorldHack, int32 );
+
+
+/** Enum to describe how to auto-center an SWindow */
+namespace EAutoCenter
+{
+	enum Type
+	{
+		/** Don't auto-center the window */
+		None,
+
+		/** Auto-center the window on the primary work area */
+		PrimaryWorkArea,
+
+		/** Auto-center the window on the preferred work area, determined using GetPreferredWorkArea() */
+		PreferredWorkArea,
+	};
+}
+
+
+/** Enum to describe how windows are sized */
+namespace ESizingRule
+{
+	enum Type
+	{
+		/* The windows size fixed and cannot be resized **/
+		FixedSize,
+		
+		/** The window size is computed from its content and cannot be resized by users */
+		Autosized,
+
+		/** The window can be resized by users */
+		UserSized,
+	};
+}
+
+
+/**
+ * SWindow is a platform-agnostic representation of a top-level window.
+ */
+class SLATECORE_API SWindow
+	: public SCompoundWidget
+{
+public:
+
+	SLATE_BEGIN_ARGS( SWindow )
+		: _Style( &FCoreStyle::Get().GetWidgetStyle<FWindowStyle>("Window") )
+		, _Title()
+		, _bDragAnywhere( false )
+		, _AutoCenter( EAutoCenter::PreferredWorkArea )
+		, _ScreenPosition( FVector2D::ZeroVector )
+		, _ClientSize( FVector2D::ZeroVector )
+		, _SupportsTransparency( false )
+		, _InitialOpacity( 1.0f )
+		, _IsInitiallyMaximized( false )
+		, _SizingRule( ESizingRule::UserSized )
+		, _IsPopupWindow( false )
+		, _FocusWhenFirstShown( true )
+		, _ActivateWhenFirstShown( true )
+		, _UseOSWindowBorder( false )
+		, _SupportsMaximize( true )
+		, _SupportsMinimize( true )
+		, _CreateTitleBar( true )
+		, _SaneWindowPlacement( true )
+	{ }
+
+		/** Style used to draw this window */
+		SLATE_STYLE_ARGUMENT( FWindowStyle, Style )
+
+		/** Title of the window */
+		SLATE_ATTRIBUTE( FText, Title )
+
+		/** When true, the window can be dragged from anywhere. */
+		SLATE_ARGUMENT( bool, bDragAnywhere )
+
+		/** The windows auto-centering mode. If set to anything other than None, then the
+		    ScreenPosition value will be ignored */
+		SLATE_ARGUMENT( EAutoCenter::Type, AutoCenter )
+
+		/** Screen-space position where the window should be initially located. */
+		SLATE_ARGUMENT( FVector2D, ScreenPosition )
+
+		/** What the initial size of the window should be. */
+		SLATE_ARGUMENT( FVector2D, ClientSize )
+
+		/** Should this window support transparency */
+		SLATE_ARGUMENT( bool, SupportsTransparency )
+
+		/** The initial opacity of the window */
+		SLATE_ARGUMENT( float, InitialOpacity )
+
+		/** Is the window initially maximized */
+		SLATE_ARGUMENT( bool, IsInitiallyMaximized )
+		
+		/** How the window should be sized */
+		SLATE_ARGUMENT( ESizingRule::Type, SizingRule )
+
+		/** True if this should be a 'pop-up' window */
+		SLATE_ARGUMENT( bool, IsPopupWindow )
+
+		/** Should this window be focused immediately after it is shown? */
+		SLATE_ARGUMENT( bool, FocusWhenFirstShown )
+	
+		/** Should this window be activated immediately after it is shown? */
+		SLATE_ARGUMENT( bool, ActivateWhenFirstShown )
+
+		/** Use the default os look for the border of the window */
+		SLATE_ARGUMENT( bool, UseOSWindowBorder )
+
+		/** Can this window be maximized? */
+		SLATE_ARGUMENT( bool, SupportsMaximize )
+		
+		/** Can this window be minimized? */
+		SLATE_ARGUMENT( bool, SupportsMinimize )
+
+		/** True if we should initially create a traditional title bar area.  If false, the user must embed the title
+		    area content into the window manually, taking into account platform-specific considerations!  Has no
+			effect for certain types of windows (popups, tool-tips, etc.) */
+		SLATE_ARGUMENT( bool, CreateTitleBar )
+
+		/** If the window appears off screen or is too large to safely fit this flag will force realistic 
+		    constraints on the window and bring it back into view. */
+		SLATE_ARGUMENT( bool, SaneWindowPlacement )
+
+		SLATE_DEFAULT_SLOT( FArguments, Content )
+
+	SLATE_END_ARGS()
+
+	/**
+	 * Default constructor. Use SNew(SWindow) instead.
+	 */
+	SWindow();
+
+public:
+
+	void Construct(const FArguments& InArgs);
+
+	/**
+	 * Make a tool tip window
+	 *
+	 * @return The new SWindow
+	 */
+	static TSharedRef<SWindow> MakeToolTipWindow();
+
+	/**
+	 * Make cursor decorator window
+	 *
+	 * @return The new SWindow
+	 */
+	static TSharedRef<SWindow> MakeCursorDecorator();
+
+	/**
+	 * Make a notification window
+	 *
+	 * @return The new SWindow
+	 */
+	static TSharedRef<SWindow> MakeNotificationWindow();
+
+	/**
+	 * @param ContentSize      The size of content that we want to accommodate 
+	 *
+	 * @return The size of the window necessary to accommodate the given content */
+	static FVector2D ComputeWindowSizeForContent( FVector2D ContentSize );
+
+	/**
+	 * Grabs the current window title
+	 *
+	 * @return The window's title
+	 */
+	FText GetTitle() const
+	{
+		return Title.Get();
+	}
+
+	/**
+	 * Sets the current window title
+	 *
+	 * @param InTitle	The new title of the window
+	 */
+	void SetTitle( const FText& InTitle )
+	{
+		if (NativeWindow.IsValid())
+		{
+			Title = InTitle;
+			NativeWindow->SetText( *InTitle.ToString() );
+		}
+	}
+
+	/**
+	 * Returns the size of the title bar as a Slate size parameter.  Does not take into account application scale!
+	 *
+	 * @return  Title bar size
+	 */
+	FOptionalSize GetTitleBarSize() const;
+
+	/**	@return The initially desired screen position of the slate window */
+	FVector2D GetInitialDesiredSizeInScreen() const;
+
+	/**	@return The initially desired size of the slate window */
+	FVector2D GetInitialDesiredPositionInScreen() const;
+
+	/** Get the Geometry that describes this window. Windows in Slate are unique in that they know their own geometry. */
+	FGeometry GetWindowGeometryInScreen() const;
+
+	/** @return The geometry of the window in window space (i.e. position and AbsolutePosition are 0) */
+	FGeometry GetWindowGeometryInWindow() const;
+
+	/** @return The position of the window in screen space */
+	FVector2D GetPositionInScreen() const;
+
+	/** @return the size of the window in screen pixels */
+	FVector2D GetSizeInScreen() const;
+
+	/** @return the rectangle of the window for its non-maximized state */
+	FSlateRect GetNonMaximizedRectInScreen() const;
+
+	/** @return Rectangle that this window occupies in screen space */
+	FSlateRect GetRectInScreen() const;
+
+	/** @return the size of the window's usable client area. */
+	FVector2D GetClientSizeInScreen() const;
+
+	/** @return a clipping rectangle that represents this window in Window Space (i.e. always starts at 0,0) */
+	FSlateRect GetClippingRectangleInWindow() const;
+
+	/** Returns the margins used for the window border. This varies based on whether it's maximized or not. */
+	FMargin GetWindowBorderSize() const;
+
+	/** Relocate the window to a screenspace position specified by NewPosition */
+	void MoveWindowTo( FVector2D NewPosition );
+	/** Relocate the window to a screenspace position specified by NewPosition and resize it to NewSize */
+	void ReshapeWindow( FVector2D NewPosition, FVector2D NewSize );
+	void ReshapeWindow( const FSlateRect& InNewShape );
+	/** Resize the window to be NewSize immediately */
+	void Resize( FVector2D NewSize );
+
+	/** Returns the rectangle of the screen the window is associated with */
+	FSlateRect GetFullScreenInfo() const;
+
+	/** @return Returns true if the window is currently morphing to a new position, shape and/or opacity */
+	bool IsMorphing() const;
+	/** @return Returns true if the window is currently morphing and is morphing by size */
+	bool IsMorphingSize() const;
+	/** Animate the window to TargetOpacity and TargetPosition over a short period of time */
+	void MorphToPosition( const FCurveSequence& Sequence, const float TargetOpacity, const FVector2D& TargetPosition );
+	/** Animate the window to TargetOpacity and TargetShape over a short period of time */
+	void MorphToShape( const FCurveSequence& Sequence, const float TargetOpacity, const FSlateRect& TargetShape );
+	/** Set a new morph shape and force the morph to run for at least one frame in order to reach that target */
+	void UpdateMorphTargetShape( const FSlateRect& TargetShape );
+	/** Set a new morph position and force the morph to run for at least one frame in order to reach that target */
+	void UpdateMorphTargetPosition( const FVector2D& TargetPosition );
+	/** @return Returns the currently set morph target position */
+	FVector2D GetMorphTargetPosition() const;
+	/** @return Returns the currently set morph target shape */
+	FSlateRect GetMorphTargetShape() const;
+
+	/** Flashed the window, used for drawing attention to modal dialogs */
+	void FlashWindow();
+
+	/** 
+	 * Bring the window to the front 
+	 *
+	 * @param bForce	Forces the window to the top of the Z order, even if that means stealing focus from other windows
+	 *					In general do not pass force in.  It can be useful for some window types, like game windows where not forcing it to the front
+	 *					would cause mouse capture and mouse lock to happen but without the window visible
+	 */
+	void BringToFront( bool bForce = false );
+
+	/** @hack Force a window to front even if a different application is in front. */
+	void HACK_ForceToFront();
+
+	/** Sets the actual screen position of the window. THIS SHOULD ONLY BE CALLED BY THE OS */
+	void SetCachedScreenPosition(FVector2D NewPosition);
+
+	/**	Sets the actual size of the window. THIS SHOULD ONLY BE CALLED BY THE OS */
+	void SetCachedSize(FVector2D NewSize);
+
+	TSharedPtr<FGenericWindow> GetNativeWindow();
+	TSharedPtr<const FGenericWindow> GetNativeWindow() const ;
+
+	/** 
+	 * Returns whether or not this window is a descendant of the specfied parent window
+	 *
+	 * @param ParentWindow	The window to check
+	 * @return true if the window is a child of ParentWindow, false otherwise.
+	 */
+	bool IsDescendantOf( const TSharedPtr<SWindow>& ParentWindow ) const;
+
+	/**
+	 * Sets the native OS window associated with this SWindow
+	 *
+	 * @param InNativeWindow	The native window
+	 */
+	void SetNativeWindow( TSharedRef<FGenericWindow> InNativeWindow );
+	
+	/**
+	 * Sets the widget content for this window
+	 *
+	 * @param	InContent	The widget to use as content for this window
+	 */
+	void SetContent( TSharedRef<SWidget> InContent );
+
+	/**
+	 * Gets the widget content for this window
+	 *
+	 * @return	The widget content for this window
+	 */
+	TSharedRef<const SWidget> GetContent() const;
+
+	/**
+	 * Adds content to draw on top of the entire window
+	 *
+	 * @param	InZOrder	Z-order to use for this widget
+	 * 
+	 * @return The added overlay slot so that it can be configured and populated
+	 */
+	SOverlay::FOverlaySlot& AddOverlaySlot( const int32 ZOrder = INDEX_NONE );
+
+	/**
+	 * Removes a widget that is being drawn over the entire window
+	 *
+	 * @param	InContent	The widget to remove
+	 */
+	void RemoveOverlaySlot( const TSharedRef<SWidget>& InContent );
+
+	/** Return a new slot in the popup layer. Assumes that the window has a popup layer. */
+	struct FPopupLayerSlot& AddPopupLayerSlot();
+
+	/** Counterpart to AddPopupLayerSlot */
+	void RemovePopupLayerSlot( const TSharedRef<SWidget>& WidgetToRemove );
+
+	/**
+	 * Sets a widget to use as a full window overlay, or clears an existing widget if set.  When set, this widget will be drawn on top of all other window content.
+	 *
+	 * @param	InContent	The widget to use for full window overlay content, or nullptr for no overlay
+	 */
+	void SetFullWindowOverlayContent( TSharedPtr<SWidget> InContent );
+
+	/**
+	 * Begins a transition from showing regular window content to overlay content
+	 * During the transition we show both sets of content
+	 */
+	void BeginFullWindowOverlayTransition();
+
+	/**
+	 * Ends a transition from showing regular window content to overlay content
+	 * When this is called content occluded by the full window overlay(if there is one) will be physically hidden
+	 */
+	void EndFullWindowOverlayTransition();
+
+	/**
+	 * Checks to see if there is content assigned as a full window overlay
+	 *
+	 * @return	True if there is an overlay widget assigned
+	 */
+	bool HasFullWindowOverlayContent() const;
+
+	/** @return should this window show up in the taskbar */
+	bool AppearsInTaskbar() const;
+
+	void SetOnWindowDeactivated( const FOnWindowDeactivated& InDelegate );
+
+	/** Sets the delegate to execute right before the window is closed */
+	void SetOnWindowClosed( const FOnWindowClosed& InDelegate );
+
+	/** Sets the delegate to execute right after the window has been moved */
+	void SetOnWindowMoved( const FOnWindowMoved& InDelegate);
+
+	/** Sets the delegate to override RequestDestroyWindow */
+	void SetRequestDestroyWindowOverride( const FRequestDestroyWindowOverride& InDelegate );
+
+	/** Request that this window be destroyed. The window is not destroyed immediately. Instead it is placed in a queue for destruction on next Tick */
+	void RequestDestroyWindow();
+
+	/** Warning: use Request Destroy Window whenever possible!  This method destroys the window immediately! */
+	void DestroyWindowImmediately();
+ 
+	/** Calls the OnWindowClosed delegate when this window is about to be closed */
+	void NotifyWindowBeingDestroyed();
+
+	/** Make the window visible */
+	void ShowWindow();
+
+	/** Make the window invisible */
+	void HideWindow();
+
+	/**
+	 * Enables or disables this window and all of its children
+	 *
+	 * @param bEnable	true to enable this window and its children false to diable this window and its children
+	 */
+	void EnableWindow( bool bEnable );
+
+	/** Toggle window between window modes (fullscreen, windowed, etc) */
+	void SetWindowMode( EWindowMode::Type WindowMode );
+
+	/** @return The current window mode (fullscreen, windowed, etc) */
+	EWindowMode::Type GetWindowMode() const { return NativeWindow->GetWindowMode(); }
+
+	/** @return true if the window is visible, false otherwise*/
+	bool IsVisible() const;
+
+	/** @return true if the window is maximized, false otherwise*/
+	bool IsWindowMaximized() const;
+
+	/** Maximize the window if bInitiallyMaximized is set */
+	void InitialMaximize();
+
+	/**
+	 * Sets the opacity of this window
+	 *
+	 * @param	InOpacity	The new window opacity represented as a floating point scalar
+	 */
+	void SetOpacity( const float InOpacity );
+
+	/** @return the window's current opacity */
+	float GetOpacity() const;
+
+	/** @return true if the window supports transparency */
+	bool SupportsTransparency() const;
+
+	/** @return A String representation of the widget */
+	virtual FString ToString() const;
+
+	/**
+	 * Sets a widget that should become focused when this window is next activated
+	 *
+	 * @param	InWidget	The widget to set focus to when this window is activated
+	 */
+	void SetWidgetToFocusOnActivate( TSharedPtr< SWidget > InWidget )
+	{
+		WidgetToFocusOnActivate = InWidget;
+	}
+
+	/** @return true if the window should be activated when first shown */
+	bool ActivateWhenFirstShown() const;
+
+	/** @return true if the window accepts input; false if the window is non-interactive */
+	bool AcceptsInput() const;
+
+	/** @return true if the user decides the size of the window */
+	bool IsUserSized() const;
+
+	/** @return true if the window is sized by the windows content */
+	bool IsAutosized() const;
+
+	/** Should this window automatically derive its size based on its content or be user-drive? */
+	void SetSizingRule( ESizingRule::Type InSizingRule );
+
+	/** @return true if this is a vanilla window, or one being used for some special purpose: e.g. tooltip or menu */
+	bool IsRegularWindow() const;
+
+	/** @return true if the window should be on top of all other windows; false otherwise */
+	bool IsTopmostWindow() const;
+
+	/** @return True if we expect the window size to change frequently. See description of bSizeWillChangeOften member variable. */
+	bool SizeWillChangeOften() const
+	{
+		return bSizeWillChangeOften;
+	}
+
+	/** @return Returns the configured expected maximum width of the window, or INDEX_NONE if not specified.  Can be used to optimize performance for window size animation */
+	int32 GetExpectedMaxWidth() const
+	{
+		return ExpectedMaxWidth;
+	}
+
+	/** @return Returns the configured expected maximum height of the window, or INDEX_NONE if not specified.  Can be used to optimize performance for window size animation */
+	int32 GetExpectedMaxHeight() const
+	{
+		return ExpectedMaxHeight;
+	}
+
+	/** @return true if the window is using the os window border instead of a slate created one */
+	bool HasOSWindowBorder() const { return bHasOSWindowBorder; }
+
+	/** @return true if mouse coordinates is within this window */
+	bool IsScreenspaceMouseWithin(FVector2D ScreenspaceMouseCoordinate) const;
+
+	/** @return true if this is a user-sized window with a thick edge */
+	bool HasSizingFrame() const;
+
+	/** @return true if this window has a maximize button/box on the titlebar area */
+	bool HasMaximizeBox() const;
+
+	/** @return true if this window has a minimize button/box on the titlebar area */
+	bool HasMinimizeBox() const;
+
+	/** Set modal window related flags - called by Slate app code during FSlateApplication::AddModalWindow() */
+	void SetAsModalWindow()
+	{
+		bHasMaximizeButton = false;
+		bHasMinimizeButton = false;
+	}
+
+	bool HasMaximizeButton( )
+	{
+		return bHasMaximizeButton;
+	}
+
+	bool HasMinimizeButton( )
+	{
+		return bHasMinimizeButton;
+	}
+
+	void SetTitleBar( const TSharedPtr<IWindowTitleBar> InTitleBar )
+	{
+		TitleBar = InTitleBar;
+	}
+
+	// Events
+	virtual FCursorReply OnCursorQuery( const FGeometry& MyGeometry, const FPointerEvent& CursorEvent ) const;
+
+	/** The system will call this method to notify the window that it has been places in the foreground or background. */
+	virtual bool OnIsActiveChanged( const FWindowActivateEvent& ActivateEvent );
+	
+	/** Windows functions */
+	void Maximize();
+	void Restore();
+	void Minimize();
+
+	/** Gets the current Window Zone that mouse position is over. */
+	EWindowZone::Type GetCurrentWindowZone(FVector2D LocalMousePosition);
+
+	/** @return Gets the radius of the corner rounding of the window. */
+	int32 GetCornerRadius();
+
+	virtual bool SupportsKeyboardFocus() const OVERRIDE;
+
+	virtual FReply OnKeyboardFocusReceived( const FGeometry& MyGeometry, const FKeyboardFocusEvent& InKeyboardFocusEvent );
+
+	virtual FReply OnMouseButtonDown( const FGeometry& MyGeometry, const FPointerEvent& MouseEvent ) OVERRIDE;
+	virtual FReply OnMouseButtonUp( const FGeometry& MyGeometry, const FPointerEvent& MouseEvent ) OVERRIDE;
+	virtual FReply OnMouseMove( const FGeometry& MyGeometry, const FPointerEvent& MouseEvent ) OVERRIDE;
+
+	/** The window's desired size takes into account the ratio between the slate units and the pixel size */
+	virtual FVector2D ComputeDesiredSize() const OVERRIDE;
+
+	/** Marks this window as having been drawn at least once */
+	inline void MarkAsDrawn()
+	{
+		bHasEverBeenDrawn = true;
+	}
+
+	/** @return Returns True if the window has even been drawn */
+	inline bool HasEverBeenDrawn() const
+	{
+		return bHasEverBeenDrawn;
+	}
+
+	/** @return true if this window will be focused when it is first shown */
+	inline bool IsFocusedInitially() const
+	{
+		return bFocusWhenFirstShown;
+	}
+
+	/** @return the list of this window's child windows */
+	const TArray< TSharedRef<SWindow> >& GetChildWindows() const;
+	
+	/** @return the list of this window's child windows */
+	TArray< TSharedRef<SWindow> >& GetChildWindows();
+
+	/** Add ChildWindow as this window's child */
+	void AddChildWindow( const TSharedRef<SWindow>& ChildWindow );
+
+	/** @return the parent of this window; Invalid shared pointer if this window is not a child */
+	TSharedPtr<SWindow> GetParentWindow() const;
+
+	/** Look up the parent chain until we find the top-level window that owns this window */
+	TSharedPtr<SWindow> GetTopmostAncestor();
+
+	/** Remove DescendantToRemove from this window's children or their children. */
+	bool RemoveDescendantWindow( const TSharedRef<SWindow>& DescendantToRemove );
+
+	/** Sets the delegate to call when switching worlds in before ticking,drawing, or sending messages to widgets in this window */
+	void SetOnWorldSwitchHack( FOnSwitchWorldHack& InOnWorldSwitchHack );
+
+	/**
+	 * Static: Marks the specified window as drawn
+	 *
+	 * @param	Window	The window to mark as drawn
+	 */
+	static void MarkWindowAsDrawn( SWindow* Window )
+	{
+		Window->MarkAsDrawn();
+	}
+
+	/**
+	 * @return Returns a delegate that can be used to mark window as drawn. The delegate is safe to call from a different thread.
+	 */
+	FSimpleDelegate MakeMarkWindowAsDrawnDelegate()
+	{
+		return FSimpleDelegate::CreateStatic( &MarkWindowAsDrawn, this );
+	}
+
+	/**
+	 * Hack to switch worlds
+	 *
+	 * @param WorldId: User ID for a world that should be restored or -1 if no restore
+	 * @param The ID of the world restore later
+	 */
+	int32 SwitchWorlds( int32 WorldId ) const;
+
+	/** Are any of our child windows active? */
+	bool HasActiveChildren() const;
+
+public:
+
+	// Begin SWidget overrides
+
+	virtual void Tick( const FGeometry& AllottedGeometry, const double InCurrentTime, const float InDeltaTime ) OVERRIDE;
+
+	// End SWidget overrides
+
+protected:
+
+	/** Get the desired color of titlebar items. These change during flashing. */
+	FSlateColor GetWindowTitleContentColor() const;
+
+	/** Kick off a morph to whatever the target shape happens to be. */
+	void StartMorph();
+
+	/** Get the brush used to draw the window background */
+	const FSlateBrush* GetWindowBackground() const;
+
+	/** Get the brush used to draw the window outline */
+	const FSlateBrush* GetWindowOutline() const;
+
+	/** Get the color used to tint the window outline */
+	FSlateColor GetWindowOutlineColor() const;
+
+protected:
+
+	/** Title of the window, displayed in the title bar as well as potentially in the task bar (Windows platform) */
+	TAttribute<FText> Title;
+
+	/** When true, grabbing anywhere on the window will allow it to be dragged. */
+	bool bDragAnywhere;
+
+	/** Current opacity of the window */
+	float Opacity;
+
+	/** How to size the window */
+	ESizingRule::Type SizingRule;
+
+	/** How to auto center the window */
+	EAutoCenter::Type AutoCenterRule;
+
+	/** true if the window supports transparency */
+	bool bIsTransparent : 1;
+
+	/** True if this is a pop up window */
+	bool bIsPopupWindow : 1;
+
+	/** True if this is a tool tip window */
+	bool bIsToolTipWindow : 1;
+
+	/** True if this is a topmost window */
+	bool bIsTopmostWindow : 1;
+
+	/** True if we expect the size of this window to change often, such as if its animated, or if it recycled for tool-tips,
+	    and we'd like to avoid costly GPU buffer resizes when that happens.  Enabling this may incur memory overhead or
+		other platform-specific side effects */
+	bool bSizeWillChangeOften : 1;
+
+	/** Whether this window is used to draw content next to the cursor; usually for drag and drop purposes. */
+	bool bIsCursorDecoratorWindow : 1;
+
+	/** true if this window is maximized when its created */
+	bool bInitiallyMaximized : 1;
+
+	/** True if this window has been shown yet */
+	bool bHasEverBeenShown : 1;
+
+	/** True if this window has ever been drawn */
+	bool bHasEverBeenDrawn : 1;
+
+	/** Focus this window immediately as it is shown */
+	bool bFocusWhenFirstShown : 1;
+
+	/** Activate this window immediately as it is shown */
+	bool bActivateWhenFirstShown : 1;
+
+	/** True if this window displays the os window border instead of drawing one in slate */
+	bool bHasOSWindowBorder : 1;
+	
+	/** True if this window displays an enabled minimize button on the toolbar area */
+	bool bHasMinimizeButton : 1;
+
+	/** True if this window displays an enabled maximize button on the toolbar area */
+	bool bHasMaximizeButton : 1;
+
+	/** True if this window displays thick edge that can be used to resize the window */
+	bool bHasSizingFrame : 1;
+
+	/** Initial desired position of the window's content in screen space */
+	FVector2D InitialDesiredScreenPosition;
+
+	/** Initial desired size of the window's content in screen space */
+	FVector2D InitialDesiredSize;
+
+	/** Position of the window's content in screen space */
+	FVector2D ScreenPosition;
+
+	/** The position of the window before entering fullscreen */
+	FVector2D PreFullscreenPosition;
+
+	/** Size of the window's content area in screen space */
+	FVector2D Size;
+
+	/** Size of this window's title bar.  Can be zero.  Set at construction and should not be changed afterwards. */
+	float TitleBarSize;
+
+	/** Utility for animating the window size. */
+	struct FMorpher
+	{
+		FMorpher()
+			: StartingMorphShape( FSlateRect(0,0,100,100) )
+			, TargetMorphShape( FSlateRect(0,0,100,100) )
+			, bIsActive(false)
+			, bIsPendingPlay(false)
+			, bIsAnimatingWindowSize(false)
+		{
+		}
+
+		/** Initial window opacity */
+		float StartingOpacity;
+		/** Desired opacity of the window */
+		float TargetOpacity;
+
+		/** Initial size of the window (i.e. at the start of animation) */
+		FSlateRect StartingMorphShape;
+		/** Desired size of the window (i.e. at the end of the animation) */
+		FSlateRect TargetMorphShape;
+		
+		/** Animation sequence to hold on to the Handle */
+		FCurveSequence Sequence;
+
+		/** True if this morph is currently active */
+		bool bIsActive : 1;
+
+		/** True if this morph sequence is pending play (i.e. will start on next Tick) */
+		bool bIsPendingPlay : 1;
+
+		/** True if we're morphing size as well as position.  False if we're just morphing position */
+		bool bIsAnimatingWindowSize : 1;
+
+	} Morpher;
+
+	/** Cached "zone" the cursor was over in the window the last time that someone called GetCurrentWindowZone() */
+	EWindowZone::Type WindowZone;
+	
+
+	TSharedPtr<SWidget> TitleArea;
+	SVerticalBox::FSlot* ContentSlot;
+
+	/** Widget to transfer keyboard focus to when this window becomes active, if any.  This is used to
+	    restore focus to a widget after a popup has been dismissed. */
+	TWeakPtr< SWidget > WidgetToFocusOnActivate;
+
+	/** Style used to draw this window */
+	const FWindowStyle* Style;
+	const FSlateBrush* WindowBackground;
+
+private:
+	/** The native window that is backing this Slate Window */
+	TSharedPtr<FGenericWindow> NativeWindow;
+	
+	/** Invoked when the window has been deactivated. */
+	FOnWindowDeactivated OnWindowDeactivated;
+	
+	/** Invoked when the window is about to be closed. */
+	FOnWindowClosed OnWindowClosed;
+
+	/** Invoked when the window is moved */
+	FOnWindowMoved OnWindowMoved;
+
+	/** Invoked when the window is requested to be destroyed. */
+	FRequestDestroyWindowOverride RequestDestroyWindowOverride;
+
+	/** Window overlay widget */
+	TSharedPtr<SOverlay> WindowOverlay;	
+	
+	/**
+	 * This layer provides mechanism for tooltips, drag-drop
+	 * decorators, and popups without creating a new window.
+	 */
+	TSharedPtr<class SPopupLayer> PopupLayer;
+
+	/** Full window overlay widget */
+	TSharedPtr<SWidget> FullWindowOverlayWidget;
+
+
+
+	/** When not null, this window will always appear on top of the parent and be closed when the parent is closed. */
+	TWeakPtr<SWindow> ParentWindowPtr;
+
+	/** Child windows of this window */
+	TArray< TSharedRef<SWindow> > ChildWindows;
+	
+	/** World switch delegate */
+	FOnSwitchWorldHack OnWorldSwitchHack;
+
+	/** 
+	 * Whether or not we should show content of the window which could be occluded by full screen window content. 
+	 * This is used to hide content when there is a full screen overlay occluding it
+	 */
+	bool bShouldShowWindowContentDuringOverlay;
+
+	/** The expected maximum width of the window.  May be used for performance optimization when bSizeWillChangeOften is set. */
+	int32 ExpectedMaxWidth;
+
+	/** The expected maximum height of the window.  May be used for performance optimization when bSizeWillChangeOften is set. */
+	int32 ExpectedMaxHeight;
+
+	// The window title bar.
+	TSharedPtr<IWindowTitleBar> TitleBar;
+
+private:
+	
+	void ConstructWindowInternals( const bool bCreateTitleBar );
+
+	/**
+	 * @return EVisibility::Visible if we are showing this viewports content.  EVisibility::Hidden otherwise (we hide the content during full screen overlays)
+	 */
+	EVisibility GetWindowContentVisibility() const;
+
+	/**
+	 * @return EVisibility::Visible if the window is flashing. Used to show/hide the white flash in the title area
+	 */
+	EVisibility GetWindowFlashVisibility() const;
+};
+
+/**
+ * Popups, tooltips, drag and drop decorators all can be executed without creating a new window.
+ * This slot along with the SWindow::AddPopupLayerSlot() API enabled it.
+ */
+struct FPopupLayerSlot
+{
+public:
+	FPopupLayerSlot()
+	: DesktopPosition_Attribute(FVector2D::ZeroVector)
+	, WidthOverride_Attribute()
+	, HeightOverride_Attribute()
+	, Scale_Attribute(1.0f)
+	, Clamp_Attribute(false)
+	, ClampBuffer_Attribute(FVector2D::ZeroVector)
+	, Widget(SNullWidget::NullWidget)
+	{}
+
+	/** Support for using brackets in slate declarative syntax */
+	FPopupLayerSlot& operator[]( TSharedRef<SWidget> InWidget )
+	{
+		Widget = InWidget;
+		return *this;
+	}
+
+	/** Pixel position in desktop space */
+	FPopupLayerSlot& DesktopPosition( const TAttribute<FVector2D>& InDesktopPosition )
+	{
+		DesktopPosition_Attribute = InDesktopPosition;
+		return *this;
+	}
+
+	/** Width override in pixels */
+	FPopupLayerSlot& WidthOverride( const TAttribute<float>& InWidthOverride )
+	{
+		WidthOverride_Attribute = InWidthOverride;
+		return *this;
+	}
+
+	/** Width override in pixels */
+	FPopupLayerSlot& HeightOverride( const TAttribute<float>& InHeightOverride )
+	{
+		HeightOverride_Attribute = InHeightOverride;
+		return *this;
+	}
+
+	/** DPI scaling to be applied to the contents of this slot */
+	FPopupLayerSlot& Scale( const TAttribute<float>& InScale )
+	{
+		Scale_Attribute = InScale;
+		return *this;
+	}
+
+	/** Should this slot be kept within the parent window */
+	FPopupLayerSlot& ClampToWindow( const TAttribute<bool>& InClamp_Attribute)
+	{
+		Clamp_Attribute = InClamp_Attribute;
+		return *this;
+	}
+
+	/** If this slot is kept within the parent window, how far from the edges should we clamp it */
+	FPopupLayerSlot& ClampBuffer( const TAttribute<FVector2D>& InClampBuffer_Attribute )
+	{
+		ClampBuffer_Attribute = InClampBuffer_Attribute;
+		return *this;
+	}
+
+	/** @return the widget present in this slot */
+	const TSharedRef<SWidget>& GetWidget()
+	{
+		return Widget;
+	}
+
+private:
+	/** SPopupLayer arranges FPopupLayerSlots, so it needs to know all about */
+	friend class SPopupLayer;
+	/** TPanelChildren need access to the Widget member */
+	friend class TPanelChildren<FPopupLayerSlot>;
+
+	TAttribute<FVector2D> DesktopPosition_Attribute;
+	TAttribute<float> WidthOverride_Attribute;
+	TAttribute<float> HeightOverride_Attribute;
+	TAttribute<float> Scale_Attribute;
+	TAttribute<bool> Clamp_Attribute;
+	TAttribute<FVector2D> ClampBuffer_Attribute;
+	TSharedRef<SWidget> Widget;
+};
+
+
+/**
+ * Hack to switch wolds in a scope and switch back when we fall out of scope                   
+ */
+struct FScopedSwitchWorldHack
+{
+	FScopedSwitchWorldHack( FWidgetPath& WidgetPath )
+		: Window( WidgetPath.TopLevelWindow )
+		, WorldId( -1 )
+	{
+		if( Window.IsValid() )
+		{
+			WorldId = Window->SwitchWorlds( WorldId );
+		}
+	}
+
+	FScopedSwitchWorldHack( TSharedPtr<SWindow> InWindow )
+		: Window( InWindow )
+		, WorldId( -1 )
+	{
+		if( Window.IsValid() )
+		{
+			WorldId = Window->SwitchWorlds( WorldId );
+		}
+	}
+
+	~FScopedSwitchWorldHack()
+	{
+		if( Window.IsValid() )
+		{
+			Window->SwitchWorlds( WorldId );
+		}
+	}
+
+private:
+
+	// The window to switch worlds for.
+	TSharedPtr<SWindow> Window;
+
+	// The worldID serves as identification to the user about the world.  It can be anything although -1 is assumed to be always invalid.
+	int32 WorldId;
+};
diff --git a/Engine/Source/Runtime/SlateCore/SlateCore.Build.cs b/Engine/Source/Runtime/SlateCore/SlateCore.Build.cs
new file mode 100644
--- /dev/null
+++ b/Engine/Source/Runtime/SlateCore/SlateCore.Build.cs
@@ -0,0 +1,48 @@
+// Copyright 1998-2014 Epic Games, Inc. All Rights Reserved.
+
+using UnrealBuildTool;
+
+public class SlateCore : ModuleRules
+{
+	public SlateCore(TargetInfo Target)
+	{
+		PublicDependencyModuleNames.AddRange(
+			new string[] { 
+				"Core",
+				"CoreUObject",
+				"InputCore",
+			}
+		);
+
+		PrivateIncludePaths.AddRange(
+			new string[] {
+				"Runtime/SlateCore/Private",
+				"Runtime/SlateCore/Private/Animation",
+				"Runtime/SlateCore/Private/Application",
+				"Runtime/SlateCore/Private/Brushes",
+				"Runtime/SlateCore/Private/Commands",
+				"Runtime/SlateCore/Private/Fonts",
+				"Runtime/SlateCore/Private/Input",
+				"Runtime/SlateCore/Private/Layout",
+				"Runtime/SlateCore/Private/Logging",
+				"Runtime/SlateCore/Private/Rendering",
+				"Runtime/SlateCore/Private/Sound",
+				"Runtime/SlateCore/Private/Styling",
+				"Runtime/SlateCore/Private/Textures",
+				"Runtime/SlateCore/Private/Types",
+				"Runtime/SlateCore/Private/Widgets",
+			}
+		);
+
+        if (!UEBuildConfiguration.bBuildDedicatedServer)
+        {
+			AddThirdPartyPrivateStaticDependencies(Target, "FreeType2");
+        }
+
+		if ((Target.Platform == UnrealTargetPlatform.Win64) ||
+			(Target.Platform == UnrealTargetPlatform.Win32))
+		{
+			AddThirdPartyPrivateStaticDependencies(Target, "XInput");
+		}
+	}
+}
diff --git a/Engine/Source/Runtime/SlateRHIRenderer/Private/SlateRHIRenderer.cpp b/Engine/Source/Runtime/SlateRHIRenderer/Private/SlateRHIRenderer.cpp
--- a/Engine/Source/Runtime/SlateRHIRenderer/Private/SlateRHIRenderer.cpp
+++ b/Engine/Source/Runtime/SlateRHIRenderer/Private/SlateRHIRenderer.cpp
@@ -5,12 +5,11 @@
 #include "SlateRHIRenderer.h"
 #include "SlateRHIFontTexture.h"
 #include "SlateRHIResourceManager.h"
-#include "Runtime/Slate/Public/ElementBatcher.h"
+#include "SlateCore.h"
 #include "SlateRHIRenderingPolicy.h"
 #include "Runtime/Engine/Public/ScreenRendering.h"
 #include "Runtime/Engine/Public/ShaderCompiler.h"
 #include "SlateShaders.h"
-#include "FontCache.h"
 
 DECLARE_CYCLE_STAT(TEXT("Map Staging Buffer"),STAT_MapStagingBuffer,STATGROUP_CrashTracker);
 DECLARE_CYCLE_STAT(TEXT("Generate Capture Buffer"),STAT_GenerateCaptureBuffer,STATGROUP_CrashTracker);
diff --git a/Engine/Source/Runtime/SlateRHIRenderer/SlateRHIRenderer.Build.cs b/Engine/Source/Runtime/SlateRHIRenderer/SlateRHIRenderer.Build.cs
--- a/Engine/Source/Runtime/SlateRHIRenderer/SlateRHIRenderer.Build.cs
+++ b/Engine/Source/Runtime/SlateRHIRenderer/SlateRHIRenderer.Build.cs
@@ -14,6 +14,7 @@ public class SlateRHIRenderer : ModuleRules
 				"CoreUObject",
                 "InputCore",
 				"Slate",
+				"SlateCore",
                 "Engine",
                 "RHI",
                 "RenderCore",
@@ -21,7 +22,5 @@ public class SlateRHIRenderer : ModuleRules
                 "ImageWrapper"
 			}
 		);
-
-
 	}
 }
diff --git a/Engine/Source/Runtime/UMG/Public/UMG.h b/Engine/Source/Runtime/UMG/Public/UMG.h
--- a/Engine/Source/Runtime/UMG/Public/UMG.h
+++ b/Engine/Source/Runtime/UMG/Public/UMG.h
@@ -15,54 +15,58 @@
 #include "Engine.h"
 #include "InputCore.h"
 #include "Slate.h"
+#include "SlateCore.h"
+
 
 /* Components
  *****************************************************************************/
 
-#include "Components/SlateWrapperTypes.h"
+#include "SlateWrapperTypes.h"
+
+#include "PanelSlot.h"
 
-#include "Components/PanelSlot.h"
+#include "SlateWrapperComponent.h"
 
-#include "Components/SlateWrapperComponent.h"
+#include "SlateLeafWidgetComponent.h"
+#include "SlateNonLeafWidgetComponent.h"
 
-#include "Components/SlateLeafWidgetComponent.h"
-#include "Components/SlateNonLeafWidgetComponent.h"
+#include "ContentWidget.h"
 
-#include "Components/ContentWidget.h"
+#include "BoxPanelComponent.h"
 
-#include "Components/BoxPanelComponent.h"
+#include "CanvasPanelSlot.h"
+#include "CanvasPanelComponent.h"
 
-#include "Components/CanvasPanelSlot.h"
-#include "Components/CanvasPanelComponent.h"
+#include "HorizontalBoxSlot.h"
+#include "HorizontalBoxComponent.h"
 
-#include "Components/HorizontalBoxSlot.h"
-#include "Components/HorizontalBoxComponent.h"
+#include "VerticalBoxSlot.h"
+#include "VerticalBoxComponent.h"
 
-#include "Components/VerticalBoxSlot.h"
-#include "Components/VerticalBoxComponent.h"
+#include "BorderComponent.h"
+#include "ButtonComponent.h"
+#include "CheckBoxComponent.h"
+#include "CircularThrobberComponent.h"
+#include "EditableText.h"
+#include "EditableTextBlockComponent.h"
+#include "GridPanelComponent.h"
+#include "ImageComponent.h"
+#include "ScrollBoxComponent.h"
+#include "SpacerComponent.h"
+#include "SpinningImageComponent.h"
+#include "TextBlockComponent.h"
+#include "ThrobberComponent.h"
 
-#include "Components/BorderComponent.h"
-#include "Components/ButtonComponent.h"
-#include "Components/CheckBoxComponent.h"
-#include "Components/CircularThrobberComponent.h"
-#include "Components/EditableText.h"
-#include "Components/EditableTextBlockComponent.h"
-#include "Components/GridPanelComponent.h"
-#include "Components/ImageComponent.h"
-#include "Components/ScrollBoxComponent.h"
-#include "Components/SpacerComponent.h"
-#include "Components/SpinningImageComponent.h"
-#include "Components/TextBlockComponent.h"
-#include "Components/ThrobberComponent.h"
 
 /* Blueprint
  *****************************************************************************/
 
-#include "Blueprint/WidgetTree.h"
-#include "Blueprint/UserWidget.h"
-#include "Blueprint/WidgetBlueprint.h"
-#include "Blueprint/WidgetBlueprintGeneratedClass.h"
+#include "WidgetTree.h"
+#include "UserWidget.h"
+#include "WidgetBlueprint.h"
+#include "WidgetBlueprintGeneratedClass.h"
  
+
 /* Interfaces
  *****************************************************************************/
 
