SlateCore: fixed FTouchKeySet static member linkage; fixed spelling errors, formatting

diff --git a/Engine/Source/Runtime/SlateCore/Private/Input/Events.cpp b/Engine/Source/Runtime/SlateCore/Private/Input/Events.cpp
new file mode 100644
--- /dev/null
+++ b/Engine/Source/Runtime/SlateCore/Private/Input/Events.cpp
@@ -0,0 +1,14 @@
+// Copyright 1998-2014 Epic Games, Inc. All Rights Reserved.
+
+/*=============================================================================
+	Events.cpp: Implements various input event related types.
+=============================================================================*/
+
+#include "SlateCorePrivatePCH.h"
+
+
+/* Static initialization
+ *****************************************************************************/
+
+const FTouchKeySet FTouchKeySet::StandardSet(EKeys::LeftMouseButton);
+const FTouchKeySet FTouchKeySet::EmptySet(EKeys::Invalid);
diff --git a/Engine/Source/Runtime/SlateCore/Public/Input/Events.h b/Engine/Source/Runtime/SlateCore/Public/Input/Events.h
--- a/Engine/Source/Runtime/SlateCore/Public/Input/Events.h
+++ b/Engine/Source/Runtime/SlateCore/Public/Input/Events.h
@@ -1,5 +1,9 @@
 // Copyright 1998-2014 Epic Games, Inc. All Rights Reserved.
 
+/*=============================================================================
+	Events.h: Declares various input event related types.
+=============================================================================*/
+
 #pragma once
 
 class SWindow;
@@ -39,7 +43,6 @@ namespace EKeyboardFocusCause
  */
 class FKeyboardFocusEvent
 {
-
 public:
 
 	/**
@@ -48,7 +51,7 @@ public:
 	 * @param  InCause  The cause of the focus event
 	 */
 	FKeyboardFocusEvent( const EKeyboardFocusCause::Type InCause )
-		: Cause( InCause )
+		: Cause(InCause)
 	{ }
 
 
@@ -62,7 +65,6 @@ public:
 		return Cause;
 	}
 
-
 private:
 
 	/** The cause of the focus change */
@@ -94,6 +96,8 @@ public:
 	 */
 	virtual ~FInputEvent( ) { }
 
+public:
+
 	/**
 	 * Returns whether or not this character is an auto-repeated keystroke
 	 *
@@ -213,13 +217,13 @@ public:
 
 protected:
 
-	/** State of modifier keys when this event happened. */
+	// State of modifier keys when this event happened.
 	FModifierKeysState ModifierKeys;
 
-	/** True if this key was auto-repeated */
+	// True if this key was auto-repeated.
 	bool bIsRepeat;
 
-	/** Events are sent along paths. See GetEventPath() */
+	// Events are sent along paths. See GetEventPath().
 	const FWidgetPath* EventPath;
 };
 
@@ -228,7 +232,8 @@ protected:
  * FKeyboardEvent describes a keyboard action (key pressed or released.)
  * It is passed to event handlers dealing with keyboard input.
  */
-class FKeyboardEvent : public FInputEvent
+class FKeyboardEvent
+	: public FInputEvent
 {
 public:
 
@@ -242,10 +247,11 @@ public:
 	FKeyboardEvent( const FKey InKey,
 					const FModifierKeysState& InModifierKeys, 
 					const bool bInIsRepeat,
-					const uint32 InCharacterCode )
-		: FInputEvent( InModifierKeys, bInIsRepeat )
-		, Key( InKey )
-		, CharacterCode( InCharacterCode )
+					const uint32 InCharacterCode
+	)
+		: FInputEvent(InModifierKeys, bInIsRepeat)
+		, Key(InKey)
+		, CharacterCode(InCharacterCode)
 	{ }
 	
 	/**
@@ -269,9 +275,11 @@ public:
 	}
 
 private:
-	/** Name of the key that was pressed */
+
+	// Name of the key that was pressed.
 	FKey Key;
-	/** The character code of the key that was pressed.  Only applicable to typed character keys, 0 otherwise*/
+
+	// The character code of the key that was pressed.  Only applicable to typed character keys, 0 otherwise.
 	uint32 CharacterCode;
 };
 
@@ -279,12 +287,13 @@ private:
 /**
  * FCharacterEvent describes a keyboard action where the utf-16 code is given.  Used for OnKeyChar messages
  */
-class FCharacterEvent : public FInputEvent
+class FCharacterEvent
+	: public FInputEvent
 {
 public:
 	FCharacterEvent( const TCHAR InCharacter, const FModifierKeysState& InModifierKeys, const bool bInIsRepeat )
-		: FInputEvent( InModifierKeys, bInIsRepeat )
-		, Character( InCharacter )
+		: FInputEvent(InModifierKeys, bInIsRepeat)
+		, Character(InCharacter)
 	{ }
 
 	/**
@@ -299,7 +308,7 @@ public:
 
 private:
 
-	/** The character that was pressed */
+	// The character that was pressed.
 	TCHAR Character;
 };
 
@@ -307,14 +316,28 @@ private:
 /**
  * Helper class to auto-populate a set with the set of "keys" a touch represents
  */
-class FTouchKeySet : public TSet<FKey>
+class FTouchKeySet
+	: public TSet<FKey>
 {
 public:
 
+	/**
+	 * Creates and initializes a new instance with the specified key.
+	 *
+	 * @param Key The key to insert into the set.
+	 */
 	FTouchKeySet(FKey Key)
 	{
 		this->Add(Key);
 	}
+
+public:
+
+	// The standard set consists of just the left mouse button key.
+	SLATECORE_API static const FTouchKeySet StandardSet;
+
+	// The empty set contains no valid keys.
+	SLATECORE_API static const FTouchKeySet EmptySet;
 };
 
 
@@ -322,7 +345,8 @@ public:
  * FPointerEvent describes a mouse or touch action (e.g. Press, Release, Move, etc).
  * It is passed to event handlers dealing with pointer-based input.
  */
-class FPointerEvent : public FInputEvent
+class FPointerEvent
+	: public FInputEvent
 {
 public:
 
@@ -334,19 +358,19 @@ public:
 		FKey InEffectingButton,
 		float InWheelDelta,
 		const FModifierKeysState& InModifierKeys
-		)
-		: FInputEvent( InModifierKeys, false )
-		, ScreenSpacePosition( InScreenSpacePosition )
-		, LastScreenSpacePosition( InLastScreenSpacePosition )
-		, CursorDelta( InScreenSpacePosition - InLastScreenSpacePosition )
-		, PressedButtons( InPressedButtons )
-		, EffectingButton (InEffectingButton )
-		, UserIndex( 0 )
-		, PointerIndex( 0 )
-		, TouchpadIndex( 0 )
-		, bIsTouchEvent( false )
-		, GestureType( EGestureEvent::None )
-		, WheelOrGestureDelta( 0, InWheelDelta )
+	)
+		: FInputEvent(InModifierKeys, false)
+		, ScreenSpacePosition(InScreenSpacePosition)
+		, LastScreenSpacePosition(InLastScreenSpacePosition)
+		, CursorDelta(InScreenSpacePosition - InLastScreenSpacePosition)
+		, PressedButtons(InPressedButtons)
+		, EffectingButton(InEffectingButton)
+		, UserIndex(0)
+		, PointerIndex(0)
+		, TouchpadIndex(0)
+		, bIsTouchEvent(false)
+		, GestureType(EGestureEvent::None)
+		, WheelOrGestureDelta(0.0f, InWheelDelta)
 	{ }
 
 	/** A constructor for raw mouse events */
@@ -356,18 +380,18 @@ public:
 		const FVector2D& InDelta,
 		const TSet<FKey>& InPressedButtons,
 		const FModifierKeysState& InModifierKeys
-		)
-		: FInputEvent( InModifierKeys, false )
-		, ScreenSpacePosition( InScreenSpacePosition )
-		, LastScreenSpacePosition( InLastScreenSpacePosition )
-		, CursorDelta( InDelta )
-		, PressedButtons( InPressedButtons )
-		, UserIndex( 0 )
-		, PointerIndex( 0 )
-		, TouchpadIndex( 0 )
-		, bIsTouchEvent( false )
-		, GestureType( EGestureEvent::None )
-		, WheelOrGestureDelta( 0, 0 )
+	)
+		: FInputEvent(InModifierKeys, false)
+		, ScreenSpacePosition(InScreenSpacePosition)
+		, LastScreenSpacePosition(InLastScreenSpacePosition)
+		, CursorDelta(InDelta)
+		, PressedButtons(InPressedButtons)
+		, UserIndex(0)
+		, PointerIndex(0)
+		, TouchpadIndex(0)
+		, bIsTouchEvent(false)
+		, GestureType(EGestureEvent::None)
+		, WheelOrGestureDelta(0.0f, 0.0f)
 	{ }
 
 	/** A constructor for touch events */
@@ -379,19 +403,19 @@ public:
 		bool bPressLeftMouseButton,
 		const FModifierKeysState& InModifierKeys = FModifierKeysState(false, false, false, false, false, false),
 		uint32 InTouchpadIndex=0
-		)
-		: FInputEvent( InModifierKeys, false )
-		, ScreenSpacePosition( InScreenSpacePosition )
-		, LastScreenSpacePosition( InLastScreenSpacePosition )
-		, CursorDelta( LastScreenSpacePosition - ScreenSpacePosition )
-		, PressedButtons( bPressLeftMouseButton ? StandardTouchKeySet : EmptyTouchKeySet )
-		, EffectingButton ( EKeys::LeftMouseButton )
-		, UserIndex( InUserIndex )
-		, PointerIndex( InPointerIndex )
-		, TouchpadIndex( InTouchpadIndex )
-		, bIsTouchEvent( true )
-		, GestureType( EGestureEvent::None )
-		, WheelOrGestureDelta( 0, 0 )
+	)
+		: FInputEvent(InModifierKeys, false)
+		, ScreenSpacePosition(InScreenSpacePosition)
+		, LastScreenSpacePosition(InLastScreenSpacePosition)
+		, CursorDelta(LastScreenSpacePosition - ScreenSpacePosition)
+		, PressedButtons(bPressLeftMouseButton ? FTouchKeySet::StandardSet : FTouchKeySet::EmptySet)
+		, EffectingButton (EKeys::LeftMouseButton)
+		, UserIndex(InUserIndex)
+		, PointerIndex(InPointerIndex)
+		, TouchpadIndex(InTouchpadIndex)
+		, bIsTouchEvent(true)
+		, GestureType(EGestureEvent::None)
+		, WheelOrGestureDelta(0.0f, 0.0f)
 	{ }
 
 	/** A constructor for gesture events */
@@ -402,19 +426,21 @@ public:
 		const FModifierKeysState& InModifierKeys,
 		EGestureEvent::Type InGestureType,
 		const FVector2D& InGestureDelta
-		)
-		: FInputEvent( InModifierKeys, false )
-		, ScreenSpacePosition( InScreenSpacePosition )
-		, LastScreenSpacePosition( InLastScreenSpacePosition )
-		, CursorDelta( LastScreenSpacePosition - ScreenSpacePosition )
-		, PressedButtons( InPressedButtons )
-		, UserIndex( 0 )
-		, PointerIndex( 0 )
-		, bIsTouchEvent( false )
-		, GestureType( InGestureType )
-		, WheelOrGestureDelta( InGestureDelta )
+	)
+		: FInputEvent(InModifierKeys, false)
+		, ScreenSpacePosition(InScreenSpacePosition)
+		, LastScreenSpacePosition(InLastScreenSpacePosition)
+		, CursorDelta(LastScreenSpacePosition - ScreenSpacePosition)
+		, PressedButtons(InPressedButtons)
+		, UserIndex(0)
+		, PointerIndex(0)
+		, bIsTouchEvent(false)
+		, GestureType(InGestureType)
+		, WheelOrGestureDelta(InGestureDelta)
 	{ }
 	
+public:
+
 	/** @return The position of the cursor in screen space */
 	const FVector2D& GetScreenSpacePosition() const { return ScreenSpacePosition; }
 
@@ -439,7 +465,7 @@ public:
 	/** @return The unique identifier of the pointer (e.g., finger index) */
 	uint32 GetPointerIndex() const { return PointerIndex; }
 
-	/** @return The index of the touchpad that generated this event (for platforms with multiple touchpads per user) */
+	/** @return The index of the touch pad that generated this event (for platforms with multiple touch pads per user) */
 	uint32 GetTouchpadIndex() const { return TouchpadIndex; }
 
 	/** @return Is this event a result from a touch (as opposed to a mouse) */
@@ -464,10 +490,6 @@ private:
 	bool bIsTouchEvent;
 	EGestureEvent::Type GestureType;
 	FVector2D WheelOrGestureDelta;
-
-	/** A reusable set of keys for touch events (since this class keeps a reference to a TSet) */
-	static FTouchKeySet StandardTouchKeySet;
-	static FTouchKeySet EmptyTouchKeySet;
 };
 
 
@@ -475,21 +497,20 @@ private:
  * FControllerEvent describes a controller action (e.g. Button Press, Release, Analog stick move, etc).
  * It is passed to event handlers dealing with controller input.
  */
-class FControllerEvent : public FInputEvent
+class FControllerEvent
+	: public FInputEvent
 {
 public:
 
-	FControllerEvent(
-		FKey InEffectingButton,
-		int32 InUserIndex,
-		float InAnalogValue,
-		bool bIsRepeat )
-		: FInputEvent( FModifierKeysState(false,false,false,false,false,false), bIsRepeat )
-		, EffectingButton( InEffectingButton )
-		, UserIndex( InUserIndex )
-		, AnalogValue( InAnalogValue )
+	FControllerEvent( FKey InEffectingButton, int32 InUserIndex, float InAnalogValue, bool bIsRepeat )
+		: FInputEvent(FModifierKeysState(false, false, false, false, false, false), bIsRepeat)
+		, EffectingButton(InEffectingButton)
+		, UserIndex(InUserIndex)
+		, AnalogValue(InAnalogValue)
 	{ }
 
+public:
+
 	/** @return The controller button that caused this event */
 	FKey GetEffectingButton() const { return EffectingButton; }
 
@@ -501,20 +522,23 @@ public:
 
 private:
 
-	/** The controller button that caused this event */
+	// The controller button that caused this event.
 	FKey EffectingButton;
-	/** The index of the user that caused the event */
+
+	// The index of the user that caused the event.
 	uint32 UserIndex;
-	/**  Analog value between 0 and 1.  1 being fully pressed, 0 being not pressed at all */
+	
+	//  Analog value between 0 and 1 (1 being fully pressed, 0 being not pressed at all).
 	float AnalogValue;
 };
 
 
 /**
- * FMotionEvent describes a touchpad action (press, move, lift)
+ * FMotionEvent describes a touch pad action (press, move, lift)
  * It is passed to event handlers dealing with touch input.
  */
-class FMotionEvent : public FInputEvent
+class FMotionEvent
+	: public FInputEvent
 {
 public:
 	FMotionEvent(
@@ -523,15 +547,17 @@ public:
 		const FVector& InRotationRate, 
 		const FVector& InGravity, 
 		const FVector& InAcceleration
-		)
-		: FInputEvent( FModifierKeysState(false,false,false,false,false,false), false )
-		, UserIndex( InUserIndex )
-		, Tilt( InTilt )
-		, RotationRate( InRotationRate )
-		, Gravity( InGravity )
-		, Acceleration( InAcceleration )
+	)
+		: FInputEvent(FModifierKeysState(false, false, false, false, false, false), false)
+		, UserIndex(InUserIndex)
+		, Tilt(InTilt)
+		, RotationRate(InRotationRate)
+		, Gravity(InGravity)
+		, Acceleration(InAcceleration)
 	{ }
 
+public:
+
 	/** @return The index of the user that caused the event */
 	uint32 GetUserIndex() const { return UserIndex; }
 
@@ -549,15 +575,19 @@ public:
 
 private:
 
-	/** The index of the user that caused the event */
+	// The index of the user that caused the event.
 	uint32 UserIndex;
-	/** The current tilt of the device/controller */
+
+	// The current tilt of the device/controller.
 	FVector Tilt;
-	/** The rotation speed */
+
+	// The rotation speed.
 	FVector RotationRate;
-	/** The gravity vector (pointing down into the ground) */
+
+	// The gravity vector (pointing down into the ground).
 	FVector Gravity;
-	/** The 3D acceleration of the device */
+
+	// The 3D acceleration of the device.
 	FVector Acceleration;
 };
 
@@ -584,7 +614,9 @@ public:
 		, AffectedWindow(InAffectedWindow)
 	{ }
 
-	/** Decscribes what actually happened to the window (e.g. Activated, Deactivated, Activated by a mouse click) */
+public:
+
+	/** Describes what actually happened to the window (e.g. Activated, Deactivated, Activated by a mouse click) */
 	EActivationType GetActivationType() const
 	{
 		return ActivationType;
