NaN safety in Vector2D and SlateElementBatcher
#UE4

diff --git a/Engine/Source/Runtime/Core/Public/Math/Vector2D.h b/Engine/Source/Runtime/Core/Public/Math/Vector2D.h
--- a/Engine/Source/Runtime/Core/Public/Math/Vector2D.h
+++ b/Engine/Source/Runtime/Core/Public/Math/Vector2D.h
@@ -433,6 +433,15 @@ struct FVector2D
 		return Ar << V.X << V.Y;
 	}
 
+#if ENABLE_NAN_DIAGNOSTIC
+	FORCEINLINE void DiagnosticCheckNaN() const
+	{
+		checkf(!ContainsNaN(), TEXT("FVector contains NaN: %s"), *ToString());
+	}
+#else
+	FORCEINLINE void DiagnosticCheckNaN() const {}
+#endif
+
 	/**
 	 * Utility to check if there are any NaNs in this vector.
 	 *
diff --git a/Engine/Source/Runtime/SlateCore/Public/Rendering/DrawElements.h b/Engine/Source/Runtime/SlateCore/Public/Rendering/DrawElements.h
--- a/Engine/Source/Runtime/SlateCore/Public/Rendering/DrawElements.h
+++ b/Engine/Source/Runtime/SlateCore/Public/Rendering/DrawElements.h
@@ -73,6 +73,7 @@ public:
 	FSlateDataPayload()
 		: Tint(FLinearColor::White)
 		, BrushResource(nullptr)
+		, RotationPoint(FVector2D::ZeroVector)
 		, Viewport(nullptr)
 	{ }
 
@@ -88,6 +89,7 @@ public:
 		Tint = InTint;
 		BrushResource = InBrush;
 		RotationPoint = LocalRotationPoint;
+		RotationPoint.DiagnosticCheckNaN();
 		Angle = InAngle;
 	}
 
