Slate - Exposing more slate enums to the UObject system.

diff --git a/Engine/Source/Runtime/Slate/Public/Framework/SlateDelegates.h b/Engine/Source/Runtime/Slate/Public/Framework/SlateDelegates.h
--- a/Engine/Source/Runtime/Slate/Public/Framework/SlateDelegates.h
+++ b/Engine/Source/Runtime/Slate/Public/Framework/SlateDelegates.h
@@ -98,38 +98,6 @@ DECLARE_DELEGATE_OneParam( FOnTextChanged, const FText& );
 /** Delegate for validating typed-in characters in SEditableText. Only invoked for typed characters */
 DECLARE_DELEGATE_RetVal_OneParam( bool, FOnIsTypedCharValid, const TCHAR /*TypedCharacter*/ );
 
-/** Additional information about a text committal */
-namespace ETextCommit
-{
-	enum Type
-	{
-		/** Losing focus or similar event caused implicit commit */
-		Default,
-		/** User committed via the enter key */
-		OnEnter,
-		/** User committed via tabbing away or moving focus explicitly away */
-		OnUserMovedFocus,
-		/** Keyboard focus was explicitly cleared via the escape key or other similar action */
-		OnCleared
-	};
-}
-
-
-/** Additional information about a selection event */
-namespace ESelectInfo
-{
-	enum Type
-	{
-		/** User selected via a key press */
-		OnKeyPress,
-		/** User selected by navigating to the item */
-		OnNavigation,
-		/** User selected by clicking on the item */
-		OnMouseClick,
-		/** Selection was directly set in code */
-		Direct
-	};
-}
 
 
 /** Delegate for hooking up to an editable text box's 'OnTextCommitted' 
diff --git a/Engine/Source/Runtime/SlateCore/Public/Types/SlateEnums.h b/Engine/Source/Runtime/SlateCore/Public/Types/SlateEnums.h
--- a/Engine/Source/Runtime/SlateCore/Public/Types/SlateEnums.h
+++ b/Engine/Source/Runtime/SlateCore/Public/Types/SlateEnums.h
@@ -8,7 +8,7 @@
 
 #include "SlateEnums.generated.h"
 
-
+UENUM()
 namespace EButtonClickMethod
 {
 	/**
@@ -35,7 +35,7 @@ namespace EButtonClickMethod
 	};
 }
 
-
+UENUM()
 namespace EFocusMoveDirection
 {
 	/**
@@ -129,6 +129,7 @@ enum EOrientation
 /**
  * Enumerates scroll directions.
  */
+UENUM()
 enum EScrollDirection
 {
 	/** Scroll down. */
@@ -137,3 +138,42 @@ enum EScrollDirection
 	/** Scroll up. */
 	Scroll_Up UMETA(DisplayName="Up"),
 };
+
+/**
+ * Additional information about a text committal
+ */
+UENUM()
+namespace ETextCommit
+{
+	enum Type
+	{
+		/** Losing focus or similar event caused implicit commit */
+		Default,
+		/** User committed via the enter key */
+		OnEnter,
+		/** User committed via tabbing away or moving focus explicitly away */
+		OnUserMovedFocus,
+		/** Keyboard focus was explicitly cleared via the escape key or other similar action */
+		OnCleared
+	};
+}
+
+
+/**
+ * Additional information about a selection event
+ */
+UENUM()
+namespace ESelectInfo
+{
+	enum Type
+	{
+		/** User selected via a key press */
+		OnKeyPress,
+		/** User selected by navigating to the item */
+		OnNavigation,
+		/** User selected by clicking on the item */
+		OnMouseClick,
+		/** Selection was directly set in code */
+		Direct
+	};
+}
