SlateCore: removed obsolete debug code (also fixed non-Unity compilation in Debug)

diff --git a/Engine/Source/Runtime/SlateCore/Private/Rendering/SlateDrawBuffer.cpp b/Engine/Source/Runtime/SlateCore/Private/Rendering/SlateDrawBuffer.cpp
--- a/Engine/Source/Runtime/SlateCore/Private/Rendering/SlateDrawBuffer.cpp
+++ b/Engine/Source/Runtime/SlateCore/Private/Rendering/SlateDrawBuffer.cpp
@@ -7,43 +7,9 @@
 #include "SlateCorePrivatePCH.h"
 
 
-/* Local constants
- *****************************************************************************/
-
-uint32 GDrawBuffers = 0;
-uint32 MaxDrawBuffers = 0;
-
-
-/* FSlateDrawBuffer structors
- *****************************************************************************/
-
-FSlateDrawBuffer::FSlateDrawBuffer( ) 
-	: Locked(0)
-{
-	++GDrawBuffers;
-
-	if (GDrawBuffers > MaxDrawBuffers)
-	{
-		MaxDrawBuffers = GDrawBuffers;
-	}
-}
-
-
-FSlateDrawBuffer::~FSlateDrawBuffer( )
-{
-	--GDrawBuffers;
-}
-
-
 /* FSlateDrawBuffer interface
  *****************************************************************************/
 
-void FSlateDrawBuffer::ClearBuffer( )
-{
-	WindowElementLists.Empty();
-}
-
-
 FSlateWindowElementList& FSlateDrawBuffer::AddWindowElementList( TSharedRef<SWindow> ForWindow )
 {
 	int32 Index = WindowElementLists.Add(FSlateWindowElementList(ForWindow));
@@ -61,4 +27,4 @@ bool FSlateDrawBuffer::Lock( )
 void FSlateDrawBuffer::Unlock( )
 {
 	FPlatformAtomics::InterlockedExchange(&Locked, 0);
-}
\ No newline at end of file
+}
diff --git a/Engine/Source/Runtime/SlateCore/Private/Rendering/SlateRenderer.cpp b/Engine/Source/Runtime/SlateCore/Private/Rendering/SlateRenderer.cpp
--- a/Engine/Source/Runtime/SlateCore/Private/Rendering/SlateRenderer.cpp
+++ b/Engine/Source/Runtime/SlateCore/Private/Rendering/SlateRenderer.cpp
@@ -7,16 +7,6 @@
 #include "SlateCorePrivatePCH.h"
 
 
-/* FSlateRenderer structors
- *****************************************************************************/
-
-FSlateRenderer::~FSlateRenderer()
-{
-	// debug check to make sure we are properly cleaning up draw buffers
-	checkSlow( GDrawBuffers == 0 );
-}
-
-
 /* FSlateRenderer interface
  *****************************************************************************/
 
diff --git a/Engine/Source/Runtime/SlateCore/Public/Rendering/SlateDrawBuffer.h b/Engine/Source/Runtime/SlateCore/Public/Rendering/SlateDrawBuffer.h
--- a/Engine/Source/Runtime/SlateCore/Public/Rendering/SlateDrawBuffer.h
+++ b/Engine/Source/Runtime/SlateCore/Public/Rendering/SlateDrawBuffer.h
@@ -7,6 +7,9 @@
 #pragma once
 
 
+/**
+ * Implements a draw buffer for Slate.
+ */
 class SLATECORE_API FSlateDrawBuffer
 {
 public:
@@ -14,19 +17,19 @@ public:
 	/**
 	 * Default constructor.
 	 */
-	explicit FSlateDrawBuffer( );
-
-	/**
-	 * Destructor.
-	 */
-	~FSlateDrawBuffer( );
+	explicit FSlateDrawBuffer( )
+		: Locked(0)
+	{ }
 
 public:
 
 	/**
 	 * Removes all data from the buffer.
 	 */
-	void ClearBuffer( );
+	void ClearBuffer( )
+	{
+		WindowElementLists.Empty();
+	}
 
 	/**
 	 * Creates a new FSlateWindowElementList and returns a reference to it so it can have draw elements added to it
diff --git a/Engine/Source/Runtime/SlateCore/Public/Rendering/SlateRenderer.h b/Engine/Source/Runtime/SlateCore/Public/Rendering/SlateRenderer.h
--- a/Engine/Source/Runtime/SlateCore/Public/Rendering/SlateRenderer.h
+++ b/Engine/Source/Runtime/SlateCore/Public/Rendering/SlateRenderer.h
@@ -16,7 +16,6 @@ class SWindow;
 class FSlateRenderer
 {
 public:
-	virtual SLATECORE_API ~FSlateRenderer();
 
 	/** Returns a draw buffer that can be used by Slate windows to draw window elements */
 	virtual FSlateDrawBuffer& GetDrawBuffer() = 0;
