SlateCore: added bakc virtual destructor for FSlateRenderer

diff --git a/Engine/Source/Runtime/SlateCore/Public/Rendering/SlateRenderer.h b/Engine/Source/Runtime/SlateCore/Public/Rendering/SlateRenderer.h
--- a/Engine/Source/Runtime/SlateCore/Public/Rendering/SlateRenderer.h
+++ b/Engine/Source/Runtime/SlateCore/Public/Rendering/SlateRenderer.h
@@ -15,6 +15,13 @@ class SWindow;
  */
 class FSlateRenderer
 {
+public:
+
+	/**
+	 * Virtual destructor.
+	 */
+	virtual ~FSlateRenderer( ) { }
+
 public:
 
 	/** Returns a draw buffer that can be used by Slate windows to draw window elements */
