Slate - Removing GetCursor from SWidget, unnessesary.  Adjusting doc on SSlider.

diff --git a/Engine/Source/Runtime/Slate/Private/Widgets/Input/SSlider.cpp b/Engine/Source/Runtime/Slate/Private/Widgets/Input/SSlider.cpp
--- a/Engine/Source/Runtime/Slate/Private/Widgets/Input/SSlider.cpp
+++ b/Engine/Source/Runtime/Slate/Private/Widgets/Input/SSlider.cpp
@@ -190,6 +190,11 @@ float SSlider::PositionToValue( const FGeometry& MyGeometry, const FVector2D& Ab
 	return FMath::Clamp(RelativeValue, 0.0f, 1.0f);
 }
 
+float SSlider::GetValue() const
+{
+	return ValueAttribute.Get();
+}
+
 void SSlider::SetValue(const TAttribute<float>& InValueAttribute)
 {
 	ValueAttribute = InValueAttribute;
diff --git a/Engine/Source/Runtime/Slate/Public/Widgets/Input/SSlider.h b/Engine/Source/Runtime/Slate/Public/Widgets/Input/SSlider.h
--- a/Engine/Source/Runtime/Slate/Public/Widgets/Input/SSlider.h
+++ b/Engine/Source/Runtime/Slate/Public/Widgets/Input/SSlider.h
@@ -42,7 +42,7 @@ public:
 		/** The style used to draw the slider. */
 		SLATE_STYLE_ARGUMENT( FSliderStyle, Style )
 
-		/** The volume value to display. */
+		/** A value that drives where the slider handle appears. Value is normalized between 0 and 1. */
 		SLATE_ATTRIBUTE( float, Value )
 
 		/** Invoked when the mouse is pressed and a capture begins. */
@@ -51,8 +51,7 @@ public:
 		/** Invoked when the mouse is released and a capture ends. */
 		SLATE_EVENT(FSimpleDelegate, OnMouseCaptureEnd)
 
-
-		/** Called when the value is changed by slider or typing. */
+		/** Called when the value is changed by the slider. */
 		SLATE_EVENT( FOnFloatValueChanged, OnValueChanged )
 
 	SLATE_END_ARGS()
@@ -64,21 +63,10 @@ public:
 	 */
 	void Construct( const SSlider::FArguments& InDeclaration );
 
-	/**
-	 * Gets the slider's value.
-	 *
-	 * @return Slider value.
-	 */
-	float GetValue() const
-	{
-		return ValueAttribute.Get();
-	}
+	/** See the Value attribute */
+	float GetValue() const;
 
-	/**
-	* Gets the slider's value.
-	*
-	* @return Slider value.
-	*/
+	/** See the Value attribute */
 	void SetValue(const TAttribute<float>& InValueAttribute);
 
 public:
diff --git a/Engine/Source/Runtime/SlateCore/Private/Widgets/SWidget.cpp b/Engine/Source/Runtime/SlateCore/Private/Widgets/SWidget.cpp
--- a/Engine/Source/Runtime/SlateCore/Private/Widgets/SWidget.cpp
+++ b/Engine/Source/Runtime/SlateCore/Private/Widgets/SWidget.cpp
@@ -459,12 +459,6 @@ void SWidget::EnableToolTipForceField( const bool bEnableForceField )
 }
 
 
-TOptional<EMouseCursor::Type> SWidget::GetCursor() const
-{
-	return Cursor.Get();
-}
-
-
 void SWidget::SetCursor( const TAttribute< TOptional<EMouseCursor::Type> >& InCursor )
 {
 	Cursor = InCursor;
diff --git a/Engine/Source/Runtime/SlateCore/Public/Widgets/SWidget.h b/Engine/Source/Runtime/SlateCore/Public/Widgets/SWidget.h
--- a/Engine/Source/Runtime/SlateCore/Public/Widgets/SWidget.h
+++ b/Engine/Source/Runtime/SlateCore/Public/Widgets/SWidget.h
@@ -90,7 +90,7 @@ public:
 	 * @param OutDrawElements   A list of FDrawElements to populate with the output.
 	 * @param LayerId           The Layer onto which this widget should be rendered.
 	 * @param InColorAndOpacity Color and Opacity to be applied to all the descendants of the widget being painted
- 	 * @param bParentEnabled	True if the parent of this widget is enabled.
+	 * @param bParentEnabled	True if the parent of this widget is enabled.
 	 *
 	 * @return The maximum layer ID attained by this widget or any of its children.
 	 */
@@ -580,11 +580,6 @@ public:
 	 */
 	void SetToolTip( const TSharedPtr<IToolTip>& InToolTip );
 
-	/**
-	 * Get the cursor that should appear when this widget is hovered
-	 */
-	TOptional<EMouseCursor::Type> GetCursor() const;
-
 	/**
 	 * Set the cursor that should appear when this widget is hovered
 	 */
@@ -713,7 +708,7 @@ private:
 
 	// Whether this widget is a "tool tip force field".  That is, tool-tips should never spawn over the area
 	// occupied by this widget, and will instead be repelled to an outside edge
- 	bool bToolTipForceFieldEnabled;
+	bool bToolTipForceFieldEnabled;
 
 	/** Is this widget hovered? */
 	bool bIsHovered;
