#slate Added support for UserWidgets. Check out SUserWidget.h, SUsetWidgetTest.h and SUserWidgetTest.cpp for usage info.
User widgets allow allocation of the widget to happen in the .CPP, which lets everything except for a (usually) small public interface to reside entirely within a .CPP, which is the only way to make something truly private in C++.

diff --git a/Engine/Source/Runtime/Slate/Private/Framework/Testing/SUserWidgetTest.cpp b/Engine/Source/Runtime/Slate/Private/Framework/Testing/SUserWidgetTest.cpp
new file mode 100644
--- /dev/null
+++ b/Engine/Source/Runtime/Slate/Private/Framework/Testing/SUserWidgetTest.cpp
@@ -0,0 +1,36 @@
+// Copyright 1998-2014 Epic Games, Inc. All Rights Reserved.
+
+#include "SlatePrivatePCH.h"
+#include "SUserWidgetTest.h"
+
+
+namespace Implementation
+{
+	
+class SUserWidgetExample : public ::SUserWidgetExample
+{
+public:
+
+	void Construct( const FArguments& InArgs )
+	{
+		SUserWidget::Construct( SUserWidget::FArguments()
+		[
+			SNew(STextBlock)
+			.Text( FText::Format( NSLOCTEXT("SlateTestSuite","UserWidgetExampleTitle"," Implemented in the .cpp : {0}"), InArgs._Title ) )
+		]);
+	}
+
+
+	virtual void DoStuff() OVERRIDE
+	{
+
+	}
+};
+
+}
+
+
+TSharedRef<SUserWidgetExample> SUserWidgetExample::New()
+{
+	return MakeShareable(new Implementation::SUserWidgetExample()); 
+}
\ No newline at end of file
diff --git a/Engine/Source/Runtime/Slate/Private/Framework/Testing/SUserWidgetTest.h b/Engine/Source/Runtime/Slate/Private/Framework/Testing/SUserWidgetTest.h
new file mode 100644
--- /dev/null
+++ b/Engine/Source/Runtime/Slate/Private/Framework/Testing/SUserWidgetTest.h
@@ -0,0 +1,20 @@
+// Copyright 1998-2014 Epic Games, Inc. All Rights Reserved.
+
+#pragma once
+
+class SUserWidgetExample : public SUserWidget
+{
+	public:
+
+	SLATE_USER_ARGS(SUserWidgetExample)
+	: _Title()
+	{}
+		SLATE_ARGUMENT( FText, Title )
+	SLATE_END_ARGS()
+
+	virtual void DoStuff() = 0;
+};
+
+
+
+
diff --git a/Engine/Source/Runtime/Slate/Private/Framework/Testing/SWidgetGallery.cpp b/Engine/Source/Runtime/Slate/Private/Framework/Testing/SWidgetGallery.cpp
--- a/Engine/Source/Runtime/Slate/Private/Framework/Testing/SWidgetGallery.cpp
+++ b/Engine/Source/Runtime/Slate/Private/Framework/Testing/SWidgetGallery.cpp
@@ -4,6 +4,8 @@
 #include "SWidgetGallery.h"
 #include "TestStyle.h"
 
+#include "SUserWidgetTest.h"
+
 
 #define LOCTEXT_NAMESPACE "WidgetGallery"
 
@@ -756,6 +758,20 @@ public:
 										.Volume(0.6f)
 								]
 						]
+
+					+ SGridPanel::Slot(0, 31)
+					[
+						SNew(STextBlock)
+							.Text(LOCTEXT("SUserWidgetLabel", "SUserWidgetExample"))
+					]
+
+					+ SGridPanel::Slot(1, 31)
+					.HAlign(HAlign_Left)
+					.Padding(0.0f, 5.0f)
+					[
+						SNew(SUserWidgetExample)
+						.Title( LOCTEXT("SUserWidgetTest", "UserWidgetTest.cpp") )
+					]
 				]
 		];
 
diff --git a/Engine/Source/Runtime/Slate/Public/Widgets/Navigation/SWizard.h b/Engine/Source/Runtime/Slate/Public/Widgets/Navigation/SWizard.h
--- a/Engine/Source/Runtime/Slate/Public/Widgets/Navigation/SWizard.h
+++ b/Engine/Source/Runtime/Slate/Public/Widgets/Navigation/SWizard.h
@@ -30,7 +30,7 @@ public:
 	{
 	public:
 
-		SLATE_BEGIN_ARGS(FWizardPage)
+		SLATE_SLOT_ARGS(FWizardPage)
 			: _ButtonContent()
 			, _CanShow(true)
 			, _PageContent()
diff --git a/Engine/Source/Runtime/Slate/Public/Widgets/Views/SHeaderRow.h b/Engine/Source/Runtime/Slate/Public/Widgets/Views/SHeaderRow.h
--- a/Engine/Source/Runtime/Slate/Public/Widgets/Views/SHeaderRow.h
+++ b/Engine/Source/Runtime/Slate/Public/Widgets/Views/SHeaderRow.h
@@ -62,7 +62,7 @@ public:
 	{
 	public:
 
-		SLATE_BEGIN_ARGS(FColumn)
+		SLATE_SLOT_ARGS(FColumn)
 			: _ColumnId()
 			, _DefaultLabel()
 			, _FillWidth( 1.0f )
diff --git a/Engine/Source/Runtime/SlateCore/Public/SlateCore.h b/Engine/Source/Runtime/SlateCore/Public/SlateCore.h
--- a/Engine/Source/Runtime/SlateCore/Public/SlateCore.h
+++ b/Engine/Source/Runtime/SlateCore/Public/SlateCore.h
@@ -143,6 +143,7 @@ DECLARE_MEMORY_STAT_EXTERN(TEXT("Index Buffer Memory"), STAT_SlateIndexBufferMem
 #include "DeclarativeSyntaxSupport.h"
 #include "SWidget.h"
 #include "SCompoundWidget.h"
+#include "SUserWidget.h"
 #include "SLeafWidget.h"
 #include "SPanel.h"
 #include "WidgetPath.inl"
diff --git a/Engine/Source/Runtime/SlateCore/Public/Widgets/DeclarativeSyntaxSupport.h b/Engine/Source/Runtime/SlateCore/Public/Widgets/DeclarativeSyntaxSupport.h
--- a/Engine/Source/Runtime/SlateCore/Public/Widgets/DeclarativeSyntaxSupport.h
+++ b/Engine/Source/Runtime/SlateCore/Public/Widgets/DeclarativeSyntaxSupport.h
@@ -44,11 +44,68 @@
 
 #define SLATE_BEGIN_ARGS( WidgetType ) \
 	public: \
+	static TSharedRef<WidgetType> New(){ return PrivateAllocateWidget<WidgetType>(); } \
 	struct FArguments : public TSlateBaseNamedArgs<WidgetType> \
 	{ \
 		typedef FArguments WidgetArgsType; \
 		FORCENOINLINE FArguments()
-		
+
+/**
+ * Just like SLATE_BEGIN_ARGS but requires the user to implement the New() method in the .CPP.
+ * Allows for widget implementation details to be entirely reside in the .CPP file.
+ * e.g.
+ *   SMyWidget.h:
+ *   ----------------------------------------------------------------
+ *    SLATE_USER_ARGS( SMyWidget )
+ *    {}
+ *    SLATE_END_ARGS()
+ *    
+ *    virtual void DoStuff() = nullptr;
+ *
+ *   SMyWidget.cpp:
+ *   ----------------------------------------------------------------
+ *   namespace Implementation{
+ *   class SMyWidget : public SMyWidget{
+ *     void Construct( const FArguments& InArgs )
+ *     {
+ *         SUserWidget::Construct( SUSerWidget::FArguments()
+ *         [
+ *             SNew(STextBlock) .Text( "Implementation Details!" )
+ *         ] );
+ *     }
+ *
+ *     virtual void DoStuff() OVERRIDE {}
+ *     
+ *     // Truly private. All handlers can be inlined (less boilerplate).
+ *     // All private members are truly private.
+ *     private:
+ *     FReply OnButtonClicked
+ *     {
+ *     }
+ *     TSharedPtr<SButton> MyButton;
+ *   }
+ *   }
+ */
+#define SLATE_USER_ARGS( WidgetType ) \
+	public: \
+	static TSharedRef<WidgetType> New(); \
+	struct FArguments; \
+	virtual void Construct( const FArguments& InArgs ) = 0; \
+	struct FArguments : public TSlateBaseNamedArgs<WidgetType> \
+	{ \
+		typedef FArguments WidgetArgsType; \
+		FORCENOINLINE FArguments()
+
+
+// @todo UMG: Probably remove this?
+#define SLATE_SLOT_ARGS( WidgetType ) \
+	public: \
+	struct FArguments : public TSlateBaseNamedArgs<WidgetType> \
+	{ \
+		typedef FArguments WidgetArgsType; \
+		FORCENOINLINE FArguments()
+
+
 
 #define SLATE_END_ARGS() \
 	};
@@ -1094,6 +1151,12 @@ namespace RequiredArgs
 	}
 }
 
+/** Responsible for allocating the widget. */
+template<typename WidgetType>
+static TSharedRef<WidgetType> PrivateAllocateWidget()
+{
+	return MakeShareable( new WidgetType() );
+}
 
 /**
  * Utility class used during widget instantiation.
@@ -1108,7 +1171,7 @@ template<class WidgetType, typename RequiredArgsPayloadType>
 struct TDecl
 {
 	TDecl( const ANSICHAR* InType, const ANSICHAR* InFile, int32 OnLine, const RequiredArgsPayloadType& InRequiredArgs )
-		: _Widget( new WidgetType() )
+		: _Widget( WidgetType::New() )
 		, _RequiredArgs( InRequiredArgs )
 	{
 		_Widget->SetDebugInfo( InType, InFile, OnLine );
@@ -1145,7 +1208,8 @@ struct TDecl
 	 */
 	TSharedRef<WidgetType> operator<<=( const typename WidgetType::FArguments& InArgs ) const
 	{
-		_Widget->SWidget::Construct(
+		//@todo UMG: This should be removed in favor of all widgets calling their superclass construct.
+		_Widget->SWidgetConstruct(
 			InArgs._ToolTipText,
 			InArgs._ToolTip ,
 			InArgs._Cursor ,
diff --git a/Engine/Source/Runtime/SlateCore/Public/Widgets/SUserWidget.h b/Engine/Source/Runtime/SlateCore/Public/Widgets/SUserWidget.h
new file mode 100644
--- /dev/null
+++ b/Engine/Source/Runtime/SlateCore/Public/Widgets/SUserWidget.h
@@ -0,0 +1,86 @@
+// Copyright 1998-2014 Epic Games, Inc. All Rights Reserved.
+
+#pragma once
+
+/**
+ * Use SUserWidget as a base class to build aggregate widgets that are not meant
+ * to serve as low-level building blocks. Examples include: a main menu, a user card,
+ * an info dialog for a selected object, a splash screen.
+ *
+ * See SUserWidgetExample
+ *
+ * SMyWidget.h
+ * -----------
+ * class SMyWidget : public SUserWidget
+ * { 
+ *   public:
+ *   SLATE_BEGIN_ARGS( SMyWidget )
+ *	 {}
+ *   SLATE_END_ARGS()
+ *    
+ *   virtual void DoSomething() = 0;
+ * };
+ *
+ * SMyWidget.cpp
+ * -------------
+ * namespace Implementation {
+ *   class SMyWidget : public SMyWidget
+ *   {
+ *     public:
+ *     void Construct( const FArgument& InArgs )
+ *     {
+ *        SUserWidget::Construct( SUserWidget::FArguments()
+ *        [
+ *           SNew(STextBlock)
+ *           .Text( NSLOCTEXT("x", "x", "My Widget's Content") )
+ *        ]
+ *     }
+ *     
+ *     private:
+ *     // Private implementation details can occur here
+ *     // without ever leaking out into the .h file!
+ *   }
+ * }
+ *
+ * TSharedRef<SMyWidget> SMyWidget::New()
+ * {
+ *   return MakeShareable( new SMyWidget() );
+ * }
+ */
+class SUserWidget : public SCompoundWidget
+{
+	public:
+	struct FArguments : public TSlateBaseNamedArgs<SUserWidget>
+	{
+		typedef FArguments WidgetArgsType;
+		FORCENOINLINE FArguments()
+			: _Content()
+			, _HAlign(HAlign_Fill)
+			, _VAlign(VAlign_Fill)
+		{}
+
+		SLATE_DEFAULT_SLOT( FArguments, Content )
+		SLATE_ARGUMENT( EHorizontalAlignment, HAlign )
+		SLATE_ARGUMENT( EVerticalAlignment, VAlign )
+
+	};
+	
+	protected:
+	void Construct( const FArguments& InArgs )
+	{
+		this->ChildSlot
+		.HAlign( InArgs._HAlign )
+		.VAlign( InArgs._VAlign )
+		[
+			InArgs._Content.Widget
+		];
+	}
+
+	protected:
+	/** User widgets can be allocated explicitly in the C++  */
+	void* operator new ( const size_t InSize )
+	{
+		return FMemory::Malloc(InSize);
+	}
+
+};
\ No newline at end of file
diff --git a/Engine/Source/Runtime/SlateCore/Public/Widgets/SWidget.h b/Engine/Source/Runtime/SlateCore/Public/Widgets/SWidget.h
--- a/Engine/Source/Runtime/SlateCore/Public/Widgets/SWidget.h
+++ b/Engine/Source/Runtime/SlateCore/Public/Widgets/SWidget.h
@@ -21,13 +21,13 @@ private:
 	void operator= (const FSlateControlledConstruction& Other){}
 
 	/** Widgets should only ever be constructed via SNew or SAssignNew */
-	void* operator new (const size_t InSize )
+	void* operator new ( const size_t InSize )
 	{
 		return FMemory::Malloc(InSize);
 	}
 
-	template<class WidgetType, typename RequiredArgsPayloadType>
-	friend struct TDecl;
+	template<class WidgetType>
+	friend TSharedRef<WidgetType> PrivateAllocateWidget();
 
 public:
 	void operator delete(void * mem)
@@ -77,6 +77,16 @@ public:
 		const TAttribute<EVisibility> & InVisibility,
 		const FName& InTag );
 
+	void SWidgetConstruct( const TAttribute<FString> & InToolTipText ,
+		const TSharedPtr<IToolTip> & InToolTip ,
+		const TAttribute< TOptional<EMouseCursor::Type> > & InCursor ,
+		const TAttribute<bool> & InEnabledState ,
+		const TAttribute<EVisibility> & InVisibility,
+		const FName& InTag )
+	{
+		Construct(InToolTipText, InToolTip, InCursor, InEnabledState, InVisibility, InTag);
+	}
+
 	//
 	// GENERAL EVENTS
 	//
