#ttp 331438 UE4: PlatformMath refactor

diff --git a/Engine/Source/Runtime/SlateCore/Private/Widgets/SWindow.cpp b/Engine/Source/Runtime/SlateCore/Private/Widgets/SWindow.cpp
--- a/Engine/Source/Runtime/SlateCore/Private/Widgets/SWindow.cpp
+++ b/Engine/Source/Runtime/SlateCore/Private/Widgets/SWindow.cpp
@@ -627,7 +627,7 @@ void SWindow::MoveWindowTo( FVector2D NewPosition )
 {
 	if (NativeWindow.IsValid())
 	{
-		NativeWindow->MoveWindowTo( FMath::Trunc(NewPosition.X), FMath::Trunc(NewPosition.Y) );
+		NativeWindow->MoveWindowTo( FMath::TruncToInt(NewPosition.X), FMath::TruncToInt(NewPosition.Y) );
 	}
 	else
 	{
@@ -639,7 +639,7 @@ void SWindow::ReshapeWindow( FVector2D NewPosition, FVector2D NewSize )
 {
 	if (NativeWindow.IsValid())
 	{
-		NativeWindow->ReshapeWindow( FMath::Trunc(NewPosition.X), FMath::Trunc(NewPosition.Y), FMath::Trunc(NewSize.X), FMath::Trunc(NewSize.Y) );
+		NativeWindow->ReshapeWindow( FMath::TruncToInt(NewPosition.X), FMath::TruncToInt(NewPosition.Y), FMath::TruncToInt(NewSize.X), FMath::TruncToInt(NewSize.Y) );
 	}
 	else
 	{
@@ -662,7 +662,7 @@ void SWindow::Resize( FVector2D NewSize )
 	{
 		if (NativeWindow.IsValid())
 		{
-			NativeWindow->ReshapeWindow( FMath::Trunc(ScreenPosition.X), FMath::Trunc(ScreenPosition.Y), FMath::Trunc(NewSize.X), FMath::Trunc(NewSize.Y) );
+			NativeWindow->ReshapeWindow( FMath::TruncToInt(ScreenPosition.X), FMath::TruncToInt(ScreenPosition.Y), FMath::TruncToInt(NewSize.X), FMath::TruncToInt(NewSize.Y) );
 		}
 		else
 		{
@@ -1110,7 +1110,7 @@ bool SWindow::IsTopmostWindow() const
 bool SWindow::IsScreenspaceMouseWithin(FVector2D ScreenspaceMouseCoordinate) const
 {
 	const FVector2D LocalMouseCoordinate = GetWindowGeometryInScreen().AbsoluteToLocal(ScreenspaceMouseCoordinate);
-	return NativeWindow->IsPointInWindow(FMath::Trunc(LocalMouseCoordinate.X), FMath::Trunc(LocalMouseCoordinate.Y));
+	return NativeWindow->IsPointInWindow(FMath::TruncToInt(LocalMouseCoordinate.X), FMath::TruncToInt(LocalMouseCoordinate.Y));
 }
 
 /** @return true if this is a user-sized window with a thick edge */
diff --git a/Engine/Source/Runtime/SlateCore/Public/Rendering/RenderingCommon.h b/Engine/Source/Runtime/SlateCore/Public/Rendering/RenderingCommon.h
--- a/Engine/Source/Runtime/SlateCore/Public/Rendering/RenderingCommon.h
+++ b/Engine/Source/Runtime/SlateCore/Public/Rendering/RenderingCommon.h
@@ -109,12 +109,12 @@ struct FSlateVertex
 		, Color( InColor )
 	{
 
-		Position[0] = FMath::Trunc(InPosition.X);
-		Position[1] = FMath::Trunc(InPosition.Y);
-		ClipCoords[0] = FMath::Trunc(InClipCoords.Left);
-		ClipCoords[1] = FMath::Trunc(InClipCoords.Top);
-		ClipCoords[2] = FMath::Trunc(InClipCoords.Right);
-		ClipCoords[3] = FMath::Trunc(InClipCoords.Bottom);
+		Position[0] = FMath::TruncToInt(InPosition.X);
+		Position[1] = FMath::TruncToInt(InPosition.Y);
+		ClipCoords[0] = FMath::TruncToInt(InClipCoords.Left);
+		ClipCoords[1] = FMath::TruncToInt(InClipCoords.Top);
+		ClipCoords[2] = FMath::TruncToInt(InClipCoords.Right);
+		ClipCoords[3] = FMath::TruncToInt(InClipCoords.Bottom);
 	}
 
 	FORCEINLINE FSlateVertex( const FVector2D& InPosition, const FVector2D& InTexCoord, const FVector2D& InTexCoord2, const FColor& InColor, const FVector4& InClipCoords) 
@@ -122,12 +122,12 @@ struct FSlateVertex
 		, Color( InColor )
 	{
 
-		Position[0] = FMath::Trunc(InPosition.X);
-		Position[1] = FMath::Trunc(InPosition.Y);
-		ClipCoords[0] = FMath::Trunc(InClipCoords.X);
-		ClipCoords[1] = FMath::Trunc(InClipCoords.Y);
-		ClipCoords[2] = FMath::Trunc(InClipCoords.Z);
-		ClipCoords[3] = FMath::Trunc(InClipCoords.W);
+		Position[0] = FMath::TruncToInt(InPosition.X);
+		Position[1] = FMath::TruncToInt(InPosition.Y);
+		ClipCoords[0] = FMath::TruncToInt(InClipCoords.X);
+		ClipCoords[1] = FMath::TruncToInt(InClipCoords.Y);
+		ClipCoords[2] = FMath::TruncToInt(InClipCoords.Z);
+		ClipCoords[3] = FMath::TruncToInt(InClipCoords.W);
 	}
 
 	FORCEINLINE FSlateVertex( const FVector2D& InPosition, const FVector2D& InTexCoord, const FColor& InColor, const FSlateRect& InClipCoords ) 
@@ -135,12 +135,12 @@ struct FSlateVertex
 		, Color( InColor )
 	{
 
-		Position[0] = FMath::Trunc(InPosition.X);
-		Position[1] = FMath::Trunc(InPosition.Y);
-		ClipCoords[0] = FMath::Trunc(InClipCoords.Left);
-		ClipCoords[1] = FMath::Trunc(InClipCoords.Top);
-		ClipCoords[2] = FMath::Trunc(InClipCoords.Right);
-		ClipCoords[3] = FMath::Trunc(InClipCoords.Bottom);
+		Position[0] = FMath::TruncToInt(InPosition.X);
+		Position[1] = FMath::TruncToInt(InPosition.Y);
+		ClipCoords[0] = FMath::TruncToInt(InClipCoords.Left);
+		ClipCoords[1] = FMath::TruncToInt(InClipCoords.Top);
+		ClipCoords[2] = FMath::TruncToInt(InClipCoords.Right);
+		ClipCoords[3] = FMath::TruncToInt(InClipCoords.Bottom);
 	}
 
 	FORCEINLINE FSlateVertex( const FVector2D& InPosition, const FVector2D& InTexCoord, const uint32 InDWORD ) 
@@ -148,8 +148,8 @@ struct FSlateVertex
 		, Color( InDWORD )
 	{
 
-		Position[0] = FMath::Trunc(InPosition.X);
-		Position[1] = FMath::Trunc(InPosition.Y);
+		Position[0] = FMath::TruncToInt(InPosition.X);
+		Position[1] = FMath::TruncToInt(InPosition.Y);
 		ClipCoords[0] = 0;
 		ClipCoords[1] = 0;
 		ClipCoords[2] = 0;
diff --git a/Engine/Source/Runtime/SlateRHIRenderer/Private/SlateRHIRenderer.cpp b/Engine/Source/Runtime/SlateRHIRenderer/Private/SlateRHIRenderer.cpp
--- a/Engine/Source/Runtime/SlateRHIRenderer/Private/SlateRHIRenderer.cpp
+++ b/Engine/Source/Runtime/SlateRHIRenderer/Private/SlateRHIRenderer.cpp
@@ -292,8 +292,8 @@ void FSlateRHIRenderer::CreateViewport( const TSharedRef<SWindow> Window )
 		const FVector2D WindowSize = Window->GetSizeInScreen();
 		// Clamp the window size to a reasonable default anything below 8 is a d3d warning and 8 is used anyway.
 		// @todo Slate: This is a hack to work around menus being summoned with 0,0 for window size until they are ticked.
-		const uint32 Width = FMath::Max(8,FMath::Trunc(WindowSize.X));
-		const uint32 Height = FMath::Max(8,FMath::Trunc(WindowSize.Y));
+		const uint32 Width = FMath::Max(8,FMath::TruncToInt(WindowSize.X));
+		const uint32 Height = FMath::Max(8,FMath::TruncToInt(WindowSize.Y));
 
 		FViewportInfo* NewInfo = new FViewportInfo();
 		// Create Viewport RHI if it doesn't exist (this must be done on the game thread)
diff --git a/Engine/Source/Runtime/Windows/D3D11RHI/Private/D3D11RHI.cpp b/Engine/Source/Runtime/Windows/D3D11RHI/Private/D3D11RHI.cpp
--- a/Engine/Source/Runtime/Windows/D3D11RHI/Private/D3D11RHI.cpp
+++ b/Engine/Source/Runtime/Windows/D3D11RHI/Private/D3D11RHI.cpp
@@ -309,7 +309,7 @@ void FD3DGPUProfiler::EndFrame()
 	{
 		uint64 GPUTiming = FrameTiming.GetTiming();
 		uint64 GPUFreq = FrameTiming.GetTimingFrequency();
-		GGPUFrameTime = FMath::Trunc( double(GPUTiming) / double(GPUFreq) / FPlatformTime::GetSecondsPerCycle() );
+		GGPUFrameTime = FMath::TruncToInt( double(GPUTiming) / double(GPUFreq) / FPlatformTime::GetSecondsPerCycle() );
 	}
 	else
 	{
diff --git a/Engine/Source/Runtime/Windows/D3D11RHI/Private/D3D11RenderTarget.cpp b/Engine/Source/Runtime/Windows/D3D11RHI/Private/D3D11RenderTarget.cpp
--- a/Engine/Source/Runtime/Windows/D3D11RHI/Private/D3D11RenderTarget.cpp
+++ b/Engine/Source/Runtime/Windows/D3D11RHI/Private/D3D11RenderTarget.cpp
@@ -480,7 +480,7 @@ struct FD3DFloat16
 
 		Result.Components.Sign = Components.Sign;
 		Result.Components.Exponent = Components.Exponent - 15 + 127; // Stored exponents are biased by half their range.
-		Result.Components.Mantissa = FMath::Min<uint32>(FMath::Floor((float)Components.Mantissa / 1024.0f * 8388608.0f),(1 << 23) - 1);
+		Result.Components.Mantissa = FMath::Min<uint32>(FMath::FloorToInt((float)Components.Mantissa / 1024.0f * 8388608.0f),(1 << 23) - 1);
 
 		return Result.Float;
 	}
@@ -508,13 +508,13 @@ struct FD3DFloatR11G11B10
 
 		Result[0].Components.Sign = 0;
 		Result[0].Components.Exponent = R_Exponent - 15 + 127;
-		Result[0].Components.Mantissa = FMath::Min<uint32>(FMath::Floor((float)R_Mantissa / 32.0f * 8388608.0f),(1 << 23) - 1);
+		Result[0].Components.Mantissa = FMath::Min<uint32>(FMath::FloorToInt((float)R_Mantissa / 32.0f * 8388608.0f),(1 << 23) - 1);
 		Result[1].Components.Sign = 0;
 		Result[1].Components.Exponent = G_Exponent - 15 + 127;
-		Result[1].Components.Mantissa = FMath::Min<uint32>(FMath::Floor((float)G_Mantissa / 64.0f * 8388608.0f),(1 << 23) - 1);
+		Result[1].Components.Mantissa = FMath::Min<uint32>(FMath::FloorToInt((float)G_Mantissa / 64.0f * 8388608.0f),(1 << 23) - 1);
 		Result[2].Components.Sign = 0;
 		Result[2].Components.Exponent = B_Exponent - 15 + 127;
-		Result[2].Components.Mantissa = FMath::Min<uint32>(FMath::Floor((float)B_Mantissa / 64.0f * 8388608.0f),(1 << 23) - 1);
+		Result[2].Components.Mantissa = FMath::Min<uint32>(FMath::FloorToInt((float)B_Mantissa / 64.0f * 8388608.0f),(1 << 23) - 1);
 
 		return FLinearColor(Result[0].Float, Result[1].Float, Result[2].Float);
 	}
diff --git a/Engine/Source/Runtime/Windows/D3D11RHI/Private/D3D11State.cpp b/Engine/Source/Runtime/Windows/D3D11RHI/Private/D3D11State.cpp
--- a/Engine/Source/Runtime/Windows/D3D11RHI/Private/D3D11State.cpp
+++ b/Engine/Source/Runtime/Windows/D3D11RHI/Private/D3D11State.cpp
@@ -199,7 +199,7 @@ FRasterizerStateRHIRef FD3D11DynamicRHI::RHICreateRasterizerState(const FRasteri
 	RasterizerDesc.FillMode = TranslateFillMode(Initializer.FillMode);
 	RasterizerDesc.SlopeScaledDepthBias = Initializer.SlopeScaleDepthBias;
 	RasterizerDesc.FrontCounterClockwise = true;
-	RasterizerDesc.DepthBias = FMath::Floor(Initializer.DepthBias * (float)(1 << 24));
+	RasterizerDesc.DepthBias = FMath::FloorToInt(Initializer.DepthBias * (float)(1 << 24));
 	RasterizerDesc.DepthClipEnable = true;
 	RasterizerDesc.MultisampleEnable = Initializer.bAllowMSAA;
 	RasterizerDesc.ScissorEnable = true;
diff --git a/Engine/Source/Runtime/Windows/D3D11RHI/Private/Windows/WindowsD3D11Device.cpp b/Engine/Source/Runtime/Windows/D3D11RHI/Private/Windows/WindowsD3D11Device.cpp
--- a/Engine/Source/Runtime/Windows/D3D11RHI/Private/Windows/WindowsD3D11Device.cpp
+++ b/Engine/Source/Runtime/Windows/D3D11RHI/Private/Windows/WindowsD3D11Device.cpp
@@ -503,7 +503,7 @@ void FD3D11DynamicRHI::InitD3DDevice()
 						float PoolSize = float(GPoolSizeVRAMPercentage) * 0.01f * float(GTotalGraphicsMemory);
 
 						// Truncate GTexturePoolSize to MB (but still counted in bytes)
-						GTexturePoolSize = int64(FGenericPlatformMath::TruncFloat(PoolSize / 1024.0f / 1024.0f)) * 1024 * 1024;
+						GTexturePoolSize = int64(FGenericPlatformMath::TruncToFloat(PoolSize / 1024.0f / 1024.0f)) * 1024 * 1024;
 
 						UE_LOG(LogRHI,Log,TEXT("Texture pool is %llu MB (%d%% of %llu MB)"),
 							GTexturePoolSize / 1024 / 1024,
