SUserWidget: The Construct function is no longer defined by the SLATE_USER_ARGS which allows us to define required args again.

diff --git a/Engine/Source/Runtime/Slate/Private/Framework/Testing/SUserWidgetTest.h b/Engine/Source/Runtime/Slate/Private/Framework/Testing/SUserWidgetTest.h
--- a/Engine/Source/Runtime/Slate/Private/Framework/Testing/SUserWidgetTest.h
+++ b/Engine/Source/Runtime/Slate/Private/Framework/Testing/SUserWidgetTest.h
@@ -12,6 +12,8 @@ class SUserWidgetExample : public SUserWidget
 		SLATE_ARGUMENT( FText, Title )
 	SLATE_END_ARGS()
 
+	virtual void Construct(const FArguments& InArgs) = 0;
+
 	virtual void DoStuff() = 0;
 };
 
diff --git a/Engine/Source/Runtime/SlateCore/Public/Widgets/DeclarativeSyntaxSupport.h b/Engine/Source/Runtime/SlateCore/Public/Widgets/DeclarativeSyntaxSupport.h
--- a/Engine/Source/Runtime/SlateCore/Public/Widgets/DeclarativeSyntaxSupport.h
+++ b/Engine/Source/Runtime/SlateCore/Public/Widgets/DeclarativeSyntaxSupport.h
@@ -89,7 +89,6 @@
 	public: \
 	static TSharedRef<WidgetType> New(); \
 	struct FArguments; \
-	virtual void Construct( const FArguments& InArgs ) = 0; \
 	struct FArguments : public TSlateBaseNamedArgs<WidgetType> \
 	{ \
 		typedef FArguments WidgetArgsType; \
