Slate - Added to IsEmpty and GetPosition to the FSlateRect

diff --git a/Engine/Source/Runtime/SlateCore/Public/Layout/SlateRect.h b/Engine/Source/Runtime/SlateCore/Public/Layout/SlateRect.h
--- a/Engine/Source/Runtime/SlateCore/Public/Layout/SlateRect.h
+++ b/Engine/Source/Runtime/SlateCore/Public/Layout/SlateRect.h
@@ -40,6 +40,15 @@ public:
 		return Left >= 0 && Right >= 0 && Bottom >=0 && Top >= 0;
 	}
 
+
+	/**
+	 * @return true, if the rectangle is empty.
+	 */
+	bool IsEmpty() const
+	{
+		return GetSize().Size() == 0.0f;
+	}
+
 	/**
 	 * Returns the size of the rectangle.
 	 *
@@ -60,6 +69,16 @@ public:
 		return FVector2D( Left, Top ) + GetSize() * 0.5f;
 	}
 
+	/**
+	 * Returns the left-top position of the rectangle
+	 * 
+	 * @return The left-top position.
+	 */
+	FVector2D GetPosition() const
+	{
+		return FVector2D( Left, Top );
+	}
+
 	/**
 	 * Return a rectangle that is contracted on each side by the amount specified in each margin.
 	 *
