Updated comments on how to use SUserWidget. IMO, we might want to delete this comment if it might rot like this quickly, and just point them to SUserWidgetExample as the compiler will enforce.

diff --git a/Engine/Source/Runtime/SlateCore/Public/Widgets/SUserWidget.h b/Engine/Source/Runtime/SlateCore/Public/Widgets/SUserWidget.h
--- a/Engine/Source/Runtime/SlateCore/Public/Widgets/SUserWidget.h
+++ b/Engine/Source/Runtime/SlateCore/Public/Widgets/SUserWidget.h
@@ -14,7 +14,7 @@
  * class SMyWidget : public SUserWidget
  * { 
  *   public:
- *   SLATE_BEGIN_ARGS( SMyWidget )
+ *   SLATE_USER_ARGS( SMyWidget )
  *	 {}
  *   SLATE_END_ARGS()
  *    
@@ -23,11 +23,12 @@
  *
  * SMyWidget.cpp
  * -------------
- * namespace Implementation {
- *   class SMyWidget : public SMyWidget
+ * namespace Implementation 
+ * {
+ *   class SMyWidget : public ::SMyWidget
  *   {
  *     public:
- *     void Construct( const FArgument& InArgs )
+ *     void Construct( const FArguments& InArgs )
  *     {
  *        SUserWidget::Construct( SUserWidget::FArguments()
  *        [
