More updates to SUserWidget usage example comment.

diff --git a/Engine/Source/Runtime/SlateCore/Public/Widgets/SUserWidget.h b/Engine/Source/Runtime/SlateCore/Public/Widgets/SUserWidget.h
--- a/Engine/Source/Runtime/SlateCore/Public/Widgets/SUserWidget.h
+++ b/Engine/Source/Runtime/SlateCore/Public/Widgets/SUserWidget.h
@@ -17,7 +17,10 @@
  *   SLATE_USER_ARGS( SMyWidget )
  *	 {}
  *   SLATE_END_ARGS()
- *    
+ *
+ *   // MUST Provide this function for SNew to call!
+ *   virtual void Construct( const FArguments& InArgs ) = 0;
+ *
  *   virtual void DoSomething() = 0;
  * };
  *
@@ -28,7 +31,7 @@
  *   class SMyWidget : public ::SMyWidget
  *   {
  *     public:
- *     void Construct( const FArguments& InArgs )
+ *     virtual void Construct( const FArguments& InArgs ) OVERRIDE
  *     {
  *        SUserWidget::Construct( SUserWidget::FArguments()
  *        [
