Minor fix.

diff --git a/Engine/Source/Runtime/SlateCore/Private/Widgets/SWindow.cpp b/Engine/Source/Runtime/SlateCore/Private/Widgets/SWindow.cpp
--- a/Engine/Source/Runtime/SlateCore/Private/Widgets/SWindow.cpp
+++ b/Engine/Source/Runtime/SlateCore/Private/Widgets/SWindow.cpp
@@ -1296,8 +1296,8 @@ FReply SWindow::OnMouseMove( const FGeometry& MyGeometry, const FPointerEvent& M
 		if (NativeWindow.IsValid())
 		{
 			NativeWindow->ReshapeWindow(
-				FMath::Trunc(MoveResizeRect.Left + MoveResizeOffset.X), FMath::Trunc(MoveResizeRect.Top + MoveResizeOffset.Y),
-				FMath::Trunc(MoveResizeRect.Right - MoveResizeRect.Left - MoveResizeOffset.X), FMath::Trunc(MoveResizeRect.Bottom - MoveResizeRect.Top - MoveResizeOffset.Y)
+				FMath::TruncToInt(MoveResizeRect.Left + MoveResizeOffset.X), FMath::TruncToInt(MoveResizeRect.Top + MoveResizeOffset.Y),
+				FMath::TruncToInt(MoveResizeRect.Right - MoveResizeRect.Left - MoveResizeOffset.X), FMath::TruncToInt(MoveResizeRect.Bottom - MoveResizeRect.Top - MoveResizeOffset.Y)
 				);
 		}
 		return FReply::Handled();
@@ -1308,8 +1308,8 @@ FReply SWindow::OnMouseMove( const FGeometry& MyGeometry, const FPointerEvent& M
 		if (NativeWindow.IsValid())
 		{
 			NativeWindow->ReshapeWindow(
-				FMath::Trunc(MoveResizeRect.Left), FMath::Trunc(MoveResizeRect.Top),
-				FMath::Trunc(MoveResizeRect.Right - MoveResizeRect.Left + MoveResizeOffset.X), FMath::Trunc(MoveResizeRect.Bottom - MoveResizeRect.Top + MoveResizeOffset.Y)
+				FMath::TruncToInt(MoveResizeRect.Left), FMath::TruncToInt(MoveResizeRect.Top),
+				FMath::TruncToInt(MoveResizeRect.Right - MoveResizeRect.Left + MoveResizeOffset.X), FMath::TruncToInt(MoveResizeRect.Bottom - MoveResizeRect.Top + MoveResizeOffset.Y)
 				);
 		}
 		return FReply::Handled();
@@ -1320,8 +1320,8 @@ FReply SWindow::OnMouseMove( const FGeometry& MyGeometry, const FPointerEvent& M
 		if (NativeWindow.IsValid())
 		{
 			NativeWindow->ReshapeWindow(
-				FMath::Trunc(MoveResizeRect.Left + MoveResizeOffset.X), FMath::Trunc(MoveResizeRect.Top),
-				FMath::Trunc(MoveResizeRect.Right - MoveResizeRect.Left - MoveResizeOffset.X), FMath::Trunc(MoveResizeRect.Bottom - MoveResizeRect.Top + MoveResizeOffset.Y)
+				FMath::TruncToInt(MoveResizeRect.Left + MoveResizeOffset.X), FMath::TruncToInt(MoveResizeRect.Top),
+				FMath::TruncToInt(MoveResizeRect.Right - MoveResizeRect.Left - MoveResizeOffset.X), FMath::TruncToInt(MoveResizeRect.Bottom - MoveResizeRect.Top + MoveResizeOffset.Y)
 				);
 		}
 		return FReply::Handled();
@@ -1332,8 +1332,8 @@ FReply SWindow::OnMouseMove( const FGeometry& MyGeometry, const FPointerEvent& M
 		if (NativeWindow.IsValid())
 		{
 			NativeWindow->ReshapeWindow(
-				FMath::Trunc(MoveResizeRect.Left), FMath::Trunc(MoveResizeRect.Top + MoveResizeOffset.Y),
-				FMath::Trunc(MoveResizeRect.Right - MoveResizeRect.Left + MoveResizeOffset.X), FMath::Trunc(MoveResizeRect.Bottom - MoveResizeRect.Top - MoveResizeOffset.Y)
+				FMath::TruncToInt(MoveResizeRect.Left), FMath::TruncToInt(MoveResizeRect.Top + MoveResizeOffset.Y),
+				FMath::TruncToInt(MoveResizeRect.Right - MoveResizeRect.Left + MoveResizeOffset.X), FMath::TruncToInt(MoveResizeRect.Bottom - MoveResizeRect.Top - MoveResizeOffset.Y)
 				);
 		}
 		return FReply::Handled();
@@ -1344,8 +1344,8 @@ FReply SWindow::OnMouseMove( const FGeometry& MyGeometry, const FPointerEvent& M
 		if (NativeWindow.IsValid())
 		{
 			NativeWindow->ReshapeWindow(
-				FMath::Trunc(MoveResizeRect.Left), FMath::Trunc(MoveResizeRect.Top + MoveResizeOffset.Y),
-				FMath::Trunc(MoveResizeRect.Right - MoveResizeRect.Left), FMath::Trunc(MoveResizeRect.Bottom - MoveResizeRect.Top - MoveResizeOffset.Y)
+				FMath::TruncToInt(MoveResizeRect.Left), FMath::TruncToInt(MoveResizeRect.Top + MoveResizeOffset.Y),
+				FMath::TruncToInt(MoveResizeRect.Right - MoveResizeRect.Left), FMath::TruncToInt(MoveResizeRect.Bottom - MoveResizeRect.Top - MoveResizeOffset.Y)
 				);
 		}
 		return FReply::Handled();
@@ -1356,8 +1356,8 @@ FReply SWindow::OnMouseMove( const FGeometry& MyGeometry, const FPointerEvent& M
 		if (NativeWindow.IsValid())
 		{
 			NativeWindow->ReshapeWindow(
-				FMath::Trunc(MoveResizeRect.Left), FMath::Trunc(MoveResizeRect.Top),
-				FMath::Trunc(MoveResizeRect.Right - MoveResizeRect.Left), FMath::Trunc(MoveResizeRect.Bottom - MoveResizeRect.Top + MoveResizeOffset.Y)
+				FMath::TruncToInt(MoveResizeRect.Left), FMath::TruncToInt(MoveResizeRect.Top),
+				FMath::TruncToInt(MoveResizeRect.Right - MoveResizeRect.Left), FMath::TruncToInt(MoveResizeRect.Bottom - MoveResizeRect.Top + MoveResizeOffset.Y)
 				);
 		}
 		return FReply::Handled();
@@ -1368,8 +1368,8 @@ FReply SWindow::OnMouseMove( const FGeometry& MyGeometry, const FPointerEvent& M
 		if (NativeWindow.IsValid())
 		{
 			NativeWindow->ReshapeWindow(
-				FMath::Trunc(MoveResizeRect.Left + MoveResizeOffset.X), FMath::Trunc(MoveResizeRect.Top),
-				FMath::Trunc(MoveResizeRect.Right - MoveResizeRect.Left - MoveResizeOffset.X), FMath::Trunc(MoveResizeRect.Bottom - MoveResizeRect.Top)
+				FMath::TruncToInt(MoveResizeRect.Left + MoveResizeOffset.X), FMath::TruncToInt(MoveResizeRect.Top),
+				FMath::TruncToInt(MoveResizeRect.Right - MoveResizeRect.Left - MoveResizeOffset.X), FMath::TruncToInt(MoveResizeRect.Bottom - MoveResizeRect.Top)
 				);
 		}
 		return FReply::Handled();
@@ -1380,8 +1380,8 @@ FReply SWindow::OnMouseMove( const FGeometry& MyGeometry, const FPointerEvent& M
 		if (NativeWindow.IsValid())
 		{
 			NativeWindow->ReshapeWindow(
-				FMath::Trunc(MoveResizeRect.Left), FMath::Trunc(MoveResizeRect.Top),
-				FMath::Trunc(MoveResizeRect.Right - MoveResizeRect.Left + MoveResizeOffset.X), FMath::Trunc(MoveResizeRect.Bottom - MoveResizeRect.Top)
+				FMath::TruncToInt(MoveResizeRect.Left), FMath::TruncToInt(MoveResizeRect.Top),
+				FMath::TruncToInt(MoveResizeRect.Right - MoveResizeRect.Left + MoveResizeOffset.X), FMath::TruncToInt(MoveResizeRect.Bottom - MoveResizeRect.Top)
 				);
 		}
 		return FReply::Handled();
