* Creating a custom assignment operator for the pointer event to allow us to continue to use const ref button list.  This should fix the issue with the viewport not releasing mouse capture.

diff --git a/Engine/Source/Runtime/SlateCore/Public/Input/Events.h b/Engine/Source/Runtime/SlateCore/Public/Input/Events.h
--- a/Engine/Source/Runtime/SlateCore/Public/Input/Events.h
+++ b/Engine/Source/Runtime/SlateCore/Public/Input/Events.h
@@ -544,12 +544,28 @@ public:
 	/** @return The change in gesture value since the last gesture event of the same type. */
 	const FVector2D& GetGestureDelta() const { return WheelOrGestureDelta; }
 
+	/** We override the assignment operator to allow generated code to compile with the const ref member. */
+	void operator=( const FPointerEvent& Other )
+	{
+		ScreenSpacePosition = Other.ScreenSpacePosition;
+		LastScreenSpacePosition = Other.LastScreenSpacePosition;
+		CursorDelta = Other.CursorDelta;
+		const_cast<TSet<FKey>&>(PressedButtons) = Other.PressedButtons;
+		EffectingButton = Other.EffectingButton;
+		UserIndex = Other.UserIndex;
+		PointerIndex = Other.PointerIndex;
+		TouchpadIndex = Other.TouchpadIndex;
+		bIsTouchEvent = Other.bIsTouchEvent;
+		GestureType = Other.GestureType;
+		WheelOrGestureDelta = Other.WheelOrGestureDelta;
+	}
+
 private:
 
 	FVector2D ScreenSpacePosition;
 	FVector2D LastScreenSpacePosition;
 	FVector2D CursorDelta;
-	TSet<FKey> PressedButtons;
+	const TSet<FKey>& PressedButtons;
 	FKey EffectingButton;
 	uint32 UserIndex;
 	uint32 PointerIndex;
