Slate API Updates
* Paint is now the non-virtual entry point to OnPaint to support required pre/post functionality. Widgets override OnPaint as usual, but calls to tell a widget to paint should be "Paint" not "OnPaint".
* ArrangeChildren -> OnArrangeChildren. ArrangeChildren is now the non-virtual entry point for the same reason as Paint/OnPain.
* OnPaint and OnArrageChildren are now private so they are not accidentally called.
* Added Stats for Paint time and ArrangeChildren time.
#BUN

diff --git a/Engine/Source/Editor/GraphEditor/Private/SGraphPanel.cpp b/Engine/Source/Editor/GraphEditor/Private/SGraphPanel.cpp
--- a/Engine/Source/Editor/GraphEditor/Private/SGraphPanel.cpp
+++ b/Engine/Source/Editor/GraphEditor/Private/SGraphPanel.cpp
@@ -274,7 +274,7 @@ int32 SGraphPanel::OnPaint( const FGeometry& AllottedGeometry, const FSlateRect&
 					const FWidgetStyle& NodeStyleToUse = (bNodeIsDifferent && !bNodeIsNotUsableInCurrentContext)? InWidgetStyle : FadedStyle;
 
 					// Draw the node.O
-					CurWidgetsMaxLayerId = CurWidget.Widget->OnPaint( CurWidget.Geometry, MyClippingRect, OutDrawElements, ChildLayerId, NodeStyleToUse, ShouldBeEnabled( bParentEnabled ) );
+					CurWidgetsMaxLayerId = CurWidget.Widget->Paint( CurWidget.Geometry, MyClippingRect, OutDrawElements, ChildLayerId, NodeStyleToUse, ShouldBeEnabled( bParentEnabled ) );
 				}
 
 				// Draw the node's overlay, if it has one.
diff --git a/Engine/Source/Editor/GraphEditor/Private/SNodePanel.cpp b/Engine/Source/Editor/GraphEditor/Private/SNodePanel.cpp
--- a/Engine/Source/Editor/GraphEditor/Private/SNodePanel.cpp
+++ b/Engine/Source/Editor/GraphEditor/Private/SNodePanel.cpp
@@ -198,7 +198,7 @@ namespace NodePanelDefs
 	static const float MouseZoomScaling = 0.05f;
 };
 
-void SNodePanel::ArrangeChildren(const FGeometry& AllottedGeometry, FArrangedChildren& ArrangedChildren) const
+void SNodePanel::OnArrangeChildren(const FGeometry& AllottedGeometry, FArrangedChildren& ArrangedChildren) const
 {
 	const TSlotlessChildren<SNode>& ChildrenToArrange = ArrangedChildren.Accepts(EVisibility::Hidden) ? Children : VisibleChildren;
 	// First pass nodes
diff --git a/Engine/Source/Editor/GraphEditor/Public/SNodePanel.h b/Engine/Source/Editor/GraphEditor/Public/SNodePanel.h
--- a/Engine/Source/Editor/GraphEditor/Public/SNodePanel.h
+++ b/Engine/Source/Editor/GraphEditor/Public/SNodePanel.h
@@ -339,7 +339,7 @@ public:
 	};
 
 	// SPanel interface
-	virtual void ArrangeChildren(const FGeometry& AllottedGeometry, FArrangedChildren& ArrangedChildren) const OVERRIDE;
+	virtual void OnArrangeChildren(const FGeometry& AllottedGeometry, FArrangedChildren& ArrangedChildren) const OVERRIDE;
 	virtual FVector2D ComputeDesiredSize() const OVERRIDE;
 	virtual FChildren* GetChildren() OVERRIDE;
 	// End of SPanel interface
diff --git a/Engine/Source/Editor/Persona/Private/STrack.cpp b/Engine/Source/Editor/Persona/Private/STrack.cpp
--- a/Engine/Source/Editor/Persona/Private/STrack.cpp
+++ b/Engine/Source/Editor/Persona/Private/STrack.cpp
@@ -346,7 +346,7 @@ void STrack::Construct( const FArguments& InArgs )
 	Font = FSlateFontInfo( FPaths::EngineContentDir() / TEXT("Slate/Fonts/Roboto-Regular.ttf"), 10 );
 }
 
-void STrack::ArrangeChildren( const FGeometry& AllottedGeometry, FArrangedChildren& ArrangedChildren ) const
+void STrack::OnArrangeChildren( const FGeometry& AllottedGeometry, FArrangedChildren& ArrangedChildren ) const
 {
 	for ( auto TrackIter = TrackNodes.CreateConstIterator() ; TrackIter ; ++TrackIter )
 	{
diff --git a/Engine/Source/Editor/Persona/Private/STrack.h b/Engine/Source/Editor/Persona/Private/STrack.h
--- a/Engine/Source/Editor/Persona/Private/STrack.h
+++ b/Engine/Source/Editor/Persona/Private/STrack.h
@@ -235,7 +235,7 @@ public:
 
 	void			Construct( const FArguments& InArgs );
 
-	virtual void	ArrangeChildren( const FGeometry& AllottedGeometry, FArrangedChildren& ArrangedChildren ) const;
+	virtual void	OnArrangeChildren( const FGeometry& AllottedGeometry, FArrangedChildren& ArrangedChildren ) const;
 	virtual int32	OnPaint( const FGeometry& AllottedGeometry, const FSlateRect& MyClippingRect, FSlateWindowElementList& OutDrawElements, int32 LayerId, const FWidgetStyle& InWidgetStyle, bool bParentEnabled ) const;
 
 	virtual FReply	OnDrop( const FGeometry& MyGeometry, const FDragDropEvent& DragDropEvent ) OVERRIDE;	
diff --git a/Engine/Source/Editor/Sequencer/Private/SAnimationOutlinerView.cpp b/Engine/Source/Editor/Sequencer/Private/SAnimationOutlinerView.cpp
--- a/Engine/Source/Editor/Sequencer/Private/SAnimationOutlinerView.cpp
+++ b/Engine/Source/Editor/Sequencer/Private/SAnimationOutlinerView.cpp
@@ -196,7 +196,7 @@ void SAnimationOutlinerView::GenerateWidgetForNode( TSharedRef<FSequencerDisplay
 }
 
 
-void SAnimationOutlinerView::ArrangeChildren( const FGeometry& AllottedGeometry, FArrangedChildren& ArrangedChildren ) const
+void SAnimationOutlinerView::OnArrangeChildren( const FGeometry& AllottedGeometry, FArrangedChildren& ArrangedChildren ) const
 {
 	const float Padding = SequencerLayoutConstants::NodePadding;
 	const float IndentAmount = SequencerLayoutConstants::IndentAmount;
diff --git a/Engine/Source/Editor/Sequencer/Private/SAnimationOutlinerView.h b/Engine/Source/Editor/Sequencer/Private/SAnimationOutlinerView.h
--- a/Engine/Source/Editor/Sequencer/Private/SAnimationOutlinerView.h
+++ b/Engine/Source/Editor/Sequencer/Private/SAnimationOutlinerView.h
@@ -92,7 +92,7 @@ private:
 	void GenerateWidgetForNode( TSharedRef<FSequencerDisplayNode>& InLayoutNode );
 
 	/** SPanel Interface */
-	virtual void ArrangeChildren( const FGeometry& AllottedGeometry, FArrangedChildren& ArrangedChildren ) const OVERRIDE;
+	virtual void OnArrangeChildren( const FGeometry& AllottedGeometry, FArrangedChildren& ArrangedChildren ) const OVERRIDE;
 
 	/**
 	 * Called when a node is selected in the outliner
diff --git a/Engine/Source/Editor/Sequencer/Private/SSequencerSectionAreaView.cpp b/Engine/Source/Editor/Sequencer/Private/SSequencerSectionAreaView.cpp
--- a/Engine/Source/Editor/Sequencer/Private/SSequencerSectionAreaView.cpp
+++ b/Engine/Source/Editor/Sequencer/Private/SSequencerSectionAreaView.cpp
@@ -129,7 +129,7 @@ int32 SSequencerSectionAreaView::OnPaint( const FGeometry& AllottedGeometry, con
 	{
 		FArrangedWidget& CurWidget = ArrangedChildren(ChildIndex);
 		FSlateRect ChildClipRect = MyClippingRect.IntersectionWith( CurWidget.Geometry.GetClippingRect() );
-		const int32 CurWidgetsMaxLayerId = CurWidget.Widget->OnPaint( CurWidget.Geometry, ChildClipRect, OutDrawElements, LayerId+1, InWidgetStyle, ShouldBeEnabled( bParentEnabled ) );
+		const int32 CurWidgetsMaxLayerId = CurWidget.Widget->Paint( CurWidget.Geometry, ChildClipRect, OutDrawElements, LayerId+1, InWidgetStyle, ShouldBeEnabled( bParentEnabled ) );
 	}
 
 	return LayerId+1;
@@ -146,7 +146,7 @@ FReply SSequencerSectionAreaView::OnMouseButtonDown( const FGeometry& MyGeometry
 	return FReply::Handled();
 }
 
-void SSequencerSectionAreaView::ArrangeChildren( const FGeometry& AllottedGeometry, FArrangedChildren& ArrangedChildren ) const
+void SSequencerSectionAreaView::OnArrangeChildren( const FGeometry& AllottedGeometry, FArrangedChildren& ArrangedChildren ) const
 {
 	int32 MaxRowIndex = 0;
 	for( int32 WidgetIndex = 0; WidgetIndex < Children.Num(); ++WidgetIndex )
diff --git a/Engine/Source/Editor/Sequencer/Private/SSequencerSectionAreaView.h b/Engine/Source/Editor/Sequencer/Private/SSequencerSectionAreaView.h
--- a/Engine/Source/Editor/Sequencer/Private/SSequencerSectionAreaView.h
+++ b/Engine/Source/Editor/Sequencer/Private/SSequencerSectionAreaView.h
@@ -39,7 +39,7 @@ private:
 	virtual FReply OnMouseButtonDown( const FGeometry& MyGeometry, const FPointerEvent& MouseEvent ) OVERRIDE;
 
 	/** SPanel Interface */
-	virtual void ArrangeChildren( const FGeometry& AllottedGeometry, FArrangedChildren& ArrangedChildren ) const OVERRIDE;
+	virtual void OnArrangeChildren( const FGeometry& AllottedGeometry, FArrangedChildren& ArrangedChildren ) const OVERRIDE;
 
 	/**
 	 * Returns the visibility of a section
diff --git a/Engine/Source/Editor/WorldBrowser/Private/Tiles/SWorldComposition.cpp b/Engine/Source/Editor/WorldBrowser/Private/Tiles/SWorldComposition.cpp
--- a/Engine/Source/Editor/WorldBrowser/Private/Tiles/SWorldComposition.cpp
+++ b/Engine/Source/Editor/WorldBrowser/Private/Tiles/SWorldComposition.cpp
@@ -253,7 +253,7 @@ public:
 	}
 	
 	/**  SWidget interface */
-	virtual void ArrangeChildren( const FGeometry& AllottedGeometry, FArrangedChildren& ArrangedChildren ) const OVERRIDE
+	virtual void OnArrangeChildren( const FGeometry& AllottedGeometry, FArrangedChildren& ArrangedChildren ) const OVERRIDE
 	{
 		for (int32 ChildIndex=0; ChildIndex < VisibleChildren.Num(); ++ChildIndex)
 		{
@@ -306,7 +306,7 @@ public:
 			TSharedRef<SWorldTileItem> ChildNode = StaticCastSharedRef<SWorldTileItem>(CurWidget.Widget);
 		
 			ChildNode->bAffectedByMarquee = SelectionToVisualize->Contains(ChildNode->GetObjectBeingDisplayed());
-			LayerId = CurWidget.Widget->OnPaint(CurWidget.Geometry, MyClippingRect, OutDrawElements, NodesLayerId, InWidgetStyle, ShouldBeEnabled(bParentEnabled));
+			LayerId = CurWidget.Widget->Paint(CurWidget.Geometry, MyClippingRect, OutDrawElements, NodesLayerId, InWidgetStyle, ShouldBeEnabled(bParentEnabled));
 			ChildNode->bAffectedByMarquee = false;
 		}
 		
diff --git a/Engine/Source/Runtime/Slate/Private/Framework/Application/SlateApplication.cpp b/Engine/Source/Runtime/Slate/Private/Framework/Application/SlateApplication.cpp
--- a/Engine/Source/Runtime/Slate/Private/Framework/Application/SlateApplication.cpp
+++ b/Engine/Source/Runtime/Slate/Private/Framework/Application/SlateApplication.cpp
@@ -550,7 +550,6 @@ void FSlateApplication::DrawWindowAndChildren( const TSharedRef<SWindow>& Window
 		FGeometry WindowGeometry = WindowToDraw->GetWindowGeometryInWindow();
 		int32 MaxLayerId = 0;
 		{
-			SCOPE_CYCLE_COUNTER( STAT_SlateOnPaint );
 			MaxLayerId = WindowToDraw->OnPaint( WindowGeometry, WindowToDraw->GetClippingRectangleInWindow(), WindowElementList, 0, FWidgetStyle(), WindowToDraw->IsEnabled() );
 		}
 
diff --git a/Engine/Source/Runtime/Slate/Private/Framework/Docking/SDockingTabWell.cpp b/Engine/Source/Runtime/Slate/Private/Framework/Docking/SDockingTabWell.cpp
--- a/Engine/Source/Runtime/Slate/Private/Framework/Docking/SDockingTabWell.cpp
+++ b/Engine/Source/Runtime/Slate/Private/Framework/Docking/SDockingTabWell.cpp
@@ -62,7 +62,7 @@ void SDockingTabWell::AddTab( const TSharedRef<SDockTab>& InTab, int32 AtIndex )
 }
 
 
-void SDockingTabWell::ArrangeChildren( const FGeometry& AllottedGeometry, FArrangedChildren& ArrangedChildren ) const
+void SDockingTabWell::OnArrangeChildren( const FGeometry& AllottedGeometry, FArrangedChildren& ArrangedChildren ) const
 {
 	// The specialized TabWell is dedicated to arranging tabs.
 	// Tabs have uniform sizing (all tabs the same size).
@@ -136,7 +136,7 @@ int32 SDockingTabWell::OnPaint( const FGeometry& AllottedGeometry, const FSlateR
 		else
 		{
 			FSlateRect ChildClipRect = MyClippingRect.IntersectionWith( CurWidget.Geometry.GetClippingRect() );
-			const int32 CurWidgetsMaxLayerId = CurWidget.Widget->OnPaint( CurWidget.Geometry, ChildClipRect, OutDrawElements, MaxLayerId, InWidgetStyle, ShouldBeEnabled( bParentEnabled ) );
+			const int32 CurWidgetsMaxLayerId = CurWidget.Widget->Paint( CurWidget.Geometry, ChildClipRect, OutDrawElements, MaxLayerId, InWidgetStyle, ShouldBeEnabled( bParentEnabled ) );
 			MaxLayerId = FMath::Max( MaxLayerId, CurWidgetsMaxLayerId );
 		}
 	}
@@ -145,7 +145,7 @@ int32 SDockingTabWell::OnPaint( const FGeometry& AllottedGeometry, const FSlateR
 	if (ForegroundTab != TSharedPtr<SDockTab>())
 	{
 		FSlateRect ChildClipRect = MyClippingRect.IntersectionWith( ForegroundTabGeometry->Geometry.GetClippingRect() );
-		const int32 CurWidgetsMaxLayerId = ForegroundTabGeometry->Widget->OnPaint( ForegroundTabGeometry->Geometry, ChildClipRect, OutDrawElements, MaxLayerId, InWidgetStyle, ShouldBeEnabled( bParentEnabled ) );
+		const int32 CurWidgetsMaxLayerId = ForegroundTabGeometry->Widget->Paint( ForegroundTabGeometry->Geometry, ChildClipRect, OutDrawElements, MaxLayerId, InWidgetStyle, ShouldBeEnabled( bParentEnabled ) );
 		MaxLayerId = FMath::Max( MaxLayerId, CurWidgetsMaxLayerId );
 	}
 
diff --git a/Engine/Source/Runtime/Slate/Private/Framework/Docking/SDockingTabWell.h b/Engine/Source/Runtime/Slate/Private/Framework/Docking/SDockingTabWell.h
--- a/Engine/Source/Runtime/Slate/Private/Framework/Docking/SDockingTabWell.h
+++ b/Engine/Source/Runtime/Slate/Private/Framework/Docking/SDockingTabWell.h
@@ -84,7 +84,7 @@ public:
 	FVector2D ComputeChildSize( const FGeometry& AllottedGeometry ) const;
 
 	// SWidget interface
-	virtual void ArrangeChildren( const FGeometry& AllottedGeometry, FArrangedChildren& ArrangedChildren ) const OVERRIDE;
+	virtual void OnArrangeChildren( const FGeometry& AllottedGeometry, FArrangedChildren& ArrangedChildren ) const OVERRIDE;
 	virtual int32 OnPaint( const FGeometry& AllottedGeometry, const FSlateRect& MyClippingRect, FSlateWindowElementList& OutDrawElements, int32 LayerId, const FWidgetStyle& InWidgetStyle, bool bParentEnabled ) const OVERRIDE;
 	virtual FVector2D ComputeDesiredSize() const OVERRIDE;
 	virtual FChildren* GetChildren() OVERRIDE;
diff --git a/Engine/Source/Runtime/Slate/Private/Framework/MultiBox/SClippingHorizontalBox.cpp b/Engine/Source/Runtime/Slate/Private/Framework/MultiBox/SClippingHorizontalBox.cpp
--- a/Engine/Source/Runtime/Slate/Private/Framework/MultiBox/SClippingHorizontalBox.cpp
+++ b/Engine/Source/Runtime/Slate/Private/Framework/MultiBox/SClippingHorizontalBox.cpp
@@ -11,13 +11,13 @@ void SClippingHorizontalBox::Tick( const FGeometry& AllottedGeometry, const doub
 	ClippedIdx = ArrangedChildren.Num() - 1;
 }
 
-void SClippingHorizontalBox::ArrangeChildren( const FGeometry& AllottedGeometry, FArrangedChildren& ArrangedChildren ) const
+void SClippingHorizontalBox::OnArrangeChildren( const FGeometry& AllottedGeometry, FArrangedChildren& ArrangedChildren ) const
 {
 	// If WrapButton hasn't been initialized, that means AddWrapButton() hasn't 
 	// been called and this method isn't going to behave properly
 	check(WrapButton.IsValid());
 
-	SHorizontalBox::ArrangeChildren(AllottedGeometry, ArrangedChildren);
+	SHorizontalBox::OnArrangeChildren(AllottedGeometry, ArrangedChildren);
 
 	// Remove children that are clipped by the allotted geometry
 	const int32 NumChildren = ArrangedChildren.Num(); 
@@ -62,8 +62,9 @@ int32 SClippingHorizontalBox::OnPaint( const FGeometry& AllottedGeometry, const
 	ArrangeChildren(AllottedGeometry, ClippedArrangedChildren);
 	
 	// Get the non-clipped children info
+	// @todo umg: One should not call the virtual OnArrangeChildren, one should only call ArrangeChildren.
 	FArrangedChildren ArrangedChildren(EVisibility::Visible);
-	SBoxPanel::ArrangeChildren(AllottedGeometry, ArrangedChildren);
+	SBoxPanel::OnArrangeChildren(AllottedGeometry, ArrangedChildren);
 	
 	if ((ClippedArrangedChildren.Num() != 0) && (ArrangedChildren.Num() != 0))
 	{
diff --git a/Engine/Source/Runtime/Slate/Private/Framework/MultiBox/SClippingHorizontalBox.h b/Engine/Source/Runtime/Slate/Private/Framework/MultiBox/SClippingHorizontalBox.h
--- a/Engine/Source/Runtime/Slate/Private/Framework/MultiBox/SClippingHorizontalBox.h
+++ b/Engine/Source/Runtime/Slate/Private/Framework/MultiBox/SClippingHorizontalBox.h
@@ -26,7 +26,7 @@ public:
 	virtual FVector2D ComputeDesiredSize() const OVERRIDE;
 
 	/** SPanel interface */
-	virtual void ArrangeChildren( const FGeometry& AllottedGeometry, FArrangedChildren& ArrangedChildren ) const OVERRIDE;
+	virtual void OnArrangeChildren( const FGeometry& AllottedGeometry, FArrangedChildren& ArrangedChildren ) const OVERRIDE;
 
 	/** Construct this widget */
 	void Construct( const FArguments& InArgs );
diff --git a/Engine/Source/Runtime/Slate/Private/Framework/Text/SlateHyperlinkRun.cpp b/Engine/Source/Runtime/Slate/Private/Framework/Text/SlateHyperlinkRun.cpp
--- a/Engine/Source/Runtime/Slate/Private/Framework/Text/SlateHyperlinkRun.cpp
+++ b/Engine/Source/Runtime/Slate/Private/Framework/Text/SlateHyperlinkRun.cpp
@@ -68,7 +68,7 @@ int32 FSlateHyperlinkRun::OnPaint( const FTextLayout::FLineView& Line, const TSh
 	const TSharedRef< FWidgetLayoutBlock > WidgetBlock = StaticCastSharedRef< FWidgetLayoutBlock >( Block );
 
 	FGeometry WidgetGeometry = AllottedGeometry.MakeChild( Block->GetLocationOffset() * ( 1 / AllottedGeometry.Scale ), Block->GetSize() * ( 1 / AllottedGeometry.Scale ), 1.0f );
-	return WidgetBlock->GetWidget()->OnPaint( WidgetGeometry, MyClippingRect, OutDrawElements, LayerId, InWidgetStyle, bParentEnabled );
+	return WidgetBlock->GetWidget()->Paint( WidgetGeometry, MyClippingRect, OutDrawElements, LayerId, InWidgetStyle, bParentEnabled );
 }
 
 FChildren* FSlateHyperlinkRun::GetChildren()
diff --git a/Engine/Source/Runtime/Slate/Private/Framework/Text/SlateSimpleRunHighlighter.cpp b/Engine/Source/Runtime/Slate/Private/Framework/Text/SlateSimpleRunHighlighter.cpp
--- a/Engine/Source/Runtime/Slate/Private/Framework/Text/SlateSimpleRunHighlighter.cpp
+++ b/Engine/Source/Runtime/Slate/Private/Framework/Text/SlateSimpleRunHighlighter.cpp
@@ -10,7 +10,7 @@ FSlateSimpleRunHighlighter::FSlateSimpleRunHighlighter()
 
 }
 
-void FSlateSimpleRunHighlighter::ArrangeChildren( const TSharedRef< ILayoutBlock >& Block, const FGeometry& AllottedGeometry, FArrangedChildren& ArrangedChildren ) const 
+void FSlateSimpleRunHighlighter::OnArrangeChildren( const TSharedRef< ILayoutBlock >& Block, const FGeometry& AllottedGeometry, FArrangedChildren& ArrangedChildren ) const 
 {
 	//No Widgets
 }
diff --git a/Engine/Source/Runtime/Slate/Private/Framework/Text/SlateWidgetRun.cpp b/Engine/Source/Runtime/Slate/Private/Framework/Text/SlateWidgetRun.cpp
--- a/Engine/Source/Runtime/Slate/Private/Framework/Text/SlateWidgetRun.cpp
+++ b/Engine/Source/Runtime/Slate/Private/Framework/Text/SlateWidgetRun.cpp
@@ -46,7 +46,7 @@ TSharedRef< ILayoutBlock > FSlateWidgetRun::CreateBlock( int32 StartIndex, int32
 int32 FSlateWidgetRun::OnPaint( const FTextLayout::FLineView& Line, const TSharedRef< ILayoutBlock >& Block, const FTextBlockStyle& DefaultStyle, const FGeometry& AllottedGeometry, const FSlateRect& MyClippingRect, FSlateWindowElementList& OutDrawElements, int32 LayerId, const FWidgetStyle& InWidgetStyle, bool bParentEnabled ) const 
 {
 	FGeometry WidgetGeometry = AllottedGeometry.MakeChild( Block->GetLocationOffset() * ( 1 / AllottedGeometry.Scale ), Block->GetSize() * ( 1 / AllottedGeometry.Scale ), 1.0f );
-	return Info.Widget->OnPaint( WidgetGeometry, MyClippingRect, OutDrawElements, LayerId, InWidgetStyle, bParentEnabled );
+	return Info.Widget->Paint( WidgetGeometry, MyClippingRect, OutDrawElements, LayerId, InWidgetStyle, bParentEnabled );
 }
 
 FChildren* FSlateWidgetRun::GetChildren()
diff --git a/Engine/Source/Runtime/Slate/Private/SlateModule.cpp b/Engine/Source/Runtime/Slate/Private/SlateModule.cpp
--- a/Engine/Source/Runtime/Slate/Private/SlateModule.cpp
+++ b/Engine/Source/Runtime/Slate/Private/SlateModule.cpp
@@ -22,7 +22,6 @@ DEFINE_STAT(STAT_SlateOnPaint_SCompoundWidget);
 DEFINE_STAT(STAT_SlateOnPaint_SBox);
 DEFINE_STAT(STAT_SlateOnPaint_SImage);
 DEFINE_STAT(STAT_SlateOnPaint_SViewport);
-DEFINE_STAT(STAT_SlateOnPaint);
 DEFINE_STAT(STAT_SlateOnPaint_SSlider);
 
 DEFINE_STAT(STAT_SlateDrawWindowTime);
diff --git a/Engine/Source/Runtime/Slate/Private/Widgets/Colors/SColorThemes.cpp b/Engine/Source/Runtime/Slate/Private/Widgets/Colors/SColorThemes.cpp
--- a/Engine/Source/Runtime/Slate/Private/Widgets/Colors/SColorThemes.cpp
+++ b/Engine/Source/Runtime/Slate/Private/Widgets/Colors/SColorThemes.cpp
@@ -379,7 +379,7 @@ EVisibility SThemeColorBlock::OnGetAlphaVisibility() const
 * @param AllottedGeometry    The geometry allotted for this widget by its parent.
 * @param ArrangedChildren    The array to which to add the WidgetGeometries that represent the arranged children.
 */
-void SThemeColorBlocksBar::ArrangeChildren( const FGeometry& AllottedGeometry, FArrangedChildren& ArrangedChildren ) const
+void SThemeColorBlocksBar::OnArrangeChildren( const FGeometry& AllottedGeometry, FArrangedChildren& ArrangedChildren ) const
 {
 	bool bPlaceHolderExists = NewColorPlaceholder.IsValid();
 
diff --git a/Engine/Source/Runtime/Slate/Private/Widgets/Layout/SBox.cpp b/Engine/Source/Runtime/Slate/Private/Widgets/Layout/SBox.cpp
--- a/Engine/Source/Runtime/Slate/Private/Widgets/Layout/SBox.cpp
+++ b/Engine/Source/Runtime/Slate/Private/Widgets/Layout/SBox.cpp
@@ -38,7 +38,7 @@ FVector2D SBox::ComputeDesiredSize() const
 	return FVector2D::ZeroVector;
 }
 
-void SBox::ArrangeChildren( const FGeometry& AllottedGeometry, FArrangedChildren& ArrangedChildren ) const
+void SBox::OnArrangeChildren( const FGeometry& AllottedGeometry, FArrangedChildren& ArrangedChildren ) const
 {
 	const EVisibility& MyCurrentVisibility = this->GetVisibility();
 	if ( ArrangedChildren.Accepts( MyCurrentVisibility ) )
@@ -81,7 +81,7 @@ int32 SBox::OnPaint( const FGeometry& AllottedGeometry, const FSlateRect& MyClip
 
 		const FSlateRect ChildClippingRect = AllottedGeometry.GetClippingRect().InsetBy( ChildSlot.SlotPadding.Get() * AllottedGeometry.Scale ).IntersectionWith(MyClippingRect);
 
-		return TheChild.Widget->OnPaint( TheChild.Geometry, ChildClippingRect, OutDrawElements, LayerId, InWidgetStyle, ShouldBeEnabled( bParentEnabled ) );
+		return TheChild.Widget->Paint( TheChild.Geometry, ChildClippingRect, OutDrawElements, LayerId, InWidgetStyle, ShouldBeEnabled( bParentEnabled ) );
 	}
 	return LayerId;
 }
diff --git a/Engine/Source/Runtime/Slate/Private/Widgets/Layout/SDPIScaler.cpp b/Engine/Source/Runtime/Slate/Private/Widgets/Layout/SDPIScaler.cpp
--- a/Engine/Source/Runtime/Slate/Private/Widgets/Layout/SDPIScaler.cpp
+++ b/Engine/Source/Runtime/Slate/Private/Widgets/Layout/SDPIScaler.cpp
@@ -9,7 +9,7 @@ void SDPIScaler::Construct( const FArguments& InArgs )
 	this->DPIScale = InArgs._DPIScale;
 }
 
-void SDPIScaler::ArrangeChildren( const FGeometry& AllottedGeometry, FArrangedChildren& ArrangedChildren ) const
+void SDPIScaler::OnArrangeChildren( const FGeometry& AllottedGeometry, FArrangedChildren& ArrangedChildren ) const
 {
 	const EVisibility MyVisibility = this->GetVisibility();
 	if ( ArrangedChildren.Accepts( MyVisibility ) )
diff --git a/Engine/Source/Runtime/Slate/Private/Widgets/Layout/SFxWidget.cpp b/Engine/Source/Runtime/Slate/Private/Widgets/Layout/SFxWidget.cpp
--- a/Engine/Source/Runtime/Slate/Private/Widgets/Layout/SFxWidget.cpp
+++ b/Engine/Source/Runtime/Slate/Private/Widgets/Layout/SFxWidget.cpp
@@ -93,7 +93,7 @@ int32 SFxWidget::OnPaint( const FGeometry& AllottedGeometry, const FSlateRect& M
 			.BlendColorAndOpacityTint(ColorAndOpacity.Get())
 			.SetForegroundColor( ForegroundColor );
 
-		return TheChild.Widget->OnPaint( TheChild.Geometry, ChildClippingRect, OutDrawElements, LayerId + 1, CompoundedWidgetStyle, ShouldBeEnabled( bParentEnabled ) );
+		return TheChild.Widget->Paint( TheChild.Geometry, ChildClippingRect, OutDrawElements, LayerId + 1, CompoundedWidgetStyle, ShouldBeEnabled( bParentEnabled ) );
 	}
 	return LayerId;
 
@@ -107,7 +107,7 @@ FVector2D SFxWidget::ComputeDesiredSize() const
 }
 
 
-void SFxWidget::ArrangeChildren( const FGeometry& AllottedGeometry, FArrangedChildren& ArrangedChildren ) const
+void SFxWidget::OnArrangeChildren( const FGeometry& AllottedGeometry, FArrangedChildren& ArrangedChildren ) const
 {
 	const EVisibility MyVisibility = this->GetVisibility();
 	if ( ArrangedChildren.Accepts( MyVisibility ) )
diff --git a/Engine/Source/Runtime/Slate/Private/Widgets/Layout/SGridPanel.cpp b/Engine/Source/Runtime/Slate/Private/Widgets/Layout/SGridPanel.cpp
--- a/Engine/Source/Runtime/Slate/Private/Widgets/Layout/SGridPanel.cpp
+++ b/Engine/Source/Runtime/Slate/Private/Widgets/Layout/SGridPanel.cpp
@@ -68,7 +68,7 @@ int32 SGridPanel::OnPaint( const FGeometry& AllottedGeometry, const FSlateRect&
 				LayerId = MaxLayerId+1;
 			}
 
-			const int32 CurWidgetsMaxLayerId = CurWidget.Widget->OnPaint(
+			const int32 CurWidgetsMaxLayerId = CurWidget.Widget->Paint(
 				CurWidget.Geometry,
 				ChildClipRect,
 				OutDrawElements,
@@ -91,7 +91,7 @@ int32 SGridPanel::OnPaint( const FGeometry& AllottedGeometry, const FSlateRect&
 }
 
 
-void SGridPanel::ArrangeChildren( const FGeometry& AllottedGeometry, FArrangedChildren& ArrangedChildren ) const
+void SGridPanel::OnArrangeChildren( const FGeometry& AllottedGeometry, FArrangedChildren& ArrangedChildren ) const
 {
 	// PREPARE PHASE
 	// Prepare some data for arranging children.
diff --git a/Engine/Source/Runtime/Slate/Private/Widgets/Layout/SScrollBar.cpp b/Engine/Source/Runtime/Slate/Private/Widgets/Layout/SScrollBar.cpp
--- a/Engine/Source/Runtime/Slate/Private/Widgets/Layout/SScrollBar.cpp
+++ b/Engine/Source/Runtime/Slate/Private/Widgets/Layout/SScrollBar.cpp
@@ -111,7 +111,7 @@ public:
 	 * @param AllottedGeometry    The geometry allotted for this widget by its parent.
 	 * @param ArrangedChildren    The array to which to add the WidgetGeometries that represent the arranged children.
 	 */
-	virtual void ArrangeChildren( const FGeometry& AllottedGeometry, FArrangedChildren& ArrangedChildren ) const
+	virtual void OnArrangeChildren( const FGeometry& AllottedGeometry, FArrangedChildren& ArrangedChildren ) const
 	{
 		const float Width = AllottedGeometry.Size.X;
 		const float Height = AllottedGeometry.Size.Y;
diff --git a/Engine/Source/Runtime/Slate/Private/Widgets/Layout/SScrollBox.cpp b/Engine/Source/Runtime/Slate/Private/Widgets/Layout/SScrollBox.cpp
--- a/Engine/Source/Runtime/Slate/Private/Widgets/Layout/SScrollBox.cpp
+++ b/Engine/Source/Runtime/Slate/Private/Widgets/Layout/SScrollBox.cpp
@@ -36,7 +36,7 @@ public:
 
 public:
 
-	virtual void ArrangeChildren( const FGeometry& AllottedGeometry, FArrangedChildren& ArrangedChildren ) const OVERRIDE
+	virtual void OnArrangeChildren( const FGeometry& AllottedGeometry, FArrangedChildren& ArrangedChildren ) const OVERRIDE
 	{
 		
 		float CurChildOffset = -PhysicalOffset;
diff --git a/Engine/Source/Runtime/Slate/Private/Widgets/Layout/SSplitter.cpp b/Engine/Source/Runtime/Slate/Private/Widgets/Layout/SSplitter.cpp
--- a/Engine/Source/Runtime/Slate/Private/Widgets/Layout/SSplitter.cpp
+++ b/Engine/Source/Runtime/Slate/Private/Widgets/Layout/SSplitter.cpp
@@ -67,7 +67,7 @@ void SSplitter::Construct( const SSplitter::FArguments& InArgs )
 * Panels arrange their children in a space described by the AllottedGeometry parameter. The results of the arrangement
 * should be returned by appending a FArrangedWidget pair for every child widget. See StackPanel for an example
 */
-void SSplitter::ArrangeChildren( const FGeometry& AllottedGeometry, FArrangedChildren& ArrangedChildren ) const
+void SSplitter::OnArrangeChildren( const FGeometry& AllottedGeometry, FArrangedChildren& ArrangedChildren ) const
 {
 	const int32 AxisIndex = (Orientation == Orient_Horizontal) ? 0 : 1;
 
@@ -665,7 +665,7 @@ void SSplitter2x2::Construct( const FArguments& InArgs )
 }
 
 
-void SSplitter2x2::ArrangeChildren( const FGeometry& AllottedGeometry, FArrangedChildren& ArrangedChildren ) const
+void SSplitter2x2::OnArrangeChildren( const FGeometry& AllottedGeometry, FArrangedChildren& ArrangedChildren ) const
 {
 	int32 NumNonCollapsedChildren = 0;
 	FVector2D CoefficientTotal(0,0);
diff --git a/Engine/Source/Runtime/Slate/Private/Widgets/Layout/SUniformGridPanel.cpp b/Engine/Source/Runtime/Slate/Private/Widgets/Layout/SUniformGridPanel.cpp
--- a/Engine/Source/Runtime/Slate/Private/Widgets/Layout/SUniformGridPanel.cpp
+++ b/Engine/Source/Runtime/Slate/Private/Widgets/Layout/SUniformGridPanel.cpp
@@ -24,7 +24,7 @@ void SUniformGridPanel::Construct( const FArguments& InArgs )
 }
 
 
-void SUniformGridPanel::ArrangeChildren( const FGeometry& AllottedGeometry, FArrangedChildren& ArrangedChildren ) const
+void SUniformGridPanel::OnArrangeChildren( const FGeometry& AllottedGeometry, FArrangedChildren& ArrangedChildren ) const
 {
 	const FVector2D CellSize(AllottedGeometry.Size.X / NumColumns, AllottedGeometry.Size.Y / NumRows);
 	const FMargin& CurrentSlotPadding(SlotPadding.Get());
diff --git a/Engine/Source/Runtime/Slate/Private/Widgets/Layout/SWidgetSwitcher.cpp b/Engine/Source/Runtime/Slate/Private/Widgets/Layout/SWidgetSwitcher.cpp
--- a/Engine/Source/Runtime/Slate/Private/Widgets/Layout/SWidgetSwitcher.cpp
+++ b/Engine/Source/Runtime/Slate/Private/Widgets/Layout/SWidgetSwitcher.cpp
@@ -90,7 +90,7 @@ TAttribute<FVector2D> ContentScale = FVector2D::UnitVector;
 /* SCompoundWidget interface
  *****************************************************************************/
 
-void SWidgetSwitcher::ArrangeChildren( const FGeometry& AllottedGeometry, FArrangedChildren& ArrangedChildren ) const
+void SWidgetSwitcher::OnArrangeChildren( const FGeometry& AllottedGeometry, FArrangedChildren& ArrangedChildren ) const
 {
 	if (AllChildren.Num() > 0)
 	{
diff --git a/Engine/Source/Runtime/Slate/Private/Widgets/Layout/SWrapBox.cpp b/Engine/Source/Runtime/Slate/Private/Widgets/Layout/SWrapBox.cpp
--- a/Engine/Source/Runtime/Slate/Private/Widgets/Layout/SWrapBox.cpp
+++ b/Engine/Source/Runtime/Slate/Private/Widgets/Layout/SWrapBox.cpp
@@ -51,7 +51,7 @@ void SWrapBox::Tick( const FGeometry& AllottedGeometry, const double InCurrentTi
 	}
 }
 
-void SWrapBox::ArrangeChildren( const FGeometry& AllottedGeometry, FArrangedChildren& ArrangedChildren ) const
+void SWrapBox::OnArrangeChildren( const FGeometry& AllottedGeometry, FArrangedChildren& ArrangedChildren ) const
 {
 	const float MyPreferredWidth = PreferredWidth.Get();
 
diff --git a/Engine/Source/Runtime/Slate/Private/Widgets/SCanvas.cpp b/Engine/Source/Runtime/Slate/Private/Widgets/SCanvas.cpp
--- a/Engine/Source/Runtime/Slate/Private/Widgets/SCanvas.cpp
+++ b/Engine/Source/Runtime/Slate/Private/Widgets/SCanvas.cpp
@@ -45,7 +45,7 @@ int32 SCanvas::RemoveSlot( const TSharedRef<SWidget>& SlotWidget )
 /* SWidget overrides
  *****************************************************************************/
 
-void SCanvas::ArrangeChildren( const FGeometry& AllottedGeometry, FArrangedChildren& ArrangedChildren ) const
+void SCanvas::OnArrangeChildren( const FGeometry& AllottedGeometry, FArrangedChildren& ArrangedChildren ) const
 {
 	if (Children.Num() > 0)
 	{
@@ -105,7 +105,7 @@ int32 SCanvas::OnPaint( const FGeometry& AllottedGeometry, const FSlateRect& MyC
 	{
 		FArrangedWidget& CurWidget = ArrangedChildren(ChildIndex);
 		FSlateRect ChildClipRect = MyClippingRect.IntersectionWith(CurWidget.Geometry.GetClippingRect());
-		const int32 CurWidgetsMaxLayerId = CurWidget.Widget->OnPaint(CurWidget.Geometry, ChildClipRect, OutDrawElements, MaxLayerId + 1, InWidgetStyle, ShouldBeEnabled(bParentEnabled));
+		const int32 CurWidgetsMaxLayerId = CurWidget.Widget->Paint(CurWidget.Geometry, ChildClipRect, OutDrawElements, MaxLayerId + 1, InWidgetStyle, ShouldBeEnabled(bParentEnabled));
 
 		MaxLayerId = FMath::Max(MaxLayerId, CurWidgetsMaxLayerId);
 	}
diff --git a/Engine/Source/Runtime/Slate/Private/Widgets/SWeakWidget.cpp b/Engine/Source/Runtime/Slate/Private/Widgets/SWeakWidget.cpp
--- a/Engine/Source/Runtime/Slate/Private/Widgets/SWeakWidget.cpp
+++ b/Engine/Source/Runtime/Slate/Private/Widgets/SWeakWidget.cpp
@@ -24,7 +24,7 @@ FVector2D SWeakWidget::ComputeDesiredSize() const
 	return FVector2D::ZeroVector;
 }
 
-void SWeakWidget::ArrangeChildren( const FGeometry& AllottedGeometry, FArrangedChildren& ArrangedChildren ) const
+void SWeakWidget::OnArrangeChildren( const FGeometry& AllottedGeometry, FArrangedChildren& ArrangedChildren ) const
 {
 	// We just want to show the child that we are presenting. Always stretch it to occupy all of the space.
 	TSharedPtr<SWidget> MyContent = WeakChild.Widget.Pin();
@@ -56,7 +56,7 @@ int32 SWeakWidget::OnPaint( const FGeometry& AllottedGeometry, const FSlateRect&
 		check( ArrangedChildren.Num() == 1 );
 		FArrangedWidget& TheChild = ArrangedChildren(0);
 
-		return TheChild.Widget->OnPaint( TheChild.Geometry, 
+		return TheChild.Widget->Paint( TheChild.Geometry, 
 			MyClippingRect, 
 			OutDrawElements, 
 			LayerId + 1,
diff --git a/Engine/Source/Runtime/Slate/Private/Widgets/Text/SRichTextBlock.cpp b/Engine/Source/Runtime/Slate/Private/Widgets/Text/SRichTextBlock.cpp
--- a/Engine/Source/Runtime/Slate/Private/Widgets/Text/SRichTextBlock.cpp
+++ b/Engine/Source/Runtime/Slate/Private/Widgets/Text/SRichTextBlock.cpp
@@ -113,7 +113,7 @@ FChildren* SRichTextBlock::GetChildren()
 	return TextLayout->GetChildren();
 }
 
-void SRichTextBlock::ArrangeChildren( const FGeometry& AllottedGeometry, FArrangedChildren& ArrangedChildren ) const
+void SRichTextBlock::OnArrangeChildren( const FGeometry& AllottedGeometry, FArrangedChildren& ArrangedChildren ) const
 {
 	TextLayout->ArrangeChildren( AllottedGeometry, ArrangedChildren );
 }
diff --git a/Engine/Source/Runtime/Slate/Private/Widgets/Views/SListPanel.cpp b/Engine/Source/Runtime/Slate/Private/Widgets/Views/SListPanel.cpp
--- a/Engine/Source/Runtime/Slate/Private/Widgets/Views/SListPanel.cpp
+++ b/Engine/Source/Runtime/Slate/Private/Widgets/Views/SListPanel.cpp
@@ -46,7 +46,7 @@ SListPanel::FSlot& SListPanel::AddSlot(int32 InsertAtIndex)
 /**
  * Arrange the children top-to-bottom with not additional layout info.
  */
-void SListPanel::ArrangeChildren( const FGeometry& AllottedGeometry, FArrangedChildren& ArrangedChildren ) const
+void SListPanel::OnArrangeChildren( const FGeometry& AllottedGeometry, FArrangedChildren& ArrangedChildren ) const
 {
 	if ( ShouldArrangeHorizontally() )
 	{
diff --git a/Engine/Source/Runtime/Slate/Private/Widgets/Views/SListPanel.h b/Engine/Source/Runtime/Slate/Private/Widgets/Views/SListPanel.h
--- a/Engine/Source/Runtime/Slate/Private/Widgets/Views/SListPanel.h
+++ b/Engine/Source/Runtime/Slate/Private/Widgets/Views/SListPanel.h
@@ -59,7 +59,7 @@ public:
 public:
 
 	// SWidget interface
-	virtual void ArrangeChildren( const FGeometry& AllottedGeometry, FArrangedChildren& ArrangedChildren ) const OVERRIDE;
+	virtual void OnArrangeChildren( const FGeometry& AllottedGeometry, FArrangedChildren& ArrangedChildren ) const OVERRIDE;
 	virtual FVector2D ComputeDesiredSize() const OVERRIDE;
 	virtual FChildren* GetChildren() OVERRIDE;
 	virtual void Tick( const FGeometry& AllottedGeometry, const double InCurrentTime, const float InDeltaTime ) OVERRIDE;
diff --git a/Engine/Source/Runtime/Slate/Public/Framework/Text/ISlateRunHighlighter.h b/Engine/Source/Runtime/Slate/Public/Framework/Text/ISlateRunHighlighter.h
--- a/Engine/Source/Runtime/Slate/Public/Framework/Text/ISlateRunHighlighter.h
+++ b/Engine/Source/Runtime/Slate/Public/Framework/Text/ISlateRunHighlighter.h
@@ -13,5 +13,5 @@ public:
 
 	virtual FChildren* GetChildren() = 0;
 
-	virtual void ArrangeChildren( const TSharedRef< ILayoutBlock >& Block, const FGeometry& AllottedGeometry, FArrangedChildren& ArrangedChildren ) const = 0;
+	virtual void OnArrangeChildren( const TSharedRef< ILayoutBlock >& Block, const FGeometry& AllottedGeometry, FArrangedChildren& ArrangedChildren ) const = 0;
 };
\ No newline at end of file
diff --git a/Engine/Source/Runtime/Slate/Public/Framework/Text/SlateSimpleRunHighlighter.h b/Engine/Source/Runtime/Slate/Public/Framework/Text/SlateSimpleRunHighlighter.h
--- a/Engine/Source/Runtime/Slate/Public/Framework/Text/SlateSimpleRunHighlighter.h
+++ b/Engine/Source/Runtime/Slate/Public/Framework/Text/SlateSimpleRunHighlighter.h
@@ -13,7 +13,7 @@ public:
 
 	virtual FChildren* GetChildren() OVERRIDE;
 
-	virtual void ArrangeChildren( const TSharedRef< ILayoutBlock >& Block, const FGeometry& AllottedGeometry, FArrangedChildren& ArrangedChildren ) const OVERRIDE;
+	virtual void OnArrangeChildren( const TSharedRef< ILayoutBlock >& Block, const FGeometry& AllottedGeometry, FArrangedChildren& ArrangedChildren ) const OVERRIDE;
 
 private:
 
diff --git a/Engine/Source/Runtime/Slate/Public/Slate.h b/Engine/Source/Runtime/Slate/Public/Slate.h
--- a/Engine/Source/Runtime/Slate/Public/Slate.h
+++ b/Engine/Source/Runtime/Slate/Public/Slate.h
@@ -24,7 +24,6 @@ DECLARE_CYCLE_STAT_EXTERN( TEXT("CacheDesiredSize"), STAT_SlateCacheDesiredSize,
 DECLARE_CYCLE_STAT_EXTERN( TEXT("STextBlock"), STAT_SlateCacheDesiredSize_STextBlock, STATGROUP_Slate , );
 DECLARE_CYCLE_STAT_EXTERN( TEXT("TickWidgets"), STAT_SlateTickWidgets, STATGROUP_Slate , );
 DECLARE_CYCLE_STAT_EXTERN( TEXT("DrawWindows"), STAT_SlateDrawWindowTime, STATGROUP_Slate , );
-DECLARE_CYCLE_STAT_EXTERN( TEXT("OnPaint"), STAT_SlateOnPaint, STATGROUP_Slate , );
 DECLARE_CYCLE_STAT_EXTERN( TEXT("OnPaint SViewport"), STAT_SlateOnPaint_SViewport, STATGROUP_Slate , );
 DECLARE_CYCLE_STAT_EXTERN( TEXT("OnPaint SImage"), STAT_SlateOnPaint_SImage, STATGROUP_Slate , );
 DECLARE_CYCLE_STAT_EXTERN( TEXT("OnPaint SBox"), STAT_SlateOnPaint_SBox, STATGROUP_Slate , );
diff --git a/Engine/Source/Runtime/Slate/Public/Widgets/Colors/SColorThemes.h b/Engine/Source/Runtime/Slate/Public/Widgets/Colors/SColorThemes.h
--- a/Engine/Source/Runtime/Slate/Public/Widgets/Colors/SColorThemes.h
+++ b/Engine/Source/Runtime/Slate/Public/Widgets/Colors/SColorThemes.h
@@ -251,7 +251,7 @@ public:
 	 * @param AllottedGeometry    The geometry allotted for this widget by its parent.
 	 * @param ArrangedChildren    The array to which to add the WidgetGeometries that represent the arranged children.
 	 */
-	virtual void ArrangeChildren( const FGeometry& AllottedGeometry, FArrangedChildren& ArrangedChildren ) const OVERRIDE;
+	virtual void OnArrangeChildren( const FGeometry& AllottedGeometry, FArrangedChildren& ArrangedChildren ) const OVERRIDE;
 
 	/**
 	 * A Panel's desired size in the space required to arrange of its children on the screen while respecting all of
diff --git a/Engine/Source/Runtime/Slate/Public/Widgets/Layout/SBox.h b/Engine/Source/Runtime/Slate/Public/Widgets/Layout/SBox.h
--- a/Engine/Source/Runtime/Slate/Public/Widgets/Layout/SBox.h
+++ b/Engine/Source/Runtime/Slate/Public/Widgets/Layout/SBox.h
@@ -53,7 +53,7 @@ class SLATE_API SBox : public SPanel
 
 
 		virtual FVector2D ComputeDesiredSize() const OVERRIDE;
-		virtual void ArrangeChildren( const FGeometry& AllottedGeometry, FArrangedChildren& ArrangedChildren ) const OVERRIDE;
+		virtual void OnArrangeChildren( const FGeometry& AllottedGeometry, FArrangedChildren& ArrangedChildren ) const OVERRIDE;
 		virtual FChildren* GetChildren() OVERRIDE;
 		virtual int32 OnPaint( const FGeometry& AllottedGeometry, const FSlateRect& MyClippingRect, FSlateWindowElementList& OutDrawElements, int32 LayerId, const FWidgetStyle& InWidgetStyle, bool bParentEnabled ) const OVERRIDE;
 
diff --git a/Engine/Source/Runtime/Slate/Public/Widgets/Layout/SDPIScaler.h b/Engine/Source/Runtime/Slate/Public/Widgets/Layout/SDPIScaler.h
--- a/Engine/Source/Runtime/Slate/Public/Widgets/Layout/SDPIScaler.h
+++ b/Engine/Source/Runtime/Slate/Public/Widgets/Layout/SDPIScaler.h
@@ -28,7 +28,7 @@ public:
 
 	void Construct( const FArguments& InArgs );
 
-	virtual void ArrangeChildren( const FGeometry& AllottedGeometry, FArrangedChildren& ArrangedChildren ) const OVERRIDE;
+	virtual void OnArrangeChildren( const FGeometry& AllottedGeometry, FArrangedChildren& ArrangedChildren ) const OVERRIDE;
 	
 	virtual FVector2D ComputeDesiredSize() const OVERRIDE;
 
diff --git a/Engine/Source/Runtime/Slate/Public/Widgets/Layout/SFxWidget.h b/Engine/Source/Runtime/Slate/Public/Widgets/Layout/SFxWidget.h
--- a/Engine/Source/Runtime/Slate/Public/Widgets/Layout/SFxWidget.h
+++ b/Engine/Source/Runtime/Slate/Public/Widgets/Layout/SFxWidget.h
@@ -67,7 +67,7 @@ public:
 protected:
 	virtual int32 OnPaint( const FGeometry& AllottedGeometry, const FSlateRect& MyClippingRect, FSlateWindowElementList& OutDrawElements, int32 LayerId, const FWidgetStyle& InWidgetStyle, bool bParentEnabled ) const OVERRIDE;
 	virtual FVector2D ComputeDesiredSize() const OVERRIDE;
-	virtual void ArrangeChildren( const FGeometry& AllottedGeometry, FArrangedChildren& ArrangedChildren ) const OVERRIDE;
+	virtual void OnArrangeChildren( const FGeometry& AllottedGeometry, FArrangedChildren& ArrangedChildren ) const OVERRIDE;
 
 	TAttribute<float> RenderScale;
 	TAttribute<FVector2D> RenderScaleOrigin;
diff --git a/Engine/Source/Runtime/Slate/Public/Widgets/Layout/SGridPanel.h b/Engine/Source/Runtime/Slate/Public/Widgets/Layout/SGridPanel.h
--- a/Engine/Source/Runtime/Slate/Public/Widgets/Layout/SGridPanel.h
+++ b/Engine/Source/Runtime/Slate/Public/Widgets/Layout/SGridPanel.h
@@ -148,7 +148,7 @@ public:
 
 	// SWidget interface
 	virtual int32 OnPaint( const FGeometry& AllottedGeometry, const FSlateRect& MyClippingRect, FSlateWindowElementList& OutDrawElements, int32 LayerId, const FWidgetStyle& InWidgetStyle, bool bParentEnabled ) const OVERRIDE;
-	virtual void ArrangeChildren( const FGeometry& AllottedGeometry, FArrangedChildren& ArrangedChildren ) const OVERRIDE;
+	virtual void OnArrangeChildren( const FGeometry& AllottedGeometry, FArrangedChildren& ArrangedChildren ) const OVERRIDE;
 	virtual void CacheDesiredSize() OVERRIDE;
 	virtual FVector2D ComputeDesiredSize() const OVERRIDE;
 	virtual FChildren* GetChildren() OVERRIDE;
diff --git a/Engine/Source/Runtime/Slate/Public/Widgets/Layout/SSplitter.h b/Engine/Source/Runtime/Slate/Public/Widgets/Layout/SSplitter.h
--- a/Engine/Source/Runtime/Slate/Public/Widgets/Layout/SSplitter.h
+++ b/Engine/Source/Runtime/Slate/Public/Widgets/Layout/SSplitter.h
@@ -152,7 +152,7 @@ public:
 
 public:
 
-	virtual void ArrangeChildren( const FGeometry& AllottedGeometry, FArrangedChildren& ArrangedChildren ) const OVERRIDE;
+	virtual void OnArrangeChildren( const FGeometry& AllottedGeometry, FArrangedChildren& ArrangedChildren ) const OVERRIDE;
 
 
 	virtual int32 OnPaint( const FGeometry& AllottedGeometry, const FSlateRect& MyClippingRect, FSlateWindowElementList& OutDrawElements, int32 LayerId, const FWidgetStyle& InWidgetStyle, bool bParentEnabled ) const OVERRIDE;
@@ -406,7 +406,7 @@ public:
 
 private:
 
-	virtual void ArrangeChildren( const FGeometry& AllottedGeometry, FArrangedChildren& ArrangedChildren ) const OVERRIDE;
+	virtual void OnArrangeChildren( const FGeometry& AllottedGeometry, FArrangedChildren& ArrangedChildren ) const OVERRIDE;
 
 	/**
 	 * A Panel's desired size in the space required to arrange of its children on the screen while respecting all of
diff --git a/Engine/Source/Runtime/Slate/Public/Widgets/Layout/SUniformGridPanel.h b/Engine/Source/Runtime/Slate/Public/Widgets/Layout/SUniformGridPanel.h
--- a/Engine/Source/Runtime/Slate/Public/Widgets/Layout/SUniformGridPanel.h
+++ b/Engine/Source/Runtime/Slate/Public/Widgets/Layout/SUniformGridPanel.h
@@ -53,7 +53,7 @@ public:
 	void Construct( const FArguments& InArgs );
 
 	// BEGIN SPanel INTERFACE	
-	virtual void ArrangeChildren( const FGeometry& AllottedGeometry, FArrangedChildren& ArrangedChildren ) const OVERRIDE;
+	virtual void OnArrangeChildren( const FGeometry& AllottedGeometry, FArrangedChildren& ArrangedChildren ) const OVERRIDE;
 	virtual FVector2D ComputeDesiredSize() const OVERRIDE;
 	virtual FChildren* GetChildren() OVERRIDE;
 	// END SPanel INTERFACE
diff --git a/Engine/Source/Runtime/Slate/Public/Widgets/Layout/SWidgetSwitcher.h b/Engine/Source/Runtime/Slate/Public/Widgets/Layout/SWidgetSwitcher.h
--- a/Engine/Source/Runtime/Slate/Public/Widgets/Layout/SWidgetSwitcher.h
+++ b/Engine/Source/Runtime/Slate/Public/Widgets/Layout/SWidgetSwitcher.h
@@ -120,7 +120,7 @@ protected:
 
 	// Begin SCompoundWidget interface
 
-	virtual void ArrangeChildren( const FGeometry& AllottedGeometry, FArrangedChildren& ArrangedChildren ) const OVERRIDE;
+	virtual void OnArrangeChildren( const FGeometry& AllottedGeometry, FArrangedChildren& ArrangedChildren ) const OVERRIDE;
 	
 	virtual FVector2D ComputeDesiredSize( ) const OVERRIDE;
 
diff --git a/Engine/Source/Runtime/Slate/Public/Widgets/Layout/SWrapBox.h b/Engine/Source/Runtime/Slate/Public/Widgets/Layout/SWrapBox.h
--- a/Engine/Source/Runtime/Slate/Public/Widgets/Layout/SWrapBox.h
+++ b/Engine/Source/Runtime/Slate/Public/Widgets/Layout/SWrapBox.h
@@ -84,7 +84,7 @@ public:
 
 	virtual void Tick( const FGeometry& AllottedGeometry, const double InCurrentTime, const float InDeltaTime ) OVERRIDE;
 
-	virtual void ArrangeChildren( const FGeometry& AllottedGeometry, FArrangedChildren& ArrangedChildren ) const OVERRIDE;
+	virtual void OnArrangeChildren( const FGeometry& AllottedGeometry, FArrangedChildren& ArrangedChildren ) const OVERRIDE;
 
 	void ClearChildren();
 
diff --git a/Engine/Source/Runtime/Slate/Public/Widgets/SCanvas.h b/Engine/Source/Runtime/Slate/Public/Widgets/SCanvas.h
--- a/Engine/Source/Runtime/Slate/Public/Widgets/SCanvas.h
+++ b/Engine/Source/Runtime/Slate/Public/Widgets/SCanvas.h
@@ -153,7 +153,7 @@ public:
 
 	// Begin SWidget overrides
 
-	virtual void ArrangeChildren( const FGeometry& AllottedGeometry, FArrangedChildren& ArrangedChildren ) const OVERRIDE;
+	virtual void OnArrangeChildren( const FGeometry& AllottedGeometry, FArrangedChildren& ArrangedChildren ) const OVERRIDE;
 	
 	virtual int32 OnPaint( const FGeometry& AllottedGeometry, const FSlateRect& MyClippingRect, FSlateWindowElementList& OutDrawElements, int32 LayerId, const FWidgetStyle& InWidgetStyle, bool bParentEnabled ) const OVERRIDE;
 
diff --git a/Engine/Source/Runtime/Slate/Public/Widgets/SWeakWidget.h b/Engine/Source/Runtime/Slate/Public/Widgets/SWeakWidget.h
--- a/Engine/Source/Runtime/Slate/Public/Widgets/SWeakWidget.h
+++ b/Engine/Source/Runtime/Slate/Public/Widgets/SWeakWidget.h
@@ -43,7 +43,7 @@ public:
 
 	virtual FVector2D ComputeDesiredSize() const OVERRIDE;
 
-	virtual void ArrangeChildren( const FGeometry& AllottedGeometry, FArrangedChildren& ArrangedChildren ) const OVERRIDE;
+	virtual void OnArrangeChildren( const FGeometry& AllottedGeometry, FArrangedChildren& ArrangedChildren ) const OVERRIDE;
 		
 	virtual FChildren* GetChildren() OVERRIDE;
 
diff --git a/Engine/Source/Runtime/Slate/Public/Widgets/Text/SRichTextBlock.h b/Engine/Source/Runtime/Slate/Public/Widgets/Text/SRichTextBlock.h
--- a/Engine/Source/Runtime/Slate/Public/Widgets/Text/SRichTextBlock.h
+++ b/Engine/Source/Runtime/Slate/Public/Widgets/Text/SRichTextBlock.h
@@ -108,7 +108,7 @@ public:
 
 	virtual FChildren* GetChildren() OVERRIDE;
 
-	virtual void ArrangeChildren( const FGeometry& AllottedGeometry, FArrangedChildren& ArrangedChildren ) const OVERRIDE;
+	virtual void OnArrangeChildren( const FGeometry& AllottedGeometry, FArrangedChildren& ArrangedChildren ) const OVERRIDE;
 
 	/**
 	 * Gets the text assigned to this text block
diff --git a/Engine/Source/Runtime/SlateCore/Private/Widgets/SBoxPanel.cpp b/Engine/Source/Runtime/SlateCore/Private/Widgets/SBoxPanel.cpp
--- a/Engine/Source/Runtime/SlateCore/Private/Widgets/SBoxPanel.cpp
+++ b/Engine/Source/Runtime/SlateCore/Private/Widgets/SBoxPanel.cpp
@@ -177,7 +177,7 @@ static void ArrangeChildrenAlong( const TPanelChildren<SBoxPanel::FSlot>& Childr
  * @param AllottedGeometry    The geometry allotted for this widget by its parent.
  * @param ArrangedChildren    The array to which to add the WidgetGeometries that represent the arranged children.
  */
-void SBoxPanel::ArrangeChildren( const FGeometry& AllottedGeometry, FArrangedChildren& ArrangedChildren ) const
+void SBoxPanel::OnArrangeChildren( const FGeometry& AllottedGeometry, FArrangedChildren& ArrangedChildren ) const
 {
 	if ( this->Orientation == Orient_Horizontal )
 	{
diff --git a/Engine/Source/Runtime/SlateCore/Private/Widgets/SCompoundWidget.cpp b/Engine/Source/Runtime/SlateCore/Private/Widgets/SCompoundWidget.cpp
--- a/Engine/Source/Runtime/SlateCore/Private/Widgets/SCompoundWidget.cpp
+++ b/Engine/Source/Runtime/SlateCore/Private/Widgets/SCompoundWidget.cpp
@@ -43,7 +43,7 @@ int32 SCompoundWidget::OnPaint( const FGeometry& AllottedGeometry, const FSlateR
 			.BlendColorAndOpacityTint(ColorAndOpacity.Get())
 			.SetForegroundColor( ForegroundColor );
 
-		return TheChild.Widget->OnPaint( TheChild.Geometry, ChildClippingRect, OutDrawElements, LayerId + 1, CompoundedWidgetStyle, ShouldBeEnabled( bParentEnabled ) );
+		return TheChild.Widget->Paint( TheChild.Geometry, ChildClippingRect, OutDrawElements, LayerId + 1, CompoundedWidgetStyle, ShouldBeEnabled( bParentEnabled ) );
 	}
 	return LayerId;
 }
@@ -65,7 +65,7 @@ FVector2D SCompoundWidget::ComputeDesiredSize() const
 	return FVector2D::ZeroVector;
 }
 
-void SCompoundWidget::ArrangeChildren( const FGeometry& AllottedGeometry, FArrangedChildren& ArrangedChildren ) const
+void SCompoundWidget::OnArrangeChildren( const FGeometry& AllottedGeometry, FArrangedChildren& ArrangedChildren ) const
 {
 	ArrangeSingleChild( AllottedGeometry, ArrangedChildren, ChildSlot, ContentScale );
 }
diff --git a/Engine/Source/Runtime/SlateCore/Private/Widgets/SLeafWidget.cpp b/Engine/Source/Runtime/SlateCore/Private/Widgets/SLeafWidget.cpp
--- a/Engine/Source/Runtime/SlateCore/Private/Widgets/SLeafWidget.cpp
+++ b/Engine/Source/Runtime/SlateCore/Private/Widgets/SLeafWidget.cpp
@@ -22,7 +22,7 @@ FChildren* SLeafWidget::GetChildren( )
 }
 
 
-void SLeafWidget::ArrangeChildren( const FGeometry& AllottedGeometry, FArrangedChildren& ArrangedChildren ) const
+void SLeafWidget::OnArrangeChildren( const FGeometry& AllottedGeometry, FArrangedChildren& ArrangedChildren ) const
 {
 	// Nothing to arrange; Leaf Widgets do not have children.
 }
diff --git a/Engine/Source/Runtime/SlateCore/Private/Widgets/SOverlay.cpp b/Engine/Source/Runtime/SlateCore/Private/Widgets/SOverlay.cpp
--- a/Engine/Source/Runtime/SlateCore/Private/Widgets/SOverlay.cpp
+++ b/Engine/Source/Runtime/SlateCore/Private/Widgets/SOverlay.cpp
@@ -18,7 +18,7 @@ void SOverlay::Construct( const SOverlay::FArguments& InArgs )
 }
 
 
-void SOverlay::ArrangeChildren( const FGeometry& AllottedGeometry, FArrangedChildren& ArrangedChildren ) const
+void SOverlay::OnArrangeChildren( const FGeometry& AllottedGeometry, FArrangedChildren& ArrangedChildren ) const
 {
 	for ( int32 ChildIndex = 0; ChildIndex < Children.Num(); ++ChildIndex )
 	{
@@ -102,7 +102,7 @@ int32 SOverlay::OnPaint( const FGeometry& AllottedGeometry, const FSlateRect& My
 	{
 		FArrangedWidget& CurWidget = ArrangedChildren(ChildIndex);
 		FSlateRect ChildClipRect = MyClippingRect.IntersectionWith( CurWidget.Geometry.GetClippingRect() );
-		const int32 CurWidgetsMaxLayerId = CurWidget.Widget->OnPaint( CurWidget.Geometry, ChildClipRect, OutDrawElements, MaxLayerId + 1, InWidgetStyle, ShouldBeEnabled( bParentEnabled ) );
+		const int32 CurWidgetsMaxLayerId = CurWidget.Widget->Paint( CurWidget.Geometry, ChildClipRect, OutDrawElements, MaxLayerId + 1, InWidgetStyle, ShouldBeEnabled( bParentEnabled ) );
 
 		MaxLayerId = FMath::Max( MaxLayerId, CurWidgetsMaxLayerId );
 	}
diff --git a/Engine/Source/Runtime/SlateCore/Private/Widgets/SPanel.cpp b/Engine/Source/Runtime/SlateCore/Private/Widgets/SPanel.cpp
--- a/Engine/Source/Runtime/SlateCore/Private/Widgets/SPanel.cpp
+++ b/Engine/Source/Runtime/SlateCore/Private/Widgets/SPanel.cpp
@@ -44,7 +44,7 @@ int32 SPanel::PaintArrangedChildren( const FArrangedChildren& ArrangedChildren,
 		
 		if (ChildClipRect.GetSize().Size() > 0)
 		{
-			const int32 CurWidgetsMaxLayerId = CurWidget.Widget->OnPaint(CurWidget.Geometry, ChildClipRect, OutDrawElements, LayerId, InWidgetStyle, ShouldBeEnabled(bParentEnabled));
+			const int32 CurWidgetsMaxLayerId = CurWidget.Widget->Paint(CurWidget.Geometry, ChildClipRect, OutDrawElements, LayerId, InWidgetStyle, ShouldBeEnabled(bParentEnabled));
 			MaxLayerId = FMath::Max(MaxLayerId, CurWidgetsMaxLayerId);
 		}		
 		
diff --git a/Engine/Source/Runtime/SlateCore/Private/Widgets/SWidget.cpp b/Engine/Source/Runtime/SlateCore/Private/Widgets/SWidget.cpp
--- a/Engine/Source/Runtime/SlateCore/Private/Widgets/SWidget.cpp
+++ b/Engine/Source/Runtime/SlateCore/Private/Widgets/SWidget.cpp
@@ -2,6 +2,8 @@
 
 #include "SlateCorePrivatePCH.h"
 
+DECLARE_CYCLE_STAT(TEXT("OnPaint"), STAT_SlateOnPaint, STATGROUP_Slate);
+DECLARE_CYCLE_STAT(TEXT("ArrangeChildren"), STAT_SlateArrangeChildren, STATGROUP_Slate);
 
 SWidget::SWidget()
 	: CreatedInFile( TEXT("") )
@@ -478,3 +480,15 @@ void SWidget::SetDebugInfo( const ANSICHAR* InType, const ANSICHAR* InFile, int3
 	this->CreatedInFile = FName( *FPaths::GetCleanFilename(InFile) );
 	this->CreatedOnLine = OnLine;
 }
+
+int32 SWidget::Paint(const FGeometry& AllottedGeometry, const FSlateRect& MyClippingRect, FSlateWindowElementList& OutDrawElements, int32 LayerId, const FWidgetStyle& InWidgetStyle, bool bParentEnabled) const
+{
+	SCOPE_CYCLE_COUNTER(STAT_SlateOnPaint);
+	return OnPaint(AllottedGeometry, MyClippingRect, OutDrawElements, LayerId, InWidgetStyle, bParentEnabled);
+}
+
+void SWidget::ArrangeChildren(const FGeometry& AllottedGeometry, FArrangedChildren& ArrangedChildren) const
+{
+	SCOPE_CYCLE_COUNTER(STAT_SlateArrangeChildren);
+	OnArrangeChildren(AllottedGeometry, ArrangedChildren);
+}
diff --git a/Engine/Source/Runtime/SlateCore/Private/Widgets/SWindow.cpp b/Engine/Source/Runtime/SlateCore/Private/Widgets/SWindow.cpp
--- a/Engine/Source/Runtime/SlateCore/Private/Widgets/SWindow.cpp
+++ b/Engine/Source/Runtime/SlateCore/Private/Widgets/SWindow.cpp
@@ -88,7 +88,7 @@ public:
 
 private:
 
-	virtual void ArrangeChildren( const FGeometry& AllottedGeometry, FArrangedChildren& ArrangedChildren ) const OVERRIDE
+	virtual void OnArrangeChildren( const FGeometry& AllottedGeometry, FArrangedChildren& ArrangedChildren ) const OVERRIDE
 	{
 		const FVector2D WindowDesktopPosition = (ensure(OwnerWindow.IsValid()))
 			? OwnerWindow.Pin()->GetPositionInScreen()
diff --git a/Engine/Source/Runtime/SlateCore/Public/Widgets/SBoxPanel.h b/Engine/Source/Runtime/SlateCore/Public/Widgets/SBoxPanel.h
--- a/Engine/Source/Runtime/SlateCore/Public/Widgets/SBoxPanel.h
+++ b/Engine/Source/Runtime/SlateCore/Public/Widgets/SBoxPanel.h
@@ -80,7 +80,7 @@ public:
 
 	// Begin SWidget overrides.
 
-	virtual void ArrangeChildren( const FGeometry& AllottedGeometry, FArrangedChildren& ArrangedChildren ) const;
+	virtual void OnArrangeChildren( const FGeometry& AllottedGeometry, FArrangedChildren& ArrangedChildren ) const;
 	
 	virtual FVector2D ComputeDesiredSize() const;
 
diff --git a/Engine/Source/Runtime/SlateCore/Public/Widgets/SCompoundWidget.h b/Engine/Source/Runtime/SlateCore/Public/Widgets/SCompoundWidget.h
--- a/Engine/Source/Runtime/SlateCore/Public/Widgets/SCompoundWidget.h
+++ b/Engine/Source/Runtime/SlateCore/Public/Widgets/SCompoundWidget.h
@@ -76,7 +76,7 @@ public:
 
 	virtual FVector2D ComputeDesiredSize() const OVERRIDE;
 
-	virtual void ArrangeChildren( const FGeometry& AllottedGeometry, FArrangedChildren& ArrangedChildren ) const OVERRIDE;
+	virtual void OnArrangeChildren( const FGeometry& AllottedGeometry, FArrangedChildren& ArrangedChildren ) const OVERRIDE;
 
 	virtual FSlateColor GetForegroundColor() const OVERRIDE;
 
diff --git a/Engine/Source/Runtime/SlateCore/Public/Widgets/SLeafWidget.h b/Engine/Source/Runtime/SlateCore/Public/Widgets/SLeafWidget.h
--- a/Engine/Source/Runtime/SlateCore/Public/Widgets/SLeafWidget.h
+++ b/Engine/Source/Runtime/SlateCore/Public/Widgets/SLeafWidget.h
@@ -48,7 +48,7 @@ public:
 	 */
 	virtual FChildren* GetChildren( );
 
-	virtual void ArrangeChildren( const FGeometry& AllottedGeometry, FArrangedChildren& ArrangedChildren ) const;
+	virtual void OnArrangeChildren( const FGeometry& AllottedGeometry, FArrangedChildren& ArrangedChildren ) const;
 
 private:
 
diff --git a/Engine/Source/Runtime/SlateCore/Public/Widgets/SOverlay.h b/Engine/Source/Runtime/SlateCore/Public/Widgets/SOverlay.h
--- a/Engine/Source/Runtime/SlateCore/Public/Widgets/SOverlay.h
+++ b/Engine/Source/Runtime/SlateCore/Public/Widgets/SOverlay.h
@@ -119,7 +119,7 @@ public:
 	}
 
 	// SWidget interface
-	virtual void ArrangeChildren( const FGeometry& AllottedGeometry, FArrangedChildren& ArrangedChildren ) const OVERRIDE;
+	virtual void OnArrangeChildren( const FGeometry& AllottedGeometry, FArrangedChildren& ArrangedChildren ) const OVERRIDE;
 	virtual FVector2D ComputeDesiredSize() const OVERRIDE;
 	virtual FChildren* GetChildren() OVERRIDE;
 	virtual int32 OnPaint( const FGeometry& AllottedGeometry, const FSlateRect& MyClippingRect, FSlateWindowElementList& OutDrawElements, int32 LayerId, const FWidgetStyle& InWidgetStyle, bool bParentEnabled ) const OVERRIDE;
diff --git a/Engine/Source/Runtime/SlateCore/Public/Widgets/SPanel.h b/Engine/Source/Runtime/SlateCore/Public/Widgets/SPanel.h
--- a/Engine/Source/Runtime/SlateCore/Public/Widgets/SPanel.h
+++ b/Engine/Source/Runtime/SlateCore/Public/Widgets/SPanel.h
@@ -26,7 +26,7 @@ public:
 	 * @param AllottedGeometry    The geometry allotted for this widget by its parent.
 	 * @param ArrangedChildren    The array to which to add the WidgetGeometries that represent the arranged children.
 	 */
-	virtual void ArrangeChildren( const FGeometry& AllottedGeometry, FArrangedChildren& ArrangedChildren ) const = 0;
+	virtual void OnArrangeChildren( const FGeometry& AllottedGeometry, FArrangedChildren& ArrangedChildren ) const = 0;
 	/**
 	 * A Panel's desired size in the space required to arrange of its children on the screen while respecting all of
 	 * the children's desired sizes and any layout-related options specified by the user. See StackPanel for an example.
diff --git a/Engine/Source/Runtime/SlateCore/Public/Widgets/SWidget.h b/Engine/Source/Runtime/SlateCore/Public/Widgets/SWidget.h
--- a/Engine/Source/Runtime/SlateCore/Public/Widgets/SWidget.h
+++ b/Engine/Source/Runtime/SlateCore/Public/Widgets/SWidget.h
@@ -92,6 +92,8 @@ public:
 	//
 
 	/**
+	 * Called to tell a widget to paint itself (and it's children).
+	 * 
 	 * The widget should respond by populating the OutDrawElements array with FDrawElements 
 	 * that represent it and any of its children.
 	 *
@@ -104,7 +106,7 @@ public:
 	 *
 	 * @return The maximum layer ID attained by this widget or any of its children.
 	 */
-	virtual int32 OnPaint( const FGeometry& AllottedGeometry, const FSlateRect& MyClippingRect, FSlateWindowElementList& OutDrawElements, int32 LayerId, const FWidgetStyle& InWidgetStyle, bool bParentEnabled ) const = 0;
+	int32 Paint( const FGeometry& AllottedGeometry, const FSlateRect& MyClippingRect, FSlateWindowElementList& OutDrawElements, int32 LayerId, const FWidgetStyle& InWidgetStyle, bool bParentEnabled ) const;
 
 	/**
 	 * Ticks this widget.  Override in derived classes, but always call the parent implementation.
@@ -465,13 +467,15 @@ public:
 	public:
 
 	/**
+	 * Non-virtual entry point for arrange children. ensures common work is executed before calling the virtual
+	 * ArrangeChildren function.
 	 * Compute the Geometry of all the children and add populate the ArrangedChildren list with their values.
 	 * Each type of Layout panel should arrange children based on desired behavior.
 	 *
 	 * @param AllottedGeometry    The geometry allotted for this widget by its parent.
 	 * @param ArrangedChildren    The array to which to add the WidgetGeometries that represent the arranged children.
 	 */
-	virtual void ArrangeChildren( const FGeometry& AllottedGeometry, FArrangedChildren& ArrangedChildren ) const = 0;
+	void ArrangeChildren( const FGeometry& AllottedGeometry, FArrangedChildren& ArrangedChildren ) const;
 
 	/**
 	 * Every widget that has children must implement this method. This allows for iteration over the Widget's
@@ -686,6 +690,32 @@ protected:
 		return IsEnabled() && InParentEnabled;
 	}
 
+private:
+	/**
+	* The widget should respond by populating the OutDrawElements array with FDrawElements
+	* that represent it and any of its children. Called by the non-virtual OnPaint to enforce pre/post conditions
+	* during OnPaint.
+	*
+	* @param AllottedGeometry  The FGeometry that describes an area in which the widget should appear.
+	* @param MyClippingRect    The clipping rectangle allocated for this widget and its children.
+	* @param OutDrawElements   A list of FDrawElements to populate with the output.
+	* @param LayerId           The Layer onto which this widget should be rendered.
+	* @param InColorAndOpacity Color and Opacity to be applied to all the descendants of the widget being painted
+	* @param bParentEnabled	True if the parent of this widget is enabled.
+	*
+	* @return The maximum layer ID attained by this widget or any of its children.
+	*/
+	virtual int32 OnPaint(const FGeometry& AllottedGeometry, const FSlateRect& MyClippingRect, FSlateWindowElementList& OutDrawElements, int32 LayerId, const FWidgetStyle& InWidgetStyle, bool bParentEnabled) const = 0;
+
+	/**
+	* Compute the Geometry of all the children and add populate the ArrangedChildren list with their values.
+	* Each type of Layout panel should arrange children based on desired behavior.
+	*
+	* @param AllottedGeometry    The geometry allotted for this widget by its parent.
+	* @param ArrangedChildren    The array to which to add the WidgetGeometries that represent the arranged children.
+	*/
+	virtual void OnArrangeChildren(const FGeometry& AllottedGeometry, FArrangedChildren& ArrangedChildren) const = 0;
+
 protected:
 	
 	//	DEBUG INFORMATION
