Make Slate font cache non-copyable

diff --git a/Engine/Source/Runtime/SlateCore/Public/Fonts/FontCache.h b/Engine/Source/Runtime/SlateCore/Public/Fonts/FontCache.h
--- a/Engine/Source/Runtime/SlateCore/Public/Fonts/FontCache.h
+++ b/Engine/Source/Runtime/SlateCore/Public/Fonts/FontCache.h
@@ -316,7 +316,10 @@ public:
 	 */
 	void FlushCache() const;
 
-
+private:
+	// Non-copyable
+	FSlateFontCache(const FSlateFontCache&);
+	FSlateFontCache& operator=(const FSlateFontCache&);
 private:
 
 	/** Mapping Font keys to cached data */
