Fix crash if a dynamic slate brush reference is not reported externally and its object reference is garbage collected before the slate resource manager has a chance to store the brush

diff --git a/Engine/Source/Runtime/SlateCore/Private/Brushes/SlateDynamicImageBrush.cpp b/Engine/Source/Runtime/SlateCore/Private/Brushes/SlateDynamicImageBrush.cpp
--- a/Engine/Source/Runtime/SlateCore/Private/Brushes/SlateDynamicImageBrush.cpp
+++ b/Engine/Source/Runtime/SlateCore/Private/Brushes/SlateDynamicImageBrush.cpp
@@ -14,6 +14,12 @@ FSlateDynamicImageBrush::~FSlateDynamicImageBrush( )
 {
 	if (FSlateApplicationBase::IsInitialized())
 	{
+		// Brush resource is no longer referenced by this object
+		if( ResourceObject )
+		{
+			ResourceObject->RemoveFromRoot();
+		}
+
 		TSharedPtr<FSlateRenderer> Renderer = FSlateApplicationBase::Get().GetRenderer();
 
 		if (Renderer.IsValid())
diff --git a/Engine/Source/Runtime/SlateCore/Public/Brushes/SlateDynamicImageBrush.h b/Engine/Source/Runtime/SlateCore/Public/Brushes/SlateDynamicImageBrush.h
--- a/Engine/Source/Runtime/SlateCore/Public/Brushes/SlateDynamicImageBrush.h
+++ b/Engine/Source/Runtime/SlateCore/Public/Brushes/SlateDynamicImageBrush.h
@@ -33,8 +33,11 @@ struct SLATECORE_API FSlateDynamicImageBrush
 		bIsDynamicallyLoaded = true;
 
 		// if we have a texture, make a unique name
-		if (GetResourceObject() != nullptr)
+		if (ResourceObject != nullptr)
 		{
+			// @todo Slate - Hack:  This is to address an issue where the brush created and a GC occurs before the brush resource object becomes referenced
+			// by the Slate resource manager.  
+			ResourceObject->AddToRoot();
 			ResourceName = InTextureName;
 		}
 	}
