Removed overhead of (it turns out) incorrectly using ICU to convert chars to a Unicode code point to check for kerning.

diff --git a/Engine/Source/Runtime/SlateCore/Private/Rendering/ElementBatcher.cpp b/Engine/Source/Runtime/SlateCore/Private/Rendering/ElementBatcher.cpp
--- a/Engine/Source/Runtime/SlateCore/Private/Rendering/ElementBatcher.cpp
+++ b/Engine/Source/Runtime/SlateCore/Private/Rendering/ElementBatcher.cpp
@@ -494,7 +494,7 @@ void FSlateElementBatcher::AddTextElement( const FVector2D& Position, const FVec
 				InvTextureSizeY = 1.0f/FontTexture->GetHeight();
 			}
 
-			const bool bIsWhitespace = FText::IsWhitespace(CurrentChar);
+			const bool bIsWhitespace = FChar::IsWhitespace(CurrentChar);
 
 			if( !bIsWhitespace && CharIndex > 0 )
 			{
