SWidget OnMouseCaptureLost() default implementation now in the CPP.

diff --git a/Engine/Source/Runtime/SlateCore/Private/Widgets/SWidget.cpp b/Engine/Source/Runtime/SlateCore/Private/Widgets/SWidget.cpp
--- a/Engine/Source/Runtime/SlateCore/Private/Widgets/SWidget.cpp
+++ b/Engine/Source/Runtime/SlateCore/Private/Widgets/SWidget.cpp
@@ -310,6 +310,12 @@ bool SWidget::HasMouseCapture() const
 }
 
 
+void SWidget::OnMouseCaptureLost()
+{
+	
+}
+
+
 bool SWidget::FindChildGeometries( const FGeometry& MyGeometry, const TSet< TSharedRef<SWidget> >& WidgetsToFind, TMap<TSharedRef<SWidget>, FArrangedWidget>& OutResult ) const
 {
 	FindChildGeometries_Helper(MyGeometry, WidgetsToFind, OutResult);
diff --git a/Engine/Source/Runtime/SlateCore/Public/Widgets/SWidget.h b/Engine/Source/Runtime/SlateCore/Public/Widgets/SWidget.h
--- a/Engine/Source/Runtime/SlateCore/Public/Widgets/SWidget.h
+++ b/Engine/Source/Runtime/SlateCore/Public/Widgets/SWidget.h
@@ -514,7 +514,7 @@ public:
 	/**
 	 * Called when this widget had captured the mouse, but that capture has been revoked for some reason.
 	 */
-	virtual void OnMouseCaptureLost() { }
+	virtual void OnMouseCaptureLost();
 
 	/**
 	 * Ticks this widget and all of it's child widgets.  Should not be called directly.
