Slate: Fixed widget reflector visibility reporting.

diff --git a/Engine/Source/Runtime/SlateCore/Private/Layout/Visibility.cpp b/Engine/Source/Runtime/SlateCore/Private/Layout/Visibility.cpp
--- a/Engine/Source/Runtime/SlateCore/Private/Layout/Visibility.cpp
+++ b/Engine/Source/Runtime/SlateCore/Private/Layout/Visibility.cpp
@@ -29,29 +29,29 @@ FString EVisibility::ToString( ) const
 	static const FString HitTestInvisibleString("HitTestInvisible");
 	static const FString SelfHitTestInvisibleString("SelfHitTestInvisible");
 
-	if ((Value & VISPRIVATE_ChildrenHitTestVisible) == 0)
+	if (Value == VIS_Visible)
 	{
-		return HitTestInvisibleString;
+		return VisibleString;
 	}
 
-	if ((Value & VISPRIVATE_SelfHitTestVisible) == 0)
+	if (Value == VIS_Collapsed)
 	{
-		return SelfHitTestInvisibleString;
+		return CollapsedString;
 	}
 
-	if ((Value & VISPRIVATE_Visible) != 0)
+	if (Value == VIS_Hidden)
 	{
-		return VisibleString;
+		return HiddenString;
 	}
 
-	if ((Value & VISPRIVATE_Hidden) != 0)
+	if (Value == VIS_HitTestInvisible)
 	{
-		return HiddenString;
+		return HitTestInvisibleString;
 	}
 
-	if ((Value & VISPRIVATE_Collapsed) != 0)
+	if (Value == VIS_SelfHitTestInvisible)
 	{
-		return CollapsedString;
+		return SelfHitTestInvisibleString;
 	}
 
 	return FString();
