Fix Gradient Batcher to correctly interpret gradient points.
* When adding boundary points if they are not there, uses the proper local space coordinates now.
* When interpreting all stops for rendering, interprets the points as % of local size so the gradient properly works invariant to the specified scale.

diff --git a/Engine/Source/Runtime/SlateCore/Private/Rendering/ElementBatcher.cpp b/Engine/Source/Runtime/SlateCore/Private/Rendering/ElementBatcher.cpp
--- a/Engine/Source/Runtime/SlateCore/Private/Rendering/ElementBatcher.cpp
+++ b/Engine/Source/Runtime/SlateCore/Private/Rendering/ElementBatcher.cpp
@@ -585,6 +585,8 @@ void FSlateElementBatcher::AddGradientElement( const FVector2D& Position, const
 	FVector2D BotLeft = FVector2D(Position.X, Position.Y + Size.Y);
 	FVector2D BotRight = FVector2D(Position.X + Size.X, Position.Y + Size.Y);
 
+	FVector2D LocalSize = Size / Scale;
+
 	// Copy the gradient stops.. We may need to add more
 	TArray<FSlateGradientStop> GradientStops = InPayload.GradientStops;
 
@@ -595,31 +597,31 @@ void FSlateElementBatcher::AddGradientElement( const FVector2D& Position, const
 	// If they are not add a gradient stop with the same color as the first and/or last stop
 	if( InPayload.GradientType == Orient_Vertical )
 	{
-		if( TopLeft.X < FirstStop.Position.X + TopLeft.X )
+		if( 0.0f < FirstStop.Position.X )
 		{
 			// The first stop is after the left side of the quad.  Add a stop at the left side of the quad using the same color as the first stop
 			GradientStops.Insert( FSlateGradientStop( FVector2D(0.0, 0.0), FirstStop.Color ), 0 );
 		}
 
-		if( TopRight.X > LastStop.Position.X + TopLeft.X )
+		if( LocalSize.X > LastStop.Position.X )
 		{
 			// The last stop is before the right side of the quad.  Add a stop at the right side of the quad using the same color as the last stop
-			GradientStops.Add( FSlateGradientStop( BotRight, LastStop.Color ) ); 
+			GradientStops.Add( FSlateGradientStop( LocalSize, LastStop.Color ) ); 
 		}
 	}
 	else
 	{
 
-		if( TopLeft.Y < FirstStop.Position.Y + TopLeft.Y )
+		if( 0.0f < FirstStop.Position.Y )
 		{
 			// The first stop is after the top side of the quad.  Add a stop at the top side of the quad using the same color as the first stop
 			GradientStops.Insert( FSlateGradientStop( FVector2D(0.0, 0.0), FirstStop.Color ), 0 );
 		}
 
-		if( BotLeft.Y > LastStop.Position.Y + TopLeft.Y )
+		if( LocalSize.Y > LastStop.Position.Y )
 		{
 			// The last stop is before the bottom side of the quad.  Add a stop at the bottom side of the quad using the same color as the last stop
-			GradientStops.Add( FSlateGradientStop( BotRight, LastStop.Color ) ); 
+			GradientStops.Add( FSlateGradientStop( LocalSize, LastStop.Color ) ); 
 		}
 	}
 
@@ -643,16 +645,18 @@ void FSlateElementBatcher::AddGradientElement( const FVector2D& Position, const
 			// Gradient stop is vertical so gradients to left to right
 			StartPt = TopLeft;
 			EndPt = BotLeft;
-			StartPt.X += CurStop.Position.X * Scale;
-			EndPt.X += CurStop.Position.X  * Scale;	
+			// Gradient stops are interpreted as % of local Size. Determine the % and use it to scale the point in screen space.
+			StartPt.X += (CurStop.Position.X / LocalSize.X) * ( BotRight.X - TopLeft.X);
+			EndPt.X += (CurStop.Position.X / LocalSize.X) * ( BotRight.X - TopLeft.X);	
 		}
 		else
 		{
 			// Gradient stop is horizontal so gradients to top to bottom
 			StartPt = TopLeft;
 			EndPt = TopRight;
-			StartPt.Y += CurStop.Position.Y * Scale;
-			EndPt.Y += CurStop.Position.Y * Scale;
+			// Gradient stops are interpreted as % of local Size. Determine the % and use it to scale the point in screen space.
+			StartPt.Y += (CurStop.Position.Y / LocalSize.Y) * (BotRight.Y - TopLeft.Y);
+			EndPt.Y += (CurStop.Position.Y / LocalSize.Y) * (BotRight.Y - TopLeft.Y);
 		}
 
 		if( StopIndex == 0 )
