Replaced SplineFilterTable in Slate's ElementBatcher with a singleton so that it doesn't call FPaths::EngineContentDir() during global variable initialization.

diff --git a/Engine/Source/Runtime/SlateCore/Private/Rendering/ElementBatcher.cpp b/Engine/Source/Runtime/SlateCore/Private/Rendering/ElementBatcher.cpp
--- a/Engine/Source/Runtime/SlateCore/Private/Rendering/ElementBatcher.cpp
+++ b/Engine/Source/Runtime/SlateCore/Private/Rendering/ElementBatcher.cpp
@@ -7,7 +7,11 @@
 #include "SlateCorePrivatePCH.h"
 
 
-static const FSlateImageBrush SplineFilterTable( FPaths::EngineContentDir() / TEXT("Slate/SplineFilterTable.png"), FVector2D(32,1) );
+static const FSlateImageBrush& GetSplineFilterTable()
+{
+	static const FSlateImageBrush SplineFilterTable( FPaths::EngineContentDir() / TEXT("Slate/SplineFilterTable.png"), FVector2D(32,1) );
+	return SplineFilterTable;
+}
 
 
 /**
@@ -707,7 +711,7 @@ void FSlateElementBatcher::AddSplineElement( const FVector2D& Position, float Sc
 	const float HalfThickness = LineThickness * .5f + Radius;
 
 	// Currently splines are not atlased because they are tiled.  So we just assume the texture proxy holds the actual texture
-	FSlateShaderResourceProxy* ResourceProxy = ResourceManager.GetShaderResource( SplineFilterTable );
+	FSlateShaderResourceProxy* ResourceProxy = ResourceManager.GetShaderResource( GetSplineFilterTable() );
 	check(ResourceProxy && ResourceProxy->Resource);
 
 	// Find a batch for the element
@@ -826,7 +830,7 @@ void FSlateElementBatcher::AddLineElement( const FVector2D& Position, const FVec
 	if( InPayload.bAntialias )
 	{
 		// Currently splines are not atlased because they are tiled.  So we just assume the texture proxy holds the actual texture
-		FSlateShaderResourceProxy* ResourceProxy = ResourceManager.GetShaderResource( SplineFilterTable );
+		FSlateShaderResourceProxy* ResourceProxy = ResourceManager.GetShaderResource( GetSplineFilterTable() );
 		check(ResourceProxy && ResourceProxy->Resource);
 
 		// The radius to use when checking the distance of pixels to the actual line.  Arbitrary value based on what looks the best
