Fixed issue where dock/undock tabs display doesn't update when highlighting section.
#ttp 340466 - Regression: EDITOR: Dock/Undock Tabs display does not update when highlighting section
#branch UE4
#proj Runtime.SlateRHIRenderer
#add Removed last remaining traces of "marking window as drawn" in SWindow.

diff --git a/Engine/Source/Developer/StandaloneRenderer/Private/OpenGL/SlateOpenGLRenderer.cpp b/Engine/Source/Developer/StandaloneRenderer/Private/OpenGL/SlateOpenGLRenderer.cpp
--- a/Engine/Source/Developer/StandaloneRenderer/Private/OpenGL/SlateOpenGLRenderer.cpp
+++ b/Engine/Source/Developer/StandaloneRenderer/Private/OpenGL/SlateOpenGLRenderer.cpp
@@ -132,13 +132,6 @@ void FSlateOpenGLRenderer::DrawWindows( FSlateDrawBuffer& InWindowDrawBuffer )
 
 			Viewport->SwapBuffers();
 
-			// Mark this window as 'drawn'
-			if( !WindowToDraw->HasEverBeenDrawn() )
-			{
-				WindowToDraw->MarkAsDrawn();
-			}
-
-
 			// Reset all batch data for this window
 			ElementBatcher->ResetBatches();
 		}
diff --git a/Engine/Source/Developer/StandaloneRenderer/Private/Windows/D3D/SlateD3DRenderer.cpp b/Engine/Source/Developer/StandaloneRenderer/Private/Windows/D3D/SlateD3DRenderer.cpp
--- a/Engine/Source/Developer/StandaloneRenderer/Private/Windows/D3D/SlateD3DRenderer.cpp
+++ b/Engine/Source/Developer/StandaloneRenderer/Private/Windows/D3D/SlateD3DRenderer.cpp
@@ -337,12 +337,6 @@ void FSlateD3DRenderer::DrawWindows( FSlateDrawBuffer& InWindowDrawBuffer )
 			const bool bUseVSync = false;
 			Viewport->D3DSwapChain->Present( bUseVSync ? 1 : 0, 0);
 
-			// Mark this window as 'drawn'
-			if( !WindowToDraw->HasEverBeenDrawn() )
-			{
-				WindowToDraw->MarkAsDrawn();
-			}
-
 			// All elements have been drawn.  Reset all cached data
 			ElementBatcher->ResetBatches();
 		}
diff --git a/Engine/Source/Runtime/SlateCore/Private/Widgets/SWindow.cpp b/Engine/Source/Runtime/SlateCore/Private/Widgets/SWindow.cpp
--- a/Engine/Source/Runtime/SlateCore/Private/Widgets/SWindow.cpp
+++ b/Engine/Source/Runtime/SlateCore/Private/Widgets/SWindow.cpp
@@ -487,7 +487,7 @@ bool SWindow::HasActiveChildren() const
 
 void SWindow::Tick( const FGeometry& AllottedGeometry, const double InCurrentTime, const float InDeltaTime )
 {
-	if( Morpher.bIsActive && bHasEverBeenDrawn )
+	if( Morpher.bIsActive )
 	{
 		if(Morpher.bIsPendingPlay)
 		{
@@ -1615,7 +1615,6 @@ SWindow::SWindow()
 	, bIsCursorDecoratorWindow( false )
 	, bInitiallyMaximized( false )
 	, bHasEverBeenShown( false )
-	, bHasEverBeenDrawn( false )
 	, bFocusWhenFirstShown(true)
 	, bActivateWhenFirstShown(true)
 	, bHasOSWindowBorder( false )
diff --git a/Engine/Source/Runtime/SlateCore/Public/Widgets/SWindow.h b/Engine/Source/Runtime/SlateCore/Public/Widgets/SWindow.h
--- a/Engine/Source/Runtime/SlateCore/Public/Widgets/SWindow.h
+++ b/Engine/Source/Runtime/SlateCore/Public/Widgets/SWindow.h
@@ -562,18 +562,6 @@ public:
 	/** The window's desired size takes into account the ratio between the slate units and the pixel size */
 	virtual FVector2D ComputeDesiredSize() const override;
 
-	/** Marks this window as having been drawn at least once */
-	inline void MarkAsDrawn()
-	{
-		bHasEverBeenDrawn = true;
-	}
-
-	/** @return Returns True if the window has even been drawn */
-	inline bool HasEverBeenDrawn() const
-	{
-		return bHasEverBeenDrawn;
-	}
-
 	/** @return true if this window will be focused when it is first shown */
 	inline bool IsFocusedInitially() const
 	{
@@ -680,9 +668,6 @@ protected:
 	/** True if this window has been shown yet */
 	bool bHasEverBeenShown : 1;
 
-	/** True if this window has ever been drawn */
-	bool bHasEverBeenDrawn : 1;
-
 	/** Focus this window immediately as it is shown */
 	bool bFocusWhenFirstShown : 1;
 
