Web camera mirroring support for Twitch
- If your web camera requires vertical mirroring, we'll do that automatically now
- For editor streaming, you can now customize horizontal mirroring in the Live Streaming preferences
- For game streaming, you can select whether to mirror horizontally using the new pin on the "Start Broadcasting Game" node

Slate
- Slate brushes now support mirroring (none, horizontal, vertical, or both)

TTP 340769

diff --git a/Engine/Config/BaseEditorGameAgnostic.ini b/Engine/Config/BaseEditorGameAgnostic.ini
--- a/Engine/Config/BaseEditorGameAgnostic.ini
+++ b/Engine/Config/BaseEditorGameAgnostic.ini
@@ -13,5 +13,6 @@ ScreenScaling=0.75
 bPrimaryMonitorOnly=true
 bEnableWebCam=true
 WebCamResolution=Normal_640x480
+bMirrorWebCamImage=false
 bCaptureAudioFromComputer=true
 bCaptureAudioFromMicrophone=true
\ No newline at end of file
diff --git a/Engine/Plugins/Runtime/TwitchLiveStreaming/Source/TwitchLiveStreaming/TwitchLiveStreaming.cpp b/Engine/Plugins/Runtime/TwitchLiveStreaming/Source/TwitchLiveStreaming/TwitchLiveStreaming.cpp
--- a/Engine/Plugins/Runtime/TwitchLiveStreaming/Source/TwitchLiveStreaming/TwitchLiveStreaming.cpp
+++ b/Engine/Plugins/Runtime/TwitchLiveStreaming/Source/TwitchLiveStreaming/TwitchLiveStreaming.cpp
@@ -64,6 +64,7 @@ void FTwitchLiveStreaming::StartupModule()
 	WebCamCapabilityIndex = INDEX_NONE;
 	WebCamVideoBufferWidth = 0;
 	WebCamVideoBufferHeight = 0;
+	bIsWebCamTextureFlippedVertically = false;
 	WebCamTexture = nullptr;
 
 	bWantsChatEnabled = false;
@@ -1110,8 +1111,10 @@ bool FTwitchLiveStreaming::IsWebCamEnabled() const
 }
 
 
-UTexture2D* FTwitchLiveStreaming::GetWebCamTexture()
+UTexture2D* FTwitchLiveStreaming::GetWebCamTexture( bool& bIsImageFlippedHorizontally, bool& bIsImageFlippedVertically )
 {
+	bIsImageFlippedHorizontally = false;	// Twitch never flips a web cam image horizontally
+	bIsImageFlippedVertically = bIsWebCamTextureFlippedVertically;
 	return WebCamTexture;
 }
 
@@ -1479,6 +1482,7 @@ void FTwitchLiveStreaming::Async_InitWebCamSupport()
 								This->WebCamCapabilityIndex = Capability.capabilityIndex;
 								This->WebCamVideoBufferWidth = Capability.resolution.width;
 								This->WebCamVideoBufferHeight = Capability.resolution.height;
+								This->bIsWebCamTextureFlippedVertically = !Capability.isTopToBottom;
 							}
 							else
 							{
@@ -1753,6 +1757,8 @@ void FTwitchLiveStreaming::ReleaseWebCamTexture( const bool bReleaseResourceToo
 		WebCamTexture->RemoveFromRoot();
 		WebCamTexture = nullptr;
 
+		bIsWebCamTextureFlippedVertically = false;
+
 		PublishStatus( EStatusType::WebCamTextureNotReady, LOCTEXT( "Status_WebCamTextureNotReady", "Web cam video disabled" ) );
 	}
 }
diff --git a/Engine/Plugins/Runtime/TwitchLiveStreaming/Source/TwitchLiveStreaming/TwitchLiveStreaming.h b/Engine/Plugins/Runtime/TwitchLiveStreaming/Source/TwitchLiveStreaming/TwitchLiveStreaming.h
--- a/Engine/Plugins/Runtime/TwitchLiveStreaming/Source/TwitchLiveStreaming/TwitchLiveStreaming.h
+++ b/Engine/Plugins/Runtime/TwitchLiveStreaming/Source/TwitchLiveStreaming/TwitchLiveStreaming.h
@@ -35,7 +35,7 @@ public:
 	virtual void StartWebCam( const FWebCamConfig& Config ) override;
 	virtual void StopWebCam() override;
 	virtual bool IsWebCamEnabled() const override;
-	virtual UTexture2D* GetWebCamTexture() override;
+	virtual UTexture2D* GetWebCamTexture( bool& bIsImageFlippedHorizontally, bool& bIsImageFlippedVertically ) override;
 	virtual FOnChatMessage& OnChatMessage() override;
 	virtual void ConnectToChat() override;
 	virtual void DisconnectFromChat() override;
@@ -291,6 +291,9 @@ private:
 	/** Web cam video buffer height */
 	int32 WebCamVideoBufferHeight;
 
+	/** True if the web cam texture is flipped vertically, as captured from the camera, and needs to be displayed bottom to top */
+	bool bIsWebCamTextureFlippedVertically;
+
 	/** Texture that stores the web cam image, copied asynchronously from the CPU */
 	class UTexture2D* WebCamTexture;
 
diff --git a/Engine/Source/Editor/EditorLiveStreaming/EditorLiveStreaming.cpp b/Engine/Source/Editor/EditorLiveStreaming/EditorLiveStreaming.cpp
--- a/Engine/Source/Editor/EditorLiveStreaming/EditorLiveStreaming.cpp
+++ b/Engine/Source/Editor/EditorLiveStreaming/EditorLiveStreaming.cpp
@@ -323,7 +323,9 @@ void FEditorLiveStreaming::BroadcastStatusCallback( const FLiveStreamingStatus&
 	// If the web cam just turned on, then we'll go ahead and show it
 	if( Status.StatusType == FLiveStreamingStatus::EStatusType::WebCamTextureReady )
 	{	
-		UTexture2D* WebCamTexture = LiveStreamer->GetWebCamTexture();
+		bool bIsImageFlippedHorizontally = false;
+		bool bIsImageFlippedVertically = false;
+		UTexture2D* WebCamTexture = LiveStreamer->GetWebCamTexture( bIsImageFlippedHorizontally, bIsImageFlippedVertically );
 		if( ensure( WebCamTexture != nullptr ) )
 		{
 			IMainFrameModule& MainFrameModule = FModuleManager::LoadModuleChecked<IMainFrameModule>( TEXT( "MainFrame" ) );
@@ -331,6 +333,13 @@ void FEditorLiveStreaming::BroadcastStatusCallback( const FLiveStreamingStatus&
 			auto RootWindow = MainFrameModule.GetParentWindow();
 			check( RootWindow.IsValid() );
 
+			// Allow the user to customize the image mirroring, too!
+			const auto& Settings = *GetDefault< UEditorLiveStreamingSettings >();
+			if( Settings.bMirrorWebCamImage )
+			{
+				bIsImageFlippedHorizontally = !bIsImageFlippedHorizontally;
+			}
+
 			// How many pixels from the edge of the main frame window to where the broadcast status window appears
 			const int WindowBorderPadding = 50;
 
@@ -361,6 +370,20 @@ void FEditorLiveStreaming::BroadcastStatusCallback( const FLiveStreamingStatus&
 				FVector2D( WebCamTexture->GetSizeX(), WebCamTexture->GetSizeY() ),
 				WebCamTexture->GetFName() ) );
 
+			// If the web cam image is coming in flipped, we'll apply mirroring to the Slate brush
+			if( bIsImageFlippedHorizontally && bIsImageFlippedVertically )
+			{
+				WebCamDynamicImageBrush->Mirroring = ESlateBrushMirrorType::Both;
+			}
+			else if( bIsImageFlippedHorizontally )
+			{ 
+				WebCamDynamicImageBrush->Mirroring = ESlateBrushMirrorType::Horizontal;
+			}
+			else if( bIsImageFlippedVertically )
+			{ 
+				WebCamDynamicImageBrush->Mirroring = ESlateBrushMirrorType::Vertical;
+			}
+
 			// @todo livestream: Currently if the user closes the window, the camera is deactivated and it doesn't turn back on unless the broadcast is restarted.  We could allow the camera to be reactivated though.
 			BroadcastStatusWindow->SetOnWindowClosed( 
 				FOnWindowClosed::CreateStatic( []( const TSharedRef<SWindow>& Window, FEditorLiveStreaming* This ) 
diff --git a/Engine/Source/Editor/EditorLiveStreaming/EditorLiveStreamingSettings.h b/Engine/Source/Editor/EditorLiveStreaming/EditorLiveStreamingSettings.h
--- a/Engine/Source/Editor/EditorLiveStreaming/EditorLiveStreamingSettings.h
+++ b/Engine/Source/Editor/EditorLiveStreaming/EditorLiveStreamingSettings.h
@@ -62,6 +62,10 @@ public:
 	UPROPERTY(Config, EditAnywhere, Category=EditorLiveStreaming, Meta=(EditCondition="bEnableWebCam") )
 	TEnumAsByte<EEditorLiveStreamingWebCamResolution::Type> WebCamResolution;
 
+	/** You can enable this to flip the web cam image horizontally so that it looks more like a mirror */
+	UPROPERTY(Config, EditAnywhere, Category=EditorLiveStreaming )
+	bool bMirrorWebCamImage;
+
 	/** Enables broadcast of audio being played by your computer, such as in-game sounds */
 	UPROPERTY(Config, EditAnywhere, Category=EditorLiveStreaming )
 	bool bCaptureAudioFromComputer;
diff --git a/Engine/Source/Runtime/Engine/Public/Features/ILiveStreamingService.h b/Engine/Source/Runtime/Engine/Public/Features/ILiveStreamingService.h
--- a/Engine/Source/Runtime/Engine/Public/Features/ILiveStreamingService.h
+++ b/Engine/Source/Runtime/Engine/Public/Features/ILiveStreamingService.h
@@ -257,10 +257,13 @@ public:
 
 	/**
 	 * For services that support web cams, this function will return the current picture on the web cam as a GPU texture
-	 *
+	 * 
+	 * @param	bIsImageFlippedHorizontally		[Out] True if the image is flipped left to right
+	 * @param	bIsImageFlippedVertically		[Out] True if the image is flipped top to bottom
+	 * 
 	 * @return	Current web cam frame as a 2D texture
 	 */
-	virtual class UTexture2D* GetWebCamTexture() = 0;
+	virtual class UTexture2D* GetWebCamTexture( bool& bIsImageFlippedHorizontally, bool& bIsImageFlippedVertically ) = 0;
 
 
 	/**
diff --git a/Engine/Source/Runtime/GameLiveStreaming/GameLiveStreaming.cpp b/Engine/Source/Runtime/GameLiveStreaming/GameLiveStreaming.cpp
--- a/Engine/Source/Runtime/GameLiveStreaming/GameLiveStreaming.cpp
+++ b/Engine/Source/Runtime/GameLiveStreaming/GameLiveStreaming.cpp
@@ -24,6 +24,7 @@ FGameLiveStreaming::FGameLiveStreaming()
 			TEXT( "     bEnableWebCam=<bool>                Enables video from your web camera to be captured and displayed\n" )
 			TEXT( "     DesiredWebCamWidth=<int32>          Horizontal resolution for web camera capture\n" )
 			TEXT( "     DesiredWebCamHeight=<int32>         Veritcal resolution for web camera capture\n" )
+			TEXT( "     bMirrorWebCamImage=<bool>			Flips the web camera image horizontally\n" )
 			TEXT( "     bCaptureAudioFromComputer=<bool>    Enables capturing sound that is played back through your PC\n" )
 			TEXT( "     bCaptureAudioFromMicrophone=<bool>  Enables capturing sound from your default microphone\n" )
 			TEXT( "     bDrawSimpleWebCamVideo=<bool>       Draws a simple web cam video on top of the viewport\n" ),
@@ -38,6 +39,7 @@ FGameLiveStreaming::FGameLiveStreaming()
 						FParse::Bool( *Arg, TEXT( "bEnableWebCam="), Config.bEnableWebCam );
 						FParse::Value( *Arg, TEXT( "DesiredWebCamWidth="), Config.DesiredWebCamWidth );
 						FParse::Value( *Arg, TEXT( "DesiredWebCamHeight="), Config.DesiredWebCamHeight );
+						FParse::Bool( *Arg, TEXT( "bMirrorWebCamImage="), Config.bMirrorWebCamImage );
 						FParse::Bool( *Arg, TEXT( "bCaptureAudioFromComputer="), Config.bCaptureAudioFromComputer );
 						FParse::Bool( *Arg, TEXT( "bCaptureAudioFromMicrophone="), Config.bCaptureAudioFromMicrophone );
 						FParse::Bool( *Arg, TEXT( "bDrawSimpleWebCamVideo=" ), Config.bDrawSimpleWebCamVideo );
@@ -62,6 +64,7 @@ FGameLiveStreaming::FGameLiveStreaming()
 	ReadbackBufferIndex = 0;
 	ReadbackBuffers[0] = nullptr;
 	ReadbackBuffers[1] = nullptr;
+	bMirrorWebCamImage = false;
 	bDrawSimpleWebCamVideo = true;
 }
 
@@ -114,6 +117,7 @@ void FGameLiveStreaming::StartBroadcastingGame( const FGameBroadcastConfig& Game
 			FSlateRenderer* SlateRenderer = FSlateApplication::Get().GetRenderer().Get();
 			SlateRenderer->OnSlateWindowRendered().AddRaw( this, &FGameLiveStreaming::OnSlateWindowRenderedDuringBroadcasting );
 
+			this->bMirrorWebCamImage = GameBroadcastConfig.bMirrorWebCamImage;
 			this->bDrawSimpleWebCamVideo = GameBroadcastConfig.bDrawSimpleWebCamVideo;
 
 			// @todo livestream: This will interfere with editor live streaming if both are running at the same time!  The editor live 
@@ -395,15 +399,23 @@ void FGameLiveStreaming::DrawSimpleWebCamVideo( UCanvas* Canvas )
 {
 	if( IsBroadcastingGame() && bDrawSimpleWebCamVideo && LiveStreamer->IsWebCamEnabled() )
 	{
-		UTexture2D* WebCamTexture = LiveStreamer->GetWebCamTexture();
+		bool bIsImageFlippedHorizontally = false;
+		bool bIsImageFlippedVertically = false;
+		UTexture2D* WebCamTexture = LiveStreamer->GetWebCamTexture( bIsImageFlippedHorizontally, bIsImageFlippedVertically );
 		if( WebCamTexture != nullptr )
 		{
+			// Give the user a chance to customize the image mirroring
+			if( this->bMirrorWebCamImage )
+			{
+				bIsImageFlippedHorizontally = !bIsImageFlippedHorizontally;
+			}
+
 			const float BorderPadding = 6.0f;
 			Canvas->Canvas->DrawTile( 
 				Canvas->SizeX - WebCamTexture->GetSizeX() - BorderPadding, BorderPadding,	// Top right justify
 				WebCamTexture->GetSizeX(), WebCamTexture->GetSizeY(),
-				0, 0, 
-				1.0f, 1.0f, 
+				bIsImageFlippedHorizontally ? 1.0f : 0.0f, bIsImageFlippedVertically ? 1.0f : 0.0f, 
+				bIsImageFlippedHorizontally ? -1.0f : 1.0f, bIsImageFlippedVertically ? -1.0f : 1.0f, 
 				FLinearColor::White, 
 				WebCamTexture->Resource, 
 				false );	// Alpha blend?
@@ -412,12 +424,17 @@ void FGameLiveStreaming::DrawSimpleWebCamVideo( UCanvas* Canvas )
 }
 
 
-UTexture2D* FGameLiveStreaming::GetWebCamTexture()
+UTexture2D* FGameLiveStreaming::GetWebCamTexture( bool& bIsImageFlippedHorizontally, bool& bIsImageFlippedVertically )
 {
+	bIsImageFlippedHorizontally = false;
+	bIsImageFlippedVertically = false;
+
 	UTexture2D* WebCamTexture = nullptr;
 	if( IsBroadcastingGame() && bDrawSimpleWebCamVideo && LiveStreamer->IsWebCamEnabled() )
 	{
-		WebCamTexture = LiveStreamer->GetWebCamTexture();
+		bool bIsImageFlippedHorizontally = false;
+		bool bIsImageFlippedVertically = false;
+		WebCamTexture = LiveStreamer->GetWebCamTexture( bIsImageFlippedHorizontally, bIsImageFlippedVertically );
 	}
 	return WebCamTexture;
 }
diff --git a/Engine/Source/Runtime/GameLiveStreaming/GameLiveStreaming.h b/Engine/Source/Runtime/GameLiveStreaming/GameLiveStreaming.h
--- a/Engine/Source/Runtime/GameLiveStreaming/GameLiveStreaming.h
+++ b/Engine/Source/Runtime/GameLiveStreaming/GameLiveStreaming.h
@@ -24,7 +24,7 @@ public:
 	virtual void StartBroadcastingGame( const FGameBroadcastConfig& GameBroadcastConfig ) override;
 	virtual void StopBroadcastingGame() override;
 	virtual void DrawSimpleWebCamVideo( UCanvas* Canvas ) override;
-	virtual class UTexture2D* GetWebCamTexture() override;
+	virtual class UTexture2D* GetWebCamTexture( bool& bIsImageFlippedHorizontally, bool& bIsImageFlippedVertically ) override;
 	virtual class ILiveStreamingService* GetLiveStreamingService() override;
 
 protected:
@@ -71,6 +71,9 @@ private:
 	/** The current buffer index.  We bounce between them to avoid stalls. */
 	int32 ReadbackBufferIndex;
 
+	/** True if the user prefers to flip the web camera image horizontally */
+	bool bMirrorWebCamImage;
+
 	/** True if we should draw a simple web cam video on top of the viewport while broadcasting */
 	bool bDrawSimpleWebCamVideo;
 
diff --git a/Engine/Source/Runtime/GameLiveStreaming/GameLiveStreamingFunctionLibrary.cpp b/Engine/Source/Runtime/GameLiveStreaming/GameLiveStreamingFunctionLibrary.cpp
--- a/Engine/Source/Runtime/GameLiveStreaming/GameLiveStreamingFunctionLibrary.cpp
+++ b/Engine/Source/Runtime/GameLiveStreaming/GameLiveStreamingFunctionLibrary.cpp
@@ -24,6 +24,7 @@ void UGameLiveStreamingFunctionLibrary::StartBroadcastingGame(
 		bool bEnableWebCam,
 		int32 DesiredWebCamWidth,
 		int32 DesiredWebCamHeight,
+		bool bMirrorWebCamImage,
 		bool bCaptureAudioFromComputer,
 		bool bCaptureAudioFromMicrophone,
 		bool bDrawSimpleWebCamVideo )
@@ -34,6 +35,7 @@ void UGameLiveStreamingFunctionLibrary::StartBroadcastingGame(
 	Config.bEnableWebCam = bEnableWebCam;
 	Config.DesiredWebCamWidth = DesiredWebCamWidth;
 	Config.DesiredWebCamHeight = DesiredWebCamHeight;
+	Config.bMirrorWebCamImage = bMirrorWebCamImage;
 	Config.bCaptureAudioFromComputer = bCaptureAudioFromComputer;
 	Config.bCaptureAudioFromMicrophone = bCaptureAudioFromMicrophone;
 	Config.bDrawSimpleWebCamVideo = bDrawSimpleWebCamVideo;
diff --git a/Engine/Source/Runtime/GameLiveStreaming/GameLiveStreamingFunctionLibrary.h b/Engine/Source/Runtime/GameLiveStreaming/GameLiveStreamingFunctionLibrary.h
--- a/Engine/Source/Runtime/GameLiveStreaming/GameLiveStreamingFunctionLibrary.h
+++ b/Engine/Source/Runtime/GameLiveStreaming/GameLiveStreamingFunctionLibrary.h
@@ -24,17 +24,19 @@ class UGameLiveStreamingFunctionLibrary : public UBlueprintFunctionLibrary
 	 * @param	bEnableWebCam				If enabled, video from your web camera will be captured and displayed while broadcasting, so that your viewers can see your presence.
 	 * @param	DesiredWebCamWidth			Desired web cam capture resolution width.  The web cam may only support a limited number of resolutions, so we'll choose one that matches as closely to this as possible
 	 * @param	DesiredWebCamHeight			Desired web cam capture resolution height.
+	 * @param	bMirrorWebCamImage			You can enable this to flip the web camera image horizontally, so that it looks like a mirror
 	 * @param	bCaptureAudioFromComputer	Enables broadcast of audio being played by your computer, such as in-game sounds
 	 * @param	bCaptureAudioFromMicrophone	Enables broadcast of audio from your default microphone recording device
 	 * @param	bDrawSimpleWebCamVideo		If enabled, the engine will draw a simple web cam image on top of the game viewport.  If you turn this off, it's up to you to draw the web cam image yourself.  You can access the web cam texture by calling IGameLiveStreaming::Get().GetWebCamTexture(). 
 	 */
-	UFUNCTION( BlueprintCallable, Category="LiveStreaming", meta=(FrameRate="30", ScreenScaling="1.0", bEnableWebCam="true", DesiredWebCamWidth="320", DesiredWebCamHeight="240", bCaptureAudioFromComputer="true", bCaptureAudioFromMicrophone="true", bDrawSimpleWebCamVideo="true") )
+	UFUNCTION( BlueprintCallable, Category="LiveStreaming", meta=(FrameRate="30", ScreenScaling="1.0", bEnableWebCam="true", DesiredWebCamWidth="320", DesiredWebCamHeight="240", bMirrorWebCamImage="false", bCaptureAudioFromComputer="true", bCaptureAudioFromMicrophone="true", bDrawSimpleWebCamVideo="true") )
 	static void StartBroadcastingGame(
 		int32 FrameRate,
 		float ScreenScaling,
 		bool bEnableWebCam,
 		int32 DesiredWebCamWidth,
 		int32 DesiredWebCamHeight,
+		bool bMirrorWebCamImage,
 		bool bCaptureAudioFromComputer,
 		bool bCaptureAudioFromMicrophone,
 		bool bDrawSimpleWebCamVideo );
diff --git a/Engine/Source/Runtime/GameLiveStreaming/Public/IGameLiveStreaming.h b/Engine/Source/Runtime/GameLiveStreaming/Public/IGameLiveStreaming.h
--- a/Engine/Source/Runtime/GameLiveStreaming/Public/IGameLiveStreaming.h
+++ b/Engine/Source/Runtime/GameLiveStreaming/Public/IGameLiveStreaming.h
@@ -25,6 +25,9 @@ struct FGameBroadcastConfig
 	/** Desired web cam capture resolution height. */
 	int32 DesiredWebCamHeight;
 
+	/** You can enable this to flip the web camera image horizontally, so that it looks like a mirror */
+	bool bMirrorWebCamImage;
+
 	/** Enables broadcast of audio being played by your computer, such as in-game sounds */
 	bool bCaptureAudioFromComputer;
 
@@ -44,6 +47,7 @@ struct FGameBroadcastConfig
 		  bEnableWebCam( true ),
 		  DesiredWebCamWidth( 320 ),
 		  DesiredWebCamHeight( 240 ),
+		  bMirrorWebCamImage( false ),
 		  bCaptureAudioFromComputer( true ),
 		  bCaptureAudioFromMicrophone( true ),
 		  bDrawSimpleWebCamVideo( true )
@@ -89,9 +93,12 @@ public:
 	/**
 	 * Exposes access to the web cam video feed as a 2D texture
 	 *
+	 * @param	bIsImageFlippedHorizontally		[Out] True if the image is flipped left to right, and you should compensate while drawing by mirroring the texture U coordinates
+	 * @param	bIsImageFlippedVertically		[Out] True if the image is flipped top to bottom, and you should compensate while drawing by mirroring the texture V coordinates
+	 * 
 	 * @return	The web cam video feed texture, or null if the web cam isn't enabled
 	 */
-	virtual class UTexture2D* GetWebCamTexture() = 0;
+	virtual class UTexture2D* GetWebCamTexture( bool& bIsImageFlippedHorizontally, bool& bIsImageFlippedVertically ) = 0;
 
 	/**
 	 * Gets direct access to the live streaming service we are using, or nullptr if we haven't selected one yet
diff --git a/Engine/Source/Runtime/SlateCore/Private/Rendering/ElementBatcher.cpp b/Engine/Source/Runtime/SlateCore/Private/Rendering/ElementBatcher.cpp
--- a/Engine/Source/Runtime/SlateCore/Private/Rendering/ElementBatcher.cpp
+++ b/Engine/Source/Runtime/SlateCore/Private/Rendering/ElementBatcher.cpp
@@ -216,10 +216,14 @@ void FSlateElementBatcher::AddBoxElement( const FVector2D& InPosition, const FVe
 
 		FColor Tint = GetElementColor( InPayload.Tint, InPayload.BrushResource );
 
-		ESlateBrushTileType::Type TilingRule = InPayload.BrushResource->Tiling;
+		const ESlateBrushTileType::Type TilingRule = InPayload.BrushResource->Tiling;
 		const bool bTileHorizontal = (TilingRule == ESlateBrushTileType::Both || TilingRule == ESlateBrushTileType::Horizontal);
 		const bool bTileVertical = (TilingRule == ESlateBrushTileType::Both || TilingRule == ESlateBrushTileType::Vertical);
 
+		const ESlateBrushMirrorType::Type MirroringRule = InPayload.BrushResource->Mirroring;
+		const bool bMirrorHorizontal = (MirroringRule == ESlateBrushMirrorType::Both || MirroringRule == ESlateBrushMirrorType::Horizontal);
+		const bool bMirrorVertical = (MirroringRule == ESlateBrushMirrorType::Both || MirroringRule == ESlateBrushMirrorType::Vertical);
+
 		// Pass the tiling information as a flag so we can pick the correct texture addressing mode
 		ESlateBatchDrawFlag::Type DrawFlags = ( ( bTileHorizontal ? ESlateBatchDrawFlag::TileU : 0 ) | ( bTileVertical ? ESlateBatchDrawFlag::TileV : 0 ) );
 
@@ -261,16 +265,16 @@ void FSlateElementBatcher::AddBoxElement( const FVector2D& InPosition, const FVe
 
 			// Determine the texture coordinates for each quad
 			// These are not scaled.
-			const float LeftMarginU = (Margin.Left > 0.0f)
+			float LeftMarginU = (Margin.Left > 0.0f)
 				? StartUV.X + Margin.Left * SizeUV.X + HalfTexel.X
 				: StartUV.X;
-			const float TopMarginV = (Margin.Top > 0.0f)
+			float TopMarginV = (Margin.Top > 0.0f)
 				? StartUV.Y + Margin.Top * SizeUV.Y + HalfTexel.Y
 				: StartUV.Y;
-			const float RightMarginU = (Margin.Right > 0.0f)
+			float RightMarginU = (Margin.Right > 0.0f)
 				? EndUV.X - Margin.Right * SizeUV.X + HalfTexel.X
 				: EndUV.X;
-			const float BottomMarginV = (Margin.Bottom > 0.0f)
+			float BottomMarginV = (Margin.Bottom > 0.0f)
 				? EndUV.Y - Margin.Bottom * SizeUV.Y + HalfTexel.Y
 				: EndUV.Y;
 
@@ -306,6 +310,27 @@ void FSlateElementBatcher::AddBoxElement( const FVector2D& InPosition, const FVe
 			EndPos.X = FMath::TruncToInt( EndPos.X );
 			EndPos.Y = FMath::TruncToInt( EndPos.Y );
 
+			if( bMirrorHorizontal || bMirrorVertical )
+			{
+				const FVector2D UVMin = StartUV;
+				const FVector2D UVMax = EndUV;
+
+				if( bMirrorHorizontal )
+				{
+					StartUV.X = UVMax.X - ( StartUV.X - UVMin.X );
+					EndUV.X = UVMax.X - ( EndUV.X - UVMin.X );
+					LeftMarginU = UVMax.X - ( LeftMarginU - UVMin.X );
+					RightMarginU = UVMax.X - ( RightMarginU - UVMin.X );
+				}
+				if( bMirrorVertical )
+				{
+					StartUV.Y = UVMax.Y - ( StartUV.Y - UVMin.Y );
+					EndUV.Y = UVMax.Y - ( EndUV.Y - UVMin.Y );
+					TopMarginV = UVMax.Y - ( TopMarginV - UVMin.Y );
+					BottomMarginV = UVMax.Y - ( BottomMarginV - UVMin.Y );
+				}
+			}
+
 			BatchVertices.Add( FSlateVertex( Position,									StartUV,									Tiling,	Tint, InClippingRect ) ); //0
 			BatchVertices.Add( FSlateVertex( FVector2D( Position.X, TopMarginY ),		FVector2D( StartUV.X, TopMarginV ),			Tiling,	Tint, InClippingRect ) ); //1
 			BatchVertices.Add( FSlateVertex( FVector2D( LeftMarginX, Position.Y ),		FVector2D( LeftMarginU, StartUV.Y ),		Tiling,	Tint, InClippingRect ) ); //2
@@ -397,6 +422,23 @@ void FSlateElementBatcher::AddBoxElement( const FVector2D& InPosition, const FVe
 			FVector2D BotLeft =	 FVector2D( Position.X, Position.Y+Size.Y);
 			FVector2D BotRight = FVector2D( Position.X + Size.X, Position.Y+Size.Y);
 
+			if( bMirrorHorizontal || bMirrorVertical )
+			{
+				const FVector2D UVMin = StartUV;
+				const FVector2D UVMax = EndUV;
+
+				if( bMirrorHorizontal )
+				{
+					StartUV.X = UVMax.X - ( StartUV.X - UVMin.X );
+					EndUV.X = UVMax.X - ( EndUV.X - UVMin.X );
+				}
+				if( bMirrorVertical )
+				{
+					StartUV.Y = UVMax.Y - ( StartUV.Y - UVMin.Y );
+					EndUV.Y = UVMax.Y - ( EndUV.Y - UVMin.Y );
+				}
+			}
+
 			// Add four vertices to the list of verts to be added to the vertex buffer
 			BatchVertices.Add( FSlateVertex(TopLeft,	StartUV,						Tiling,	Tint, InClippingRect ) );
 			BatchVertices.Add( FSlateVertex(TopRight,	FVector2D(EndUV.X,StartUV.Y),	Tiling,	Tint, InClippingRect ) );
diff --git a/Engine/Source/Runtime/SlateCore/Public/Styling/SlateBrush.h b/Engine/Source/Runtime/SlateCore/Public/Styling/SlateBrush.h
--- a/Engine/Source/Runtime/SlateCore/Public/Styling/SlateBrush.h
+++ b/Engine/Source/Runtime/SlateCore/Public/Styling/SlateBrush.h
@@ -58,6 +58,29 @@ namespace ESlateBrushTileType
 }
 
 
+/**
+ * Possible options for mirroring the brush image
+ */
+UENUM()
+namespace ESlateBrushMirrorType
+{
+	enum Type
+	{
+		/** Don't mirror anything, just draw the texture as it is */
+		NoMirror,
+
+		/** Mirror the image horizontally */
+		Horizontal,
+
+		/** Mirror the image vertically */
+		Vertical,
+
+		/** Mirror in both directions */
+		Both
+	};
+}
+
+
 /*
  * Enumerates brush image types.
  */
@@ -110,6 +133,11 @@ struct SLATECORE_API FSlateBrush
 	UPROPERTY(EditAnywhere, Category=Brush)
 	TEnumAsByte<enum ESlateBrushTileType::Type> Tiling;
 
+	/** How to mirror the image in Image mode.  This is normally only used for dynamic image brushes where the source texture
+	    comes from a hardware device such as a web camera. */
+	UPROPERTY(EditAnywhere, Category=Brush)
+	TEnumAsByte<enum ESlateBrushMirrorType::Type> Mirroring;
+
 	/** The type of image */
 	UPROPERTY()
 	TEnumAsByte<enum ESlateBrushImageType::Type> ImageType;
@@ -126,6 +154,7 @@ public:
 		, Tint_DEPRECATED(FLinearColor::White)
 		, TintColor(FLinearColor::White)
 		, Tiling(ESlateBrushTileType::NoTile)
+		, Mirroring(ESlateBrushMirrorType::NoMirror)
 		, ImageType(ESlateBrushImageType::NoImage)
 		, ResourceObject(nullptr)
 		, ResourceName(NAME_None)
