Add an Expose that takes a TSharedRef so that SAssignNew can assing to a TSharedRef if desired

diff --git a/Engine/Source/Runtime/SlateCore/Public/Widgets/DeclarativeSyntaxSupport.h b/Engine/Source/Runtime/SlateCore/Public/Widgets/DeclarativeSyntaxSupport.h
--- a/Engine/Source/Runtime/SlateCore/Public/Widgets/DeclarativeSyntaxSupport.h
+++ b/Engine/Source/Runtime/SlateCore/Public/Widgets/DeclarativeSyntaxSupport.h
@@ -1205,6 +1205,17 @@ struct TDecl
 		return *this;
 	}
 
+	/**
+	 * Initialize OutVarToInit with the widget that is being constructed.
+	 * @see SAssignNew
+	 */
+	template<class ExposeAsWidgetType>
+	TDecl& Expose( TSharedRef<ExposeAsWidgetType>& OutVarToInit )
+	{
+		OutVarToInit = _Widget;
+		return *this;
+	}
+
 	/**
 	 * Initialize a WEAK OutVarToInit with the widget that is being constructed.
 	 * @see SAssignNew
