Slate Hittesting 2.0 : Hittesting no longer bound to logical widget hierarchy.

diff --git a/Engine/Source/Runtime/Slate/Public/Widgets/Text/SMultiLineEditableText.h b/Engine/Source/Runtime/Slate/Public/Widgets/Text/SMultiLineEditableText.h
--- a/Engine/Source/Runtime/Slate/Public/Widgets/Text/SMultiLineEditableText.h
+++ b/Engine/Source/Runtime/Slate/Public/Widgets/Text/SMultiLineEditableText.h
@@ -2,6 +2,8 @@
 
 #pragma once
 
+#if WITH_FANCY_TEXT
+
 #include "ITextEditorWidget.h"
 #include "ITextInputMethodSystem.h"
 
@@ -231,7 +233,7 @@ private:
 
 		virtual ~FCursorLineHighlighter() {}
 
-		virtual int32 OnPaint( const FTextLayout::FLineView& Line, const float OffsetX, const float Width, const FTextBlockStyle& DefaultStyle, const FGeometry& AllottedGeometry, const FSlateRect& MyClippingRect, FSlateWindowElementList& OutDrawElements, int32 LayerId, const FWidgetStyle& InWidgetStyle, bool bParentEnabled ) const override;
+		virtual int32 OnPaint( const FPaintArgs& Args, const FTextLayout::FLineView& Line, const float OffsetX, const float Width, const FTextBlockStyle& DefaultStyle, const FGeometry& AllottedGeometry, const FSlateRect& MyClippingRect, FSlateWindowElementList& OutDrawElements, int32 LayerId, const FWidgetStyle& InWidgetStyle, bool bParentEnabled ) const override;
 
 	protected:
 
@@ -250,7 +252,7 @@ private:
 
 		virtual ~FTextCompositionHighlighter() {}
 
-		virtual int32 OnPaint( const FTextLayout::FLineView& Line, const float OffsetX, const float Width, const FTextBlockStyle& DefaultStyle, const FGeometry& AllottedGeometry, const FSlateRect& MyClippingRect, FSlateWindowElementList& OutDrawElements, int32 LayerId, const FWidgetStyle& InWidgetStyle, bool bParentEnabled ) const override;
+		virtual int32 OnPaint( const FPaintArgs& Args, const FTextLayout::FLineView& Line, const float OffsetX, const float Width, const FTextBlockStyle& DefaultStyle, const FGeometry& AllottedGeometry, const FSlateRect& MyClippingRect, FSlateWindowElementList& OutDrawElements, int32 LayerId, const FWidgetStyle& InWidgetStyle, bool bParentEnabled ) const override;
 
 	protected:
 
@@ -266,7 +268,7 @@ private:
 
 		virtual ~FTextSelectionRunRenderer() {}
 
-		virtual int32 OnPaint( const FTextLayout::FLineView& Line, const TSharedRef< ISlateRun >& Run, const TSharedRef< ILayoutBlock >& Block, const FTextBlockStyle& DefaultStyle, const FGeometry& AllottedGeometry, const FSlateRect& MyClippingRect, FSlateWindowElementList& OutDrawElements, int32 LayerId, const FWidgetStyle& InWidgetStyle, bool bParentEnabled ) const override;
+		virtual int32 OnPaint( const FPaintArgs& Args, const FTextLayout::FLineView& Line, const TSharedRef< ISlateRun >& Run, const TSharedRef< ILayoutBlock >& Block, const FTextBlockStyle& DefaultStyle, const FGeometry& AllottedGeometry, const FSlateRect& MyClippingRect, FSlateWindowElementList& OutDrawElements, int32 LayerId, const FWidgetStyle& InWidgetStyle, bool bParentEnabled ) const override;
 
 		void SetHasKeyboardFocus(const bool bInHasKeyboardFocus)
 		{
@@ -324,7 +326,7 @@ private:
 private:
 	// BEGIN SWidget interface
 	virtual void Tick( const FGeometry& AllottedGeometry, const double InCurrentTime, const float InDeltaTime ) override;
-	virtual int32 OnPaint( const FGeometry& AllottedGeometry, const FSlateRect& MyClippingRect, FSlateWindowElementList& OutDrawElements, int32 LayerId, const FWidgetStyle& InWidgetStyle, bool bParentEnabled ) const override;
+	virtual int32 OnPaint( const FPaintArgs& Args, const FGeometry& AllottedGeometry, const FSlateRect& MyClippingRect, FSlateWindowElementList& OutDrawElements, int32 LayerId, const FWidgetStyle& InWidgetStyle, bool bParentEnabled ) const override;
 	virtual void CacheDesiredSize() override;
 	virtual FVector2D ComputeDesiredSize() const override;
 	virtual FChildren* GetChildren() override;
@@ -519,3 +521,6 @@ private:
 	/** Notification interface object for text input method systems. */
 	TSharedPtr<ITextInputMethodChangeNotifier> TextInputMethodChangeNotifier;
 };
+
+
+#endif //WITH_FANCY_TEXT
\ No newline at end of file
diff --git a/Engine/Source/Runtime/Slate/Public/Widgets/Text/SRichTextBlock.h b/Engine/Source/Runtime/Slate/Public/Widgets/Text/SRichTextBlock.h
--- a/Engine/Source/Runtime/Slate/Public/Widgets/Text/SRichTextBlock.h
+++ b/Engine/Source/Runtime/Slate/Public/Widgets/Text/SRichTextBlock.h
@@ -1,6 +1,8 @@
 // Copyright 1998-2014 Epic Games, Inc. All Rights Reserved.
 #pragma once
 
+#if WITH_FANCY_TEXT
+
 /**
  * A rich static text widget. 
  * Through the use of markup and text decorators, text with different styles, embedded image and widgets can be achieved.
@@ -95,12 +97,10 @@ public:
 		return FHyperlinkDecorator::Create( Id, FSlateHyperlinkRun::FOnClick::CreateSP( InUserObjectPtr, NavigateFunc ) );
 	}
 
-	/************************************************************************/
-	/* SWidget                                                              */
-	/************************************************************************/
+	// Begin SWidget Interface
 	void Construct( const FArguments& InArgs );
 
-	virtual int32 OnPaint( const FGeometry& AllottedGeometry, const FSlateRect& MyClippingRect, FSlateWindowElementList& OutDrawElements, int32 LayerId, const FWidgetStyle& InWidgetStyle, bool bParentEnabled ) const;
+	virtual int32 OnPaint( const FPaintArgs& Args, const FGeometry& AllottedGeometry, const FSlateRect& MyClippingRect, FSlateWindowElementList& OutDrawElements, int32 LayerId, const FWidgetStyle& InWidgetStyle, bool bParentEnabled ) const;
 	virtual void Tick( const FGeometry& AllottedGeometry, const double InCurrentTime, const float InDeltaTime ) override;
 
 	virtual void CacheDesiredSize() override;
@@ -109,6 +109,7 @@ public:
 	virtual FChildren* GetChildren() override;
 
 	virtual void OnArrangeChildren( const FGeometry& AllottedGeometry, FArrangedChildren& ArrangedChildren ) const override;
+	// End SWidget Interface
 
 	/**
 	 * Gets the text assigned to this text block
@@ -161,3 +162,5 @@ private:
 
 	TSharedPtr< class IRichTextMarkupParser > Parser;
 };
+
+#endif //WITH_FANCY_TEXT
\ No newline at end of file
diff --git a/Engine/Source/Runtime/Slate/Public/Widgets/Tutorials/STutorialWrapper.h b/Engine/Source/Runtime/Slate/Public/Widgets/Tutorials/STutorialWrapper.h
--- a/Engine/Source/Runtime/Slate/Public/Widgets/Tutorials/STutorialWrapper.h
+++ b/Engine/Source/Runtime/Slate/Public/Widgets/Tutorials/STutorialWrapper.h
@@ -23,7 +23,7 @@ public:
 
 	/** Begin SWidget interface */
 	SLATE_API virtual void Tick( const FGeometry& AllottedGeometry, const double InCurrentTime, const float InDeltaTime ) override;
-	SLATE_API virtual int32 OnPaint(const FGeometry& AllottedGeometry, const FSlateRect& MyClippingRect, FSlateWindowElementList& OutDrawElements, int32 LayerId, const FWidgetStyle& InWidgetStyle, bool bParentEnabled ) const override;
+	SLATE_API virtual int32 OnPaint(const FPaintArgs& Args, const FGeometry& AllottedGeometry, const FSlateRect& MyClippingRect, FSlateWindowElementList& OutDrawElements, int32 LayerId, const FWidgetStyle& InWidgetStyle, bool bParentEnabled ) const override;
 	/** End SWidget interface */
 
 	/** 
diff --git a/Engine/Source/Runtime/Slate/Public/Widgets/Views/STableRow.h b/Engine/Source/Runtime/Slate/Public/Widgets/Views/STableRow.h
--- a/Engine/Source/Runtime/Slate/Public/Widgets/Views/STableRow.h
+++ b/Engine/Source/Runtime/Slate/Public/Widgets/Views/STableRow.h
@@ -134,9 +134,9 @@ public:
 		}
 	}
 
-	virtual int32 OnPaint( const FGeometry& AllottedGeometry, const FSlateRect& MyClippingRect, FSlateWindowElementList& OutDrawElements, int32 LayerId, const FWidgetStyle& InWidgetStyle, bool bParentEnabled ) const override
+	virtual int32 OnPaint( const FPaintArgs& Args, const FGeometry& AllottedGeometry, const FSlateRect& MyClippingRect, FSlateWindowElementList& OutDrawElements, int32 LayerId, const FWidgetStyle& InWidgetStyle, bool bParentEnabled ) const override
 	{
-		int32 Result = SBorder::OnPaint(AllottedGeometry, MyClippingRect, OutDrawElements, LayerId, InWidgetStyle, bParentEnabled);
+		int32 Result = SBorder::OnPaint(Args, AllottedGeometry, MyClippingRect, OutDrawElements, LayerId, InWidgetStyle, bParentEnabled);
 
 		TSharedRef< ITypedTableView<ItemType> > OwnerWidget = OwnerTablePtr.Pin().ToSharedRef();
 		const bool bIsActive = OwnerWidget->AsWidget()->HasKeyboardFocus();
diff --git a/Engine/Source/Runtime/Slate/Public/Widgets/Views/STableViewBase.h b/Engine/Source/Runtime/Slate/Public/Widgets/Views/STableViewBase.h
--- a/Engine/Source/Runtime/Slate/Public/Widgets/Views/STableViewBase.h
+++ b/Engine/Source/Runtime/Slate/Public/Widgets/Views/STableViewBase.h
@@ -82,13 +82,7 @@ public:
 	virtual TSharedRef<class SWidget> GetScrollWidget() override;
 	// End of IScrollableWidget interface
 
-	virtual int32 OnPaint(
-		const FGeometry& AllottedGeometry,
-		const FSlateRect& MyClippingRect,
-		FSlateWindowElementList& OutDrawElements,
-		int32 LayerId,
-		const FWidgetStyle& InWidgetStyle,
-		bool bParentEnabled) const override;
+	virtual int32 OnPaint(const FPaintArgs& Args, const FGeometry& AllottedGeometry, const FSlateRect& MyClippingRect, FSlateWindowElementList& OutDrawElements, int32 LayerId, const FWidgetStyle& InWidgetStyle, bool bParentEnabled) const override;
 
 	/** Is this list backing a tree or just a standalone list */
 	const ETableViewMode::Type TableViewMode;
diff --git a/Engine/Source/Runtime/SlateCore/Private/Input/HittestGrid.cpp b/Engine/Source/Runtime/SlateCore/Private/Input/HittestGrid.cpp
new file mode 100644
--- /dev/null
+++ b/Engine/Source/Runtime/SlateCore/Private/Input/HittestGrid.cpp
@@ -0,0 +1,208 @@
+// Copyright 1998-2014 Epic Games, Inc. All Rights Reserved.
+
+#include "SlateCorePrivatePCH.h"
+#include "HittestGrid.h"
+
+DEFINE_LOG_CATEGORY_STATIC(LogHittestDebug, Display, All);
+
+static const FVector2D CellSize(128.0f, 128.0f);
+
+struct FHittestGrid::FCachedWidget
+{
+	FCachedWidget( int32 InParentIndex, const FArrangedWidget& InWidget, const FSlateRect& InClippingRect )
+	: ArrangedWidget( InWidget )
+	, ClippingRect( InClippingRect )
+	, Children()
+	, ParentIndex( InParentIndex )
+	{
+	}
+
+	void AddChild( const int32 ChildIndex )
+	{
+		Children.Add( ChildIndex );
+	}
+
+	FArrangedWidget ArrangedWidget;
+	// @todo umg : ideally this clipping rect is optional and we only have them on a small number of widgets.
+	FSlateRect ClippingRect;
+	TArray<int32> Children;
+	int32 ParentIndex;
+};
+
+
+FHittestGrid::FHittestGrid()
+: WidgetsCachedThisFrame( new TArray<FCachedWidget>() )
+{
+}
+
+TArray<FArrangedWidget> FHittestGrid::GetBubblePath( FVector2D DesktopSpaceCoordinate )
+{
+	if (WidgetsCachedThisFrame->Num() > 0)
+	{
+		const FVector2D CursorPositionInGrid = DesktopSpaceCoordinate - GridOrigin;
+		const FIntPoint CellCoordinate = FIntPoint(
+			FMath::Max(FMath::FloorToInt(CursorPositionInGrid.X / CellSize.X), 0),
+			FMath::Max(FMath::FloorToInt(CursorPositionInGrid.Y / CellSize.Y), 0) );
+		
+		static FVector2D LastCoordinate = FVector2D::ZeroVector;
+		if ( LastCoordinate != CursorPositionInGrid )
+		{
+			LastCoordinate = CursorPositionInGrid;
+		}
+
+		const TArray<int32>& IndexesInCell = CellAt( CellCoordinate.X, CellCoordinate.Y ).CachedWidgetIndexes;
+		int32 HitWidgetIndex = INDEX_NONE;
+	
+		// Consider front-most widgets first for hittesting.
+		for ( int32 i = IndexesInCell.Num()-1; i>=0 && HitWidgetIndex==INDEX_NONE; --i )
+		{
+			const FCachedWidget& TestCandidate = (*WidgetsCachedThisFrame)[IndexesInCell[i]];
+			if ( TestCandidate.ArrangedWidget.Geometry.IsUnderLocation( DesktopSpaceCoordinate ) && TestCandidate.ClippingRect.ContainsPoint( DesktopSpaceCoordinate ) )
+			{
+				HitWidgetIndex = IndexesInCell[i];
+			}
+		}
+
+		if (HitWidgetIndex != INDEX_NONE)
+		{
+			TArray<FArrangedWidget> BubblePath;
+			int32 CurWidgetIndex=HitWidgetIndex;
+			do
+			{
+				const FCachedWidget& CurCachedWidget = (*WidgetsCachedThisFrame)[CurWidgetIndex];
+				BubblePath.Insert( CurCachedWidget.ArrangedWidget, 0 );
+				CurWidgetIndex = CurCachedWidget.ParentIndex;
+			}
+			while( CurWidgetIndex != INDEX_NONE );
+			return BubblePath;
+		}
+		else
+		{
+			return TArray<FArrangedWidget>();
+		}
+	}
+	else
+	{
+		return TArray<FArrangedWidget>();
+	}
+}
+
+void FHittestGrid::BeginFrame( const FSlateRect& HittestArea )
+{
+	//LogGrid();
+
+	GridOrigin = HittestArea.GetTopLeft();
+	const FVector2D GridSize = HittestArea.GetSize();
+	NumCells = FIntPoint( FMath::CeilToInt(GridSize.X / CellSize.X), FMath::CeilToInt(GridSize.Y / CellSize.Y) );
+	WidgetsCachedThisFrame->Empty();
+	Cells.Reset( NumCells.X * NumCells.Y );	
+	Cells.SetNum( NumCells.X * NumCells.Y );
+}
+
+int32 FHittestGrid::InsertWidget( const int32 ParentHittestIndex, const EVisibility& Visibility, const FArrangedWidget& Widget, const FVector2D InWindowOffset, const FSlateRect& InClippingRect )
+{
+	FArrangedWidget WindowAdjustedWidget(Widget);
+	WindowAdjustedWidget.Geometry.AbsolutePosition += InWindowOffset;
+
+	const FSlateRect WindowAdjustedRect(
+		InClippingRect.Left + InWindowOffset.X, InClippingRect.Top + InWindowOffset.Y,
+		InClippingRect.Right + InWindowOffset.X, InClippingRect.Bottom + InWindowOffset.Y );
+
+	// Remember this widget, its geometry, and its place in the logical hierarchy.
+	const int32 WidgetIndex = WidgetsCachedThisFrame->Add( FCachedWidget( ParentHittestIndex, WindowAdjustedWidget, WindowAdjustedRect ) );
+	if (ParentHittestIndex != INDEX_NONE)
+	{
+		(*WidgetsCachedThisFrame)[ParentHittestIndex].AddChild( WidgetIndex );
+	}
+	
+	if (Visibility.IsHitTestVisible())
+	{
+		// Mark any cell that is overlapped by this widget.
+		const FVector2D UpperLeft = WindowAdjustedWidget.Geometry.AbsolutePosition - GridOrigin;
+		const FVector2D LowerRight = UpperLeft + WindowAdjustedWidget.Geometry.Size*WindowAdjustedWidget.Geometry.Scale;
+		const FSlateRect GridSpaceWindowRect = FSlateRect(WindowAdjustedRect.Left - GridOrigin.X, WindowAdjustedRect.Top - GridOrigin.Y, WindowAdjustedRect.Right - GridOrigin.X, WindowAdjustedRect.Bottom - GridOrigin.Y);
+		const FSlateRect CellCoverageRect = FSlateRect(UpperLeft, LowerRight).IntersectionWith(GridSpaceWindowRect);
+		
+		// Starting and ending cells covered by this widget.	
+		const FIntPoint UpperLeftCell = FIntPoint(
+			FMath::Max(0, FMath::FloorToInt(CellCoverageRect.Left / CellSize.X)),
+			FMath::Max(0, FMath::FloorToInt(CellCoverageRect.Top / CellSize.Y)));
+
+		const FIntPoint LowerRightCell = FIntPoint(
+			FMath::Min( NumCells.X-1, FMath::FloorToInt(CellCoverageRect.Right / CellSize.X)),
+			FMath::Min( NumCells.Y-1, FMath::FloorToInt(CellCoverageRect.Bottom / CellSize.Y)));
+
+		for (int32 XIndex=UpperLeftCell.X; XIndex <= LowerRightCell.X; ++ XIndex )
+		{
+			for(int32 YIndex=UpperLeftCell.Y; YIndex <= LowerRightCell.Y; ++YIndex)
+			{
+				CellAt(XIndex, YIndex).CachedWidgetIndexes.Add( WidgetIndex );
+			}
+		}
+	}
+	
+	return WidgetIndex;
+
+}
+
+FHittestGrid::FCell& FHittestGrid::CellAt( const int32 X, const int32 Y )
+{
+	return Cells[ Y*NumCells.X + X ];
+}
+
+const FHittestGrid::FCell& FHittestGrid::CellAt( const int32 X, const int32 Y ) const
+{
+	return Cells[ Y*NumCells.X + X ];
+}
+
+
+void FHittestGrid::LogGrid() const
+{
+	FString TempString;
+	for (int y=0; y<NumCells.Y; ++y)
+	{
+		for (int x=0; x<NumCells.X; ++x)
+		{
+			TempString += "\t";
+			TempString += "[";
+			for( int i : CellAt(x, y).CachedWidgetIndexes )
+			{
+				TempString += FString::Printf(TEXT("%d,"), i);
+			}
+			TempString += "]";
+		}
+		TempString += "\n";
+	}
+
+	TempString += "\n";
+
+	UE_LOG( LogHittestDebug, Warning,TEXT("\n%s"), *TempString );
+
+	for ( int i=0; i<WidgetsCachedThisFrame->Num(); ++i )
+	{
+		if ( (*WidgetsCachedThisFrame)[i].ParentIndex == INDEX_NONE )
+		{
+			LogChildren( i, 0, *WidgetsCachedThisFrame );
+		}
+	}
+	
+
+}
+
+
+void FHittestGrid::LogChildren(int32 Index, int32 IndentLevel, const TArray<FCachedWidget>& WidgetsCachedThisFrame)
+{
+	FString IndentString;
+	for (int i = 0; i<IndentLevel; ++i)
+	{
+		IndentString += "|\t";
+	}
+
+	const FCachedWidget& CachedWidget = WidgetsCachedThisFrame[Index];
+	UE_LOG( LogHittestDebug, Warning, TEXT("%s[%d] => %s"), *IndentString, Index, *CachedWidget.ArrangedWidget.ToString() );
+
+	for ( int i=0; i<CachedWidget.Children.Num(); ++i )
+	{
+		LogChildren( CachedWidget.Children[i], IndentLevel+1, WidgetsCachedThisFrame );
+	}
+}
\ No newline at end of file
diff --git a/Engine/Source/Runtime/SlateCore/Private/Types/PaintArgs.cpp b/Engine/Source/Runtime/SlateCore/Private/Types/PaintArgs.cpp
new file mode 100644
--- /dev/null
+++ b/Engine/Source/Runtime/SlateCore/Private/Types/PaintArgs.cpp
@@ -0,0 +1,36 @@
+// Copyright 1998-2014 Epic Games, Inc. All Rights Reserved.
+
+#include "SlateCorePrivatePCH.h"
+
+FPaintArgs::FPaintArgs( const TSharedRef<SWidget>& Parent, FHittestGrid& InHittestGrid, FVector2D InWindowOffset )
+: ParentPtr(Parent)
+, Grid(InHittestGrid)
+, LastHittestIndex(INDEX_NONE)
+, LastRecordedVisibility( EVisibility::Visible )
+, WindowOffset(InWindowOffset)
+{
+}
+
+
+FPaintArgs FPaintArgs::WithNewParent( const SWidget* Parent ) const
+{
+	FPaintArgs Args = FPaintArgs( const_cast<SWidget*>(Parent)->AsShared(), this->Grid, this->WindowOffset );
+	Args.LastHittestIndex = this->LastHittestIndex;
+	Args.LastRecordedVisibility = this->LastRecordedVisibility;
+	return Args;
+}
+
+
+FPaintArgs FPaintArgs::RecordHittestGeometry(const SWidget* Widget, const FGeometry& WidgetGeometry, const FSlateRect& InClippingRect) const
+{
+	const EVisibility RecordedVisibility = (LastRecordedVisibility.AreChildrenHitTestVisible())
+		? Widget->GetVisibility()
+		: LastRecordedVisibility;
+
+	const int32 RecordedHittestIndex = Grid.InsertWidget( LastHittestIndex, RecordedVisibility, FArrangedWidget(const_cast<SWidget*>(Widget)->AsShared(), WidgetGeometry), WindowOffset, InClippingRect );
+
+	FPaintArgs UpdatedArgs(*this);
+	UpdatedArgs.LastHittestIndex = RecordedHittestIndex;
+	UpdatedArgs.LastRecordedVisibility = RecordedVisibility;
+	return UpdatedArgs;
+}
\ No newline at end of file
diff --git a/Engine/Source/Runtime/SlateCore/Private/Widgets/SCompoundWidget.cpp b/Engine/Source/Runtime/SlateCore/Private/Widgets/SCompoundWidget.cpp
--- a/Engine/Source/Runtime/SlateCore/Private/Widgets/SCompoundWidget.cpp
+++ b/Engine/Source/Runtime/SlateCore/Private/Widgets/SCompoundWidget.cpp
@@ -7,20 +7,7 @@
 #include "SlateCorePrivatePCH.h"
 
 
-/**
- * The widget should respond by populating the OutDrawElements array with FDrawElements 
- * that represent it and any of its children.
- *
- * @param AllottedGeometry  The FGeometry that describes an area in which the widget should appear.
- * @param MyClippingRect    The clipping rectangle allocated for this widget and its children.
- * @param OutDrawElements   A list of FDrawElements to populate with the output.
- * @param LayerId           The Layer onto which this widget should be rendered.
- * @param InColorAndOpacity Color and Opacity to be applied to all the descendants of the widget being painted
- * @param bParentEnabled	True if the parent of this widget is enabled.
- *
- * @return The maximum layer ID attained by this widget or any of its children.
- */
-int32 SCompoundWidget::OnPaint( const FGeometry& AllottedGeometry, const FSlateRect& MyClippingRect, FSlateWindowElementList& OutDrawElements, int32 LayerId, const FWidgetStyle& InWidgetStyle, bool bParentEnabled ) const
+int32 SCompoundWidget::OnPaint( const FPaintArgs& Args, const FGeometry& AllottedGeometry, const FSlateRect& MyClippingRect, FSlateWindowElementList& OutDrawElements, int32 LayerId, const FWidgetStyle& InWidgetStyle, bool bParentEnabled ) const
 {
 	// A CompoundWidget just draws its children
 	FArrangedChildren ArrangedChildren(EVisibility::Visible);
@@ -43,7 +30,7 @@ int32 SCompoundWidget::OnPaint( const FGeometry& AllottedGeometry, const FSlateR
 			.BlendColorAndOpacityTint(ColorAndOpacity.Get())
 			.SetForegroundColor( ForegroundColor );
 
-		return TheChild.Widget->Paint( TheChild.Geometry, ChildClippingRect, OutDrawElements, LayerId + 1, CompoundedWidgetStyle, ShouldBeEnabled( bParentEnabled ) );
+		return TheChild.Widget->Paint( Args.WithNewParent(this), TheChild.Geometry, ChildClippingRect, OutDrawElements, LayerId + 1, CompoundedWidgetStyle, ShouldBeEnabled( bParentEnabled ) );
 	}
 	return LayerId;
 }
diff --git a/Engine/Source/Runtime/SlateCore/Private/Widgets/SNullWidget.cpp b/Engine/Source/Runtime/SlateCore/Private/Widgets/SNullWidget.cpp
--- a/Engine/Source/Runtime/SlateCore/Private/Widgets/SNullWidget.cpp
+++ b/Engine/Source/Runtime/SlateCore/Private/Widgets/SNullWidget.cpp
@@ -34,7 +34,7 @@ public:
 
 	// SWidget interface
 
-	virtual int32 OnPaint( const FGeometry& AllottedGeometry, const FSlateRect& MyClippingRect, FSlateWindowElementList& OutDrawElements, int32 LayerId, const FWidgetStyle& InWidgetStyle, bool bParentEnabled ) const override
+	virtual int32 OnPaint( const FPaintArgs& Args, const FGeometry& AllottedGeometry, const FSlateRect& MyClippingRect, FSlateWindowElementList& OutDrawElements, int32 LayerId, const FWidgetStyle& InWidgetStyle, bool bParentEnabled ) const override
 	{
 		return LayerId;
 	}
diff --git a/Engine/Source/Runtime/SlateCore/Private/Widgets/SOverlay.cpp b/Engine/Source/Runtime/SlateCore/Private/Widgets/SOverlay.cpp
--- a/Engine/Source/Runtime/SlateCore/Private/Widgets/SOverlay.cpp
+++ b/Engine/Source/Runtime/SlateCore/Private/Widgets/SOverlay.cpp
@@ -70,20 +70,7 @@ FChildren* SOverlay::GetChildren()
 	return &Children;
 }
 
-/**
- * The widget should respond by populating the OutDrawElements array with FDrawElements 
- * that represent it and any of its children.
- *
- * @param AllottedGeometry  The FGeometry that describes an area in which the widget should appear.
- * @param MyClippingRect    The clipping rectangle allocated for this widget and its children.
- * @param OutDrawElements   A list of FDrawElements to populate with the output.
- * @param LayerId           The Layer onto which this widget should be rendered.
- * @param InColorAndOpacity Color and Opacity to be applied to all the descendants of the widget being painted
- * @param bParentEnabled	True if the parent of this widget is enabled.
- *
- * @return The maximum layer ID attained by this widget or any of its children.
- */
-int32 SOverlay::OnPaint( const FGeometry& AllottedGeometry, const FSlateRect& MyClippingRect, FSlateWindowElementList& OutDrawElements, int32 LayerId, const FWidgetStyle& InWidgetStyle, bool bParentEnabled ) const
+int32 SOverlay::OnPaint( const FPaintArgs& Args, const FGeometry& AllottedGeometry, const FSlateRect& MyClippingRect, FSlateWindowElementList& OutDrawElements, int32 LayerId, const FWidgetStyle& InWidgetStyle, bool bParentEnabled ) const
 {
 	FArrangedChildren ArrangedChildren(EVisibility::Visible);
 	{
@@ -102,7 +89,7 @@ int32 SOverlay::OnPaint( const FGeometry& AllottedGeometry, const FSlateRect& My
 	{
 		FArrangedWidget& CurWidget = ArrangedChildren(ChildIndex);
 		FSlateRect ChildClipRect = MyClippingRect.IntersectionWith( CurWidget.Geometry.GetClippingRect() );
-		const int32 CurWidgetsMaxLayerId = CurWidget.Widget->Paint( CurWidget.Geometry, ChildClipRect, OutDrawElements, MaxLayerId + 1, InWidgetStyle, ShouldBeEnabled( bParentEnabled ) );
+		const int32 CurWidgetsMaxLayerId = CurWidget.Widget->Paint( Args.WithNewParent(this), CurWidget.Geometry, ChildClipRect, OutDrawElements, MaxLayerId + 1, InWidgetStyle, ShouldBeEnabled( bParentEnabled ) );
 
 		MaxLayerId = FMath::Max( MaxLayerId, CurWidgetsMaxLayerId );
 	}
diff --git a/Engine/Source/Runtime/SlateCore/Private/Widgets/SPanel.cpp b/Engine/Source/Runtime/SlateCore/Private/Widgets/SPanel.cpp
--- a/Engine/Source/Runtime/SlateCore/Private/Widgets/SPanel.cpp
+++ b/Engine/Source/Runtime/SlateCore/Private/Widgets/SPanel.cpp
@@ -10,7 +10,7 @@
 /* SWidget overrides
  *****************************************************************************/
 
-int32 SPanel::OnPaint( const FGeometry& AllottedGeometry, const FSlateRect& MyClippingRect, FSlateWindowElementList& OutDrawElements, int32 LayerId, const FWidgetStyle& InWidgetStyle, bool bParentEnabled ) const
+int32 SPanel::OnPaint( const FPaintArgs& Args, const FGeometry& AllottedGeometry, const FSlateRect& MyClippingRect, FSlateWindowElementList& OutDrawElements, int32 LayerId, const FWidgetStyle& InWidgetStyle, bool bParentEnabled ) const
 {
 	// REMOVE THIS: Currently adding a debug quad so that BoxPanels are visible during development
 	// OutDrawElements.AddItem( FSlateDrawElement::MakeQuad( AllottedGeometry.AbsolutePosition, AllottedGeometry.Size ) );	
@@ -24,14 +24,14 @@ int32 SPanel::OnPaint( const FGeometry& AllottedGeometry, const FSlateRect& MyCl
 		this->ArrangeChildren(AllottedGeometry, ArrangedChildren);
 	}
 
-	return PaintArrangedChildren(ArrangedChildren, MyClippingRect, OutDrawElements, LayerId, InWidgetStyle, bParentEnabled);
+	return PaintArrangedChildren(Args, ArrangedChildren, MyClippingRect, OutDrawElements, LayerId, InWidgetStyle, bParentEnabled);
 }
 
 
 /* SPanel implementation
  *****************************************************************************/
 
-int32 SPanel::PaintArrangedChildren( const FArrangedChildren& ArrangedChildren, const FSlateRect& MyClippingRect, FSlateWindowElementList& OutDrawElements, int32 LayerId, const FWidgetStyle& InWidgetStyle, bool bParentEnabled  ) const
+int32 SPanel::PaintArrangedChildren( const FPaintArgs& Args, const FArrangedChildren& ArrangedChildren, const FSlateRect& MyClippingRect, FSlateWindowElementList& OutDrawElements, int32 LayerId, const FWidgetStyle& InWidgetStyle, bool bParentEnabled  ) const
 {
 	// Because we paint multiple children, we must track the maximum layer id that they produced in case one of our parents
 	// wants to an overlay for all of its contents.
@@ -44,7 +44,7 @@ int32 SPanel::PaintArrangedChildren( const FArrangedChildren& ArrangedChildren,
 		
 		if (ChildClipRect.GetSize().Size() > 0)
 		{
-			const int32 CurWidgetsMaxLayerId = CurWidget.Widget->Paint(CurWidget.Geometry, ChildClipRect, OutDrawElements, LayerId, InWidgetStyle, ShouldBeEnabled(bParentEnabled));
+			const int32 CurWidgetsMaxLayerId = CurWidget.Widget->Paint(Args.WithNewParent(this), CurWidget.Geometry, ChildClipRect, OutDrawElements, LayerId, InWidgetStyle, ShouldBeEnabled(bParentEnabled));
 			MaxLayerId = FMath::Max(MaxLayerId, CurWidgetsMaxLayerId);
 		}		
 		
diff --git a/Engine/Source/Runtime/SlateCore/Private/Widgets/SWidget.cpp b/Engine/Source/Runtime/SlateCore/Private/Widgets/SWidget.cpp
--- a/Engine/Source/Runtime/SlateCore/Private/Widgets/SWidget.cpp
+++ b/Engine/Source/Runtime/SlateCore/Private/Widgets/SWidget.cpp
@@ -5,6 +5,12 @@
 DECLARE_CYCLE_STAT(TEXT("OnPaint"), STAT_SlateOnPaint, STATGROUP_Slate);
 DECLARE_CYCLE_STAT(TEXT("ArrangeChildren"), STAT_SlateArrangeChildren, STATGROUP_Slate);
 
+int32 SWidget::bUseLegacyHittest = false;
+FAutoConsoleVariableRef SWidget::CVarUseLegacyHittest(
+	TEXT("Slate.UseLegacyHittest"),
+	SWidget::bUseLegacyHittest,
+	TEXT("Should slate fall-back to using the legacy hittest path that does not support child widgets outside their parents' bounds.") );
+
 SWidget::SWidget()
 	: CreatedInFile( TEXT("") )
 	, CreatedOnLine( -1 )
@@ -487,10 +493,20 @@ void SWidget::SetDebugInfo( const ANSICHAR* InType, const ANSICHAR* InFile, int3
 	this->CreatedOnLine = OnLine;
 }
 
-int32 SWidget::Paint(const FGeometry& AllottedGeometry, const FSlateRect& MyClippingRect, FSlateWindowElementList& OutDrawElements, int32 LayerId, const FWidgetStyle& InWidgetStyle, bool bParentEnabled) const
+int32 SWidget::Paint(const FPaintArgs& Args, const FGeometry& AllottedGeometry, const FSlateRect& MyClippingRect, FSlateWindowElementList& OutDrawElements, int32 LayerId, const FWidgetStyle& InWidgetStyle, bool bParentEnabled) const
 {
 	SCOPE_CYCLE_COUNTER(STAT_SlateOnPaint);
-	return OnPaint(AllottedGeometry, MyClippingRect, OutDrawElements, LayerId, InWidgetStyle, bParentEnabled);
+	if ( !bUseLegacyHittest )
+	{
+		FPaintArgs UpdatedArgs = Args.RecordHittestGeometry( this, AllottedGeometry, MyClippingRect );
+		return OnPaint(UpdatedArgs, AllottedGeometry, MyClippingRect, OutDrawElements, LayerId, InWidgetStyle, bParentEnabled);
+	}
+	else
+	{
+		return OnPaint(Args, AllottedGeometry, MyClippingRect, OutDrawElements, LayerId, InWidgetStyle, bParentEnabled);
+	}
+
+	
 }
 
 void SWidget::ArrangeChildren(const FGeometry& AllottedGeometry, FArrangedChildren& ArrangedChildren) const
@@ -498,3 +514,9 @@ void SWidget::ArrangeChildren(const FGeometry& AllottedGeometry, FArrangedChildr
 	SCOPE_CYCLE_COUNTER(STAT_SlateArrangeChildren);
 	OnArrangeChildren(AllottedGeometry, ArrangedChildren);
 }
+
+
+bool SWidget::UseLegacyHittest()
+{
+	return bUseLegacyHittest != 0;
+}
\ No newline at end of file
diff --git a/Engine/Source/Runtime/SlateCore/Private/Widgets/SWindow.cpp b/Engine/Source/Runtime/SlateCore/Private/Widgets/SWindow.cpp
--- a/Engine/Source/Runtime/SlateCore/Private/Widgets/SWindow.cpp
+++ b/Engine/Source/Runtime/SlateCore/Private/Widgets/SWindow.cpp
@@ -266,6 +266,9 @@ void SWindow::Construct(const FArguments& InArgs)
 	this->InitialDesiredScreenPosition = WindowPosition;
 	this->InitialDesiredSize = WindowSize;
 
+	// Window visibility is currently driven by whether the window is interactive.
+	this->Visibility = TAttribute<EVisibility>::Create( TAttribute<EVisibility>::FGetter::CreateRaw(this, &SWindow::GetWindowVisibility) );
+
 	this->ConstructWindowInternals( bCreateTitleBar );
 	this->SetContent( InArgs._Content.Widget );
 }
@@ -339,7 +342,6 @@ FVector2D SWindow::ComputeWindowSizeForContent( FVector2D ContentSize )
 	return ContentSize + FVector2D(0, SWindowDefs::DefaultTitleBarSize) + SWindowDefs::WindowBorderSize.GetDesiredSize();
 }
 
-
 void SWindow::ConstructWindowInternals( const bool bCreateTitleBar )
 {
 	ForegroundColor = FCoreStyle::Get().GetSlateColor("DefaultForeground");
@@ -470,6 +472,7 @@ void SWindow::ConstructWindowInternals( const bool bCreateTitleBar )
 	}
 }
 
+
 /** Are any of our child windows active? */
 bool SWindow::HasActiveChildren() const
 {
@@ -775,6 +778,13 @@ FSlateColor SWindow::GetWindowOutlineColor() const
 	return Style->OutlineColor;
 }
 
+EVisibility SWindow::GetWindowVisibility() const
+{
+	return ( AcceptsInput() )
+		? EVisibility::Visible
+		: EVisibility::HitTestInvisible;
+}
+
 void SWindow::UpdateMorphTargetShape( const FSlateRect& TargetShape )
 {
 	Morpher.TargetMorphShape = TargetShape;
@@ -1634,6 +1644,7 @@ SWindow::SWindow()
 	, bShouldShowWindowContentDuringOverlay( false )
 	, ExpectedMaxWidth( INDEX_NONE )
 	, ExpectedMaxHeight( INDEX_NONE )
+	, TitleBar()
 {
 }
 
diff --git a/Engine/Source/Runtime/SlateCore/Public/Application/SlateApplicationBase.h b/Engine/Source/Runtime/SlateCore/Public/Application/SlateApplicationBase.h
--- a/Engine/Source/Runtime/SlateCore/Public/Application/SlateApplicationBase.h
+++ b/Engine/Source/Runtime/SlateCore/Public/Application/SlateApplicationBase.h
@@ -6,6 +6,8 @@
 
 #pragma once
 
+class FHittestGrid;
+
 
 /**
  * Interface for window title bars.
diff --git a/Engine/Source/Runtime/SlateCore/Public/Input/HittestGrid.h b/Engine/Source/Runtime/SlateCore/Public/Input/HittestGrid.h
new file mode 100644
--- /dev/null
+++ b/Engine/Source/Runtime/SlateCore/Public/Input/HittestGrid.h
@@ -0,0 +1,76 @@
+// Copyright 1998-2014 Epic Games, Inc. All Rights Reserved.
+
+#pragma once
+
+class FArrangedWidget;
+
+
+class SLATECORE_API FHittestGrid
+{
+public:
+
+	FHittestGrid();
+
+	/**
+	 * Given a Slate Units coordinate in virtual desktop space, perform a hittest
+	 * and return the path along which the corresponding event would be bubbled.
+	 */
+	TArray<FArrangedWidget> GetBubblePath( FVector2D DesktopSpaceCoordinate );
+
+	/**
+	 * Clear the hittesting area and prepare to execute a new frame.
+	 * Depending on monitor arrangement, the area to be hittested could begin
+	 * in the negative coordinates.
+	 *
+	 * @param HittestArea       Size in Slate Units of the area we are considering for hittesting.
+	 */
+	void BeginFrame(const FSlateRect& HittestArea);
+
+	/** Add Widget into the hittest data structure so that we can later make queries about it. */
+	int32 InsertWidget( const int32 ParentHittestIndex, const EVisibility& Visibility, const FArrangedWidget& Widget, const FVector2D InWindowOffset, const FSlateRect& InClippingRect );
+
+private:
+
+	/**
+	 * All the available space is partitioned into Cells.
+	 * Each Cell contains a list of widgets that overlap the cell.
+	 * The list is ordered from back to front.
+	 */
+	struct FCell
+	{
+		TArray<int32> CachedWidgetIndexes;
+	};
+
+	/**
+	 * @todo umg: can we eliminate this?
+	 *
+	 * FNodes are a duplicate of the widget tree, with each widget
+	 * that is encountered on the way to outputting a hit-testable element
+	 * being recorded such that the event bubbling path can be reconstructed.
+	 */
+	struct FCachedWidget;
+
+private:
+	
+	friend class SWidgetReflector;
+
+	/** Access a cell at coordinates X, Y. Coordinates are row and column indexes. */
+	FCell& CellAt( const int32 X, const int32 Y );
+	const FCell& CellAt( const int32 X, const int32 Y ) const;
+
+	/** All the widgets and their arranged geometries encountered this frame. */
+	TSharedRef< TArray<FCachedWidget> > WidgetsCachedThisFrame;
+
+	/** The cells that make up the space partition. */
+	TArray<FCell> Cells;
+
+	/** Where the 0,0 of the upper-left-most cell corresponds to in desktop space. */
+	FVector2D GridOrigin;
+
+	/** The size of the grid in cells. */
+	FIntPoint NumCells;
+
+	void LogGrid() const;
+
+	static void LogChildren(int32 Index, int32 IndentLevel, const TArray<FCachedWidget>& WidgetsCachedThisFrame);
+};
\ No newline at end of file
diff --git a/Engine/Source/Runtime/SlateCore/Public/Layout/ArrangedChildren.h b/Engine/Source/Runtime/SlateCore/Public/Layout/ArrangedChildren.h
--- a/Engine/Source/Runtime/SlateCore/Public/Layout/ArrangedChildren.h
+++ b/Engine/Source/Runtime/SlateCore/Public/Layout/ArrangedChildren.h
@@ -32,6 +32,14 @@ class SLATECORE_API FArrangedChildren
 	{
 	}
 
+	// @todo hittest2.0 : we should get rid of this eventually.
+	static FArrangedChildren Hittest2_FromArray(const TArray<FArrangedWidget>& InWidgets)
+	{
+		FArrangedChildren Temp( EVisibility::All );
+		Temp.Array = InWidgets;
+		return Temp;
+	}
+
 	/** Reverse the order of the arranged children */
 	void Reverse()
 	{
diff --git a/Engine/Source/Runtime/SlateCore/Public/Layout/WidgetPath.h b/Engine/Source/Runtime/SlateCore/Public/Layout/WidgetPath.h
--- a/Engine/Source/Runtime/SlateCore/Public/Layout/WidgetPath.h
+++ b/Engine/Source/Runtime/SlateCore/Public/Layout/WidgetPath.h
@@ -44,6 +44,13 @@ public:
 	
 	}
 
+	FWidgetPath( const TArray<FArrangedWidget>& InWidgetPath )
+	: Widgets( FArrangedChildren::Hittest2_FromArray(InWidgetPath) )
+	, TopLevelWindow( InWidgetPath.Num() > 0 ? StaticCastSharedRef<SWindow>(InWidgetPath[0].Widget) : TSharedPtr<SWindow>(nullptr) )
+	{
+	
+	}
+
 	/**
 	 * @param MarkerWidget Copy the path up to and including this widget 
 	 *
diff --git a/Engine/Source/Runtime/SlateCore/Public/SlateCore.h b/Engine/Source/Runtime/SlateCore/Public/SlateCore.h
--- a/Engine/Source/Runtime/SlateCore/Public/SlateCore.h
+++ b/Engine/Source/Runtime/SlateCore/Public/SlateCore.h
@@ -57,6 +57,9 @@ DECLARE_MEMORY_STAT_EXTERN(TEXT("Vertex Buffer Memory"), STAT_SlateVertexBufferM
 DECLARE_MEMORY_STAT_EXTERN(TEXT("Index Buffer Memory"), STAT_SlateIndexBufferMemory, STATGROUP_SlateMemory, SLATECORE_API);
 
 
+// Compile all the RichText and MultiLine editable text?
+#define WITH_FANCY_TEXT 1
+
 /* Includes
  *****************************************************************************/
 
@@ -64,6 +67,7 @@ DECLARE_MEMORY_STAT_EXTERN(TEXT("Index Buffer Memory"), STAT_SlateIndexBufferMem
 #include "SlateConstants.h"
 #include "SlateEnums.h"
 #include "SlateStructs.h"
+#include "PaintArgs.h"
 
 // Layout
 #include "Margin.h"
diff --git a/Engine/Source/Runtime/SlateCore/Public/Types/PaintArgs.h b/Engine/Source/Runtime/SlateCore/Public/Types/PaintArgs.h
new file mode 100644
--- /dev/null
+++ b/Engine/Source/Runtime/SlateCore/Public/Types/PaintArgs.h
@@ -0,0 +1,31 @@
+// Copyright 1998-2014 Epic Games, Inc. All Rights Reserved.
+
+#pragma once
+
+#include "Visibility.h"
+
+class SWidget;
+struct FGeometry;
+class FHittestGrid;
+class FSlateRect;
+
+/**
+ * SWidget::OnPaint and SWidget::Paint use FPaintArgs as their
+ * sole parameter in order to ease the burden of passing
+ * through multiple fields.
+ */
+class SLATECORE_API FPaintArgs
+{
+public:
+	FPaintArgs( const TSharedRef<SWidget>& Parent, FHittestGrid& InHittestGrid, FVector2D InWindowOffset );
+	FPaintArgs WithNewParent( const SWidget* Parent ) const;
+	FPaintArgs RecordHittestGeometry(const SWidget* Widget, const FGeometry& WidgetGeometry, const FSlateRect& InClippingRect) const;
+
+
+private:
+	const TSharedRef<SWidget>& ParentPtr;
+	FHittestGrid& Grid;
+	int32 LastHittestIndex;
+	EVisibility LastRecordedVisibility;
+	FVector2D WindowOffset;
+};
\ No newline at end of file
diff --git a/Engine/Source/Runtime/SlateCore/Public/Widgets/SLeafWidget.h b/Engine/Source/Runtime/SlateCore/Public/Widgets/SLeafWidget.h
--- a/Engine/Source/Runtime/SlateCore/Public/Widgets/SLeafWidget.h
+++ b/Engine/Source/Runtime/SlateCore/Public/Widgets/SLeafWidget.h
@@ -30,7 +30,7 @@ public:
 	 *
 	 * Whenever possible, LeafWidgets should avoid dealing with layout properties. See TextBlock for an example.
 	 */
-	virtual int32 OnPaint( const FGeometry& AllottedGeometry, const FSlateRect& MyClippingRect, FSlateWindowElementList& OutDrawElements, int32 LayerId, const FWidgetStyle& InWidgetStyle, bool bParentEnabled ) const = 0;
+	virtual int32 OnPaint( const FPaintArgs& Args, const FGeometry& AllottedGeometry, const FSlateRect& MyClippingRect, FSlateWindowElementList& OutDrawElements, int32 LayerId, const FWidgetStyle& InWidgetStyle, bool bParentEnabled ) const = 0;
 	
 	/**
 	 * Overwritten from SWidget.
diff --git a/Engine/Source/Runtime/SlateCore/Public/Widgets/SPanel.h b/Engine/Source/Runtime/SlateCore/Public/Widgets/SPanel.h
--- a/Engine/Source/Runtime/SlateCore/Public/Widgets/SPanel.h
+++ b/Engine/Source/Runtime/SlateCore/Public/Widgets/SPanel.h
@@ -56,7 +56,7 @@ public:
 
 	// Begin SWidget overrides
 
-	virtual int32 OnPaint( const FGeometry& AllottedGeometry, const FSlateRect& MyClippingRect, FSlateWindowElementList& OutDrawElements, int32 LayerId, const FWidgetStyle& InWidgetStyle, bool bParentEnabled ) const override;
+	virtual int32 OnPaint( const FPaintArgs& Args, const FGeometry& AllottedGeometry, const FSlateRect& MyClippingRect, FSlateWindowElementList& OutDrawElements, int32 LayerId, const FWidgetStyle& InWidgetStyle, bool bParentEnabled ) const override;
 
 	// End SWidget overrides
 
@@ -65,7 +65,7 @@ protected:
 	/**
 	 * Just like OnPaint, but takes already arranged children. Can be handy for writing custom SPanels.
 	 */
-	int32 PaintArrangedChildren( const FArrangedChildren& ArrangedChildren, const FSlateRect& MyClippingRect, FSlateWindowElementList& OutDrawElements, int32 LayerId, const FWidgetStyle& InWidgetStyle, bool bParentEnabled  ) const;
+	int32 PaintArrangedChildren( const FPaintArgs& Args, const FArrangedChildren& ArrangedChildren, const FSlateRect& MyClippingRect, FSlateWindowElementList& OutDrawElements, int32 LayerId, const FWidgetStyle& InWidgetStyle, bool bParentEnabled  ) const;
 	
 protected:
 
diff --git a/Engine/Source/Runtime/SlateCore/Public/Widgets/SWidget.h b/Engine/Source/Runtime/SlateCore/Public/Widgets/SWidget.h
--- a/Engine/Source/Runtime/SlateCore/Public/Widgets/SWidget.h
+++ b/Engine/Source/Runtime/SlateCore/Public/Widgets/SWidget.h
@@ -97,6 +97,7 @@ public:
 	 * The widget should respond by populating the OutDrawElements array with FDrawElements 
 	 * that represent it and any of its children.
 	 *
+	 * @param Args              All the arguments necessary to paint this widget (@todo umg: move all params into this struct)
 	 * @param AllottedGeometry  The FGeometry that describes an area in which the widget should appear.
 	 * @param MyClippingRect    The clipping rectangle allocated for this widget and its children.
 	 * @param OutDrawElements   A list of FDrawElements to populate with the output.
@@ -106,7 +107,7 @@ public:
 	 *
 	 * @return The maximum layer ID attained by this widget or any of its children.
 	 */
-	int32 Paint( const FGeometry& AllottedGeometry, const FSlateRect& MyClippingRect, FSlateWindowElementList& OutDrawElements, int32 LayerId, const FWidgetStyle& InWidgetStyle, bool bParentEnabled ) const;
+	int32 Paint( const FPaintArgs& Args, const FGeometry& AllottedGeometry, const FSlateRect& MyClippingRect, FSlateWindowElementList& OutDrawElements, int32 LayerId, const FWidgetStyle& InWidgetStyle, bool bParentEnabled ) const;
 
 	/**
 	 * Ticks this widget.  Override in derived classes, but always call the parent implementation.
@@ -696,6 +697,7 @@ private:
 	* that represent it and any of its children. Called by the non-virtual OnPaint to enforce pre/post conditions
 	* during OnPaint.
 	*
+	* @param Args              All the arguments necessary to paint this widget (@todo umg: move all params into this struct)
 	* @param AllottedGeometry  The FGeometry that describes an area in which the widget should appear.
 	* @param MyClippingRect    The clipping rectangle allocated for this widget and its children.
 	* @param OutDrawElements   A list of FDrawElements to populate with the output.
@@ -705,7 +707,7 @@ private:
 	*
 	* @return The maximum layer ID attained by this widget or any of its children.
 	*/
-	virtual int32 OnPaint(const FGeometry& AllottedGeometry, const FSlateRect& MyClippingRect, FSlateWindowElementList& OutDrawElements, int32 LayerId, const FWidgetStyle& InWidgetStyle, bool bParentEnabled) const = 0;
+	virtual int32 OnPaint(const FPaintArgs& Args, const FGeometry& AllottedGeometry, const FSlateRect& MyClippingRect, FSlateWindowElementList& OutDrawElements, int32 LayerId, const FWidgetStyle& InWidgetStyle, bool bParentEnabled) const = 0;
 
 	/**
 	* Compute the Geometry of all the children and add populate the ArrangedChildren list with their values.
@@ -755,4 +757,13 @@ private:
 
 	/** Is this widget hovered? */
 	bool bIsHovered;
+
+	// @see UseLegacyHittest()
+	static int32 bUseLegacyHittest;
+	static FAutoConsoleVariableRef CVarUseLegacyHittest;
+
+public:
+	/** For perf testing. Old hittest schema does not support widgets outside their parents' bounds. */
+	static bool UseLegacyHittest();
+
 };
diff --git a/Engine/Source/Runtime/SlateCore/Public/Widgets/SWindow.h b/Engine/Source/Runtime/SlateCore/Public/Widgets/SWindow.h
--- a/Engine/Source/Runtime/SlateCore/Public/Widgets/SWindow.h
+++ b/Engine/Source/Runtime/SlateCore/Public/Widgets/SWindow.h
@@ -625,6 +625,9 @@ protected:
 	/** Get the color used to tint the window outline */
 	FSlateColor GetWindowOutlineColor() const;
 
+	/** Windows that are not hittestable should not show up in the hittest grid. */
+	EVisibility SWindow::GetWindowVisibility() const;
+
 protected:
 
 	/** Title of the window, displayed in the title bar as well as potentially in the task bar (Windows platform) */
diff --git a/Engine/Source/Runtime/UMG/Private/Slate/SConstraintCanvas.cpp b/Engine/Source/Runtime/UMG/Private/Slate/SConstraintCanvas.cpp
--- a/Engine/Source/Runtime/UMG/Private/Slate/SConstraintCanvas.cpp
+++ b/Engine/Source/Runtime/UMG/Private/Slate/SConstraintCanvas.cpp
@@ -137,7 +137,7 @@ void SConstraintCanvas::OnArrangeChildren( const FGeometry& AllottedGeometry, FA
 }
 
 
-int32 SConstraintCanvas::OnPaint( const FGeometry& AllottedGeometry, const FSlateRect& MyClippingRect, FSlateWindowElementList& OutDrawElements, int32 LayerId, const FWidgetStyle& InWidgetStyle, bool bParentEnabled ) const
+int32 SConstraintCanvas::OnPaint( const FPaintArgs& Args, const FGeometry& AllottedGeometry, const FSlateRect& MyClippingRect, FSlateWindowElementList& OutDrawElements, int32 LayerId, const FWidgetStyle& InWidgetStyle, bool bParentEnabled ) const
 {
 	CachedGeometry = AllottedGeometry;
 
@@ -152,7 +152,7 @@ int32 SConstraintCanvas::OnPaint( const FGeometry& AllottedGeometry, const FSlat
 	{
 		FArrangedWidget& CurWidget = ArrangedChildren(ChildIndex);
 		FSlateRect ChildClipRect = MyClippingRect.IntersectionWith(CurWidget.Geometry.GetClippingRect());
-		const int32 CurWidgetsMaxLayerId = CurWidget.Widget->Paint(CurWidget.Geometry, ChildClipRect, OutDrawElements, MaxLayerId + 1, InWidgetStyle, ShouldBeEnabled(bParentEnabled));
+		const int32 CurWidgetsMaxLayerId = CurWidget.Widget->Paint( Args.WithNewParent(this), CurWidget.Geometry, ChildClipRect, OutDrawElements, MaxLayerId + 1, InWidgetStyle, ShouldBeEnabled(bParentEnabled));
 
 		MaxLayerId = FMath::Max(MaxLayerId, CurWidgetsMaxLayerId);
 	}
diff --git a/Engine/Source/Runtime/UMG/Private/Slate/SObjectWidget.cpp b/Engine/Source/Runtime/UMG/Private/Slate/SObjectWidget.cpp
--- a/Engine/Source/Runtime/UMG/Private/Slate/SObjectWidget.cpp
+++ b/Engine/Source/Runtime/UMG/Private/Slate/SObjectWidget.cpp
@@ -40,9 +40,9 @@ void SObjectWidget::Tick(const FGeometry& AllottedGeometry, const double InCurre
 //	}
 //}
 
-int32 SObjectWidget::OnPaint(const FGeometry& AllottedGeometry, const FSlateRect& MyClippingRect, FSlateWindowElementList& OutDrawElements, int32 LayerId, const FWidgetStyle& InWidgetStyle, bool bParentEnabled) const
+int32 SObjectWidget::OnPaint(const FPaintArgs& Args, const FGeometry& AllottedGeometry, const FSlateRect& MyClippingRect, FSlateWindowElementList& OutDrawElements, int32 LayerId, const FWidgetStyle& InWidgetStyle, bool bParentEnabled) const
 {
-	int32 MaxLayer = SCompoundWidget::OnPaint(AllottedGeometry, MyClippingRect, OutDrawElements, LayerId, InWidgetStyle, bParentEnabled);
+	int32 MaxLayer = SCompoundWidget::OnPaint(Args, AllottedGeometry, MyClippingRect, OutDrawElements, LayerId, InWidgetStyle, bParentEnabled);
 
 	if ( WidgetObject && !WidgetObject->IsDesignTime() )
 	{
