Fixed redundant qualifier: void SWindow::GetWindowVisibility().

diff --git a/Engine/Source/Runtime/SlateCore/Public/Widgets/SWindow.h b/Engine/Source/Runtime/SlateCore/Public/Widgets/SWindow.h
--- a/Engine/Source/Runtime/SlateCore/Public/Widgets/SWindow.h
+++ b/Engine/Source/Runtime/SlateCore/Public/Widgets/SWindow.h
@@ -626,7 +626,7 @@ protected:
 	FSlateColor GetWindowOutlineColor() const;
 
 	/** Windows that are not hittestable should not show up in the hittest grid. */
-	EVisibility SWindow::GetWindowVisibility() const;
+	EVisibility GetWindowVisibility() const;
 
 protected:
 
