Slate Hittest2.0 : now storing TWeakPtr<SWidget> so that widgets do not stick around for longer than they should.
There is now also a possibility of hittesting a ghost of a widget that is now removed.
We simply pretend we did not hit anything at the moment.

diff --git a/Engine/Source/Runtime/SlateCore/Private/Input/HittestGrid.cpp b/Engine/Source/Runtime/SlateCore/Private/Input/HittestGrid.cpp
--- a/Engine/Source/Runtime/SlateCore/Private/Input/HittestGrid.cpp
+++ b/Engine/Source/Runtime/SlateCore/Private/Input/HittestGrid.cpp
@@ -10,7 +10,8 @@ static const FVector2D CellSize(128.0f, 128.0f);
 struct FHittestGrid::FCachedWidget
 {
 	FCachedWidget( int32 InParentIndex, const FArrangedWidget& InWidget, const FSlateRect& InClippingRect )
-	: ArrangedWidget( InWidget )
+	: WidgetPtr(InWidget.Widget)
+	, CachedGeometry( InWidget.Geometry )
 	, ClippingRect( InClippingRect )
 	, Children()
 	, ParentIndex( InParentIndex )
@@ -22,7 +23,8 @@ struct FHittestGrid::FCachedWidget
 		Children.Add( ChildIndex );
 	}
 
-	FArrangedWidget ArrangedWidget;
+	TWeakPtr<SWidget> WidgetPtr;
+	FGeometry CachedGeometry;
 	// @todo umg : ideally this clipping rect is optional and we only have them on a small number of widgets.
 	FSlateRect ClippingRect;
 	TArray<int32> Children;
@@ -57,7 +59,7 @@ TArray<FArrangedWidget> FHittestGrid::GetBubblePath( FVector2D DesktopSpaceCoord
 		for ( int32 i = IndexesInCell.Num()-1; i>=0 && HitWidgetIndex==INDEX_NONE; --i )
 		{
 			const FCachedWidget& TestCandidate = (*WidgetsCachedThisFrame)[IndexesInCell[i]];
-			if ( TestCandidate.ArrangedWidget.Geometry.IsUnderLocation( DesktopSpaceCoordinate ) && TestCandidate.ClippingRect.ContainsPoint( DesktopSpaceCoordinate ) )
+			if ( TestCandidate.CachedGeometry.IsUnderLocation( DesktopSpaceCoordinate ) && TestCandidate.ClippingRect.ContainsPoint( DesktopSpaceCoordinate ) && TestCandidate.WidgetPtr.IsValid() )
 			{
 				HitWidgetIndex = IndexesInCell[i];
 			}
@@ -67,13 +69,29 @@ TArray<FArrangedWidget> FHittestGrid::GetBubblePath( FVector2D DesktopSpaceCoord
 		{
 			TArray<FArrangedWidget> BubblePath;
 			int32 CurWidgetIndex=HitWidgetIndex;
+			bool bPathUninterrupted = false;
 			do
 			{
 				const FCachedWidget& CurCachedWidget = (*WidgetsCachedThisFrame)[CurWidgetIndex];
-				BubblePath.Insert( CurCachedWidget.ArrangedWidget, 0 );
-				CurWidgetIndex = CurCachedWidget.ParentIndex;
+				const TSharedPtr<SWidget> CachedWidgetPtr = CurCachedWidget.WidgetPtr.Pin();
+				
+				bPathUninterrupted = CachedWidgetPtr.IsValid();
+				if (bPathUninterrupted)
+				{
+					BubblePath.Insert(FArrangedWidget(CachedWidgetPtr.ToSharedRef(), CurCachedWidget.CachedGeometry), 0);
+					CurWidgetIndex = CurCachedWidget.ParentIndex;
+				}
 			}
-			while( CurWidgetIndex != INDEX_NONE );
+			while (CurWidgetIndex != INDEX_NONE && bPathUninterrupted);
+			
+			if (!bPathUninterrupted)
+			{
+				// A widget in the path to the root has been removed, so anything
+				// we thought we had hittest is no longer actually there.
+				// Pretend we didn't hit anything.
+				BubblePath = TArray<FArrangedWidget>();
+			}
+
 			return BubblePath;
 		}
 		else
@@ -199,7 +217,9 @@ void FHittestGrid::LogChildren(int32 Index, int32 IndentLevel, const TArray<FCac
 	}
 
 	const FCachedWidget& CachedWidget = WidgetsCachedThisFrame[Index];
-	UE_LOG( LogHittestDebug, Warning, TEXT("%s[%d] => %s"), *IndentString, Index, *CachedWidget.ArrangedWidget.ToString() );
+	const TSharedPtr<SWidget> CachedWidgetPtr = CachedWidget.WidgetPtr.Pin();
+	const FString WidgetString = CachedWidgetPtr.IsValid() ? CachedWidgetPtr->ToString() : TEXT("(null)");
+	UE_LOG( LogHittestDebug, Warning, TEXT("%s[%d] => %s @ %s"), *IndentString, Index, *WidgetString , *CachedWidget.CachedGeometry.ToString() );
 
 	for ( int i=0; i<CachedWidget.Children.Num(); ++i )
 	{
