Added back deprecated methods for OnPaint() and Paint(); these don't fix the resulting compile errors, but should help understand how to fix the compile errors.

diff --git a/Engine/Source/Runtime/SlateCore/Public/Widgets/SWidget.h b/Engine/Source/Runtime/SlateCore/Public/Widgets/SWidget.h
--- a/Engine/Source/Runtime/SlateCore/Public/Widgets/SWidget.h
+++ b/Engine/Source/Runtime/SlateCore/Public/Widgets/SWidget.h
@@ -109,6 +109,9 @@ public:
 	 */
 	int32 Paint( const FPaintArgs& Args, const FGeometry& AllottedGeometry, const FSlateRect& MyClippingRect, FSlateWindowElementList& OutDrawElements, int32 LayerId, const FWidgetStyle& InWidgetStyle, bool bParentEnabled ) const;
 
+	DEPRECATED(4.4, "Paint() now requires an extra FPaintArgs parameter. When calling paint on a child widget, use Args.WithNewParent(this).")
+	int32 Paint(const FGeometry& AllottedGeometry, const FSlateRect& MyClippingRect, FSlateWindowElementList& OutDrawElements, int32 LayerId, const FWidgetStyle& InWidgetStyle, bool bParentEnabled) const { return 0; }
+
 	/**
 	 * Ticks this widget.  Override in derived classes, but always call the parent implementation.
 	 *
@@ -709,6 +712,9 @@ private:
 	*/
 	virtual int32 OnPaint(const FPaintArgs& Args, const FGeometry& AllottedGeometry, const FSlateRect& MyClippingRect, FSlateWindowElementList& OutDrawElements, int32 LayerId, const FWidgetStyle& InWidgetStyle, bool bParentEnabled) const = 0;
 
+	DEPRECATED(4.4, "OnPaint() now requires an extra FPaintArgs parameter. When calling Superclass::OnPaint() you can simply pass through the existing Args.")
+	virtual int32 OnPaint(const FGeometry& AllottedGeometry, const FSlateRect& MyClippingRect, FSlateWindowElementList& OutDrawElements, int32 LayerId, const FWidgetStyle& InWidgetStyle, bool bParentEnabled) const { return 0; }
+
 	/**
 	* Compute the Geometry of all the children and add populate the ArrangedChildren list with their values.
 	* Each type of Layout panel should arrange children based on desired behavior.
