Slate slots now use a common base class, which has a private '.Widget' member; use GetWidget() instead.
Slate TPanelChildren, TSlotlessChildren are no longer full-fledged arrays; removed any unused Array API.
Fixed up a lot of improper use cases.

This better encapsulates the action of adding/removing Slate Widgets from the runtime tree as a Slate Core concern.
Alows us to add parent pointers in the future with significantly less effort.

diff --git a/Engine/Plugins/Editor/PluginsEditor/Source/PluginsEditor/Private/SPluginCategories.cpp b/Engine/Plugins/Editor/PluginsEditor/Source/PluginsEditor/Private/SPluginCategories.cpp
--- a/Engine/Plugins/Editor/PluginsEditor/Source/PluginsEditor/Private/SPluginCategories.cpp
+++ b/Engine/Plugins/Editor/PluginsEditor/Source/PluginsEditor/Private/SPluginCategories.cpp
@@ -45,7 +45,7 @@ void SPluginCategories::Construct( const FArguments& Args, const TSharedRef< SPl
 		PluginCategoryTreeView->SetSelection( RootPluginCategories[ 0 ] );
 	}
 
-	ChildSlot.Widget = PluginCategoryTreeView.ToSharedRef();
+	ChildSlot.AttachWidget( PluginCategoryTreeView.ToSharedRef() );
 }
 
 
diff --git a/Engine/Plugins/Editor/PluginsEditor/Source/PluginsEditor/Private/SPluginCategoryTreeItem.cpp b/Engine/Plugins/Editor/PluginsEditor/Source/PluginsEditor/Private/SPluginCategoryTreeItem.cpp
--- a/Engine/Plugins/Editor/PluginsEditor/Source/PluginsEditor/Private/SPluginCategoryTreeItem.cpp
+++ b/Engine/Plugins/Editor/PluginsEditor/Source/PluginsEditor/Private/SPluginCategoryTreeItem.cpp
@@ -23,51 +23,53 @@ void SPluginCategoryTreeItem::Construct( const FArguments& Args, const TSharedRe
 	// Figure out which font size to use
 	const auto bIsRootItem = !Item->GetParentCategory().IsValid();
 
-	ChildSlot.Widget = 
-	SNew( SBorder )
-	.BorderImage( FPluginStyle::Get()->GetBrush(bIsRootItem ? "CategoryTreeItem.Root.BackgroundBrush" : "CategoryTreeItem.BackgroundBrush") )
-	.Padding(FPluginStyle::Get()->GetMargin(bIsRootItem ? "CategoryTreeItem.Root.BackgroundPadding" : "CategoryTreeItem.BackgroundPadding") )
-	[
-		SNew( SHorizontalBox )
-
-		// Icon image
-		+SHorizontalBox::Slot()
-		.Padding( PaddingAmount )
-		.AutoWidth()
-		.VAlign(VAlign_Center)
+	ChildSlot
+	[ 
+		SNew( SBorder )
+		.BorderImage( FPluginStyle::Get()->GetBrush(bIsRootItem ? "CategoryTreeItem.Root.BackgroundBrush" : "CategoryTreeItem.BackgroundBrush") )
+		.Padding(FPluginStyle::Get()->GetMargin(bIsRootItem ? "CategoryTreeItem.Root.BackgroundPadding" : "CategoryTreeItem.BackgroundPadding") )
 		[
-			SNew( SBox )
-			.WidthOverride( CategoryIconSize )
-			.HeightOverride( CategoryIconSize )
+			SNew( SHorizontalBox )
+
+			// Icon image
+			+SHorizontalBox::Slot()
+			.Padding( PaddingAmount )
+			.AutoWidth()
+			.VAlign(VAlign_Center)
 			[
-				SNew( SImage )
-				.Image( this, &SPluginCategoryTreeItem::GetIconBrush )
+				SNew( SBox )
+				.WidthOverride( CategoryIconSize )
+				.HeightOverride( CategoryIconSize )
+				[
+					SNew( SImage )
+					.Image( this, &SPluginCategoryTreeItem::GetIconBrush )
+				]
 			]
-		]
 
-		// Category name
-		+SHorizontalBox::Slot()
-		.Padding( PaddingAmount )
-		.VAlign(VAlign_Center)
-		[
-			SNew( STextBlock )
-			.Text( Item->GetCategoryDisplayName() )
-			.TextStyle( FPluginStyle::Get(), bIsRootItem ? "CategoryTreeItem.Root.Text" : "CategoryTreeItem.Text" )
-		]
+			// Category name
+			+SHorizontalBox::Slot()
+			.Padding( PaddingAmount )
+			.VAlign(VAlign_Center)
+			[
+				SNew( STextBlock )
+				.Text( Item->GetCategoryDisplayName() )
+				.TextStyle( FPluginStyle::Get(), bIsRootItem ? "CategoryTreeItem.Root.Text" : "CategoryTreeItem.Text" )
+			]
 			
-		// Plugin count
-		+SHorizontalBox::Slot()
-		.AutoWidth()
-		.Padding( PaddingAmount )
-		.VAlign(VAlign_Center)
-		[
-			SNew( STextBlock )
+			// Plugin count
+			+SHorizontalBox::Slot()
+			.AutoWidth()
+			.Padding( PaddingAmount )
+			.VAlign(VAlign_Center)
+			[
+				SNew( STextBlock )
 
-			// Only display if at there is least one plugin is in this category
-			.Visibility( Item->GetPlugins().Num() > 0 ? EVisibility::Visible : EVisibility::Collapsed )
+				// Only display if at there is least one plugin is in this category
+				.Visibility( Item->GetPlugins().Num() > 0 ? EVisibility::Visible : EVisibility::Collapsed )
 
-			.Text( FText::Format( LOCTEXT( "NumberOfPluginsWrapper", "({0})" ), FText::AsNumber( Item->GetPlugins().Num() ) ) )
-			.TextStyle( FPluginStyle::Get(), bIsRootItem ? "CategoryTreeItem.Root.PluginCountText" : "CategoryTreeItem.PluginCountText" )
+				.Text( FText::Format( LOCTEXT( "NumberOfPluginsWrapper", "({0})" ), FText::AsNumber( Item->GetPlugins().Num() ) ) )
+				.TextStyle( FPluginStyle::Get(), bIsRootItem ? "CategoryTreeItem.Root.PluginCountText" : "CategoryTreeItem.PluginCountText" )
+			]
 		]
 	];
 }
diff --git a/Engine/Plugins/Editor/PluginsEditor/Source/PluginsEditor/Private/SPluginList.cpp b/Engine/Plugins/Editor/PluginsEditor/Source/PluginsEditor/Private/SPluginList.cpp
--- a/Engine/Plugins/Editor/PluginsEditor/Source/PluginsEditor/Private/SPluginList.cpp
+++ b/Engine/Plugins/Editor/PluginsEditor/Source/PluginsEditor/Private/SPluginList.cpp
@@ -33,7 +33,7 @@ void SPluginList::Construct( const FArguments& Args, const TSharedRef< SPluginsE
 		.OnGenerateRow( this, &SPluginList::PluginListView_OnGenerateRow )
 		;
 
-	ChildSlot.Widget = PluginListView.ToSharedRef();
+	ChildSlot.AttachWidget( PluginListView.ToSharedRef() );
 }
 
 
diff --git a/Engine/Plugins/Editor/PluginsEditor/Source/PluginsEditor/Private/SPluginListTile.cpp b/Engine/Plugins/Editor/PluginsEditor/Source/PluginsEditor/Private/SPluginListTile.cpp
--- a/Engine/Plugins/Editor/PluginsEditor/Source/PluginsEditor/Private/SPluginListTile.cpp
+++ b/Engine/Plugins/Editor/PluginsEditor/Source/PluginsEditor/Private/SPluginListTile.cpp
@@ -75,7 +75,8 @@ void SPluginListTile::Construct( const FArguments& Args, const TSharedRef< SPlug
 	}
 
 
-	ChildSlot.Widget = 
+	ChildSlot
+	[
 		SNew( SBorder )
 		.BorderImage( FEditorStyle::GetBrush( "NoBorder" ) )
 		.Padding( PaddingAmount )
@@ -248,7 +249,7 @@ void SPluginListTile::Construct( const FArguments& Args, const TSharedRef< SPlug
 				]
 			]
 		]
-	;
+	];
 }
 
 
diff --git a/Engine/Plugins/Editor/PluginsEditor/Source/PluginsEditor/Private/SPluginsEditor.cpp b/Engine/Plugins/Editor/PluginsEditor/Source/PluginsEditor/Private/SPluginsEditor.cpp
--- a/Engine/Plugins/Editor/PluginsEditor/Source/PluginsEditor/Private/SPluginsEditor.cpp
+++ b/Engine/Plugins/Editor/PluginsEditor/Source/PluginsEditor/Private/SPluginsEditor.cpp
@@ -36,7 +36,8 @@ void SPluginsEditor::Construct( const FArguments& Args )
 	PluginCategories = SNew( SPluginCategories, SharedThis( this ) );
 
 
-	ChildSlot.Widget = 
+	ChildSlot
+	[ 
 		SNew( SHorizontalBox )
 
 		+SHorizontalBox::Slot()
@@ -118,7 +119,8 @@ void SPluginsEditor::Construct( const FArguments& Args )
 					]
 				]
 			]
-		];
+		]
+	];
 }
 
 
diff --git a/Engine/Source/Developer/OutputLog/Private/SDebugConsole.cpp b/Engine/Source/Developer/OutputLog/Private/SDebugConsole.cpp
--- a/Engine/Source/Developer/OutputLog/Private/SDebugConsole.cpp
+++ b/Engine/Source/Developer/OutputLog/Private/SDebugConsole.cpp
@@ -17,61 +17,63 @@ void SDebugConsole::Construct( const FArguments& InArgs, const EDebugConsoleStyl
 	CurrentStyle = InStyle;
 
 	check( OutputLogModule != NULL );
-	ChildSlot.Widget = 
+	ChildSlot
+	[ 
 		SNew( SVerticalBox )
-			+SVerticalBox::Slot()
-			.AutoHeight()
-			[
-				SNew( SVerticalBox )
-					.Visibility( this, &SDebugConsole::MakeVisibleIfLogIsShown )
+		+SVerticalBox::Slot()
+		.AutoHeight()
+		[
+			SNew( SVerticalBox )
+				.Visibility( this, &SDebugConsole::MakeVisibleIfLogIsShown )
 					
-					+SVerticalBox::Slot()
-					.AutoHeight()
-					.Padding( 10.0f )
+				+SVerticalBox::Slot()
+				.AutoHeight()
+				.Padding( 10.0f )
+				[
+					SNew(SBox)
+					.HeightOverride( 200.0f )
 					[
-						SNew(SBox)
-						.HeightOverride( 200.0f )
-						[
-							SNew( SBorder )
-								.BorderImage( FEditorStyle::GetBrush( "DebugConsole.Background" ) )
-								.ColorAndOpacity( this, &SDebugConsole::GetAnimatedColorAndOpacity )
-								.BorderBackgroundColor( this, &SDebugConsole::GetAnimatedSlateColor )
-								[
-									SNew( SSpacer )
-								]
-						]
+						SNew( SBorder )
+							.BorderImage( FEditorStyle::GetBrush( "DebugConsole.Background" ) )
+							.ColorAndOpacity( this, &SDebugConsole::GetAnimatedColorAndOpacity )
+							.BorderBackgroundColor( this, &SDebugConsole::GetAnimatedSlateColor )
+							[
+								SNew( SSpacer )
+							]
 					]
-			]
-
-			+SVerticalBox::Slot()
-			.AutoHeight()
-			.Padding( 10.0f )
-			[
-				SNew(SBox)
-				.HeightOverride( 26.0f )
-				[
-					SNew( SBorder )
-						.BorderImage( FEditorStyle::GetBrush( "Console.Background" ) )
-						.ColorAndOpacity( this, &SDebugConsole::GetAnimatedColorAndOpacity )
-						.BorderBackgroundColor( this, &SDebugConsole::GetAnimatedSlateColor )
-						[
-							SNew( SHorizontalBox )
-								+SHorizontalBox::Slot()
-									.AutoWidth()
-									.Padding( 3.0f )
-								[
-									SNew( STextBlock )
-										.Text( NSLOCTEXT("Console", "ConsoleLabel", "Console >").ToString() )
-								]
-								+SHorizontalBox::Slot()
-									.Padding( 3.0f )
-									.FillWidth( 1.0f )
-								[
-									OutputLogModule->MakeConsoleInputBox( EditableTextBox )
-								]
-						]
 				]
-			];
+		]
+
+		+SVerticalBox::Slot()
+		.AutoHeight()
+		.Padding( 10.0f )
+		[
+			SNew(SBox)
+			.HeightOverride( 26.0f )
+			[
+				SNew( SBorder )
+					.BorderImage( FEditorStyle::GetBrush( "Console.Background" ) )
+					.ColorAndOpacity( this, &SDebugConsole::GetAnimatedColorAndOpacity )
+					.BorderBackgroundColor( this, &SDebugConsole::GetAnimatedSlateColor )
+					[
+						SNew( SHorizontalBox )
+							+SHorizontalBox::Slot()
+								.AutoWidth()
+								.Padding( 3.0f )
+							[
+								SNew( STextBlock )
+									.Text( NSLOCTEXT("Console", "ConsoleLabel", "Console >").ToString() )
+							]
+							+SHorizontalBox::Slot()
+								.Padding( 3.0f )
+								.FillWidth( 1.0f )
+							[
+								OutputLogModule->MakeConsoleInputBox( EditableTextBox )
+							]
+					]
+			]
+		]
+	];
 
 	// Kick off intro animation
 	AnimCurve = FCurveSequence();
diff --git a/Engine/Source/Developer/SettingsEditor/Private/Widgets/SSettingsEditor.cpp b/Engine/Source/Developer/SettingsEditor/Private/Widgets/SSettingsEditor.cpp
--- a/Engine/Source/Developer/SettingsEditor/Private/Widgets/SSettingsEditor.cpp
+++ b/Engine/Source/Developer/SettingsEditor/Private/Widgets/SSettingsEditor.cpp
@@ -691,11 +691,11 @@ void SSettingsEditor::HandleModelSelectionChanged( )
 		// show settings widget
 		if (CustomWidget.IsValid())
 		{
-			CustomWidgetSlot->Widget = CustomWidget.ToSharedRef();
+			CustomWidgetSlot->AttachWidget( CustomWidget.ToSharedRef() );
 		}
 		else
 		{
-			CustomWidgetSlot->Widget = SNullWidget::NullWidget;
+			CustomWidgetSlot->AttachWidget( SNullWidget::NullWidget );
 		}
 
 		SettingsView->SetObject(SelectedSection->GetSettingsObject().Get());
@@ -721,7 +721,7 @@ void SSettingsEditor::HandleModelSelectionChanged( )
 	}
 	else
 	{
-		CustomWidgetSlot->Widget = SNullWidget::NullWidget;
+		CustomWidgetSlot->AttachWidget( SNullWidget::NullWidget );
 		SettingsView->SetObject(nullptr);
 	}
 
diff --git a/Engine/Source/Editor/GraphEditor/Private/SNodePanel.cpp b/Engine/Source/Editor/GraphEditor/Private/SNodePanel.cpp
--- a/Engine/Source/Editor/GraphEditor/Private/SNodePanel.cpp
+++ b/Engine/Source/Editor/GraphEditor/Private/SNodePanel.cpp
@@ -198,6 +198,12 @@ namespace NodePanelDefs
 	static const float MouseZoomScaling = 0.05f;
 };
 
+SNodePanel::SNodePanel()
+: Children()
+, VisibleChildren()
+{
+}
+
 void SNodePanel::OnArrangeChildren(const FGeometry& AllottedGeometry, FArrangedChildren& ArrangedChildren) const
 {
 	ArrangeChildNodes(AllottedGeometry, ArrangedChildren);
diff --git a/Engine/Source/Editor/GraphEditor/Public/SNodePanel.h b/Engine/Source/Editor/GraphEditor/Public/SNodePanel.h
--- a/Engine/Source/Editor/GraphEditor/Public/SNodePanel.h
+++ b/Engine/Source/Editor/GraphEditor/Public/SNodePanel.h
@@ -368,6 +368,8 @@ public:
 		}
 	};
 
+	SNodePanel();
+
 	// SPanel interface
 	virtual void OnArrangeChildren(const FGeometry& AllottedGeometry, FArrangedChildren& ArrangedChildren) const override;
 	virtual FVector2D ComputeDesiredSize() const override;
diff --git a/Engine/Source/Editor/Kismet/Private/SBlueprintPalette.cpp b/Engine/Source/Editor/Kismet/Private/SBlueprintPalette.cpp
--- a/Engine/Source/Editor/Kismet/Private/SBlueprintPalette.cpp
+++ b/Engine/Source/Editor/Kismet/Private/SBlueprintPalette.cpp
@@ -1551,11 +1551,11 @@ void SBlueprintPalette::OnSplitterResized() const
 	{
 		SSplitter::FSlot const& SplitterSlot = PaletteSplitter->SlotAt(SlotIndex);
 
-		if (SplitterSlot.Widget == FavoritesWrapper)
+		if (SplitterSlot.GetWidget() == FavoritesWrapper)
 		{
 			GConfig->SetFloat(*BlueprintPalette::ConfigSection, *BlueprintPalette::FavoritesHeightConfigKey, SplitterSlot.SizeValue.Get(), GEditorUserSettingsIni);
 		}
-		else if (SplitterSlot.Widget == LibraryWrapper)
+		else if (SplitterSlot.GetWidget() == LibraryWrapper)
 		{
 			GConfig->SetFloat(*BlueprintPalette::ConfigSection, *BlueprintPalette::LibraryHeightConfigKey, SplitterSlot.SizeValue.Get(), GEditorUserSettingsIni);
 		}
diff --git a/Engine/Source/Editor/LandscapeEditor/Private/LandscapeEditorDetailWidgets.cpp b/Engine/Source/Editor/LandscapeEditor/Private/LandscapeEditorDetailWidgets.cpp
--- a/Engine/Source/Editor/LandscapeEditor/Private/LandscapeEditorDetailWidgets.cpp
+++ b/Engine/Source/Editor/LandscapeEditor/Private/LandscapeEditorDetailWidgets.cpp
@@ -191,7 +191,8 @@ void SToolSelector::BuildMultiBlockWidget(const ISlateStyle* StyleSet, const FNa
 	EMultiBlockLocation::Type BlockLocation = GetMultiBlockLocation();
 	FName BlockStyle = EMultiBlockLocation::ToName(FEditorStyle::Join( StyleName, ".Button" ), BlockLocation);
 
-	ChildSlot.Widget =
+	ChildSlot
+	[
 		SNew(SComboButton)
 			.ContentPadding(0)
 			.ButtonStyle( FEditorStyle::Get(), BlockStyle )
@@ -202,7 +203,7 @@ void SToolSelector::BuildMultiBlockWidget(const ISlateStyle* StyleSet, const FNa
 				ButtonContent
 			]
 			.OnGetMenuContent(this, &SToolSelector::OnGetMenuContent)
-		;
+	];
 
 	ChildSlot.Padding(FEditorStyle::GetMargin(FEditorStyle::Join(StyleName, ".SToolBarComboButtonBlock.Padding")));
 
diff --git a/Engine/Source/Editor/LevelEditor/Private/SLevelViewport.cpp b/Engine/Source/Editor/LevelEditor/Private/SLevelViewport.cpp
--- a/Engine/Source/Editor/LevelEditor/Private/SLevelViewport.cpp
+++ b/Engine/Source/Editor/LevelEditor/Private/SLevelViewport.cpp
@@ -2424,7 +2424,8 @@ void SActorPreview::Construct( const FArguments& InArgs )
 	// widgets that are added to the viewport overlay.
 	this->SetVisibility(EVisibility::SelfHitTestInvisible);
 
-	this->ChildSlot.Widget =
+	this->ChildSlot
+	[
 		SNew(SBorder)
 		.Padding(0)
 
@@ -2502,7 +2503,8 @@ void SActorPreview::Construct( const FArguments& InArgs )
 					// Pass along the block's tool-tip string
 					.ToolTipText( this, &SActorPreview::GetPinButtonToolTipText )
 			]
-		];
+		]
+	];
 
 	// Setup animation curve for fading in and out.  Note that we add a bit of lead-in time on the fade-in
 	// to avoid hysteresis as the user moves the mouse over the view
@@ -2891,8 +2893,11 @@ void SLevelViewport::PreviewActors( const TArray< AActor* >& ActorsToPreview )
 
 			// Add our new widget to our viewport's overlay
 			// @todo camerapip: Consider using a canvas instead of an overlay widget -- our viewports get SQUASHED when the view shrinks!
-			auto& HorizontalBoxSlot = ActorPreviewHorizontalBox->AddSlot().AutoWidth();
-			HorizontalBoxSlot.Widget = ActorPreviewWidget.ToSharedRef();
+			ActorPreviewHorizontalBox->AddSlot()
+			.AutoWidth()
+			[
+				ActorPreviewWidget.ToSharedRef()
+			];
 		}
 
 		// OK, at least one new preview viewport was added, so update settings for all views immediately.
diff --git a/Engine/Source/Editor/Persona/Private/STrack.cpp b/Engine/Source/Editor/Persona/Private/STrack.cpp
--- a/Engine/Source/Editor/Persona/Private/STrack.cpp
+++ b/Engine/Source/Editor/Persona/Private/STrack.cpp
@@ -320,6 +320,11 @@ float STrackNode::GetDataStartPos() const
 //////////////////////////////////////////////////////////////////////////
 // STrack
 
+STrack::STrack()
+:TrackNodes()
+{
+}
+
 void STrack::Construct( const FArguments& InArgs )
 {
 	TrackColor = InArgs._TrackColor;
@@ -348,9 +353,9 @@ void STrack::Construct( const FArguments& InArgs )
 
 void STrack::OnArrangeChildren( const FGeometry& AllottedGeometry, FArrangedChildren& ArrangedChildren ) const
 {
-	for ( auto TrackIter = TrackNodes.CreateConstIterator() ; TrackIter ; ++TrackIter )
+	for ( int32 TrackIndex=0; TrackIndex < TrackNodes.Num(); ++TrackIndex )
 	{
-		TSharedRef<STrackNode> TrackNode = (*TrackIter);
+		TSharedRef<STrackNode> TrackNode = (TrackNodes[TrackIndex]);
 		if (TrackNode->bBeingDragged)
 		{
 			continue;
diff --git a/Engine/Source/Editor/Persona/Private/STrack.h b/Engine/Source/Editor/Persona/Private/STrack.h
--- a/Engine/Source/Editor/Persona/Private/STrack.h
+++ b/Engine/Source/Editor/Persona/Private/STrack.h
@@ -233,6 +233,8 @@ public:
 		SLATE_ARGUMENT( const FSlateBrush*, StyleInfo )
 	SLATE_END_ARGS()
 
+	STrack();
+
 	void			Construct( const FArguments& InArgs );
 
 	virtual void	OnArrangeChildren( const FGeometry& AllottedGeometry, FArrangedChildren& ArrangedChildren ) const;
diff --git a/Engine/Source/Editor/PropertyEditor/Private/SDetailSingleItemRow.cpp b/Engine/Source/Editor/PropertyEditor/Private/SDetailSingleItemRow.cpp
--- a/Engine/Source/Editor/PropertyEditor/Private/SDetailSingleItemRow.cpp
+++ b/Engine/Source/Editor/PropertyEditor/Private/SDetailSingleItemRow.cpp
@@ -46,7 +46,7 @@ public:
 		}
 		else
 		{
-			FVector2D ChildSize = ChildSlot.Widget->GetDesiredSize();
+			FVector2D ChildSize = ChildSlot.GetWidget()->GetDesiredSize();
 
 			float XVal = FMath::Max( MinWidthVal, ChildSize.X );
 			if( MaxWidthVal >= MinWidthVal )
diff --git a/Engine/Source/Editor/PropertyEditor/Private/UserInterface/PropertyTable/SPropertyTable.h b/Engine/Source/Editor/PropertyEditor/Private/UserInterface/PropertyTable/SPropertyTable.h
--- a/Engine/Source/Editor/PropertyEditor/Private/UserInterface/PropertyTable/SPropertyTable.h
+++ b/Engine/Source/Editor/PropertyEditor/Private/UserInterface/PropertyTable/SPropertyTable.h
@@ -60,7 +60,7 @@ public:
 
 		if ( Table->GetIsUserAllowedToChangeRoot() )
 		{
-			TSharedRef< SWidget > TreeContent = ChildSlot.Widget;
+			TSharedRef< SWidget > TreeContent = ChildSlot.GetWidget();
 
 			ChildSlot
 			[
diff --git a/Engine/Source/Editor/Sequencer/Private/SSequencerTrackViewPanel.h b/Engine/Source/Editor/Sequencer/Private/SSequencerTrackViewPanel.h
--- a/Engine/Source/Editor/Sequencer/Private/SSequencerTrackViewPanel.h
+++ b/Engine/Source/Editor/Sequencer/Private/SSequencerTrackViewPanel.h
@@ -12,6 +12,11 @@ public:
 	SLATE_BEGIN_ARGS( TSequencerTrackViewPanel<WidgetType> ) {}
 	SLATE_END_ARGS()
 
+	TSequencerTrackViewPanel()
+	: Children()
+	{
+	}
+
 	/** Construct this widget.  Called by the SNew() Slate macro. */
 	void Construct( const FArguments& InArgs, TSharedRef<FSequencerDisplayNode> InRootNode )
 	{
diff --git a/Engine/Source/Editor/UMGEditor/Private/Designer/SZoomPan.cpp b/Engine/Source/Editor/UMGEditor/Private/Designer/SZoomPan.cpp
--- a/Engine/Source/Editor/UMGEditor/Private/Designer/SZoomPan.cpp
+++ b/Engine/Source/Editor/UMGEditor/Private/Designer/SZoomPan.cpp
@@ -22,7 +22,7 @@ void SZoomPan::Construct(const FArguments& InArgs)
 
 void SZoomPan::OnArrangeChildren(const FGeometry& AllottedGeometry, FArrangedChildren& ArrangedChildren) const
 {
-	const EVisibility ChildVisibility = ChildSlot.Widget->GetVisibility();
+	const EVisibility ChildVisibility = ChildSlot.GetWidget()->GetVisibility();
 	if ( ArrangedChildren.Accepts(ChildVisibility) )
 	{
 		const FMargin SlotPadding(ChildSlot.SlotPadding.Get());
@@ -30,7 +30,7 @@ void SZoomPan::OnArrangeChildren(const FGeometry& AllottedGeometry, FArrangedChi
 		AlignmentArrangeResult YResult = AlignChild<Orient_Vertical>(AllottedGeometry.Size.Y, ChildSlot, SlotPadding, 1);
 
 		ArrangedChildren.AddWidget( ChildVisibility, AllottedGeometry.MakeChild(
-				ChildSlot.Widget,
+				ChildSlot.GetWidget(),
 				FVector2D(XResult.Offset, YResult.Offset) - ViewOffset.Get(),
 				FVector2D(XResult.Size, YResult.Size),
 				ZoomAmount.Get()
diff --git a/Engine/Source/Editor/UnrealEd/Private/Dialogs/Dialogs.cpp b/Engine/Source/Editor/UnrealEd/Private/Dialogs/Dialogs.cpp
--- a/Engine/Source/Editor/UnrealEd/Private/Dialogs/Dialogs.cpp
+++ b/Engine/Source/Editor/UnrealEd/Private/Dialogs/Dialogs.cpp
@@ -777,27 +777,30 @@ void SGenericDialogWidget::Construct( const FArguments& InArgs )
 	TSharedPtr< SScrollBox > ScrollBox;
 
 	this->ChildSlot
+	[
+		SNew(SVerticalBox)
+		+SVerticalBox::Slot()
+		.AutoHeight()
+		.MaxHeight(300.0f)
 		[
-			SNew(SVerticalBox)
-			+SVerticalBox::Slot()
+			SAssignNew(ScrollBox, SScrollBox)
+		]
+
+		+SVerticalBox::Slot()
+			.HAlign(HAlign_Right)
 			.AutoHeight()
-			.MaxHeight(300.0f)
+			.Padding(0.0f, 2.0f, 0.0f, 0.0f)
 			[
-				SAssignNew(ScrollBox, SScrollBox)
+				SNew(SButton)
+					.Text( NSLOCTEXT("UnrealEd", "OK", "OK") )
+					.OnClicked(this, &SGenericDialogWidget::OnOK_Clicked)
 			]
+	];
 
-			+SVerticalBox::Slot()
-				.HAlign(HAlign_Right)
-				.AutoHeight()
-				.Padding(0.0f, 2.0f, 0.0f, 0.0f)
-				[
-					SNew(SButton)
-						.Text( NSLOCTEXT("UnrealEd", "OK", "OK") )
-						.OnClicked(this, &SGenericDialogWidget::OnOK_Clicked)
-				]
-		];
-
-	ScrollBox->AddSlot().Widget = InArgs._Content.Widget;
+	ScrollBox->AddSlot()
+	[
+		InArgs._Content.Widget
+	];
 }
 
 void SGenericDialogWidget::OpenDialog (const FText& InDialogTitle, const TSharedRef< SWidget >& DisplayContent)
diff --git a/Engine/Source/Editor/UnrealEd/Private/Toolkits/SAssetEditorCommon.cpp b/Engine/Source/Editor/UnrealEd/Private/Toolkits/SAssetEditorCommon.cpp
--- a/Engine/Source/Editor/UnrealEd/Private/Toolkits/SAssetEditorCommon.cpp
+++ b/Engine/Source/Editor/UnrealEd/Private/Toolkits/SAssetEditorCommon.cpp
@@ -32,7 +32,8 @@ void SAssetEditorCommon::Construct( const FArguments& InArgs, TWeakPtr< FAssetEd
 	ParentToolkit->FillDefaultFileMenuCommands( MenuBuilder );
 	ParentToolkit->FillDefaultAssetMenuCommands( MenuBuilder );
 
-	this->ChildSlot.Widget =
+	this->ChildSlot
+	[
 		SNew( SVerticalBox )
 			+SVerticalBox::Slot()
 			.AutoHeight()
@@ -40,6 +41,6 @@ void SAssetEditorCommon::Construct( const FArguments& InArgs, TWeakPtr< FAssetEd
 			[
 				MenuBuilder.MakeWidget()
 			]
-		;
+	];
 
 }
\ No newline at end of file
diff --git a/Engine/Source/Editor/UserFeedback/Private/UserFeedback.cpp b/Engine/Source/Editor/UserFeedback/Private/UserFeedback.cpp
--- a/Engine/Source/Editor/UserFeedback/Private/UserFeedback.cpp
+++ b/Engine/Source/Editor/UserFeedback/Private/UserFeedback.cpp
@@ -71,8 +71,8 @@ public:
 		PopulateContextNames(InArgs._Context);
 
 		const FMargin DefaultPadding(5, 5, 5, 5);
-		ChildSlot.Widget =
-
+		ChildSlot
+		[
 			SNew(SBorder)
 			.BorderImage(FEditorStyle::GetBrush("ToolPanel.GroupBorder"))
 			[
@@ -175,9 +175,9 @@ public:
 					]
 				]
 			]
-			;
+		];
 
-			TextBox->SetError( FText::GetEmpty() );
+		TextBox->SetError( FText::GetEmpty() );
 	}
 
 	/** Called when a key is pressed on the widget */
@@ -484,7 +484,7 @@ public:
 		}
 		else
 		{
-			FVector2D ChildSize = ChildSlot.Widget->GetDesiredSize();
+			FVector2D ChildSize = ChildSlot.GetWidget()->GetDesiredSize();
 			return FVector2D(FMath::Max(MinWidthVal, ChildSize.X), ChildSize.Y );
 		}
 	}
diff --git a/Engine/Source/Runtime/Core/Public/Containers/Array.h b/Engine/Source/Runtime/Core/Public/Containers/Array.h
--- a/Engine/Source/Runtime/Core/Public/Containers/Array.h
+++ b/Engine/Source/Runtime/Core/Public/Containers/Array.h
@@ -2307,6 +2307,10 @@ public:
 		}
 		Array.RemoveAtSwap( Index, Count, AllowShrinking );
 	}
+	void Swap( int32 A, int32 B )
+	{
+		Array.Swap(A,B);
+	}
 	void Empty( int32 Slack=0 )
 	{
 		DestructAndFreeItems();
diff --git a/Engine/Source/Runtime/Slate/Private/Framework/Application/SlateApplication.cpp b/Engine/Source/Runtime/Slate/Private/Framework/Application/SlateApplication.cpp
--- a/Engine/Source/Runtime/Slate/Private/Framework/Application/SlateApplication.cpp
+++ b/Engine/Source/Runtime/Slate/Private/Framework/Application/SlateApplication.cpp
@@ -1,8 +1,5 @@
 // Copyright 1998-2014 Epic Games, Inc. All Rights Reserved.
 
-/*=============================================================================
-	SlateApplication.cpp: Implements the FSlateApplication class.
-=============================================================================*/
 
 #include "SlatePrivatePCH.h"
 #include "SWindowTitleBar.h"
diff --git a/Engine/Source/Runtime/Slate/Private/Framework/Docking/SDockingTabStack.cpp b/Engine/Source/Runtime/Slate/Private/Framework/Docking/SDockingTabStack.cpp
--- a/Engine/Source/Runtime/Slate/Private/Framework/Docking/SDockingTabStack.cpp
+++ b/Engine/Source/Runtime/Slate/Private/Framework/Docking/SDockingTabStack.cpp
@@ -280,7 +280,7 @@ void SDockingTabStack::AddTabWidget( const TSharedRef<SDockTab>& InTab, int32 At
 
 }
 
-const TArray< TSharedRef<SDockTab> >& SDockingTabStack::GetTabs() const
+const TSlotlessChildren<SDockTab>& SDockingTabStack::GetTabs() const
 {
 	return TabWell->GetTabs();
 }
@@ -435,7 +435,7 @@ FReply SDockingTabStack::OnUserAttemptingDock( SDockingNode::RelativeDirection D
 
 TArray< TSharedRef<SDockTab> > SDockingTabStack::GetAllChildTabs() const
 {
-	return GetTabs();
+	return GetTabs().AsArrayCopy();
 }
 
 void SDockingTabStack::CloseForegroundTab()
@@ -455,7 +455,7 @@ void SDockingTabStack::CloseAllButForegroundTab(ETabsToClose TabsToClose)
 		int32 DestroyIndex = 0;
 		while ((TabWell->GetNumTabs() > 1) && (DestroyIndex < TabWell->GetNumTabs()))
 		{
-			auto Tab = TabWell->GetTabs()[DestroyIndex];
+			const TSharedRef<SDockTab>& Tab = TabWell->GetTabs()[DestroyIndex];
 
 			const bool bCanClose = (TabsToClose == CloseDocumentsAndTools) || (Tab->GetTabRole() == ETabRole::DocumentTab);
 
@@ -537,7 +537,7 @@ TSharedRef<SWidget> SDockingTabStack::MakeContextMenu()
 			);
 
 			// If the active tab is a document tab, and there is more than one open in this tab well, offer to close the others
-			auto ForegroundTabPtr = TabWell->GetForegroundTab();
+			TSharedPtr<SDockTab> ForegroundTabPtr = TabWell->GetForegroundTab();
 			if (ForegroundTabPtr.IsValid() && (ForegroundTabPtr->GetTabRole() == ETabRole::DocumentTab) && (TabWell->GetNumTabs() > 1))
 			{
 				const ETabsToClose TabsToClose = CloseDocumentTabs;
@@ -661,7 +661,7 @@ TSharedPtr<FTabManager::FLayoutNode> SDockingTabStack::GatherPersistentLayout()
 	{
 		// Each live tab might want to save custom visual state.
 		{
-			const TArray< TSharedRef<SDockTab> >& MyTabs = this->GetTabs();
+			const TArray< TSharedRef<SDockTab> > MyTabs = this->GetTabs().AsArrayCopy();
 			for (int32 TabIndex=0; TabIndex < MyTabs.Num(); ++TabIndex)
 			{
 				MyTabs[TabIndex]->PersistVisualState();
diff --git a/Engine/Source/Runtime/Slate/Private/Framework/Docking/SDockingTabStack.h b/Engine/Source/Runtime/Slate/Private/Framework/Docking/SDockingTabStack.h
--- a/Engine/Source/Runtime/Slate/Private/Framework/Docking/SDockingTabStack.h
+++ b/Engine/Source/Runtime/Slate/Private/Framework/Docking/SDockingTabStack.h
@@ -36,7 +36,7 @@ public:
 	void AddTabWidget( const TSharedRef<SDockTab>& InTab, int32 AtLocation = INDEX_NONE);
 
 	/** @return All child tabs in this node */
-	const TArray< TSharedRef<SDockTab> >& GetTabs() const;
+	const TSlotlessChildren<SDockTab>& GetTabs() const;
 	
 	/** @return How many tabs are in this node */
 	int32 GetNumTabs() const;
diff --git a/Engine/Source/Runtime/Slate/Private/Framework/Docking/SDockingTabWell.cpp b/Engine/Source/Runtime/Slate/Private/Framework/Docking/SDockingTabWell.cpp
--- a/Engine/Source/Runtime/Slate/Private/Framework/Docking/SDockingTabWell.cpp
+++ b/Engine/Source/Runtime/Slate/Private/Framework/Docking/SDockingTabWell.cpp
@@ -14,6 +14,10 @@ const FVector2D FDockingConstants::GetMaxTabSizeFor( ETabRole::Type TabRole )
 		: MaxMinorTabSize;
 }
 
+SDockingTabWell::SDockingTabWell()
+: Tabs()
+{
+}
 
 void SDockingTabWell::Construct( const FArguments& InArgs )
 {
@@ -27,7 +31,7 @@ void SDockingTabWell::Construct( const FArguments& InArgs )
 	ParentTabStackPtr = InArgs._ParentStackNode.Get();
 }
 
-const TArray< TSharedRef<SDockTab> >& SDockingTabWell::GetTabs() const
+const TSlotlessChildren<SDockTab>& SDockingTabWell::GetTabs() const
 {
 	return Tabs;
 }
diff --git a/Engine/Source/Runtime/Slate/Private/Framework/Docking/SDockingTabWell.h b/Engine/Source/Runtime/Slate/Private/Framework/Docking/SDockingTabWell.h
--- a/Engine/Source/Runtime/Slate/Private/Framework/Docking/SDockingTabWell.h
+++ b/Engine/Source/Runtime/Slate/Private/Framework/Docking/SDockingTabWell.h
@@ -25,13 +25,15 @@ public:
 		SLATE_ATTRIBUTE( TSharedPtr<SDockingTabStack>, ParentStackNode )
 	SLATE_END_ARGS()
 
+	SDockingTabWell();
+
 	void Construct( const FArguments& InArgs );
 
 	/** @return How many tabs there are. */
 	int32 GetNumTabs() const;
 
 	/** @return All child tabs in this node */
-	const TArray< TSharedRef<SDockTab> >& GetTabs() const;
+	const TSlotlessChildren<SDockTab>& GetTabs() const;
 
 	/**
 	 * Add a new tab (InTab) to the TabWell at location (AtIndex).
diff --git a/Engine/Source/Runtime/Slate/Private/Framework/MultiBox/SButtonRowBlock.cpp b/Engine/Source/Runtime/Slate/Private/Framework/MultiBox/SButtonRowBlock.cpp
--- a/Engine/Source/Runtime/Slate/Private/Framework/MultiBox/SButtonRowBlock.cpp
+++ b/Engine/Source/Runtime/Slate/Private/Framework/MultiBox/SButtonRowBlock.cpp
@@ -137,8 +137,8 @@ void SButtonRowBlock::BuildMultiBlockWidget(const ISlateStyle* StyleSet, const F
 	const EUserInterfaceActionType::Type UserInterfaceType = ButtonRowBlock->GetAction().IsValid() ? ButtonRowBlock->GetAction()->GetUserInterfaceType() : ButtonRowBlock->UserInterfaceActionTypeOverride;
 	if( UserInterfaceType == EUserInterfaceActionType::Button )
 	{
-		ChildSlot.Widget =
-
+		ChildSlot
+		[
 			// Create a button
 			SNew( SButton )
 
@@ -158,12 +158,12 @@ void SButtonRowBlock::BuildMultiBlockWidget(const ISlateStyle* StyleSet, const F
 
 				// Bind the button's "on clicked" event to our object's method for this
 				.OnClicked( this, &SButtonRowBlock::OnClicked )
-			;
+		];
 	}
 	else if( ensure( UserInterfaceType == EUserInterfaceActionType::ToggleButton ) )
 	{
-		ChildSlot.Widget =
-
+		ChildSlot
+		[
 			// Create a check box
  			SAssignNew( ToggleButton, SCheckBox )
  
@@ -184,8 +184,8 @@ void SButtonRowBlock::BuildMultiBlockWidget(const ISlateStyle* StyleSet, const F
 				.OnCheckStateChanged( this, &SButtonRowBlock::OnCheckStateChanged )
 
 				// Bind the check box's "checked" state to our user interface action
-				.IsChecked( this, &SButtonRowBlock::OnIsChecked );
-			;
+				.IsChecked( this, &SButtonRowBlock::OnIsChecked )
+		];
 	}
 
 
diff --git a/Engine/Source/Runtime/Slate/Private/Framework/MultiBox/SClippingHorizontalBox.cpp b/Engine/Source/Runtime/Slate/Private/Framework/MultiBox/SClippingHorizontalBox.cpp
--- a/Engine/Source/Runtime/Slate/Private/Framework/MultiBox/SClippingHorizontalBox.cpp
+++ b/Engine/Source/Runtime/Slate/Private/Framework/MultiBox/SClippingHorizontalBox.cpp
@@ -117,7 +117,7 @@ FVector2D SClippingHorizontalBox::ComputeDesiredSize() const
 	{
 		// If the wrap button isn't being shown, subtract it's size from the total desired size
 		const SBoxPanel::FSlot& Child = Children[Children.Num() - 1];
-		const FVector2D& ChildDesiredSize = Child.Widget->GetDesiredSize();
+		const FVector2D& ChildDesiredSize = Child.GetWidget()->GetDesiredSize();
 		Size.X -= ChildDesiredSize.X;
 	}
 	return Size;
diff --git a/Engine/Source/Runtime/Slate/Private/Framework/MultiBox/SEditableTextBlock.cpp b/Engine/Source/Runtime/Slate/Private/Framework/MultiBox/SEditableTextBlock.cpp
--- a/Engine/Source/Runtime/Slate/Private/Framework/MultiBox/SEditableTextBlock.cpp
+++ b/Engine/Source/Runtime/Slate/Private/Framework/MultiBox/SEditableTextBlock.cpp
@@ -58,7 +58,8 @@ void SEditableTextBlock::BuildMultiBlockWidget( const ISlateStyle* StyleSet, con
 		}
 	}
 
-	ChildSlot.Widget =
+	ChildSlot
+	[
 		SNew( SHorizontalBox )
 		+ SHorizontalBox::Slot()
 		.AutoWidth()
@@ -116,7 +117,8 @@ void SEditableTextBlock::BuildMultiBlockWidget( const ISlateStyle* StyleSet, con
 			.OnTextChanged( EditableTextBlock->OnTextChanged )
 			.OnTextCommitted( EditableTextBlock->OnTextCommitted )
 			.ToolTipText( ToolTip )
-		];
+		]
+	];
 
 	// Bind our widget's enabled state to whether or not our action can execute
 	SetEnabled( TAttribute<bool>( this, &SEditableTextBlock::IsEnabled ) );
diff --git a/Engine/Source/Runtime/Slate/Private/Framework/MultiBox/SMenuEntryBlock.cpp b/Engine/Source/Runtime/Slate/Private/Framework/MultiBox/SMenuEntryBlock.cpp
--- a/Engine/Source/Runtime/Slate/Private/Framework/MultiBox/SMenuEntryBlock.cpp
+++ b/Engine/Source/Runtime/Slate/Private/Framework/MultiBox/SMenuEntryBlock.cpp
@@ -740,12 +740,12 @@ void SMenuEntryBlock::BuildMultiBlockWidget(const ISlateStyle* StyleSet, const F
 		if( MenuEntryBlock->bIsSubMenu )
 		{
 			// This menu entry is actually a submenu that opens a new menu to the right
-			ChildSlot.Widget = BuildSubMenuWidget( BuildParams );
+			ChildSlot.AttachWidget( BuildSubMenuWidget( BuildParams ) );
 		}
 		else
 		{
 			// Standard menu entry 
-			ChildSlot.Widget = BuildMenuEntryWidget( BuildParams );
+			ChildSlot.AttachWidget( BuildMenuEntryWidget( BuildParams ) );
 		}
 	}
 	else if( ensure( MultiBox->GetType() == EMultiBoxType::MenuBar ) )
@@ -753,19 +753,21 @@ void SMenuEntryBlock::BuildMultiBlockWidget(const ISlateStyle* StyleSet, const F
 		// Menu bar items cannot be submenus
 		check( !MenuEntryBlock->bIsSubMenu );
 		
-		ChildSlot.Widget = BuildMenuBarWidget( BuildParams );
+		ChildSlot[ BuildMenuBarWidget( BuildParams ) ];
 	}
 
 	// Insert named STutorialWrapper if desired
 	FName TutorialName = MenuEntryBlock->GetTutorialHighlightName();
 	if(TutorialName != NAME_None)
 	{
-		TSharedRef<SWidget> ChildWidget = ChildSlot.Widget;
-		ChildSlot.Widget = 
+		TSharedRef<SWidget> ChildWidget = ChildSlot.GetWidget();
+		ChildSlot
+		[
 			SNew( STutorialWrapper, TutorialName )
 			[
 				ChildWidget
-			];
+			]
+		];
 	}
 
 	// Bind our widget's enabled state to whether or not our action can execute
diff --git a/Engine/Source/Runtime/Slate/Private/Framework/MultiBox/SMenuSeparatorBlock.cpp b/Engine/Source/Runtime/Slate/Private/Framework/MultiBox/SMenuSeparatorBlock.cpp
--- a/Engine/Source/Runtime/Slate/Private/Framework/MultiBox/SMenuSeparatorBlock.cpp
+++ b/Engine/Source/Runtime/Slate/Private/Framework/MultiBox/SMenuSeparatorBlock.cpp
@@ -42,7 +42,8 @@ void SMenuSeparatorBlock::Construct( const FArguments& InArgs )
  */
 void SMenuSeparatorBlock::BuildMultiBlockWidget(const ISlateStyle* StyleSet, const FName& StyleName)
 {
-	ChildSlot.Widget =
+	ChildSlot
+	[
 		SNew( SVerticalBox )
 			+SVerticalBox::Slot()
 				.AutoHeight()
@@ -57,7 +58,8 @@ void SMenuSeparatorBlock::BuildMultiBlockWidget(const ISlateStyle* StyleSet, con
 
 						// Separator graphic
 						.BorderImage( StyleSet->GetBrush( StyleName, ".Separator" ) )
-				];
+				]
+	];
 }
 
 
diff --git a/Engine/Source/Runtime/Slate/Private/Framework/MultiBox/SToolBarButtonBlock.cpp b/Engine/Source/Runtime/Slate/Private/Framework/MultiBox/SToolBarButtonBlock.cpp
--- a/Engine/Source/Runtime/Slate/Private/Framework/MultiBox/SToolBarButtonBlock.cpp
+++ b/Engine/Source/Runtime/Slate/Private/Framework/MultiBox/SToolBarButtonBlock.cpp
@@ -210,8 +210,8 @@ void SToolBarButtonBlock::BuildMultiBlockWidget(const ISlateStyle* StyleSet, con
 	if( UserInterfaceType == EUserInterfaceActionType::Button )
 	{
 		FName BlockStyle = EMultiBlockLocation::ToName(ISlateStyle::Join( StyleName, ".Button" ), BlockLocation);
-		ChildSlot.Widget =
-
+		ChildSlot
+		[
 			// Create a button
 			SNew( SButton )
 				.ContentPadding(0)
@@ -231,15 +231,15 @@ void SToolBarButtonBlock::BuildMultiBlockWidget(const ISlateStyle* StyleSet, con
 
 				// Pass along the block's tool-tip string
 				.ToolTip( FMultiBoxSettings::ToolTipConstructor.Execute( ActualToolTip, NULL, Action.Pin() ) )
-			;
+		];
 	}
 	else if( ensure( UserInterfaceType == EUserInterfaceActionType::ToggleButton || UserInterfaceType == EUserInterfaceActionType::RadioButton ) )
 	{
 		FName BlockStyle = EMultiBlockLocation::ToName(ISlateStyle::Join( StyleName, ".ToggleButton" ), BlockLocation);
 		FName CheckboxStyle = ISlateStyle::Join( StyleName, ".SToolBarButtonBlock.CheckBox.Padding" );
 
-		ChildSlot.Widget =
-
+		ChildSlot
+		[
 			// Create a check box
 			SNew( SCheckBox )
 
@@ -261,8 +261,8 @@ void SToolBarButtonBlock::BuildMultiBlockWidget(const ISlateStyle* StyleSet, con
 				// Bind the check box's "checked" state to our user interface action
 				.IsChecked( this, &SToolBarButtonBlock::OnIsChecked )
 
-				.Padding( StyleSet->GetMargin(CheckboxStyle) );
-			;
+				.Padding( StyleSet->GetMargin(CheckboxStyle) )
+		];
 	}
 	
 	ChildSlot.Padding(StyleSet->GetMargin(ISlateStyle::Join( StyleName, ".SToolBarButtonBlock.Padding" )));
diff --git a/Engine/Source/Runtime/Slate/Private/Framework/MultiBox/SToolBarComboButtonBlock.cpp b/Engine/Source/Runtime/Slate/Private/Framework/MultiBox/SToolBarComboButtonBlock.cpp
--- a/Engine/Source/Runtime/Slate/Private/Framework/MultiBox/SToolBarComboButtonBlock.cpp
+++ b/Engine/Source/Runtime/Slate/Private/Framework/MultiBox/SToolBarComboButtonBlock.cpp
@@ -141,7 +141,8 @@ void SToolBarComboButtonBlock::BuildMultiBlockWidget(const ISlateStyle* StyleSet
 	FName BlockStyle = EMultiBlockLocation::ToName(ISlateStyle::Join( StyleName, ".Button" ), BlockLocation);
 	FName ColorStyle = ISlateStyle::Join( StyleName, ".SToolBarComboButtonBlock.ComboButton.Color" );
 
-	ChildSlot.Widget =
+	ChildSlot
+	[
 		SNew( SComboButton )
 			.ContentPadding(0)
 
@@ -160,7 +161,7 @@ void SToolBarComboButtonBlock::BuildMultiBlockWidget(const ISlateStyle* StyleSet
 
 			// Route the content generator event
 			.OnGetMenuContent( this, &SToolBarComboButtonBlock::OnGetMenuContent )
-		;
+	];
 
 	ChildSlot.Padding(StyleSet->GetMargin(ISlateStyle::Join( StyleName, ".SToolBarComboButtonBlock.Padding" )));
 	// Bind our widget's enabled state to whether or not our action can execute
diff --git a/Engine/Source/Runtime/Slate/Private/Framework/MultiBox/SToolBarSeparatorBlock.cpp b/Engine/Source/Runtime/Slate/Private/Framework/MultiBox/SToolBarSeparatorBlock.cpp
--- a/Engine/Source/Runtime/Slate/Private/Framework/MultiBox/SToolBarSeparatorBlock.cpp
+++ b/Engine/Source/Runtime/Slate/Private/Framework/MultiBox/SToolBarSeparatorBlock.cpp
@@ -50,16 +50,18 @@ void SToolBarSeparatorBlock::Construct( const FArguments& InArgs )
  */
 void SToolBarSeparatorBlock::BuildMultiBlockWidget(const ISlateStyle* StyleSet, const FName& StyleName)
 {
-	ChildSlot.Widget =
+	ChildSlot
+	[
 		SNew( SHorizontalBox )
-			+SHorizontalBox::Slot()
-				.AutoWidth()
-				.Padding( 0.0f, 0.0f, 0.0f, 0.0f )
-				[
-					SNew(SSeparator)
-						.Orientation(Orient_Vertical)
-						.SeparatorImage( StyleSet->GetBrush( ISlateStyle::Join( StyleName, ".Separator" ) ) )
-				];
+		+SHorizontalBox::Slot()
+		.AutoWidth()
+		.Padding( 0.0f, 0.0f, 0.0f, 0.0f )
+		[
+			SNew(SSeparator)
+				.Orientation(Orient_Vertical)
+				.SeparatorImage( StyleSet->GetBrush( ISlateStyle::Join( StyleName, ".Separator" ) ) )
+		]
+	];
 }
 
 
diff --git a/Engine/Source/Runtime/Slate/Private/Framework/Text/SlateHyperlinkRun.cpp b/Engine/Source/Runtime/Slate/Private/Framework/Text/SlateHyperlinkRun.cpp
--- a/Engine/Source/Runtime/Slate/Private/Framework/Text/SlateHyperlinkRun.cpp
+++ b/Engine/Source/Runtime/Slate/Private/Framework/Text/SlateHyperlinkRun.cpp
@@ -79,9 +79,9 @@ int32 FSlateHyperlinkRun::OnPaint( const FPaintArgs& Args, const FTextLayout::FL
 	return WidgetBlock->GetWidget()->Paint( Args, WidgetGeometry, MyClippingRect, OutDrawElements, LayerId, InWidgetStyle, bParentEnabled );
 }
 
-FChildren* FSlateHyperlinkRun::GetChildren()
+const TArray< TSharedRef<SWidget> >& FSlateHyperlinkRun::GetChildren()
 {
-	return &Children;
+	return Children;
 }
 
 void FSlateHyperlinkRun::ArrangeChildren( const TSharedRef< ILayoutBlock >& Block, const FGeometry& AllottedGeometry, FArrangedChildren& ArrangedChildren ) const 
diff --git a/Engine/Source/Runtime/Slate/Private/Framework/Text/SlateImageRun.cpp b/Engine/Source/Runtime/Slate/Private/Framework/Text/SlateImageRun.cpp
--- a/Engine/Source/Runtime/Slate/Private/Framework/Text/SlateImageRun.cpp
+++ b/Engine/Source/Runtime/Slate/Private/Framework/Text/SlateImageRun.cpp
@@ -6,8 +6,6 @@
 
 #include "SlateImageRun.h"
 
-FNoChildren FSlateImageRun::NoChildrenInstance;
-
 TSharedRef< FSlateImageRun > FSlateImageRun::Create( const TSharedRef< const FString >& InText, const FSlateBrush* InImage, int16 InBaseline )
 {
 	if ( InImage == NULL )
@@ -46,9 +44,10 @@ FSlateImageRun::FSlateImageRun( const TSharedRef< const FString >& InText, const
 	check( Image != NULL );
 }
 
-FChildren* FSlateImageRun::GetChildren() 
+const TArray< TSharedRef<SWidget> >& FSlateImageRun::GetChildren() 
 {
-	return &NoChildrenInstance;
+	static TArray< TSharedRef<SWidget> > NoChildren;
+	return NoChildren;
 }
 
 void FSlateImageRun::ArrangeChildren( const TSharedRef< ILayoutBlock >& Block, const FGeometry& AllottedGeometry, FArrangedChildren& ArrangedChildren ) const 
diff --git a/Engine/Source/Runtime/Slate/Private/Framework/Text/SlateTextLayout.cpp b/Engine/Source/Runtime/Slate/Private/Framework/Text/SlateTextLayout.cpp
--- a/Engine/Source/Runtime/Slate/Private/Framework/Text/SlateTextLayout.cpp
+++ b/Engine/Source/Runtime/Slate/Private/Framework/Text/SlateTextLayout.cpp
@@ -134,10 +134,10 @@ void FSlateTextLayout::AggregateChildren()
 			const FRunModel& LineRun = LineModel.Runs[ RunIndex ];
 			const TSharedRef< ISlateRun > SlateRun = StaticCastSharedRef< ISlateRun >( LineRun.GetRun() );
 
-			FChildren* RunChildren = SlateRun->GetChildren();
-			for (int ChildIndex = 0; ChildIndex < RunChildren->Num(); ChildIndex++)
+			const TArray< TSharedRef<SWidget> >& RunChildren = SlateRun->GetChildren();
+			for (int ChildIndex = 0; ChildIndex < RunChildren.Num(); ChildIndex++)
 			{
-				const TSharedRef< SWidget >& Child = RunChildren->GetChildAt( ChildIndex );
+				const TSharedRef< SWidget >& Child = RunChildren[ ChildIndex ];
 				Children.Add( Child );
 			}
 		}
diff --git a/Engine/Source/Runtime/Slate/Private/Framework/Text/SlateTextRun.cpp b/Engine/Source/Runtime/Slate/Private/Framework/Text/SlateTextRun.cpp
--- a/Engine/Source/Runtime/Slate/Private/Framework/Text/SlateTextRun.cpp
+++ b/Engine/Source/Runtime/Slate/Private/Framework/Text/SlateTextRun.cpp
@@ -6,8 +6,6 @@
 
 #include "SlateTextRun.h"
 
-FNoChildren FSlateTextRun::NoChildrenInstance;
-
 TSharedRef< FSlateTextRun > FSlateTextRun::Create( const TSharedRef< const FString >& InText, const FTextBlockStyle& Style )
 {
 	return MakeShareable( new FSlateTextRun( InText, Style ) );
@@ -111,9 +109,10 @@ int32 FSlateTextRun::OnPaint( const FPaintArgs& Args, const FTextLayout::FLineVi
 	return LayerId;
 }
 
-FChildren* FSlateTextRun::GetChildren()
+const TArray< TSharedRef<SWidget> >& FSlateTextRun::GetChildren()
 {
-	return &NoChildrenInstance;
+	static TArray< TSharedRef<SWidget> > NoChildren;
+	return NoChildren;
 }
 
 void FSlateTextRun::ArrangeChildren( const TSharedRef< ILayoutBlock >& Block, const FGeometry& AllottedGeometry, FArrangedChildren& ArrangedChildren ) const 
diff --git a/Engine/Source/Runtime/Slate/Private/Framework/Text/SlateWidgetRun.cpp b/Engine/Source/Runtime/Slate/Private/Framework/Text/SlateWidgetRun.cpp
--- a/Engine/Source/Runtime/Slate/Private/Framework/Text/SlateWidgetRun.cpp
+++ b/Engine/Source/Runtime/Slate/Private/Framework/Text/SlateWidgetRun.cpp
@@ -57,9 +57,9 @@ int32 FSlateWidgetRun::OnPaint( const FPaintArgs& Args, const FTextLayout::FLine
 	return Info.Widget->Paint( Args, WidgetGeometry, MyClippingRect, OutDrawElements, LayerId, InWidgetStyle, bParentEnabled );
 }
 
-FChildren* FSlateWidgetRun::GetChildren()
+const TArray< TSharedRef<SWidget> >& FSlateWidgetRun::GetChildren()
 {
-	return &Children;
+	return Children;
 }
 
 void FSlateWidgetRun::ArrangeChildren( const TSharedRef< ILayoutBlock >& Block, const FGeometry& AllottedGeometry, FArrangedChildren& ArrangedChildren ) const
@@ -123,7 +123,7 @@ FSlateWidgetRun::FSlateWidgetRun( const FSlateWidgetRun& Run )
 	: Text( Run.Text )
 	, Range( Run.Range )
 	, Info( Run.Info )
-	, Children( Run.Children )
+	, Children()
 {
 
 }
diff --git a/Engine/Source/Runtime/Slate/Private/Widgets/Colors/SColorThemes.cpp b/Engine/Source/Runtime/Slate/Private/Widgets/Colors/SColorThemes.cpp
--- a/Engine/Source/Runtime/Slate/Private/Widgets/Colors/SColorThemes.cpp
+++ b/Engine/Source/Runtime/Slate/Private/Widgets/Colors/SColorThemes.cpp
@@ -372,6 +372,12 @@ EVisibility SThemeColorBlock::OnGetAlphaVisibility() const
 	return bUseAlpha.Get() ? EVisibility::Visible : EVisibility::Collapsed;
 }
 
+SThemeColorBlocksBar::SThemeColorBlocksBar()
+: ColorBlocks()
+{
+}
+
+
 /**
 * Panels arrange their children in a space described by the AllottedGeometry parameter. The results of the arrangement
 * should be returned by appending a FArrangedWidget pair for every child widget.
diff --git a/Engine/Source/Runtime/Slate/Private/Widgets/Docking/SDockTab.cpp b/Engine/Source/Runtime/Slate/Private/Widgets/Docking/SDockTab.cpp
--- a/Engine/Source/Runtime/Slate/Private/Widgets/Docking/SDockTab.cpp
+++ b/Engine/Source/Runtime/Slate/Private/Widgets/Docking/SDockTab.cpp
@@ -630,7 +630,7 @@ FSlateColor SDockTab::GetFlashColor() const
 }
 
 
-float SDockTab::GetOverlapWidth()
+float SDockTab::GetOverlapWidth() const
 {
 	return GetCurrentStyle().OverlapWidth;
 }
diff --git a/Engine/Source/Runtime/Slate/Private/Widgets/Images/SThrobber.cpp b/Engine/Source/Runtime/Slate/Private/Widgets/Images/SThrobber.cpp
--- a/Engine/Source/Runtime/Slate/Private/Widgets/Images/SThrobber.cpp
+++ b/Engine/Source/Runtime/Slate/Private/Widgets/Images/SThrobber.cpp
@@ -11,7 +11,10 @@ void SThrobber::Construct(const FArguments& InArgs)
 
 	HBox = SNew(SHorizontalBox);
 
-	this->ChildSlot.Widget = HBox.ToSharedRef();
+	this->ChildSlot
+	[
+		HBox.ToSharedRef()
+	];
 
 	ConstructPieces();
 }
diff --git a/Engine/Source/Runtime/Slate/Private/Widgets/Input/SButton.cpp b/Engine/Source/Runtime/Slate/Private/Widgets/Input/SButton.cpp
--- a/Engine/Source/Runtime/Slate/Private/Widgets/Input/SButton.cpp
+++ b/Engine/Source/Runtime/Slate/Private/Widgets/Input/SButton.cpp
@@ -255,7 +255,7 @@ FVector2D SButton::ComputeDesiredSize() const
 {
 	// When there is no widget in the button, it sizes itself based on
 	// the border image specified by the style.
-	if (ChildSlot.Widget == SNullWidget::NullWidget)
+	if (ChildSlot.GetWidget() == SNullWidget::NullWidget)
 	{
 		return GetBorder()->ImageSize;
 	}
diff --git a/Engine/Source/Runtime/Slate/Private/Widgets/Layout/SBorder.cpp b/Engine/Source/Runtime/Slate/Private/Widgets/Layout/SBorder.cpp
--- a/Engine/Source/Runtime/Slate/Private/Widgets/Layout/SBorder.cpp
+++ b/Engine/Source/Runtime/Slate/Private/Widgets/Layout/SBorder.cpp
@@ -56,13 +56,13 @@ void SBorder::SetContent( const TSharedRef< SWidget >& InContent )
 
 const TSharedRef< SWidget >& SBorder::GetContent() const
 {
-	return ChildSlot.Widget;
+	return ChildSlot.GetWidget();
 }
 
 /** Clears out the content for the border */
 void SBorder::ClearContent()
 {
-	ChildSlot = FSimpleSlot();
+	ChildSlot.DetachWidget();
 }
 
 int32 SBorder::OnPaint( const FPaintArgs& Args, const FGeometry& AllottedGeometry, const FSlateRect& MyClippingRect, FSlateWindowElementList& OutDrawElements, int32 LayerId, const FWidgetStyle& InWidgetStyle, bool bParentEnabled ) const
diff --git a/Engine/Source/Runtime/Slate/Private/Widgets/Layout/SBox.cpp b/Engine/Source/Runtime/Slate/Private/Widgets/Layout/SBox.cpp
--- a/Engine/Source/Runtime/Slate/Private/Widgets/Layout/SBox.cpp
+++ b/Engine/Source/Runtime/Slate/Private/Widgets/Layout/SBox.cpp
@@ -3,6 +3,11 @@
 #include "SlatePrivatePCH.h"
 #include "LayoutUtils.h"
 
+SBox::SBox()
+: ChildSlot()
+{
+}
+
 void SBox::Construct( const FArguments& InArgs )
 {
 	WidthOverride = InArgs._WidthOverride;
@@ -21,12 +26,12 @@ void SBox::Construct( const FArguments& InArgs )
 
 FVector2D SBox::ComputeDesiredSize() const
 {
-	EVisibility ChildVisibility = ChildSlot.Widget->GetVisibility();
+	EVisibility ChildVisibility = ChildSlot.GetWidget()->GetVisibility();
 
 	if ( ChildVisibility != EVisibility::Collapsed )
 	{
 		// If the user specified a fixed width or height, those values override the Box's content.
-		const FVector2D& UnmodifiedChildDesiredSize = ChildSlot.Widget->GetDesiredSize() + ChildSlot.SlotPadding.Get().GetDesiredSize();
+		const FVector2D& UnmodifiedChildDesiredSize = ChildSlot.GetWidget()->GetDesiredSize() + ChildSlot.SlotPadding.Get().GetDesiredSize();
 		const FOptionalSize CurrentWidthOverride = WidthOverride.Get();
 		const FOptionalSize CurrentHeightOverride = HeightOverride.Get();
 		return FVector2D(
@@ -49,7 +54,7 @@ void SBox::OnArrangeChildren( const FGeometry& AllottedGeometry, FArrangedChildr
 
 		ArrangedChildren.AddWidget(
 			AllottedGeometry.MakeChild(
-				ChildSlot.Widget,
+				ChildSlot.GetWidget(),
 				FVector2D(XAlignmentResult.Offset, YAlignmentResult.Offset),
 				FVector2D(XAlignmentResult.Size, YAlignmentResult.Size)
 			)
@@ -85,5 +90,3 @@ int32 SBox::OnPaint( const FPaintArgs& Args, const FGeometry& AllottedGeometry,
 	}
 	return LayerId;
 }
-
-
diff --git a/Engine/Source/Runtime/Slate/Private/Widgets/Layout/SDPIScaler.cpp b/Engine/Source/Runtime/Slate/Private/Widgets/Layout/SDPIScaler.cpp
--- a/Engine/Source/Runtime/Slate/Private/Widgets/Layout/SDPIScaler.cpp
+++ b/Engine/Source/Runtime/Slate/Private/Widgets/Layout/SDPIScaler.cpp
@@ -2,10 +2,17 @@
 
 #include "SlatePrivatePCH.h"
 
+SDPIScaler::SDPIScaler()
+: ChildSlot()
+{
+}
 
 void SDPIScaler::Construct( const FArguments& InArgs )
 {
-	this->ChildSlot.Widget = InArgs._Content.Widget;
+	this->ChildSlot
+	[
+		InArgs._Content.Widget
+	];
 	this->DPIScale = InArgs._DPIScale;
 }
 
@@ -17,7 +24,7 @@ void SDPIScaler::OnArrangeChildren( const FGeometry& AllottedGeometry, FArranged
 		const float MyDPIScale = DPIScale.Get();
 
 		ArrangedChildren.AddWidget( AllottedGeometry.MakeChild(
-			this->ChildSlot.Widget,
+			this->ChildSlot.GetWidget(),
 			FVector2D::ZeroVector,
 			AllottedGeometry.Size / MyDPIScale,
 			MyDPIScale
@@ -28,11 +35,10 @@ void SDPIScaler::OnArrangeChildren( const FGeometry& AllottedGeometry, FArranged
 	
 FVector2D SDPIScaler::ComputeDesiredSize() const
 {
-	return DPIScale.Get() * ChildSlot.Widget->GetDesiredSize();
+	return DPIScale.Get() * ChildSlot.GetWidget()->GetDesiredSize();
 }
 
 FChildren* SDPIScaler::GetChildren()
 {
 	return &ChildSlot;
 }
-
diff --git a/Engine/Source/Runtime/Slate/Private/Widgets/Layout/SExpandableArea.cpp b/Engine/Source/Runtime/Slate/Private/Widgets/Layout/SExpandableArea.cpp
--- a/Engine/Source/Runtime/Slate/Private/Widgets/Layout/SExpandableArea.cpp
+++ b/Engine/Source/Runtime/Slate/Private/Widgets/Layout/SExpandableArea.cpp
@@ -175,10 +175,10 @@ FVector2D SExpandableArea::GetSectionScale() const
 
 FVector2D SExpandableArea::ComputeDesiredSize() const
 {
-	EVisibility ChildVisibility = ChildSlot.Widget->GetVisibility();
+	EVisibility ChildVisibility = ChildSlot.GetWidget()->GetVisibility();
 	if ( ChildVisibility != EVisibility::Collapsed )
 	{
-		FVector2D DesiredSize = ChildSlot.Widget->GetDesiredSize() + ChildSlot.SlotPadding.Get().GetDesiredSize();
+		FVector2D DesiredSize = ChildSlot.GetWidget()->GetDesiredSize() + ChildSlot.SlotPadding.Get().GetDesiredSize();
 		
 		// Only clamp if the user specified a max height
 		if( MaxHeight > 0.0f )
diff --git a/Engine/Source/Runtime/Slate/Private/Widgets/Layout/SFxWidget.cpp b/Engine/Source/Runtime/Slate/Private/Widgets/Layout/SFxWidget.cpp
--- a/Engine/Source/Runtime/Slate/Private/Widgets/Layout/SFxWidget.cpp
+++ b/Engine/Source/Runtime/Slate/Private/Widgets/Layout/SFxWidget.cpp
@@ -103,7 +103,7 @@ int32 SFxWidget::OnPaint( const FPaintArgs& Args, const FGeometry& AllottedGeome
 FVector2D SFxWidget::ComputeDesiredSize() const
 {
 	// Layout scale affects out desired size.
-	return LayoutScale.Get() * ChildSlot.Widget->GetDesiredSize();
+	return LayoutScale.Get() * ChildSlot.GetWidget()->GetDesiredSize();
 }
 
 
@@ -116,7 +116,7 @@ void SFxWidget::OnArrangeChildren( const FGeometry& AllottedGeometry, FArrangedC
 
 		// Only layout scale affects the arranged geometry. This part is identical to DPI scaling.
 		ArrangedChildren.AddWidget( AllottedGeometry.MakeChild(
-			this->ChildSlot.Widget,
+			this->ChildSlot.GetWidget(),
 			FVector2D::ZeroVector,
 			AllottedGeometry.Size / MyDPIScale,
 			MyDPIScale
diff --git a/Engine/Source/Runtime/Slate/Private/Widgets/Layout/SGridPanel.cpp b/Engine/Source/Runtime/Slate/Private/Widgets/Layout/SGridPanel.cpp
--- a/Engine/Source/Runtime/Slate/Private/Widgets/Layout/SGridPanel.cpp
+++ b/Engine/Source/Runtime/Slate/Private/Widgets/Layout/SGridPanel.cpp
@@ -3,7 +3,9 @@
 #include "SlatePrivatePCH.h"
 #include "LayoutUtils.h"
 
-
+SGridPanel::SGridPanel()
+: Slots()
+{}
 SGridPanel::FSlot& SGridPanel::AddSlot( int32 Column, int32 Row, SGridPanel::Layer InLayer )
 {
 	return InsertSlot( new FSlot( Column, Row, InLayer.TheLayer ) );
@@ -170,7 +172,7 @@ void SGridPanel::OnArrangeChildren( const FGeometry& AllottedGeometry, FArranged
 	for( int32 SlotIndex=0; SlotIndex < Slots.Num(); ++SlotIndex )
 	{
 		const FSlot& CurSlot = Slots[SlotIndex];
-		const EVisibility ChildVisibility = CurSlot.Widget->GetVisibility();
+		const EVisibility ChildVisibility = CurSlot.GetWidget()->GetVisibility();
 		if ( ArrangedChildren.Accepts(ChildVisibility) )
 		{
 			// Figure out the position of this cell.
@@ -189,7 +191,7 @@ void SGridPanel::OnArrangeChildren( const FGeometry& AllottedGeometry, FArranged
 
 			// Output the result
 			ArrangedChildren.AddWidget( ChildVisibility, AllottedGeometry.MakeChild( 
-				CurSlot.Widget,
+				CurSlot.GetWidget(),
 				ThisCellOffset + FVector2D( XAxisResult.Offset, YAxisResult.Offset ) + CurSlot.NudgeParam,
 				FVector2D(XAxisResult.Size, YAxisResult.Size)
 			));
@@ -349,10 +351,10 @@ void SGridPanel::ComputeDesiredCellSizes( TArray<float>& OutColumns, TArray<floa
 	for (int32 SlotIndex=0; SlotIndex<Slots.Num(); ++SlotIndex)
 	{
 		const FSlot& CurSlot = Slots[SlotIndex];
-		if (CurSlot.Widget->GetVisibility() != EVisibility::Collapsed)
+		if (CurSlot.GetWidget()->GetVisibility() != EVisibility::Collapsed)
 		{
 			// The slots wants to be as big as its content along with the required padding.
-			const FVector2D SlotDesiredSize = CurSlot.Widget->GetDesiredSize() + CurSlot.SlotPadding.Get().GetDesiredSize();
+			const FVector2D SlotDesiredSize = CurSlot.GetWidget()->GetDesiredSize() + CurSlot.SlotPadding.Get().GetDesiredSize();
 
 			// If the slot has a (colspan,rowspan) of (1,1) it will only affect that slot.
 			// For larger spans, the slots size will be evenly distributed across all the affected slots.
diff --git a/Engine/Source/Runtime/Slate/Private/Widgets/Layout/SScrollBar.cpp b/Engine/Source/Runtime/Slate/Private/Widgets/Layout/SScrollBar.cpp
--- a/Engine/Source/Runtime/Slate/Private/Widgets/Layout/SScrollBar.cpp
+++ b/Engine/Source/Runtime/Slate/Private/Widgets/Layout/SScrollBar.cpp
@@ -12,21 +12,12 @@ class SScrollBarTrack : public SPanel
 public:
 
 	/** A ListPanel slot is very simple - it just stores a widget. */
-	class FSlot
+	class FSlot : public TSlotBase<FSlot>
 	{
-	public:		
+	public:
 		FSlot()
-			: Widget( SNullWidget::NullWidget )
-		{
-		}
-
-		FSlot& operator[]( TSharedRef<SWidget> InWidget )
-		{
-			Widget = InWidget;
-			return *this;
-		}
-
-		TSharedRef<SWidget> Widget;
+		: TSlotBase<FSlot>()
+		{}
 	};
 
 	SLATE_BEGIN_ARGS( SScrollBarTrack )
@@ -45,6 +36,11 @@ public:
 
 	SLATE_END_ARGS()
 
+	SScrollBarTrack()
+	: Children()
+	{
+	}
+
 	/**
 	 * Construct the widget from a Declaration
 	 *
@@ -125,7 +121,7 @@ public:
 
 		FVector2D ChildPos(0, 0);
 		ArrangedChildren.AddWidget(			
-			AllottedGeometry.MakeChild( Children[TOP_SLOT_INDEX].Widget, ChildPos, ChildSize )
+			AllottedGeometry.MakeChild( Children[TOP_SLOT_INDEX].GetWidget(), ChildPos, ChildSize )
 		);
 
 		// Arrange bottom half of the track
@@ -138,7 +134,7 @@ public:
 			: FVector2D(Width, AllottedGeometry.Size.Y - TrackSizeInfo.GetThumbEnd());
 
 		ArrangedChildren.AddWidget( 
-			AllottedGeometry.MakeChild( Children[BOTTOM_SLOT_INDEX].Widget, ChildPos, ChildSize )
+			AllottedGeometry.MakeChild( Children[BOTTOM_SLOT_INDEX].GetWidget(), ChildPos, ChildSize )
 		);
 		
 		// Arrange the thumb; only show it in enabled scrollbars
@@ -153,7 +149,7 @@ public:
 				: FVector2D(Width, TrackSizeInfo.ThumbSize);
 
 			ArrangedChildren.AddWidget(
-				AllottedGeometry.MakeChild( Children[THUMB_SLOT_INDEX].Widget, ChildPos, ChildSize )
+				AllottedGeometry.MakeChild( Children[THUMB_SLOT_INDEX].GetWidget(), ChildPos, ChildSize )
 			);
 		}
 
@@ -169,12 +165,12 @@ public:
 	{
 		if ( Orientation == Orient_Horizontal )
 		{
-			const float DesiredHeight = FMath::Max3( Children[0].Widget->GetDesiredSize().Y, Children[1].Widget->GetDesiredSize().Y, Children[2].Widget->GetDesiredSize().Y );
+			const float DesiredHeight = FMath::Max3( Children[0].GetWidget()->GetDesiredSize().Y, Children[1].GetWidget()->GetDesiredSize().Y, Children[2].GetWidget()->GetDesiredSize().Y );
 			return FVector2D( MinThumbSize, DesiredHeight );
 		}
 		else
 		{
-			const float DesiredWidth = FMath::Max3( Children[0].Widget->GetDesiredSize().X, Children[1].Widget->GetDesiredSize().X, Children[2].Widget->GetDesiredSize().X );
+			const float DesiredWidth = FMath::Max3( Children[0].GetWidget()->GetDesiredSize().X, Children[1].GetWidget()->GetDesiredSize().X, Children[2].GetWidget()->GetDesiredSize().X );
 			return FVector2D( DesiredWidth, MinThumbSize );
 		}
 	}
diff --git a/Engine/Source/Runtime/Slate/Private/Widgets/Layout/SScrollBox.cpp b/Engine/Source/Runtime/Slate/Private/Widgets/Layout/SScrollBox.cpp
--- a/Engine/Source/Runtime/Slate/Private/Widgets/Layout/SScrollBox.cpp
+++ b/Engine/Source/Runtime/Slate/Private/Widgets/Layout/SScrollBox.cpp
@@ -20,6 +20,11 @@ public:
 		}
 	SLATE_END_ARGS()
 
+	SScrollPanel()
+	: Children()
+	{
+	}
+
 	void Construct( const FArguments& InArgs, const TArray<SScrollBox::FSlot*>& InSlots )
 	{
 		PhysicalOffset = 0;
@@ -39,19 +44,19 @@ public:
 		for( int32 SlotIndex=0; SlotIndex < Children.Num(); ++SlotIndex )
 		{
 			const SScrollBox::FSlot& ThisSlot = Children[SlotIndex];
-			const EVisibility ChildVisibility = ThisSlot.Widget->GetVisibility();
+			const EVisibility ChildVisibility = ThisSlot.GetWidget()->GetVisibility();
 
 			if ( ChildVisibility != EVisibility::Collapsed )
 			{
 				const FMargin& ThisPadding = ThisSlot.SlotPadding.Get();
-				const FVector2D& WidgetDesiredSize = ThisSlot.Widget->GetDesiredSize();
+				const FVector2D& WidgetDesiredSize = ThisSlot.GetWidget()->GetDesiredSize();
 				const float ThisSlotDesiredHeight = WidgetDesiredSize.Y + ThisSlot.SlotPadding.Get().GetTotalSpaceAlong<Orient_Vertical>();
 
 				// Figure out the size and local position of the child within the slot
 				// There is no vertical alignment, because it does not make sense in a panel where items are stacked end-to-end
 				AlignmentArrangeResult XAlignmentResult = AlignChild<Orient_Horizontal>( AllottedGeometry.Size.X, ThisSlot, ThisPadding );
 
-				ArrangedChildren.AddWidget( AllottedGeometry.MakeChild( ThisSlot.Widget, FVector2D(XAlignmentResult.Offset, CurChildOffset + ThisPadding.Top), FVector2D(XAlignmentResult.Size, WidgetDesiredSize.Y) ) );
+				ArrangedChildren.AddWidget( AllottedGeometry.MakeChild( ThisSlot.GetWidget(), FVector2D(XAlignmentResult.Offset, CurChildOffset + ThisPadding.Top), FVector2D(XAlignmentResult.Size, WidgetDesiredSize.Y) ) );
 				CurChildOffset += ThisSlotDesiredHeight;
 			}
 		}
@@ -63,7 +68,7 @@ public:
 		for(int32 SlotIndex=0; SlotIndex < Children.Num(); ++SlotIndex )
 		{
 			const SScrollBox::FSlot& ThisSlot = Children[SlotIndex];
-			const FVector2D ChildDesiredSize = ThisSlot.Widget->GetDesiredSize();
+			const FVector2D ChildDesiredSize = ThisSlot.GetWidget()->GetDesiredSize();
 			ThisDesiredSize.X  = FMath::Max(ChildDesiredSize.X, ThisDesiredSize.X);
 			ThisDesiredSize.Y += ChildDesiredSize.Y + ThisSlot.SlotPadding.Get().GetTotalSpaceAlong<Orient_Vertical>();
 		}
@@ -167,7 +172,7 @@ void SScrollBox::Construct( const FArguments& InArgs )
 /** Adds a slot to SScrollBox */
 SScrollBox::FSlot& SScrollBox::AddSlot()
 {
-	SScrollBox::FSlot& NewSlot = SScrollBox::Slot();
+	SScrollBox::FSlot& NewSlot = *new SScrollBox::FSlot();
 	ScrollPanel->Children.Add( &NewSlot );
 
 	return NewSlot;
@@ -179,7 +184,7 @@ void SScrollBox::RemoveSlot( const TSharedRef<SWidget>& WidgetToRemove )
 	TPanelChildren<SScrollBox::FSlot>& Children = ScrollPanel->Children;
 	for( int32 SlotIndex=0; SlotIndex < Children.Num(); ++SlotIndex )
 	{
-		if ( Children[SlotIndex].Widget == WidgetToRemove )
+		if ( Children[SlotIndex].GetWidget() == WidgetToRemove )
 		{
 			Children.RemoveAt(SlotIndex);
 			return;
diff --git a/Engine/Source/Runtime/Slate/Private/Widgets/Layout/SSplitter.cpp b/Engine/Source/Runtime/Slate/Private/Widgets/Layout/SSplitter.cpp
--- a/Engine/Source/Runtime/Slate/Private/Widgets/Layout/SSplitter.cpp
+++ b/Engine/Source/Runtime/Slate/Private/Widgets/Layout/SSplitter.cpp
@@ -14,7 +14,7 @@ SSplitter::FSlot& SSplitter::Slot()
 
 SSplitter::FSlot& SSplitter::AddSlot( int32 AtIndex )
 {
-	FSlot& NewSlot = SSplitter::Slot();
+	FSlot& NewSlot = *new FSlot();
 	if ( AtIndex == INDEX_NONE )
 	{
 		// No index was specified; just add to the end of the list.
@@ -94,13 +94,13 @@ void SSplitter::OnArrangeChildren( const FGeometry& AllottedGeometry, FArrangedC
 	{
 		for (int32 ChildIndex=0; ChildIndex < Children.Num(); ++ChildIndex)	
 		{
-			if (Children[ChildIndex].Widget->GetVisibility() != EVisibility::Collapsed)
+			if (Children[ChildIndex].GetWidget()->GetVisibility() != EVisibility::Collapsed)
 			{
 				++NumNonCollapsedChildren;
 
 				if ( Children[ChildIndex].SizingRule == SSplitter::SizeToContent )
 				{
-					NonResizeableSpace += Children[ChildIndex].Widget->GetDesiredSize()[AxisIndex];
+					NonResizeableSpace += Children[ChildIndex].GetWidget()->GetDesiredSize()[AxisIndex];
 				}
 				else // SizingRule == SSplitter::FractionOfParent
 				{
@@ -123,21 +123,21 @@ void SSplitter::OnArrangeChildren( const FGeometry& AllottedGeometry, FArrangedC
 		const FSlot& CurSlot = Children[ChildIndex];
 
 		const float ChildSpace = ( CurSlot.SizingRule == SSplitter::SizeToContent )
-			? CurSlot.Widget->GetDesiredSize()[AxisIndex]
+			? CurSlot.GetWidget()->GetDesiredSize()[AxisIndex]
 		: ResizeableSpace * CurSlot.SizeValue.Get() / CoefficientTotal;
 
-		const EVisibility ChildVisibility = CurSlot.Widget->GetVisibility();
+		const EVisibility ChildVisibility = CurSlot.GetWidget()->GetVisibility();
 
 		// If the output array wants arranged children
 		if ( ArrangedChildren.Accepts(ChildVisibility) )
 		{
 			if (Orientation == Orient_Horizontal)
 			{
-				ArrangedChildren.AddWidget( ChildVisibility, AllottedGeometry.MakeChild( CurSlot.Widget, FVector2D(XOffset, 0), FVector2D( ChildSpace, AllottedGeometry.Size.Y ) ) );
+				ArrangedChildren.AddWidget( ChildVisibility, AllottedGeometry.MakeChild( CurSlot.GetWidget(), FVector2D(XOffset, 0), FVector2D( ChildSpace, AllottedGeometry.Size.Y ) ) );
 			}
 			else
 			{
-				ArrangedChildren.AddWidget( ChildVisibility, AllottedGeometry.MakeChild( CurSlot.Widget, FVector2D(0, XOffset), FVector2D( AllottedGeometry.Size.X, ChildSpace ) ) );
+				ArrangedChildren.AddWidget( ChildVisibility, AllottedGeometry.MakeChild( CurSlot.GetWidget(), FVector2D(0, XOffset), FVector2D( AllottedGeometry.Size.X, ChildSpace ) ) );
 			}
 		}
 
@@ -220,12 +220,12 @@ static FVector2D ComputeDesiredSizeForSplitter( const float PhysicalSplitterHand
 	for (int32 ChildIndex=0; ChildIndex < Children.Num(); ++ChildIndex)
 	{
 		const SSplitter::FSlot& CurSlot = Children[ChildIndex];
-		const EVisibility ChildVisibility = CurSlot.Widget->GetVisibility();
+		const EVisibility ChildVisibility = CurSlot.GetWidget()->GetVisibility();
 		if ( ChildVisibility != EVisibility::Collapsed )
 		{
 			++NumNonCollapsed;
 
-			FVector2D ChildDesiredSize( CurSlot.Widget->GetDesiredSize() );
+			FVector2D ChildDesiredSize( CurSlot.GetWidget()->GetDesiredSize() );
 			if ( Orientation == Orient_Horizontal )
 			{
 				MyDesiredSize.X += ChildDesiredSize.X;
@@ -426,7 +426,8 @@ EOrientation SSplitter::GetOrientation() const
 
 
 SSplitter::SSplitter()
-	: HoveredHandleIndex( INDEX_NONE )
+	: Children()
+	, HoveredHandleIndex( INDEX_NONE )
 	, bIsResizing( false )
 	, Orientation( Orient_Horizontal )
 	, Style( nullptr )
@@ -439,7 +440,7 @@ int32 SSplitter::FindResizeableSlotBeforeHandle( int32 DraggedHandle, const TPan
 	// Resizing collapsed or autosizing slots does not make sense (their size is predetermined).
 	// Search out from the DraggedHandle to find the first non-collapsed, non-autosizing slot we can resize.
 	int32 SlotBeforeDragHandle = DraggedHandle;
-	while ( SlotBeforeDragHandle >= 0 && (Children[SlotBeforeDragHandle].Widget->GetVisibility() == EVisibility::Collapsed || Children[SlotBeforeDragHandle].SizingRule.Get() == SSplitter::SizeToContent ) )
+	while ( SlotBeforeDragHandle >= 0 && (Children[SlotBeforeDragHandle].GetWidget()->GetVisibility() == EVisibility::Collapsed || Children[SlotBeforeDragHandle].SizingRule.Get() == SSplitter::SizeToContent ) )
 	{
 		--SlotBeforeDragHandle;
 	}
@@ -456,7 +457,7 @@ int32 SSplitter::FindResizeableSlotAfterHandle( int32 DraggedHandle, const TPane
 	// We also cannot resize auto-sizing slots.
 	int32 SlotAfterDragHandle = DraggedHandle+1;
 	{
-		while( SlotAfterDragHandle < NumChildren && (Children[SlotAfterDragHandle].Widget->GetVisibility() == EVisibility::Collapsed || Children[SlotAfterDragHandle].SizingRule.Get() == SSplitter::SizeToContent) )
+		while( SlotAfterDragHandle < NumChildren && (Children[SlotAfterDragHandle].GetWidget()->GetVisibility() == EVisibility::Collapsed || Children[SlotAfterDragHandle].SizingRule.Get() == SSplitter::SizeToContent) )
 		{
 			++SlotAfterDragHandle;
 		}
@@ -472,7 +473,7 @@ void SSplitter::FindAllResizeableSlotsAfterHandle( int32 DraggedHandle, const TP
 
 	for (int SlotIndex = DraggedHandle+1; SlotIndex < NumChildren; SlotIndex++)
 	{
-		if (Children[ SlotIndex ].Widget->GetVisibility() == EVisibility::Collapsed || Children[ SlotIndex ].SizingRule.Get() == SSplitter::SizeToContent)
+		if (Children[ SlotIndex ].GetWidget()->GetVisibility() == EVisibility::Collapsed || Children[ SlotIndex ].SizingRule.Get() == SSplitter::SizeToContent)
 		{
 			continue;
 		}
@@ -652,6 +653,11 @@ int32 SSplitter::GetHandleBeingResizedFromMousePosition( float PhysicalSplitterH
 * by dragging the center of the splitter.							*
 ********************************************************************/
 
+SSplitter2x2::SSplitter2x2()
+: Children()
+{
+}
+
 void SSplitter2x2::Construct( const FArguments& InArgs )
 {
 	Children.Add( new FSlot(InArgs._TopLeft.Widget) );
@@ -686,12 +692,12 @@ void SSplitter2x2::OnArrangeChildren( const FGeometry& AllottedGeometry, FArrang
 		// It is based on the current percentage of space it should take up which is defined by a user moving the splitters
 		FVector2D ChildSpace = SpaceAllottedForChildren * CurSlot.PercentageAttribute.Get();
 
-		const EVisibility ChildVisibility = CurSlot.Widget->GetVisibility();
+		const EVisibility ChildVisibility = CurSlot.GetWidget()->GetVisibility();
 
 		// If the child is visible, put them in their spot
 		if ( ArrangedChildren.Accepts( ChildVisibility ) )
 		{
-			ArrangedChildren.AddWidget( ChildVisibility, AllottedGeometry.MakeChild( CurSlot.Widget, Offset, ChildSpace ) );
+			ArrangedChildren.AddWidget( ChildVisibility, AllottedGeometry.MakeChild( CurSlot.GetWidget(), Offset, ChildSpace ) );
 		}
 
 		// Advance to the next slot. If the child is collapsed, it takes up no room and does not need a splitter
@@ -918,49 +924,49 @@ int32 SSplitter2x2::CalculateResizingAxis( const FGeometry& MyGeometry, const FV
 
 TSharedRef< SWidget > SSplitter2x2::GetTopLeftContent()
 {
-	return Children[ 0 ].Widget;
+	return Children[ 0 ].GetWidget();
 }
 
 
 TSharedRef< SWidget > SSplitter2x2::GetBottomLeftContent()
 {
-	return Children[ 1 ].Widget;
+	return Children[ 1 ].GetWidget();
 }
 
 
 TSharedRef< SWidget > SSplitter2x2::GetTopRightContent()
 {
-	return Children[ 2 ].Widget;
+	return Children[ 2 ].GetWidget();
 }
 
 
 TSharedRef< SWidget > SSplitter2x2::GetBottomRightContent()
 {
-	return Children[ 3 ].Widget;
+	return Children[ 3 ].GetWidget();
 }
 
 
 void SSplitter2x2::SetTopLeftContent( TSharedRef< SWidget > TopLeftContent )
 {
-	Children[ 0 ].Widget = TopLeftContent;
+	Children[ 0 ][ TopLeftContent ];
 }
 
 
 void SSplitter2x2::SetBottomLeftContent( TSharedRef< SWidget > BottomLeftContent )
 {
-	Children[ 1 ].Widget = BottomLeftContent;
+	Children[ 1 ][ BottomLeftContent ];
 }
 
 
 void SSplitter2x2::SetTopRightContent( TSharedRef< SWidget > TopRightContent )
 {
-	Children[ 2 ].Widget = TopRightContent;
+	Children[ 2 ][ TopRightContent ];
 }
 
 
 void SSplitter2x2::SetBottomRightContent( TSharedRef< SWidget > BottomRightContent )
 {
-	Children[ 3 ].Widget = BottomRightContent;
+	Children[ 3 ][ BottomRightContent ];
 }
 
 void SSplitter2x2::GetSplitterPercentages( TArray< FVector2D >& OutPercentages ) const
diff --git a/Engine/Source/Runtime/Slate/Private/Widgets/Layout/SUniformGridPanel.cpp b/Engine/Source/Runtime/Slate/Private/Widgets/Layout/SUniformGridPanel.cpp
--- a/Engine/Source/Runtime/Slate/Private/Widgets/Layout/SUniformGridPanel.cpp
+++ b/Engine/Source/Runtime/Slate/Private/Widgets/Layout/SUniformGridPanel.cpp
@@ -3,6 +3,10 @@
 #include "SlatePrivatePCH.h"
 #include "LayoutUtils.h"
 
+SUniformGridPanel::SUniformGridPanel()
+: Children()
+{
+}
 
 void SUniformGridPanel::Construct( const FArguments& InArgs )
 {
@@ -29,7 +33,7 @@ void SUniformGridPanel::OnArrangeChildren( const FGeometry& AllottedGeometry, FA
 		for ( int32 ChildIndex=0; ChildIndex < Children.Num(); ++ChildIndex )
 		{
 			const FSlot& Child = Children[ChildIndex];
-			const EVisibility ChildVisibility = Child.Widget->GetVisibility();
+			const EVisibility ChildVisibility = Child.GetWidget()->GetVisibility();
 			if ( ArrangedChildren.Accepts(ChildVisibility) )
 			{
 				// Do the standard arrangement of elements within a slot
@@ -38,7 +42,7 @@ void SUniformGridPanel::OnArrangeChildren( const FGeometry& AllottedGeometry, FA
 				AlignmentArrangeResult YAxisResult = AlignChild<Orient_Vertical>(CellSize.Y, Child, CurrentSlotPadding);
 
 				ArrangedChildren.AddWidget(ChildVisibility,
-					AllottedGeometry.MakeChild(Child.Widget,
+					AllottedGeometry.MakeChild(Child.GetWidget(),
 					FVector2D(CellSize.X*Child.Column + XAxisResult.Offset, CellSize.Y*Child.Row + YAxisResult.Offset),
 					FVector2D(XAxisResult.Size, YAxisResult.Size)
 					));
@@ -67,9 +71,9 @@ FVector2D SUniformGridPanel::ComputeDesiredSize() const
 		NumColumns = FMath::Max(Child.Column + 1, NumColumns);
 		NumRows = FMath::Max(Child.Row + 1, NumRows);
 
-		if (Child.Widget->GetVisibility() != EVisibility::Collapsed)
+		if (Child.GetWidget()->GetVisibility() != EVisibility::Collapsed)
 		{
-			FVector2D ChildDesiredSize = Child.Widget->GetDesiredSize() + SlotPaddingDesiredSize;
+			FVector2D ChildDesiredSize = Child.GetWidget()->GetDesiredSize() + SlotPaddingDesiredSize;
 
 			ChildDesiredSize.X = FMath::Max( ChildDesiredSize.X, CachedMinDesiredSlotWidth);
 			ChildDesiredSize.Y = FMath::Max( ChildDesiredSize.Y, CachedMinDesiredSlotHeight);
@@ -94,7 +98,7 @@ void SUniformGridPanel::SetSlotPadding(TAttribute<FMargin> InSlotPadding)
 
 SUniformGridPanel::FSlot& SUniformGridPanel::AddSlot( int32 Column, int32 Row )
 {
-	FSlot& NewSlot = SUniformGridPanel::Slot(Column, Row);
+	FSlot& NewSlot = *(new FSlot( Column, Row ));
 
 	Children.Add( &NewSlot );
 
@@ -105,7 +109,7 @@ bool SUniformGridPanel::RemoveSlot( const TSharedRef<SWidget>& SlotWidget )
 {
 	for (int32 SlotIdx = 0; SlotIdx < Children.Num(); ++SlotIdx)
 	{
-		if ( SlotWidget == Children[SlotIdx].Widget )
+		if ( SlotWidget == Children[SlotIdx].GetWidget() )
 		{
 			Children.RemoveAt(SlotIdx);
 			return true;
@@ -120,4 +124,4 @@ void SUniformGridPanel::ClearChildren()
 	NumColumns = 0;
 	NumRows = 0;
 	Children.Empty();
-}
+}
\ No newline at end of file
diff --git a/Engine/Source/Runtime/Slate/Private/Widgets/Layout/SWidgetSwitcher.cpp b/Engine/Source/Runtime/Slate/Private/Widgets/Layout/SWidgetSwitcher.cpp
--- a/Engine/Source/Runtime/Slate/Private/Widgets/Layout/SWidgetSwitcher.cpp
+++ b/Engine/Source/Runtime/Slate/Private/Widgets/Layout/SWidgetSwitcher.cpp
@@ -11,20 +11,25 @@
 /* SWidgetSwitcher interface
  *****************************************************************************/
 
-FSimpleSlot& SWidgetSwitcher::AddSlot( int32 SlotIndex )
+ SWidgetSwitcher::SWidgetSwitcher()
+ : AllChildren()
+ {
+ }
+
+SWidgetSwitcher::FSlot& SWidgetSwitcher::AddSlot( int32 SlotIndex )
 {
-	FSimpleSlot& NewSlot = SWidgetSwitcher::Slot();
+	FSlot* NewSlot = new FSlot();
 
 	if (!AllChildren.IsValidIndex(SlotIndex))
 	{
-		AllChildren.Add(&NewSlot);
+		AllChildren.Add(NewSlot);
 	}
 	else
 	{
-		AllChildren.Insert(&NewSlot, SlotIndex);
+		AllChildren.Insert(NewSlot, SlotIndex);
 	}
 
-	return NewSlot;
+	return *NewSlot;
 }
 
 
@@ -32,7 +37,7 @@ int32 SWidgetSwitcher::RemoveSlot(TSharedRef<SWidget> WidgetToRemove)
 {
 	for( int32 SlotIndex=0; SlotIndex < AllChildren.Num(); ++SlotIndex )
 	{
-		if ( AllChildren[SlotIndex].Widget == WidgetToRemove )
+		if ( AllChildren[SlotIndex].GetWidget() == WidgetToRemove )
 		{
 			AllChildren.RemoveAt(SlotIndex);
 			return SlotIndex;
@@ -61,7 +66,7 @@ TSharedPtr<SWidget> SWidgetSwitcher::GetActiveWidget( ) const
 	const int32 ActiveWidgetIndex = WidgetIndex.Get();
 	if (ActiveWidgetIndex >= 0)
 	{
-		return AllChildren[ActiveWidgetIndex].Widget;
+		return AllChildren[ActiveWidgetIndex].GetWidget();
 	}
 
 	return NULL;
@@ -72,7 +77,7 @@ TSharedPtr<SWidget> SWidgetSwitcher::GetWidget( int32 SlotIndex ) const
 {
 	if (AllChildren.IsValidIndex(SlotIndex))
 	{
-		return AllChildren[SlotIndex].Widget;
+		return AllChildren[SlotIndex].GetWidget();
 	}
 
 	return NULL;
@@ -83,9 +88,9 @@ int32 SWidgetSwitcher::GetWidgetIndex( TSharedRef<SWidget> Widget ) const
 {
 	for (int32 Index = 0; Index < AllChildren.Num(); ++Index)
 	{
-		const FSimpleSlot& Slot = AllChildren[Index];
+		const FSlot& Slot = AllChildren[Index];
 
-		if (Slot.Widget == Widget)
+		if (Slot.GetWidget() == Widget)
 		{
 			return Index;
 		}
@@ -116,7 +121,7 @@ void SWidgetSwitcher::OnArrangeChildren( const FGeometry& AllottedGeometry, FArr
 FVector2D SWidgetSwitcher::ComputeDesiredSize( ) const
 {
 	return AllChildren.Num() > 0
-		? AllChildren[WidgetIndex.Get()].Widget->GetDesiredSize()
+		? AllChildren[WidgetIndex.Get()].GetWidget()->GetDesiredSize()
 		: FVector2D::ZeroVector;
 }
 
diff --git a/Engine/Source/Runtime/Slate/Private/Widgets/Layout/SWrapBox.cpp b/Engine/Source/Runtime/Slate/Private/Widgets/Layout/SWrapBox.cpp
--- a/Engine/Source/Runtime/Slate/Private/Widgets/Layout/SWrapBox.cpp
+++ b/Engine/Source/Runtime/Slate/Private/Widgets/Layout/SWrapBox.cpp
@@ -3,6 +3,10 @@
 #include "SlatePrivatePCH.h"
 #include "LayoutUtils.h"
 
+SWrapBox::SWrapBox()
+: Slots()
+{
+}
 
 FWrapSlot& SWrapBox::Slot()
 {
@@ -11,7 +15,7 @@ FWrapSlot& SWrapBox::Slot()
 
 FWrapSlot& SWrapBox::AddSlot()
 {
-	FWrapSlot* NewSlot = &Slot();
+	FWrapSlot* NewSlot = new FWrapSlot();
 	Slots.Add(NewSlot);
 	return *NewSlot;
 }
@@ -20,7 +24,7 @@ int32 SWrapBox::RemoveSlot( const TSharedRef<SWidget>& SlotWidget )
 {
 	for (int32 SlotIdx = 0; SlotIdx < Slots.Num(); ++SlotIdx)
 	{
-		if ( SlotWidget == Slots[SlotIdx].Widget )
+		if ( SlotWidget == Slots[SlotIdx].GetWidget() )
 		{
 			Slots.RemoveAt(SlotIdx);
 			return SlotIdx;
@@ -64,7 +68,7 @@ void SWrapBox::OnArrangeChildren( const FGeometry& AllottedGeometry, FArrangedCh
 		float LineHeightSoFar = 0;
 		for ( int32 ChildIndex=0; ChildIndex < Slots.Num(); ++ChildIndex )
 		{
-			const TSharedRef<SWidget>& ThisChild = Slots[ChildIndex].Widget;
+			const TSharedRef<SWidget>& ThisChild = Slots[ChildIndex].GetWidget();
 
 			// If slot isn't the first one on the line, we need to add the inner slot padding.
 			const float ConditionPaddingX = ( WidthSoFar == 0 ) ? 0 : InnerSlotPadding.X;
@@ -129,7 +133,7 @@ void SWrapBox::OnArrangeChildren( const FGeometry& AllottedGeometry, FArrangedCh
 		const float CurrentLineHeight = LineHeights[ CurrentLineIndex ];
 		const FMargin& SlotPadding( Slots[ChildIndex].SlotPadding.Get() );
 
-		const TSharedRef<SWidget>& ThisChild = Slots[ChildIndex].Widget;
+		const TSharedRef<SWidget>& ThisChild = Slots[ChildIndex].GetWidget();
 		float ThisSlotSizeX = ThisChild->GetDesiredSize().X + SlotPadding.GetTotalSpaceAlong<Orient_Horizontal>();
 
 		// If the width is smaller than this slots fill limit, we need to account for it filling the line
@@ -183,7 +187,7 @@ FVector2D SWrapBox::ComputeDesiredSize() const
 		// If slot isn't the first one on the line, we need to add the inner slot padding.
 		const float ConditionPaddingX = ( RowDesiredSize.X == 0 ) ? 0 : InnerSlotPadding.X;
 
-		FVector2D ThisChildDesiredSize = Slots[ChildIndex].Widget->GetDesiredSize() + Slots[ChildIndex].SlotPadding.Get().GetDesiredSize() + ConditionPaddingX;
+		FVector2D ThisChildDesiredSize = Slots[ChildIndex].GetWidget()->GetDesiredSize() + Slots[ChildIndex].SlotPadding.Get().GetDesiredSize() + ConditionPaddingX;
 
 		// If the width is smaller than this slots fill limit, we need to account for it filling the line
 		// and forcing all controls that follow it to wrap.
diff --git a/Engine/Source/Runtime/Slate/Private/Widgets/SCanvas.cpp b/Engine/Source/Runtime/Slate/Private/Widgets/SCanvas.cpp
--- a/Engine/Source/Runtime/Slate/Private/Widgets/SCanvas.cpp
+++ b/Engine/Source/Runtime/Slate/Private/Widgets/SCanvas.cpp
@@ -7,6 +7,10 @@
 #include "SlatePrivatePCH.h"
 #include "LayoutUtils.h"
 
+SCanvas::SCanvas()
+: Children()
+{
+}
 
 /* SCanvas interface
  *****************************************************************************/
@@ -31,7 +35,7 @@ int32 SCanvas::RemoveSlot( const TSharedRef<SWidget>& SlotWidget )
 {
 	for (int32 SlotIdx = 0; SlotIdx < Children.Num(); ++SlotIdx)
 	{
-		if (SlotWidget == Children[SlotIdx].Widget)
+		if (SlotWidget == Children[SlotIdx].GetWidget())
 		{
 			Children.RemoveAt(SlotIdx);
 			return SlotIdx;
@@ -81,7 +85,7 @@ void SCanvas::OnArrangeChildren( const FGeometry& AllottedGeometry, FArrangedChi
 			// Add the information about this child to the output list (ArrangedChildren)
 			ArrangedChildren.AddWidget( AllottedGeometry.MakeChild(
 				// The child widget being arranged
-				CurChild.Widget,
+				CurChild.GetWidget(),
 				// Child's local position (i.e. position within parent)
 				CurChild.PositionAttr.Get() + Offset,
 				// Child's size
diff --git a/Engine/Source/Runtime/Slate/Private/Widgets/SWeakWidget.cpp b/Engine/Source/Runtime/Slate/Private/Widgets/SWeakWidget.cpp
--- a/Engine/Source/Runtime/Slate/Private/Widgets/SWeakWidget.cpp
+++ b/Engine/Source/Runtime/Slate/Private/Widgets/SWeakWidget.cpp
@@ -6,17 +6,21 @@
 
 #include "SlatePrivatePCH.h"
 
+SWeakWidget::SWeakWidget()
+: WeakChild()
+{
+}
 
 void SWeakWidget::Construct(const FArguments& InArgs)
 {
-	WeakChild.Widget = InArgs._PossiblyNullContent;
+	WeakChild.AttachWidget( InArgs._PossiblyNullContent );
 }
 
 FVector2D SWeakWidget::ComputeDesiredSize() const
 {	
-	TSharedPtr<SWidget> ReferencedWidget = WeakChild.Widget.Pin();
+	TSharedRef<SWidget> ReferencedWidget = WeakChild.GetWidget();
 
-	if (ReferencedWidget.IsValid() && ReferencedWidget->GetVisibility() != EVisibility::Collapsed)
+	if ( ReferencedWidget != SNullWidget::NullWidget && ReferencedWidget->GetVisibility() != EVisibility::Collapsed )
 	{
 		return ReferencedWidget->GetDesiredSize();
 	}
@@ -27,12 +31,12 @@ FVector2D SWeakWidget::ComputeDesiredSize() const
 void SWeakWidget::OnArrangeChildren( const FGeometry& AllottedGeometry, FArrangedChildren& ArrangedChildren ) const
 {
 	// We just want to show the child that we are presenting. Always stretch it to occupy all of the space.
-	TSharedPtr<SWidget> MyContent = WeakChild.Widget.Pin();
+	TSharedRef<SWidget> MyContent = WeakChild.GetWidget();
 
-	if( MyContent.IsValid() && ArrangedChildren.Accepts(MyContent->GetVisibility()))
+	if( MyContent!=SNullWidget::NullWidget && ArrangedChildren.Accepts(MyContent->GetVisibility()) )
 	{
 		ArrangedChildren.AddWidget( AllottedGeometry.MakeChild(
-			MyContent.ToSharedRef(),
+			MyContent,
 			FVector2D(0,0),
 			AllottedGeometry.Size
 			) );
@@ -68,17 +72,12 @@ int32 SWeakWidget::OnPaint( const FPaintArgs& Args, const FGeometry& AllottedGeo
 	return LayerId;
 }
 
-void SWeakWidget::SetContent(TWeakPtr<SWidget> InWidget)
+void SWeakWidget::SetContent(const TSharedRef<SWidget>& InWidget)
 {
-	WeakChild.Widget = InWidget;
+	WeakChild.AttachWidget( InWidget );
 }
 
 bool SWeakWidget::ChildWidgetIsValid() const
 {
-	return WeakChild.Widget.IsValid();
-}
-
-TWeakPtr<SWidget> SWeakWidget::GetChildWidget() const
-{
-	return WeakChild.Widget;
+	return WeakChild.GetWidget() != SNullWidget::NullWidget;
 }
diff --git a/Engine/Source/Runtime/Slate/Private/Widgets/Views/SListPanel.cpp b/Engine/Source/Runtime/Slate/Private/Widgets/Views/SListPanel.cpp
--- a/Engine/Source/Runtime/Slate/Private/Widgets/Views/SListPanel.cpp
+++ b/Engine/Source/Runtime/Slate/Private/Widgets/Views/SListPanel.cpp
@@ -6,6 +6,11 @@
 
 FNoChildren SListPanel::NoChildren = FNoChildren();
 
+SListPanel::SListPanel()
+: Children()
+{
+}
+
 /**
  * Construct the widget
  *
@@ -63,7 +68,7 @@ void SListPanel::OnArrangeChildren( const FGeometry& AllottedGeometry, FArranged
 		for( int32 ItemIndex = 0; ItemIndex < Children.Num(); ++ItemIndex )
 		{
 			ArrangedChildren.AddWidget(
-				AllottedGeometry.MakeChild( Children[ItemIndex].Widget, FVector2D(WidthSoFar + HalfItemPadding, HeightSoFar), FVector2D(LocalItemWidth, LocalItemHeight) )
+				AllottedGeometry.MakeChild( Children[ItemIndex].GetWidget(), FVector2D(WidthSoFar + HalfItemPadding, HeightSoFar), FVector2D(LocalItemWidth, LocalItemHeight) )
 				);
 
 			WidthSoFar += LocalItemWidth + ItemPadding;
@@ -80,16 +85,16 @@ void SListPanel::OnArrangeChildren( const FGeometry& AllottedGeometry, FArranged
 		if (Children.Num() > 0)
 		{
 			// This is a normal list, arrange items vertically
-			float HeightSoFar = -FMath::FloorToInt(SmoothScrollOffsetInItems * Children[0].Widget->GetDesiredSize().Y)-OverscrollAmount;
+			float HeightSoFar = -FMath::FloorToInt(SmoothScrollOffsetInItems * Children[0].GetWidget()->GetDesiredSize().Y)-OverscrollAmount;
 			for( int32 ItemIndex = 0; ItemIndex < Children.Num(); ++ItemIndex )
 			{
-					const FVector2D ItemDesiredSize = Children[ItemIndex].Widget->GetDesiredSize();
+					const FVector2D ItemDesiredSize = Children[ItemIndex].GetWidget()->GetDesiredSize();
 					const float LocalItemHeight = ItemDesiredSize.Y;
 
 				// Note that ListPanel does not respect child Visibility.
 				// It is simply not useful for ListPanels.
 				ArrangedChildren.AddWidget(
-					AllottedGeometry.MakeChild( Children[ItemIndex].Widget, FVector2D(0, HeightSoFar), FVector2D(AllottedGeometry.Size.X, LocalItemHeight) )
+					AllottedGeometry.MakeChild( Children[ItemIndex].GetWidget(), FVector2D(0, HeightSoFar), FVector2D(AllottedGeometry.Size.X, LocalItemHeight) )
 					);
 
 				HeightSoFar += LocalItemHeight;
@@ -136,7 +141,7 @@ FVector2D SListPanel::ComputeDesiredSize() const
 	{
 		// Notice that we do not respect child Visibility.
 		// It is simply not useful for ListPanels.
-		FVector2D ChildDesiredSize = Children[ItemIndex].Widget->GetDesiredSize();
+		FVector2D ChildDesiredSize = Children[ItemIndex].GetWidget()->GetDesiredSize();
 		MaxWidth = FMath::Max( ChildDesiredSize.X, MaxWidth );
 		TotalHeight += ChildDesiredSize.Y;
 	}
diff --git a/Engine/Source/Runtime/Slate/Private/Widgets/Views/SListPanel.h b/Engine/Source/Runtime/Slate/Private/Widgets/Views/SListPanel.h
--- a/Engine/Source/Runtime/Slate/Private/Widgets/Views/SListPanel.h
+++ b/Engine/Source/Runtime/Slate/Private/Widgets/Views/SListPanel.h
@@ -12,21 +12,12 @@ class SListPanel : public SPanel
 {
 public:
 	/** A ListPanel slot is very simple - it just stores a widget. */
-	class FSlot
+	class FSlot : public TSlotBase<FSlot>
 	{
-	public:		
+	public:
 		FSlot()
-		: Widget( SNullWidget::NullWidget )
-		{
-		}
-
-		FSlot& operator[]( TSharedRef<SWidget> InWidget )
-		{
-			Widget = InWidget;
-			return *this;
-		}
-
-		TSharedRef<SWidget> Widget;
+		: TSlotBase<FSlot>()
+		{}
 	};
 	
 	/** Make a new ListPanel::Slot  */
@@ -49,6 +40,8 @@ public:
 
 	SLATE_END_ARGS()
 	
+	SListPanel();
+
 	/**
 	 * Construct the widget
 	 *
diff --git a/Engine/Source/Runtime/Slate/Public/Framework/Text/ISlateRun.h b/Engine/Source/Runtime/Slate/Public/Framework/Text/ISlateRun.h
--- a/Engine/Source/Runtime/Slate/Public/Framework/Text/ISlateRun.h
+++ b/Engine/Source/Runtime/Slate/Public/Framework/Text/ISlateRun.h
@@ -11,7 +11,7 @@ public:
 
 	virtual int32 OnPaint( const FPaintArgs& Args, const FTextLayout::FLineView& Line, const TSharedRef< ILayoutBlock >& Block, const FTextBlockStyle& DefaultStyle, const FGeometry& AllottedGeometry, const FSlateRect& MyClippingRect, FSlateWindowElementList& OutDrawElements, int32 LayerId, const FWidgetStyle& InWidgetStyle, bool bParentEnabled ) const = 0;
 
-	virtual FChildren* GetChildren() = 0;
+	virtual const TArray< TSharedRef<SWidget> >& GetChildren() = 0;
 
 	virtual void ArrangeChildren( const TSharedRef< ILayoutBlock >& Block, const FGeometry& AllottedGeometry, FArrangedChildren& ArrangedChildren ) const = 0;
 };
diff --git a/Engine/Source/Runtime/Slate/Public/Framework/Text/SlateHyperlinkRun.h b/Engine/Source/Runtime/Slate/Public/Framework/Text/SlateHyperlinkRun.h
--- a/Engine/Source/Runtime/Slate/Public/Framework/Text/SlateHyperlinkRun.h
+++ b/Engine/Source/Runtime/Slate/Public/Framework/Text/SlateHyperlinkRun.h
@@ -54,7 +54,7 @@ public:
 
 	virtual int32 OnPaint( const FPaintArgs& Args, const FTextLayout::FLineView& Line, const TSharedRef< ILayoutBlock >& Block, const FTextBlockStyle& DefaultStyle, const FGeometry& AllottedGeometry, const FSlateRect& MyClippingRect, FSlateWindowElementList& OutDrawElements, int32 LayerId, const FWidgetStyle& InWidgetStyle, bool bParentEnabled ) const override;
 
-	virtual FChildren* GetChildren() override;
+	virtual const TArray< TSharedRef<SWidget> >& GetChildren() override;
 
 	virtual void ArrangeChildren( const TSharedRef< ILayoutBlock >& Block, const FGeometry& AllottedGeometry, FArrangedChildren& ArrangedChildren ) const override;
 
@@ -92,7 +92,7 @@ private:
 	FOnClick NavigateDelegate;
 
 	TSharedRef< FWidgetViewModel > ViewModel;
-	TSlotlessChildren< SWidget > Children;
+	TArray< TSharedRef<SWidget> > Children;
 };
 
 #endif //WITH_FANCY_TEXT
\ No newline at end of file
diff --git a/Engine/Source/Runtime/Slate/Public/Framework/Text/SlateImageRun.h b/Engine/Source/Runtime/Slate/Public/Framework/Text/SlateImageRun.h
--- a/Engine/Source/Runtime/Slate/Public/Framework/Text/SlateImageRun.h
+++ b/Engine/Source/Runtime/Slate/Public/Framework/Text/SlateImageRun.h
@@ -29,7 +29,7 @@ public:
 
 	virtual int32 OnPaint( const FPaintArgs& Args, const FTextLayout::FLineView& Line, const TSharedRef< ILayoutBlock >& Block, const FTextBlockStyle& DefaultStyle, const FGeometry& AllottedGeometry, const FSlateRect& MyClippingRect, FSlateWindowElementList& OutDrawElements, int32 LayerId, const FWidgetStyle& InWidgetStyle, bool bParentEnabled ) const override;
 
-	virtual FChildren* GetChildren() override;
+	virtual const TArray< TSharedRef<SWidget> >& GetChildren() override;
 
 	virtual void ArrangeChildren( const TSharedRef< ILayoutBlock >& Block, const FGeometry& AllottedGeometry, FArrangedChildren& ArrangedChildren ) const override;
 
@@ -58,8 +58,6 @@ private:
 	FTextRange Range;
 	const FSlateBrush* Image;
 	int16 Baseline;
-
-	static FNoChildren NoChildrenInstance;
 };
 
 #endif //WITH_FANCY_TEXT
\ No newline at end of file
diff --git a/Engine/Source/Runtime/Slate/Public/Framework/Text/SlateTextRun.h b/Engine/Source/Runtime/Slate/Public/Framework/Text/SlateTextRun.h
--- a/Engine/Source/Runtime/Slate/Public/Framework/Text/SlateTextRun.h
+++ b/Engine/Source/Runtime/Slate/Public/Framework/Text/SlateTextRun.h
@@ -27,7 +27,7 @@ public:
 
 	virtual int32 OnPaint( const FPaintArgs& Args, const FTextLayout::FLineView& Line, const TSharedRef< ILayoutBlock >& Block, const FTextBlockStyle& DefaultStyle, const FGeometry& AllottedGeometry, const FSlateRect& MyClippingRect, FSlateWindowElementList& OutDrawElements, int32 LayerId, const FWidgetStyle& InWidgetStyle, bool bParentEnabled ) const override;
 
-	virtual FChildren* GetChildren() override;
+	virtual const TArray< TSharedRef<SWidget> >& GetChildren() override;
 
 	virtual void ArrangeChildren( const TSharedRef< ILayoutBlock >& Block, const FGeometry& AllottedGeometry, FArrangedChildren& ArrangedChildren ) const override;
 
@@ -61,8 +61,6 @@ private:
 #if TEXT_LAYOUT_DEBUG
 	FString DebugSlice;
 #endif
-
-	static FNoChildren NoChildrenInstance;
 };
 
 #endif //WITH_FANCY_TEXT
\ No newline at end of file
diff --git a/Engine/Source/Runtime/Slate/Public/Framework/Text/SlateWidgetRun.h b/Engine/Source/Runtime/Slate/Public/Framework/Text/SlateWidgetRun.h
--- a/Engine/Source/Runtime/Slate/Public/Framework/Text/SlateWidgetRun.h
+++ b/Engine/Source/Runtime/Slate/Public/Framework/Text/SlateWidgetRun.h
@@ -49,7 +49,7 @@ public:
 
 	virtual int32 OnPaint( const FPaintArgs& Args, const FTextLayout::FLineView& Line, const TSharedRef< ILayoutBlock >& Block, const FTextBlockStyle& DefaultStyle, const FGeometry& AllottedGeometry, const FSlateRect& MyClippingRect, FSlateWindowElementList& OutDrawElements, int32 LayerId, const FWidgetStyle& InWidgetStyle, bool bParentEnabled ) const override;
 
-	virtual FChildren* GetChildren() override;
+	virtual const TArray< TSharedRef<SWidget> >& GetChildren() override;
 
 	virtual void ArrangeChildren( const TSharedRef< ILayoutBlock >& Block, const FGeometry& AllottedGeometry, FArrangedChildren& ArrangedChildren ) const;
 
@@ -79,6 +79,6 @@ private:
 	TSharedRef< const FString > Text;
 	FTextRange Range;
 	FWidgetRunInfo Info;
-	TSlotlessChildren< SWidget > Children;
+	TArray< TSharedRef<SWidget> > Children;
 };
 #endif //WITH_FANCY_TEXT
\ No newline at end of file
diff --git a/Engine/Source/Runtime/Slate/Public/Widgets/Colors/SColorThemes.h b/Engine/Source/Runtime/Slate/Public/Widgets/Colors/SColorThemes.h
--- a/Engine/Source/Runtime/Slate/Public/Widgets/Colors/SColorThemes.h
+++ b/Engine/Source/Runtime/Slate/Public/Widgets/Colors/SColorThemes.h
@@ -238,6 +238,8 @@ public:
 
 	SLATE_END_ARGS()
 
+	SThemeColorBlocksBar();
+
 	/**
 	 * Most panels do not create widgets as part of their implementation, so
 	 * they do not need to implement a Construct()
diff --git a/Engine/Source/Runtime/Slate/Public/Widgets/Docking/SDockTab.h b/Engine/Source/Runtime/Slate/Public/Widgets/Docking/SDockTab.h
--- a/Engine/Source/Runtime/Slate/Public/Widgets/Docking/SDockTab.h
+++ b/Engine/Source/Runtime/Slate/Public/Widgets/Docking/SDockTab.h
@@ -153,7 +153,7 @@ public:
 	TSharedPtr<SWindow> GetParentWindow() const;
 
 	/** The width that this tab will overlap with side-by-side tabs. */
-	float GetOverlapWidth();
+	float GetOverlapWidth() const;
 
 	/** The label on the tab */
 	FText GetTabLabel() const;
diff --git a/Engine/Source/Runtime/Slate/Public/Widgets/Layout/SBox.h b/Engine/Source/Runtime/Slate/Public/Widgets/Layout/SBox.h
--- a/Engine/Source/Runtime/Slate/Public/Widgets/Layout/SBox.h
+++ b/Engine/Source/Runtime/Slate/Public/Widgets/Layout/SBox.h
@@ -9,11 +9,12 @@
 class SLATE_API SBox : public SPanel
 {
 	public:
-		class FBoxSlot : public TSupportsContentAlignmentMixin<FBoxSlot>, public TSupportsContentPaddingMixin<FBoxSlot>, public TSupportsOneChildMixin<SWidget, FBoxSlot>
+		class FBoxSlot : public TSupportsOneChildMixin<FBoxSlot>, public TSupportsContentAlignmentMixin<FBoxSlot>, public TSupportsContentPaddingMixin<FBoxSlot>
 		{
 			public:
 				FBoxSlot()
-				: TSupportsContentAlignmentMixin<FBoxSlot>(HAlign_Fill, VAlign_Fill)
+				: TSupportsOneChildMixin<FBoxSlot>()
+				, TSupportsContentAlignmentMixin<FBoxSlot>(HAlign_Fill, VAlign_Fill)
 				{
 				}
 		};
@@ -49,6 +50,8 @@ class SLATE_API SBox : public SPanel
 
 		SLATE_END_ARGS()
 
+		SBox();
+
 		void Construct( const FArguments& InArgs );
 
 
@@ -58,6 +61,7 @@ class SLATE_API SBox : public SPanel
 		virtual int32 OnPaint( const FPaintArgs& Args, const FGeometry& AllottedGeometry, const FSlateRect& MyClippingRect, FSlateWindowElementList& OutDrawElements, int32 LayerId, const FWidgetStyle& InWidgetStyle, bool bParentEnabled ) const override;
 
 private:
+
 		FBoxSlot ChildSlot;
 
 		/** When specified, ignore the content's desired size and report the.WidthOverride as the Box's desired width. */
diff --git a/Engine/Source/Runtime/Slate/Public/Widgets/Layout/SDPIScaler.h b/Engine/Source/Runtime/Slate/Public/Widgets/Layout/SDPIScaler.h
--- a/Engine/Source/Runtime/Slate/Public/Widgets/Layout/SDPIScaler.h
+++ b/Engine/Source/Runtime/Slate/Public/Widgets/Layout/SDPIScaler.h
@@ -26,6 +26,8 @@ public:
 
 	SLATE_END_ARGS()
 
+	SDPIScaler();
+
 	void Construct( const FArguments& InArgs );
 
 	virtual void OnArrangeChildren( const FGeometry& AllottedGeometry, FArrangedChildren& ArrangedChildren ) const override;
diff --git a/Engine/Source/Runtime/Slate/Public/Widgets/Layout/SGridPanel.h b/Engine/Source/Runtime/Slate/Public/Widgets/Layout/SGridPanel.h
--- a/Engine/Source/Runtime/Slate/Public/Widgets/Layout/SGridPanel.h
+++ b/Engine/Source/Runtime/Slate/Public/Widgets/Layout/SGridPanel.h
@@ -18,12 +18,13 @@ public:
 		int TheLayer;
 	};
 
-	class FSlot : public TSupportsOneChildMixin<SWidget, FSlot>, public TSupportsContentAlignmentMixin<FSlot>, public TSupportsContentPaddingMixin<FSlot>
+	class FSlot : public TSlotBase<FSlot>, public TSupportsContentAlignmentMixin<FSlot>, public TSupportsContentPaddingMixin<FSlot>
 	{
 		public:
 			/** Default values for a slot. */
 			FSlot( int32 Column, int32 Row, int32 InLayer )
-				: TSupportsContentAlignmentMixin<FSlot>(HAlign_Fill, VAlign_Fill)
+				: TSlotBase<FSlot>()
+				, TSupportsContentAlignmentMixin<FSlot>(HAlign_Fill, VAlign_Fill)
 				, ColumnParam( Column )
 				, ColumnSpanParam( 1 )
 				, RowParam( Row )
@@ -141,6 +142,8 @@ public:
 		
 	SLATE_END_ARGS()
 
+	SGridPanel();
+
 	/** Removes all slots from the panel */
 	void ClearChildren();
 	
diff --git a/Engine/Source/Runtime/Slate/Public/Widgets/Layout/SScrollBox.h b/Engine/Source/Runtime/Slate/Public/Widgets/Layout/SScrollBox.h
--- a/Engine/Source/Runtime/Slate/Public/Widgets/Layout/SScrollBox.h
+++ b/Engine/Source/Runtime/Slate/Public/Widgets/Layout/SScrollBox.h
@@ -7,11 +7,12 @@ class SLATE_API SScrollBox : public SCompoundWidget
 {
 public:
 	/** A Slot that provides layout options for the contents of a scrollable box. */
-	class SLATE_API FSlot : public TSupportsOneChildMixin<SWidget, FSlot>, public TSupportsContentPaddingMixin<FSlot>
+	class SLATE_API FSlot : public TSlotBase<FSlot>, public TSupportsContentPaddingMixin<FSlot>
 	{
 	public:
 		FSlot()
-			: HAlignment(HAlign_Fill)
+			: TSlotBase<FSlot>()
+			, HAlignment(HAlign_Fill)
 		{
 
 		}
diff --git a/Engine/Source/Runtime/Slate/Public/Widgets/Layout/SSplitter.h b/Engine/Source/Runtime/Slate/Public/Widgets/Layout/SSplitter.h
--- a/Engine/Source/Runtime/Slate/Public/Widgets/Layout/SSplitter.h
+++ b/Engine/Source/Runtime/Slate/Public/Widgets/Layout/SSplitter.h
@@ -34,22 +34,16 @@ public:
 		float );
 
 public:
-	class FSlot
+	class FSlot : public TSlotBase<FSlot>
 	{
 	public:		
 		FSlot()
-		: SizingRule( FractionOfParent )
-		, SizeValue( 1 )
-		, Widget( SNullWidget::NullWidget )
+			: TSlotBase<FSlot>()
+			, SizingRule( FractionOfParent )
+			, SizeValue( 1 )
 		{
 		}
 
-		FSlot& operator[]( TSharedRef<SWidget> InWidget )
-		{
-			Widget = InWidget;
-			return *this;
-		}
-
 		FSlot& Value( const TAttribute<float>& InValue )
 		{
 			SizeValue = InValue;
@@ -68,15 +62,8 @@ public:
 			return *this;
 		}
 
-		FSlot& Expose( FSlot*& OutVarToInit )
-		{
-			OutVarToInit = this;
-			return *this;
-		}
-
 		TAttribute<ESizeRule> SizingRule;
 		TAttribute<float> SizeValue;
-		TSharedRef<SWidget> Widget;
 		FOnSlotResized OnSlotResized_Handler;
 	};
 
@@ -290,28 +277,21 @@ protected:
 class SLATE_API SSplitter2x2 : public SPanel
 {
 public:
-	class FSlot
+	class FSlot : public TSlotBase<FSlot>
 	{
 	public:	
 		/** Default Constructor.  Initially each slot takes up a quarter of the entire space */
 		FSlot()
-			: PercentageAttribute( FVector2D(0.5, 0.5) )
-			, Widget( SNullWidget::NullWidget )
+			: TSlotBase<FSlot>( SNullWidget::NullWidget )
+			, PercentageAttribute( FVector2D(0.5, 0.5) )
 		{
 		}
 
 		/** Copy Constructor */
 		FSlot( const TSharedRef<SWidget>& InWidget )
-			: PercentageAttribute( FVector2D(0.5, 0.5) )
-			, Widget( InWidget )
-		{
-		}
-
-		/** Declarative syntax support */
-		FSlot& operator[]( TSharedRef<SWidget> InWidget )
+			: TSlotBase<FSlot>( InWidget )
+			, PercentageAttribute( FVector2D(0.5, 0.5) )
 		{
-			Widget = InWidget;
-			return *this;
 		}
 
 		/**
@@ -327,8 +307,6 @@ public:
 	public:
 		/** The percentage of the alloted space of the splitter that this slot requires */
 		TAttribute<FVector2D> PercentageAttribute;
-		/** The widget in this slot */
-		TSharedRef<SWidget> Widget;
 	};
 
 	SLATE_BEGIN_ARGS( SSplitter2x2 ){}
@@ -338,6 +316,7 @@ public:
 		SLATE_NAMED_SLOT( FArguments, BottomRight )
 	SLATE_END_ARGS()
 
+	SSplitter2x2();
 
 	void Construct( const FArguments& InArgs );
 
diff --git a/Engine/Source/Runtime/Slate/Public/Widgets/Layout/SUniformGridPanel.h b/Engine/Source/Runtime/Slate/Public/Widgets/Layout/SUniformGridPanel.h
--- a/Engine/Source/Runtime/Slate/Public/Widgets/Layout/SUniformGridPanel.h
+++ b/Engine/Source/Runtime/Slate/Public/Widgets/Layout/SUniformGridPanel.h
@@ -7,10 +7,11 @@ class SLATE_API SUniformGridPanel : public SPanel
 {
 public:
 	/** Stores the per-child info for this panel type */
-	struct FSlot : public TSupportsOneChildMixin<SWidget, FSlot>, public TSupportsContentAlignmentMixin<FSlot>
+	struct FSlot : public TSlotBase<FSlot>, public TSupportsContentAlignmentMixin<FSlot>
 	{
 		FSlot( int32 InColumn, int32 InRow )
-		: TSupportsContentAlignmentMixin<FSlot>(HAlign_Fill, VAlign_Fill)
+		: TSlotBase<FSlot>()
+		, TSupportsContentAlignmentMixin<FSlot>(HAlign_Fill, VAlign_Fill)
 		, Column( InColumn )
 		, Row( InRow )
 		{
@@ -20,6 +21,8 @@ public:
 		int32 Row;
 	};
 
+	SUniformGridPanel();
+
 	/**
 	 * Used by declarative syntax to create a Slot in the specified Column, Row.
 	 */
diff --git a/Engine/Source/Runtime/Slate/Public/Widgets/Layout/SWidgetSwitcher.h b/Engine/Source/Runtime/Slate/Public/Widgets/Layout/SWidgetSwitcher.h
--- a/Engine/Source/Runtime/Slate/Public/Widgets/Layout/SWidgetSwitcher.h
+++ b/Engine/Source/Runtime/Slate/Public/Widgets/Layout/SWidgetSwitcher.h
@@ -17,19 +17,30 @@ class SLATE_API SWidgetSwitcher
 {
 public:
 
+	class FSlot : public TSlotBase<FSlot>, public TSupportsContentAlignmentMixin<FSlot>, public TSupportsContentPaddingMixin<FSlot>
+	{
+	public:
+		FSlot()
+			: TSlotBase<FSlot>()
+			,TSupportsContentAlignmentMixin<FSlot>( HAlign_Fill, VAlign_Fill )
+		{}
+	};
+
 	SLATE_BEGIN_ARGS(SWidgetSwitcher)
 		: _WidgetIndex(0)
 		{
 			_Visibility = EVisibility::SelfHitTestInvisible;
 		}
 
-		SLATE_SUPPORTS_SLOT(FSimpleSlot)
+		SLATE_SUPPORTS_SLOT(FSlot)
 
 		/** Holds the index of the initial widget to be displayed (INDEX_NONE = default). */
 		SLATE_ATTRIBUTE(int32, WidgetIndex)
 
 	SLATE_END_ARGS()
 
+	SWidgetSwitcher();
+
 public:
 
 	/**
@@ -37,7 +48,7 @@ public:
 	 *
 	 * @param SlotIndex - The index at which to insert the slot, or INDEX_NONE to append.
 	 */
-	FSimpleSlot& AddSlot( int32 SlotIndex = INDEX_NONE );
+	FSlot& AddSlot( int32 SlotIndex = INDEX_NONE );
 
 	/**
 	 * Removes a slot with the corresponding widget in it.  Returns the index where the widget was found, otherwise -1.
@@ -118,9 +129,9 @@ public:
 	 *
 	 * @return A new slot.
 	 */
-	static FSimpleSlot& Slot( )
+	static FSlot& Slot( )
 	{
-		return *(new FSimpleSlot());
+		return *(new FSlot());
 	}
 
 protected:
@@ -143,20 +154,20 @@ private:
 	{
 	public:
 
-		FOneDynamicChild( TPanelChildren<FSimpleSlot>* InAllChildren = NULL, const TAttribute<int32>* InWidgetIndex = NULL )
+		FOneDynamicChild( TPanelChildren<FSlot>* InAllChildren = NULL, const TAttribute<int32>* InWidgetIndex = NULL )
 			: AllChildren( InAllChildren )
 			, WidgetIndex( InWidgetIndex )
 		{ }
 		
 		virtual int32 Num() const override { return AllChildren->Num() > 0 ? 1 : 0; }
 		
-		virtual TSharedRef<SWidget> GetChildAt( int32 Index ) override { check(Index == 0); return (*AllChildren)[WidgetIndex->Get()].Widget; }
+		virtual TSharedRef<SWidget> GetChildAt( int32 Index ) override { check(Index == 0); return AllChildren->GetChildAt(WidgetIndex->Get()); }
 		
-		virtual TSharedRef<const SWidget> GetChildAt( int32 Index ) const override { check(Index == 0); return (*AllChildren)[WidgetIndex->Get()].Widget; }
+		virtual TSharedRef<const SWidget> GetChildAt( int32 Index ) const override { check(Index == 0); return AllChildren->GetChildAt(WidgetIndex->Get()); }
 		
 	private:
 
-		TPanelChildren<FSimpleSlot>* AllChildren;
+		TPanelChildren<FSlot>* AllChildren;
 		const TAttribute<int32>* WidgetIndex;
 
 	} OneDynamicChild;
@@ -165,5 +176,5 @@ private:
 	TAttribute<int32> WidgetIndex;
 
 	// Holds the collection of widgets.
-	TPanelChildren<FSimpleSlot> AllChildren;
+	TPanelChildren<FSlot> AllChildren;
 };
diff --git a/Engine/Source/Runtime/Slate/Public/Widgets/Layout/SWrapBox.h b/Engine/Source/Runtime/Slate/Public/Widgets/Layout/SWrapBox.h
--- a/Engine/Source/Runtime/Slate/Public/Widgets/Layout/SWrapBox.h
+++ b/Engine/Source/Runtime/Slate/Public/Widgets/Layout/SWrapBox.h
@@ -3,11 +3,12 @@
 #pragma once
 
 /** A slot that support alignment of content and padding */
-class SLATE_API FWrapSlot : public TSupportsOneChildMixin<SWidget, FWrapSlot>, public TSupportsContentAlignmentMixin<FWrapSlot>, public TSupportsContentPaddingMixin<FWrapSlot>
+class SLATE_API FWrapSlot : public TSlotBase<FWrapSlot>, public TSupportsContentAlignmentMixin<FWrapSlot>, public TSupportsContentPaddingMixin<FWrapSlot>
 {
 public:
 	FWrapSlot()
-		: TSupportsContentAlignmentMixin<FWrapSlot>( HAlign_Fill, VAlign_Fill )
+		: TSlotBase<FWrapSlot>()
+		, TSupportsContentAlignmentMixin<FWrapSlot>( HAlign_Fill, VAlign_Fill )
 		, bSlotFillEmptySpace( false )
 	{
 	}
@@ -69,6 +70,8 @@ public:
 		SLATE_ARGUMENT( bool, UseAllottedWidth )
 	SLATE_END_ARGS()
 
+	SWrapBox();
+
 	static FWrapSlot& Slot();
 
 	FWrapSlot& AddSlot();
diff --git a/Engine/Source/Runtime/Slate/Public/Widgets/SCanvas.h b/Engine/Source/Runtime/Slate/Public/Widgets/SCanvas.h
--- a/Engine/Source/Runtime/Slate/Public/Widgets/SCanvas.h
+++ b/Engine/Source/Runtime/Slate/Public/Widgets/SCanvas.h
@@ -18,7 +18,7 @@ public:
 	/**
 	 * Canvas slots allow child widgets to be positioned and sized
 	 */
-	class FSlot
+	class FSlot : public TSlotBase<FSlot>
 	{
 	public:		
 		FSlot& Position( const TAttribute<FVector2D>& InPosition )
@@ -33,18 +33,6 @@ public:
 			return *this;
 		}
 
-		FSlot& operator[]( TSharedRef<SWidget> InWidget )
-		{
-			Widget = InWidget;
-			return *this;
-		}
-
-		FSlot& Expose( FSlot*& OutVarToInit )
-		{
-			OutVarToInit = this;
-			return *this;
-		}
-
 		FSlot& HAlign( EHorizontalAlignment InHAlignment )
 		{
 			HAlignment = InHAlignment;
@@ -57,9 +45,6 @@ public:
 			return *this;
 		}
 
-		/** The child widget contained in this slot. */
-		TSharedRef<SWidget> Widget;
-
 		/** Position */
 		TAttribute<FVector2D> PositionAttr;
 
@@ -96,7 +81,7 @@ public:
 
 		/** Default values for a slot. */
 		FSlot()
-			: Widget( SNullWidget::NullWidget )
+			: TSlotBase<FSlot>()
 			, PositionAttr( FVector2D::ZeroVector )
 			, SizeAttr( FVector2D( 1.0f, 1.0f ) )
 			, HAlignment(HAlign_Left)
@@ -113,6 +98,8 @@ public:
 
 	SLATE_END_ARGS()
 
+	SCanvas();
+
 	/**
 	 * Construct this widget
 	 *
@@ -132,7 +119,7 @@ public:
 	 */
 	FSlot& AddSlot()
 	{
-		SCanvas::FSlot& NewSlot = SCanvas::Slot();
+		SCanvas::FSlot& NewSlot = *new FSlot();
 		this->Children.Add( &NewSlot );
 		return NewSlot;
 	}
diff --git a/Engine/Source/Runtime/Slate/Public/Widgets/SViewport.h b/Engine/Source/Runtime/Slate/Public/Widgets/SViewport.h
--- a/Engine/Source/Runtime/Slate/Public/Widgets/SViewport.h
+++ b/Engine/Source/Runtime/Slate/Public/Widgets/SViewport.h
@@ -102,7 +102,7 @@ public:
 		];
 	}
 
-	const TSharedPtr<SWidget> GetContent() const { return ChildSlot.Widget; }
+	const TSharedPtr<SWidget> GetContent() const { return ChildSlot.GetWidget(); }
 
 	/**
 	 * A delegate called when the viewports top level window is being closed
diff --git a/Engine/Source/Runtime/Slate/Public/Widgets/SWeakWidget.h b/Engine/Source/Runtime/Slate/Public/Widgets/SWeakWidget.h
--- a/Engine/Source/Runtime/Slate/Public/Widgets/SWeakWidget.h
+++ b/Engine/Source/Runtime/Slate/Public/Widgets/SWeakWidget.h
@@ -1,9 +1,5 @@
 // Copyright 1998-2014 Epic Games, Inc. All Rights Reserved.
 
-/*=============================================================================
-	SWeakWidget.h: Declares the SWeakWidget class.
-=============================================================================*/
-
 #pragma once
 
 
@@ -15,8 +11,7 @@
  *      by a floating window. That window cannot own the tooltip
  *      and must therefore use an SWeakWidget.
  */
-class SLATE_API SWeakWidget
-	: public SWidget
+class SLATE_API SWeakWidget : public SWidget
 {
 public:
 
@@ -28,10 +23,11 @@ public:
 	SLATE_END_ARGS()
 
 public:
+	SWeakWidget();
 
 	void Construct(const FArguments& InArgs);
 
-	void SetContent(TWeakPtr<SWidget> InWidget);
+	void SetContent(const TSharedRef<SWidget>& InWidget);
 
 	bool ChildWidgetIsValid() const;
 
diff --git a/Engine/Source/Runtime/SlateCore/Private/Children.cpp b/Engine/Source/Runtime/SlateCore/Private/Children.cpp
new file mode 100644
--- /dev/null
+++ b/Engine/Source/Runtime/SlateCore/Private/Children.cpp
@@ -0,0 +1,4 @@
+// Copyright 1998-2014 Epic Games, Inc. All Rights Reserved.
+
+#include "SlateCorePrivatePCH.h"
+
diff --git a/Engine/Source/Runtime/SlateCore/Private/SlotBase.cpp b/Engine/Source/Runtime/SlateCore/Private/SlotBase.cpp
new file mode 100644
--- /dev/null
+++ b/Engine/Source/Runtime/SlateCore/Private/SlotBase.cpp
@@ -0,0 +1,45 @@
+// Copyright 1998-2014 Epic Games, Inc. All Rights Reserved.
+
+#include "SlateCorePrivatePCH.h"
+
+FSlotBase::FSlotBase()
+: Widget( SNullWidget::NullWidget )
+{
+
+}
+
+FSlotBase::FSlotBase( const TSharedRef<SWidget>& InWidget )
+: Widget( InWidget )
+{
+	
+}
+
+void FSlotBase::AttachWidget( const TSharedRef<SWidget>& InWidget )
+{
+	Widget = InWidget;
+}
+
+
+const TSharedRef<SWidget>& FSlotBase::GetWidget() const
+{
+	return Widget;
+}
+
+const TSharedPtr<SWidget> FSlotBase::DetachWidget()
+{
+	if (Widget != SNullWidget::NullWidget)
+	{
+		const TSharedRef<SWidget> MyExWidget = Widget;
+		Widget = SNullWidget::NullWidget;	
+		return MyExWidget;
+	}
+	else
+	{
+		// Nothing to detach!
+		return TSharedPtr<SWidget>();
+	}
+}
+
+FSlotBase::~FSlotBase()
+{
+}
\ No newline at end of file
diff --git a/Engine/Source/Runtime/SlateCore/Private/Widgets/SBoxPanel.cpp b/Engine/Source/Runtime/SlateCore/Private/Widgets/SBoxPanel.cpp
--- a/Engine/Source/Runtime/SlateCore/Private/Widgets/SBoxPanel.cpp
+++ b/Engine/Source/Runtime/SlateCore/Private/Widgets/SBoxPanel.cpp
@@ -50,7 +50,7 @@ static void ArrangeChildrenAlong( const TPanelChildren<SBoxPanel::FSlot>& Childr
 		{
 			const SBoxPanel::FSlot& CurChild = Children[ChildIndex];
 
-			if ( CurChild.Widget->GetVisibility() != EVisibility::Collapsed )
+			if ( CurChild.GetWidget()->GetVisibility() != EVisibility::Collapsed )
 			{
 				// All widgets contribute their margin to the fixed space requirement
 				FixedTotal += CurChild.SlotPadding.Get().GetTotalSpaceAlong<Orientation>();
@@ -77,8 +77,8 @@ static void ArrangeChildrenAlong( const TPanelChildren<SBoxPanel::FSlot>& Childr
 				{
 					// Auto-sized children contribute their desired size to the fixed space requirement
 					float ChildSize = (Orientation == Orient_Vertical)
-					? CurChild.Widget->GetDesiredSize().Y
-					: CurChild.Widget->GetDesiredSize().X;
+					? CurChild.GetWidget()->GetDesiredSize().Y
+					: CurChild.GetWidget()->GetDesiredSize().X;
 
 					// Clamp to the max size if it was specified
 					float MaxSize = CurChild.MaxSize.Get();
@@ -100,9 +100,9 @@ static void ArrangeChildrenAlong( const TPanelChildren<SBoxPanel::FSlot>& Childr
 		for( int32 ChildIndex=0; ChildIndex < Children.Num(); ++ChildIndex )
 		{
 			const SBoxPanel::FSlot& CurChild = Children[ChildIndex];
-			const EVisibility ChildVisibility = CurChild.Widget->GetVisibility();
+			const EVisibility ChildVisibility = CurChild.GetWidget()->GetVisibility();
 
-			FVector2D ChildDesiredSize = CurChild.Widget->GetDesiredSize();
+			FVector2D ChildDesiredSize = CurChild.GetWidget()->GetDesiredSize();
 			
 			// Figure out the area allocated to the child in the direction of BoxPanel
 			// The area allocated to the slot is ChildSize + the associated margin.
@@ -152,7 +152,7 @@ static void ArrangeChildrenAlong( const TPanelChildren<SBoxPanel::FSlot>& Childr
 			// Add the information about this child to the output list (ArrangedChildren)
 			ArrangedChildren.AddWidget( ChildVisibility, AllottedGeometry.MakeChild(
 				// The child widget being arranged
-				CurChild.Widget,
+				CurChild.GetWidget(),
 				// Child's local position (i.e. position within parent)
 				(Orientation == Orient_Vertical)
 					? FVector2D( XAlignmentResult.Offset, PositionSoFar + YAlignmentResult.Offset )
@@ -208,8 +208,8 @@ static FVector2D ComputeDesiredSizeForBox( const TPanelChildren<SBoxPanel::FSlot
 	for( int32 ChildIndex=0; ChildIndex < Children.Num(); ++ChildIndex )
 	{
 		const SBoxPanel::FSlot& CurChild = Children[ChildIndex];
-		const FVector2D& CurChildDesiredSize = CurChild.Widget->GetDesiredSize();
-		if ( CurChild.Widget->GetVisibility() != EVisibility::Collapsed )
+		const FVector2D& CurChildDesiredSize = CurChild.GetWidget()->GetDesiredSize();
+		if ( CurChild.GetWidget()->GetVisibility() != EVisibility::Collapsed )
 		{
 			if (Orientation == Orient_Vertical)
 			{
@@ -300,7 +300,7 @@ int32 SBoxPanel::RemoveSlot( const TSharedRef<SWidget>& SlotWidget )
 {
 	for (int32 SlotIdx = 0; SlotIdx < Children.Num(); ++SlotIdx)
 	{
-		if ( SlotWidget == Children[SlotIdx].Widget )
+		if ( SlotWidget == Children[SlotIdx].GetWidget() )
 		{
 			Children.RemoveAt(SlotIdx);
 			return SlotIdx;
@@ -322,6 +322,7 @@ void SBoxPanel::ClearChildren()
  */
 SBoxPanel::SBoxPanel( EOrientation InOrientation )
 : Orientation(InOrientation)
+, Children()
 {
 
 }
diff --git a/Engine/Source/Runtime/SlateCore/Private/Widgets/SCompoundWidget.cpp b/Engine/Source/Runtime/SlateCore/Private/Widgets/SCompoundWidget.cpp
--- a/Engine/Source/Runtime/SlateCore/Private/Widgets/SCompoundWidget.cpp
+++ b/Engine/Source/Runtime/SlateCore/Private/Widgets/SCompoundWidget.cpp
@@ -43,10 +43,10 @@ FChildren* SCompoundWidget::GetChildren()
 
 FVector2D SCompoundWidget::ComputeDesiredSize() const
 {
-	EVisibility ChildVisibility = ChildSlot.Widget->GetVisibility();
+	EVisibility ChildVisibility = ChildSlot.GetWidget()->GetVisibility();
 	if ( ChildVisibility != EVisibility::Collapsed )
 	{
-		return ChildSlot.Widget->GetDesiredSize() + ChildSlot.SlotPadding.Get().GetDesiredSize();
+		return ChildSlot.GetWidget()->GetDesiredSize() + ChildSlot.SlotPadding.Get().GetDesiredSize();
 	}
 	
 	return FVector2D::ZeroVector;
@@ -60,4 +60,12 @@ void SCompoundWidget::OnArrangeChildren( const FGeometry& AllottedGeometry, FArr
 FSlateColor SCompoundWidget::GetForegroundColor() const
 {
 	return ForegroundColor.Get();
+}
+
+SCompoundWidget::SCompoundWidget()
+	: ChildSlot()
+	, ContentScale( FVector2D(1.0f,1.0f) )
+	, ColorAndOpacity( FLinearColor::White )
+	, ForegroundColor( FSlateColor::UseForeground() )
+{
 }
\ No newline at end of file
diff --git a/Engine/Source/Runtime/SlateCore/Private/Widgets/SOverlay.cpp b/Engine/Source/Runtime/SlateCore/Private/Widgets/SOverlay.cpp
--- a/Engine/Source/Runtime/SlateCore/Private/Widgets/SOverlay.cpp
+++ b/Engine/Source/Runtime/SlateCore/Private/Widgets/SOverlay.cpp
@@ -2,6 +2,12 @@
 
  #include "SlateCorePrivatePCH.h"
 
+#include "LayoutUtils.h"
+
+SOverlay::SOverlay()
+: Children()
+{
+}
 
 /**
  * Construct this widget
@@ -23,7 +29,7 @@ void SOverlay::OnArrangeChildren( const FGeometry& AllottedGeometry, FArrangedCh
 	for ( int32 ChildIndex = 0; ChildIndex < Children.Num(); ++ChildIndex )
 	{
 		const FOverlaySlot& CurChild = Children[ChildIndex];
-		const EVisibility ChildVisibility = CurChild.Widget->GetVisibility();
+		const EVisibility ChildVisibility = CurChild.GetWidget()->GetVisibility();
 		if ( ArrangedChildren.Accepts(ChildVisibility) )
 		{
 			const FMargin SlotPadding(CurChild.SlotPadding.Get());
@@ -31,7 +37,7 @@ void SOverlay::OnArrangeChildren( const FGeometry& AllottedGeometry, FArrangedCh
 			AlignmentArrangeResult YResult = AlignChild<Orient_Vertical>(AllottedGeometry.Size.Y, CurChild, SlotPadding);
 
 			ArrangedChildren.AddWidget( ChildVisibility, AllottedGeometry.MakeChild(
-				CurChild.Widget,
+				CurChild.GetWidget(),
 				FVector2D(XResult.Offset,YResult.Offset),
 				FVector2D(XResult.Size, YResult.Size)
 			) );
@@ -52,10 +58,10 @@ FVector2D SOverlay::ComputeDesiredSize() const
 	for ( int32 ChildIndex=0; ChildIndex < Children.Num(); ++ChildIndex )
 	{
 		const FOverlaySlot& CurSlot = Children[ChildIndex];
-		const EVisibility ChildVisibilty = CurSlot.Widget->GetVisibility();
+		const EVisibility ChildVisibilty = CurSlot.GetWidget()->GetVisibility();
 		if ( ChildVisibilty != EVisibility::Collapsed )
 		{
-			FVector2D ChildDesiredSize = CurSlot.Widget->GetDesiredSize() + CurSlot.SlotPadding.Get().GetDesiredSize();
+			FVector2D ChildDesiredSize = CurSlot.GetWidget()->GetDesiredSize() + CurSlot.SlotPadding.Get().GetDesiredSize();
 			MaxSize.X = FMath::Max( MaxSize.X, ChildDesiredSize.X );
 			MaxSize.Y = FMath::Max( MaxSize.Y, ChildDesiredSize.Y );
 		}
@@ -99,7 +105,7 @@ int32 SOverlay::OnPaint( const FPaintArgs& Args, const FGeometry& AllottedGeomet
 
 SOverlay::FOverlaySlot& SOverlay::AddSlot( int32 ZOrder )
 {
-	FOverlaySlot& NewSlot = SOverlay::Slot();
+	FOverlaySlot& NewSlot = *new FOverlaySlot();
 	if ( ZOrder == INDEX_NONE )
 	{
 		// No ZOrder was specified; just add to the end of the list.
@@ -177,7 +183,7 @@ void SOverlay::RemoveSlot( TSharedRef< SWidget > Widget )
 	for( int32 CurSlotIndex = 0; CurSlotIndex < Children.Num(); ++CurSlotIndex )
 	{
 		const FOverlaySlot& CurSlot = Children[ CurSlotIndex ];
-		if( CurSlot.Widget == Widget )
+		if( CurSlot.GetWidget() == Widget )
 		{
 			Children.RemoveAt( CurSlotIndex );
 			break;
diff --git a/Engine/Source/Runtime/SlateCore/Private/Widgets/SWidget.cpp b/Engine/Source/Runtime/SlateCore/Private/Widgets/SWidget.cpp
--- a/Engine/Source/Runtime/SlateCore/Private/Widgets/SWidget.cpp
+++ b/Engine/Source/Runtime/SlateCore/Private/Widgets/SWidget.cpp
@@ -61,6 +61,13 @@ FReply SWidget::OnKeyboardFocusReceived( const FGeometry& MyGeometry, const FKey
 	return FReply::Unhandled();
 }
 
+void SWidget::OnKeyboardFocusLost( const FKeyboardFocusEvent& InKeyboardFocusEvent )
+{
+}
+
+void SWidget::OnKeyboardFocusChanging( const FWeakWidgetPath& PreviousFocusPath, const FWidgetPath& NewWidgetPath )
+{
+}
 
 FReply SWidget::OnKeyChar( const FGeometry& MyGeometry, const FCharacterEvent& InCharacterEvent )
 {
@@ -231,6 +238,9 @@ EWindowZone::Type SWidget::GetWindowZoneOverride() const
 	return EWindowZone::Unspecified;
 }
 
+void SWidget::Tick( const FGeometry& AllottedGeometry, const double InCurrentTime, const float InDeltaTime )
+{
+}
 
 bool SWidget::OnHitTest( const FGeometry& MyGeometry, FVector2D InAbsoluteCursorPosition )
 {
@@ -249,7 +259,7 @@ void SWidget::TickWidgetsRecursively( const FGeometry& AllottedGeometry, const d
 	FArrangedChildren ArrangedChildren(EVisibility::All);
 	ArrangeChildren(AllottedGeometry, ArrangedChildren);
 
-	// Recurse!
+	// Recur!
 	for(int32 ChildIndex=0; ChildIndex < ArrangedChildren.Num(); ++ChildIndex)
 	{
 		FArrangedWidget& SomeChild = ArrangedChildren(ChildIndex);
@@ -472,6 +482,9 @@ TSharedPtr<IToolTip> SWidget::GetToolTip()
 	return ToolTip;
 }
 
+void SWidget::OnToolTipClosing()
+{
+}
 
 void SWidget::EnableToolTipForceField( const bool bEnableForceField )
 {
@@ -519,4 +532,4 @@ void SWidget::ArrangeChildren(const FGeometry& AllottedGeometry, FArrangedChildr
 bool SWidget::UseLegacyHittest()
 {
 	return bUseLegacyHittest != 0;
-}
\ No newline at end of file
+}
diff --git a/Engine/Source/Runtime/SlateCore/Private/Widgets/SWindow.cpp b/Engine/Source/Runtime/SlateCore/Private/Widgets/SWindow.cpp
--- a/Engine/Source/Runtime/SlateCore/Private/Widgets/SWindow.cpp
+++ b/Engine/Source/Runtime/SlateCore/Private/Widgets/SWindow.cpp
@@ -40,6 +40,10 @@ public:
 
 	SLATE_END_ARGS()
 
+	SPopupLayer()
+	: Children()
+	{}
+
 	void Construct( const FArguments& InArgs, const TSharedRef<SWindow>& InWindow )
 	{
 		OwnerWindow = InWindow;
@@ -60,7 +64,7 @@ public:
 	/** Add a slot to the ListPanel */
 	FPopupLayerSlot& AddSlot(int32 InsertAtIndex = INDEX_NONE)
 	{
-		FPopupLayerSlot& NewSlot = SPopupLayer::Slot();
+		FPopupLayerSlot& NewSlot = *new FPopupLayerSlot();
 		if (InsertAtIndex == INDEX_NONE)
 		{
 			this->Children.Add( &NewSlot );
@@ -78,7 +82,7 @@ public:
 		for( int32 CurSlotIndex = 0; CurSlotIndex < Children.Num(); ++CurSlotIndex )
 		{
 			const FPopupLayerSlot& CurSlot = Children[ CurSlotIndex ];
-			if( CurSlot.Widget == WidgetToRemove )
+			if( CurSlot.GetWidget() == WidgetToRemove )
 			{
 				Children.RemoveAt( CurSlotIndex );
 				return;
@@ -101,10 +105,10 @@ private:
 		for ( int32 ChildIndex = 0; ChildIndex < Children.Num(); ++ChildIndex )
 		{
 			const FPopupLayerSlot& CurChild = Children[ChildIndex];
-			const EVisibility ChildVisibility = CurChild.Widget->GetVisibility();
+			const EVisibility ChildVisibility = CurChild.GetWidget()->GetVisibility();
 			if ( ArrangedChildren.Accepts(ChildVisibility) )
 			{
-				const FVector2D WidgetDesiredSize = CurChild.Widget->GetDesiredSize();
+				const FVector2D WidgetDesiredSize = CurChild.GetWidget()->GetDesiredSize();
 				const float ChildScale = CurChild.Scale_Attribute.Get();
 				const bool bClampToWindow = CurChild.Clamp_Attribute.Get();
 				const FVector2D ChildSize = WidgetDesiredSize * ChildScale;
@@ -120,7 +124,7 @@ private:
 				// The position is explicitly in desktop pixels.
 				// The size and DPI scale come from the widget that is using
 				// this overlay to "punch" through the UI.
-				ArrangedChildren.AddWidget( ChildVisibility, FArrangedWidget( CurChild.Widget,
+				ArrangedChildren.AddWidget( ChildVisibility, FArrangedWidget( CurChild.GetWidget(),
 					FGeometry(
 						ChildLocalPosition / ChildScale,
 						AllottedGeometry.AbsolutePosition,
@@ -887,7 +891,7 @@ TSharedRef<const SWidget> SWindow::GetContent() const
 	}
 	else
 	{
-		return this->ContentSlot->Widget;
+		return this->ContentSlot->GetWidget();
 	}
 }
 
diff --git a/Engine/Source/Runtime/SlateCore/Public/Layout/Children.h b/Engine/Source/Runtime/SlateCore/Public/Layout/Children.h
--- a/Engine/Source/Runtime/SlateCore/Public/Layout/Children.h
+++ b/Engine/Source/Runtime/SlateCore/Public/Layout/Children.h
@@ -4,6 +4,7 @@
 
 class SWidget;
 
+
 /**
  * FChildren is an interface that must be implemented by all child containers.
  * It allows iteration over a list of any Widget's children regardless of how
@@ -34,9 +35,9 @@ protected:
 class SLATECORE_API FNoChildren : public FChildren
 {
 public:
-	virtual int32 Num() const { return 0; }
+	virtual int32 Num() const override { return 0; }
 	
-	virtual TSharedRef<SWidget> GetChildAt( int32 )
+	virtual TSharedRef<SWidget> GetChildAt( int32 ) override
 	{
 		// Nobody should be getting a child when there aren't any children.
 		// We expect this to crash!
@@ -44,43 +45,28 @@ public:
 		return TSharedPtr<SWidget>(nullptr).ToSharedRef();
 	}
 	
-	virtual TSharedRef<const SWidget> GetChildAt( int32 ) const
+	virtual TSharedRef<const SWidget> GetChildAt( int32 ) const override
 	{
 		// Nobody should be getting a child when there aren't any children.
 		// We expect this to crash!
 		check( false );
 		return TSharedPtr<const SWidget>(nullptr).ToSharedRef();
 	}
+
 };
 
 /**
  * Widgets that will only have one child can return an instance of FOneChild.
  */
-template <typename ChildType, typename MixedIntoType>
-class TSupportsOneChildMixin : public FChildren
+template <typename MixedIntoType>
+class TSupportsOneChildMixin : public FChildren, public TSlotBase<MixedIntoType>
 {
 public:
-	TSupportsOneChildMixin()
-		: Widget( SNullWidget::NullWidget )
-	{ }
+	TSupportsOneChildMixin():TSlotBase<MixedIntoType>(){}
 
 	virtual int32 Num() const { return 1; }
-	virtual TSharedRef<SWidget> GetChildAt( int32 ChildIndex ) { check(ChildIndex == 0); return Widget; }
-	virtual TSharedRef<const SWidget> GetChildAt( int32 ChildIndex ) const { check(ChildIndex == 0); return Widget; }
-	
-	MixedIntoType& operator[]( TSharedRef<SWidget> InWidget )
-	{
-		Widget = InWidget;
-		return *(static_cast<MixedIntoType*>(this));
-	}
-
-	MixedIntoType& Expose( MixedIntoType*& OutVarToInit )
-	{
-		OutVarToInit = static_cast<MixedIntoType*>(this);
-		return *(static_cast<MixedIntoType*>(this));
-	}
-
-	TSharedRef<ChildType> Widget;
+	virtual TSharedRef<SWidget> GetChildAt( int32 ChildIndex ) { check(ChildIndex == 0); return FSlotBase::GetWidget(); }
+	virtual TSharedRef<const SWidget> GetChildAt( int32 ChildIndex ) const { check(ChildIndex == 0); return FSlotBase::GetWidget(); }
 };
 
 /**
@@ -93,14 +79,27 @@ class TWeakChild : public FChildren
 {
 public:
 	TWeakChild()
-		: Widget( SNullWidget::NullWidget )
-	{ }
+	: WidgetPtr()
+	{
+	}
 
-	virtual int32 Num() const { return Widget.IsValid() ? 1 : 0 ; }
-	virtual TSharedRef<SWidget> GetChildAt( int32 ChildIndex ) { check(ChildIndex == 0); return Widget.Pin().ToSharedRef(); }
-	virtual TSharedRef<const SWidget> GetChildAt( int32 ChildIndex ) const { check(ChildIndex == 0); return Widget.Pin().ToSharedRef(); }
+	virtual int32 Num() const { return WidgetPtr.IsValid() ? 1 : 0 ; }
+	virtual TSharedRef<SWidget> GetChildAt( int32 ChildIndex ) { check(ChildIndex == 0); return WidgetPtr.Pin().ToSharedRef(); }
+	virtual TSharedRef<const SWidget> GetChildAt( int32 ChildIndex ) const { check(ChildIndex == 0); return WidgetPtr.Pin().ToSharedRef(); }
 	
-	TWeakPtr<ChildType> Widget;
+	void AttachWidget(const TSharedPtr<SWidget>& InWidget)
+	{
+		WidgetPtr = InWidget;
+	}
+
+	TSharedRef<SWidget> GetWidget() const
+	{
+		TSharedPtr<SWidget> Widget = WidgetPtr.Pin();
+		return (Widget.IsValid()) ? Widget.ToSharedRef() : SNullWidget::NullWidget;
+	}
+
+private:
+	TWeakPtr<ChildType> WidgetPtr;
 };
 
 template <typename MixedIntoType>
@@ -162,11 +161,12 @@ public:
 };
 
 /** A slot that support alignment of content and padding */
-class SLATECORE_API FSimpleSlot : public TSupportsOneChildMixin<SWidget, FSimpleSlot>, public TSupportsContentAlignmentMixin<FSimpleSlot>, public TSupportsContentPaddingMixin<FSimpleSlot>
+class SLATECORE_API FSimpleSlot : public TSupportsOneChildMixin<FSimpleSlot>, public TSupportsContentAlignmentMixin<FSimpleSlot>, public TSupportsContentPaddingMixin<FSimpleSlot>
 {
 public:
 	FSimpleSlot()
-	: TSupportsContentAlignmentMixin<FSimpleSlot>(HAlign_Fill, VAlign_Fill)
+	: TSupportsOneChildMixin<FSimpleSlot>()
+	, TSupportsContentAlignmentMixin<FSimpleSlot>(HAlign_Fill, VAlign_Fill)
 	{
 	}
 };
@@ -178,22 +178,71 @@ public:
  * ChildType must have a public member SWidget* Widget;
  */
 template<typename SlotType>
-class TPanelChildren : public FChildren, public TIndirectArray< SlotType >
+class TPanelChildren : public FChildren, private TIndirectArray< SlotType >
 {
 public:
-	virtual int32 Num() const
+	TPanelChildren()
+	{
+	}
+
+	virtual int32 Num() const override
 	{
 		return TIndirectArray<SlotType>::Num();
 	}
 
-	virtual TSharedRef<SWidget> GetChildAt( int32 Index )
+	virtual TSharedRef<SWidget> GetChildAt( int32 Index ) override
+	{
+		return (*this)[Index].GetWidget();
+	}
+
+	virtual TSharedRef<const SWidget> GetChildAt( int32 Index ) const override
 	{
-		return this->operator[](Index).Widget;
+		return (*this)[Index].GetWidget();
 	}
 
-	virtual TSharedRef<const SWidget> GetChildAt( int32 Index ) const
+	int32 Add( SlotType* Slot )
 	{
-		return this->operator[](Index).Widget;
+		return TIndirectArray< SlotType >::Add(Slot);
+	}
+
+	void RemoveAt( int32 Index )
+	{
+		TIndirectArray< SlotType >::RemoveAt(Index);
+	}
+
+	void Empty()
+	{
+		TIndirectArray< SlotType >::Empty();
+	}
+
+	void Insert(SlotType* Item, int32 Index)
+	{
+		TIndirectArray< SlotType >::Insert(Item, Index);
+	}
+
+	void Reserve( int32 NumToReserve )
+	{
+		TIndirectArray< SlotType >::Reserve(NumToReserve);
+	}
+
+	bool IsValidIndex( int32 Index ) const
+	{
+		return TIndirectArray< SlotType >::IsValidIndex( Index );
+	}
+
+	const SlotType& operator[](int32 Index) const { return TIndirectArray< SlotType >::operator[](Index); }
+
+	SlotType& operator[](int32 Index) { return TIndirectArray< SlotType >::operator[](Index ); }
+
+	template <class PREDICATE_CLASS>
+	void Sort( const PREDICATE_CLASS& Predicate )
+	{
+		TIndirectArray< SlotType >::Sort( Predicate );
+	}
+
+	void Swap( int32 IndexA, int32 IndexB )
+	{
+		TIndirectArray< SlotType >::Swap(IndexA, IndexB);
 	}
 };
 
@@ -205,22 +254,81 @@ public:
  * TSlotlessChildren should not be used for general-purpose widgets.
  */
 template<typename ChildType>
-class TSlotlessChildren : public FChildren, public TArray< TSharedRef<ChildType> >
+class TSlotlessChildren : public FChildren, private TArray< TSharedRef<ChildType> >
 {
 public:
-	/** @return the number of children */
-	virtual int32 Num() const
+	TSlotlessChildren()
+	{
+	}
+
+	virtual int32 Num() const override
 	{
 		return TArray< TSharedRef<ChildType> >::Num();
 	}
-	/** @return pointer to the Widget at the specified Index. */
-	virtual TSharedRef<SWidget> GetChildAt( int32 Index )
+
+	virtual TSharedRef<SWidget> GetChildAt( int32 Index ) override
 	{
-		return this->operator[](Index);
+		return (*this)[Index];
 	}
-	/** @return const pointer to the Widget at the specified Index. */
-	virtual TSharedRef<const SWidget> GetChildAt( int32 Index ) const
+
+	virtual TSharedRef<const SWidget> GetChildAt( int32 Index ) const override
+	{
+		return (*this)[Index];
+	}
+
+	int32 Add( const TSharedRef<ChildType>& Child )
+	{
+		return TArray< TSharedRef<ChildType> >::Add(Child);
+	}
+
+	void Empty()
+	{
+		TArray< TSharedRef<ChildType> >::Empty();
+	}
+
+	void Insert(const TSharedRef<ChildType>& Child, int32 Index)
+	{
+		TArray< TSharedRef<ChildType> >::Insert(Child, Index);
+	}
+
+	int32 Remove( const TSharedRef<ChildType>& Child )
+	{
+		const int32 NumFoundAndRemoved = TArray< TSharedRef<ChildType> >::Remove( Child );
+		return NumFoundAndRemoved;
+	}
+
+	void RemoveAt( int32 Index )
+	{
+		TArray< TSharedRef<ChildType> >::RemoveAt( Index );
+	}
+
+	int32 Find( const TSharedRef<ChildType>& Item ) const
+	{
+		return TArray< TSharedRef<ChildType> >::Find( Item );
+	}
+
+	TArray< TSharedRef< ChildType > > AsArrayCopy() const
+	{
+		const int32 NumChildren = this->Num();
+		TArray< TSharedRef< ChildType > > Copy(*this);
+		for (int32 Index=0; Index < NumChildren; ++Index)
+		{
+			Copy.Add( (*this)[Index] );
+		}
+		return Copy;
+	}
+
+	const TSharedRef<ChildType>& operator[](int32 Index) const { return TArray< TSharedRef<ChildType> >::operator[](Index); }
+	TSharedRef<ChildType>& operator[](int32 Index) { return TArray< TSharedRef<ChildType> >::operator[](Index); }
+
+	template <class PREDICATE_CLASS>
+	void Sort( const PREDICATE_CLASS& Predicate )
+	{
+		TArray< TSharedRef<ChildType> >::Sort( Predicate );
+	}
+
+	void Swap( int32 IndexA, int32 IndexB )
 	{
-		return this->operator[](Index);
+		TIndirectArray< ChildType >::Swap(IndexA, IndexB);
 	}
 };
diff --git a/Engine/Source/Runtime/SlateCore/Public/Layout/LayoutUtils.h b/Engine/Source/Runtime/SlateCore/Public/Layout/LayoutUtils.h
--- a/Engine/Source/Runtime/SlateCore/Public/Layout/LayoutUtils.h
+++ b/Engine/Source/Runtime/SlateCore/Public/Layout/LayoutUtils.h
@@ -59,8 +59,8 @@ template<EOrientation Orientation, typename SlotType>
 static AlignmentArrangeResult AlignChild( float AllottedSize, const SlotType& ChildToArrange, const FMargin& SlotPadding, const float& ContentScale = 1.0f )
 {
 	float ChildDesiredSize = (Orientation == Orient_Horizontal)
-		? ( ChildToArrange.Widget->GetDesiredSize().X * ContentScale )
-		: ( ChildToArrange.Widget->GetDesiredSize().Y * ContentScale );
+		? ( ChildToArrange.GetWidget()->GetDesiredSize().X * ContentScale )
+		: ( ChildToArrange.GetWidget()->GetDesiredSize().Y * ContentScale );
 
 	const FMargin& Margin = SlotPadding;
 	const float TotalMargin = Margin.GetTotalSpaceAlong<Orientation>();	
@@ -100,7 +100,7 @@ static AlignmentArrangeResult AlignChild( float AllottedSize, const SlotType& Ch
 template<typename SlotType>
 void ArrangeSingleChild( const FGeometry& AllottedGeometry, FArrangedChildren& ArrangedChildren, const SlotType& ChildSlot, const TAttribute<FVector2D>& ContentScale )
 {
-	const EVisibility ChildVisibility = ChildSlot.Widget->GetVisibility();
+	const EVisibility ChildVisibility = ChildSlot.GetWidget()->GetVisibility();
 	if ( ArrangedChildren.Accepts(ChildVisibility) )
 	{
 		const FVector2D ThisContentScale = ContentScale.Get();
@@ -109,7 +109,7 @@ void ArrangeSingleChild( const FGeometry& AllottedGeometry, FArrangedChildren& A
 		AlignmentArrangeResult YResult = AlignChild<Orient_Vertical>(AllottedGeometry.Size.Y, ChildSlot, SlotPadding, ThisContentScale.Y);
 
 		ArrangedChildren.AddWidget( ChildVisibility, AllottedGeometry.MakeChild(
-				ChildSlot.Widget,
+				ChildSlot.GetWidget(),
 				FVector2D(XResult.Offset,YResult.Offset),
 				FVector2D(XResult.Size, YResult.Size)
 		) );
diff --git a/Engine/Source/Runtime/SlateCore/Public/SlateCore.h b/Engine/Source/Runtime/SlateCore/Public/SlateCore.h
--- a/Engine/Source/Runtime/SlateCore/Public/SlateCore.h
+++ b/Engine/Source/Runtime/SlateCore/Public/SlateCore.h
@@ -70,6 +70,7 @@ DECLARE_MEMORY_STAT_EXTERN(TEXT("Index Buffer Memory"), STAT_SlateIndexBufferMem
 #include "PaintArgs.h"
 
 // Layout
+#include "SlotBase.h"
 #include "Margin.h"
 #include "SlateRect.h"
 #include "PaintGeometry.h"
diff --git a/Engine/Source/Runtime/SlateCore/Public/SlotBase.h b/Engine/Source/Runtime/SlateCore/Public/SlotBase.h
new file mode 100644
--- /dev/null
+++ b/Engine/Source/Runtime/SlateCore/Public/SlotBase.h
@@ -0,0 +1,71 @@
+// Copyright 1998-2014 Epic Games, Inc. All Rights Reserved.
+
+#pragma once
+
+class SWidget;
+
+class SLATECORE_API FSlotBase
+{
+public:
+	FSlotBase& operator=( const FSlotBase& ) = delete;
+	FSlotBase( const FSlotBase& ) = delete;
+
+public:
+
+	FSlotBase();
+
+	FSlotBase( const TSharedRef<SWidget>& InWidget );
+
+	void AttachWidget( const TSharedRef<SWidget>& InWidget );
+
+	virtual ~FSlotBase();
+
+	/**
+	 * Access the widget in the current slot.
+	 * There will always be a widget in the slot; sometimes it is
+	 * the SNullWidget instance.
+	 */
+	const TSharedRef<SWidget>& GetWidget() const;
+
+	/**
+	 * Remove the widget from its current slot.
+	 * The removed widget is returned so that operations could be performed on it.
+	 * If the null widget was being stored, an invalid shared ptr is returned instead.
+	 */
+	const TSharedPtr<SWidget> DetachWidget();
+
+private:
+	/**
+	 *
+	 * Widget IS NOW PRIVATE!
+	 *
+	 */
+	TSharedRef<SWidget> Widget;
+};
+
+
+template<typename SlotType>
+class TSlotBase : public FSlotBase
+{
+public:
+
+	TSlotBase()
+	: FSlotBase()
+	{}
+
+	TSlotBase( const TSharedRef<SWidget>& InWidget )
+	: FSlotBase( InWidget )
+	{}
+
+	SlotType& operator[]( const TSharedRef<SWidget>& InChildWidget )
+	{
+		this->AttachWidget(InChildWidget);
+		return (SlotType&)(*this);
+	}
+
+	SlotType& Expose( SlotType*& OutVarToInit )
+	{
+		OutVarToInit = (SlotType*)this;
+		return (SlotType&)(*this);
+	}
+};
diff --git a/Engine/Source/Runtime/SlateCore/Public/Widgets/SBoxPanel.h b/Engine/Source/Runtime/SlateCore/Public/Widgets/SBoxPanel.h
--- a/Engine/Source/Runtime/SlateCore/Public/Widgets/SBoxPanel.h
+++ b/Engine/Source/Runtime/SlateCore/Public/Widgets/SBoxPanel.h
@@ -1,30 +1,25 @@
 // Copyright 1998-2014 Epic Games, Inc. All Rights Reserved.
 
-/*=============================================================================
-	SBoxPanel.h: Declares the SBoxPanel, SHorizontalBox and SVerticalBox classes.
-=============================================================================*/
-
 #pragma once
 
 
 /**
- * Implements a box panel widget.
- *
- * A BoxPanel contains one BoxPanel child and describes how that child should be arranged on the screen.
+ * A BoxPanel contains one child and describes how that child should be arranged on the screen.
  */
 class SLATECORE_API SBoxPanel
 	: public SPanel
 {
 public:
 
-	class FSlot
+	/**
+	 * A BoxPanel contains one BoxPanel child and describes how that
+	 * child should be arranged on the screen.
+	 */
+	class FSlot : public TSlotBase<FSlot>
 	{
 	public:		
 		/** Horizontal and Vertical Boxes inherit from FSlot */
-		virtual ~FSlot() { }
-
-		/** The child widget contained in this slot. */
-		TSharedRef<SWidget> Widget;
+		virtual ~FSlot(){}
 		/**
 		 * How much space this slot should occupy along panel's direction.
 		 *   When SizeRule is SizeRule_Auto, the widget's DesiredSize will be used as the space required.
@@ -48,12 +43,9 @@ public:
 		TAttribute<float> MaxSize;
 			
 	protected:
-
-		/**
-		 * Hidden default constructor.
-		 */
+		/** Default values for a slot. */
 		FSlot()
-			: Widget( SNullWidget::NullWidget )
+			: TSlotBase<FSlot>()
 			, SizeParam( FStretch(1) )
 			, HAlignment( HAlign_Fill )
 			, VAlignment( VAlign_Fill )
@@ -110,6 +102,11 @@ public:
 	class FSlot : public SBoxPanel::FSlot
 	{
 		public:
+		FSlot()
+		: SBoxPanel::FSlot()
+		{
+		}
+
 		FSlot& AutoWidth()
 		{
 			SizeParam = FAuto();
@@ -172,7 +169,7 @@ public:
 
 		FSlot& operator[]( TSharedRef<SWidget> InWidget )
 		{
-			Widget = InWidget;
+			SBoxPanel::FSlot::operator[](InWidget);
 			return *this;
 		}
 
@@ -199,7 +196,7 @@ public:
 
 	FSlot& AddSlot()
 	{
-		SHorizontalBox::FSlot& NewSlot = SHorizontalBox::Slot();
+		SHorizontalBox::FSlot& NewSlot = *new SHorizontalBox::FSlot();
 		this->Children.Add( &NewSlot );
 		return NewSlot;
 	}
@@ -224,7 +221,12 @@ public:
 	class FSlot : public SBoxPanel::FSlot
 	{
 		public:
-		
+
+		FSlot()
+		: SBoxPanel::FSlot()
+		{
+		}
+
 		FSlot& AutoHeight()
 		{
 			SizeParam = FAuto();
@@ -293,7 +295,7 @@ public:
 
 		FSlot& operator[]( TSharedRef<SWidget> InWidget )
 		{
-			Widget = InWidget;
+			SBoxPanel::FSlot::operator[](InWidget);
 			return *this;
 		}
 
@@ -321,7 +323,7 @@ public:
 
 	FSlot& AddSlot()
 	{
-		SVerticalBox::FSlot& NewSlot = SVerticalBox::Slot();
+		SVerticalBox::FSlot& NewSlot = *new SVerticalBox::FSlot();
 		this->Children.Add( &NewSlot );
 		return NewSlot;
 	}
diff --git a/Engine/Source/Runtime/SlateCore/Public/Widgets/SCompoundWidget.h b/Engine/Source/Runtime/SlateCore/Public/Widgets/SCompoundWidget.h
--- a/Engine/Source/Runtime/SlateCore/Public/Widgets/SCompoundWidget.h
+++ b/Engine/Source/Runtime/SlateCore/Public/Widgets/SCompoundWidget.h
@@ -8,13 +8,10 @@
 
 
 /**
- * Implements a compound widget.
- *
  * A CompoundWidget is the base from which most non-primitive widgets should be built.
  * CompoundWidgets have a protected member named ChildSlot.
  */
-class SLATECORE_API SCompoundWidget
-	: public SWidget
+class SLATECORE_API SCompoundWidget : public SWidget
 {
 public:
 
@@ -84,16 +81,8 @@ public:
 
 protected:
 
-	/**
-	 * Hidden constructor.
-	 */
-	FORCENOINLINE SCompoundWidget()
-		: ContentScale(FVector2D(1.0f, 1.0f))
-		, ColorAndOpacity(FLinearColor::White)
-		, ForegroundColor(FSlateColor::UseForeground())
-	{ }
-
-protected:
+	/** Disallow public construction */
+	SCompoundWidget();
 
 	/** The slot that contains this widget's descendants.*/
 	FSimpleSlot ChildSlot;
diff --git a/Engine/Source/Runtime/SlateCore/Public/Widgets/SOverlay.h b/Engine/Source/Runtime/SlateCore/Public/Widgets/SOverlay.h
--- a/Engine/Source/Runtime/SlateCore/Public/Widgets/SOverlay.h
+++ b/Engine/Source/Runtime/SlateCore/Public/Widgets/SOverlay.h
@@ -1,9 +1,5 @@
 // Copyright 1998-2014 Epic Games, Inc. All Rights Reserved.
 
-/*=============================================================================
-	SOverlay.h: Declares the SOverlay class.
-=============================================================================*/
-
 #pragma once
 
 
@@ -28,12 +24,12 @@ class SLATECORE_API SOverlay
 public:	
 
 	/** A slot that support alignment of content and padding and z-order */
-	class SLATECORE_API FOverlaySlot
-		: public TSupportsOneChildMixin<SWidget, FOverlaySlot>
+	class SLATECORE_API FOverlaySlot : public TSlotBase<FOverlaySlot>
 	{
 	public:
 		FOverlaySlot()
-			: HAlignment(HAlign_Fill)
+			: TSlotBase<FOverlaySlot>()
+			, HAlignment(HAlign_Fill)
 			, VAlignment(VAlign_Fill)
 			, SlotPadding(0.0f)
  			, ZOrder(0)
@@ -71,20 +67,19 @@ public:
 
 
 	SLATE_BEGIN_ARGS( SOverlay )
-		{
-			_Visibility = EVisibility::SelfHitTestInvisible;
-		}
+	{
+		_Visibility = EVisibility::SelfHitTestInvisible;
+	}
 
 		SLATE_SUPPORTS_SLOT( SOverlay::FOverlaySlot )
 
 	SLATE_END_ARGS()
 
-	/**
-	 * Virtual destructor.
-	 */
-	virtual ~SOverlay( ) { }
+	SOverlay();
 
-public:
+	virtual ~SOverlay()
+	{
+	}
 
 	/**
 	 * Construct this widget.
diff --git a/Engine/Source/Runtime/SlateCore/Public/Widgets/SWidget.h b/Engine/Source/Runtime/SlateCore/Public/Widgets/SWidget.h
--- a/Engine/Source/Runtime/SlateCore/Public/Widgets/SWidget.h
+++ b/Engine/Source/Runtime/SlateCore/Public/Widgets/SWidget.h
@@ -119,7 +119,7 @@ public:
 	 * @param  InCurrentTime  Current absolute real time
 	 * @param  InDeltaTime  Real time passed since last tick
 	 */
-	virtual void Tick( const FGeometry& AllottedGeometry, const double InCurrentTime, const float InDeltaTime ) { }
+	virtual void Tick( const FGeometry& AllottedGeometry, const double InCurrentTime, const float InDeltaTime );
 
 	/**
 	 * Called when a widget is being hit tested and has passed the bounding box test to determine
@@ -151,10 +151,10 @@ public:
 	 *
 	 * @param  InKeyboardFocusEvent  KeyboardFocusEvent
 	 */
-	virtual void OnKeyboardFocusLost( const FKeyboardFocusEvent& InKeyboardFocusEvent ) { }
+	virtual void OnKeyboardFocusLost( const FKeyboardFocusEvent& InKeyboardFocusEvent );
 
 	/** Called whenever a focus path is changing on all the widgets within the old and new focus paths */
-	virtual void OnKeyboardFocusChanging( const FWeakWidgetPath& PreviousFocusPath, const FWidgetPath& NewWidgetPath ) { }
+	virtual void OnKeyboardFocusChanging( const FWeakWidgetPath& PreviousFocusPath, const FWidgetPath& NewWidgetPath );
 
 	/**
 	 * Called after a character is entered while this widget has keyboard focus
@@ -549,7 +549,7 @@ public:
 	virtual TSharedPtr<IToolTip> GetToolTip();
 
 	/** Called when a tooltip displayed from this widget is being closed */
-	virtual void OnToolTipClosing() { }
+	virtual void OnToolTipClosing();
 
 	/**
 	 * Sets whether this widget is a "tool tip force field".  That is, tool-tips should never spawn over the area
diff --git a/Engine/Source/Runtime/SlateCore/Public/Widgets/SWindow.h b/Engine/Source/Runtime/SlateCore/Public/Widgets/SWindow.h
--- a/Engine/Source/Runtime/SlateCore/Public/Widgets/SWindow.h
+++ b/Engine/Source/Runtime/SlateCore/Public/Widgets/SWindow.h
@@ -844,26 +844,19 @@ private:
  * Popups, tooltips, drag and drop decorators all can be executed without creating a new window.
  * This slot along with the SWindow::AddPopupLayerSlot() API enabled it.
  */
-struct FPopupLayerSlot
+struct FPopupLayerSlot : public TSlotBase<FPopupLayerSlot>
 {
 public:
 	FPopupLayerSlot()
-	: DesktopPosition_Attribute(FVector2D::ZeroVector)
+	: TSlotBase<FPopupLayerSlot>()
+	, DesktopPosition_Attribute(FVector2D::ZeroVector)
 	, WidthOverride_Attribute()
 	, HeightOverride_Attribute()
 	, Scale_Attribute(1.0f)
 	, Clamp_Attribute(false)
 	, ClampBuffer_Attribute(FVector2D::ZeroVector)
-	, Widget(SNullWidget::NullWidget)
 	{}
 
-	/** Support for using brackets in slate declarative syntax */
-	FPopupLayerSlot& operator[]( TSharedRef<SWidget> InWidget )
-	{
-		Widget = InWidget;
-		return *this;
-	}
-
 	/** Pixel position in desktop space */
 	FPopupLayerSlot& DesktopPosition( const TAttribute<FVector2D>& InDesktopPosition )
 	{
@@ -906,12 +899,6 @@ public:
 		return *this;
 	}
 
-	/** @return the widget present in this slot */
-	const TSharedRef<SWidget>& GetWidget()
-	{
-		return Widget;
-	}
-
 private:
 	/** SPopupLayer arranges FPopupLayerSlots, so it needs to know all about */
 	friend class SPopupLayer;
@@ -924,7 +911,6 @@ private:
 	TAttribute<float> Scale_Attribute;
 	TAttribute<bool> Clamp_Attribute;
 	TAttribute<FVector2D> ClampBuffer_Attribute;
-	TSharedRef<SWidget> Widget;
 };
 
 
diff --git a/Engine/Source/Runtime/UMG/Private/Slate/SConstraintCanvas.cpp b/Engine/Source/Runtime/UMG/Private/Slate/SConstraintCanvas.cpp
--- a/Engine/Source/Runtime/UMG/Private/Slate/SConstraintCanvas.cpp
+++ b/Engine/Source/Runtime/UMG/Private/Slate/SConstraintCanvas.cpp
@@ -7,6 +7,12 @@
 /* SConstraintCanvas interface
  *****************************************************************************/
 
+SConstraintCanvas::SConstraintCanvas()
+: Children()
+{
+
+}
+
 void SConstraintCanvas::Construct( const SConstraintCanvas::FArguments& InArgs )
 {
 	const int32 NumSlots = InArgs.Slots.Num();
@@ -27,7 +33,7 @@ int32 SConstraintCanvas::RemoveSlot( const TSharedRef<SWidget>& SlotWidget )
 {
 	for (int32 SlotIdx = 0; SlotIdx < Children.Num(); ++SlotIdx)
 	{
-		if (SlotWidget == Children[SlotIdx].Widget)
+		if (SlotWidget == Children[SlotIdx].GetWidget())
 		{
 			Children.RemoveAt(SlotIdx);
 			return SlotIdx;
@@ -126,7 +132,7 @@ void SConstraintCanvas::OnArrangeChildren( const FGeometry& AllottedGeometry, FA
 			// Add the information about this child to the output list (ArrangedChildren)
 			ArrangedChildren.AddWidget( AllottedGeometry.MakeChild(
 				// The child widget being arranged
-				CurChild.Widget,
+				CurChild.GetWidget(),
 				// Child's local position (i.e. position within parent)
 				LocalPosition,
 				// Child's size
diff --git a/Engine/Source/Runtime/UMG/Public/Components/WidgetSwitcherSlot.h b/Engine/Source/Runtime/UMG/Public/Components/WidgetSwitcherSlot.h
--- a/Engine/Source/Runtime/UMG/Public/Components/WidgetSwitcherSlot.h
+++ b/Engine/Source/Runtime/UMG/Public/Components/WidgetSwitcherSlot.h
@@ -46,5 +46,5 @@ private:
 	//TODO UMG Slots should hold weak or shared refs to slots.
 
 	/** A raw pointer to the slot to allow us to adjust the size, padding...etc at runtime. */
-	FSimpleSlot* Slot;
+	SWidgetSwitcher::FSlot* Slot;
 };
diff --git a/Engine/Source/Runtime/UMG/Public/Slate/SConstraintCanvas.h b/Engine/Source/Runtime/UMG/Public/Slate/SConstraintCanvas.h
--- a/Engine/Source/Runtime/UMG/Public/Slate/SConstraintCanvas.h
+++ b/Engine/Source/Runtime/UMG/Public/Slate/SConstraintCanvas.h
@@ -14,7 +14,7 @@ public:
 	/**
 	 * ConstraintCanvas slots allow child widgets to be positioned and sized
 	 */
-	class FSlot
+	class FSlot : public TSlotBase<FSlot>
 	{
 	public:		
 		FSlot& Offset( const TAttribute<FMargin>& InOffset )
@@ -41,21 +41,12 @@ public:
 			return *this;
 		}
 
-		FSlot& operator[]( TSharedRef<SWidget> InWidget )
-		{
-			Widget = InWidget;
-			return *this;
-		}
-
 		FSlot& Expose( FSlot*& OutVarToInit )
 		{
 			OutVarToInit = this;
 			return *this;
 		}
 
-		/** The child widget contained in this slot. */
-		TSharedRef<SWidget> Widget;
-
 		/** Offset */
 		TAttribute<FMargin> OffsetAttr;
 
@@ -70,7 +61,7 @@ public:
 
 		/** Default values for a slot. */
 		FSlot()
-			: Widget( SNullWidget::NullWidget )
+			: TSlotBase<FSlot>()
 			, OffsetAttr( FMargin( 0, 0, 1, 1 ) )
 			, AnchorsAttr( FAnchors( 0.5f, 0.5f ) )
 			, AlignmentAttr( FVector2D( 0.5f, 0.5f ) )
@@ -87,6 +78,8 @@ public:
 
 	SLATE_END_ARGS()
 
+	SConstraintCanvas();
+
 	/**
 	 * Construct this widget
 	 *
@@ -106,7 +99,7 @@ public:
 	 */
 	FSlot& AddSlot()
 	{
-		SConstraintCanvas::FSlot& NewSlot = SConstraintCanvas::Slot();
+		SConstraintCanvas::FSlot& NewSlot = *(new FSlot());
 		this->Children.Add( &NewSlot );
 		return NewSlot;
 	}
