Fix Slate compiler errors on Xbox one

diff --git a/Engine/Source/Runtime/SlateCore/Public/SlotBase.h b/Engine/Source/Runtime/SlateCore/Public/SlotBase.h
--- a/Engine/Source/Runtime/SlateCore/Public/SlotBase.h
+++ b/Engine/Source/Runtime/SlateCore/Public/SlotBase.h
@@ -6,10 +6,6 @@ class SWidget;
 
 class SLATECORE_API FSlotBase
 {
-public:
-	FSlotBase& operator=( const FSlotBase& ) = delete;
-	FSlotBase( const FSlotBase& ) = delete;
-
 public:
 
 	FSlotBase();
@@ -34,6 +30,11 @@ public:
 	 */
 	const TSharedPtr<SWidget> DetachWidget();
 
+private:
+	// non-copyable
+	FSlotBase& operator=(const FSlotBase&);
+	FSlotBase(const FSlotBase&);
+
 private:
 	/**
 	 *
