To fix project thumbnails being upside down (TTP #341524) ensure that the new data member in FSlateBrush is always initialised. Unlike on other platforms, allocations on Mac aren't guaranteed to be zero-filled.

diff --git a/Engine/Source/Runtime/SlateCore/Private/Styling/SlateBrush.cpp b/Engine/Source/Runtime/SlateCore/Private/Styling/SlateBrush.cpp
--- a/Engine/Source/Runtime/SlateCore/Private/Styling/SlateBrush.cpp
+++ b/Engine/Source/Runtime/SlateCore/Private/Styling/SlateBrush.cpp
@@ -16,6 +16,7 @@ FSlateBrush::FSlateBrush( ESlateBrushDrawType::Type InDrawType, const FName InRe
 	, Margin( InMargin )
 	, TintColor( InTint )
 	, Tiling( InTiling )
+	, Mirroring( ESlateBrushMirrorType::NoMirror )
 	, ImageType( InImageType )
 	, ResourceObject( InObjectResource )
 	, ResourceName( InResourceName )
@@ -37,6 +38,7 @@ FSlateBrush::FSlateBrush( ESlateBrushDrawType::Type InDrawType, const FName InRe
 	, Margin( InMargin )
 	, TintColor( InTint )
 	, Tiling( InTiling )
+	, Mirroring( ESlateBrushMirrorType::NoMirror )
 	, ImageType( InImageType )
 	, ResourceObject( InObjectResource )
 	, ResourceName( InResourceName )
@@ -58,6 +60,7 @@ FSlateBrush::FSlateBrush( ESlateBrushDrawType::Type InDrawType, const FName InRe
 	, Margin(InMargin)
 	, TintColor(InTint)
 	, Tiling(InTiling)
+	, Mirroring( ESlateBrushMirrorType::NoMirror )
 	, ImageType(InImageType)
 	, ResourceObject(InObjectResource)
 	, ResourceName(InResourceName)
