Slate : Hittest 2.0 now properly supports the IsEnabled() state on all widgets.

diff --git a/Engine/Source/Runtime/Slate/Private/Framework/Application/SlateApplication.cpp b/Engine/Source/Runtime/Slate/Private/Framework/Application/SlateApplication.cpp
--- a/Engine/Source/Runtime/Slate/Private/Framework/Application/SlateApplication.cpp
+++ b/Engine/Source/Runtime/Slate/Private/Framework/Application/SlateApplication.cpp
@@ -497,7 +497,7 @@ FWidgetPath FSlateApplication::LocateWindowUnderMouse( FVector2D ScreenspaceMous
 		{
 			if ( !SWidget::UseLegacyHittest() )
 			{
-				const TArray<FArrangedWidget> ArrangedWidgets = HittestGrid->GetBubblePath( ScreenspaceMouseCoordinate );
+				const TArray<FArrangedWidget> ArrangedWidgets = HittestGrid->GetBubblePath( ScreenspaceMouseCoordinate, bIgnoreEnabledStatus );
 				return FWidgetPath( ArrangedWidgets );
 			}
 			else
diff --git a/Engine/Source/Runtime/SlateCore/Private/Input/HittestGrid.cpp b/Engine/Source/Runtime/SlateCore/Private/Input/HittestGrid.cpp
--- a/Engine/Source/Runtime/SlateCore/Private/Input/HittestGrid.cpp
+++ b/Engine/Source/Runtime/SlateCore/Private/Input/HittestGrid.cpp
@@ -37,7 +37,7 @@ FHittestGrid::FHittestGrid()
 {
 }
 
-TArray<FArrangedWidget> FHittestGrid::GetBubblePath( FVector2D DesktopSpaceCoordinate )
+TArray<FArrangedWidget> FHittestGrid::GetBubblePath( FVector2D DesktopSpaceCoordinate, bool bIgnoreEnabledStatus )
 {
 	if (WidgetsCachedThisFrame->Num() > 0)
 	{
@@ -92,6 +92,18 @@ TArray<FArrangedWidget> FHittestGrid::GetBubblePath( FVector2D DesktopSpaceCoord
 				BubblePath = TArray<FArrangedWidget>();
 			}
 
+			// Disabling a widget disables all of its logical children
+			// This effect is achieved by truncating the path to the
+			// root-most enabled widget.
+			if ( !bIgnoreEnabledStatus )
+			{
+				const int32 DisabledWidgetIndex = BubblePath.IndexOfByPredicate( []( const FArrangedWidget& SomeWidget ){ return !SomeWidget.Widget->IsEnabled( ); } );
+				if (DisabledWidgetIndex != INDEX_NONE)
+				{
+					BubblePath.RemoveAt( DisabledWidgetIndex, BubblePath.Num() - DisabledWidgetIndex );
+				}
+			}
+			
 			return BubblePath;
 		}
 		else
diff --git a/Engine/Source/Runtime/SlateCore/Public/Input/HittestGrid.h b/Engine/Source/Runtime/SlateCore/Public/Input/HittestGrid.h
--- a/Engine/Source/Runtime/SlateCore/Public/Input/HittestGrid.h
+++ b/Engine/Source/Runtime/SlateCore/Public/Input/HittestGrid.h
@@ -15,7 +15,7 @@ public:
 	 * Given a Slate Units coordinate in virtual desktop space, perform a hittest
 	 * and return the path along which the corresponding event would be bubbled.
 	 */
-	TArray<FArrangedWidget> GetBubblePath( FVector2D DesktopSpaceCoordinate );
+	TArray<FArrangedWidget> GetBubblePath( FVector2D DesktopSpaceCoordinate, bool bIgnoreEnabledStatus );
 
 	/**
 	 * Clear the hittesting area and prepare to execute a new frame.
