UMG/Slate - Adding a new parameter to AlignChild for the scale box or whenever you want to align things but not always clamp to the parents size, useful when you want to center things that can clip.

diff --git a/Engine/Source/Runtime/SlateCore/Public/Layout/LayoutUtils.h b/Engine/Source/Runtime/SlateCore/Public/Layout/LayoutUtils.h
--- a/Engine/Source/Runtime/SlateCore/Public/Layout/LayoutUtils.h
+++ b/Engine/Source/Runtime/SlateCore/Public/Layout/LayoutUtils.h
@@ -56,7 +56,7 @@ namespace ArrangeUtils
  * @return  Offset and Size of widget
  */
 template<EOrientation Orientation, typename SlotType>
-static AlignmentArrangeResult AlignChild( float AllottedSize, const SlotType& ChildToArrange, const FMargin& SlotPadding, const float& ContentScale = 1.0f )
+static AlignmentArrangeResult AlignChild(float AllottedSize, const SlotType& ChildToArrange, const FMargin& SlotPadding, const float& ContentScale = 1.0f, bool ClampToParent = true)
 {
 	float ChildDesiredSize = (Orientation == Orient_Horizontal)
 		? ( ChildToArrange.GetWidget()->GetDesiredSize().X * ContentScale )
@@ -69,7 +69,7 @@ static AlignmentArrangeResult AlignChild( float AllottedSize, const SlotType& Ch
 	
 	const int32 Alignment = ArrangeUtils::GetChildAlignment<Orientation>::AsInt(ChildToArrange);
 		
-	const float ChildSize = FMath::Min(ChildDesiredSize, AllottedSize-TotalMargin);
+	const float ChildSize = ClampToParent ? FMath::Min(ChildDesiredSize, AllottedSize - TotalMargin) : ChildDesiredSize;
 		
 	switch( Alignment )	
 	{
