Added more checks into the Hittest code to help track down an unusual bug.

diff --git a/Engine/Source/Runtime/SlateCore/Private/Input/HittestGrid.cpp b/Engine/Source/Runtime/SlateCore/Private/Input/HittestGrid.cpp
--- a/Engine/Source/Runtime/SlateCore/Private/Input/HittestGrid.cpp
+++ b/Engine/Source/Runtime/SlateCore/Private/Input/HittestGrid.cpp
@@ -58,6 +58,8 @@ TArray<FArrangedWidget> FHittestGrid::GetBubblePath( FVector2D DesktopSpaceCoord
 		// Consider front-most widgets first for hittesting.
 		for ( int32 i = IndexesInCell.Num()-1; i>=0 && HitWidgetIndex==INDEX_NONE; --i )
 		{
+			check( IndexesInCell[i] < WidgetsCachedThisFrame->Num() ); 
+
 			const FCachedWidget& TestCandidate = (*WidgetsCachedThisFrame)[IndexesInCell[i]];
 			if ( TestCandidate.CachedGeometry.IsUnderLocation( DesktopSpaceCoordinate ) && TestCandidate.ClippingRect.ContainsPoint( DesktopSpaceCoordinate ) && TestCandidate.WidgetPtr.IsValid() )
 			{
@@ -72,6 +74,7 @@ TArray<FArrangedWidget> FHittestGrid::GetBubblePath( FVector2D DesktopSpaceCoord
 			bool bPathUninterrupted = false;
 			do
 			{
+				check( CurWidgetIndex < WidgetsCachedThisFrame->Num() );
 				const FCachedWidget& CurCachedWidget = (*WidgetsCachedThisFrame)[CurWidgetIndex];
 				const TSharedPtr<SWidget> CachedWidgetPtr = CurCachedWidget.WidgetPtr.Pin();
 				
@@ -131,6 +134,8 @@ void FHittestGrid::BeginFrame( const FSlateRect& HittestArea )
 
 int32 FHittestGrid::InsertWidget( const int32 ParentHittestIndex, const EVisibility& Visibility, const FArrangedWidget& Widget, const FVector2D InWindowOffset, const FSlateRect& InClippingRect )
 {
+	check( ParentHittestIndex < WidgetsCachedThisFrame->Num() );
+
 	FArrangedWidget WindowAdjustedWidget(Widget);
 	WindowAdjustedWidget.Geometry.AbsolutePosition += InWindowOffset;
 
@@ -140,6 +145,7 @@ int32 FHittestGrid::InsertWidget( const int32 ParentHittestIndex, const EVisibil
 
 	// Remember this widget, its geometry, and its place in the logical hierarchy.
 	const int32 WidgetIndex = WidgetsCachedThisFrame->Add( FCachedWidget( ParentHittestIndex, WindowAdjustedWidget, WindowAdjustedRect ) );
+	check( WidgetIndex < WidgetsCachedThisFrame->Num() ); 
 	if (ParentHittestIndex != INDEX_NONE)
 	{
 		(*WidgetsCachedThisFrame)[ParentHittestIndex].AddChild( WidgetIndex );
