Slate: Removed file header comments

diff --git a/Engine/Source/Runtime/SlateCore/Private/SlateCorePrivatePCH.h b/Engine/Source/Runtime/SlateCore/Private/SlateCorePrivatePCH.h
--- a/Engine/Source/Runtime/SlateCore/Private/SlateCorePrivatePCH.h
+++ b/Engine/Source/Runtime/SlateCore/Private/SlateCorePrivatePCH.h
@@ -1,9 +1,5 @@
 // Copyright 1998-2014 Epic Games, Inc. All Rights Reserved.
 
-/*=============================================================================
-	SlateCorePrivatePCH.h: Pre-compiled header file for the SlateCore module.
-=============================================================================*/
-
 #pragma once
 
 #include "SlateCore.h"
diff --git a/Engine/Source/Runtime/SlateCore/Private/Sound/SlateSound.cpp b/Engine/Source/Runtime/SlateCore/Private/Sound/SlateSound.cpp
--- a/Engine/Source/Runtime/SlateCore/Private/Sound/SlateSound.cpp
+++ b/Engine/Source/Runtime/SlateCore/Private/Sound/SlateSound.cpp
@@ -1,9 +1,5 @@
 // Copyright 1998-2014 Epic Games, Inc. All Rights Reserved.
 
-/*=============================================================================
-	SlateSound.cpp: Implements the FSlateSound class.
-=============================================================================*/
-
 #include "SlateCorePrivatePCH.h"
 
 
diff --git a/Engine/Source/Runtime/SlateCore/Private/Styling/CoreStyle.cpp b/Engine/Source/Runtime/SlateCore/Private/Styling/CoreStyle.cpp
--- a/Engine/Source/Runtime/SlateCore/Private/Styling/CoreStyle.cpp
+++ b/Engine/Source/Runtime/SlateCore/Private/Styling/CoreStyle.cpp
@@ -1,9 +1,5 @@
 // Copyright 1998-2014 Epic Games, Inc. All Rights Reserved.
 
-/*=============================================================================
-	CoreStyle.cpp: Implements the FCoreStyle class.
-=============================================================================*/
-
 #include "SlateCorePrivatePCH.h"
 
 
diff --git a/Engine/Source/Runtime/SlateCore/Private/Styling/SlateBrush.cpp b/Engine/Source/Runtime/SlateCore/Private/Styling/SlateBrush.cpp
--- a/Engine/Source/Runtime/SlateCore/Private/Styling/SlateBrush.cpp
+++ b/Engine/Source/Runtime/SlateCore/Private/Styling/SlateBrush.cpp
@@ -1,9 +1,5 @@
 // Copyright 1998-2014 Epic Games, Inc. All Rights Reserved.
 
-/*=============================================================================
-	SlateBrush.cpp: Implements the FSlateBrush structure.
-=============================================================================*/
-
 #include "SlateCorePrivatePCH.h"
 
 
diff --git a/Engine/Source/Runtime/SlateCore/Private/Styling/SlateStyleRegistry.cpp b/Engine/Source/Runtime/SlateCore/Private/Styling/SlateStyleRegistry.cpp
--- a/Engine/Source/Runtime/SlateCore/Private/Styling/SlateStyleRegistry.cpp
+++ b/Engine/Source/Runtime/SlateCore/Private/Styling/SlateStyleRegistry.cpp
@@ -1,9 +1,5 @@
 // Copyright 1998-2014 Epic Games, Inc. All Rights Reserved.
 
-/*=============================================================================
-	SlateStyleRegistry.cpp: Implements the FSlateStyleRegistry class.
-=============================================================================*/
-
 #include "SlateCorePrivatePCH.h"
 
 
diff --git a/Engine/Source/Runtime/SlateCore/Private/Styling/SlateWidgetStyleContainerBase.cpp b/Engine/Source/Runtime/SlateCore/Private/Styling/SlateWidgetStyleContainerBase.cpp
--- a/Engine/Source/Runtime/SlateCore/Private/Styling/SlateWidgetStyleContainerBase.cpp
+++ b/Engine/Source/Runtime/SlateCore/Private/Styling/SlateWidgetStyleContainerBase.cpp
@@ -1,9 +1,5 @@
 // Copyright 1998-2014 Epic Games, Inc. All Rights Reserved.
 
-/*=============================================================================
-	SlateWidgetStyleContainerBase.cpp: Implements the USlateWidgetStyleContainerBase class.
-=============================================================================*/
-
 #include "SlateCorePrivatePCH.h"
 #include "SlateWidgetStyleContainerBase.h"
 
diff --git a/Engine/Source/Runtime/SlateCore/Private/Textures/TextureAtlas.cpp b/Engine/Source/Runtime/SlateCore/Private/Textures/TextureAtlas.cpp
--- a/Engine/Source/Runtime/SlateCore/Private/Textures/TextureAtlas.cpp
+++ b/Engine/Source/Runtime/SlateCore/Private/Textures/TextureAtlas.cpp
@@ -1,9 +1,5 @@
 // Copyright 1998-2014 Epic Games, Inc. All Rights Reserved.
 
-/*=============================================================================
-	TextureAtlas.cpp: Implements the FSlateTextureAtlas class.
-=============================================================================*/
-
 #include "SlateCorePrivatePCH.h"
 
 
diff --git a/Engine/Source/Runtime/SlateCore/Private/Types/WidgetStyle.cpp b/Engine/Source/Runtime/SlateCore/Private/Types/WidgetStyle.cpp
--- a/Engine/Source/Runtime/SlateCore/Private/Types/WidgetStyle.cpp
+++ b/Engine/Source/Runtime/SlateCore/Private/Types/WidgetStyle.cpp
@@ -1,9 +1,5 @@
 // Copyright 1998-2014 Epic Games, Inc. All Rights Reserved.
 
-/*=============================================================================
-	WidgetStyle.cpp: Implements the FWidgetStyle structure.
-=============================================================================*/
-
 #include "SlateCorePrivatePCH.h"
 
 
diff --git a/Engine/Source/Runtime/SlateCore/Private/Widgets/SCompoundWidget.cpp b/Engine/Source/Runtime/SlateCore/Private/Widgets/SCompoundWidget.cpp
--- a/Engine/Source/Runtime/SlateCore/Private/Widgets/SCompoundWidget.cpp
+++ b/Engine/Source/Runtime/SlateCore/Private/Widgets/SCompoundWidget.cpp
@@ -1,9 +1,5 @@
 // Copyright 1998-2014 Epic Games, Inc. All Rights Reserved.
 
-/*=============================================================================
-	SCompoundWidget.cpp: Implements the SCompoundWidget class.
-=============================================================================*/
-
 #include "SlateCorePrivatePCH.h"
 
 
diff --git a/Engine/Source/Runtime/SlateCore/Private/Widgets/SLeafWidget.cpp b/Engine/Source/Runtime/SlateCore/Private/Widgets/SLeafWidget.cpp
--- a/Engine/Source/Runtime/SlateCore/Private/Widgets/SLeafWidget.cpp
+++ b/Engine/Source/Runtime/SlateCore/Private/Widgets/SLeafWidget.cpp
@@ -1,9 +1,5 @@
 //// Copyright 1998-2014 Epic Games, Inc. All Rights Reserved.
 
-/*=============================================================================
-	SLeafWidget.cpp: Implements the SLeafWidget class.
-=============================================================================*/
-
 #include "SlateCorePrivatePCH.h"
 
 
diff --git a/Engine/Source/Runtime/SlateCore/Private/Widgets/SNullWidget.cpp b/Engine/Source/Runtime/SlateCore/Private/Widgets/SNullWidget.cpp
--- a/Engine/Source/Runtime/SlateCore/Private/Widgets/SNullWidget.cpp
+++ b/Engine/Source/Runtime/SlateCore/Private/Widgets/SNullWidget.cpp
@@ -1,9 +1,5 @@
 // Copyright 1998-2014 Epic Games, Inc. All Rights Reserved.
 
-/*=============================================================================
-	SNullWidget.cpp: Implements the SNullWidget class.
-=============================================================================*/
- 
 #include "SlateCorePrivatePCH.h"
 
 
diff --git a/Engine/Source/Runtime/SlateCore/Private/Widgets/SPanel.cpp b/Engine/Source/Runtime/SlateCore/Private/Widgets/SPanel.cpp
--- a/Engine/Source/Runtime/SlateCore/Private/Widgets/SPanel.cpp
+++ b/Engine/Source/Runtime/SlateCore/Private/Widgets/SPanel.cpp
@@ -1,9 +1,5 @@
 // Copyright 1998-2014 Epic Games, Inc. All Rights Reserved.
 
-/*=============================================================================
-	SPanel.cpp: Implements the SPanel class.
-=============================================================================*/
-
 #include "SlateCorePrivatePCH.h"
 
 
diff --git a/Engine/Source/Runtime/SlateCore/Private/Widgets/SWindow.cpp b/Engine/Source/Runtime/SlateCore/Private/Widgets/SWindow.cpp
--- a/Engine/Source/Runtime/SlateCore/Private/Widgets/SWindow.cpp
+++ b/Engine/Source/Runtime/SlateCore/Private/Widgets/SWindow.cpp
@@ -1,9 +1,5 @@
 // Copyright 1998-2014 Epic Games, Inc. All Rights Reserved.
 
-/*=============================================================================
-	SWindow.cpp: Implements the SWindow class.
-=============================================================================*/
-
 #include "SlateCorePrivatePCH.h"
 
 
diff --git a/Engine/Source/Runtime/SlateCore/Public/Animation/CurveHandle.h b/Engine/Source/Runtime/SlateCore/Public/Animation/CurveHandle.h
--- a/Engine/Source/Runtime/SlateCore/Public/Animation/CurveHandle.h
+++ b/Engine/Source/Runtime/SlateCore/Public/Animation/CurveHandle.h
@@ -1,9 +1,5 @@
 // Copyright 1998-2014 Epic Games, Inc. All Rights Reserved.
 
-/*=============================================================================
-	CurveHandle.h: Declares the FCurveHandle structure.
-=============================================================================*/
-
 #pragma once
 
 
@@ -57,7 +53,6 @@ public:
 	 * Gets the linearly interpolated value between 0 and 1 for this curve.
 	 *
 	 * @return Lerp value.
-	 *
 	 * @see GetLerpLooping
 	 */
 	float GetLerp( ) const;
@@ -66,7 +61,6 @@ public:
 	 * Just like GetLerp, but loops forever.
 	 *
 	 * @return Looped lerp value.
-	 *
 	 * @see GetLerp
 	 */
 	float GetLerpLooping( ) const;
@@ -75,7 +69,6 @@ public:
 	 * Checks whether this handle is initialized.
 	 *
 	 * A curve handle is considered initialized if it has an owner sequence.
-	 *
 	 * @return true if initialized, false otherwise.
 	 */
 	bool IsInitialized( ) const
diff --git a/Engine/Source/Runtime/SlateCore/Public/Animation/CurveSequence.h b/Engine/Source/Runtime/SlateCore/Public/Animation/CurveSequence.h
--- a/Engine/Source/Runtime/SlateCore/Public/Animation/CurveSequence.h
+++ b/Engine/Source/Runtime/SlateCore/Public/Animation/CurveSequence.h
@@ -1,9 +1,5 @@
 // Copyright 1998-2014 Epic Games, Inc. All Rights Reserved.
 
-/*=============================================================================
-	CurveSequence.h: Declares the FCurveSequence structure.
-=============================================================================*/
-
 #pragma once
 
 
@@ -54,7 +50,6 @@ public:
 	 * @param InStartTimeSeconds   When to start this curve.
 	 * @param InDurationSeconds    How long this curve lasts.
 	 * @param InEaseFunction       Easing function to use for this curve.  Defaults to Linear.  Use this to smooth out your animation transitions.
-	 *
 	 * @return A FCurveHandle that can be used to get the value of this curve after the animation starts playing.
 	 */
 	FCurveSequence( const float InStartTimeSeconds, const float InDurationSeconds, const ECurveEaseFunction::Type InEaseFunction = ECurveEaseFunction::Linear  );
@@ -65,7 +60,6 @@ public:
 	 * @param InStartTimeSeconds   When to start this curve.
 	 * @param InDurationSeconds    How long this curve lasts.
 	 * @param InEaseFunction       Easing function to use for this curve.  Defaults to Linear.  Use this to smooth out your animation transitions.
-	 *
 	 * @return A FCurveHandle that can be used to get the value of this curve after the animation starts playing.
 	 */
 	FCurveHandle AddCurve( const float InStartTimeSeconds, const float InDurationSeconds, const ECurveEaseFunction::Type InEaseFunction = ECurveEaseFunction::Linear );
diff --git a/Engine/Source/Runtime/SlateCore/Public/Application/SlateApplicationBase.h b/Engine/Source/Runtime/SlateCore/Public/Application/SlateApplicationBase.h
--- a/Engine/Source/Runtime/SlateCore/Public/Application/SlateApplicationBase.h
+++ b/Engine/Source/Runtime/SlateCore/Public/Application/SlateApplicationBase.h
@@ -1,11 +1,8 @@
 // Copyright 1998-2014 Epic Games, Inc. All Rights Reserved.
 
-/*=============================================================================
-	SlateApplicationBase.h: Declares the FSlateApplicationBase class.
-=============================================================================*/
-
 #pragma once
 
+
 class FHittestGrid;
 
 
@@ -67,7 +64,6 @@ public:
 	 * @param  InWidget       Widget to generate a path to
 	 * @param  OutWidgetPath  The generated widget path
 	 * @param  VisibilityFilter	Widgets must have this type of visibility to be included the path
-	 *
 	 * @return	True if the widget path was found
 	 */
 	virtual bool FindPathToWidget( TSharedRef<const SWidget> InWidget, FWidgetPath& OutWidgetPath, EVisibility VisibilityFilter = EVisibility::Visible ) = 0;
@@ -159,7 +155,6 @@ public:
 	 * Checks whether the specified widget has any descendants which are currently focused.
 	 *
 	 * @param Widget The widget to check.
-	 *
 	 * @return true if any descendants are focused, false otherwise.
 	 */
 	virtual bool HasFocusedDescendants( const TSharedRef< const SWidget >& Widget ) const = 0;
@@ -195,7 +190,6 @@ public:
 	 * Creates a tool tip with the specified string.
 	 *
 	 * @param ToolTipString The string attribute to assign to the tool tip.
-	 *
 	 * @return The tool tip.
 	 */
 	virtual TSharedRef<IToolTip> MakeToolTip( const TAttribute<FString>& ToolTipString ) = 0;
@@ -213,7 +207,6 @@ public:
 	 * Creates a tool tip with the specified text.
 	 *
 	 * @param ToolTipText The text to assign to the tool tip.
-	 *
 	 * @return The tool tip.
 	 */
 	virtual TSharedRef<IToolTip> MakeToolTip( const FText& ToolTipText ) = 0;
@@ -225,7 +218,6 @@ public:
 	 * @param CenterContent Optional content for the title bar's center (will override window title).
 	 * @param CenterContentAlignment The horizontal alignment of the center content.
 	 * @param OutTitleBar Will hold a pointer to the title bar's interface.
-	 *
 	 * @return The new title bar widget.
 	 */
 	virtual TSharedRef<SWidget> MakeWindowTitleBar( const TSharedRef<SWindow>& Window, const TSharedPtr<SWidget>& CenterContent, EHorizontalAlignment CenterContentAlignment, TSharedPtr<IWindowTitleBar>& OutTitleBar ) const = 0;
@@ -244,7 +236,6 @@ public:
 	 *
 	 * @param  InWidget WidgetPath to the Widget to being focused
 	 * @param InCause The reason that keyboard focus is changing
-	 *
 	 * @return true if the widget is now focused, false otherwise.
 	 */
 	virtual bool SetKeyboardFocus( const FWidgetPath& InFocusPath, const EKeyboardFocusCause::Type InCause ) = 0;
diff --git a/Engine/Source/Runtime/SlateCore/Public/Application/SlateWindowHelper.h b/Engine/Source/Runtime/SlateCore/Public/Application/SlateWindowHelper.h
--- a/Engine/Source/Runtime/SlateCore/Public/Application/SlateWindowHelper.h
+++ b/Engine/Source/Runtime/SlateCore/Public/Application/SlateWindowHelper.h
@@ -1,9 +1,5 @@
 // Copyright 1998-2014 Epic Games, Inc. All Rights Reserved.
 
-/*=============================================================================
-	SlateWindowHelper.h: Declares the FSlateWindowHelper class.
-=============================================================================*/
-
 #pragma once
 
 
@@ -37,7 +33,6 @@ public:
 	 *
 	 * @param WindowsToSearch The collection of windows to check.
 	 * @param WorkAreaRect The work area to check.
-	 *
 	 * @return true if at least one window overlaps with the work area, false otherwise.
 	 */
 	static bool CheckWorkAreaForWindows( const TArray< TSharedRef<SWindow> >& WindowsToSearch, const FSlateRect& WorkAreaRect );
@@ -47,7 +42,6 @@ public:
 	 *
 	 * @param WindowsToSearch The collection of windows to search.
 	 * @param WindowToFind The window to find.
-	 *
 	 * @return true if the collection contains the window, false otherwise.
 	 */
 	static bool ContainsWindow( const TArray<TSharedRef<SWindow>>& WindowsToSearch, const TSharedRef<SWindow>& WindowToFind );
@@ -60,7 +54,6 @@ public:
 	 * @param InWidget Widget to generate a path to.
 	 * @param OutWidgetPath The generated widget path.
 	 * @param VisibilityFilter Widgets must have this type of visibility to be included the path.
-	 *
 	 * @return True if the widget path was found.
 	 */
 	static bool FindPathToWidget( const TArray<TSharedRef<SWindow>> WindowsToSearch, TSharedRef<const SWidget> InWidget, FWidgetPath& OutWidgetPath, EVisibility VisibilityFilter = EVisibility::Visible );
@@ -70,7 +63,6 @@ public:
 	 *
 	 * @param WindowsToSearch The collection of windows to search.
 	 * @param PlatformWindow The platform window to find the Slate window for.
-	 *
 	 * @return The corresponding Slate window, or nullptr if it was not found.
 	 */
 	static TSharedPtr<SWindow> FindWindowByPlatformWindow( const TArray<TSharedRef<SWindow>>& WindowsToSearch, const TSharedRef<FGenericWindow>& PlatformWindow );
diff --git a/Engine/Source/Runtime/SlateCore/Public/Application/ThrottleManager.h b/Engine/Source/Runtime/SlateCore/Public/Application/ThrottleManager.h
--- a/Engine/Source/Runtime/SlateCore/Public/Application/ThrottleManager.h
+++ b/Engine/Source/Runtime/SlateCore/Public/Application/ThrottleManager.h
@@ -1,9 +1,5 @@
 // Copyright 1998-2014 Epic Games, Inc. All Rights Reserved.
 
-/*=============================================================================
-	ThrottleManager.h: Declares the FSlateThrottleManager class.
-=============================================================================*/
-
 #pragma once
 
 
diff --git a/Engine/Source/Runtime/SlateCore/Public/Brushes/SlateBorderBrush.h b/Engine/Source/Runtime/SlateCore/Public/Brushes/SlateBorderBrush.h
--- a/Engine/Source/Runtime/SlateCore/Public/Brushes/SlateBorderBrush.h
+++ b/Engine/Source/Runtime/SlateCore/Public/Brushes/SlateBorderBrush.h
@@ -1,9 +1,5 @@
 // Copyright 1998-2014 Epic Games, Inc. All Rights Reserved.
 
-/*=============================================================================
-	SlateBoxBrush.h: Declares the FSlateBoxBrush structure.
-=============================================================================*/
-
 #pragma once
 
 
diff --git a/Engine/Source/Runtime/SlateCore/Public/Brushes/SlateBoxBrush.h b/Engine/Source/Runtime/SlateCore/Public/Brushes/SlateBoxBrush.h
--- a/Engine/Source/Runtime/SlateCore/Public/Brushes/SlateBoxBrush.h
+++ b/Engine/Source/Runtime/SlateCore/Public/Brushes/SlateBoxBrush.h
@@ -1,9 +1,5 @@
 // Copyright 1998-2014 Epic Games, Inc. All Rights Reserved.
 
-/*=============================================================================
-	SlateBoxBrush.h: Declares the FSlateBoxBrush structure.
-=============================================================================*/
-
 #pragma once
 
 
diff --git a/Engine/Source/Runtime/SlateCore/Public/Brushes/SlateColorBrush.h b/Engine/Source/Runtime/SlateCore/Public/Brushes/SlateColorBrush.h
--- a/Engine/Source/Runtime/SlateCore/Public/Brushes/SlateColorBrush.h
+++ b/Engine/Source/Runtime/SlateCore/Public/Brushes/SlateColorBrush.h
@@ -1,9 +1,5 @@
 // Copyright 1998-2014 Epic Games, Inc. All Rights Reserved.
 
-/*=============================================================================
-	SlateColorBrush.h: Declares the FSlateColorBrush structure.
-=============================================================================*/
-
 #pragma once
 
 
diff --git a/Engine/Source/Runtime/SlateCore/Public/Brushes/SlateDynamicImageBrush.h b/Engine/Source/Runtime/SlateCore/Public/Brushes/SlateDynamicImageBrush.h
--- a/Engine/Source/Runtime/SlateCore/Public/Brushes/SlateDynamicImageBrush.h
+++ b/Engine/Source/Runtime/SlateCore/Public/Brushes/SlateDynamicImageBrush.h
@@ -1,9 +1,5 @@
 // Copyright 1998-2014 Epic Games, Inc. All Rights Reserved.
 
-/*=============================================================================
-	SlateDynamicImageBrush.h: Declares the FSlateDynamicImageBrush structure.
-=============================================================================*/
-
 #pragma once
 
 
@@ -61,8 +57,6 @@ struct SLATECORE_API FSlateDynamicImageBrush
 		bIsDynamicallyLoaded = true;
 	}
 
-	/**
-	 * Destructor.
-	 */
+	/** Destructor. */
 	virtual ~FSlateDynamicImageBrush();
 };
diff --git a/Engine/Source/Runtime/SlateCore/Public/Brushes/SlateImageBrush.h b/Engine/Source/Runtime/SlateCore/Public/Brushes/SlateImageBrush.h
--- a/Engine/Source/Runtime/SlateCore/Public/Brushes/SlateImageBrush.h
+++ b/Engine/Source/Runtime/SlateCore/Public/Brushes/SlateImageBrush.h
@@ -1,9 +1,5 @@
 // Copyright 1998-2014 Epic Games, Inc. All Rights Reserved.
 
-/*=============================================================================
-	SlateImageBrush.h: Declares the FSlateImageBrush structure.
-=============================================================================*/
-
 #pragma once
 
 
diff --git a/Engine/Source/Runtime/SlateCore/Public/Brushes/SlateNoResource.h b/Engine/Source/Runtime/SlateCore/Public/Brushes/SlateNoResource.h
--- a/Engine/Source/Runtime/SlateCore/Public/Brushes/SlateNoResource.h
+++ b/Engine/Source/Runtime/SlateCore/Public/Brushes/SlateNoResource.h
@@ -1,9 +1,5 @@
 // Copyright 1998-2014 Epic Games, Inc. All Rights Reserved.
 
-/*=============================================================================
-	SlateNoResource.h: Declares the FSlateNoResource structure.
-=============================================================================*/
-
 #pragma once
 
 
diff --git a/Engine/Source/Runtime/SlateCore/Public/Fonts/SlateFontInfo.h b/Engine/Source/Runtime/SlateCore/Public/Fonts/SlateFontInfo.h
--- a/Engine/Source/Runtime/SlateCore/Public/Fonts/SlateFontInfo.h
+++ b/Engine/Source/Runtime/SlateCore/Public/Fonts/SlateFontInfo.h
@@ -1,9 +1,5 @@
 // Copyright 1998-2014 Epic Games, Inc. All Rights Reserved.
 
-/*=============================================================================
-	SlateFontInfo.h: Declares the FSlateFontInfo structure.
-=============================================================================*/
-
 #pragma once
 
 #include "SlateFontInfo.generated.h"
@@ -27,9 +23,7 @@ struct SLATECORE_API FSlateFontInfo
 
 public:
 
-	/**
-	 * Default constructor.
-	 */
+	/** Default constructor. */
 	FSlateFontInfo();
 
 	/**
@@ -70,7 +64,6 @@ public:
 	 * Compares this font info with another for equality.
 	 *
 	 * @param Other The other font info.
-	 *
 	 * @return true if the two font infos are equal, false otherwise.
 	 */
 	bool operator==( const FSlateFontInfo& Other ) const 
@@ -84,7 +77,6 @@ public:
 	 * Type hashes are used in certain collection types, such as TMap.
 	 *
 	 * @param FontInfo The font info to calculate the hash for.
-	 *
 	 * @return The hash value.
 	 */
 	friend inline uint32 GetTypeHash( const FSlateFontInfo& FontInfo )
diff --git a/Engine/Source/Runtime/SlateCore/Public/Input/CursorReply.h b/Engine/Source/Runtime/SlateCore/Public/Input/CursorReply.h
--- a/Engine/Source/Runtime/SlateCore/Public/Input/CursorReply.h
+++ b/Engine/Source/Runtime/SlateCore/Public/Input/CursorReply.h
@@ -1,9 +1,5 @@
 // Copyright 1998-2014 Epic Games, Inc. All Rights Reserved.
 
-/*=============================================================================
-	CursorReply.h: Declares the FCursorReply class.
-=============================================================================*/
-
 #pragma once
 
 
@@ -13,6 +9,7 @@
 class FCursorReply
 {
 public:
+
 	/** Makes a NULL response; this response will not change the previous value. */
 	static FCursorReply Unhandled()
 	{
@@ -32,19 +29,18 @@ public:
 	EMouseCursor::Type GetCursor() { return MouseCursor; }
 		
 private:
+
 	/** Internal constructor - makes a NULL result. */
 	FCursorReply()
-	: bIsHandled(false)
-	, MouseCursor( EMouseCursor::Default )
-	{
-	}
+		: bIsHandled(false)
+		, MouseCursor( EMouseCursor::Default )
+	{ }
 		
 	/** Internal constructor - makes a non-NULL result. */
 	FCursorReply( EMouseCursor::Type InCursorType )
-	: bIsHandled(true)
-	, MouseCursor( InCursorType )
-	{		
-	}
+		: bIsHandled(true)
+		, MouseCursor( InCursorType )
+	{ }
 		
 	/** Does this reply have any meaning */ 
 	bool bIsHandled;
diff --git a/Engine/Source/Runtime/SlateCore/Public/Input/Events.h b/Engine/Source/Runtime/SlateCore/Public/Input/Events.h
--- a/Engine/Source/Runtime/SlateCore/Public/Input/Events.h
+++ b/Engine/Source/Runtime/SlateCore/Public/Input/Events.h
@@ -1,15 +1,13 @@
 // Copyright 1998-2014 Epic Games, Inc. All Rights Reserved.
 
-/*=============================================================================
-	Events.h: Declares various input event related types.
-=============================================================================*/
-
 #pragma once
 
 #include "Events.generated.h"
 
+
 class SWindow;
 
+
 /**
  * Context for keyboard focus change
  */
diff --git a/Engine/Source/Runtime/SlateCore/Public/Layout/ArrangedChildren.h b/Engine/Source/Runtime/SlateCore/Public/Layout/ArrangedChildren.h
--- a/Engine/Source/Runtime/SlateCore/Public/Layout/ArrangedChildren.h
+++ b/Engine/Source/Runtime/SlateCore/Public/Layout/ArrangedChildren.h
@@ -1,9 +1,5 @@
 // Copyright 1998-2014 Epic Games, Inc. All Rights Reserved.
 
-/*=============================================================================
-	ArrangedChildren.h: Declares the FArrangedChildren class.
-=============================================================================*/
-
 #pragma once
 
 
@@ -140,5 +136,4 @@ class SLATECORE_API FArrangedChildren
 	{
 		Array.Empty();
 	}
-
 };
diff --git a/Engine/Source/Runtime/SlateCore/Public/Layout/ArrangedWidget.h b/Engine/Source/Runtime/SlateCore/Public/Layout/ArrangedWidget.h
--- a/Engine/Source/Runtime/SlateCore/Public/Layout/ArrangedWidget.h
+++ b/Engine/Source/Runtime/SlateCore/Public/Layout/ArrangedWidget.h
@@ -1,9 +1,5 @@
 // Copyright 1998-2014 Epic Games, Inc. All Rights Reserved.
 
-/*=============================================================================
-	ArrangedWidget.h: Declares the FArrangedWidget class.
-=============================================================================*/
-
 #pragma once
 
 
@@ -35,7 +31,6 @@ public:
 	 * Compares this widget arrangement with another for equality.
 	 *
 	 * @param Other The other arrangement to compare with.
-	 *
 	 * @return true if the two arrangements are equal, false otherwise.
 	 */
 	bool operator==( const FArrangedWidget& Other ) const 
diff --git a/Engine/Source/Runtime/SlateCore/Public/Layout/Geometry.h b/Engine/Source/Runtime/SlateCore/Public/Layout/Geometry.h
--- a/Engine/Source/Runtime/SlateCore/Public/Layout/Geometry.h
+++ b/Engine/Source/Runtime/SlateCore/Public/Layout/Geometry.h
@@ -1,9 +1,5 @@
 // Copyright 1998-2014 Epic Games, Inc. All Rights Reserved.
 
-/*=============================================================================
-	Geometry.h: Declares the FGeometry class.
-=============================================================================*/
-
 #pragma once
 
 #include "Geometry.generated.h"
@@ -29,9 +25,7 @@ struct SLATECORE_API FGeometry
 
 public:
 
-	/**
-	 * Default constructor.
-	 */
+	/** Default constructor. */
 	FGeometry( )
 		: Position(0.0f, 0.0f)
 		, AbsolutePosition(0.0f, 0.0f)
@@ -60,7 +54,6 @@ public:
 	 * Test geometries for strict equality (i.e. exact float equality).
 	 *
 	 * @param Other A geometry to compare to.
-	 *
 	 * @return true if this is equal to other, false otherwise.
 	 */
 	bool operator==( const FGeometry& Other ) const
@@ -76,7 +69,6 @@ public:
 	 * Test geometries for strict inequality (i.e. exact float equality negated).
 	 *
 	 * @param Other A geometry to compare to.
-	 *
 	 * @return false if this is equal to other, true otherwise.
 	 */
 	bool operator!=( const FGeometry& Other ) const
@@ -92,7 +84,6 @@ public:
 	 * @param ChildOffset  The offset of the child from this Geometry's position.
 	 * @param ChildSize    The size of the child geometry in screen units.
 	 * @param InScale      How much to scale this child with respect to its parent.
-	 *
 	 * @return A Geometry that is offset from this geometry by ChildOffset
 	 */
 	FGeometry MakeChild( const FVector2D& ChildOffset, const FVector2D& ChildSize, float InScale = 1.0f ) const
@@ -108,7 +99,6 @@ public:
 	 * @param ChildOffset  The offset of the child from this Geometry's position.
 	 * @param ChildSize    The size of the child geometry in screen units.
 	 * @param InScale      How much to scale this child with respect to its parent.
-	 *
 	 * @return A WidgetGeometry that is offset from this geometry by ChildOffset
 	 */
 	FArrangedWidget MakeChild( const TSharedRef<SWidget>& InWidget, const FVector2D& ChildOffset, const FVector2D& ChildSize, float InScale = 1.0f ) const;
@@ -119,7 +109,6 @@ public:
 	 * @param InOffset   Offset in SlateUnits from this FGeometry's origin.
 	 * @param InSize     The size of this paint element in SlateUnits.
 	 * @param InScale    Additional scaling to apply to the DrawElements.
-	 *
 	 * @return a FPaintGeometry derived this FGeometry.
 	 */
 	FPaintGeometry ToPaintGeometry( FVector2D InOffset, FVector2D InSize, float InScale ) const;
diff --git a/Engine/Source/Runtime/SlateCore/Public/Layout/Margin.h b/Engine/Source/Runtime/SlateCore/Public/Layout/Margin.h
--- a/Engine/Source/Runtime/SlateCore/Public/Layout/Margin.h
+++ b/Engine/Source/Runtime/SlateCore/Public/Layout/Margin.h
@@ -1,9 +1,5 @@
 // Copyright 1998-2014 Epic Games, Inc. All Rights Reserved.
 
-/*=============================================================================
-	Margin.h: Declares the FMargin structure.
-=============================================================================*/
-
 #pragma once
 
 #include "Margin.generated.h"
@@ -77,7 +73,6 @@ public:
 	 * Multiply the margin by a scalar.
 	 *
 	 * @param Scale How much to scale the margin.
-	 *
 	 * @return An FMargin where each value is scaled by Scale.
 	 */
 	FMargin operator*( float Scale ) const
@@ -89,7 +84,6 @@ public:
 	 * Adds another margin to this margin.
 	 *
 	 * @param Other The margin to add.
-	 *
 	 * @return A margin that represents this margin plus the other margin.
 	 */
 	FMargin operator+( const FMargin& InDelta ) const
@@ -101,7 +95,6 @@ public:
 	 * Subtracts another margin from this margin.
 	 *
 	 * @param Other The margin to subtract.
-	 *
 	 * @return A margin that represents this margin minues the other margin.
 	 */
 	FMargin operator-( const FMargin& Other ) const
@@ -113,7 +106,6 @@ public:
 	 * Compares this margin with another for equality.
 	 *
 	 * @param Other The other margin.
-	 *
 	 * @return true if the two margins are equal, false otherwise.
 	 */
 	bool operator==( const FMargin& Other ) const 
@@ -125,7 +117,6 @@ public:
 	 * Compares this margin with another for inequality.
 	 *
 	 * @param Other The other margin.
-	 *
 	 * @return true if the two margins are not equal, false otherwise.
 	 */
 	bool operator!=( const FMargin& Other ) const 
diff --git a/Engine/Source/Runtime/SlateCore/Public/Layout/PaintGeometry.h b/Engine/Source/Runtime/SlateCore/Public/Layout/PaintGeometry.h
--- a/Engine/Source/Runtime/SlateCore/Public/Layout/PaintGeometry.h
+++ b/Engine/Source/Runtime/SlateCore/Public/Layout/PaintGeometry.h
@@ -1,9 +1,5 @@
 // Copyright 1998-2014 Epic Games, Inc. All Rights Reserved.
 
-/*=============================================================================
-	PaintGeometry.h: Declares the FPaintGeometry structure.
-=============================================================================*/
-
 #pragma once
 
 
diff --git a/Engine/Source/Runtime/SlateCore/Public/Layout/SlateRect.h b/Engine/Source/Runtime/SlateCore/Public/Layout/SlateRect.h
--- a/Engine/Source/Runtime/SlateCore/Public/Layout/SlateRect.h
+++ b/Engine/Source/Runtime/SlateCore/Public/Layout/SlateRect.h
@@ -1,9 +1,5 @@
 // Copyright 1998-2014 Epic Games, Inc. All Rights Reserved.
 
-/*=============================================================================
-	SlateRect.h: Declares the FSlateRect structure.
-=============================================================================*/
-
 #pragma once
 
 
diff --git a/Engine/Source/Runtime/SlateCore/Public/Layout/Visibility.h b/Engine/Source/Runtime/SlateCore/Public/Layout/Visibility.h
--- a/Engine/Source/Runtime/SlateCore/Public/Layout/Visibility.h
+++ b/Engine/Source/Runtime/SlateCore/Public/Layout/Visibility.h
@@ -1,9 +1,5 @@
 // Copyright 1998-2014 Epic Games, Inc. All Rights Reserved.
 
-/*=============================================================================
-	Visibility.h: Declares the FVisibility class.
-=============================================================================*/
-
 #pragma once
 
 
diff --git a/Engine/Source/Runtime/SlateCore/Public/Logging/EventLogger.h b/Engine/Source/Runtime/SlateCore/Public/Logging/EventLogger.h
--- a/Engine/Source/Runtime/SlateCore/Public/Logging/EventLogger.h
+++ b/Engine/Source/Runtime/SlateCore/Public/Logging/EventLogger.h
@@ -1,9 +1,5 @@
 // Copyright 1998-2014 Epic Games, Inc. All Rights Reserved.
 
-/*=============================================================================
-	EventLogger.h: Declares various Slate event logger implementations.
-=============================================================================*/
-
 #pragma once
 
 
@@ -29,7 +25,7 @@ public:
 
 public:
 
-	// Begin IEventLogger interface
+	// IEventLogger interface
 
 	virtual FString GetLog( ) const override
 	{
@@ -48,8 +44,6 @@ public:
 		LoggerB->SaveToFile();
 	}
 
-	// EndIEventLogger interface
-
 private:
 
 	TSharedRef<IEventLogger> LoggerA;
@@ -65,16 +59,12 @@ class SLATECORE_API FFileEventLogger
 {
 public:
 
-	// Begin IEventLogger interface
+	// IEventLogger interface
 
 	virtual FString GetLog( ) const override;
-
 	virtual void Log( EEventLog::Type Event, const FString& AdditionalContent, TSharedPtr<SWidget> Widget ) override;
-	
 	virtual void SaveToFile( ) override;
 
-	// EndIEventLogger interface
-
 private:
 
 	// Holds the collection of logged events.
@@ -87,17 +77,11 @@ class SLATECORE_API FConsoleEventLogger
 {
 public:
 
-	// Begin IEventLogger interface
+	// IEventLogger interface
 
 	virtual FString GetLog( ) const override;
-
 	virtual void Log( EEventLog::Type Event, const FString& AdditionalContent, TSharedPtr<SWidget> Widget ) override;
-	
 	virtual void SaveToFile( ) override { }
-	
-	// End IEventLogger interface
-
-private:
 };
 
 
@@ -109,15 +93,11 @@ class SLATECORE_API FStabilityEventLogger
 {
 public:
 
-	// Begin IEventLogger interface
+	// IEventLogger interface
 	
 	virtual FString GetLog( ) const override;
-
 	virtual void Log( EEventLog::Type Event, const FString& AdditionalContent, TSharedPtr<SWidget> Widget ) override;
-	
 	virtual void SaveToFile( ) override { }
-	
-	// End IEventLogger interface
 
 private:
 
diff --git a/Engine/Source/Runtime/SlateCore/Public/Logging/IEventLogger.h b/Engine/Source/Runtime/SlateCore/Public/Logging/IEventLogger.h
--- a/Engine/Source/Runtime/SlateCore/Public/Logging/IEventLogger.h
+++ b/Engine/Source/Runtime/SlateCore/Public/Logging/IEventLogger.h
@@ -1,9 +1,5 @@
 // Copyright 1998-2014 Epic Games, Inc. All Rights Reserved.
 
-/*=============================================================================
-	IEventLogger.h: Declares the IEventLogger interface.
-=============================================================================*/
-
 #pragma once
 
 
@@ -41,7 +37,6 @@ namespace EEventLog
 	 * Converts a log event to a string.
 	 *
 	 * @param Event The event to convert.
-	 *
 	 * @return The event's string representation.
 	 */
 	inline FString ToString(EEventLog::Type Event)
@@ -107,8 +102,6 @@ public:
 
 public:
 
-	/**
-	 * Virtual destructor.
-	 */
+	/** Virtual destructor. */
 	virtual ~IEventLogger( ) { }
 };
diff --git a/Engine/Source/Runtime/SlateCore/Public/Rendering/ElementBatcher.h b/Engine/Source/Runtime/SlateCore/Public/Rendering/ElementBatcher.h
--- a/Engine/Source/Runtime/SlateCore/Public/Rendering/ElementBatcher.h
+++ b/Engine/Source/Runtime/SlateCore/Public/Rendering/ElementBatcher.h
@@ -1,13 +1,11 @@
 // Copyright 1998-2014 Epic Games, Inc. All Rights Reserved.
 
-/*=============================================================================
-	ElementBatcher.h: Declares the FSlateElementBatcher class.
-=============================================================================*/
-
 #pragma once
 
+
 class FSlateShaderResource;
 
+
 /**
  * A class which batches Slate elements for rendering
  */
diff --git a/Engine/Source/Runtime/SlateCore/Public/Rendering/RenderingPolicy.h b/Engine/Source/Runtime/SlateCore/Public/Rendering/RenderingPolicy.h
--- a/Engine/Source/Runtime/SlateCore/Public/Rendering/RenderingPolicy.h
+++ b/Engine/Source/Runtime/SlateCore/Public/Rendering/RenderingPolicy.h
@@ -1,9 +1,5 @@
 // Copyright 1998-2014 Epic Games, Inc. All Rights Reserved.
 
-/*=============================================================================
-	RenderingPolicy.h: Declares the FRenderingPolicy class.
-=============================================================================*/
-
 #pragma once
 
 
@@ -28,9 +24,7 @@ public:
 		: PixelCenterOffset( InPixelCenterOffset )
 	{ }
 
-	/**
-	 * Virtual constructor.
-	 */
+	/** Virtual constructor. */
 	virtual ~FSlateRenderingPolicy( ) { }
 
 	virtual void UpdateBuffers( const FSlateWindowElementList& WindowElementList ) = 0;
@@ -45,8 +39,12 @@ public:
 	}
 
 private:
+
 	FSlateRenderingPolicy(const FSlateRenderingPolicy&);
 	FSlateRenderingPolicy& operator=(const FSlateRenderingPolicy&);
+
 private:
+
 	float PixelCenterOffset;
+
 };
diff --git a/Engine/Source/Runtime/SlateCore/Public/Rendering/SlateDrawBuffer.h b/Engine/Source/Runtime/SlateCore/Public/Rendering/SlateDrawBuffer.h
--- a/Engine/Source/Runtime/SlateCore/Public/Rendering/SlateDrawBuffer.h
+++ b/Engine/Source/Runtime/SlateCore/Public/Rendering/SlateDrawBuffer.h
@@ -1,9 +1,5 @@
 // Copyright 1998-2014 Epic Games, Inc. All Rights Reserved.
 
-/*=============================================================================
-	SlateDrawBuffer.h: Declares the FSlateDrawBuffer class.
-=============================================================================*/
-
 #pragma once
 
 
@@ -14,18 +10,14 @@ class SLATECORE_API FSlateDrawBuffer
 {
 public:
 
-	/**
-	 * Default constructor.
-	 */
+	/** Default constructor. */
 	explicit FSlateDrawBuffer( )
 		: Locked(0)
 	{ }
 
 public:
 
-	/**
-	 * Removes all data from the buffer.
-	 */
+	/** Removes all data from the buffer. */
 	void ClearBuffer( )
 	{
 		WindowElementLists.Empty();
@@ -47,14 +39,17 @@ public:
 	}
 
 	/** 
-	 * Locks the draw buffer.  Indicates that the viewport is in use
+	 * Locks the draw buffer.  Indicates that the viewport is in use.
 	 *
-	 * @return true if the viewport could be locked.  False otherwise
+	 * @return true if the viewport could be locked.  False otherwise.
+	 * @see Unlock
 	 */
 	bool Lock( );
 
 	/**
-	 * Unlocks the buffer.  Indicates that the buffer is free                   
+	 * Unlocks the buffer.  Indicates that the buffer is free.
+	 *
+	 * @see Lock
 	 */
 	void Unlock( );
 
diff --git a/Engine/Source/Runtime/SlateCore/Public/Rendering/SlateRenderer.h b/Engine/Source/Runtime/SlateCore/Public/Rendering/SlateRenderer.h
--- a/Engine/Source/Runtime/SlateCore/Public/Rendering/SlateRenderer.h
+++ b/Engine/Source/Runtime/SlateCore/Public/Rendering/SlateRenderer.h
@@ -1,26 +1,25 @@
 // Copyright 1998-2014 Epic Games, Inc. All Rights Reserved.
 
-/*=============================================================================
-	SlateRenderer.h: Declares the FSlateRenderer class.
-=============================================================================*/
-
 #pragma once
 
 
 class SWindow;
 class FRHITexture2D;
+
 typedef FRHITexture2D* FTexture2DRHIParamRef;
 
+
 /**
  * Abstract base class for Slate renderers.
  */
 class FSlateRenderer
 {
 public:
-	FSlateRenderer(){}
-	/**
-	 * Virtual destructor.
-	 */
+
+	/** Default constructor. */
+	FSlateRenderer() { }
+
+	/** Virtual destructor. */
 	virtual ~FSlateRenderer( ) { }
 
 public:
@@ -76,7 +75,6 @@ public:
 	 * @param	bPrimaryWorkAreaOnly	True if we should capture only the primary monitor's work area, or false to capture the entire desktop spanning all monitors
 	 * @param	ScreenScaling	How much to downscale the desktop size
 	 * @param	LiveStreamingService	Optional pointer to a live streaming service this buffer needs to work with
-	 *
 	 * @return	The virtual screen rectangle.  The size of this rectangle will be the size of the render target buffer.
 	 */
 	virtual FIntRect SetupVirtualScreenBuffer( const bool bPrimaryWorkAreaOnly, const float ScreenScaling, class ILiveStreamingService* LiveStreamingService) { return FIntRect(); }
@@ -197,9 +195,11 @@ public:
 	virtual void SetWindowRenderTarget(const SWindow& Window, FTexture2DRHIParamRef RT) {}
 
 private:
+
 	// Non-copyable
 	FSlateRenderer(const FSlateRenderer&);
 	FSlateRenderer& operator=(const FSlateRenderer&);
+
 protected:
 
 	TSharedPtr<class FSlateFontCache> FontCache;
diff --git a/Engine/Source/Runtime/SlateCore/Public/SlateCore.h b/Engine/Source/Runtime/SlateCore/Public/SlateCore.h
--- a/Engine/Source/Runtime/SlateCore/Public/SlateCore.h
+++ b/Engine/Source/Runtime/SlateCore/Public/SlateCore.h
@@ -1,9 +1,5 @@
 // Copyright 1998-2014 Epic Games, Inc. All Rights Reserved.
 
-/*=============================================================================
-	SlateCore.h: SlateCore module public header file.
-=============================================================================*/
-
 #pragma once
 
 
@@ -60,6 +56,7 @@ DECLARE_MEMORY_STAT_EXTERN(TEXT("Index Buffer Memory"), STAT_SlateIndexBufferMem
 // Compile all the RichText and MultiLine editable text?
 #define WITH_FANCY_TEXT 1
 
+
 /* Includes
  *****************************************************************************/
 
diff --git a/Engine/Source/Runtime/SlateCore/Public/Sound/SlateSound.h b/Engine/Source/Runtime/SlateCore/Public/Sound/SlateSound.h
--- a/Engine/Source/Runtime/SlateCore/Public/Sound/SlateSound.h
+++ b/Engine/Source/Runtime/SlateCore/Public/Sound/SlateSound.h
@@ -1,9 +1,5 @@
 // Copyright 1998-2014 Epic Games, Inc. All Rights Reserved.
 
-/*=============================================================================
-	SlateSound.h: Declares the FSlateSound structure.
-=============================================================================*/
-
 #pragma once
 
 #include "SlateSound.generated.h"
@@ -19,9 +15,7 @@ struct SLATECORE_API FSlateSound
 
 public:
 
-	/**
-	 * Default constructor.
-	 */
+	/** Default constructor. */
 	FSlateSound( )
 		: ResourceObject(nullptr)
 	{ }
@@ -31,7 +25,6 @@ public:
 	/** Get the resource object associated with this sound
 	 *
 	 * @note Ensure that you only access the resource as a USoundBase 
-	 *
 	 * @return The resource object, or null if this sound resource is empty
 	 */
 	UObject* GetResourceObject( ) const;
@@ -50,10 +43,9 @@ public:
 
 	/**
 	 * Construct a FSlateSound from a FName
-	 * @note This functionality is deprecated and should only be used for upgrading old data
 	 *
+	 * @note This functionality is deprecated and should only be used for upgrading old data
 	 * @param SoundName The name of the sound to load, eg) SoundCue'/Game/QA_Assets/Sound/TEST_Music_Ambient.TEST_Music_Ambient'
-	 *
 	 * @return An FSlateSound object with a ResourceObject pointer that's either set the given sound, or points to null if the sound could not be found
 	 */
 	static FSlateSound FromName_DEPRECATED( const FName& SoundName );
diff --git a/Engine/Source/Runtime/SlateCore/Public/Styling/CoreStyle.h b/Engine/Source/Runtime/SlateCore/Public/Styling/CoreStyle.h
--- a/Engine/Source/Runtime/SlateCore/Public/Styling/CoreStyle.h
+++ b/Engine/Source/Runtime/SlateCore/Public/Styling/CoreStyle.h
@@ -1,9 +1,5 @@
 // Copyright 1998-2014 Epic Games, Inc. All Rights Reserved.
 
-/*=============================================================================
-	CoreStyle.h: Declares the FCoreStyle class.
-=============================================================================*/
-
 #pragma once
 
 
diff --git a/Engine/Source/Runtime/SlateCore/Public/Styling/SlateBrush.h b/Engine/Source/Runtime/SlateCore/Public/Styling/SlateBrush.h
--- a/Engine/Source/Runtime/SlateCore/Public/Styling/SlateBrush.h
+++ b/Engine/Source/Runtime/SlateCore/Public/Styling/SlateBrush.h
@@ -1,9 +1,5 @@
 // Copyright 1998-2014 Epic Games, Inc. All Rights Reserved.
 
-/*=============================================================================
-	SlateBrush.h: Declares the FSlateBrush structure.
-=============================================================================*/
-
 #pragma once
 
 #include "Margin.h"
@@ -200,7 +196,6 @@ public:
 	 * Gets the brush's tint color.
 	 *
 	 * @param InWidgetStyle The widget style to get the tint for.
-	 *
 	 * @return Tint color.
 	 */
 	FLinearColor GetTint( const FWidgetStyle& InWidgetStyle ) const
diff --git a/Engine/Source/Runtime/SlateCore/Public/Styling/SlateColor.h b/Engine/Source/Runtime/SlateCore/Public/Styling/SlateColor.h
--- a/Engine/Source/Runtime/SlateCore/Public/Styling/SlateColor.h
+++ b/Engine/Source/Runtime/SlateCore/Public/Styling/SlateColor.h
@@ -1,9 +1,5 @@
 // Copyright 1998-2014 Epic Games, Inc. All Rights Reserved.
 
-/*=============================================================================
-	SlateColor.h: Declares the FSlateColor structure.
-=============================================================================*/
-
 #pragma once
 
 #include "SlateColor.generated.h"
@@ -81,9 +77,7 @@ public:
 	 * Gets the color value represented by this Slate color.
 	 *
 	 * @param InWidgetStyle The widget style to use when this color represents a foreground or subdued color.
-	 *
 	 * @return The color value.
-	 *
 	 * @see GetSpecifiedColor
 	 */
 	const FLinearColor& GetColor( const FWidgetStyle& InWidgetStyle ) const
@@ -113,9 +107,7 @@ public:
 	 * Gets the specified color value.
 	 *
 	 * @return The specified color value.
-	 *
-	 * @see GetColor
-	 * @see IsColorSpecified
+	 * @see GetColor, IsColorSpecified
 	 */
 	FLinearColor GetSpecifiedColor( ) const
 	{
@@ -131,7 +123,6 @@ public:
 	 * Checks whether the values for this color have been specified.
 	 *
 	 * @return true if specified, false otherwise.
-	 *
 	 * @see GetSpecifiedColor
 	 */
 	bool IsColorSpecified( ) const
@@ -141,8 +132,7 @@ public:
 
 public:
 
-	/**
-	 * @returns an FSlateColor that is the widget's foreground. */
+	/** @returns an FSlateColor that is the widget's foreground. */
 	static FSlateColor UseForeground( )
 	{
 		return FSlateColor( ESlateColorStylingMode::UseColor_Foreground );
diff --git a/Engine/Source/Runtime/SlateCore/Public/Styling/SlateStyleRegistry.h b/Engine/Source/Runtime/SlateCore/Public/Styling/SlateStyleRegistry.h
--- a/Engine/Source/Runtime/SlateCore/Public/Styling/SlateStyleRegistry.h
+++ b/Engine/Source/Runtime/SlateCore/Public/Styling/SlateStyleRegistry.h
@@ -1,9 +1,5 @@
 // Copyright 1998-2014 Epic Games, Inc. All Rights Reserved.
 
-/*=============================================================================
-	SlateStyleRegistry.h: Declares the FSlateStyleRegistry class.
-=============================================================================*/
-
 #pragma once
 
 
@@ -13,9 +9,11 @@
 class SLATECORE_API FSlateStyleRegistry
 {
 public:
+
 	/**
 	 * Add a slate style to the repository.
-	 * @param ISlateStyle - interface to a user-definable style to add.
+	 *
+	 * @param ISlateStyle interface to a user-definable style to add.
 	 */
 	static void RegisterSlateStyle( const ISlateStyle& InSlateStyle );
 
@@ -35,14 +33,16 @@ public:
 
 	/**
 	 * Find a slate style in the repository.
-	 * @param InSlateStyleName - The name of the slate style to find.
+	 *
+	 * @param InSlateStyleName The name of the slate style to find.
 	 * @return The slate style, or null if it couldn't be found
 	 */
 	static const ISlateStyle* FindSlateStyle( const FName& InSlateStyleName );
 
 	/**
 	 * Populate an array of slate brushes with all of the resources used by the registered styles.
-	 * @param OutResources - array of slate brushes to populate.
+	 *
+	 * @param OutResources Array of slate brushes to populate.
 	 */
 	static void GetAllResources( TArray<const FSlateBrush*>& OutResources );
 
diff --git a/Engine/Source/Runtime/SlateCore/Public/Styling/SlateWidgetStyle.h b/Engine/Source/Runtime/SlateCore/Public/Styling/SlateWidgetStyle.h
--- a/Engine/Source/Runtime/SlateCore/Public/Styling/SlateWidgetStyle.h
+++ b/Engine/Source/Runtime/SlateCore/Public/Styling/SlateWidgetStyle.h
@@ -1,9 +1,5 @@
 // Copyright 1998-2014 Epic Games, Inc. All Rights Reserved.
 
-/*=============================================================================
-	SlateWidgetStyle.h: Declares the FSlateWidgetStyle structure.
-=============================================================================*/
-
 #pragma once
 
 #include "SlateWidgetStyle.generated.h"
@@ -19,14 +15,10 @@ struct SLATECORE_API FSlateWidgetStyle
 
 public:
 
-	/**
-	 * Default constructor.
-	 */
+	/** Default constructor. */
 	FSlateWidgetStyle( );
 
-	/**
-	 * Virtual Destructor.
-	 */
+	/** Virtual Destructor. */
 	virtual ~FSlateWidgetStyle() { }
 
 public:
diff --git a/Engine/Source/Runtime/SlateCore/Public/Styling/WidgetStyle.h b/Engine/Source/Runtime/SlateCore/Public/Styling/WidgetStyle.h
--- a/Engine/Source/Runtime/SlateCore/Public/Styling/WidgetStyle.h
+++ b/Engine/Source/Runtime/SlateCore/Public/Styling/WidgetStyle.h
@@ -1,9 +1,5 @@
 // Copyright 1998-2014 Epic Games, Inc. All Rights Reserved.
 
-/*=============================================================================
-	WidgetStyle.h: Declares the FWidgetStyle structure.
-=============================================================================*/
-
 #pragma once
 
 
@@ -14,9 +10,7 @@ class SLATECORE_API FWidgetStyle
 {
 public:
 	
-	/**
-	 * Default constructor.
-	 */
+	/** Default constructor. */
 	FWidgetStyle( )
 		: ColorAndOpacityTint(FLinearColor::White)
 		, ForegroundColor(FLinearColor::White)
@@ -29,7 +23,6 @@ public:
 	 * Blends the current tint color with the specified tint.
 	 *
 	 * @param InTint The color to blend with.
-	 *
 	 * @return This instance (for method chaining).
 	 */
 	FWidgetStyle& BlendColorAndOpacityTint( const FLinearColor& InTint )
@@ -42,9 +35,7 @@ public:
 	 * Sets the current foreground color from the given linear color.
 	 *
 	 * @param InForeground The foreground color value to set.
-	 *
 	 * @return This instance (for method chaining).
-	 *
 	 * @see GetForegroundColor
 	 */
 	FWidgetStyle& SetForegroundColor( const FLinearColor& InForeground )
@@ -63,9 +54,7 @@ public:
 	 * If the attribute passed in is not set, the foreground will remain unchanged.
 	 *
 	 * @param InForeground The foreground color to set.
-	 *
 	 * @return This instance (for method chaining).
-	 *
 	 * @see GetForegroundColor
 	 */
 	FWidgetStyle& SetForegroundColor( const TAttribute<struct FSlateColor>& InForeground );
@@ -76,9 +65,7 @@ public:
 	 * Gets the style's color opacity and tint.
 	 *
 	 * @return Current opacity and tint.
-	 *
-	 * @see GetForegroundColor
-	 * @see GetSubduedForegroundColor
+	 * @see GetForegroundColor, GetSubduedForegroundColor
 	 */
 	const FLinearColor& GetColorAndOpacityTint( ) const
 	{
@@ -89,10 +76,7 @@ public:
 	 * Gets the style's foreground color.
 	 *
 	 * @return Current foreground color.
-	 *
-	 * @see GetColorAndOpacityTint
-	 * @see GetSubduedForegroundColor
-	 * @see SetForegroundColor
+	 * @see GetColorAndOpacityTint, GetSubduedForegroundColor, SetForegroundColor
 	 */
 	const FLinearColor& GetForegroundColor( ) const
 	{
@@ -103,9 +87,7 @@ public:
 	 * Gets the style's subdued color.
 	 *
 	 * @return Current subdued color.
-	 *
-	 * @see GetColorAndOpacityTint
-	 * @see GetForegroundColor
+	 * @see GetColorAndOpacityTint, GetForegroundColor
 	 */
 	const FLinearColor& GetSubduedForegroundColor() const
 	{
diff --git a/Engine/Source/Runtime/SlateCore/Public/Textures/SlateIcon.h b/Engine/Source/Runtime/SlateCore/Public/Textures/SlateIcon.h
--- a/Engine/Source/Runtime/SlateCore/Public/Textures/SlateIcon.h
+++ b/Engine/Source/Runtime/SlateCore/Public/Textures/SlateIcon.h
@@ -1,13 +1,11 @@
 // Copyright 1998-2014 Epic Games, Inc. All Rights Reserved.
 
-/*=============================================================================
-	SlateIcon.h: Declares the FSlateIcon structure.
-=============================================================================*/
-
 #pragma once
 
+
 struct FSlateBrush;
 
+
 /**
  * Struct used to represent an icon in Slate
  */
@@ -41,7 +39,6 @@ public:
 	 * Gets the resolved icon.
 	 *
 	 * @return Icon brush, or FStyleDefaults::GetNoBrush() if the icon wasn't found.
-	 *
 	 * @see GetSmallIcon
 	 */
 	const FSlateBrush* GetIcon( ) const;
@@ -50,7 +47,6 @@ public:
 	 * Gets the resolved small icon.
 	 *
 	 * @return Icon brush, or FStyleDefaults::GetNoBrush() if the icon wasn't found.
-	 *
 	 * @see GetIcon
 	 */
 	const FSlateBrush* GetSmallIcon( ) const;
@@ -59,10 +55,7 @@ public:
 	 * Gets the name of the style for the icon.
 	 *
 	 * @return Style name.
-	 *
-	 * @see GetStyleName
-	 * @see GetStyleSet
-	 * @see GetStyleSetName
+	 * @see GetStyleName, GetStyleSet, GetStyleSetName
 	 */
 	const FName& GetSmallStyleName( ) const
 	{
@@ -73,10 +66,7 @@ public:
 	 * Gets the name of the style for the icon.
 	 *
 	 * @return Style name.
-	 *
-	 * @see GetSmallStyleName
-	 * @see GetStyleSet
-	 * @see GetStyleSetName
+	 * @see GetSmallStyleName, GetStyleSet, GetStyleSetName
 	 */
 	const FName& GetStyleName( ) const
 	{
@@ -87,10 +77,7 @@ public:
 	 * Gets the resolved style set.
 	 *
 	 * @return Style set, or nullptr if the style set wasn't found.
-	 *
-	 * @see GetSmallStyleName
-	 * @see GetStyleName
-	 * @see GetStyleSetName
+	 * @see GetSmallStyleName, GetStyleName, GetStyleSetName
 	 */
 	const class ISlateStyle* GetStyleSet( ) const;
 
@@ -98,10 +85,7 @@ public:
 	 * Gets the name of the style set the icon can be found in.
 	 *
 	 * @return Style name.
-	 *
-	 * @see GetSmallStyleName
-	 * @see GetStyleName
-	 * @see GetStyleSet
+	 * @see GetSmallStyleName, GetStyleName, GetStyleSet
 	 */
 	const FName& GetStyleSetName( ) const
 	{
diff --git a/Engine/Source/Runtime/SlateCore/Public/Textures/SlateShaderResource.h b/Engine/Source/Runtime/SlateCore/Public/Textures/SlateShaderResource.h
--- a/Engine/Source/Runtime/SlateCore/Public/Textures/SlateShaderResource.h
+++ b/Engine/Source/Runtime/SlateCore/Public/Textures/SlateShaderResource.h
@@ -1,9 +1,5 @@
 // Copyright 1998-2014 Epic Games, Inc. All Rights Reserved.
 
-/*=============================================================================
-	SlateShaderResource.h: Declares various shader resources types.
-=============================================================================*/
-
 #pragma once
 
 
@@ -53,9 +49,7 @@ public:
 
 public:
 
-	/**
-	 * Virtual destructor.
-	 */
+	/** Virtual destructor. */
 	virtual ~FSlateShaderResource() { }
 };
 
@@ -82,9 +76,7 @@ public:
 	/** The size of the texture.  Regardless of atlasing this is the size of the actual texture */
 	FIntPoint ActualSize;
 
-	/**
-	 * Default constructor.
-	 */
+	/** Default constructor. */
 	FSlateShaderResourceProxy( )
 		: StartUV(0.0f, 0.0f)
 		, SizeUV(1.0f, 1.0f)
@@ -103,9 +95,7 @@ class TSlateTexture
 {
 public:
 
-	/**
-	 * Default constructor.
-	 */
+	/** Default constructor. */
 	TSlateTexture( ) { }
 
 	/**
@@ -133,15 +123,13 @@ public:
 
 public:
 
-	// Begin FSlateShaderResource interface
+	// FSlateShaderResource interface
 
 	virtual ESlateShaderResource::Type GetType() const override
 	{
 		return ESlateShaderResource::Texture;
 	}
 
-	// End FSlateShaderResource interface
-
 protected:
 
 	// Holds the resource.
diff --git a/Engine/Source/Runtime/SlateCore/Public/Textures/SlateTextureData.h b/Engine/Source/Runtime/SlateCore/Public/Textures/SlateTextureData.h
--- a/Engine/Source/Runtime/SlateCore/Public/Textures/SlateTextureData.h
+++ b/Engine/Source/Runtime/SlateCore/Public/Textures/SlateTextureData.h
@@ -1,9 +1,5 @@
 // Copyright 1998-2014 Epic Games, Inc. All Rights Reserved.
 
-/*=============================================================================
-	FSlateTextureData.h: Declares the FFSlateTextureData structure.
-=============================================================================*/
-
 #pragma once
 
 
diff --git a/Engine/Source/Runtime/SlateCore/Public/Types/SlateConstants.h b/Engine/Source/Runtime/SlateCore/Public/Types/SlateConstants.h
--- a/Engine/Source/Runtime/SlateCore/Public/Types/SlateConstants.h
+++ b/Engine/Source/Runtime/SlateCore/Public/Types/SlateConstants.h
@@ -1,9 +1,5 @@
 // Copyright 1998-2014 Epic Games, Inc. All Rights Reserved.
 
-/*=============================================================================
-	SlateConstants.h: Declares various constants for Slate.
-=============================================================================*/
-
 #pragma once
 
 
diff --git a/Engine/Source/Runtime/SlateCore/Public/Types/SlateEnums.h b/Engine/Source/Runtime/SlateCore/Public/Types/SlateEnums.h
--- a/Engine/Source/Runtime/SlateCore/Public/Types/SlateEnums.h
+++ b/Engine/Source/Runtime/SlateCore/Public/Types/SlateEnums.h
@@ -1,13 +1,10 @@
 // Copyright 1998-2014 Epic Games, Inc. All Rights Reserved.
 
-/*=============================================================================
-	SlateEnums.h: Declares various enumeration types for Slate.
-=============================================================================*/
-
 #pragma once
 
 #include "SlateEnums.generated.h"
 
+
 UENUM()
 namespace EButtonClickMethod
 {
@@ -35,6 +32,7 @@ namespace EButtonClickMethod
 	};
 }
 
+
 UENUM()
 namespace EButtonTouchMethod
 {
@@ -52,6 +50,7 @@ namespace EButtonTouchMethod
 	};
 }
 
+
 UENUM()
 namespace EFocusMoveDirection
 {
@@ -156,6 +155,7 @@ enum EScrollDirection
 	Scroll_Up UMETA(DisplayName="Up"),
 };
 
+
 /**
  * Additional information about a text committal
  */
diff --git a/Engine/Source/Runtime/SlateCore/Public/Types/SlateStructs.h b/Engine/Source/Runtime/SlateCore/Public/Types/SlateStructs.h
--- a/Engine/Source/Runtime/SlateCore/Public/Types/SlateStructs.h
+++ b/Engine/Source/Runtime/SlateCore/Public/Types/SlateStructs.h
@@ -1,9 +1,5 @@
 // Copyright 1998-2014 Epic Games, Inc. All Rights Reserved.
 
-/*=============================================================================
-	SlateStructs.h: Declares various structured types for Slate.
-=============================================================================*/
-
 #pragma once
 
 
@@ -34,7 +30,6 @@ public:
 	 * Checks whether the size is set.
 	 *
 	 * @return true if the size is set, false if it is unespecified.
-	 *
 	 * @see Get
 	 */
 	bool IsSet( ) const
@@ -104,9 +99,7 @@ protected:
 	 *
 	 * Use FAspectRatio, FAuto, FStretch to instantiate size parameters.
 	 *
-	 * @see FAspectRatio
-	 * @see FAuto
-	 * @see FStretch
+	 * @see FAspectRatio, FAuto, FStretch
 	 */
 	FSizeParam( ESizeRule InTypeOfSize, const TAttribute<float>& InValue )
 		: SizeRule(InTypeOfSize)
@@ -137,8 +130,7 @@ struct FStretch
 /**
  * Structure for size parameters with SizeRule = SizeRule_Auto.
  *
- * @see FAspectRatio
- * @see FStretch
+ * @see FAspectRatio, FStretch
  */
 struct FAuto
 	: public FSizeParam
@@ -152,8 +144,7 @@ struct FAuto
 /**
  * Structure for size parameters with SizeRule == SizeRule_AspectRatio.
  *
- * @see FAuto
- * @see FStretch
+ * @see FAuto, FStretch
  */
 struct FAspectRatio
 	: public FSizeParam
diff --git a/Engine/Source/Runtime/SlateCore/Public/Widgets/DeclarativeSyntaxSupport.h b/Engine/Source/Runtime/SlateCore/Public/Widgets/DeclarativeSyntaxSupport.h
--- a/Engine/Source/Runtime/SlateCore/Public/Widgets/DeclarativeSyntaxSupport.h
+++ b/Engine/Source/Runtime/SlateCore/Public/Widgets/DeclarativeSyntaxSupport.h
@@ -1,9 +1,5 @@
 // Copyright 1998-2014 Epic Games, Inc. All Rights Reserved.
 
-/*=============================================================================
-	DeclarativeSyntaxSupport.h: Defines various macros and types for Slate's declarative syntax.
-=============================================================================*/
-
 #pragma once
 
 
diff --git a/Engine/Source/Runtime/SlateCore/Public/Widgets/IToolTip.h b/Engine/Source/Runtime/SlateCore/Public/Widgets/IToolTip.h
--- a/Engine/Source/Runtime/SlateCore/Public/Widgets/IToolTip.h
+++ b/Engine/Source/Runtime/SlateCore/Public/Widgets/IToolTip.h
@@ -1,9 +1,5 @@
 // Copyright 1998-2014 Epic Games, Inc. All Rights Reserved.
 
-/*=============================================================================
-	IToolTip.h: Declares the IToolTip interface.
-=============================================================================*/
-
 #pragma once
 
 
@@ -44,8 +40,6 @@ public:
 
 public:
 
-	/**
-	 * Virtual destructor.
-	 */
+	/** Virtual destructor. */
 	virtual ~IToolTip( ) { }
 };
diff --git a/Engine/Source/Runtime/SlateCore/Public/Widgets/SCompoundWidget.h b/Engine/Source/Runtime/SlateCore/Public/Widgets/SCompoundWidget.h
--- a/Engine/Source/Runtime/SlateCore/Public/Widgets/SCompoundWidget.h
+++ b/Engine/Source/Runtime/SlateCore/Public/Widgets/SCompoundWidget.h
@@ -1,9 +1,5 @@
 // Copyright 1998-2014 Epic Games, Inc. All Rights Reserved.
 
-/*=============================================================================
-	SCompoundWidget.h: Declares the SCompoundWidget class.
-=============================================================================*/
-
 #pragma once
 
 
@@ -65,20 +61,14 @@ public:
 
 public:
 
-	// Begin SWidgetOverrides
+	// SWidgetOverrides
 
 	virtual int32 OnPaint( const FPaintArgs& Args, const FGeometry& AllottedGeometry, const FSlateRect& MyClippingRect, FSlateWindowElementList& OutDrawElements, int32 LayerId, const FWidgetStyle& InWidgetStyle, bool bParentEnabled ) const override;
-	
 	virtual FChildren* GetChildren() override;
-
 	virtual FVector2D ComputeDesiredSize() const override;
-
 	virtual void OnArrangeChildren( const FGeometry& AllottedGeometry, FArrangedChildren& ArrangedChildren ) const override;
-
 	virtual FSlateColor GetForegroundColor() const override;
 
-	// Begin SWidgetOverrides
-
 protected:
 
 	/** Disallow public construction */
diff --git a/Engine/Source/Runtime/SlateCore/Public/Widgets/SLeafWidget.h b/Engine/Source/Runtime/SlateCore/Public/Widgets/SLeafWidget.h
--- a/Engine/Source/Runtime/SlateCore/Public/Widgets/SLeafWidget.h
+++ b/Engine/Source/Runtime/SlateCore/Public/Widgets/SLeafWidget.h
@@ -1,9 +1,5 @@
 // Copyright 1998-2014 Epic Games, Inc. All Rights Reserved.
 
-/*=============================================================================
-	SLeafWidget.h: Declares the SLeafWidget class.
-=============================================================================*/
-
 #pragma once
 
 
diff --git a/Engine/Source/Runtime/SlateCore/Public/Widgets/SNullWidget.h b/Engine/Source/Runtime/SlateCore/Public/Widgets/SNullWidget.h
--- a/Engine/Source/Runtime/SlateCore/Public/Widgets/SNullWidget.h
+++ b/Engine/Source/Runtime/SlateCore/Public/Widgets/SNullWidget.h
@@ -1,9 +1,5 @@
 // Copyright 1998-2014 Epic Games, Inc. All Rights Reserved.
 
-/*=============================================================================
-	SNullWidget.h: Declares the SNullWidget class.
-=============================================================================*/
-
 #pragma once
 
 
diff --git a/Engine/Source/Runtime/SlateCore/Public/Widgets/SPanel.h b/Engine/Source/Runtime/SlateCore/Public/Widgets/SPanel.h
--- a/Engine/Source/Runtime/SlateCore/Public/Widgets/SPanel.h
+++ b/Engine/Source/Runtime/SlateCore/Public/Widgets/SPanel.h
@@ -1,9 +1,5 @@
 // Copyright 1998-2014 Epic Games, Inc. All Rights Reserved.
 
-/*=============================================================================
-	SPanel.h: Declares the SPanel class.
-=============================================================================*/
-
 #pragma once
 
 
@@ -27,6 +23,7 @@ public:
 	 * @param ArrangedChildren    The array to which to add the WidgetGeometries that represent the arranged children.
 	 */
 	virtual void OnArrangeChildren( const FGeometry& AllottedGeometry, FArrangedChildren& ArrangedChildren ) const = 0;
+
 	/**
 	 * A Panel's desired size in the space required to arrange of its children on the screen while respecting all of
 	 * the children's desired sizes and any layout-related options specified by the user. See StackPanel for an example.
@@ -54,12 +51,10 @@ public:
 
 public:
 
-	// Begin SWidget overrides
+	// SWidget overrides
 
 	virtual int32 OnPaint( const FPaintArgs& Args, const FGeometry& AllottedGeometry, const FSlateRect& MyClippingRect, FSlateWindowElementList& OutDrawElements, int32 LayerId, const FWidgetStyle& InWidgetStyle, bool bParentEnabled ) const override;
 
-	// End SWidget overrides
-
 protected:
 
 	/**
@@ -69,8 +64,6 @@ protected:
 	
 protected:
 
-	/**
-	 * Hidden default constructor.
-	 */
+	/** Hidden default constructor. */
 	SPanel( ) { }
 };
diff --git a/Engine/Source/Runtime/SlateCore/Public/Widgets/SWidget.h b/Engine/Source/Runtime/SlateCore/Public/Widgets/SWidget.h
--- a/Engine/Source/Runtime/SlateCore/Public/Widgets/SWidget.h
+++ b/Engine/Source/Runtime/SlateCore/Public/Widgets/SWidget.h
@@ -1,9 +1,5 @@
 // Copyright 1998-2014 Epic Games, Inc. All Rights Reserved.
 
-/*=============================================================================
-	SWidget.h: Declares the SWidget class.
-=============================================================================*/
-
 #pragma once
 
 
@@ -104,7 +100,6 @@ public:
 	 * @param LayerId           The Layer onto which this widget should be rendered.
 	 * @param InColorAndOpacity Color and Opacity to be applied to all the descendants of the widget being painted
 	 * @param bParentEnabled	True if the parent of this widget is enabled.
-	 *
 	 * @return The maximum layer ID attained by this widget or any of its children.
 	 */
 	int32 Paint( const FPaintArgs& Args, const FGeometry& AllottedGeometry, const FSlateRect& MyClippingRect, FSlateWindowElementList& OutDrawElements, int32 LayerId, const FWidgetStyle& InWidgetStyle, bool bParentEnabled ) const;
@@ -127,7 +122,6 @@ public:
 	 *
 	 * @param MyGeometry The Geometry of the widget receiving the event
 	 * @param InAbsoluteCursorPosition	The absolute cursor position
-	 *
 	 * @return true if this widget considers itself hit by the passed in cursor position.
 	 */
 	virtual bool OnHitTest( const FGeometry& MyGeometry, FVector2D InAbsoluteCursorPosition );
@@ -141,7 +135,6 @@ public:
 	 *
 	 * @param MyGeometry The Geometry of the widget receiving the event
 	 * @param  InKeyboardFocusEvent  KeyboardFocusEvent
-	 *
 	 * @return  Returns whether the event was handled, along with other possible actions
 	 */
 	virtual FReply OnKeyboardFocusReceived( const FGeometry& MyGeometry, const FKeyboardFocusEvent& InKeyboardFocusEvent );
@@ -161,7 +154,6 @@ public:
 	 *
 	 * @param MyGeometry The Geometry of the widget receiving the event
 	 * @param  InCharacterEvent  Character event
-	 *
 	 * @return  Returns whether the event was handled, along with other possible actions
 	 */
 	virtual FReply OnKeyChar( const FGeometry& MyGeometry,const FCharacterEvent& InCharacterEvent );
@@ -175,7 +167,6 @@ public:
 	 *
 	 * @param MyGeometry The Geometry of the widget receiving the event
 	 * @param  InKeyboardEvent  Keyboard event
-	 *
 	 * @return  Returns whether the event was handled, along with other possible actions
 	 */
 	virtual FReply OnPreviewKeyDown( const FGeometry& MyGeometry, const FKeyboardEvent& InKeyboardEvent );
@@ -185,7 +176,6 @@ public:
 	 *
 	 * @param MyGeometry The Geometry of the widget receiving the event
 	 * @param  InKeyboardEvent  Keyboard event
-	 *
 	 * @return  Returns whether the event was handled, along with other possible actions
 	 */
 	virtual FReply OnKeyDown( const FGeometry& MyGeometry, const FKeyboardEvent& InKeyboardEvent );
@@ -195,7 +185,6 @@ public:
 	 *
 	 * @param MyGeometry The Geometry of the widget receiving the event
 	 * @param  InKeyboardEvent  Keyboard event
-	 *
 	 * @return  Returns whether the event was handled, along with other possible actions
 	 */
 	virtual FReply OnKeyUp( const FGeometry& MyGeometry, const FKeyboardEvent& InKeyboardEvent );
@@ -209,7 +198,6 @@ public:
 	 *
 	 * @param MyGeometry The Geometry of the widget receiving the event
 	 * @param MouseEvent Information about the input event
-	 *
 	 * @return Whether the event was handled along with possible requests for the system to take action.
 	 */
 	virtual FReply OnMouseButtonDown( const FGeometry& MyGeometry, const FPointerEvent& MouseEvent );
@@ -228,7 +216,6 @@ public:
 	 *
 	 * @param MyGeometry The Geometry of the widget receiving the event
 	 * @param MouseEvent Information about the input event
-	 *
 	 * @return Whether the event was handled along with possible requests for the system to take action.
 	 */
 	virtual FReply OnMouseButtonUp( const FGeometry& MyGeometry, const FPointerEvent& MouseEvent );
@@ -238,7 +225,6 @@ public:
 	 *
 	 * @param MyGeometry The Geometry of the widget receiving the event
 	 * @param MouseEvent Information about the input event
-	 *
 	 * @return Whether the event was handled along with possible requests for the system to take action.
 	 */
 	virtual FReply OnMouseMove( const FGeometry& MyGeometry, const FPointerEvent& MouseEvent );
@@ -262,7 +248,6 @@ public:
 	 * Called when the mouse wheel is spun. This event is bubbled.
 	 *
 	 * @param  MouseEvent  Mouse event
-	 *
 	 * @return  Returns whether the event was handled, along with other possible actions
 	 */
 	virtual FReply OnMouseWheel( const FGeometry& MyGeometry, const FPointerEvent& MouseEvent );
@@ -279,7 +264,6 @@ public:
 	 *
 	 * @param  InMyGeometry  Widget geometry
 	 * @param  InMouseEvent  Mouse button event
-	 *
 	 * @return  Returns whether the event was handled, along with other possible actions
 	 */
 	virtual FReply OnMouseButtonDoubleClick( const FGeometry& InMyGeometry, const FPointerEvent& InMouseEvent );
@@ -290,7 +274,6 @@ public:
 	 * If you override this event, you should probably return true.
 	 * 
 	 * @param  TooltipContent    The TooltipContent that I may want to visualize.
-	 *
 	 * @return true if this widget visualized the tooltip content; i.e., the event is handled.
 	 */
 	virtual bool OnVisualizeTooltip( const TSharedPtr<SWidget>& TooltipContent );
@@ -307,7 +290,6 @@ public:
 	 *
 	 * @param  InMyGeometry  Widget geometry
 	 * @param  InMouseEvent  MouseMove that triggered the drag
-	 * 
 	 */
 	virtual FReply OnDragDetected( const FGeometry& MyGeometry, const FPointerEvent& MouseEvent );
 
@@ -346,7 +328,6 @@ public:
 	 *
 	 * @param MyGeometry      The geometry of the widget receiving the event.
 	 * @param DragDropEvent   The drag and drop event.
-	 *
 	 * @return A reply that indicated whether this event was handled.
 	 */
 	virtual FReply OnDragOver( const FGeometry& MyGeometry, const FDragDropEvent& DragDropEvent );
@@ -356,7 +337,6 @@ public:
 	 *
 	 * @param MyGeometry      The geometry of the widget receiving the event.
 	 * @param DragDropEvent   The drag and drop event.
-	 *
 	 * @return A reply that indicated whether this event was handled.
 	 */
 	virtual FReply OnDrop( const FGeometry& MyGeometry, const FDragDropEvent& DragDropEvent );
@@ -386,7 +366,6 @@ public:
 	 * Called when the user performs a gesture on trackpad. This event is bubbled.
 	 *
 	 * @param  GestureEvent  gesture event
-	 *
 	 * @return  Returns whether the event was handled, along with other possible actions
 	 */
 	virtual FReply OnTouchGesture( const FGeometry& MyGeometry, const FPointerEvent& GestureEvent );
@@ -515,9 +494,7 @@ public:
 	 */
 	bool HasMouseCapture() const;
 
-	/**
-	 * Called when this widget had captured the mouse, but that capture has been revoked for some reason.
-	 */
+	/** Called when this widget had captured the mouse, but that capture has been revoked for some reason. */
 	virtual void OnMouseCaptureLost();
 
 	/**
@@ -612,9 +589,8 @@ public:
 	void SetDebugInfo( const ANSICHAR* InType, const ANSICHAR* InFile, int32 OnLine );
 
 public:
-	//
+
 	// Widget Inspector and debugging methods
-	//
 
 	/** @return A String representation of the widget */
 	virtual FString ToString() const;
@@ -655,7 +631,6 @@ protected:
 	 * @param MyGeometry      The geometry of this widget.
 	 * @param WidgetsToFind   The widgets whose geometries we wish to discover.
 	 * @param OutResult       A map of widget references to their respective geometries.
-	 *
 	 * @return True if all the WidgetGeometries were found. False otherwise.
 	 */
 	bool FindChildGeometries( const FGeometry& MyGeometry, const TSet< TSharedRef<SWidget> >& WidgetsToFind, TMap<TSharedRef<SWidget>, FArrangedWidget>& OutResult ) const;
@@ -674,7 +649,6 @@ protected:
 	 *
 	 * @param MyGeometry   The geometry of this widget.
 	 * @param WidgetToFind The widget whose geometry we wish to discover.
-	 *
 	 * @return the geometry of WidgetToFind.
 	 */
 	FGeometry FindChildGeometry( const FGeometry& MyGeometry, TSharedRef<SWidget> WidgetToFind ) const;
@@ -695,33 +669,33 @@ protected:
 	}
 
 private:
+
 	/**
-	* The widget should respond by populating the OutDrawElements array with FDrawElements
-	* that represent it and any of its children. Called by the non-virtual OnPaint to enforce pre/post conditions
-	* during OnPaint.
-	*
-	* @param Args              All the arguments necessary to paint this widget (@todo umg: move all params into this struct)
-	* @param AllottedGeometry  The FGeometry that describes an area in which the widget should appear.
-	* @param MyClippingRect    The clipping rectangle allocated for this widget and its children.
-	* @param OutDrawElements   A list of FDrawElements to populate with the output.
-	* @param LayerId           The Layer onto which this widget should be rendered.
-	* @param InColorAndOpacity Color and Opacity to be applied to all the descendants of the widget being painted
-	* @param bParentEnabled	True if the parent of this widget is enabled.
-	*
-	* @return The maximum layer ID attained by this widget or any of its children.
-	*/
+	 * The widget should respond by populating the OutDrawElements array with FDrawElements
+	 * that represent it and any of its children. Called by the non-virtual OnPaint to enforce pre/post conditions
+	 * during OnPaint.
+	 *
+	 * @param Args              All the arguments necessary to paint this widget (@todo umg: move all params into this struct)
+	 * @param AllottedGeometry  The FGeometry that describes an area in which the widget should appear.
+	 * @param MyClippingRect    The clipping rectangle allocated for this widget and its children.
+	 * @param OutDrawElements   A list of FDrawElements to populate with the output.
+	 * @param LayerId           The Layer onto which this widget should be rendered.
+	 * @param InColorAndOpacity Color and Opacity to be applied to all the descendants of the widget being painted
+	 * @param bParentEnabled	True if the parent of this widget is enabled.
+	 * @return The maximum layer ID attained by this widget or any of its children.
+	 */
 	virtual int32 OnPaint(const FPaintArgs& Args, const FGeometry& AllottedGeometry, const FSlateRect& MyClippingRect, FSlateWindowElementList& OutDrawElements, int32 LayerId, const FWidgetStyle& InWidgetStyle, bool bParentEnabled) const = 0;
 
 	DEPRECATED(4.4, "OnPaint() now requires an extra FPaintArgs parameter. When calling Superclass::OnPaint() you can simply pass through the existing Args.")
 	virtual int32 OnPaint(const FGeometry& AllottedGeometry, const FSlateRect& MyClippingRect, FSlateWindowElementList& OutDrawElements, int32 LayerId, const FWidgetStyle& InWidgetStyle, bool bParentEnabled) const { return 0; }
 
 	/**
-	* Compute the Geometry of all the children and add populate the ArrangedChildren list with their values.
-	* Each type of Layout panel should arrange children based on desired behavior.
-	*
-	* @param AllottedGeometry    The geometry allotted for this widget by its parent.
-	* @param ArrangedChildren    The array to which to add the WidgetGeometries that represent the arranged children.
-	*/
+	 * Compute the Geometry of all the children and add populate the ArrangedChildren list with their values.
+	 * Each type of Layout panel should arrange children based on desired behavior.
+	 *
+	 * @param AllottedGeometry    The geometry allotted for this widget by its parent.
+	 * @param ArrangedChildren    The array to which to add the WidgetGeometries that represent the arranged children.
+	 */
 	virtual void OnArrangeChildren(const FGeometry& AllottedGeometry, FArrangedChildren& ArrangedChildren) const = 0;
 
 protected:
@@ -769,7 +743,7 @@ private:
 	static FAutoConsoleVariableRef CVarUseLegacyHittest;
 
 public:
+
 	/** For perf testing. Old hittest schema does not support widgets outside their parents' bounds. */
 	static bool UseLegacyHittest();
-
 };
diff --git a/Engine/Source/Runtime/SlateCore/Public/Widgets/SWindow.h b/Engine/Source/Runtime/SlateCore/Public/Widgets/SWindow.h
--- a/Engine/Source/Runtime/SlateCore/Public/Widgets/SWindow.h
+++ b/Engine/Source/Runtime/SlateCore/Public/Widgets/SWindow.h
@@ -1,9 +1,5 @@
 // Copyright 1998-2014 Epic Games, Inc. All Rights Reserved.
 
-/*=============================================================================
-	SWindow.h: Declares the SWindow class.
-=============================================================================*/
-
 #pragma once
 
 
@@ -336,7 +332,6 @@ public:
 	 * Adds content to draw on top of the entire window
 	 *
 	 * @param	InZOrder	Z-order to use for this widget
-	 * 
 	 * @return The added overlay slot so that it can be configured and populated
 	 */
 	SOverlay::FOverlaySlot& AddOverlaySlot( const int32 ZOrder = INDEX_NONE );
@@ -609,12 +604,10 @@ public:
 
 public:
 
-	// Begin SWidget overrides
+	// SWidget overrides
 
 	virtual void Tick( const FGeometry& AllottedGeometry, const double InCurrentTime, const float InDeltaTime ) override;
 
-	// End SWidget overrides
-
 protected:
 
 	/** Get the desired color of titlebar items. These change during flashing. */
@@ -726,8 +719,7 @@ protected:
 			, bIsActive(false)
 			, bIsPendingPlay(false)
 			, bIsAnimatingWindowSize(false)
-		{
-		}
+		{ }
 
 		/** Initial window opacity */
 		float StartingOpacity;
@@ -769,6 +761,7 @@ protected:
 	const FSlateBrush* WindowBackground;
 
 private:
+
 	/** The native window that is backing this Slate Window */
 	TSharedPtr<FGenericWindow> NativeWindow;
 	
@@ -799,8 +792,6 @@ private:
 	/** Full window overlay widget */
 	TSharedPtr<SWidget> FullWindowOverlayWidget;
 
-
-
 	/** When not null, this window will always appear on top of the parent and be closed when the parent is closed. */
 	TWeakPtr<SWindow> ParentWindowPtr;
 
@@ -840,6 +831,7 @@ private:
 	EVisibility GetWindowFlashVisibility() const;
 };
 
+
 /**
  * Popups, tooltips, drag and drop decorators all can be executed without creating a new window.
  * This slot along with the SWindow::AddPopupLayerSlot() API enabled it.
diff --git a/Engine/Source/Runtime/SlateRHIRenderer/Private/SlateMaterialShader.h b/Engine/Source/Runtime/SlateRHIRenderer/Private/SlateMaterialShader.h
--- a/Engine/Source/Runtime/SlateRHIRenderer/Private/SlateMaterialShader.h
+++ b/Engine/Source/Runtime/SlateRHIRenderer/Private/SlateMaterialShader.h
@@ -1,24 +1,18 @@
 // Copyright 1998-2014 Epic Games, Inc. All Rights Reserved.
 
-/*=============================================================================
-	PostProcessMaterial.h: Post processing Material
-=============================================================================*/
-
 #pragma once
 
 #include "Runtime/Renderer/Public/MaterialShader.h"
 
+
 class FSlateMaterialShaderPS : public FMaterialShader
 {
 public:
-	/**
-	  * Only compile shaders used with UI
-	  */
+
+	/** Only compile shaders used with UI. */
 	static bool ShouldCache(EShaderPlatform Platform, const FMaterial* Material);
 
-	/**
-	 * Modifies the compilation of this shader
-	 */
+	/** Modifies the compilation of this shader. */
 	static void ModifyCompilationEnvironment(EShaderPlatform Platform, const FMaterial* Material, FShaderCompilerEnvironment& OutEnvironment);
 
 	FSlateMaterialShaderPS() {}
@@ -27,36 +21,33 @@ public:
 	void SetParameters(FRHICommandList& RHICmdList, const FSceneView& View, const FMaterialRenderProxy* MaterialRenderProxy, const FMaterial* Material, float InDisplayGamma, const FVector4& InShaderParams );
 
 	virtual bool Serialize(FArchive& Ar) override;
+
 private:
+
 	FShaderParameter DisplayGamma;
 	FShaderParameter ShaderParams;
 };
 
+
 template<ESlateShader::Type ShaderType,bool bDrawDisabledEffect> 
 class TSlateMaterialShaderPS : public FSlateMaterialShaderPS
 {
 public:
 	DECLARE_SHADER_TYPE(TSlateMaterialShaderPS,Material);
 
-	TSlateMaterialShaderPS() {}
+	TSlateMaterialShaderPS() { }
+
 	TSlateMaterialShaderPS(const ShaderMetaType::CompiledShaderInitializerType& Initializer)
 		: FSlateMaterialShaderPS( Initializer )
-	{
-	
-	}
+	{ }
 	
-	/**
-	  * Only compile shaders used with UI
-	  */
+	/** Only compile shaders used with UI. */
 	static bool ShouldCache(EShaderPlatform Platform, const FMaterial* Material)
 	{
 		return FSlateMaterialShaderPS::ShouldCache(Platform,Material);
 	}
 
-
-	/**
-	 * Modifies the compilation of this shader
-	 */
+	/** Modifies the compilation of this shader. */
 	static void ModifyCompilationEnvironment(EShaderPlatform Platform, const FMaterial* Material, FShaderCompilerEnvironment& OutEnvironment)
 	{
 		FSlateMaterialShaderPS::ModifyCompilationEnvironment(Platform, Material,OutEnvironment);
@@ -69,4 +60,4 @@ public:
 	{
 		return FSlateMaterialShaderPS::Serialize( Ar );
 	}
-};
\ No newline at end of file
+};
