Missed a file

diff --git a/Engine/Source/Runtime/SlateCore/Public/Input/Events.h b/Engine/Source/Runtime/SlateCore/Public/Input/Events.h
--- a/Engine/Source/Runtime/SlateCore/Public/Input/Events.h
+++ b/Engine/Source/Runtime/SlateCore/Public/Input/Events.h
@@ -96,7 +96,7 @@ public:
 	 * UStruct Constructor.  Not meant for normal usage.
 	 */
 	FInputEvent()
-		: ModifierKeys(FModifierKeysState(false, false, false, false, false, false, false, false))
+		: ModifierKeys(FModifierKeysState())
 		, bIsRepeat(false)
 		, EventPath(nullptr)
 	{ }
@@ -295,7 +295,7 @@ public:
 	 * UStruct Constructor.  Not meant for normal usage.
 	 */
 	FKeyboardEvent()
-		: FInputEvent(FModifierKeysState(false, false, false, false, false, false, false, false), false)
+		: FInputEvent(FModifierKeysState(), false)
 		, Key(EKeys::SpaceBar)
 		, CharacterCode(0)
 	{
@@ -362,7 +362,7 @@ public:
 	 * UStruct Constructor.  Not meant for normal usage.
 	 */
 	FCharacterEvent()
-		: FInputEvent(FModifierKeysState(false, false, false, false, false, false, false, false), false)
+		: FInputEvent(FModifierKeysState(), false)
 		, Character(0)
 	{
 	}
@@ -496,7 +496,7 @@ public:
 		const FVector2D& InScreenSpacePosition,
 		const FVector2D& InLastScreenSpacePosition,
 		bool bPressLeftMouseButton,
-		const FModifierKeysState& InModifierKeys = FModifierKeysState(false, false, false, false, false, false, false, false),
+		const FModifierKeysState& InModifierKeys = FModifierKeysState(),
 		uint32 InTouchpadIndex=0
 	)
 		: FInputEvent(InModifierKeys, false)
@@ -625,7 +625,7 @@ public:
 	{ }
 
 	FControllerEvent( FKey InEffectingButton, int32 InUserIndex, float InAnalogValue, bool bIsRepeat )
-		: FInputEvent(FModifierKeysState(false, false, false, false, false, false, false, false), bIsRepeat)
+		: FInputEvent(FModifierKeysState(), bIsRepeat)
 		, EffectingButton(InEffectingButton)
 		, UserIndex(InUserIndex)
 		, AnalogValue(InAnalogValue)
@@ -684,7 +684,7 @@ public:
 		const FVector& InGravity, 
 		const FVector& InAcceleration
 	)
-		: FInputEvent(FModifierKeysState(false, false, false, false, false, false, false, false), false)
+		: FInputEvent(FModifierKeysState(), false)
 		, UserIndex(InUserIndex)
 		, Tilt(InTilt)
 		, RotationRate(InRotationRate)
