Potential fix for the hittest grid issue.  Looks like we are not handling a case of huge numbers being passed in as the coordinate to check.

diff --git a/Engine/Source/Runtime/SlateCore/Private/Input/HittestGrid.cpp b/Engine/Source/Runtime/SlateCore/Private/Input/HittestGrid.cpp
--- a/Engine/Source/Runtime/SlateCore/Private/Input/HittestGrid.cpp
+++ b/Engine/Source/Runtime/SlateCore/Private/Input/HittestGrid.cpp
@@ -43,8 +43,8 @@ TArray<FArrangedWidget> FHittestGrid::GetBubblePath( FVector2D DesktopSpaceCoord
 	{
 		const FVector2D CursorPositionInGrid = DesktopSpaceCoordinate - GridOrigin;
 		const FIntPoint CellCoordinate = FIntPoint(
-			FMath::Max(FMath::FloorToInt(CursorPositionInGrid.X / CellSize.X), 0),
-			FMath::Max(FMath::FloorToInt(CursorPositionInGrid.Y / CellSize.Y), 0) );
+			FMath::Min( FMath::Max(FMath::FloorToInt(CursorPositionInGrid.X / CellSize.X), 0), NumCells.X ),
+			FMath::Min( FMath::Max(FMath::FloorToInt(CursorPositionInGrid.Y / CellSize.Y), 0), NumCells.Y ) );
 		
 		static FVector2D LastCoordinate = FVector2D::ZeroVector;
 		if ( LastCoordinate != CursorPositionInGrid )
@@ -52,6 +52,8 @@ TArray<FArrangedWidget> FHittestGrid::GetBubblePath( FVector2D DesktopSpaceCoord
 			LastCoordinate = CursorPositionInGrid;
 		}
 
+		checkf( (CellCoordinate.Y*NumCells.X + CellCoordinate.X) < Cells.Num(), TEXT("Index out of range, CellCoordinate is: %d %d CursorPosition is: %f %f"), CellCoordinate.X, CellCoordinate.Y, CursorPositionInGrid.X, CursorPositionInGrid.Y );
+
 		const TArray<int32>& IndexesInCell = CellAt( CellCoordinate.X, CellCoordinate.Y ).CachedWidgetIndexes;
 		int32 HitWidgetIndex = INDEX_NONE;
 	
@@ -183,11 +185,13 @@ int32 FHittestGrid::InsertWidget( const int32 ParentHittestIndex, const EVisibil
 
 FHittestGrid::FCell& FHittestGrid::CellAt( const int32 X, const int32 Y )
 {
+	check( (Y*NumCells.X + X) < Cells.Num() );
 	return Cells[ Y*NumCells.X + X ];
 }
 
 const FHittestGrid::FCell& FHittestGrid::CellAt( const int32 X, const int32 Y ) const
 {
+	check( (Y*NumCells.X + X) < Cells.Num() );
 	return Cells[ Y*NumCells.X + X ];
 }
 
