Slate - Adding code doc.

diff --git a/Engine/Source/Runtime/SlateCore/Public/Layout/LayoutUtils.h b/Engine/Source/Runtime/SlateCore/Public/Layout/LayoutUtils.h
--- a/Engine/Source/Runtime/SlateCore/Public/Layout/LayoutUtils.h
+++ b/Engine/Source/Runtime/SlateCore/Public/Layout/LayoutUtils.h
@@ -52,11 +52,14 @@ namespace ArrangeUtils
  * 
  * @param AllottedSize         The size available to arrange the widget along the given orientation
  * @param ChildToArrange       The widget and associated layout information
+ * @param SlotPadding          The padding to when aligning the child
+ * @param ContentScale         The scale to apply to the child before aligning it.
+ * @param bClampToParent       If true the child's size is clamped to the allotted size before alignment occurs, if false, the child's desired size is used, even if larger than the allotted size.
  * 
  * @return  Offset and Size of widget
  */
 template<EOrientation Orientation, typename SlotType>
-static AlignmentArrangeResult AlignChild(float AllottedSize, const SlotType& ChildToArrange, const FMargin& SlotPadding, const float& ContentScale = 1.0f, bool ClampToParent = true)
+static AlignmentArrangeResult AlignChild(float AllottedSize, const SlotType& ChildToArrange, const FMargin& SlotPadding, const float& ContentScale = 1.0f, bool bClampToParent = true)
 {
 	float ChildDesiredSize = (Orientation == Orient_Horizontal)
 		? ( ChildToArrange.GetWidget()->GetDesiredSize().X * ContentScale )
@@ -69,7 +72,7 @@ static AlignmentArrangeResult AlignChild(float AllottedSize, const SlotType& Chi
 	
 	const int32 Alignment = ArrangeUtils::GetChildAlignment<Orientation>::AsInt(ChildToArrange);
 		
-	const float ChildSize = ClampToParent ? FMath::Min(ChildDesiredSize, AllottedSize - TotalMargin) : ChildDesiredSize;
+	const float ChildSize = bClampToParent ? FMath::Min(ChildDesiredSize, AllottedSize - TotalMargin) : ChildDesiredSize;
 		
 	switch( Alignment )	
 	{
