Fix up Min case for NumCells.

diff --git a/Engine/Source/Runtime/SlateCore/Private/Input/HittestGrid.cpp b/Engine/Source/Runtime/SlateCore/Private/Input/HittestGrid.cpp
--- a/Engine/Source/Runtime/SlateCore/Private/Input/HittestGrid.cpp
+++ b/Engine/Source/Runtime/SlateCore/Private/Input/HittestGrid.cpp
@@ -43,8 +43,8 @@ TArray<FArrangedWidget> FHittestGrid::GetBubblePath( FVector2D DesktopSpaceCoord
 	{
 		const FVector2D CursorPositionInGrid = DesktopSpaceCoordinate - GridOrigin;
 		const FIntPoint CellCoordinate = FIntPoint(
-			FMath::Min( FMath::Max(FMath::FloorToInt(CursorPositionInGrid.X / CellSize.X), 0), NumCells.X ),
-			FMath::Min( FMath::Max(FMath::FloorToInt(CursorPositionInGrid.Y / CellSize.Y), 0), NumCells.Y ) );
+			FMath::Min( FMath::Max(FMath::FloorToInt(CursorPositionInGrid.X / CellSize.X), 0), NumCells.X-1),
+			FMath::Min( FMath::Max(FMath::FloorToInt(CursorPositionInGrid.Y / CellSize.Y), 0), NumCells.Y-1 ) );
 		
 		static FVector2D LastCoordinate = FVector2D::ZeroVector;
 		if ( LastCoordinate != CursorPositionInGrid )
