Slate/Fonts - Fixing a large bug in our font system, that was causing all fonts to be measured based on fallback fonts instead of the real font being used.

diff --git a/Engine/Source/Runtime/SlateCore/Private/Fonts/FontCache.cpp b/Engine/Source/Runtime/SlateCore/Private/Fonts/FontCache.cpp
--- a/Engine/Source/Runtime/SlateCore/Private/Fonts/FontCache.cpp
+++ b/Engine/Source/Runtime/SlateCore/Private/Fonts/FontCache.cpp
@@ -135,7 +135,7 @@ public:
 		uint32 LocalGlyphFlags = GlyphFlags;
 
 		// If the requested glyph doesn't exist, use the localization fallback font.
-		if( GlyphIndex == 0 )
+		if ( FontFace == nullptr )
 		{
 			static FName FallbackFontName( NAME_None );
 			if( FallbackFontName == NAME_None )
@@ -159,7 +159,7 @@ public:
 		}
 
 		// If the requested glyph doesn't exist, use the last resort fallback font.
-		if( GlyphIndex == 0 )
+		if ( FontFace == nullptr )
 		{
 			static FName LastResortFontName( NAME_None );
 			if( LastResortFontName == NAME_None )
@@ -274,8 +274,10 @@ public:
 		OutRenderData.MaxHeight = Height;
 
 		// Need to divide by 64 to get pixels;
+		// Ascender is not scaled by freetype.  Scale it now. 
+		OutRenderData.MeasureInfo.GlobalAscender = ( FontFace->size->metrics.ascender / 64 ) * Scale;
 		// Descender is not scaled by freetype.  Scale it now. 
-		OutRenderData.MeasureInfo.GlobalDescender = (FontFace->size->metrics.descender / 64) * Scale;
+		OutRenderData.MeasureInfo.GlobalDescender = ( FontFace->size->metrics.descender / 64 ) * Scale;
 		// Note we use Slot->advance instead of Slot->metrics.horiAdvance because Slot->Advance contains transformed position (needed if we scale)
 		OutRenderData.MeasureInfo.XAdvance =  Slot->advance.x / 64;
 		OutRenderData.MeasureInfo.HorizontalOffset = Slot->bitmap_left;
@@ -309,7 +311,7 @@ public:
 	 *
 	 * @param First			The first character in the pair
 	 * @param Second		The second character in the pair
- 	 * @param InFontInfo	Information about the font that used to draw the string with the first and second characters
+	 * @param InFontInfo	Information about the font that used to draw the string with the first and second characters
 	 * @return The kerning amount, 0 if no kerning
 	 */
 	int8 GetKerning( TCHAR First, TCHAR Second, const FSlateFontKey& FontKey )
@@ -693,6 +695,8 @@ uint16 FSlateFontCache::GetMaxCharacterHeight( const FSlateFontInfo& InFontInfo,
 	// Render the character 
 	FTInterface->GetRenderData( InFontInfo, Char, NewRenderData, FontScale );	
 
+	//return NewRenderData.MeasureInfo.GlobalAscender - NewRenderData.MeasureInfo.GlobalDescender;
+
 	return NewRenderData.MaxHeight;
 }
 
diff --git a/Engine/Source/Runtime/SlateCore/Public/Fonts/FontTypes.h b/Engine/Source/Runtime/SlateCore/Public/Fonts/FontTypes.h
--- a/Engine/Source/Runtime/SlateCore/Public/Fonts/FontTypes.h
+++ b/Engine/Source/Runtime/SlateCore/Public/Fonts/FontTypes.h
@@ -33,6 +33,8 @@ struct FCharacterMeasureInfo
 	int16 SizeY;
 	/** How many pixels to advance in X after drawing this character (when its part of a string)*/
 	int16 XAdvance;
+	/** The largest vertical distance above the baseline for any character in the font */
+	int16 GlobalAscender;
 	/** The largest vertical distance below the baseline for any character in the font */
 	int16 GlobalDescender;
 	/** The vertical distance from the baseline to the topmost border of the glyph bitmap */
