Added more Slate texture memory stats

diff --git a/Engine/Source/Runtime/Engine/Private/Slate/SlateTextures.cpp b/Engine/Source/Runtime/Engine/Private/Slate/SlateTextures.cpp
--- a/Engine/Source/Runtime/Engine/Private/Slate/SlateTextures.cpp
+++ b/Engine/Source/Runtime/Engine/Private/Slate/SlateTextures.cpp
@@ -45,6 +45,8 @@ void FSlateTexture2DRHIRef::InitDynamicRHI()
 			FRHIResourceCreateInfo CreateInfo;
 			ShaderResource = RHICreateTexture2D( Width, Height, PixelFormat, 1, 1, TexCreateFlags, CreateInfo );
 			check( IsValidRef( ShaderResource ) );
+
+			INC_MEMORY_STAT_BY(STAT_SlateTextureGPUMemory, Width*Height*GPixelFormats[PixelFormat].BlockBytes);
 		}
 
 		if( TextureData.IsValid() && TextureData->GetRawBytes().Num() > 0 )
@@ -66,6 +68,8 @@ void FSlateTexture2DRHIRef::ReleaseDynamicRHI()
 {
 	check( IsInRenderingThread() );
 	ShaderResource.SafeRelease();
+
+	DEC_MEMORY_STAT_BY(STAT_SlateTextureGPUMemory, Width*Height*GPixelFormats[PixelFormat].BlockBytes);
 }
 
 void FSlateTexture2DRHIRef::Resize( uint32 InWidth, uint32 InHeight )
diff --git a/Engine/Source/Runtime/SlateCore/Private/SlateCoreModule.cpp b/Engine/Source/Runtime/SlateCore/Private/SlateCoreModule.cpp
--- a/Engine/Source/Runtime/SlateCore/Private/SlateCoreModule.cpp
+++ b/Engine/Source/Runtime/SlateCore/Private/SlateCoreModule.cpp
@@ -35,6 +35,7 @@ DEFINE_STAT(STAT_SlateVertexCount);
 
 DEFINE_STAT(STAT_SlateTextureDataMemory);
 DEFINE_STAT(STAT_SlateTextureAtlasMemory);
+DEFINE_STAT(STAT_SlateTextureGPUMemory);
 DEFINE_STAT(STAT_SlateFontKerningTableMemory);
 DEFINE_STAT(STAT_SlateFontMeasureCacheMemory);
 DEFINE_STAT(STAT_SlateVertexBatchMemory);
diff --git a/Engine/Source/Runtime/SlateCore/Private/Textures/TextureAtlas.cpp b/Engine/Source/Runtime/SlateCore/Private/Textures/TextureAtlas.cpp
--- a/Engine/Source/Runtime/SlateCore/Private/Textures/TextureAtlas.cpp
+++ b/Engine/Source/Runtime/SlateCore/Private/Textures/TextureAtlas.cpp
@@ -51,7 +51,6 @@ const FAtlasedTextureSlot* FSlateTextureAtlas::FindSlotForTexture( uint32 InWidt
 	return FindSlotForTexture(*RootNode, InWidth, InHeight);
 }
 
-
 void FSlateTextureAtlas::DestroyNodes( FAtlasedTextureSlot* StartNode )
 {
 	if (StartNode->Left)
@@ -73,6 +72,7 @@ void FSlateTextureAtlas::InitAtlasData()
 	check(RootNode == nullptr && AtlasData.Num() == 0);
 
 	RootNode = new FAtlasedTextureSlot(0, 0, AtlasWidth, AtlasHeight, Padding);
+	AtlasData.Reserve(AtlasWidth * AtlasHeight * Stride);
 	AtlasData.AddZeroed(AtlasWidth * AtlasHeight * Stride);
 
 	INC_MEMORY_STAT_BY(STAT_SlateTextureAtlasMemory, AtlasData.GetAllocatedSize());
diff --git a/Engine/Source/Runtime/SlateCore/Public/SlateCore.h b/Engine/Source/Runtime/SlateCore/Public/SlateCore.h
--- a/Engine/Source/Runtime/SlateCore/Public/SlateCore.h
+++ b/Engine/Source/Runtime/SlateCore/Public/SlateCore.h
@@ -44,7 +44,8 @@ DECLARE_DWORD_COUNTER_STAT_EXTERN(TEXT("Num Batches"), STAT_SlateNumBatches, STA
 DECLARE_DWORD_COUNTER_STAT_EXTERN(TEXT("Num Vertices"), STAT_SlateVertexCount, STATGROUP_Slate,);
 
 DECLARE_MEMORY_STAT_EXTERN(TEXT("Texture Data CPU Memory"), STAT_SlateTextureDataMemory, STATGROUP_SlateMemory, SLATECORE_API);
-DECLARE_MEMORY_STAT_EXTERN(TEXT("Texture Atlas Memory"), STAT_SlateTextureAtlasMemory, STATGROUP_SlateMemory,);
+DECLARE_MEMORY_STAT_EXTERN(TEXT("Texture Atlas CPU Memory"), STAT_SlateTextureAtlasMemory, STATGROUP_SlateMemory,);
+DECLARE_MEMORY_STAT_EXTERN(TEXT("Texture Data GPU Memory"), STAT_SlateTextureGPUMemory, STATGROUP_SlateMemory, SLATECORE_API);
 DECLARE_MEMORY_STAT_EXTERN(TEXT("Font Kerning Table Memory"), STAT_SlateFontKerningTableMemory, STATGROUP_SlateMemory,);
 DECLARE_MEMORY_STAT_EXTERN(TEXT("Font Measure Memory"), STAT_SlateFontMeasureCacheMemory, STATGROUP_SlateMemory,);
 DECLARE_MEMORY_STAT_EXTERN(TEXT("Batch Vertex Memory"), STAT_SlateVertexBatchMemory, STATGROUP_SlateMemory,);
diff --git a/Engine/Source/Runtime/SlateRHIRenderer/Private/SlateRHIResourceManager.h b/Engine/Source/Runtime/SlateRHIRenderer/Private/SlateRHIResourceManager.h
--- a/Engine/Source/Runtime/SlateRHIRenderer/Private/SlateRHIResourceManager.h
+++ b/Engine/Source/Runtime/SlateRHIRenderer/Private/SlateRHIResourceManager.h
@@ -150,9 +150,6 @@ public:
 	 */
 	TSharedRef<SWidget> CreateTextureDisplayWidget();
 private:
-	
-	FSlateShaderResourceProxy* AddTexture( const FSlateBrush& InBrush );
-
 	/**
 	 * Deletes resources created by the manager
 	 */
