Perfect forwarding of "required args" in Slate declarative syntax
No longer makes _multiple_ copies of required args passed to SNew()

diff --git a/Engine/Source/Runtime/SlateCore/Public/Widgets/DeclarativeSyntaxSupport.h b/Engine/Source/Runtime/SlateCore/Public/Widgets/DeclarativeSyntaxSupport.h
--- a/Engine/Source/Runtime/SlateCore/Public/Widgets/DeclarativeSyntaxSupport.h
+++ b/Engine/Source/Runtime/SlateCore/Public/Widgets/DeclarativeSyntaxSupport.h
@@ -990,103 +990,102 @@ namespace RequiredArgs
 		{
 		}
 
-		template<class WidgetType, typename NamedArgsType>
-		void CallConstruct( const TSharedRef<WidgetType>& OnWidget, const NamedArgsType& WithNamedArgs ) const
+		template<class WidgetType>
+		void CallConstruct(const TSharedRef<WidgetType>& OnWidget, const typename WidgetType::FArguments& WithNamedArgs) const
 		{
-			// YOUR WIDGET MUST IMPLEMENT void Construct( const FArguments& InArgs )
-			OnWidget->Construct( WithNamedArgs );
+			// YOUR WIDGET MUST IMPLEMENT void Construct(const FArguments& InArgs)
+			OnWidget->Construct(WithNamedArgs);
 		}
 	};
 
 	template<typename Arg0Type>
 	struct T1RequiredArgs
 	{
-		T1RequiredArgs( Arg0Type InArg0 )
-		: Arg0(InArg0)
+		T1RequiredArgs(Arg0Type&& InArg0)
+			: Arg0(InArg0)
 		{
 		}
 
-		template<class WidgetType, typename NamedArgsType>
-		void CallConstruct( const TSharedRef<WidgetType>& OnWidget, const NamedArgsType& WithNamedArgs ) const
+		template<class WidgetType>
+		void CallConstruct(const TSharedRef<WidgetType>& OnWidget, const typename WidgetType::FArguments& WithNamedArgs) const
 		{
-			// YOUR WIDGET MUST IMPLEMENT void Construct( const FArguments& InArgs )
-			OnWidget->Construct( WithNamedArgs, Arg0 );
+			// YOUR WIDGET MUST IMPLEMENT void Construct(const FArguments& InArgs)
+			OnWidget->Construct(WithNamedArgs, Forward<Arg0Type>(Arg0));
 		}
 
-		Arg0Type Arg0;
-
+		Arg0Type& Arg0;
 	};
 
 	template<typename Arg0Type, typename Arg1Type>
 	struct T2RequiredArgs
 	{
-		T2RequiredArgs( Arg0Type InArg0, Arg1Type InArg1 )
-		: Arg0(InArg0)
-		, Arg1(InArg1)
+		T2RequiredArgs(Arg0Type&& InArg0, Arg1Type&& InArg1)
+			: Arg0(InArg0)
+			, Arg1(InArg1)
 		{
 		}
 
-		template<class WidgetType, typename NamedArgsType>
-		void CallConstruct( const TSharedRef<WidgetType>& OnWidget, const NamedArgsType& WithNamedArgs ) const
+		template<class WidgetType>
+		void CallConstruct(const TSharedRef<WidgetType>& OnWidget, const typename WidgetType::FArguments& WithNamedArgs) const
 		{
-			// YOUR WIDGET MUST IMPLEMENT void Construct( const FArguments& InArgs )
-			OnWidget->Construct( WithNamedArgs, Arg0, Arg1 );
+			// YOUR WIDGET MUST IMPLEMENT Construct(const FArguments& InArgs)
+			OnWidget->Construct(WithNamedArgs, Forward<Arg0Type>(Arg0), Forward<Arg1Type>(Arg1));
 		}
 
-		Arg0Type Arg0;
-		Arg1Type Arg1;
+		Arg0Type& Arg0;
+		Arg1Type& Arg1;
 	};
 
 	template<typename Arg0Type, typename Arg1Type, typename Arg2Type>
 	struct T3RequiredArgs
 	{
-		T3RequiredArgs( Arg0Type InArg0, Arg1Type InArg1, Arg2Type InArg2 )
-		: Arg0(InArg0)
-		, Arg1(InArg1)
-		, Arg2(InArg2)
+		T3RequiredArgs(Arg0Type&& InArg0, Arg1Type&& InArg1, Arg2Type&& InArg2)
+			: Arg0(InArg0)
+			, Arg1(InArg1)
+			, Arg2(InArg2)
 		{
 		}
 
-		template<class WidgetType, typename NamedArgsType>
-		void CallConstruct( const TSharedRef<WidgetType>& OnWidget, const NamedArgsType& WithNamedArgs ) const
+		template<class WidgetType>
+		void CallConstruct(const TSharedRef<WidgetType>& OnWidget, const typename WidgetType::FArguments& WithNamedArgs) const
 		{
-			// YOUR WIDGET MUST IMPLEMENT void Construct( const FArguments& InArgs )
-			OnWidget->Construct( WithNamedArgs, Arg0, Arg1, Arg2 );
+			// YOUR WIDGET MUST IMPLEMENT Construct(const FArguments& InArgs)
+			OnWidget->Construct(WithNamedArgs, Forward<Arg0Type>(Arg0), Forward<Arg1Type>(Arg1), Forward<Arg2Type>(Arg2));
 		}
 
-		Arg0Type Arg0;
-		Arg1Type Arg1;
-		Arg2Type Arg2;
+		Arg0Type& Arg0;
+		Arg1Type& Arg1;
+		Arg2Type& Arg2;
 	};
 
 	template<typename Arg0Type, typename Arg1Type, typename Arg2Type, typename Arg3Type>
 	struct T4RequiredArgs
 	{
-		T4RequiredArgs( Arg0Type InArg0, Arg1Type InArg1, Arg2Type InArg2, Arg3Type InArg3 )
-		: Arg0(InArg0)
-		, Arg1(InArg1)
-		, Arg2(InArg2)
-		, Arg3(InArg3)
+		T4RequiredArgs(Arg0Type&& InArg0, Arg1Type&& InArg1, Arg2Type&& InArg2, Arg3Type&& InArg3)
+			: Arg0(InArg0)
+			, Arg1(InArg1)
+			, Arg2(InArg2)
+			, Arg3(InArg3)
 		{
 		}
 
-		template<class WidgetType, typename NamedArgsType>
-		void CallConstruct( const TSharedRef<WidgetType>& OnWidget, const NamedArgsType& WithNamedArgs ) const
+		template<class WidgetType>
+		void CallConstruct(const TSharedRef<WidgetType>& OnWidget, const typename WidgetType::FArguments& WithNamedArgs) const
 		{
-			// YOUR WIDGET MUST IMPLEMENT void Construct( const FArguments& InArgs )
-			OnWidget->Construct( WithNamedArgs, Arg0, Arg1, Arg2, Arg3 );
+			// YOUR WIDGET MUST IMPLEMENT Construct(const FArguments& InArgs)
+			OnWidget->Construct(WithNamedArgs, Forward<Arg0Type>(Arg0), Forward<Arg1Type>(Arg1), Forward<Arg2Type>(Arg2), Forward<Arg3Type>(Arg3));
 		}
 
-		Arg0Type Arg0;
-		Arg1Type Arg1;
-		Arg2Type Arg2;
-		Arg3Type Arg3;
+		Arg0Type& Arg0;
+		Arg1Type& Arg1;
+		Arg2Type& Arg2;
+		Arg3Type& Arg3;
 	};
 
 	template<typename Arg0Type, typename Arg1Type, typename Arg2Type, typename Arg3Type, typename Arg4Type>
 	struct T5RequiredArgs
 	{
-		T5RequiredArgs( Arg0Type InArg0, Arg1Type InArg1, Arg2Type InArg2, Arg3Type InArg3, Arg4Type InArg4 )
+		T5RequiredArgs(Arg0Type&& InArg0, Arg1Type&& InArg1, Arg2Type&& InArg2, Arg3Type&& InArg3, Arg4Type&& InArg4)
 			: Arg0(InArg0)
 			, Arg1(InArg1)
 			, Arg2(InArg2)
@@ -1095,18 +1094,18 @@ namespace RequiredArgs
 		{
 		}
 
-		template<class WidgetType, typename NamedArgsType>
-		void CallConstruct( const TSharedRef<WidgetType>& OnWidget, const NamedArgsType& WithNamedArgs ) const
+		template<class WidgetType>
+		void CallConstruct(const TSharedRef<WidgetType>& OnWidget, const typename WidgetType::FArguments& WithNamedArgs) const
 		{
-			// YOUR WIDGET MUST IMPLEMENT void Construct( const FArguments& InArgs )
-			OnWidget->Construct( WithNamedArgs, Arg0, Arg1, Arg2, Arg3, Arg4 );
+			// YOUR WIDGET MUST IMPLEMENT Construct(const FArguments& InArgs)
+			OnWidget->Construct(WithNamedArgs, Forward<Arg0Type>(Arg0), Forward<Arg1Type>(Arg1), Forward<Arg2Type>(Arg2), Forward<Arg3Type>(Arg3), Forward<Arg4Type>(Arg4));
 		}
 
-		Arg0Type Arg0;
-		Arg1Type Arg1;
-		Arg2Type Arg2;
-		Arg3Type Arg3;
-		Arg4Type Arg4;
+		Arg0Type& Arg0;
+		Arg1Type& Arg1;
+		Arg2Type& Arg2;
+		Arg3Type& Arg3;
+		Arg4Type& Arg4;
 	};
 
 	FORCEINLINE T0RequiredArgs MakeRequiredArgs()
@@ -1115,33 +1114,33 @@ namespace RequiredArgs
 	}
 
 	template<typename Arg0Type>
-	T1RequiredArgs<Arg0Type> MakeRequiredArgs(Arg0Type InArg0)
+	T1RequiredArgs<Arg0Type&&> MakeRequiredArgs(Arg0Type&& InArg0)
 	{
-		return T1RequiredArgs<Arg0Type>( InArg0 );
+		return T1RequiredArgs<Arg0Type&&>(Forward<Arg0Type>(InArg0));
 	}
 
 	template<typename Arg0Type, typename Arg1Type>
-	T2RequiredArgs<Arg0Type, Arg1Type> MakeRequiredArgs(Arg0Type InArg0, Arg1Type InArg1)
+	T2RequiredArgs<Arg0Type&&, Arg1Type&&> MakeRequiredArgs(Arg0Type&& InArg0, Arg1Type&& InArg1)
 	{
-		return T2RequiredArgs<Arg0Type, Arg1Type>( InArg0, InArg1 );
+		return T2RequiredArgs<Arg0Type&&, Arg1Type&&>(Forward<Arg0Type>(InArg0), Forward<Arg1Type>(InArg1));
 	}
-	
+
 	template<typename Arg0Type, typename Arg1Type, typename Arg2Type>
-	T3RequiredArgs<Arg0Type, Arg1Type, Arg2Type> MakeRequiredArgs(Arg0Type InArg0, Arg1Type InArg1, Arg2Type InArg2)
+	T3RequiredArgs<Arg0Type&&, Arg1Type&&, Arg2Type&&> MakeRequiredArgs(Arg0Type&& InArg0, Arg1Type&& InArg1, Arg2Type&& InArg2)
 	{
-		return T3RequiredArgs<Arg0Type, Arg1Type, Arg2Type>( InArg0, InArg1, InArg2 );
+		return T3RequiredArgs<Arg0Type&&, Arg1Type&&, Arg2Type&&>(Forward<Arg0Type>(InArg0), Forward<Arg1Type>(InArg1), Forward<Arg2Type>(InArg2));
 	}
 
 	template<typename Arg0Type, typename Arg1Type, typename Arg2Type, typename Arg3Type>
-	T4RequiredArgs<Arg0Type, Arg1Type, Arg2Type, Arg3Type> MakeRequiredArgs(Arg0Type InArg0, Arg1Type InArg1, Arg2Type InArg2, Arg3Type InArg3)
+	T4RequiredArgs<Arg0Type&&, Arg1Type&&, Arg2Type&&, Arg3Type&&> MakeRequiredArgs(Arg0Type&& InArg0, Arg1Type&& InArg1, Arg2Type&& InArg2, Arg3Type&& InArg3)
 	{
-		return T4RequiredArgs<Arg0Type, Arg1Type, Arg2Type, Arg3Type>( InArg0, InArg1, InArg2, InArg3 );
+		return T4RequiredArgs<Arg0Type&&, Arg1Type&&, Arg2Type&&, Arg3Type&&>(Forward<Arg0Type>(InArg0), Forward<Arg1Type>(InArg1), Forward<Arg2Type>(InArg2), Forward<Arg3Type>(InArg3));
 	}
 
 	template<typename Arg0Type, typename Arg1Type, typename Arg2Type, typename Arg3Type, typename Arg4Type>
-	T5RequiredArgs<Arg0Type, Arg1Type, Arg2Type, Arg3Type, Arg4Type> MakeRequiredArgs(Arg0Type InArg0, Arg1Type InArg1, Arg2Type InArg2, Arg3Type InArg3, Arg4Type InArg4)
+	T5RequiredArgs<Arg0Type&&, Arg1Type&&, Arg2Type&&, Arg3Type&&, Arg4Type&&> MakeRequiredArgs(Arg0Type&& InArg0, Arg1Type&& InArg1, Arg2Type&& InArg2, Arg3Type&& InArg3, Arg4Type&& InArg4)
 	{
-		return T5RequiredArgs<Arg0Type, Arg1Type, Arg2Type, Arg3Type, Arg4Type>( InArg0, InArg1, InArg2, InArg3, InArg4 );
+		return T5RequiredArgs<Arg0Type&&, Arg1Type&&, Arg2Type&&, Arg3Type&&, Arg4Type&&>(Forward<Arg0Type>(InArg0), Forward<Arg1Type>(InArg1), Forward<Arg2Type>(InArg2), Forward<Arg3Type>(InArg3), Forward<Arg4Type>(InArg4));
 	}
 }
 
@@ -1183,9 +1182,9 @@ class SUserWidget;
 template<class WidgetType, typename RequiredArgsPayloadType>
 struct TDecl
 {
-	TDecl( const ANSICHAR* InType, const ANSICHAR* InFile, int32 OnLine, const RequiredArgsPayloadType& InRequiredArgs )
+	TDecl( const ANSICHAR* InType, const ANSICHAR* InFile, int32 OnLine, RequiredArgsPayloadType&& InRequiredArgs )
 		: _Widget( TWidgetAllocator<WidgetType, TIsDerivedFrom<WidgetType, SUserWidget>::IsDerived >::PrivateAllocateWidget() )
-		, _RequiredArgs( InRequiredArgs )
+		, _RequiredArgs(InRequiredArgs)
 	{
 		_Widget->SetDebugInfo( InType, InFile, OnLine );
 	}
@@ -1241,18 +1240,18 @@ struct TDecl
 			InArgs._Visibility,
 			InArgs._Tag );
 
-		_RequiredArgs.CallConstruct( _Widget, InArgs );
+		_RequiredArgs.CallConstruct(_Widget, InArgs);
 
 		return _Widget;
 	}
 
 	const TSharedRef<WidgetType> _Widget;
-	const RequiredArgsPayloadType _RequiredArgs;
+	RequiredArgsPayloadType& _RequiredArgs;
 };
 
 
 template<typename WidgetType, typename RequiredArgsPayloadType>
-TDecl<WidgetType, RequiredArgsPayloadType> MakeTDecl( const ANSICHAR* InType, const ANSICHAR* InFile, int32 OnLine, const RequiredArgsPayloadType& InRequiredArgs )
+TDecl<WidgetType, RequiredArgsPayloadType> MakeTDecl( const ANSICHAR* InType, const ANSICHAR* InFile, int32 OnLine, RequiredArgsPayloadType&& InRequiredArgs )
 {
-	return TDecl<WidgetType, RequiredArgsPayloadType>(InType, InFile, OnLine, InRequiredArgs);
+	return TDecl<WidgetType, RequiredArgsPayloadType>(InType, InFile, OnLine, Forward<RequiredArgsPayloadType>(InRequiredArgs));
 }
