Slate/Font - Improving the fallback code so that it both measures fonts correctly and fallsback on missing characters.

diff --git a/Engine/Source/Runtime/SlateCore/Private/Fonts/FontCache.cpp b/Engine/Source/Runtime/SlateCore/Private/Fonts/FontCache.cpp
--- a/Engine/Source/Runtime/SlateCore/Private/Fonts/FontCache.cpp
+++ b/Engine/Source/Runtime/SlateCore/Private/Fonts/FontCache.cpp
@@ -135,7 +135,7 @@ public:
 		uint32 LocalGlyphFlags = GlyphFlags;
 
 		// If the requested glyph doesn't exist, use the localization fallback font.
-		if ( FontFace == nullptr )
+		if ( FontFace == nullptr || (Char != 0 && GlyphIndex == 0) )
 		{
 			static FName FallbackFontName( NAME_None );
 			if( FallbackFontName == NAME_None )
@@ -159,7 +159,7 @@ public:
 		}
 
 		// If the requested glyph doesn't exist, use the last resort fallback font.
-		if ( FontFace == nullptr )
+		if ( FontFace == nullptr || ( Char != 0 && GlyphIndex == 0 ) )
 		{
 			static FName LastResortFontName( NAME_None );
 			if( LastResortFontName == NAME_None )
