Prevent FHittestGrid from crashing if Cells has not yet been initialized.

diff --git a/Engine/Source/Runtime/SlateCore/Private/Input/HittestGrid.cpp b/Engine/Source/Runtime/SlateCore/Private/Input/HittestGrid.cpp
--- a/Engine/Source/Runtime/SlateCore/Private/Input/HittestGrid.cpp
+++ b/Engine/Source/Runtime/SlateCore/Private/Input/HittestGrid.cpp
@@ -39,7 +39,7 @@ FHittestGrid::FHittestGrid()
 
 TArray<FArrangedWidget> FHittestGrid::GetBubblePath( FVector2D DesktopSpaceCoordinate, bool bIgnoreEnabledStatus )
 {
-	if (WidgetsCachedThisFrame->Num() > 0)
+	if (WidgetsCachedThisFrame->Num() > 0 && Cells.Num() > 0)
 	{
 		const FVector2D CursorPositionInGrid = DesktopSpaceCoordinate - GridOrigin;
 		const FIntPoint CellCoordinate = FIntPoint(
