Fix typo

diff --git a/Engine/Source/Runtime/SlateCore/Private/Styling/SlateBrush.cpp b/Engine/Source/Runtime/SlateCore/Private/Styling/SlateBrush.cpp
--- a/Engine/Source/Runtime/SlateCore/Private/Styling/SlateBrush.cpp
+++ b/Engine/Source/Runtime/SlateCore/Private/Styling/SlateBrush.cpp
@@ -3,7 +3,7 @@
 #include "SlateCorePrivatePCH.h"
 
 
-/* FSlateBrush structors
+/* FSlateBrush constructors
  *****************************************************************************/
 
 FSlateBrush::FSlateBrush( ESlateBrushDrawType::Type InDrawType, const FName InResourceName, const FMargin& InMargin, ESlateBrushTileType::Type InTiling, ESlateBrushImageType::Type InImageType, const FVector2D& InImageSize, const FLinearColor& InTint, UObject* InObjectResource, bool bInDynamicallyLoaded )
