Fixed default border brush overdraw

diff --git a/Engine/Source/Runtime/SlateCore/Private/Styling/CoreStyle.cpp b/Engine/Source/Runtime/SlateCore/Private/Styling/CoreStyle.cpp
--- a/Engine/Source/Runtime/SlateCore/Private/Styling/CoreStyle.cpp
+++ b/Engine/Source/Runtime/SlateCore/Private/Styling/CoreStyle.cpp
@@ -447,7 +447,7 @@ TSharedRef<ISlateStyle> FCoreStyle::Create( const FName& InStyleSetName )
 
 	// SBorder defaults...
 	{
-		Style->Set( "Border", new BOX_BRUSH( "Old/Border", 4.0f/16.0f ) );
+		Style->Set( "Border", new BORDER_BRUSH( "Old/Border", 4.0f/16.0f ) );
 	}
 
 	// SHyperlink defaults...
