Linux Slate fixes from PR #306.

- Discussed with Matt and Nick.
- Tooltips hack TTP'd.

diff --git a/Engine/Source/Runtime/Slate/Private/Framework/Application/SlateApplication.cpp b/Engine/Source/Runtime/Slate/Private/Framework/Application/SlateApplication.cpp
--- a/Engine/Source/Runtime/Slate/Private/Framework/Application/SlateApplication.cpp
+++ b/Engine/Source/Runtime/Slate/Private/Framework/Application/SlateApplication.cpp
@@ -353,7 +353,7 @@ FSlateApplication::FSlateApplication()
 	, OnExitRequested()
 	, EventLogger( TSharedPtr<IEventLogger>() )
 	, NumExternalModalWindowsActive( 0 )
-	, bAllowToolTips( true )
+	, bAllowToolTips( PLATFORM_LINUX ? false : true )	// @todo: support tooltip windows on Linux (TTP #344052)
 	, ToolTipDelay( 0.15f )
 	, ToolTipFadeInDuration( 0.1f )
 	, ToolTipSummonTime( 0.0 )
diff --git a/Engine/Source/Runtime/Slate/Private/Widgets/Input/SMenuAnchor.cpp b/Engine/Source/Runtime/Slate/Private/Widgets/Input/SMenuAnchor.cpp
--- a/Engine/Source/Runtime/Slate/Private/Widgets/Input/SMenuAnchor.cpp
+++ b/Engine/Source/Runtime/Slate/Private/Widgets/Input/SMenuAnchor.cpp
@@ -179,7 +179,7 @@ void SMenuAnchor::SetIsOpen( bool InIsOpen, const bool bFocusMenu )
 			// other pop-up menus are currently auto-sized
 			const FVector2D DesiredContentSize = MenuContentRef->GetDesiredSize();  // @todo: This is ignoring any window border size!
 			const EMenuPlacement PlacementMode = Placement.Get();
-			bool bShouldAutoSize = PlacementMode != MenuPlacement_ComboBox;
+			bool bShouldAutoSize = (PlacementMode != MenuPlacement_ComboBox);
 
 			const FVector2D NewPosition = MyGeometry.AbsolutePosition;
 			FVector2D NewWindowSize = DesiredContentSize;
diff --git a/Engine/Source/Runtime/SlateCore/Private/Widgets/SWindow.cpp b/Engine/Source/Runtime/SlateCore/Private/Widgets/SWindow.cpp
--- a/Engine/Source/Runtime/SlateCore/Private/Widgets/SWindow.cpp
+++ b/Engine/Source/Runtime/SlateCore/Private/Widgets/SWindow.cpp
@@ -181,7 +181,7 @@ void SWindow::Construct(const FArguments& InArgs)
 	const bool bCreateTitleBar = InArgs._CreateTitleBar && !bIsPopupWindow && !bIsCursorDecoratorWindow && !bHasOSWindowBorder;
 	FVector2D WindowSize = InArgs._ClientSize;
 
-	// Do not adjust the client size if we have an OS border.  
+	// If the window has no OS border, simulate it ourselves, enlarging window by the size that OS border would have.
 	if (!HasOSWindowBorder())
 	{
 		const FMargin BorderSize = GetWindowBorderSize();
@@ -628,7 +628,11 @@ FSlateRect SWindow::GetClippingRectangleInWindow() const
 FMargin SWindow::GetWindowBorderSize() const
 {
 // Mac didn't want a window border, and consoles don't either, so only do this in Windows
-#if PLATFORM_WINDOWS || PLATFORM_LINUX
+
+// @TODO This is not working for Linux. The window is not yet valid when this gets
+// called from SWindow::Construct which is causing a default border to be retured even when the
+// window is borderless. This causes problems for menu positioning.
+#if PLATFORM_WINDOWS // || PLATFORM_LINUX
 	if (NativeWindow.IsValid() && NativeWindow->IsMaximized())
 	{
 		int32 OSWindowBorderSize = NativeWindow->GetWindowBorderSize();
