Slate SetVisbility() no longer has any effect on SNullWidget.
SetVisiblity() is marked as override final in LeafWidget, Panel, and CompoundWidget.

diff --git a/Engine/Source/Runtime/SlateCore/Private/Widgets/SCompoundWidget.cpp b/Engine/Source/Runtime/SlateCore/Private/Widgets/SCompoundWidget.cpp
--- a/Engine/Source/Runtime/SlateCore/Private/Widgets/SCompoundWidget.cpp
+++ b/Engine/Source/Runtime/SlateCore/Private/Widgets/SCompoundWidget.cpp
@@ -64,4 +64,9 @@ SCompoundWidget::SCompoundWidget()
 	, ColorAndOpacity( FLinearColor::White )
 	, ForegroundColor( FSlateColor::UseForeground() )
 {
+}
+
+void SCompoundWidget::SetVisibility( TAttribute<EVisibility> InVisibility )
+{
+	SWidget::SetVisibility( InVisibility );
 }
\ No newline at end of file
diff --git a/Engine/Source/Runtime/SlateCore/Private/Widgets/SLeafWidget.cpp b/Engine/Source/Runtime/SlateCore/Private/Widgets/SLeafWidget.cpp
--- a/Engine/Source/Runtime/SlateCore/Private/Widgets/SLeafWidget.cpp
+++ b/Engine/Source/Runtime/SlateCore/Private/Widgets/SLeafWidget.cpp
@@ -12,6 +12,11 @@ FNoChildren SLeafWidget::NoChildrenInstance;
 /* SLeafWidget interface
  *****************************************************************************/
 
+void SLeafWidget::SetVisibility( TAttribute<EVisibility> InVisibility )
+{
+	SWidget::SetVisibility( InVisibility );
+}
+
 FChildren* SLeafWidget::GetChildren( )
 {
 	return &SLeafWidget::NoChildrenInstance;
diff --git a/Engine/Source/Runtime/SlateCore/Private/Widgets/SNullWidget.cpp b/Engine/Source/Runtime/SlateCore/Private/Widgets/SNullWidget.cpp
--- a/Engine/Source/Runtime/SlateCore/Private/Widgets/SNullWidget.cpp
+++ b/Engine/Source/Runtime/SlateCore/Private/Widgets/SNullWidget.cpp
@@ -3,8 +3,10 @@
 #include "SlateCorePrivatePCH.h"
 
 
+FNoChildren NullWidgetNoChildren;
+
 class SLATECORE_API SNullWidgetContent
-	: public SLeafWidget
+	: public SWidget
 {
 public:
 
@@ -26,8 +28,13 @@ public:
 		);
 	}
 
+private:
+	virtual void SetVisibility( TAttribute<EVisibility> InVisibility ) override final
+	{
+		ensureMsg( false, TEXT("Attempting to SetVisibility() on SNullWidget. Mutating SNullWidget is not allowed.") );
+	}
 public:
-
+	
 	// SWidget interface
 
 	virtual int32 OnPaint( const FPaintArgs& Args, const FGeometry& AllottedGeometry, const FSlateRect& MyClippingRect, FSlateWindowElementList& OutDrawElements, int32 LayerId, const FWidgetStyle& InWidgetStyle, bool bParentEnabled ) const override
@@ -35,11 +42,22 @@ public:
 		return LayerId;
 	}
 
-	virtual FVector2D ComputeDesiredSize( ) const override
+	virtual FVector2D ComputeDesiredSize( ) const override final
 	{
 		return FVector2D(0.0f, 0.0f);
 	}
 
+	virtual FChildren* GetChildren( ) override final
+	{
+		return &NullWidgetNoChildren;
+	}
+
+
+	virtual void OnArrangeChildren( const FGeometry& AllottedGeometry, FArrangedChildren& ArrangedChildren ) const override final
+	{
+		// Nothing to arrange; Null Widgets do not have children.
+	}
+
 	// End of SWidget interface
 };
 
diff --git a/Engine/Source/Runtime/SlateCore/Private/Widgets/SPanel.cpp b/Engine/Source/Runtime/SlateCore/Private/Widgets/SPanel.cpp
--- a/Engine/Source/Runtime/SlateCore/Private/Widgets/SPanel.cpp
+++ b/Engine/Source/Runtime/SlateCore/Private/Widgets/SPanel.cpp
@@ -48,3 +48,9 @@ int32 SPanel::PaintArrangedChildren( const FPaintArgs& Args, const FArrangedChil
 	
 	return MaxLayerId;
 }
+
+
+void SPanel::SetVisibility( TAttribute<EVisibility> InVisibility )
+{
+	SWidget::SetVisibility( InVisibility );
+}
\ No newline at end of file
diff --git a/Engine/Source/Runtime/SlateCore/Public/Widgets/SCompoundWidget.h b/Engine/Source/Runtime/SlateCore/Public/Widgets/SCompoundWidget.h
--- a/Engine/Source/Runtime/SlateCore/Public/Widgets/SCompoundWidget.h
+++ b/Engine/Source/Runtime/SlateCore/Public/Widgets/SCompoundWidget.h
@@ -69,6 +69,9 @@ public:
 	virtual void OnArrangeChildren( const FGeometry& AllottedGeometry, FArrangedChildren& ArrangedChildren ) const override;
 	virtual FSlateColor GetForegroundColor() const override;
 
+public:
+	virtual void SetVisibility( TAttribute<EVisibility> InVisibility ) override final;
+
 protected:
 
 	/** Disallow public construction */
diff --git a/Engine/Source/Runtime/SlateCore/Public/Widgets/SLeafWidget.h b/Engine/Source/Runtime/SlateCore/Public/Widgets/SLeafWidget.h
--- a/Engine/Source/Runtime/SlateCore/Public/Widgets/SLeafWidget.h
+++ b/Engine/Source/Runtime/SlateCore/Public/Widgets/SLeafWidget.h
@@ -13,6 +13,9 @@ class SLATECORE_API SLeafWidget
 	: public SWidget
 {
 public:
+	virtual void SetVisibility( TAttribute<EVisibility> InVisibility ) override final;
+
+private:
 
 	// Begin SWidget overrides
 
diff --git a/Engine/Source/Runtime/SlateCore/Public/Widgets/SPanel.h b/Engine/Source/Runtime/SlateCore/Public/Widgets/SPanel.h
--- a/Engine/Source/Runtime/SlateCore/Public/Widgets/SPanel.h
+++ b/Engine/Source/Runtime/SlateCore/Public/Widgets/SPanel.h
@@ -66,4 +66,7 @@ protected:
 
 	/** Hidden default constructor. */
 	SPanel( ) { }
+
+public:
+	virtual void SetVisibility( TAttribute<EVisibility> InVisibility ) override final;
 };
diff --git a/Engine/Source/Runtime/SlateCore/Public/Widgets/SWidget.h b/Engine/Source/Runtime/SlateCore/Public/Widgets/SWidget.h
--- a/Engine/Source/Runtime/SlateCore/Public/Widgets/SWidget.h
+++ b/Engine/Source/Runtime/SlateCore/Public/Widgets/SWidget.h
@@ -555,7 +555,7 @@ public:
 	}
 
 	/** @param InVisibility  should this widget be */
-	void SetVisibility( TAttribute<EVisibility> InVisibility )
+	virtual void SetVisibility( TAttribute<EVisibility> InVisibility )
 	{
 		Visibility = InVisibility;
 	}
