Fix crash rendering mono, non-antialiased fonts

diff --git a/Engine/Source/Runtime/SlateCore/Private/Fonts/FontCache.cpp b/Engine/Source/Runtime/SlateCore/Private/Fonts/FontCache.cpp
--- a/Engine/Source/Runtime/SlateCore/Private/Fonts/FontCache.cpp
+++ b/Engine/Source/Runtime/SlateCore/Private/Fonts/FontCache.cpp
@@ -18,6 +18,7 @@
 	#include FT_FREETYPE_H
 	#include FT_GLYPH_H
 	#include FT_MODULE_H
+	#include FT_BITMAP_H
 
 #endif // WITH_FREETYPE
 
@@ -210,55 +211,53 @@ public:
 		// one byte per pixel 
 		const uint32 GlyphPixelSize = 1;
 
-		FT_Bitmap Bitmap = Slot->bitmap;
+	
+		FT_Bitmap* Bitmap = nullptr;
 
+		if( Slot->bitmap.pixel_mode == FT_PIXEL_MODE_MONO )
+		{
+			FT_Bitmap NewBitmap;
+
+			FT_Bitmap_New( &NewBitmap );
+			// Convert the mono font to 8bbp from 1bpp
+			FT_Bitmap_Convert( FTLibrary, &Slot->bitmap, &NewBitmap, 4 );
+
+			Bitmap = &NewBitmap;
+		}
+		else
+		{
+			Bitmap = &Slot->bitmap;
+		}
+		
 		OutRenderData.RawPixels.Reset();
-		OutRenderData.RawPixels.AddUninitialized( Bitmap.rows * Bitmap.width );
+		OutRenderData.RawPixels.AddUninitialized( Bitmap->rows * Bitmap->width );
+
 
-		if( Bitmap.pixel_mode == FT_PIXEL_MODE_MONO )
 		{
-			uint32 BytesPerLine = Bitmap.pitch;
+			// Copy the rendered bitmap to our raw pixels array
 
-			// For each texel generate a corresponding color for a 32 bit image
-			for( int32 Height = 0; Height < Bitmap.rows; ++Height )
+			if( Slot->bitmap.pixel_mode != FT_PIXEL_MODE_MONO )
 			{
-				for( uint32 ByteIdx = 0; ByteIdx < BytesPerLine; ++ByteIdx )
+				for (int32 Row = 0; Row < Bitmap->rows; ++Row)
 				{
-					// The current byte
-					uint8 Byte = Bitmap.buffer[ Height*BytesPerLine + ByteIdx ];
-
-					// We may not need to use all the bits in a byte.  If there are more than 8 pixels left to check, check all 8
-					// otherwise check whats left
-					uint8 NumBitsToCheck = FMath::Min<uint8>(Bitmap.width-ByteIdx*8, 8);
+					// Copy a single row. Note Bitmap.pitch contains the offset (in bytes) between rows.  Not always equal to Bitmap.width!
+					FMemory::Memcpy(&OutRenderData.RawPixels[Row*Bitmap->width], &Bitmap->buffer[Row*Bitmap->pitch], Bitmap->width*GlyphPixelSize);
+				}
 
-					// Iterate through each bit in this byte.
-					// Since each byte has up to 8 pixels, we need to generate an uint8 for each bit
-					for( int32 Bit = 0; Bit < NumBitsToCheck; ++Bit )
+			}
+			else
+			{
+				// In Mono a value of 1 means the pixel is drawn and a value of zero means it is not. 
+				// So we must check each pixel and convert it to a color.
+				for( int32 Height = 0; Height < Bitmap->rows; ++Height )
+				{
+					for( int32 Width = 0; Width < Bitmap->width; ++Width )
 					{
-						// Most significant bit is left most bit
-						if( ( Byte & ( 1 << (7-Bit) ) ) == 0 )
-						{
-							// Bit not set, pixel is black
-							OutRenderData.RawPixels[ Height * Bitmap.width + Bit + 8*ByteIdx ] = 0;
-						}
-						else
-						{
-							// Bit set, pixel is white
-							OutRenderData.RawPixels[ Height * Bitmap.width + Bit + 8*ByteIdx ] = 255;
-						}
+						OutRenderData.RawPixels[Height*Bitmap->width+Width] = Bitmap->buffer[Height*Bitmap->pitch+Width] == 1 ? 255 : 0;
 					}
 				}
 			}
 		}
-		else
-		{
-			// Copy the rendered bitmap to our raw pixels array
-			for( int32 Row = 0; Row < Bitmap.rows; ++Row )
-			{
-				// Copy a single row. Note Bitmap.pitch contains the offset (in bytes) between rows.  Not always equal to Bitmap.width!
-				FMemory::Memcpy( &OutRenderData.RawPixels[Row*Bitmap.width], &Bitmap.buffer[Row*Bitmap.pitch], Bitmap.width*GlyphPixelSize );
-			}
-		}
 
 		FT_BBox GlyphBox;
 		FT_Glyph Glyph;
@@ -269,8 +268,8 @@ public:
 
 		// Set measurement info for this character
 		OutRenderData.Char = Char;
-		OutRenderData.MeasureInfo.SizeX = Bitmap.width;
-		OutRenderData.MeasureInfo.SizeY = Bitmap.rows;
+		OutRenderData.MeasureInfo.SizeX = Bitmap->width;
+		OutRenderData.MeasureInfo.SizeY = Bitmap->rows;
 		OutRenderData.MaxHeight = Height;
 
 		// Need to divide by 64 to get pixels;
@@ -283,7 +282,13 @@ public:
 		OutRenderData.MeasureInfo.HorizontalOffset = Slot->bitmap_left;
 		OutRenderData.MeasureInfo.VerticalOffset = Slot->bitmap_top;
 
+		if( Slot->bitmap.pixel_mode == FT_PIXEL_MODE_MONO )
+		{
+			FT_Bitmap_Done( FTLibrary, Bitmap );
+		}
+
 		FT_Done_Glyph( Glyph );
+
 #endif // WITH_FREETYPE
 	}
 
