Added widget metadata

Expanded tagging system to more general metadata. Supplied a base class with simple RTTI as well as a simple tagging metadata that will be used in place of the FName tag. Once all tags are converted the Tag field will be deprecated.

diff --git a/Engine/Source/Editor/IntroTutorials/Private/STutorialOverlay.cpp b/Engine/Source/Editor/IntroTutorials/Private/STutorialOverlay.cpp
--- a/Engine/Source/Editor/IntroTutorials/Private/STutorialOverlay.cpp
+++ b/Engine/Source/Editor/IntroTutorials/Private/STutorialOverlay.cpp
@@ -102,8 +102,10 @@ int32 STutorialOverlay::OnPaint( const FPaintArgs& Args, const FGeometry& Allott
 
 int32 STutorialOverlay::TraverseWidgets(TSharedRef<SWidget> InWidget, const FGeometry& InGeometry, const FSlateRect& MyClippingRect, FSlateWindowElementList& OutDrawElements, int32 LayerId) const
 {
+	TSharedPtr<FTagMetaData> MetaData = InWidget->GetMetaData<FTagMetaData>();
+
 	const FName Tag = InWidget->GetTag();
-	if(Tag != NAME_None)
+	if(Tag != NAME_None || MetaData.IsValid())
 	{
 		// we are a named widget - ask it to draw
 		OnPaintNamedWidget.Broadcast(InWidget, InGeometry);
diff --git a/Engine/Source/Editor/IntroTutorials/Private/TutorialStructCustomization.cpp b/Engine/Source/Editor/IntroTutorials/Private/TutorialStructCustomization.cpp
--- a/Engine/Source/Editor/IntroTutorials/Private/TutorialStructCustomization.cpp
+++ b/Engine/Source/Editor/IntroTutorials/Private/TutorialStructCustomization.cpp
@@ -212,6 +212,15 @@ private:
 				PickedWidgetName = PathWidget->GetTag();
 				break;
 			}
+			else
+			{
+				TSharedPtr<FTagMetaData> MetaData = PathWidget->GetMetaData<FTagMetaData>();
+				if(MetaData.IsValid())
+				{
+					PickedWidgetName = MetaData->Tag;
+					break;
+				}
+			}		
 		}
 
 		// kind of a hack, but we need to maintain keyboard focus otherwise we wont get our keypress to 'pick'
diff --git a/Engine/Source/Runtime/Slate/Private/Widgets/Layout/SSpacer.cpp b/Engine/Source/Runtime/Slate/Private/Widgets/Layout/SSpacer.cpp
--- a/Engine/Source/Runtime/Slate/Private/Widgets/Layout/SSpacer.cpp
+++ b/Engine/Source/Runtime/Slate/Private/Widgets/Layout/SSpacer.cpp
@@ -18,7 +18,8 @@ void SSpacer::Construct( const FArguments& InArgs )
 		InArgs._Cursor, 
 		InArgs._IsEnabled,
 		InArgs._Visibility,
-		InArgs._Tag);
+		InArgs._Tag,
+		InArgs.MetaData);
 }
 
 
diff --git a/Engine/Source/Runtime/Slate/Private/Widgets/Views/SHeaderRow.cpp b/Engine/Source/Runtime/Slate/Private/Widgets/Views/SHeaderRow.cpp
--- a/Engine/Source/Runtime/Slate/Private/Widgets/Views/SHeaderRow.cpp
+++ b/Engine/Source/Runtime/Slate/Private/Widgets/Views/SHeaderRow.cpp
@@ -34,7 +34,7 @@ public:
 	{
 		check(InArgs._Style);
 
-		SWidget::Construct( InArgs._ToolTipText, InArgs._ToolTip, InArgs._Cursor, InArgs._IsEnabled, InArgs._Visibility, InArgs._Tag );
+		SWidget::Construct( InArgs._ToolTipText, InArgs._ToolTip, InArgs._Cursor, InArgs._IsEnabled, InArgs._Visibility, InArgs._Tag, InArgs.MetaData );
 
 		Style = InArgs._Style;
 		ColumnId = Column.ColumnId;
@@ -384,7 +384,7 @@ void SHeaderRow::Construct( const FArguments& InArgs )
 {
 	check(InArgs._Style);
 
-	SWidget::Construct( InArgs._ToolTipText, InArgs._ToolTip, InArgs._Cursor, InArgs._IsEnabled, InArgs._Visibility, InArgs._Tag );
+	SWidget::Construct( InArgs._ToolTipText, InArgs._ToolTip, InArgs._Cursor, InArgs._IsEnabled, InArgs._Visibility, InArgs._Tag, InArgs.MetaData );
 
 	ScrollBarThickness = FVector2D::ZeroVector;
 	ScrollBarVisibility = EVisibility::Collapsed;
diff --git a/Engine/Source/Runtime/SlateCore/Private/Widgets/SNullWidget.cpp b/Engine/Source/Runtime/SlateCore/Private/Widgets/SNullWidget.cpp
--- a/Engine/Source/Runtime/SlateCore/Private/Widgets/SNullWidget.cpp
+++ b/Engine/Source/Runtime/SlateCore/Private/Widgets/SNullWidget.cpp
@@ -24,7 +24,8 @@ public:
 			InArgs._Cursor, 
 			InArgs._IsEnabled,
 			InArgs._Visibility,
-			InArgs._Tag
+			InArgs._Tag,
+			InArgs.MetaData
 		);
 	}
 
diff --git a/Engine/Source/Runtime/SlateCore/Private/Widgets/SWidget.cpp b/Engine/Source/Runtime/SlateCore/Private/Widgets/SWidget.cpp
--- a/Engine/Source/Runtime/SlateCore/Private/Widgets/SWidget.cpp
+++ b/Engine/Source/Runtime/SlateCore/Private/Widgets/SWidget.cpp
@@ -31,7 +31,8 @@ void SWidget::Construct(
 	const TAttribute< TOptional<EMouseCursor::Type> > & InCursor ,
 	const TAttribute<bool> & InEnabledState ,
 	const TAttribute<EVisibility> & InVisibility,
-	const FName& InTag
+	const FName& InTag,
+	const TArray<TSharedRef<ISlateMetaData>>& InMetaData
 )
 {
 	if ( InToolTip.IsValid() )
@@ -54,6 +55,7 @@ void SWidget::Construct(
 	EnabledState = InEnabledState;
 	Visibility = InVisibility;
 	Tag = InTag;
+	MetaData = InMetaData;
 }
 
 
diff --git a/Engine/Source/Runtime/SlateCore/Public/SlateCore.h b/Engine/Source/Runtime/SlateCore/Public/SlateCore.h
--- a/Engine/Source/Runtime/SlateCore/Public/SlateCore.h
+++ b/Engine/Source/Runtime/SlateCore/Public/SlateCore.h
@@ -66,6 +66,7 @@ DECLARE_MEMORY_STAT_EXTERN(TEXT("Index Buffer Memory"), STAT_SlateIndexBufferMem
 #include "SlateEnums.h"
 #include "SlateStructs.h"
 #include "PaintArgs.h"
+#include "ISlateMetaData.h"
 
 // Layout
 #include "SlotBase.h"
diff --git a/Engine/Source/Runtime/SlateCore/Public/Types/ISlateMetaData.h b/Engine/Source/Runtime/SlateCore/Public/Types/ISlateMetaData.h
new file mode 100644
--- /dev/null
+++ b/Engine/Source/Runtime/SlateCore/Public/Types/ISlateMetaData.h
@@ -0,0 +1,60 @@
+// Copyright 1998-2014 Epic Games, Inc. All Rights Reserved.
+
+#pragma once
+
+/**
+ * Base class for all slate metadata
+ */
+class ISlateMetaData
+{
+public:
+	/** Check if this metadata operation can cast safely to the specified template type */
+	template<class TType> 
+	bool IsOfType() const 
+	{
+		return IsOfTypeImpl(TType::GetTypeId());
+	}
+
+	/** Virtual destructor. */
+	virtual ~ISlateMetaData() { }
+
+protected:
+	/**
+	 * Checks whether this drag and drop operation can cast safely to the specified type.
+	 */
+	virtual bool IsOfTypeImpl( const FName& Type ) const
+	{
+		return false;
+	}
+};
+
+/**
+ * All metadata-derived classes must include this macro.
+ * Example Usage:
+ *	class FMyMetaData : public ISlateMetaData
+ *	{
+ *	public:
+ *		SLATE_METADATA_TYPE(FMyMetaData, ISlateMetaData)
+ *		...
+ *	};
+ */
+#define SLATE_METADATA_TYPE(TYPE, BASE) \
+	static const FName& GetTypeId() { static FName Type(TEXT(#TYPE)); return Type; } \
+	virtual bool IsOfTypeImpl(const FName& Type) const override { return GetTypeId() == Type || BASE::IsOfTypeImpl(Type); }
+
+/**
+ * Simple tagging metadata
+ */
+class FTagMetaData : public ISlateMetaData
+{
+public:
+	SLATE_METADATA_TYPE(FTagMetaData, ISlateMetaData)
+
+	FTagMetaData(FName InTag)
+	 : Tag(InTag)
+	{
+	}
+
+	/** Tag name for a widget */
+	FName Tag;
+};
diff --git a/Engine/Source/Runtime/SlateCore/Public/Widgets/DeclarativeSyntaxSupport.h b/Engine/Source/Runtime/SlateCore/Public/Widgets/DeclarativeSyntaxSupport.h
--- a/Engine/Source/Runtime/SlateCore/Public/Widgets/DeclarativeSyntaxSupport.h
+++ b/Engine/Source/Runtime/SlateCore/Public/Widgets/DeclarativeSyntaxSupport.h
@@ -973,14 +973,38 @@ struct TSlateBaseNamedArgs
 		return *(static_cast<WidgetArgsType*>(this));
 	}
 
+	/** Add metadata to this widget. */
+	WidgetArgsType& AddMetaData(TSharedRef<ISlateMetaData> InMetaData)
+	{
+		MetaData.Add(InMetaData);
+		return *(static_cast<WidgetArgsType*>(this));
+	}
+
+	/** Add metadata to this widget - convenience method - 1 argument */
+	template<typename MetaDataType, typename Arg0Type>
+	WidgetArgsType& AddMetaData(Arg0Type InArg0)
+	{
+		MetaData.Add(MakeShareable(new MetaDataType(InArg0)));
+		return *(static_cast<WidgetArgsType*>(this));
+	}
+
+	/** Add metadata to this widget - convenience method - 2 arguments */
+	template<typename MetaDataType, typename Arg0Type, typename Arg1Type>
+	WidgetArgsType& AddMetaData(Arg0Type InArg0, Arg1Type InArg1)
+	{
+		MetaData.Add(MakeShareable(new MetaDataType(InArg0, InArg1)));
+		return *(static_cast<WidgetArgsType*>(this));
+	}
+
 	SLATE_TEXT_ATTRIBUTE( ToolTipText )
 	SLATE_ARGUMENT( TSharedPtr<IToolTip>, ToolTip )
 	SLATE_ATTRIBUTE( TOptional<EMouseCursor::Type>, Cursor )
 	SLATE_ATTRIBUTE( bool, IsEnabled )
 	SLATE_ATTRIBUTE( EVisibility, Visibility )
 	SLATE_ARGUMENT( FName, Tag )
-};
 
+	TArray<TSharedRef<ISlateMetaData>> MetaData;
+};
 
 namespace RequiredArgs
 {
@@ -1238,7 +1262,8 @@ struct TDecl
 			InArgs._Cursor ,
 			InArgs._IsEnabled ,
 			InArgs._Visibility,
-			InArgs._Tag );
+			InArgs._Tag,
+			InArgs.MetaData );
 
 		_RequiredArgs.CallConstruct(_Widget, InArgs);
 
diff --git a/Engine/Source/Runtime/SlateCore/Public/Widgets/SWidget.h b/Engine/Source/Runtime/SlateCore/Public/Widgets/SWidget.h
--- a/Engine/Source/Runtime/SlateCore/Public/Widgets/SWidget.h
+++ b/Engine/Source/Runtime/SlateCore/Public/Widgets/SWidget.h
@@ -71,16 +71,18 @@ public:
 		const TAttribute< TOptional<EMouseCursor::Type> > & InCursor ,
 		const TAttribute<bool> & InEnabledState ,
 		const TAttribute<EVisibility> & InVisibility,
-		const FName& InTag );
+		const FName& InTag,
+		const TArray<TSharedRef<ISlateMetaData>>& InMetaData);
 
 	void SWidgetConstruct( const TAttribute<FString> & InToolTipText ,
 		const TSharedPtr<IToolTip> & InToolTip ,
 		const TAttribute< TOptional<EMouseCursor::Type> > & InCursor ,
 		const TAttribute<bool> & InEnabledState ,
 		const TAttribute<EVisibility> & InVisibility,
-		const FName& InTag )
+		const FName& InTag,
+		const TArray<TSharedRef<ISlateMetaData>>& InMetaData)
 	{
-		Construct(InToolTipText, InToolTip, InCursor, InEnabledState, InVisibility, InTag);
+		Construct(InToolTipText, InToolTip, InCursor, InEnabledState, InVisibility, InTag, InMetaData);
 	}
 
 	//
@@ -587,6 +589,41 @@ public:
 	 * Used by Slate to set the runtime debug info about this widget.
 	 */
 	void SetDebugInfo( const ANSICHAR* InType, const ANSICHAR* InFile, int32 OnLine );
+	
+	/**
+	 * Get the metadata of the type provided.
+	 * @return the first metadata of the type supplied that we encounter
+	 */
+	template<typename MetaDataType>
+	TSharedPtr<MetaDataType> GetMetaData() const
+	{
+		for (const auto& MetaDataEntry : MetaData)
+		{
+			if (MetaDataEntry->IsOfType<MetaDataType>())
+			{
+				return StaticCastSharedRef<MetaDataType>(MetaDataEntry);
+			}
+		}
+		return TSharedPtr<MetaDataType>();	
+	}
+
+	/**
+	 * Get all metadata of the type provided.
+	 * @return all the metadata found of the specified type.
+	 */
+	template<typename MetaDataType>
+	TArray<TSharedRef<MetaDataType>> GetAllMetaData() const
+	{
+		TArray<TSharedRef<MetaDataType>> FoundMetaData;
+		for (const auto& MetaDataEntry : MetaData)
+		{
+			if (MetaDataEntry->IsOfType<MetaDataType>())
+			{
+				FoundMetaData.Add(StaticCastSharedRef<MetaDataType>(MetaDataEntry));
+			}
+		}
+		return FoundMetaData;	
+	}
 
 public:
 
@@ -716,6 +753,9 @@ protected:
 	/** Tag for this widget */
 	FName Tag;
 
+	/** Metadata associated with this widget */
+	TArray<TSharedRef<ISlateMetaData>> MetaData;
+
 	/** The cursor to show when the mouse is hovering over this widget. */
 	TAttribute< TOptional<EMouseCursor::Type> > Cursor;
 
