Change 2122846 on 2014/07/01 00:51:24 by [REDACTED_PERSON][REDACTED_HANDLE]

diff --git a/Engine/Source/Runtime/SlateCore/Private/Layout/Geometry.cpp b/Engine/Source/Runtime/SlateCore/Private/Layout/Geometry.cpp
--- a/Engine/Source/Runtime/SlateCore/Private/Layout/Geometry.cpp
+++ b/Engine/Source/Runtime/SlateCore/Private/Layout/Geometry.cpp
@@ -1,61 +1,182 @@
 // Copyright 1998-2014 Epic Games, Inc. All Rights Reserved.
 
 #include "SlateCorePrivatePCH.h"
+#include "LayoutGeometry.h"
 
+FGeometry::FGeometry() 
+	: Size(0.0f, 0.0f)
+	, Scale(1.0f)
+	, AbsolutePosition(0.0f, 0.0f)
+{
+}
+
+FGeometry& FGeometry::operator=(const FGeometry& RHS)
+{
+	// HACK to allow us to make FGeometry public members immutable to catch misuse.
+	if (this != &RHS)
+	{
+		FMemory::MemCopy(*this, RHS);
+	}
+	return *this;
+}
+
+FGeometry::FGeometry(const FVector2D& OffsetFromParent, const FVector2D& ParentAbsolutePosition, const FVector2D& InLocalSize, float InScale) 
+	: Size(InLocalSize)
+	, Scale(1.0f)
+	, AbsolutePosition(0.0f, 0.0f)
+{
+	// Since OffsetFromParent is given as a LocalSpaceOffset, we MUST convert this offset into the space of the parent to construct a valid layout transform.
+	// The extra TransformPoint below does this by converting the local offset to an offset in parent space.
+	FVector2D LayoutOffset = TransformPoint(InScale, OffsetFromParent);
+
+	FSlateLayoutTransform ParentAccumulatedLayoutTransform(InScale, ParentAbsolutePosition);
+	FSlateLayoutTransform LocalLayoutTransform(LayoutOffset);
+	AccumulatedRenderTransform = TransformCast<FSlateRenderTransform>(LocalLayoutTransform);
+	FSlateLayoutTransform AccumulatedLayoutTransform = Concatenate(LocalLayoutTransform, ParentAccumulatedLayoutTransform);
+	// HACK to allow us to make FGeometry public members immutable to catch misuse.
+	const_cast<FVector2D&>(AbsolutePosition) = AccumulatedLayoutTransform.GetTranslation();
+	const_cast<float&>(Scale) = AccumulatedLayoutTransform.GetScale();
+	const_cast<FVector2D&>(Position) = LocalLayoutTransform.GetTranslation();
+}
+
+FGeometry::FGeometry(const FVector2D& InLocalSize, const FSlateLayoutTransform& InLocalLayoutTransform, const FSlateRenderTransform& InLocalRenderTransform, const FSlateLayoutTransform& ParentAccumulatedLayoutTransform, const FSlateRenderTransform& ParentAccumulatedRenderTransform) 
+	: Size(InLocalSize)
+	, Scale(1.0f)
+	, AbsolutePosition(0.0f, 0.0f)
+	, AccumulatedRenderTransform(Concatenate(InLocalRenderTransform, InLocalLayoutTransform, ParentAccumulatedRenderTransform))
+{
+	FSlateLayoutTransform AccumulatedLayoutTransform = Concatenate(InLocalLayoutTransform, ParentAccumulatedLayoutTransform);
+	// HACK to allow us to make FGeometry public members immutable to catch misuse.
+	const_cast<FVector2D&>(AbsolutePosition) = AccumulatedLayoutTransform.GetTranslation();
+	const_cast<float&>(Scale) = AccumulatedLayoutTransform.GetScale();
+	const_cast<FVector2D&>(Position) = InLocalLayoutTransform.GetTranslation();
+}
+
+FGeometry::FGeometry(const FVector2D& InLocalSize, const FSlateLayoutTransform& InLocalLayoutTransform, const FSlateLayoutTransform& ParentAccumulatedLayoutTransform, const FSlateRenderTransform& ParentAccumulatedRenderTransform) 
+	: Size(InLocalSize)
+	, Scale(1.0f)
+	, AbsolutePosition(0.0f, 0.0f)
+	, AccumulatedRenderTransform(Concatenate(InLocalLayoutTransform, ParentAccumulatedRenderTransform))
+{
+	FSlateLayoutTransform AccumulatedLayoutTransform = Concatenate(InLocalLayoutTransform, ParentAccumulatedLayoutTransform);
+	// HACK to allow us to make FGeometry public members immutable to catch misuse.
+	const_cast<FVector2D&>(AbsolutePosition) = AccumulatedLayoutTransform.GetTranslation();
+	const_cast<float&>(Scale) = AccumulatedLayoutTransform.GetScale();
+	const_cast<FVector2D&>(Position) = InLocalLayoutTransform.GetTranslation();
+}
+
+void FGeometry::AppendTransform(const FSlateLayoutTransform& LayoutTransform)
+{
+	FSlateLayoutTransform AccumulatedLayoutTransform = ::Concatenate(GetAccumulatedLayoutTransform(), LayoutTransform);
+	AccumulatedRenderTransform = ::Concatenate(AccumulatedRenderTransform, LayoutTransform);
+	const_cast<FVector2D&>(AbsolutePosition) = AccumulatedLayoutTransform.GetTranslation();
+	const_cast<float&>(Scale) = AccumulatedLayoutTransform.GetScale();
+}
+
+FGeometry FGeometry::MakeRoot(const FVector2D& LocalSize, const FSlateLayoutTransform& LayoutTransform)
+{
+	return FGeometry(LocalSize, LayoutTransform, FSlateLayoutTransform(), FSlateRenderTransform());
+}
+
+FGeometry FGeometry::MakeChild(const FVector2D& LocalSize, const FSlateLayoutTransform& LayoutTransform, const FSlateRenderTransform& RenderTransform) const
+{
+	return FGeometry(LocalSize, LayoutTransform, RenderTransform, GetAccumulatedLayoutTransform(), GetAccumulatedRenderTransform());
+}
+
+FGeometry FGeometry::MakeChild( const FVector2D& LocalSize, const FSlateLayoutTransform& LayoutTransform ) const
+{
+	return FGeometry(LocalSize, LayoutTransform, GetAccumulatedLayoutTransform(), GetAccumulatedRenderTransform());
+}
 
-/* FArrangedChildren interface
- *****************************************************************************/
+FArrangedWidget FGeometry::MakeChild(const TSharedRef<SWidget>& ChildWidget, const FLayoutGeometry& LayoutGeometry) const
+{
+	return MakeChild(ChildWidget, LayoutGeometry.GetSizeInLocalSpace(), LayoutGeometry.GetLocalToParentTransform());
+}
+
+FArrangedWidget FGeometry::MakeChild(const TSharedRef<SWidget>& ChildWidget, const FVector2D& LocalSize, const FSlateLayoutTransform& LayoutTransform) const
+{
+	return FArrangedWidget(ChildWidget, MakeChild(LocalSize, LayoutTransform, ChildWidget->GetRenderTransform()));
+}
+
+FGeometry FGeometry::MakeChild( const FVector2D& ChildOffset, const FVector2D& LocalSize, float LocalScale ) const
+{
+	// Since ChildOffset is given as a LocalSpaceOffset, we MUST convert this offset into the space of the parent to construct a valid layout transform.
+	// The extra TransformPoint below does this by converting the local offset to an offset in parent space.
+	return FGeometry(LocalSize, FSlateLayoutTransform(LocalScale, TransformPoint(LocalScale, ChildOffset)), GetAccumulatedLayoutTransform(), GetAccumulatedRenderTransform());
+}
 
-FArrangedWidget FGeometry::MakeChild( const TSharedRef<SWidget>& InWidget, const FVector2D& ChildOffset, const FVector2D& ChildSize, float InScale ) const
+FArrangedWidget FGeometry::MakeChild( const TSharedRef<SWidget>& ChildWidget, const FVector2D& ChildOffset, const FVector2D& LocalSize, float ChildScale) const
 {
-	return FArrangedWidget( InWidget, this->MakeChild(ChildOffset, ChildSize, InScale) );
+	// Since ChildOffset is given as a LocalSpaceOffset, we MUST convert this offset into the space of the parent to construct a valid layout transform.
+	// The extra TransformPoint below does this by converting the local offset to an offset in parent space.
+	return FArrangedWidget(ChildWidget, MakeChild(LocalSize, FSlateLayoutTransform(ChildScale, TransformPoint(ChildScale, ChildOffset)), ChildWidget->GetRenderTransform()));
 }
 
+FPaintGeometry FGeometry::ToPaintGeometry() const
+{
+	return FPaintGeometry(GetAccumulatedLayoutTransform(), GetAccumulatedRenderTransform(), Size);
+}
 
-FPaintGeometry FGeometry::ToPaintGeometry( FVector2D InOffset, FVector2D InSize, float InScale ) const
+FPaintGeometry FGeometry::ToPaintGeometry(const FSlateLayoutTransform& LayoutTransform) const
 {
-	const float CombinedScale = this->Scale*InScale;
-	return FPaintGeometry( this->AbsolutePosition + InOffset*CombinedScale, InSize * CombinedScale, CombinedScale );
+	return ToPaintGeometry(Size, LayoutTransform);
 }
 
 
+FPaintGeometry FGeometry::ToPaintGeometry(const FVector2D& LocalSize, const FSlateLayoutTransform& LayoutTransform) const
+{
+	FSlateLayoutTransform NewAccumulatedLayoutTransform = Concatenate(LayoutTransform, GetAccumulatedLayoutTransform());
+	return FPaintGeometry(NewAccumulatedLayoutTransform, Concatenate(LayoutTransform, GetAccumulatedRenderTransform()), LocalSize);
+}
+
+FPaintGeometry FGeometry::ToPaintGeometry(const FVector2D& LocalOffset, const FVector2D& LocalSize, float LocalScale) const
+{
+	// Since ChildOffset is given as a LocalSpaceOffset, we MUST convert this offset into the space of the parent to construct a valid layout transform.
+	// The extra TransformPoint below does this by converting the local offset to an offset in parent space.
+	return ToPaintGeometry(LocalSize, FSlateLayoutTransform(LocalScale, TransformPoint(LocalScale, LocalOffset)));
+}
+
+FPaintGeometry FGeometry::ToOffsetPaintGeometry(const FVector2D& LocalOffset) const
+{
+	return ToPaintGeometry(FSlateLayoutTransform(LocalOffset));
+}
+
 FPaintGeometry FGeometry::ToInflatedPaintGeometry( const FVector2D& InflateAmount ) const
 {
-	return FPaintGeometry( this->AbsolutePosition - InflateAmount*this->Scale, (this->Size+InflateAmount*2)*this->Scale, this->Scale );
+	// This essentially adds (or subtracts) a border around the widget. We scale the size then offset by the border amount.
+	// Note this is not scaling child widgets, so the scale is not changing.
+	FVector2D NewSize = Size + InflateAmount * 2;
+	return ToPaintGeometry(NewSize, FSlateLayoutTransform(-InflateAmount));
 }
 
+FSlateRect FGeometry::GetClippingRect() const
+{
+	return TransformRect(GetAccumulatedLayoutTransform(), FSlateRect(FVector2D(0.0f,0.0f), Size));
+}
 
-FPaintGeometry FGeometry::CenteredPaintGeometryOnLeft( const FVector2D& SizeBeingAligned, float InScale ) const
+FVector2D FGeometry::AbsoluteToLocal(FVector2D AbsoluteCoordinate) const
 {
-	const float CombinedScale = this->Scale*InScale;
-	
-	return FPaintGeometry(
-		this->AbsolutePosition + FVector2D(-SizeBeingAligned.X, this->Size.Y/2 - SizeBeingAligned.Y/2) * InScale,
-		SizeBeingAligned * CombinedScale,
-		CombinedScale
-	);
+	return TransformPoint(Inverse(GetAccumulatedLayoutTransform()), AbsoluteCoordinate);
 }
 
 
-FPaintGeometry FGeometry::CenteredPaintGeometryOnRight( const FVector2D& SizeBeingAligned, float InScale ) const
+FVector2D FGeometry::LocalToAbsolute(FVector2D LocalCoordinate) const
 {
-	const float CombinedScale = this->Scale*InScale;
-	
-	return FPaintGeometry(
-		this->AbsolutePosition + FVector2D(this->Size.X, this->Size.Y/2 - SizeBeingAligned.Y/2) * InScale,
-		SizeBeingAligned * CombinedScale,
-		CombinedScale
-	);
+	return TransformPoint(GetAccumulatedLayoutTransform(), LocalCoordinate);
 }
 
+bool FGeometry::IsUnderLocation(const FVector2D& AbsoluteCoordinate) const
+{
+	return GetClippingRect().ContainsPoint(AbsoluteCoordinate);
+}
+
+
+FString FGeometry::ToString() const
+{
+	return FString::Printf(TEXT("[Abs=%s, Scale=%.2f, Size=%s]"), *AbsolutePosition.ToString(), Scale, *Size.ToString());
+}
 
-FPaintGeometry FGeometry::CenteredPaintGeometryBelow( const FVector2D& SizeBeingAligned, float InScale ) const
+FVector2D FGeometry::GetDrawSize() const
 {
-	const float CombinedScale = this->Scale*InScale;
-	
-	return FPaintGeometry(
-		this->AbsolutePosition + FVector2D(this->Size.X/2 - SizeBeingAligned.X/2, this->Size.Y) * InScale,
-		SizeBeingAligned * CombinedScale,
-		CombinedScale
-	);
+	return TransformVector(GetAccumulatedLayoutTransform(), Size);
 }
diff --git a/Engine/Source/Runtime/SlateCore/Private/Layout/LayoutGeometry.cpp b/Engine/Source/Runtime/SlateCore/Private/Layout/LayoutGeometry.cpp
new file mode 100644
--- /dev/null
+++ b/Engine/Source/Runtime/SlateCore/Private/Layout/LayoutGeometry.cpp
@@ -0,0 +1,5 @@
+// Copyright 1998-2014 Epic Games, Inc. All Rights Reserved.
+
+#include "SlateCorePrivatePCH.h"
+#include "LayoutGeometry.h"
+
diff --git a/Engine/Source/Runtime/SlateCore/Private/Layout/SlateRect.cpp b/Engine/Source/Runtime/SlateCore/Private/Layout/SlateRect.cpp
--- a/Engine/Source/Runtime/SlateCore/Private/Layout/SlateRect.cpp
+++ b/Engine/Source/Runtime/SlateCore/Private/Layout/SlateRect.cpp
@@ -10,3 +10,8 @@ FSlateRect FSlateRect::InsetBy( const FMargin& InsetAmount ) const
 {
 	return FSlateRect(Left + InsetAmount.Left, Top + InsetAmount.Top, Right - InsetAmount.Right, Bottom - InsetAmount.Bottom);
 }
+
+FSlateRect FSlateRect::OffsetBy(const FVector2D& OffsetAmount) const
+{
+	return FSlateRect(GetTopLeft() + OffsetAmount, GetBottomRight() + OffsetAmount);
+}
diff --git a/Engine/Source/Runtime/SlateCore/Private/Rendering/DrawElements.cpp b/Engine/Source/Runtime/SlateCore/Private/Rendering/DrawElements.cpp
--- a/Engine/Source/Runtime/SlateCore/Private/Rendering/DrawElements.cpp
+++ b/Engine/Source/Runtime/SlateCore/Private/Rendering/DrawElements.cpp
@@ -5,10 +5,12 @@
 
 void FSlateDrawElement::MakeDebugQuad( FSlateWindowElementList& ElementList, uint32 InLayer, const FPaintGeometry& PaintGeometry, const FSlateRect& InClippingRect)
 {
+	PaintGeometry.CommitTransformsIfUsingLegacyConstructor();
 	FSlateDrawElement& DrawElt = ElementList.AddUninitialized();
 	DrawElt.ElementType = ET_DebugQuad;
+	DrawElt.RenderTransform = PaintGeometry.GetAccumulatedRenderTransform();
 	DrawElt.Position = PaintGeometry.DrawPosition;
-	DrawElt.Size = PaintGeometry.DrawSize;
+	DrawElt.LocalSize = PaintGeometry.GetLocalSize();
 	DrawElt.ClippingRect = InClippingRect;
 	DrawElt.Layer = InLayer;
 	DrawElt.DrawEffects = ESlateDrawEffect::None;
@@ -25,10 +27,12 @@ void FSlateDrawElement::MakeBox(
 	ESlateDrawEffect::Type InDrawEffects, 
 	const FLinearColor& InTint )
 {
+	PaintGeometry.CommitTransformsIfUsingLegacyConstructor();
 	FSlateDrawElement& DrawElt = ElementList.AddUninitialized();
 	DrawElt.ElementType = (InBrush->DrawAs == ESlateBrushDrawType::Border) ? ET_Border : ET_Box;
+	DrawElt.RenderTransform = PaintGeometry.GetAccumulatedRenderTransform();
 	DrawElt.Position = PaintGeometry.DrawPosition;
-	DrawElt.Size = PaintGeometry.DrawSize;
+	DrawElt.LocalSize = PaintGeometry.GetLocalSize();
 	DrawElt.ClippingRect = InClippingRect;
 	DrawElt.DataPayload.SetBoxPayloadProperties( InBrush, InTint );
 	DrawElt.Layer = InLayer;
@@ -36,7 +40,6 @@ void FSlateDrawElement::MakeBox(
 	DrawElt.Scale = PaintGeometry.DrawScale;
 }
 
-
 void FSlateDrawElement::MakeRotatedBox( 
 	FSlateWindowElementList& ElementList,
 	uint32 InLayer, 
@@ -49,10 +52,12 @@ void FSlateDrawElement::MakeRotatedBox(
 	ERotationSpace RotationSpace,
 	const FLinearColor& InTint )
 {
+	PaintGeometry.CommitTransformsIfUsingLegacyConstructor();
 	FSlateDrawElement& DrawElt = ElementList.AddUninitialized();
 	DrawElt.ElementType = (InBrush->DrawAs == ESlateBrushDrawType::Border) ? ET_Border : ET_Box;
+	DrawElt.RenderTransform = PaintGeometry.GetAccumulatedRenderTransform();
 	DrawElt.Position = PaintGeometry.DrawPosition;
-	DrawElt.Size = PaintGeometry.DrawSize;
+	DrawElt.LocalSize = PaintGeometry.GetLocalSize();
 	DrawElt.ClippingRect = InClippingRect;
 	DrawElt.Layer = InLayer;
 	DrawElt.DrawEffects = InDrawEffects;
@@ -65,10 +70,12 @@ void FSlateDrawElement::MakeRotatedBox(
 
 void FSlateDrawElement::MakeText( FSlateWindowElementList& ElementList, uint32 InLayer, const FPaintGeometry& PaintGeometry, const FString& InText, const int32 StartIndex, const int32 EndIndex, const FSlateFontInfo& InFontInfo, const FSlateRect& InClippingRect,ESlateDrawEffect::Type InDrawEffects, const FLinearColor& InTint )
 {
+	PaintGeometry.CommitTransformsIfUsingLegacyConstructor();
 	FSlateDrawElement& DrawElt = ElementList.AddUninitialized();
 	DrawElt.ElementType = ET_Text;
+	DrawElt.RenderTransform = PaintGeometry.GetAccumulatedRenderTransform();
 	DrawElt.Position = PaintGeometry.DrawPosition;
-	DrawElt.Size = PaintGeometry.DrawSize;
+	DrawElt.LocalSize = PaintGeometry.GetLocalSize();
 	DrawElt.ClippingRect = InClippingRect;
 	DrawElt.DataPayload.SetTextPayloadProperties( FString( EndIndex - StartIndex, *InText + StartIndex ), InFontInfo, InTint );
 	DrawElt.Layer = InLayer;
@@ -79,10 +86,12 @@ void FSlateDrawElement::MakeText( FSlateWindowElementList& ElementList, uint32 I
 
 void FSlateDrawElement::MakeText( FSlateWindowElementList& ElementList, uint32 InLayer, const FPaintGeometry& PaintGeometry, const FString& InText, const FSlateFontInfo& InFontInfo, const FSlateRect& InClippingRect,ESlateDrawEffect::Type InDrawEffects, const FLinearColor& InTint )
 {
+	PaintGeometry.CommitTransformsIfUsingLegacyConstructor();
 	FSlateDrawElement& DrawElt = ElementList.AddUninitialized();
 	DrawElt.ElementType = ET_Text;
+	DrawElt.RenderTransform = PaintGeometry.GetAccumulatedRenderTransform();
 	DrawElt.Position = PaintGeometry.DrawPosition;
-	DrawElt.Size = PaintGeometry.DrawSize;
+	DrawElt.LocalSize = PaintGeometry.GetLocalSize();
 	DrawElt.ClippingRect = InClippingRect;
 	DrawElt.DataPayload.SetTextPayloadProperties( InText, InFontInfo, InTint );
 	DrawElt.Layer = InLayer;
@@ -93,10 +102,12 @@ void FSlateDrawElement::MakeText( FSlateWindowElementList& ElementList, uint32 I
 
 void FSlateDrawElement::MakeText( FSlateWindowElementList& ElementList, uint32 InLayer, const FPaintGeometry& PaintGeometry, const FText& InText, const FSlateFontInfo& InFontInfo, const FSlateRect& InClippingRect,ESlateDrawEffect::Type InDrawEffects, const FLinearColor& InTint )
 {
+	PaintGeometry.CommitTransformsIfUsingLegacyConstructor();
 	FSlateDrawElement& DrawElt = ElementList.AddUninitialized();
 	DrawElt.ElementType = ET_Text;
+	DrawElt.RenderTransform = PaintGeometry.GetAccumulatedRenderTransform();
 	DrawElt.Position = PaintGeometry.DrawPosition;
-	DrawElt.Size = PaintGeometry.DrawSize;
+	DrawElt.LocalSize = PaintGeometry.GetLocalSize();
 	DrawElt.ClippingRect = InClippingRect;
 	DrawElt.DataPayload.SetTextPayloadProperties( InText.ToString(), InFontInfo, InTint );
 	DrawElt.Layer = InLayer;
@@ -104,26 +115,30 @@ void FSlateDrawElement::MakeText( FSlateWindowElementList& ElementList, uint32 I
 	DrawElt.Scale = PaintGeometry.DrawScale;
 }
 
-
-void FSlateDrawElement::MakeGradient( FSlateWindowElementList& ElementList, uint32 InLayer, const FPaintGeometry& InPaintGeometry, TArray<FSlateGradientStop> InGradientStops, EOrientation InGradientType, const FSlateRect& InClippingRect, ESlateDrawEffect::Type InDrawEffects, bool bGammaCorrect )
+void FSlateDrawElement::MakeGradient( FSlateWindowElementList& ElementList, uint32 InLayer, const FPaintGeometry& PaintGeometry, TArray<FSlateGradientStop> InGradientStops, EOrientation InGradientType, const FSlateRect& InClippingRect, ESlateDrawEffect::Type InDrawEffects, bool bGammaCorrect )
 {
+	PaintGeometry.CommitTransformsIfUsingLegacyConstructor();
 	FSlateDrawElement& DrawElt = ElementList.AddUninitialized();
 	DrawElt.ElementType = ET_Gradient;
-	DrawElt.Position = InPaintGeometry.DrawPosition;
-	DrawElt.Size = InPaintGeometry.DrawSize;
+	DrawElt.RenderTransform = PaintGeometry.GetAccumulatedRenderTransform();
+	DrawElt.Position = PaintGeometry.DrawPosition;
+	DrawElt.LocalSize = PaintGeometry.GetLocalSize();
 	DrawElt.ClippingRect = InClippingRect;
 	DrawElt.DataPayload.SetGradientPayloadProperties( InGradientStops, InGradientType, bGammaCorrect );
 	DrawElt.Layer = InLayer;
 	DrawElt.DrawEffects = InDrawEffects;
-	DrawElt.Scale = InPaintGeometry.DrawScale;
+	DrawElt.Scale = PaintGeometry.DrawScale;
 }
 
 
 void FSlateDrawElement::MakeSpline( FSlateWindowElementList& ElementList, uint32 InLayer, const FPaintGeometry& PaintGeometry, const FVector2D& InStart, const FVector2D& InStartDir, const FVector2D& InEnd, const FVector2D& InEndDir, const FSlateRect InClippingRect, float InThickness, ESlateDrawEffect::Type InDrawEffects, const FLinearColor& InTint )
 {
+	PaintGeometry.CommitTransformsIfUsingLegacyConstructor();
 	FSlateDrawElement& DrawElt = ElementList.AddUninitialized();
 	DrawElt.ElementType = ET_Spline;
+	DrawElt.RenderTransform = PaintGeometry.GetAccumulatedRenderTransform();
 	DrawElt.Position = PaintGeometry.DrawPosition;
+	DrawElt.LocalSize = PaintGeometry.GetLocalSize();
 	DrawElt.ClippingRect = InClippingRect;
 	DrawElt.DataPayload.SetSplinePayloadProperties( InStart, InStartDir, InEnd, InEndDir, InThickness, InTint );
 	DrawElt.Layer = InLayer;
@@ -134,15 +149,18 @@ void FSlateDrawElement::MakeSpline( FSlateWindowElementList& ElementList, uint32
 
 void FSlateDrawElement::MakeDrawSpaceSpline( FSlateWindowElementList& ElementList, uint32 InLayer, const FVector2D& InStart, const FVector2D& InStartDir, const FVector2D& InEnd, const FVector2D& InEndDir, const FSlateRect InClippingRect, float InThickness, ESlateDrawEffect::Type InDrawEffects, const FColor& InTint )
 {
-	MakeSpline( ElementList, InLayer, FPaintGeometry::Identity(), InStart, InStartDir, InEnd, InEndDir, InClippingRect, InThickness, InDrawEffects, InTint );
+	MakeSpline( ElementList, InLayer, FPaintGeometry(), InStart, InStartDir, InEnd, InEndDir, InClippingRect, InThickness, InDrawEffects, InTint );
 }
 
 
 void FSlateDrawElement::MakeLines( FSlateWindowElementList& ElementList, uint32 InLayer, const FPaintGeometry& PaintGeometry, const TArray<FVector2D>& Points, const FSlateRect InClippingRect, ESlateDrawEffect::Type InDrawEffects, const FLinearColor& InTint, bool bAntialias )
 {
+	PaintGeometry.CommitTransformsIfUsingLegacyConstructor();
 	FSlateDrawElement& DrawElt = ElementList.AddUninitialized();
 	DrawElt.ElementType = ET_Line;
+	DrawElt.RenderTransform = PaintGeometry.GetAccumulatedRenderTransform();
 	DrawElt.Position = PaintGeometry.DrawPosition;
+	DrawElt.LocalSize = PaintGeometry.GetLocalSize();
 	DrawElt.ClippingRect = InClippingRect;
 	DrawElt.DataPayload.SetLinesPayloadProperties( Points, InTint, bAntialias, ESlateLineJoinType::Sharp );
 	DrawElt.Layer = InLayer;
@@ -153,10 +171,12 @@ void FSlateDrawElement::MakeLines( FSlateWindowElementList& ElementList, uint32
 
 void FSlateDrawElement::MakeViewport( FSlateWindowElementList& ElementList, uint32 InLayer, const FPaintGeometry& PaintGeometry, TSharedPtr<const ISlateViewport> Viewport, const FSlateRect& InClippingRect, bool bGammaCorrect, bool bAllowBlending, ESlateDrawEffect::Type InDrawEffects, const FLinearColor& InTint )
 {
+	PaintGeometry.CommitTransformsIfUsingLegacyConstructor();
 	FSlateDrawElement& DrawElt = ElementList.AddUninitialized();
 	DrawElt.ElementType = ET_Viewport;
+	DrawElt.RenderTransform = PaintGeometry.GetAccumulatedRenderTransform();
 	DrawElt.Position = PaintGeometry.DrawPosition;
-	DrawElt.Size = PaintGeometry.DrawSize;
+	DrawElt.LocalSize = PaintGeometry.GetLocalSize();
 	DrawElt.ClippingRect = InClippingRect;
 	DrawElt.DataPayload.SetViewportPayloadProperties( Viewport, InTint, bGammaCorrect, bAllowBlending );
 	DrawElt.Layer = InLayer;
@@ -168,12 +188,40 @@ void FSlateDrawElement::MakeViewport( FSlateWindowElementList& ElementList, uint
 void FSlateDrawElement::MakeCustom( FSlateWindowElementList& ElementList, uint32 InLayer, TSharedPtr<ICustomSlateElement, ESPMode::ThreadSafe> CustomDrawer )
 {
 	FSlateDrawElement& DrawElt = ElementList.AddUninitialized();
+	DrawElt.RenderTransform = FSlateRenderTransform();
 	DrawElt.ElementType = ET_Custom;
 	DrawElt.Layer = InLayer;
 	DrawElt.DataPayload.SetCustomDrawerPayloadProperties( CustomDrawer );
 	DrawElt.ClippingRect = FSlateRect(1,1,1,1);
 }
 
+FVector2D FSlateDrawElement::GetRotationPoint(const FPaintGeometry& PaintGeometry, const TOptional<FVector2D>& UserRotationPoint, ERotationSpace RotationSpace)
+{
+	FVector2D RotationPoint(0, 0);
+
+	const FVector2D& LocalSize = PaintGeometry.GetLocalSize();
+
+	switch (RotationSpace)
+	{
+	case RelativeToElement:
+	{
+		// If the user did not specify a rotation point, we rotate about the center of the element
+		RotationPoint = UserRotationPoint.Get(LocalSize * 0.5f);
+	}
+		break;
+	case RelativeToWorld:
+	{
+		// its in world space, must convert the point to local space.
+		RotationPoint = TransformPoint(Inverse(PaintGeometry.GetAccumulatedRenderTransform()), UserRotationPoint.Get(FVector2D::ZeroVector));
+	}
+		break;
+	default:
+		check(0);
+		break;
+	}
+
+	return RotationPoint;
+}
 
 FSlateWindowElementList::FDeferredPaint::FDeferredPaint( const TSharedRef<SWidget>& InWidgetToPaint, const FPaintArgs& InArgs, const FGeometry InAllottedGeometry, const FSlateRect InMyClippingRect, const FWidgetStyle& InWidgetStyle, bool InParentEnabled )
 : WidgetToPaintPtr( InWidgetToPaint )
@@ -213,5 +261,3 @@ int32 FSlateWindowElementList::PaintDeferred( int32 LayerId )
 
 	return LayerId;
 }
-
-
diff --git a/Engine/Source/Runtime/SlateCore/Private/Rendering/ElementBatcher.cpp b/Engine/Source/Runtime/SlateCore/Private/Rendering/ElementBatcher.cpp
--- a/Engine/Source/Runtime/SlateCore/Private/Rendering/ElementBatcher.cpp
+++ b/Engine/Source/Runtime/SlateCore/Private/Rendering/ElementBatcher.cpp
@@ -1,6 +1,90 @@
 // Copyright 1998-2014 Epic Games, Inc. All Rights Reserved.
 
 #include "SlateCorePrivatePCH.h"
+#include "ElementBatcher.h"
+#include "SlateRenderTransform.h"
+
+FSlateRotatedRect::FSlateRotatedRect(const FSlateLayoutTransform& InverseLayoutTransform, const FSlateRenderTransform& RenderTransform, const FSlateRect& ClipRectInLayoutWindowSpace)
+{
+	const FSlateRenderTransform LayoutWindowSpaceToRenderWindowSpace = Concatenate(InverseLayoutTransform, RenderTransform);
+
+	const FVector2D SizeInLayoutWindowSpace = ClipRectInLayoutWindowSpace.GetSize();
+	const FVector2D TopLeftInRenderWindowSpace = TransformPoint(LayoutWindowSpaceToRenderWindowSpace, ClipRectInLayoutWindowSpace.GetTopLeft());
+	const FVector2D ClipExtentXInRenderWindowSpace = TransformVector(LayoutWindowSpaceToRenderWindowSpace, FVector2D(SizeInLayoutWindowSpace.X, 0.0f));
+	const FVector2D ClipExtentYInRenderWindowSpace = TransformVector(LayoutWindowSpaceToRenderWindowSpace, FVector2D(0.0f, SizeInLayoutWindowSpace.Y));
+	TopLeft.X = FMath::RoundToInt(TopLeftInRenderWindowSpace.X);
+	TopLeft.Y = FMath::RoundToInt(TopLeftInRenderWindowSpace.Y);
+	ExtentX.X = FMath::RoundToInt(ClipExtentXInRenderWindowSpace.X);
+	ExtentX.Y = FMath::RoundToInt(ClipExtentXInRenderWindowSpace.Y);
+	ExtentY.X = FMath::RoundToInt(ClipExtentYInRenderWindowSpace.X);
+	ExtentY.Y = FMath::RoundToInt(ClipExtentYInRenderWindowSpace.Y);
+}
+
+// !!! WRH 2014/08/25 - this is a brute-force, not efficient implementation, uses a bunch of extra conditionals.
+FSlateRect FSlateRotatedRect::ToBoundingRect() const
+{
+	FVector2D Points[4] = 
+	{
+		TopLeft,
+		TopLeft + ExtentX,
+		TopLeft + ExtentY,
+		TopLeft + ExtentX + ExtentY
+	};
+	return FSlateRect(
+		FMath::Min(Points[0].X, FMath::Min3(Points[1].X, Points[2].X, Points[3].X)),
+		FMath::Min(Points[0].Y, FMath::Min3(Points[1].Y, Points[2].Y, Points[3].Y)),
+		FMath::Max(Points[0].X, FMath::Max3(Points[1].X, Points[2].X, Points[3].X)),
+		FMath::Max(Points[0].Y, FMath::Max3(Points[1].Y, Points[2].Y, Points[3].Y))
+		);
+}
+
+bool FSlateRotatedRect::IsUnderLocation(const FVector2D& Location) const
+{
+	const FVector2D Offset = Location - TopLeft;
+	const float Det = FVector2D::CrossProduct(ExtentX, ExtentY);
+
+	// Not exhaustively efficient. Could optimize the checks for [0..1] to short circuit faster.
+	const float S = -FVector2D::CrossProduct(Offset, ExtentX) / Det;
+	if (FMath::IsWithinInclusive(S, 0.0f, 1.0f))
+	{
+		const float T = FVector2D::CrossProduct(Offset, ExtentY) / Det;
+		return FMath::IsWithinInclusive(T, 0.0f, 1.0f);
+	}
+	return false;
+}
+
+FSlateRotatedRectFloat16::FSlateRotatedRectFloat16() {}
+FSlateRotatedRectFloat16::FSlateRotatedRectFloat16(const FSlateRotatedRect& RotatedRect)
+{
+	TopLeft[0] = RotatedRect.TopLeft[0];
+	TopLeft[1] = RotatedRect.TopLeft[1];
+	ExtentX[0] = RotatedRect.ExtentX[0];
+	ExtentX[1] = RotatedRect.ExtentX[1];
+	ExtentY[0] = RotatedRect.ExtentY[0];
+	ExtentY[1] = RotatedRect.ExtentY[1];
+}
+
+FSlateVertex::FSlateVertex() {}
+
+FSlateVertex::FSlateVertex(const FSlateRenderTransform& RenderTransform, const FVector2D& InLocalPosition, const FVector2D& InTexCoord, const FVector2D& InTexCoord2, const FColor& InColor, const FSlateRotatedRectFloat16& InClipRect )
+	: TexCoords( InTexCoord, InTexCoord2 )
+	, ClipRect( InClipRect )
+	, Color( InColor )
+{
+	const FVector2D WindowPosition = TransformPoint(RenderTransform, InLocalPosition);
+	Position[0] = FMath::RoundToInt(WindowPosition.X);
+	Position[1] = FMath::RoundToInt(WindowPosition.Y);
+}
+
+FSlateVertex::FSlateVertex( const FSlateRenderTransform& RenderTransform, const FVector2D& InLocalPosition, const FVector2D& InTexCoord, const FColor& InColor, const FSlateRotatedRectFloat16& InClipRect )
+	: TexCoords( InTexCoord, FVector2D(1.0f, 1.0f) )
+	, ClipRect( InClipRect )
+	, Color( InColor )
+{
+	const FVector2D WindowPosition = TransformPoint(RenderTransform, InLocalPosition);
+	Position[0] = FMath::RoundToInt(WindowPosition.X);
+	Position[1] = FMath::RoundToInt(WindowPosition.Y);
+}
 
 
 static const FSlateImageBrush& GetSplineFilterTable()
@@ -21,35 +105,6 @@ static FColor GetElementColor( const FLinearColor& InColor, const FSlateBrush* I
 	return InColor.ToFColor(false);
 }
 
-
-/**
- * Rotates a point around another point 
- *
- * @param InPoint	The point to rotate
- * @param AboutPoint	The point to rotate about
- * @param Radians	The angle of rotation in radians
- * @return The rotated point
- */
-static FVector2D RotatePoint( FVector2D InPoint, const FVector2D& AboutPoint, float Radians )
-{
-	if( Radians != 0.0f )
-	{
-		float CosAngle = FMath::Cos(Radians);
-		float SinAngle = FMath::Sin(Radians);
-
-		InPoint.X -= AboutPoint.X;
-		InPoint.Y -= AboutPoint.Y;
-
-		float X = (InPoint.X * CosAngle) - (InPoint.Y * SinAngle);
-		float Y = (InPoint.X * SinAngle) + (InPoint.Y * CosAngle);
-
-		return FVector2D( X + AboutPoint.X, Y + AboutPoint.Y );
-	}
-
-	return InPoint;
-}
-
-
 FSlateElementBatcher::FSlateElementBatcher( TSharedRef<FSlateRenderingPolicy> InRenderingPolicy )
 	: ResourceManager( *InRenderingPolicy->GetResourceManager() )
 	, FontCache( *InRenderingPolicy->GetFontCache() )
@@ -83,31 +138,31 @@ void FSlateElementBatcher::AddElements( const TArray<FSlateDrawElement>& DrawEle
 			switch( DrawElement.GetElementType() )
 			{
 			case FSlateDrawElement::ET_Box:
-				AddBoxElement( DrawElement.GetPosition(), DrawElement.GetSize(), DrawElement.GetScale(), DrawElement.GetDataPayload(), DrawElement.GetClippingRect(), DrawElement.GetDrawEffects(), DrawElement.GetLayer() );
+				AddBoxElement( DrawElement );
 				break;
 			case FSlateDrawElement::ET_DebugQuad:
-				AddQuadElement( DrawElement.GetPosition(), DrawElement.GetSize(), DrawElement.GetScale(), DrawElement.GetClippingRect(), DrawElement.GetLayer() );
+				AddQuadElement( DrawElement );
 				break;
 			case FSlateDrawElement::ET_Text:
-				AddTextElement( DrawElement.GetPosition(), DrawElement.GetSize(), DrawElement.GetScale(), DrawElement.GetDataPayload(), DrawElement.GetClippingRect(), DrawElement.GetDrawEffects(), DrawElement.GetLayer() );
+				AddTextElement( DrawElement );
 				break;
 			case FSlateDrawElement::ET_Spline:
-				AddSplineElement( DrawElement.GetPosition(), DrawElement.GetScale(), DrawElement.GetDataPayload(), DrawElement.GetClippingRect(), DrawElement.GetDrawEffects(), DrawElement.GetLayer() );
+				AddSplineElement( DrawElement );
 				break;
 			case FSlateDrawElement::ET_Line:
-				AddLineElement( DrawElement.GetPosition(), DrawElement.GetSize(), DrawElement.GetScale(), DrawElement.GetDataPayload(), DrawElement.GetClippingRect(), DrawElement.GetDrawEffects(), DrawElement.GetLayer() );
+				AddLineElement( DrawElement );
 				break;
 			case FSlateDrawElement::ET_Gradient:
-				AddGradientElement( DrawElement.GetPosition(), DrawElement.GetSize(), DrawElement.GetScale(), DrawElement.GetDataPayload(), DrawElement.GetClippingRect(), DrawElement.GetDrawEffects(), DrawElement.GetLayer() );
+				AddGradientElement( DrawElement );
 				break;
 			case FSlateDrawElement::ET_Viewport:
-				AddViewportElement( DrawElement.GetPosition(), DrawElement.GetSize(), DrawElement.GetScale(), DrawElement.GetDataPayload(), DrawElement.GetClippingRect(), DrawElement.GetDrawEffects(), DrawElement.GetLayer() );
+				AddViewportElement( DrawElement );
 				break;
 			case FSlateDrawElement::ET_Border:
-				AddBorderElement( DrawElement.GetPosition(), DrawElement.GetSize(), DrawElement.GetScale(), DrawElement.GetDataPayload(), DrawElement.GetClippingRect(), DrawElement.GetDrawEffects(), DrawElement.GetLayer() ); 
+				AddBorderElement( DrawElement );
 				break;
 			case FSlateDrawElement::ET_Custom:
-				AddCustomElement( DrawElement.GetPosition(), DrawElement.GetSize(), DrawElement.GetScale(), DrawElement.GetDataPayload(), DrawElement.GetClippingRect(), DrawElement.GetDrawEffects(), DrawElement.GetLayer() ); 
+				AddCustomElement( DrawElement );
 				break;
 			default:
 				checkf(0, TEXT("Invalid element type"));
@@ -117,26 +172,40 @@ void FSlateElementBatcher::AddElements( const TArray<FSlateDrawElement>& DrawEle
 	}
 }
 
-void FSlateElementBatcher::AddQuadElement( const FVector2D& Position, const FVector2D& Size, float Scale, const FSlateRect& InClippingRect, uint32 Layer, FColor Color )
+void FSlateElementBatcher::AddQuadElement( const FSlateDrawElement& DrawElement, FColor Color )
 {
+	const FSlateRenderTransform& RenderTransform = DrawElement.GetRenderTransform();
+	//const FVector2D& InPosition = DrawElement.GetPosition();
+	//const FVector2D& Size = DrawElement.GetSize();
+	const FVector2D& LocalSize = DrawElement.GetLocalSize();
+	//float Scale = DrawElement.GetScale();
+	//const FSlateDataPayload& InPayload = DrawElement.GetDataPayload();
+	const FSlateRect& InClippingRect = DrawElement.GetClippingRect();
+	//ESlateDrawEffect::Type InDrawEffects = DrawElement.GetDrawEffects();
+	uint32 Layer = DrawElement.GetLayer();
+
+	// extract the layout transform from the draw element
+	FSlateLayoutTransform InverseLayoutTransform(Inverse(FSlateLayoutTransform(DrawElement.GetScale(), DrawElement.GetPosition())));
+	FSlateRotatedRectFloat16 RenderClipRect = FSlateRotatedRect(InverseLayoutTransform, RenderTransform, InClippingRect);
+
 	FSlateElementBatch& ElementBatch = FindBatchForElement( Layer, FShaderParams(), nullptr, ESlateDrawPrimitive::TriangleList, ESlateShader::Default, ESlateDrawEffect::None, ESlateBatchDrawFlag::Wireframe|ESlateBatchDrawFlag::NoBlending );
 	TArray<FSlateVertex>& BatchVertices = BatchVertexArrays[ElementBatch.VertexArrayIndex];
 	TArray<SlateIndex>& BatchIndices = BatchIndexArrays[ElementBatch.IndexArrayIndex];
 
 	// Determine the four corners of the quad
-	FVector2D TopLeft = Position;
-	FVector2D TopRight = FVector2D( Position.X+Size.X, Position.Y);
-	FVector2D BotLeft = FVector2D( Position.X, Position.Y+Size.Y);
-	FVector2D BotRight = FVector2D( Position.X + Size.X, Position.Y+Size.Y);
+	FVector2D TopLeft = FVector2D::ZeroVector;
+	FVector2D TopRight = FVector2D(LocalSize.X, 0);
+	FVector2D BotLeft = FVector2D(0, LocalSize.Y);
+	FVector2D BotRight = FVector2D(LocalSize.X, LocalSize.Y);
 
 	// The start index of these vertices in the index buffer
 	uint32 IndexStart = BatchVertices.Num();
 
 	// Add four vertices to the list of verts to be added to the vertex buffer
-	BatchVertices.Add( FSlateVertex(TopLeft, FVector2D(0.0f,0.0f), Color, InClippingRect ) );
-	BatchVertices.Add( FSlateVertex(TopRight, FVector2D(1.0f,0.0f), Color, InClippingRect ) );
-	BatchVertices.Add( FSlateVertex(BotLeft, FVector2D(0.0f,1.0f), Color, InClippingRect ) );
-	BatchVertices.Add( FSlateVertex(BotRight, FVector2D(1.0f,1.0f),Color, InClippingRect ) );
+	BatchVertices.Add( FSlateVertex( RenderTransform, TopLeft, FVector2D(0.0f,0.0f), Color, RenderClipRect ) );
+	BatchVertices.Add( FSlateVertex( RenderTransform, TopRight, FVector2D(1.0f,0.0f), Color, RenderClipRect ) );
+	BatchVertices.Add( FSlateVertex( RenderTransform, BotLeft, FVector2D(0.0f,1.0f), Color, RenderClipRect ) );
+	BatchVertices.Add( FSlateVertex( RenderTransform, BotRight, FVector2D(1.0f,1.0f),Color, RenderClipRect ) );
 
 	// The offset into the index buffer where this quads indices start
 	uint32 IndexOffsetStart = BatchIndices.Num();
@@ -151,18 +220,46 @@ void FSlateElementBatcher::AddQuadElement( const FVector2D& Position, const FVec
 }
 
 
-void FSlateElementBatcher::AddBoxElement( const FVector2D& InPosition, const FVector2D& Size, float Scale, const FSlateDataPayload& InPayload, const FSlateRect& InClippingRect, ESlateDrawEffect::Type InDrawEffects, uint32 Layer )
+FSlateRenderTransform GetBoxRenderTransform(const FSlateDrawElement& DrawElement)
 {
-	check( InPayload.BrushResource );
+	const FSlateRenderTransform& ElementRenderTransform = DrawElement.GetRenderTransform();
+	const float RotationAngle = DrawElement.GetDataPayload().Angle;
+	if (RotationAngle == 0.0f)
+	{
+		return ElementRenderTransform;
+	}
+	const FVector2D RotationPoint = DrawElement.GetDataPayload().RotationPoint;
+	const FSlateRenderTransform RotationTransform = Concatenate(Inverse(RotationPoint), FQuat2D(RotationAngle), RotationPoint);
+	return Concatenate(RotationTransform, ElementRenderTransform);
+}
 
+void FSlateElementBatcher::AddBoxElement( const FSlateDrawElement& DrawElement )
+{
 	SCOPE_CYCLE_COUNTER(STAT_SlateBuildBoxes);
 
+	const FSlateRenderTransform& ElementRenderTransform = DrawElement.GetRenderTransform();
+	const FSlateRenderTransform RenderTransform = GetBoxRenderTransform(DrawElement);
+	//const FVector2D& InPosition = DrawElement.GetPosition();
+	//const FVector2D& Size = DrawElement.GetSize();
+	const FVector2D& LocalSize = DrawElement.GetLocalSize();
+	//float Scale = DrawElement.GetScale();
+	const FSlateDataPayload& InPayload = DrawElement.GetDataPayload();
+	const FSlateRect& InClippingRect = DrawElement.GetClippingRect();
+	ESlateDrawEffect::Type InDrawEffects = DrawElement.GetDrawEffects();
+	uint32 Layer = DrawElement.GetLayer();
+
+	// extract the layout transform from the draw element
+	FSlateLayoutTransform InverseLayoutTransform(Inverse(FSlateLayoutTransform(DrawElement.GetScale(), DrawElement.GetPosition())));
+	// The clip rect is NOT subject to the rotations specified by MakeRotatedBox.
+	FSlateRotatedRectFloat16 RenderClipRect = FSlateRotatedRect(InverseLayoutTransform, ElementRenderTransform, InClippingRect);
+
+	check( InPayload.BrushResource );
+
 	if(InPayload.BrushResource->DrawAs != ESlateBrushDrawType::NoDrawType)
 	{
-		FVector2D Position( FMath::TruncToInt(InPosition.X), FMath::TruncToInt(InPosition.Y) );
-
-		const float Angle = InPayload.Angle;
-		const FVector2D RotationPoint = InPayload.RotationPoint;
+		// Do pixel snapping
+		FVector2D TopLeft(0,0);
+		FVector2D BotRight(LocalSize);
 
 		uint32 TextureWidth = 1;
 		uint32 TextureHeight = 1;
@@ -212,14 +309,12 @@ void FSlateElementBatcher::AddBoxElement( const FVector2D& InPosition, const FVe
 		// Pass the tiling information as a flag so we can pick the correct texture addressing mode
 		ESlateBatchDrawFlag::Type DrawFlags = ( ( bTileHorizontal ? ESlateBatchDrawFlag::TileU : 0 ) | ( bTileVertical ? ESlateBatchDrawFlag::TileV : 0 ) );
 
-		FShaderParams RotationParams = Angle != 0.0f ? FShaderParams::MakeVertexShaderParams( FVector4( Angle, RotationPoint.X, RotationPoint.Y, 0 ) ) : FShaderParams();
-
-		FSlateElementBatch& ElementBatch = FindBatchForElement( Layer, RotationParams, Resource, ESlateDrawPrimitive::TriangleList, ESlateShader::Default, InDrawEffects, DrawFlags );
+		FSlateElementBatch& ElementBatch = FindBatchForElement( Layer, FShaderParams(), Resource, ESlateDrawPrimitive::TriangleList, ESlateShader::Default, InDrawEffects, DrawFlags );
 		TArray<FSlateVertex>& BatchVertices = BatchVertexArrays[ElementBatch.VertexArrayIndex];
 		TArray<SlateIndex>& BatchIndices = BatchIndexArrays[ElementBatch.IndexArrayIndex];
 
-		float HorizontalTiling = bTileHorizontal ? Size.X/Scale/TextureWidth : 1.0f;
-		float VerticalTiling = bTileVertical ? Size.Y/Scale/TextureHeight : 1.0f;
+		float HorizontalTiling = bTileHorizontal ? LocalSize.X/TextureWidth : 1.0f;
+		float VerticalTiling = bTileVertical ? LocalSize.Y/TextureHeight : 1.0f;
 
 		const FVector2D Tiling( HorizontalTiling,VerticalTiling );
 
@@ -263,38 +358,6 @@ void FSlateElementBatcher::AddBoxElement( const FVector2D& InPosition, const FVe
 				? EndUV.Y - Margin.Bottom * SizeUV.Y + HalfTexel.Y
 				: EndUV.Y;
 
-			// Determine the margins for each quad
-			const FMargin ScaledMargin = Margin * Scale;
-
-			float LeftMarginX = TextureWidth * ScaledMargin.Left;
-			float TopMarginY = TextureHeight * ScaledMargin.Top;
-			float RightMarginX = Size.X - TextureWidth * ScaledMargin.Right;
-			float BottomMarginY = Size.Y - TextureHeight * ScaledMargin.Bottom;
-
-			// If the margins are overlapping the margins are too big or the button is too small
-			// so clamp margins to half of the box size
-			if( RightMarginX < LeftMarginX )
-			{
-				LeftMarginX = Size.X / 2;
-				RightMarginX = LeftMarginX;
-			}
-
-			if( BottomMarginY < TopMarginY )
-			{
-				TopMarginY = Size.Y / 2;
-				BottomMarginY = TopMarginY;
-			}
-
-			LeftMarginX += Position.X;
-			TopMarginY += Position.Y;
-			RightMarginX += Position.X;
-			BottomMarginY += Position.Y;
-
-			FVector2D EndPos = FVector2D(Position.X, Position.Y) + Size;
-
-			EndPos.X = FMath::TruncToInt( EndPos.X );
-			EndPos.Y = FMath::TruncToInt( EndPos.Y );
-
 			if( bMirrorHorizontal || bMirrorVertical )
 			{
 				const FVector2D UVMin = StartUV;
@@ -316,23 +379,47 @@ void FSlateElementBatcher::AddBoxElement( const FVector2D& InPosition, const FVe
 				}
 			}
 
-			BatchVertices.Add( FSlateVertex( Position,									StartUV,									Tiling,	Tint, InClippingRect ) ); //0
-			BatchVertices.Add( FSlateVertex( FVector2D( Position.X, TopMarginY ),		FVector2D( StartUV.X, TopMarginV ),			Tiling,	Tint, InClippingRect ) ); //1
-			BatchVertices.Add( FSlateVertex( FVector2D( LeftMarginX, Position.Y ),		FVector2D( LeftMarginU, StartUV.Y ),		Tiling,	Tint, InClippingRect ) ); //2
-			BatchVertices.Add( FSlateVertex( FVector2D( LeftMarginX, TopMarginY ),		FVector2D( LeftMarginU, TopMarginV ),		Tiling,	Tint, InClippingRect ) ); //3
-			BatchVertices.Add( FSlateVertex( FVector2D( RightMarginX, Position.Y ),		FVector2D( RightMarginU, StartUV.Y ),		Tiling,	Tint, InClippingRect ) ); //4
-			BatchVertices.Add( FSlateVertex( FVector2D( RightMarginX, TopMarginY ),		FVector2D( RightMarginU,TopMarginV),		Tiling,	Tint, InClippingRect ) ); //5
-			BatchVertices.Add( FSlateVertex( FVector2D( EndPos.X, Position.Y ),			FVector2D( EndUV.X, StartUV.Y ),			Tiling,	Tint, InClippingRect ) ); //6
-			BatchVertices.Add( FSlateVertex( FVector2D( EndPos.X, TopMarginY ),			FVector2D( EndUV.X, TopMarginV),			Tiling,	Tint, InClippingRect ) ); //7
-
-			BatchVertices.Add( FSlateVertex( FVector2D( Position.X, BottomMarginY ),	FVector2D( StartUV.X, BottomMarginV ),		Tiling,	Tint, InClippingRect ) ); //8
-			BatchVertices.Add( FSlateVertex( FVector2D( LeftMarginX, BottomMarginY ),	FVector2D( LeftMarginU, BottomMarginV ),	Tiling,	Tint, InClippingRect ) ); //9
-			BatchVertices.Add( FSlateVertex( FVector2D( RightMarginX, BottomMarginY ),	FVector2D( RightMarginU, BottomMarginV ),	Tiling,	Tint, InClippingRect ) ); //10
-			BatchVertices.Add( FSlateVertex( FVector2D( EndPos.X, BottomMarginY ),		FVector2D( EndUV.X, BottomMarginV ),		Tiling,	Tint, InClippingRect ) ); //11
-			BatchVertices.Add( FSlateVertex( FVector2D( Position.X, EndPos.Y ),			FVector2D( StartUV.X, EndUV.Y ),			Tiling,	Tint, InClippingRect ) ); //12
-			BatchVertices.Add( FSlateVertex( FVector2D( LeftMarginX, EndPos.Y ),		FVector2D( LeftMarginU, EndUV.Y ),			Tiling,	Tint, InClippingRect ) ); //13
-			BatchVertices.Add( FSlateVertex( FVector2D( RightMarginX, EndPos.Y ),		FVector2D( RightMarginU, EndUV.Y ),			Tiling,	Tint, InClippingRect ) ); //14
-			BatchVertices.Add( FSlateVertex( FVector2D( EndPos.X, EndPos.Y ),			EndUV,										Tiling,	Tint, InClippingRect ) ); //15
+			// Determine the margins for each quad
+
+			float LeftMarginX = TextureWidth * Margin.Left;
+			float TopMarginY = TextureHeight * Margin.Top;
+			float RightMarginX = LocalSize.X - TextureWidth * Margin.Right;
+			float BottomMarginY = LocalSize.Y - TextureHeight * Margin.Bottom;
+
+			// If the margins are overlapping the margins are too big or the button is too small
+			// so clamp margins to half of the box size
+			if( RightMarginX < LeftMarginX )
+			{
+				LeftMarginX = LocalSize.X / 2;
+				RightMarginX = LeftMarginX;
+			}
+
+			if( BottomMarginY < TopMarginY )
+			{
+				TopMarginY = LocalSize.Y / 2;
+				BottomMarginY = TopMarginY;
+			}
+
+			FVector2D Position = TopLeft;
+			FVector2D EndPos = BotRight;
+
+			BatchVertices.Add( FSlateVertex( RenderTransform, FVector2D( Position.X, Position.Y ),		StartUV,									Tiling,	Tint, RenderClipRect ) ); //0
+			BatchVertices.Add( FSlateVertex( RenderTransform, FVector2D( Position.X, TopMarginY ),		FVector2D( StartUV.X, TopMarginV ),			Tiling,	Tint, RenderClipRect ) ); //1
+			BatchVertices.Add( FSlateVertex( RenderTransform, FVector2D( LeftMarginX, Position.Y ),		FVector2D( LeftMarginU, StartUV.Y ),		Tiling,	Tint, RenderClipRect ) ); //2
+			BatchVertices.Add( FSlateVertex( RenderTransform, FVector2D( LeftMarginX, TopMarginY ),		FVector2D( LeftMarginU, TopMarginV ),		Tiling,	Tint, RenderClipRect ) ); //3
+			BatchVertices.Add( FSlateVertex( RenderTransform, FVector2D( RightMarginX, Position.Y ),	FVector2D( RightMarginU, StartUV.Y ),		Tiling,	Tint, RenderClipRect ) ); //4
+			BatchVertices.Add( FSlateVertex( RenderTransform, FVector2D( RightMarginX, TopMarginY ),	FVector2D( RightMarginU,TopMarginV),		Tiling,	Tint, RenderClipRect ) ); //5
+			BatchVertices.Add( FSlateVertex( RenderTransform, FVector2D( EndPos.X, Position.Y ),		FVector2D( EndUV.X, StartUV.Y ),			Tiling,	Tint, RenderClipRect ) ); //6
+			BatchVertices.Add( FSlateVertex( RenderTransform, FVector2D( EndPos.X, TopMarginY ),		FVector2D( EndUV.X, TopMarginV),			Tiling,	Tint, RenderClipRect ) ); //7
+
+			BatchVertices.Add( FSlateVertex( RenderTransform, FVector2D( Position.X, BottomMarginY ),	FVector2D( StartUV.X, BottomMarginV ),		Tiling,	Tint, RenderClipRect ) ); //8
+			BatchVertices.Add( FSlateVertex( RenderTransform, FVector2D( LeftMarginX, BottomMarginY ),	FVector2D( LeftMarginU, BottomMarginV ),	Tiling,	Tint, RenderClipRect ) ); //9
+			BatchVertices.Add( FSlateVertex( RenderTransform, FVector2D( RightMarginX, BottomMarginY ),	FVector2D( RightMarginU, BottomMarginV ),	Tiling,	Tint, RenderClipRect ) ); //10
+			BatchVertices.Add( FSlateVertex( RenderTransform, FVector2D( EndPos.X, BottomMarginY ),		FVector2D( EndUV.X, BottomMarginV ),		Tiling,	Tint, RenderClipRect ) ); //11
+			BatchVertices.Add( FSlateVertex( RenderTransform, FVector2D( Position.X, EndPos.Y ),		FVector2D( StartUV.X, EndUV.Y ),			Tiling,	Tint, RenderClipRect ) ); //12
+			BatchVertices.Add( FSlateVertex( RenderTransform, FVector2D( LeftMarginX, EndPos.Y ),		FVector2D( LeftMarginU, EndUV.Y ),			Tiling,	Tint, RenderClipRect ) ); //13
+			BatchVertices.Add( FSlateVertex( RenderTransform, FVector2D( RightMarginX, EndPos.Y ),		FVector2D( RightMarginU, EndUV.Y ),			Tiling,	Tint, RenderClipRect ) ); //14
+			BatchVertices.Add( FSlateVertex( RenderTransform, FVector2D( EndPos.X, EndPos.Y ),			EndUV,										Tiling,	Tint, RenderClipRect ) ); //15
 
 			// Top
 			BatchIndices.Add( IndexStart + 0 );
@@ -402,10 +489,8 @@ void FSlateElementBatcher::AddBoxElement( const FVector2D& InPosition, const FVe
 		}
 		else
 		{
-			FVector2D TopLeft =	Position;
-			FVector2D TopRight = FVector2D( Position.X+Size.X, Position.Y);
-			FVector2D BotLeft =	 FVector2D( Position.X, Position.Y+Size.Y);
-			FVector2D BotRight = FVector2D( Position.X + Size.X, Position.Y+Size.Y);
+			FVector2D TopRight = FVector2D( BotRight.X, TopLeft.Y);
+			FVector2D BotLeft =	 FVector2D( TopLeft.X, BotRight.Y);
 
 			if( bMirrorHorizontal || bMirrorVertical )
 			{
@@ -425,10 +510,10 @@ void FSlateElementBatcher::AddBoxElement( const FVector2D& InPosition, const FVe
 			}
 
 			// Add four vertices to the list of verts to be added to the vertex buffer
-			BatchVertices.Add( FSlateVertex(TopLeft,	StartUV,						Tiling,	Tint, InClippingRect ) );
-			BatchVertices.Add( FSlateVertex(TopRight,	FVector2D(EndUV.X,StartUV.Y),	Tiling,	Tint, InClippingRect ) );
-			BatchVertices.Add( FSlateVertex(BotLeft,	FVector2D(StartUV.X,EndUV.Y),	Tiling,	Tint, InClippingRect ) );
-			BatchVertices.Add( FSlateVertex(BotRight,	EndUV,							Tiling,	Tint, InClippingRect ) );
+			BatchVertices.Add( FSlateVertex( RenderTransform, TopLeft,	StartUV,						Tiling,	Tint, RenderClipRect ) );
+			BatchVertices.Add( FSlateVertex( RenderTransform, TopRight,	FVector2D(EndUV.X,StartUV.Y),	Tiling,	Tint, RenderClipRect ) );
+			BatchVertices.Add( FSlateVertex( RenderTransform, BotLeft,	FVector2D(StartUV.X,EndUV.Y),	Tiling,	Tint, RenderClipRect ) );
+			BatchVertices.Add( FSlateVertex( RenderTransform, BotRight,	EndUV,							Tiling,	Tint, RenderClipRect ) );
 
 			BatchIndices.Add( IndexStart + 0 );
 			BatchIndices.Add( IndexStart + 1 );
@@ -442,10 +527,33 @@ void FSlateElementBatcher::AddBoxElement( const FVector2D& InPosition, const FVe
 }
 
 
-void FSlateElementBatcher::AddTextElement( const FVector2D& Position, const FVector2D& Size, float InScale, const FSlateDataPayload& InPayload, const FSlateRect& InClippingRect, ESlateDrawEffect::Type InDrawEffects, uint32 Layer )
+void FSlateElementBatcher::AddTextElement(const FSlateDrawElement& DrawElement)
 {
 	SCOPE_CYCLE_COUNTER(STAT_SlateBuildText);
 
+	//const FVector2D& InPosition = DrawElement.GetPosition();
+	//const FVector2D& Size = DrawElement.GetSize();
+	//const FVector2D& LocalSize = DrawElement.GetLocalSize();
+	//float Scale = DrawElement.GetScale();
+	const FSlateDataPayload& InPayload = DrawElement.GetDataPayload();
+	const FSlateRect& InClippingRect = DrawElement.GetClippingRect();
+	ESlateDrawEffect::Type InDrawEffects = DrawElement.GetDrawEffects();
+	uint32 Layer = DrawElement.GetLayer();
+
+	// extract the layout transform from the draw element
+	FSlateLayoutTransform LayoutTransform(DrawElement.GetScale(), DrawElement.GetPosition());
+
+	// We don't just scale up fonts, we draw them in local space pre-scaled so we don't get scaling artifcats.
+	// So we need to pull the layout scale out of the layout and render transform so we can apply them
+	// in local space with pre-scaled fonts.
+	const float FontScale = LayoutTransform.GetScale();
+	FSlateLayoutTransform InverseLayoutTransform = Inverse(Concatenate(Inverse(FontScale), LayoutTransform));
+	const FSlateRenderTransform RenderTransform = Concatenate(Inverse(FontScale), DrawElement.GetRenderTransform());
+
+	FSlateRotatedRectFloat16 RenderClipRect = FSlateRotatedRect(InverseLayoutTransform, RenderTransform, InClippingRect);
+	// Used to clip individual characters as we generate them.
+	FSlateRect LocalClipRect = TransformRect(InverseLayoutTransform, InClippingRect);
+
 	const FString& Text = InPayload.Text;
 	// Nothing to do if no text
 	if( Text.Len() == 0 )
@@ -453,7 +561,6 @@ void FSlateElementBatcher::AddTextElement( const FVector2D& Position, const FVec
 		return;
 	}
 
-	const float FontScale = InScale;
 
 	FCharacterList& CharacterList = FontCache.GetCharacterList( InPayload.FontInfo, FontScale );
 
@@ -472,14 +579,16 @@ void FSlateElementBatcher::AddTextElement( const FVector2D& Position, const FVec
 	float InvTextureSizeX = 0;
 	float InvTextureSizeY = 0;
 
-	int32 LineX = 0;
+	float LineX = 0;
 
 	TCHAR PreviousChar = 0;
 
 	int32 Kerning = 0;
 
-	int32 PosX = FMath::TruncToInt(Position.X);
-	int32 PosY = FMath::TruncToInt(Position.Y);
+	FVector2D TopLeft(0,0);
+
+	const float PosX = TopLeft.X;
+	float PosY = TopLeft.Y;
 
 	LineX = PosX;
 	
@@ -549,7 +658,7 @@ void FSlateElementBatcher::AddTextElement( const FVector2D& Position, const FVec
 
 				FSlateRect CharRect( X, Y, X+SizeX, Y+SizeY );
 
-				if( FSlateRect::DoRectanglesIntersect( InClippingRect, CharRect )	)
+				if( FSlateRect::DoRectanglesIntersect( LocalClipRect, CharRect ) )
 				{
 					TArray<FSlateVertex>& BatchVerticesRef = *BatchVertices;
 					TArray<SlateIndex>& BatchIndicesRef = *BatchIndices;
@@ -568,10 +677,10 @@ void FSlateElementBatcher::AddTextElement( const FVector2D& Position, const FVec
 					uint32 IndexStart = VertexOffset;
 
 					// Add four vertices to the list of verts to be added to the vertex buffer
-					BatchVerticesRef[ VertexOffset++ ] = FSlateVertex( UpperLeft,								FVector2D(U,V),				FinalColor, InClippingRect );
-					BatchVerticesRef[ VertexOffset++ ] = FSlateVertex( FVector2D(LowerRight.X,UpperLeft.Y),		FVector2D(U+SizeU, V),		FinalColor, InClippingRect );
-					BatchVerticesRef[ VertexOffset++ ] = FSlateVertex( FVector2D(UpperLeft.X,LowerRight.Y),		FVector2D(U, V+SizeV),		FinalColor, InClippingRect );
-					BatchVerticesRef[ VertexOffset++ ] = FSlateVertex( LowerRight,								FVector2D(U+SizeU, V+SizeV),FinalColor, InClippingRect );
+					BatchVerticesRef[ VertexOffset++ ] = FSlateVertex( RenderTransform, UpperLeft,								FVector2D(U,V),				FinalColor, RenderClipRect );
+					BatchVerticesRef[ VertexOffset++ ] = FSlateVertex( RenderTransform, FVector2D(LowerRight.X,UpperLeft.Y),	FVector2D(U+SizeU, V),		FinalColor, RenderClipRect );
+					BatchVerticesRef[ VertexOffset++ ] = FSlateVertex( RenderTransform, FVector2D(UpperLeft.X,LowerRight.Y),	FVector2D(U, V+SizeV),		FinalColor, RenderClipRect );
+					BatchVerticesRef[ VertexOffset++ ] = FSlateVertex( RenderTransform, LowerRight,								FVector2D(U+SizeU, V+SizeV),FinalColor, RenderClipRect );
 
 					BatchIndicesRef[IndexOffset++] = IndexStart + 0;
 					BatchIndicesRef[IndexOffset++] = IndexStart + 1;
@@ -588,10 +697,24 @@ void FSlateElementBatcher::AddTextElement( const FVector2D& Position, const FVec
 }
 
 
-void FSlateElementBatcher::AddGradientElement( const FVector2D& Position, const FVector2D& Size, float Scale, const class FSlateDataPayload& InPayload, const FSlateRect& InClippingRect, ESlateDrawEffect::Type InDrawEffects, uint32 Layer )
+void FSlateElementBatcher::AddGradientElement( const FSlateDrawElement& DrawElement )
 {
 	SCOPE_CYCLE_COUNTER(STAT_SlateBuildGradients);
 
+	const FSlateRenderTransform& RenderTransform = DrawElement.GetRenderTransform();
+	//const FVector2D& InPosition = DrawElement.GetPosition();
+	//const FVector2D& Size = DrawElement.GetSize();
+	const FVector2D& LocalSize = DrawElement.GetLocalSize();
+	//float Scale = DrawElement.GetScale();
+	const FSlateDataPayload& InPayload = DrawElement.GetDataPayload();
+	const FSlateRect& InClippingRect = DrawElement.GetClippingRect();
+	ESlateDrawEffect::Type InDrawEffects = DrawElement.GetDrawEffects();
+	uint32 Layer = DrawElement.GetLayer();
+
+	// extract the layout transform from the draw element
+	FSlateLayoutTransform InverseLayoutTransform(Inverse(FSlateLayoutTransform(DrawElement.GetScale(), DrawElement.GetPosition())));
+	FSlateRotatedRectFloat16 RenderClipRect = FSlateRotatedRect(InverseLayoutTransform, RenderTransform, InClippingRect);
+
 	// There must be at least one gradient stop
 	check( InPayload.GradientStops.Num() > 0 );
 
@@ -600,12 +723,10 @@ void FSlateElementBatcher::AddGradientElement( const FVector2D& Position, const
 	TArray<SlateIndex>& BatchIndices = BatchIndexArrays[ElementBatch.IndexArrayIndex];
 
 	// Determine the four corners of the quad containing the gradient
-	FVector2D TopLeft = Position;
-	FVector2D TopRight = FVector2D(Position.X + Size.X, Position.Y);
-	FVector2D BotLeft = FVector2D(Position.X, Position.Y + Size.Y);
-	FVector2D BotRight = FVector2D(Position.X + Size.X, Position.Y + Size.Y);
-
-	FVector2D LocalSize = Size / Scale;
+	FVector2D TopLeft = FVector2D::ZeroVector;
+	FVector2D TopRight = FVector2D(LocalSize.X, 0);
+	FVector2D BotLeft = FVector2D(0, LocalSize.Y);
+	FVector2D BotRight = FVector2D(LocalSize.X, LocalSize.Y);
 
 	// Copy the gradient stops.. We may need to add more
 	TArray<FSlateGradientStop> GradientStops = InPayload.GradientStops;
@@ -665,31 +786,31 @@ void FSlateElementBatcher::AddGradientElement( const FVector2D& Position, const
 			// Gradient stop is vertical so gradients to left to right
 			StartPt = TopLeft;
 			EndPt = BotLeft;
-			// Gradient stops are interpreted as % of local Size. Determine the % and use it to scale the point in screen space.
-			StartPt.X += (CurStop.Position.X / LocalSize.X) * ( BotRight.X - TopLeft.X);
-			EndPt.X += (CurStop.Position.X / LocalSize.X) * ( BotRight.X - TopLeft.X);	
+			// Gradient stops are interpreted in local space.
+			StartPt.X += CurStop.Position.X;
+			EndPt.X += CurStop.Position.X;
 		}
 		else
 		{
 			// Gradient stop is horizontal so gradients to top to bottom
 			StartPt = TopLeft;
 			EndPt = TopRight;
-			// Gradient stops are interpreted as % of local Size. Determine the % and use it to scale the point in screen space.
-			StartPt.Y += (CurStop.Position.Y / LocalSize.Y) * (BotRight.Y - TopLeft.Y);
-			EndPt.Y += (CurStop.Position.Y / LocalSize.Y) * (BotRight.Y - TopLeft.Y);
+			// Gradient stops are interpreted in local space.
+			StartPt.Y += CurStop.Position.Y;
+			EndPt.Y += CurStop.Position.Y;
 		}
 
 		if( StopIndex == 0 )
 		{
 			// First stop does not have a full quad yet so do not create indices
-			BatchVertices.Add( FSlateVertex(StartPt, FVector2D::ZeroVector, FVector2D::ZeroVector, CurStop.Color.ToFColor(false), InClippingRect ) );
-			BatchVertices.Add( FSlateVertex(EndPt, FVector2D::ZeroVector, FVector2D::ZeroVector, CurStop.Color.ToFColor(false), InClippingRect ) );
+			BatchVertices.Add( FSlateVertex(RenderTransform, StartPt, FVector2D::ZeroVector, FVector2D::ZeroVector, CurStop.Color.ToFColor(false), RenderClipRect ) );
+			BatchVertices.Add( FSlateVertex(RenderTransform, EndPt, FVector2D::ZeroVector, FVector2D::ZeroVector, CurStop.Color.ToFColor(false), RenderClipRect ) );
 		}
 		else
 		{
 			// All stops after the first have indices and generate quads
-			BatchVertices.Add( FSlateVertex(StartPt, FVector2D::ZeroVector, FVector2D::ZeroVector, CurStop.Color.ToFColor(false), InClippingRect ) );
-			BatchVertices.Add( FSlateVertex(EndPt, FVector2D::ZeroVector, FVector2D::ZeroVector, CurStop.Color.ToFColor(false), InClippingRect ) );
+			BatchVertices.Add( FSlateVertex(RenderTransform, StartPt, FVector2D::ZeroVector, FVector2D::ZeroVector, CurStop.Color.ToFColor(false), RenderClipRect ) );
+			BatchVertices.Add( FSlateVertex(RenderTransform, EndPt, FVector2D::ZeroVector, FVector2D::ZeroVector, CurStop.Color.ToFColor(false), RenderClipRect ) );
 
 			// Connect the indices to the previous vertices
 			BatchIndices.Add( IndexStart - 2 );
@@ -703,19 +824,34 @@ void FSlateElementBatcher::AddGradientElement( const FVector2D& Position, const
 	}
 }
 
-
-void FSlateElementBatcher::AddSplineElement( const FVector2D& Position, float Scale, const FSlateDataPayload& InPayload, const FSlateRect& InClippingRect, ESlateDrawEffect::Type InDrawEffects, uint32 Layer)
+void FSlateElementBatcher::AddSplineElement( const FSlateDrawElement& DrawElement )
 {
 	SCOPE_CYCLE_COUNTER(STAT_SlateBuildSplines);
 
+	const FSlateRenderTransform& RenderTransform = DrawElement.GetRenderTransform();
+	//const FVector2D& InPosition = DrawElement.GetPosition();
+	//const FVector2D& Size = DrawElement.GetSize();
+	//const FVector2D& LocalSize = DrawElement.GetLocalSize();
+	//float Scale = DrawElement.GetScale();
+	const FSlateDataPayload& InPayload = DrawElement.GetDataPayload();
+	const FSlateRect& InClippingRect = DrawElement.GetClippingRect();
+	ESlateDrawEffect::Type InDrawEffects = DrawElement.GetDrawEffects();
+	uint32 Layer = DrawElement.GetLayer();
+
+	// extract the layout transform from the draw element
+	FSlateLayoutTransform InverseLayoutTransform(Inverse(FSlateLayoutTransform(DrawElement.GetScale(), DrawElement.GetPosition())));
+	FSlateRotatedRectFloat16 RenderClipRect = FSlateRotatedRect(InverseLayoutTransform, RenderTransform, InClippingRect);
+
 	//@todo SLATE: Merge with AddLineElement?
 
+	//@todo SLATE: This should probably be done in window space so there are no scaling artifacts?
 	const float DirectLength = (InPayload.EndPt - InPayload.StartPt).Size();
 	const float HandleLength = ((InPayload.EndPt - InPayload.EndDir) - (InPayload.StartPt + InPayload.StartDir)).Size();
 	float NumSteps = FMath::Clamp<float>(FMath::CeilToInt(FMath::Max(DirectLength,HandleLength)/15.0f), 1, 256);
 
 	// 1 is the minimum thickness we support
-	float InThickness = FMath::Max( 1.0f, InPayload.Thickness );
+	// Thickness is given in screenspace, so convert it to local space before proceeding.
+	float InThickness = FMath::Max( 1.0f, InverseLayoutTransform.GetScale() * InPayload.Thickness );
 
 	// The radius to use when checking the distance of pixels to the actual line.  Arbitrary value based on what looks the best
 	const float Radius = 1.5f;
@@ -737,10 +873,10 @@ void FSlateElementBatcher::AddSplineElement( const FVector2D& Position, float Sc
 	TArray<SlateIndex>& BatchIndices = BatchIndexArrays[ElementBatch.IndexArrayIndex];
 
 
-	const FVector2D StartPt = Position + InPayload.StartPt * Scale;
-	const FVector2D StartDir = InPayload.StartDir * Scale;
-	const FVector2D EndPt = Position + InPayload.EndPt * Scale;
-	const FVector2D EndDir = InPayload.EndDir * Scale;
+	const FVector2D StartPt = InPayload.StartPt;
+	const FVector2D StartDir = InPayload.StartDir;
+	const FVector2D EndPt = InPayload.EndPt;
+	const FVector2D EndDir = InPayload.EndDir;
 	
 	// Compute the normal to the line
 	FVector2D Normal = FVector2D( StartPt.Y - EndPt.Y, EndPt.X - StartPt.X ).SafeNormal();
@@ -754,8 +890,8 @@ void FSlateElementBatcher::AddSplineElement( const FVector2D& Position, float Sc
 
 	FColor FinalColor = GetElementColor( InPayload.Tint, nullptr );
 
-	BatchVertices.Add( FSlateVertex( StartPos + Up, StartPos, EndPos, FinalColor, InClippingRect ) );
-	BatchVertices.Add( FSlateVertex( StartPos - Up, StartPos, EndPos, FinalColor, InClippingRect ) );
+	BatchVertices.Add( FSlateVertex( RenderTransform, StartPos + Up, TransformPoint(RenderTransform, StartPos), TransformPoint(RenderTransform, EndPos), FinalColor, RenderClipRect ) );
+	BatchVertices.Add( FSlateVertex( RenderTransform, StartPos - Up, TransformPoint(RenderTransform, StartPos), TransformPoint(RenderTransform, EndPos), FinalColor, RenderClipRect ) );
 
 	// Generate the rest of the segments
 	for( int32 Step = 0; Step < NumSteps; ++Step )
@@ -775,8 +911,8 @@ void FSlateElementBatcher::AddSplineElement( const FVector2D& Position, float Sc
 		// Create the new vertices for the thick line segment
 		Up = SegmentNormal * HalfThickness;
 
-		BatchVertices.Add( FSlateVertex( EndPos + Up, StartPos, EndPos, FinalColor, InClippingRect ) );
-		BatchVertices.Add( FSlateVertex( EndPos - Up, StartPos, EndPos, FinalColor, InClippingRect ) );
+		BatchVertices.Add( FSlateVertex( RenderTransform, EndPos + Up, TransformPoint(RenderTransform, StartPos), TransformPoint(RenderTransform, EndPos), FinalColor, RenderClipRect ) );
+		BatchVertices.Add( FSlateVertex( RenderTransform, EndPos - Up, TransformPoint(RenderTransform, StartPos), TransformPoint(RenderTransform, EndPos), FinalColor, RenderClipRect ) );
 
 		BatchIndices.Add( IndexStart - 2 );
 		BatchIndices.Add( IndexStart - 1 );
@@ -834,10 +970,24 @@ static bool LineIntersect( const FVector2D& P1, const FVector2D& P2, const FVect
 }
 	
 
-void FSlateElementBatcher::AddLineElement( const FVector2D& Position, const FVector2D& Size, float Scale, const FSlateDataPayload& InPayload, const FSlateRect& InClippingRect, ESlateDrawEffect::Type DrawEffects, uint32 Layer )
+void FSlateElementBatcher::AddLineElement( const FSlateDrawElement& DrawElement )
 {
 	SCOPE_CYCLE_COUNTER(STAT_SlateBuildLines);
 
+	const FSlateRenderTransform& RenderTransform = DrawElement.GetRenderTransform();
+	//const FVector2D& InPosition = DrawElement.GetPosition();
+	//const FVector2D& Size = DrawElement.GetSize();
+	//const FVector2D& LocalSize = DrawElement.GetLocalSize();
+	//float Scale = DrawElement.GetScale();
+	const FSlateDataPayload& InPayload = DrawElement.GetDataPayload();
+	const FSlateRect& InClippingRect = DrawElement.GetClippingRect();
+	ESlateDrawEffect::Type DrawEffects = DrawElement.GetDrawEffects();
+	uint32 Layer = DrawElement.GetLayer();
+
+	// extract the layout transform from the draw element
+	FSlateLayoutTransform InverseLayoutTransform(Inverse(FSlateLayoutTransform(DrawElement.GetScale(), DrawElement.GetPosition())));
+	FSlateRotatedRectFloat16 RenderClipRect = FSlateRotatedRect(InverseLayoutTransform, RenderTransform, InClippingRect);
+
 	if( InPayload.Points.Num() < 2 )
 	{
 		return;
@@ -854,7 +1004,8 @@ void FSlateElementBatcher::AddLineElement( const FVector2D& Position, const FVec
 		// The radius to use when checking the distance of pixels to the actual line.  Arbitrary value based on what looks the best
 		const float Radius = 1.5f;
 
-		float RequestedThickness = FMath::Max( 1.0f, InPayload.Thickness );
+		// Thickness is given in screenspace, so convert it to local space before proceeding.
+		float RequestedThickness = FMath::Max( 1.0f, InverseLayoutTransform.GetScale() * InPayload.Thickness );
 
 		// Compute the actual size of the line we need based on thickness.  Need to ensure pixels that are at least Thickness/2 + Sample radius are generated so that we have enough pixels to blend.
 		// The idea for the anti-aliasing technique is based on the fast prefiltered lines technique published in GPU Gems 2 
@@ -870,20 +1021,20 @@ void FSlateElementBatcher::AddLineElement( const FVector2D& Position, const FVec
 
 		const TArray<FVector2D>& Points = InPayload.Points;
 
-		FVector2D StartPos = Position + Points[0]*Scale;
-		FVector2D EndPos = Position + Points[1]*Scale;
+		FVector2D StartPos = Points[0];
+		FVector2D EndPos = Points[1];
 
 		FVector2D Normal = FVector2D( StartPos.Y - EndPos.Y, EndPos.X - StartPos.X ).SafeNormal();
 
 		FVector2D Up = Normal * HalfThickness;
 
-		BatchVertices.Add( FSlateVertex( StartPos + Up, StartPos, EndPos, FinalColor, InClippingRect ) );
-		BatchVertices.Add( FSlateVertex( StartPos - Up, StartPos, EndPos, FinalColor, InClippingRect ) );
+		BatchVertices.Add( FSlateVertex( RenderTransform, StartPos + Up, TransformPoint(RenderTransform, StartPos), TransformPoint(RenderTransform, EndPos), FinalColor, RenderClipRect ) );
+		BatchVertices.Add( FSlateVertex( RenderTransform, StartPos - Up, TransformPoint(RenderTransform, StartPos), TransformPoint(RenderTransform, EndPos), FinalColor, RenderClipRect ) );
 
 		// Generate the rest of the segments
 		for( int32 Point = 1; Point < Points.Num(); ++Point )
 		{
-			EndPos = Position + Points[Point]*Scale;
+			EndPos = Points[Point];
 			// Determine if we should check the intersection point with the next line segment.
 			// We will adjust were this line ends to the intersection
 			bool bCheckIntersection = Points.IsValidIndex( Point+1 ) ;
@@ -902,7 +1053,7 @@ void FSlateElementBatcher::AddLineElement( const FVector2D& Position, const FVec
 			if( bCheckIntersection )
 			{
 				// The end point of the next segment
-				const FVector2D NextEndPos = Position + Points[Point+1]*Scale;
+				const FVector2D NextEndPos = Points[Point+1];
 
 				// The normal of the next segment
 				const FVector2D NextNormal = FVector2D( EndPos.Y - NextEndPos.Y, NextEndPos.X - EndPos.X ).SafeNormal();
@@ -933,29 +1084,33 @@ void FSlateElementBatcher::AddLineElement( const FVector2D& Position, const FVec
 				IntersectCenter = (IntersectUpper+IntersectLower) * .5f;
 			}
 
+			// We use these points when making the copy of the vert below, so go ahead and cache it.
+			FVector2D StartPosRenderSpace = TransformPoint(RenderTransform, StartPos);
+			FVector2D IntersectCenterRenderSpace = TransformPoint(RenderTransform, IntersectCenter);
+
 			if( Point > 1 )
 			{
 				// Make a copy of the last two vertices and update their start and end position to reflect the new line segment
 				FSlateVertex StartV1 = BatchVertices[IndexStart-1];
 				FSlateVertex StartV2 = BatchVertices[IndexStart-2];
 
-				StartV1.TexCoords.X = StartPos.X;
-				StartV1.TexCoords.Y = StartPos.Y;
-				StartV1.TexCoords.Z = IntersectCenter.X;
-				StartV1.TexCoords.W = IntersectCenter.Y;
+				StartV1.TexCoords.X = StartPosRenderSpace.X;
+				StartV1.TexCoords.Y = StartPosRenderSpace.Y;
+				StartV1.TexCoords.Z = IntersectCenterRenderSpace.X;
+				StartV1.TexCoords.W = IntersectCenterRenderSpace.Y;
 
-				StartV2.TexCoords.X = StartPos.X;
-				StartV2.TexCoords.Y = StartPos.Y;
-				StartV2.TexCoords.Z = IntersectCenter.X;
-				StartV2.TexCoords.W = IntersectCenter.Y;
+				StartV2.TexCoords.X = StartPosRenderSpace.X;
+				StartV2.TexCoords.Y = StartPosRenderSpace.Y;
+				StartV2.TexCoords.Z = IntersectCenterRenderSpace.X;
+				StartV2.TexCoords.W = IntersectCenterRenderSpace.Y;
 
 				IndexStart += 2;
 				BatchVertices.Add( StartV2 );
 				BatchVertices.Add( StartV1 );
 			}
 
-			BatchVertices.Add( FSlateVertex( IntersectUpper, StartPos, IntersectCenter, FinalColor, InClippingRect ) );
-			BatchVertices.Add( FSlateVertex( IntersectLower, StartPos, IntersectCenter, FinalColor, InClippingRect ) );
+			BatchVertices.Add( FSlateVertex( RenderTransform, IntersectUpper, StartPosRenderSpace, IntersectCenterRenderSpace, FinalColor, RenderClipRect ) );
+			BatchVertices.Add( FSlateVertex( RenderTransform, IntersectLower, StartPosRenderSpace, IntersectCenterRenderSpace, FinalColor, RenderClipRect ) );
 
 			BatchIndices.Add( IndexStart - 1 );
 			BatchIndices.Add( IndexStart - 2 );
@@ -979,12 +1134,12 @@ void FSlateElementBatcher::AddLineElement( const FVector2D& Position, const FVec
 		for( int32 Point = 0; Point < InPayload.Points.Num()-1; ++Point )
 		{
 			uint32 IndexStart = BatchVertices.Num();
-			FVector2D StartPos = Position + InPayload.Points[Point]*Scale;
-			FVector2D EndPos = Position + InPayload.Points[Point+1]*Scale;
+			FVector2D StartPos = InPayload.Points[Point];
+			FVector2D EndPos = InPayload.Points[Point+1];
 
 
-			BatchVertices.Add( FSlateVertex( StartPos, FVector2D::ZeroVector, FinalColor, InClippingRect ) );
-			BatchVertices.Add( FSlateVertex( EndPos, FVector2D::ZeroVector, FinalColor, InClippingRect ) );
+			BatchVertices.Add( FSlateVertex( RenderTransform, StartPos, FVector2D::ZeroVector, FinalColor, RenderClipRect ) );
+			BatchVertices.Add( FSlateVertex( RenderTransform, EndPos, FVector2D::ZeroVector, FinalColor, RenderClipRect ) );
 
 			BatchIndices.Add( IndexStart );
 			BatchIndices.Add( IndexStart + 1 );
@@ -993,10 +1148,24 @@ void FSlateElementBatcher::AddLineElement( const FVector2D& Position, const FVec
 }
 
 
-void FSlateElementBatcher::AddViewportElement( const FVector2D& InPosition, const FVector2D& InSize, float Scale, const FSlateDataPayload& InPayload, const FSlateRect& InClippingRect, ESlateDrawEffect::Type InDrawEffects, uint32 Layer )
+void FSlateElementBatcher::AddViewportElement( const FSlateDrawElement& DrawElement )
 {
 	SCOPE_CYCLE_COUNTER(STAT_SlateBuildViewports);
 
+	const FSlateRenderTransform& RenderTransform = DrawElement.GetRenderTransform();
+	//const FVector2D& InPosition = DrawElement.GetPosition();
+	//const FVector2D& Size = DrawElement.GetSize();
+	const FVector2D& LocalSize = DrawElement.GetLocalSize();
+	//float Scale = DrawElement.GetScale();
+	const FSlateDataPayload& InPayload = DrawElement.GetDataPayload();
+	const FSlateRect& InClippingRect = DrawElement.GetClippingRect();
+	ESlateDrawEffect::Type InDrawEffects = DrawElement.GetDrawEffects();
+	uint32 Layer = DrawElement.GetLayer();
+
+	// extract the layout transform from the draw element
+	FSlateLayoutTransform InverseLayoutTransform(Inverse(FSlateLayoutTransform(DrawElement.GetScale(), DrawElement.GetPosition())));
+	FSlateRotatedRectFloat16 RenderClipRect = FSlateRotatedRect(InverseLayoutTransform, RenderTransform, InClippingRect);
+
 	const FColor FinalColor = GetElementColor( InPayload.Tint, nullptr );
 
 	ESlateBatchDrawFlag::Type DrawFlags = ESlateBatchDrawFlag::None;
@@ -1026,23 +1195,21 @@ void FSlateElementBatcher::AddViewportElement( const FVector2D& InPosition, cons
 		bRequiresVsync |= ViewportPin->RequiresVsync();
 	}
 
-	FVector2D Position = FVector2D( FMath::TruncToInt( InPosition.X ), FMath::TruncToInt( InPosition.Y ) );
-	FVector2D Size = FVector2D( FMath::TruncToInt( InSize.X ), FMath::TruncToInt( InSize.Y ) );
+	// Do pixel snapping
+	FVector2D TopLeft(0,0);
+	FVector2D BotRight(LocalSize);
 
-	// Determine the four corners of the quad
-	FVector2D TopLeft = Position;
-	FVector2D TopRight = FVector2D(Position.X+Size.X, Position.Y);
-	FVector2D BotLeft = FVector2D(Position.X, Position.Y+Size.Y);
-	FVector2D BotRight = FVector2D(Position.X + Size.X, Position.Y+Size.Y);
+	FVector2D TopRight = FVector2D( BotRight.X, TopLeft.Y);
+	FVector2D BotLeft =	 FVector2D( TopLeft.X, BotRight.Y);
 
 	// The start index of these vertices in the index buffer
 	uint32 IndexStart = BatchVertices.Num();
 
 	// Add four vertices to the list of verts to be added to the vertex buffer
-	BatchVertices.Add( FSlateVertex( TopLeft,	FVector2D(0.0f,0.0f),	FinalColor, InClippingRect ) );
-	BatchVertices.Add( FSlateVertex( TopRight,	FVector2D(1.0f,0.0f),	FinalColor, InClippingRect ) );
-	BatchVertices.Add( FSlateVertex( BotLeft,	FVector2D(0.0f,1.0f),	FinalColor, InClippingRect ) );
-	BatchVertices.Add( FSlateVertex( BotRight,	FVector2D(1.0f,1.0f),	FinalColor, InClippingRect ) );
+	BatchVertices.Add( FSlateVertex( RenderTransform, TopLeft,	FVector2D(0.0f,0.0f),	FinalColor, RenderClipRect ) );
+	BatchVertices.Add( FSlateVertex( RenderTransform, TopRight,	FVector2D(1.0f,0.0f),	FinalColor, RenderClipRect ) );
+	BatchVertices.Add( FSlateVertex( RenderTransform, BotLeft,	FVector2D(0.0f,1.0f),	FinalColor, RenderClipRect ) );
+	BatchVertices.Add( FSlateVertex( RenderTransform, BotRight,	FVector2D(1.0f,1.0f),	FinalColor, RenderClipRect ) );
 
 	// The offset into the index buffer where this quads indices start
 	uint32 IndexOffsetStart = BatchIndices.Num();
@@ -1059,13 +1226,25 @@ void FSlateElementBatcher::AddViewportElement( const FVector2D& InPosition, cons
 }
 
 
-void FSlateElementBatcher::AddBorderElement( const FVector2D& InPosition, const FVector2D& Size, float Scale, const FSlateDataPayload& InPayload, const FSlateRect& InClippingRect, ESlateDrawEffect::Type InDrawEffects, uint32 Layer )
+void FSlateElementBatcher::AddBorderElement( const FSlateDrawElement& DrawElement )
 {
 	SCOPE_CYCLE_COUNTER(STAT_SlateBuildBorders);
 
-	check( InPayload.BrushResource );
+	const FSlateRenderTransform& RenderTransform = DrawElement.GetRenderTransform();
+	//const FVector2D& InPosition = DrawElement.GetPosition();
+	//const FVector2D& Size = DrawElement.GetSize();
+	const FVector2D& LocalSize = DrawElement.GetLocalSize();
+	//float Scale = DrawElement.GetScale();
+	const FSlateDataPayload& InPayload = DrawElement.GetDataPayload();
+	const FSlateRect& InClippingRect = DrawElement.GetClippingRect();
+	ESlateDrawEffect::Type InDrawEffects = DrawElement.GetDrawEffects();
+	uint32 Layer = DrawElement.GetLayer();
 
-	const FVector2D Position( FMath::TruncToInt(InPosition.X), FMath::TruncToInt(InPosition.Y) );
+	// extract the layout transform from the draw element
+	FSlateLayoutTransform InverseLayoutTransform(Inverse(FSlateLayoutTransform(DrawElement.GetScale(), DrawElement.GetPosition())));
+	FSlateRotatedRectFloat16 RenderClipRect = FSlateRotatedRect(InverseLayoutTransform, RenderTransform, InClippingRect);
+
+	check( InPayload.BrushResource );
 
 	uint32 TextureWidth = 1;
 	uint32 TextureHeight = 1;
@@ -1078,6 +1257,7 @@ void FSlateElementBatcher::AddBorderElement( const FVector2D& InPosition, const
 		TextureWidth = Resource->GetWidth();
 		TextureHeight = Resource->GetHeight();
 	}
+	FVector2D TextureSizeLocalSpace = TransformVector(InverseLayoutTransform, FVector2D(TextureWidth, TextureHeight));
  
 	// Texel offset
 	const FVector2D HalfTexel( PixelCenterOffset/TextureWidth, PixelCenterOffset/TextureHeight );
@@ -1087,25 +1267,29 @@ void FSlateElementBatcher::AddBorderElement( const FVector2D& InPosition, const
 
 	const FMargin& Margin = InPayload.BrushResource->Margin;
 
-	const FMargin ScaledMargin = Margin * Scale;
-
+	// Do pixel snapping
+	FVector2D TopLeft(0,0);
+	FVector2D BotRight(LocalSize);
 	// Determine the margins for each quad
-	float LeftMarginX = TextureWidth * ScaledMargin.Left;
-	float TopMarginY = TextureHeight * ScaledMargin.Top;
-	float RightMarginX = Size.X - TextureWidth * ScaledMargin.Right;
-	float BottomMarginY = Size.Y - TextureHeight * ScaledMargin.Bottom;
+	FVector2D TopLeftMargin(TextureSizeLocalSpace * FVector2D(Margin.Left, Margin.Top));
+	FVector2D BotRightMargin(LocalSize - TextureSizeLocalSpace * FVector2D(Margin.Right, Margin.Bottom));
+
+	float LeftMarginX = TopLeftMargin.X;
+	float TopMarginY = TopLeftMargin.Y;
+	float RightMarginX = BotRightMargin.X;
+	float BottomMarginY = BotRightMargin.Y;
 
 	// If the margins are overlapping the margins are too big or the button is too small
 	// so clamp margins to half of the box size
 	if( RightMarginX < LeftMarginX )
 	{
-		LeftMarginX = Size.X / 2;
+		LeftMarginX = LocalSize.X / 2;
 		RightMarginX = LeftMarginX;
 	}
 
 	if( BottomMarginY < TopMarginY )
 	{
-		TopMarginY = Size.Y / 2;
+		TopMarginY = LocalSize.Y / 2;
 		BottomMarginY = TopMarginY;
 	}
 
@@ -1121,8 +1305,8 @@ void FSlateElementBatcher::AddBorderElement( const FVector2D& InPosition, const
 	RightMarginU += HalfTexel.X;
 
 	// Determine the amount of tiling needed for the texture in this element.  The formula is number of pixels covered by the tiling portion of the texture / the number number of texels corresponding to the tiled portion of the texture.
-	float TopTiling = (RightMarginX-LeftMarginX)/(TextureWidth * ( 1 - Margin.GetTotalSpaceAlong<Orient_Horizontal>() ));
-	float LeftTiling = (BottomMarginY-TopMarginY)/(TextureHeight * (1 - Margin.GetTotalSpaceAlong<Orient_Vertical>() ));
+	float TopTiling = (RightMarginX-LeftMarginX)/(TextureSizeLocalSpace.X * ( 1 - Margin.GetTotalSpaceAlong<Orient_Horizontal>() ));
+	float LeftTiling = (BottomMarginY-TopMarginY)/(TextureSizeLocalSpace.Y * ( 1 - Margin.GetTotalSpaceAlong<Orient_Vertical>() ));
 	
 	FShaderParams ShaderParams = FShaderParams::MakePixelShaderParams( FVector4(LeftMarginU,RightMarginU,TopMarginV,BottomMarginV) );
 
@@ -1142,15 +1326,8 @@ void FSlateElementBatcher::AddBorderElement( const FVector2D& InPosition, const
 	float RightTiling = LeftTiling;
 	float BottomTiling = TopTiling;
 
-	LeftMarginX += Position.X;
-	TopMarginY += Position.Y;
-	RightMarginX += Position.X;
-	BottomMarginY += Position.Y;
-
-	FVector2D EndPos = FVector2D(Position.X, Position.Y) + Size;
-
-	EndPos.X = FMath::TruncToInt(EndPos.X);
-	EndPos.Y = FMath::TruncToInt(EndPos.Y);
+	FVector2D Position = TopLeft;
+	FVector2D EndPos = BotRight;
 
 	// The start index of these vertices in the index buffer
 	uint32 IndexStart = BatchVertices.Num();
@@ -1159,45 +1336,45 @@ void FSlateElementBatcher::AddBorderElement( const FVector2D& InPosition, const
 	FVector2D Zero(0,0);
 
 	// Add all the vertices needed for this element.  Vertices are duplicated so that we can have some sections with no tiling and some with tiling.
-	BatchVertices.Add( FSlateVertex( Position,									FVector2D( StartUV.X, StartUV.Y ),			Zero,							Tint, InClippingRect ) ); //0
-	BatchVertices.Add( FSlateVertex( FVector2D( Position.X, TopMarginY ),		FVector2D( StartUV.X, TopMarginV ),			Zero,							Tint, InClippingRect ) ); //1
-	BatchVertices.Add( FSlateVertex( FVector2D( LeftMarginX, Position.Y ),		FVector2D( LeftMarginU, StartUV.Y ),		Zero,							Tint, InClippingRect ) ); //2
-	BatchVertices.Add( FSlateVertex( FVector2D( LeftMarginX, TopMarginY ),		FVector2D( LeftMarginU, TopMarginV ),		Zero,							Tint, InClippingRect ) ); //3
-
-	BatchVertices.Add( FSlateVertex( FVector2D( LeftMarginX, Position.Y ),		FVector2D( StartUV.X, StartUV.Y ),			FVector2D(TopTiling, 0.0f),		Tint, InClippingRect ) ); //4
-	BatchVertices.Add( FSlateVertex( FVector2D( LeftMarginX, TopMarginY ),		FVector2D( StartUV.X, TopMarginV ),			FVector2D(TopTiling, 0.0f),		Tint, InClippingRect ) ); //5
-	BatchVertices.Add( FSlateVertex( FVector2D( RightMarginX, Position.Y ),		FVector2D( EndUV.X, StartUV.Y ),			FVector2D(TopTiling, 0.0f),		Tint, InClippingRect ) ); //6
-	BatchVertices.Add( FSlateVertex( FVector2D( RightMarginX, TopMarginY ),		FVector2D( EndUV.X, TopMarginV),			FVector2D(TopTiling, 0.0f),		Tint, InClippingRect ) ); //7
-
-	BatchVertices.Add( FSlateVertex( FVector2D( RightMarginX, Position.Y ),		FVector2D( RightMarginU, StartUV.Y ),		Zero,							Tint, InClippingRect ) ); //8
-	BatchVertices.Add( FSlateVertex( FVector2D( RightMarginX, TopMarginY ),		FVector2D( RightMarginU, TopMarginV),		Zero,							Tint, InClippingRect ) ); //9
-	BatchVertices.Add( FSlateVertex( FVector2D( EndPos.X, Position.Y ),			FVector2D( EndUV.X, StartUV.Y ),			Zero,							Tint, InClippingRect ) ); //10
-	BatchVertices.Add( FSlateVertex( FVector2D( EndPos.X, TopMarginY ),			FVector2D( EndUV.X, TopMarginV),			Zero,							Tint, InClippingRect ) ); //11
-
-	BatchVertices.Add( FSlateVertex( FVector2D( Position.X, TopMarginY ),		FVector2D( StartUV.X, StartUV.Y ),			FVector2D(0.0f, LeftTiling),	Tint, InClippingRect ) ); //12
-	BatchVertices.Add( FSlateVertex( FVector2D( Position.X, BottomMarginY ),	FVector2D( StartUV.X, EndUV.Y ),			FVector2D(0.0f, LeftTiling),	Tint, InClippingRect ) ); //13
-	BatchVertices.Add( FSlateVertex( FVector2D( LeftMarginX, TopMarginY ),		FVector2D( LeftMarginU, StartUV.Y ),		FVector2D(0.0f, LeftTiling),	Tint, InClippingRect ) ); //14
-	BatchVertices.Add( FSlateVertex( FVector2D( LeftMarginX, BottomMarginY ),	FVector2D( LeftMarginU, EndUV.Y ),			FVector2D(0.0f, LeftTiling),	Tint, InClippingRect ) ); //15
-
-	BatchVertices.Add( FSlateVertex( FVector2D( RightMarginX, TopMarginY ),		FVector2D( RightMarginU, StartUV.Y ),		FVector2D(0.0f, RightTiling),	Tint, InClippingRect ) ); //16
-	BatchVertices.Add( FSlateVertex( FVector2D( RightMarginX, BottomMarginY ),	FVector2D( RightMarginU, EndUV.Y ),			FVector2D(0.0f, RightTiling),	Tint, InClippingRect ) ); //17
-	BatchVertices.Add( FSlateVertex( FVector2D( EndPos.X, TopMarginY ),			FVector2D( EndUV.X, StartUV.Y ),			FVector2D(0.0f, RightTiling),	Tint, InClippingRect ) ); //18
-	BatchVertices.Add( FSlateVertex( FVector2D( EndPos.X, BottomMarginY ),		FVector2D( EndUV.X, EndUV.Y ),				FVector2D(0.0f, RightTiling),	Tint, InClippingRect ) ); //19
-
-	BatchVertices.Add( FSlateVertex( FVector2D( Position.X, BottomMarginY ),	FVector2D( StartUV.X, BottomMarginV ),		Zero,							Tint, InClippingRect ) ); //20
-	BatchVertices.Add( FSlateVertex( FVector2D( Position.X, EndPos.Y ),			FVector2D( StartUV.X, EndUV.Y ),			Zero,							Tint, InClippingRect ) ); //21
-	BatchVertices.Add( FSlateVertex( FVector2D( LeftMarginX, BottomMarginY ),	FVector2D( LeftMarginU, BottomMarginV ),	Zero,							Tint, InClippingRect ) ); //22
-	BatchVertices.Add( FSlateVertex( FVector2D( LeftMarginX, EndPos.Y ),		FVector2D( LeftMarginU, EndUV.Y ),			Zero,							Tint, InClippingRect ) ); //23
+	BatchVertices.Add( FSlateVertex( RenderTransform, Position,									FVector2D( StartUV.X, StartUV.Y ),			Zero,							Tint, RenderClipRect ) ); //0
+	BatchVertices.Add( FSlateVertex( RenderTransform, FVector2D( Position.X, TopMarginY ),		FVector2D( StartUV.X, TopMarginV ),			Zero,							Tint, RenderClipRect ) ); //1
+	BatchVertices.Add( FSlateVertex( RenderTransform, FVector2D( LeftMarginX, Position.Y ),		FVector2D( LeftMarginU, StartUV.Y ),		Zero,							Tint, RenderClipRect ) ); //2
+	BatchVertices.Add( FSlateVertex( RenderTransform, FVector2D( LeftMarginX, TopMarginY ),		FVector2D( LeftMarginU, TopMarginV ),		Zero,							Tint, RenderClipRect ) ); //3
 	
-	BatchVertices.Add( FSlateVertex( FVector2D( LeftMarginX, BottomMarginY ),	FVector2D( StartUV.X, BottomMarginV ),		FVector2D(BottomTiling, 0.0f),	Tint, InClippingRect ) ); //24
-	BatchVertices.Add( FSlateVertex( FVector2D( LeftMarginX, EndPos.Y ),		FVector2D( StartUV.X, EndUV.Y ),			FVector2D(BottomTiling, 0.0f),	Tint, InClippingRect ) ); //25
-	BatchVertices.Add( FSlateVertex( FVector2D( RightMarginX,BottomMarginY ),	FVector2D( EndUV.X, BottomMarginV ),		FVector2D(BottomTiling, 0.0f),	Tint, InClippingRect ) ); //26
-	BatchVertices.Add( FSlateVertex( FVector2D( RightMarginX, EndPos.Y ),		FVector2D( EndUV.X, EndUV.Y ),				FVector2D(BottomTiling, 0.0f),	Tint, InClippingRect ) ); //27
-
-	BatchVertices.Add( FSlateVertex( FVector2D( RightMarginX, BottomMarginY ),	FVector2D( RightMarginU, BottomMarginV ),	Zero,							Tint, InClippingRect ) ); //29
-	BatchVertices.Add( FSlateVertex( FVector2D( RightMarginX, EndPos.Y ),		FVector2D( RightMarginU, EndUV.Y ),			Zero,							Tint, InClippingRect ) ); //30
-	BatchVertices.Add( FSlateVertex( FVector2D( EndPos.X, BottomMarginY ),		FVector2D( EndUV.X, BottomMarginV ),		Zero,							Tint, InClippingRect ) ); //31
-	BatchVertices.Add( FSlateVertex( FVector2D( EndPos.X, EndPos.Y ),			FVector2D( EndUV.X, EndUV.Y ),				Zero,							Tint, InClippingRect ) ); //32
+	BatchVertices.Add( FSlateVertex( RenderTransform, FVector2D( LeftMarginX, Position.Y ),		FVector2D( StartUV.X, StartUV.Y ),			FVector2D(TopTiling, 0.0f),		Tint, RenderClipRect ) ); //4
+	BatchVertices.Add( FSlateVertex( RenderTransform, FVector2D( LeftMarginX, TopMarginY ),		FVector2D( StartUV.X, TopMarginV ),			FVector2D(TopTiling, 0.0f),		Tint, RenderClipRect ) ); //5
+	BatchVertices.Add( FSlateVertex( RenderTransform, FVector2D( RightMarginX, Position.Y ),	FVector2D( EndUV.X, StartUV.Y ),			FVector2D(TopTiling, 0.0f),		Tint, RenderClipRect ) ); //6
+	BatchVertices.Add( FSlateVertex( RenderTransform, FVector2D( RightMarginX, TopMarginY ),	FVector2D( EndUV.X, TopMarginV),			FVector2D(TopTiling, 0.0f),		Tint, RenderClipRect ) ); //7
+	
+	BatchVertices.Add( FSlateVertex( RenderTransform, FVector2D( RightMarginX, Position.Y ),	FVector2D( RightMarginU, StartUV.Y ),		Zero,							Tint, RenderClipRect ) ); //8
+	BatchVertices.Add( FSlateVertex( RenderTransform, FVector2D( RightMarginX, TopMarginY ),	FVector2D( RightMarginU, TopMarginV),		Zero,							Tint, RenderClipRect ) ); //9
+	BatchVertices.Add( FSlateVertex( RenderTransform, FVector2D( EndPos.X, Position.Y ),		FVector2D( EndUV.X, StartUV.Y ),			Zero,							Tint, RenderClipRect ) ); //10
+	BatchVertices.Add( FSlateVertex( RenderTransform, FVector2D( EndPos.X, TopMarginY ),		FVector2D( EndUV.X, TopMarginV),			Zero,							Tint, RenderClipRect ) ); //11
+	
+	BatchVertices.Add( FSlateVertex( RenderTransform, FVector2D( Position.X, TopMarginY ),		FVector2D( StartUV.X, StartUV.Y ),			FVector2D(0.0f, LeftTiling),	Tint, RenderClipRect ) ); //12
+	BatchVertices.Add( FSlateVertex( RenderTransform, FVector2D( Position.X, BottomMarginY ),	FVector2D( StartUV.X, EndUV.Y ),			FVector2D(0.0f, LeftTiling),	Tint, RenderClipRect ) ); //13
+	BatchVertices.Add( FSlateVertex( RenderTransform, FVector2D( LeftMarginX, TopMarginY ),		FVector2D( LeftMarginU, StartUV.Y ),		FVector2D(0.0f, LeftTiling),	Tint, RenderClipRect ) ); //14
+	BatchVertices.Add( FSlateVertex( RenderTransform, FVector2D( LeftMarginX, BottomMarginY ),	FVector2D( LeftMarginU, EndUV.Y ),			FVector2D(0.0f, LeftTiling),	Tint, RenderClipRect ) ); //15
+	
+	BatchVertices.Add( FSlateVertex( RenderTransform, FVector2D( RightMarginX, TopMarginY ),	FVector2D( RightMarginU, StartUV.Y ),		FVector2D(0.0f, RightTiling),	Tint, RenderClipRect ) ); //16
+	BatchVertices.Add( FSlateVertex( RenderTransform, FVector2D( RightMarginX, BottomMarginY ),	FVector2D( RightMarginU, EndUV.Y ),			FVector2D(0.0f, RightTiling),	Tint, RenderClipRect ) ); //17
+	BatchVertices.Add( FSlateVertex( RenderTransform, FVector2D( EndPos.X, TopMarginY ),		FVector2D( EndUV.X, StartUV.Y ),			FVector2D(0.0f, RightTiling),	Tint, RenderClipRect ) ); //18
+	BatchVertices.Add( FSlateVertex( RenderTransform, FVector2D( EndPos.X, BottomMarginY ),		FVector2D( EndUV.X, EndUV.Y ),				FVector2D(0.0f, RightTiling),	Tint, RenderClipRect ) ); //19
+	
+	BatchVertices.Add( FSlateVertex( RenderTransform, FVector2D( Position.X, BottomMarginY ),	FVector2D( StartUV.X, BottomMarginV ),		Zero,							Tint, RenderClipRect ) ); //20
+	BatchVertices.Add( FSlateVertex( RenderTransform, FVector2D( Position.X, EndPos.Y ),		FVector2D( StartUV.X, EndUV.Y ),			Zero,							Tint, RenderClipRect ) ); //21
+	BatchVertices.Add( FSlateVertex( RenderTransform, FVector2D( LeftMarginX, BottomMarginY ),	FVector2D( LeftMarginU, BottomMarginV ),	Zero,							Tint, RenderClipRect ) ); //22
+	BatchVertices.Add( FSlateVertex( RenderTransform, FVector2D( LeftMarginX, EndPos.Y ),		FVector2D( LeftMarginU, EndUV.Y ),			Zero,							Tint, RenderClipRect ) ); //23
+	
+	BatchVertices.Add( FSlateVertex( RenderTransform, FVector2D( LeftMarginX, BottomMarginY ),	FVector2D( StartUV.X, BottomMarginV ),		FVector2D(BottomTiling, 0.0f),	Tint, RenderClipRect ) ); //24
+	BatchVertices.Add( FSlateVertex( RenderTransform, FVector2D( LeftMarginX, EndPos.Y ),		FVector2D( StartUV.X, EndUV.Y ),			FVector2D(BottomTiling, 0.0f),	Tint, RenderClipRect ) ); //25
+	BatchVertices.Add( FSlateVertex( RenderTransform, FVector2D( RightMarginX,BottomMarginY ),	FVector2D( EndUV.X, BottomMarginV ),		FVector2D(BottomTiling, 0.0f),	Tint, RenderClipRect ) ); //26
+	BatchVertices.Add( FSlateVertex( RenderTransform, FVector2D( RightMarginX, EndPos.Y ),		FVector2D( EndUV.X, EndUV.Y ),				FVector2D(BottomTiling, 0.0f),	Tint, RenderClipRect ) ); //27
+	
+	BatchVertices.Add( FSlateVertex( RenderTransform, FVector2D( RightMarginX, BottomMarginY ),	FVector2D( RightMarginU, BottomMarginV ),	Zero,							Tint, RenderClipRect ) ); //29
+	BatchVertices.Add( FSlateVertex( RenderTransform, FVector2D( RightMarginX, EndPos.Y ),		FVector2D( RightMarginU, EndUV.Y ),			Zero,							Tint, RenderClipRect ) ); //30
+	BatchVertices.Add( FSlateVertex( RenderTransform, FVector2D( EndPos.X, BottomMarginY ),		FVector2D( EndUV.X, BottomMarginV ),		Zero,							Tint, RenderClipRect ) ); //31
+	BatchVertices.Add( FSlateVertex( RenderTransform, FVector2D( EndPos.X, EndPos.Y ),			FVector2D( EndUV.X, EndUV.Y ),				Zero,							Tint, RenderClipRect ) ); //32
 
 
 	// The offset into the index buffer where this elements indices start
@@ -1264,8 +1441,11 @@ void FSlateElementBatcher::AddBorderElement( const FVector2D& InPosition, const
 }
 
 
-void FSlateElementBatcher::AddCustomElement( const FVector2D& Position, const FVector2D& Size, float Scale, const FSlateDataPayload& InPayload, const FSlateRect& InClippingRect, ESlateDrawEffect::Type DrawEffects, uint32 Layer )
+void FSlateElementBatcher::AddCustomElement( const FSlateDrawElement& DrawElement )
 {
+	const FSlateDataPayload& InPayload = DrawElement.GetDataPayload();
+	uint32 Layer = DrawElement.GetLayer();
+
 	if( InPayload.CustomDrawer.IsValid() )
 	{
 		// See if the layer already exists.
diff --git a/Engine/Source/Runtime/SlateCore/Private/Textures/TextureAtlas.cpp b/Engine/Source/Runtime/SlateCore/Private/Textures/TextureAtlas.cpp
--- a/Engine/Source/Runtime/SlateCore/Private/Textures/TextureAtlas.cpp
+++ b/Engine/Source/Runtime/SlateCore/Private/Textures/TextureAtlas.cpp
@@ -36,7 +36,9 @@ const FAtlasedTextureSlot* FSlateTextureAtlas::AddTexture( uint32 TextureWidth,
 	// Find a spot for the character in the texture
 	const FAtlasedTextureSlot* NewSlot = FindSlotForTexture(TextureWidth, TextureHeight);
 
-	if (NewSlot)
+	// handle cases like space, where the size of the glyph is zero. The copy data code doesn't handle zero sized source data with a padding
+	// so make sure to skip this call.
+	if (NewSlot && TextureWidth > 0 && TextureHeight > 0)
 	{
 		CopyDataIntoSlot(NewSlot, Data);
 		MarkTextureDirty();
@@ -71,7 +73,7 @@ void FSlateTextureAtlas::InitAtlasData()
 {
 	check(RootNode == nullptr && AtlasData.Num() == 0);
 
-	RootNode = new FAtlasedTextureSlot(0, 0, AtlasWidth, AtlasHeight, Padding);
+	RootNode = new FAtlasedTextureSlot(0, 0, AtlasWidth, AtlasHeight, PaddingStyle == NoPadding ? 0 : 1);
 	AtlasData.Reserve(AtlasWidth * AtlasHeight * Stride);
 	AtlasData.AddZeroed(AtlasWidth * AtlasHeight * Stride);
 
@@ -79,32 +81,51 @@ void FSlateTextureAtlas::InitAtlasData()
 }
 
 
-void FSlateTextureAtlas::CopyRow( FCopyRowData& CopyRowData )
+void FSlateTextureAtlas::CopyRow( const FCopyRowData& CopyRowData )
 {
 	const uint8* Data = CopyRowData.SrcData;
 	uint8* Start = CopyRowData.DestData;
 	const uint32 SourceWidth = CopyRowData.SrcTextureWidth;
 	const uint32 DestWidth = CopyRowData.DestTextureWidth;
-	const uint32 CopyStride = CopyRowData.DataStride;
-	const uint32 CopyPadding = CopyRowData.Padding;
 	const uint32 SrcRow = CopyRowData.SrcRow;
 	const uint32 DestRow = CopyRowData.DestRow;
+	// this can only be one or zero
+	const uint32 Padding = PaddingStyle != ESlateTextureAtlasPaddingStyle::NoPadding ? 1 : 0;
 
-	const uint8* SourceDataAddr = &Data[SrcRow * SourceWidth * CopyStride]; 
-	uint8* DestDataAddr = &Start[DestRow * DestWidth * CopyStride + CopyPadding * CopyStride]; 
-	FMemory::Memcpy(DestDataAddr, SourceDataAddr, SourceWidth * CopyStride); 
+	const uint8* SourceDataAddr = &Data[(SrcRow * SourceWidth) * Stride]; 
+	uint8* DestDataAddr = &Start[(DestRow * DestWidth + Padding) * Stride]; 
+	FMemory::Memcpy(DestDataAddr, SourceDataAddr, SourceWidth * Stride); 
 
-	if (CopyRowData.bDialateToPadding)
+	if (Padding > 0)
 	{ 
-		const uint8* FirstPixel = SourceDataAddr; 
-		const uint8* LastPixel = SourceDataAddr + ((SourceWidth - 1) * CopyStride); 
-		uint8* DestDataAddrNoPadding = &Start[DestRow * DestWidth * CopyStride]; 
-		FMemory::Memcpy(DestDataAddrNoPadding, FirstPixel, CopyStride ); 
-		DestDataAddrNoPadding = &Start[DestRow * DestWidth * CopyStride + (CopyRowData.RowWidth - 1) * CopyStride]; 
-		FMemory::Memcpy(DestDataAddrNoPadding, LastPixel, CopyStride); 
+		uint8* DestPaddingPixelLeft = &Start[(DestRow * DestWidth) * Stride];
+		uint8* DestPaddingPixelRight = DestPaddingPixelLeft + ((CopyRowData.RowWidth - 1) * Stride);
+		if (PaddingStyle == ESlateTextureAtlasPaddingStyle::DilateBorder)
+		{
+			const uint8* FirstPixel = SourceDataAddr; 
+			const uint8* LastPixel = SourceDataAddr + ((SourceWidth - 1) * Stride); 
+			FMemory::Memcpy(DestPaddingPixelLeft, FirstPixel, Stride);
+			FMemory::Memcpy(DestPaddingPixelRight, LastPixel, Stride);
+		}
+		else
+		{
+			FMemory::Memzero(DestPaddingPixelLeft, Stride);
+			FMemory::Memzero(DestPaddingPixelRight, Stride);
+		}
+
 	} 
 }
 
+void FSlateTextureAtlas::ZeroRow( const FCopyRowData& CopyRowData )
+{
+	const uint32 SourceWidth = CopyRowData.SrcTextureWidth;
+	const uint32 DestWidth = CopyRowData.DestTextureWidth;
+	const uint32 DestRow = CopyRowData.DestRow;
+
+	uint8* DestDataAddr = &CopyRowData.DestData[DestRow * DestWidth * Stride]; 
+	FMemory::Memzero(DestDataAddr, CopyRowData.RowWidth * Stride); 
+}
+
 
 void FSlateTextureAtlas::CopyDataIntoSlot( const FAtlasedTextureSlot* SlotToCopyTo, const TArray<uint8>& Data )
 {
@@ -112,29 +133,35 @@ void FSlateTextureAtlas::CopyDataIntoSlot( const FAtlasedTextureSlot* SlotToCopy
 	uint8* Start = &AtlasData[SlotToCopyTo->Y*AtlasWidth*Stride + SlotToCopyTo->X*Stride];
 	
 	// Account for same padding on each sides
+	const uint32 Padding = PaddingStyle != ESlateTextureAtlasPaddingStyle::NoPadding ? 1 : 0;
 	const uint32 AllPadding = Padding*2;
 	// The width of the source texture without padding (actual width)
 	const uint32 SourceWidth = SlotToCopyTo->Width-AllPadding; 
+	const uint32 SourceHeight = SlotToCopyTo->Width-AllPadding; 
 
 	FCopyRowData CopyRowData;
-	CopyRowData.bDialateToPadding = Padding != 0;
-	CopyRowData.DataStride = Stride;
 	CopyRowData.DestData = Start;
 	CopyRowData.SrcData = Data.GetData();
 	CopyRowData.DestTextureWidth = AtlasWidth;
 	CopyRowData.SrcTextureWidth = SourceWidth;
-	CopyRowData.Padding = Padding;
 	CopyRowData.RowWidth = SlotToCopyTo->Width;
 
-	// Copy color data into padding for bilinear filtering. 
+	// Apply the padding for bilinear filtering. 
 	// Not used if no padding (assumes sampling outside boundaries of the sub texture is not possible)
-	if (Padding)
+	if (Padding > 0)
 	{
 		// Copy first color row into padding.  
 		CopyRowData.SrcRow = 0;
 		CopyRowData.DestRow = 0;
 
-		CopyRow(CopyRowData);	
+		if (PaddingStyle == ESlateTextureAtlasPaddingStyle::DilateBorder)
+		{
+			CopyRow(CopyRowData);
+		}
+		else
+		{
+			ZeroRow(CopyRowData);
+		}
 	}
 
 	// Copy each row of the texture
@@ -146,13 +173,20 @@ void FSlateTextureAtlas::CopyDataIntoSlot( const FAtlasedTextureSlot* SlotToCopy
 		CopyRow(CopyRowData);
 	}
 
-	if (Padding)
+	if (Padding > 0)
 	{
 		// Copy last color row into padding row for bilinear filtering
-		CopyRowData.SrcRow = SlotToCopyTo->Height - AllPadding - 1;
+		CopyRowData.SrcRow = SourceHeight - 1;
 		CopyRowData.DestRow = SlotToCopyTo->Height - Padding;
 
-		CopyRow(CopyRowData);
+		if (PaddingStyle == ESlateTextureAtlasPaddingStyle::DilateBorder)
+		{
+			CopyRow(CopyRowData);
+		}
+		else
+		{
+			ZeroRow(CopyRowData);
+		}
 	}
 }
 
@@ -188,6 +222,7 @@ const FAtlasedTextureSlot* FSlateTextureAtlas::FindSlotForTexture( FAtlasedTextu
 	}
 
 	// Account for padding on both sides
+	const uint32 Padding = PaddingStyle != ESlateTextureAtlasPaddingStyle::NoPadding ? 1 : 0;
 	uint32 TotalPadding = Padding * 2;
 
 	// This slot can't fit the character
diff --git a/Engine/Source/Runtime/SlateCore/Private/Widgets/SNullWidget.cpp b/Engine/Source/Runtime/SlateCore/Private/Widgets/SNullWidget.cpp
--- a/Engine/Source/Runtime/SlateCore/Private/Widgets/SNullWidget.cpp
+++ b/Engine/Source/Runtime/SlateCore/Private/Widgets/SNullWidget.cpp
@@ -24,6 +24,7 @@ public:
 			InArgs._Cursor, 
 			InArgs._IsEnabled,
 			InArgs._Visibility,
+			InArgs._RenderTransform,
 			InArgs._Tag,
 			InArgs.MetaData
 		);
diff --git a/Engine/Source/Runtime/SlateCore/Private/Widgets/SWidget.cpp b/Engine/Source/Runtime/SlateCore/Private/Widgets/SWidget.cpp
--- a/Engine/Source/Runtime/SlateCore/Private/Widgets/SWidget.cpp
+++ b/Engine/Source/Runtime/SlateCore/Private/Widgets/SWidget.cpp
@@ -17,6 +17,7 @@ SWidget::SWidget()
 	, Cursor( TOptional<EMouseCursor::Type>() )
 	, EnabledState( true )
 	, Visibility( EVisibility::Visible )
+	, RenderTransform()
 	, DesiredSize(FVector2D::ZeroVector)
 	, ToolTip()
 	, bToolTipForceFieldEnabled( false )
@@ -31,6 +32,7 @@ void SWidget::Construct(
 	const TAttribute< TOptional<EMouseCursor::Type> > & InCursor ,
 	const TAttribute<bool> & InEnabledState ,
 	const TAttribute<EVisibility> & InVisibility,
+	const TAttribute<FSlateRenderTransform>& InTransform,
 	const FName& InTag,
 	const TArray<TSharedRef<ISlateMetaData>>& InMetaData
 )
@@ -54,6 +56,7 @@ void SWidget::Construct(
 	Cursor = InCursor;
 	EnabledState = InEnabledState;
 	Visibility = InVisibility;
+	RenderTransform = InTransform;
 	Tag = InTag;
 	MetaData = InMetaData;
 }
