Change 2122846 on 2014/07/01 00:51:24 by [REDACTED_PERSON][REDACTED_HANDLE]

diff --git a/Engine/Source/Runtime/SlateCore/Private/Widgets/SWindow.cpp b/Engine/Source/Runtime/SlateCore/Private/Widgets/SWindow.cpp
--- a/Engine/Source/Runtime/SlateCore/Private/Widgets/SWindow.cpp
+++ b/Engine/Source/Runtime/SlateCore/Private/Widgets/SWindow.cpp
@@ -88,15 +88,31 @@ public:
 
 private:
 
+	/**
+	 * Each child slot essentially tries to place their contents at a specified position on the screen
+	 * and scale as the widget initiating the popup, both of which are stored in the slot attributes.
+	 * The tricky part is that the scale we are given is the fully accumulated layout scale of the widget, which already incorporates 
+	 * the DPI Scale of the window. The DPI Scale is also applied to the overlay since it is part of the window,
+	 * so this scale needs to be factored out when determining the scale of the child geometry that will be created to hold the popup.
+	 * We also optionally adjust the window position to keep it within the client bounds of the top-level window. This must be done in screenspace.
+	 * This means some hairy transformation calculus goes on to ensure the computations are done in the proper space so scale is respected.
+	 * 
+	 * There are 3 transformational spaces involved, each clearly specified in the variable names:
+	 * Screen      - Basically desktop space. Contains desktop offset and DPI scale.
+	 * WindowLocal - local space of the SWindow containing this popup. Screenspace == Concat(WindowLocal, DPI Scale, Desktop Offset)
+	 * ChildLocal  - space of the child widget we want to display in the popup. The widget's LayoutTransform takes us from ChildLocal to WindowLocal space.
+	 */
 	virtual void OnArrangeChildren( const FGeometry& AllottedGeometry, FArrangedChildren& ArrangedChildren ) const override
 	{
-		const FVector2D WindowDesktopPosition = (ensure(OwnerWindow.IsValid()))
-			? OwnerWindow.Pin()->GetPositionInScreen()
-			: FVector2D::ZeroVector;
+		// skip all this work if there are no children to arrange.
+		if (Children.Num() == 0) return;
 
-		const FVector2D WindowSize = (ensure(OwnerWindow.IsValid()))
-			? OwnerWindow.Pin()->GetClientSizeInScreen()
-			: FVector2D::ZeroVector;
+		// create a transform from screen to local space.
+		// This assumes that the PopupLayer is part of an Overlay that takes up the entire window space.
+		// We should technically be using the AllottedGeometry to transform from AbsoluteToLocal space just in case it has an additional scale on it.
+		// But we can't because the absolute space of the geometry is sometimes given in desktop space (picking, ticking) 
+		// and sometimes in window space (painting), and we can't necessarily tell by inspection so we have to just make an assumption here.
+		FSlateLayoutTransform ScreenToWindowLocal = (ensure(OwnerWindow.IsValid())) ? Inverse(OwnerWindow.Pin()->GetLocalToScreenTransform()) : FSlateLayoutTransform();
 		
 		for ( int32 ChildIndex = 0; ChildIndex < Children.Num(); ++ChildIndex )
 		{
@@ -104,33 +120,43 @@ private:
 			const EVisibility ChildVisibility = CurChild.GetWidget()->GetVisibility();
 			if ( ArrangedChildren.Accepts(ChildVisibility) )
 			{
-				const FVector2D WidgetDesiredSize = CurChild.GetWidget()->GetDesiredSize();
-				const float ChildScale = CurChild.Scale_Attribute.Get();
-				const bool bClampToWindow = CurChild.Clamp_Attribute.Get();
-				const FVector2D ChildSize = WidgetDesiredSize * ChildScale;
-				FVector2D ChildLocalPosition = CurChild.DesktopPosition_Attribute.Get() - WindowDesktopPosition;
+				// This scale+translate forms the ChildLocal to Screenspace transform.
+				// The translation may be adjusted based on clamping, but the scale is accurate, 
+				// so we can transform vectors into screenspace using the scale alone.
+				const float ChildLocalToScreenScale = CurChild.Scale_Attribute.Get();
+				FVector2D ChildLocalToScreenOffset = CurChild.DesktopPosition_Attribute.Get();
+				// The size of the child is either the desired size of the widget (computed in the child's local space) or the size override (specified in screen space)
+				const FVector2D ChildSizeChildLocal = CurChild.GetWidget()->GetDesiredSize();
+				// Convert the desired size to screen space. Here is were we convert a vector to screenspace
+				// before we have the final position in screenspace (which would be needed to transform a point).
+				FVector2D ChildSizeScreenspace = TransformVector(ChildLocalToScreenScale, ChildSizeChildLocal);
+				// But then allow each size dimension to be overridden by the slot, which specifies the overrides in screen space.
+				ChildSizeScreenspace = FVector2D(
+						CurChild.WidthOverride_Attribute.IsSet() ? CurChild.WidthOverride_Attribute.Get() : ChildSizeScreenspace.X,
+						CurChild.HeightOverride_Attribute.IsSet() ? CurChild.HeightOverride_Attribute.Get() : ChildSizeScreenspace.Y);
 				
-				if(bClampToWindow)
+				// If clamping, move the screen space position to ensure the screen space size stays within the client rect of the top level window.
+				if(CurChild.Clamp_Attribute.Get())
 				{
-					FVector2D ClampBufferReducedWindowSize = WindowSize - CurChild.ClampBuffer_Attribute.Get();
-					ChildLocalPosition.X = FMath::Clamp(ChildLocalPosition.X - FMath::Max(0.0f, (ChildLocalPosition.X  + WidgetDesiredSize.X ) - ClampBufferReducedWindowSize.X), 0.0f, MAX_flt);
-					ChildLocalPosition.Y = FMath::Clamp(ChildLocalPosition.Y - FMath::Max(0.0f, (ChildLocalPosition.Y  + WidgetDesiredSize.Y ) - ClampBufferReducedWindowSize.Y), 0.0f, MAX_flt);
+					const FSlateRect WindowClientRectScreenspace = (ensure(OwnerWindow.IsValid())) ? OwnerWindow.Pin()->GetClientRectInScreen() : FSlateRect();
+					const FVector2D ClampBufferScreenspace = CurChild.ClampBuffer_Attribute.Get();
+					const FSlateRect ClampedWindowClientRectScreenspace = WindowClientRectScreenspace.InsetBy(FMargin(ClampBufferScreenspace.X, ClampBufferScreenspace.Y));
+					// Find how much our child wants to extend beyond our client space and subtract that amount, but don't push it past the client edge.
+					ChildLocalToScreenOffset.X = FMath::Max(WindowClientRectScreenspace.Left, ChildLocalToScreenOffset.X - FMath::Max(0.0f, (ChildLocalToScreenOffset.X  + ChildSizeScreenspace.X ) - ClampedWindowClientRectScreenspace.Right));
+					ChildLocalToScreenOffset.Y = FMath::Max(WindowClientRectScreenspace.Top, ChildLocalToScreenOffset.Y - FMath::Max(0.0f, (ChildLocalToScreenOffset.Y  + ChildSizeScreenspace.Y ) - ClampedWindowClientRectScreenspace.Bottom));
 				}
 
+				// We now have the final position, so construct the transform from ChildLocal to Screenspace
+				const FSlateLayoutTransform ChildLocalToScreen(ChildLocalToScreenScale, ChildLocalToScreenOffset);
+				// Using this we can compute the transform from ChildLocal to WindowLocal, which is effectively the LayoutTransform of the child widget.
+				const FSlateLayoutTransform ChildLocalToWindowLocal = Concatenate(ChildLocalToScreen, ScreenToWindowLocal);
+				// The ChildSize needs to be given in ChildLocal space when constructing a geometry.
+				const FVector2D ChildSizeLocalspace = TransformVector(Inverse(ChildLocalToScreen), ChildSizeScreenspace);
+
 				// The position is explicitly in desktop pixels.
 				// The size and DPI scale come from the widget that is using
 				// this overlay to "punch" through the UI.
-				ArrangedChildren.AddWidget( ChildVisibility, FArrangedWidget( CurChild.GetWidget(),
-					FGeometry(
-						ChildLocalPosition / ChildScale,
-						AllottedGeometry.AbsolutePosition,
-						FVector2D(
-							// Notice that override sizes get divided by ChildScale because any explicit size overrides are in pixel units and are not affected by DPI scaling.
-							CurChild.WidthOverride_Attribute.IsSet() ? CurChild.WidthOverride_Attribute.Get() / ChildScale : WidgetDesiredSize.X,
-							CurChild.HeightOverride_Attribute.IsSet() ? CurChild.HeightOverride_Attribute.Get() / ChildScale : WidgetDesiredSize.Y
-						),
-						ChildScale
-				) ) );
+				ArrangedChildren.AddWidget(AllottedGeometry.MakeChild(CurChild.GetWidget(), ChildSizeLocalspace, ChildLocalToWindowLocal));
 			}
 		}
 	}
@@ -562,15 +588,34 @@ FVector2D SWindow::GetInitialDesiredPositionInScreen() const
 FGeometry SWindow::GetWindowGeometryInScreen() const
 {
 	const float AppScale = FSlateApplicationBase::Get().GetApplicationScale();
-	return FGeometry( ScreenPosition/AppScale, FVector2D::ZeroVector, Size/AppScale, AppScale );
+	// We are scaling children for layout, but our pixel bounds are not changing. 
+	// FGeometry expects Size in Local space, but our size is stored in screen space.
+	// So we need to transform Size into the window's local space for FGeometry.
+	FSlateLayoutTransform LocalToScreen = GetLocalToScreenTransform();
+	return FGeometry::MakeRoot( TransformVector(Inverse(LocalToScreen), Size), LocalToScreen );
 }
 
 FGeometry SWindow::GetWindowGeometryInWindow() const
 {
 	const float AppScale = FSlateApplicationBase::Get().GetApplicationScale();
-	return FGeometry( FVector2D::ZeroVector, FVector2D::ZeroVector, Size/AppScale, AppScale );
+	// We are scaling children for layout, but our pixel bounds are not changing. 
+	// FGeometry expects Size in Local space, but our size is stored in screen space (same as window space + screen offset).
+	// So we need to transform Size into the window's local space for FGeometry.
+	FSlateLayoutTransform LocalToWindow = GetLocalToWindowTransform();
+	return FGeometry::MakeRoot( TransformVector(Inverse(LocalToWindow), Size), LocalToWindow );
+}
+
+FSlateLayoutTransform SWindow::GetLocalToScreenTransform() const
+{
+	return FSlateLayoutTransform(FSlateApplicationBase::Get().GetApplicationScale(), ScreenPosition);
+}
+
+FSlateLayoutTransform SWindow::GetLocalToWindowTransform() const
+{
+	return FSlateLayoutTransform(FSlateApplicationBase::Get().GetApplicationScale());
 }
 
+
 FVector2D SWindow::GetPositionInScreen() const
 {
 	return ScreenPosition;
@@ -600,23 +645,24 @@ FSlateRect SWindow::GetNonMaximizedRectInScreen() const
 
 FSlateRect SWindow::GetRectInScreen() const
 { 
-	return FSlateRect( ScreenPosition.X, ScreenPosition.Y, ScreenPosition.X + Size.X, ScreenPosition.Y + Size.Y );
+	return FSlateRect( ScreenPosition, ScreenPosition + Size );
 }
 
-FVector2D SWindow::GetClientSizeInScreen() const
+FSlateRect SWindow::GetClientRectInScreen() const
 {
 	if (HasOSWindowBorder())
 	{
-		return Size;
+		return GetRectInScreen();
 	}
 
-	FVector2D ClientSize = Size;
-	FMargin BorderSize = GetWindowBorderSize();
-
-	ClientSize.X -= (BorderSize.Left + BorderSize.Right);
-	ClientSize.Y -= (BorderSize.Top + BorderSize.Bottom + TitleBarSize);
+	return GetRectInScreen()
+		.InsetBy(GetWindowBorderSize())
+		.InsetBy(FMargin(0.0f, TitleBarSize, 0.0f, 0.0f));
+}
 
-	return ClientSize;
+FVector2D SWindow::GetClientSizeInScreen() const
+{
+	return GetClientRectInScreen().GetSize();
 }
 
 FSlateRect SWindow::GetClippingRectangleInWindow() const
diff --git a/Engine/Source/Runtime/SlateCore/Public/Layout/ArrangedWidget.h b/Engine/Source/Runtime/SlateCore/Public/Layout/ArrangedWidget.h
--- a/Engine/Source/Runtime/SlateCore/Public/Layout/ArrangedWidget.h
+++ b/Engine/Source/Runtime/SlateCore/Public/Layout/ArrangedWidget.h
@@ -2,6 +2,7 @@
 
 #pragma once
 
+#include "Geometry.h"
 
 /**
  * A pair: Widget and its Geometry. Widgets populate an list of WidgetGeometries
diff --git a/Engine/Source/Runtime/SlateCore/Public/Layout/Geometry.h b/Engine/Source/Runtime/SlateCore/Public/Layout/Geometry.h
--- a/Engine/Source/Runtime/SlateCore/Public/Layout/Geometry.h
+++ b/Engine/Source/Runtime/SlateCore/Public/Layout/Geometry.h
@@ -2,12 +2,14 @@
 
 #pragma once
 
+#include "SlateLayoutTransform.h"
+#include "SlateRenderTransform.h"
+#include "PaintGeometry.h"
 #include "Geometry.generated.h"
 
-
 class FArrangedWidget;
 class SWidget;
-
+class FLayoutGeometry;
 
 /**
  * Represents the position, size, and absolute position of a Widget in Slate.
@@ -25,15 +27,19 @@ struct SLATECORE_API FGeometry
 
 public:
 
-	/** Default constructor. */
-	FGeometry( )
-		: Position(0.0f, 0.0f)
-		, AbsolutePosition(0.0f, 0.0f)
-		, Size(0.0f, 0.0f)
-		, Scale(1.0f)
-	{ }
+	/**
+	 * Default constructor. Creates a geometry with identity transforms.
+	 */
+	FGeometry();
+
+	/**
+	 * !!! HACK!!! We're keeping members of FGeometry const to prevent mutability without making them private, for backward compatibility.
+	 * But this means the assignment operator no longer works. We implement one ourselves now and force a memcpy.
+	 */
+	FGeometry& operator=(const FGeometry& RHS);
 
 	/**
+	 * !!! DEPRECATED FUNCTION !!! Use MakeChild taking a layout transform instead!
 	 * Construct a new geometry given the following parameters:
 	 * 
 	 * @param OffsetFromParent         Local position of this geometry within its parent geometry.
@@ -41,12 +47,38 @@ public:
 	 * @param InSize                   The size of this geometry.
 	 * @param InScale                  The scale of this geometry with respect to Normal Slate Coordinates.
 	 */
-	FGeometry( const FVector2D& OffsetFromParent, const FVector2D& ParentAbsolutePosition, const FVector2D& InSize, float InScale )
-		: Position(OffsetFromParent)
-		, AbsolutePosition(ParentAbsolutePosition + (OffsetFromParent * InScale))
-		, Size(InSize)	
-		, Scale(InScale)
-	{ }
+	FGeometry( const FVector2D& OffsetFromParent, const FVector2D& ParentAbsolutePosition, const FVector2D& InLocalSize, float InScale );
+
+private:
+	/**
+	 * Construct a new geometry with a given size in LocalSpace that is attached to a parent geometry with the given layout and render transform. 
+	 * 
+	 * @param InLocalSize						The size of the geometry in Local Space.
+	 * @param InLocalLayoutTransform			A layout transform from local space to the parent geoemtry's local space.
+	 * @param InLocalRenderTransform			A render-only transform in local space that will be prepended to the LocalLayoutTransform when rendering.
+	 * @param ParentAccumulatedLayoutTransform	The accumulated layout transform of the parent widget. AccumulatedLayoutTransform = Concat(LocalLayoutTransform, ParentAccumulatedLayoutTransform).
+	 * @param ParentAccumulatedRenderTransform	The accumulated render transform of the parent widget. AccumulatedRenderTransform = Concat(LocalRenderTransform, LocalLayoutTransform, ParentAccumulatedRenderTransform).
+	 */
+	FGeometry( 
+		const FVector2D& InLocalSize, 
+		const FSlateLayoutTransform& InLocalLayoutTransform, 
+		const FSlateRenderTransform& InLocalRenderTransform, 
+		const FSlateLayoutTransform& ParentAccumulatedLayoutTransform, 
+		const FSlateRenderTransform& ParentAccumulatedRenderTransform);
+
+	/**
+	 * Construct a new geometry with a given size in LocalSpace (and identity render transform) that is attached to a parent geometry with the given layout and render transform. 
+	 * 
+	 * @param InLocalSize						The size of the geometry in Local Space.
+	 * @param InLocalLayoutTransform			A layout transform from local space to the parent geoemtry's local space.
+	 * @param ParentAccumulatedLayoutTransform	The accumulated layout transform of the parent widget. AccumulatedLayoutTransform = Concat(LocalLayoutTransform, ParentAccumulatedLayoutTransform).
+	 * @param ParentAccumulatedRenderTransform	The accumulated render transform of the parent widget. AccumulatedRenderTransform = Concat(LocalRenderTransform, LocalLayoutTransform, ParentAccumulatedRenderTransform).
+	 */
+	FGeometry( 
+		const FVector2D& InLocalSize, 
+		const FSlateLayoutTransform& InLocalLayoutTransform, 
+		const FSlateLayoutTransform& ParentAccumulatedLayoutTransform, 
+		const FSlateRenderTransform& ParentAccumulatedRenderTransform);
 
 public:
 
@@ -59,9 +91,9 @@ public:
 	bool operator==( const FGeometry& Other ) const
 	{
 		return
-			this->Position == Other.Position &&
+			this->Size == Other.Size && 
 			this->AbsolutePosition == Other.AbsolutePosition &&
-			this->Size == Other.Size &&
+			this->Position == Other.Position &&
 			this->Scale == Other.Scale;
 	}
 	
@@ -77,143 +109,276 @@ public:
 	}
 
 public:
-	
 	/**
-	 * Create a child geometry; i.e. a geometry relative to this geometry.
+	 * Makes a new geometry that is essentially the root of a hierarchy (has no parent transforms to inherit).
+	 * For a root Widget, the LayoutTransform is often the window DPI scale + window offset.
+	 * 
+	 * @param LocalSize			Size of the geoemtry in Local Space.
+	 * @param LayoutTransform	Layout transform of the geometry.
+	 * @return					The new root geometry
+	 */
+	static FGeometry MakeRoot( const FVector2D& LocalSize, const FSlateLayoutTransform& LayoutTransform );
+
+	/**
+	 * Create a child geometry relative to this one with a given local space size, layout transform, and render transform.
+	 * For example, a widget with a 5x5 margin will create a geometry for it's child contents having a LayoutTransform of Translate(5,5) and a LocalSize 10 units smaller 
+	 * than it's own.
 	 *
-	 * @param ChildOffset  The offset of the child from this Geometry's position.
-	 * @param ChildSize    The size of the child geometry in screen units.
-	 * @param InScale      How much to scale this child with respect to its parent.
-	 * @return A Geometry that is offset from this geometry by ChildOffset
+	 * @param LocalSize			The size of the child geometry in local space.
+	 * @param LayoutTransform	Layout transform of the new child relative to this Geometry. Goes from the child's layout space to the this widget's layout space.
+	 * @param RenderTransform	Render-only transform of the new child that is applied before the layout transform for rendering purposes only.
+	 *
+	 * @return					The new child geometry.
 	 */
-	FGeometry MakeChild( const FVector2D& ChildOffset, const FVector2D& ChildSize, float InScale = 1.0f ) const
-	{
-		return FGeometry( ChildOffset, this->AbsolutePosition, ChildSize, this->Scale * InScale );
-	}
-	
+	FGeometry MakeChild( const FVector2D& LocalSize, const FSlateLayoutTransform& LayoutTransform, const FSlateRenderTransform& RenderTransform ) const;
+
 	/**
-	 * Create a child widget geometry; i.e. a geometry relative to this geometry that
-	 * is also associated with a Widget.
+	 * Create a child geometry relative to this one with a given local space size, layout transform, and identity render transform.
+	 * For example, a widget with a 5x5 margin will create a geometry for it's child contents having a LayoutTransform of Translate(5,5) and a LocalSize 10 units smaller 
+	 * than it's own.
+	 *
+	 * @param LocalSize			The size of the child geometry in local space.
+	 * @param LayoutTransform	Layout transform of the new child relative to this Geometry. Goes from the child's layout space to the this widget's layout space.
 	 *
-	 * @param InWidget     The widget to associate with the newly created child geometry
-	 * @param ChildOffset  The offset of the child from this Geometry's position.
-	 * @param ChildSize    The size of the child geometry in screen units.
-	 * @param InScale      How much to scale this child with respect to its parent.
-	 * @return A WidgetGeometry that is offset from this geometry by ChildOffset
+	 * @return					The new child geometry.
 	 */
-	FArrangedWidget MakeChild( const TSharedRef<SWidget>& InWidget, const FVector2D& ChildOffset, const FVector2D& ChildSize, float InScale = 1.0f ) const;
+	FGeometry MakeChild( const FVector2D& LocalSize, const FSlateLayoutTransform& LayoutTransform ) const;
 
 	/**
-	 * Make an FPaintGeometry in this FGeometry at the specified Offset.
+	 * Create a child geometry+widget relative to this one using the given LayoutGeometry.
 	 *
-	 * @param InOffset   Offset in SlateUnits from this FGeometry's origin.
-	 * @param InSize     The size of this paint element in SlateUnits.
-	 * @param InScale    Additional scaling to apply to the DrawElements.
-	 * @return a FPaintGeometry derived this FGeometry.
+	 * @param ChildWidget		The child widget this geometry is being created for.
+	 * @param LayoutGeometry	Layout geometry of the child.
+	 *
+	 * @return					The new child geometry.
 	 */
-	FPaintGeometry ToPaintGeometry( FVector2D InOffset, FVector2D InSize, float InScale ) const;
+	FArrangedWidget MakeChild( const TSharedRef<SWidget>& ChildWidget, const FLayoutGeometry& LayoutGeometry ) const;
 
-	/** A PaintGeometry that is not scaled additionally with respect to its originating FGeometry */
-	FPaintGeometry ToPaintGeometry( FVector2D InOffset, FVector2D InSize ) const
-	{
-		return ToPaintGeometry( InOffset, InSize, 1.0f );
-	}
+	/**
+	 * Create a child geometry+widget relative to this one with a given local space size and layout transform.
+	 * The Geometry inherits the child widget's render transform.
+	 * For example, a widget with a 5x5 margin will create a geometry for it's child contents having a LayoutTransform of Translate(5,5) and a LocalSize 10 units smaller 
+	 * than it's own.
+	 *
+	 * @param ChildWidget		The child widget this geometry is being created for.
+	 * @param LocalSize			The size of the child geometry in local space.
+	 * @param LayoutTransform	Layout transform of the new child relative to this Geometry. 
+	 *
+	 * @return					The new child geometry+widget.
+	 */
+	FArrangedWidget MakeChild( const TSharedRef<SWidget>& ChildWidget, const FVector2D& LocalSize, const FSlateLayoutTransform& LayoutTransform ) const;
 
-	/** Convert the FGeometry to an FPaintGeometry with no changes */
-	FPaintGeometry ToPaintGeometry() const
-	{
-		return ToPaintGeometry( FVector2D(0,0), this->Size, 1.0f );
-	}
+	/**
+	 * !!! DEPRECATED FUNCTION !!! Use MakeChild taking a layout transform instead!
+	 * Create a child geometry relative to this one with a given local space size and layout transform given by a scale+offset. The Render Transform is identity.
+	 *
+	 * @param ChildOffset	Offset of the child relative to the parent. Scale+Offset effectively define the layout transform.
+	 * @param LocalSize		The size of the child geometry in local space.
+	 * @param ChildScale	Scale of the child relative to the parent. Scale+Offset effectively define the layout transform.
+	 *
+	 * @return				The new child geometry.
+	 */
+	FGeometry MakeChild( const FVector2D& ChildOffset, const FVector2D& LocalSize, float ChildScale = 1.0f ) const;
 
-	/** An FPaintGeometry that is offset into the FGeometry from which it originates and occupies all the available room */
-	FPaintGeometry ToOffsetPaintGeometry( FVector2D InOffset ) const
-	{
-		return ToPaintGeometry( InOffset, this->Size, 1.0f );
-	}
+	/**
+	 * !!! DEPRECATED FUNCTION !!! Use MakeChild taking a layout transform instead!
+	 * Create a child geometry+widget relative to this one with a given local space size and layout transform given by a scale+offset.
+	 * The Geometry inherits the child widget's render transform.
+	 *
+	 * @param ChildWidget	The child widget this geometry is being created for.
+	 * @param ChildOffset	Offset of the child relative to the parent. Scale+Offset effectively define the layout transform.
+	 * @param LocalSize		The size of the child geometry in local space.
+	 * @param ChildScale	Scale of the child relative to the parent. Scale+Offset effectively define the layout transform.
+	 *
+	 * @return				The new child geometry+widget.
+	 */
+	FArrangedWidget MakeChild( const TSharedRef<SWidget>& ChildWidget, const FVector2D& ChildOffset, const FVector2D& LocalSize, float ChildScale = 1.0f) const;
 
 	/**
-	 * Make an FPaintGeometry from this FGeometry by inflating the FGeometry by InflateAmount.
+	 * Create a paint geometry that represents this geometry.
+	 * 
+	 * @return	The new paint geometry.
 	 */
-	FPaintGeometry ToInflatedPaintGeometry( const FVector2D& InflateAmount ) const;
+	FPaintGeometry ToPaintGeometry() const;
 
-	FPaintGeometry CenteredPaintGeometryOnLeft( const FVector2D& SizeBeingAligned, float InScale ) const;
+	/**
+	 * Create a paint geometry relative to this one with a given local space size and layout transform.
+	 * The paint geometry inherits the widget's render transform.
+	 *
+	 * @param LocalSize			The size of the child geometry in local space.
+	 * @param LayoutTransform	Layout transform of the paint geometry relative to this Geometry. 
+	 *
+	 * @return					The new paint geometry derived from this one.
+	 */
+	FPaintGeometry ToPaintGeometry( const FVector2D& LocalSize, const FSlateLayoutTransform& LayoutTransform) const;
 
-	FPaintGeometry CenteredPaintGeometryOnRight( const FVector2D& SizeBeingAligned, float InScale ) const;
+	/**
+	 * Create a paint geometry with the same size as this geometry with a given layout transform.
+	 * The paint geometry inherits the widget's render transform.
+	 *
+	 * @param LayoutTransform	Layout transform of the paint geometry relative to this Geometry. 
+	 * 
+	 * @return					The new paint geometry derived from this one.
+	 */
+	FPaintGeometry ToPaintGeometry( const FSlateLayoutTransform& LayoutTransform) const;
 
-	FPaintGeometry CenteredPaintGeometryBelow( const FVector2D& SizeBeingAligned, float InScale ) const;
+	/**
+	 * !!! DEPRECATED FUNCTION !!! Use ToPaintGeometry taking a layout transform instead!
+	 * Create a paint geometry relative to this one with a given local space size and layout transform given as a scale+offset.
+	 * The paint geometry inherits the widget's render transform.
+	 *
+	 * @param LocalOffset	Offset of the paint geometry relative to the parent. Scale+Offset effectively define the layout transform.
+	 * @param LocalSize		The size of the paint geometry in local space.
+	 * @param LocalScale	Scale of the paint geometry relative to the parent. Scale+Offset effectively define the layout transform.
+	 * 
+	 * @return				The new paint geometry derived from this one.
+	 */
+	FPaintGeometry ToPaintGeometry( const FVector2D& LocalOffset, const FVector2D& LocalSize, float LocalScale = 1.0f ) const;
 
+	/**
+	 * !!! DEPRECATED FUNCTION !!! Use ToPaintGeometry taking a layout transform instead!
+	 * Create a paint geometry relative to this one with the same size size and a given local space offset.
+	 * The paint geometry inherits the widget's render transform.
+	 *
+	 * @param LocalOffset	Offset of the paint geometry relative to the parent. Effectively defines the layout transform.
+	 * 
+	 * @return				The new paint geometry derived from this one.
+	 */
+	FPaintGeometry ToOffsetPaintGeometry( const FVector2D& LocalOffset ) const;
 
-	/** @return true if the provided location is within the bounds of this geometry. */
-	const bool IsUnderLocation( const FVector2D& AbsoluteCoordinate ) const
-	{
-		return
-			( AbsolutePosition.X <= AbsoluteCoordinate.X ) && ( ( AbsolutePosition.X + Size.X * Scale ) > AbsoluteCoordinate.X ) &&
-			( AbsolutePosition.Y <= AbsoluteCoordinate.Y ) && ( ( AbsolutePosition.Y + Size.Y * Scale ) > AbsoluteCoordinate.Y );
-	}
+	/**
+	 * Create a paint geometry relative to this one that whose local space is "inflated" by the specified amount in each direction.
+	 * The paint geometry inherits the widget's render transform.
+	 *
+	 * @param InflateAmount	Amount by which to "inflate" the geometry in each direction around the center point. Effectively defines a layout transform offset and an inflation of the LocalSize.
+	 * 
+	 * @return				The new paint geometry derived from this one.
+	 */
+	FPaintGeometry ToInflatedPaintGeometry( const FVector2D& InflateAmount ) const;
 
-	/** @return Translates AbsoluteCoordinate into the space defined by this Geometry. */
-	const FVector2D AbsoluteToLocal( FVector2D AbsoluteCoordinate ) const
-	{
-		return (AbsoluteCoordinate - this->AbsolutePosition) / Scale;
-	}
+	/** 
+	 * Absolute coordinates could be either desktop or window space depending on what space the root of the widget hierarchy is in.
+	 * 
+	 * @return true if the provided location in absolute coordinates is within the bounds of this geometry. 
+	 */
+	bool IsUnderLocation( const FVector2D& AbsoluteCoordinate ) const;
+
+	/** 
+	 * Absolute coordinates could be either desktop or window space depending on what space the root of the widget hierarchy is in.
+	 * 
+	 * @return Transforms AbsoluteCoordinate into the local space of this Geometry. 
+	 */
+	FVector2D AbsoluteToLocal( FVector2D AbsoluteCoordinate ) const;
 
 	/**
 	 * Translates local coordinates into absolute coordinates
-	 *
+	 * 
+	 * Absolute coordinates could be either desktop or window space depending on what space the root of the widget hierarchy is in.
+	 * 
 	 * @return  Absolute coordinates
 	 */
-	const FVector2D LocalToAbsolute( FVector2D LocalCoordinate ) const
-	{
-		return (LocalCoordinate * Scale) + this->AbsolutePosition;
-	}
+	FVector2D LocalToAbsolute( FVector2D LocalCoordinate ) const;
 	
 	/**
-	 * Returns the allocated geometry's relative position and size as a rect
-	 *
-	 * @return  Allotted geometry relative position and size rectangle
-	 */
-	FSlateRect GetRect( ) const
-	{
-		return FSlateRect( Position.X, Position.Y, Position.X + Size.X, Position.Y + Size.Y );
-	}
-
-	/**
+	 * !!! DEPRECATED !!! This legacy function does not account for render transforms.
+	 * 
 	 * Returns a clipping rectangle corresponding to the allocated geometry's absolute position and size.
 	 * Note that the clipping rectangle starts 1 pixel above and left of the geometry because clipping is not
 	 * inclusive on the lower bound.
+	 * 
+	 * Absolute coordinates could be either desktop or window space depending on what space the root of the widget hierarchy is in.
 	 *
-	 * @return  Allotted geometry absolute position and size rectangle
+	 * @return  Allotted geometry rectangle in absolute coordinates.
 	 */
-	FSlateRect GetClippingRect( ) const
-	{
-		return FSlateRect( AbsolutePosition.X, AbsolutePosition.Y, AbsolutePosition.X + Size.X * Scale, AbsolutePosition.Y + Size.Y * Scale );
-	}
+	FSlateRect GetClippingRect( ) const;
 	
 	/** @return A String representation of this Geometry */
-	FString ToString( ) const
-	{
-		return FString::Printf(TEXT("[Pos=%s, Abs=%s, Size=%s]"), *Position.ToString(), *AbsolutePosition.ToString(), *Size.ToString() );
-	}
+	FString ToString( ) const;
 
-	/** @return the size of the geometry in screen space */
-	FVector2D GetDrawSize() const
-	{
-		return Size * Scale;
-	}
+	/** 
+	 * !!! DEPRECATED !!! This legacy function does not account for render transforms.
+	 * 
+	 * Absolute coordinates could be either desktop or window space depending on what space the root of the widget hierarchy is in.
+	 *
+	 * @return the size of the geometry in absolute space */
+	FVector2D GetDrawSize() const;
+
+	/** @return the size of the geometry in local space. */
+	const FVector2D& GetLocalSize() const { return Size; }
 
+	/** @return the accumulated render transform. Shouldn't be needed in general. */
+	const FSlateRenderTransform& GetAccumulatedRenderTransform() const { return AccumulatedRenderTransform; }
+
+	/** @return the accumulated layout transform. Shouldn't be needed in general. */
+	FSlateLayoutTransform GetAccumulatedLayoutTransform() const { return FSlateLayoutTransform(Scale, AbsolutePosition); }
+
+	/**
+	 * Special case method to append a layout transform to a geometry.
+	 * This is used in cases where the FGeometry was arranged in window space
+	 * and we need to add the root desktop translation.
+	 * If you find yourself wanting to use this function, ask someone if there's a better way.
+	 * 
+	 * @param LayoutTransform	An additional layout transform to append to this geoemtry.
+	 */
+	void AppendTransform(const FSlateLayoutTransform& LayoutTransform);
 public:
+	/** 
+	 * 
+	 * !!! DEPRECATED !!! Use GetLocalSize() accessor instead of directly accessing public members.
+	 * 
+	 *	   This member has been made const to prevent mutation.
+	 *	   There is no way to easily detect mutation of public members, thus no way to update the render transforms when they are modified.
+	 * 
+	 * 
+	 * Size of the geometry in local space. 
+	 */
+	FVector2D /*Local*/Size;
 
-	/** Position relative to the parent */
-	FVector2D Position;
+	/** 
+	 * !!! DEPRECATED !!! These legacy public members should ideally not be referenced, as they do not account for the render transform.
+	 *     FGeometry manipulation should be done in local space as much as possible so logic can be done in aligned local space, but 
+	 *     still support arbitrary render transforms.
+	 * 
+	 *	   This member has been made const to prevent mutation, which would also break render transforms, which are computed during construction.
+	 *	   There is no way to easily detect mutation of public members, thus no way to update the render transforms when they are modified.
+	 * 
+	 * Scale in absolute space. Equivalent to the scale of the accumulated layout transform. 
+	 * 
+	 * Absolute coordinates could be either desktop or window space depending on what space the root of the widget hierarchy is in.
+	 */
+	const float /*Absolute*/Scale;
 
-	/** Position in screen space or (window space during rendering) */
-	FVector2D AbsolutePosition;	
+	/** 
+	 * !!! DEPRECATED !!! These legacy public members should ideally not be referenced, as they do not account for the render transform.
+	 *     FGeometry manipulation should be done in local space as much as possible so logic can be done in aligned local space, but 
+	 *     still support arbitrary render transforms.
+	 * 
+	 *	   This member has been made const to prevent mutation, which would also break render transforms, which are computed during construction.
+	 *	   There is no way to easily detect mutation of public members, thus no way to update the render transforms when they are modified.
+	 * 
+	 * Position in absolute space. Equivalent to the translation of the accumulated layout transform. 
+	 * 
+	 * Absolute coordinates could be either desktop or window space depending on what space the root of the widget hierarchy is in.
+	 */
+	const FVector2D AbsolutePosition;	
 
-	/** The dimensions */
-	FVector2D Size;
+	/** 
+	 * !!! DEPRECATED !!! 
+	 * 
+	 * Position of the geometry with respect to its parent in local space. Equivalent to the translation portion of the Local->Parent layout transform.
+	 * If you know your children have no additional scale applied to them, you can use this as the Local->Parent layout transform. If your children
+	 * DO have additional scale applied, there is no way to determine the actual Local->Parent layout transform, since the scale is accumulated.
+	 */
+	const FVector2D /*Local*/Position;
 
-	/** How much this geometry is scaled: WidgetCoodinates / SlateCoordinates */
-	float Scale;
+private:
+
+	/** Concatenated Render Transform. Actual transform used for rendering.
+	 * Formed as: Concat(LocalRenderTransform, LocalLayoutTransform, Parent->AccumulatedRenderTransform) 
+	 * 
+	 * For rendering, absolute coordinates will always be in window space (relative to the root window).
+	 */
+	FSlateRenderTransform AccumulatedRenderTransform;
 };
 
 
diff --git a/Engine/Source/Runtime/SlateCore/Public/Layout/LayoutGeometry.h b/Engine/Source/Runtime/SlateCore/Public/Layout/LayoutGeometry.h
new file mode 100644
--- /dev/null
+++ b/Engine/Source/Runtime/SlateCore/Public/Layout/LayoutGeometry.h
@@ -0,0 +1,77 @@
+// Copyright 1998-2014 Epic Games, Inc. All Rights Reserved.
+
+#pragma once
+
+#include "SlateLayoutTransform.h"
+#include "SlateRect.h"
+
+class FLayoutGeometry
+{
+public:
+	FLayoutGeometry()
+		: LocalSize(1.0f, 1.0f)
+	{
+	}
+
+	explicit FLayoutGeometry(const FSlateLayoutTransform& InLocalToParent, const FVector2D& SizeInLocalSpace)
+		: LocalToParent(InLocalToParent)
+		, LocalSize(SizeInLocalSpace)
+	{
+	}
+
+	const FSlateLayoutTransform& GetLocalToParentTransform() const 
+	{ 
+		return LocalToParent; 
+	}
+
+	const FVector2D& GetSizeInLocalSpace() const
+	{
+		return LocalSize;
+	}
+
+	FVector2D GetSizeInParentSpace() const
+	{
+		return TransformVector(LocalToParent, LocalSize);
+	}
+
+	FVector2D GetOffsetInParentSpace() const
+	{
+		return LocalToParent.GetTranslation();
+	}
+
+	FSlateRect GetRectInLocalSpace() const
+	{
+		return FSlateRect(FVector2D(0.0f, 0.0f), LocalSize);
+	}
+
+	FSlateRect GetRectInParentSpace() const
+	{
+		return TransformRect(LocalToParent, GetRectInLocalSpace());
+	}
+
+public:
+	//FLayoutGeometry MakeChild(const FSlateLayoutTransform& ChildToThisTransform)
+	//{
+	//	return FLayoutGeometry(Concatenate(ChildToThisTransform, LocalToParent), LocalSize);
+	//}
+
+	//FLayoutGeometry MakeChild(const FSlateLayoutTransform& ChildToThisTransform, const FVector2D& SizeInChildSpace)
+	//{
+	//	return FLayoutGeometry(Concatenate(ChildToThisTransform, LocalToParent), SizeInChildSpace);
+	//}
+
+	//FLayoutGeometry MakeOffsetChild(const FVector2D& ChildOffset, const FVector2D& SizeInChildSpace)
+	//{
+	//	return FLayoutGeometry(Concatenate(ChildOffset, LocalToParent), SizeInChildSpace);
+	//}
+
+	//FLayoutGeometry MakeInflatedChild(const FVector2D& InflateAmount)
+	//{
+	//	return FLayoutGeometry(Concatenate(Inverse(InflateAmount), LocalToParent), LocalSize + 2.0f*InflateAmount);
+	//}
+
+private:
+	FSlateLayoutTransform LocalToParent;
+	FVector2D LocalSize;
+};
+
diff --git a/Engine/Source/Runtime/SlateCore/Public/Layout/PaintGeometry.h b/Engine/Source/Runtime/SlateCore/Public/Layout/PaintGeometry.h
--- a/Engine/Source/Runtime/SlateCore/Public/Layout/PaintGeometry.h
+++ b/Engine/Source/Runtime/SlateCore/Public/Layout/PaintGeometry.h
@@ -2,67 +2,127 @@
 
 #pragma once
 
+#include "SlateRenderTransform.h"
+#include "SlateLayoutTransform.h"
 
 /**
- * A Paint geometry is a draw-space rectangle.
- * The DrawPosition and DrawSize are already positioned and
- * scaled for immediate consumption by the draw code.
+ * A Paint geometry contains the window-space (draw-space) info to draw an element on the screen.
+ * 
+ * It contains the size of the element in the element's local space along with the 
+ * transform needed to position the element in window space.
+ * 
+ * DrawPosition/DrawSize/DrawScale are maintained for legacy reasons, but are deprecated:
+ * 
+ *		The DrawPosition and DrawSize are already positioned and
+ *		scaled for immediate consumption by the draw code.
  *
- * The DrawScale is only applied to the internal aspects of the draw primitives. e.g. Line thickness, 3x3 grid margins, etc.
+ *		The DrawScale is only applied to the internal aspects of the draw primitives. e.g. Line thickness, 3x3 grid margins, etc.
  */
 struct SLATECORE_API FPaintGeometry
 {
-	/** Render-space position at which we will draw */
+	/** 
+	 * !!! DEPRECATED!!! Drawing should only happen in local space to ensure render transforms work.
+	 *	WARNING: This legacy member represents is LAYOUT space, and does not account for render transforms.
+	 *	
+	 * Render-space position at which we will draw.
+	 *
+	 * 
+	 */
 	FVector2D DrawPosition;
 
 	/**
-	 * The size in render-space that we want to make this element.
-	 * This field is ignored by some elements (e.g. spline, lines)
+	 * !!! DEPRECATED!!! Drawing should only happen in local space to ensure render transforms work.
+	 *	WARNING: This legacy member represents is LAYOUT space, and does not account for render transforms.
+	 * 
+	 * Only affects the draw aspects of individual draw elements. e.g. line thickness, 3x3 grid margins 
 	 */
-	FVector2D DrawSize;
-
-	/** Only affects the draw aspects of individual draw elements. e.g. line thickness, 3x3 grid margins */
 	float DrawScale;
 
-public:
+	/** Get the Size of the geometry in local space. Must call CommitTransformsIfUsingLegacyConstructor() first if legacy ctor is used. */
+	const FVector2D& GetLocalSize() const { return LocalSize; }
+
+	/** Access the final render transform. Must call CommitTransformsIfUsingLegacyConstructor() first if legacy ctor is used. */
+	const FSlateRenderTransform& GetAccumulatedRenderTransform() const { return AccumulatedRenderTransform; }
 
 	/**
-	 * Creates and initializes a new instance.
-	 *
-	 * @param InDrawPosition The draw position.
-	 * @param InDrawSize The draw size.
-	 * @param InDrawScale The draw scale.
+	 * Support mutable geometries constructed in window space, and possibly mutated later, as all legacy members are public.
+	 * In these cases we defer creating of the RenderTransform and LocalSize until rendering time to ensure that all member changes have finished.
+	 * WARNING: Legacy usage does NOT support render transforms!
 	 */
-	FPaintGeometry( FVector2D InDrawPosition, FVector2D InDrawSize, float InDrawScale )
-		: DrawPosition(InDrawPosition)
-		, DrawSize (InDrawSize)
-		, DrawScale(InDrawScale)
-	{ }
+	void CommitTransformsIfUsingLegacyConstructor() const 
+	{ 
+		if (!bUsingLegacyConstructor) return;
+		
+		AccumulatedRenderTransform = FSlateRenderTransform(DrawScale, DrawPosition);
+		FSlateLayoutTransform AccumulatedLayoutTransform = FSlateLayoutTransform(DrawScale, DrawPosition);
+		LocalSize = TransformVector(Inverse(AccumulatedLayoutTransform), DrawSize);
+	}
+
+private:
+	// Mutable to support legacy constructors. Doesn't account for render transforms.
+	mutable FVector2D DrawSize;
 
+	// Mutable to support legacy constructors.
+	mutable FVector2D LocalSize;
+
+	// final render transform for drawing. Transforms from local space to window space for the draw element.
+	// Mutable to support legacy constructors.
+	mutable FSlateRenderTransform AccumulatedRenderTransform;
+
+	// Support legacy constructors.
+	bool bUsingLegacyConstructor;
 public:
+	/** Default ctor. */
+	FPaintGeometry() 
+		: DrawPosition(0.0f, 0.0f)
+		, DrawScale(1.0f)
+		, DrawSize(0.0f, 0.0f)
+		, LocalSize(0.0f, 0.0f)
+		, AccumulatedRenderTransform()
+		, bUsingLegacyConstructor(true)
+	{}
 
 	/**
-	 * Converts this paint geometry to a Slate rectangle.
+	 * Creates and initializes a new instance.
 	 *
-	 * @return The Slate rectangle.
+	 * @param InLocalSize					The size in local space
+	 * @param InAccumulatedLayoutTransform	The accumulated layout transform (from an FGeometry)
+	 * @param InAccumulatedRenderTransform	The accumulated render transform (from an FGeometry)
 	 */
-	FSlateRect ToSlateRect( ) const
-	{
-		return FSlateRect(DrawPosition.X, DrawPosition.Y, DrawPosition.X + DrawSize.X, DrawPosition.Y + DrawSize.Y);
+	FPaintGeometry( const FSlateLayoutTransform& InAccumulatedLayoutTransform, const FSlateRenderTransform& InAccumulatedRenderTransform, const FVector2D& InLocalSize )
+		: DrawPosition(InAccumulatedLayoutTransform.GetTranslation())
+		, DrawScale(InAccumulatedLayoutTransform.GetScale())
+		, DrawSize(0.0f, 0.0f)
+		, LocalSize(InLocalSize)
+		, AccumulatedRenderTransform(InAccumulatedRenderTransform)
+		, bUsingLegacyConstructor(false)
+	{ 
 	}
 
-public:
+	// !!! DEPRECATED!!! This is legacy and should be removed!
+	FPaintGeometry( FVector2D InDrawPosition, FVector2D InDrawSize, float InDrawScale )
+		: DrawPosition(InDrawPosition)
+		, DrawScale(InDrawScale)
+		, DrawSize(InDrawSize)
+		, LocalSize(0.0f, 0.0f)
+		, AccumulatedRenderTransform()
+		, bUsingLegacyConstructor(true)
+	{ 
+	}
 
-	static const FPaintGeometry& Identity( )
+	/**
+	 * Special case method to append a layout transform to a paint geometry.
+	 * This is used in cases where the FGeometry was arranged in desktop space
+	 * and we need to undo the root desktop translation to get into window space.
+	 * If you find yourself wanting to use this function, ask someone if there's a better way.
+	 * 
+	 * @param LayoutTransform	An additional layout transform to append to this paint geoemtry.
+	 */
+	void AppendTransform(const FSlateLayoutTransform& LayoutTransform)
 	{
-		static FPaintGeometry IdentityPaintGeometry
-		(
-			FVector2D(0.0f, 0.0f),	// Position: Draw at the origin of draw space
-			FVector2D(0.0f, 0.0f),	// Size    : Size is ignored by splines
-			1.0f					// Scale   : Do not do any scaling.
-		);
-
-		return IdentityPaintGeometry;
+		AccumulatedRenderTransform = ::Concatenate(AccumulatedRenderTransform, LayoutTransform);
+		DrawPosition = TransformPoint(LayoutTransform, DrawPosition);
+		DrawScale = ::Concatenate(LayoutTransform.GetScale(), DrawScale);
 	}
 };
 
diff --git a/Engine/Source/Runtime/SlateCore/Public/Layout/SlateRect.h b/Engine/Source/Runtime/SlateCore/Public/Layout/SlateRect.h
--- a/Engine/Source/Runtime/SlateCore/Public/Layout/SlateRect.h
+++ b/Engine/Source/Runtime/SlateCore/Public/Layout/SlateRect.h
@@ -4,7 +4,20 @@
 
 
 /** 
- * A rectangle defined by upper-left and lower-right corners
+ * A rectangle defined by upper-left and lower-right corners.
+ * 
+ * Assumes a "screen-like" coordinate system where the origin is in the top-left, with the Y-axis going down.
+ * Functions like "contains" etc will not work with other conventions.
+ * 
+ *      +---------> X
+ *      |
+ *      |    (Left,Top)  
+ *      |            o----o 
+ *      |            |    |
+ *      |            o----o 
+ *      |                (Right, Bottom)
+ *      \/
+ *      Y
  */
 class SLATECORE_API FSlateRect
 {
@@ -75,15 +88,34 @@ public:
 		return FVector2D( Left, Top );
 	}
 
+	/**
+	 * Returns the bottom-right position of the rectangle
+	 * 
+	 * @return The bottom-right position.
+	 */
+	FVector2D GetBottomRight() const
+	{
+		return FVector2D( Right, Bottom );
+	}
+
 	/**
 	 * Return a rectangle that is contracted on each side by the amount specified in each margin.
 	 *
-	 * @param FMargin The amount to contract the geometry.
+	 * @param InsetAmount The amount to contract the geometry.
 	 *
 	 * @return An inset rectangle.
 	 */
 	FSlateRect InsetBy( const struct FMargin& InsetAmount ) const;
 
+	/**
+	 * Return a rectangle that is offset by the amount specified .
+	 *
+	 * @param OffsetAmount The amount to contract the geometry.
+	 *
+	 * @return An offset rectangle.
+	 */
+	FSlateRect OffsetBy( const FVector2D& OffsetAmount ) const;
+
 	/**
 	 * Returns the rect that encompasses both rectangles
 	 * 
@@ -96,7 +128,6 @@ public:
 		return FSlateRect( FMath::Min( Left, Other.Left ), FMath::Min( Top, Other.Top ), FMath::Max( Right, Other.Right ), FMath::Max( Bottom, Other.Bottom ) );
 	}
 
-	
 	FSlateRect IntersectionWith( const FSlateRect& Other ) const 
 	{
 		FSlateRect Intersected( FMath::Max( this->Left, Other.Left ), FMath::Max(this->Top, Other.Top), FMath::Min( this->Right, Other.Right ), FMath::Min( this->Bottom, Other.Bottom ) );
@@ -179,5 +210,24 @@ public:
 	}
 };
 
+/**
+ * Transforms a rect by the given transform, ensuring the rect does not get inverted.
+ */
+template <typename TransformType>
+FSlateRect TransformRect(const TransformType& Transform, const FSlateRect& Rect)
+{
+	FVector2D TopLeftTransformed = TransformPoint(Transform, FVector2D(Rect.Left, Rect.Top));
+	FVector2D BottomRightTransformed = TransformPoint(Transform, FVector2D(Rect.Right, Rect.Bottom));
+
+	if (TopLeftTransformed.X > BottomRightTransformed.X)
+	{
+		Swap(TopLeftTransformed.X, BottomRightTransformed.X);
+	}
+	if (TopLeftTransformed.Y > BottomRightTransformed.Y)
+	{
+		Swap(TopLeftTransformed.Y, BottomRightTransformed.Y);
+	}
+	return FSlateRect(TopLeftTransformed, BottomRightTransformed);
+}
 
 template<> struct TIsPODType<FSlateRect> { enum { Value = true }; };
diff --git a/Engine/Source/Runtime/SlateCore/Public/Rendering/DrawElements.h b/Engine/Source/Runtime/SlateCore/Public/Rendering/DrawElements.h
--- a/Engine/Source/Runtime/SlateCore/Public/Rendering/DrawElements.h
+++ b/Engine/Source/Runtime/SlateCore/Public/Rendering/DrawElements.h
@@ -210,6 +210,7 @@ public:
 		ESlateDrawEffect::Type InDrawEffects = ESlateDrawEffect::None, 
 		const FLinearColor& InTint = FLinearColor::White );
 
+	// !!! DEPRECATED !!! Use a render transform om your widget instead.
 	SLATECORE_API static void MakeRotatedBox(
 		FSlateWindowElementList& ElementList,
 		uint32 InLayer, 
@@ -253,7 +254,7 @@ public:
 	 * @param InClippingRect           Parts of the element are clipped if it falls outside of this rectangle
 	 * @param InDrawEffects            Optional draw effects to apply
 	 */
-	SLATECORE_API static void MakeGradient( FSlateWindowElementList& ElementList, uint32 InLayer, const FPaintGeometry& InPaintGeometry, TArray<FSlateGradientStop> InGradientStops, EOrientation InGradientType, const FSlateRect& InClippingRect, ESlateDrawEffect::Type InDrawEffects = ESlateDrawEffect::None, bool bGammaCorrect = true );
+	SLATECORE_API static void MakeGradient( FSlateWindowElementList& ElementList, uint32 InLayer, const FPaintGeometry& PaintGeometry, TArray<FSlateGradientStop> InGradientStops, EOrientation InGradientType, const FSlateRect& InClippingRect, ESlateDrawEffect::Type InDrawEffects = ESlateDrawEffect::None, bool bGammaCorrect = true );
 
 	/**
 	 * Creates a spline element
@@ -314,8 +315,9 @@ public:
 
 	EElementType GetElementType() const { return ElementType; }
 	uint32 GetLayer() const { return Layer; }
+	const FSlateRenderTransform& GetRenderTransform() const { return RenderTransform; }
 	const FVector2D& GetPosition() const { return Position; }
-	const FVector2D& GetSize() const { return Size; }
+	const FVector2D& GetLocalSize() const { return LocalSize; }
 	float GetScale() const { return Scale; }
 	const FSlateRect& GetClippingRect() const { return ClippingRect; }
 	const FSlateDataPayload& GetDataPayload() const { return DataPayload; }
@@ -323,48 +325,14 @@ public:
 
 private:
 
-	static FVector2D GetRotationPoint( const FPaintGeometry& PaintGeometry, const TOptional<FVector2D>& UserRotationPoint, ERotationSpace RotationSpace )
-	{
-		FVector2D RotationPoint(0,0);
-	
-		const FVector2D& Position = PaintGeometry.DrawPosition;
-		const FVector2D& Size = PaintGeometry.DrawSize;
-
-		switch( RotationSpace )
-		{
-		case RelativeToElement:
-			{
-				if( !UserRotationPoint.IsSet() )
-				{
-					// If the user did not specify a rotation point, we rotate about the center of the element
-					RotationPoint = ( Position + ( Position + Size ) ) / 2.0f;
-				}
-				else
-				{
-					// Rotate relative to the position of the element
-					RotationPoint = Position + UserRotationPoint.GetValue();
-				}
-			}
-			break;
-		case RelativeToWorld:
-			{
-				// Passthrough, its in world space
-				RotationPoint =  UserRotationPoint.Get( FVector2D::ZeroVector );
-			}
-			break;
-		default:
-			check(0);
-			break;
-		}
-
-		return RotationPoint;
-	}
+	static FVector2D GetRotationPoint( const FPaintGeometry& PaintGeometry, const TOptional<FVector2D>& UserRotationPoint, ERotationSpace RotationSpace );
 
 private:
 	FSlateDataPayload DataPayload;
+	FSlateRenderTransform RenderTransform;
 	FSlateRect ClippingRect;
 	FVector2D Position;
-	FVector2D Size;
+	FVector2D LocalSize;
 	float Scale;
 	uint32 Layer;
 	uint32 DrawEffects;
@@ -376,34 +344,25 @@ private:
  */
 struct FShaderParams
 {
-	/** Vertex shader parameters */
-	FVector4 VertexParams;
 	/** Pixel shader parameters */
 	FVector4 PixelParams;
 
 	FShaderParams()
-		: VertexParams( 0,0,0,0 )
-		, PixelParams( 0,0,0,0 )
+		: PixelParams( 0,0,0,0 )
 	{}
 
-	FShaderParams( const FVector4& InVertexParams, const FVector4& InPixelParams )
-		: VertexParams( InVertexParams )
-		, PixelParams( InPixelParams )
+	FShaderParams( const FVector4& InPixelParams )
+		: PixelParams( InPixelParams )
 	{}
 
 	bool operator==( const FShaderParams& Other ) const
 	{
-		return VertexParams == Other.VertexParams && PixelParams == Other.PixelParams;
+		return PixelParams == Other.PixelParams;
 	}
 
 	static FShaderParams MakePixelShaderParams( const FVector4& PixelShaderParams )
 	{
-		return FShaderParams( FVector4(0,0,0,0), PixelShaderParams );
-	}
-
-	static FShaderParams MakeVertexShaderParams( const FVector4& VertexShaderParams )
-	{
-		return FShaderParams( VertexShaderParams, FVector4(0,0,0,0) );
+		return FShaderParams( PixelShaderParams );
 	}
 };
 
diff --git a/Engine/Source/Runtime/SlateCore/Public/Rendering/ElementBatcher.h b/Engine/Source/Runtime/SlateCore/Public/Rendering/ElementBatcher.h
--- a/Engine/Source/Runtime/SlateCore/Public/Rendering/ElementBatcher.h
+++ b/Engine/Source/Runtime/SlateCore/Public/Rendering/ElementBatcher.h
@@ -47,104 +47,43 @@ private:
 
 	/** 
 	 * Creates vertices necessary to draw a Quad element 
-	 *
-	 * @param Position			The top left screen space position of the element
-	 * @param Size				The size of the element
-	 * @param Scale				The amount to scale the element by
-	 * @param InClippingRect	The clipping rectangle for the element
-	 * @param Layer				The layer to draw this element in
-	 * @param DrawEffects		DrawEffects to apply
-	 * @param Color				The color of the quad
 	 */
-	void AddQuadElement( const FVector2D& Position, const FVector2D& Size, float Scale, const FSlateRect& InClippingRect, uint32 Layer, FColor Color = FColor(255,255,255));
+	void AddQuadElement( const FSlateDrawElement& DrawElement, FColor Color = FColor(255,255,255));
 
 	/** 
 	 * Creates vertices necessary to draw a 3x3 element
-	 *
-	 * @param Position			The top left screen space position of the element
-	 * @param Size				The size of the element
-	 * @param Scale				The amount to scale the element by
-	 * @param InPayload			The data payload for this element
-	 * @param InClippingRect	The clipping rectangle for the element
-	 * @param DrawEffects		DrawEffects to apply
-	 * @param Layer				The layer to draw this element in
 	 */
-	void AddBoxElement( const FVector2D& Position, const FVector2D& Size, float Scale, const FSlateDataPayload& InPayload, const FSlateRect& InClippingRect, ESlateDrawEffect::Type DrawEffects, uint32 Layer );
+	void AddBoxElement( const FSlateDrawElement& DrawElement );
 
 	/** 
 	 * Creates vertices necessary to draw a string (one quad per character)
-	 *
-	 * @param Position			The top left screen space position of the element
-	 * @param Size				The size of the element
-	 * @param Scale				The amount to scale the element by
-	 * @param InPayload			The data payload for this element
-	 * @param InClippingRect	The clipping rectangle for the element
-	 * @param DrawEffects		DrawEffects to apply
-	 * @param Layer				The layer to draw this element in
 	 */
-	void AddTextElement( const FVector2D& Position, const FVector2D& Size, float Scale, const FSlateDataPayload& InPayload, const FSlateRect& InClippingRect, ESlateDrawEffect::Type DrawEffects, uint32 Layer );
+	void AddTextElement( const FSlateDrawElement& DrawElement );
 
 	/** 
 	 * Creates vertices necessary to draw a gradient box (horizontal or vertical)
-	 *
-	 * @param Position			The top left screen space position of the element
-	 * @param Size				The size of the element
-	 * @param Scale				The amount to scale the element by
-	 * @param InPayload			The data payload for this element
-	 * @param InClippingRect	The clipping rectangle for the element
-	 * @param DrawEffects		DrawEffects to apply
-	 * @param Layer				The layer to draw this element in
 	 */
-	void AddGradientElement( const FVector2D& Position, const FVector2D& Size, float Scale, const FSlateDataPayload& InPayload, const FSlateRect& InClippingRect, ESlateDrawEffect::Type DrawEffects, uint32 Layer );
+	void AddGradientElement( const FSlateDrawElement& DrawElement );
 
 	/** 
 	 * Creates vertices necessary to draw a spline (Bezier curve)
-	 *
-	 * @param Position			The top left screen space position of the element
-	 * @param Scale				The amount to scale the element by
-	 * @param InPayload			The data payload for this element
-	 * @param InClippingRect	The clipping rectangle for the element
-	 * @param DrawEffects		DrawEffects to apply
-	 * @param Layer				The layer to draw this element in
 	 */
-	void AddSplineElement( const FVector2D& Position, float Scale, const FSlateDataPayload& InPayload, const FSlateRect& InClippingRect, ESlateDrawEffect::Type DrawEffects, uint32 Layer );
+	void AddSplineElement( const FSlateDrawElement& DrawElement );
 
 	/** 
 	 * Creates vertices necessary to draw a series of attached line segments
-	 *
-	 * @param Position			The top left screen space position of the element
-	 * @param Size				The size of the element (not used)
-	 * @param Scale				The amount to scale the element by (not used)
-	 * @param InPayload			The data payload for this element
-	 * @param InClippingRect	The clipping rectangle for the element
-	 * @param DrawEffects		DrawEffects to apply
-	 * @param Layer				The layer to draw this element in
 	 */
-	void AddLineElement( const FVector2D& Position, const FVector2D& Size, float Scale, const FSlateDataPayload& InPayload, const FSlateRect& InClippingRect, ESlateDrawEffect::Type DrawEffects, uint32 Layer );
+	void AddLineElement( const FSlateDrawElement& DrawElement );
 	
 	/** 
 	 * Creates vertices necessary to draw a viewport (just a textured quad)
-	 *
-	 * @param Position			The top left screen space position of the element
-	 * @param Scale				The amount to scale the element by
-	 * @param Size				The size of the element
-	 * @param InPayload			The data payload for this element
-	 * @param InClippingRect	The clipping rectangle for the element
-	 * @param Layer				The layer to draw this element in
 	 */
-	void AddViewportElement( const FVector2D& Position, const FVector2D& Size, float Scale, const FSlateDataPayload& InPayload, const FSlateRect& InClippingRect, ESlateDrawEffect::Type DrawEffects, uint32 Layer );
+	void AddViewportElement( const FSlateDrawElement& DrawElement );
 
 	/** 
 	 * Creates vertices necessary to draw a border element
-	 *
-	 * @param Position		The top left screen space position of the element
-	 * @param Size			The size of the element
-	 * @param Scale			The amount to scale the element by
-	 * @param InPayload		The data payload for this element
-	 * @param DrawEffects	DrawEffects to apply
-	 * @param Layer			The layer to draw this element in
 	 */
-	void AddBorderElement( const FVector2D& Position, const FVector2D& Size, float Scale, const FSlateDataPayload& InPayload, const FSlateRect& InClippingRect, ESlateDrawEffect::Type DrawEffects, uint32 Layer );
+	void AddBorderElement( const FSlateDrawElement& DrawElement );
 
 	/**
 	 * Batches a custom slate drawing element
@@ -156,7 +95,7 @@ private:
 	 * @param DrawEffects	DrawEffects to apply
 	 * @param Layer			The layer to draw this element in
 	 */
-	void AddCustomElement( const FVector2D& Position, const FVector2D& Size, float Scale, const FSlateDataPayload& InPayload, const FSlateRect& InClippingRect, ESlateDrawEffect::Type DrawEffects, uint32 Layer );
+	void AddCustomElement( const FSlateDrawElement& DrawElement );
 
 	/** 
 	 * Finds an batch for an element based on the passed in parameters
diff --git a/Engine/Source/Runtime/SlateCore/Public/Rendering/RenderingCommon.h b/Engine/Source/Runtime/SlateCore/Public/Rendering/RenderingCommon.h
--- a/Engine/Source/Runtime/SlateCore/Public/Rendering/RenderingCommon.h
+++ b/Engine/Source/Runtime/SlateCore/Public/Rendering/RenderingCommon.h
@@ -87,6 +87,39 @@ namespace ESlateLineJoinType
 	};
 };
 
+/**
+ * Stores a rectangle that has been transformed by an arbitrary render transform. 
+ * We provide a ctor that does the work common to slate drawing, but you could technically 
+ * create this any way you want.
+ */
+struct FSlateRotatedRect
+{
+	FSlateRotatedRect(const FSlateLayoutTransform& InverseLayoutTransform, const FSlateRenderTransform& RenderTransform, const FSlateRect& ClipRectInLayoutWindowSpace);
+	/** transformed Top-left corner. */
+	FVector2D TopLeft;
+	/** transformed X extent (right-left). */
+	FVector2D ExtentX;
+	/** transformed Y extent (bottom-top). */
+	FVector2D ExtentY;
+
+	FSlateRect ToBoundingRect() const;
+	bool IsUnderLocation(const FVector2D& Location) const;
+};
+
+/**
+ * Stores a Rotated rect as float16 for rendering.
+ */
+struct FSlateRotatedRectFloat16
+{
+	FSlateRotatedRectFloat16();
+	FSlateRotatedRectFloat16(const FSlateRotatedRect& RotatedRect);
+	/** transformed Top-left corner. */
+	FFloat16 TopLeft[2];
+	/** transformed X extent (right-left). */
+	FFloat16 ExtentX[2];
+	/** transformed Y extent (bottom-top). */
+	FFloat16 ExtentY[2];
+};
 
 /** 
  * A struct which defines a basic vertex seen by the Slate vertex buffers and shaders
@@ -95,66 +128,16 @@ struct FSlateVertex
 {
 	/** Texture coordinates.  The first 2 are in xy and the 2nd are in zw */
 	FVector4 TexCoords; 
-	/** Clipping coordinates (left,top,right,bottom)*/
-	uint16 ClipCoords[4];
-	/** Position of the vertex in world space */
-	uint16 Position[2];
+	/** Position of the vertex in window space */
+	int16 Position[2];
+	/** clip center/extents in render window space (window space with render transforms applied) */
+	FSlateRotatedRectFloat16 ClipRect;
 	/** Vertex color */
 	FColor Color;
 	
-	FSlateVertex() {}
-
-	FORCEINLINE FSlateVertex( const FVector2D& InPosition, const FVector2D& InTexCoord, const FVector2D& InTexCoord2, const FColor& InColor, const FSlateRect& InClipCoords) 
-		: TexCoords( InTexCoord.X, InTexCoord.Y, InTexCoord2.X, InTexCoord2.Y )
-		, Color( InColor )
-	{
-
-		Position[0] = FMath::TruncToInt(InPosition.X);
-		Position[1] = FMath::TruncToInt(InPosition.Y);
-		ClipCoords[0] = FMath::TruncToInt(InClipCoords.Left);
-		ClipCoords[1] = FMath::TruncToInt(InClipCoords.Top);
-		ClipCoords[2] = FMath::TruncToInt(InClipCoords.Right);
-		ClipCoords[3] = FMath::TruncToInt(InClipCoords.Bottom);
-	}
-
-	FORCEINLINE FSlateVertex( const FVector2D& InPosition, const FVector2D& InTexCoord, const FVector2D& InTexCoord2, const FColor& InColor, const FVector4& InClipCoords) 
-		: TexCoords( InTexCoord.X, InTexCoord.Y, InTexCoord2.X, InTexCoord2.Y )
-		, Color( InColor )
-	{
-
-		Position[0] = FMath::TruncToInt(InPosition.X);
-		Position[1] = FMath::TruncToInt(InPosition.Y);
-		ClipCoords[0] = FMath::TruncToInt(InClipCoords.X);
-		ClipCoords[1] = FMath::TruncToInt(InClipCoords.Y);
-		ClipCoords[2] = FMath::TruncToInt(InClipCoords.Z);
-		ClipCoords[3] = FMath::TruncToInt(InClipCoords.W);
-	}
-
-	FORCEINLINE FSlateVertex( const FVector2D& InPosition, const FVector2D& InTexCoord, const FColor& InColor, const FSlateRect& InClipCoords ) 
-		: TexCoords( InTexCoord.X, InTexCoord.Y, 1.0f, 1.0f )
-		, Color( InColor )
-	{
-
-		Position[0] = FMath::TruncToInt(InPosition.X);
-		Position[1] = FMath::TruncToInt(InPosition.Y);
-		ClipCoords[0] = FMath::TruncToInt(InClipCoords.Left);
-		ClipCoords[1] = FMath::TruncToInt(InClipCoords.Top);
-		ClipCoords[2] = FMath::TruncToInt(InClipCoords.Right);
-		ClipCoords[3] = FMath::TruncToInt(InClipCoords.Bottom);
-	}
-
-	FORCEINLINE FSlateVertex( const FVector2D& InPosition, const FVector2D& InTexCoord, const uint32 InDWORD ) 
-		: TexCoords( InTexCoord.X, InTexCoord.Y, 1.0f, 1.0f )
-		, Color( InDWORD )
-	{
-
-		Position[0] = FMath::TruncToInt(InPosition.X);
-		Position[1] = FMath::TruncToInt(InPosition.Y);
-		ClipCoords[0] = 0;
-		ClipCoords[1] = 0;
-		ClipCoords[2] = 0;
-		ClipCoords[3] = 0;
-	}
+	FSlateVertex();
+	FSlateVertex( const FSlateRenderTransform& RenderTransform, const FVector2D& InLocalPosition, const FVector2D& InTexCoord, const FVector2D& InTexCoord2, const FColor& InColor, const FSlateRotatedRectFloat16& InClipRect );
+	FSlateVertex( const FSlateRenderTransform& RenderTransform, const FVector2D& InLocalPosition, const FVector2D& InTexCoord, const FColor& InColor, const FSlateRotatedRectFloat16& InClipRect );
 };
 
 template<> struct TIsPODType<FSlateVertex> { enum { Value = true }; };
diff --git a/Engine/Source/Runtime/SlateCore/Public/Rendering/SlateLayoutTransform.h b/Engine/Source/Runtime/SlateCore/Public/Rendering/SlateLayoutTransform.h
new file mode 100644
--- /dev/null
+++ b/Engine/Source/Runtime/SlateCore/Public/Rendering/SlateLayoutTransform.h
@@ -0,0 +1,147 @@
+// Copyright 1998-2014 Epic Games, Inc. All Rights Reserved.
+#pragma once
+
+#include "TransformCalculus2D.h"
+
+/**
+ * Represents a 2D transformation in the following order: scale then translate.
+ * Used by FGeometry for it's layout transformations.
+ *
+ * Matrix form looks like:
+ *   [Vx Vy 1] * [ S   0   0 ]
+ *               [ 0   S   0 ]
+ *               [ Tx  Ty  1 ]
+ *
+ */
+class FSlateLayoutTransform
+{
+public:
+	/** Ctor from a scale followed by translate. Shortcut to Concatenate(InScale, InTranslation). */
+	explicit FSlateLayoutTransform(float InScale = 1.0f, const FVector2D& InTranslation = FVector2D(ForceInit))
+		:Scale(InScale)
+		,Translation(InTranslation)
+	{
+	}
+
+	/** Ctor from a 2D translation followed by a scale. Shortcut to Concatenate(InTranslation, InScale). While this is the opposite order we internally store them, we can represent this correctly. */
+	explicit FSlateLayoutTransform(const FVector2D& InTranslation/*, float InScale = 1.0f*/)
+		:Scale(1.0f/*InScale*/)
+		,Translation(InTranslation/**InScale*/)
+	{
+	}
+
+	/** Access to the 2D translation */
+	const FVector2D& GetTranslation() const
+	{
+		return Translation;
+	}
+
+	/** Access to the scale. */
+	float GetScale() const
+	{
+		return Scale;
+	}
+
+	/** Support for converting to an FMatrix. */
+	FMatrix ToMatrix() const
+	{
+		FMatrix Matrix = FScaleMatrix(GetScale());
+		Matrix.SetOrigin(FVector(GetTranslation(), 0.0f));
+		return Matrix;
+	}
+
+	/** 2D transform support. */
+	FVector2D TransformPoint(const FVector2D& Point) const
+	{
+		return ::TransformPoint(Translation, ::TransformPoint(Scale, Point));
+	}
+
+	/** 2D transform support. */
+	FVector2D TransformVector(const FVector2D& Vector) const
+	{
+		return ::TransformVector(Translation, ::TransformVector(Scale, Vector));
+	}
+
+	/**
+	 * This works by transforming the origin through LHS then RHS.
+	 * In matrix form, looks like this:
+	 * [ Sa  0   0 ]   [ Sb  0   0 ]
+	 * [ 0   Sa  0 ] * [ 0   Sb  0 ] 
+	 * [ Tax Tay 1 ]   [ Tbx Tby 1 ]
+	 */
+	FSlateLayoutTransform Concatenate(const FSlateLayoutTransform& RHS) const
+	{
+		// New Translation is essentially: RHS.TransformPoint(TransformPoint(FVector2D::ZeroVector))
+		// Since Zero through LHS -> Translation we optimize it slightly to skip the zero multiplies.
+		return FSlateLayoutTransform(::Concatenate(Scale, RHS.Scale), RHS.TransformPoint(Translation));
+	}
+
+	/** Invert the transform/scale. */
+	FSlateLayoutTransform Inverse() const
+	{
+		return FSlateLayoutTransform(::Inverse(Scale), ::Inverse(Translation) * ::Inverse(Scale));
+	}
+
+	/** Equality. */
+	bool operator==(const FSlateLayoutTransform& Other) const
+	{
+		return Scale == Other.Scale && Translation == Other.Translation;
+	}
+	
+	/** Inequality. */
+	bool operator!=(const FSlateLayoutTransform& Other) const
+	{
+		return !operator==(Other);
+	}
+
+private:
+	float Scale;
+	FVector2D Translation;
+};
+
+/** Specialization for concatenating a uniform scale and 2D Translation. */
+inline FSlateLayoutTransform Concatenate(float Scale, const FVector2D& Translation)
+{
+	return FSlateLayoutTransform(Scale, Translation);
+}
+
+/** Specialization for concatenating a 2D Translation and uniform scale. */
+inline FSlateLayoutTransform Concatenate(const FVector2D& Translation, float Scale)
+{
+	return FSlateLayoutTransform(Scale, TransformPoint(Scale, Translation));
+}
+
+/** concatenation rules for LayoutTransform. */
+template<> struct ConcatenateRules<FSlateLayoutTransform, float                > { typedef FSlateLayoutTransform ResultType; };
+/** concatenation rules for LayoutTransform. */
+template<> struct ConcatenateRules<float                , FSlateLayoutTransform> { typedef FSlateLayoutTransform ResultType; };
+/** concatenation rules for LayoutTransform. */
+template<> struct ConcatenateRules<FSlateLayoutTransform, FVector2D            > { typedef FSlateLayoutTransform ResultType; };
+/** concatenation rules for LayoutTransform. */
+template<> struct ConcatenateRules<FVector2D            , FSlateLayoutTransform> { typedef FSlateLayoutTransform ResultType; };
+
+/** concatenation rules for LayoutTransform. */
+template<> struct ConcatenateRules<FSlateLayoutTransform, FMatrix              > { typedef FMatrix ResultType; };
+/** concatenation rules for LayoutTransform. */
+template<> struct ConcatenateRules<FMatrix              , FSlateLayoutTransform> { typedef FMatrix ResultType; };
+
+/** concatenation rules for layout transforms and 2x2 generalized transforms. Need to be upcast to FTransform2D. */
+template<> struct ConcatenateRules<FScale2D             , FSlateLayoutTransform> { typedef FTransform2D ResultType; };
+template<> struct ConcatenateRules<FShear2D             , FSlateLayoutTransform> { typedef FTransform2D ResultType; };
+template<> struct ConcatenateRules<FQuat2D              , FSlateLayoutTransform> { typedef FTransform2D ResultType; };
+template<> struct ConcatenateRules<FMatrix2x2           , FSlateLayoutTransform> { typedef FTransform2D ResultType; };
+template<> struct ConcatenateRules<FSlateLayoutTransform, FScale2D             > { typedef FTransform2D ResultType; };
+template<> struct ConcatenateRules<FSlateLayoutTransform, FShear2D             > { typedef FTransform2D ResultType; };
+template<> struct ConcatenateRules<FSlateLayoutTransform, FQuat2D              > { typedef FTransform2D ResultType; };
+template<> struct ConcatenateRules<FSlateLayoutTransform, FMatrix2x2           > { typedef FTransform2D ResultType; };
+
+//////////////////////////////////////////////////////////////////////////
+// FSlateLayoutTransform adapters.
+// 
+// Adapt FTransform2D to accept FSlateLayoutTransforms as well.
+//////////////////////////////////////////////////////////////////////////
+template<> template<> inline FTransform2D TransformConverter<FTransform2D>::Convert<FSlateLayoutTransform>(const FSlateLayoutTransform& Transform)
+{
+	return FTransform2D(FScale2D(Transform.GetScale()), Transform.GetTranslation());
+}
+
diff --git a/Engine/Source/Runtime/SlateCore/Public/Rendering/SlateRenderTransform.h b/Engine/Source/Runtime/SlateCore/Public/Rendering/SlateRenderTransform.h
new file mode 100644
--- /dev/null
+++ b/Engine/Source/Runtime/SlateCore/Public/Rendering/SlateRenderTransform.h
@@ -0,0 +1,8 @@
+// Copyright 1998-2014 Epic Games, Inc. All Rights Reserved.
+#pragma once
+
+#include "TransformCalculus2D.h"
+
+/** Typecast so it's more clear to the reader when we are dealing with a render transform vs. a layout transform. */
+typedef FTransform2D FSlateRenderTransform;
+
diff --git a/Engine/Source/Runtime/SlateCore/Public/SlateCore.h b/Engine/Source/Runtime/SlateCore/Public/SlateCore.h
--- a/Engine/Source/Runtime/SlateCore/Public/SlateCore.h
+++ b/Engine/Source/Runtime/SlateCore/Public/SlateCore.h
@@ -134,7 +134,6 @@ DECLARE_MEMORY_STAT_EXTERN(TEXT("Index Buffer Memory"), STAT_SlateIndexBufferMem
 #include "RenderingCommon.h"
 #include "DrawElements.h"
 #include "RenderingPolicy.h"
-#include "ElementBatcher.h"
 #include "SlateDrawBuffer.h"
 #include "SlateRenderer.h"
 
diff --git a/Engine/Source/Runtime/SlateCore/Public/Textures/TextureAtlas.h b/Engine/Source/Runtime/SlateCore/Public/Textures/TextureAtlas.h
--- a/Engine/Source/Runtime/SlateCore/Public/Textures/TextureAtlas.h
+++ b/Engine/Source/Runtime/SlateCore/Public/Textures/TextureAtlas.h
@@ -2,6 +2,20 @@
 
 #pragma once
 
+/** 
+ * Specifies how to handle texture atlas padding (when specified for the atlas). 
+ * We only support one pixel of padding because we don't support mips or aniso filtering on atlas textures right now.
+ */
+enum ESlateTextureAtlasPaddingStyle
+{
+	/** Don't pad the atlas. */
+	NoPadding,
+	/** Dilate the texture by one pixel to pad the atlas. */
+	DilateBorder,
+	/** One pixel uniform padding border filled with zeros. */
+	PadWithZero,
+};
+
 /**
  * Structure holding information about where a texture is located in the atlas
  * Pointers to left and right children build a tree of texture rectangles so we can easily find optimal slots for textures
@@ -20,6 +34,7 @@ struct FAtlasedTextureSlot
 	uint32 Width;
 	/** The height of the character */
 	uint32 Height;
+	/** Uniform Padding. can only be zero or one. See ESlateTextureAtlasPaddingStyle. */
 	uint8 Padding;
 	FAtlasedTextureSlot( uint32 InX, uint32 InY, uint32 InWidth, uint32 InHeight, uint8 InPadding )
 		: Left(nullptr)
@@ -35,20 +50,19 @@ struct FAtlasedTextureSlot
 	}
 };
 
-
 /**
  * Base class texture atlases in Slate
  */
 class SLATECORE_API FSlateTextureAtlas
 {
 public:
-	FSlateTextureAtlas( uint32 InWidth, uint32 InHeight, uint32 StrideBytes, uint32 InPadding )
+	FSlateTextureAtlas( uint32 InWidth, uint32 InHeight, uint32 StrideBytes, ESlateTextureAtlasPaddingStyle InPaddingStyle )
 		: AtlasData()
 		, RootNode( nullptr )
 		, AtlasWidth( InWidth )
 		, AtlasHeight( InHeight )
 		, Stride( StrideBytes )
-		, Padding( InPadding )
+		, PaddingStyle( InPaddingStyle )
 		, bNeedsUpdate( false )
 	{
 		InitAtlasData();
@@ -120,21 +134,23 @@ protected:
 		uint32 SrcTextureWidth;
 		/** The width of the dest texture */
 		uint32 DestTextureWidth;
-		/** The stride of each pixel */
-		uint32 DataStride;
-		/** Padding in the destination */
-		uint32 Padding;
-		/** Whether to dialate the texture and put the extra data in the padding area */
-		bool bDialateToPadding;
 	};
 
 	/**
-	 * Copies a single row from a source texture to a dest texture
-	 * Optionally dilating the texture by the padding amount
+	 * Copies a single row from a source texture to a dest texture,
+	 * respecting the padding.
 	 *
 	 * @param CopyRowData	Information for how to copy a row
 	 */
-	void CopyRow( FCopyRowData& CopyRowData );
+	void CopyRow( const FCopyRowData& CopyRowData );
+
+	/**
+	 * Zeros out a row in the dest texture (used with PaddingStyle == PadWithZero).
+	 * respecting the padding.
+	 *
+	 * @param CopyRowData	Information for how to copy a row
+	 */
+	void ZeroRow( const FCopyRowData& CopyRowData );
 
 	/** 
 	 * Copies texture data into the atlas at a given slot
@@ -167,8 +183,9 @@ protected:
 	uint32 AtlasHeight;
 	/** Stride of the atlas in bytes */
 	uint32 Stride;
-	/** Padding on all sides of each texture in the atlas */
-	uint32 Padding;
+	/** Padding style */
+	ESlateTextureAtlasPaddingStyle PaddingStyle;
+
 	/** True if this texture needs to have its rendering resources updated */
 	bool bNeedsUpdate;
 };
diff --git a/Engine/Source/Runtime/SlateCore/Public/Widgets/DeclarativeSyntaxSupport.h b/Engine/Source/Runtime/SlateCore/Public/Widgets/DeclarativeSyntaxSupport.h
--- a/Engine/Source/Runtime/SlateCore/Public/Widgets/DeclarativeSyntaxSupport.h
+++ b/Engine/Source/Runtime/SlateCore/Public/Widgets/DeclarativeSyntaxSupport.h
@@ -964,6 +964,7 @@ struct TSlateBaseNamedArgs
 	, _Cursor( TOptional<EMouseCursor::Type>() )
 	, _IsEnabled( true )
 	, _Visibility( EVisibility::Visible )
+	, _RenderTransform()
 	{
 	}
 
@@ -1001,6 +1002,7 @@ struct TSlateBaseNamedArgs
 	SLATE_ATTRIBUTE( TOptional<EMouseCursor::Type>, Cursor )
 	SLATE_ATTRIBUTE( bool, IsEnabled )
 	SLATE_ATTRIBUTE( EVisibility, Visibility )
+	SLATE_ATTRIBUTE( FSlateRenderTransform, RenderTransform )
 	SLATE_ARGUMENT( FName, Tag )
 
 	TArray<TSharedRef<ISlateMetaData>> MetaData;
@@ -1262,6 +1264,7 @@ struct TDecl
 			InArgs._Cursor ,
 			InArgs._IsEnabled ,
 			InArgs._Visibility,
+			InArgs._RenderTransform,
 			InArgs._Tag,
 			InArgs.MetaData );
 
diff --git a/Engine/Source/Runtime/SlateCore/Public/Widgets/SWidget.h b/Engine/Source/Runtime/SlateCore/Public/Widgets/SWidget.h
--- a/Engine/Source/Runtime/SlateCore/Public/Widgets/SWidget.h
+++ b/Engine/Source/Runtime/SlateCore/Public/Widgets/SWidget.h
@@ -2,6 +2,7 @@
 
 #pragma once
 
+#include "SlateRenderTransform.h"
 
 class SLATECORE_API FSlateControlledConstruction
 {
@@ -71,6 +72,7 @@ public:
 		const TAttribute< TOptional<EMouseCursor::Type> > & InCursor ,
 		const TAttribute<bool> & InEnabledState ,
 		const TAttribute<EVisibility> & InVisibility,
+		const TAttribute<FSlateRenderTransform>& InTransform,
 		const FName& InTag,
 		const TArray<TSharedRef<ISlateMetaData>>& InMetaData);
 
@@ -79,10 +81,11 @@ public:
 		const TAttribute< TOptional<EMouseCursor::Type> > & InCursor ,
 		const TAttribute<bool> & InEnabledState ,
 		const TAttribute<EVisibility> & InVisibility,
+		const TAttribute<FSlateRenderTransform>& InTransform,
 		const FName& InTag,
 		const TArray<TSharedRef<ISlateMetaData>>& InMetaData)
 	{
-		Construct(InToolTipText, InToolTip, InCursor, InEnabledState, InVisibility, InTag, InMetaData);
+		Construct(InToolTipText, InToolTip, InCursor, InEnabledState, InVisibility, InTransform, InTag, InMetaData);
 	}
 
 	//
@@ -562,6 +565,16 @@ public:
 		Visibility = InVisibility;
 	}
 
+	const FSlateRenderTransform& GetRenderTransform() const
+	{
+		return RenderTransform.Get();
+	}
+
+	void SetRenderTransform( TAttribute<FSlateRenderTransform> InTransform )
+	{
+		RenderTransform = InTransform;
+	}
+
 	/**
 	 * Set the tool tip that should appear when this widget is hovered.
 	 *
@@ -765,6 +778,8 @@ protected:
 	/** Is this widget visible, hidden or collapsed */
 	TAttribute< EVisibility > Visibility;
 
+	/** Render transform of this widget */
+	TAttribute< FSlateRenderTransform > RenderTransform;
 private:
 
 	/**
diff --git a/Engine/Source/Runtime/SlateCore/Public/Widgets/SWindow.h b/Engine/Source/Runtime/SlateCore/Public/Widgets/SWindow.h
--- a/Engine/Source/Runtime/SlateCore/Public/Widgets/SWindow.h
+++ b/Engine/Source/Runtime/SlateCore/Public/Widgets/SWindow.h
@@ -229,6 +229,12 @@ public:
 	/** @return The geometry of the window in window space (i.e. position and AbsolutePosition are 0) */
 	FGeometry GetWindowGeometryInWindow() const;
 
+	/** @return the transform from local space to screen space (desktop space). */
+	FSlateLayoutTransform GetLocalToScreenTransform() const;
+
+	/** @return the transform from local space to window space, which is basically desktop space without the offset. Essentially contains the DPI scale. */
+	FSlateLayoutTransform GetLocalToWindowTransform() const;
+
 	/** @return The position of the window in screen space */
 	FVector2D GetPositionInScreen() const;
 
@@ -241,6 +247,9 @@ public:
 	/** @return Rectangle that this window occupies in screen space */
 	FSlateRect GetRectInScreen() const;
 
+	/** @return Rectangle of the window's usable client area in screen space. */
+	FSlateRect GetClientRectInScreen() const;
+
 	/** @return the size of the window's usable client area. */
 	FVector2D GetClientSizeInScreen() const;
 
diff --git a/Engine/Source/Runtime/SlateRHIRenderer/Private/SlateRHIRenderer.cpp b/Engine/Source/Runtime/SlateRHIRenderer/Private/SlateRHIRenderer.cpp
--- a/Engine/Source/Runtime/SlateRHIRenderer/Private/SlateRHIRenderer.cpp
+++ b/Engine/Source/Runtime/SlateRHIRenderer/Private/SlateRHIRenderer.cpp
@@ -1,6 +1,7 @@
 // Copyright 1998-2014 Epic Games, Inc. All Rights Reserved.
 
 #include "SlateRHIRendererPrivatePCH.h"
+#include "ElementBatcher.h"
 #include "StereoRendering.h"
 #include "Runtime/Engine/Public/Features/ILiveStreamingService.h"
 #include "SlateNativeTextureResource.h"
diff --git a/Engine/Source/Runtime/SlateRHIRenderer/Private/SlateRHIRenderer.h b/Engine/Source/Runtime/SlateRHIRenderer/Private/SlateRHIRenderer.h
--- a/Engine/Source/Runtime/SlateRHIRenderer/Private/SlateRHIRenderer.h
+++ b/Engine/Source/Runtime/SlateRHIRenderer/Private/SlateRHIRenderer.h
@@ -4,6 +4,7 @@
 
 class FSlateRHIResourceManager;
 class FSlateRHIRenderingPolicy;
+class FSlateElementBatcher;
 
 
 #define USE_MAX_DRAWBUFFERS 1
diff --git a/Engine/Source/Runtime/SlateRHIRenderer/Private/SlateRHIResourceManager.cpp b/Engine/Source/Runtime/SlateRHIRenderer/Private/SlateRHIResourceManager.cpp
--- a/Engine/Source/Runtime/SlateRHIRenderer/Private/SlateRHIResourceManager.cpp
+++ b/Engine/Source/Runtime/SlateRHIRenderer/Private/SlateRHIResourceManager.cpp
@@ -276,7 +276,7 @@ FSlateShaderResourceProxy* FSlateRHIResourceManager::GenerateTextureResource( co
 		{
 			INC_DWORD_STAT_BY(STAT_SlateNumTextureAtlases, 1);
 
-			Atlas = new FSlateTextureAtlasRHI( AtlasSize, AtlasSize, AtlasStride, Padding );
+			Atlas = new FSlateTextureAtlasRHI( AtlasSize, AtlasSize, AtlasStride, ESlateTextureAtlasPaddingStyle::DilateBorder );
 			TextureAtlases.Add( Atlas );
 			NewSlot = TextureAtlases.Last()->AddTexture( Width, Height, Info.TextureData->GetRawBytes() );
 		}
diff --git a/Engine/Source/Runtime/SlateRHIRenderer/Private/SlateRHITextureAtlas.cpp b/Engine/Source/Runtime/SlateRHIRenderer/Private/SlateRHITextureAtlas.cpp
--- a/Engine/Source/Runtime/SlateRHIRenderer/Private/SlateRHITextureAtlas.cpp
+++ b/Engine/Source/Runtime/SlateRHIRenderer/Private/SlateRHITextureAtlas.cpp
@@ -6,8 +6,8 @@
 /* FSlateTextureAtlasRHI structors
  *****************************************************************************/
 
-FSlateTextureAtlasRHI::FSlateTextureAtlasRHI( uint32 Width, uint32 Height, uint32 StrideBytes, uint32 Padding )
-	: FSlateTextureAtlas(Width, Height, StrideBytes, Padding)
+FSlateTextureAtlasRHI::FSlateTextureAtlasRHI( uint32 Width, uint32 Height, uint32 StrideBytes, ESlateTextureAtlasPaddingStyle PaddingStyle )
+	: FSlateTextureAtlas(Width, Height, StrideBytes, PaddingStyle)
 	, AtlasTexture(new FSlateTexture2DRHIRef(Width, Height, PF_B8G8R8A8, NULL, TexCreate_SRGB, true)) 
 { }
 
diff --git a/Engine/Source/Runtime/SlateRHIRenderer/Private/SlateRHITextureAtlas.h b/Engine/Source/Runtime/SlateRHIRenderer/Private/SlateRHITextureAtlas.h
--- a/Engine/Source/Runtime/SlateRHIRenderer/Private/SlateRHITextureAtlas.h
+++ b/Engine/Source/Runtime/SlateRHIRenderer/Private/SlateRHITextureAtlas.h
@@ -19,7 +19,7 @@ public:
 	 * @param StrideBytes
 	 * @param Padding
 	 */
-	FSlateTextureAtlasRHI( uint32 Width, uint32 Height, uint32 SrideBytes, uint32 Padding );
+	FSlateTextureAtlasRHI( uint32 Width, uint32 Height, uint32 SrideBytes, ESlateTextureAtlasPaddingStyle PaddingStyle );
 
 	/**
 	 * Destructor.
diff --git a/Engine/Source/Runtime/SlateRHIRenderer/Private/SlateShaders.cpp b/Engine/Source/Runtime/SlateRHIRenderer/Private/SlateShaders.cpp
--- a/Engine/Source/Runtime/SlateRHIRenderer/Private/SlateShaders.cpp
+++ b/Engine/Source/Runtime/SlateRHIRenderer/Private/SlateShaders.cpp
@@ -44,9 +44,10 @@ void FSlateVertexDeclaration::InitRHI()
 	FVertexDeclarationElementList Elements;
 	uint32 Stride = sizeof(FSlateVertex);
 	Elements.Add(FVertexElement(0,STRUCT_OFFSET(FSlateVertex,TexCoords),VET_Float4,0,Stride));
-	Elements.Add(FVertexElement(0,STRUCT_OFFSET(FSlateVertex,ClipCoords),VET_Short4,1,Stride));
-	Elements.Add(FVertexElement(0,STRUCT_OFFSET(FSlateVertex,Position),VET_Short2,2,Stride));
-	Elements.Add(FVertexElement(0,STRUCT_OFFSET(FSlateVertex,Color),VET_Color,3,Stride));
+	Elements.Add(FVertexElement(0,STRUCT_OFFSET(FSlateVertex,Position),VET_Short2,1,Stride));
+	Elements.Add(FVertexElement(0,STRUCT_OFFSET(FSlateVertex,ClipRect),VET_Half2,2,Stride));
+	Elements.Add(FVertexElement(0,STRUCT_OFFSET(FSlateVertex,ClipRect)+STRUCT_OFFSET(FSlateRotatedRectFloat16,ExtentX),VET_Half4,3,Stride));
+	Elements.Add(FVertexElement(0,STRUCT_OFFSET(FSlateVertex,Color),VET_Color,4,Stride));
 
 	VertexDeclarationRHI = RHICreateVertexDeclaration(Elements);
 }
