Fix for an IOS bad access crash. Probably affected other platforms too.

diff --git a/Engine/Source/Runtime/SlateCore/Private/Textures/TextureAtlas.cpp b/Engine/Source/Runtime/SlateCore/Private/Textures/TextureAtlas.cpp
--- a/Engine/Source/Runtime/SlateCore/Private/Textures/TextureAtlas.cpp
+++ b/Engine/Source/Runtime/SlateCore/Private/Textures/TextureAtlas.cpp
@@ -144,7 +144,7 @@ void FSlateTextureAtlas::CopyDataIntoSlot( const FAtlasedTextureSlot* SlotToCopy
 	const uint32 AllPadding = Padding*2;
 	// The width of the source texture without padding (actual width)
 	const uint32 SourceWidth = SlotToCopyTo->Width-AllPadding; 
-	const uint32 SourceHeight = SlotToCopyTo->Width-AllPadding; 
+	const uint32 SourceHeight = SlotToCopyTo->Height-AllPadding;
 
 	FCopyRowData CopyRowData;
 	CopyRowData.DestData = Start;
