Slate: fixed code should have been using 'return Me();' but was not.

diff --git a/Engine/Source/Runtime/SlateCore/Public/Widgets/DeclarativeSyntaxSupport.h b/Engine/Source/Runtime/SlateCore/Public/Widgets/DeclarativeSyntaxSupport.h
--- a/Engine/Source/Runtime/SlateCore/Public/Widgets/DeclarativeSyntaxSupport.h
+++ b/Engine/Source/Runtime/SlateCore/Public/Widgets/DeclarativeSyntaxSupport.h
@@ -978,7 +978,7 @@ struct TSlateBaseNamedArgs
 	WidgetArgsType& AddMetaData(TSharedRef<ISlateMetaData> InMetaData)
 	{
 		MetaData.Add(InMetaData);
-		return *(static_cast<WidgetArgsType*>(this));
+		return Me();
 	}
 
 	/** Add metadata to this widget - convenience method - 1 argument */
@@ -986,7 +986,7 @@ struct TSlateBaseNamedArgs
 	WidgetArgsType& AddMetaData(Arg0Type InArg0)
 	{
 		MetaData.Add(MakeShareable(new MetaDataType(InArg0)));
-		return *(static_cast<WidgetArgsType*>(this));
+		return Me();
 	}
 
 	/** Add metadata to this widget - convenience method - 2 arguments */
@@ -994,7 +994,7 @@ struct TSlateBaseNamedArgs
 	WidgetArgsType& AddMetaData(Arg0Type InArg0, Arg1Type InArg1)
 	{
 		MetaData.Add(MakeShareable(new MetaDataType(InArg0, InArg1)));
-		return *(static_cast<WidgetArgsType*>(this));
+		return Me();
 	}
 
 	SLATE_TEXT_ATTRIBUTE( ToolTipText )
