Fixed a deprecated FGeometry constructor to take absolute position into account when it determines the appropriate render transform.
Fixes TTP 345278.

Note, this FGeometry constructor is still being used, but will be removed soon.

diff --git a/Engine/Source/Runtime/SlateCore/Private/Layout/Geometry.cpp b/Engine/Source/Runtime/SlateCore/Private/Layout/Geometry.cpp
--- a/Engine/Source/Runtime/SlateCore/Private/Layout/Geometry.cpp
+++ b/Engine/Source/Runtime/SlateCore/Private/Layout/Geometry.cpp
@@ -31,8 +31,8 @@ FGeometry::FGeometry(const FVector2D& OffsetFromParent, const FVector2D& ParentA
 
 	FSlateLayoutTransform ParentAccumulatedLayoutTransform(InScale, ParentAbsolutePosition);
 	FSlateLayoutTransform LocalLayoutTransform(LayoutOffset);
-	AccumulatedRenderTransform = TransformCast<FSlateRenderTransform>(LocalLayoutTransform);
 	FSlateLayoutTransform AccumulatedLayoutTransform = Concatenate(LocalLayoutTransform, ParentAccumulatedLayoutTransform);
+	AccumulatedRenderTransform = TransformCast<FSlateRenderTransform>(AccumulatedLayoutTransform);
 	// HACK to allow us to make FGeometry public members immutable to catch misuse.
 	const_cast<FVector2D&>(AbsolutePosition) = AccumulatedLayoutTransform.GetTranslation();
 	const_cast<float&>(Scale) = AccumulatedLayoutTransform.GetScale();
