Converted STextBlock to use an FTextLayout

This solves a wrapping issue where STextBlock would incorrectly consume new-line characters when wrapping, as well as allowing it to use custom wrapping behavior (eg, CamelCase wrapping for asset names in the content browser).

This change also implements FString -> FText passthrough for SLATE_TEXT_ATTRIBUTE (it was previously performing FText -> FString passthrough, however this didn't fit well with the fact that FText has a more efficient method of comparing whether a bound text attribute has changed (see FTextSnapshot), so can quickly work out that the FTextLayout doesn't need updating).

This change adds FTextBlockLayout to handle the text layout cache for text block types (STextBlock and SRichTextBlock) and serves a dual purpose of moving the duplicated caching logic into a single place, as well as hiding it from the widget (so allowing the cache to mutate, even when called from immutable widget functions).

This change also increases the unification between the STextBlock and SRichTextBlock construction parameters, by allowing STextBlock to specify a Justification, Margin, and LineHeightPercentage, and allowing SRichTextBlock to specify its HighlightText as an attribute.

diff --git a/Engine/Source/Editor/ContentBrowser/Private/AssetViewWidgets.cpp b/Engine/Source/Editor/ContentBrowser/Private/AssetViewWidgets.cpp
--- a/Engine/Source/Editor/ContentBrowser/Private/AssetViewWidgets.cpp
+++ b/Engine/Source/Editor/ContentBrowser/Private/AssetViewWidgets.cpp
@@ -9,6 +9,7 @@
 #include "SThumbnailEditModeTools.h"
 #include "DragAndDrop/AssetDragDropOp.h"
 #include "DragAndDrop/AssetPathDragDropOp.h"
+#include "BreakIterator.h"
 
 #define LOCTEXT_NAMESPACE "ContentBrowser"
 
@@ -1274,6 +1275,7 @@ void SAssetTileItem::Construct( const FArguments& InArgs )
 					.IsSelected(InArgs._IsSelected)
 					.IsReadOnly(ThumbnailEditMode)
 					.WrapTextAt(this, &SAssetTileItem::GetNameTextWrapWidth)
+					.LineBreakPolicy(FBreakIterator::CreateCamelCaseBreakIterator())
 			]
 		]
 	
diff --git a/Engine/Source/Editor/Documentation/Private/SDocumentationAnchor.cpp b/Engine/Source/Editor/Documentation/Private/SDocumentationAnchor.cpp
--- a/Engine/Source/Editor/Documentation/Private/SDocumentationAnchor.cpp
+++ b/Engine/Source/Editor/Documentation/Private/SDocumentationAnchor.cpp
@@ -9,24 +9,7 @@ void SDocumentationAnchor::Construct(const FArguments& InArgs )
 {
 	Link = InArgs._Link;
 	
-	struct StringToText_Local
-	{
-		static FText PassThroughAttribute( TAttribute<FString> StringAttribute )
-		{
-			return FText::FromString( StringAttribute.Get( FString() ) );
-		}
-	};
-
-	TAttribute<FText> ToolTipText;
-	if ( InArgs._ToolTipText.IsBound() )
-	{
-		ToolTipText = TAttribute< FText >::Create( TAttribute< FText >::FGetter::CreateStatic( &StringToText_Local::PassThroughAttribute, InArgs._ToolTipText ) );
-	}
-	else
-	{
-		ToolTipText = FText::FromString( InArgs._ToolTipText.Get( FString() ) );
-	}
-
+	TAttribute<FText> ToolTipText = InArgs._ToolTipText;
 	if ( !ToolTipText.IsBound() && ToolTipText.Get().IsEmpty() )
 	{
 		ToolTipText = NSLOCTEXT("DocumentationAnchor", "DefaultToolTip", "Click to open documentation");
diff --git a/Engine/Source/Editor/IntroTutorials/Private/STutorialEditableText.cpp b/Engine/Source/Editor/IntroTutorials/Private/STutorialEditableText.cpp
--- a/Engine/Source/Editor/IntroTutorials/Private/STutorialEditableText.cpp
+++ b/Engine/Source/Editor/IntroTutorials/Private/STutorialEditableText.cpp
@@ -7,7 +7,7 @@
 #include "ContentBrowserModule.h"
 #include "SColorPicker.h"
 #include "DesktopPlatformModule.h"
-
+#include "RichTextLayoutMarshaller.h"
 
 #define LOCTEXT_NAMESPACE "STutorialEditableText"
 
diff --git a/Engine/Source/Editor/Kismet/Private/SBlueprintFavoritesPalette.cpp b/Engine/Source/Editor/Kismet/Private/SBlueprintFavoritesPalette.cpp
--- a/Engine/Source/Editor/Kismet/Private/SBlueprintFavoritesPalette.cpp
+++ b/Engine/Source/Editor/Kismet/Private/SBlueprintFavoritesPalette.cpp
@@ -276,9 +276,9 @@ SBlueprintFavoritesPalette::~SBlueprintFavoritesPalette()
 void SBlueprintFavoritesPalette::Construct(FArguments const& InArgs, TWeakPtr<FBlueprintEditor> InBlueprintEditor)
 {
 	SBlueprintSubPalette::FArguments SuperArgs;
-	SuperArgs._Title       = LOCTEXT("PaletteTitle", "Favorites").ToString();
+	SuperArgs._Title       = LOCTEXT("PaletteTitle", "Favorites");
 	SuperArgs._Icon        = FEditorStyle::GetBrush("Kismet.Palette.Favorites");
-	SuperArgs._ToolTipText = LOCTEXT("PaletteToolTip", "A listing of your favorite and most used nodes.").ToString();
+	SuperArgs._ToolTipText = LOCTEXT("PaletteToolTip", "A listing of your favorite and most used nodes.");
 
 	static FString const ShowFreqUsedConfigKey("bShowFrequentlyUsed");
 	// should be set before we call the super (so CollectAllActions() has the right value)
diff --git a/Engine/Source/Editor/Kismet/Private/SBlueprintLibraryPalette.cpp b/Engine/Source/Editor/Kismet/Private/SBlueprintLibraryPalette.cpp
--- a/Engine/Source/Editor/Kismet/Private/SBlueprintLibraryPalette.cpp
+++ b/Engine/Source/Editor/Kismet/Private/SBlueprintLibraryPalette.cpp
@@ -186,9 +186,9 @@ public:
 void SBlueprintLibraryPalette::Construct(FArguments const& InArgs, TWeakPtr<FBlueprintEditor> InBlueprintEditor)
 {
 	SBlueprintSubPalette::FArguments SuperArgs;
-	SuperArgs._Title       = LOCTEXT("PaletteTitle", "Find a Node").ToString();
+	SuperArgs._Title       = LOCTEXT("PaletteTitle", "Find a Node");
 	SuperArgs._Icon        = FEditorStyle::GetBrush("Kismet.Palette.Library");
-	SuperArgs._ToolTipText = LOCTEXT("PaletteToolTip", "An all encompassing list of every node that is available for this blueprint.").ToString();
+	SuperArgs._ToolTipText = LOCTEXT("PaletteToolTip", "An all encompassing list of every node that is available for this blueprint.");
 	SuperArgs._ShowFavoriteToggles = true;
 
 	bUseLegacyLayout = InArgs._UseLegacyLayout.Get();
@@ -226,7 +226,7 @@ void SBlueprintLibraryPalette::CollectAllActions(FGraphActionListBuilderBase& Ou
 }
 
 //------------------------------------------------------------------------------
-TSharedRef<SVerticalBox> SBlueprintLibraryPalette::ConstructHeadingWidget(FSlateBrush const* const Icon, FString const& TitleText, FString const& ToolTip)
+TSharedRef<SVerticalBox> SBlueprintLibraryPalette::ConstructHeadingWidget(FSlateBrush const* const Icon, FText const& TitleText, FText const& ToolTip)
 {
 	TSharedRef<SVerticalBox> SuperHeading = SBlueprintSubPalette::ConstructHeadingWidget(Icon, TitleText, ToolTip);
 
diff --git a/Engine/Source/Editor/Kismet/Private/SBlueprintLibraryPalette.h b/Engine/Source/Editor/Kismet/Private/SBlueprintLibraryPalette.h
--- a/Engine/Source/Editor/Kismet/Private/SBlueprintLibraryPalette.h
+++ b/Engine/Source/Editor/Kismet/Private/SBlueprintLibraryPalette.h
@@ -33,7 +33,7 @@ private:
 	// End SGraphPalette Interface
 
 	// SBlueprintSubPalette Interface
-	virtual TSharedRef<SVerticalBox> ConstructHeadingWidget(FSlateBrush const* const Icon, FString const& TitleText, FString const& ToolTip) override;
+	virtual TSharedRef<SVerticalBox> ConstructHeadingWidget(FSlateBrush const* const Icon, FText const& TitleText, FText const& ToolTip) override;
 	virtual void BindCommands(TSharedPtr<FUICommandList> CommandListIn) const override;
 	virtual void GenerateContextMenuEntries(FMenuBuilder& MenuBuilder) const override;
 	// End SBlueprintSubPalette Interface
diff --git a/Engine/Source/Editor/Kismet/Private/SBlueprintSubPalette.cpp b/Engine/Source/Editor/Kismet/Private/SBlueprintSubPalette.cpp
--- a/Engine/Source/Editor/Kismet/Private/SBlueprintSubPalette.cpp
+++ b/Engine/Source/Editor/Kismet/Private/SBlueprintSubPalette.cpp
@@ -384,7 +384,7 @@ void SBlueprintSubPalette::RequestRefreshActionsList()
 *******************************************************************************/
 
 //------------------------------------------------------------------------------
-TSharedRef<SVerticalBox> SBlueprintSubPalette::ConstructHeadingWidget(FSlateBrush const* const Icon, FString const& TitleText, FString const& ToolTipText)
+TSharedRef<SVerticalBox> SBlueprintSubPalette::ConstructHeadingWidget(FSlateBrush const* const Icon, FText const& TitleText, FText const& ToolTipText)
 {
 	TSharedPtr<SToolTip> ToolTip;
 	SAssignNew(ToolTip, SToolTip).Text(ToolTipText);
diff --git a/Engine/Source/Editor/Kismet/Private/SBlueprintSubPalette.h b/Engine/Source/Editor/Kismet/Private/SBlueprintSubPalette.h
--- a/Engine/Source/Editor/Kismet/Private/SBlueprintSubPalette.h
+++ b/Engine/Source/Editor/Kismet/Private/SBlueprintSubPalette.h
@@ -82,7 +82,7 @@ protected:
 	 * @param  ToolTip		A tooltip you want displayed when the user hovers over the heading.
 	 * @return A reference to the newly created vertical slate box containing the header.
 	 */
-	virtual TSharedRef<SVerticalBox> ConstructHeadingWidget(FSlateBrush const* const Icon, FString const& TitleText, FString const& ToolTip);
+	virtual TSharedRef<SVerticalBox> ConstructHeadingWidget(FSlateBrush const* const Icon, FText const& TitleText, FText const& ToolTip);
 
 	/**
 	 * An overridable method, that fills out the provided menu-builder with 
diff --git a/Engine/Source/Editor/UnrealEd/Private/SEditorViewportToolBarMenu.cpp b/Engine/Source/Editor/UnrealEd/Private/SEditorViewportToolBarMenu.cpp
--- a/Engine/Source/Editor/UnrealEd/Private/SEditorViewportToolBarMenu.cpp
+++ b/Engine/Source/Editor/UnrealEd/Private/SEditorViewportToolBarMenu.cpp
@@ -7,7 +7,7 @@
 
 void SEditorViewportToolbarMenu::Construct( const FArguments& Declaration )
 {
-	const TAttribute<FString>& Label = Declaration._Label;
+	const TAttribute<FText>& Label = Declaration._Label;
 
 	const FName ImageName = Declaration._Image;
 	const FSlateBrush* ImageBrush = FEditorStyle::GetBrush( ImageName );
diff --git a/Engine/Source/Runtime/AppFramework/Private/Widgets/Testing/STestSuite.cpp b/Engine/Source/Runtime/AppFramework/Private/Widgets/Testing/STestSuite.cpp
--- a/Engine/Source/Runtime/AppFramework/Private/Widgets/Testing/STestSuite.cpp
+++ b/Engine/Source/Runtime/AppFramework/Private/Widgets/Testing/STestSuite.cpp
@@ -5,6 +5,7 @@
 #include "SLayoutExample.h"
 #include "SWidgetGallery.h"
 #include "TestStyle.h"
+#include "RichTextLayoutMarshaller.h"
 #include "SyntaxHighlighterTextLayoutMarshaller.h"
 #include "STestSuite.h"
 #include "TransformCalculus3D.h"
@@ -1030,6 +1031,7 @@ class SMultiLineEditingTest : public SCompoundWidget
 
 	void Construct( const FArguments& InArgs )
 	{
+		bIsReadOnly = false;
 		MultilineEditableText = LOCTEXT( "MultiLineTextTest", "He has refused his Assent to Laws, the most wholesome and necessary for the public good.\nHe has forbidden his Governors to pass Laws of immediate and pressing importance, unless suspended in their operation till his Assent should be obtained; and when so suspended, he has utterly neglected to attend to them.\nHe has refused to pass other Laws for the accommodation of large districts of people, unless those people would relinquish the right of Representation in the Legislature, a right inestimable to them and formidable to tyrants only.\n\nHe has called together legislative bodies at places unusual, uncomfortable, and distant from the depository of their public Records, for the sole purpose of fatiguing them into compliance with his measures.\nHe has dissolved Representative Houses repeatedly, for opposing with manly firmness his invasions on the rights of the people.\nHe has refused for a long time, after such dissolutions, to cause others to be elected; whereby the Legislative powers, incapable of Annihilation, have returned to the People at large for their exercise; the State remaining in the mean time exposed to all the dangers of invasion from without, and convulsions within.\nHe has endeavoured to prevent the population of these States; for that purpose obstructing the Laws for Naturalization of Foreigners; refusing to pass others to encourage their migrations hither, and raising the conditions of new Appropriations of Lands.\n" );
 
 		this->ChildSlot
@@ -1070,14 +1072,33 @@ class SMultiLineEditingTest : public SCompoundWidget
 					+ SVerticalBox::Slot()
 					.AutoHeight()
 					[
-						SNew( STextBlock )
-						.Text( LOCTEXT( "MultiLineTextAutoWrap", "Multi-line editable text auto-wrapping" ) )
+						SNew( SHorizontalBox )
+						
+						+SHorizontalBox::Slot()
+						.AutoWidth()
+						[
+							SNew( STextBlock )
+							.Text( LOCTEXT( "MultiLineTextAutoWrap", "Multi-line editable text auto-wrapping" ) )
+						]
+
+						+SHorizontalBox::Slot()
+						.HAlign(HAlign_Right)
+						[
+							SNew( SCheckBox )
+							.IsChecked(this, &SMultiLineEditingTest::IsReadOnlyChecked)
+							.OnCheckStateChanged(this, &SMultiLineEditingTest::OnReadOnlyCheckedStateChanged)
+							[
+								SNew( STextBlock )
+								.Text( LOCTEXT( "MultiLineTextReadOnly", "Read-only?" ) )
+							]
+						]
 					]
 					+ SVerticalBox::Slot()
 					[
 						SNew( SMultiLineEditableTextBox )
 						.Margin( 10 )
 						.Text( MultilineEditableText )
+						.IsReadOnly( this, &SMultiLineEditingTest::IsReadOnly )
 						//.Justification(ETextJustify::Right)
 						.Font( FSlateFontInfo( FPaths::EngineContentDir() / TEXT( "Slate/Fonts/Roboto-Regular.ttf" ), 12 ) )
 						.AutoWrapText( true )
@@ -1128,7 +1149,22 @@ private:
 		MultilineEditableText = Text;
 	}
 
+	bool IsReadOnly() const
+	{
+		return bIsReadOnly;
+	}
+
+	ESlateCheckBoxState::Type IsReadOnlyChecked() const
+	{
+		return (bIsReadOnly) ? ESlateCheckBoxState::Checked : ESlateCheckBoxState::Unchecked;
+	}
+
+	void OnReadOnlyCheckedStateChanged(ESlateCheckBoxState::Type InState)
+	{
+		bIsReadOnly = (InState == ESlateCheckBoxState::Checked);
+	}
 
+	bool bIsReadOnly;
 	FText MultilineEditableText;
 };
 
diff --git a/Engine/Source/Runtime/AppFramework/Public/Widgets/Workflow/SWizard.h b/Engine/Source/Runtime/AppFramework/Public/Widgets/Workflow/SWizard.h
--- a/Engine/Source/Runtime/AppFramework/Public/Widgets/Workflow/SWizard.h
+++ b/Engine/Source/Runtime/AppFramework/Public/Widgets/Workflow/SWizard.h
@@ -138,8 +138,8 @@ public:
 
 	SLATE_BEGIN_ARGS(SWizard)
 		: _CanFinish(true)
-		, _FinishButtonText(NSLOCTEXT("SWizard", "DefaultFinishButtonText", "Finish").ToString())
-		, _FinishButtonToolTip(NSLOCTEXT("SWizard", "DefaultFinishButtonTooltip", "Finish the wizard").ToString())
+		, _FinishButtonText(NSLOCTEXT("SWizard", "DefaultFinishButtonText", "Finish"))
+		, _FinishButtonToolTip(NSLOCTEXT("SWizard", "DefaultFinishButtonTooltip", "Finish the wizard"))
 		, _InitialPageIndex(0)
 		, _DesiredSize(FVector2D(0, 0))
 		, _ShowPageList(true)
diff --git a/Engine/Source/Runtime/Slate/Private/Framework/MultiBox/Mac/MacMenu.cpp b/Engine/Source/Runtime/Slate/Private/Framework/MultiBox/Mac/MacMenu.cpp
--- a/Engine/Source/Runtime/Slate/Private/Framework/MultiBox/Mac/MacMenu.cpp
+++ b/Engine/Source/Runtime/Slate/Private/Framework/MultiBox/Mac/MacMenu.cpp
@@ -227,7 +227,7 @@ CFStringRef FSlateMacMenu::GetMenuItemTitle(const TSharedRef<const FMenuEntryBlo
 			const TSharedRef<SWidget>& TextBlockWidget = FindTextBlockWidget(Block->EntryWidget.ToSharedRef());
 			if (TextBlockWidget != SNullWidget::NullWidget)
 			{
-				Label = FText::FromString(StaticCastSharedRef<STextBlock>(TextBlockWidget)->GetText());
+				Label = StaticCastSharedRef<STextBlock>(TextBlockWidget)->GetText();
 			}
 		}
 
diff --git a/Engine/Source/Runtime/Slate/Private/Framework/Text/PlainTextLayoutMarshaller.cpp b/Engine/Source/Runtime/Slate/Private/Framework/Text/PlainTextLayoutMarshaller.cpp
--- a/Engine/Source/Runtime/Slate/Private/Framework/Text/PlainTextLayoutMarshaller.cpp
+++ b/Engine/Source/Runtime/Slate/Private/Framework/Text/PlainTextLayoutMarshaller.cpp
@@ -35,19 +35,8 @@ void FPlainTextLayoutMarshaller::GetText(FString& TargetString, const FTextLayou
 	SourceTextLayout.GetAsText(TargetString);
 }
 
-void FPlainTextLayoutMarshaller::SetDefaultTextStyle(FTextBlockStyle InDefaultTextStyle)
-{
-	DefaultTextStyle = MoveTemp(InDefaultTextStyle);
-	MakeDirty();
-}
-
-const FTextBlockStyle& FPlainTextLayoutMarshaller::GetDefaultTextStyle() const
-{
-	return DefaultTextStyle;
-}
-
 FPlainTextLayoutMarshaller::FPlainTextLayoutMarshaller(FTextBlockStyle InDefaultTextStyle)
-	: DefaultTextStyle(MoveTemp(InDefaultTextStyle))
+	: FSlateTextLayoutMarshaller(MoveTemp(InDefaultTextStyle))
 {
 }
 
diff --git a/Engine/Source/Runtime/Slate/Private/Framework/Text/RichTextLayoutMarshaller.cpp b/Engine/Source/Runtime/Slate/Private/Framework/Text/RichTextLayoutMarshaller.cpp
--- a/Engine/Source/Runtime/Slate/Private/Framework/Text/RichTextLayoutMarshaller.cpp
+++ b/Engine/Source/Runtime/Slate/Private/Framework/Text/RichTextLayoutMarshaller.cpp
@@ -102,24 +102,13 @@ void FRichTextLayoutMarshaller::GetText(FString& TargetString, const FTextLayout
 	Writer->Write(WriterLines, TargetString);
 }
 
-void FRichTextLayoutMarshaller::SetDefaultTextStyle(FTextBlockStyle InDefaultTextStyle)
-{
-	DefaultTextStyle = MoveTemp(InDefaultTextStyle);
-	MakeDirty();
-}
-
-const FTextBlockStyle& FRichTextLayoutMarshaller::GetDefaultTextStyle() const
-{
-	return DefaultTextStyle;
-}
-
 FRichTextLayoutMarshaller::FRichTextLayoutMarshaller(TSharedPtr< IRichTextMarkupParser > InParser, TSharedPtr< IRichTextMarkupWriter > InWriter, TArray< TSharedRef< ITextDecorator > > InDecorators, const ISlateStyle* const InDecoratorStyleSet, FTextBlockStyle InDefaultTextStyle)
-	: Parser(MoveTemp(InParser))
+	: FSlateTextLayoutMarshaller(MoveTemp(InDefaultTextStyle))
+	, Parser(MoveTemp(InParser))
 	, Writer(MoveTemp(InWriter))
 	, Decorators(MoveTemp(InDecorators))
 	, InlineDecorators()
 	, DecoratorStyleSet(InDecoratorStyleSet)
-	, DefaultTextStyle(MoveTemp(InDefaultTextStyle))
 {
 }
 
diff --git a/Engine/Source/Runtime/Slate/Private/Framework/Text/SlateTextHighlightRunRenderer.cpp b/Engine/Source/Runtime/Slate/Private/Framework/Text/SlateTextHighlightRunRenderer.cpp
--- a/Engine/Source/Runtime/Slate/Private/Framework/Text/SlateTextHighlightRunRenderer.cpp
+++ b/Engine/Source/Runtime/Slate/Private/Framework/Text/SlateTextHighlightRunRenderer.cpp
@@ -27,11 +27,11 @@ int32 FSlateTextHighlightRunRenderer::OnPaint( const FPaintArgs& Args, const FTe
 		InWidgetStyle.GetColorAndOpacityTint() * DefaultStyle.HighlightColor
 		);
 
-	FLinearColor InvertedForeground = FLinearColor::White - InWidgetStyle.GetForegroundColor();
-	InvertedForeground.A = InWidgetStyle.GetForegroundColor().A;
+	FLinearColor InvertedHighlightColor = FLinearColor::White - DefaultStyle.HighlightColor;
+	InvertedHighlightColor.A = InWidgetStyle.GetForegroundColor().A;
 
 	FWidgetStyle WidgetStyle( InWidgetStyle );
-	WidgetStyle.SetForegroundColor( InvertedForeground );
+	WidgetStyle.SetForegroundColor( InvertedHighlightColor );
 
 	return Run->OnPaint( Args, Line, Block, DefaultStyle, AllottedGeometry, MyClippingRect, OutDrawElements, LayerId, WidgetStyle, bParentEnabled );
 }
diff --git a/Engine/Source/Runtime/Slate/Private/Framework/Text/SlateTextLayoutMarshaller.cpp b/Engine/Source/Runtime/Slate/Private/Framework/Text/SlateTextLayoutMarshaller.cpp
new file mode 100644
--- /dev/null
+++ b/Engine/Source/Runtime/Slate/Private/Framework/Text/SlateTextLayoutMarshaller.cpp
@@ -0,0 +1,24 @@
+// Copyright 1998-2014 Epic Games, Inc. All Rights Reserved.
+
+#include "SlatePrivatePCH.h"
+#include "SlateTextLayoutMarshaller.h"
+
+#if WITH_FANCY_TEXT
+
+void FSlateTextLayoutMarshaller::SetDefaultTextStyle(FTextBlockStyle InDefaultTextStyle)
+{
+	DefaultTextStyle = MoveTemp(InDefaultTextStyle);
+	MakeDirty();
+}
+
+const FTextBlockStyle& FSlateTextLayoutMarshaller::GetDefaultTextStyle() const
+{
+	return DefaultTextStyle;
+}
+
+FSlateTextLayoutMarshaller::FSlateTextLayoutMarshaller(FTextBlockStyle InDefaultTextStyle)
+	: DefaultTextStyle(MoveTemp(InDefaultTextStyle))
+{
+}
+
+#endif //WITH_FANCY_TEXT
diff --git a/Engine/Source/Runtime/Slate/Private/Framework/Text/TextLayout.cpp b/Engine/Source/Runtime/Slate/Private/Framework/Text/TextLayout.cpp
--- a/Engine/Source/Runtime/Slate/Private/Framework/Text/TextLayout.cpp
+++ b/Engine/Source/Runtime/Slate/Private/Framework/Text/TextLayout.cpp
@@ -587,13 +587,13 @@ void FTextLayout::UpdateHighlights()
 
 void FTextLayout::DirtyRunLayout(const TSharedRef<const IRun>& Run)
 {
-	for (int LineModelIndex = 0; LineModelIndex < LineModels.Num(); LineModelIndex++)
+	for (int32 LineModelIndex = 0; LineModelIndex < LineModels.Num(); LineModelIndex++)
 	{
 		FLineModel& LineModel = LineModels[LineModelIndex];
 
 		if (LineModel.HasWrappingInformation)
 		{
-			for (int RunIndex = 0; RunIndex < LineModel.Runs.Num(); RunIndex++)
+			for (int32 RunIndex = 0; RunIndex < LineModel.Runs.Num(); RunIndex++)
 			{
 				if (LineModel.Runs[RunIndex].GetRun() == Run)
 				{
@@ -609,20 +609,10 @@ void FTextLayout::DirtyRunLayout(const TSharedRef<const IRun>& Run)
 
 void FTextLayout::DirtyLayout()
 {
-	for (int LineModelIndex = 0; LineModelIndex < LineModels.Num(); LineModelIndex++)
-	{
-		FLineModel& LineModel = LineModels[LineModelIndex];
-
-		if (LineModel.HasWrappingInformation)
-		{
-			for (int RunIndex = 0; RunIndex < LineModel.Runs.Num(); RunIndex++)
-			{
-				LineModel.Runs[RunIndex].ClearCache();
-			}
-		}
-	}
-
 	DirtyFlags |= EDirtyState::Layout;
+
+	// Clear out the entire cache so it gets regenerated on the text call to FlowLayout
+	ClearWrappingCache();
 }
 
 void FTextLayout::ClearRunRenderers()
diff --git a/Engine/Source/Runtime/Slate/Private/Widgets/Text/SInlineEditableTextBlock.cpp b/Engine/Source/Runtime/Slate/Private/Widgets/Text/SInlineEditableTextBlock.cpp
--- a/Engine/Source/Runtime/Slate/Private/Widgets/Text/SInlineEditableTextBlock.cpp
+++ b/Engine/Source/Runtime/Slate/Private/Widgets/Text/SInlineEditableTextBlock.cpp
@@ -32,6 +32,7 @@ void SInlineEditableTextBlock::Construct( const FArguments& InArgs )
 			.HighlightText( InArgs._HighlightText )
 			.ToolTipText( InArgs._ToolTipText )
 			.WrapTextAt( InArgs._WrapTextAt )
+			.LineBreakPolicy( InArgs._LineBreakPolicy )
 		]
 	];
 
@@ -57,11 +58,7 @@ SInlineEditableTextBlock::~SInlineEditableTextBlock()
 
 void SInlineEditableTextBlock::CancelEditMode()
 {
-	HorizontalBox->RemoveSlot(TextBox.ToSharedRef());
-	TextBlock->SetVisibility(EVisibility::Visible);
-
-	// Clear the error so it will vanish.
-	TextBox->SetError( FText::GetEmpty() );
+	ExitEditingMode();
 
 	// Get the text from source again.
 	TextBox->SetText(Text);
@@ -95,7 +92,11 @@ void SInlineEditableTextBlock::EnterEditingMode()
 
 void SInlineEditableTextBlock::ExitEditingMode()
 {
-	FSlateApplication::Get().ClearKeyboardFocus(EKeyboardFocusCause::SetDirectly);	
+	HorizontalBox->RemoveSlot(TextBox.ToSharedRef());
+	TextBlock->SetVisibility(EVisibility::Visible);
+
+	// Clear the error so it will vanish.
+	TextBox->SetError( FText::GetEmpty() );
 }
 
 bool SInlineEditableTextBlock::IsInEditMode() const
@@ -243,14 +244,10 @@ void SInlineEditableTextBlock::OnTextBoxCommitted(const FText& InText, ETextComm
 			}
 		}
 		
-		HorizontalBox->RemoveSlot(TextBox.ToSharedRef());
-		TextBlock->SetVisibility(EVisibility::Visible);
+		ExitEditingMode();
 
 		OnTextCommittedDelegate.ExecuteIfBound(InText, InCommitType);
 
-		// Clear the error so it will vanish.
-		TextBox->SetError( FText::GetEmpty() );
-
 		if ( !Text.IsBound() )
 		{
 			TextBlock->SetText( Text );
diff --git a/Engine/Source/Runtime/Slate/Private/Widgets/Text/SRichTextBlock.cpp b/Engine/Source/Runtime/Slate/Private/Widgets/Text/SRichTextBlock.cpp
--- a/Engine/Source/Runtime/Slate/Private/Widgets/Text/SRichTextBlock.cpp
+++ b/Engine/Source/Runtime/Slate/Private/Widgets/Text/SRichTextBlock.cpp
@@ -5,14 +5,16 @@
 #if WITH_FANCY_TEXT
 
 #include "TextLayoutEngine.h"
-#include "SlateTextLayout.h"
-#include "SlateTextHighlightRunRenderer.h"
+#include "TextBlockLayout.h"
 #include "IRichTextMarkupParser.h"
 #include "RichTextMarkupProcessing.h"
-
+#include "RichTextLayoutMarshaller.h"
 
 void SRichTextBlock::Construct( const FArguments& InArgs )
 {
+	BoundText = InArgs._Text;
+	HighlightText = InArgs._HighlightText;
+
 	TextStyle = *InArgs._TextStyle;
 	WrapTextAt = InArgs._WrapTextAt;
 	AutoWrapText = InArgs._AutoWrapText;
@@ -20,7 +22,7 @@ void SRichTextBlock::Construct( const FArguments& InArgs )
 	LineHeightPercentage = InArgs._LineHeightPercentage;
 	Justification = InArgs._Justification;
 
-	CachedSize = FVector2D::ZeroVector;
+	CachedScale = 1.0f;
 
 	{
 		TSharedPtr<IRichTextMarkupParser> Parser = InArgs._Parser;
@@ -29,169 +31,69 @@ void SRichTextBlock::Construct( const FArguments& InArgs )
 			Parser = FDefaultRichTextMarkupParser::Create();
 		}
 
-		Marshaller = FRichTextLayoutMarshaller::Create(Parser, nullptr, InArgs._Decorators, InArgs._DecoratorStyleSet, TextStyle);
+		TSharedRef<FRichTextLayoutMarshaller> Marshaller = FRichTextLayoutMarshaller::Create(Parser, nullptr, InArgs._Decorators, InArgs._DecoratorStyleSet, TextStyle);
 		for ( const TSharedRef< ITextDecorator >& Decorator : InArgs.InlineDecorators )
 		{
 			Marshaller->AppendInlineDecorator( Decorator );
 		}
-	}
-
-	TextHighlighter = FSlateTextHighlightRunRenderer::Create();
-
-	TextLayout = FSlateTextLayout::Create();
-
-	SetText( InArgs._Text );
-
-	TextLayout->SetWrappingWidth( WrapTextAt.Get() );
-	TextLayout->SetMargin( Margin.Get() );
-	TextLayout->SetLineHeightPercentage( LineHeightPercentage.Get() );
-	TextLayout->SetJustification( Justification.Get() );
-	TextLayout->UpdateIfNeeded();
-}
-
-void SRichTextBlock::Tick( const FGeometry& AllottedGeometry, const double InCurrentTime, const float InDeltaTime )
-{
-	TextLayout->SetScale( AllottedGeometry.Scale );
-
-	bool bRequiresTextUpdate = false;
-	const FText& TextToSet = GetText();
-	if (!BoundTextLastTick.IdenticalTo(TextToSet))
-	{
-		// The pointer used by the bound text has changed, however the text may still be the same - check that now
-		if (!BoundTextLastTick.IsDisplayStringEqualTo(TextToSet))
-		{
-			// The source text has changed, so update the internal editable text
-			bRequiresTextUpdate = true;
-		}
 
-		// Update this even if the text is lexically identical, as it will update the pointer compared by IdenticalTo for the next Tick
-		BoundTextLastTick = FTextSnapshot(TextToSet);
-	}
-
-	if (bRequiresTextUpdate || Marshaller->IsDirty())
-	{
-		SetText(TextToSet);
+		TextLayoutCache = FTextBlockLayout::Create(Marshaller, nullptr);
 	}
-
-	TextLayout->SetVisibleRegion(AllottedGeometry.Size, FVector2D::ZeroVector);
 }
 
 int32 SRichTextBlock::OnPaint( const FPaintArgs& Args, const FGeometry& AllottedGeometry, const FSlateRect& MyClippingRect, FSlateWindowElementList& OutDrawElements, int32 LayerId, const FWidgetStyle& InWidgetStyle, bool bParentEnabled ) const
 {
 	SCOPE_CYCLE_COUNTER( STAT_SlateOnPaint_SRichTextBlock );
 
-	// Update the auto-wrap size now that we have computed paint geometry; won't take affect until text frame
-	// Note: This is done here rather than in Tick(), because Tick() doesn't get called while resizing windows, but OnPaint() does
-	CachedSize = AllottedGeometry.Size;
+	// todo: jdale - The scale needs to be passed to ComputeDesiredSize
+	CachedScale = AllottedGeometry.Scale;
 
-	LayerId = TextLayout->OnPaint( Args.WithNewParent(this), TextStyle, AllottedGeometry, MyClippingRect, OutDrawElements, LayerId, InWidgetStyle, ShouldBeEnabled( bParentEnabled ) );
+	// OnPaint will also update the text layout cache if required
+	LayerId = TextLayoutCache->OnPaint(
+		FTextBlockLayout::FWidgetArgs(BoundText, HighlightText, WrapTextAt, AutoWrapText, Margin, LineHeightPercentage, Justification), 
+		Args.WithNewParent(this), AllottedGeometry, MyClippingRect, OutDrawElements, LayerId, InWidgetStyle, ShouldBeEnabled(bParentEnabled)
+		);
 
 	return LayerId;
 }
 
-void SRichTextBlock::CacheDesiredSize()
-{
-	float WrappingWidth = WrapTextAt.Get();
-
-	// Text wrapping can either be used defined (WrapTextAt), automatic (AutoWrapText and CachedSize), 
-	// or a mixture of both.  Take whichever has the smallest value (>1)
-	if ( AutoWrapText.Get() && CachedSize.X >= 1.0f )
-	{
-		WrappingWidth = (WrappingWidth >= 1.0f) ? FMath::Min(WrappingWidth, CachedSize.X) : CachedSize.X;
-	}
-
-	TextLayout->SetWrappingWidth( WrappingWidth );
-	TextLayout->SetMargin( Margin.Get() );
-	TextLayout->SetLineHeightPercentage( LineHeightPercentage.Get() );
-	TextLayout->SetJustification( Justification.Get() );
-	TextLayout->SetVisibleRegion( CachedSize, FVector2D::ZeroVector );
-	TextLayout->UpdateIfNeeded();
-
-	SWidget::CacheDesiredSize();
-}
-
 FVector2D SRichTextBlock::ComputeDesiredSize() const
 {
-	//The layouts current margin size. We should not report a size smaller then the margins.
-	const FMargin LayoutMargin = TextLayout->GetMargin();
-	const FVector2D TextLayoutSize = TextLayout->GetSize();
+	// todo: jdale - The scale needs to be passed to ComputeDesiredSize
+	const float Scale = CachedScale;
 
-	//If a wrapping width has been provided that should be reported as the desired width.
-	const float Width = TextLayout->GetWrappingWidth() > 0 ? TextLayout->GetWrappingWidth() : TextLayoutSize.X;
+	// ComputeDesiredSize will also update the text layout cache if required
+	const FVector2D TextSize = TextLayoutCache->ComputeDesiredSize(
+		FTextBlockLayout::FWidgetArgs(BoundText, HighlightText, WrapTextAt, AutoWrapText, Margin, LineHeightPercentage, Justification), 
+		Scale, TextStyle
+		);
 
-	return FVector2D(FMath::Max(LayoutMargin.GetTotalSpaceAlong<Orient_Horizontal>(), Width), FMath::Max(LayoutMargin.GetTotalSpaceAlong<Orient_Vertical>(), TextLayoutSize.Y));
+	return TextSize;
 }
 
 FChildren* SRichTextBlock::GetChildren()
 {
-	return TextLayout->GetChildren();
+	return TextLayoutCache->GetChildren();
 }
 
 void SRichTextBlock::OnArrangeChildren( const FGeometry& AllottedGeometry, FArrangedChildren& ArrangedChildren ) const
 {
-	TextLayout->ArrangeChildren( AllottedGeometry, ArrangedChildren );
+	TextLayoutCache->ArrangeChildren( AllottedGeometry, ArrangedChildren );
 }
 
 void SRichTextBlock::SetText( const TAttribute<FText>& InTextAttr )
 {
 	BoundText = InTextAttr;
-
-	const FText& InText = GetText();
-
-	// Update the cached BoundText value to prevent it triggering another SetText update again next Tick
-	if (BoundText.IsBound())
-	{
-		BoundTextLastTick = FTextSnapshot(InText);
-	}
-
-	TextLayout->ClearLines();
-
-	Marshaller->SetText(InText.ToString(), *TextLayout);
 }
 
-void SRichTextBlock::SetHighlightText( const FText& InHighlightText )
+void SRichTextBlock::SetHighlightText( const TAttribute<FText>& InHighlightText )
 {
-	if ( !TextHighlighter.IsValid() )
-	{
-		return;
-	}
-
 	HighlightText = InHighlightText;
-
-	const FString& HighlightTextString = HighlightText.ToString();
-	const int32 HighlightTextLength = HighlightTextString.Len();
-
-	const TArray< FTextLayout::FLineModel >& LineModels = TextLayout->GetLineModels();
-
-	TArray< FTextRunRenderer > TextHighlights;
-	for (int32 LineIndex = 0; LineIndex < LineModels.Num(); LineIndex++)
-	{
-		const FTextLayout::FLineModel& LineModel = LineModels[ LineIndex ];
-
-		int32 FindBegin = 0;
-		int32 CurrentHighlightBegin;
-		const int32 TextLength = LineModel.Text->Len();
-		while( FindBegin < TextLength && (CurrentHighlightBegin = LineModel.Text->Find(HighlightTextString, ESearchCase::IgnoreCase, ESearchDir::FromStart, FindBegin)) != INDEX_NONE )
-		{
-			FindBegin = CurrentHighlightBegin + HighlightTextLength;
-
-			if ( TextHighlights.Num() > 0 && TextHighlights.Last().LineIndex == LineIndex && TextHighlights.Last().Range.EndIndex == CurrentHighlightBegin )
-			{
-				TextHighlights[ TextHighlights.Num() - 1 ] = FTextRunRenderer( LineIndex, FTextRange( TextHighlights.Last().Range.BeginIndex, FindBegin ), TextHighlighter.ToSharedRef() );
-			}
-			else
-			{
-				TextHighlights.Add( FTextRunRenderer( LineIndex, FTextRange( CurrentHighlightBegin, FindBegin ), TextHighlighter.ToSharedRef() ) );
-			}
-		}
-	}
-
-	TextLayout->SetRunRenderers( TextHighlights );
 }
 
 void SRichTextBlock::Refresh()
 {
-	TextLayout->DirtyLayout();
+	TextLayoutCache->DirtyLayout();
 }
 
-#endif //WITH_FANCY_TEXT
\ No newline at end of file
+#endif //WITH_FANCY_TEXT
diff --git a/Engine/Source/Runtime/Slate/Private/Widgets/Text/STextBlock.cpp b/Engine/Source/Runtime/Slate/Private/Widgets/Text/STextBlock.cpp
--- a/Engine/Source/Runtime/Slate/Private/Widgets/Text/STextBlock.cpp
+++ b/Engine/Source/Runtime/Slate/Private/Widgets/Text/STextBlock.cpp
@@ -1,27 +1,22 @@
 ﻿// Copyright 1998-2014 Epic Games, Inc. All Rights Reserved.
 
 #include "SlatePrivatePCH.h"
-#include "SlateWordWrapper.h"
+#include "PlainTextLayoutMarshaller.h"
+#include "TextBlockLayout.h"
 
-
-/**
- * Construct this widget
- *
- * @param	InArgs	The declaration data for this widget
- */
 void STextBlock::Construct( const FArguments& InArgs )
 {
-	Text = InArgs._Text;
-
 	TextStyle = InArgs._TextStyle;
 
 	HighlightText = InArgs._HighlightText;
 	WrapTextAt = InArgs._WrapTextAt;
 	AutoWrapText = InArgs._AutoWrapText;
+	Margin = InArgs._Margin;
+	LineHeightPercentage = InArgs._LineHeightPercentage;
+	Justification = InArgs._Justification;
 	MinDesiredWidth = InArgs._MinDesiredWidth;
 
-	CachedWrapTextWidth = 0.0f;
-	CachedAutoWrapTextWidth = 0.0f;
+	CachedScale = 1.0f;
 
 	Font = InArgs._Font;
 	ColorAndOpacity = InArgs._ColorAndOpacity;
@@ -32,8 +27,14 @@ void STextBlock::Construct( const FArguments& InArgs )
 
 	OnDoubleClicked = InArgs._OnDoubleClicked;
 
-	// Request text size be cached
-	bRequestCache = true;
+	BoundText = InArgs._Text;
+
+#if WITH_FANCY_TEXT
+
+	// We use a dummy style here (as it may not be safe to call the delegates used to compute the style), but the correct style is set by ComputeDesiredSize
+	TextLayoutCache = FTextBlockLayout::Create(FPlainTextLayoutMarshaller::Create(FTextBlockStyle::GetDefault()), InArgs._LineBreakPolicy);
+
+#endif//WITH_FANCY_TEXT
 }
 
 FSlateFontInfo STextBlock::GetFont() const
@@ -66,18 +67,32 @@ const FSlateBrush* STextBlock::GetHighlightShape() const
 	return HighlightShape.IsSet() ? HighlightShape.Get() : &TextStyle->HighlightShape;
 }
 
-void STextBlock::SetText( const TAttribute< FText >& InText )
+void STextBlock::SetText( const TAttribute< FString >& InText )
 {
 	struct Local
 	{
-		static FString PassThroughAttribute( TAttribute< FText > TextAttribute )
+		static FText PassThroughAttribute( TAttribute< FString > InString )
 		{
-			return TextAttribute.Get( FText::GetEmpty() ).ToString();
+			return FText::FromString( InString.Get( TEXT("") ) );
 		}
 	};
 
-	Text = TAttribute< FString >::Create(TAttribute<FString>::FGetter::CreateStatic( &Local::PassThroughAttribute, InText));
-	bRequestCache = true;
+	BoundText = TAttribute< FText >::Create(TAttribute<FText>::FGetter::CreateStatic( &Local::PassThroughAttribute, InText) );
+}
+
+void STextBlock::SetText( const FString& InText )
+{
+	BoundText = FText::FromString( InText );
+}
+
+void STextBlock::SetText( const TAttribute< FText >& InText )
+{
+	BoundText = InText;
+}
+
+void STextBlock::SetText( const FText& InText )
+{
+	BoundText = InText;
 }
 
 int32 STextBlock::OnPaint( const FPaintArgs& Args, const FGeometry& AllottedGeometry, const FSlateRect& MyClippingRect, FSlateWindowElementList& OutDrawElements, int32 LayerId, const FWidgetStyle& InWidgetStyle, bool bParentEnabled ) const
@@ -86,128 +101,63 @@ int32 STextBlock::OnPaint( const FPaintArgs& Args, const FGeometry& AllottedGeom
 	SCOPE_CYCLE_COUNTER( STAT_SlateOnPaint_STextBlock );
 #endif
 
-	const TSharedRef< FSlateFontMeasure > FontMeasureService = FSlateApplication::Get().GetRenderer()->GetFontMeasureService();
-
-	const FSlateRect ClippingRect = AllottedGeometry.GetClippingRect().IntersectionWith(MyClippingRect);
-
-	bool bEnabled = ShouldBeEnabled( bParentEnabled );
-	ESlateDrawEffect::Type DrawEffects = bEnabled ? ESlateDrawEffect::None : ESlateDrawEffect::DisabledEffect;
+	// todo: jdale - The scale needs to be passed to ComputeDesiredSize
+	CachedScale = AllottedGeometry.Scale;
 
-	FLinearColor CurShadowColor = GetShadowColorAndOpacity();
-	const FVector2D CurShadowOffset = GetShadowOffset();
+#if WITH_FANCY_TEXT
 
-	const bool ShouldDropShadow = CurShadowOffset.Size() > 0 && CurShadowColor.A > 0;
+	// OnPaint will also update the text layout cache if required
+	LayerId = TextLayoutCache->OnPaint(
+		FTextBlockLayout::FWidgetArgs(BoundText, HighlightText, WrapTextAt, AutoWrapText, Margin, LineHeightPercentage, Justification), 
+		Args.WithNewParent(this), AllottedGeometry, MyClippingRect, OutDrawElements, LayerId, InWidgetStyle, ShouldBeEnabled(bParentEnabled)
+		);
 
-	FSlateFontInfo FontInfo = GetFont();
+#else//WITH_FANCY_TEXT
 
-	// Perform text auto-wrapping if that was enabled.
-	// Note: This is done here rather than in Tick(), because Tick() doesn't get called while resizing windows, but OnPaint() does
-	if( AutoWrapText.Get() )
+	// Simple stubbed version for when WITH_FANCY_TEXT is disabled (which should only ever be for testing purposes)
+	// Just draw the text as is, with no attempt at handling wrapping or highlighting
 	{
-		const float OldWrapTextAt = CachedAutoWrapTextWidth;
-		const float NewWrapTextAt = AllottedGeometry.Size.X;
-		if( OldWrapTextAt != NewWrapTextAt )
+		const FSlateRect ClippingRect = AllottedGeometry.GetClippingRect().IntersectionWith(MyClippingRect);
+		const ESlateDrawEffect::Type DrawEffects = ShouldBeEnabled(bParentEnabled) ? ESlateDrawEffect::None : ESlateDrawEffect::DisabledEffect;
+		const FLinearColor CurShadowColor = GetShadowColorAndOpacity();
+		const FVector2D CurShadowOffset = GetShadowOffset();
+		const bool ShouldDropShadow = CurShadowOffset.Size() > 0 && CurShadowColor.A > 0;
+		const FSlateFontInfo FontInfo = GetFont();
+		const FText& TextToDraw = BoundText.Get(FText::GetEmpty());
+
+		// Draw the optional shadow
+		if (ShouldDropShadow)
 		{
-			// Available space has changed, so make sure that we recompute wrapping
-			bRequestCache = true;
-			CachedAutoWrapTextWidth = NewWrapTextAt;
+			FSlateDrawElement::MakeText(
+				OutDrawElements,
+				LayerId,
+				AllottedGeometry.ToPaintGeometry(CurShadowOffset, AllottedGeometry.Size),
+				TextToDraw,
+				FontInfo,
+				ClippingRect,
+				DrawEffects,
+				CurShadowColor * InWidgetStyle.GetColorAndOpacityTint()
+				);
 		}
-	}
 
-	// Draw the text highlight
-	{
-		const FString& StringToHighlight = HighlightText.Get().ToString();
-		const int32 CurrentHighlightLength = StringToHighlight.Len();
-
-		// Do we have text to highlight?
-		if( CurrentHighlightLength > 0 )
-		{
-			// Check to see if the original string contains text which needs to be highlighted
-			// We need to use the original as the wrapped string has had newlines added to it
-			const int32 CurrentHighlightStart = CachedOriginalString.Find(StringToHighlight, ESearchCase::IgnoreCase);
-			if ( INDEX_NONE != CurrentHighlightStart )
-			{
-				const int32 CurrentHighlightEnd = CurrentHighlightStart + CurrentHighlightLength;
-
-				TSharedPtr<FSlateRenderer> SlateRenderer = FSlateApplication::Get().GetRenderer();
-				const int32 LineHeight = FontMeasureService->GetMaxCharacterHeight(FontInfo);
-
-				// We might have to highlight multiple lines if the range we've found spans multiple entries in CachedWrappedLineData
-				for(auto It = CachedWrappedLineData.CreateConstIterator(); It; ++It)
-				{
-					const int32 LineStart	= It->Key;
-					const int32 LineEnd		= It->Value;
-
-					// We need to highlight this line if it started after a break from a previous line, or has highlighted characters before its end
-					if(CurrentHighlightStart <= LineStart || CurrentHighlightStart <= LineEnd)
-					{
-						// Clamp the highlight indices to this line so we can measure the highlighted text for just this line
-						const int32 LineHighlightStart  = FMath::Max(CurrentHighlightStart, LineStart);
-						const int32 LineHighlightEnd	= FMath::Min(CurrentHighlightEnd, LineEnd);
-
-						// Figure out screen space locations to start drawing the highlight rectangle
-						const FString TextUpToHighlightStart = CachedOriginalString.Mid(LineStart, LineHighlightStart - LineStart);
-						const FVector2D HighlightStartOffset(FontMeasureService->Measure(TextUpToHighlightStart, FontInfo).X, LineHeight * It.GetIndex());
-
-						// Figure out the actual text being highlighted
-						// Cannot just use the filter string; it might be different case and therefore different width characters
-						const FString TextToHighlight = CachedOriginalString.Mid(LineHighlightStart, LineHighlightEnd - LineHighlightStart);
-						const FVector2D HighlightSize = FontMeasureService->Measure(TextToHighlight, FontInfo);
-
-						// Draw the actual highlight rectangle
-						FSlateDrawElement::MakeBox(
-							OutDrawElements,
-							++LayerId,
-							AllottedGeometry.ToPaintGeometry(HighlightStartOffset, HighlightSize),
-							GetHighlightShape(),
-							ClippingRect,
-							DrawEffects,
-							InWidgetStyle.GetColorAndOpacityTint() * GetHighlightColor()
-							);
-					}
-				}
-			}
-		}
-	}
-
-	// Draw the optional shadow
-	if ( ShouldDropShadow )
-	{
+		// Draw the text itself
 		FSlateDrawElement::MakeText(
 			OutDrawElements,
-			LayerId,
-			AllottedGeometry.ToPaintGeometry( CurShadowOffset, AllottedGeometry.Size ),
-			CachedWrappedString,
+			++LayerId,
+			AllottedGeometry.ToPaintGeometry(),
+			TextToDraw,
 			FontInfo,
 			ClippingRect,
 			DrawEffects,
-			CurShadowColor * InWidgetStyle.GetColorAndOpacityTint()
-		);
+			InWidgetStyle.GetColorAndOpacityTint() * GetColorAndOpacity().GetColor(InWidgetStyle)
+			);
 	}
 
-	// Draw the text itself
-	FSlateDrawElement::MakeText(
-		OutDrawElements,
-		++LayerId,
-		AllottedGeometry.ToPaintGeometry(),
-		CachedWrappedString,
-		FontInfo,
-		ClippingRect,
-		DrawEffects,
-		InWidgetStyle.GetColorAndOpacityTint() * GetColorAndOpacity().GetColor(InWidgetStyle)
-	);
+#endif//WITH_FANCY_TEXT
 
 	return LayerId;
 }
 
-/**
- * See SWidget::OnMouseButtonDoubleClick.
- *
- * @param MyGeometry The Geometry of the widget receiving the event
- * @param MouseEvent Information about the input event
- *
- * @return Whether the event was handled along with possible requests for the system to take action.
- */
 FReply STextBlock::OnMouseButtonDoubleClick( const FGeometry& InMyGeometry, const FPointerEvent& InMouseEvent )
 {
 	if ( InMouseEvent.GetEffectingButton() == EKeys::LeftMouseButton )
@@ -223,82 +173,38 @@ FReply STextBlock::OnMouseButtonDoubleClick( const FGeometry& InMyGeometry, cons
 	return FReply::Unhandled();
 }
 
-void STextBlock::CacheDesiredSize()
+FVector2D STextBlock::ComputeDesiredSize() const
 {
 #if SLATE_HD_STATS
-	SCOPE_CYCLE_COUNTER( STAT_SlateCacheDesiredSize_STextBlock );
+	SCOPE_CYCLE_COUNTER( STAT_SlateComputeDesiredSize_STextBlock );
 #endif
 
-	// Get the wrapping width and font to see if they have changed
-	float WrappingWidth = WrapTextAt.Get();
-	const FSlateFontInfo& FontInfo = GetFont();
-
-	// Text wrapping can either be used defined (WrapTextAt), automatic (AutoWrapText and CachedAutoWrapTextWidth), or a mixture of both
-	// Take whichever has the smallest value (>1)
-	if(AutoWrapText.Get() && CachedAutoWrapTextWidth >= 1.0f)
-	{
-		WrappingWidth = (WrappingWidth >= 1.0f) ? FMath::Min(WrappingWidth, CachedAutoWrapTextWidth) : CachedAutoWrapTextWidth;
-	}
-
-	// We should recache text size, width, font, and wrapped string if...
-	// A cache was requested, the text is dynamic, the wrapping width has changed, or the font has changed
-	const bool bShouldCacheText = bRequestCache || Text.IsBound() || (WrappingWidth != CachedWrapTextWidth) || (!(FontInfo == CachedFont));
-
-	if(bShouldCacheText)
-	{
-		CachedOriginalString = Text.Get();
-
-		const bool bWrapText = (WrappingWidth >= 1.0f);
-
-		// Handle optional text wrapping, caching the appropriate result
-		if(bWrapText)
-		{
-			// WrapText takes care of clearing CachedWrappedLineData
-			CachedWrappedString = SlateWordWrapper::WrapText(CachedOriginalString, FontInfo, FMath::TruncToInt(WrappingWidth), 1.0f, &CachedWrappedLineData);
-		}
-		else
-		{
-			CachedWrappedString = CachedOriginalString;
-
-			CachedWrappedLineData.Empty();
-			CachedWrappedLineData.Add(TPairInitializer<int32, int32>(0, CachedOriginalString.Len()));
-		}
+#if WITH_FANCY_TEXT
 
-		const TSharedRef< FSlateFontMeasure > FontMeasureService = FSlateApplication::Get().GetRenderer()->GetFontMeasureService();
-		FVector2D TextMeasurement = FontMeasureService->Measure(CachedWrappedString, FontInfo);
+	// todo: jdale - The scale needs to be passed to ComputeDesiredSize
+	const float Scale = CachedScale;
 
-		// Force the result to exactly the wrapping width
-		if(bWrapText)
-		{
-			TextMeasurement.X = FMath::Min(TextMeasurement.X, WrappingWidth);
-		}
+	// ComputeDesiredSize will also update the text layout cache if required
+	const FVector2D TextSize = TextLayoutCache->ComputeDesiredSize(
+		FTextBlockLayout::FWidgetArgs(BoundText, HighlightText, WrapTextAt, AutoWrapText, Margin, LineHeightPercentage, Justification), 
+		Scale, GetComputedTextStyle()
+		);
 
-		TextMeasurement.X = FMath::Max(TextMeasurement.X, MinDesiredWidth.Get());
+#else//WITH_FANCY_TEXT
 
-		const FVector2D CurrentShadowOffset = GetShadowOffset();
-		const FVector2D AbsoluteShadowOffset(FMath::Abs(CurrentShadowOffset.X), FMath::Abs(CurrentShadowOffset.Y));
-		this->Advanced_SetDesiredSize(TextMeasurement + AbsoluteShadowOffset);
+	// Simple stubbed version for when WITH_FANCY_TEXT is disabled (which should only ever be for testing purposes)
+	// Just measure the text as is, with no attempt at handling wrapping
+	const TSharedRef<FSlateFontMeasure> FontMeasureService = FSlateApplication::Get().GetRenderer()->GetFontMeasureService();
+	const FVector2D TextSize = FontMeasureService->Measure(BoundText.Get(FText::GetEmpty()), GetFont());
 
-		// Update cached values
-		CachedWrapTextWidth = WrappingWidth;
-		CachedFont = FontInfo;
-		bRequestCache = false;
-	}
-}
+#endif//WITH_FANCY_TEXT
 
-FVector2D STextBlock::ComputeDesiredSize() const
-{
-	// Usually widgets just override ComputeDesiredSize(), but STextBlock
-	// overrides CacheDesiredSize() and does all the work in there.
-
-	// Dummy implementation.
-	return FVector2D::ZeroVector;
+	return FVector2D(FMath::Max(MinDesiredWidth.Get(0.0f), TextSize.X), TextSize.Y);
 }
 
 void STextBlock::SetFont(const TAttribute< FSlateFontInfo >& InFont)
 {
 	Font = InFont;
-	bRequestCache = true;
 }
 
 void STextBlock::SetColorAndOpacity(const TAttribute<FSlateColor>& InColorAndOpacity)
@@ -309,7 +215,6 @@ void STextBlock::SetColorAndOpacity(const TAttribute<FSlateColor>& InColorAndOpa
 void STextBlock::SetTextStyle(const FTextBlockStyle* InTextStyle)
 {
 	TextStyle = InTextStyle;
-	bRequestCache = true;
 }
 
 void STextBlock::SetWrapTextAt(const TAttribute<float>& InWrapTextAt)
@@ -336,3 +241,15 @@ void STextBlock::SetMinDesiredWidth(const TAttribute<float>& InMinDesiredWidth)
 {
 	MinDesiredWidth = InMinDesiredWidth;
 }
+
+FTextBlockStyle STextBlock::GetComputedTextStyle() const
+{
+	FTextBlockStyle ComputedStyle = *TextStyle;
+	ComputedStyle.SetFont( GetFont() );
+	ComputedStyle.SetColorAndOpacity( GetColorAndOpacity() );
+	ComputedStyle.SetShadowOffset( GetShadowOffset() );
+	ComputedStyle.SetShadowColorAndOpacity( GetShadowColorAndOpacity() );
+	ComputedStyle.SetHighlightColor( GetHighlightColor() );
+	ComputedStyle.SetHighlightShape( *GetHighlightShape() );
+	return ComputedStyle;
+}
diff --git a/Engine/Source/Runtime/Slate/Private/Widgets/Text/TextBlockLayout.cpp b/Engine/Source/Runtime/Slate/Private/Widgets/Text/TextBlockLayout.cpp
new file mode 100644
--- /dev/null
+++ b/Engine/Source/Runtime/Slate/Private/Widgets/Text/TextBlockLayout.cpp
@@ -0,0 +1,181 @@
+// Copyright 1998-2014 Epic Games, Inc. All Rights Reserved.
+
+#include "SlatePrivatePCH.h"
+#include "TextBlockLayout.h"
+#include "SlateTextLayout.h"
+#include "SlateTextLayoutMarshaller.h"
+#include "SlateTextHighlightRunRenderer.h"
+
+#if WITH_FANCY_TEXT
+
+TSharedRef<FTextBlockLayout> FTextBlockLayout::Create(TSharedRef<FSlateTextLayoutMarshaller> InMarshaller, TSharedPtr<IBreakIterator> InLineBreakPolicy)
+{
+	return MakeShareable(new FTextBlockLayout(InMarshaller, InLineBreakPolicy));
+}
+
+FTextBlockLayout::FTextBlockLayout(TSharedRef<FSlateTextLayoutMarshaller> InMarshaller, TSharedPtr<IBreakIterator> InLineBreakPolicy)
+	: TextLayout(FSlateTextLayout::Create())
+	, Marshaller(InMarshaller)
+	, TextHighlighter(FSlateTextHighlightRunRenderer::Create())
+	, CachedSize(ForceInitToZero)
+{
+	TextLayout->SetLineBreakIterator(InLineBreakPolicy);
+}
+
+FVector2D FTextBlockLayout::ComputeDesiredSize(const FWidgetArgs& InWidgetArgs, const float InScale, const FTextBlockStyle& InTextStyle)
+{
+	TextLayout->SetScale(InScale);
+	TextLayout->SetWrappingWidth(CalculateWrappingWidth(InWidgetArgs));
+	TextLayout->SetMargin(InWidgetArgs.Margin.Get());
+	TextLayout->SetJustification(InWidgetArgs.Justification.Get());
+	TextLayout->SetLineHeightPercentage(InWidgetArgs.LineHeightPercentage.Get());
+
+	// Has the style used for this text block changed?
+	if(!IsStyleUpToDate(InTextStyle))
+	{
+		Marshaller->SetDefaultTextStyle(InTextStyle); // will make marshaller dirty
+	}
+
+	{
+		bool bRequiresTextUpdate = false;
+		const FText& TextToSet = InWidgetArgs.Text.Get(FText::GetEmpty());
+		if(!TextLastUpdate.IdenticalTo(TextToSet))
+		{
+			// The pointer used by the bound text has changed, however the text may still be the same - check that now
+			if(!TextLastUpdate.IsDisplayStringEqualTo(TextToSet))
+			{
+				// The source text has changed, so update the internal text
+				bRequiresTextUpdate = true;
+			}
+
+			// Update this even if the text is lexically identical, as it will update the pointer compared by IdenticalTo for the next Tick
+			TextLastUpdate = FTextSnapshot(TextToSet);
+		}
+
+		if(bRequiresTextUpdate || Marshaller->IsDirty())
+		{
+			UpdateTextLayout(TextToSet);
+		}
+	}
+
+	{
+		const FText& HighlightTextToSet = InWidgetArgs.HighlightText.Get(FText::GetEmpty());
+		if(!HighlightTextLastUpdate.IdenticalTo(HighlightTextToSet))
+		{
+			// The pointer used by the bound text has changed, however the text may still be the same - check that now
+			if(!HighlightTextLastUpdate.IsDisplayStringEqualTo(HighlightTextToSet))
+			{
+				UpdateTextHighlights(HighlightTextToSet);
+			}
+
+			// Update this even if the text is lexically identical, as it will update the pointer compared by IdenticalTo for the next Tick
+			HighlightTextLastUpdate = FTextSnapshot(HighlightTextToSet);
+		}
+	}
+
+	TextLayout->UpdateIfNeeded();
+
+	return TextLayout->GetSize();
+}
+
+int32 FTextBlockLayout::OnPaint(const FWidgetArgs& InWidgetArgs, const FPaintArgs& InPaintArgs, const FGeometry& InAllottedGeometry, const FSlateRect& InClippingRect, FSlateWindowElementList& OutDrawElements, int32 LayerId, const FWidgetStyle& InWidgetStyle, bool bParentEnabled)
+{
+	CachedSize = InAllottedGeometry.Size;
+
+	TextLayout->SetWrappingWidth(CalculateWrappingWidth(InWidgetArgs));
+	TextLayout->SetVisibleRegion(InAllottedGeometry.Size, FVector2D::ZeroVector);
+
+	TextLayout->UpdateIfNeeded();
+
+	return TextLayout->OnPaint(InPaintArgs, Marshaller->GetDefaultTextStyle(), InAllottedGeometry, InClippingRect, OutDrawElements, LayerId, InWidgetStyle, bParentEnabled);
+}
+
+void FTextBlockLayout::DirtyLayout()
+{
+	TextLayout->DirtyLayout();
+}
+
+FChildren* FTextBlockLayout::GetChildren()
+{
+	return TextLayout->GetChildren();
+}
+
+void FTextBlockLayout::ArrangeChildren(const FGeometry& AllottedGeometry, FArrangedChildren& ArrangedChildren) const
+{
+	TextLayout->ArrangeChildren(AllottedGeometry, ArrangedChildren);
+}
+
+void FTextBlockLayout::UpdateTextLayout(const FText& InText)
+{
+	Marshaller->ClearDirty();
+	TextLayout->ClearLines();
+	Marshaller->SetText(InText.ToString(), *TextLayout);
+
+	TextLayout->ClearLineHighlights();
+	TextLayout->ClearRunRenderers();
+
+	HighlightTextLastUpdate = FTextSnapshot();
+}
+
+void FTextBlockLayout::UpdateTextHighlights(const FText& InHighlightText)
+{
+	const FString& HighlightTextString = InHighlightText.ToString();
+	const int32 HighlightTextLength = HighlightTextString.Len();
+
+	const TArray< FTextLayout::FLineModel >& LineModels = TextLayout->GetLineModels();
+
+	TArray<FTextRunRenderer> TextHighlights;
+	for(int32 LineIndex = 0; LineIndex < LineModels.Num(); ++LineIndex)
+	{
+		const FTextLayout::FLineModel& LineModel = LineModels[LineIndex];
+
+		int32 FindBegin = 0;
+		int32 CurrentHighlightBegin;
+		const int32 TextLength = LineModel.Text->Len();
+		while(FindBegin < TextLength && (CurrentHighlightBegin = LineModel.Text->Find(HighlightTextString, ESearchCase::IgnoreCase, ESearchDir::FromStart, FindBegin)) != INDEX_NONE)
+		{
+			FindBegin = CurrentHighlightBegin + HighlightTextLength;
+
+			if(TextHighlights.Num() > 0 && TextHighlights.Last().LineIndex == LineIndex && TextHighlights.Last().Range.EndIndex == CurrentHighlightBegin)
+			{
+				TextHighlights[TextHighlights.Num() - 1] = FTextRunRenderer(LineIndex, FTextRange(TextHighlights.Last().Range.BeginIndex, FindBegin), TextHighlighter.ToSharedRef());
+			}
+			else
+			{
+				TextHighlights.Add(FTextRunRenderer(LineIndex, FTextRange(CurrentHighlightBegin, FindBegin), TextHighlighter.ToSharedRef()));
+			}
+		}
+	}
+
+	TextLayout->SetRunRenderers(TextHighlights);
+}
+
+bool FTextBlockLayout::IsStyleUpToDate(const FTextBlockStyle& NewStyle) const
+{
+	const FTextBlockStyle& CurrentStyle = Marshaller->GetDefaultTextStyle();
+
+	return (CurrentStyle.Font == NewStyle.Font)
+		&& (CurrentStyle.ColorAndOpacity == NewStyle.ColorAndOpacity)
+		&& (CurrentStyle.ShadowOffset == NewStyle.ShadowOffset)
+		&& (CurrentStyle.ShadowColorAndOpacity == NewStyle.ShadowColorAndOpacity)
+		&& (CurrentStyle.HighlightColor == NewStyle.HighlightColor)
+		&& (CurrentStyle.HighlightShape == NewStyle.HighlightShape);
+}
+
+float FTextBlockLayout::CalculateWrappingWidth(const FWidgetArgs& InWidgetArgs) const
+{
+	const float WrapTextAt = InWidgetArgs.WrapTextAt.Get(0.0f);
+	const bool bAutoWrapText = InWidgetArgs.AutoWrapText.Get(false);
+
+	// Text wrapping can either be used defined (WrapTextAt), automatic (bAutoWrapText and CachedSize), 
+	// or a mixture of both. Take whichever has the smallest value (>1)
+	float WrappingWidth = WrapTextAt;
+	if(bAutoWrapText && CachedSize.X >= 1.0f)
+	{
+		WrappingWidth = (WrappingWidth >= 1.0f) ? FMath::Min(WrappingWidth, CachedSize.X) : CachedSize.X;
+	}
+
+	return FMath::Max(0.0f, WrappingWidth);
+}
+
+#endif //WITH_FANCY_TEXT
diff --git a/Engine/Source/Runtime/Slate/Private/Widgets/Views/SHeaderRow.cpp b/Engine/Source/Runtime/Slate/Private/Widgets/Views/SHeaderRow.cpp
--- a/Engine/Source/Runtime/Slate/Private/Widgets/Views/SHeaderRow.cpp
+++ b/Engine/Source/Runtime/Slate/Private/Widgets/Views/SHeaderRow.cpp
@@ -48,11 +48,11 @@ public:
 
 		FMargin AdjustedDefaultHeaderContentPadding = DefaultHeaderContentPadding;
 
-		TAttribute< FString > LabelText = Column.DefaultText;
+		TAttribute< FText > LabelText = Column.DefaultText;
 
 		if (Column.HeaderContent.Widget == SNullWidget::NullWidget && !Column.DefaultText.IsBound() && Column.DefaultText.Get().IsEmpty())
 		{
-			LabelText = Column.ColumnId.ToString() + TEXT("[LabelMissing]");
+			LabelText = FText::FromString( Column.ColumnId.ToString() + TEXT("[LabelMissing]") );
 		}
 
 		TSharedPtr< SHorizontalBox > Box;
diff --git a/Engine/Source/Runtime/Slate/Public/Framework/Text/PlainTextLayoutMarshaller.h b/Engine/Source/Runtime/Slate/Public/Framework/Text/PlainTextLayoutMarshaller.h
--- a/Engine/Source/Runtime/Slate/Public/Framework/Text/PlainTextLayoutMarshaller.h
+++ b/Engine/Source/Runtime/Slate/Public/Framework/Text/PlainTextLayoutMarshaller.h
@@ -3,12 +3,12 @@
 
 #if WITH_FANCY_TEXT
 
-#include "BaseTextLayoutMarshaller.h"
+#include "SlateTextLayoutMarshaller.h"
 
 /**
  * Get/set the raw text to/from a text layout as plain text
  */
-class SLATE_API FPlainTextLayoutMarshaller : public FBaseTextLayoutMarshaller
+class SLATE_API FPlainTextLayoutMarshaller : public FSlateTextLayoutMarshaller
 {
 public:
 
@@ -20,16 +20,10 @@ public:
 	virtual void SetText(const FString& SourceString, FTextLayout& TargetTextLayout) override;
 	virtual void GetText(FString& TargetString, const FTextLayout& SourceTextLayout) override;
 
-	void SetDefaultTextStyle(FTextBlockStyle InDefaultTextStyle);
-	const FTextBlockStyle& GetDefaultTextStyle() const;
-
 protected:
 
 	FPlainTextLayoutMarshaller(FTextBlockStyle InDefaultTextStyle);
 
-	/** Default style used by the TextLayout */
-	FTextBlockStyle DefaultTextStyle;
-
 };
 
 #endif //WITH_FANCY_TEXT
diff --git a/Engine/Source/Runtime/Slate/Public/Framework/Text/RichTextLayoutMarshaller.h b/Engine/Source/Runtime/Slate/Public/Framework/Text/RichTextLayoutMarshaller.h
--- a/Engine/Source/Runtime/Slate/Public/Framework/Text/RichTextLayoutMarshaller.h
+++ b/Engine/Source/Runtime/Slate/Public/Framework/Text/RichTextLayoutMarshaller.h
@@ -3,7 +3,7 @@
 
 #if WITH_FANCY_TEXT
 
-#include "BaseTextLayoutMarshaller.h"
+#include "SlateTextLayoutMarshaller.h"
 #include "ITextDecorator.h"
 #include "IRichTextMarkupParser.h"
 #include "IRichTextMarkupWriter.h"
@@ -11,7 +11,7 @@
 /**
  * Get/set the raw text to/from a text layout as rich text
  */
-class SLATE_API FRichTextLayoutMarshaller : public FBaseTextLayoutMarshaller
+class SLATE_API FRichTextLayoutMarshaller : public FSlateTextLayoutMarshaller
 {
 public:
 
@@ -24,9 +24,6 @@ public:
 	virtual void SetText(const FString& SourceString, FTextLayout& TargetTextLayout) override;
 	virtual void GetText(FString& TargetString, const FTextLayout& SourceTextLayout) override;
 
-	void SetDefaultTextStyle(FTextBlockStyle InDefaultTextStyle);
-	const FTextBlockStyle& GetDefaultTextStyle() const;
-
 	/**
 	 * Append an inline decorator to this marshaller
 	 */
@@ -57,9 +54,6 @@ protected:
 	/** The style set used for looking up styles used by decorators */
 	const ISlateStyle* DecoratorStyleSet;
 
-	/** Default style used by the TextLayout */
-	FTextBlockStyle DefaultTextStyle;
-
 };
 
 #endif //WITH_FANCY_TEXT
diff --git a/Engine/Source/Runtime/Slate/Public/Framework/Text/SlateTextLayoutMarshaller.h b/Engine/Source/Runtime/Slate/Public/Framework/Text/SlateTextLayoutMarshaller.h
new file mode 100644
--- /dev/null
+++ b/Engine/Source/Runtime/Slate/Public/Framework/Text/SlateTextLayoutMarshaller.h
@@ -0,0 +1,26 @@
+// Copyright 1998-2014 Epic Games, Inc. All Rights Reserved.
+#pragma once
+
+#if WITH_FANCY_TEXT
+
+#include "BaseTextLayoutMarshaller.h"
+
+/**
+ * Base class implementing some common functionality for all Slate text marshallers
+ */
+class SLATE_API FSlateTextLayoutMarshaller : public FBaseTextLayoutMarshaller
+{
+public:
+
+	void SetDefaultTextStyle(FTextBlockStyle InDefaultTextStyle);
+	const FTextBlockStyle& GetDefaultTextStyle() const;
+
+protected:
+
+	FSlateTextLayoutMarshaller(FTextBlockStyle InDefaultTextStyle);
+
+	/** Default style used by the TextLayout */
+	FTextBlockStyle DefaultTextStyle;
+};
+
+#endif //WITH_FANCY_TEXT
diff --git a/Engine/Source/Runtime/Slate/Public/Widgets/Layout/SExpandableArea.h b/Engine/Source/Runtime/Slate/Public/Widgets/Layout/SExpandableArea.h
--- a/Engine/Source/Runtime/Slate/Public/Widgets/Layout/SExpandableArea.h
+++ b/Engine/Source/Runtime/Slate/Public/Widgets/Layout/SExpandableArea.h
@@ -14,7 +14,7 @@ public:
 		: _Style( &FCoreStyle::Get().GetWidgetStyle<FExpandableAreaStyle>("ExpandableArea") )
 		, _BorderBackgroundColor( FLinearColor::White )
 		, _BorderImage( FCoreStyle::Get().GetBrush("ExpandableArea.Border") )
-		, _AreaTitle( TEXT("") )
+		, _AreaTitle( )
 		, _InitiallyCollapsed( false )
 		, _MaxHeight( 0.0f )
 		, _Padding( 1.0f )
diff --git a/Engine/Source/Runtime/Slate/Public/Widgets/SToolTip.h b/Engine/Source/Runtime/Slate/Public/Widgets/SToolTip.h
--- a/Engine/Source/Runtime/Slate/Public/Widgets/SToolTip.h
+++ b/Engine/Source/Runtime/Slate/Public/Widgets/SToolTip.h
@@ -78,7 +78,7 @@ public:
 private:
 
 	// Text block widget.
-	TAttribute<FString> TextContent;
+	TAttribute<FText> TextContent;
 
 	// Content widget.
 	TWeakPtr<SWidget> WidgetContent;
diff --git a/Engine/Source/Runtime/Slate/Public/Widgets/Text/SInlineEditableTextBlock.h b/Engine/Source/Runtime/Slate/Public/Widgets/Text/SInlineEditableTextBlock.h
--- a/Engine/Source/Runtime/Slate/Public/Widgets/Text/SInlineEditableTextBlock.h
+++ b/Engine/Source/Runtime/Slate/Public/Widgets/Text/SInlineEditableTextBlock.h
@@ -20,6 +20,7 @@ class SLATE_API SInlineEditableTextBlock: public SCompoundWidget
 		, _ShadowColorAndOpacity()
 		, _HighlightText()
 		, _WrapTextAt(0.0f)
+		, _LineBreakPolicy()
 		, _IsReadOnly(false)
 	{}
 
@@ -47,6 +48,9 @@ class SLATE_API SInlineEditableTextBlock: public SCompoundWidget
 		/** Whether text wraps onto a new line when it's length exceeds this width; if this value is zero or negative, no wrapping occurs. */
 		SLATE_ATTRIBUTE( float, WrapTextAt )
 
+		/** The iterator to use to detect appropriate soft-wrapping points for lines (or null to use the default) */
+		SLATE_ARGUMENT( TSharedPtr<IBreakIterator>, LineBreakPolicy )
+
 		/** True if the editable text block is read-only. It will not be able to enter edit mode if read-only */
 		SLATE_ATTRIBUTE( bool, IsReadOnly )
 
diff --git a/Engine/Source/Runtime/Slate/Public/Widgets/Text/SRichTextBlock.h b/Engine/Source/Runtime/Slate/Public/Widgets/Text/SRichTextBlock.h
--- a/Engine/Source/Runtime/Slate/Public/Widgets/Text/SRichTextBlock.h
+++ b/Engine/Source/Runtime/Slate/Public/Widgets/Text/SRichTextBlock.h
@@ -3,8 +3,6 @@
 
 #if WITH_FANCY_TEXT
 
-#include "RichTextLayoutMarshaller.h"
-
 /**
  * A rich static text widget. 
  * Through the use of markup and text decorators, text with different styles, embedded image and widgets can be achieved.
@@ -14,7 +12,8 @@ class SLATE_API SRichTextBlock : public SWidget
 public:
 
 	SLATE_BEGIN_ARGS( SRichTextBlock )
-		: _Text(  )
+		: _Text()
+		, _HighlightText()
 		, _WrapTextAt( 0.0f )
 		, _AutoWrapText(false)
 		, _DecoratorStyleSet( &FCoreStyle::Get() )
@@ -28,6 +27,9 @@ public:
 		/** The text displayed in this text block */
 		SLATE_ATTRIBUTE( FText, Text )
 
+		/** Highlight this text in the text block */
+		SLATE_ATTRIBUTE( FText, HighlightText )
+
 		/** Whether text wraps onto a new line when it's length exceeds this width; if this value is zero or negative, no wrapping occurs. */
 		SLATE_ATTRIBUTE( float, WrapTextAt )
 
@@ -103,9 +105,7 @@ public:
 	void Construct( const FArguments& InArgs );
 
 	virtual int32 OnPaint( const FPaintArgs& Args, const FGeometry& AllottedGeometry, const FSlateRect& MyClippingRect, FSlateWindowElementList& OutDrawElements, int32 LayerId, const FWidgetStyle& InWidgetStyle, bool bParentEnabled ) const;
-	virtual void Tick( const FGeometry& AllottedGeometry, const double InCurrentTime, const float InDeltaTime ) override;
 
-	virtual void CacheDesiredSize() override;
 	virtual FVector2D ComputeDesiredSize() const override;
 
 	virtual FChildren* GetChildren() override;
@@ -126,7 +126,7 @@ public:
 	 */
 	void SetText( const TAttribute<FText>& InTextAttr );
 
-	void SetHighlightText( const FText& InHighlightText );
+	void SetHighlightText( const TAttribute<FText>& InHighlightText );
 
 	/**
 	 * Causes the text to reflow it's layout
@@ -138,19 +138,14 @@ private:
 	/** The text displayed in this text block */
 	TAttribute< FText > BoundText;
 
-	/** The state of BoundText last Tick() (used to allow updates when the text is changed) */
-	FTextSnapshot BoundTextLastTick;
-
-	/** The marshaller used to set the Text text in the text layout. */
-	TSharedPtr< FRichTextLayoutMarshaller > Marshaller;
-
-	TSharedPtr< FSlateTextLayout > TextLayout;
+	/** The wrapped layout for this text block */
+	TSharedPtr< FTextBlockLayout > TextLayoutCache;
 
 	/** Default style used by the TextLayout */
 	FTextBlockStyle TextStyle;
 
-	FText HighlightText;
-	TSharedPtr< ISlateRunRenderer> TextHighlighter;
+	/** Highlight this text in the textblock */
+	TAttribute<FText> HighlightText;
 
 	/** Whether text wraps onto a new line when it's length exceeds this width; if this value is zero or negative, no wrapping occurs. */
 	TAttribute<float> WrapTextAt;
@@ -158,12 +153,17 @@ private:
 	/** True if we're wrapping text automatically based on the computed horizontal space for this widget */
 	TAttribute<bool> AutoWrapText;
 
-	/** The last known size of the control from the previous OnPaint, used to recalculate wrapping. */
-	mutable FVector2D CachedSize;
-
+	/** The amount of blank space left around the edges of text area. */
 	TAttribute< FMargin > Margin;
+
+	/** The amount to scale each lines height by. */
 	TAttribute< ETextJustify::Type > Justification; 
+
+	/** How the text should be aligned with the margin. */
 	TAttribute< float > LineHeightPercentage;
+
+	/** todo: jdale - The scale needs to be passed to ComputeDesiredSize */
+	mutable float CachedScale;
 };
 
 #endif //WITH_FANCY_TEXT
\ No newline at end of file
diff --git a/Engine/Source/Runtime/Slate/Public/Widgets/Text/STextBlock.h b/Engine/Source/Runtime/Slate/Public/Widgets/Text/STextBlock.h
--- a/Engine/Source/Runtime/Slate/Public/Widgets/Text/STextBlock.h
+++ b/Engine/Source/Runtime/Slate/Public/Widgets/Text/STextBlock.h
@@ -2,7 +2,9 @@
  
 #pragma once
 
-#include "WordWrapper.h"
+#include "IBreakIterator.h"
+
+class FTextBlockLayout;
 
 namespace ETextRole
 {
@@ -35,7 +37,11 @@ public:
 		, _HighlightText()
 		, _WrapTextAt(0.0f)
 		, _AutoWrapText(false)
+		, _Margin()
+		, _LineHeightPercentage(1.0f)
+		, _Justification(ETextJustify::Left)
 		, _MinDesiredWidth(0.0f)
+		, _LineBreakPolicy()
 		{}
 
 		/** The text displayed in this text block */
@@ -73,9 +79,21 @@ public:
 			desired size will not be clamped.  This works best in cases where the text block's size is not affecting other widget's layout. */
 		SLATE_ATTRIBUTE( bool, AutoWrapText )
 
+		/** The amount of blank space left around the edges of text area. */
+		SLATE_ATTRIBUTE( FMargin, Margin )
+
+		/** The amount to scale each lines height by. */
+		SLATE_ATTRIBUTE( float, LineHeightPercentage )
+
+		/** How the text should be aligned with the margin. */
+		SLATE_ATTRIBUTE( ETextJustify::Type, Justification )
+
 		/** Minimum width that a text block should be */
 		SLATE_ATTRIBUTE( float, MinDesiredWidth )
 
+		/** The iterator to use to detect appropriate soft-wrapping points for lines (or null to use the default) */
+		SLATE_ARGUMENT( TSharedPtr<IBreakIterator>, LineBreakPolicy )
+
 		/** Called when this text is double clicked */
 		SLATE_EVENT( FOnClicked, OnDoubleClicked )
 
@@ -93,9 +111,9 @@ public:
 	 *
 	 * @return	This text block's text string
 	 */
-	const FString& GetText() const
+	const FText& GetText() const
 	{
-		return Text.Get();
+		return BoundText.Get();
 	}
 	
 	/**
@@ -103,35 +121,10 @@ public:
 	 *
 	 * @param	InText	The new text to display
 	 */
-	void SetText( const TAttribute< FString >& InText )
-	{
-		Text = InText;
-		bRequestCache = true;
-	}
-
-	void SetText( const FString& InText )
-	{
-		Text = InText;
-		bRequestCache = true;
-	}
-
-	static FString PassThroughAttribute( TAttribute< FText > TextAttribute )
-	{
-		return TextAttribute.Get( FText::GetEmpty() ).ToString();
-	}
-
-	/**
-	 * Sets the text for this text block
-	 *
-	 * @param	InText	The new text to display
-	 */
+	void SetText( const TAttribute< FString >& InText );
+	void SetText( const FString& InText );
 	void SetText( const TAttribute< FText >& InText );
-
-	void SetText( const FText& InText )
-	{
-		Text = InText.ToString();
-		bRequestCache = true;
-	}
+	void SetText( const FText& InText );
 
 	/**
 	 * Sets the font used to draw the text
@@ -167,13 +160,12 @@ public:
 	// SWidget interface
 	virtual int32 OnPaint( const FPaintArgs& Args, const FGeometry& AllottedGeometry, const FSlateRect& MyClippingRect, FSlateWindowElementList& OutDrawElements, int32 LayerId, const FWidgetStyle& InWidgetStyle, bool bParentEnabled ) const override;
 	virtual FReply OnMouseButtonDoubleClick( const FGeometry& InMyGeometry, const FPointerEvent& InMouseEvent ) override;
-	virtual void CacheDesiredSize() override;
 	virtual FVector2D ComputeDesiredSize() const override;
 	// End of SWidget interface
 
 private:
-	/** Caches string size to avoid remeasuring if possible */
-	void CacheStringSizeIfNeeded();
+	/** Get the computed text style to use with the text marshaller */
+	FTextBlockStyle GetComputedTextStyle() const;
 
 	/** Gets the current foreground color */
 	FSlateColor GetColorAndOpacity() const;
@@ -195,9 +187,16 @@ private:
 
 private:
 	/** The text displayed in this text block */
-	TAttribute< FString > Text;
+	TAttribute< FText > BoundText;
+
+#if WITH_FANCY_TEXT
 
-	/** The style for the text block */
+	/** The wrapped layout for this text block */
+	TSharedPtr< FTextBlockLayout > TextLayoutCache;
+
+#endif//WITH_FANCY_TEXT
+
+	/** Default style used by the TextLayout */
 	const FTextBlockStyle* TextStyle;
 
 	/** Sets the font used to draw the text */
@@ -225,30 +224,21 @@ private:
 	/** True if we're wrapping text automatically based on the computed horizontal space for this widget */
 	TAttribute<bool> AutoWrapText;
 
-	/** The delegate to execute when this text is double clicked */
-	FOnClicked OnDoubleClicked;
-
-	/** Cached Font that this text is using. This is used when determining if the cached string size should be updated */
-	FSlateFontInfo CachedFont;
-
-	/** Cached wrap width that this text is using. This is used when determining if the cached string size should be updated */
-	float CachedWrapTextWidth;
-
-	/** Cached auto-wrap width that this text is using. This is used when determining if the cached string size should be updated */
-	mutable float CachedAutoWrapTextWidth;
-
-	/** Text that was used to generate CachedWrappedString */
-	FString CachedOriginalString;
-
-	/** Cached wrapped text. This gets cached once in tick so it does not have to be re-generated at paint time every frame */
-	FString CachedWrappedString;
+	/** The amount of blank space left around the edges of text area. */
+	TAttribute< FMargin > Margin;
 
-	/** Line break data for the original string -> cached string */
-	FWordWrapper::FWrappedLineData CachedWrappedLineData;
+	/** The amount to scale each lines height by. */
+	TAttribute< ETextJustify::Type > Justification; 
 
-	/** If a request to cache the string size occurred and we should recache the text */
-	mutable bool bRequestCache;
+	/** How the text should be aligned with the margin. */
+	TAttribute< float > LineHeightPercentage;
 
 	/** Prevents the text block from being smaller than desired in certain cases (e.g. when it is empty) */
 	TAttribute<float> MinDesiredWidth;
+
+	/** The delegate to execute when this text is double clicked */
+	FOnClicked OnDoubleClicked;
+
+	/** todo: jdale - The scale needs to be passed to ComputeDesiredSize */
+	mutable float CachedScale;
 };
diff --git a/Engine/Source/Runtime/Slate/Public/Widgets/Text/TextBlockLayout.h b/Engine/Source/Runtime/Slate/Public/Widgets/Text/TextBlockLayout.h
new file mode 100644
--- /dev/null
+++ b/Engine/Source/Runtime/Slate/Public/Widgets/Text/TextBlockLayout.h
@@ -0,0 +1,106 @@
+// Copyright 1998-2014 Epic Games, Inc. All Rights Reserved.
+
+#pragma once
+
+class FSlateTextLayout;
+class FSlateTextLayoutMarshaller;
+class ISlateRunRenderer;
+
+#if WITH_FANCY_TEXT
+
+/** Class to handle the cached layout of STextBlock/SRichTextBlock by proxying around a FTextLayout */
+class FTextBlockLayout
+{
+public:
+	struct FWidgetArgs
+	{
+		FORCEINLINE FWidgetArgs(
+			const TAttribute<FText>& InText, 
+			const TAttribute<FText>& InHighlightText, 
+			const TAttribute<float>& InWrapTextAt, 
+			const TAttribute<bool>& InAutoWrapText, 
+			const TAttribute<FMargin>& InMargin, 
+			const TAttribute<float>& InLineHeightPercentage, 
+			const TAttribute<ETextJustify::Type>& InJustification
+			)
+			: Text(InText)
+			, HighlightText(InHighlightText)
+			, WrapTextAt(InWrapTextAt)
+			, AutoWrapText(InAutoWrapText)
+			, Margin(InMargin)
+			, LineHeightPercentage(InLineHeightPercentage)
+			, Justification(InJustification)
+		{
+		}
+
+		const TAttribute<FText>& Text;
+		const TAttribute<FText>& HighlightText;
+		const TAttribute<float>& WrapTextAt;
+		const TAttribute<bool>& AutoWrapText;
+		const TAttribute<FMargin>& Margin;
+		const TAttribute<float>& LineHeightPercentage;
+		const TAttribute<ETextJustify::Type>& Justification;
+	};
+
+	static TSharedRef<FTextBlockLayout> Create(TSharedRef<FSlateTextLayoutMarshaller> InMarshaller, TSharedPtr<IBreakIterator> InLineBreakPolicy);
+
+	/**
+	 * Get the computed desired size for this layout, updating the internal cache as required
+	 */
+	FVector2D ComputeDesiredSize(const FWidgetArgs& InWidgetArgs, const float InScale, const FTextBlockStyle& InTextStyle);
+
+	/**
+	 * Paint this layout, updating the internal cache as required
+	 */
+	int32 OnPaint(const FWidgetArgs& InWidgetArgs, const FPaintArgs& InPaintArgs, const FGeometry& InAllottedGeometry, const FSlateRect& InClippingRect, FSlateWindowElementList& OutDrawElements, int32 LayerId, const FWidgetStyle& InWidgetStyle, bool bParentEnabled);
+
+	/**
+	 * Force dirty the layout due to an external change that can't be picked up automatically by this cache
+	 */
+	void DirtyLayout();
+
+	/**
+	 * Get the child widgets of this layout
+	 */
+	FChildren* GetChildren();
+
+	/**
+	 * Arrange any child widgets in this layout
+	 */
+	void ArrangeChildren(const FGeometry& AllottedGeometry, FArrangedChildren& ArrangedChildren) const;
+
+private:
+	/** Updates the text layout to contain the given text */
+	void UpdateTextLayout(const FText& InText);
+
+	/** Update the text highlights */
+	void UpdateTextHighlights(const FText& InHighlightText);
+
+	/** Is the style used by the text marshaller up-to-date? */
+	bool IsStyleUpToDate(const FTextBlockStyle& NewStyle) const;
+
+	/** Calculate the wrapping width based on the given fixed wrap width, and whether we're auto-wrapping */
+	float CalculateWrappingWidth(const FWidgetArgs& InWidgetArgs) const;
+
+	FTextBlockLayout(TSharedRef<FSlateTextLayoutMarshaller> InMarshaller, TSharedPtr<IBreakIterator> InLineBreakPolicy);
+
+	/** In control of the layout and wrapping of the text */
+	TSharedPtr<FSlateTextLayout> TextLayout;
+
+	/** The marshaller used to get/set the text to/from the text layout. */
+	TSharedPtr<FSlateTextLayoutMarshaller> Marshaller;
+
+	/** Used to render the current highlights in the text layout */
+	TSharedPtr<ISlateRunRenderer> TextHighlighter;
+
+	/** The last known size of the layout from the previous OnPaint, used to guess at an auto-wrapping width in ComputeDesiredSize */
+	FVector2D CachedSize;
+
+	/** The state of the text the last time it was updated (used to allow updates when the text is changed) */
+	FTextSnapshot TextLastUpdate;
+
+	/** The state of the highlight text the last time it was updated (used to allow updates when the text is changed) */
+	FTextSnapshot HighlightTextLastUpdate;
+};
+
+#endif //WITH_FANCY_TEXT
diff --git a/Engine/Source/Runtime/Slate/Public/Widgets/Views/SHeaderRow.h b/Engine/Source/Runtime/Slate/Public/Widgets/Views/SHeaderRow.h
--- a/Engine/Source/Runtime/Slate/Public/Widgets/Views/SHeaderRow.h
+++ b/Engine/Source/Runtime/Slate/Public/Widgets/Views/SHeaderRow.h
@@ -171,7 +171,7 @@ public:
 		FName ColumnId;
 
 		/** Default text to use if no widget is passed in. */
-		TAttribute< FString > DefaultText;
+		TAttribute< FText > DefaultText;
 
 		/** A column width in Slate Units */
 		TAttribute< float > Width;
diff --git a/Engine/Source/Runtime/SlateCore/Private/Widgets/SWidget.cpp b/Engine/Source/Runtime/SlateCore/Private/Widgets/SWidget.cpp
--- a/Engine/Source/Runtime/SlateCore/Private/Widgets/SWidget.cpp
+++ b/Engine/Source/Runtime/SlateCore/Private/Widgets/SWidget.cpp
@@ -27,7 +27,7 @@ SWidget::SWidget()
 }
 
 void SWidget::Construct(
-	const TAttribute<FString> & InToolTipText ,
+	const TAttribute<FText> & InToolTipText ,
 	const TSharedPtr<IToolTip> & InToolTip ,
 	const TAttribute< TOptional<EMouseCursor::Type> > & InCursor ,
 	const TAttribute<bool> & InEnabledState ,
diff --git a/Engine/Source/Runtime/SlateCore/Public/Fonts/SlateFontInfo.h b/Engine/Source/Runtime/SlateCore/Public/Fonts/SlateFontInfo.h
--- a/Engine/Source/Runtime/SlateCore/Public/Fonts/SlateFontInfo.h
+++ b/Engine/Source/Runtime/SlateCore/Public/Fonts/SlateFontInfo.h
@@ -71,6 +71,18 @@ public:
 		return FontName == Other.FontName && Size == Other.Size;
 	}
 
+	/**
+	 * Compares this font info with another for inequality.
+	 *
+	 * @param Other The other font info.
+	 *
+	 * @return false if the two font infos are equal, true otherwise.
+	 */
+	bool operator!=( const FSlateFontInfo& Other ) const 
+	{
+		return !(*this == Other);
+	}
+
 	/**
 	 * Calculates a type hash value for a font info.
 	 *
diff --git a/Engine/Source/Runtime/SlateCore/Public/Styling/SlateBrush.h b/Engine/Source/Runtime/SlateCore/Public/Styling/SlateBrush.h
--- a/Engine/Source/Runtime/SlateCore/Public/Styling/SlateBrush.h
+++ b/Engine/Source/Runtime/SlateCore/Public/Styling/SlateBrush.h
@@ -223,6 +223,38 @@ public:
 		return bIsDynamicallyLoaded;
 	}
 
+	/**
+	 * Compares this brush with another for equality.
+	 *
+	 * @param Other The other brush.
+	 *
+	 * @return true if the two brushes are equal, false otherwise.
+	 */
+	bool operator==( const FSlateBrush& Other ) const 
+	{
+		return ImageSize == Other.ImageSize
+			&& DrawAs == Other.DrawAs
+			&& Margin == Other.Margin
+			&& TintColor == Other.TintColor
+			&& Tiling == Other.Tiling
+			&& Mirroring == Other.Mirroring
+			&& ResourceObject == Other.ResourceObject
+			&& ResourceName == Other.ResourceName
+			&& bIsDynamicallyLoaded == Other.bIsDynamicallyLoaded;
+	}
+
+	/**
+	 * Compares this brush with another for inequality.
+	 *
+	 * @param Other The other brush.
+	 *
+	 * @return false if the two brushes are equal, true otherwise.
+	 */
+	bool operator!=( const FSlateBrush& Other ) const 
+	{
+		return !(*this == Other);
+	}
+
 public:
 
 	/**
diff --git a/Engine/Source/Runtime/SlateCore/Public/Styling/SlateColor.h b/Engine/Source/Runtime/SlateCore/Public/Styling/SlateColor.h
--- a/Engine/Source/Runtime/SlateCore/Public/Styling/SlateColor.h
+++ b/Engine/Source/Runtime/SlateCore/Public/Styling/SlateColor.h
@@ -130,6 +130,32 @@ public:
 		return (ColorUseRule == ESlateColorStylingMode::UseColor_Specified) || (ColorUseRule == ESlateColorStylingMode::UseColor_Specified_Link);
 	}
 
+	/**
+	 * Compares this color with another for equality.
+	 *
+	 * @param Other The other color.
+	 *
+	 * @return true if the two colors are equal, false otherwise.
+	 */
+	bool operator==( const FSlateColor& Other ) const 
+	{
+		return SpecifiedColor == Other.SpecifiedColor
+			&& ColorUseRule == Other.ColorUseRule
+			&& (ColorUseRule != ESlateColorStylingMode::UseColor_Specified_Link || LinkedSpecifiedColor == Other.LinkedSpecifiedColor);
+	}
+
+	/**
+	 * Compares this color with another for inequality.
+	 *
+	 * @param Other The other color.
+	 *
+	 * @return false if the two colors are equal, true otherwise.
+	 */
+	bool operator!=( const FSlateColor& Other ) const 
+	{
+		return !(*this == Other);
+	}
+
 public:
 
 	/** @returns an FSlateColor that is the widget's foreground. */
diff --git a/Engine/Source/Runtime/SlateCore/Public/Widgets/DeclarativeSyntaxSupport.h b/Engine/Source/Runtime/SlateCore/Public/Widgets/DeclarativeSyntaxSupport.h
--- a/Engine/Source/Runtime/SlateCore/Public/Widgets/DeclarativeSyntaxSupport.h
+++ b/Engine/Source/Runtime/SlateCore/Public/Widgets/DeclarativeSyntaxSupport.h
@@ -277,203 +277,170 @@
 /**
  * Internal macro to declare overloads for setting attributes
  *
- * This macro expects 'WidgetArgsType' to be defined in the scope
- * where it is used.
+ * This macro expects 'WidgetArgsType' and 'AttrName##_Local' 
+ * to be defined in the scope where it is used.
  */
-#define INTERNAL_SLATE_DECL_ATTRIBUTE_TEXT_PASSTHROUGH( Var, AttrType, AttrName ) \
-		struct AttrName##_Local\
-		{\
-			static FString PassThroughDelegate( AttrType::FGetter TextDelegate )\
-			{\
-				return TextDelegate.Execute().ToString();\
-			}\
-			\
-			static FString PassThroughAttribute( AttrType TextAttribute )\
-			{\
-				return TextAttribute.Get( FText::GetEmpty() ).ToString();\
-			}\
-		};\
-		\
-		WidgetArgsType& AttrName( const FString& String ) \
-		{ \
-			Var = TAttribute< FString >( String ); \
-			return this->Me(); \
-		} \
-		WidgetArgsType& AttrName( const FText& Text ) \
-		{ \
-			Var = TAttribute< FString >( Text.ToString() ); \
-			return this->Me(); \
-		} \
-		WidgetArgsType& AttrName( const AttrType& InAttribute ) \
-		{ \
-			if ( InAttribute.IsBound() ) {\
-				Var = TAttribute< FString >::Create( TAttribute< FString >::FGetter::CreateStatic( &AttrName##_Local::PassThroughAttribute, InAttribute ) ); \
-			} else {\
-				Var = InAttribute.Get().ToString();\
-			}\
-			return this->Me(); \
-		} \
-	\
+#define INTERNAL_SLATE_DECL_ATTRIBUTE_PASSTHROUGH_COMMON( TargetAttrType, SourceAttrType, Var, AttrName ) \
 		/* Bind attribute with delegate to a global function */ \
 		/* NOTE: We use a template here to avoid 'typename' issues when hosting attributes inside templated classes */ \
-		WidgetArgsType& AttrName##_Static( AttrType::FGetter::FStaticDelegate::FFuncPtr InFunc )	\
+		WidgetArgsType& AttrName##_Static( SourceAttrType::FGetter::FStaticDelegate::FFuncPtr InFunc )	\
 		{ \
-			Var = TAttribute< FString >::Create( TAttribute< FString >::FGetter::CreateStatic( &AttrName##_Local::PassThroughDelegate, AttrType::FGetter::CreateStatic( InFunc ) ) ); \
+			Var = TargetAttrType::Create( TargetAttrType::FGetter::CreateStatic( &AttrName##_Local::PassThroughDelegate, SourceAttrType::FGetter::CreateStatic( InFunc ) ) ); \
 			return this->Me(); \
 		} \
 		template< typename Var1Type > \
-		WidgetArgsType& AttrName##_Static( typename AttrType::FGetter::template TStaticDelegate_OneVar< Var1Type >::FFuncPtr InFunc, Var1Type Var1 )	\
+		WidgetArgsType& AttrName##_Static( typename SourceAttrType::FGetter::template TStaticDelegate_OneVar< Var1Type >::FFuncPtr InFunc, Var1Type Var1 )	\
 		{ \
-			Var = TAttribute< FString >::Create( TAttribute< FString >::FGetter::CreateStatic( &AttrName##_Local::PassThroughDelegate, AttrType::FGetter::CreateStatic( InFunc, Var1 ) ) ); \
+			Var = TargetAttrType::Create( TargetAttrType::FGetter::CreateStatic( &AttrName##_Local::PassThroughDelegate, SourceAttrType::FGetter::CreateStatic( InFunc, Var1 ) ) ); \
 			return this->Me(); \
 		} \
 		template< typename Var1Type, typename Var2Type > \
-		WidgetArgsType& AttrName##_Static( typename AttrType::FGetter::template TStaticDelegate_TwoVars< Var1Type, Var2Type >::FFuncPtr InFunc, Var1Type Var1, Var2Type Var2 )	\
+		WidgetArgsType& AttrName##_Static( typename SourceAttrType::FGetter::template TStaticDelegate_TwoVars< Var1Type, Var2Type >::FFuncPtr InFunc, Var1Type Var1, Var2Type Var2 )	\
 		{ \
-			Var = TAttribute< FString >::Create( TAttribute< FString >::FGetter::CreateStatic( &AttrName##_Local::PassThroughDelegate, AttrType::FGetter::CreateStatic( InFunc, Var1, Var2 ) ) ); \
+			Var = TargetAttrType::Create( TargetAttrType::FGetter::CreateStatic( &AttrName##_Local::PassThroughDelegate, SourceAttrType::FGetter::CreateStatic( InFunc, Var1, Var2 ) ) ); \
 			return this->Me(); \
 		} \
 		template< typename Var1Type, typename Var2Type, typename Var3Type > \
-		WidgetArgsType& AttrName##_Static( typename AttrType::FGetter::template TStaticDelegate_ThreeVars< Var1Type, Var2Type, Var3Type >::FFuncPtr InFunc, Var1Type Var1, Var2Type Var2, Var3Type Var3 )	\
+		WidgetArgsType& AttrName##_Static( typename SourceAttrType::FGetter::template TStaticDelegate_ThreeVars< Var1Type, Var2Type, Var3Type >::FFuncPtr InFunc, Var1Type Var1, Var2Type Var2, Var3Type Var3 )	\
 		{ \
-			Var = TAttribute< FString >::Create( TAttribute< FString >::FGetter::CreateStatic( &AttrName##_Local::PassThroughDelegate, AttrType::FGetter::CreateStatic( InFunc, Var1, Var2, Var3 ) ) ); \
+			Var = TargetAttrType::Create( TargetAttrType::FGetter::CreateStatic( &AttrName##_Local::PassThroughDelegate, SourceAttrType::FGetter::CreateStatic( InFunc, Var1, Var2, Var3 ) ) ); \
 			return this->Me(); \
 		} \
 		template< typename Var1Type, typename Var2Type, typename Var3Type, typename Var4Type > \
-		WidgetArgsType& AttrName##_Static( typename AttrType::FGetter::template TStaticDelegate_FourVars< Var1Type, Var2Type, Var3Type, Var4Type >::FFuncPtr InFunc, Var1Type Var1, Var2Type Var2, Var3Type Var3, Var4Type Var4 )	\
+		WidgetArgsType& AttrName##_Static( typename SourceAttrType::FGetter::template TStaticDelegate_FourVars< Var1Type, Var2Type, Var3Type, Var4Type >::FFuncPtr InFunc, Var1Type Var1, Var2Type Var2, Var3Type Var3, Var4Type Var4 )	\
 		{ \
-			Var = TAttribute< FString >::Create( TAttribute< FString >::FGetter::CreateStatic( &AttrName##_Local::PassThroughDelegate, AttrType::FGetter::CreateStatic( InFunc, Var1, Var2, Var3, Var4 ) ) ); \
+			Var = TargetAttrType::Create( TargetAttrType::FGetter::CreateStatic( &AttrName##_Local::PassThroughDelegate, SourceAttrType::FGetter::CreateStatic( InFunc, Var1, Var2, Var3, Var4 ) ) ); \
 			return this->Me(); \
 		} \
 	\
 		/* Bind attribute with delegate to a raw C++ class method */ \
 		template< class UserClass >	\
-		WidgetArgsType& AttrName##_Raw( UserClass* InUserObject, typename AttrType::FGetter::template TRawMethodDelegate_Const< UserClass >::FMethodPtr InFunc )	\
+		WidgetArgsType& AttrName##_Raw( UserClass* InUserObject, typename SourceAttrType::FGetter::template TRawMethodDelegate_Const< UserClass >::FMethodPtr InFunc )	\
 		{ \
-			Var = TAttribute< FString >::Create( TAttribute< FString >::FGetter::CreateStatic( &AttrName##_Local::PassThroughDelegate, AttrType::FGetter::CreateRaw( InUserObject, InFunc ) ) ); \
+			Var = TargetAttrType::Create( TargetAttrType::FGetter::CreateStatic( &AttrName##_Local::PassThroughDelegate, SourceAttrType::FGetter::CreateRaw( InUserObject, InFunc ) ) ); \
 			return this->Me(); \
 		} \
 		template< class UserClass, typename Var1Type >	\
-		WidgetArgsType& AttrName##_Raw( UserClass* InUserObject, typename AttrType::FGetter::template TRawMethodDelegate_OneVar_Const< UserClass, Var1Type  >::FMethodPtr InFunc, Var1Type Var1 )	\
+		WidgetArgsType& AttrName##_Raw( UserClass* InUserObject, typename SourceAttrType::FGetter::template TRawMethodDelegate_OneVar_Const< UserClass, Var1Type  >::FMethodPtr InFunc, Var1Type Var1 )	\
 		{ \
-			Var = TAttribute< FString >::Create( TAttribute< FString >::FGetter::CreateStatic( &AttrName##_Local::PassThroughDelegate, AttrType::FGetter::CreateRaw( InUserObject, InFunc, Var1 ) ) ); \
+			Var = TargetAttrType::Create( TargetAttrType::FGetter::CreateStatic( &AttrName##_Local::PassThroughDelegate, SourceAttrType::FGetter::CreateRaw( InUserObject, InFunc, Var1 ) ) ); \
 			return this->Me(); \
 		} \
 		template< class UserClass, typename Var1Type, typename Var2Type >	\
-		WidgetArgsType& AttrName##_Raw( UserClass* InUserObject, typename AttrType::FGetter::template TRawMethodDelegate_TwoVars_Const< UserClass, Var1Type, Var2Type >::FMethodPtr InFunc, Var1Type Var1, Var2Type Var2 )	\
+		WidgetArgsType& AttrName##_Raw( UserClass* InUserObject, typename SourceAttrType::FGetter::template TRawMethodDelegate_TwoVars_Const< UserClass, Var1Type, Var2Type >::FMethodPtr InFunc, Var1Type Var1, Var2Type Var2 )	\
 		{ \
-			Var = TAttribute< FString >::Create( TAttribute< FString >::FGetter::CreateStatic( &AttrName##_Local::PassThroughDelegate, AttrType::FGetter::CreateRaw( InUserObject, InFunc, Var1, Var2 ) ) ); \
+			Var = TargetAttrType::Create( TargetAttrType::FGetter::CreateStatic( &AttrName##_Local::PassThroughDelegate, SourceAttrType::FGetter::CreateRaw( InUserObject, InFunc, Var1, Var2 ) ) ); \
 			return this->Me(); \
 		} \
 		template< class UserClass, typename Var1Type, typename Var2Type, typename Var3Type >	\
-		WidgetArgsType& AttrName##_Raw( UserClass* InUserObject, typename AttrType::FGetter::template TRawMethodDelegate_ThreeVars_Const< UserClass, Var1Type, Var2Type, Var3Type >::FMethodPtr InFunc, Var1Type Var1, Var2Type Var2, Var3Type Var3 )	\
+		WidgetArgsType& AttrName##_Raw( UserClass* InUserObject, typename SourceAttrType::FGetter::template TRawMethodDelegate_ThreeVars_Const< UserClass, Var1Type, Var2Type, Var3Type >::FMethodPtr InFunc, Var1Type Var1, Var2Type Var2, Var3Type Var3 )	\
 		{ \
-			Var = TAttribute< FString >::Create( TAttribute< FString >::FGetter::CreateStatic( &AttrName##_Local::PassThroughDelegate, AttrType::FGetter::CreateRaw( InUserObject, InFunc, Var1, Var2, Var3 ) ) ); \
+			Var = TargetAttrType::Create( TargetAttrType::FGetter::CreateStatic( &AttrName##_Local::PassThroughDelegate, SourceAttrType::FGetter::CreateRaw( InUserObject, InFunc, Var1, Var2, Var3 ) ) ); \
 			return this->Me(); \
 		} \
 		template< class UserClass, typename Var1Type, typename Var2Type, typename Var3Type, typename Var4Type >	\
-		WidgetArgsType& AttrName##_Raw( UserClass* InUserObject, typename AttrType::FGetter::template TRawMethodDelegate_FourVars_Const< UserClass, Var1Type, Var2Type, Var3Type, Var4Type >::FMethodPtr InFunc, Var1Type Var1, Var2Type Var2, Var3Type Var3, Var4Type Var4 )	\
+		WidgetArgsType& AttrName##_Raw( UserClass* InUserObject, typename SourceAttrType::FGetter::template TRawMethodDelegate_FourVars_Const< UserClass, Var1Type, Var2Type, Var3Type, Var4Type >::FMethodPtr InFunc, Var1Type Var1, Var2Type Var2, Var3Type Var3, Var4Type Var4 )	\
 		{ \
-			Var = TAttribute< FString >::Create( TAttribute< FString >::FGetter::CreateStatic( &AttrName##_Local::PassThroughDelegate, AttrType::FGetter::CreateRaw( InUserObject, InFunc, Var1, Var2, Var3, Var4 ) ) ); \
+			Var = TargetAttrType::Create( TargetAttrType::FGetter::CreateStatic( &AttrName##_Local::PassThroughDelegate, SourceAttrType::FGetter::CreateRaw( InUserObject, InFunc, Var1, Var2, Var3, Var4 ) ) ); \
 			return this->Me(); \
 		} \
 	\
 		/* Bind attribute with delegate to a shared pointer-based class method.  Slate mostly uses shared pointers so we use an overload for this type of binding. */ \
 		template< class UserClass >	\
-		WidgetArgsType& AttrName( TSharedRef< UserClass > InUserObjectRef, typename AttrType::FGetter::template TSPMethodDelegate_Const< UserClass >::FMethodPtr InFunc )	\
+		WidgetArgsType& AttrName( TSharedRef< UserClass > InUserObjectRef, typename SourceAttrType::FGetter::template TSPMethodDelegate_Const< UserClass >::FMethodPtr InFunc )	\
 		{ \
-			Var = TAttribute< FString >::Create( TAttribute< FString >::FGetter::CreateStatic( &AttrName##_Local::PassThroughDelegate, AttrType::FGetter::CreateSP( InUserObjectRef, InFunc ) ) ); \
+			Var = TargetAttrType::Create( TargetAttrType::FGetter::CreateStatic( &AttrName##_Local::PassThroughDelegate, SourceAttrType::FGetter::CreateSP( InUserObjectRef, InFunc ) ) ); \
 			return this->Me(); \
 		} \
 		template< class UserClass, typename Var1Type >	\
-		WidgetArgsType& AttrName( TSharedRef< UserClass > InUserObjectRef, typename AttrType::FGetter::template TSPMethodDelegate_OneVar_Const< UserClass, Var1Type >::FMethodPtr InFunc, Var1Type Var1 )	\
+		WidgetArgsType& AttrName( TSharedRef< UserClass > InUserObjectRef, typename SourceAttrType::FGetter::template TSPMethodDelegate_OneVar_Const< UserClass, Var1Type >::FMethodPtr InFunc, Var1Type Var1 )	\
 		{ \
-			Var = TAttribute< FString >::Create( TAttribute< FString >::FGetter::CreateStatic( &AttrName##_Local::PassThroughDelegate, AttrType::FGetter::CreateSP( InUserObjectRef, InFunc, Var1 ) ) ); \
+			Var = TargetAttrType::Create( TargetAttrType::FGetter::CreateStatic( &AttrName##_Local::PassThroughDelegate, SourceAttrType::FGetter::CreateSP( InUserObjectRef, InFunc, Var1 ) ) ); \
 			return this->Me(); \
 		} \
 		template< class UserClass, typename Var1Type, typename Var2Type >	\
-		WidgetArgsType& AttrName( TSharedRef< UserClass > InUserObjectRef, typename AttrType::FGetter::template TSPMethodDelegate_TwoVars_Const< UserClass, Var1Type, Var2Type >::FMethodPtr InFunc, Var1Type Var1, Var2Type Var2 )	\
+		WidgetArgsType& AttrName( TSharedRef< UserClass > InUserObjectRef, typename SourceAttrType::FGetter::template TSPMethodDelegate_TwoVars_Const< UserClass, Var1Type, Var2Type >::FMethodPtr InFunc, Var1Type Var1, Var2Type Var2 )	\
 		{ \
-			Var = TAttribute< FString >::Create( TAttribute< FString >::FGetter::CreateStatic( &AttrName##_Local::PassThroughDelegate, AttrType::FGetter::CreateSP( InUserObjectRef, InFunc, Var1, Var2 ) ) ); \
+			Var = TargetAttrType::Create( TargetAttrType::FGetter::CreateStatic( &AttrName##_Local::PassThroughDelegate, SourceAttrType::FGetter::CreateSP( InUserObjectRef, InFunc, Var1, Var2 ) ) ); \
 			return this->Me(); \
 		} \
 		template< class UserClass, typename Var1Type, typename Var2Type, typename Var3Type >	\
-		WidgetArgsType& AttrName( TSharedRef< UserClass > InUserObjectRef, typename AttrType::FGetter::template TSPMethodDelegate_ThreeVars_Const< UserClass, Var1Type, Var2Type, Var3Type >::FMethodPtr InFunc, Var1Type Var1, Var2Type Var2, Var3Type Var3 )	\
+		WidgetArgsType& AttrName( TSharedRef< UserClass > InUserObjectRef, typename SourceAttrType::FGetter::template TSPMethodDelegate_ThreeVars_Const< UserClass, Var1Type, Var2Type, Var3Type >::FMethodPtr InFunc, Var1Type Var1, Var2Type Var2, Var3Type Var3 )	\
 		{ \
-			Var = TAttribute< FString >::Create( TAttribute< FString >::FGetter::CreateStatic( &AttrName##_Local::PassThroughDelegate, AttrType::FGetter::CreateSP( InUserObjectRef, InFunc, Var1, Var2, Var3 ) ) ); \
+			Var = TargetAttrType::Create( TargetAttrType::FGetter::CreateStatic( &AttrName##_Local::PassThroughDelegate, SourceAttrType::FGetter::CreateSP( InUserObjectRef, InFunc, Var1, Var2, Var3 ) ) ); \
 			return this->Me(); \
 		} \
 		template< class UserClass, typename Var1Type, typename Var2Type, typename Var3Type, typename Var4Type >	\
-		WidgetArgsType& AttrName( TSharedRef< UserClass > InUserObjectRef, typename AttrType::FGetter::template TSPMethodDelegate_FourVars_Const< UserClass, Var1Type, Var2Type, Var3Type, Var4Type >::FMethodPtr InFunc, Var1Type Var1, Var2Type Var2, Var3Type Var3, Var4Type Var4 )	\
+		WidgetArgsType& AttrName( TSharedRef< UserClass > InUserObjectRef, typename SourceAttrType::FGetter::template TSPMethodDelegate_FourVars_Const< UserClass, Var1Type, Var2Type, Var3Type, Var4Type >::FMethodPtr InFunc, Var1Type Var1, Var2Type Var2, Var3Type Var3, Var4Type Var4 )	\
 		{ \
-			Var = TAttribute< FString >::Create( TAttribute< FString >::FGetter::CreateStatic( &AttrName##_Local::PassThroughDelegate, AttrType::FGetter::CreateSP( InUserObjectRef, InFunc, Var1, Var2, Var3, Var4 ) ) ); \
+			Var = TargetAttrType::Create( TargetAttrType::FGetter::CreateStatic( &AttrName##_Local::PassThroughDelegate, SourceAttrType::FGetter::CreateSP( InUserObjectRef, InFunc, Var1, Var2, Var3, Var4 ) ) ); \
 			return this->Me(); \
 		} \
 	\
 		/* Bind attribute with delegate to a shared pointer-based class method.  Slate mostly uses shared pointers so we use an overload for this type of binding. */ \
 		template< class UserClass >	\
-		WidgetArgsType& AttrName( UserClass* InUserObject, typename AttrType::FGetter::template TSPMethodDelegate_Const< UserClass >::FMethodPtr InFunc )	\
+		WidgetArgsType& AttrName( UserClass* InUserObject, typename SourceAttrType::FGetter::template TSPMethodDelegate_Const< UserClass >::FMethodPtr InFunc )	\
 		{ \
-			Var = TAttribute< FString >::Create( TAttribute< FString >::FGetter::CreateStatic( &AttrName##_Local::PassThroughDelegate, AttrType::FGetter::CreateSP( InUserObject, InFunc ) ) ); \
+			Var = TargetAttrType::Create( TargetAttrType::FGetter::CreateStatic( &AttrName##_Local::PassThroughDelegate, SourceAttrType::FGetter::CreateSP( InUserObject, InFunc ) ) ); \
 			return this->Me(); \
 		} \
 		template< class UserClass, typename Var1Type >	\
-		WidgetArgsType& AttrName( UserClass* InUserObject, typename AttrType::FGetter::template TSPMethodDelegate_OneVar_Const< UserClass, Var1Type >::FMethodPtr InFunc, Var1Type Var1 )	\
+		WidgetArgsType& AttrName( UserClass* InUserObject, typename SourceAttrType::FGetter::template TSPMethodDelegate_OneVar_Const< UserClass, Var1Type >::FMethodPtr InFunc, Var1Type Var1 )	\
 		{ \
-			Var = TAttribute< FString >::Create( TAttribute< FString >::FGetter::CreateStatic( &AttrName##_Local::PassThroughDelegate, AttrType::FGetter::CreateSP( InUserObject, InFunc, Var1 ) ) ); \
+			Var = TargetAttrType::Create( TargetAttrType::FGetter::CreateStatic( &AttrName##_Local::PassThroughDelegate, SourceAttrType::FGetter::CreateSP( InUserObject, InFunc, Var1 ) ) ); \
 			return this->Me(); \
 		} \
 		template< class UserClass, typename Var1Type, typename Var2Type >	\
-		WidgetArgsType& AttrName( UserClass* InUserObject, typename AttrType::FGetter::template TSPMethodDelegate_TwoVars_Const< UserClass, Var1Type, Var2Type >::FMethodPtr InFunc, Var1Type Var1, Var2Type Var2 )	\
+		WidgetArgsType& AttrName( UserClass* InUserObject, typename SourceAttrType::FGetter::template TSPMethodDelegate_TwoVars_Const< UserClass, Var1Type, Var2Type >::FMethodPtr InFunc, Var1Type Var1, Var2Type Var2 )	\
 		{ \
-			Var = TAttribute< FString >::Create( TAttribute< FString >::FGetter::CreateStatic( &AttrName##_Local::PassThroughDelegate, AttrType::FGetter::CreateSP( InUserObject, InFunc, Var1, Var2 ) ) ); \
+			Var = TargetAttrType::Create( TargetAttrType::FGetter::CreateStatic( &AttrName##_Local::PassThroughDelegate, SourceAttrType::FGetter::CreateSP( InUserObject, InFunc, Var1, Var2 ) ) ); \
 			return this->Me(); \
 		} \
 		template< class UserClass, typename Var1Type, typename Var2Type, typename Var3Type >	\
-		WidgetArgsType& AttrName( UserClass* InUserObject, typename AttrType::FGetter::template TSPMethodDelegate_ThreeVars_Const< UserClass, Var1Type, Var2Type, Var3Type >::FMethodPtr InFunc, Var1Type Var1, Var2Type Var2, Var3Type Var3 )	\
+		WidgetArgsType& AttrName( UserClass* InUserObject, typename SourceAttrType::FGetter::template TSPMethodDelegate_ThreeVars_Const< UserClass, Var1Type, Var2Type, Var3Type >::FMethodPtr InFunc, Var1Type Var1, Var2Type Var2, Var3Type Var3 )	\
 		{ \
-			Var = TAttribute< FString >::Create( TAttribute< FString >::FGetter::CreateStatic( &AttrName##_Local::PassThroughDelegate, AttrType::FGetter::CreateSP( InUserObject, InFunc, Var1, Var2, Var3 ) ) ); \
+			Var = TargetAttrType::Create( TargetAttrType::FGetter::CreateStatic( &AttrName##_Local::PassThroughDelegate, SourceAttrType::FGetter::CreateSP( InUserObject, InFunc, Var1, Var2, Var3 ) ) ); \
 			return this->Me(); \
 		} \
 		template< class UserClass, typename Var1Type, typename Var2Type, typename Var3Type, typename Var4Type >	\
-		WidgetArgsType& AttrName( UserClass* InUserObject, typename AttrType::FGetter::template TSPMethodDelegate_FourVars_Const< UserClass, Var1Type, Var2Type, Var3Type, Var4Type >::FMethodPtr InFunc, Var1Type Var1, Var2Type Var2, Var3Type Var3, Var4Type Var4 )	\
+		WidgetArgsType& AttrName( UserClass* InUserObject, typename SourceAttrType::FGetter::template TSPMethodDelegate_FourVars_Const< UserClass, Var1Type, Var2Type, Var3Type, Var4Type >::FMethodPtr InFunc, Var1Type Var1, Var2Type Var2, Var3Type Var3, Var4Type Var4 )	\
 		{ \
-			Var = TAttribute< FString >::Create( TAttribute< FString >::FGetter::CreateStatic( &AttrName##_Local::PassThroughDelegate, AttrType::FGetter::CreateSP( InUserObject, InFunc, Var1, Var2, Var3, Var4 ) ) ); \
+			Var = TargetAttrType::Create( TargetAttrType::FGetter::CreateStatic( &AttrName##_Local::PassThroughDelegate, SourceAttrType::FGetter::CreateSP( InUserObject, InFunc, Var1, Var2, Var3, Var4 ) ) ); \
 			return this->Me(); \
 		} \
 	\
 		/* Bind attribute with delegate to a UObject-based class method */ \
 		template< class UserClass >	\
-		WidgetArgsType& AttrName##_UObject( UserClass* InUserObject, typename AttrType::FGetter::template TUObjectMethodDelegate_Const< UserClass >::FMethodPtr InFunc )	\
+		WidgetArgsType& AttrName##_UObject( UserClass* InUserObject, typename SourceAttrType::FGetter::template TUObjectMethodDelegate_Const< UserClass >::FMethodPtr InFunc )	\
 		{ \
-			Var = TAttribute< FString >::Create( TAttribute< FString >::FGetter::CreateStatic( &AttrName##_Local::PassThroughDelegate, AttrType::FGetter::CreateUObject( InUserObject, InFunc ) ) ); \
+			Var = TargetAttrType::Create( TargetAttrType::FGetter::CreateStatic( &AttrName##_Local::PassThroughDelegate, SourceAttrType::FGetter::CreateUObject( InUserObject, InFunc ) ) ); \
 			return this->Me(); \
 		} \
 		template< class UserClass, typename Var1Type >	\
-		WidgetArgsType& AttrName##_UObject( UserClass* InUserObject, typename AttrType::FGetter::template TUObjectMethodDelegate_OneVar_Const< UserClass, Var1Type >::FMethodPtr InFunc, Var1Type Var1 )	\
+		WidgetArgsType& AttrName##_UObject( UserClass* InUserObject, typename SourceAttrType::FGetter::template TUObjectMethodDelegate_OneVar_Const< UserClass, Var1Type >::FMethodPtr InFunc, Var1Type Var1 )	\
 		{ \
-			Var = TAttribute< FString >::Create( TAttribute< FString >::FGetter::CreateStatic( &AttrName##_Local::PassThroughDelegate, AttrType::FGetter::CreateUObject( InUserObject, InFunc, Var1 ) ) ); \
+			Var = TargetAttrType::Create( TargetAttrType::FGetter::CreateStatic( &AttrName##_Local::PassThroughDelegate, SourceAttrType::FGetter::CreateUObject( InUserObject, InFunc, Var1 ) ) ); \
 			return this->Me(); \
 		} \
 		template< class UserClass, typename Var1Type, typename Var2Type  >	\
-		WidgetArgsType& AttrName##_UObject( UserClass* InUserObject, typename AttrType::FGetter::template TUObjectMethodDelegate_TwoVars_Const< UserClass, Var1Type, Var2Type >::FMethodPtr InFunc, Var1Type Var1, Var2Type Var2 )	\
+		WidgetArgsType& AttrName##_UObject( UserClass* InUserObject, typename SourceAttrType::FGetter::template TUObjectMethodDelegate_TwoVars_Const< UserClass, Var1Type, Var2Type >::FMethodPtr InFunc, Var1Type Var1, Var2Type Var2 )	\
 		{ \
-			Var = TAttribute< FString >::Create( TAttribute< FString >::FGetter::CreateStatic( &AttrName##_Local::PassThroughDelegate, AttrType::FGetter::CreateUObject( InUserObject, InFunc, Var1, Var2 ) ) ); \
+			Var = TargetAttrType::Create( TargetAttrType::FGetter::CreateStatic( &AttrName##_Local::PassThroughDelegate, SourceAttrType::FGetter::CreateUObject( InUserObject, InFunc, Var1, Var2 ) ) ); \
 			return this->Me(); \
 		} \
 		template< class UserClass, typename Var1Type, typename Var2Type, typename Var3Type >	\
-		WidgetArgsType& AttrName##_UObject( UserClass* InUserObject, typename AttrType::FGetter::template TUObjectMethodDelegate_ThreeVars_Const< UserClass, Var1Type, Var2Type, Var3Type >::FMethodPtr InFunc, Var1Type Var1, Var2Type Var2, Var3Type Var3 )	\
+		WidgetArgsType& AttrName##_UObject( UserClass* InUserObject, typename SourceAttrType::FGetter::template TUObjectMethodDelegate_ThreeVars_Const< UserClass, Var1Type, Var2Type, Var3Type >::FMethodPtr InFunc, Var1Type Var1, Var2Type Var2, Var3Type Var3 )	\
 		{ \
-			Var = TAttribute< FString >::Create( TAttribute< FString >::FGetter::CreateStatic( &AttrName##_Local::PassThroughDelegate, AttrType::FGetter::CreateUObject( InUserObject, InFunc, Var1, Var2, Var3 ) ) ); \
+			Var = TargetAttrType::Create( TargetAttrType::FGetter::CreateStatic( &AttrName##_Local::PassThroughDelegate, SourceAttrType::FGetter::CreateUObject( InUserObject, InFunc, Var1, Var2, Var3 ) ) ); \
 			return this->Me(); \
 		} \
 		template< class UserClass, typename Var1Type, typename Var2Type, typename Var3Type, typename Var4Type >	\
-		WidgetArgsType& AttrName##_UObject( UserClass* InUserObject, typename AttrType::FGetter::template TUObjectMethodDelegate_FourVars_Const< UserClass, Var1Type, Var2Type, Var3Type, Var4Type >::FMethodPtr InFunc, Var1Type Var1, Var2Type Var2, Var3Type Var3, Var4Type Var4 )	\
+		WidgetArgsType& AttrName##_UObject( UserClass* InUserObject, typename SourceAttrType::FGetter::template TUObjectMethodDelegate_FourVars_Const< UserClass, Var1Type, Var2Type, Var3Type, Var4Type >::FMethodPtr InFunc, Var1Type Var1, Var2Type Var2, Var3Type Var3, Var4Type Var4 )	\
 		{ \
-			Var = TAttribute< FString >::Create( TAttribute< FString >::FGetter::CreateStatic( &AttrName##_Local::PassThroughDelegate, AttrType::FGetter::CreateUObject( InUserObject, InFunc, Var1, Var2, Var3, Var4 ) ) ); \
+			Var = TargetAttrType::Create( TargetAttrType::FGetter::CreateStatic( &AttrName##_Local::PassThroughDelegate, SourceAttrType::FGetter::CreateUObject( InUserObject, InFunc, Var1, Var2, Var3, Var4 ) ) ); \
 			return this->Me(); \
 		}
-
+		
 
 /**
  * Use this macro to add a attribute to the declaration of your widget.
@@ -484,10 +451,88 @@
 	INTERNAL_SLATE_DECL_ATTRIBUTE( _##AttrName, TAttribute< AttrType >, AttrName )
 
 
-#define SLATE_TEXT_ATTRIBUTE( AttrName ) \
-	TAttribute< FString > _##AttrName; \
-	INTERNAL_SLATE_DECL_ATTRIBUTE_TEXT_PASSTHROUGH( _##AttrName, TAttribute< FText >, AttrName ) \
-	INTERNAL_SLATE_DECL_ATTRIBUTE( _##AttrName, TAttribute< FString >, AttrName )
+#define SLATE_STRING_ATTRIBUTE_IMPL( Var, AttrName ) \
+	TAttribute< FString > Var; \
+	struct AttrName##_Local\
+	{\
+		static FString PassThroughDelegate( TAttribute< FText >::FGetter TextDelegate )\
+		{\
+			return TextDelegate.Execute().ToString();\
+		}\
+		\
+		static FString PassThroughAttribute( TAttribute< FText > TextAttribute )\
+		{\
+			return TextAttribute.Get( FText::GetEmpty() ).ToString();\
+		}\
+	};\
+	WidgetArgsType& AttrName( const FString& String ) \
+	{ \
+		Var = TAttribute< FString >( String ); \
+		return this->Me(); \
+	} \
+	WidgetArgsType& AttrName( const FText& Text ) \
+	{ \
+		Var = TAttribute< FString >( Text.ToString() ); \
+		return this->Me(); \
+	} \
+	WidgetArgsType& AttrName( const TAttribute< FText >& InAttribute ) \
+	{ \
+		if ( InAttribute.IsBound() ) {\
+			Var = TAttribute< FString >::Create( TAttribute< FString >::FGetter::CreateStatic( &AttrName##_Local::PassThroughAttribute, InAttribute ) ); \
+		} else {\
+			Var = InAttribute.Get().ToString();\
+		}\
+		return this->Me(); \
+	} \
+	INTERNAL_SLATE_DECL_ATTRIBUTE_PASSTHROUGH_COMMON( TAttribute< FString >, TAttribute< FText >, Var, AttrName )\
+	INTERNAL_SLATE_DECL_ATTRIBUTE( Var, TAttribute< FString >, AttrName )
+
+/**
+ * Declare an FString attribute which also allows content passthrough from FText
+ */
+#define SLATE_STRING_ATTRIBUTE( AttrName ) SLATE_STRING_ATTRIBUTE_IMPL( _##AttrName, AttrName )
+
+
+#define SLATE_TEXT_ATTRIBUTE_IMPL( Var, AttrName ) \
+	TAttribute< FText > Var; \
+	struct AttrName##_Local\
+	{\
+		static FText PassThroughDelegate( TAttribute< FString >::FGetter StringDelegate )\
+		{\
+			return FText::FromString( StringDelegate.Execute() );\
+		}\
+		\
+		static FText PassThroughAttribute( TAttribute< FString > StringAttribute )\
+		{\
+			return FText::FromString( StringAttribute.Get( FString() ) );\
+		}\
+	};\
+	WidgetArgsType& AttrName( const FText& Text ) \
+	{ \
+		Var = TAttribute< FText >( Text ); \
+		return this->Me(); \
+	} \
+	WidgetArgsType& AttrName( const FString& String ) \
+	{ \
+		Var = TAttribute< FText >( FText::FromString( String ) ); \
+		return this->Me(); \
+	} \
+	WidgetArgsType& AttrName( const TAttribute< FString >& InAttribute ) \
+	{ \
+		if ( InAttribute.IsBound() ) {\
+			Var = TAttribute< FText >::Create( TAttribute< FText >::FGetter::CreateStatic( &AttrName##_Local::PassThroughAttribute, InAttribute ) ); \
+		} else {\
+			Var = FText::FromString( InAttribute.Get() );\
+		}\
+		return this->Me(); \
+	} \
+	INTERNAL_SLATE_DECL_ATTRIBUTE_PASSTHROUGH_COMMON( TAttribute< FText >, TAttribute< FString >, Var, AttrName )\
+	INTERNAL_SLATE_DECL_ATTRIBUTE( Var, TAttribute< FText >, AttrName )
+
+/**
+ * Declare an FText attribute which also allows content passthrough from FString
+ */
+#define SLATE_TEXT_ATTRIBUTE( AttrName ) SLATE_TEXT_ATTRIBUTE_IMPL( _##AttrName, AttrName )
 
 
 /**
diff --git a/Engine/Source/Runtime/SlateCore/Public/Widgets/SWidget.h b/Engine/Source/Runtime/SlateCore/Public/Widgets/SWidget.h
--- a/Engine/Source/Runtime/SlateCore/Public/Widgets/SWidget.h
+++ b/Engine/Source/Runtime/SlateCore/Public/Widgets/SWidget.h
@@ -67,7 +67,7 @@ public:
 	 * Construct a SWidget based on initial parameters.
 	 */
 	void Construct(
-		const TAttribute<FString> & InToolTipText ,
+		const TAttribute<FText> & InToolTipText ,
 		const TSharedPtr<IToolTip> & InToolTip ,
 		const TAttribute< TOptional<EMouseCursor::Type> > & InCursor ,
 		const TAttribute<bool> & InEnabledState ,
@@ -76,7 +76,7 @@ public:
 		const FName& InTag,
 		const TArray<TSharedRef<ISlateMetaData>>& InMetaData);
 
-	void SWidgetConstruct( const TAttribute<FString> & InToolTipText ,
+	void SWidgetConstruct( const TAttribute<FText> & InToolTipText ,
 		const TSharedPtr<IToolTip> & InToolTip ,
 		const TAttribute< TOptional<EMouseCursor::Type> > & InCursor ,
 		const TAttribute<bool> & InEnabledState ,
