Fixed crash when using SLATE_TEXT_ATTRIBUTE with a delegate that had become invalid

SLATE_TEXT_ATTRIBUTE acts as a passthrough from FString -> FText, however it wasn't checking that the delegate it was bound to was valid before calling it. I'd seen this intermittently before (when SLATE_TEXT_ATTRIBUTE used to passthrough from FText -> FString) however the changes in CL# 2280351 gave me a better repro case.

diff --git a/Engine/Source/Runtime/SlateCore/Public/Widgets/DeclarativeSyntaxSupport.h b/Engine/Source/Runtime/SlateCore/Public/Widgets/DeclarativeSyntaxSupport.h
--- a/Engine/Source/Runtime/SlateCore/Public/Widgets/DeclarativeSyntaxSupport.h
+++ b/Engine/Source/Runtime/SlateCore/Public/Widgets/DeclarativeSyntaxSupport.h
@@ -457,7 +457,7 @@
 	{\
 		static FString PassThroughDelegate( TAttribute< FText >::FGetter TextDelegate )\
 		{\
-			return TextDelegate.Execute().ToString();\
+			return (TextDelegate.IsBound()) ? TextDelegate.Execute().ToString() : FString();\
 		}\
 		\
 		static FString PassThroughAttribute( TAttribute< FText > TextAttribute )\
@@ -499,7 +499,7 @@
 	{\
 		static FText PassThroughDelegate( TAttribute< FString >::FGetter StringDelegate )\
 		{\
-			return FText::FromString( StringDelegate.Execute() );\
+			return (StringDelegate.IsBound()) ? FText::FromString( StringDelegate.Execute() ) : FText::GetEmpty();\
 		}\
 		\
 		static FText PassThroughAttribute( TAttribute< FString > StringAttribute )\
