Fix ES2 renderer for Slate as it doesn't necessarily understand FLOAT16.

diff --git a/Engine/Shaders/SlateVertexShader.usf b/Engine/Shaders/SlateVertexShader.usf
--- a/Engine/Shaders/SlateVertexShader.usf
+++ b/Engine/Shaders/SlateVertexShader.usf
@@ -21,8 +21,8 @@ VertexOut Main(
 */
 #if ES2_PROFILE && COMPILER_GLSL_ES2
 	in float2 InPosition : ATTRIBUTE1,
-	in half2 InClipOrigin : ATTRIBUTE2,
-	in half4 InClipExtents : ATTRIBUTE3,
+	in float2 InClipOrigin : ATTRIBUTE2,
+	in float4 InClipExtents : ATTRIBUTE3,
 #else
 	in int2 InPosition : ATTRIBUTE1,
 	in half2 InClipOrigin : ATTRIBUTE2,
diff --git a/Engine/Source/Developer/StandaloneRenderer/Private/OpenGL/SlateOpenGLRenderingPolicy.cpp b/Engine/Source/Developer/StandaloneRenderer/Private/OpenGL/SlateOpenGLRenderingPolicy.cpp
--- a/Engine/Source/Developer/StandaloneRenderer/Private/OpenGL/SlateOpenGLRenderingPolicy.cpp
+++ b/Engine/Source/Developer/StandaloneRenderer/Private/OpenGL/SlateOpenGLRenderingPolicy.cpp
@@ -245,13 +245,20 @@ void FSlateOpenGLRenderingPolicy::DrawElements( const FMatrix& ViewProjectionMat
 		Offset = STRUCT_OFFSET( FSlateVertex, Position );
 		glVertexAttribPointer( 1, 2, GL_SHORT, GL_FALSE, Stride, BUFFER_OFFSET(Stride*BaseVertexIndex+Offset) );
 
+		bool bUseFloat16 =
+#if SLATE_USE_FLOAT16
+			true;
+#else
+			false;
+#endif
+
 		glEnableVertexAttribArray(2);
 		Offset = STRUCT_OFFSET( FSlateVertex, ClipRect );
-		glVertexAttribPointer( 2, 2, GL_HALF_FLOAT, GL_FALSE, Stride, BUFFER_OFFSET(Stride*BaseVertexIndex+Offset) );
+		glVertexAttribPointer( 2, 2, bUseFloat16 ? GL_HALF_FLOAT : GL_FLOAT, GL_FALSE, Stride, BUFFER_OFFSET(Stride*BaseVertexIndex+Offset) );
 
 		glEnableVertexAttribArray(3);
-		Offset = STRUCT_OFFSET( FSlateVertex, ClipRect ) + STRUCT_OFFSET( FSlateRotatedRectFloat16, ExtentX );
-		glVertexAttribPointer( 3, 4, GL_HALF_FLOAT, GL_FALSE, Stride, BUFFER_OFFSET(Stride*BaseVertexIndex+Offset) );
+		Offset = STRUCT_OFFSET( FSlateVertex, ClipRect ) + STRUCT_OFFSET( FSlateClipRectType, ExtentX );
+		glVertexAttribPointer( 3, 4, bUseFloat16 ? GL_HALF_FLOAT : GL_FLOAT, GL_FALSE, Stride, BUFFER_OFFSET(Stride*BaseVertexIndex+Offset) );
 
 		glEnableVertexAttribArray(4);
 		Offset = STRUCT_OFFSET( FSlateVertex, Color );
diff --git a/Engine/Source/Runtime/SlateCore/Private/Rendering/ElementBatcher.cpp b/Engine/Source/Runtime/SlateCore/Private/Rendering/ElementBatcher.cpp
--- a/Engine/Source/Runtime/SlateCore/Private/Rendering/ElementBatcher.cpp
+++ b/Engine/Source/Runtime/SlateCore/Private/Rendering/ElementBatcher.cpp
@@ -4,6 +4,9 @@
 #include "ElementBatcher.h"
 #include "SlateRenderTransform.h"
 
+/** Default ctor. */
+FSlateRotatedRect::FSlateRotatedRect() {}
+
 FSlateRotatedRect::FSlateRotatedRect(const FSlateLayoutTransform& InverseLayoutTransform, const FSlateRenderTransform& RenderTransform, const FSlateRect& ClipRectInLayoutWindowSpace)
 {
 	const FSlateRenderTransform LayoutWindowSpaceToRenderWindowSpace = Concatenate(InverseLayoutTransform, RenderTransform);
@@ -53,7 +56,10 @@ bool FSlateRotatedRect::IsUnderLocation(const FVector2D& Location) const
 	return false;
 }
 
+/** Default ctor. */
 FSlateRotatedRectFloat16::FSlateRotatedRectFloat16() {}
+
+/** Construct a float16 version of a rotated rect from a full-float version. */
 FSlateRotatedRectFloat16::FSlateRotatedRectFloat16(const FSlateRotatedRect& RotatedRect)
 {
 	TopLeft[0] = RotatedRect.TopLeft[0];
@@ -64,9 +70,22 @@ FSlateRotatedRectFloat16::FSlateRotatedRectFloat16(const FSlateRotatedRect& Rota
 	ExtentY[1] = RotatedRect.ExtentY[1];
 }
 
+/** Used by the typedef magic to construct a rotated rect of the type needed by the RHI being used. */
+FSlateRotatedRectFloat16::FSlateRotatedRectFloat16(const FSlateLayoutTransform& InverseLayoutTransform, const FSlateRenderTransform& RenderTransform, const FSlateRect& ClipRectInLayoutWindowSpace)
+{
+	FSlateRotatedRect RotatedRect(InverseLayoutTransform, RenderTransform, ClipRectInLayoutWindowSpace);
+	TopLeft[0] = RotatedRect.TopLeft[0];
+	TopLeft[1] = RotatedRect.TopLeft[1];
+	ExtentX[0] = RotatedRect.ExtentX[0];
+	ExtentX[1] = RotatedRect.ExtentX[1];
+	ExtentY[0] = RotatedRect.ExtentY[0];
+	ExtentY[1] = RotatedRect.ExtentY[1];
+}
+
+
 FSlateVertex::FSlateVertex() {}
 
-FSlateVertex::FSlateVertex(const FSlateRenderTransform& RenderTransform, const FVector2D& InLocalPosition, const FVector2D& InTexCoord, const FVector2D& InTexCoord2, const FColor& InColor, const FSlateRotatedRectFloat16& InClipRect )
+FSlateVertex::FSlateVertex(const FSlateRenderTransform& RenderTransform, const FVector2D& InLocalPosition, const FVector2D& InTexCoord, const FVector2D& InTexCoord2, const FColor& InColor, const FSlateClipRectType& InClipRect )
 	: TexCoords( InTexCoord, InTexCoord2 )
 	, ClipRect( InClipRect )
 	, Color( InColor )
@@ -76,7 +95,7 @@ FSlateVertex::FSlateVertex(const FSlateRenderTransform& RenderTransform, const F
 	Position[1] = FMath::RoundToInt(WindowPosition.Y);
 }
 
-FSlateVertex::FSlateVertex( const FSlateRenderTransform& RenderTransform, const FVector2D& InLocalPosition, const FVector2D& InTexCoord, const FColor& InColor, const FSlateRotatedRectFloat16& InClipRect )
+FSlateVertex::FSlateVertex( const FSlateRenderTransform& RenderTransform, const FVector2D& InLocalPosition, const FVector2D& InTexCoord, const FColor& InColor, const FSlateClipRectType& InClipRect )
 	: TexCoords( InTexCoord, FVector2D(1.0f, 1.0f) )
 	, ClipRect( InClipRect )
 	, Color( InColor )
@@ -186,7 +205,7 @@ void FSlateElementBatcher::AddQuadElement( const FSlateDrawElement& DrawElement,
 
 	// extract the layout transform from the draw element
 	FSlateLayoutTransform InverseLayoutTransform(Inverse(FSlateLayoutTransform(DrawElement.GetScale(), DrawElement.GetPosition())));
-	FSlateRotatedRectFloat16 RenderClipRect = FSlateRotatedRect(InverseLayoutTransform, RenderTransform, InClippingRect);
+	FSlateClipRectType RenderClipRect = FSlateClipRectType(InverseLayoutTransform, RenderTransform, InClippingRect);
 
 	FSlateElementBatch& ElementBatch = FindBatchForElement( Layer, FShaderParams(), nullptr, ESlateDrawPrimitive::TriangleList, ESlateShader::Default, ESlateDrawEffect::None, ESlateBatchDrawFlag::Wireframe|ESlateBatchDrawFlag::NoBlending );
 	TArray<FSlateVertex>& BatchVertices = BatchVertexArrays[ElementBatch.VertexArrayIndex];
@@ -251,7 +270,7 @@ void FSlateElementBatcher::AddBoxElement( const FSlateDrawElement& DrawElement )
 	// extract the layout transform from the draw element
 	FSlateLayoutTransform InverseLayoutTransform(Inverse(FSlateLayoutTransform(DrawElement.GetScale(), DrawElement.GetPosition())));
 	// The clip rect is NOT subject to the rotations specified by MakeRotatedBox.
-	FSlateRotatedRectFloat16 RenderClipRect = FSlateRotatedRect(InverseLayoutTransform, ElementRenderTransform, InClippingRect);
+	FSlateClipRectType RenderClipRect = FSlateClipRectType(InverseLayoutTransform, ElementRenderTransform, InClippingRect);
 
 	check( InPayload.BrushResource );
 
@@ -550,7 +569,7 @@ void FSlateElementBatcher::AddTextElement(const FSlateDrawElement& DrawElement)
 	FSlateLayoutTransform InverseLayoutTransform = Inverse(Concatenate(Inverse(FontScale), LayoutTransform));
 	const FSlateRenderTransform RenderTransform = Concatenate(Inverse(FontScale), DrawElement.GetRenderTransform());
 
-	FSlateRotatedRectFloat16 RenderClipRect = FSlateRotatedRect(InverseLayoutTransform, RenderTransform, InClippingRect);
+	FSlateClipRectType RenderClipRect = FSlateClipRectType(InverseLayoutTransform, RenderTransform, InClippingRect);
 	// Used to clip individual characters as we generate them.
 	FSlateRect LocalClipRect = TransformRect(InverseLayoutTransform, InClippingRect);
 
@@ -713,7 +732,7 @@ void FSlateElementBatcher::AddGradientElement( const FSlateDrawElement& DrawElem
 
 	// extract the layout transform from the draw element
 	FSlateLayoutTransform InverseLayoutTransform(Inverse(FSlateLayoutTransform(DrawElement.GetScale(), DrawElement.GetPosition())));
-	FSlateRotatedRectFloat16 RenderClipRect = FSlateRotatedRect(InverseLayoutTransform, RenderTransform, InClippingRect);
+	FSlateClipRectType RenderClipRect = FSlateClipRectType(InverseLayoutTransform, RenderTransform, InClippingRect);
 
 	// There must be at least one gradient stop
 	check( InPayload.GradientStops.Num() > 0 );
@@ -840,7 +859,7 @@ void FSlateElementBatcher::AddSplineElement( const FSlateDrawElement& DrawElemen
 
 	// extract the layout transform from the draw element
 	FSlateLayoutTransform InverseLayoutTransform(Inverse(FSlateLayoutTransform(DrawElement.GetScale(), DrawElement.GetPosition())));
-	FSlateRotatedRectFloat16 RenderClipRect = FSlateRotatedRect(InverseLayoutTransform, RenderTransform, InClippingRect);
+	FSlateClipRectType RenderClipRect = FSlateClipRectType(InverseLayoutTransform, RenderTransform, InClippingRect);
 
 	//@todo SLATE: Merge with AddLineElement?
 
@@ -986,7 +1005,7 @@ void FSlateElementBatcher::AddLineElement( const FSlateDrawElement& DrawElement
 
 	// extract the layout transform from the draw element
 	FSlateLayoutTransform InverseLayoutTransform(Inverse(FSlateLayoutTransform(DrawElement.GetScale(), DrawElement.GetPosition())));
-	FSlateRotatedRectFloat16 RenderClipRect = FSlateRotatedRect(InverseLayoutTransform, RenderTransform, InClippingRect);
+	FSlateClipRectType RenderClipRect = FSlateClipRectType(InverseLayoutTransform, RenderTransform, InClippingRect);
 
 	if( InPayload.Points.Num() < 2 )
 	{
@@ -1164,7 +1183,7 @@ void FSlateElementBatcher::AddViewportElement( const FSlateDrawElement& DrawElem
 
 	// extract the layout transform from the draw element
 	FSlateLayoutTransform InverseLayoutTransform(Inverse(FSlateLayoutTransform(DrawElement.GetScale(), DrawElement.GetPosition())));
-	FSlateRotatedRectFloat16 RenderClipRect = FSlateRotatedRect(InverseLayoutTransform, RenderTransform, InClippingRect);
+	FSlateClipRectType RenderClipRect = FSlateClipRectType(InverseLayoutTransform, RenderTransform, InClippingRect);
 
 	const FColor FinalColor = GetElementColor( InPayload.Tint, nullptr );
 
@@ -1242,7 +1261,7 @@ void FSlateElementBatcher::AddBorderElement( const FSlateDrawElement& DrawElemen
 
 	// extract the layout transform from the draw element
 	FSlateLayoutTransform InverseLayoutTransform(Inverse(FSlateLayoutTransform(DrawElement.GetScale(), DrawElement.GetPosition())));
-	FSlateRotatedRectFloat16 RenderClipRect = FSlateRotatedRect(InverseLayoutTransform, RenderTransform, InClippingRect);
+	FSlateClipRectType RenderClipRect = FSlateClipRectType(InverseLayoutTransform, RenderTransform, InClippingRect);
 
 	check( InPayload.BrushResource );
 
diff --git a/Engine/Source/Runtime/SlateCore/Public/Rendering/RenderingCommon.h b/Engine/Source/Runtime/SlateCore/Public/Rendering/RenderingCommon.h
--- a/Engine/Source/Runtime/SlateCore/Public/Rendering/RenderingCommon.h
+++ b/Engine/Source/Runtime/SlateCore/Public/Rendering/RenderingCommon.h
@@ -6,6 +6,7 @@ struct FVector2D;
 class FSlateRect;
 
 #define SLATE_USE_32BIT_INDICES !PLATFORM_USES_ES2
+#define SLATE_USE_FLOAT16 !PLATFORM_USES_ES2
 
 #if SLATE_USE_32BIT_INDICES
 typedef uint32 SlateIndex;
@@ -94,6 +95,7 @@ namespace ESlateLineJoinType
  */
 struct FSlateRotatedRect
 {
+	FSlateRotatedRect();
 	FSlateRotatedRect(const FSlateLayoutTransform& InverseLayoutTransform, const FSlateRenderTransform& RenderTransform, const FSlateRect& ClipRectInLayoutWindowSpace);
 	/** transformed Top-left corner. */
 	FVector2D TopLeft;
@@ -113,6 +115,7 @@ struct FSlateRotatedRectFloat16
 {
 	FSlateRotatedRectFloat16();
 	FSlateRotatedRectFloat16(const FSlateRotatedRect& RotatedRect);
+	FSlateRotatedRectFloat16(const FSlateLayoutTransform& InverseLayoutTransform, const FSlateRenderTransform& RenderTransform, const FSlateRect& ClipRectInLayoutWindowSpace);
 	/** transformed Top-left corner. */
 	FFloat16 TopLeft[2];
 	/** transformed X extent (right-left). */
@@ -121,6 +124,16 @@ struct FSlateRotatedRectFloat16
 	FFloat16 ExtentY[2];
 };
 
+/**
+ * Not all platforms support Float16, so we have to be tricky here and declare the proper vertex type.
+ */
+#if SLATE_USE_FLOAT16
+typedef FSlateRotatedRectFloat16 FSlateClipRectType;
+#else
+typedef FSlateRotatedRect FSlateClipRectType;
+#endif
+
+
 /** 
  * A struct which defines a basic vertex seen by the Slate vertex buffers and shaders
  */
@@ -131,13 +144,13 @@ struct FSlateVertex
 	/** Position of the vertex in window space */
 	int16 Position[2];
 	/** clip center/extents in render window space (window space with render transforms applied) */
-	FSlateRotatedRectFloat16 ClipRect;
+	FSlateClipRectType ClipRect;
 	/** Vertex color */
 	FColor Color;
 	
 	FSlateVertex();
-	FSlateVertex( const FSlateRenderTransform& RenderTransform, const FVector2D& InLocalPosition, const FVector2D& InTexCoord, const FVector2D& InTexCoord2, const FColor& InColor, const FSlateRotatedRectFloat16& InClipRect );
-	FSlateVertex( const FSlateRenderTransform& RenderTransform, const FVector2D& InLocalPosition, const FVector2D& InTexCoord, const FColor& InColor, const FSlateRotatedRectFloat16& InClipRect );
+	FSlateVertex( const FSlateRenderTransform& RenderTransform, const FVector2D& InLocalPosition, const FVector2D& InTexCoord, const FVector2D& InTexCoord2, const FColor& InColor, const FSlateClipRectType& InClipRect );
+	FSlateVertex( const FSlateRenderTransform& RenderTransform, const FVector2D& InLocalPosition, const FVector2D& InTexCoord, const FColor& InColor, const FSlateClipRectType& InClipRect );
 };
 
 template<> struct TIsPODType<FSlateVertex> { enum { Value = true }; };
diff --git a/Engine/Source/Runtime/SlateRHIRenderer/Private/SlateShaders.cpp b/Engine/Source/Runtime/SlateRHIRenderer/Private/SlateShaders.cpp
--- a/Engine/Source/Runtime/SlateRHIRenderer/Private/SlateShaders.cpp
+++ b/Engine/Source/Runtime/SlateRHIRenderer/Private/SlateShaders.cpp
@@ -45,8 +45,14 @@ void FSlateVertexDeclaration::InitRHI()
 	uint32 Stride = sizeof(FSlateVertex);
 	Elements.Add(FVertexElement(0,STRUCT_OFFSET(FSlateVertex,TexCoords),VET_Float4,0,Stride));
 	Elements.Add(FVertexElement(0,STRUCT_OFFSET(FSlateVertex,Position),VET_Short2,1,Stride));
-	Elements.Add(FVertexElement(0,STRUCT_OFFSET(FSlateVertex,ClipRect),VET_Half2,2,Stride));
-	Elements.Add(FVertexElement(0,STRUCT_OFFSET(FSlateVertex,ClipRect)+STRUCT_OFFSET(FSlateRotatedRectFloat16,ExtentX),VET_Half4,3,Stride));
+	bool bUseFloat16 =
+#if SLATE_USE_FLOAT16
+		true;
+#else
+		false;
+#endif
+	Elements.Add(FVertexElement(0,STRUCT_OFFSET(FSlateVertex,ClipRect),bUseFloat16 ? VET_Half2 : VET_Float2,2,Stride));
+	Elements.Add(FVertexElement(0,STRUCT_OFFSET(FSlateVertex,ClipRect)+STRUCT_OFFSET(FSlateClipRectType,ExtentX),bUseFloat16 ? VET_Half4 : VET_Float4,3,Stride));
 	Elements.Add(FVertexElement(0,STRUCT_OFFSET(FSlateVertex,Color),VET_Color,4,Stride));
 
 	VertexDeclarationRHI = RHICreateVertexDeclaration(Elements);
