UMG - Making several slate structs play nice with the blueprint compiler and not warn about being removed.

diff --git a/Engine/Source/Runtime/SlateCore/Public/Fonts/SlateFontInfo.h b/Engine/Source/Runtime/SlateCore/Public/Fonts/SlateFontInfo.h
--- a/Engine/Source/Runtime/SlateCore/Public/Fonts/SlateFontInfo.h
+++ b/Engine/Source/Runtime/SlateCore/Public/Fonts/SlateFontInfo.h
@@ -8,17 +8,17 @@
 /**
  * A representation of a font in Slate.
  */
-USTRUCT()
+USTRUCT(BlueprintType)
 struct SLATECORE_API FSlateFontInfo
 {
 	GENERATED_USTRUCT_BODY()
 
 	/** The name of the font */
-	UPROPERTY(EditAnywhere, Category=SlateStyleRules)
+	UPROPERTY(EditAnywhere, BlueprintReadWrite, Category=SlateStyleRules)
 	FName FontName;
 
 	/** The size of the font */
-	UPROPERTY(EditAnywhere, Category=SlateStyleRules)
+	UPROPERTY(EditAnywhere, BlueprintReadWrite, Category=SlateStyleRules)
 	int32 Size;
 
 public:
diff --git a/Engine/Source/Runtime/SlateCore/Public/Input/Events.h b/Engine/Source/Runtime/SlateCore/Public/Input/Events.h
--- a/Engine/Source/Runtime/SlateCore/Public/Input/Events.h
+++ b/Engine/Source/Runtime/SlateCore/Public/Input/Events.h
@@ -85,7 +85,7 @@ private:
 /**
  * Base class for all mouse and keyboard events.
  */
-USTRUCT()
+USTRUCT(BlueprintType)
 struct FInputEvent
 {
 	GENERATED_USTRUCT_BODY()
@@ -284,7 +284,7 @@ protected:
  * FKeyboardEvent describes a keyboard action (key pressed or released.)
  * It is passed to event handlers dealing with keyboard input.
  */
-USTRUCT()
+USTRUCT(BlueprintType)
 struct FKeyboardEvent
 	: public FInputEvent
 {
@@ -421,7 +421,7 @@ public:
  * FPointerEvent describes a mouse or touch action (e.g. Press, Release, Move, etc).
  * It is passed to event handlers dealing with pointer-based input.
  */
-USTRUCT()
+USTRUCT(BlueprintType)
 struct FPointerEvent
 	: public FInputEvent
 {
@@ -608,7 +608,7 @@ private:
  * FControllerEvent describes a controller action (e.g. Button Press, Release, Analog stick move, etc).
  * It is passed to event handlers dealing with controller input.
  */
-USTRUCT()
+USTRUCT(BlueprintType)
 struct FControllerEvent
 	: public FInputEvent
 {
@@ -659,7 +659,7 @@ private:
  * FMotionEvent describes a touch pad action (press, move, lift)
  * It is passed to event handlers dealing with touch input.
  */
-USTRUCT()
+USTRUCT(BlueprintType)
 struct FMotionEvent
 	: public FInputEvent
 {
diff --git a/Engine/Source/Runtime/SlateCore/Public/Layout/Geometry.h b/Engine/Source/Runtime/SlateCore/Public/Layout/Geometry.h
--- a/Engine/Source/Runtime/SlateCore/Public/Layout/Geometry.h
+++ b/Engine/Source/Runtime/SlateCore/Public/Layout/Geometry.h
@@ -20,7 +20,7 @@ class FLayoutGeometry;
  * A Geometry's parent is generally thought to be the Geometry of the
  * the corresponding parent widget.
  */
-USTRUCT()
+USTRUCT(BlueprintType)
 struct SLATECORE_API FGeometry
 {
 	GENERATED_USTRUCT_BODY()
diff --git a/Engine/Source/Runtime/SlateCore/Public/Styling/SlateBrush.h b/Engine/Source/Runtime/SlateCore/Public/Styling/SlateBrush.h
--- a/Engine/Source/Runtime/SlateCore/Public/Styling/SlateBrush.h
+++ b/Engine/Source/Runtime/SlateCore/Public/Styling/SlateBrush.h
@@ -100,21 +100,21 @@ namespace ESlateBrushImageType
 /**
  * An brush which contains information about how to draw a Slate element
  */
-USTRUCT()
+USTRUCT(BlueprintType)
 struct SLATECORE_API FSlateBrush
 {
 	GENERATED_USTRUCT_BODY()
 
 	/** Size of the resource in Slate Units */
-	UPROPERTY(EditAnywhere, Category=Brush)
+	UPROPERTY(EditAnywhere, BlueprintReadWrite, Category=Brush)
 	FVector2D ImageSize;
 
 	/** How to draw the image */
-	UPROPERTY(EditAnywhere, Category=Brush)
+	UPROPERTY(EditAnywhere, BlueprintReadWrite, Category=Brush)
 	TEnumAsByte<enum ESlateBrushDrawType::Type > DrawAs;
 
 	/** The margin to use in Box and Border modes */
-	UPROPERTY(EditAnywhere, Category=Brush, meta=(UVSpace="true"))
+	UPROPERTY(EditAnywhere, BlueprintReadWrite, Category=Brush, meta=( UVSpace="true" ))
 	FMargin Margin;
 
 	/** Tinting applied to the image. */
@@ -122,16 +122,16 @@ struct SLATECORE_API FSlateBrush
 	FLinearColor Tint_DEPRECATED;
 
 	/** Tinting applied to the image. */
-	UPROPERTY(EditAnywhere, Category=Brush, meta=(DisplayName="Tint"))
+	UPROPERTY(EditAnywhere, BlueprintReadWrite, Category=Brush, meta=( DisplayName="Tint" ))
 	FSlateColor TintColor;
 
 	/** How to tile the image in Image mode */
-	UPROPERTY(EditAnywhere, Category=Brush)
+	UPROPERTY(EditAnywhere, BlueprintReadWrite, Category=Brush)
 	TEnumAsByte<enum ESlateBrushTileType::Type> Tiling;
 
 	/** How to mirror the image in Image mode.  This is normally only used for dynamic image brushes where the source texture
 	    comes from a hardware device such as a web camera. */
-	UPROPERTY(EditAnywhere, Category=Brush)
+	UPROPERTY(EditAnywhere, BlueprintReadWrite, Category=Brush)
 	TEnumAsByte<enum ESlateBrushMirrorType::Type> Mirroring;
 
 	/** The type of image */
@@ -266,8 +266,10 @@ public:
 
 protected:
 
-	/** Pointer to the UTexture2D. Holding onto it as a UObject because UTexture2D is not available in Slate core. */
-	UPROPERTY(EditAnywhere, Category=Brush, meta=(DisplayThumbnail="true", DisplayName="Texture or Material Asset"))
+	/**
+	 * The image to render for this brush, can be a UTexture2D or Material.
+	 */
+	UPROPERTY(EditAnywhere, BlueprintReadWrite, Category=Brush, meta=( DisplayThumbnail="true", ThumbnailSize="X=40 Y=40", DisplayName="Image" ))
 	UObject* ResourceObject;
 
 	/** The name of the rendering resource to use */
diff --git a/Engine/Source/Runtime/SlateCore/Public/Styling/SlateColor.h b/Engine/Source/Runtime/SlateCore/Public/Styling/SlateColor.h
--- a/Engine/Source/Runtime/SlateCore/Public/Styling/SlateColor.h
+++ b/Engine/Source/Runtime/SlateCore/Public/Styling/SlateColor.h
@@ -33,7 +33,7 @@ namespace ESlateColorStylingMode
 /**
  * A Slate color can be a directly specified value, or the color can be pulled from a WidgetStyle.
  */
-USTRUCT()
+USTRUCT(BlueprintType)
 struct FSlateColor
 {
 	GENERATED_USTRUCT_BODY()
@@ -181,11 +181,11 @@ protected:
 protected:
 
 	// The current specified color; only meaningful when ColorToUse == UseColor_Specified.
-	UPROPERTY(EditAnywhere, Category=Color)
+	UPROPERTY(EditAnywhere, BlueprintReadWrite, Category=Color)
 	FLinearColor SpecifiedColor;
 
 	// The rule for which color to pick.
-	UPROPERTY(EditAnywhere, Category=Color)
+	UPROPERTY(EditAnywhere, BlueprintReadWrite, Category=Color)
 	TEnumAsByte<ESlateColorStylingMode::Type> ColorUseRule;
 
 private:
diff --git a/Engine/Source/Runtime/SlateCore/Public/Styling/SlateTypes.h b/Engine/Source/Runtime/SlateCore/Public/Styling/SlateTypes.h
--- a/Engine/Source/Runtime/SlateCore/Public/Styling/SlateTypes.h
+++ b/Engine/Source/Runtime/SlateCore/Public/Styling/SlateTypes.h
@@ -26,7 +26,7 @@ namespace ESlateCheckBoxType
 /**
  * Represents the appearance of an SCheckBox
  */
-USTRUCT()
+USTRUCT(BlueprintType)
 struct SLATECORE_API FCheckBoxStyle : public FSlateWidgetStyle
 {
 	GENERATED_USTRUCT_BODY()
@@ -43,88 +43,88 @@ struct SLATECORE_API FCheckBoxStyle : public FSlateWidgetStyle
 	static const FCheckBoxStyle& GetDefault();
 
 	/** The visual type of the checkbox */	
-	UPROPERTY(EditAnywhere, Category=Appearance)
+	UPROPERTY( EditAnywhere, BlueprintReadWrite, Category=Appearance )
 	TEnumAsByte<ESlateCheckBoxType::Type> CheckBoxType;
 	FCheckBoxStyle& SetCheckBoxType( ESlateCheckBoxType::Type InCheckBoxType ){ CheckBoxType = InCheckBoxType; return *this; }
 
 	/* CheckBox appearance when the CheckBox is unchecked (normal) */
-	UPROPERTY( EditAnywhere, Category = Appearance )
+	UPROPERTY( EditAnywhere, BlueprintReadWrite, Category = Appearance )
 	FSlateBrush UncheckedImage;
 	FCheckBoxStyle & SetUncheckedImage( const FSlateBrush & InUncheckedImage ){ UncheckedImage = InUncheckedImage; return *this; }
 
 	/* CheckBox appearance when the CheckBox is unchecked and hovered */
-	UPROPERTY( EditAnywhere, Category = Appearance )
+	UPROPERTY( EditAnywhere, BlueprintReadWrite, Category = Appearance )
 	FSlateBrush UncheckedHoveredImage;
 	FCheckBoxStyle & SetUncheckedHoveredImage( const FSlateBrush & InUncheckedHoveredImage ){ UncheckedHoveredImage = InUncheckedHoveredImage; return *this; }
 
 	/* CheckBox appearance when the CheckBox is unchecked and hovered */
-	UPROPERTY( EditAnywhere, Category = Appearance )
+	UPROPERTY( EditAnywhere, BlueprintReadWrite, Category = Appearance )
 	FSlateBrush UncheckedPressedImage;
 	FCheckBoxStyle & SetUncheckedPressedImage( const FSlateBrush & InUncheckedPressedImage ){ UncheckedPressedImage = InUncheckedPressedImage; return *this; }
 
 	/* CheckBox appearance when the CheckBox is checked */
-	UPROPERTY( EditAnywhere, Category = Appearance )
+	UPROPERTY( EditAnywhere, BlueprintReadWrite, Category = Appearance )
 	FSlateBrush CheckedImage;
 	FCheckBoxStyle & SetCheckedImage( const FSlateBrush & InCheckedImage ){ CheckedImage = InCheckedImage; return *this; }
 
 	/* CheckBox appearance when checked and hovered */
-	UPROPERTY( EditAnywhere, Category = Appearance )
+	UPROPERTY( EditAnywhere, BlueprintReadWrite, Category = Appearance )
 	FSlateBrush CheckedHoveredImage;
 	FCheckBoxStyle & SetCheckedHoveredImage( const FSlateBrush & InCheckedHoveredImage ){ CheckedHoveredImage = InCheckedHoveredImage; return *this; }
 
 	/* CheckBox appearance when checked and pressed */
-	UPROPERTY( EditAnywhere, Category = Appearance )
+	UPROPERTY( EditAnywhere, BlueprintReadWrite, Category = Appearance )
 	FSlateBrush CheckedPressedImage;
 	FCheckBoxStyle & SetCheckedPressedImage( const FSlateBrush & InCheckedPressedImage ){ CheckedPressedImage = InCheckedPressedImage; return *this; }
 	
 	/* CheckBox appearance when the CheckBox is undetermined */
-	UPROPERTY( EditAnywhere, Category = Appearance )
+	UPROPERTY( EditAnywhere, BlueprintReadWrite, Category = Appearance )
 	FSlateBrush UndeterminedImage;
 	FCheckBoxStyle & SetUndeterminedImage( const FSlateBrush & InUndeterminedImage ){ UndeterminedImage = InUndeterminedImage; return *this; }
 
 	/* CheckBox appearance when CheckBox is undetermined and hovered */
-	UPROPERTY( EditAnywhere, Category = Appearance )
+	UPROPERTY( EditAnywhere, BlueprintReadWrite, Category = Appearance )
 	FSlateBrush UndeterminedHoveredImage;
 	FCheckBoxStyle & SetUndeterminedHoveredImage( const FSlateBrush & InUndeterminedHoveredImage ){ UndeterminedHoveredImage = InUndeterminedHoveredImage; return *this; }
 
 	/* CheckBox appearance when CheckBox is undetermined and pressed */
-	UPROPERTY( EditAnywhere, Category = Appearance )
+	UPROPERTY( EditAnywhere, BlueprintReadWrite, Category = Appearance )
 	FSlateBrush UndeterminedPressedImage;
 	FCheckBoxStyle & SetUndeterminedPressedImage( const FSlateBrush & InUndeterminedPressedImage ){ UndeterminedPressedImage = InUndeterminedPressedImage; return *this; }
 
 	/** Padding */
-	UPROPERTY(EditAnywhere, Category=Appearance)
+	UPROPERTY(EditAnywhere, BlueprintReadWrite, Category=Appearance)
 	FMargin Padding;
 	FCheckBoxStyle& SetPadding( const FMargin& InPadding ){ Padding = InPadding; return *this; }
 
 	/** The foreground color */
-	UPROPERTY(EditAnywhere, Category=Appearance)
+	UPROPERTY(EditAnywhere, BlueprintReadWrite, Category=Appearance)
 	FSlateColor ForegroundColor;
 	FCheckBoxStyle& SetForegroundColor(const FSlateColor& InForegroundColor) { ForegroundColor = InForegroundColor; return *this; }
 
 	/** BorderBackgroundColor refers to the actual color and opacity of the supplied border image on toggle buttons */
-	UPROPERTY(EditAnywhere, Category=Appearance)
+	UPROPERTY(EditAnywhere, BlueprintReadWrite, Category=Appearance)
 	FSlateColor BorderBackgroundColor;
 	FCheckBoxStyle& SetBorderBackgroundColor(const FSlateColor& InBorderBackgroundColor) { BorderBackgroundColor = InBorderBackgroundColor; return *this; }
 
 	/**
 	 * The sound the check box should play when checked
 	 */	
-	UPROPERTY(EditAnywhere, Category=Sound, meta=(DisplayName="Checked Sound"))
+	UPROPERTY(EditAnywhere, BlueprintReadWrite, Category=Sound, meta=( DisplayName="Checked Sound" ))
 	FSlateSound CheckedSlateSound;
 	FCheckBoxStyle& SetCheckedSound( const FSlateSound& InCheckedSound ){ CheckedSlateSound = InCheckedSound; return *this; }
 
 	/**
 	 * The sound the check box should play when unchecked
 	 */	
-	UPROPERTY(EditAnywhere, Category=Sound, meta=(DisplayName="Unchecked Sound"))
+	UPROPERTY(EditAnywhere, BlueprintReadWrite, Category=Sound, meta=( DisplayName="Unchecked Sound" ))
 	FSlateSound UncheckedSlateSound;
 	FCheckBoxStyle& SetUncheckedSound( const FSlateSound& InUncheckedSound ){ UncheckedSlateSound = InUncheckedSound; return *this; }
 
 	/**
 	 * The sound the check box should play when initially hovered over
 	 */	
-	UPROPERTY(EditAnywhere, Category=Sound, meta=(DisplayName="Hovered Sound"))
+	UPROPERTY(EditAnywhere, BlueprintReadWrite, Category=Sound, meta=( DisplayName="Hovered Sound" ))
 	FSlateSound HoveredSlateSound;
 	FCheckBoxStyle& SetHoveredSound( const FSlateSound& InHoveredSound ){ HoveredSlateSound = InHoveredSound; return *this; }
 
@@ -153,7 +153,7 @@ struct TStructOpsTypeTraits<FCheckBoxStyle> : public TStructOpsTypeTraitsBase
 /**
  * Represents the appearance of an STextBlock
  */
-USTRUCT()
+USTRUCT(BlueprintType)
 struct SLATECORE_API FTextBlockStyle : public FSlateWidgetStyle
 {
 	GENERATED_USTRUCT_BODY()
@@ -168,7 +168,7 @@ struct SLATECORE_API FTextBlockStyle : public FSlateWidgetStyle
 	static const FTextBlockStyle& GetDefault();
 
 	/** Font family and size to be used when displaying this text. */
-	UPROPERTY(EditAnywhere, Category=Appearance)
+	UPROPERTY(EditAnywhere, BlueprintReadWrite, Category=Appearance)
 	FSlateFontInfo Font;
 	FTextBlockStyle& SetFont(const FSlateFontInfo& InFont) { Font = InFont; return *this; }
 	FTextBlockStyle& SetFont(const FName& InFontName, uint16 InSize) { Font = FSlateFontInfo(InFontName, InSize); return *this; }
@@ -182,27 +182,27 @@ struct SLATECORE_API FTextBlockStyle : public FSlateWidgetStyle
 	FTextBlockStyle& SetFontSize(uint16 InSize) { Font.Size = InSize; return *this; }
 
 	/** The color and opacity of this text */
-	UPROPERTY(EditAnywhere, Category=Appearance)
+	UPROPERTY(EditAnywhere, BlueprintReadWrite, Category=Appearance)
 	FSlateColor ColorAndOpacity;
 	FTextBlockStyle& SetColorAndOpacity(const FSlateColor& InColorAndOpacity) { ColorAndOpacity = InColorAndOpacity; return *this; }
 
 	/** How much should the shadow be offset? An offset of 0 implies no shadow. */
-	UPROPERTY(EditAnywhere, Category=Appearance)
+	UPROPERTY(EditAnywhere, BlueprintReadWrite, Category=Appearance)
 	FVector2D ShadowOffset;
 	FTextBlockStyle& SetShadowOffset(const FVector2D& InShadowOffset) { ShadowOffset = InShadowOffset; return *this; }
 
 	/** The color and opacity of the shadow */
-	UPROPERTY(EditAnywhere, Category=Appearance)
+	UPROPERTY(EditAnywhere, BlueprintReadWrite, Category=Appearance)
 	FLinearColor ShadowColorAndOpacity;
 	FTextBlockStyle& SetShadowColorAndOpacity(const FLinearColor& InShadowColorAndOpacity) { ShadowColorAndOpacity = InShadowColorAndOpacity; return *this; }
 
 	/** The color of highlighted text */
-	UPROPERTY(EditAnywhere, Category=Appearance)
+	UPROPERTY(EditAnywhere, BlueprintReadWrite, Category=Appearance)
 	FLinearColor HighlightColor;
 	FTextBlockStyle& SetHighlightColor(const FLinearColor& InHighlightColor) { HighlightColor = InHighlightColor; return *this; }
 
 	/** The shape of highlighted text */
-	UPROPERTY(EditAnywhere, Category=Appearance)
+	UPROPERTY(EditAnywhere, BlueprintReadWrite, Category=Appearance)
 	FSlateBrush HighlightShape;
 	FTextBlockStyle& SetHighlightShape( const FSlateBrush& InHighlightShape ){ HighlightShape = InHighlightShape; return *this; }
 };
@@ -210,7 +210,7 @@ struct SLATECORE_API FTextBlockStyle : public FSlateWidgetStyle
 /**
  * Represents the appearance of an SButton
  */
-USTRUCT()
+USTRUCT(BlueprintType)
 struct SLATECORE_API FButtonStyle : public FSlateWidgetStyle
 {
 	GENERATED_USTRUCT_BODY()
@@ -227,22 +227,22 @@ struct SLATECORE_API FButtonStyle : public FSlateWidgetStyle
 	static const FButtonStyle& GetDefault();
 
 	/** Button appearance when the button is not hovered or pressed */
-	UPROPERTY(EditAnywhere, Category=Appearance)
+	UPROPERTY(EditAnywhere, BlueprintReadWrite, Category=Appearance)
 	FSlateBrush Normal;
 	FButtonStyle& SetNormal( const FSlateBrush& InNormal ){ Normal = InNormal; return *this; }
 
 	/** Button appearance when hovered */
-	UPROPERTY(EditAnywhere, Category=Appearance)
+	UPROPERTY(EditAnywhere, BlueprintReadWrite, Category=Appearance)
 	FSlateBrush Hovered;
 	FButtonStyle& SetHovered( const FSlateBrush& InHovered){ Hovered = InHovered; return *this; }
 
 	/** Button appearance when pressed */
-	UPROPERTY(EditAnywhere, Category=Appearance)
+	UPROPERTY(EditAnywhere, BlueprintReadWrite, Category=Appearance)
 	FSlateBrush Pressed;
 	FButtonStyle& SetPressed( const FSlateBrush& InPressed ){ Pressed = InPressed; return *this; }
 
 	/** Button appearance when disabled */
-	UPROPERTY(EditAnywhere, Category=Appearance)
+	UPROPERTY(EditAnywhere, BlueprintReadWrite, Category=Appearance)
 	FSlateBrush Disabled;
 	FButtonStyle& SetDisabled( const FSlateBrush& InDisabled ){ Disabled = InDisabled; return *this; }
 
@@ -251,7 +251,7 @@ struct SLATECORE_API FButtonStyle : public FSlateWidgetStyle
 	 * When this is applied, the content of the button should appear flush
 	 * with the button's border. Use this padding when the button is not pressed.
 	 */
-	UPROPERTY(EditAnywhere, Category=Appearance)
+	UPROPERTY(EditAnywhere, BlueprintReadWrite, Category=Appearance)
 	FMargin NormalPadding;
 	FButtonStyle& SetNormalPadding( const FMargin& InNormalPadding){ NormalPadding = InNormalPadding; return *this; }
 
@@ -259,21 +259,21 @@ struct SLATECORE_API FButtonStyle : public FSlateWidgetStyle
 	 * Same as NormalPadding but used when the button is pressed. Allows for moving the content to match
 	 * any "movement" in the button's border image.
 	 */	
-	UPROPERTY(EditAnywhere, Category=Appearance)
+	UPROPERTY(EditAnywhere, BlueprintReadWrite, Category=Appearance)
 	FMargin PressedPadding;
 	FButtonStyle& SetPressedPadding( const FMargin& InPressedPadding){ PressedPadding = InPressedPadding; return *this; }
 
 	/**
 	 * The sound the button should play when pressed
 	 */	
-	UPROPERTY(EditAnywhere, Category=Sound, meta=(DisplayName="Pressed Sound"))
+	UPROPERTY(EditAnywhere, BlueprintReadWrite, Category=Sound, meta=( DisplayName="Pressed Sound" ))
 	FSlateSound PressedSlateSound;
 	FButtonStyle& SetPressedSound( const FSlateSound& InPressedSound ){ PressedSlateSound = InPressedSound; return *this; }
 
 	/**
 	 * The sound the button should play when initially hovered over
 	 */	
-	UPROPERTY(EditAnywhere, Category=Sound, meta=(DisplayName="Hovered Sound"))
+	UPROPERTY(EditAnywhere, BlueprintReadWrite, Category=Sound, meta=( DisplayName="Hovered Sound" ))
 	FSlateSound HoveredSlateSound;
 	FButtonStyle& SetHoveredSound( const FSlateSound& InHoveredSound ){ HoveredSlateSound = InHoveredSound; return *this; }
 
@@ -300,7 +300,7 @@ struct TStructOpsTypeTraits<FButtonStyle> : public TStructOpsTypeTraitsBase
 /**
  * Represents the appearance of an SComboButton
  */
-USTRUCT()
+USTRUCT(BlueprintType)
 struct SLATECORE_API FComboButtonStyle : public FSlateWidgetStyle
 {
 	GENERATED_USTRUCT_BODY()
@@ -319,28 +319,28 @@ struct SLATECORE_API FComboButtonStyle : public FSlateWidgetStyle
 	/**
 	 * The style to use for our SButton
 	 */	
-	UPROPERTY(EditAnywhere, Category=Appearance)
+	UPROPERTY(EditAnywhere, BlueprintReadWrite, Category=Appearance)
 	FButtonStyle ButtonStyle;
 	FComboButtonStyle& SetButtonStyle( const FButtonStyle& InButtonStyle ){ ButtonStyle = InButtonStyle; return *this; }
 
 	/**
 	 * Image to use for the down arrow
 	 */	
-	UPROPERTY(EditAnywhere, Category=Appearance)
+	UPROPERTY(EditAnywhere, BlueprintReadWrite, Category=Appearance)
 	FSlateBrush DownArrowImage;
 	FComboButtonStyle& SetDownArrowImage( const FSlateBrush& InDownArrowImage ){ DownArrowImage = InDownArrowImage; return *this; }
 
 	/**
 	 * Brush to use to add a "menu border" around the drop-down content
 	 */
-	UPROPERTY(EditAnywhere, Category=Appearance)
+	UPROPERTY(EditAnywhere, BlueprintReadWrite, Category=Appearance)
 	FSlateBrush MenuBorderBrush;
 	FComboButtonStyle& SetMenuBorderBrush( const FSlateBrush& InMenuBorderBrush ){ MenuBorderBrush = InMenuBorderBrush; return *this; }
 
 	/**
 	 * Padding to use to add a "menu border" around the drop-down content
 	 */
-	UPROPERTY(EditAnywhere, Category=Appearance)
+	UPROPERTY(EditAnywhere, BlueprintReadWrite, Category=Appearance)
 	FMargin MenuBorderPadding;
 	FComboButtonStyle& SetMenuBorderPadding( const FMargin& InMenuBorderPadding ){ MenuBorderPadding = InMenuBorderPadding; return *this; }
 };
@@ -349,7 +349,7 @@ struct SLATECORE_API FComboButtonStyle : public FSlateWidgetStyle
 /**
  * Represents the appearance of an SComboBox
  */
-USTRUCT()
+USTRUCT(BlueprintType)
 struct SLATECORE_API FComboBoxStyle : public FSlateWidgetStyle
 {
 	GENERATED_USTRUCT_BODY()
@@ -444,7 +444,7 @@ struct SLATECORE_API FHyperlinkStyle : public FSlateWidgetStyle
 /**
  * Represents the appearance of an SEditableText
  */
-USTRUCT()
+USTRUCT(BlueprintType)
 struct SLATECORE_API FEditableTextStyle : public FSlateWidgetStyle
 {
 	GENERATED_USTRUCT_BODY()
@@ -461,33 +461,33 @@ struct SLATECORE_API FEditableTextStyle : public FSlateWidgetStyle
 	static const FEditableTextStyle& GetDefault();
 
 	/** Font family and size to be used when displaying this text. */
-	UPROPERTY(EditAnywhere, Category=Appearance)
+	UPROPERTY(EditAnywhere, BlueprintReadWrite, Category=Appearance)
 	FSlateFontInfo Font;
 	FEditableTextStyle& SetFont(const FSlateFontInfo& InFont) { Font = InFont; return *this; }
 	FEditableTextStyle& SetFont(const FName& InFontName, uint16 InSize) { Font = FSlateFontInfo(InFontName, InSize); return *this; }
 
 	/** The color and opacity of this text */
-	UPROPERTY(EditAnywhere, Category=Appearance)
+	UPROPERTY(EditAnywhere, BlueprintReadWrite, Category=Appearance)
 	FSlateColor ColorAndOpacity;
 	FEditableTextStyle& SetColorAndOpacity(const FSlateColor& InColorAndOpacity) { ColorAndOpacity = InColorAndOpacity; return *this; }
 
 	/** Background image for the selected text */
-	UPROPERTY(EditAnywhere, Category=Appearance)
+	UPROPERTY(EditAnywhere, BlueprintReadWrite, Category=Appearance)
 	FSlateBrush BackgroundImageSelected;
 	FEditableTextStyle& SetBackgroundImageSelected( const FSlateBrush& InBackgroundImageSelected ){ BackgroundImageSelected = InBackgroundImageSelected; return *this; }
 
 	/** Background image for the selection targeting effect */
-	UPROPERTY(EditAnywhere, Category=Appearance)
+	UPROPERTY(EditAnywhere, BlueprintReadWrite, Category=Appearance)
 	FSlateBrush BackgroundImageSelectionTarget;
 	FEditableTextStyle& SetBackgroundImageSelectionTarget( const FSlateBrush& InBackgroundImageSelectionTarget ){ BackgroundImageSelectionTarget = InBackgroundImageSelectionTarget; return *this; }
 
 	/** Background image for the selected text */
-	UPROPERTY(EditAnywhere, Category=Appearance)
+	UPROPERTY(EditAnywhere, BlueprintReadWrite, Category=Appearance)
 	FSlateBrush BackgroundImageComposing;
 	FEditableTextStyle& SetBackgroundImageComposing( const FSlateBrush& InBackgroundImageComposing ){ BackgroundImageComposing = InBackgroundImageComposing; return *this; }	
 
 	/** Image brush used for the caret */
-	UPROPERTY(EditAnywhere, Category=Appearance)
+	UPROPERTY(EditAnywhere, BlueprintReadWrite, Category=Appearance)
 	FSlateBrush CaretImage;
 	FEditableTextStyle& SetCaretImage( const FSlateBrush& InCaretImage ){ CaretImage = InCaretImage; return *this; }
 };
@@ -496,7 +496,7 @@ struct SLATECORE_API FEditableTextStyle : public FSlateWidgetStyle
 /**
  * Represents the appearance of an SScrollBar
  */
-USTRUCT()
+USTRUCT(BlueprintType)
 struct SLATECORE_API FScrollBarStyle : public FSlateWidgetStyle
 {
 	GENERATED_USTRUCT_BODY()
@@ -513,27 +513,27 @@ struct SLATECORE_API FScrollBarStyle : public FSlateWidgetStyle
 	static const FScrollBarStyle& GetDefault();
 
 	/** Background image to use when the scrollbar is oriented horizontally */
-	UPROPERTY(EditAnywhere, Category=Appearance)
+	UPROPERTY(EditAnywhere, BlueprintReadWrite, Category=Appearance)
 	FSlateBrush HorizontalBackgroundImage;
 	FScrollBarStyle& SetHorizontalBackgroundImage( const FSlateBrush& InHorizontalBackgroundImage ){ HorizontalBackgroundImage = InHorizontalBackgroundImage; return *this; }
 
 	/** Background image to use when the scrollbar is oriented vertically */
-	UPROPERTY(EditAnywhere, Category=Appearance)
+	UPROPERTY(EditAnywhere, BlueprintReadWrite, Category=Appearance)
 	FSlateBrush VerticalBackgroundImage;
 	FScrollBarStyle& SetVerticalBackgroundImage( const FSlateBrush& InVerticalBackgroundImage ){ VerticalBackgroundImage = InVerticalBackgroundImage; return *this; }
 
 	/** Image to use when the scrollbar thumb is in its normal state */
-	UPROPERTY(EditAnywhere, Category=Appearance)
+	UPROPERTY(EditAnywhere, BlueprintReadWrite, Category=Appearance)
 	FSlateBrush NormalThumbImage;
 	FScrollBarStyle& SetNormalThumbImage( const FSlateBrush& InNormalThumbImage ){ NormalThumbImage = InNormalThumbImage; return *this; }
 
 	/** Image to use when the scrollbar thumb is in its hovered state */
-	UPROPERTY(EditAnywhere, Category=Appearance)
+	UPROPERTY(EditAnywhere, BlueprintReadWrite, Category=Appearance)
 	FSlateBrush HoveredThumbImage;
 	FScrollBarStyle& SetHoveredThumbImage( const FSlateBrush& InHoveredThumbImage ){ HoveredThumbImage = InHoveredThumbImage; return *this; }
 
 	/** Image to use when the scrollbar thumb is in its dragged state */
-	UPROPERTY(EditAnywhere, Category=Appearance)
+	UPROPERTY(EditAnywhere, BlueprintReadWrite, Category=Appearance)
 	FSlateBrush DraggedThumbImage;
 	FScrollBarStyle& SetDraggedThumbImage( const FSlateBrush& InDraggedThumbImage ){ DraggedThumbImage = InDraggedThumbImage; return *this; }
 };
@@ -542,7 +542,7 @@ struct SLATECORE_API FScrollBarStyle : public FSlateWidgetStyle
 /**
  * Represents the appearance of an SEditableTextBox
  */
-USTRUCT()
+USTRUCT(BlueprintType)
 struct SLATECORE_API FEditableTextBoxStyle : public FSlateWidgetStyle
 {
 	GENERATED_USTRUCT_BODY()
@@ -559,63 +559,63 @@ struct SLATECORE_API FEditableTextBoxStyle : public FSlateWidgetStyle
 	static const FEditableTextBoxStyle& GetDefault();
 
 	/** Border background image when the box is not hovered or focused */
-	UPROPERTY(EditAnywhere, Category=Appearance)
+	UPROPERTY(EditAnywhere, BlueprintReadWrite, Category=Appearance)
 	FSlateBrush BackgroundImageNormal;
 	FEditableTextBoxStyle& SetBackgroundImageNormal( const FSlateBrush& InBackgroundImageNormal ){ BackgroundImageNormal = InBackgroundImageNormal; return *this; }
 
 	/** Border background image when the box is hovered */
-	UPROPERTY(EditAnywhere, Category=Appearance)
+	UPROPERTY(EditAnywhere, BlueprintReadWrite, Category=Appearance)
 	FSlateBrush BackgroundImageHovered;
 	FEditableTextBoxStyle& SetBackgroundImageHovered( const FSlateBrush& InBackgroundImageHovered ){ BackgroundImageHovered = InBackgroundImageHovered; return *this; }
 
 	/** Border background image when the box is focused */
-	UPROPERTY(EditAnywhere, Category=Appearance)
+	UPROPERTY(EditAnywhere, BlueprintReadWrite, Category=Appearance)
 	FSlateBrush BackgroundImageFocused;
 	FEditableTextBoxStyle& SetBackgroundImageFocused( const FSlateBrush& InBackgroundImageFocused ){ BackgroundImageFocused = InBackgroundImageFocused; return *this; }
 
 	/** Border background image when the box is read-only */
-	UPROPERTY(EditAnywhere, Category=Appearance)
+	UPROPERTY(EditAnywhere, BlueprintReadWrite, Category=Appearance)
 	FSlateBrush BackgroundImageReadOnly;
 	FEditableTextBoxStyle& SetBackgroundImageReadOnly( const FSlateBrush& InBackgroundImageReadOnly ){ BackgroundImageReadOnly = InBackgroundImageReadOnly; return *this; }
 
 	/** Padding */
-	UPROPERTY(EditAnywhere, Category=Appearance)
+	UPROPERTY(EditAnywhere, BlueprintReadWrite, Category=Appearance)
 	FMargin Padding;
 	FEditableTextBoxStyle& SetPadding( const FMargin& InPadding ){ Padding = InPadding; return *this; }
 
 	/** Font family and size to be used when displaying this text. */
-	UPROPERTY(EditAnywhere, Category=Appearance)
+	UPROPERTY(EditAnywhere, BlueprintReadWrite, Category=Appearance)
 	FSlateFontInfo Font;
 	FEditableTextBoxStyle& SetFont(const FSlateFontInfo& InFont) { Font = InFont; return *this; }
 	FEditableTextBoxStyle& SetFont(const FName& InFontName, uint16 InSize) { Font = FSlateFontInfo(InFontName, InSize); return *this; }
 
 	/** The foreground color of text. */
-	UPROPERTY(EditAnywhere, Category=Appearance)
+	UPROPERTY(EditAnywhere, BlueprintReadWrite, Category=Appearance)
 	FSlateColor ForegroundColor;
 	FEditableTextBoxStyle& SetForegroundColor(const FSlateColor& InForegroundColor) { ForegroundColor = InForegroundColor; return *this; }
 
 	/** The background color applied to the active background image */
-	UPROPERTY(EditAnywhere, Category=Appearance)
+	UPROPERTY(EditAnywhere, BlueprintReadWrite, Category=Appearance)
 	FSlateColor BackgroundColor;
 	FEditableTextBoxStyle& SetBackgroundColor(const FSlateColor& InBackgroundColor) { BackgroundColor = InBackgroundColor; return *this; }
 
 	/** The read-only foreground color of text in read-only mode. */
-	UPROPERTY(EditAnywhere, Category=Appearance)
+	UPROPERTY(EditAnywhere, BlueprintReadWrite, Category=Appearance)
 	FSlateColor ReadOnlyForegroundColor;
 	FEditableTextBoxStyle& SetReadOnlyForegroundColor(const FSlateColor& InReadOnlyForegroundColor) { ReadOnlyForegroundColor = InReadOnlyForegroundColor; return *this; }
 
 	/** Padding around the horizontal scrollbar */
-	UPROPERTY(EditAnywhere, Category=Appearance)
+	UPROPERTY(EditAnywhere, BlueprintReadWrite, Category=Appearance)
 	FMargin HScrollBarPadding;
 	FEditableTextBoxStyle& SetHScrollBarPadding( const FMargin& InPadding ){ HScrollBarPadding = InPadding; return *this; }
 
 	/** Padding around the vertical scrollbar */
-	UPROPERTY(EditAnywhere, Category=Appearance)
+	UPROPERTY(EditAnywhere, BlueprintReadWrite, Category=Appearance)
 	FMargin VScrollBarPadding;
 	FEditableTextBoxStyle& SetVScrollBarPadding( const FMargin& InPadding ){ VScrollBarPadding = InPadding; return *this; }
 
 	/** Style used for the scrollbars */
-	UPROPERTY(EditAnywhere, Category=Appearance)
+	UPROPERTY(EditAnywhere, BlueprintReadWrite, Category=Appearance)
 	FScrollBarStyle ScrollBarStyle;
 	FEditableTextBoxStyle& SetScrollBarStyle( const FScrollBarStyle& InScrollBarStyle ){ ScrollBarStyle = InScrollBarStyle; return *this; }
 };
@@ -655,7 +655,7 @@ struct SLATECORE_API FInlineEditableTextBlockStyle : public FSlateWidgetStyle
 /**
  * Represents the appearance of an SProgressBar
  */
-USTRUCT()
+USTRUCT(BlueprintType)
 struct SLATECORE_API FProgressBarStyle : public FSlateWidgetStyle
 {
 	GENERATED_USTRUCT_BODY()
@@ -672,17 +672,17 @@ struct SLATECORE_API FProgressBarStyle : public FSlateWidgetStyle
 	static const FProgressBarStyle& GetDefault();
 
 	/** Background image to use for the progress bar */
-	UPROPERTY(EditAnywhere, Category=Appearance)
+	UPROPERTY(EditAnywhere, BlueprintReadWrite, Category=Appearance)
 	FSlateBrush BackgroundImage;
 	FProgressBarStyle& SetBackgroundImage( const FSlateBrush& InBackgroundImage ){ BackgroundImage = InBackgroundImage; return *this; }
 
 	/** Foreground image to use for the progress bar */
-	UPROPERTY(EditAnywhere, Category=Appearance)
+	UPROPERTY(EditAnywhere, BlueprintReadWrite, Category=Appearance)
 	FSlateBrush FillImage;
 	FProgressBarStyle& SetFillImage( const FSlateBrush& InFillImage ){ FillImage = InFillImage; return *this; }
 
 	/** Image to use for marquee mode */
-	UPROPERTY(EditAnywhere, Category=Appearance)
+	UPROPERTY(EditAnywhere, BlueprintReadWrite, Category=Appearance)
 	FSlateBrush MarqueeImage;
 	FProgressBarStyle& SetMarqueeImage( const FSlateBrush& InMarqueeImage ){ MarqueeImage = InMarqueeImage; return *this; }
 };
@@ -773,7 +773,7 @@ struct SLATECORE_API FSearchBoxStyle : public FSlateWidgetStyle
 /**
  * Represents the appearance of an SSlider
  */
-USTRUCT()
+USTRUCT(BlueprintType)
 struct SLATECORE_API FSliderStyle : public FSlateWidgetStyle
 {
 	GENERATED_USTRUCT_BODY()
@@ -790,12 +790,12 @@ struct SLATECORE_API FSliderStyle : public FSlateWidgetStyle
 	static const FSliderStyle& GetDefault();
 
 	/** Image to use when the slider thumb is in its normal state */
-	UPROPERTY(EditAnywhere, Category=Appearance)
+	UPROPERTY(EditAnywhere, BlueprintReadWrite, Category=Appearance)
 	FSlateBrush NormalThumbImage;
 	FSliderStyle& SetNormalThumbImage( const FSlateBrush& InNormalThumbImage ){ NormalThumbImage = InNormalThumbImage; return *this; }
 
 	/** Image to use when the slider thumb is in its disabled state */
-	UPROPERTY(EditAnywhere, Category=Appearance)
+	UPROPERTY(EditAnywhere, BlueprintReadWrite, Category=Appearance)
 	FSlateBrush DisabledThumbImage;
 	FSliderStyle& SetDisabledThumbImage( const FSlateBrush& InDisabledThumbImage ){ DisabledThumbImage = InDisabledThumbImage; return *this; }
 };
@@ -1235,7 +1235,7 @@ struct SLATECORE_API FDockTabStyle : public FSlateWidgetStyle
 /**
  * Represents the appearance of an SScrollBox
  */
-USTRUCT()
+USTRUCT(BlueprintType)
 struct SLATECORE_API FScrollBoxStyle : public FSlateWidgetStyle
 {
 	GENERATED_USTRUCT_BODY()
@@ -1252,17 +1252,17 @@ struct SLATECORE_API FScrollBoxStyle : public FSlateWidgetStyle
 	static const FScrollBoxStyle& GetDefault();
 
 	/** Brush used to draw the top shadow of a scrollbox */
-	UPROPERTY(EditAnywhere, Category=Appearance)
+	UPROPERTY(EditAnywhere, BlueprintReadWrite, Category=Appearance)
 	FSlateBrush TopShadowBrush;
 	FScrollBoxStyle& SetTopShadowBrush( const FSlateBrush& InTopShadowBrush ){ TopShadowBrush = InTopShadowBrush; return *this; }
 
 	/** Brush used to draw the bottom shadow of a scrollbox */
-	UPROPERTY(EditAnywhere, Category=Appearance)
+	UPROPERTY(EditAnywhere, BlueprintReadWrite, Category=Appearance)
 	FSlateBrush BottomShadowBrush;
 	FScrollBoxStyle& SetBottomShadowBrush( const FSlateBrush& InBottomShadowBrush ){ BottomShadowBrush = InBottomShadowBrush; return *this; }
 
 	/** Brush used to draw the left shadow of a scrollbox */
-	UPROPERTY(EditAnywhere, Category = Appearance)
+	UPROPERTY(EditAnywhere, BlueprintReadWrite, Category = Appearance)
 	FSlateBrush LeftShadowBrush;
 	FScrollBoxStyle& SetLeftShadowBrush(const FSlateBrush& InLeftShadowBrush)
 	{
@@ -1271,7 +1271,7 @@ struct SLATECORE_API FScrollBoxStyle : public FSlateWidgetStyle
 	}
 
 	/** Brush used to draw the right shadow of a scrollbox */
-	UPROPERTY(EditAnywhere, Category = Appearance)
+	UPROPERTY(EditAnywhere, BlueprintReadWrite, Category = Appearance)
 	FSlateBrush RightShadowBrush;
 	FScrollBoxStyle& SetRightShadowBrush(const FSlateBrush& InRightShadowBrush)
 	{
@@ -1302,12 +1302,12 @@ struct SLATECORE_API FScrollBorderStyle : public FSlateWidgetStyle
 
 	/** Brush used to draw the top shadow of a scrollborder */
 	UPROPERTY( EditAnywhere, Category = Appearance )
-		FSlateBrush TopShadowBrush;
+	FSlateBrush TopShadowBrush;
 	FScrollBorderStyle& SetTopShadowBrush( const FSlateBrush& InTopShadowBrush ){ TopShadowBrush = InTopShadowBrush; return *this; }
 
 	/** Brush used to draw the bottom shadow of a scrollborder */
 	UPROPERTY( EditAnywhere, Category = Appearance )
-		FSlateBrush BottomShadowBrush;
+	FSlateBrush BottomShadowBrush;
 	FScrollBorderStyle& SetBottomShadowBrush( const FSlateBrush& InBottomShadowBrush ){ BottomShadowBrush = InBottomShadowBrush; return *this; }
 };
 
