Clean up Rotated Rect types in Slate.
* FSlateClipRectType -> FSlateRotatedClipRectType to better represent its usage.
* Remove special purpose ctors in place of helper code in the element batchers.
* Added common ctors to init rotated rects by their individual components.
* Added transform calculus support for rotated rects, used it to simplify a few internal functions.

diff --git a/Engine/Source/Developer/StandaloneRenderer/Private/OpenGL/SlateOpenGLRenderingPolicy.cpp b/Engine/Source/Developer/StandaloneRenderer/Private/OpenGL/SlateOpenGLRenderingPolicy.cpp
--- a/Engine/Source/Developer/StandaloneRenderer/Private/OpenGL/SlateOpenGLRenderingPolicy.cpp
+++ b/Engine/Source/Developer/StandaloneRenderer/Private/OpenGL/SlateOpenGLRenderingPolicy.cpp
@@ -257,7 +257,7 @@ void FSlateOpenGLRenderingPolicy::DrawElements( const FMatrix& ViewProjectionMat
 		glVertexAttribPointer( 2, 2, bUseFloat16 ? GL_HALF_FLOAT : GL_FLOAT, GL_FALSE, Stride, BUFFER_OFFSET(Stride*BaseVertexIndex+Offset) );
 
 		glEnableVertexAttribArray(3);
-		Offset = STRUCT_OFFSET( FSlateVertex, ClipRect ) + STRUCT_OFFSET( FSlateClipRectType, ExtentX );
+		Offset = STRUCT_OFFSET( FSlateVertex, ClipRect ) + STRUCT_OFFSET( FSlateRotatedClipRectType, ExtentX );
 		glVertexAttribPointer( 3, 4, bUseFloat16 ? GL_HALF_FLOAT : GL_FLOAT, GL_FALSE, Stride, BUFFER_OFFSET(Stride*BaseVertexIndex+Offset) );
 
 		glEnableVertexAttribArray(4);
diff --git a/Engine/Source/Runtime/SlateCore/Private/Input/HittestGrid.cpp b/Engine/Source/Runtime/SlateCore/Private/Input/HittestGrid.cpp
--- a/Engine/Source/Runtime/SlateCore/Private/Input/HittestGrid.cpp
+++ b/Engine/Source/Runtime/SlateCore/Private/Input/HittestGrid.cpp
@@ -64,11 +64,15 @@ TArray<FArrangedWidget> FHittestGrid::GetBubblePath( FVector2D DesktopSpaceCoord
 
 			const FCachedWidget& TestCandidate = (*WidgetsCachedThisFrame)[IndexesInCell[i]];
 
-			// Compute the render space clipping rect.
-			FSlateRotatedRect DesktopOrientedClipRect(
-				Inverse(TestCandidate.CachedGeometry.GetAccumulatedLayoutTransform()),
-				TestCandidate.CachedGeometry.GetAccumulatedRenderTransform(),
-				TestCandidate.CachedGeometry.GetClippingRect().IntersectionWith(TestCandidate.ClippingRect));
+			// Compute the render space clipping rect (FGeometry exposes a layout space clipping rect).
+			FSlateRotatedRect DesktopOrientedClipRect = 
+				TransformRect(
+					Concatenate(
+						Inverse(TestCandidate.CachedGeometry.GetAccumulatedLayoutTransform()), 
+						TestCandidate.CachedGeometry.GetAccumulatedRenderTransform()
+					), 
+					FSlateRotatedRect(TestCandidate.CachedGeometry.GetClippingRect().IntersectionWith(TestCandidate.ClippingRect))
+				);
 
 			if (DesktopOrientedClipRect.IsUnderLocation(DesktopSpaceCoordinate) && TestCandidate.WidgetPtr.IsValid())
 			{
@@ -163,10 +167,16 @@ int32 FHittestGrid::InsertWidget( const int32 ParentHittestIndex, const EVisibil
 		// Mark any cell that is overlapped by this widget.
 
 		// Compute the render space clipping rect, and compute it's aligned bounds so we can insert conservatively into the hit test grid.
-		FSlateRect GridRelativeBoundingClipRect = FSlateRotatedRect(
-			Inverse(WindowAdjustedWidget.Geometry.GetAccumulatedLayoutTransform()),
-			WindowAdjustedWidget.Geometry.GetAccumulatedRenderTransform(),
-			WindowAdjustedWidget.Geometry.GetClippingRect().IntersectionWith(WindowAdjustedRect)).ToBoundingRect().OffsetBy(-GridOrigin);
+		FSlateRect GridRelativeBoundingClipRect = 
+			TransformRect(
+				Concatenate(
+					Inverse(WindowAdjustedWidget.Geometry.GetAccumulatedLayoutTransform()),
+					WindowAdjustedWidget.Geometry.GetAccumulatedRenderTransform()
+				),
+				FSlateRotatedRect(WindowAdjustedWidget.Geometry.GetClippingRect().IntersectionWith(WindowAdjustedRect))
+			)
+			.ToBoundingRect()
+			.OffsetBy(-GridOrigin);
 
 		// Starting and ending cells covered by this widget.	
 		const FIntPoint UpperLeftCell = FIntPoint(
diff --git a/Engine/Source/Runtime/SlateCore/Private/Rendering/ElementBatcher.cpp b/Engine/Source/Runtime/SlateCore/Private/Rendering/ElementBatcher.cpp
--- a/Engine/Source/Runtime/SlateCore/Private/Rendering/ElementBatcher.cpp
+++ b/Engine/Source/Runtime/SlateCore/Private/Rendering/ElementBatcher.cpp
@@ -4,23 +4,25 @@
 #include "ElementBatcher.h"
 #include "SlateRenderTransform.h"
 
-/** Default ctor. */
+FVector2D RoundToInt(const FVector2D& Vec)
+{
+	return FVector2D(FMath::RoundToInt(Vec.X), FMath::RoundToInt(Vec.Y));
+}
+
 FSlateRotatedRect::FSlateRotatedRect() {}
 
-FSlateRotatedRect::FSlateRotatedRect(const FSlateLayoutTransform& InverseLayoutTransform, const FSlateRenderTransform& RenderTransform, const FSlateRect& ClipRectInLayoutWindowSpace)
+FSlateRotatedRect::FSlateRotatedRect(const FSlateRect& AlignedRect)
+	: TopLeft(AlignedRect.GetTopLeft())
+	, ExtentX(AlignedRect.Right - AlignedRect.Left, 0.0f)
+	, ExtentY(0.0f, AlignedRect.Bottom - AlignedRect.Top)
+{
+}
+
+FSlateRotatedRect::FSlateRotatedRect(const FVector2D& InTopLeft, const FVector2D& InExtentX, const FVector2D& InExtentY)
+	: TopLeft(InTopLeft)
+	, ExtentX(InExtentX)
+	, ExtentY(InExtentY)
 {
-	const FSlateRenderTransform LayoutWindowSpaceToRenderWindowSpace = Concatenate(InverseLayoutTransform, RenderTransform);
-
-	const FVector2D SizeInLayoutWindowSpace = ClipRectInLayoutWindowSpace.GetSize();
-	const FVector2D TopLeftInRenderWindowSpace = TransformPoint(LayoutWindowSpaceToRenderWindowSpace, ClipRectInLayoutWindowSpace.GetTopLeft());
-	const FVector2D ClipExtentXInRenderWindowSpace = TransformVector(LayoutWindowSpaceToRenderWindowSpace, FVector2D(SizeInLayoutWindowSpace.X, 0.0f));
-	const FVector2D ClipExtentYInRenderWindowSpace = TransformVector(LayoutWindowSpaceToRenderWindowSpace, FVector2D(0.0f, SizeInLayoutWindowSpace.Y));
-	TopLeft.X = FMath::RoundToInt(TopLeftInRenderWindowSpace.X);
-	TopLeft.Y = FMath::RoundToInt(TopLeftInRenderWindowSpace.Y);
-	ExtentX.X = FMath::RoundToInt(ClipExtentXInRenderWindowSpace.X);
-	ExtentX.Y = FMath::RoundToInt(ClipExtentXInRenderWindowSpace.Y);
-	ExtentY.X = FMath::RoundToInt(ClipExtentYInRenderWindowSpace.X);
-	ExtentY.Y = FMath::RoundToInt(ClipExtentYInRenderWindowSpace.Y);
 }
 
 // !!! WRH 2014/08/25 - this is a brute-force, not efficient implementation, uses a bunch of extra conditionals.
@@ -56,10 +58,8 @@ bool FSlateRotatedRect::IsUnderLocation(const FVector2D& Location) const
 	return false;
 }
 
-/** Default ctor. */
 FSlateRotatedRectFloat16::FSlateRotatedRectFloat16() {}
 
-/** Construct a float16 version of a rotated rect from a full-float version. */
 FSlateRotatedRectFloat16::FSlateRotatedRectFloat16(const FSlateRotatedRect& RotatedRect)
 {
 	TopLeft[0] = RotatedRect.TopLeft[0];
@@ -70,41 +70,54 @@ FSlateRotatedRectFloat16::FSlateRotatedRectFloat16(const FSlateRotatedRect& Rota
 	ExtentY[1] = RotatedRect.ExtentY[1];
 }
 
-/** Used by the typedef magic to construct a rotated rect of the type needed by the RHI being used. */
-FSlateRotatedRectFloat16::FSlateRotatedRectFloat16(const FSlateLayoutTransform& InverseLayoutTransform, const FSlateRenderTransform& RenderTransform, const FSlateRect& ClipRectInLayoutWindowSpace)
+FSlateRotatedRectFloat16::FSlateRotatedRectFloat16(const FVector2D& InTopLeft, const FVector2D& InExtentX, const FVector2D& InExtentY)
 {
-	FSlateRotatedRect RotatedRect(InverseLayoutTransform, RenderTransform, ClipRectInLayoutWindowSpace);
-	TopLeft[0] = RotatedRect.TopLeft[0];
-	TopLeft[1] = RotatedRect.TopLeft[1];
-	ExtentX[0] = RotatedRect.ExtentX[0];
-	ExtentX[1] = RotatedRect.ExtentX[1];
-	ExtentY[0] = RotatedRect.ExtentY[0];
-	ExtentY[1] = RotatedRect.ExtentY[1];
+	TopLeft[0] = InTopLeft.X;
+	TopLeft[1] = InTopLeft.Y;
+	ExtentX[0] = InExtentX.X;
+	ExtentX[1] = InExtentX.Y;
+	ExtentY[0] = InExtentY.X;
+	ExtentY[1] = InExtentY.Y;
 }
 
+FSlateVertex::FSlateVertex() 
+{
+}
 
-FSlateVertex::FSlateVertex() {}
-
-FSlateVertex::FSlateVertex(const FSlateRenderTransform& RenderTransform, const FVector2D& InLocalPosition, const FVector2D& InTexCoord, const FVector2D& InTexCoord2, const FColor& InColor, const FSlateClipRectType& InClipRect )
+FSlateVertex::FSlateVertex(const FSlateRenderTransform& RenderTransform, const FVector2D& InLocalPosition, const FVector2D& InTexCoord, const FVector2D& InTexCoord2, const FColor& InColor, const FSlateRotatedClipRectType& InClipRect )
 	: TexCoords( InTexCoord, InTexCoord2 )
 	, ClipRect( InClipRect )
 	, Color( InColor )
 {
 	const FVector2D WindowPosition = TransformPoint(RenderTransform, InLocalPosition);
+	// Pixel snapping here.
 	Position[0] = FMath::RoundToInt(WindowPosition.X);
 	Position[1] = FMath::RoundToInt(WindowPosition.Y);
 }
 
-FSlateVertex::FSlateVertex( const FSlateRenderTransform& RenderTransform, const FVector2D& InLocalPosition, const FVector2D& InTexCoord, const FColor& InColor, const FSlateClipRectType& InClipRect )
+FSlateVertex::FSlateVertex( const FSlateRenderTransform& RenderTransform, const FVector2D& InLocalPosition, const FVector2D& InTexCoord, const FColor& InColor, const FSlateRotatedClipRectType& InClipRect )
 	: TexCoords( InTexCoord, FVector2D(1.0f, 1.0f) )
 	, ClipRect( InClipRect )
 	, Color( InColor )
 {
 	const FVector2D WindowPosition = TransformPoint(RenderTransform, InLocalPosition);
+	// Pixel snapping here.
 	Position[0] = FMath::RoundToInt(WindowPosition.X);
 	Position[1] = FMath::RoundToInt(WindowPosition.Y);
 }
 
+/**
+ * Used to construct a rotated rect from an aligned clip rect and a set of layout and render transforms from the geometry, snapped to pixel boundaries. Returns a float or float16 version of the rect based on the typedef.
+ */
+FSlateRotatedClipRectType ToSnappedRotatedRect(const FSlateRect& ClipRectInLayoutWindowSpace, const FSlateLayoutTransform& InverseLayoutTransform, const FSlateRenderTransform& RenderTransform)
+{
+	FSlateRotatedRect RotatedRect = TransformRect(Concatenate(InverseLayoutTransform, RenderTransform), FSlateRotatedRect(ClipRectInLayoutWindowSpace));
+	// Pixel snapping is done here by rounding the resulting floats to ints.
+	return FSlateRotatedClipRectType(
+		RoundToInt(RotatedRect.TopLeft), 
+		RoundToInt(RotatedRect.ExtentX), 
+		RoundToInt(RotatedRect.ExtentY));
+}
 
 static const FSlateImageBrush& GetSplineFilterTable()
 {
@@ -205,7 +218,7 @@ void FSlateElementBatcher::AddQuadElement( const FSlateDrawElement& DrawElement,
 
 	// extract the layout transform from the draw element
 	FSlateLayoutTransform InverseLayoutTransform(Inverse(FSlateLayoutTransform(DrawElement.GetScale(), DrawElement.GetPosition())));
-	FSlateClipRectType RenderClipRect = FSlateClipRectType(InverseLayoutTransform, RenderTransform, InClippingRect);
+	FSlateRotatedClipRectType RenderClipRect = ToSnappedRotatedRect(InClippingRect, InverseLayoutTransform, RenderTransform);
 
 	FSlateElementBatch& ElementBatch = FindBatchForElement( Layer, FShaderParams(), nullptr, ESlateDrawPrimitive::TriangleList, ESlateShader::Default, ESlateDrawEffect::None, ESlateBatchDrawFlag::Wireframe|ESlateBatchDrawFlag::NoBlending );
 	TArray<FSlateVertex>& BatchVertices = BatchVertexArrays[ElementBatch.VertexArrayIndex];
@@ -270,7 +283,7 @@ void FSlateElementBatcher::AddBoxElement( const FSlateDrawElement& DrawElement )
 	// extract the layout transform from the draw element
 	FSlateLayoutTransform InverseLayoutTransform(Inverse(FSlateLayoutTransform(DrawElement.GetScale(), DrawElement.GetPosition())));
 	// The clip rect is NOT subject to the rotations specified by MakeRotatedBox.
-	FSlateClipRectType RenderClipRect = FSlateClipRectType(InverseLayoutTransform, ElementRenderTransform, InClippingRect);
+	FSlateRotatedClipRectType RenderClipRect = ToSnappedRotatedRect(InClippingRect, InverseLayoutTransform, ElementRenderTransform);
 
 	check( InPayload.BrushResource );
 
@@ -569,7 +582,7 @@ void FSlateElementBatcher::AddTextElement(const FSlateDrawElement& DrawElement)
 	FSlateLayoutTransform InverseLayoutTransform = Inverse(Concatenate(Inverse(FontScale), LayoutTransform));
 	const FSlateRenderTransform RenderTransform = Concatenate(Inverse(FontScale), DrawElement.GetRenderTransform());
 
-	FSlateClipRectType RenderClipRect = FSlateClipRectType(InverseLayoutTransform, RenderTransform, InClippingRect);
+	FSlateRotatedClipRectType RenderClipRect = ToSnappedRotatedRect(InClippingRect, InverseLayoutTransform, RenderTransform);
 	// Used to clip individual characters as we generate them.
 	FSlateRect LocalClipRect = TransformRect(InverseLayoutTransform, InClippingRect);
 
@@ -732,7 +745,7 @@ void FSlateElementBatcher::AddGradientElement( const FSlateDrawElement& DrawElem
 
 	// extract the layout transform from the draw element
 	FSlateLayoutTransform InverseLayoutTransform(Inverse(FSlateLayoutTransform(DrawElement.GetScale(), DrawElement.GetPosition())));
-	FSlateClipRectType RenderClipRect = FSlateClipRectType(InverseLayoutTransform, RenderTransform, InClippingRect);
+	FSlateRotatedClipRectType RenderClipRect = ToSnappedRotatedRect(InClippingRect, InverseLayoutTransform, RenderTransform);
 
 	// There must be at least one gradient stop
 	check( InPayload.GradientStops.Num() > 0 );
@@ -859,7 +872,7 @@ void FSlateElementBatcher::AddSplineElement( const FSlateDrawElement& DrawElemen
 
 	// extract the layout transform from the draw element
 	FSlateLayoutTransform InverseLayoutTransform(Inverse(FSlateLayoutTransform(DrawElement.GetScale(), DrawElement.GetPosition())));
-	FSlateClipRectType RenderClipRect = FSlateClipRectType(InverseLayoutTransform, RenderTransform, InClippingRect);
+	FSlateRotatedClipRectType RenderClipRect = ToSnappedRotatedRect(InClippingRect, InverseLayoutTransform, RenderTransform);
 
 	//@todo SLATE: Merge with AddLineElement?
 
@@ -1005,7 +1018,7 @@ void FSlateElementBatcher::AddLineElement( const FSlateDrawElement& DrawElement
 
 	// extract the layout transform from the draw element
 	FSlateLayoutTransform InverseLayoutTransform(Inverse(FSlateLayoutTransform(DrawElement.GetScale(), DrawElement.GetPosition())));
-	FSlateClipRectType RenderClipRect = FSlateClipRectType(InverseLayoutTransform, RenderTransform, InClippingRect);
+	FSlateRotatedClipRectType RenderClipRect = ToSnappedRotatedRect(InClippingRect, InverseLayoutTransform, RenderTransform);
 
 	if( InPayload.Points.Num() < 2 )
 	{
@@ -1183,7 +1196,7 @@ void FSlateElementBatcher::AddViewportElement( const FSlateDrawElement& DrawElem
 
 	// extract the layout transform from the draw element
 	FSlateLayoutTransform InverseLayoutTransform(Inverse(FSlateLayoutTransform(DrawElement.GetScale(), DrawElement.GetPosition())));
-	FSlateClipRectType RenderClipRect = FSlateClipRectType(InverseLayoutTransform, RenderTransform, InClippingRect);
+	FSlateRotatedClipRectType RenderClipRect = ToSnappedRotatedRect(InClippingRect, InverseLayoutTransform, RenderTransform);
 
 	const FColor FinalColor = GetElementColor( InPayload.Tint, nullptr );
 
@@ -1261,7 +1274,7 @@ void FSlateElementBatcher::AddBorderElement( const FSlateDrawElement& DrawElemen
 
 	// extract the layout transform from the draw element
 	FSlateLayoutTransform InverseLayoutTransform(Inverse(FSlateLayoutTransform(DrawElement.GetScale(), DrawElement.GetPosition())));
-	FSlateClipRectType RenderClipRect = FSlateClipRectType(InverseLayoutTransform, RenderTransform, InClippingRect);
+	FSlateRotatedClipRectType RenderClipRect = ToSnappedRotatedRect(InClippingRect, InverseLayoutTransform, RenderTransform);
 
 	check( InPayload.BrushResource );
 
diff --git a/Engine/Source/Runtime/SlateCore/Public/Rendering/RenderingCommon.h b/Engine/Source/Runtime/SlateCore/Public/Rendering/RenderingCommon.h
--- a/Engine/Source/Runtime/SlateCore/Public/Rendering/RenderingCommon.h
+++ b/Engine/Source/Runtime/SlateCore/Public/Rendering/RenderingCommon.h
@@ -95,8 +95,12 @@ namespace ESlateLineJoinType
  */
 struct FSlateRotatedRect
 {
+	/** Default ctor. */
 	FSlateRotatedRect();
-	FSlateRotatedRect(const FSlateLayoutTransform& InverseLayoutTransform, const FSlateRenderTransform& RenderTransform, const FSlateRect& ClipRectInLayoutWindowSpace);
+	/** Construct a rotated rect from a given aligned rect. */
+	explicit FSlateRotatedRect(const FSlateRect& AlignedRect);
+	/** Per-element constructor. */
+	FSlateRotatedRect(const FVector2D& InTopLeft, const FVector2D& InExtentX, const FVector2D& InExtentY);
 	/** transformed Top-left corner. */
 	FVector2D TopLeft;
 	/** transformed X extent (right-left). */
@@ -104,18 +108,38 @@ struct FSlateRotatedRect
 	/** transformed Y extent (bottom-top). */
 	FVector2D ExtentY;
 
+	/** Convert to a bounding, aligned rect. */
 	FSlateRect ToBoundingRect() const;
+	/** Point-in-rect test. */
 	bool IsUnderLocation(const FVector2D& Location) const;
 };
 
 /**
- * Stores a Rotated rect as float16 for rendering.
+ * Transforms a rect by the given transform.
+ */
+template <typename TransformType>
+FSlateRotatedRect TransformRect(const TransformType& Transform, const FSlateRotatedRect& Rect)
+{
+	return FSlateRotatedRect
+	{
+		TransformPoint(Transform, Rect.TopLeft),
+		TransformVector(Transform, Rect.ExtentX),
+		TransformVector(Transform, Rect.ExtentY),
+	};
+}
+
+
+/**
+ * Stores a Rotated rect as float16 (for rendering).
  */
 struct FSlateRotatedRectFloat16
 {
+	/** Default ctor. */
 	FSlateRotatedRectFloat16();
-	FSlateRotatedRectFloat16(const FSlateRotatedRect& RotatedRect);
-	FSlateRotatedRectFloat16(const FSlateLayoutTransform& InverseLayoutTransform, const FSlateRenderTransform& RenderTransform, const FSlateRect& ClipRectInLayoutWindowSpace);
+	/** Construct a float16 version of a rotated rect from a full-float version. */
+	explicit FSlateRotatedRectFloat16(const FSlateRotatedRect& RotatedRect);
+	/** Per-element constructor. */
+	FSlateRotatedRectFloat16(const FVector2D& InTopLeft, const FVector2D& InExtentX, const FVector2D& InExtentY);
 	/** transformed Top-left corner. */
 	FFloat16 TopLeft[2];
 	/** transformed X extent (right-left). */
@@ -128,9 +152,9 @@ struct FSlateRotatedRectFloat16
  * Not all platforms support Float16, so we have to be tricky here and declare the proper vertex type.
  */
 #if SLATE_USE_FLOAT16
-typedef FSlateRotatedRectFloat16 FSlateClipRectType;
+typedef FSlateRotatedRectFloat16 FSlateRotatedClipRectType;
 #else
-typedef FSlateRotatedRect FSlateClipRectType;
+typedef FSlateRotatedRect FSlateRotatedClipRectType;
 #endif
 
 
@@ -144,13 +168,13 @@ struct FSlateVertex
 	/** Position of the vertex in window space */
 	int16 Position[2];
 	/** clip center/extents in render window space (window space with render transforms applied) */
-	FSlateClipRectType ClipRect;
+	FSlateRotatedClipRectType ClipRect;
 	/** Vertex color */
 	FColor Color;
 	
 	FSlateVertex();
-	FSlateVertex( const FSlateRenderTransform& RenderTransform, const FVector2D& InLocalPosition, const FVector2D& InTexCoord, const FVector2D& InTexCoord2, const FColor& InColor, const FSlateClipRectType& InClipRect );
-	FSlateVertex( const FSlateRenderTransform& RenderTransform, const FVector2D& InLocalPosition, const FVector2D& InTexCoord, const FColor& InColor, const FSlateClipRectType& InClipRect );
+	FSlateVertex( const FSlateRenderTransform& RenderTransform, const FVector2D& InLocalPosition, const FVector2D& InTexCoord, const FVector2D& InTexCoord2, const FColor& InColor, const FSlateRotatedClipRectType& InClipRect );
+	FSlateVertex( const FSlateRenderTransform& RenderTransform, const FVector2D& InLocalPosition, const FVector2D& InTexCoord, const FColor& InColor, const FSlateRotatedClipRectType& InClipRect );
 };
 
 template<> struct TIsPODType<FSlateVertex> { enum { Value = true }; };
diff --git a/Engine/Source/Runtime/SlateRHIRenderer/Private/SlateShaders.cpp b/Engine/Source/Runtime/SlateRHIRenderer/Private/SlateShaders.cpp
--- a/Engine/Source/Runtime/SlateRHIRenderer/Private/SlateShaders.cpp
+++ b/Engine/Source/Runtime/SlateRHIRenderer/Private/SlateShaders.cpp
@@ -52,7 +52,7 @@ void FSlateVertexDeclaration::InitRHI()
 		false;
 #endif
 	Elements.Add(FVertexElement(0,STRUCT_OFFSET(FSlateVertex,ClipRect),bUseFloat16 ? VET_Half2 : VET_Float2,2,Stride));
-	Elements.Add(FVertexElement(0,STRUCT_OFFSET(FSlateVertex,ClipRect)+STRUCT_OFFSET(FSlateClipRectType,ExtentX),bUseFloat16 ? VET_Half4 : VET_Float4,3,Stride));
+	Elements.Add(FVertexElement(0,STRUCT_OFFSET(FSlateVertex,ClipRect)+STRUCT_OFFSET(FSlateRotatedClipRectType,ExtentX),bUseFloat16 ? VET_Half4 : VET_Float4,3,Stride));
 	Elements.Add(FVertexElement(0,STRUCT_OFFSET(FSlateVertex,Color),VET_Color,4,Stride));
 
 	VertexDeclarationRHI = RHICreateVertexDeclaration(Elements);
