FGeometry AbsoluteToLocal, LocalToAbsolute, and IsUnderLocation now respect render transforms.

diff --git a/Engine/Source/Runtime/SlateCore/Private/Layout/Geometry.cpp b/Engine/Source/Runtime/SlateCore/Private/Layout/Geometry.cpp
--- a/Engine/Source/Runtime/SlateCore/Private/Layout/Geometry.cpp
+++ b/Engine/Source/Runtime/SlateCore/Private/Layout/Geometry.cpp
@@ -167,18 +167,20 @@ FSlateRect FGeometry::GetClippingRect() const
 
 FVector2D FGeometry::AbsoluteToLocal(FVector2D AbsoluteCoordinate) const
 {
-	return TransformPoint(Inverse(GetAccumulatedLayoutTransform()), AbsoluteCoordinate);
+	// this render transform invert is a little expensive. We might consider caching it.
+	return TransformPoint(Inverse(GetAccumulatedRenderTransform()), AbsoluteCoordinate);
 }
 
 
 FVector2D FGeometry::LocalToAbsolute(FVector2D LocalCoordinate) const
 {
-	return TransformPoint(GetAccumulatedLayoutTransform(), LocalCoordinate);
+	return TransformPoint(GetAccumulatedRenderTransform(), LocalCoordinate);
 }
 
 bool FGeometry::IsUnderLocation(const FVector2D& AbsoluteCoordinate) const
 {
-	return GetClippingRect().ContainsPoint(AbsoluteCoordinate);
+	// this render transform invert is a little expensive. We might consider caching it.
+	return FSlateRect(FVector2D(0.0f, 0.0f), Size).ContainsPoint(TransformPoint(Inverse(GetAccumulatedRenderTransform()), AbsoluteCoordinate));
 }
 
 
