VS2012 compile fix.

diff --git a/Engine/Source/Runtime/SlateCore/Public/Rendering/RenderingCommon.h b/Engine/Source/Runtime/SlateCore/Public/Rendering/RenderingCommon.h
--- a/Engine/Source/Runtime/SlateCore/Public/Rendering/RenderingCommon.h
+++ b/Engine/Source/Runtime/SlateCore/Public/Rendering/RenderingCommon.h
@@ -121,11 +121,11 @@ template <typename TransformType>
 FSlateRotatedRect TransformRect(const TransformType& Transform, const FSlateRotatedRect& Rect)
 {
 	return FSlateRotatedRect
-	{
+	(
 		TransformPoint(Transform, Rect.TopLeft),
 		TransformVector(Transform, Rect.ExtentX),
-		TransformVector(Transform, Rect.ExtentY),
-	};
+		TransformVector(Transform, Rect.ExtentY)
+	);
 }
 
 
