Add FromPointAndExtent factory function to FSlateRect, made the other ctor explicit.

diff --git a/Engine/Source/Runtime/SlateCore/Public/Layout/SlateRect.h b/Engine/Source/Runtime/SlateCore/Public/Layout/SlateRect.h
--- a/Engine/Source/Runtime/SlateCore/Public/Layout/SlateRect.h
+++ b/Engine/Source/Runtime/SlateCore/Public/Layout/SlateRect.h
@@ -28,7 +28,7 @@ public:
 	float Right;
 	float Bottom;
 
-	FSlateRect( float InLeft = -1, float InTop = -1, float InRight = -1, float InBottom = -1 )
+	explicit FSlateRect( float InLeft = -1, float InTop = -1, float InRight = -1, float InBottom = -1 )
 		: Left(InLeft)
 		, Top(InTop)
 		, Right(InRight)
@@ -41,6 +41,15 @@ public:
 		, Right(InEndPos.X)
 		, Bottom(InEndPos.Y)
 	{ }
+
+	/**
+	 * Creates a rect from a top left point and extent. Provided as a factory function to not conflict
+	 * with the TopLeft + BottomRight ctor.
+	 */
+	static FSlateRect FromPointAndExtent(const FVector2D& TopLeft, const FVector2D& Size)
+	{
+		return FSlateRect(TopLeft, TopLeft + Size);
+	}
 	
 public:
 
@@ -212,6 +221,9 @@ public:
 
 /**
  * Transforms a rect by the given transform, ensuring the rect does not get inverted.
+ * WARNING: this only really supports scales and offsets. Any skew or rotation that 
+ * would turn this into an un-aligned rect will won't work because FSlateRect doesn't support
+ * non-axis-alignment. Instead, convert to ta FSlateRotatedRect first and transform that.
  */
 template <typename TransformType>
 FSlateRect TransformRect(const TransformType& Transform, const FSlateRect& Rect)
