* Added RenderingCommon.cpp and moves appropriate functions there.
* Renamed FSlateRotatedRectFloat16 -> FSlateRotatedRectHalf and made it use existing FVector2DHalf.

diff --git a/Engine/Source/Runtime/SlateCore/Private/Rendering/ElementBatcher.cpp b/Engine/Source/Runtime/SlateCore/Private/Rendering/ElementBatcher.cpp
--- a/Engine/Source/Runtime/SlateCore/Private/Rendering/ElementBatcher.cpp
+++ b/Engine/Source/Runtime/SlateCore/Private/Rendering/ElementBatcher.cpp
@@ -9,103 +9,6 @@ FVector2D RoundToInt(const FVector2D& Vec)
 	return FVector2D(FMath::RoundToInt(Vec.X), FMath::RoundToInt(Vec.Y));
 }
 
-FSlateRotatedRect::FSlateRotatedRect() {}
-
-FSlateRotatedRect::FSlateRotatedRect(const FSlateRect& AlignedRect)
-	: TopLeft(AlignedRect.GetTopLeft())
-	, ExtentX(AlignedRect.Right - AlignedRect.Left, 0.0f)
-	, ExtentY(0.0f, AlignedRect.Bottom - AlignedRect.Top)
-{
-}
-
-FSlateRotatedRect::FSlateRotatedRect(const FVector2D& InTopLeft, const FVector2D& InExtentX, const FVector2D& InExtentY)
-	: TopLeft(InTopLeft)
-	, ExtentX(InExtentX)
-	, ExtentY(InExtentY)
-{
-}
-
-// !!! WRH 2014/08/25 - this is a brute-force, not efficient implementation, uses a bunch of extra conditionals.
-FSlateRect FSlateRotatedRect::ToBoundingRect() const
-{
-	FVector2D Points[4] = 
-	{
-		TopLeft,
-		TopLeft + ExtentX,
-		TopLeft + ExtentY,
-		TopLeft + ExtentX + ExtentY
-	};
-	return FSlateRect(
-		FMath::Min(Points[0].X, FMath::Min3(Points[1].X, Points[2].X, Points[3].X)),
-		FMath::Min(Points[0].Y, FMath::Min3(Points[1].Y, Points[2].Y, Points[3].Y)),
-		FMath::Max(Points[0].X, FMath::Max3(Points[1].X, Points[2].X, Points[3].X)),
-		FMath::Max(Points[0].Y, FMath::Max3(Points[1].Y, Points[2].Y, Points[3].Y))
-		);
-}
-
-bool FSlateRotatedRect::IsUnderLocation(const FVector2D& Location) const
-{
-	const FVector2D Offset = Location - TopLeft;
-	const float Det = FVector2D::CrossProduct(ExtentX, ExtentY);
-
-	// Not exhaustively efficient. Could optimize the checks for [0..1] to short circuit faster.
-	const float S = -FVector2D::CrossProduct(Offset, ExtentX) / Det;
-	if (FMath::IsWithinInclusive(S, 0.0f, 1.0f))
-	{
-		const float T = FVector2D::CrossProduct(Offset, ExtentY) / Det;
-		return FMath::IsWithinInclusive(T, 0.0f, 1.0f);
-	}
-	return false;
-}
-
-FSlateRotatedRectFloat16::FSlateRotatedRectFloat16() {}
-
-FSlateRotatedRectFloat16::FSlateRotatedRectFloat16(const FSlateRotatedRect& RotatedRect)
-{
-	TopLeft[0] = RotatedRect.TopLeft[0];
-	TopLeft[1] = RotatedRect.TopLeft[1];
-	ExtentX[0] = RotatedRect.ExtentX[0];
-	ExtentX[1] = RotatedRect.ExtentX[1];
-	ExtentY[0] = RotatedRect.ExtentY[0];
-	ExtentY[1] = RotatedRect.ExtentY[1];
-}
-
-FSlateRotatedRectFloat16::FSlateRotatedRectFloat16(const FVector2D& InTopLeft, const FVector2D& InExtentX, const FVector2D& InExtentY)
-{
-	TopLeft[0] = InTopLeft.X;
-	TopLeft[1] = InTopLeft.Y;
-	ExtentX[0] = InExtentX.X;
-	ExtentX[1] = InExtentX.Y;
-	ExtentY[0] = InExtentY.X;
-	ExtentY[1] = InExtentY.Y;
-}
-
-FSlateVertex::FSlateVertex() 
-{
-}
-
-FSlateVertex::FSlateVertex(const FSlateRenderTransform& RenderTransform, const FVector2D& InLocalPosition, const FVector2D& InTexCoord, const FVector2D& InTexCoord2, const FColor& InColor, const FSlateRotatedClipRectType& InClipRect )
-	: TexCoords( InTexCoord, InTexCoord2 )
-	, ClipRect( InClipRect )
-	, Color( InColor )
-{
-	const FVector2D WindowPosition = TransformPoint(RenderTransform, InLocalPosition);
-	// Pixel snapping here.
-	Position[0] = FMath::RoundToInt(WindowPosition.X);
-	Position[1] = FMath::RoundToInt(WindowPosition.Y);
-}
-
-FSlateVertex::FSlateVertex( const FSlateRenderTransform& RenderTransform, const FVector2D& InLocalPosition, const FVector2D& InTexCoord, const FColor& InColor, const FSlateRotatedClipRectType& InClipRect )
-	: TexCoords( InTexCoord, FVector2D(1.0f, 1.0f) )
-	, ClipRect( InClipRect )
-	, Color( InColor )
-{
-	const FVector2D WindowPosition = TransformPoint(RenderTransform, InLocalPosition);
-	// Pixel snapping here.
-	Position[0] = FMath::RoundToInt(WindowPosition.X);
-	Position[1] = FMath::RoundToInt(WindowPosition.Y);
-}
-
 /**
  * Used to construct a rotated rect from an aligned clip rect and a set of layout and render transforms from the geometry, snapped to pixel boundaries. Returns a float or float16 version of the rect based on the typedef.
  */
diff --git a/Engine/Source/Runtime/SlateCore/Private/Rendering/RenderingCommon.cpp b/Engine/Source/Runtime/SlateCore/Private/Rendering/RenderingCommon.cpp
new file mode 100644
--- /dev/null
+++ b/Engine/Source/Runtime/SlateCore/Private/Rendering/RenderingCommon.cpp
@@ -0,0 +1,96 @@
+// Copyright 1998-2014 Epic Games, Inc. All Rights Reserved.
+
+#include "SlateCorePrivatePCH.h"
+#include "RenderingCommon.h"
+
+FSlateRotatedRect::FSlateRotatedRect() {}
+
+FSlateRotatedRect::FSlateRotatedRect(const FSlateRect& AlignedRect)
+	: TopLeft(AlignedRect.GetTopLeft())
+	, ExtentX(AlignedRect.Right - AlignedRect.Left, 0.0f)
+	, ExtentY(0.0f, AlignedRect.Bottom - AlignedRect.Top)
+{
+}
+
+FSlateRotatedRect::FSlateRotatedRect(const FVector2D& InTopLeft, const FVector2D& InExtentX, const FVector2D& InExtentY)
+	: TopLeft(InTopLeft)
+	, ExtentX(InExtentX)
+	, ExtentY(InExtentY)
+{
+}
+
+// !!! WRH 2014/08/25 - this is a brute-force, not efficient implementation, uses a bunch of extra conditionals.
+FSlateRect FSlateRotatedRect::ToBoundingRect() const
+{
+	FVector2D Points[4] = 
+	{
+		TopLeft,
+		TopLeft + ExtentX,
+		TopLeft + ExtentY,
+		TopLeft + ExtentX + ExtentY
+	};
+	return FSlateRect(
+		FMath::Min(Points[0].X, FMath::Min3(Points[1].X, Points[2].X, Points[3].X)),
+		FMath::Min(Points[0].Y, FMath::Min3(Points[1].Y, Points[2].Y, Points[3].Y)),
+		FMath::Max(Points[0].X, FMath::Max3(Points[1].X, Points[2].X, Points[3].X)),
+		FMath::Max(Points[0].Y, FMath::Max3(Points[1].Y, Points[2].Y, Points[3].Y))
+		);
+}
+
+bool FSlateRotatedRect::IsUnderLocation(const FVector2D& Location) const
+{
+	const FVector2D Offset = Location - TopLeft;
+	const float Det = FVector2D::CrossProduct(ExtentX, ExtentY);
+
+	// Not exhaustively efficient. Could optimize the checks for [0..1] to short circuit faster.
+	const float S = -FVector2D::CrossProduct(Offset, ExtentX) / Det;
+	if (FMath::IsWithinInclusive(S, 0.0f, 1.0f))
+	{
+		const float T = FVector2D::CrossProduct(Offset, ExtentY) / Det;
+		return FMath::IsWithinInclusive(T, 0.0f, 1.0f);
+	}
+	return false;
+}
+
+FSlateRotatedRectHalf::FSlateRotatedRectHalf() {}
+
+FSlateRotatedRectHalf::FSlateRotatedRectHalf(const FSlateRotatedRect& RotatedRect)
+	: TopLeft(RotatedRect.TopLeft)
+	, ExtentX(RotatedRect.ExtentX)
+	, ExtentY(RotatedRect.ExtentY)
+{
+}
+
+FSlateRotatedRectHalf::FSlateRotatedRectHalf(const FVector2D& InTopLeft, const FVector2D& InExtentX, const FVector2D& InExtentY)
+	: TopLeft(InTopLeft)
+	, ExtentX(InExtentX)
+	, ExtentY(InExtentY)
+{
+}
+
+FSlateVertex::FSlateVertex() 
+{
+}
+
+FSlateVertex::FSlateVertex(const FSlateRenderTransform& RenderTransform, const FVector2D& InLocalPosition, const FVector2D& InTexCoord, const FVector2D& InTexCoord2, const FColor& InColor, const FSlateRotatedClipRectType& InClipRect )
+	: TexCoords( InTexCoord, InTexCoord2 )
+	, ClipRect( InClipRect )
+	, Color( InColor )
+{
+	const FVector2D WindowPosition = TransformPoint(RenderTransform, InLocalPosition);
+	// Pixel snapping here.
+	Position[0] = FMath::RoundToInt(WindowPosition.X);
+	Position[1] = FMath::RoundToInt(WindowPosition.Y);
+}
+
+FSlateVertex::FSlateVertex( const FSlateRenderTransform& RenderTransform, const FVector2D& InLocalPosition, const FVector2D& InTexCoord, const FColor& InColor, const FSlateRotatedClipRectType& InClipRect )
+	: TexCoords( InTexCoord, FVector2D(1.0f, 1.0f) )
+	, ClipRect( InClipRect )
+	, Color( InColor )
+{
+	const FVector2D WindowPosition = TransformPoint(RenderTransform, InLocalPosition);
+	// Pixel snapping here.
+	Position[0] = FMath::RoundToInt(WindowPosition.X);
+	Position[1] = FMath::RoundToInt(WindowPosition.Y);
+}
+
diff --git a/Engine/Source/Runtime/SlateCore/Public/Rendering/RenderingCommon.h b/Engine/Source/Runtime/SlateCore/Public/Rendering/RenderingCommon.h
--- a/Engine/Source/Runtime/SlateCore/Public/Rendering/RenderingCommon.h
+++ b/Engine/Source/Runtime/SlateCore/Public/Rendering/RenderingCommon.h
@@ -132,27 +132,27 @@ FSlateRotatedRect TransformRect(const TransformType& Transform, const FSlateRota
 /**
  * Stores a Rotated rect as float16 (for rendering).
  */
-struct FSlateRotatedRectFloat16
+struct FSlateRotatedRectHalf
 {
 	/** Default ctor. */
-	FSlateRotatedRectFloat16();
+	FSlateRotatedRectHalf();
 	/** Construct a float16 version of a rotated rect from a full-float version. */
-	explicit FSlateRotatedRectFloat16(const FSlateRotatedRect& RotatedRect);
+	explicit FSlateRotatedRectHalf(const FSlateRotatedRect& RotatedRect);
 	/** Per-element constructor. */
-	FSlateRotatedRectFloat16(const FVector2D& InTopLeft, const FVector2D& InExtentX, const FVector2D& InExtentY);
+	FSlateRotatedRectHalf(const FVector2D& InTopLeft, const FVector2D& InExtentX, const FVector2D& InExtentY);
 	/** transformed Top-left corner. */
-	FFloat16 TopLeft[2];
+	FVector2DHalf TopLeft;
 	/** transformed X extent (right-left). */
-	FFloat16 ExtentX[2];
+	FVector2DHalf ExtentX;
 	/** transformed Y extent (bottom-top). */
-	FFloat16 ExtentY[2];
+	FVector2DHalf ExtentY;
 };
 
 /**
  * Not all platforms support Float16, so we have to be tricky here and declare the proper vertex type.
  */
 #if SLATE_USE_FLOAT16
-typedef FSlateRotatedRectFloat16 FSlateRotatedClipRectType;
+typedef FSlateRotatedRectHalf FSlateRotatedClipRectType;
 #else
 typedef FSlateRotatedRect FSlateRotatedClipRectType;
 #endif
