Slate - Making FMargin a blueprint type struct and making the properties read/writeable.

diff --git a/Engine/Source/Runtime/SlateCore/Public/Layout/Margin.h b/Engine/Source/Runtime/SlateCore/Public/Layout/Margin.h
--- a/Engine/Source/Runtime/SlateCore/Public/Layout/Margin.h
+++ b/Engine/Source/Runtime/SlateCore/Public/Layout/Margin.h
@@ -8,25 +8,25 @@
 /**
  * Describes the space around a Widget.
  */
-USTRUCT()
+USTRUCT(BlueprintType)
 struct FMargin
 {
 	GENERATED_USTRUCT_BODY()
 
 	/** Holds the margin to the left. */
-	UPROPERTY(EditAnywhere, Category=Appearance)
+	UPROPERTY(EditAnywhere, BlueprintReadWrite, Category=Appearance)
 	float Left;
 
 	/** Holds the margin to the top. */
-	UPROPERTY(EditAnywhere, Category=Appearance)
+	UPROPERTY(EditAnywhere, BlueprintReadWrite, Category=Appearance)
 	float Top;
 	
 	/** Holds the margin to the right. */
-	UPROPERTY(EditAnywhere, Category=Appearance)
+	UPROPERTY(EditAnywhere, BlueprintReadWrite, Category=Appearance)
 	float Right;
 	
 	/** Holds the margin to the bottom. */
-	UPROPERTY(EditAnywhere, Category=Appearance)
+	UPROPERTY(EditAnywhere, BlueprintReadWrite, Category=Appearance)
 	float Bottom;
 
 public:
