Slate: OnHitTest() is now deprecated.

diff --git a/Engine/Source/Editor/GraphEditor/Private/SGraphNodeDocumentation.cpp b/Engine/Source/Editor/GraphEditor/Private/SGraphNodeDocumentation.cpp
--- a/Engine/Source/Editor/GraphEditor/Private/SGraphNodeDocumentation.cpp
+++ b/Engine/Source/Editor/GraphEditor/Private/SGraphNodeDocumentation.cpp
@@ -294,16 +294,6 @@ FReply SGraphNodeDocumentation::OnMouseWheel( const FGeometry& MyGeometry, const
 	return FReply::Unhandled();
 }
 
-bool SGraphNodeDocumentation::OnHitTest( const FGeometry& MyGeometry, FVector2D InAbsoluteCursorPosition )
-{
-	if( !bUserIsDragging )
-	{
-		const FVector2D LocalMouseCoordinates = MyGeometry.AbsoluteToLocal( InAbsoluteCursorPosition );
-		return FindMouseZone( LocalMouseCoordinates ) != CRWZ_NotInWindow;
-	}
-	return true;
-}
-
 EVisibility SGraphNodeDocumentation::GetWidgetVisibility() const
 {
 	return ChildWidgetVisibility;
diff --git a/Engine/Source/Editor/GraphEditor/Private/SGraphNodeResizable.cpp b/Engine/Source/Editor/GraphEditor/Private/SGraphNodeResizable.cpp
--- a/Engine/Source/Editor/GraphEditor/Private/SGraphNodeResizable.cpp
+++ b/Engine/Source/Editor/GraphEditor/Private/SGraphNodeResizable.cpp
@@ -111,12 +111,6 @@ FReply SGraphNodeResizable::OnMouseButtonUp( const FGeometry& MyGeometry, const
 	return FReply::Unhandled();
 }
 
-bool SGraphNodeResizable::OnHitTest( const FGeometry& MyGeometry, FVector2D InAbsoluteCursorPosition )
-{
-	FVector2D LocalMouseCoordinates = MyGeometry.AbsoluteToLocal( InAbsoluteCursorPosition );
-	EResizableWindowZone InMouseZone = FindMouseZone(LocalMouseCoordinates);
-	return (InMouseZone != CRWZ_NotInWindow) && (InMouseZone != CRWZ_InWindow);
-}
 
 FReply SGraphNodeResizable::OnMouseMove(const FGeometry& MyGeometry, const FPointerEvent& MouseEvent)
 {
diff --git a/Engine/Source/Editor/GraphEditor/Public/SGraphNodeDocumentation.h b/Engine/Source/Editor/GraphEditor/Public/SGraphNodeDocumentation.h
--- a/Engine/Source/Editor/GraphEditor/Public/SGraphNodeDocumentation.h
+++ b/Engine/Source/Editor/GraphEditor/Public/SGraphNodeDocumentation.h
@@ -17,7 +17,6 @@ public:
 	// Begin SWidget Interface
 	void Construct( const FArguments& InArgs, UEdGraphNode* InNode );
 	virtual FReply OnMouseWheel( const FGeometry& MyGeometry, const FPointerEvent& MouseEvent );
-	virtual bool OnHitTest( const FGeometry& MyGeometry, FVector2D InAbsoluteCursorPosition ) override;
 	virtual void Tick( const FGeometry& AllottedGeometry, const double InCurrentTime, const float InDeltaTime ) override;
 	// End SWidget Interface
 
diff --git a/Engine/Source/Editor/GraphEditor/Public/SGraphNodeResizable.h b/Engine/Source/Editor/GraphEditor/Public/SGraphNodeResizable.h
--- a/Engine/Source/Editor/GraphEditor/Public/SGraphNodeResizable.h
+++ b/Engine/Source/Editor/GraphEditor/Public/SGraphNodeResizable.h
@@ -31,7 +31,6 @@ public:
 	virtual FReply OnMouseButtonDown( const FGeometry& MyGeometry, const FPointerEvent& MouseEvent ) override;
 	virtual FReply OnMouseButtonUp( const FGeometry& MyGeometry, const FPointerEvent& MouseEvent ) override;
 	virtual FCursorReply OnCursorQuery( const FGeometry& MyGeometry, const FPointerEvent& CursorEvent ) const override;
-	virtual bool OnHitTest( const FGeometry& MyGeometry, FVector2D InAbsoluteCursorPosition ) override;
 	// End SWidget Interface
 	
 protected:
diff --git a/Engine/Source/Editor/UMGEditor/Private/Designer/SDesignerView.cpp b/Engine/Source/Editor/UMGEditor/Private/Designer/SDesignerView.cpp
--- a/Engine/Source/Editor/UMGEditor/Private/Designer/SDesignerView.cpp
+++ b/Engine/Source/Editor/UMGEditor/Private/Designer/SDesignerView.cpp
@@ -60,9 +60,7 @@ static bool LocateWidgetsUnderCursor_Helper(FArrangedWidget& Candidate, FVector2
 {
 	const bool bCandidateUnderCursor =
 		// Candidate is physically under the cursor
-		Candidate.Geometry.IsUnderLocation(InAbsoluteCursorLocation) &&
-		// Candidate actually considers itself hit by this test
-		( Candidate.Widget->OnHitTest(Candidate.Geometry, InAbsoluteCursorLocation) );
+		Candidate.Geometry.IsUnderLocation(InAbsoluteCursorLocation);
 
 	bool bHitAnyWidget = false;
 	if ( bCandidateUnderCursor )
diff --git a/Engine/Source/Runtime/SlateCore/Private/Widgets/SWidget.cpp b/Engine/Source/Runtime/SlateCore/Private/Widgets/SWidget.cpp
--- a/Engine/Source/Runtime/SlateCore/Private/Widgets/SWidget.cpp
+++ b/Engine/Source/Runtime/SlateCore/Private/Widgets/SWidget.cpp
@@ -246,12 +246,6 @@ void SWidget::Tick( const FGeometry& AllottedGeometry, const double InCurrentTim
 {
 }
 
-bool SWidget::OnHitTest( const FGeometry& MyGeometry, FVector2D InAbsoluteCursorPosition )
-{
-	return true;
-}
-
-
 void SWidget::TickWidgetsRecursively( const FGeometry& AllottedGeometry, const double InCurrentTime, const float InDeltaTime )
 {
 	// Tick this widget
@@ -400,9 +394,7 @@ int32 SWidget::FindChildUnderMouse( const FArrangedChildren& Children, const FPo
 		const FArrangedWidget& Candidate = Children[ChildIndex];
 		const bool bCandidateUnderCursor = 
 			// Candidate is physically under the cursor
-			Candidate.Geometry.IsUnderLocation( AbsoluteCursorLocation ) && 
-			// Candidate actually considers itself hit by this test
-			(Candidate.Widget->OnHitTest( Candidate.Geometry, AbsoluteCursorLocation ));
+			Candidate.Geometry.IsUnderLocation( AbsoluteCursorLocation );
 
 		if (bCandidateUnderCursor)
 		{
diff --git a/Engine/Source/Runtime/SlateCore/Public/Widgets/SWidget.h b/Engine/Source/Runtime/SlateCore/Public/Widgets/SWidget.h
--- a/Engine/Source/Runtime/SlateCore/Public/Widgets/SWidget.h
+++ b/Engine/Source/Runtime/SlateCore/Public/Widgets/SWidget.h
@@ -124,14 +124,11 @@ public:
 	virtual void Tick( const FGeometry& AllottedGeometry, const double InCurrentTime, const float InDeltaTime );
 
 	/**
-	 * Called when a widget is being hit tested and has passed the bounding box test to determine
-	 * if the precise widget geometry is under the cursor.
-	 *
-	 * @param MyGeometry The Geometry of the widget receiving the event
-	 * @param InAbsoluteCursorPosition	The absolute cursor position
-	 * @return true if this widget considers itself hit by the passed in cursor position.
+	 * DEPRECATED: This function caused a lot of confusion, and we found that all known uses of it were unnecessary.
+	 * We will provide a less confusing API for supporting non-rectangular hittesting and pixel-perfect hit testing.
 	 */
-	virtual bool OnHitTest( const FGeometry& MyGeometry, FVector2D InAbsoluteCursorPosition );
+	DEPRECATED(4.5, "OnHitTest is no longer called. It will be replaced by more robust hit testing logic in the near future.")
+	virtual bool OnHitTest( const FGeometry& MyGeometry, FVector2D InAbsoluteCursorPosition ){return false;}
 
 	//
 	// KEYBOARD INPUT
diff --git a/Engine/Source/Runtime/UMG/Private/Slate/SObjectWidget.cpp b/Engine/Source/Runtime/UMG/Private/Slate/SObjectWidget.cpp
--- a/Engine/Source/Runtime/UMG/Private/Slate/SObjectWidget.cpp
+++ b/Engine/Source/Runtime/UMG/Private/Slate/SObjectWidget.cpp
@@ -35,12 +35,6 @@ void SObjectWidget::Tick(const FGeometry& AllottedGeometry, const double InCurre
 	}
 }
 
-//bool SObjectWidget::OnHitTest(const FGeometry& MyGeometry, FVector2D InAbsoluteCursorPosition)
-//{
-//	if ( WidgetObject && !WidgetObject->IsDesignTime() )
-//	{
-//	}
-//}
 
 int32 SObjectWidget::OnPaint(const FPaintArgs& Args, const FGeometry& AllottedGeometry, const FSlateRect& MyClippingRect, FSlateWindowElementList& OutDrawElements, int32 LayerId, const FWidgetStyle& InWidgetStyle, bool bParentEnabled) const
 {
diff --git a/Engine/Source/Runtime/UMG/Public/Slate/SObjectWidget.h b/Engine/Source/Runtime/UMG/Public/Slate/SObjectWidget.h
--- a/Engine/Source/Runtime/UMG/Public/Slate/SObjectWidget.h
+++ b/Engine/Source/Runtime/UMG/Public/Slate/SObjectWidget.h
@@ -29,7 +29,6 @@ class SObjectWidget : public SCompoundWidget, public FGCObject
 	// End of FGCObject interface
 
 	virtual void Tick(const FGeometry& AllottedGeometry, const double InCurrentTime, const float InDeltaTime) override;
-	//virtual bool OnHitTest(const FGeometry& MyGeometry, FVector2D InAbsoluteCursorPosition) override;
 	virtual int32 OnPaint(const FPaintArgs& Args, const FGeometry& AllottedGeometry, const FSlateRect& MyClippingRect, FSlateWindowElementList& OutDrawElements, int32 LayerId, const FWidgetStyle& InWidgetStyle, bool bParentEnabled) const override;
 
 	virtual bool SupportsKeyboardFocus() const override;
