UMG - We no longer allow the user to select "NoDrawType" for brush types.  Exposing more brush options to blueprints correctly.

diff --git a/Engine/Source/Runtime/SlateCore/Public/Styling/SlateBrush.h b/Engine/Source/Runtime/SlateCore/Public/Styling/SlateBrush.h
--- a/Engine/Source/Runtime/SlateCore/Public/Styling/SlateBrush.h
+++ b/Engine/Source/Runtime/SlateCore/Public/Styling/SlateBrush.h
@@ -11,13 +11,13 @@
 /**
  * Enumerates ways in which an image can be drawn.
  */
-UENUM()
+UENUM(BlueprintType)
 namespace ESlateBrushDrawType
 {
 	enum Type
 	{
 		/** Don't do anything */
-		NoDrawType,
+		NoDrawType UMETA(Hidden),
 
 		/** Draw a 3x3 box, where the sides and the middle stretch based on the Margin */
 		Box,
@@ -34,7 +34,7 @@ namespace ESlateBrushDrawType
 /**
  * Enumerates tiling options for image drawing.
  */
-UENUM()
+UENUM(BlueprintType)
 namespace ESlateBrushTileType
 {
 	enum Type
diff --git a/Engine/Source/Runtime/SlateCore/Public/Styling/SlateColor.h b/Engine/Source/Runtime/SlateCore/Public/Styling/SlateColor.h
--- a/Engine/Source/Runtime/SlateCore/Public/Styling/SlateColor.h
+++ b/Engine/Source/Runtime/SlateCore/Public/Styling/SlateColor.h
@@ -10,7 +10,7 @@
  *
  * Should we use the specified color? If not, then which color from the style should we use.
  */
-UENUM()
+UENUM(BlueprintType)
 namespace ESlateColorStylingMode
 {
 	enum Type
