Slate - We now use the virtual function for getting the foreground color in the SCompoundWidget when painting.

diff --git a/Engine/Source/Runtime/SlateCore/Private/Widgets/SCompoundWidget.cpp b/Engine/Source/Runtime/SlateCore/Private/Widgets/SCompoundWidget.cpp
--- a/Engine/Source/Runtime/SlateCore/Private/Widgets/SCompoundWidget.cpp
+++ b/Engine/Source/Runtime/SlateCore/Private/Widgets/SCompoundWidget.cpp
@@ -24,7 +24,7 @@ int32 SCompoundWidget::OnPaint( const FPaintArgs& Args, const FGeometry& Allotte
 
 		FWidgetStyle CompoundedWidgetStyle = FWidgetStyle(InWidgetStyle)
 			.BlendColorAndOpacityTint(ColorAndOpacity.Get())
-			.SetForegroundColor( ForegroundColor );
+			.SetForegroundColor( GetForegroundColor() );
 
 		return TheChild.Widget->Paint( Args.WithNewParent(this), TheChild.Geometry, ChildClippingRect, OutDrawElements, LayerId + 1, CompoundedWidgetStyle, ShouldBeEnabled( bParentEnabled ) );
 	}
