UMG - Adding more blueprint friendly metadata to Slate objects.

diff --git a/Engine/Source/Runtime/SlateCore/Public/Styling/SlateTypes.h b/Engine/Source/Runtime/SlateCore/Public/Styling/SlateTypes.h
--- a/Engine/Source/Runtime/SlateCore/Public/Styling/SlateTypes.h
+++ b/Engine/Source/Runtime/SlateCore/Public/Styling/SlateTypes.h
@@ -368,21 +368,21 @@ struct SLATECORE_API FComboBoxStyle : public FSlateWidgetStyle
 	/**
 	 * The style to use for our SComboButton
 	 */	
-	UPROPERTY(EditAnywhere, Category=Appearance, meta=(ShowOnlyInnerProperties))
+	UPROPERTY(EditAnywhere, BlueprintReadWrite, Category=Appearance, meta=( ShowOnlyInnerProperties ))
 	FComboButtonStyle ComboButtonStyle;
 	FComboBoxStyle& SetComboButtonStyle( const FComboButtonStyle& InComboButtonStyle ){ ComboButtonStyle = InComboButtonStyle; return *this; }
 
 	/**
 	 * The sound the button should play when pressed
 	 */	
-	UPROPERTY(EditAnywhere, Category=Sound, meta=(DisplayName="Pressed Sound"))
+	UPROPERTY(EditAnywhere, BlueprintReadWrite, Category=Sound, meta=( DisplayName="Pressed Sound" ))
 	FSlateSound PressedSlateSound;
 	FComboBoxStyle& SetPressedSound( const FSlateSound& InPressedSound ){ PressedSlateSound = InPressedSound; return *this; }
 
 	/**
 	 * The Sound to play when the selection is changed
 	 */	
-	UPROPERTY(EditAnywhere, Category=Sound, meta=(DisplayName="Selection Change Sound"))
+	UPROPERTY(EditAnywhere, BlueprintReadWrite, Category=Sound, meta=( DisplayName="Selection Change Sound" ))
 	FSlateSound SelectionChangeSlateSound;
 	FComboBoxStyle& SetSelectionChangeSound( const FSlateSound& InSelectionChangeSound ){ SelectionChangeSlateSound = InSelectionChangeSound; return *this; }
 
@@ -409,7 +409,7 @@ struct TStructOpsTypeTraits<FComboBoxStyle> : public TStructOpsTypeTraitsBase
 /**
  * Represents the appearance of an SHyperlink
  */
-USTRUCT()
+USTRUCT(BlueprintType)
 struct SLATECORE_API FHyperlinkStyle : public FSlateWidgetStyle
 {
 	GENERATED_USTRUCT_BODY()
@@ -426,17 +426,17 @@ struct SLATECORE_API FHyperlinkStyle : public FSlateWidgetStyle
 	static const FHyperlinkStyle& GetDefault();
 
 	/** Underline style */
-	UPROPERTY(EditAnywhere, Category=Appearance)
+	UPROPERTY(EditAnywhere, BlueprintReadWrite, Category=Appearance)
 	FButtonStyle UnderlineStyle;
 	FHyperlinkStyle& SetUnderlineStyle( const FButtonStyle& InUnderlineStyle ){ UnderlineStyle = InUnderlineStyle; return *this; }
 
 	/** Text style */
-	UPROPERTY(EditAnywhere, Category=Appearance)
+	UPROPERTY(EditAnywhere, BlueprintReadWrite, Category=Appearance)
 	FTextBlockStyle TextStyle;
 	FHyperlinkStyle& SetTextStyle( const FTextBlockStyle& InTextStyle ){ TextStyle = InTextStyle; return *this; }
 
 	/** Padding */
-	UPROPERTY(EditAnywhere, Category=Appearance)
+	UPROPERTY(EditAnywhere, BlueprintReadWrite, Category=Appearance)
 	FMargin Padding;
 	FHyperlinkStyle& SetPadding( const FMargin& InPadding ){ Padding = InPadding; return *this; }
 };
@@ -624,7 +624,7 @@ struct SLATECORE_API FEditableTextBoxStyle : public FSlateWidgetStyle
 /**
  * Represents the appearance of an SInlineEditableTextBlock
  */
-USTRUCT()
+USTRUCT(BlueprintType)
 struct SLATECORE_API FInlineEditableTextBlockStyle : public FSlateWidgetStyle
 {
 	GENERATED_USTRUCT_BODY()
@@ -641,12 +641,12 @@ struct SLATECORE_API FInlineEditableTextBlockStyle : public FSlateWidgetStyle
 	static const FInlineEditableTextBlockStyle& GetDefault();
 
 	/** The style of the editable text box, which dictates the font, color, and shadow options. */
-	UPROPERTY(EditAnywhere, Category=Appearance)
+	UPROPERTY(EditAnywhere, BlueprintReadWrite, Category=Appearance)
 	FEditableTextBoxStyle EditableTextBoxStyle;
 	FInlineEditableTextBlockStyle& SetEditableTextBoxStyle( const FEditableTextBoxStyle& InEditableTextBoxStyle ){ EditableTextBoxStyle = InEditableTextBoxStyle; return *this; }
 
 	/** The style of the text block, which dictates the font, color, and shadow options. Style overrides all other properties! */
-	UPROPERTY(EditAnywhere, Category=Appearance)
+	UPROPERTY(EditAnywhere, BlueprintReadWrite, Category=Appearance)
 	FTextBlockStyle TextStyle;
 	FInlineEditableTextBlockStyle& SetTextStyle( const FTextBlockStyle& InTextStyle ){ TextStyle = InTextStyle; return *this; }
 };
@@ -884,7 +884,7 @@ struct SLATECORE_API FInlineTextImageStyle : public FSlateWidgetStyle
 /**
  * Represents the appearance of an SSpinBox
  */
-USTRUCT()
+USTRUCT(BlueprintType)
 struct SLATECORE_API FSpinBoxStyle : public FSlateWidgetStyle
 {
 	GENERATED_USTRUCT_BODY()
@@ -901,37 +901,37 @@ struct SLATECORE_API FSpinBoxStyle : public FSlateWidgetStyle
 	static const FSpinBoxStyle& GetDefault();
 
 	/** Brush used to draw the background of the spinbox */
-	UPROPERTY(EditAnywhere, Category=Appearance)
+	UPROPERTY(EditAnywhere, BlueprintReadWrite, Category=Appearance)
 	FSlateBrush BackgroundBrush;
 	FSpinBoxStyle& SetBackgroundBrush( const FSlateBrush& InBackgroundBrush ){ BackgroundBrush = InBackgroundBrush; return *this; }
 
 	/** Brush used to draw the background of the spinbox when it's hovered over */
-	UPROPERTY(EditAnywhere, Category=Appearance)
+	UPROPERTY(EditAnywhere, BlueprintReadWrite, Category=Appearance)
 	FSlateBrush HoveredBackgroundBrush;
 	FSpinBoxStyle& SetHoveredBackgroundBrush( const FSlateBrush& InHoveredBackgroundBrush ){ HoveredBackgroundBrush = InHoveredBackgroundBrush; return *this; }
 
 	/** Brush used to fill the spinbox when it's active */
-	UPROPERTY(EditAnywhere, Category=Appearance)
+	UPROPERTY(EditAnywhere, BlueprintReadWrite, Category=Appearance)
 	FSlateBrush ActiveFillBrush;
 	FSpinBoxStyle& SetActiveFillBrush( const FSlateBrush& InActiveFillBrush ){ ActiveFillBrush = InActiveFillBrush; return *this; }
 
 	/** Brush used to fill the spinbox when it's inactive */
-	UPROPERTY(EditAnywhere, Category=Appearance)
+	UPROPERTY(EditAnywhere, BlueprintReadWrite, Category=Appearance)
 	FSlateBrush InactiveFillBrush;
 	FSpinBoxStyle& SetInactiveFillBrush( const FSlateBrush& InInactiveFillBrush ){ InactiveFillBrush = InInactiveFillBrush; return *this; }
 
 	/** Image used to draw the spinbox arrows */
-	UPROPERTY(EditAnywhere, Category=Appearance)
+	UPROPERTY(EditAnywhere, BlueprintReadWrite, Category=Appearance)
 	FSlateBrush ArrowsImage;
 	FSpinBoxStyle& SetArrowsImage( const FSlateBrush& InArrowsImage ){ ArrowsImage = InArrowsImage; return *this; }
 
 	/** Color used to draw the spinbox foreground elements */
-	UPROPERTY(EditAnywhere, Category=Appearance)
+	UPROPERTY(EditAnywhere, BlueprintReadWrite, Category=Appearance)
 	FSlateColor ForegroundColor;
 	FSpinBoxStyle& SetForegroundColor( const FSlateColor& InForegroundColor ){ ForegroundColor = InForegroundColor; return *this; }
 
 	/** Padding to add around the spinbox and its text */
-	UPROPERTY(EditAnywhere, Category=Appearance)
+	UPROPERTY(EditAnywhere, BlueprintReadWrite, Category=Appearance)
 	FMargin TextPadding;
 	FSpinBoxStyle& SetTextPadding( const FMargin& InTextPadding ){ TextPadding = InTextPadding; return *this; }
 };
@@ -940,7 +940,7 @@ struct SLATECORE_API FSpinBoxStyle : public FSlateWidgetStyle
 /**
  * Represents the appearance of an SSplitter
  */
-USTRUCT()
+USTRUCT(BlueprintType)
 struct SLATECORE_API FSplitterStyle : public FSlateWidgetStyle
 {
 	GENERATED_USTRUCT_BODY()
@@ -957,12 +957,12 @@ struct SLATECORE_API FSplitterStyle : public FSlateWidgetStyle
 	static const FSplitterStyle& GetDefault();
 
 	/** Brush used to draw the handle in its normal state */
-	UPROPERTY(EditAnywhere, Category=Appearance)
+	UPROPERTY(EditAnywhere, BlueprintReadWrite, Category=Appearance)
 	FSlateBrush HandleNormalBrush;
 	FSplitterStyle& SetHandleNormalBrush( const FSlateBrush& InHandleNormalBrush ){ HandleNormalBrush = InHandleNormalBrush; return *this; }
 
 	/** Brush used to draw the handle in its highlight state */
-	UPROPERTY(EditAnywhere, Category=Appearance)
+	UPROPERTY(EditAnywhere, BlueprintReadWrite, Category=Appearance)
 	FSlateBrush HandleHighlightBrush;
 	FSplitterStyle& SetHandleHighlightBrush( const FSlateBrush& InHandleHighlightBrush ){ HandleHighlightBrush = InHandleHighlightBrush; return *this; }
 };
@@ -971,7 +971,7 @@ struct SLATECORE_API FSplitterStyle : public FSlateWidgetStyle
 /**
  * Represents the appearance of an STableRow
  */
-USTRUCT()
+USTRUCT(BlueprintType)
 struct SLATECORE_API FTableRowStyle : public FSlateWidgetStyle
 {
 	GENERATED_USTRUCT_BODY()
@@ -988,57 +988,57 @@ struct SLATECORE_API FTableRowStyle : public FSlateWidgetStyle
 	static const FTableRowStyle& GetDefault();
 
 	/** Brush used as a selector when a row is focused */
-	UPROPERTY(EditAnywhere, Category=Appearance)
+	UPROPERTY(EditAnywhere, BlueprintReadWrite, Category=Appearance)
 	FSlateBrush SelectorFocusedBrush;
 	FTableRowStyle& SetSelectorFocusedBrush( const FSlateBrush& InSelectorFocusedBrush ){ SelectorFocusedBrush = InSelectorFocusedBrush; return *this; }
 
 	/** Brush used when a selected row is active and hovered */
-	UPROPERTY(EditAnywhere, Category=Appearance)
+	UPROPERTY(EditAnywhere, BlueprintReadWrite, Category=Appearance)
 	FSlateBrush ActiveHoveredBrush;
 	FTableRowStyle& SetActiveHoveredBrush( const FSlateBrush& InActiveHoveredBrush ){ ActiveHoveredBrush = InActiveHoveredBrush; return *this; }
 
 	/** Brush used when a selected row is active */
-	UPROPERTY(EditAnywhere, Category=Appearance)
+	UPROPERTY(EditAnywhere, BlueprintReadWrite, Category=Appearance)
 	FSlateBrush ActiveBrush;
 	FTableRowStyle& SetActiveBrush( const FSlateBrush& InActiveBrush ){ ActiveBrush = InActiveBrush; return *this; }
 
 	/** Brush used when an selected row is inactive and hovered */
-	UPROPERTY(EditAnywhere, Category=Appearance)
+	UPROPERTY(EditAnywhere, BlueprintReadWrite, Category=Appearance)
 	FSlateBrush InactiveHoveredBrush;
 	FTableRowStyle& SetInactiveHoveredBrush( const FSlateBrush& InInactiveHoveredBrush ){ InactiveHoveredBrush = InInactiveHoveredBrush; return *this; }
 
 	/** Brush used when a selected row is inactive */
-	UPROPERTY(EditAnywhere, Category=Appearance)
+	UPROPERTY(EditAnywhere, BlueprintReadWrite, Category=Appearance)
 	FSlateBrush InactiveBrush;
 	FTableRowStyle& SetInactiveBrush( const FSlateBrush& InInactiveBrush ){ InactiveBrush = InInactiveBrush; return *this; }
 
 	/** Brush used when an even row is hovered */
-	UPROPERTY(EditAnywhere, Category=Appearance)
+	UPROPERTY(EditAnywhere, BlueprintReadWrite, Category=Appearance)
 	FSlateBrush EvenRowBackgroundHoveredBrush;
 	FTableRowStyle& SetEvenRowBackgroundHoveredBrush( const FSlateBrush& InEvenRowBackgroundHoveredBrush ){ EvenRowBackgroundHoveredBrush = InEvenRowBackgroundHoveredBrush; return *this; }
 
 	/** Brush used when an even row is in its normal state */
-	UPROPERTY(EditAnywhere, Category=Appearance)
+	UPROPERTY(EditAnywhere, BlueprintReadWrite, Category=Appearance)
 	FSlateBrush EvenRowBackgroundBrush;
 	FTableRowStyle& SetEvenRowBackgroundBrush( const FSlateBrush& InEvenRowBackgroundBrush ){ EvenRowBackgroundBrush = InEvenRowBackgroundBrush; return *this; }
 
 	/** Brush used when an odd row is hovered */
-	UPROPERTY(EditAnywhere, Category=Appearance)
+	UPROPERTY(EditAnywhere, BlueprintReadWrite, Category=Appearance)
 	FSlateBrush OddRowBackgroundHoveredBrush;
 	FTableRowStyle& SetOddRowBackgroundHoveredBrush( const FSlateBrush& InOddRowBackgroundHoveredBrush ){ OddRowBackgroundHoveredBrush = InOddRowBackgroundHoveredBrush; return *this; }
 
 	/** Brush to used when an odd row is in its normal state */
-	UPROPERTY(EditAnywhere, Category=Appearance)
+	UPROPERTY(EditAnywhere, BlueprintReadWrite, Category=Appearance)
 	FSlateBrush OddRowBackgroundBrush;
 	FTableRowStyle& SetOddRowBackgroundBrush( const FSlateBrush& InOddRowBackgroundBrush ){ OddRowBackgroundBrush = InOddRowBackgroundBrush; return *this; }
 
 	/** Text color used for all rows */
-	UPROPERTY(EditAnywhere, Category=Appearance)
+	UPROPERTY(EditAnywhere, BlueprintReadWrite, Category=Appearance)
 	FSlateColor TextColor;
 	FTableRowStyle& SetTextColor( const FSlateColor& InTextColor ){ TextColor = InTextColor; return *this; }
 
 	/** Text color used for the selected row */
-	UPROPERTY(EditAnywhere, Category=Appearance)
+	UPROPERTY(EditAnywhere, BlueprintReadWrite, Category=Appearance)
 	FSlateColor SelectedTextColor;
 	FTableRowStyle& SetSelectedTextColor( const FSlateColor& InSelectedTextColor ){ SelectedTextColor = InSelectedTextColor; return *this; }
 };
@@ -1284,7 +1284,7 @@ struct SLATECORE_API FScrollBoxStyle : public FSlateWidgetStyle
 /**
 * Represents the appearance of an FScrollBorderStyle
 */
-USTRUCT()
+USTRUCT(BlueprintType)
 struct SLATECORE_API FScrollBorderStyle : public FSlateWidgetStyle
 {
 	GENERATED_USTRUCT_BODY()
@@ -1301,12 +1301,12 @@ struct SLATECORE_API FScrollBorderStyle : public FSlateWidgetStyle
 	static const FScrollBorderStyle& GetDefault();
 
 	/** Brush used to draw the top shadow of a scrollborder */
-	UPROPERTY( EditAnywhere, Category = Appearance )
+	UPROPERTY(EditAnywhere, BlueprintReadWrite, Category = Appearance)
 	FSlateBrush TopShadowBrush;
 	FScrollBorderStyle& SetTopShadowBrush( const FSlateBrush& InTopShadowBrush ){ TopShadowBrush = InTopShadowBrush; return *this; }
 
 	/** Brush used to draw the bottom shadow of a scrollborder */
-	UPROPERTY( EditAnywhere, Category = Appearance )
+	UPROPERTY(EditAnywhere, BlueprintReadWrite, Category = Appearance)
 	FSlateBrush BottomShadowBrush;
 	FScrollBorderStyle& SetBottomShadowBrush( const FSlateBrush& InBottomShadowBrush ){ BottomShadowBrush = InBottomShadowBrush; return *this; }
 };
@@ -1315,7 +1315,7 @@ struct SLATECORE_API FScrollBorderStyle : public FSlateWidgetStyle
 /**
  * Represents the appearance of an SWindow
  */
-USTRUCT()
+USTRUCT(BlueprintType)
 struct SLATECORE_API FWindowStyle : public FSlateWidgetStyle
 {
 	GENERATED_USTRUCT_BODY()
@@ -1332,67 +1332,67 @@ struct SLATECORE_API FWindowStyle : public FSlateWidgetStyle
 	static const FWindowStyle& GetDefault();
 
 	/** Style used to draw the window minimize button */
-	UPROPERTY(EditAnywhere, Category=Appearance)
+	UPROPERTY(EditAnywhere, BlueprintReadWrite, Category=Appearance)
 	FButtonStyle MinimizeButtonStyle;
 	FWindowStyle& SetMinimizeButtonStyle( const FButtonStyle& InMinimizeButtonStyle ){ MinimizeButtonStyle = InMinimizeButtonStyle; return *this; }
 
 	/** Style used to draw the window maximize button */
-	UPROPERTY(EditAnywhere, Category=Appearance)
+	UPROPERTY(EditAnywhere, BlueprintReadWrite, Category=Appearance)
 	FButtonStyle MaximizeButtonStyle;
 	FWindowStyle& SetMaximizeButtonStyle( const FButtonStyle& InMaximizeButtonStyle ){ MaximizeButtonStyle = InMaximizeButtonStyle; return *this; }
 
 	/** Style used to draw the window restore button */
-	UPROPERTY(EditAnywhere, Category=Appearance)
+	UPROPERTY(EditAnywhere, BlueprintReadWrite, Category=Appearance)
 	FButtonStyle RestoreButtonStyle;
 	FWindowStyle& SetRestoreButtonStyle( const FButtonStyle& InRestoreButtonStyle ){ RestoreButtonStyle = InRestoreButtonStyle; return *this; }
 
 	/** Style used to draw the window close button */
-	UPROPERTY(EditAnywhere, Category=Appearance)
+	UPROPERTY(EditAnywhere, BlueprintReadWrite, Category=Appearance)
 	FButtonStyle CloseButtonStyle;
 	FWindowStyle& SetCloseButtonStyle( const FButtonStyle& InCloseButtonStyle ){ CloseButtonStyle = InCloseButtonStyle; return *this; }
 
 	/** Style used to draw the window title text */
-	UPROPERTY(EditAnywhere, Category=Appearance)
+	UPROPERTY(EditAnywhere, BlueprintReadWrite, Category=Appearance)
 	FTextBlockStyle TitleTextStyle;
 	FWindowStyle& SetTitleTextStyle( const FTextBlockStyle& InTitleTextStyle ){ TitleTextStyle = InTitleTextStyle; return *this; }
 
 	/** Brush used to draw the window title area when the window is active */
-	UPROPERTY(EditAnywhere, Category=Appearance)
+	UPROPERTY(EditAnywhere, BlueprintReadWrite, Category=Appearance)
 	FSlateBrush ActiveTitleBrush;
 	FWindowStyle& SetActiveTitleBrush( const FSlateBrush& InActiveTitleBrush ){ ActiveTitleBrush = InActiveTitleBrush; return *this; }
 
 	/** Brush used to draw the window title area when the window is inactive */
-	UPROPERTY(EditAnywhere, Category=Appearance)
+	UPROPERTY(EditAnywhere, BlueprintReadWrite, Category=Appearance)
 	FSlateBrush InactiveTitleBrush;
 	FWindowStyle& SetInactiveTitleBrush( const FSlateBrush& InInactiveTitleBrush ){ InactiveTitleBrush = InInactiveTitleBrush; return *this; }
 
 	/** Brush used to draw the window title area when the window is flashing */
-	UPROPERTY(EditAnywhere, Category=Appearance)
+	UPROPERTY(EditAnywhere, BlueprintReadWrite, Category=Appearance)
 	FSlateBrush FlashTitleBrush;
 	FWindowStyle& SetFlashTitleBrush( const FSlateBrush& InFlashTitleBrush ){ FlashTitleBrush = InFlashTitleBrush; return *this; }
 
 	/** Brush used to draw the window outline */
-	UPROPERTY(EditAnywhere, Category=Appearance)
+	UPROPERTY(EditAnywhere, BlueprintReadWrite, Category=Appearance)
 	FSlateBrush OutlineBrush;
 	FWindowStyle& SetOutlineBrush( const FSlateBrush& InOutlineBrush ){ OutlineBrush = InOutlineBrush; return *this; }
 
 	/** Color used to draw the window outline */
-	UPROPERTY(EditAnywhere, Category=Appearance)
+	UPROPERTY(EditAnywhere, BlueprintReadWrite, Category=Appearance)
 	FSlateColor OutlineColor;
 	FWindowStyle& SetOutlineColor( const FSlateColor& InOutlineColor ){ OutlineColor = InOutlineColor; return *this; }
 
 	/** Brush used to draw the window border */
-	UPROPERTY(EditAnywhere, Category=Appearance)
+	UPROPERTY(EditAnywhere, BlueprintReadWrite, Category=Appearance)
 	FSlateBrush BorderBrush;
 	FWindowStyle& SetBorderBrush( const FSlateBrush& InBorderBrush ){ BorderBrush = InBorderBrush; return *this; }
 
 	/** Brush used to draw the window background */
-	UPROPERTY(EditAnywhere, Category=Appearance)
+	UPROPERTY(EditAnywhere, BlueprintReadWrite, Category=Appearance)
 	FSlateBrush BackgroundBrush;
 	FWindowStyle& SetBackgroundBrush( const FSlateBrush& InBackgroundBrush ){ BackgroundBrush = InBackgroundBrush; return *this; }
 
 	/** Brush used to draw the background of child windows */
-	UPROPERTY(EditAnywhere, Category=Appearance)
+	UPROPERTY(EditAnywhere, BlueprintReadWrite, Category=Appearance)
 	FSlateBrush ChildBackgroundBrush;
 	FWindowStyle& SetChildBackgroundBrush( const FSlateBrush& InChildBackgroundBrush ){ ChildBackgroundBrush = InChildBackgroundBrush; return *this; }
 };
