Use non-deprecated hash for font caching strings

diff --git a/Engine/Source/Runtime/SlateCore/Private/Fonts/FontCacheUtils.h b/Engine/Source/Runtime/SlateCore/Private/Fonts/FontCacheUtils.h
--- a/Engine/Source/Runtime/SlateCore/Private/Fonts/FontCacheUtils.h
+++ b/Engine/Source/Runtime/SlateCore/Private/Fonts/FontCacheUtils.h
@@ -159,7 +159,7 @@ private:
 
 		FORCEINLINE static uint32 GetKeyHash(const KeyType& Identifier)
 		{
-			return FCrc::StrCrc_DEPRECATED( *Identifier );
+			return FCrc::StrCrc32( *Identifier );
 		}
 	};
 
