Slate : Deprecated .AspectRatio() and removed all references.

diff --git a/Engine/Source/Developer/CollisionAnalyzer/Private/SCollisionAnalyzer.cpp b/Engine/Source/Developer/CollisionAnalyzer/Private/SCollisionAnalyzer.cpp
--- a/Engine/Source/Developer/CollisionAnalyzer/Private/SCollisionAnalyzer.cpp
+++ b/Engine/Source/Developer/CollisionAnalyzer/Private/SCollisionAnalyzer.cpp
@@ -45,7 +45,6 @@ void SCollisionAnalyzer::Construct(const FArguments& InArgs, FCollisionAnalyzer*
 				+SHorizontalBox::Slot()
 				.AutoWidth()
 				.Padding(1)
-				.AspectRatio()
 				[
 					SNew(SCheckBox)
 					.Style(FEditorStyle::Get(), "ToggleButtonCheckbox")
diff --git a/Engine/Source/Developer/LogVisualizer/Private/SLogVisualizer.cpp b/Engine/Source/Developer/LogVisualizer/Private/SLogVisualizer.cpp
--- a/Engine/Source/Developer/LogVisualizer/Private/SLogVisualizer.cpp
+++ b/Engine/Source/Developer/LogVisualizer/Private/SLogVisualizer.cpp
@@ -163,7 +163,6 @@ void SLogVisualizer::Construct(const FArguments& InArgs, FLogVisualizer* InLogVi
 					+SHorizontalBox::Slot()
 					.AutoWidth()
 					.Padding(1)
-					.AspectRatio()
 					[
 						SNew(SButton)
 						.OnClicked(this, &SLogVisualizer::OnRecordButtonClicked)
@@ -177,7 +176,6 @@ void SLogVisualizer::Construct(const FArguments& InArgs, FLogVisualizer* InLogVi
 					+SHorizontalBox::Slot()
 					.AutoWidth()
 					.Padding(1)
-					.AspectRatio()
 					[
 						SNew(SCheckBox)
 						.Style(FEditorStyle::Get(), "ToggleButtonCheckbox")
@@ -193,7 +191,6 @@ void SLogVisualizer::Construct(const FArguments& InArgs, FLogVisualizer* InLogVi
 					+SHorizontalBox::Slot()
 					.AutoWidth()
 					.Padding(1)
-					.AspectRatio()
 					[
 						SNew(SCheckBox)
 						.Style(FEditorStyle::Get(), "ToggleButtonCheckbox")
@@ -208,7 +205,6 @@ void SLogVisualizer::Construct(const FArguments& InArgs, FLogVisualizer* InLogVi
 					+SHorizontalBox::Slot()
 					.MaxWidth(3)
 					.Padding(1)
-					.AspectRatio()
 					[
 						SNew(SSeparator)
 						.Orientation(Orient_Vertical)
@@ -217,7 +213,6 @@ void SLogVisualizer::Construct(const FArguments& InArgs, FLogVisualizer* InLogVi
 					+SHorizontalBox::Slot()
 					.AutoWidth()
 					.Padding(1)
-					.AspectRatio()
 					[
 						SNew(SButton)
 						.OnClicked(this, &SLogVisualizer::OnSave)
@@ -231,7 +226,6 @@ void SLogVisualizer::Construct(const FArguments& InArgs, FLogVisualizer* InLogVi
 					+SHorizontalBox::Slot()
 					.AutoWidth()
 					.Padding(1)
-					.AspectRatio()
 					[
 						SNew(SButton)
 						.OnClicked(this, &SLogVisualizer::OnLoad)
@@ -245,7 +239,6 @@ void SLogVisualizer::Construct(const FArguments& InArgs, FLogVisualizer* InLogVi
 					+SHorizontalBox::Slot()
 					.AutoWidth()
 					.Padding(1)
-					.AspectRatio()
 					[
 						SNew(SButton)
 						.OnClicked(this, &SLogVisualizer::OnRemove)
@@ -258,7 +251,6 @@ void SLogVisualizer::Construct(const FArguments& InArgs, FLogVisualizer* InLogVi
 					+SHorizontalBox::Slot()
 					.MaxWidth(3)
 					.Padding(1)
-					.AspectRatio()
 					[
 						SNew(SSeparator)
 						.Orientation(Orient_Vertical)
diff --git a/Engine/Source/Editor/BehaviorTreeEditor/Private/SGraphNode_BehaviorTree.cpp b/Engine/Source/Editor/BehaviorTreeEditor/Private/SGraphNode_BehaviorTree.cpp
--- a/Engine/Source/Editor/BehaviorTreeEditor/Private/SGraphNode_BehaviorTree.cpp
+++ b/Engine/Source/Editor/BehaviorTreeEditor/Private/SGraphNode_BehaviorTree.cpp
@@ -441,14 +441,11 @@ void SGraphNode_BehaviorTree::UpdateGraphNode()
 										.ToolTipText( this, &SGraphNode_BehaviorTree::GetErrorMsgToolTip )
 									]
 									+SHorizontalBox::Slot()
-									.AspectRatio()
+									.AutoWidth()
+									.VAlign(VAlign_Center)
 									[
-										SNew(SOverlay)
-										+SOverlay::Slot()
-										[
-											SNew(SImage)
-											.Image(this, &SGraphNode_BehaviorTree::GetNameIcon)
-										]
+										SNew(SImage)
+										.Image(this, &SGraphNode_BehaviorTree::GetNameIcon)
 									]
 									+SHorizontalBox::Slot()
 									.Padding(FMargin(4.0f, 0.0f, 4.0f, 0.0f))
diff --git a/Engine/Source/Editor/GraphEditor/Private/AnimationNodes/SGraphNodeAnimationResult.cpp b/Engine/Source/Editor/GraphEditor/Private/AnimationNodes/SGraphNodeAnimationResult.cpp
--- a/Engine/Source/Editor/GraphEditor/Private/AnimationNodes/SGraphNodeAnimationResult.cpp
+++ b/Engine/Source/Editor/GraphEditor/Private/AnimationNodes/SGraphNodeAnimationResult.cpp
@@ -36,7 +36,8 @@ TSharedRef<SWidget> SGraphNodeAnimationResult::CreateNodeContentArea()
 				SAssignNew(LeftNodeBox, SVerticalBox)
 			]
 			+SHorizontalBox::Slot()
-			.AspectRatio()
+			.AutoWidth()
+			.VAlign(VAlign_Center)
 			[
 				SNew(SImage)
 				.Image( FEditorStyle::GetBrush("Graph.AnimationResultNode.Body") )
diff --git a/Engine/Source/Editor/GraphEditor/Private/AnimationStateNodes/SGraphNodeAnimState.cpp b/Engine/Source/Editor/GraphEditor/Private/AnimationStateNodes/SGraphNodeAnimState.cpp
--- a/Engine/Source/Editor/GraphEditor/Private/AnimationStateNodes/SGraphNodeAnimState.cpp
+++ b/Engine/Source/Editor/GraphEditor/Private/AnimationStateNodes/SGraphNodeAnimState.cpp
@@ -219,7 +219,8 @@ void SGraphNodeAnimState::UpdateGraphNode()
 							.ToolTipText( this, &SGraphNodeAnimState::GetErrorMsgToolTip )
 						]
 						+SHorizontalBox::Slot()
-						.AspectRatio()
+						.AutoWidth()
+						.VAlign(VAlign_Center)
 						[
 							SNew(SImage)
 							.Image(NodeTypeIcon)
diff --git a/Engine/Source/Editor/GraphEditor/Private/SoundNodes/SGraphNodeSoundResult.cpp b/Engine/Source/Editor/GraphEditor/Private/SoundNodes/SGraphNodeSoundResult.cpp
--- a/Engine/Source/Editor/GraphEditor/Private/SoundNodes/SGraphNodeSoundResult.cpp
+++ b/Engine/Source/Editor/GraphEditor/Private/SoundNodes/SGraphNodeSoundResult.cpp
@@ -35,7 +35,8 @@ TSharedRef<SWidget> SGraphNodeSoundResult::CreateNodeContentArea()
 				SAssignNew(LeftNodeBox, SVerticalBox)
 			]
 			+SHorizontalBox::Slot()
-			.AspectRatio()
+			.AutoWidth()
+			.VAlign(VAlign_Center)
 			[
 				SNew(SImage)
 				.Image(FEditorStyle::GetBrush("Graph.SoundResultNode.Body"))
diff --git a/Engine/Source/Editor/LandscapeEditor/Private/LandscapeEditorDetailCustomization_NewLandscape.cpp b/Engine/Source/Editor/LandscapeEditor/Private/LandscapeEditorDetailCustomization_NewLandscape.cpp
--- a/Engine/Source/Editor/LandscapeEditor/Private/LandscapeEditorDetailCustomization_NewLandscape.cpp
+++ b/Engine/Source/Editor/LandscapeEditor/Private/LandscapeEditorDetailCustomization_NewLandscape.cpp
@@ -1396,7 +1396,6 @@ void FLandscapeEditorStructCustomization_FLandscapeImportLayer::CustomizeChildre
 		.AutoWidth()
 		.VAlign(VAlign_Center)
 		.Padding(FMargin(2))
-		//.AspectRatio()
 		[
 			SNew(SLandscapeAssetThumbnail, ThumbnailMIC, StructCustomizationUtils.GetThumbnailPool().ToSharedRef())
 			.ThumbnailSize(FIntPoint(48, 48))
diff --git a/Engine/Source/Editor/LandscapeEditor/Private/LandscapeEditorDetailCustomization_TargetLayers.cpp b/Engine/Source/Editor/LandscapeEditor/Private/LandscapeEditorDetailCustomization_TargetLayers.cpp
--- a/Engine/Source/Editor/LandscapeEditor/Private/LandscapeEditorDetailCustomization_TargetLayers.cpp
+++ b/Engine/Source/Editor/LandscapeEditor/Private/LandscapeEditorDetailCustomization_TargetLayers.cpp
@@ -170,7 +170,6 @@ void FLandscapeEditorCustomNodeBuilder_TargetLayers::GenerateRow(IDetailChildren
 				.AutoWidth()
 				.VAlign(VAlign_Center)
 				.Padding(FMargin(2))
-				//.AspectRatio()
 				[
 					SNew(SImage)
 					.Image(FEditorStyle::GetBrush(Target->TargetType == ELandscapeToolTargetType::Heightmap ? TEXT("LandscapeEditor.Target_Heightmap") : TEXT("LandscapeEditor.Target_Visibility")))
@@ -210,7 +209,6 @@ void FLandscapeEditorCustomNodeBuilder_TargetLayers::GenerateRow(IDetailChildren
 				.AutoWidth()
 				.VAlign(VAlign_Center)
 				.Padding(FMargin(2))
-				//.AspectRatio()
 				[
 					(Target->bValid)
 					? (TSharedRef<SWidget>)(
diff --git a/Engine/Source/Editor/LevelEditor/Private/SToolkitDisplay.cpp b/Engine/Source/Editor/LevelEditor/Private/SToolkitDisplay.cpp
--- a/Engine/Source/Editor/LevelEditor/Private/SToolkitDisplay.cpp
+++ b/Engine/Source/Editor/LevelEditor/Private/SToolkitDisplay.cpp
@@ -98,7 +98,7 @@ void SLevelEditorActiveToolkit::Construct( const FArguments&, const TSharedPtr<
 
 		ContentBox->AddSlot()
 			.AutoWidth()
-			.AspectRatio()
+			.VAlign(VAlign_Center)
 			.HAlign( HAlign_Right )
 			[
 				SNew( SComboButton )
diff --git a/Engine/Source/Editor/Matinee/Private/SMatineeRecorder.cpp b/Engine/Source/Editor/Matinee/Private/SMatineeRecorder.cpp
--- a/Engine/Source/Editor/Matinee/Private/SMatineeRecorder.cpp
+++ b/Engine/Source/Editor/Matinee/Private/SMatineeRecorder.cpp
@@ -81,7 +81,6 @@ void SMatineeRecorder::Construct(const FArguments& InArgs)
 						SNew(SHorizontalBox)
 						+ SHorizontalBox::Slot()
 						.AutoWidth()
-						. AspectRatio()
 						[
 							SNew(SImage)
 							.Image( this, &SMatineeRecorder::GetRecordImageDelegate )
diff --git a/Engine/Source/Editor/PropertyEditor/Private/SDetailNameArea.cpp b/Engine/Source/Editor/PropertyEditor/Private/SDetailNameArea.cpp
--- a/Engine/Source/Editor/PropertyEditor/Private/SDetailNameArea.cpp
+++ b/Engine/Source/Editor/PropertyEditor/Private/SDetailNameArea.cpp
@@ -87,7 +87,6 @@ TSharedRef< SWidget > SDetailNameArea::BuildObjectNameArea( const TArray< TWeakO
 		.HAlign(HAlign_Left)
 		.VAlign(VAlign_Center)
 		.Padding(0,0,0,0)
-		.AspectRatio()
 		[
 			SNew(SImage)
 			.Image(ActorIcon)
diff --git a/Engine/Source/Runtime/AppFramework/Private/Framework/Testing/SLayoutExample.cpp b/Engine/Source/Runtime/AppFramework/Private/Framework/Testing/SLayoutExample.cpp
--- a/Engine/Source/Runtime/AppFramework/Private/Framework/Testing/SLayoutExample.cpp
+++ b/Engine/Source/Runtime/AppFramework/Private/Framework/Testing/SLayoutExample.cpp
@@ -115,24 +115,7 @@ public:
 			]
 			+SScrollBox::Slot() .Padding(5)
 			[
-				SNew(SHorizontalBox)
-				+SHorizontalBox::Slot()
-				.AutoWidth()
-				[
-					SNew(SHorizontalBox)
-					+SHorizontalBox::Slot()
-					.FillWidth(1)
-					[
-						SNew(STextBlock) .Text( LOCTEXT("ExampleLayout-TextLabel09", "A somewhat long piece of text.") )
-					]
-					+SHorizontalBox::Slot()
-					.AutoWidth()
-					.AspectRatio()
-					.Padding( 2.0f, 1.0f )
-					[
-						SNew(SImage) .Image( FCoreStyle::Get().GetBrush("ToolBar.SubMenuIndicator") )
-					]
-				]
+				SNew(STextBlock) .Text( LOCTEXT("ExampleLayout-TextLabel09", "A somewhat long piece of text.") )
 			]
 			// Fixed Size example
 			+SScrollBox::Slot() .Padding(5)
diff --git a/Engine/Source/Runtime/AppFramework/Private/Widgets/Testing/STestSuite.cpp b/Engine/Source/Runtime/AppFramework/Private/Widgets/Testing/STestSuite.cpp
--- a/Engine/Source/Runtime/AppFramework/Private/Widgets/Testing/STestSuite.cpp
+++ b/Engine/Source/Runtime/AppFramework/Private/Widgets/Testing/STestSuite.cpp
@@ -3348,7 +3348,6 @@ public:
 								+ SVerticalBox::Slot()
 									.AutoHeight()
 									. HAlign(HAlign_Center)
-									. AspectRatio()
 								[
 									SNew(SImage)
 										. Image( FCoreStyle::Get().GetBrush( TEXT("DefaultAppIcon") ) )
@@ -3372,13 +3371,12 @@ public:
 									[
 										SNew(SHorizontalBox)
 										+ SHorizontalBox::Slot()
-											.AutoWidth()
-											. AspectRatio()
+										.AutoWidth()
 										[
 											SNew(SImage) .Image( FCoreStyle::Get().GetBrush( TEXT("DefaultAppIcon") ) )
 										]
 										+ SHorizontalBox::Slot()
-											.AutoWidth()
+										.AutoWidth()
 										[
 											SNew(STextBlock) .Text( LOCTEXT("ButtonTextLabel", "Button with content") )
 										]
@@ -4496,10 +4494,9 @@ TSharedRef<SDockTab> SpawnTab(const FSpawnTabArgs& Args, FName TabIdentifier)
 					+SHorizontalBox::Slot()
 					.AutoWidth()
 					.Padding( 2.0f )
-					.AspectRatio()
 					[
 						SNew(SImage)
-							.Image( FCoreStyle::Get().GetBrush( TEXT("DefaultAppIcon") ) )
+						.Image( FCoreStyle::Get().GetBrush( TEXT("DefaultAppIcon") ) )
 					]
 					+SHorizontalBox::Slot()
 						.AutoWidth()
diff --git a/Engine/Source/Runtime/Slate/Private/Widgets/Input/SSearchBox.cpp b/Engine/Source/Runtime/Slate/Private/Widgets/Input/SSearchBox.cpp
--- a/Engine/Source/Runtime/Slate/Private/Widgets/Input/SSearchBox.cpp
+++ b/Engine/Source/Runtime/Slate/Private/Widgets/Input/SSearchBox.cpp
@@ -37,7 +37,6 @@ void SSearchBox::Construct( const FArguments& InArgs )
 	{
 		Box->AddSlot()
 		.AutoWidth()
-		.AspectRatio()
 		.HAlign(HAlign_Center)
 		.VAlign(VAlign_Center)
 		[
@@ -57,7 +56,6 @@ void SSearchBox::Construct( const FArguments& InArgs )
 		];
 		Box->AddSlot()
 		.AutoWidth()
-		.AspectRatio()
 		.HAlign(HAlign_Center)
 		.VAlign(VAlign_Center)
 		[
@@ -80,7 +78,6 @@ void SSearchBox::Construct( const FArguments& InArgs )
 	// Add a search glass image so that the user knows this text box is for searching
 	Box->AddSlot()
 	.AutoWidth()
-	.AspectRatio()
 	.HAlign(HAlign_Center)
 	.VAlign(VAlign_Center)
 	[
@@ -93,7 +90,6 @@ void SSearchBox::Construct( const FArguments& InArgs )
 	// Add an X to clear the search whenever there is some text typed into it
 	Box->AddSlot()
 	.AutoWidth()
-	.AspectRatio()
 	[
 		SNew(SButton)
 		.Visibility(this, &SSearchBox::GetXVisibility)
diff --git a/Engine/Source/Runtime/SlateCore/Private/Widgets/SBoxPanel.cpp b/Engine/Source/Runtime/SlateCore/Private/Widgets/SBoxPanel.cpp
--- a/Engine/Source/Runtime/SlateCore/Private/Widgets/SBoxPanel.cpp
+++ b/Engine/Source/Runtime/SlateCore/Private/Widgets/SBoxPanel.cpp
@@ -60,19 +60,6 @@ static void ArrangeChildrenAlong( const TPanelChildren<SBoxPanel::FSlot>& Childr
 					// for stretch children we save sum up the stretch coefficients
 					StretchCoefficientTotal += CurChild.SizeParam.Value.Get();
 				}
-				else if (CurChild.SizeParam.SizeRule == FSizeParam::SizeRule_AspectRatio)
-				{
-					// Aspect ratio children contribute size such that they maintain
-					// the aspect ratio of 1 (i.e. they're square).
-					float ChildSize = (Orientation == Orient_Vertical)
-						? AllottedGeometry.Size.X
-						: AllottedGeometry.Size.Y;
-
-					// Clamp to the max size if it was specified
-					float MaxSize = CurChild.MaxSize.Get();
-					FixedTotal += MaxSize > 0 ? FMath::Min( MaxSize, ChildSize ) : ChildSize ;
-				
-				}
 				else
 				{
 					// Auto-sized children contribute their desired size to the fixed space requirement
@@ -115,14 +102,6 @@ static void ArrangeChildrenAlong( const TPanelChildren<SBoxPanel::FSlot>& Childr
 					// Stretch widgets get a fraction of the space remaining after all the fixed-space requirements are met
 					ChildSize = NonFixedSpace * CurChild.SizeParam.Value.Get() / StretchCoefficientTotal;
 				}
-				else if (CurChild.SizeParam.SizeRule == FSizeParam::SizeRule_AspectRatio)
-				{
-					// Aspect ratio children contribute size such that they maintain
-					// the aspect ratio of 1.
-					ChildSize = (Orientation == Orient_Vertical)
-						? AllottedGeometry.Size.X
-						: AllottedGeometry.Size.Y;
-				}
 				else
 				{
 					// Auto-sized widgets get their desired-size value
@@ -204,7 +183,6 @@ static FVector2D ComputeDesiredSizeForBox( const TPanelChildren<SBoxPanel::FSlot
 	// The layout along the panel's axis is describe dy the SizeParam, while the perpendicular layout is described by the
 	// alignment property.
 	FVector2D MyDesiredSize(0,0);
-	int32 NumAspectRatioChildren = 0;
 	for( int32 ChildIndex=0; ChildIndex < Children.Num(); ++ChildIndex )
 	{
 		const SBoxPanel::FSlot& CurChild = Children[ChildIndex];
@@ -217,24 +195,15 @@ static FVector2D ComputeDesiredSizeForBox( const TPanelChildren<SBoxPanel::FSlot
 				// That will be the desired width of the whole panel.
 				MyDesiredSize.X = FMath::Max(MyDesiredSize.X, CurChildDesiredSize.X + CurChild.SlotPadding.Get().GetTotalSpaceAlong<Orient_Horizontal>());
 
-				if ( CurChild.SizeParam.SizeRule == FSizeParam::SizeRule_AspectRatio )
+				// Clamp to the max size if it was specified
+				float FinalChildDesiredSize = CurChildDesiredSize.Y;
+				float MaxSize = CurChild.MaxSize.Get();
+				if( MaxSize > 0 )
 				{
-					// Omit aspect ratio children for now. We do not know how big they need to be yet.
-					NumAspectRatioChildren++;
+					FinalChildDesiredSize = FMath::Min( MaxSize, FinalChildDesiredSize );
 				}
-				else
-				{
-					// Clamp to the max size if it was specified
-					float FinalChildDesiredSize = CurChildDesiredSize.Y;
-					float MaxSize = CurChild.MaxSize.Get();
-					if( MaxSize > 0 )
-					{
-						FinalChildDesiredSize = FMath::Min( MaxSize, FinalChildDesiredSize );
-					}
 
-					MyDesiredSize.Y += FinalChildDesiredSize + CurChild.SlotPadding.Get().GetTotalSpaceAlong<Orient_Vertical>();
-				}
-						
+				MyDesiredSize.Y += FinalChildDesiredSize + CurChild.SlotPadding.Get().GetTotalSpaceAlong<Orient_Vertical>();				
 			}
 			else
 			{
@@ -242,38 +211,19 @@ static FVector2D ComputeDesiredSizeForBox( const TPanelChildren<SBoxPanel::FSlot
 
 				MyDesiredSize.Y = FMath::Max(MyDesiredSize.Y, CurChildDesiredSize.Y + CurChild.SlotPadding.Get().GetTotalSpaceAlong<Orient_Vertical>());
 
-				if ( CurChild.SizeParam.SizeRule == FSizeParam::SizeRule_AspectRatio )
+				// Clamp to the max size if it was specified
+				float FinalChildDesiredSize = CurChildDesiredSize.X;
+				float MaxSize = CurChild.MaxSize.Get();
+				if( MaxSize > 0 )
 				{
-					// Omit aspect ratio children for now. We do not know how big they need to be yet.
-					NumAspectRatioChildren++;
+					FinalChildDesiredSize = FMath::Min( MaxSize, FinalChildDesiredSize );
 				}
-				else
-				{
-					// Clamp to the max size if it was specified
-					float FinalChildDesiredSize = CurChildDesiredSize.X;
-					float MaxSize = CurChild.MaxSize.Get();
-					if( MaxSize > 0 )
-					{
-						FinalChildDesiredSize = FMath::Min( MaxSize, FinalChildDesiredSize );
-					}
 
-					MyDesiredSize.X += FinalChildDesiredSize + CurChild.SlotPadding.Get().GetTotalSpaceAlong<Orient_Horizontal>();
-				}
+				MyDesiredSize.X += FinalChildDesiredSize + CurChild.SlotPadding.Get().GetTotalSpaceAlong<Orient_Horizontal>();
 			}
 		}
 	}
 
-	// Aspect ratio children need room proportional to the height/width of the entire box
-	if ( Orientation == Orient_Vertical )
-	{
-		MyDesiredSize.Y += MyDesiredSize.X * NumAspectRatioChildren;
-	}
-	else
-	{
-		MyDesiredSize.X += MyDesiredSize.Y * NumAspectRatioChildren;
-	}
-
-
 	return MyDesiredSize;
 }
 
diff --git a/Engine/Source/Runtime/SlateCore/Public/Types/SlateStructs.h b/Engine/Source/Runtime/SlateCore/Public/Types/SlateStructs.h
--- a/Engine/Source/Runtime/SlateCore/Public/Types/SlateStructs.h
+++ b/Engine/Source/Runtime/SlateCore/Public/Types/SlateStructs.h
@@ -66,7 +66,6 @@ private:
  * Describes a way in which a parent widget allocates available space to its child widgets.
  *
  * When SizeRule is SizeRule_Auto, the widget's DesiredSize will be used as the space required.
- * When SizeRule is SizeRule_AspectRatio, the widget will attempt to maintain the specified aspect ratio.
  * When SizeRule is SizeRule_Stretch, the available space will be distributed proportionately between
  * peer Widgets depending on the Value property. Available space is space remaining after all the
  * peers' SizeRule_Auto requirements have been satisfied.
@@ -78,8 +77,7 @@ struct FSizeParam
 	enum ESizeRule
 	{
 		SizeRule_Auto,
-		SizeRule_Stretch,
-		SizeRule_AspectRatio
+		SizeRule_Stretch
 	};
 	
 	/** The sizing rule to use. */
@@ -139,17 +137,3 @@ struct FAuto
 		: FSizeParam(SizeRule_Auto, 0.0f)
 	{ }
 };
-
-
-/**
- * Structure for size parameters with SizeRule == SizeRule_AspectRatio.
- *
- * @see FAuto, FStretch
- */
-struct FAspectRatio
-	: public FSizeParam
-{
-	FAspectRatio( )
-		: FSizeParam(SizeRule_AspectRatio, 1.0f)
-	{ }
-};
diff --git a/Engine/Source/Runtime/SlateCore/Public/Widgets/SBoxPanel.h b/Engine/Source/Runtime/SlateCore/Public/Widgets/SBoxPanel.h
--- a/Engine/Source/Runtime/SlateCore/Public/Widgets/SBoxPanel.h
+++ b/Engine/Source/Runtime/SlateCore/Public/Widgets/SBoxPanel.h
@@ -23,7 +23,6 @@ public:
 		/**
 		 * How much space this slot should occupy along panel's direction.
 		 *   When SizeRule is SizeRule_Auto, the widget's DesiredSize will be used as the space required.
-		 *   When SizeRule is SizeRule_AspectRatio, the widget will attempt to maintain the specified aspect ratio.
 		 *   When SizeRule is SizeRule_Stretch, the available space will be distributed proportionately between
 		 *   peer Widgets depending on the Value property. Available space is space remaining after all the
 		 *   peers' SizeRule_Auto requirements have been satisfied.
@@ -125,9 +124,9 @@ public:
 			return *this;
 		}
 
+		DEPRECATED( 4.5, "AspectRatio() property is no longer supported; it did not behave correctly to begin with and was rarely used." )
 		FSlot& AspectRatio()
 		{
-			SizeParam = FAspectRatio();
 			return *this;
 		}
 		
@@ -245,9 +244,9 @@ public:
 			return *this;
 		}
 
+		DEPRECATED( 4.5, "AspectRatio() property is no longer supported; it did not behave correctly to begin with and was rarely used." )
 		FSlot& AspectRatio()
 		{
-			SizeParam = FAspectRatio();
 			return *this;
 		}
 
diff --git a/Engine/Source/Runtime/SlateCore/Public/Widgets/SWidget.h b/Engine/Source/Runtime/SlateCore/Public/Widgets/SWidget.h
--- a/Engine/Source/Runtime/SlateCore/Public/Widgets/SWidget.h
+++ b/Engine/Source/Runtime/SlateCore/Public/Widgets/SWidget.h
@@ -449,7 +449,7 @@ public:
 	 *
 	 * @return The desired size.
 	 */
-	protected: virtual FVector2D ComputeDesiredSize() const = 0;
+	private: virtual FVector2D ComputeDesiredSize() const = 0;
 
 	public:
 
