Added ability to specify a font hinting algorithm to use when rendering fonts with slate

Available algorithms are:
    Default: Prefer to use hinting embedded into the font
    Auto: Use freetype's automatic hinting algorithm
    AutoLight: Use a lighter automatic hinting algorithm
    Monochrome: Use a hinting algorithm optimised for monochrome displays
    None: Disable hinting altogether

diff --git a/Engine/Source/Runtime/SlateCore/Private/Fonts/FontCache.cpp b/Engine/Source/Runtime/SlateCore/Private/Fonts/FontCache.cpp
--- a/Engine/Source/Runtime/SlateCore/Private/Fonts/FontCache.cpp
+++ b/Engine/Source/Runtime/SlateCore/Private/Fonts/FontCache.cpp
@@ -37,7 +37,7 @@ namespace FontCacheConstants
 }
 
 #if WITH_FREETYPE
-const uint32 GlyphFlags = FT_LOAD_TARGET_NORMAL | FT_LOAD_NO_BITMAP;
+const uint32 GlyphFlags = FT_LOAD_NO_BITMAP;
 
 /**
  * Memory allocation functions to be used only by freetype
@@ -135,6 +135,16 @@ public:
 
 		uint32 LocalGlyphFlags = GlyphFlags;
 
+		switch(InFontInfo.Hinting)
+		{
+		case EFontHinting::Auto:		LocalGlyphFlags |= FT_LOAD_FORCE_AUTOHINT; break;
+		case EFontHinting::AutoLight:	LocalGlyphFlags |= FT_LOAD_TARGET_LIGHT; break;
+		case EFontHinting::Monochrome:	LocalGlyphFlags |= FT_LOAD_TARGET_MONO | FT_LOAD_FORCE_AUTOHINT; break;
+		case EFontHinting::None:		LocalGlyphFlags |= FT_LOAD_NO_AUTOHINT | FT_LOAD_NO_HINTING; break;
+		case EFontHinting::Default:
+		default:						LocalGlyphFlags |= FT_LOAD_TARGET_NORMAL; break;
+		}
+
 		// If the requested glyph doesn't exist, use the localization fallback font.
 		if ( FontFace == nullptr || (Char != 0 && GlyphIndex == 0) )
 		{
@@ -205,7 +215,7 @@ public:
 
 		// Get the slot for the glyph.  This contains measurement info
 		FT_GlyphSlot Slot = FontFace->glyph;
-
+		
 		FT_Render_Glyph( Slot, FT_RENDER_MODE_NORMAL );
 
 		// one byte per pixel 
diff --git a/Engine/Source/Runtime/SlateCore/Private/Fonts/SlateFontInfo.cpp b/Engine/Source/Runtime/SlateCore/Private/Fonts/SlateFontInfo.cpp
--- a/Engine/Source/Runtime/SlateCore/Private/Fonts/SlateFontInfo.cpp
+++ b/Engine/Source/Runtime/SlateCore/Private/Fonts/SlateFontInfo.cpp
@@ -9,39 +9,44 @@
 FSlateFontInfo::FSlateFontInfo( )
 	: FontName(NAME_None)
 	, Size(0)
+	, Hinting(EFontHinting::Default)
 { }
 
 
-FSlateFontInfo::FSlateFontInfo( const FString& InFontName, uint16 InSize )
+FSlateFontInfo::FSlateFontInfo( const FString& InFontName, uint16 InSize, EFontHinting InHinting )
 	: FontName(*InFontName)
 	, Size(InSize)
+	, Hinting(InHinting)
 {
 	//Useful for debugging style breakages
 	//check( FPaths::FileExists( FontName.ToString() ) );
 }
 
 
-FSlateFontInfo::FSlateFontInfo( const FName& InFontName, uint16 InSize )
+FSlateFontInfo::FSlateFontInfo( const FName& InFontName, uint16 InSize, EFontHinting InHinting )
 	: FontName(InFontName)
 	, Size(InSize)
+	, Hinting(InHinting)
 {
 	//Useful for debugging style breakages
 	//check( FPaths::FileExists( FontName.ToString() ) );
 }
 
 
-FSlateFontInfo::FSlateFontInfo( const ANSICHAR* InFontName, uint16 InSize )
+FSlateFontInfo::FSlateFontInfo( const ANSICHAR* InFontName, uint16 InSize, EFontHinting InHinting )
 	: FontName(InFontName)
 	, Size(InSize)
+	, Hinting(InHinting)
 {
 	//Useful for debugging style breakages
 	//check( FPaths::FileExists( FontName.ToString() ) );
 }
 
 
-FSlateFontInfo::FSlateFontInfo( const WIDECHAR* InFontName, uint16 InSize )
+FSlateFontInfo::FSlateFontInfo( const WIDECHAR* InFontName, uint16 InSize, EFontHinting InHinting )
 	: FontName(InFontName)
 	, Size(InSize)
+	, Hinting(InHinting)
 {
 	//Useful for debugging style breakages
 	//check( FPaths::FileExists( FontName.ToString() ) );
diff --git a/Engine/Source/Runtime/SlateCore/Public/Fonts/SlateFontInfo.h b/Engine/Source/Runtime/SlateCore/Public/Fonts/SlateFontInfo.h
--- a/Engine/Source/Runtime/SlateCore/Public/Fonts/SlateFontInfo.h
+++ b/Engine/Source/Runtime/SlateCore/Public/Fonts/SlateFontInfo.h
@@ -4,6 +4,20 @@
 
 #include "SlateFontInfo.generated.h"
 
+UENUM()
+enum class EFontHinting : uint8
+{
+	/** Use the default hinting specified in the font */
+	Default,
+	/** Force the use of an automatic hinting algorithm */
+	Auto,
+	/** Force the use of an automatic light hinting algorithm, optimized for non-monochrome displays */
+	AutoLight,
+	/** Force the use of an automatic hinting algorithm optimized for monochrome displays */
+	Monochrome,
+	/** Do not use hinting */
+	None,
+};
 
 /**
  * A representation of a font in Slate.
@@ -21,6 +35,10 @@ struct SLATECORE_API FSlateFontInfo
 	UPROPERTY(EditAnywhere, BlueprintReadWrite, Category=SlateStyleRules)
 	int32 Size;
 
+	/** The hinting algorithm to use with the font */
+	UPROPERTY(EditAnywhere, BlueprintReadWrite, Category=SlateStyleRules)
+	EFontHinting Hinting;
+
 public:
 
 	/** Default constructor. */
@@ -31,32 +49,36 @@ public:
 	 *
 	 * @param InFontName The name of the font.
 	 * @param InSize The size of the font.
+	 * @param InHinting The type of hinting to use for the font.
 	 */
-	FSlateFontInfo( const FString& InFontName, uint16 InSize );
+	FSlateFontInfo( const FString& InFontName, uint16 InSize, EFontHinting InHinting = EFontHinting::Default );
 
 	/**
 	 * Creates and initializes a new instance with the specified font name and size.
 	 *
 	 * @param InFontName The name of the font.
 	 * @param InSize The size of the font.
+	 * @param InHinting The type of hinting to use for the font.
 	 */
-	FSlateFontInfo( const FName& InFontName, uint16 InSize );
+	FSlateFontInfo( const FName& InFontName, uint16 InSize, EFontHinting InHinting = EFontHinting::Default );
 
 	/**
 	 * Creates and initializes a new instance with the specified font name and size.
 	 *
 	 * @param InFontName The name of the font.
 	 * @param InSize The size of the font.
+	 * @param InHinting The type of hinting to use for the font.
 	 */
-	FSlateFontInfo( const ANSICHAR* InFontName, uint16 InSize );
+	FSlateFontInfo( const ANSICHAR* InFontName, uint16 InSize, EFontHinting InHinting = EFontHinting::Default );
 
 	/**
 	 * Creates and initializes a new instance with the specified font name and size.
 	 *
 	 * @param InFontName The name of the font.
 	 * @param InSize The size of the font.
+	 * @param InHinting The type of hinting to use for the font.
 	 */
-	FSlateFontInfo( const WIDECHAR* InFontName, uint16 InSize );
+	FSlateFontInfo( const WIDECHAR* InFontName, uint16 InSize, EFontHinting InHinting = EFontHinting::Default );
 
 public:
 
@@ -68,7 +90,7 @@ public:
 	 */
 	bool operator==( const FSlateFontInfo& Other ) const 
 	{
-		return FontName == Other.FontName && Size == Other.Size;
+		return FontName == Other.FontName && Size == Other.Size && Hinting == Other.Hinting;
 	}
 
 	/**
@@ -93,6 +115,9 @@ public:
 	 */
 	friend inline uint32 GetTypeHash( const FSlateFontInfo& FontInfo )
 	{
-		return GetTypeHash(FontInfo.FontName) + (FontInfo.Size << 17);
+		uint32 Hash = 0;
+		Hash = HashCombine(Hash, GetTypeHash(FontInfo.FontName));
+		Hash = HashCombine(Hash, FontInfo.Size);
+		return HashCombine(Hash, uint32(FontInfo.Hinting));
 	}
 };
