Slate: TabManager now adds a tutorial tag if the content being spawned is not already tagged for tutorials.

// NOTE: relying on this is bad because it couples tutorials to the internals of tab management!
// We should find another way.

diff --git a/Engine/Source/Runtime/Slate/Private/Framework/Docking/TabManager.cpp b/Engine/Source/Runtime/Slate/Private/Framework/Docking/TabManager.cpp
--- a/Engine/Source/Runtime/Slate/Private/Framework/Docking/TabManager.cpp
+++ b/Engine/Source/Runtime/Slate/Private/Framework/Docking/TabManager.cpp
@@ -1167,6 +1167,10 @@ TSharedRef<SDockTab> FTabManager::SpawnTab( const FTabId& TabId, const TSharedPt
 		NewTabWidget->ProvideDefaultLabel( Spawner->GetDisplayName().IsEmpty() ? FText::FromName( Spawner->TabType ) : Spawner->GetDisplayName() );
 		NewTabWidget->ProvideDefaultIcon( Spawner->GetIcon().GetIcon() );
 
+		// If this tabs' content was not tagged for tutorials by its creator, we can provide an automated tag based  on the TabId.
+		// NOTE: relying on this is bad because it couples tutorials to the internals of tab management!
+		NewTabWidget->GetContent()->AddMetadataIfMissing<FTagMetaData>( MakeShareable(new FTagMetaData( *TabId.ToString() ) ) );
+
 		// The spawner tracks that last tab it spawned
 		Spawner->SpawnedTabPtr = NewTabWidget;
 	}
diff --git a/Engine/Source/Runtime/Slate/Private/Widgets/Docking/SDockTab.cpp b/Engine/Source/Runtime/Slate/Private/Widgets/Docking/SDockTab.cpp
--- a/Engine/Source/Runtime/Slate/Private/Widgets/Docking/SDockTab.cpp
+++ b/Engine/Source/Runtime/Slate/Private/Widgets/Docking/SDockTab.cpp
@@ -290,7 +290,7 @@ TSharedPtr<SWindow> SDockTab::GetParentWindow() const
 }
 
 SDockTab::SDockTab()
-	: Content(SNullWidget::NullWidget)
+	: Content(SNew(SSpacer))
 	, TabWellContentLeft(SNullWidget::NullWidget)
 	, TabWellContentRight(SNullWidget::NullWidget)
 	, LayoutIdentifier(NAME_None)
diff --git a/Engine/Source/Runtime/SlateCore/Public/Widgets/SWidget.h b/Engine/Source/Runtime/SlateCore/Public/Widgets/SWidget.h
--- a/Engine/Source/Runtime/SlateCore/Public/Widgets/SWidget.h
+++ b/Engine/Source/Runtime/SlateCore/Public/Widgets/SWidget.h
@@ -648,7 +648,33 @@ public:
 				FoundMetaData.Add(StaticCastSharedRef<MetaDataType>(MetaDataEntry));
 			}
 		}
-		return FoundMetaData;	
+		return FoundMetaData;
+	}
+
+private:
+	/**
+	 * Add metadata to the widget unless this type of metadata is already present.
+	 *
+	 * MIGHT BE REMOVED ! ! !
+	 *
+	 * Not sure we want to allow such mutation of metadata yet.
+	 * We might remove this method, so it is only exposed to the
+	 * one used case at the moment; thus the 'friend FTabManager'
+	 */
+	friend class FTabManager;
+
+	template<typename MetaDataType>
+	void AddMetadataIfMissing( const TSharedRef<MetaDataType>& AddMe )
+	{
+		const bool bMetadataAlreadyPresent = MetaData.ContainsByPredicate([]( const TSharedRef<ISlateMetaData>& Candidate )
+		{
+			return Candidate->IsOfType<MetaDataType>();
+		});
+
+		if (!bMetadataAlreadyPresent)
+		{
+			MetaData.Add( AddMe );
+		}		
 	}
 
 public:
