Support for a ScissorRectBox that ensures child widgets are clipped even under render transforms by using hardware scissor rect. #BUN
* This is a hack for the UMG designer to ensure any render transforms in designed widgets do not leak outside the designer. Implemented in a low-level hacky way so we don't have to change any public APIs because clipping control can and will change in better ways going forward.
* SlateViewer now uses it for all it's views as a test, and the UMG designer uses it.
* ElementBatcher now uses a new hash function that incorporates the ScissorRect, and slightly improves the efficiency, definitely improving the accuracy.

#TTP346004: UMG: Scaling a widget will result in it scaling outside of the graph

diff --git a/Engine/Source/Developer/StandaloneRenderer/Private/Windows/D3D/SlateD3DRenderingPolicy.cpp b/Engine/Source/Developer/StandaloneRenderer/Private/Windows/D3D/SlateD3DRenderingPolicy.cpp
--- a/Engine/Source/Developer/StandaloneRenderer/Private/Windows/D3D/SlateD3DRenderingPolicy.cpp
+++ b/Engine/Source/Developer/StandaloneRenderer/Private/Windows/D3D/SlateD3DRenderingPolicy.cpp
@@ -238,7 +238,21 @@ void FSlateD3D11RenderingPolicy::DrawElements( const FMatrix& ViewProjectionMatr
 		}
 		else
 		{
-			GD3DDeviceContext->RSSetState( NormalRasterState  );
+			if (RenderBatch.ScissorRect.IsSet())
+			{
+				D3D11_RECT R;
+				FShortRect ScissorRect = RenderBatch.ScissorRect.GetValue();
+				R.left = ScissorRect.Left;
+				R.top = ScissorRect.Top;
+				R.bottom = ScissorRect.Bottom;
+				R.right = ScissorRect.Right;
+				GD3DDeviceContext->RSSetScissorRects(1, &R);
+				GD3DDeviceContext->RSSetState(ScissorRasterState);
+			}
+			else
+			{
+				GD3DDeviceContext->RSSetState(NormalRasterState);
+			}
 		}
 	
 
diff --git a/Engine/Source/Editor/UMGEditor/Private/TabFactory/DesignerTabSummoner.cpp b/Engine/Source/Editor/UMGEditor/Private/TabFactory/DesignerTabSummoner.cpp
--- a/Engine/Source/Editor/UMGEditor/Private/TabFactory/DesignerTabSummoner.cpp
+++ b/Engine/Source/Editor/UMGEditor/Private/TabFactory/DesignerTabSummoner.cpp
@@ -4,6 +4,7 @@
 
 #include "DesignerTabSummoner.h"
 #include "SDesignerView.h"
+#include "SScissorRectBox.h"
 
 #define LOCTEXT_NAMESPACE "UMG"
 
@@ -24,12 +25,16 @@ FDesignerTabSummoner::FDesignerTabSummoner(TSharedPtr<class FWidgetBlueprintEdit
 
 TSharedRef<SWidget> FDesignerTabSummoner::CreateTabBody(const FWorkflowTabSpawnInfo& Info) const
 {
-	return SNew(SHorizontalBox)
-		+ SHorizontalBox::Slot()
-		.FillWidth(1)
+	return 
+		SNew(SScissorRectBox)
 		[
-			SNew(SDesignerView, BlueprintEditor.Pin())
-			.AddMetaData<FTagMetaData>(FTagMetaData(TEXT("Designer")))
+			SNew(SHorizontalBox)
+			+ SHorizontalBox::Slot()
+			.FillWidth(1)
+			[
+				SNew(SDesignerView, BlueprintEditor.Pin())
+				.AddMetaData<FTagMetaData>(FTagMetaData(TEXT("Designer")))
+			]
 		];
 }
 
diff --git a/Engine/Source/Runtime/AppFramework/Private/Widgets/Testing/STestSuite.cpp b/Engine/Source/Runtime/AppFramework/Private/Widgets/Testing/STestSuite.cpp
--- a/Engine/Source/Runtime/AppFramework/Private/Widgets/Testing/STestSuite.cpp
+++ b/Engine/Source/Runtime/AppFramework/Private/Widgets/Testing/STestSuite.cpp
@@ -8,6 +8,7 @@
 #include "RichTextLayoutMarshaller.h"
 #include "SyntaxHighlighterTextLayoutMarshaller.h"
 #include "STestSuite.h"
+#include "SScissorRectBox.h"
 #include "TransformCalculus3D.h"
 #include "SlateRenderTransform.h"
 #include "SlateLayoutTransform.h"
@@ -4509,16 +4510,19 @@ TSharedRef<SDockTab> SpawnTab(const FSpawnTabArgs& Args, FName TabIdentifier)
 				]
 			)
 		[
-			SNew(SVerticalBox)
-			.RenderTransform_Static(&::GetTestRenderTransform)
-			.RenderTransformPivot_Static(&::GetTestRenderTransformPivot)
-			+SVerticalBox::Slot()
-			[
-				SNew(SAnimTest)
-			]
-			+SVerticalBox::Slot()
+			SNew(SScissorRectBox)
 			[
-				SNew(SFxTest)
+				SNew(SVerticalBox)
+				.RenderTransform_Static(&::GetTestRenderTransform)
+				.RenderTransformPivot_Static(&::GetTestRenderTransformPivot)
+				+SVerticalBox::Slot()
+				[
+					SNew(SAnimTest)
+				]
+				+SVerticalBox::Slot()
+				[
+					SNew(SFxTest)
+				]
 			]
 		];
 	}
@@ -4530,10 +4534,13 @@ TSharedRef<SDockTab> SpawnTab(const FSpawnTabArgs& Args, FName TabIdentifier)
 		SNew(SDockTab)
 		.Label( NSLOCTEXT("TestSuite1", "DocumentsTab", "Documents") )
 		[
-			SNew( SDocumentsTest, TabManagerRef )
-			.RenderTransform_Static(&::GetTestRenderTransform)
-			.RenderTransformPivot_Static(&::GetTestRenderTransformPivot)
-			.AddMetaData<FTagMetaData>(FTagMetaData("DocumentSpawner"))
+			SNew(SScissorRectBox)
+			[
+				SNew( SDocumentsTest, TabManagerRef )
+				.RenderTransform_Static(&::GetTestRenderTransform)
+				.RenderTransformPivot_Static(&::GetTestRenderTransformPivot)
+				.AddMetaData<FTagMetaData>(FTagMetaData("DocumentSpawner"))
+			]
 		];
 	}
 	else if (TabIdentifier == FName(TEXT("TableViewTestTab")))
@@ -4542,7 +4549,10 @@ TSharedRef<SDockTab> SpawnTab(const FSpawnTabArgs& Args, FName TabIdentifier)
 			. Label( LOCTEXT("ListViewTestTab", "ListView Test") )
 			. ToolTipText( LOCTEXT( "ListViewTestToolTip", "Switches to the List View test tab, which allows you to test list widgets in Slate." ) )
 		[
-			MakeTableViewTesting()
+			SNew(SScissorRectBox)
+			[
+				MakeTableViewTesting()
+			]
 		];
 	}
 	else if (TabIdentifier == FName(TEXT("LayoutExampleTab")))
@@ -4551,7 +4561,10 @@ TSharedRef<SDockTab> SpawnTab(const FSpawnTabArgs& Args, FName TabIdentifier)
 			. Label( LOCTEXT("LayoutExampleTab", "Layout Example") )
 			. ToolTipText( LOCTEXT( "LayoutExampleTabToolTip", "Switches to the Layout Example tab, which shows off examples of various Slate layout primitives." ) )
 		[
-			MakeLayoutExample()
+			SNew(SScissorRectBox)
+			[
+				MakeLayoutExample()
+			]
 		];
 	}
 #if WITH_FANCY_TEXT
@@ -4560,9 +4573,12 @@ TSharedRef<SDockTab> SpawnTab(const FSpawnTabArgs& Args, FName TabIdentifier)
 		return SNew(SDockTab)
 			. Label( LOCTEXT("RichTextTestTab", "Rich Text") )
 			[
-				SNew( SRichTextTest )
-				.RenderTransform_Static(&::GetTestRenderTransform)
-				.RenderTransformPivot_Static(&::GetTestRenderTransformPivot)
+				SNew(SScissorRectBox)
+				[
+					SNew( SRichTextTest )
+					.RenderTransform_Static(&::GetTestRenderTransform)
+					.RenderTransformPivot_Static(&::GetTestRenderTransformPivot)
+				]
 			];
 	}
 	else if ( TabIdentifier == FName( TEXT( "MultiLineEditTab" ) ) )
@@ -4570,13 +4586,16 @@ TSharedRef<SDockTab> SpawnTab(const FSpawnTabArgs& Args, FName TabIdentifier)
 		return SNew(SDockTab)
 			. Label( LOCTEXT("MultiLineEditTab", "MultiLine Edit") )
 			[
-				#if WITH_FANCY_TEXT
-				SNew( SMultiLineEditingTest )
-				#else
-				SNew( SSpacer )
-				#endif //WITH_FANCY_TEXT
-				.RenderTransform_Static(&::GetTestRenderTransform)
-				.RenderTransformPivot_Static(&::GetTestRenderTransformPivot)
+				SNew(SScissorRectBox)
+				[
+					#if WITH_FANCY_TEXT
+					SNew( SMultiLineEditingTest )
+					#else
+					SNew( SSpacer )
+					#endif //WITH_FANCY_TEXT
+					.RenderTransform_Static(&::GetTestRenderTransform)
+					.RenderTransformPivot_Static(&::GetTestRenderTransformPivot)
+				]
 			];
 	}
 #endif //WITH_FANCY_TEXT
@@ -4586,9 +4605,12 @@ TSharedRef<SDockTab> SpawnTab(const FSpawnTabArgs& Args, FName TabIdentifier)
 			. Label( LOCTEXT("EditableTextTestTab", "Editable Text") )
 			. ToolTipText( LOCTEXT( "EditableTextTestTabToolTip", "Switches to the Editable Text tab, where you can test the various inline text editing controls." ) )
 		[
-			SNew( STextEditTest )
-			.RenderTransform_Static(&::GetTestRenderTransform)
-			.RenderTransformPivot_Static(&::GetTestRenderTransformPivot)
+			SNew(SScissorRectBox)
+			[
+				SNew( STextEditTest )
+				.RenderTransform_Static(&::GetTestRenderTransform)
+				.RenderTransformPivot_Static(&::GetTestRenderTransformPivot)
+			]
 		];
 	}
 #if WITH_FANCY_TEXT
@@ -4598,9 +4620,12 @@ TSharedRef<SDockTab> SpawnTab(const FSpawnTabArgs& Args, FName TabIdentifier)
 			. Label( LOCTEXT("RichEditableTextTestTab", "Rich Editable Text") )
 			. ToolTipText( LOCTEXT( "RichEditableTextTestTabToolTip", "Switches to the Rich Editable Text tab, where you can test the various rich editable text features." ) )
 			[
-				SNew( SRichTextEditTest )
-				.RenderTransform_Static(&::GetTestRenderTransform)
-				.RenderTransformPivot_Static(&::GetTestRenderTransformPivot)
+				SNew(SScissorRectBox)
+				[
+					SNew( SRichTextEditTest )
+					.RenderTransform_Static(&::GetTestRenderTransform)
+					.RenderTransformPivot_Static(&::GetTestRenderTransformPivot)
+				]
 			];
 	}
 #endif //WITH_FANCY_TEXT
@@ -4609,9 +4634,12 @@ TSharedRef<SDockTab> SpawnTab(const FSpawnTabArgs& Args, FName TabIdentifier)
 		return SNew(SDockTab)
 			. Label( LOCTEXT("LayoutRoundingTab", "Layout Rounding") )
 		[
-			SNew( SLayoutRoundingTest )
-			.RenderTransform_Static(&::GetTestRenderTransform)
-			.RenderTransformPivot_Static(&::GetTestRenderTransformPivot)
+			SNew(SScissorRectBox)
+			[
+				SNew( SLayoutRoundingTest )
+				.RenderTransform_Static(&::GetTestRenderTransform)
+				.RenderTransformPivot_Static(&::GetTestRenderTransformPivot)
+			]
 		];
 	}
 	else if (TabIdentifier == FName(TEXT("ElementTestsTab")))
@@ -4620,9 +4648,12 @@ TSharedRef<SDockTab> SpawnTab(const FSpawnTabArgs& Args, FName TabIdentifier)
 			. Label( LOCTEXT("ElementTestsTab", "Element Tests") )
 			. ToolTipText( LOCTEXT( "ElementTestsTabToolTip", "Switches to the Element Tests tab, which allows you to view various rendering-related features of Slate." ) )
 		[
-			SNew( SElementTesting )
-			.RenderTransform_Static(&::GetTestRenderTransform)
-			.RenderTransformPivot_Static(&::GetTestRenderTransformPivot)
+			SNew(SScissorRectBox)
+			[
+				SNew( SElementTesting )
+				.RenderTransform_Static(&::GetTestRenderTransform)
+				.RenderTransformPivot_Static(&::GetTestRenderTransformPivot)
+			]
 		];
 	}
 	else if (TabIdentifier == FName(TEXT("SplitterTestTab")))
@@ -4631,9 +4662,12 @@ TSharedRef<SDockTab> SpawnTab(const FSpawnTabArgs& Args, FName TabIdentifier)
 			. Label( LOCTEXT("SplitterTestTab", "Splitter Test") )
 			. ToolTipText( LOCTEXT( "SplitterTestTabToolTip", "Switches to the Splitter Test tab, which you can use to test splitters." ) )
 		[
-			SNew( SSplitterTest )
-			.RenderTransform_Static(&::GetTestRenderTransform)
-			.RenderTransformPivot_Static(&::GetTestRenderTransformPivot)
+			SNew(SScissorRectBox)
+			[
+				SNew( SSplitterTest )
+				.RenderTransform_Static(&::GetTestRenderTransform)
+				.RenderTransformPivot_Static(&::GetTestRenderTransformPivot)
+			]
 		];
 	}
 	else if (TabIdentifier == FName(TEXT("MultiBoxTestTab")))
@@ -4642,9 +4676,12 @@ TSharedRef<SDockTab> SpawnTab(const FSpawnTabArgs& Args, FName TabIdentifier)
 			. Label( LOCTEXT("MultiBoxTextTab", "MultiBox Text") )
 			. ToolTipText( LOCTEXT( "MultiBoxTextTabToolTip", "Switches to the MultiBox tab, where you can test out MultiBoxes and MultiBlocks." ) )
 		[
-			SNew( SMultiBoxTest )
-			.RenderTransform_Static(&::GetTestRenderTransform)
-			.RenderTransformPivot_Static(&::GetTestRenderTransformPivot)
+			SNew(SScissorRectBox)
+			[
+				SNew( SMultiBoxTest )
+				.RenderTransform_Static(&::GetTestRenderTransform)
+				.RenderTransformPivot_Static(&::GetTestRenderTransformPivot)
+			]
 		];
 	}
 	else if (TabIdentifier == FName(TEXT("WidgetGalleryTab")))
@@ -4653,7 +4690,10 @@ TSharedRef<SDockTab> SpawnTab(const FSpawnTabArgs& Args, FName TabIdentifier)
 			. Label( LOCTEXT("WidgetGalleryTab", "Widget Gallery") )
 			. ToolTipText( LOCTEXT( "WidgetGalleryTabTextToolTip", "Switch to the widget gallery." ) )
 		[
-			MakeWidgetGallery()
+			SNew(SScissorRectBox)
+			[
+				MakeWidgetGallery()
+			]
 		];
 	}
 	else if (TabIdentifier == FName(TEXT("ColorPickerTestTab")))
@@ -4662,18 +4702,24 @@ TSharedRef<SDockTab> SpawnTab(const FSpawnTabArgs& Args, FName TabIdentifier)
 			. Label( LOCTEXT("ColorPickerTestTab", "Color Picker Test") )
 			. ToolTipText( LOCTEXT( "ColorPickerTestTabToolTip", "Switches to the Color Picker tab, where you can test out the color picker." ) )
 		[
-			SNew(SColorPickerTest)
-			.RenderTransform_Static(&::GetTestRenderTransform)
-			.RenderTransformPivot_Static(&::GetTestRenderTransformPivot)
+			SNew(SScissorRectBox)
+			[
+				SNew(SColorPickerTest)
+				.RenderTransform_Static(&::GetTestRenderTransform)
+				.RenderTransformPivot_Static(&::GetTestRenderTransformPivot)
+			]
 		];
 	}
 	else if (TabIdentifier == "DPIScalingTest")
 	{
 		return SNew(SDockTab)
 		[
-			SNew(SDPIScalingTest)
-			.RenderTransform_Static(&::GetTestRenderTransform)
-			.RenderTransformPivot_Static(&::GetTestRenderTransformPivot)
+			SNew(SScissorRectBox)
+			[
+				SNew(SDPIScalingTest)
+				.RenderTransform_Static(&::GetTestRenderTransform)
+				.RenderTransformPivot_Static(&::GetTestRenderTransformPivot)
+			]
 		];
 	}
 	else if (TabIdentifier == FName(TEXT("NotificationListTestTab")))
@@ -4682,18 +4728,24 @@ TSharedRef<SDockTab> SpawnTab(const FSpawnTabArgs& Args, FName TabIdentifier)
 			. Label( LOCTEXT("NotificationListTestTab", "Notification List Test") )
 			. ToolTipText( LOCTEXT( "NotificationListTestTabToolTip", "Switches to the Notification List tab, where you can test out the notification list." ) )
 			[
-				SNew(SNotificationListTest)
-				.RenderTransform_Static(&::GetTestRenderTransform)
-				.RenderTransformPivot_Static(&::GetTestRenderTransformPivot)
+				SNew(SScissorRectBox)
+				[
+					SNew(SNotificationListTest)
+					.RenderTransform_Static(&::GetTestRenderTransform)
+					.RenderTransformPivot_Static(&::GetTestRenderTransformPivot)
+				]
 			];
 	}
 	else if (TabIdentifier == FName("GridPanelTest"))
 	{
 		return SNew(SDockTab)
 			[
-				SNew(SGridPanelTest)
-				.RenderTransform_Static(&::GetTestRenderTransform)
-				.RenderTransformPivot_Static(&::GetTestRenderTransformPivot)
+				SNew(SScissorRectBox)
+				[
+					SNew(SGridPanelTest)
+					.RenderTransform_Static(&::GetTestRenderTransform)
+					.RenderTransformPivot_Static(&::GetTestRenderTransformPivot)
+				]
 			];
 	}
 	else
diff --git a/Engine/Source/Runtime/Slate/Private/Widgets/Layout/SScissorRectBox.cpp b/Engine/Source/Runtime/Slate/Private/Widgets/Layout/SScissorRectBox.cpp
new file mode 100644
--- /dev/null
+++ b/Engine/Source/Runtime/Slate/Private/Widgets/Layout/SScissorRectBox.cpp
@@ -0,0 +1,74 @@
+// Copyright 1998-2014 Epic Games, Inc. All Rights Reserved.
+
+#include "SlatePrivatePCH.h"
+#include "SScissorRectBox.h"
+
+SScissorRectBox::SScissorRectBox()
+{
+}
+
+void SScissorRectBox::Construct(const FArguments& InArgs)
+{
+	ChildSlot
+	[
+		InArgs._Content.Widget
+	];
+}
+
+
+FVector2D SScissorRectBox::ComputeDesiredSize() const
+{
+	EVisibility ChildVisibility = ChildSlot.GetWidget()->GetVisibility();
+
+	if (ChildVisibility != EVisibility::Collapsed)
+	{
+		return ChildSlot.GetWidget()->GetDesiredSize();
+	}
+	
+	return FVector2D::ZeroVector;
+}
+
+void SScissorRectBox::OnArrangeChildren(const FGeometry& AllottedGeometry, FArrangedChildren& ArrangedChildren) const
+{
+	// the child takes up the full space of this panel, always.
+	ArrangedChildren.AddWidget(AllottedGeometry.MakeChild(ChildSlot.GetWidget(), AllottedGeometry.GetLocalSize(), FSlateLayoutTransform()));
+}
+
+FChildren* SScissorRectBox::GetChildren()
+{
+	return &ChildSlot;
+}
+
+int32 SScissorRectBox::OnPaint(const FPaintArgs& Args, const FGeometry& AllottedGeometry, const FSlateRect& MyClippingRect, FSlateWindowElementList& OutDrawElements, int32 LayerId, const FWidgetStyle& InWidgetStyle, bool bParentEnabled) const
+{
+	// just draw the only child
+	FArrangedChildren ArrangedChildren(EVisibility::Visible);
+	{
+		this->ArrangeChildren(AllottedGeometry, ArrangedChildren);
+	}
+
+	// Maybe none of our children are visible
+	if( ArrangedChildren.Num() > 0 )
+	{
+		check( ArrangedChildren.Num() == 1 );
+		FArrangedWidget& TheChild = ArrangedChildren[0];
+
+		// Child clipping rect is made effectively infinite to ensure that any rotated children do not get prematurely cut off. We let the scissor rect do all the work.
+		const FSlateRect ChildClippingRect = AllottedGeometry.GetClippingRect().IntersectionWith(MyClippingRect);
+
+		// set the hacky, back door global scissor rect settings.
+		// This will get picked up by the element batcher and added to the draw element.
+		// We do it this way because this scissor rect is a real hack until better solution is made available.
+		// In the meantime, we don't want to modify the public interface of all element drawing to ensure all widgets pick up this scissor rect.
+		extern SLATECORE_API TOptional<FShortRect> GSlateScissorRect;
+		TOptional<FShortRect> OldRect = GSlateScissorRect;
+		GSlateScissorRect = FShortRect(ChildClippingRect);
+		int32 Result = TheChild.Widget->Paint( Args.WithNewParent(this), TheChild.Geometry, FSlateRect(FVector2D(-FLT_MAX, -FLT_MAX), FVector2D(FLT_MAX, FLT_MAX)), OutDrawElements, LayerId, InWidgetStyle, ShouldBeEnabled( bParentEnabled ) );
+		// We've added all the draw elements, so put the scissor rect back
+		GSlateScissorRect = OldRect;
+		return Result;
+
+	}
+	return LayerId;
+}
+
diff --git a/Engine/Source/Runtime/Slate/Public/Widgets/Layout/SScissorRectBox.h b/Engine/Source/Runtime/Slate/Public/Widgets/Layout/SScissorRectBox.h
new file mode 100644
--- /dev/null
+++ b/Engine/Source/Runtime/Slate/Public/Widgets/Layout/SScissorRectBox.h
@@ -0,0 +1,48 @@
+// Copyright 1998-2014 Epic Games, Inc. All Rights Reserved.
+
+#pragma once
+
+/**
+ * A ScirrorRectBox is a widget that clips all its children using a scissor rect that matches it clipping rect
+ * in LAYOUT SPACE (the way it is provided in OnPaint).
+ * This ensures that any child widgets with render transforms do not "bleed out" of this parent container.
+ * Note there is a significant performance impact to using one of these widgets because any widgets inside 
+ * this one end up in a separate set of batches using that scissor rect.
+ * 
+ * Keep in mind that this widget will use the LAYOUT clippping rect to clip children, since hardware clipping
+ * rects do not support non-aligned rectangles. If a render transform is applied to this widget itself, the results
+ * will not be as expected.
+ * 
+ * NOTE: This class is provided as a STOPGAP for design-time widgets (like the UMG designer itself)
+ *       until a more robust clipping solution can be put in place.
+ */
+class SLATE_API SScissorRectBox: public SPanel
+{
+
+public:
+	class FScissorRectSlot : public TSupportsOneChildMixin<FScissorRectSlot>
+	{
+	};
+
+	SLATE_BEGIN_ARGS(SScissorRectBox)
+	{}
+		SLATE_DEFAULT_SLOT(FArguments, Content)
+	SLATE_END_ARGS()
+
+	SScissorRectBox();
+
+	void Construct( const FArguments& InArgs );
+
+	/**
+	 * See the Content slot.
+     */
+	void SetContent(const TSharedRef< SWidget >& InContent);
+
+private:
+	virtual FVector2D ComputeDesiredSize() const override;
+	virtual void OnArrangeChildren( const FGeometry& AllottedGeometry, FArrangedChildren& ArrangedChildren ) const override;
+	virtual FChildren* GetChildren() override;
+	virtual int32 OnPaint( const FPaintArgs& Args, const FGeometry& AllottedGeometry, const FSlateRect& MyClippingRect, FSlateWindowElementList& OutDrawElements, int32 LayerId, const FWidgetStyle& InWidgetStyle, bool bParentEnabled ) const override;
+
+	FScissorRectSlot ChildSlot;
+};
diff --git a/Engine/Source/Runtime/SlateCore/Private/Rendering/DrawElements.cpp b/Engine/Source/Runtime/SlateCore/Private/Rendering/DrawElements.cpp
--- a/Engine/Source/Runtime/SlateCore/Private/Rendering/DrawElements.cpp
+++ b/Engine/Source/Runtime/SlateCore/Private/Rendering/DrawElements.cpp
@@ -2,19 +2,25 @@
 
 #include "SlateCorePrivatePCH.h"
 
+void FSlateDrawElement::Init(uint32 InLayer, const FPaintGeometry& PaintGeometry, const FSlateRect& InClippingRect, ESlateDrawEffect::Type InDrawEffects)
+{
+	RenderTransform = PaintGeometry.GetAccumulatedRenderTransform();
+	Position = PaintGeometry.DrawPosition;
+	LocalSize = PaintGeometry.GetLocalSize();
+	ClippingRect = InClippingRect;
+	Layer = InLayer;
+	Scale = PaintGeometry.DrawScale;
+	DrawEffects = InDrawEffects;
+	extern TOptional<FShortRect> GSlateScissorRect;
+	ScissorRect = GSlateScissorRect;
+}
 
 void FSlateDrawElement::MakeDebugQuad( FSlateWindowElementList& ElementList, uint32 InLayer, const FPaintGeometry& PaintGeometry, const FSlateRect& InClippingRect)
 {
 	PaintGeometry.CommitTransformsIfUsingLegacyConstructor();
 	FSlateDrawElement& DrawElt = ElementList.AddUninitialized();
+	DrawElt.Init(InLayer, PaintGeometry, InClippingRect, ESlateDrawEffect::None);
 	DrawElt.ElementType = ET_DebugQuad;
-	DrawElt.RenderTransform = PaintGeometry.GetAccumulatedRenderTransform();
-	DrawElt.Position = PaintGeometry.DrawPosition;
-	DrawElt.LocalSize = PaintGeometry.GetLocalSize();
-	DrawElt.ClippingRect = InClippingRect;
-	DrawElt.Layer = InLayer;
-	DrawElt.DrawEffects = ESlateDrawEffect::None;
-	DrawElt.Scale = PaintGeometry.DrawScale;
 }
 
 
@@ -29,15 +35,9 @@ void FSlateDrawElement::MakeBox(
 {
 	PaintGeometry.CommitTransformsIfUsingLegacyConstructor();
 	FSlateDrawElement& DrawElt = ElementList.AddUninitialized();
+	DrawElt.Init(InLayer, PaintGeometry, InClippingRect, InDrawEffects);
 	DrawElt.ElementType = (InBrush->DrawAs == ESlateBrushDrawType::Border) ? ET_Border : ET_Box;
-	DrawElt.RenderTransform = PaintGeometry.GetAccumulatedRenderTransform();
-	DrawElt.Position = PaintGeometry.DrawPosition;
-	DrawElt.LocalSize = PaintGeometry.GetLocalSize();
-	DrawElt.ClippingRect = InClippingRect;
 	DrawElt.DataPayload.SetBoxPayloadProperties( InBrush, InTint );
-	DrawElt.Layer = InLayer;
-	DrawElt.DrawEffects = InDrawEffects;
-	DrawElt.Scale = PaintGeometry.DrawScale;
 }
 
 void FSlateDrawElement::MakeRotatedBox( 
@@ -54,14 +54,8 @@ void FSlateDrawElement::MakeRotatedBox(
 {
 	PaintGeometry.CommitTransformsIfUsingLegacyConstructor();
 	FSlateDrawElement& DrawElt = ElementList.AddUninitialized();
+	DrawElt.Init(InLayer, PaintGeometry, InClippingRect, InDrawEffects);
 	DrawElt.ElementType = (InBrush->DrawAs == ESlateBrushDrawType::Border) ? ET_Border : ET_Box;
-	DrawElt.RenderTransform = PaintGeometry.GetAccumulatedRenderTransform();
-	DrawElt.Position = PaintGeometry.DrawPosition;
-	DrawElt.LocalSize = PaintGeometry.GetLocalSize();
-	DrawElt.ClippingRect = InClippingRect;
-	DrawElt.Layer = InLayer;
-	DrawElt.DrawEffects = InDrawEffects;
-	DrawElt.Scale = PaintGeometry.DrawScale;
 
 	FVector2D RotationPoint = GetRotationPoint( PaintGeometry, InRotationPoint, RotationSpace );
 	DrawElt.DataPayload.SetRotatedBoxPayloadProperties( InBrush, Angle, RotationPoint, InTint );
@@ -72,15 +66,9 @@ void FSlateDrawElement::MakeText( FSlateWindowElementList& ElementList, uint32 I
 {
 	PaintGeometry.CommitTransformsIfUsingLegacyConstructor();
 	FSlateDrawElement& DrawElt = ElementList.AddUninitialized();
+	DrawElt.Init(InLayer, PaintGeometry, InClippingRect, InDrawEffects);
 	DrawElt.ElementType = ET_Text;
-	DrawElt.RenderTransform = PaintGeometry.GetAccumulatedRenderTransform();
-	DrawElt.Position = PaintGeometry.DrawPosition;
-	DrawElt.LocalSize = PaintGeometry.GetLocalSize();
-	DrawElt.ClippingRect = InClippingRect;
 	DrawElt.DataPayload.SetTextPayloadProperties( FString( EndIndex - StartIndex, *InText + StartIndex ), InFontInfo, InTint );
-	DrawElt.Layer = InLayer;
-	DrawElt.DrawEffects = InDrawEffects;
-	DrawElt.Scale = PaintGeometry.DrawScale;
 }
 
 
@@ -88,15 +76,9 @@ void FSlateDrawElement::MakeText( FSlateWindowElementList& ElementList, uint32 I
 {
 	PaintGeometry.CommitTransformsIfUsingLegacyConstructor();
 	FSlateDrawElement& DrawElt = ElementList.AddUninitialized();
+	DrawElt.Init(InLayer, PaintGeometry, InClippingRect, InDrawEffects);
 	DrawElt.ElementType = ET_Text;
-	DrawElt.RenderTransform = PaintGeometry.GetAccumulatedRenderTransform();
-	DrawElt.Position = PaintGeometry.DrawPosition;
-	DrawElt.LocalSize = PaintGeometry.GetLocalSize();
-	DrawElt.ClippingRect = InClippingRect;
 	DrawElt.DataPayload.SetTextPayloadProperties( InText, InFontInfo, InTint );
-	DrawElt.Layer = InLayer;
-	DrawElt.DrawEffects = InDrawEffects;
-	DrawElt.Scale = PaintGeometry.DrawScale;
 }
 
 
@@ -104,30 +86,18 @@ void FSlateDrawElement::MakeText( FSlateWindowElementList& ElementList, uint32 I
 {
 	PaintGeometry.CommitTransformsIfUsingLegacyConstructor();
 	FSlateDrawElement& DrawElt = ElementList.AddUninitialized();
+	DrawElt.Init(InLayer, PaintGeometry, InClippingRect, InDrawEffects);
 	DrawElt.ElementType = ET_Text;
-	DrawElt.RenderTransform = PaintGeometry.GetAccumulatedRenderTransform();
-	DrawElt.Position = PaintGeometry.DrawPosition;
-	DrawElt.LocalSize = PaintGeometry.GetLocalSize();
-	DrawElt.ClippingRect = InClippingRect;
 	DrawElt.DataPayload.SetTextPayloadProperties( InText.ToString(), InFontInfo, InTint );
-	DrawElt.Layer = InLayer;
-	DrawElt.DrawEffects = InDrawEffects;
-	DrawElt.Scale = PaintGeometry.DrawScale;
 }
 
 void FSlateDrawElement::MakeGradient( FSlateWindowElementList& ElementList, uint32 InLayer, const FPaintGeometry& PaintGeometry, TArray<FSlateGradientStop> InGradientStops, EOrientation InGradientType, const FSlateRect& InClippingRect, ESlateDrawEffect::Type InDrawEffects, bool bGammaCorrect )
 {
 	PaintGeometry.CommitTransformsIfUsingLegacyConstructor();
 	FSlateDrawElement& DrawElt = ElementList.AddUninitialized();
+	DrawElt.Init(InLayer, PaintGeometry, InClippingRect, InDrawEffects);
 	DrawElt.ElementType = ET_Gradient;
-	DrawElt.RenderTransform = PaintGeometry.GetAccumulatedRenderTransform();
-	DrawElt.Position = PaintGeometry.DrawPosition;
-	DrawElt.LocalSize = PaintGeometry.GetLocalSize();
-	DrawElt.ClippingRect = InClippingRect;
 	DrawElt.DataPayload.SetGradientPayloadProperties( InGradientStops, InGradientType, bGammaCorrect );
-	DrawElt.Layer = InLayer;
-	DrawElt.DrawEffects = InDrawEffects;
-	DrawElt.Scale = PaintGeometry.DrawScale;
 }
 
 
@@ -135,15 +105,9 @@ void FSlateDrawElement::MakeSpline( FSlateWindowElementList& ElementList, uint32
 {
 	PaintGeometry.CommitTransformsIfUsingLegacyConstructor();
 	FSlateDrawElement& DrawElt = ElementList.AddUninitialized();
+	DrawElt.Init(InLayer, PaintGeometry, InClippingRect, InDrawEffects);
 	DrawElt.ElementType = ET_Spline;
-	DrawElt.RenderTransform = PaintGeometry.GetAccumulatedRenderTransform();
-	DrawElt.Position = PaintGeometry.DrawPosition;
-	DrawElt.LocalSize = PaintGeometry.GetLocalSize();
-	DrawElt.ClippingRect = InClippingRect;
 	DrawElt.DataPayload.SetSplinePayloadProperties( InStart, InStartDir, InEnd, InEndDir, InThickness, InTint );
-	DrawElt.Layer = InLayer;
-	DrawElt.DrawEffects = InDrawEffects;
-	DrawElt.Scale = PaintGeometry.DrawScale;
 }
 
 
@@ -157,15 +121,9 @@ void FSlateDrawElement::MakeLines( FSlateWindowElementList& ElementList, uint32
 {
 	PaintGeometry.CommitTransformsIfUsingLegacyConstructor();
 	FSlateDrawElement& DrawElt = ElementList.AddUninitialized();
+	DrawElt.Init(InLayer, PaintGeometry, InClippingRect, InDrawEffects);
 	DrawElt.ElementType = ET_Line;
-	DrawElt.RenderTransform = PaintGeometry.GetAccumulatedRenderTransform();
-	DrawElt.Position = PaintGeometry.DrawPosition;
-	DrawElt.LocalSize = PaintGeometry.GetLocalSize();
-	DrawElt.ClippingRect = InClippingRect;
 	DrawElt.DataPayload.SetLinesPayloadProperties( Points, InTint, bAntialias, ESlateLineJoinType::Sharp );
-	DrawElt.Layer = InLayer;
-	DrawElt.DrawEffects = InDrawEffects;
-	DrawElt.Scale = PaintGeometry.DrawScale;
 }
 
 
@@ -173,26 +131,19 @@ void FSlateDrawElement::MakeViewport( FSlateWindowElementList& ElementList, uint
 {
 	PaintGeometry.CommitTransformsIfUsingLegacyConstructor();
 	FSlateDrawElement& DrawElt = ElementList.AddUninitialized();
+	DrawElt.Init(InLayer, PaintGeometry, InClippingRect, InDrawEffects);
 	DrawElt.ElementType = ET_Viewport;
-	DrawElt.RenderTransform = PaintGeometry.GetAccumulatedRenderTransform();
-	DrawElt.Position = PaintGeometry.DrawPosition;
-	DrawElt.LocalSize = PaintGeometry.GetLocalSize();
-	DrawElt.ClippingRect = InClippingRect;
 	DrawElt.DataPayload.SetViewportPayloadProperties( Viewport, InTint, bGammaCorrect, bAllowBlending );
-	DrawElt.Layer = InLayer;
-	DrawElt.DrawEffects = InDrawEffects;
-	DrawElt.Scale = PaintGeometry.DrawScale;
 }
 
 
 void FSlateDrawElement::MakeCustom( FSlateWindowElementList& ElementList, uint32 InLayer, TSharedPtr<ICustomSlateElement, ESPMode::ThreadSafe> CustomDrawer )
 {
 	FSlateDrawElement& DrawElt = ElementList.AddUninitialized();
+	DrawElt.Init(InLayer, FPaintGeometry(), FSlateRect(1,1,1,1), ESlateDrawEffect::None);
 	DrawElt.RenderTransform = FSlateRenderTransform();
 	DrawElt.ElementType = ET_Custom;
-	DrawElt.Layer = InLayer;
 	DrawElt.DataPayload.SetCustomDrawerPayloadProperties( CustomDrawer );
-	DrawElt.ClippingRect = FSlateRect(1,1,1,1);
 }
 
 FVector2D FSlateDrawElement::GetRotationPoint(const FPaintGeometry& PaintGeometry, const TOptional<FVector2D>& UserRotationPoint, ERotationSpace RotationSpace)
diff --git a/Engine/Source/Runtime/SlateCore/Private/Rendering/ElementBatcher.cpp b/Engine/Source/Runtime/SlateCore/Private/Rendering/ElementBatcher.cpp
--- a/Engine/Source/Runtime/SlateCore/Private/Rendering/ElementBatcher.cpp
+++ b/Engine/Source/Runtime/SlateCore/Private/Rendering/ElementBatcher.cpp
@@ -4,6 +4,9 @@
 #include "ElementBatcher.h"
 #include "SlateRenderTransform.h"
 
+// Super-hacky way of storing the scissor rect so we don't have to change all the FSlateDrawElement APIs for this hacky support.
+SLATECORE_API TOptional<FShortRect> GSlateScissorRect;
+
 FVector2D RoundToInt(const FVector2D& Vec)
 {
 	return FVector2D(FMath::RoundToInt(Vec.X), FMath::RoundToInt(Vec.Y));
@@ -66,42 +69,49 @@ void FSlateElementBatcher::AddElements( const TArray<FSlateDrawElement>& DrawEle
 	
 		// A zero or negatively sized clipping rect means the geometry will not be displayed
 		const bool bIsFullyClipped = (!InClippingRect.IsValid() || (InClippingRect.Top == 0.0f && InClippingRect.Left == 0.0f && InClippingRect.Bottom == 0.0f && InClippingRect.Right == 0.0f));
-
+		
 		if ( !bIsFullyClipped )
 		{
-			// Determine what type of element to add
-			switch( DrawElement.GetElementType() )
+			// !!! WRH 2014/09/17 - Profile and see how much this costs.
+			// do this check in here do we can short circuit above more quickly, but still name this variable so its meaning is clear.
+			const bool bIsScissored = DrawElement.GetScissorRect().IsSet() && !DrawElement.GetScissorRect().GetValue().DoesIntersect(FShortRect(InClippingRect));
+			// scissor rects are sort of a low level hack, so no one konws to clip against them. Instead we do it here to make sure the element is not actually rendered.
+			if (!bIsScissored)
 			{
-			case FSlateDrawElement::ET_Box:
-				AddBoxElement( DrawElement );
-				break;
-			case FSlateDrawElement::ET_DebugQuad:
-				AddQuadElement( DrawElement );
-				break;
-			case FSlateDrawElement::ET_Text:
-				AddTextElement( DrawElement );
-				break;
-			case FSlateDrawElement::ET_Spline:
-				AddSplineElement( DrawElement );
-				break;
-			case FSlateDrawElement::ET_Line:
-				AddLineElement( DrawElement );
-				break;
-			case FSlateDrawElement::ET_Gradient:
-				AddGradientElement( DrawElement );
-				break;
-			case FSlateDrawElement::ET_Viewport:
-				AddViewportElement( DrawElement );
-				break;
-			case FSlateDrawElement::ET_Border:
-				AddBorderElement( DrawElement );
-				break;
-			case FSlateDrawElement::ET_Custom:
-				AddCustomElement( DrawElement );
-				break;
-			default:
-				checkf(0, TEXT("Invalid element type"));
-				break;
+				// Determine what type of element to add
+				switch( DrawElement.GetElementType() )
+				{
+				case FSlateDrawElement::ET_Box:
+					AddBoxElement( DrawElement );
+					break;
+				case FSlateDrawElement::ET_DebugQuad:
+					AddQuadElement( DrawElement );
+					break;
+				case FSlateDrawElement::ET_Text:
+					AddTextElement( DrawElement );
+					break;
+				case FSlateDrawElement::ET_Spline:
+					AddSplineElement( DrawElement );
+					break;
+				case FSlateDrawElement::ET_Line:
+					AddLineElement( DrawElement );
+					break;
+				case FSlateDrawElement::ET_Gradient:
+					AddGradientElement( DrawElement );
+					break;
+				case FSlateDrawElement::ET_Viewport:
+					AddViewportElement( DrawElement );
+					break;
+				case FSlateDrawElement::ET_Border:
+					AddBorderElement( DrawElement );
+					break;
+				case FSlateDrawElement::ET_Custom:
+					AddCustomElement( DrawElement );
+					break;
+				default:
+					checkf(0, TEXT("Invalid element type"));
+					break;
+				}
 			}
 		}
 	}
@@ -123,7 +133,7 @@ void FSlateElementBatcher::AddQuadElement( const FSlateDrawElement& DrawElement,
 	FSlateLayoutTransform InverseLayoutTransform(Inverse(FSlateLayoutTransform(DrawElement.GetScale(), DrawElement.GetPosition())));
 	FSlateRotatedClipRectType RenderClipRect = ToSnappedRotatedRect(InClippingRect, InverseLayoutTransform, RenderTransform);
 
-	FSlateElementBatch& ElementBatch = FindBatchForElement( Layer, FShaderParams(), nullptr, ESlateDrawPrimitive::TriangleList, ESlateShader::Default, ESlateDrawEffect::None, ESlateBatchDrawFlag::Wireframe|ESlateBatchDrawFlag::NoBlending );
+	FSlateElementBatch& ElementBatch = FindBatchForElement( Layer, FShaderParams(), nullptr, ESlateDrawPrimitive::TriangleList, ESlateShader::Default, ESlateDrawEffect::None, ESlateBatchDrawFlag::Wireframe|ESlateBatchDrawFlag::NoBlending, DrawElement.GetScissorRect() );
 	TArray<FSlateVertex>& BatchVertices = BatchVertexArrays[ElementBatch.VertexArrayIndex];
 	TArray<SlateIndex>& BatchIndices = BatchIndexArrays[ElementBatch.IndexArrayIndex];
 
@@ -244,7 +254,7 @@ void FSlateElementBatcher::AddBoxElement( const FSlateDrawElement& DrawElement )
 		// Pass the tiling information as a flag so we can pick the correct texture addressing mode
 		ESlateBatchDrawFlag::Type DrawFlags = ( ( bTileHorizontal ? ESlateBatchDrawFlag::TileU : 0 ) | ( bTileVertical ? ESlateBatchDrawFlag::TileV : 0 ) );
 
-		FSlateElementBatch& ElementBatch = FindBatchForElement( Layer, FShaderParams(), Resource, ESlateDrawPrimitive::TriangleList, ESlateShader::Default, InDrawEffects, DrawFlags );
+		FSlateElementBatch& ElementBatch = FindBatchForElement( Layer, FShaderParams(), Resource, ESlateDrawPrimitive::TriangleList, ESlateShader::Default, InDrawEffects, DrawFlags, DrawElement.GetScissorRect() );
 		TArray<FSlateVertex>& BatchVertices = BatchVertexArrays[ElementBatch.VertexArrayIndex];
 		TArray<SlateIndex>& BatchIndices = BatchIndexArrays[ElementBatch.IndexArrayIndex];
 
@@ -552,7 +562,7 @@ void FSlateElementBatcher::AddTextElement(const FSlateDrawElement& DrawElement)
 				FontTextureIndex = Entry.TextureIndex;
 
 				FontTexture = FontCache.GetTextureResource( FontTextureIndex );
-				ElementBatch = &FindBatchForElement( Layer, FShaderParams(), FontTexture, ESlateDrawPrimitive::TriangleList, ESlateShader::Font, InDrawEffects );
+				ElementBatch = &FindBatchForElement( Layer, FShaderParams(), FontTexture, ESlateDrawPrimitive::TriangleList, ESlateShader::Font, InDrawEffects, ESlateBatchDrawFlag::None, DrawElement.GetScissorRect() );
 
 				BatchVertices = &BatchVertexArrays[ElementBatch->VertexArrayIndex];
 				BatchIndices = &BatchIndexArrays[ElementBatch->IndexArrayIndex];
@@ -653,7 +663,7 @@ void FSlateElementBatcher::AddGradientElement( const FSlateDrawElement& DrawElem
 	// There must be at least one gradient stop
 	check( InPayload.GradientStops.Num() > 0 );
 
-	FSlateElementBatch& ElementBatch = FindBatchForElement( Layer, FShaderParams(), nullptr, ESlateDrawPrimitive::TriangleList, ESlateShader::Default, InDrawEffects, InPayload.bGammaCorrect == false ? ESlateBatchDrawFlag::NoGamma : ESlateBatchDrawFlag::None );
+	FSlateElementBatch& ElementBatch = FindBatchForElement( Layer, FShaderParams(), nullptr, ESlateDrawPrimitive::TriangleList, ESlateShader::Default, InDrawEffects, InPayload.bGammaCorrect == false ? ESlateBatchDrawFlag::NoGamma : ESlateBatchDrawFlag::None, DrawElement.GetScissorRect() );
 	TArray<FSlateVertex>& BatchVertices = BatchVertexArrays[ElementBatch.VertexArrayIndex];
 	TArray<SlateIndex>& BatchIndices = BatchIndexArrays[ElementBatch.IndexArrayIndex];
 
@@ -803,7 +813,7 @@ void FSlateElementBatcher::AddSplineElement( const FSlateDrawElement& DrawElemen
 	check(ResourceProxy && ResourceProxy->Resource);
 
 	// Find a batch for the element
-	FSlateElementBatch& ElementBatch = FindBatchForElement( Layer, FShaderParams::MakePixelShaderParams( FVector4( InPayload.Thickness,Radius,0,0) ), ResourceProxy->Resource, ESlateDrawPrimitive::TriangleList, ESlateShader::LineSegment, InDrawEffects );
+	FSlateElementBatch& ElementBatch = FindBatchForElement( Layer, FShaderParams::MakePixelShaderParams( FVector4( InPayload.Thickness,Radius,0,0) ), ResourceProxy->Resource, ESlateDrawPrimitive::TriangleList, ESlateShader::LineSegment, InDrawEffects, ESlateBatchDrawFlag::None, DrawElement.GetScissorRect() );
 	TArray<FSlateVertex>& BatchVertices = BatchVertexArrays[ElementBatch.VertexArrayIndex];
 	TArray<SlateIndex>& BatchIndices = BatchIndexArrays[ElementBatch.IndexArrayIndex];
 
@@ -950,7 +960,7 @@ void FSlateElementBatcher::AddLineElement( const FSlateDrawElement& DrawElement
 		const float HalfThickness = LineThickness * .5f + Radius;
 
 		// Find a batch for the element
-		FSlateElementBatch& ElementBatch = FindBatchForElement( Layer, FShaderParams::MakePixelShaderParams( FVector4(RequestedThickness,Radius,0,0) ), ResourceProxy->Resource, ESlateDrawPrimitive::TriangleList, ESlateShader::LineSegment, DrawEffects );
+		FSlateElementBatch& ElementBatch = FindBatchForElement( Layer, FShaderParams::MakePixelShaderParams( FVector4(RequestedThickness,Radius,0,0) ), ResourceProxy->Resource, ESlateDrawPrimitive::TriangleList, ESlateShader::LineSegment, DrawEffects, ESlateBatchDrawFlag::None, DrawElement.GetScissorRect() );
 		TArray<FSlateVertex>& BatchVertices = BatchVertexArrays[ElementBatch.VertexArrayIndex];
 		TArray<SlateIndex>& BatchIndices = BatchIndexArrays[ElementBatch.IndexArrayIndex];
 
@@ -1061,7 +1071,7 @@ void FSlateElementBatcher::AddLineElement( const FSlateDrawElement& DrawElement
 	else
 	{
 		// Find a batch for the element
-		FSlateElementBatch& ElementBatch = FindBatchForElement( Layer, FShaderParams(), nullptr, ESlateDrawPrimitive::LineList, ESlateShader::Default, DrawEffects );
+		FSlateElementBatch& ElementBatch = FindBatchForElement( Layer, FShaderParams(), nullptr, ESlateDrawPrimitive::LineList, ESlateShader::Default, DrawEffects, ESlateBatchDrawFlag::None, DrawElement.GetScissorRect() );
 		TArray<FSlateVertex>& BatchVertices = BatchVertexArrays[ElementBatch.VertexArrayIndex];
 		TArray<SlateIndex>& BatchIndices = BatchIndexArrays[ElementBatch.IndexArrayIndex];
 
@@ -1120,7 +1130,7 @@ void FSlateElementBatcher::AddViewportElement( const FSlateDrawElement& DrawElem
 
 	FSlateShaderResource* ViewportResource = ViewportPin.IsValid() ? ViewportPin->GetViewportRenderTargetTexture() : nullptr;
 
-	FSlateElementBatch& ElementBatch = FindBatchForElement( Layer, FShaderParams(), ViewportResource, ESlateDrawPrimitive::TriangleList, ESlateShader::Default, InDrawEffects, DrawFlags );
+	FSlateElementBatch& ElementBatch = FindBatchForElement( Layer, FShaderParams(), ViewportResource, ESlateDrawPrimitive::TriangleList, ESlateShader::Default, InDrawEffects, DrawFlags, DrawElement.GetScissorRect() );
 	TArray<FSlateVertex>& BatchVertices = BatchVertexArrays[ElementBatch.VertexArrayIndex];
 	TArray<SlateIndex>& BatchIndices = BatchIndexArrays[ElementBatch.IndexArrayIndex];
 
@@ -1251,7 +1261,7 @@ void FSlateElementBatcher::AddBorderElement( const FSlateDrawElement& DrawElemen
 	// Pass the tiling information as a flag so we can pick the correct texture addressing mode
 	ESlateBatchDrawFlag::Type DrawFlags = (ESlateBatchDrawFlag::TileU|ESlateBatchDrawFlag::TileV);
 
-	FSlateElementBatch& ElementBatch = FindBatchForElement( Layer, ShaderParams, Resource, ESlateDrawPrimitive::TriangleList, ESlateShader::Border, InDrawEffects, DrawFlags );
+	FSlateElementBatch& ElementBatch = FindBatchForElement( Layer, ShaderParams, Resource, ESlateDrawPrimitive::TriangleList, ESlateShader::Border, InDrawEffects, DrawFlags, DrawElement.GetScissorRect() );
 	TArray<FSlateVertex>& BatchVertices = BatchVertexArrays[ ElementBatch.VertexArrayIndex ];
 	TArray<SlateIndex>& BatchIndices = BatchIndexArrays[ ElementBatch.IndexArrayIndex ];
 
@@ -1394,7 +1404,7 @@ void FSlateElementBatcher::AddCustomElement( const FSlateDrawElement& DrawElemen
 
 
 		// Custom elements are not batched together 
-		ElementBatches->Add( FSlateElementBatch( InPayload.CustomDrawer ) );
+		ElementBatches->Add( FSlateElementBatch( InPayload.CustomDrawer, DrawElement.GetScissorRect() ) );
 	}
 }
 
@@ -1406,7 +1416,8 @@ FSlateElementBatch& FSlateElementBatcher::FindBatchForElement(
 	ESlateDrawPrimitive::Type PrimitiveType,
 	ESlateShader::Type ShaderType, 
 	ESlateDrawEffect::Type DrawEffects, 
-	ESlateBatchDrawFlag::Type DrawFlags)
+	ESlateBatchDrawFlag::Type DrawFlags,
+	const TOptional<FShortRect>& ScissorRect)
 {
 	SCOPE_CYCLE_COUNTER(STAT_SlateFindBatchTime);
 
@@ -1420,7 +1431,7 @@ FSlateElementBatch& FSlateElementBatcher::FindBatchForElement(
 	check( ElementBatches );
 
 	// Create a temp batch so we can use it as our key to find if the same batch already exists
-	FSlateElementBatch TempBatch( InTexture, ShaderParams, ShaderType, PrimitiveType, DrawEffects, DrawFlags );
+	FSlateElementBatch TempBatch( InTexture, ShaderParams, ShaderType, PrimitiveType, DrawEffects, DrawFlags, ScissorRect );
 
 	FSlateElementBatch* ElementBatch = ElementBatches->Find( TempBatch );
 	if( !ElementBatch )
diff --git a/Engine/Source/Runtime/SlateCore/Public/Rendering/DrawElements.h b/Engine/Source/Runtime/SlateCore/Public/Rendering/DrawElements.h
--- a/Engine/Source/Runtime/SlateCore/Public/Rendering/DrawElements.h
+++ b/Engine/Source/Runtime/SlateCore/Public/Rendering/DrawElements.h
@@ -322,8 +322,11 @@ public:
 	const FSlateRect& GetClippingRect() const { return ClippingRect; }
 	const FSlateDataPayload& GetDataPayload() const { return DataPayload; }
 	uint32 GetDrawEffects() const { return DrawEffects; }
+	const TOptional<FShortRect>& GetScissorRect() const { return ScissorRect; }
 
 private:
+	void Init(uint32 InLayer, const FPaintGeometry& PaintGeometry, const FSlateRect& InClippingRect, ESlateDrawEffect::Type InDrawEffects);
+
 
 	static FVector2D GetRotationPoint( const FPaintGeometry& PaintGeometry, const TOptional<FVector2D>& UserRotationPoint, ERotationSpace RotationSpace );
 
@@ -337,6 +340,7 @@ private:
 	uint32 Layer;
 	uint32 DrawEffects;
 	EElementType ElementType;
+	TOptional<FShortRect> ScissorRect;
 };
 
 /**
@@ -372,8 +376,8 @@ struct FShaderParams
 class FSlateElementBatch
 {
 public: 
-	FSlateElementBatch( const FSlateShaderResource* InShaderResource, const FShaderParams& InShaderParams, ESlateShader::Type ShaderType, ESlateDrawPrimitive::Type PrimitiveType, ESlateDrawEffect::Type DrawEffects, ESlateBatchDrawFlag::Type DrawFlags )
-		: BatchKey( InShaderParams, ShaderType, PrimitiveType, DrawEffects, DrawFlags )
+	FSlateElementBatch( const FSlateShaderResource* InShaderResource, const FShaderParams& InShaderParams, ESlateShader::Type ShaderType, ESlateDrawPrimitive::Type PrimitiveType, ESlateDrawEffect::Type DrawEffects, ESlateBatchDrawFlag::Type DrawFlags, const TOptional<FShortRect>& ScissorRect )
+		: BatchKey( InShaderParams, ShaderType, PrimitiveType, DrawEffects, DrawFlags, ScissorRect )
 		, ShaderResource(InShaderResource)
 		, VertexOffset(0)
 		, IndexOffset(0)
@@ -384,8 +388,8 @@ public:
 		, IndexArrayIndex(INDEX_NONE)
 	{}
 
-	FSlateElementBatch( TWeakPtr<ICustomSlateElement, ESPMode::ThreadSafe> InCustomDrawer )
-		: BatchKey( InCustomDrawer )
+	FSlateElementBatch( TWeakPtr<ICustomSlateElement, ESPMode::ThreadSafe> InCustomDrawer, const TOptional<FShortRect>& ScissorRect )
+		: BatchKey( InCustomDrawer, ScissorRect )
 		, ShaderResource( nullptr )
 		, VertexOffset(0)
 		, IndexOffset(0)
@@ -413,48 +417,62 @@ public:
 	ESlateDrawPrimitive::Type GetPrimitiveType() const { return BatchKey.DrawPrimitiveType; }
 	ESlateShader::Type GetShaderType() const { return BatchKey.ShaderType; } 
 	ESlateDrawEffect::Type GetDrawEffects() const { return BatchKey.DrawEffects; }
+	const TOptional<FShortRect>& GetScissorRect() const { return BatchKey.ScissorRect; }
 	const TWeakPtr<ICustomSlateElement, ESPMode::ThreadSafe> GetCustomDrawer() const { return BatchKey.CustomDrawer; }
 private:
 	struct FBatchKey
 	{
-		const FShaderParams ShaderParams;
 		const TWeakPtr<ICustomSlateElement, ESPMode::ThreadSafe> CustomDrawer;
+		const FShaderParams ShaderParams;
 		const ESlateBatchDrawFlag::Type DrawFlags;	
 		const ESlateShader::Type ShaderType;
 		const ESlateDrawPrimitive::Type DrawPrimitiveType;
 		const ESlateDrawEffect::Type DrawEffects;
+		const TOptional<FShortRect> ScissorRect;
 
-		FBatchKey( const FShaderParams& InShaderParams, ESlateShader::Type InShaderType, ESlateDrawPrimitive::Type InDrawPrimitiveType, ESlateDrawEffect::Type InDrawEffects, ESlateBatchDrawFlag::Type InDrawFlags )
+		FBatchKey( const FShaderParams& InShaderParams, ESlateShader::Type InShaderType, ESlateDrawPrimitive::Type InDrawPrimitiveType, ESlateDrawEffect::Type InDrawEffects, ESlateBatchDrawFlag::Type InDrawFlags, const TOptional<FShortRect>& InScissorRect )
 			: ShaderParams( InShaderParams )
 			, DrawFlags( InDrawFlags )
 			, ShaderType( InShaderType )
 			, DrawPrimitiveType( InDrawPrimitiveType )
 			, DrawEffects( InDrawEffects )
+			, ScissorRect( InScissorRect )
 		{
 		}
 
-		FBatchKey( TWeakPtr<ICustomSlateElement, ESPMode::ThreadSafe> InCustomDrawer )
-			: ShaderParams()
-			, CustomDrawer( InCustomDrawer )
+		FBatchKey( TWeakPtr<ICustomSlateElement, ESPMode::ThreadSafe> InCustomDrawer, const TOptional<FShortRect>& InScissorRect )
+			: CustomDrawer( InCustomDrawer )
+			, ShaderParams()
 			, DrawFlags( ESlateBatchDrawFlag::None )
 			, ShaderType( ESlateShader::Default )
 			, DrawPrimitiveType( ESlateDrawPrimitive::TriangleList )
 			, DrawEffects( ESlateDrawEffect::None )
+			, ScissorRect( InScissorRect )
 		{}
 
-		FORCEINLINE bool operator==( const FBatchKey& Other ) const
+		bool operator==( const FBatchKey& Other ) const
 		{
 			return DrawFlags == Other.DrawFlags
 				&& ShaderType == Other.ShaderType
 				&& DrawPrimitiveType == Other.DrawPrimitiveType
 				&& DrawEffects == Other.DrawEffects
-			    && ShaderParams == Other.ShaderParams
-			   	&& CustomDrawer == Other.CustomDrawer;
+				&& ShaderParams == Other.ShaderParams
+				&& ScissorRect == Other.ScissorRect
+				&& CustomDrawer == Other.CustomDrawer
+				;
 		}
 
+		/** Compute an efficient hash for this type for use in hash containers. */
 		friend uint32 GetTypeHash( const FBatchKey& InBatchKey )
 		{
-			return FCrc::MemCrc32(&InBatchKey.ShaderParams, sizeof(FShaderParams)) ^ ((InBatchKey.ShaderType << 16) | (InBatchKey.DrawFlags+InBatchKey.ShaderType+InBatchKey.DrawPrimitiveType+InBatchKey.DrawEffects));
+			// NOTE: Assumes these enum types are 8 bits.
+			uint32 RunningHash = (uint32)InBatchKey.DrawFlags << 24 | (uint32)InBatchKey.ShaderType << 16 | (uint32)InBatchKey.DrawPrimitiveType << 8 | (uint32)InBatchKey.DrawEffects << 0;
+			RunningHash = InBatchKey.CustomDrawer.IsValid() ? PointerHash(InBatchKey.CustomDrawer.Pin().Get(), RunningHash) : RunningHash;
+			RunningHash = HashCombine(GetTypeHash(InBatchKey.ShaderParams.PixelParams), RunningHash);
+			// NOTE: Assumes this type is 64 bits, no padding.
+			RunningHash = InBatchKey.ScissorRect.IsSet() ? HashCombine(GetTypeHash(*reinterpret_cast<const uint64*>(&InBatchKey.ScissorRect.GetValue())), RunningHash) : RunningHash;
+			return RunningHash;
+			//return FCrc::MemCrc32(&InBatchKey.ShaderParams, sizeof(FShaderParams)) ^ ((InBatchKey.ShaderType << 16) | (InBatchKey.DrawFlags+InBatchKey.ShaderType+InBatchKey.DrawPrimitiveType+InBatchKey.DrawEffects));
 		}
 	};
 
@@ -491,6 +509,7 @@ public:
 		, ShaderType( InBatch.GetShaderType() )
 		, DrawPrimitiveType( InBatch.GetPrimitiveType() )
 		, DrawEffects( InBatch.GetDrawEffects() )
+		, ScissorRect( InBatch.GetScissorRect() )
 		, VertexOffset( InBatch.VertexOffset )
 		, IndexOffset( InBatch.IndexOffset )
 		, NumVertices( InBatch.NumVertices )
@@ -513,6 +532,8 @@ public:
 
 	const ESlateDrawEffect::Type DrawEffects;
 
+	const TOptional<FShortRect> ScissorRect;
+
 	/** How far into the vertex buffer is this batch*/
 	const uint32 VertexOffset;
 	/** How far into the index buffer this batch is*/	
diff --git a/Engine/Source/Runtime/SlateCore/Public/Rendering/ElementBatcher.h b/Engine/Source/Runtime/SlateCore/Public/Rendering/ElementBatcher.h
--- a/Engine/Source/Runtime/SlateCore/Public/Rendering/ElementBatcher.h
+++ b/Engine/Source/Runtime/SlateCore/Public/Rendering/ElementBatcher.h
@@ -114,7 +114,8 @@ private:
 											 ESlateDrawPrimitive::Type PrimitiveType, 
 											 ESlateShader::Type ShaderType, 
 											 ESlateDrawEffect::Type DrawEffects, 
-											 ESlateBatchDrawFlag::Type DrawFlags = ESlateBatchDrawFlag::None);
+											 ESlateBatchDrawFlag::Type DrawFlags,
+											 const TOptional<FShortRect>& ScissorRect);
 
 	void AddVertices( TArray<FSlateVertex>& OutVertices, FSlateElementBatch& ElementBatch, const TArray<FSlateVertex>& VertexBatch );
 
diff --git a/Engine/Source/Runtime/SlateCore/Public/Rendering/RenderingCommon.h b/Engine/Source/Runtime/SlateCore/Public/Rendering/RenderingCommon.h
--- a/Engine/Source/Runtime/SlateCore/Public/Rendering/RenderingCommon.h
+++ b/Engine/Source/Runtime/SlateCore/Public/Rendering/RenderingCommon.h
@@ -179,6 +179,29 @@ struct FSlateVertex
 
 template<> struct TIsPODType<FSlateVertex> { enum { Value = true }; };
 
+/** Stores an aligned rect as shorts. */
+struct FShortRect
+{
+	FShortRect() : Left(0), Top(0), Right(0), Bottom(0) {}
+	FShortRect(uint16 InLeft, uint16 InTop, uint16 InRight, uint16 InBottom) : Left(InLeft), Top(InTop), Right(InRight), Bottom(InBottom) {}
+	explicit FShortRect(const FSlateRect& Rect) : Left((uint16)Rect.Left), Top((uint16)Rect.Top), Right((uint16)Rect.Right), Bottom((uint16)Rect.Bottom) {}
+	bool operator==(const FShortRect& RHS) const { return Left == RHS.Left && Top == RHS.Top && Right == RHS.Right && Bottom == RHS.Bottom; }
+	bool operator!=(const FShortRect& RHS) const { return !(*this == RHS); }
+	bool DoesIntersect( const FShortRect& B ) const
+	{
+		const bool bDoNotOverlap =
+			B.Right < Left || Right < B.Left ||
+			B.Bottom < Top || Bottom < B.Top;
+
+		return ! bDoNotOverlap;
+	}
+
+	uint16 Left;
+	uint16 Top;
+	uint16 Right;
+	uint16 Bottom;
+};
+
 /**
  * Viewport implementation interface that is used by SViewport when it needs to draw and processes input.                   
  */
diff --git a/Engine/Source/Runtime/SlateRHIRenderer/Private/SlateRHIRenderingPolicy.cpp b/Engine/Source/Runtime/SlateRHIRenderer/Private/SlateRHIRenderingPolicy.cpp
--- a/Engine/Source/Runtime/SlateRHIRenderer/Private/SlateRHIRenderingPolicy.cpp
+++ b/Engine/Source/Runtime/SlateRHIRenderer/Private/SlateRHIRenderingPolicy.cpp
@@ -348,6 +348,16 @@ void FSlateRHIRenderingPolicy::DrawElements(FRHICommandListImmediate& RHICmdList
 					RHICmdList.SetRasterizerState(TStaticRasterizerState<FM_Solid, CM_None, true>::GetRHI());
 				}
 
+				if (RenderBatch.ScissorRect.IsSet())
+				{
+					RHICmdList.SetScissorRect(true, RenderBatch.ScissorRect.GetValue().Left, RenderBatch.ScissorRect.GetValue().Top, RenderBatch.ScissorRect.GetValue().Right, RenderBatch.ScissorRect.GetValue().Bottom);
+				}
+				else
+				{
+					RHICmdList.SetScissorRect(false, 0, 0, 0, 0);
+				}
+
+
 				FTexture2DRHIRef TextureRHI;
 				if(ShaderResource)
 				{
