- Slider bar now draws a thin box instead of a line to make it show up on mobile
- FSliderStyle includes a brush for the enabled & disabled bar

diff --git a/Engine/Source/Editor/EditorStyle/Private/SlateEditorStyle.cpp b/Engine/Source/Editor/EditorStyle/Private/SlateEditorStyle.cpp
--- a/Engine/Source/Editor/EditorStyle/Private/SlateEditorStyle.cpp
+++ b/Engine/Source/Editor/EditorStyle/Private/SlateEditorStyle.cpp
@@ -1322,6 +1322,8 @@ void FSlateEditorStyle::FStyle::SetupGeneralStyles()
 	// Slider and Volume Control
 	{
 		FSliderStyle SliderStyle = FSliderStyle()
+			.SetNormalBarImage(FSlateColorBrush(FColor::White))
+			.SetDisabledBarImage(FSlateColorBrush(FLinearColor::Gray))
 			.SetNormalThumbImage( BOX_BRUSH( "Common/Button", 8.0f/32.0f ) )
 			.SetDisabledThumbImage( BOX_BRUSH( "Common/Button_Disabled", 8.0f/32.0f ) );
 		Set( "Slider", SliderStyle );
diff --git a/Engine/Source/Runtime/Slate/Private/Widgets/Input/SSlider.cpp b/Engine/Source/Runtime/Slate/Private/Widgets/Input/SSlider.cpp
--- a/Engine/Source/Runtime/Slate/Private/Widgets/Input/SSlider.cpp
+++ b/Engine/Source/Runtime/Slate/Private/Widgets/Input/SSlider.cpp
@@ -76,28 +76,25 @@ int32 SSlider::OnPaint( const FPaintArgs& Args, const FGeometry& AllottedGeometr
 			MyClippingRect);
 	}
 
-
-	// draw slider bar
-	TArray<FVector2D> LinePoints;
-	LinePoints.Add(SliderStartPoint);
-	LinePoints.Add(SliderEndPoint);
-
 	const bool bEnabled = ShouldBeEnabled(bParentEnabled);
 	const ESlateDrawEffect::Type DrawEffects = bEnabled ? ESlateDrawEffect::None : ESlateDrawEffect::DisabledEffect;
 
-	FSlateDrawElement::MakeLines( 
+	// draw slider bar
+	auto BarTopLeft = FVector2D(SliderStartPoint.X, SliderStartPoint.Y - 1);
+	auto BarSize = FVector2D(SliderEndPoint.X - SliderStartPoint.X, 2);
+	FSlateDrawElement::MakeBox(
 		OutDrawElements,
 		LayerId,
-		SliderGeometry.ToPaintGeometry(),
-		LinePoints,
+		SliderGeometry.ToPaintGeometry(BarTopLeft, BarSize),
+		LockedAttribute.Get() ? &Style->DisabledBarImage : &Style->NormalBarImage,
 		RotatedClippingRect,
 		DrawEffects,
 		SliderBarColor.Get().GetColor(InWidgetStyle) * InWidgetStyle.GetColorAndOpacityTint()
-	);
+		);
 
 	++LayerId;
 
-	// draw slider handle
+	// draw slider thumb
 	FSlateDrawElement::MakeBox( 
 		OutDrawElements,
 		LayerId,
diff --git a/Engine/Source/Runtime/SlateCore/Private/Styling/CoreStyle.cpp b/Engine/Source/Runtime/SlateCore/Private/Styling/CoreStyle.cpp
--- a/Engine/Source/Runtime/SlateCore/Private/Styling/CoreStyle.cpp
+++ b/Engine/Source/Runtime/SlateCore/Private/Styling/CoreStyle.cpp
@@ -543,6 +543,8 @@ TSharedRef<ISlateStyle> FCoreStyle::Create( const FName& InStyleSetName )
 	// SSlider and SVolumeControl defaults...
 	{
 		FSliderStyle SliderStyle = FSliderStyle()
+			.SetNormalBarImage(FSlateColorBrush(FColor::White))
+			.SetDisabledBarImage(FSlateColorBrush(FLinearColor::Gray))
 			.SetNormalThumbImage( IMAGE_BRUSH( "Common/Button", FVector2D(8.0f, 14.0f) ) )
 			.SetDisabledThumbImage( IMAGE_BRUSH( "Common/Button_Disabled", FVector2D(8.0f, 14.0f) ) );
 		Style->Set( "Slider", SliderStyle );
diff --git a/Engine/Source/Runtime/SlateCore/Private/Styling/SlateTypes.cpp b/Engine/Source/Runtime/SlateCore/Private/Styling/SlateTypes.cpp
--- a/Engine/Source/Runtime/SlateCore/Private/Styling/SlateTypes.cpp
+++ b/Engine/Source/Runtime/SlateCore/Private/Styling/SlateTypes.cpp
@@ -358,13 +358,14 @@ const FSearchBoxStyle& FSearchBoxStyle::GetDefault()
 	return Default;
 }
 
-
 FSliderStyle::FSliderStyle()
 {
 }
 
 void FSliderStyle::GetResources( TArray< const FSlateBrush* >& OutBrushes ) const
 {
+	OutBrushes.Add( &NormalBarImage );
+	OutBrushes.Add( &DisabledBarImage );
 	OutBrushes.Add( &NormalThumbImage );
 	OutBrushes.Add( &DisabledThumbImage );
 }
diff --git a/Engine/Source/Runtime/SlateCore/Public/Styling/SlateTypes.h b/Engine/Source/Runtime/SlateCore/Public/Styling/SlateTypes.h
--- a/Engine/Source/Runtime/SlateCore/Public/Styling/SlateTypes.h
+++ b/Engine/Source/Runtime/SlateCore/Public/Styling/SlateTypes.h
@@ -794,6 +794,16 @@ struct SLATECORE_API FSliderStyle : public FSlateWidgetStyle
 
 	static const FSliderStyle& GetDefault();
 
+	/** Image to use when the slider bar is in its normal state */
+	UPROPERTY(EditAnywhere, BlueprintReadWrite, Category=Appearance)
+	FSlateBrush NormalBarImage;
+	FSliderStyle& SetNormalBarImage(const FSlateBrush& InNormalBarImage){ NormalBarImage = InNormalBarImage; return *this; }
+
+	/** Image to use when the slider bar is in its disabled state */
+	UPROPERTY(EditAnywhere, BlueprintReadWrite, Category = Appearance)
+	FSlateBrush DisabledBarImage;
+	FSliderStyle& SetDisabledBarImage(const FSlateBrush& InDisabledBarImage){ DisabledBarImage = InDisabledBarImage; return *this; }
+
 	/** Image to use when the slider thumb is in its normal state */
 	UPROPERTY(EditAnywhere, BlueprintReadWrite, Category=Appearance)
 	FSlateBrush NormalThumbImage;
@@ -1037,6 +1047,16 @@ struct SLATECORE_API FTableRowStyle : public FSlateWidgetStyle
 	FSlateBrush OddRowBackgroundBrush;
 	FTableRowStyle& SetOddRowBackgroundBrush( const FSlateBrush& InOddRowBackgroundBrush ){ OddRowBackgroundBrush = InOddRowBackgroundBrush; return *this; }
 
+	/** Brush used when drawing the drag-drop indicator lines above and below rows */
+	UPROPERTY(EditAnywhere, Category = Appearance)
+	FSlateBrush DragDropLineIndicatorBrush;
+	FTableRowStyle& SetDragDropLineIndicatorBrush(const FSlateBrush& InDragDropLineIndicatorBrush){ DragDropLineIndicatorBrush = InDragDropLineIndicatorBrush; return *this; }
+
+	/** Brush used when drawing the drag-drop indicator border around a row */
+	UPROPERTY(EditAnywhere, Category = Appearance)
+	FSlateBrush DragDropBorderIndicatorBrush;
+	FTableRowStyle& SetDragDropBorderIndicatorBrush(const FSlateBrush& InDragDropBorderIndicatorBrush){ DragDropBorderIndicatorBrush = InDragDropBorderIndicatorBrush; return *this; }
+
 	/** Text color used for all rows */
 	UPROPERTY(EditAnywhere, BlueprintReadWrite, Category=Appearance)
 	FSlateColor TextColor;
