SlateApplication - Add back an API which was removed when multitouch capture was added, and mark it as deprecated.

diff --git a/Engine/Source/Runtime/Slate/Private/Framework/Application/SlateApplication.cpp b/Engine/Source/Runtime/Slate/Private/Framework/Application/SlateApplication.cpp
--- a/Engine/Source/Runtime/Slate/Private/Framework/Application/SlateApplication.cpp
+++ b/Engine/Source/Runtime/Slate/Private/Framework/Application/SlateApplication.cpp
@@ -2879,6 +2879,10 @@ TSharedPtr< SWidget > FSlateApplication::GetKeyboardFocusedWidget() const
 	return TSharedPtr< SWidget >();
 }
 
+TSharedPtr<SWidget> FSlateApplication::GetMouseCaptor() const
+{
+	return MouseCaptor.ToSharedWidget(CursorPointerIndex);
+}
 
 bool FSlateApplication::HasAnyMouseCaptor() const
 {
diff --git a/Engine/Source/Runtime/Slate/Public/Framework/Application/SlateApplication.h b/Engine/Source/Runtime/Slate/Public/Framework/Application/SlateApplication.h
--- a/Engine/Source/Runtime/Slate/Public/Framework/Application/SlateApplication.h
+++ b/Engine/Source/Runtime/Slate/Public/Framework/Application/SlateApplication.h
@@ -946,6 +946,8 @@ public:
 	
 	virtual TSharedPtr<SWidget> GetKeyboardFocusedWidget( ) const override;
 
+	virtual TSharedPtr< SWidget > GetMouseCaptor() const override;
+
 	virtual bool HasAnyMouseCaptor( ) const override;
 
 	virtual FSlateRect GetPreferredWorkArea( ) const override;
diff --git a/Engine/Source/Runtime/SlateCore/Public/Application/SlateApplicationBase.h b/Engine/Source/Runtime/SlateCore/Public/Application/SlateApplicationBase.h
--- a/Engine/Source/Runtime/SlateCore/Public/Application/SlateApplicationBase.h
+++ b/Engine/Source/Runtime/SlateCore/Public/Application/SlateApplicationBase.h
@@ -135,6 +135,14 @@ public:
 	 */
 	virtual TSharedPtr< SWidget > GetKeyboardFocusedWidget( ) const = 0;
 
+	/**
+	* Gets the Widget that currently captures the mouse.
+	*
+	* @return The captor widget, or nullptr if no widget captured the mouse.
+	*/
+	DEPRECATED(4.5, "This API is no longer supported.  You can check if a specific widget has capture using the SWidget::HasMouseCapture() API")
+	virtual TSharedPtr< SWidget > GetMouseCaptor() const = 0;
+
 	/**
 	 * Gets whether or not a widget has captured the mouse.
 	 *
