Fix for warning about implementation of deprecated method in FSlateApplication.

diff --git a/Engine/Source/Runtime/Slate/Private/Framework/Application/SlateApplication.cpp b/Engine/Source/Runtime/Slate/Private/Framework/Application/SlateApplication.cpp
--- a/Engine/Source/Runtime/Slate/Private/Framework/Application/SlateApplication.cpp
+++ b/Engine/Source/Runtime/Slate/Private/Framework/Application/SlateApplication.cpp
@@ -2879,7 +2879,7 @@ TSharedPtr< SWidget > FSlateApplication::GetKeyboardFocusedWidget() const
 	return TSharedPtr< SWidget >();
 }
 
-TSharedPtr<SWidget> FSlateApplication::GetMouseCaptor() const
+TSharedPtr<SWidget> FSlateApplication::GetMouseCaptorImpl() const
 {
 	return MouseCaptor.ToSharedWidget(CursorPointerIndex);
 }
diff --git a/Engine/Source/Runtime/Slate/Public/Framework/Application/SlateApplication.h b/Engine/Source/Runtime/Slate/Public/Framework/Application/SlateApplication.h
--- a/Engine/Source/Runtime/Slate/Public/Framework/Application/SlateApplication.h
+++ b/Engine/Source/Runtime/Slate/Public/Framework/Application/SlateApplication.h
@@ -946,8 +946,10 @@ public:
 	
 	virtual TSharedPtr<SWidget> GetKeyboardFocusedWidget( ) const override;
 
-	virtual TSharedPtr< SWidget > GetMouseCaptor() const override;
+protected:
+	virtual TSharedPtr< SWidget > GetMouseCaptorImpl() const override;
 
+public:
 	virtual bool HasAnyMouseCaptor( ) const override;
 
 	virtual FSlateRect GetPreferredWorkArea( ) const override;
diff --git a/Engine/Source/Runtime/SlateCore/Public/Application/SlateApplicationBase.h b/Engine/Source/Runtime/SlateCore/Public/Application/SlateApplicationBase.h
--- a/Engine/Source/Runtime/SlateCore/Public/Application/SlateApplicationBase.h
+++ b/Engine/Source/Runtime/SlateCore/Public/Application/SlateApplicationBase.h
@@ -141,8 +141,20 @@ public:
 	* @return The captor widget, or nullptr if no widget captured the mouse.
 	*/
 	DEPRECATED(4.5, "This API is no longer supported.  You can check if a specific widget has capture using the SWidget::HasMouseCapture() API")
-	virtual TSharedPtr< SWidget > GetMouseCaptor() const = 0;
+	inline TSharedPtr< SWidget > GetMouseCaptor() const
+	{
+		return GetMouseCaptorImpl();
+	}
+
+protected:
+	/**
+	 * Implementation of GetMouseCaptor which can be overriden without warnings.
+	 * 
+	 * @return Widget with the mouse capture
+	 */
+	virtual TSharedPtr< SWidget > GetMouseCaptorImpl() const = 0;
 
+public:
 	/**
 	 * Gets whether or not a widget has captured the mouse.
 	 *
