SlateApplication - User a safer constant index for the cursor to prevent some crashes.

diff --git a/Engine/Source/Runtime/InputCore/Classes/InputCoreTypes.h b/Engine/Source/Runtime/InputCore/Classes/InputCoreTypes.h
--- a/Engine/Source/Runtime/InputCore/Classes/InputCoreTypes.h
+++ b/Engine/Source/Runtime/InputCore/Classes/InputCoreTypes.h
@@ -394,7 +394,7 @@ struct INPUTCORE_API EKeys
 
 	static const FKey Invalid;
 
-	static const int32 NUM_TOUCH_KEYS = 10;
+	static const int32 NUM_TOUCH_KEYS = 11;
 	static const FKey TouchKeys[NUM_TOUCH_KEYS];
 
 	static EConsoleForGamepadLabels::Type ConsoleForGamepadLabels;
diff --git a/Engine/Source/Runtime/SlateCore/Private/Application/SlateApplicationBase.cpp b/Engine/Source/Runtime/SlateCore/Private/Application/SlateApplicationBase.cpp
--- a/Engine/Source/Runtime/SlateCore/Private/Application/SlateApplicationBase.cpp
+++ b/Engine/Source/Runtime/SlateCore/Private/Application/SlateApplicationBase.cpp
@@ -8,4 +8,5 @@
 
 TSharedPtr<FSlateApplicationBase> FSlateApplicationBase::CurrentBaseApplication = nullptr;
 TSharedPtr<GenericApplication> FSlateApplicationBase::PlatformApplication = nullptr;
-const uint32 FSlateApplicationBase::CursorPointerIndex = 1000;
+// TODO: Identifier the cursor index in a smarter way.
+const uint32 FSlateApplicationBase::CursorPointerIndex = EKeys::NUM_TOUCH_KEYS - 1;
