Fixed a crash caused by changing an FLinearColor property to an FSlateColor

Added a SerializeFromMismatchedTag to FSlateColor that can convert an FColor or FLinearColor property to an FSlateColor.

Had to update the logic for calling SerializeFromMismatchedTag so that it would also call it when the StructName in the Tag of a StructProperty had changed.

diff --git a/Engine/Source/Runtime/CoreUObject/Private/UObject/Class.cpp b/Engine/Source/Runtime/CoreUObject/Private/UObject/Class.cpp
--- a/Engine/Source/Runtime/CoreUObject/Private/UObject/Class.cpp
+++ b/Engine/Source/Runtime/CoreUObject/Private/UObject/Class.cpp
@@ -1004,7 +1004,7 @@ void UStruct::SerializeTaggedProperties(FArchive& Ar, uint8* Data, UStruct* Defa
 				AdvanceProperty = true;
 				continue; 
 			}
-			else if( Tag.Type!=Property->GetID() && Cast<UStructProperty>(Property) && Cast<UStructProperty>(Property)->Struct && (Cast<UStructProperty>(Property)->Struct->StructFlags & STRUCT_SerializeFromMismatchedTag))
+			else if( Cast<UStructProperty>(Property) && Cast<UStructProperty>(Property)->Struct && (Tag.Type != Property->GetID() || (Tag.Type == NAME_StructProperty && Tag.StructName != Cast<UStructProperty>(Property)->Struct->GetFName())) && (Cast<UStructProperty>(Property)->Struct->StructFlags & STRUCT_SerializeFromMismatchedTag))
 			{
 				UScriptStruct::ICppStructOps* CppStructOps = Cast<UStructProperty>(Property)->Struct->GetCppStructOps();
 				check(CppStructOps && CppStructOps->HasSerializeFromMismatchedTag()); // else should not have STRUCT_SerializeFromMismatchedTag
diff --git a/Engine/Source/Runtime/SlateCore/Private/Styling/SlateColor.cpp b/Engine/Source/Runtime/SlateCore/Private/Styling/SlateColor.cpp
new file mode 100644
--- /dev/null
+++ b/Engine/Source/Runtime/SlateCore/Private/Styling/SlateColor.cpp
@@ -0,0 +1,28 @@
+// Copyright 1998-2014 Epic Games, Inc. All Rights Reserved.
+
+#include "SlateCorePrivatePCH.h"
+
+bool FSlateColor::SerializeFromMismatchedTag( const FPropertyTag& Tag, FArchive& Ar )
+{
+	if (Tag.Type == NAME_StructProperty)
+	{
+		if (Tag.StructName == NAME_Color)
+		{
+			FColor OldColor;
+			Ar << OldColor;
+			*this = FSlateColor(FLinearColor(OldColor));
+
+			return true;
+		}
+		else if(Tag.StructName == NAME_LinearColor)
+		{
+			FLinearColor OldColor;
+			Ar << OldColor;
+			*this = FSlateColor(OldColor);
+
+			return true;
+		}
+	}
+
+	return false;
+}
diff --git a/Engine/Source/Runtime/SlateCore/Public/Styling/SlateColor.h b/Engine/Source/Runtime/SlateCore/Public/Styling/SlateColor.h
--- a/Engine/Source/Runtime/SlateCore/Public/Styling/SlateColor.h
+++ b/Engine/Source/Runtime/SlateCore/Public/Styling/SlateColor.h
@@ -1,6 +1,7 @@
 // Copyright 1998-2014 Epic Games, Inc. All Rights Reserved.
 
 #pragma once
+
 #include "Styling/WidgetStyle.h"
 #include "SlateColor.generated.h"
 
@@ -34,7 +35,7 @@ namespace ESlateColorStylingMode
  * A Slate color can be a directly specified value, or the color can be pulled from a WidgetStyle.
  */
 USTRUCT(BlueprintType)
-struct FSlateColor
+struct SLATECORE_API FSlateColor
 {
 	GENERATED_USTRUCT_BODY()
 
@@ -170,6 +171,9 @@ public:
 		return FSlateColor( ESlateColorStylingMode::UseColor_Foreground_Subdued );
 	}
 
+	/** Used to upgrade an FColor or FLinearColor property to an FSlateColor property */
+	bool SerializeFromMismatchedTag( const struct FPropertyTag& Tag, FArchive& Ar );
+
 protected:
 	
 	// Private constructor to prevent construction of invalid FSlateColors
@@ -193,3 +197,13 @@ private:
 	// A shared pointer to the linked color value, if any.
 	TSharedPtr<FLinearColor> LinkedSpecifiedColor;
 };
+
+template<>
+struct TStructOpsTypeTraits<FSlateColor>
+	: public TStructOpsTypeTraitsBase
+{
+	enum 
+	{
+		WithSerializeFromMismatchedTag = true,
+	};
+};
