Fixed typo in SearchForWidgetRecursively_Reverse, so that Shift+Tab to focus to the previous widget works correctly.
#ttp 345543 - NUX:Details: NUX: Shift+Tab is behaving inconsistently with Tab -- it skips over value fields!
#branch UE4
#proj Runtime.SlateCore

diff --git a/Engine/Source/Runtime/SlateCore/Public/Layout/WidgetPath.inl b/Engine/Source/Runtime/SlateCore/Public/Layout/WidgetPath.inl
--- a/Engine/Source/Runtime/SlateCore/Public/Layout/WidgetPath.inl
+++ b/Engine/Source/Runtime/SlateCore/Public/Layout/WidgetPath.inl
@@ -53,7 +53,7 @@ bool FWidgetPath::SearchForWidgetRecursively_Reverse( const MatchRuleType& Match
 			OutReversedPath.AddWidget( SomeChild );
 			return true;
 		}
-		else if ( SearchForWidgetRecursively( MatchRule, SomeChild, OutReversedPath, VisibilityFilter) )
+		else if ( SearchForWidgetRecursively_Reverse( MatchRule, SomeChild, OutReversedPath, VisibilityFilter) )
 		{
 			OutReversedPath.AddWidget( SomeChild );
 			return true;
