Converted SWidget::RenderTransform into a TOptional for performance reasons.
* Allows code to skip the expensive stuff when no render transform is used. Specifically, skips render transform related concatenate calls in FGeometry ctor.
* UMG now sets a NULL RenderTransform when it detects the transform is identity.
TTP#346674: Slate: Editor performance has regressed in debug and release.
TTP#346490: MAIN: Performance drop when gathering resources

diff --git a/Engine/Source/Runtime/AppFramework/Private/Framework/Testing/SLayoutExample.cpp b/Engine/Source/Runtime/AppFramework/Private/Framework/Testing/SLayoutExample.cpp
--- a/Engine/Source/Runtime/AppFramework/Private/Framework/Testing/SLayoutExample.cpp
+++ b/Engine/Source/Runtime/AppFramework/Private/Framework/Testing/SLayoutExample.cpp
@@ -380,7 +380,7 @@ public:
 
 APPFRAMEWORK_API TSharedRef<SWidget> MakeLayoutExample()
 {
-	extern FSlateRenderTransform GetTestRenderTransform();
+	extern TOptional<FSlateRenderTransform> GetTestRenderTransform();
 	extern FVector2D GetTestRenderTransformPivot();
 	return
 		SNew(SExampleLayout)
diff --git a/Engine/Source/Runtime/AppFramework/Private/Framework/Testing/STableViewTesting.cpp b/Engine/Source/Runtime/AppFramework/Private/Framework/Testing/STableViewTesting.cpp
--- a/Engine/Source/Runtime/AppFramework/Private/Framework/Testing/STableViewTesting.cpp
+++ b/Engine/Source/Runtime/AppFramework/Private/Framework/Testing/STableViewTesting.cpp
@@ -1102,7 +1102,7 @@ protected:
 
 TSharedRef<SWidget> MakeTableViewTesting()
 {
-	extern FSlateRenderTransform GetTestRenderTransform();
+	extern TOptional<FSlateRenderTransform> GetTestRenderTransform();
 	extern FVector2D GetTestRenderTransformPivot();
 	return
 		SNew(STableViewTesting)
diff --git a/Engine/Source/Runtime/AppFramework/Private/Framework/Testing/SWidgetGallery.cpp b/Engine/Source/Runtime/AppFramework/Private/Framework/Testing/SWidgetGallery.cpp
--- a/Engine/Source/Runtime/AppFramework/Private/Framework/Testing/SWidgetGallery.cpp
+++ b/Engine/Source/Runtime/AppFramework/Private/Framework/Testing/SWidgetGallery.cpp
@@ -1123,7 +1123,7 @@ private:
  */
 TSharedRef<SWidget> MakeWidgetGallery()
 {
-	extern FSlateRenderTransform GetTestRenderTransform();
+	extern TOptional<FSlateRenderTransform> GetTestRenderTransform();
 	extern FVector2D GetTestRenderTransformPivot();
 	return
 		SNew(SWidgetGallery)
diff --git a/Engine/Source/Runtime/AppFramework/Private/Widgets/Testing/STestSuite.cpp b/Engine/Source/Runtime/AppFramework/Private/Widgets/Testing/STestSuite.cpp
--- a/Engine/Source/Runtime/AppFramework/Private/Widgets/Testing/STestSuite.cpp
+++ b/Engine/Source/Runtime/AppFramework/Private/Widgets/Testing/STestSuite.cpp
@@ -4331,7 +4331,7 @@ namespace
 /**
  * Global access here because we need other translation units to access this function.
  */
-FSlateRenderTransform GetTestRenderTransform()
+TOptional<FSlateRenderTransform> GetTestRenderTransform()
 {
 	return TransformCast<FSlateRenderTransform>(Concatenate(Shear, Scale, Rot));
 }
@@ -4482,8 +4482,6 @@ static TSharedPtr<FTabManager> TestSuite1TabManager;
 BEGIN_SLATE_FUNCTION_BUILD_OPTIMIZATION
 TSharedRef<SDockTab> SpawnTab(const FSpawnTabArgs& Args, FName TabIdentifier)
 {
-	extern FSlateRenderTransform GetTestRenderTransform();
-	extern FVector2D GetTestRenderTransformPivot();
 	if (TabIdentifier == FName(TEXT("AnimationTestTab")))
 	{
 		return SNew(SDockTab)
diff --git a/Engine/Source/Runtime/SlateCore/Private/Layout/Geometry.cpp b/Engine/Source/Runtime/SlateCore/Private/Layout/Geometry.cpp
--- a/Engine/Source/Runtime/SlateCore/Private/Layout/Geometry.cpp
+++ b/Engine/Source/Runtime/SlateCore/Private/Layout/Geometry.cpp
@@ -106,7 +106,16 @@ FArrangedWidget FGeometry::MakeChild(const TSharedRef<SWidget>& ChildWidget, con
 
 FArrangedWidget FGeometry::MakeChild(const TSharedRef<SWidget>& ChildWidget, const FVector2D& LocalSize, const FSlateLayoutTransform& LayoutTransform) const
 {
-	return FArrangedWidget(ChildWidget, MakeChild(LocalSize, LayoutTransform, ChildWidget->GetRenderTransform(), ChildWidget->GetRenderTransformPivot()));
+	// If there is no render transform set, use the simpler MakeChild call that doesn't bother concatenating the render transforms.
+	// This saves a significant amount of overhead since every widget does this, and most children don't have a render transform.
+	if (ChildWidget->GetRenderTransform().IsSet())
+	{
+		return FArrangedWidget(ChildWidget, MakeChild(LocalSize, LayoutTransform, ChildWidget->GetRenderTransform().GetValue(), ChildWidget->GetRenderTransformPivot()));
+	}
+	else
+	{
+		return FArrangedWidget(ChildWidget, MakeChild(LocalSize, LayoutTransform));
+	}
 }
 
 FGeometry FGeometry::MakeChild( const FVector2D& ChildOffset, const FVector2D& LocalSize, float LocalScale ) const
@@ -120,7 +129,7 @@ FArrangedWidget FGeometry::MakeChild( const TSharedRef<SWidget>& ChildWidget, co
 {
 	// Since ChildOffset is given as a LocalSpaceOffset, we MUST convert this offset into the space of the parent to construct a valid layout transform.
 	// The extra TransformPoint below does this by converting the local offset to an offset in parent space.
-	return FArrangedWidget(ChildWidget, MakeChild(LocalSize, FSlateLayoutTransform(ChildScale, TransformPoint(ChildScale, ChildOffset)), ChildWidget->GetRenderTransform(), ChildWidget->GetRenderTransformPivot()));
+	return MakeChild(ChildWidget, LocalSize, FSlateLayoutTransform(ChildScale, TransformPoint(ChildScale, ChildOffset)));
 }
 
 FPaintGeometry FGeometry::ToPaintGeometry() const
diff --git a/Engine/Source/Runtime/SlateCore/Private/Widgets/SWidget.cpp b/Engine/Source/Runtime/SlateCore/Private/Widgets/SWidget.cpp
--- a/Engine/Source/Runtime/SlateCore/Private/Widgets/SWidget.cpp
+++ b/Engine/Source/Runtime/SlateCore/Private/Widgets/SWidget.cpp
@@ -12,7 +12,7 @@ SWidget::SWidget()
 	, Cursor( TOptional<EMouseCursor::Type>() )
 	, EnabledState( true )
 	, Visibility( EVisibility::Visible )
-	, RenderTransform( FSlateRenderTransform() )
+	, RenderTransform( )
 	, RenderTransformPivot( FVector2D::ZeroVector )
 	, DesiredSize(FVector2D::ZeroVector)
 	, ToolTip()
@@ -28,7 +28,7 @@ void SWidget::Construct(
 	const TAttribute< TOptional<EMouseCursor::Type> > & InCursor ,
 	const TAttribute<bool> & InEnabledState ,
 	const TAttribute<EVisibility> & InVisibility,
-	const TAttribute<FSlateRenderTransform>& InTransform,
+	const TAttribute<TOptional<FSlateRenderTransform>>& InTransform,
 	const TAttribute<FVector2D>& InTransformPivot,
 	const FName& InTag,
 	const TArray<TSharedRef<ISlateMetaData>>& InMetaData
diff --git a/Engine/Source/Runtime/SlateCore/Public/Widgets/DeclarativeSyntaxSupport.h b/Engine/Source/Runtime/SlateCore/Public/Widgets/DeclarativeSyntaxSupport.h
--- a/Engine/Source/Runtime/SlateCore/Public/Widgets/DeclarativeSyntaxSupport.h
+++ b/Engine/Source/Runtime/SlateCore/Public/Widgets/DeclarativeSyntaxSupport.h
@@ -1009,7 +1009,7 @@ struct TSlateBaseNamedArgs
 	, _Cursor( TOptional<EMouseCursor::Type>() )
 	, _IsEnabled( true )
 	, _Visibility( EVisibility::Visible )
-	, _RenderTransform( FSlateRenderTransform() )
+	, _RenderTransform( )
 	, _RenderTransformPivot( FVector2D::ZeroVector )
 	{
 	}
@@ -1048,7 +1048,7 @@ struct TSlateBaseNamedArgs
 	SLATE_ATTRIBUTE( TOptional<EMouseCursor::Type>, Cursor )
 	SLATE_ATTRIBUTE( bool, IsEnabled )
 	SLATE_ATTRIBUTE( EVisibility, Visibility )
-	SLATE_ATTRIBUTE( FSlateRenderTransform, RenderTransform )
+	SLATE_ATTRIBUTE( TOptional<FSlateRenderTransform>, RenderTransform )
 	SLATE_ATTRIBUTE( FVector2D, RenderTransformPivot )
 	SLATE_ARGUMENT( FName, Tag )
 
diff --git a/Engine/Source/Runtime/SlateCore/Public/Widgets/SWidget.h b/Engine/Source/Runtime/SlateCore/Public/Widgets/SWidget.h
--- a/Engine/Source/Runtime/SlateCore/Public/Widgets/SWidget.h
+++ b/Engine/Source/Runtime/SlateCore/Public/Widgets/SWidget.h
@@ -72,7 +72,7 @@ public:
 		const TAttribute< TOptional<EMouseCursor::Type> > & InCursor ,
 		const TAttribute<bool> & InEnabledState ,
 		const TAttribute<EVisibility> & InVisibility,
-		const TAttribute<FSlateRenderTransform>& InTransform,
+		const TAttribute<TOptional<FSlateRenderTransform>>& InTransform,
 		const TAttribute<FVector2D>& InTransformPivot,
 		const FName& InTag,
 		const TArray<TSharedRef<ISlateMetaData>>& InMetaData);
@@ -82,7 +82,7 @@ public:
 		const TAttribute< TOptional<EMouseCursor::Type> > & InCursor ,
 		const TAttribute<bool> & InEnabledState ,
 		const TAttribute<EVisibility> & InVisibility,
-		const TAttribute<FSlateRenderTransform>& InTransform,
+		const TAttribute<TOptional<FSlateRenderTransform>>& InTransform,
 		const TAttribute<FVector2D>& InTransformPivot,
 		const FName& InTag,
 		const TArray<TSharedRef<ISlateMetaData>>& InMetaData)
@@ -565,13 +565,13 @@ public:
 	}
 
 	/** @return the render transform of the widget. */
-	const FSlateRenderTransform& GetRenderTransform() const
+	const TOptional<FSlateRenderTransform>& GetRenderTransform() const
 	{
 		return RenderTransform.Get();
 	}
 
-	/** @param InTransform the render transform to set for the widget (tranforms from widget's local space). */
-	void SetRenderTransform( TAttribute<FSlateRenderTransform> InTransform )
+	/** @param InTransform the render transform to set for the widget (tranforms from widget's local space). TOptional<> to allow code to skip expensive overhead if there is no render transform applied. */
+	void SetRenderTransform( TAttribute<TOptional<FSlateRenderTransform>> InTransform )
 	{
 		RenderTransform = InTransform;
 	}
@@ -817,8 +817,8 @@ protected:
 	/** Is this widget visible, hidden or collapsed */
 	TAttribute< EVisibility > Visibility;
 
-	/** Render transform of this widget */
-	TAttribute< FSlateRenderTransform > RenderTransform;
+	/** Render transform of this widget. TOptional<> to allow code to skip expensive overhead if there is no render transform applied. */
+	TAttribute< TOptional<FSlateRenderTransform> > RenderTransform;
 
 	/** Render transform pivot of this widget (in normalized local space) */
 	TAttribute< FVector2D > RenderTransformPivot;
diff --git a/Engine/Source/Runtime/UMG/Private/Components/Widget.cpp b/Engine/Source/Runtime/UMG/Private/Components/Widget.cpp
--- a/Engine/Source/Runtime/UMG/Private/Components/Widget.cpp
+++ b/Engine/Source/Runtime/UMG/Private/Components/Widget.cpp
@@ -49,12 +49,19 @@ void UWidget::UpdateRenderTransform()
 {
 	if ( MyWidget.IsValid() )
 	{
-		// Compute the M (Shear Slot) = CoTan(90 - SlopeAngle)
-		float ShearX = 1.0f / FMath::Tan(FMath::DegreesToRadians(90 - FMath::Clamp(RenderTransform.Shear.X, -89.0f, 89.0f)));
-		float ShearY = 1.0f / FMath::Tan(FMath::DegreesToRadians(90 - FMath::Clamp(RenderTransform.Shear.Y, -89.0f, 89.0f)));
+		if (!RenderTransform.IsIdentity())
+		{
+			// Compute the M (Shear Slot) = CoTan(90 - SlopeAngle)
+			float ShearX = 1.0f / FMath::Tan(FMath::DegreesToRadians(90 - FMath::Clamp(RenderTransform.Shear.X, -89.0f, 89.0f)));
+			float ShearY = 1.0f / FMath::Tan(FMath::DegreesToRadians(90 - FMath::Clamp(RenderTransform.Shear.Y, -89.0f, 89.0f)));
 
-		FSlateRenderTransform Transform2D = ::Concatenate(FScale2D(RenderTransform.Scale), FShear2D(ShearX, ShearY), FQuat2D(FMath::DegreesToRadians(RenderTransform.Angle)), FVector2D(RenderTransform.Translation));
-		MyWidget.Pin()->SetRenderTransform(Transform2D);
+			FSlateRenderTransform Transform2D = ::Concatenate(FScale2D(RenderTransform.Scale), FShear2D(ShearX, ShearY), FQuat2D(FMath::DegreesToRadians(RenderTransform.Angle)), FVector2D(RenderTransform.Translation));
+			MyWidget.Pin()->SetRenderTransform(Transform2D);
+		}
+		else
+		{
+			MyWidget.Pin()->SetRenderTransform(TOptional<FSlateRenderTransform>());
+		}
 	}
 }
 
