Add some new Slate counters.
  * Num Painted Widgets
  * Num Ticked Widgets

diff --git a/Engine/Source/Runtime/SlateCore/Private/Widgets/SWidget.cpp b/Engine/Source/Runtime/SlateCore/Private/Widgets/SWidget.cpp
--- a/Engine/Source/Runtime/SlateCore/Private/Widgets/SWidget.cpp
+++ b/Engine/Source/Runtime/SlateCore/Private/Widgets/SWidget.cpp
@@ -4,7 +4,8 @@
 
 DECLARE_CYCLE_STAT(TEXT("OnPaint"), STAT_SlateOnPaint, STATGROUP_Slate);
 DECLARE_CYCLE_STAT(TEXT("ArrangeChildren"), STAT_SlateArrangeChildren, STATGROUP_Slate);
-
+DECLARE_DWORD_COUNTER_STAT(TEXT("Num Painted Widgets"), STAT_SlateNumPaintedWidgets, STATGROUP_Slate);
+DECLARE_DWORD_COUNTER_STAT(TEXT("Num Ticked Widgets"), STAT_SlateNumTickedWidgets, STATGROUP_Slate);
 
 SWidget::SWidget()
 	: CreatedInFile( TEXT("") )
@@ -248,6 +249,8 @@ void SWidget::Tick( const FGeometry& AllottedGeometry, const double InCurrentTim
 
 void SWidget::TickWidgetsRecursively( const FGeometry& AllottedGeometry, const double InCurrentTime, const float InDeltaTime )
 {
+	INC_DWORD_STAT(STAT_SlateNumTickedWidgets);
+
 	// Tick this widget
 	Tick( AllottedGeometry, InCurrentTime, InDeltaTime );
 
@@ -511,6 +514,7 @@ void SWidget::SetDebugInfo( const ANSICHAR* InType, const ANSICHAR* InFile, int3
 int32 SWidget::Paint(const FPaintArgs& Args, const FGeometry& AllottedGeometry, const FSlateRect& MyClippingRect, FSlateWindowElementList& OutDrawElements, int32 LayerId, const FWidgetStyle& InWidgetStyle, bool bParentEnabled) const
 {
 	SCOPE_CYCLE_COUNTER(STAT_SlateOnPaint);
+	INC_DWORD_STAT(STAT_SlateNumPaintedWidgets);
 
 	FPaintArgs UpdatedArgs = Args.RecordHittestGeometry( this, AllottedGeometry, MyClippingRect );
 	return OnPaint(UpdatedArgs, AllottedGeometry, MyClippingRect, OutDrawElements, LayerId, InWidgetStyle, bParentEnabled);
