Remove stale comment.

diff --git a/Engine/Source/Runtime/SlateCore/Private/Rendering/ElementBatcher.cpp b/Engine/Source/Runtime/SlateCore/Private/Rendering/ElementBatcher.cpp
--- a/Engine/Source/Runtime/SlateCore/Private/Rendering/ElementBatcher.cpp
+++ b/Engine/Source/Runtime/SlateCore/Private/Rendering/ElementBatcher.cpp
@@ -72,7 +72,6 @@ void FSlateElementBatcher::AddElements( const TArray<FSlateDrawElement>& DrawEle
 		
 		if ( !bIsFullyClipped )
 		{
-			// !!! WRH 2014/09/17 - Profile and see how much this costs.
 			// do this check in here do we can short circuit above more quickly, but still name this variable so its meaning is clear.
 			const bool bIsScissored = DrawElement.GetScissorRect().IsSet() && !DrawElement.GetScissorRect().GetValue().DoesIntersect(FShortRect(InClippingRect));
 			// scissor rects are sort of a low level hack, so no one konws to clip against them. Instead we do it here to make sure the element is not actually rendered.
