Made FGeometry::Size const like all the other members. This was an oversight when I made the rest const.

diff --git a/Engine/Source/Runtime/SlateCore/Public/Layout/Geometry.h b/Engine/Source/Runtime/SlateCore/Public/Layout/Geometry.h
--- a/Engine/Source/Runtime/SlateCore/Public/Layout/Geometry.h
+++ b/Engine/Source/Runtime/SlateCore/Public/Layout/Geometry.h
@@ -335,7 +335,7 @@ public:
 	 * 
 	 * Size of the geometry in local space. 
 	 */
-	FVector2D /*Local*/Size;
+	const FVector2D /*Local*/Size;
 
 	/** 
 	 * !!! DEPRECATED !!! These legacy public members should ideally not be referenced, as they do not account for the render transform.
