Added a cvar Slate.TickInvisibleWidgets that controls whether invisible widgets are ticked. Defaults to true to match existing functionality.
* In the editor in a standard view, not ticking invisibile widgets has about a 4-5ms impact on total slate tick time, or about a 14% improvement.

diff --git a/Engine/Source/Runtime/SlateCore/Private/Widgets/SWidget.cpp b/Engine/Source/Runtime/SlateCore/Private/Widgets/SWidget.cpp
--- a/Engine/Source/Runtime/SlateCore/Private/Widgets/SWidget.cpp
+++ b/Engine/Source/Runtime/SlateCore/Private/Widgets/SWidget.cpp
@@ -247,6 +247,8 @@ void SWidget::Tick( const FGeometry& AllottedGeometry, const double InCurrentTim
 {
 }
 
+TAutoConsoleVariable<int32> TickInvisibleWidgets(TEXT("Slate.TickInvisibleWidgets"), 1, TEXT("Controls whether invisible widgets are ticked."));
+
 void SWidget::TickWidgetsRecursively( const FGeometry& AllottedGeometry, const double InCurrentTime, const float InDeltaTime )
 {
 	INC_DWORD_STAT(STAT_SlateNumTickedWidgets);
@@ -257,7 +259,7 @@ void SWidget::TickWidgetsRecursively( const FGeometry& AllottedGeometry, const d
 	// Gather all children, whether they're visible or not.  We need to allow invisible widgets to
 	// consider whether they should still be invisible in their tick functions, as well as maintain
 	// other state when hidden,
-	FArrangedChildren ArrangedChildren(EVisibility::All);
+	FArrangedChildren ArrangedChildren(TickInvisibleWidgets.GetValueOnGameThread() ? EVisibility::All : EVisibility::Visible);
 	ArrangeChildren(AllottedGeometry, ArrangedChildren);
 
 	// Recur!
