Adding original key code to keyboard events and exposed AreCapsLocked from input events

This will be needed for the web browser module so that keys like delete can processed as their character code comes through as 0.
It also needs to know if caps lock is down when you pass through keyboard and mouse events.

diff --git a/Engine/Source/Editor/Matinee/Private/MatineeViewportClient.cpp b/Engine/Source/Editor/Matinee/Private/MatineeViewportClient.cpp
--- a/Engine/Source/Editor/Matinee/Private/MatineeViewportClient.cpp
+++ b/Engine/Source/Editor/Matinee/Private/MatineeViewportClient.cpp
@@ -832,7 +832,7 @@ bool FMatineeViewportClient::InputKey(FViewport* Viewport, int32 ControllerId, F
 		}
 
 		FModifierKeysState ModKeys(bShiftDown, bShiftDown, bCtrlDown, bCtrlDown, bAltDown, bAltDown, bCmdDown, bCmdDown, bCapsDown);
-		FKeyboardEvent KeyEvent(Key, ModKeys, false, 0);
+		FKeyboardEvent KeyEvent(Key, ModKeys, false, 0, 0);
 		InterpEd->ProcessCommandBindings(KeyEvent);
 	}
 
diff --git a/Engine/Source/Editor/UnrealEd/Private/Tests/EditorBuildPromotionTests.cpp b/Engine/Source/Editor/UnrealEd/Private/Tests/EditorBuildPromotionTests.cpp
--- a/Engine/Source/Editor/UnrealEd/Private/Tests/EditorBuildPromotionTests.cpp
+++ b/Engine/Source/Editor/UnrealEd/Private/Tests/EditorBuildPromotionTests.cpp
@@ -359,7 +359,7 @@ namespace EditorBuildPromotionTestUtils
 
 			//Send the command
 			FModifierKeysState ModifierKeys(InGesture.bShift, false, InGesture.bCtrl, false, InGesture.bAlt, false, InGesture.bCmd, false, false);
-			FKeyboardEvent KeyEvent(InGesture.Key, ModifierKeys, false, 0);
+			FKeyboardEvent KeyEvent(InGesture.Key, ModifierKeys, false, 0, 0);
 			FSlateApplication::Get().ProcessKeyDownEvent(KeyEvent);
 			FSlateApplication::Get().ProcessKeyUpEvent(KeyEvent);
 		}
@@ -4370,7 +4370,7 @@ bool FBuildPromotionSettingsTest::RunTest(const FString& Parameters)
 	FSlateApplication::Get().ProcessWindowActivatedEvent(FWindowActivateEvent(FWindowActivateEvent::EA_Activate, AllWindows[0]));
 
 	//Send the PIE event
-	FKeyboardEvent PIEKeyEvent(EKeys::L, FModifierKeysState(false, false, true, false, false, false, false, false, false), false, 0x4C);
+	FKeyboardEvent PIEKeyEvent(EKeys::L, FModifierKeysState(false, false, true, false, false, false, false, false, false), false, 0x4C, 0x4C);
 	FSlateApplication::Get().ProcessKeyDownEvent(PIEKeyEvent);
 	FSlateApplication::Get().ProcessKeyUpEvent(PIEKeyEvent);
 
diff --git a/Engine/Source/Runtime/Slate/Private/Framework/Application/SlateApplication.cpp b/Engine/Source/Runtime/Slate/Private/Framework/Application/SlateApplication.cpp
--- a/Engine/Source/Runtime/Slate/Private/Framework/Application/SlateApplication.cpp
+++ b/Engine/Source/Runtime/Slate/Private/Framework/Application/SlateApplication.cpp
@@ -3030,7 +3030,7 @@ bool FSlateApplication::ProcessKeyCharEvent( FCharacterEvent& InCharacterEvent )
 bool FSlateApplication::OnKeyDown( const int32 KeyCode, const uint32 CharacterCode, const bool IsRepeat ) 
 {
 	FKey const Key = FInputKeyManager::Get().GetKeyFromCodes( KeyCode, CharacterCode );
-	FKeyboardEvent KeyboardEvent( Key, PlatformApplication->GetModifierKeys(), IsRepeat, CharacterCode );
+	FKeyboardEvent KeyboardEvent( Key, PlatformApplication->GetModifierKeys(), IsRepeat, CharacterCode, KeyCode );
 
 	return ProcessKeyDownEvent( KeyboardEvent );
 }
@@ -3139,7 +3139,7 @@ bool FSlateApplication::ProcessKeyDownEvent( FKeyboardEvent& InKeyboardEvent )
 bool FSlateApplication::OnKeyUp( const int32 KeyCode, const uint32 CharacterCode, const bool IsRepeat )
 {
 	FKey const Key = FInputKeyManager::Get().GetKeyFromCodes( KeyCode, CharacterCode );
-	FKeyboardEvent KeyboardEvent( Key, PlatformApplication->GetModifierKeys(), IsRepeat, CharacterCode );
+	FKeyboardEvent KeyboardEvent( Key, PlatformApplication->GetModifierKeys(), IsRepeat, CharacterCode, KeyCode );
 
 	return ProcessKeyUpEvent( KeyboardEvent );
 }
diff --git a/Engine/Source/Runtime/SlateCore/Public/Input/Events.h b/Engine/Source/Runtime/SlateCore/Public/Input/Events.h
--- a/Engine/Source/Runtime/SlateCore/Public/Input/Events.h
+++ b/Engine/Source/Runtime/SlateCore/Public/Input/Events.h
@@ -250,6 +250,16 @@ public:
 		return ModifierKeys.IsRightCommandDown();
 	}
 
+	/**
+	 * Returns true if caps lock was on when this event occurred
+	 * 
+	 * @return True if caps lock is on
+	 */
+	bool AreCapsLocked() const
+	{
+		return ModifierKeys.AreCapsLocked();
+	}
+
 	/** The event path provides additional context for handling */
 	FGeometry FindGeometry( const TSharedRef<SWidget>& WidgetToFind ) const
 	{
@@ -298,6 +308,7 @@ public:
 		: FInputEvent(FModifierKeysState(), false)
 		, Key(EKeys::SpaceBar)
 		, CharacterCode(0)
+		, KeyCode(0)
 	{
 	}
 
@@ -311,11 +322,13 @@ public:
 	FKeyboardEvent( const FKey InKey,
 					const FModifierKeysState& InModifierKeys, 
 					const bool bInIsRepeat,
-					const uint32 InCharacterCode
+					const uint32 InCharacterCode,
+					const uint32 InKeyCode
 	)
 		: FInputEvent(InModifierKeys, bInIsRepeat)
 		, Key(InKey)
 		, CharacterCode(InCharacterCode)
+		, KeyCode(InKeyCode)
 	{ }
 	
 	/**
@@ -338,6 +351,11 @@ public:
 		return CharacterCode;
 	}
 
+	uint32 GetKeyCode() const
+	{
+		return KeyCode;
+	}
+
 private:
 
 	// Name of the key that was pressed.
@@ -345,6 +363,9 @@ private:
 
 	// The character code of the key that was pressed.  Only applicable to typed character keys, 0 otherwise.
 	uint32 CharacterCode;
+
+	// Original key code received from hardware before any conversion/mapping
+	uint32 KeyCode;
 };
 
 
