Fixed Clang compile errors

diff --git a/Engine/Source/Runtime/SlateCore/Public/Widgets/DeclarativeSyntaxSupport.h b/Engine/Source/Runtime/SlateCore/Public/Widgets/DeclarativeSyntaxSupport.h
--- a/Engine/Source/Runtime/SlateCore/Public/Widgets/DeclarativeSyntaxSupport.h
+++ b/Engine/Source/Runtime/SlateCore/Public/Widgets/DeclarativeSyntaxSupport.h
@@ -112,8 +112,8 @@
 			return this->Me(); \
 		} \
 	\
-		/* Bind attribute with delegate to a global function */ \
-		/* NOTE: We use a template here to avoid 'typename' issues when hosting attributes inside templated classes */ \
+		/* Bind attribute with delegate to a global function
+		 * NOTE: We use a template here to avoid 'typename' issues when hosting attributes inside templated classes */ \
 		template< typename StaticFuncPtr > \
 		WidgetArgsType& AttrName##_Static( StaticFuncPtr InFunc )	\
 		{ \
@@ -290,8 +290,8 @@
  * to be defined in the scope where it is used.
  */
 #define INTERNAL_SLATE_DECL_ATTRIBUTE_PASSTHROUGH_COMMON( TargetAttrType, SourceAttrType, Var, AttrName ) \
-		/* Bind attribute with delegate to a global function */ \
-		/* NOTE: We use a template here to avoid 'typename' issues when hosting attributes inside templated classes */ \
+		/* Bind attribute with delegate to a global function \
+		 * NOTE: We use a template here to avoid 'typename' issues when hosting attributes inside templated classes */ \
 		WidgetArgsType& AttrName##_Static( SourceAttrType::FGetter::FStaticDelegate::FFuncPtr InFunc )	\
 		{ \
 			Var = TargetAttrType::Create( TargetAttrType::FGetter::CreateStatic( &AttrName##_Local::PassThroughDelegate, SourceAttrType::FGetter::CreateStatic( InFunc ) ) ); \
@@ -322,11 +322,11 @@
 			return this->Me(); \
 		} \
 	\
-		/* We don't support binding FString lambdas to FText attributes because we currently have no way to disambiguate them */ \
-		/*template<typename FunctorType> \
+		/* We don't support binding FString lambdas to FText attributes because we currently have no way to disambiguate them \
+		 * template<typename FunctorType> \
 		WidgetArgsType& AttrName##_Lambda(FunctorType&& InFunctor) \
 	\
-		/* Bind attribute with delegate to a raw C++ class method */ \
+		* Bind attribute with delegate to a raw C++ class method */ \
 		template< class UserClass >	\
 		WidgetArgsType& AttrName##_Raw( UserClass* InUserObject, typename SourceAttrType::FGetter::template TRawMethodDelegate_Const< UserClass >::FMethodPtr InFunc )	\
 		{ \
