Adding support for lambdas to slate events, on request of Nick Atamas

diff --git a/Engine/Source/Runtime/SlateCore/Public/Widgets/DeclarativeSyntaxSupport.h b/Engine/Source/Runtime/SlateCore/Public/Widgets/DeclarativeSyntaxSupport.h
--- a/Engine/Source/Runtime/SlateCore/Public/Widgets/DeclarativeSyntaxSupport.h
+++ b/Engine/Source/Runtime/SlateCore/Public/Widgets/DeclarativeSyntaxSupport.h
@@ -150,7 +150,7 @@
 		template<typename FunctorType> \
 		WidgetArgsType& AttrName##_Lambda(FunctorType&& InFunctor) \
 		{ \
-			Var = AttrType::Create(AttrType::FGetter::CreateLambda(InFunctor)); \
+			Var = AttrType::Create(AttrType::FGetter::CreateLambda(Forward<FunctorType>(InFunctor))); \
 			return this->Me(); \
 		} \
 	\
@@ -322,11 +322,11 @@
 			return this->Me(); \
 		} \
 	\
-		/* We don't support binding FString lambdas to FText attributes because we currently have no way to disambiguate them \
-		 * template<typename FunctorType> \
-		WidgetArgsType& AttrName##_Lambda(FunctorType&& InFunctor) \
+		/* We don't support binding FString lambdas to FText attributes because we currently have no way to disambiguate them from FText lambdas \
+		template<typename FunctorType> \
+		WidgetArgsType& AttrName##_Lambda(FunctorType&& InFunctor) */ \
 	\
-		* Bind attribute with delegate to a raw C++ class method */ \
+		/* Bind attribute with delegate to a raw C++ class method */ \
 		template< class UserClass >	\
 		WidgetArgsType& AttrName##_Raw( UserClass* InUserObject, typename SourceAttrType::FGetter::template TRawMethodDelegate_Const< UserClass >::FMethodPtr InFunc )	\
 		{ \
@@ -759,6 +759,15 @@ struct NamedSlotProperty
 			return *this; \
 		} \
 		\
+		/* Set event delegate to a lambda
+		 * technically this works for any functor types, but lambdas are the primary use case */ \
+		template<typename FunctorType> \
+		WidgetArgsType& EventName##_Lambda(FunctorType&& InFunctor) \
+		{ \
+			_##EventName = DelegateName::CreateLambda(Forward<FunctorType>(InFunctor)); \
+			return this->Me(); \
+		} \
+		\
 		/* Set event delegate to a raw C++ class method */ \
 		template< class UserClass >	\
 		WidgetArgsType& EventName##_Raw( UserClass* InUserObject, typename DelegateName::template TRawMethodDelegate< UserClass >::FMethodPtr InFunc )	\
