Rename FPostConstructInitializeProperties to something simpler

diff --git a/Engine/Source/Runtime/Slate/Private/Framework/Styling/ComboButtonWidgetStyle.cpp b/Engine/Source/Runtime/Slate/Private/Framework/Styling/ComboButtonWidgetStyle.cpp
--- a/Engine/Source/Runtime/Slate/Private/Framework/Styling/ComboButtonWidgetStyle.cpp
+++ b/Engine/Source/Runtime/Slate/Private/Framework/Styling/ComboButtonWidgetStyle.cpp
@@ -3,8 +3,8 @@
 #include "SlatePrivatePCH.h"
 
 
-UComboButtonWidgetStyle::UComboButtonWidgetStyle( const class FPostConstructInitializeProperties& PCIP )
-	: Super(PCIP)
+UComboButtonWidgetStyle::UComboButtonWidgetStyle( const FObjectInitializer& ObjectInitializer )
+	: Super(ObjectInitializer)
 {
 	
 }
diff --git a/Engine/Source/Runtime/Slate/Private/Framework/Styling/EditableTextBoxWidgetStyle.cpp b/Engine/Source/Runtime/Slate/Private/Framework/Styling/EditableTextBoxWidgetStyle.cpp
--- a/Engine/Source/Runtime/Slate/Private/Framework/Styling/EditableTextBoxWidgetStyle.cpp
+++ b/Engine/Source/Runtime/Slate/Private/Framework/Styling/EditableTextBoxWidgetStyle.cpp
@@ -3,8 +3,8 @@
 #include "SlatePrivatePCH.h"
 
 
-UEditableTextBoxWidgetStyle::UEditableTextBoxWidgetStyle( const class FPostConstructInitializeProperties& PCIP )
-	: Super(PCIP)
+UEditableTextBoxWidgetStyle::UEditableTextBoxWidgetStyle( const FObjectInitializer& ObjectInitializer )
+	: Super(ObjectInitializer)
 {
 	
 }
\ No newline at end of file
diff --git a/Engine/Source/Runtime/Slate/Private/Framework/Styling/EditableTextWidgetStyle.cpp b/Engine/Source/Runtime/Slate/Private/Framework/Styling/EditableTextWidgetStyle.cpp
--- a/Engine/Source/Runtime/Slate/Private/Framework/Styling/EditableTextWidgetStyle.cpp
+++ b/Engine/Source/Runtime/Slate/Private/Framework/Styling/EditableTextWidgetStyle.cpp
@@ -3,8 +3,8 @@
 #include "SlatePrivatePCH.h"
 
 
-UEditableTextWidgetStyle::UEditableTextWidgetStyle( const class FPostConstructInitializeProperties& PCIP )
-	: Super(PCIP)
+UEditableTextWidgetStyle::UEditableTextWidgetStyle( const FObjectInitializer& ObjectInitializer )
+	: Super(ObjectInitializer)
 {
 	
 }
\ No newline at end of file
diff --git a/Engine/Source/Runtime/Slate/Private/Framework/Styling/ProgressWidgetStyle.cpp b/Engine/Source/Runtime/Slate/Private/Framework/Styling/ProgressWidgetStyle.cpp
--- a/Engine/Source/Runtime/Slate/Private/Framework/Styling/ProgressWidgetStyle.cpp
+++ b/Engine/Source/Runtime/Slate/Private/Framework/Styling/ProgressWidgetStyle.cpp
@@ -3,8 +3,8 @@
 #include "SlatePrivatePCH.h"
 
 
-UProgressWidgetStyle::UProgressWidgetStyle( const class FPostConstructInitializeProperties& PCIP )
-	: Super(PCIP)
+UProgressWidgetStyle::UProgressWidgetStyle( const FObjectInitializer& ObjectInitializer )
+	: Super(ObjectInitializer)
 {
 	
 }
\ No newline at end of file
diff --git a/Engine/Source/Runtime/Slate/Private/Framework/Styling/ScrollBarWidgetStyle.cpp b/Engine/Source/Runtime/Slate/Private/Framework/Styling/ScrollBarWidgetStyle.cpp
--- a/Engine/Source/Runtime/Slate/Private/Framework/Styling/ScrollBarWidgetStyle.cpp
+++ b/Engine/Source/Runtime/Slate/Private/Framework/Styling/ScrollBarWidgetStyle.cpp
@@ -3,8 +3,8 @@
 #include "SlatePrivatePCH.h"
 
 
-UScrollBarWidgetStyle::UScrollBarWidgetStyle( const class FPostConstructInitializeProperties& PCIP )
-	: Super(PCIP)
+UScrollBarWidgetStyle::UScrollBarWidgetStyle( const FObjectInitializer& ObjectInitializer )
+	: Super(ObjectInitializer)
 {
 	
 }
\ No newline at end of file
diff --git a/Engine/Source/Runtime/Slate/Private/Framework/Styling/ScrollBoxWidgetStyle.cpp b/Engine/Source/Runtime/Slate/Private/Framework/Styling/ScrollBoxWidgetStyle.cpp
--- a/Engine/Source/Runtime/Slate/Private/Framework/Styling/ScrollBoxWidgetStyle.cpp
+++ b/Engine/Source/Runtime/Slate/Private/Framework/Styling/ScrollBoxWidgetStyle.cpp
@@ -3,8 +3,8 @@
 #include "SlatePrivatePCH.h"
 
 
-UScrollBoxWidgetStyle::UScrollBoxWidgetStyle( const class FPostConstructInitializeProperties& PCIP )
-	: Super(PCIP)
+UScrollBoxWidgetStyle::UScrollBoxWidgetStyle( const FObjectInitializer& ObjectInitializer )
+	: Super(ObjectInitializer)
 {
 	
 }
\ No newline at end of file
diff --git a/Engine/Source/Runtime/Slate/Private/Framework/Styling/SpinBoxWidgetStyle.cpp b/Engine/Source/Runtime/Slate/Private/Framework/Styling/SpinBoxWidgetStyle.cpp
--- a/Engine/Source/Runtime/Slate/Private/Framework/Styling/SpinBoxWidgetStyle.cpp
+++ b/Engine/Source/Runtime/Slate/Private/Framework/Styling/SpinBoxWidgetStyle.cpp
@@ -3,8 +3,8 @@
 #include "SlatePrivatePCH.h"
 
 
-USpinBoxWidgetStyle::USpinBoxWidgetStyle(const class FPostConstructInitializeProperties& PCIP)
-	: Super(PCIP)
+USpinBoxWidgetStyle::USpinBoxWidgetStyle(const FObjectInitializer& ObjectInitializer)
+	: Super(ObjectInitializer)
 {
 
 }
\ No newline at end of file
diff --git a/Engine/Source/Runtime/Slate/Private/Framework/Styling/TextBlockWidgetStyle.cpp b/Engine/Source/Runtime/Slate/Private/Framework/Styling/TextBlockWidgetStyle.cpp
--- a/Engine/Source/Runtime/Slate/Private/Framework/Styling/TextBlockWidgetStyle.cpp
+++ b/Engine/Source/Runtime/Slate/Private/Framework/Styling/TextBlockWidgetStyle.cpp
@@ -3,8 +3,8 @@
 #include "SlatePrivatePCH.h"
 
 
-UTextBlockWidgetStyle::UTextBlockWidgetStyle( const class FPostConstructInitializeProperties& PCIP )
-	: Super(PCIP)
+UTextBlockWidgetStyle::UTextBlockWidgetStyle( const FObjectInitializer& ObjectInitializer )
+	: Super(ObjectInitializer)
 {
 	
 }
\ No newline at end of file
diff --git a/Engine/Source/Runtime/SlateCore/Private/Styling/SlateTypes.cpp b/Engine/Source/Runtime/SlateCore/Private/Styling/SlateTypes.cpp
--- a/Engine/Source/Runtime/SlateCore/Private/Styling/SlateTypes.cpp
+++ b/Engine/Source/Runtime/SlateCore/Private/Styling/SlateTypes.cpp
@@ -636,7 +636,7 @@ const FWindowStyle& FWindowStyle::GetDefault()
 }
 
 
-USlateTypes::USlateTypes( const class FPostConstructInitializeProperties& PCIP )
-	: Super(PCIP)
+USlateTypes::USlateTypes( const FObjectInitializer& ObjectInitializer )
+	: Super(ObjectInitializer)
 {
 }
diff --git a/Engine/Source/Runtime/SlateCore/Private/Styling/SlateWidgetStyleAsset.cpp b/Engine/Source/Runtime/SlateCore/Private/Styling/SlateWidgetStyleAsset.cpp
--- a/Engine/Source/Runtime/SlateCore/Private/Styling/SlateWidgetStyleAsset.cpp
+++ b/Engine/Source/Runtime/SlateCore/Private/Styling/SlateWidgetStyleAsset.cpp
@@ -3,8 +3,8 @@
 #include "SlateCorePrivatePCH.h"
 
 
-USlateWidgetStyleAsset::USlateWidgetStyleAsset( const class FPostConstructInitializeProperties& PCIP )
-	: Super(PCIP)
+USlateWidgetStyleAsset::USlateWidgetStyleAsset( const FObjectInitializer& ObjectInitializer )
+	: Super(ObjectInitializer)
 {
 	
 }
diff --git a/Engine/Source/Runtime/SlateCore/Private/Styling/SlateWidgetStyleContainerBase.cpp b/Engine/Source/Runtime/SlateCore/Private/Styling/SlateWidgetStyleContainerBase.cpp
--- a/Engine/Source/Runtime/SlateCore/Private/Styling/SlateWidgetStyleContainerBase.cpp
+++ b/Engine/Source/Runtime/SlateCore/Private/Styling/SlateWidgetStyleContainerBase.cpp
@@ -7,8 +7,8 @@
 DEFINE_LOG_CATEGORY(LogSlateStyle);
 
 
-USlateWidgetStyleContainerBase::USlateWidgetStyleContainerBase( const class FPostConstructInitializeProperties& PCIP )
-	: Super(PCIP)
+USlateWidgetStyleContainerBase::USlateWidgetStyleContainerBase( const FObjectInitializer& ObjectInitializer )
+	: Super(ObjectInitializer)
 { }
 
 
diff --git a/Engine/Source/Runtime/SlateCore/Private/Styling/SlateWidgetStyleContainerInterface.cpp b/Engine/Source/Runtime/SlateCore/Private/Styling/SlateWidgetStyleContainerInterface.cpp
--- a/Engine/Source/Runtime/SlateCore/Private/Styling/SlateWidgetStyleContainerInterface.cpp
+++ b/Engine/Source/Runtime/SlateCore/Private/Styling/SlateWidgetStyleContainerInterface.cpp
@@ -3,8 +3,8 @@
 #include "SlateCorePrivatePCH.h"
 #include "SlateWidgetStyleContainerInterface.h"
 
-USlateWidgetStyleContainerInterface::USlateWidgetStyleContainerInterface( const class FPostConstructInitializeProperties& PCIP )
-	: Super(PCIP)
+USlateWidgetStyleContainerInterface::USlateWidgetStyleContainerInterface( const FObjectInitializer& ObjectInitializer )
+	: Super(ObjectInitializer)
 {
 	
 }
diff --git a/Engine/Source/Runtime/UMG/Private/Animation/MovieScene2DTransformSection.cpp b/Engine/Source/Runtime/UMG/Private/Animation/MovieScene2DTransformSection.cpp
--- a/Engine/Source/Runtime/UMG/Private/Animation/MovieScene2DTransformSection.cpp
+++ b/Engine/Source/Runtime/UMG/Private/Animation/MovieScene2DTransformSection.cpp
@@ -4,8 +4,8 @@
 #include "MovieScene2DTransformSection.h"
 #include "MovieScene2DTransformTrack.h"
 
-UMovieScene2DTransformSection::UMovieScene2DTransformSection( const FPostConstructInitializeProperties& PCIP )
-	: Super( PCIP )
+UMovieScene2DTransformSection::UMovieScene2DTransformSection( const FObjectInitializer& ObjectInitializer )
+	: Super( ObjectInitializer )
 {
 
 }
diff --git a/Engine/Source/Runtime/UMG/Private/Animation/MovieScene2DTransformTrack.cpp b/Engine/Source/Runtime/UMG/Private/Animation/MovieScene2DTransformTrack.cpp
--- a/Engine/Source/Runtime/UMG/Private/Animation/MovieScene2DTransformTrack.cpp
+++ b/Engine/Source/Runtime/UMG/Private/Animation/MovieScene2DTransformTrack.cpp
@@ -7,8 +7,8 @@
 #include "MovieScene2DTransformTrackInstance.h"
 #include "MovieSceneCommonHelpers.h"
 
-UMovieScene2DTransformTrack::UMovieScene2DTransformTrack( const FPostConstructInitializeProperties& PCIP )
-	: Super( PCIP )
+UMovieScene2DTransformTrack::UMovieScene2DTransformTrack( const FObjectInitializer& ObjectInitializer )
+	: Super( ObjectInitializer )
 {
 }
 
diff --git a/Engine/Source/Runtime/UMG/Private/Animation/MovieSceneMarginSection.cpp b/Engine/Source/Runtime/UMG/Private/Animation/MovieSceneMarginSection.cpp
--- a/Engine/Source/Runtime/UMG/Private/Animation/MovieSceneMarginSection.cpp
+++ b/Engine/Source/Runtime/UMG/Private/Animation/MovieSceneMarginSection.cpp
@@ -4,8 +4,8 @@
 #include "MovieSceneMarginSection.h"
 #include "MovieSceneMarginTrack.h"
 
-UMovieSceneMarginSection::UMovieSceneMarginSection( const FPostConstructInitializeProperties& PCIP )
-	: Super( PCIP )
+UMovieSceneMarginSection::UMovieSceneMarginSection( const FObjectInitializer& ObjectInitializer )
+	: Super( ObjectInitializer )
 {
 }
 
diff --git a/Engine/Source/Runtime/UMG/Private/Animation/MovieSceneMarginTrack.cpp b/Engine/Source/Runtime/UMG/Private/Animation/MovieSceneMarginTrack.cpp
--- a/Engine/Source/Runtime/UMG/Private/Animation/MovieSceneMarginTrack.cpp
+++ b/Engine/Source/Runtime/UMG/Private/Animation/MovieSceneMarginTrack.cpp
@@ -7,8 +7,8 @@
 #include "MovieSceneMarginTrackInstance.h"
 #include "MovieSceneCommonHelpers.h"
 
-UMovieSceneMarginTrack::UMovieSceneMarginTrack(const FPostConstructInitializeProperties& PCIP)
-: Super(PCIP)
+UMovieSceneMarginTrack::UMovieSceneMarginTrack(const FObjectInitializer& ObjectInitializer)
+: Super(ObjectInitializer)
 {
 }
 
diff --git a/Engine/Source/Runtime/UMG/Private/Animation/UMGSequencePlayer.cpp b/Engine/Source/Runtime/UMG/Private/Animation/UMGSequencePlayer.cpp
--- a/Engine/Source/Runtime/UMG/Private/Animation/UMGSequencePlayer.cpp
+++ b/Engine/Source/Runtime/UMG/Private/Animation/UMGSequencePlayer.cpp
@@ -8,8 +8,8 @@
 #include "MovieScene.h"
 #include "WidgetAnimation.h"
 
-UUMGSequencePlayer::UUMGSequencePlayer(const FPostConstructInitializeProperties& PCIP)
-	: Super(PCIP)
+UUMGSequencePlayer::UUMGSequencePlayer(const FObjectInitializer& ObjectInitializer)
+	: Super(ObjectInitializer)
 {
 	PlayerStatus = EMovieScenePlayerStatus::Stopped;
 	TimeCursorPosition = 0.0f;
diff --git a/Engine/Source/Runtime/UMG/Private/Animation/WidgetAnimation.cpp b/Engine/Source/Runtime/UMG/Private/Animation/WidgetAnimation.cpp
--- a/Engine/Source/Runtime/UMG/Private/Animation/WidgetAnimation.cpp
+++ b/Engine/Source/Runtime/UMG/Private/Animation/WidgetAnimation.cpp
@@ -23,8 +23,8 @@ UObject* FWidgetAnimationBinding::FindRuntimeObject( UWidgetTree& WidgetTree ) c
 	return FoundObject;
 }
 
-UWidgetAnimation::UWidgetAnimation(const FPostConstructInitializeProperties& PCIP)
-	: Super(PCIP)
+UWidgetAnimation::UWidgetAnimation(const FObjectInitializer& ObjectInitializer)
+	: Super(ObjectInitializer)
 {
 
 }
diff --git a/Engine/Source/Runtime/UMG/Private/Components/Border.cpp b/Engine/Source/Runtime/UMG/Private/Components/Border.cpp
--- a/Engine/Source/Runtime/UMG/Private/Components/Border.cpp
+++ b/Engine/Source/Runtime/UMG/Private/Components/Border.cpp
@@ -9,8 +9,8 @@
 /////////////////////////////////////////////////////
 // UBorder
 
-UBorder::UBorder(const FPostConstructInitializeProperties& PCIP)
-	: Super(PCIP)
+UBorder::UBorder(const FObjectInitializer& ObjectInitializer)
+	: Super(ObjectInitializer)
 {
 	bIsVariable = false;
 
diff --git a/Engine/Source/Runtime/UMG/Private/Components/BorderSlot.cpp b/Engine/Source/Runtime/UMG/Private/Components/BorderSlot.cpp
--- a/Engine/Source/Runtime/UMG/Private/Components/BorderSlot.cpp
+++ b/Engine/Source/Runtime/UMG/Private/Components/BorderSlot.cpp
@@ -6,8 +6,8 @@
 /////////////////////////////////////////////////////
 // UBorderSlot
 
-UBorderSlot::UBorderSlot(const FPostConstructInitializeProperties& PCIP)
-	: Super(PCIP)
+UBorderSlot::UBorderSlot(const FObjectInitializer& ObjectInitializer)
+	: Super(ObjectInitializer)
 {
 	Padding = FMargin(4, 2);
 
diff --git a/Engine/Source/Runtime/UMG/Private/Components/Button.cpp b/Engine/Source/Runtime/UMG/Private/Components/Button.cpp
--- a/Engine/Source/Runtime/UMG/Private/Components/Button.cpp
+++ b/Engine/Source/Runtime/UMG/Private/Components/Button.cpp
@@ -7,8 +7,8 @@
 /////////////////////////////////////////////////////
 // UButton
 
-UButton::UButton(const FPostConstructInitializeProperties& PCIP)
-	: Super(PCIP)
+UButton::UButton(const FObjectInitializer& ObjectInitializer)
+	: Super(ObjectInitializer)
 {
 	SButton::FArguments ButtonDefaults;
 	WidgetStyle = *ButtonDefaults._ButtonStyle;
diff --git a/Engine/Source/Runtime/UMG/Private/Components/ButtonSlot.cpp b/Engine/Source/Runtime/UMG/Private/Components/ButtonSlot.cpp
--- a/Engine/Source/Runtime/UMG/Private/Components/ButtonSlot.cpp
+++ b/Engine/Source/Runtime/UMG/Private/Components/ButtonSlot.cpp
@@ -5,8 +5,8 @@
 /////////////////////////////////////////////////////
 // UButtonSlot
 
-UButtonSlot::UButtonSlot(const FPostConstructInitializeProperties& PCIP)
-	: Super(PCIP)
+UButtonSlot::UButtonSlot(const FObjectInitializer& ObjectInitializer)
+	: Super(ObjectInitializer)
 {
 	Padding = FMargin(4, 2);
 
diff --git a/Engine/Source/Runtime/UMG/Private/Components/CanvasPanel.cpp b/Engine/Source/Runtime/UMG/Private/Components/CanvasPanel.cpp
--- a/Engine/Source/Runtime/UMG/Private/Components/CanvasPanel.cpp
+++ b/Engine/Source/Runtime/UMG/Private/Components/CanvasPanel.cpp
@@ -7,8 +7,8 @@
 /////////////////////////////////////////////////////
 // UCanvasPanel
 
-UCanvasPanel::UCanvasPanel(const FPostConstructInitializeProperties& PCIP)
-	: Super(PCIP)
+UCanvasPanel::UCanvasPanel(const FObjectInitializer& ObjectInitializer)
+	: Super(ObjectInitializer)
 {
 	bIsVariable = false;
 
diff --git a/Engine/Source/Runtime/UMG/Private/Components/CanvasPanelSlot.cpp b/Engine/Source/Runtime/UMG/Private/Components/CanvasPanelSlot.cpp
--- a/Engine/Source/Runtime/UMG/Private/Components/CanvasPanelSlot.cpp
+++ b/Engine/Source/Runtime/UMG/Private/Components/CanvasPanelSlot.cpp
@@ -5,8 +5,8 @@
 /////////////////////////////////////////////////////
 // UCanvasPanelSlot
 
-UCanvasPanelSlot::UCanvasPanelSlot(const FPostConstructInitializeProperties& PCIP)
-	: Super(PCIP)
+UCanvasPanelSlot::UCanvasPanelSlot(const FObjectInitializer& ObjectInitializer)
+	: Super(ObjectInitializer)
 	, Slot(NULL)
 {
 	LayoutData.Offsets = FMargin(0, 0, 100, 30);
diff --git a/Engine/Source/Runtime/UMG/Private/Components/CheckBox.cpp b/Engine/Source/Runtime/UMG/Private/Components/CheckBox.cpp
--- a/Engine/Source/Runtime/UMG/Private/Components/CheckBox.cpp
+++ b/Engine/Source/Runtime/UMG/Private/Components/CheckBox.cpp
@@ -7,8 +7,8 @@
 /////////////////////////////////////////////////////
 // UCheckBox
 
-UCheckBox::UCheckBox(const FPostConstructInitializeProperties& PCIP)
-	: Super(PCIP)
+UCheckBox::UCheckBox(const FObjectInitializer& ObjectInitializer)
+	: Super(ObjectInitializer)
 {
 	SCheckBox::FArguments SlateDefaults;
 	WidgetStyle = *SlateDefaults._Style;
diff --git a/Engine/Source/Runtime/UMG/Private/Components/CircularThrobber.cpp b/Engine/Source/Runtime/UMG/Private/Components/CircularThrobber.cpp
--- a/Engine/Source/Runtime/UMG/Private/Components/CircularThrobber.cpp
+++ b/Engine/Source/Runtime/UMG/Private/Components/CircularThrobber.cpp
@@ -8,8 +8,8 @@
 /////////////////////////////////////////////////////
 // UCircularThrobber
 
-UCircularThrobber::UCircularThrobber(const FPostConstructInitializeProperties& PCIP)
-	: Super(PCIP)
+UCircularThrobber::UCircularThrobber(const FObjectInitializer& ObjectInitializer)
+	: Super(ObjectInitializer)
 {
 	SCircularThrobber::FArguments DefaultArgs;
 	Image = *DefaultArgs._PieceImage;
diff --git a/Engine/Source/Runtime/UMG/Private/Components/ComboBox.cpp b/Engine/Source/Runtime/UMG/Private/Components/ComboBox.cpp
--- a/Engine/Source/Runtime/UMG/Private/Components/ComboBox.cpp
+++ b/Engine/Source/Runtime/UMG/Private/Components/ComboBox.cpp
@@ -7,8 +7,8 @@
 /////////////////////////////////////////////////////
 // UComboBox
 
-UComboBox::UComboBox(const FPostConstructInitializeProperties& PCIP)
-	: Super(PCIP)
+UComboBox::UComboBox(const FObjectInitializer& ObjectInitializer)
+	: Super(ObjectInitializer)
 {
 }
 
diff --git a/Engine/Source/Runtime/UMG/Private/Components/ComboBoxString.cpp b/Engine/Source/Runtime/UMG/Private/Components/ComboBoxString.cpp
--- a/Engine/Source/Runtime/UMG/Private/Components/ComboBoxString.cpp
+++ b/Engine/Source/Runtime/UMG/Private/Components/ComboBoxString.cpp
@@ -7,8 +7,8 @@
 /////////////////////////////////////////////////////
 // UComboBoxString
 
-UComboBoxString::UComboBoxString(const FPostConstructInitializeProperties& PCIP)
-	: Super(PCIP)
+UComboBoxString::UComboBoxString(const FObjectInitializer& ObjectInitializer)
+	: Super(ObjectInitializer)
 {
 	SComboBox< TSharedPtr<FString> >::FArguments SlateDefaults;
 	WidgetStyle = *SlateDefaults._ComboBoxStyle;
diff --git a/Engine/Source/Runtime/UMG/Private/Components/ContentWidget.cpp b/Engine/Source/Runtime/UMG/Private/Components/ContentWidget.cpp
--- a/Engine/Source/Runtime/UMG/Private/Components/ContentWidget.cpp
+++ b/Engine/Source/Runtime/UMG/Private/Components/ContentWidget.cpp
@@ -12,8 +12,8 @@
 /////////////////////////////////////////////////////
 // UContentWidget
 
-UContentWidget::UContentWidget(const FPostConstructInitializeProperties& PCIP)
-	: Super(PCIP)
+UContentWidget::UContentWidget(const FObjectInitializer& ObjectInitializer)
+	: Super(ObjectInitializer)
 {
 	bCanHaveMultipleChildren = false;
 }
diff --git a/Engine/Source/Runtime/UMG/Private/Components/EditableText.cpp b/Engine/Source/Runtime/UMG/Private/Components/EditableText.cpp
--- a/Engine/Source/Runtime/UMG/Private/Components/EditableText.cpp
+++ b/Engine/Source/Runtime/UMG/Private/Components/EditableText.cpp
@@ -7,8 +7,8 @@
 /////////////////////////////////////////////////////
 // UEditableText
 
-UEditableText::UEditableText(const FPostConstructInitializeProperties& PCIP)
-	: Super(PCIP)
+UEditableText::UEditableText(const FObjectInitializer& ObjectInitializer)
+	: Super(ObjectInitializer)
 {
 	SEditableText::FArguments Defaults;
 	WidgetStyle = *Defaults._Style;
diff --git a/Engine/Source/Runtime/UMG/Private/Components/EditableTextBox.cpp b/Engine/Source/Runtime/UMG/Private/Components/EditableTextBox.cpp
--- a/Engine/Source/Runtime/UMG/Private/Components/EditableTextBox.cpp
+++ b/Engine/Source/Runtime/UMG/Private/Components/EditableTextBox.cpp
@@ -7,8 +7,8 @@
 /////////////////////////////////////////////////////
 // UEditableTextBox
 
-UEditableTextBox::UEditableTextBox(const FPostConstructInitializeProperties& PCIP)
-	: Super(PCIP)
+UEditableTextBox::UEditableTextBox(const FObjectInitializer& ObjectInitializer)
+	: Super(ObjectInitializer)
 {
 	ForegroundColor = FLinearColor::Black;
 	BackgroundColor = FLinearColor::White;
diff --git a/Engine/Source/Runtime/UMG/Private/Components/GridPanel.cpp b/Engine/Source/Runtime/UMG/Private/Components/GridPanel.cpp
--- a/Engine/Source/Runtime/UMG/Private/Components/GridPanel.cpp
+++ b/Engine/Source/Runtime/UMG/Private/Components/GridPanel.cpp
@@ -7,8 +7,8 @@
 /////////////////////////////////////////////////////
 // UGridPanel
 
-UGridPanel::UGridPanel(const FPostConstructInitializeProperties& PCIP)
-	: Super(PCIP)
+UGridPanel::UGridPanel(const FObjectInitializer& ObjectInitializer)
+	: Super(ObjectInitializer)
 {
 	bIsVariable = false;
 
diff --git a/Engine/Source/Runtime/UMG/Private/Components/GridSlot.cpp b/Engine/Source/Runtime/UMG/Private/Components/GridSlot.cpp
--- a/Engine/Source/Runtime/UMG/Private/Components/GridSlot.cpp
+++ b/Engine/Source/Runtime/UMG/Private/Components/GridSlot.cpp
@@ -5,8 +5,8 @@
 /////////////////////////////////////////////////////
 // UGridSlot
 
-UGridSlot::UGridSlot(const FPostConstructInitializeProperties& PCIP)
-	: Super(PCIP)
+UGridSlot::UGridSlot(const FObjectInitializer& ObjectInitializer)
+	: Super(ObjectInitializer)
 	, Slot(nullptr)
 {
 	HorizontalAlignment = HAlign_Fill;
diff --git a/Engine/Source/Runtime/UMG/Private/Components/HorizontalBox.cpp b/Engine/Source/Runtime/UMG/Private/Components/HorizontalBox.cpp
--- a/Engine/Source/Runtime/UMG/Private/Components/HorizontalBox.cpp
+++ b/Engine/Source/Runtime/UMG/Private/Components/HorizontalBox.cpp
@@ -7,8 +7,8 @@
 /////////////////////////////////////////////////////
 // UHorizontalBox
 
-UHorizontalBox::UHorizontalBox(const FPostConstructInitializeProperties& PCIP)
-	: Super(PCIP)
+UHorizontalBox::UHorizontalBox(const FObjectInitializer& ObjectInitializer)
+	: Super(ObjectInitializer)
 {
 	bIsVariable = false;
 
diff --git a/Engine/Source/Runtime/UMG/Private/Components/HorizontalBoxSlot.cpp b/Engine/Source/Runtime/UMG/Private/Components/HorizontalBoxSlot.cpp
--- a/Engine/Source/Runtime/UMG/Private/Components/HorizontalBoxSlot.cpp
+++ b/Engine/Source/Runtime/UMG/Private/Components/HorizontalBoxSlot.cpp
@@ -5,8 +5,8 @@
 /////////////////////////////////////////////////////
 // UHorizontalBoxSlot
 
-UHorizontalBoxSlot::UHorizontalBoxSlot(const FPostConstructInitializeProperties& PCIP)
-	: Super(PCIP)
+UHorizontalBoxSlot::UHorizontalBoxSlot(const FObjectInitializer& ObjectInitializer)
+	: Super(ObjectInitializer)
 {
 	HorizontalAlignment = HAlign_Fill;
 	VerticalAlignment = VAlign_Fill;
diff --git a/Engine/Source/Runtime/UMG/Private/Components/Image.cpp b/Engine/Source/Runtime/UMG/Private/Components/Image.cpp
--- a/Engine/Source/Runtime/UMG/Private/Components/Image.cpp
+++ b/Engine/Source/Runtime/UMG/Private/Components/Image.cpp
@@ -8,8 +8,8 @@
 /////////////////////////////////////////////////////
 // UImage
 
-UImage::UImage(const FPostConstructInitializeProperties& PCIP)
-	: Super(PCIP)
+UImage::UImage(const FObjectInitializer& ObjectInitializer)
+	: Super(ObjectInitializer)
 	, ColorAndOpacity(FLinearColor::White)
 {
 }
diff --git a/Engine/Source/Runtime/UMG/Private/Components/ListView.cpp b/Engine/Source/Runtime/UMG/Private/Components/ListView.cpp
--- a/Engine/Source/Runtime/UMG/Private/Components/ListView.cpp
+++ b/Engine/Source/Runtime/UMG/Private/Components/ListView.cpp
@@ -7,8 +7,8 @@
 /////////////////////////////////////////////////////
 // UListView
 
-UListView::UListView(const FPostConstructInitializeProperties& PCIP)
-	: Super(PCIP)
+UListView::UListView(const FObjectInitializer& ObjectInitializer)
+	: Super(ObjectInitializer)
 {
 	bIsVariable = true;
 
diff --git a/Engine/Source/Runtime/UMG/Private/Components/MenuAnchor.cpp b/Engine/Source/Runtime/UMG/Private/Components/MenuAnchor.cpp
--- a/Engine/Source/Runtime/UMG/Private/Components/MenuAnchor.cpp
+++ b/Engine/Source/Runtime/UMG/Private/Components/MenuAnchor.cpp
@@ -7,8 +7,8 @@
 /////////////////////////////////////////////////////
 // UMenuAnchor
 
-UMenuAnchor::UMenuAnchor(const FPostConstructInitializeProperties& PCIP)
-	: Super(PCIP)
+UMenuAnchor::UMenuAnchor(const FObjectInitializer& ObjectInitializer)
+	: Super(ObjectInitializer)
 {
 	Placement = MenuPlacement_ComboBox;
 }
diff --git a/Engine/Source/Runtime/UMG/Private/Components/MultiLineEditableText.cpp b/Engine/Source/Runtime/UMG/Private/Components/MultiLineEditableText.cpp
--- a/Engine/Source/Runtime/UMG/Private/Components/MultiLineEditableText.cpp
+++ b/Engine/Source/Runtime/UMG/Private/Components/MultiLineEditableText.cpp
@@ -7,8 +7,8 @@
 /////////////////////////////////////////////////////
 // UMultiLineEditableText
 
-UMultiLineEditableText::UMultiLineEditableText(const FPostConstructInitializeProperties& PCIP)
-	: Super(PCIP)
+UMultiLineEditableText::UMultiLineEditableText(const FObjectInitializer& ObjectInitializer)
+	: Super(ObjectInitializer)
 {
 	SMultiLineEditableText::FArguments Defaults;
 	WidgetStyle = *Defaults._TextStyle;
diff --git a/Engine/Source/Runtime/UMG/Private/Components/MultiLineEditableTextBox.cpp b/Engine/Source/Runtime/UMG/Private/Components/MultiLineEditableTextBox.cpp
--- a/Engine/Source/Runtime/UMG/Private/Components/MultiLineEditableTextBox.cpp
+++ b/Engine/Source/Runtime/UMG/Private/Components/MultiLineEditableTextBox.cpp
@@ -7,8 +7,8 @@
 /////////////////////////////////////////////////////
 // UMultiLineEditableTextBox
 
-UMultiLineEditableTextBox::UMultiLineEditableTextBox(const FPostConstructInitializeProperties& PCIP)
-	: Super(PCIP)
+UMultiLineEditableTextBox::UMultiLineEditableTextBox(const FObjectInitializer& ObjectInitializer)
+	: Super(ObjectInitializer)
 {
 	ForegroundColor = FLinearColor::Black;
 	BackgroundColor = FLinearColor::White;
diff --git a/Engine/Source/Runtime/UMG/Private/Components/NamedSlot.cpp b/Engine/Source/Runtime/UMG/Private/Components/NamedSlot.cpp
--- a/Engine/Source/Runtime/UMG/Private/Components/NamedSlot.cpp
+++ b/Engine/Source/Runtime/UMG/Private/Components/NamedSlot.cpp
@@ -7,8 +7,8 @@
 /////////////////////////////////////////////////////
 // UNamedSlot
 
-UNamedSlot::UNamedSlot(const FPostConstructInitializeProperties& PCIP)
-	: Super(PCIP)
+UNamedSlot::UNamedSlot(const FObjectInitializer& ObjectInitializer)
+	: Super(ObjectInitializer)
 {
 	bIsVariable = true;
 }
diff --git a/Engine/Source/Runtime/UMG/Private/Components/NamedSlotInterface.cpp b/Engine/Source/Runtime/UMG/Private/Components/NamedSlotInterface.cpp
--- a/Engine/Source/Runtime/UMG/Private/Components/NamedSlotInterface.cpp
+++ b/Engine/Source/Runtime/UMG/Private/Components/NamedSlotInterface.cpp
@@ -2,7 +2,7 @@
 
 #include "UMGPrivatePCH.h"
 
-UNamedSlotInterface::UNamedSlotInterface(const FPostConstructInitializeProperties& PCIP)
-	: Super(PCIP)
+UNamedSlotInterface::UNamedSlotInterface(const FObjectInitializer& ObjectInitializer)
+	: Super(ObjectInitializer)
 {
 }
diff --git a/Engine/Source/Runtime/UMG/Private/Components/NativeWidgetHost.cpp b/Engine/Source/Runtime/UMG/Private/Components/NativeWidgetHost.cpp
--- a/Engine/Source/Runtime/UMG/Private/Components/NativeWidgetHost.cpp
+++ b/Engine/Source/Runtime/UMG/Private/Components/NativeWidgetHost.cpp
@@ -7,8 +7,8 @@
 /////////////////////////////////////////////////////
 // UNativeWidgetHost
 
-UNativeWidgetHost::UNativeWidgetHost(const FPostConstructInitializeProperties& PCIP)
-	: Super(PCIP)
+UNativeWidgetHost::UNativeWidgetHost(const FObjectInitializer& ObjectInitializer)
+	: Super(ObjectInitializer)
 {
 	bIsVariable = false;
 }
diff --git a/Engine/Source/Runtime/UMG/Private/Components/Overlay.cpp b/Engine/Source/Runtime/UMG/Private/Components/Overlay.cpp
--- a/Engine/Source/Runtime/UMG/Private/Components/Overlay.cpp
+++ b/Engine/Source/Runtime/UMG/Private/Components/Overlay.cpp
@@ -7,8 +7,8 @@
 /////////////////////////////////////////////////////
 // UOverlay
 
-UOverlay::UOverlay(const FPostConstructInitializeProperties& PCIP)
-	: Super(PCIP)
+UOverlay::UOverlay(const FObjectInitializer& ObjectInitializer)
+	: Super(ObjectInitializer)
 {
 	bIsVariable = false;
 
diff --git a/Engine/Source/Runtime/UMG/Private/Components/OverlaySlot.cpp b/Engine/Source/Runtime/UMG/Private/Components/OverlaySlot.cpp
--- a/Engine/Source/Runtime/UMG/Private/Components/OverlaySlot.cpp
+++ b/Engine/Source/Runtime/UMG/Private/Components/OverlaySlot.cpp
@@ -5,8 +5,8 @@
 /////////////////////////////////////////////////////
 // UOverlaySlot
 
-UOverlaySlot::UOverlaySlot(const FPostConstructInitializeProperties& PCIP)
-	: Super(PCIP)
+UOverlaySlot::UOverlaySlot(const FObjectInitializer& ObjectInitializer)
+	: Super(ObjectInitializer)
 {
 	HorizontalAlignment = HAlign_Left;
 	VerticalAlignment = VAlign_Top;
diff --git a/Engine/Source/Runtime/UMG/Private/Components/PanelSlot.cpp b/Engine/Source/Runtime/UMG/Private/Components/PanelSlot.cpp
--- a/Engine/Source/Runtime/UMG/Private/Components/PanelSlot.cpp
+++ b/Engine/Source/Runtime/UMG/Private/Components/PanelSlot.cpp
@@ -5,8 +5,8 @@
 /////////////////////////////////////////////////////
 // UPanelSlot
 
-UPanelSlot::UPanelSlot(const FPostConstructInitializeProperties& PCIP)
-	: Super(PCIP)
+UPanelSlot::UPanelSlot(const FObjectInitializer& ObjectInitializer)
+	: Super(ObjectInitializer)
 {
 }
 
diff --git a/Engine/Source/Runtime/UMG/Private/Components/PanelWidget.cpp b/Engine/Source/Runtime/UMG/Private/Components/PanelWidget.cpp
--- a/Engine/Source/Runtime/UMG/Private/Components/PanelWidget.cpp
+++ b/Engine/Source/Runtime/UMG/Private/Components/PanelWidget.cpp
@@ -12,8 +12,8 @@
 /////////////////////////////////////////////////////
 // UPanelWidget
 
-UPanelWidget::UPanelWidget(const FPostConstructInitializeProperties& PCIP)
-	: Super(PCIP)
+UPanelWidget::UPanelWidget(const FObjectInitializer& ObjectInitializer)
+	: Super(ObjectInitializer)
 	, bCanHaveMultipleChildren(true)
 {
 }
diff --git a/Engine/Source/Runtime/UMG/Private/Components/ProgressBar.cpp b/Engine/Source/Runtime/UMG/Private/Components/ProgressBar.cpp
--- a/Engine/Source/Runtime/UMG/Private/Components/ProgressBar.cpp
+++ b/Engine/Source/Runtime/UMG/Private/Components/ProgressBar.cpp
@@ -7,8 +7,8 @@
 /////////////////////////////////////////////////////
 // UProgressBar
 
-UProgressBar::UProgressBar(const FPostConstructInitializeProperties& PCIP)
-	: Super(PCIP)
+UProgressBar::UProgressBar(const FObjectInitializer& ObjectInitializer)
+	: Super(ObjectInitializer)
 {
 	SProgressBar::FArguments SlateDefaults;
 	WidgetStyle = *SlateDefaults._Style;
diff --git a/Engine/Source/Runtime/UMG/Private/Components/ScaleBox.cpp b/Engine/Source/Runtime/UMG/Private/Components/ScaleBox.cpp
--- a/Engine/Source/Runtime/UMG/Private/Components/ScaleBox.cpp
+++ b/Engine/Source/Runtime/UMG/Private/Components/ScaleBox.cpp
@@ -7,8 +7,8 @@
 /////////////////////////////////////////////////////
 // UScaleBox
 
-UScaleBox::UScaleBox(const FPostConstructInitializeProperties& PCIP)
-	: Super(PCIP)
+UScaleBox::UScaleBox(const FObjectInitializer& ObjectInitializer)
+	: Super(ObjectInitializer)
 {
 	bIsVariable = false;
 
diff --git a/Engine/Source/Runtime/UMG/Private/Components/ScaleBoxSlot.cpp b/Engine/Source/Runtime/UMG/Private/Components/ScaleBoxSlot.cpp
--- a/Engine/Source/Runtime/UMG/Private/Components/ScaleBoxSlot.cpp
+++ b/Engine/Source/Runtime/UMG/Private/Components/ScaleBoxSlot.cpp
@@ -5,8 +5,8 @@
 /////////////////////////////////////////////////////
 // UScaleBoxSlot
 
-UScaleBoxSlot::UScaleBoxSlot(const FPostConstructInitializeProperties& PCIP)
-	: Super(PCIP)
+UScaleBoxSlot::UScaleBoxSlot(const FObjectInitializer& ObjectInitializer)
+	: Super(ObjectInitializer)
 {
 	Padding = FMargin(0, 0);
 
diff --git a/Engine/Source/Runtime/UMG/Private/Components/ScrollBar.cpp b/Engine/Source/Runtime/UMG/Private/Components/ScrollBar.cpp
--- a/Engine/Source/Runtime/UMG/Private/Components/ScrollBar.cpp
+++ b/Engine/Source/Runtime/UMG/Private/Components/ScrollBar.cpp
@@ -7,8 +7,8 @@
 /////////////////////////////////////////////////////
 // UScrollBar
 
-UScrollBar::UScrollBar(const FPostConstructInitializeProperties& PCIP)
-	: Super(PCIP)
+UScrollBar::UScrollBar(const FObjectInitializer& ObjectInitializer)
+	: Super(ObjectInitializer)
 {
 	bIsVariable = false;
 
diff --git a/Engine/Source/Runtime/UMG/Private/Components/ScrollBox.cpp b/Engine/Source/Runtime/UMG/Private/Components/ScrollBox.cpp
--- a/Engine/Source/Runtime/UMG/Private/Components/ScrollBox.cpp
+++ b/Engine/Source/Runtime/UMG/Private/Components/ScrollBox.cpp
@@ -7,8 +7,8 @@
 /////////////////////////////////////////////////////
 // UScrollBox
 
-UScrollBox::UScrollBox(const FPostConstructInitializeProperties& PCIP)
-	: Super(PCIP)
+UScrollBox::UScrollBox(const FObjectInitializer& ObjectInitializer)
+	: Super(ObjectInitializer)
 {
 	bIsVariable = false;
 
diff --git a/Engine/Source/Runtime/UMG/Private/Components/ScrollBoxSlot.cpp b/Engine/Source/Runtime/UMG/Private/Components/ScrollBoxSlot.cpp
--- a/Engine/Source/Runtime/UMG/Private/Components/ScrollBoxSlot.cpp
+++ b/Engine/Source/Runtime/UMG/Private/Components/ScrollBoxSlot.cpp
@@ -5,8 +5,8 @@
 /////////////////////////////////////////////////////
 // UScrollBoxSlot
 
-UScrollBoxSlot::UScrollBoxSlot(const FPostConstructInitializeProperties& PCIP)
-	: Super(PCIP)
+UScrollBoxSlot::UScrollBoxSlot(const FObjectInitializer& ObjectInitializer)
+	: Super(ObjectInitializer)
 	, Slot(nullptr)
 {
 	HorizontalAlignment = HAlign_Fill;
diff --git a/Engine/Source/Runtime/UMG/Private/Components/SizeBox.cpp b/Engine/Source/Runtime/UMG/Private/Components/SizeBox.cpp
--- a/Engine/Source/Runtime/UMG/Private/Components/SizeBox.cpp
+++ b/Engine/Source/Runtime/UMG/Private/Components/SizeBox.cpp
@@ -8,8 +8,8 @@
 /////////////////////////////////////////////////////
 // USizeBox
 
-USizeBox::USizeBox(const FPostConstructInitializeProperties& PCIP)
-	: Super(PCIP)
+USizeBox::USizeBox(const FObjectInitializer& ObjectInitializer)
+	: Super(ObjectInitializer)
 {
 	bIsVariable = false;
 }
diff --git a/Engine/Source/Runtime/UMG/Private/Components/SizeBoxSlot.cpp b/Engine/Source/Runtime/UMG/Private/Components/SizeBoxSlot.cpp
--- a/Engine/Source/Runtime/UMG/Private/Components/SizeBoxSlot.cpp
+++ b/Engine/Source/Runtime/UMG/Private/Components/SizeBoxSlot.cpp
@@ -5,8 +5,8 @@
 /////////////////////////////////////////////////////
 // USizeBoxSlot
 
-USizeBoxSlot::USizeBoxSlot(const FPostConstructInitializeProperties& PCIP)
-	: Super(PCIP)
+USizeBoxSlot::USizeBoxSlot(const FObjectInitializer& ObjectInitializer)
+	: Super(ObjectInitializer)
 {
 	Padding = FMargin(0, 0);
 
diff --git a/Engine/Source/Runtime/UMG/Private/Components/Slider.cpp b/Engine/Source/Runtime/UMG/Private/Components/Slider.cpp
--- a/Engine/Source/Runtime/UMG/Private/Components/Slider.cpp
+++ b/Engine/Source/Runtime/UMG/Private/Components/Slider.cpp
@@ -7,8 +7,8 @@
 /////////////////////////////////////////////////////
 // USlider
 
-USlider::USlider(const FPostConstructInitializeProperties& PCIP)
-	: Super(PCIP)
+USlider::USlider(const FObjectInitializer& ObjectInitializer)
+	: Super(ObjectInitializer)
 {
 	Orientation = EOrientation::Orient_Horizontal;
 	SliderBarColor = FLinearColor::White;
diff --git a/Engine/Source/Runtime/UMG/Private/Components/Spacer.cpp b/Engine/Source/Runtime/UMG/Private/Components/Spacer.cpp
--- a/Engine/Source/Runtime/UMG/Private/Components/Spacer.cpp
+++ b/Engine/Source/Runtime/UMG/Private/Components/Spacer.cpp
@@ -7,8 +7,8 @@
 /////////////////////////////////////////////////////
 // USpacer
 
-USpacer::USpacer(const FPostConstructInitializeProperties& PCIP)
-	: Super(PCIP)
+USpacer::USpacer(const FObjectInitializer& ObjectInitializer)
+	: Super(ObjectInitializer)
 	, Size(1.0f, 1.0f)
 {
 	bIsVariable = false;
diff --git a/Engine/Source/Runtime/UMG/Private/Components/SpinBox.cpp b/Engine/Source/Runtime/UMG/Private/Components/SpinBox.cpp
--- a/Engine/Source/Runtime/UMG/Private/Components/SpinBox.cpp
+++ b/Engine/Source/Runtime/UMG/Private/Components/SpinBox.cpp
@@ -7,8 +7,8 @@
 /////////////////////////////////////////////////////
 // USpinBox
 
-USpinBox::USpinBox(const FPostConstructInitializeProperties& PCIP)
-	: Super(PCIP)
+USpinBox::USpinBox(const FObjectInitializer& ObjectInitializer)
+	: Super(ObjectInitializer)
 {
 	// HACK Special font initialization hack since there are no font assets yet for slate.
 	Font = FSlateFontInfo(TEXT("Slate/Fonts/Roboto-Bold.ttf"), 12);
diff --git a/Engine/Source/Runtime/UMG/Private/Components/TableViewBase.cpp b/Engine/Source/Runtime/UMG/Private/Components/TableViewBase.cpp
--- a/Engine/Source/Runtime/UMG/Private/Components/TableViewBase.cpp
+++ b/Engine/Source/Runtime/UMG/Private/Components/TableViewBase.cpp
@@ -7,8 +7,8 @@
 /////////////////////////////////////////////////////
 // UTableViewBase
 
-UTableViewBase::UTableViewBase(const FPostConstructInitializeProperties& PCIP)
-	: Super(PCIP)
+UTableViewBase::UTableViewBase(const FObjectInitializer& ObjectInitializer)
+	: Super(ObjectInitializer)
 {
 	bIsVariable = true;
 }
diff --git a/Engine/Source/Runtime/UMG/Private/Components/TextBlock.cpp b/Engine/Source/Runtime/UMG/Private/Components/TextBlock.cpp
--- a/Engine/Source/Runtime/UMG/Private/Components/TextBlock.cpp
+++ b/Engine/Source/Runtime/UMG/Private/Components/TextBlock.cpp
@@ -7,8 +7,8 @@
 /////////////////////////////////////////////////////
 // UTextBlock
 
-UTextBlock::UTextBlock(const FPostConstructInitializeProperties& PCIP)
-	: Super(PCIP)
+UTextBlock::UTextBlock(const FObjectInitializer& ObjectInitializer)
+	: Super(ObjectInitializer)
 {
 	bIsVariable = false;
 
diff --git a/Engine/Source/Runtime/UMG/Private/Components/Throbber.cpp b/Engine/Source/Runtime/UMG/Private/Components/Throbber.cpp
--- a/Engine/Source/Runtime/UMG/Private/Components/Throbber.cpp
+++ b/Engine/Source/Runtime/UMG/Private/Components/Throbber.cpp
@@ -8,8 +8,8 @@
 /////////////////////////////////////////////////////
 // UThrobber
 
-UThrobber::UThrobber(const FPostConstructInitializeProperties& PCIP)
-	: Super(PCIP)
+UThrobber::UThrobber(const FObjectInitializer& ObjectInitializer)
+	: Super(ObjectInitializer)
 {
 	SThrobber::FArguments DefaultArgs;
 	Image = *DefaultArgs._PieceImage;
diff --git a/Engine/Source/Runtime/UMG/Private/Components/TileView.cpp b/Engine/Source/Runtime/UMG/Private/Components/TileView.cpp
--- a/Engine/Source/Runtime/UMG/Private/Components/TileView.cpp
+++ b/Engine/Source/Runtime/UMG/Private/Components/TileView.cpp
@@ -7,8 +7,8 @@
 /////////////////////////////////////////////////////
 // UTileView
 
-UTileView::UTileView(const FPostConstructInitializeProperties& PCIP)
-	: Super(PCIP)
+UTileView::UTileView(const FObjectInitializer& ObjectInitializer)
+	: Super(ObjectInitializer)
 {
 	bIsVariable = true;
 
diff --git a/Engine/Source/Runtime/UMG/Private/Components/UniformGridPanel.cpp b/Engine/Source/Runtime/UMG/Private/Components/UniformGridPanel.cpp
--- a/Engine/Source/Runtime/UMG/Private/Components/UniformGridPanel.cpp
+++ b/Engine/Source/Runtime/UMG/Private/Components/UniformGridPanel.cpp
@@ -7,8 +7,8 @@
 /////////////////////////////////////////////////////
 // UUniformGridPanel
 
-UUniformGridPanel::UUniformGridPanel(const FPostConstructInitializeProperties& PCIP)
-	: Super(PCIP)
+UUniformGridPanel::UUniformGridPanel(const FObjectInitializer& ObjectInitializer)
+	: Super(ObjectInitializer)
 {
 	bIsVariable = false;
 
diff --git a/Engine/Source/Runtime/UMG/Private/Components/UniformGridSlot.cpp b/Engine/Source/Runtime/UMG/Private/Components/UniformGridSlot.cpp
--- a/Engine/Source/Runtime/UMG/Private/Components/UniformGridSlot.cpp
+++ b/Engine/Source/Runtime/UMG/Private/Components/UniformGridSlot.cpp
@@ -5,8 +5,8 @@
 /////////////////////////////////////////////////////
 // UUniformGridSlot
 
-UUniformGridSlot::UUniformGridSlot(const FPostConstructInitializeProperties& PCIP)
-	: Super(PCIP)
+UUniformGridSlot::UUniformGridSlot(const FObjectInitializer& ObjectInitializer)
+	: Super(ObjectInitializer)
 	, Slot(NULL)
 {
 	HorizontalAlignment = HAlign_Left;
diff --git a/Engine/Source/Runtime/UMG/Private/Components/VerticalBox.cpp b/Engine/Source/Runtime/UMG/Private/Components/VerticalBox.cpp
--- a/Engine/Source/Runtime/UMG/Private/Components/VerticalBox.cpp
+++ b/Engine/Source/Runtime/UMG/Private/Components/VerticalBox.cpp
@@ -7,8 +7,8 @@
 /////////////////////////////////////////////////////
 // UVerticalBox
 
-UVerticalBox::UVerticalBox(const FPostConstructInitializeProperties& PCIP)
-	: Super(PCIP)
+UVerticalBox::UVerticalBox(const FObjectInitializer& ObjectInitializer)
+	: Super(ObjectInitializer)
 {
 	bIsVariable = false;
 
diff --git a/Engine/Source/Runtime/UMG/Private/Components/VerticalBoxSlot.cpp b/Engine/Source/Runtime/UMG/Private/Components/VerticalBoxSlot.cpp
--- a/Engine/Source/Runtime/UMG/Private/Components/VerticalBoxSlot.cpp
+++ b/Engine/Source/Runtime/UMG/Private/Components/VerticalBoxSlot.cpp
@@ -5,8 +5,8 @@
 /////////////////////////////////////////////////////
 // UVerticalBoxSlot
 
-UVerticalBoxSlot::UVerticalBoxSlot(const FPostConstructInitializeProperties& PCIP)
-	: Super(PCIP)
+UVerticalBoxSlot::UVerticalBoxSlot(const FObjectInitializer& ObjectInitializer)
+	: Super(ObjectInitializer)
 	, Slot(NULL)
 {
 	HorizontalAlignment = HAlign_Fill;
diff --git a/Engine/Source/Runtime/UMG/Private/Components/Viewport.cpp b/Engine/Source/Runtime/UMG/Private/Components/Viewport.cpp
--- a/Engine/Source/Runtime/UMG/Private/Components/Viewport.cpp
+++ b/Engine/Source/Runtime/UMG/Private/Components/Viewport.cpp
@@ -486,8 +486,8 @@ public:
 /////////////////////////////////////////////////////
 // UViewport
 
-UViewport::UViewport(const FPostConstructInitializeProperties& PCIP)
-	: Super(PCIP)
+UViewport::UViewport(const FObjectInitializer& ObjectInitializer)
+	: Super(ObjectInitializer)
 	, ShowFlags(ESFIM_Game)
 {
 	bIsVariable = true;
diff --git a/Engine/Source/Runtime/UMG/Private/Components/Visual.cpp b/Engine/Source/Runtime/UMG/Private/Components/Visual.cpp
--- a/Engine/Source/Runtime/UMG/Private/Components/Visual.cpp
+++ b/Engine/Source/Runtime/UMG/Private/Components/Visual.cpp
@@ -5,8 +5,8 @@
 /////////////////////////////////////////////////////
 // UVisual
 
-UVisual::UVisual(const FPostConstructInitializeProperties& PCIP)
-	: Super(PCIP)
+UVisual::UVisual(const FObjectInitializer& ObjectInitializer)
+	: Super(ObjectInitializer)
 {
 }
 
diff --git a/Engine/Source/Runtime/UMG/Private/Components/Widget.cpp b/Engine/Source/Runtime/UMG/Private/Components/Widget.cpp
--- a/Engine/Source/Runtime/UMG/Private/Components/Widget.cpp
+++ b/Engine/Source/Runtime/UMG/Private/Components/Widget.cpp
@@ -5,8 +5,8 @@
 /////////////////////////////////////////////////////
 // UWidget
 
-UWidget::UWidget(const FPostConstructInitializeProperties& PCIP)
-	: Super(PCIP)
+UWidget::UWidget(const FObjectInitializer& ObjectInitializer)
+	: Super(ObjectInitializer)
 {
 	bIsEnabled = true;
 	bIsVariable = true;
diff --git a/Engine/Source/Runtime/UMG/Private/Components/WidgetSwitcher.cpp b/Engine/Source/Runtime/UMG/Private/Components/WidgetSwitcher.cpp
--- a/Engine/Source/Runtime/UMG/Private/Components/WidgetSwitcher.cpp
+++ b/Engine/Source/Runtime/UMG/Private/Components/WidgetSwitcher.cpp
@@ -7,8 +7,8 @@
 /////////////////////////////////////////////////////
 // UWidgetSwitcher
 
-UWidgetSwitcher::UWidgetSwitcher(const FPostConstructInitializeProperties& PCIP)
-	: Super(PCIP)
+UWidgetSwitcher::UWidgetSwitcher(const FObjectInitializer& ObjectInitializer)
+	: Super(ObjectInitializer)
 {
 	bIsVariable = true;
 
diff --git a/Engine/Source/Runtime/UMG/Private/Components/WidgetSwitcherSlot.cpp b/Engine/Source/Runtime/UMG/Private/Components/WidgetSwitcherSlot.cpp
--- a/Engine/Source/Runtime/UMG/Private/Components/WidgetSwitcherSlot.cpp
+++ b/Engine/Source/Runtime/UMG/Private/Components/WidgetSwitcherSlot.cpp
@@ -5,8 +5,8 @@
 /////////////////////////////////////////////////////
 // UWidgetSwitcherSlot
 
-UWidgetSwitcherSlot::UWidgetSwitcherSlot(const FPostConstructInitializeProperties& PCIP)
-	: Super(PCIP)
+UWidgetSwitcherSlot::UWidgetSwitcherSlot(const FObjectInitializer& ObjectInitializer)
+	: Super(ObjectInitializer)
 	, Slot(NULL)
 {
 	HorizontalAlignment = HAlign_Fill;
diff --git a/Engine/Source/Runtime/UMG/Private/Components/WrapBox.cpp b/Engine/Source/Runtime/UMG/Private/Components/WrapBox.cpp
--- a/Engine/Source/Runtime/UMG/Private/Components/WrapBox.cpp
+++ b/Engine/Source/Runtime/UMG/Private/Components/WrapBox.cpp
@@ -7,8 +7,8 @@
 /////////////////////////////////////////////////////
 // UWrapBox
 
-UWrapBox::UWrapBox(const FPostConstructInitializeProperties& PCIP)
-	: Super(PCIP)
+UWrapBox::UWrapBox(const FObjectInitializer& ObjectInitializer)
+	: Super(ObjectInitializer)
 {
 	bIsVariable = false;
 
diff --git a/Engine/Source/Runtime/UMG/Private/Components/WrapBoxSlot.cpp b/Engine/Source/Runtime/UMG/Private/Components/WrapBoxSlot.cpp
--- a/Engine/Source/Runtime/UMG/Private/Components/WrapBoxSlot.cpp
+++ b/Engine/Source/Runtime/UMG/Private/Components/WrapBoxSlot.cpp
@@ -5,8 +5,8 @@
 /////////////////////////////////////////////////////
 // UWrapBoxSlot
 
-UWrapBoxSlot::UWrapBoxSlot(const FPostConstructInitializeProperties& PCIP)
-	: Super(PCIP)
+UWrapBoxSlot::UWrapBoxSlot(const FObjectInitializer& ObjectInitializer)
+	: Super(ObjectInitializer)
 	, Slot(NULL)
 {
 	HorizontalAlignment = HAlign_Fill;
diff --git a/Engine/Source/Runtime/UMG/Private/DragDropOperation.cpp b/Engine/Source/Runtime/UMG/Private/DragDropOperation.cpp
--- a/Engine/Source/Runtime/UMG/Private/DragDropOperation.cpp
+++ b/Engine/Source/Runtime/UMG/Private/DragDropOperation.cpp
@@ -5,8 +5,8 @@
 /////////////////////////////////////////////////////
 // UDragDropOperation
 
-UDragDropOperation::UDragDropOperation(const FPostConstructInitializeProperties& PCIP)
-	: Super(PCIP)
+UDragDropOperation::UDragDropOperation(const FObjectInitializer& ObjectInitializer)
+	: Super(ObjectInitializer)
 {
 }
 
diff --git a/Engine/Source/Runtime/UMG/Private/SlateBlueprintLibrary.cpp b/Engine/Source/Runtime/UMG/Private/SlateBlueprintLibrary.cpp
--- a/Engine/Source/Runtime/UMG/Private/SlateBlueprintLibrary.cpp
+++ b/Engine/Source/Runtime/UMG/Private/SlateBlueprintLibrary.cpp
@@ -10,8 +10,8 @@
 /////////////////////////////////////////////////////
 // USlateBlueprintLibrary
 
-USlateBlueprintLibrary::USlateBlueprintLibrary(const FPostConstructInitializeProperties& PCIP)
-: Super(PCIP)
+USlateBlueprintLibrary::USlateBlueprintLibrary(const FObjectInitializer& ObjectInitializer)
+: Super(ObjectInitializer)
 {
 }
 
diff --git a/Engine/Source/Runtime/UMG/Private/UserWidget.cpp b/Engine/Source/Runtime/UMG/Private/UserWidget.cpp
--- a/Engine/Source/Runtime/UMG/Private/UserWidget.cpp
+++ b/Engine/Source/Runtime/UMG/Private/UserWidget.cpp
@@ -14,8 +14,8 @@
 /////////////////////////////////////////////////////
 // UUserWidget
 
-UUserWidget::UUserWidget(const FPostConstructInitializeProperties& PCIP)
-	: Super(PCIP)
+UUserWidget::UUserWidget(const FObjectInitializer& ObjectInitializer)
+	: Super(ObjectInitializer)
 {
 	ViewportAnchors = FAnchors(0, 0, 1, 1);
 	Visiblity = ESlateVisibility::SelfHitTestInvisible;
diff --git a/Engine/Source/Runtime/UMG/Private/WidgetBlueprintGeneratedClass.cpp b/Engine/Source/Runtime/UMG/Private/WidgetBlueprintGeneratedClass.cpp
--- a/Engine/Source/Runtime/UMG/Private/WidgetBlueprintGeneratedClass.cpp
+++ b/Engine/Source/Runtime/UMG/Private/WidgetBlueprintGeneratedClass.cpp
@@ -10,8 +10,8 @@
 /////////////////////////////////////////////////////
 // UWidgetBlueprintGeneratedClass
 
-UWidgetBlueprintGeneratedClass::UWidgetBlueprintGeneratedClass(const FPostConstructInitializeProperties& PCIP)
-	: Super(PCIP)
+UWidgetBlueprintGeneratedClass::UWidgetBlueprintGeneratedClass(const FObjectInitializer& ObjectInitializer)
+	: Super(ObjectInitializer)
 {
 }
 
diff --git a/Engine/Source/Runtime/UMG/Private/WidgetBlueprintLibrary.cpp b/Engine/Source/Runtime/UMG/Private/WidgetBlueprintLibrary.cpp
--- a/Engine/Source/Runtime/UMG/Private/WidgetBlueprintLibrary.cpp
+++ b/Engine/Source/Runtime/UMG/Private/WidgetBlueprintLibrary.cpp
@@ -9,8 +9,8 @@
 /////////////////////////////////////////////////////
 // UWidgetBlueprintLibrary
 
-UWidgetBlueprintLibrary::UWidgetBlueprintLibrary(const FPostConstructInitializeProperties& PCIP)
-: Super(PCIP)
+UWidgetBlueprintLibrary::UWidgetBlueprintLibrary(const FObjectInitializer& ObjectInitializer)
+: Super(ObjectInitializer)
 {
 }
 
diff --git a/Engine/Source/Runtime/UMG/Private/WidgetLayoutLibrary.cpp b/Engine/Source/Runtime/UMG/Private/WidgetLayoutLibrary.cpp
--- a/Engine/Source/Runtime/UMG/Private/WidgetLayoutLibrary.cpp
+++ b/Engine/Source/Runtime/UMG/Private/WidgetLayoutLibrary.cpp
@@ -10,8 +10,8 @@
 /////////////////////////////////////////////////////
 // UWidgetLayoutLibrary
 
-UWidgetLayoutLibrary::UWidgetLayoutLibrary(const FPostConstructInitializeProperties& PCIP)
-: Super(PCIP)
+UWidgetLayoutLibrary::UWidgetLayoutLibrary(const FObjectInitializer& ObjectInitializer)
+: Super(ObjectInitializer)
 {
 }
 
diff --git a/Engine/Source/Runtime/UMG/Private/WidgetTree.cpp b/Engine/Source/Runtime/UMG/Private/WidgetTree.cpp
--- a/Engine/Source/Runtime/UMG/Private/WidgetTree.cpp
+++ b/Engine/Source/Runtime/UMG/Private/WidgetTree.cpp
@@ -5,8 +5,8 @@
 /////////////////////////////////////////////////////
 // UWidgetTree
 
-UWidgetTree::UWidgetTree(const FPostConstructInitializeProperties& PCIP)
-	: Super(PCIP)
+UWidgetTree::UWidgetTree(const FObjectInitializer& ObjectInitializer)
+	: Super(ObjectInitializer)
 {
 }
 
