Added Stub class.

diff --git a/Engine/Source/Runtime/SlateCore/Public/Types/NavigationMetaData.h b/Engine/Source/Runtime/SlateCore/Public/Types/NavigationMetaData.h
new file mode 100644
--- /dev/null
+++ b/Engine/Source/Runtime/SlateCore/Public/Types/NavigationMetaData.h
@@ -0,0 +1,28 @@
+// Copyright 1998-2014 Epic Games, Inc. All Rights Reserved.
+
+#pragma once
+
+#include "ISlateMetaData.h"
+
+/**
+ * Simple tagging metadata
+ */
+class FNavigationMetaData : public ISlateMetaData
+{
+public:
+	SLATE_METADATA_TYPE(FNavigationMetaData, ISlateMetaData)
+
+	FNavigationMetaData()
+	{
+	}
+
+private:
+	/*
+	struct SNavData
+	{
+		EUINavigationRule BoundaryRule;
+		TAttribute<TSharedPtr<SWidget>> FocusRecipient;
+	};
+	SNavData Rules[EUINavigation::Num];
+	*/
+};
