Slight modification to previous CL#2331543, TitleBarHeight is not a generic plaform function so others can specify it

diff --git a/Engine/Source/Runtime/Core/Private/GenericPlatform/GenericWindow.cpp b/Engine/Source/Runtime/Core/Private/GenericPlatform/GenericWindow.cpp
--- a/Engine/Source/Runtime/Core/Private/GenericPlatform/GenericWindow.cpp
+++ b/Engine/Source/Runtime/Core/Private/GenericPlatform/GenericWindow.cpp
@@ -135,6 +135,12 @@ int32 FGenericWindow::GetWindowBorderSize() const
 	return 0;
 }
 
+int32 FGenericWindow::GetWindowTitleBarSize() const
+{
+	// empty default functionality
+	return 0;
+}
+
 void* FGenericWindow::GetOSWindowHandle() const
 {
 	// empty default functionality (for platforms that don't need an OS handle)
diff --git a/Engine/Source/Runtime/Core/Private/Windows/WindowsWindow.cpp b/Engine/Source/Runtime/Core/Private/Windows/WindowsWindow.cpp
--- a/Engine/Source/Runtime/Core/Private/Windows/WindowsWindow.cpp
+++ b/Engine/Source/Runtime/Core/Private/Windows/WindowsWindow.cpp
@@ -644,6 +644,11 @@ int32 FWindowsWindow::GetWindowBorderSize() const
 	return WindowInfo.cxWindowBorders;
 }
 
+int32 FWindowsWindow::GetWindowTitleBarSize() const
+{
+	return GetSystemMetrics(SM_CYCAPTION);
+}
+
 bool FWindowsWindow::IsForegroundWindow() const
 {
 	return ::GetForegroundWindow() == HWnd;
diff --git a/Engine/Source/Runtime/Core/Public/GenericPlatform/GenericWindow.h b/Engine/Source/Runtime/Core/Public/GenericPlatform/GenericWindow.h
--- a/Engine/Source/Runtime/Core/Public/GenericPlatform/GenericWindow.h
+++ b/Engine/Source/Runtime/Core/Public/GenericPlatform/GenericWindow.h
@@ -139,6 +139,9 @@ public:
 	/** Gets OS specific window border size. This is necessary because Win32 does not give control over this size. */
 	virtual int32 GetWindowBorderSize() const;
 
+	/** Gets OS specific window title bar size */
+	virtual int32 GetWindowTitleBarSize() const;
+
 	/** Gets the OS Window handle in the form of a void pointer for other API's */
 	virtual void* GetOSWindowHandle() const;
 
diff --git a/Engine/Source/Runtime/Core/Public/Windows/WindowsWindow.h b/Engine/Source/Runtime/Core/Public/Windows/WindowsWindow.h
--- a/Engine/Source/Runtime/Core/Public/Windows/WindowsWindow.h
+++ b/Engine/Source/Runtime/Core/Public/Windows/WindowsWindow.h
@@ -82,6 +82,8 @@ public:
 
 	virtual int32 GetWindowBorderSize() const override;
 
+	virtual int32 GetWindowTitleBarSize() const override;
+
 	virtual void* GetOSWindowHandle() const  override { return HWnd; }
 
 	virtual bool IsForegroundWindow() const override;
diff --git a/Engine/Source/Runtime/SlateCore/Private/Widgets/SWindow.cpp b/Engine/Source/Runtime/SlateCore/Private/Widgets/SWindow.cpp
--- a/Engine/Source/Runtime/SlateCore/Private/Widgets/SWindow.cpp
+++ b/Engine/Source/Runtime/SlateCore/Private/Widgets/SWindow.cpp
@@ -678,18 +678,17 @@ FMargin SWindow::GetWindowBorderSize( bool bIncTitleBar ) const
 // @TODO This is not working for Linux. The window is not yet valid when this gets
 // called from SWindow::Construct which is causing a default border to be retured even when the
 // window is borderless. This causes problems for menu positioning.
-#if PLATFORM_WINDOWS // || PLATFORM_LINUX
 	if (NativeWindow.IsValid() && NativeWindow->IsMaximized())
 	{
 		FMargin BorderSize(NativeWindow->GetWindowBorderSize());
 		if(bIncTitleBar)
 		{
 			// Add title bar size (whether it's visible or not)
-			BorderSize.Top += GetSystemMetrics(SM_CYCAPTION);
+			BorderSize.Top += NativeWindow->GetWindowTitleBarSize();
 		}
 		return BorderSize;
 	}
-
+#if PLATFORM_WINDOWS // || PLATFORM_LINUX
 	return SWindowDefs::WindowBorderSize;
 #else
 	return FMargin();
