Disable platform-specific hack in Slate (SWindow).

diff --git a/Engine/Source/Runtime/SlateCore/Private/Widgets/SWindow.cpp b/Engine/Source/Runtime/SlateCore/Private/Widgets/SWindow.cpp
--- a/Engine/Source/Runtime/SlateCore/Private/Widgets/SWindow.cpp
+++ b/Engine/Source/Runtime/SlateCore/Private/Widgets/SWindow.cpp
@@ -2,6 +2,8 @@
 
 #include "SlateCorePrivatePCH.h"
 
+// this define is the first step to complete removal of the below hack
+#define PLATFORM_SPECIFIC_HACK 	0
 
 namespace SWindowDefs
 {
@@ -1353,7 +1355,7 @@ FReply SWindow::OnKeyboardFocusReceived( const FGeometry& MyGeometry, const FKey
 
 FReply SWindow::OnMouseButtonDown( const FGeometry& MyGeometry, const FPointerEvent& MouseEvent )
 {
-#if PLATFORM_LINUX
+#if PLATFORM_SPECIFIC_HACK && PLATFORM_LINUX
 	if (bHasSizingFrame && MouseEvent.GetEffectingButton() == EKeys::LeftMouseButton)
 	{
 		if ((WindowZone == EWindowZone::TopLeftBorder || WindowZone == EWindowZone::BottomRightBorder ||
@@ -1382,7 +1384,7 @@ FReply SWindow::OnMouseButtonDown( const FGeometry& MyGeometry, const FPointerEv
 
 FReply SWindow::OnMouseButtonUp( const FGeometry& MyGeometry, const FPointerEvent& MouseEvent )
 {
-#if PLATFORM_LINUX
+#if PLATFORM_SPECIFIC_HACK && PLATFORM_LINUX
 	if (MoveResizeZone != EWindowZone::Unspecified && MouseEvent.GetEffectingButton() == EKeys::LeftMouseButton)
 	{
 		MoveResizeZone =  EWindowZone::Unspecified;
@@ -1401,7 +1403,7 @@ FReply SWindow::OnMouseButtonUp( const FGeometry& MyGeometry, const FPointerEven
 
 FReply SWindow::OnMouseMove( const FGeometry& MyGeometry, const FPointerEvent& MouseEvent )
 {
-#if PLATFORM_LINUX
+#if PLATFORM_SPECIFIC_HACK && PLATFORM_LINUX
 	if (MoveResizeZone == EWindowZone::TopLeftBorder)
 	{
 		FVector2D MoveResizeOffset = MouseEvent.GetScreenSpacePosition() - MoveResizeStart;
