#orion

- Adding in Depth Sorting

diff --git a/Engine/Source/Runtime/SlateCore/Public/Layout/Children.h b/Engine/Source/Runtime/SlateCore/Public/Layout/Children.h
--- a/Engine/Source/Runtime/SlateCore/Public/Layout/Children.h
+++ b/Engine/Source/Runtime/SlateCore/Public/Layout/Children.h
@@ -237,7 +237,7 @@ public:
 	template <class PREDICATE_CLASS>
 	void Sort( const PREDICATE_CLASS& Predicate )
 	{
-		TIndirectArray< SlotType >::Sort( Predicate );
+		::Sort(GetData(), Num(), Predicate);
 	}
 
 	void Swap( int32 IndexA, int32 IndexB )
