#UE4

- Making Sort work on TPanelChildren

diff --git a/Engine/Source/Runtime/SlateCore/Public/Layout/Children.h b/Engine/Source/Runtime/SlateCore/Public/Layout/Children.h
--- a/Engine/Source/Runtime/SlateCore/Public/Layout/Children.h
+++ b/Engine/Source/Runtime/SlateCore/Public/Layout/Children.h
@@ -237,7 +237,7 @@ public:
 	template <class PREDICATE_CLASS>
 	void Sort( const PREDICATE_CLASS& Predicate )
 	{
-		TIndirectArray< SlotType >::Sort( Predicate );
+		::Sort(TIndirectArray< SlotType >::GetData(), TIndirectArray<SlotType>::Num(), Predicate);
 	}
 
 	void Swap( int32 IndexA, int32 IndexB )
