CIS fix

diff --git a/Engine/Source/Runtime/SlateCore/Private/Fonts/FontCache.cpp b/Engine/Source/Runtime/SlateCore/Private/Fonts/FontCache.cpp
--- a/Engine/Source/Runtime/SlateCore/Private/Fonts/FontCache.cpp
+++ b/Engine/Source/Runtime/SlateCore/Private/Fonts/FontCache.cpp
@@ -585,6 +585,7 @@ public:
 	{
 		static const TSet<FName> DummyAttributes;
 
+#if WITH_FREETYPE
 		FFontFaceAndMemory* FaceAndMemory = FontFaceMap.Find(&InFontData);
 		if (!FaceAndMemory)
 		{
@@ -593,10 +594,14 @@ public:
 		}
 
 		return (FaceAndMemory) ? FaceAndMemory->Attributes : DummyAttributes;
+#else
+		return DummyAttributes;
+#endif // WITH_FREETYPE
 	}
 
 private:
 
+#if WITH_FREETYPE
 	/** Get the cached composite font data for the given composite font */
 	TSharedPtr<const FCachedCompositeFontData> GetCachedCompositeFont(const FCompositeFont* const InCompositeFont)
 	{
@@ -730,7 +735,6 @@ private:
 		return BestMatchFont;
 	}
 
-#if WITH_FREETYPE
 	/**
 	 * Gets or loads a freetype font face
 	 *
