Converted composite fonts to store their embedded TTF/OTF font data as "bulk data"

This allows the data to be compressed and lazy loaded, so the embedded font data is now only loaded when FreeType needs to load it into a face.

diff --git a/Engine/Source/Editor/FontEditor/Private/SCompositeFontEditor.cpp b/Engine/Source/Editor/FontEditor/Private/SCompositeFontEditor.cpp
--- a/Engine/Source/Editor/FontEditor/Private/SCompositeFontEditor.cpp
+++ b/Engine/Source/Editor/FontEditor/Private/SCompositeFontEditor.cpp
@@ -699,10 +699,11 @@ void STypefaceEntryEditor::OnTypefaceEntryFontPathPicked(const FString& InNewFon
 		const FScopedTransaction Transaction(LOCTEXT("SetFontFile", "Set Font File"));
 		CompositeFontEditorPtr->GetFontObject()->Modify();
 
-		if(TypefaceEntryPtr->Font.SetFont(InNewFontFilename))
-		{
-			CompositeFontEditorPtr->FlushCachedFont();
-		}
+		// We need to allocate the bulk data with the font as its outer
+		const UFontBulkData* const BulkData = new(CompositeFontEditorPtr->GetFontObject()) UFontBulkData(InNewFontFilename);
+
+		TypefaceEntryPtr->Font.SetFont(InNewFontFilename, BulkData);
+		CompositeFontEditorPtr->FlushCachedFont();
 	}
 
 	FEditorDirectories::Get().SetLastDirectory(ELastDirectory::GENERIC_OPEN, FPaths::GetPath(InNewFontFilename));
diff --git a/Engine/Source/Editor/UnrealEd/Private/Factories/EditorFactories.cpp b/Engine/Source/Editor/UnrealEd/Private/Factories/EditorFactories.cpp
--- a/Engine/Source/Editor/UnrealEd/Private/Factories/EditorFactories.cpp
+++ b/Engine/Source/Editor/UnrealEd/Private/Factories/EditorFactories.cpp
@@ -5117,12 +5117,10 @@ UObject* UFontFileImportFactory::FactoryCreateBinary(UClass* InClass, UObject* I
 	{
 		Font->FontCacheType = EFontCacheType::Runtime;
 
-		FString FontFilename = GetCurrentFilename();
+		// We need to allocate the bulk data with the font as its outer
+		const UFontBulkData* const BulkData = new(Font) UFontBulkData(InBuffer, InBufferEnd - InBuffer);
 
-		TArray<uint8> FontData;
-		FontData.Insert(InBuffer, InBufferEnd - InBuffer, 0);
-
-		Font->CompositeFont.DefaultTypeface.Fonts.Add(FTypefaceEntry("Default", MoveTemp(FontFilename), MoveTemp(FontData), EFontHinting::Auto));
+		Font->CompositeFont.DefaultTypeface.Fonts.Add(FTypefaceEntry("Default", GetCurrentFilename(), BulkData, EFontHinting::Auto));
 	}
 
 	FEditorDelegates::OnAssetPostImport.Broadcast(this, Font);
diff --git a/Engine/Source/Runtime/Core/Public/UObject/ObjectVersion.h b/Engine/Source/Runtime/Core/Public/UObject/ObjectVersion.h
--- a/Engine/Source/Runtime/Core/Public/UObject/ObjectVersion.h
+++ b/Engine/Source/Runtime/Core/Public/UObject/ObjectVersion.h
@@ -634,6 +634,8 @@ enum EUnrealEngineObjectUE4Version
 	VER_UE4_REMOVE_SAVEGAMESUMMARY,
 	//Remove bodyseutp serialization from skeletal mesh component
 	VER_UE4_REMOVE_SKELETALMESH_COMPONENT_BODYSETUP_SERIALIZATION,
+	// Made Slate font data use bulk data to store the embedded font data
+	VER_UE4_SLATE_BULK_FONT_DATA,
 
 	// -----<new versions can be added before this line>-------------------------------------------------
 	// - this needs to be the last line (see note below)
diff --git a/Engine/Source/Runtime/Engine/Private/Font.cpp b/Engine/Source/Runtime/Engine/Private/Font.cpp
--- a/Engine/Source/Runtime/Engine/Private/Font.cpp
+++ b/Engine/Source/Runtime/Engine/Private/Font.cpp
@@ -41,6 +41,32 @@ void UFont::Serialize( FArchive& Ar )
 {
 	Super::Serialize( Ar );
 	Ar << CharRemap;
+
+	// Update any composite font data to use bulk data with us as its outer
+	if( Ar.UE4Ver() < VER_UE4_SLATE_BULK_FONT_DATA )
+	{
+		auto UpgradeFontDataToBulk = [this](FFontData& InFontData)
+		{
+			if( InFontData.FontData_DEPRECATED.Num() > 0 )
+			{
+				InFontData.BulkDataPtr = new(this) UFontBulkData(InFontData.FontData_DEPRECATED.GetData(), InFontData.FontData_DEPRECATED.Num());
+				InFontData.FontData_DEPRECATED.Empty();
+			}
+		};
+
+		for( FTypefaceEntry& TypefaceEntry : CompositeFont.DefaultTypeface.Fonts )
+		{
+			UpgradeFontDataToBulk(TypefaceEntry.Font);
+		}
+
+		for( FCompositeSubFont& SubFont : CompositeFont.SubTypefaces )
+		{
+			for( FTypefaceEntry& TypefaceEntry : SubFont.Typeface.Fonts )
+			{
+				UpgradeFontDataToBulk(TypefaceEntry.Font);
+			}
+		}
+	}
 }
 
 void UFont::PostLoad()
diff --git a/Engine/Source/Runtime/SlateCore/Private/Fonts/CompositeFont.cpp b/Engine/Source/Runtime/SlateCore/Private/Fonts/CompositeFont.cpp
--- a/Engine/Source/Runtime/SlateCore/Private/Fonts/CompositeFont.cpp
+++ b/Engine/Source/Runtime/SlateCore/Private/Fonts/CompositeFont.cpp
@@ -3,33 +3,40 @@
 #include "SlateCorePrivatePCH.h"
 #include "CompositeFont.h"
 
-FFontData::FFontData(FString InFontFilename, const EFontHinting InHinting)
+FFontData::FFontData()
+	: FontFilename()
+	, BulkDataPtr(nullptr)
+	, Hinting(EFontHinting::Default)
+	, FontData_DEPRECATED()
+{
+}
+
+FFontData::FFontData(FString InFontFilename, const UFontBulkData* const InBulkData, const EFontHinting InHinting)
 	: FontFilename(MoveTemp(InFontFilename))
-	, FontData()
+	, BulkDataPtr(InBulkData)
 	, Hinting(InHinting)
+	, FontData_DEPRECATED()
 {
-	if (!FFileHelper::LoadFileToArray(FontData, *FontFilename, FILEREAD_Silent))
-	{
-		UE_LOG(LogSlate, Warning, TEXT("Failed to load font data from '%s'"), *FontFilename);
-	}
 }
 
-bool FFontData::SetFont(FString InFontFilename)
+void FFontData::SetFont(FString InFontFilename, const UFontBulkData* const InBulkData)
 {
-	TArray<uint8> TempFontData;
-	if (!FFileHelper::LoadFileToArray(TempFontData, *InFontFilename, FILEREAD_Silent))
-	{
-		UE_LOG(LogSlate, Warning, TEXT("Failed to load font data from '%s'"), *InFontFilename);
-		return false;
-	}
-
 	FontFilename = MoveTemp(InFontFilename);
-	FontData = MoveTemp(TempFontData);
-	return true;
+	BulkDataPtr = InBulkData;
 }
 
-void FFontData::SetFont(FString InFontFilename, TArray<uint8> InFontData)
+void FStandaloneCompositeFont::AddReferencedObjects(FReferenceCollector& Collector)
 {
-	FontFilename = MoveTemp(InFontFilename);
-	FontData = MoveTemp(InFontData);
+	for(FTypefaceEntry& TypefaceEntry : DefaultTypeface.Fonts)
+	{
+		Collector.AddReferencedObject(TypefaceEntry.Font.BulkDataPtr);
+	}
+
+	for(FCompositeSubFont& SubFont : SubTypefaces)
+	{
+		for(FTypefaceEntry& TypefaceEntry : SubFont.Typeface.Fonts)
+		{
+			Collector.AddReferencedObject(TypefaceEntry.Font.BulkDataPtr);
+		}
+	}
 }
diff --git a/Engine/Source/Runtime/SlateCore/Private/Fonts/FontBulkData.cpp b/Engine/Source/Runtime/SlateCore/Private/Fonts/FontBulkData.cpp
new file mode 100644
--- /dev/null
+++ b/Engine/Source/Runtime/SlateCore/Private/Fonts/FontBulkData.cpp
@@ -0,0 +1,56 @@
+// Copyright 1998-2014 Epic Games, Inc. All Rights Reserved.
+
+#include "SlateCorePrivatePCH.h"
+#include "FontBulkData.h"
+
+UFontBulkData::UFontBulkData()
+{
+	BulkData.SetBulkDataFlags(BULKDATA_SerializeCompressed);
+}
+
+UFontBulkData::UFontBulkData(const FString& InFontFilename)
+{
+	BulkData.SetBulkDataFlags(BULKDATA_SerializeCompressed);
+
+	TUniquePtr<FArchive> Reader(IFileManager::Get().CreateFileReader(*InFontFilename, 0));
+	if(Reader)
+	{
+		const int32 FontDataSizeBytes = Reader->TotalSize();
+
+		BulkData.Lock(LOCK_READ_WRITE);
+		void* const LockedFontData = BulkData.Realloc(FontDataSizeBytes);
+		Reader->Serialize(LockedFontData, FontDataSizeBytes);
+		BulkData.Unlock();
+	}
+	else
+	{
+		UE_LOG(LogSlate, Warning, TEXT("Failed to load font data from '%s'"), *InFontFilename);
+	}
+}
+
+UFontBulkData::UFontBulkData(const void* const InFontData, const int32 InFontDataSizeBytes)
+{
+	BulkData.SetBulkDataFlags(BULKDATA_SerializeCompressed);
+
+	BulkData.Lock(LOCK_READ_WRITE);
+	void* const LockedFontData = BulkData.Realloc(InFontDataSizeBytes);
+	FMemory::Memcpy(LockedFontData, InFontData, InFontDataSizeBytes);
+	BulkData.Unlock();
+}
+
+const void* UFontBulkData::Lock(int32& OutFontDataSizeBytes) const
+{
+	OutFontDataSizeBytes = BulkData.GetBulkDataSize();
+	return BulkData.LockReadOnly();
+}
+
+void UFontBulkData::Unlock() const
+{
+	BulkData.Unlock();
+}
+
+void UFontBulkData::Serialize(FArchive& Ar)
+{
+	Super::Serialize(Ar);
+	BulkData.Serialize(Ar, this);
+}
diff --git a/Engine/Source/Runtime/SlateCore/Private/Fonts/FontCache.cpp b/Engine/Source/Runtime/SlateCore/Private/Fonts/FontCache.cpp
--- a/Engine/Source/Runtime/SlateCore/Private/Fonts/FontCache.cpp
+++ b/Engine/Source/Runtime/SlateCore/Private/Fonts/FontCache.cpp
@@ -743,40 +743,44 @@ private:
 	FT_Face GetFontFace( const FFontData& InFontData )
 	{
 		FFontFaceAndMemory* FaceAndMemory = FontFaceMap.Find(&InFontData);
-		if (!FaceAndMemory && InFontData.FontData.Num() > 0)
+		if (!FaceAndMemory && InFontData.BulkDataPtr)
 		{
-			// make a new entry
-			FaceAndMemory = &FontFaceMap.Add(&InFontData, FFontFaceAndMemory());
+			int32 LockedFontDataSizeBytes = 0;
+			const void* const LockedFontData = InFontData.BulkDataPtr->Lock(LockedFontDataSizeBytes);
+			if (LockedFontDataSizeBytes > 0)
+			{
+				// make a new entry
+				FaceAndMemory = &FontFaceMap.Add(&InFontData, FFontFaceAndMemory());
 
-			// todo: jdale - can we avoid this copy without potentially crashing if the source FFontData is unloaded (when using a UObject)?
-			//				 We'd just need to remove this entry from the font cache before unloading the UFont and then we 
-			//				 could just take a pointer to the array data (which we don't own, so won't delete)
-			FaceAndMemory->Memory = static_cast<uint8*>(FMemory::Malloc(InFontData.FontData.Num()));
-			FMemory::Memcpy(FaceAndMemory->Memory, InFontData.FontData.GetData(), InFontData.FontData.Num());
+				FaceAndMemory->Memory = static_cast<uint8*>(FMemory::Malloc(LockedFontDataSizeBytes));
+				FMemory::Memcpy(FaceAndMemory->Memory, LockedFontData, LockedFontDataSizeBytes);
 
-			// initialize the font, setting the error code
-			const bool bFailedToLoadFace = FT_New_Memory_Face(FTLibrary, FaceAndMemory->Memory, static_cast<FT_Long>(InFontData.FontData.Num()), 0, &FaceAndMemory->Face) != 0;
+				// initialize the font, setting the error code
+				const bool bFailedToLoadFace = FT_New_Memory_Face(FTLibrary, FaceAndMemory->Memory, static_cast<FT_Long>(LockedFontDataSizeBytes), 0, &FaceAndMemory->Face) != 0;
 
-			// if it failed, we don't want to keep the memory around
-			if (bFailedToLoadFace)
-			{
-				FMemory::Free(FaceAndMemory->Memory);
-				FontFaceMap.Remove(&InFontData);
-				FaceAndMemory = nullptr;
-				UE_LOG(LogSlate, Warning, TEXT("GetFontFace failed to load or process '%s'"), *InFontData.FontFilename);
-			}
+				// if it failed, we don't want to keep the memory around
+				if (bFailedToLoadFace)
+				{
+					FMemory::Free(FaceAndMemory->Memory);
+					FontFaceMap.Remove(&InFontData);
+					FaceAndMemory = nullptr;
+					UE_LOG(LogSlate, Warning, TEXT("GetFontFace failed to load or process '%s'"), *InFontData.FontFilename);
+				}
 			
-			if (FaceAndMemory)
-			{
-				// Parse out the font attributes
-				TArray<FString> Styles;
-				FString(FaceAndMemory->Face->style_name).ParseIntoArray(&Styles, TEXT(" "), true);
-
-				for (const FString& Style : Styles)
+				if (FaceAndMemory)
 				{
-					FaceAndMemory->Attributes.Add(*Style);
+					// Parse out the font attributes
+					TArray<FString> Styles;
+					FString(FaceAndMemory->Face->style_name).ParseIntoArray(&Styles, TEXT(" "), true);
+
+					for (const FString& Style : Styles)
+					{
+						FaceAndMemory->Attributes.Add(*Style);
+					}
 				}
 			}
+
+			InFontData.BulkDataPtr->Unlock();
 		}
 
 		return (FaceAndMemory) ? FaceAndMemory->Face : nullptr;
diff --git a/Engine/Source/Runtime/SlateCore/Private/Fonts/LegacySlateFontInfoCache.cpp b/Engine/Source/Runtime/SlateCore/Private/Fonts/LegacySlateFontInfoCache.cpp
--- a/Engine/Source/Runtime/SlateCore/Private/Fonts/LegacySlateFontInfoCache.cpp
+++ b/Engine/Source/Runtime/SlateCore/Private/Fonts/LegacySlateFontInfoCache.cpp
@@ -49,7 +49,7 @@ TSharedPtr<const FCompositeFont> FLegacySlateFontInfoCache::GetCompositeFont(con
 		}
 	}
 
-	TSharedRef<const FCompositeFont> NewCompositeFont = MakeShareable(new FCompositeFont(NAME_None, LegacyFontPath, InLegacyFontHinting));
+	TSharedRef<const FCompositeFont> NewCompositeFont = MakeShareable(new FStandaloneCompositeFont(NAME_None, LegacyFontPath, new UFontBulkData(LegacyFontPath), InLegacyFontHinting));
 	LegacyFontNameToCompositeFont.Add(LegacyFontKey, NewCompositeFont);
 	return NewCompositeFont;
 }
@@ -59,7 +59,7 @@ TSharedPtr<const FCompositeFont> FLegacySlateFontInfoCache::GetSystemFont()
 	if (!SystemFont.IsValid())
 	{
 		TArray<uint8> FontBytes = FPlatformMisc::GetSystemFontBytes();
-		SystemFont = MakeShareable(new FCompositeFont(NAME_None, TEXT("DefaultSystemFont"), MoveTemp(FontBytes), EFontHinting::Default));
+		SystemFont = MakeShareable(new FStandaloneCompositeFont(NAME_None, TEXT("DefaultSystemFont"), new UFontBulkData(FontBytes.GetData(), FontBytes.Num()), EFontHinting::Default));
 	}
 	return SystemFont;
 }
@@ -76,7 +76,8 @@ const FFontData& FLegacySlateFontInfoCache::GetFallbackFont()
 		}
 	}
 
-	TSharedRef<const FFontData> NewFallbackFont = MakeShareable(new FFontData(FallbackFontName.ToString(), EFontHinting::Default));
+	const FString FallbackFontPath = FallbackFontName.ToString();
+	TSharedRef<const FFontData> NewFallbackFont = MakeShareable(new FFontData(FallbackFontPath, new UFontBulkData(FallbackFontPath), EFontHinting::Default));
 	FallbackFonts.Add(FallbackFontName, NewFallbackFont);
 	return *NewFallbackFont;
 }
@@ -86,7 +87,22 @@ const FFontData& FLegacySlateFontInfoCache::GetLastResortFont()
 	if (!LastResortFont.IsValid())
 	{
 		const FString LastResortFontPath = FPaths::EngineContentDir() / TEXT("Slate/Fonts/LastResort.ttf");
-		LastResortFont = MakeShareable(new FFontData(LastResortFontPath, EFontHinting::Default));
+		LastResortFont = MakeShareable(new FFontData(LastResortFontPath, new UFontBulkData(LastResortFontPath), EFontHinting::Default));
 	}
 	return *LastResortFont;
 }
+
+void FLegacySlateFontInfoCache::AddReferencedObjects(FReferenceCollector& Collector)
+{
+	for(auto& FallbackFontEntry : FallbackFonts)
+	{
+		const UFontBulkData* TmpPtr = FallbackFontEntry.Value->BulkDataPtr;
+		Collector.AddReferencedObject(TmpPtr);
+	}
+
+	if(LastResortFont.IsValid())
+	{
+		const UFontBulkData* TmpPtr = LastResortFont->BulkDataPtr;
+		Collector.AddReferencedObject(TmpPtr);
+	}
+}
diff --git a/Engine/Source/Runtime/SlateCore/Private/Fonts/LegacySlateFontInfoCache.h b/Engine/Source/Runtime/SlateCore/Private/Fonts/LegacySlateFontInfoCache.h
--- a/Engine/Source/Runtime/SlateCore/Private/Fonts/LegacySlateFontInfoCache.h
+++ b/Engine/Source/Runtime/SlateCore/Private/Fonts/LegacySlateFontInfoCache.h
@@ -5,7 +5,7 @@
 /**
  * Cache used to efficiently upgrade legacy FSlateFontInfo structs to use a composite font by reducing the amount of duplicate instances that are created
  */
-class FLegacySlateFontInfoCache
+class FLegacySlateFontInfoCache : public FGCObject
 {
 public:
 
@@ -34,6 +34,9 @@ public:
 	 */
 	const FFontData& GetLastResortFont();
 
+	// FGCObject interface
+	virtual void AddReferencedObjects(FReferenceCollector& Collector) override;
+
 private:
 	struct FLegacyFontKey
 	{
diff --git a/Engine/Source/Runtime/SlateCore/Public/Fonts/CompositeFont.h b/Engine/Source/Runtime/SlateCore/Public/Fonts/CompositeFont.h
--- a/Engine/Source/Runtime/SlateCore/Public/Fonts/CompositeFont.h
+++ b/Engine/Source/Runtime/SlateCore/Public/Fonts/CompositeFont.h
@@ -2,6 +2,7 @@
 
 #pragma once
 
+#include "FontBulkData.h"
 #include "CompositeFont.generated.h"
 
 UENUM()
@@ -26,41 +27,29 @@ struct SLATECORE_API FFontData
 	GENERATED_USTRUCT_BODY()
 
 	/** Default constructor */
-	FFontData()
-		: FontFilename()
-		, FontData()
-		, Hinting(EFontHinting::Default)
-	{
-	}
-
-	/** Construct the raw font data from a filename (it will load the file into FontData) */
-	FFontData(FString InFontFilename, const EFontHinting InHinting);
+	FFontData();
 
 	/** Construct the raw font data from a filename, and the data associated with that file */
-	FFontData(FString InFontFilename, TArray<uint8> InFontData, const EFontHinting InHinting)
-		: FontFilename(MoveTemp(InFontFilename))
-		, FontData(MoveTemp(InFontData))
-		, Hinting(InHinting)
-	{
-	}
-
-	/** Set the new font data from a filename (it will load the file into FontData) */
-	bool SetFont(FString InFontFilename);
+	FFontData(FString InFontFilename, const UFontBulkData* const InBulkData, const EFontHinting InHinting);
 
 	/** Set the new font data from a filename, and the data associated with that file */
-	void SetFont(FString InFontFilename, TArray<uint8> InFontData);
+	void SetFont(FString InFontFilename, const UFontBulkData* const InBulkData);
 
 	/** The filename of the font to use - this may not always exist on disk, as we may have previously loaded and cached the font data inside an asset */
 	UPROPERTY(EditAnywhere, BlueprintReadWrite, Category=Font)
 	FString FontFilename;
 
-	/** The data associated with the font - this should always be filled in providing the source font filename is valid */
+	/** The data associated with the font */
 	UPROPERTY()
-	TArray<uint8> FontData;
+	const UFontBulkData* BulkDataPtr;
 
 	/** The hinting algorithm to use with the font */
 	UPROPERTY(EditAnywhere, BlueprintReadWrite, Category=Font)
 	EFontHinting Hinting;
+
+	/** The data associated with the font - this should always be filled in providing the source font filename is valid */
+	UPROPERTY()
+	TArray<uint8> FontData_DEPRECATED;
 };
 
 /** A single entry in a typeface */
@@ -80,17 +69,10 @@ struct SLATECORE_API FTypefaceEntry
 	{
 	}
 
-	/** Construct the entry from a filename */
-	FTypefaceEntry(const FName& InFontName, FString InFontFilename, const EFontHinting InHinting)
-		: Name(InFontName)
-		, Font(MoveTemp(InFontFilename), InHinting)
-	{
-	}
-
 	/** Construct the entry from a filename, and the data associated with that file */
-	FTypefaceEntry(const FName& InFontName, FString InFontFilename, TArray<uint8> InFontData, const EFontHinting InHinting)
+	FTypefaceEntry(const FName& InFontName, FString InFontFilename, const UFontBulkData* const InBulkData, const EFontHinting InHinting)
 		: Name(InFontName)
-		, Font(MoveTemp(InFontFilename), MoveTemp(InFontData), InHinting)
+		, Font(MoveTemp(InFontFilename), InBulkData, InHinting)
 	{
 	}
 
@@ -115,28 +97,15 @@ struct SLATECORE_API FTypeface
 	}
 
 	/** Convenience constructor for when your font family only contains a single font */
-	FTypeface(const FName& InFontName, FString InFontFilename, const EFontHinting InHinting)
-	{
-		Fonts.Emplace(FTypefaceEntry(InFontName, MoveTemp(InFontFilename), InHinting));
-	}
-
-	/** Convenience constructor for when your font family only contains a single font */
-	FTypeface(const FName& InFontName, FString InFontFilename, TArray<uint8> InFontData, const EFontHinting InHinting)
-	{
-		Fonts.Emplace(FTypefaceEntry(InFontName, MoveTemp(InFontFilename), MoveTemp(InFontData), InHinting));
-	}
-
-	/** Append a new font into this family */
-	FTypeface& AppendFont(const FName& InFontName, FString InFontFilename, const EFontHinting InHinting = EFontHinting::Default)
+	FTypeface(const FName& InFontName, FString InFontFilename, const UFontBulkData* const InBulkData, const EFontHinting InHinting)
 	{
-		Fonts.Emplace(FTypefaceEntry(InFontName, MoveTemp(InFontFilename), InHinting));
-		return *this;
+		Fonts.Emplace(FTypefaceEntry(InFontName, MoveTemp(InFontFilename), InBulkData, InHinting));
 	}
 
 	/** Append a new font into this family */
-	FTypeface& AppendFont(const FName& InFontName, FString InFontFilename, TArray<uint8> InFontData, const EFontHinting InHinting = EFontHinting::Default)
+	FTypeface& AppendFont(const FName& InFontName, FString InFontFilename, const UFontBulkData* const InBulkData, const EFontHinting InHinting = EFontHinting::Default)
 	{
-		Fonts.Emplace(FTypefaceEntry(InFontName, MoveTemp(InFontFilename), MoveTemp(InFontData), InHinting));
+		Fonts.Emplace(FTypefaceEntry(InFontName, MoveTemp(InFontFilename), InBulkData, InHinting));
 		return *this;
 	}
 
@@ -185,16 +154,8 @@ struct SLATECORE_API FCompositeFont
 	}
 
 	/** Convenience constructor for when your composite font only contains a single font */
-	FCompositeFont(const FName& InFontName, FString InFontFilename, const EFontHinting InHinting = EFontHinting::Default)
-		: DefaultTypeface(InFontName, MoveTemp(InFontFilename), InHinting)
-		, SubTypefaces()
-		, HistoryRevision(0)
-	{
-	}
-
-	/** Convenience constructor for when your composite font only contains a single font */
-	FCompositeFont(const FName& InFontName, FString InFontFilename, TArray<uint8> InFontData, const EFontHinting InHinting = EFontHinting::Default)
-		: DefaultTypeface(InFontName, MoveTemp(InFontFilename), MoveTemp(InFontData), InHinting)
+	FCompositeFont(const FName& InFontName, FString InFontFilename, const UFontBulkData* const InBulkData, const EFontHinting InHinting = EFontHinting::Default)
+		: DefaultTypeface(InFontName, MoveTemp(InFontFilename), InBulkData, InHinting)
 		, SubTypefaces()
 		, HistoryRevision(0)
 	{
@@ -220,3 +181,24 @@ struct SLATECORE_API FCompositeFont
 	 */
 	int32 HistoryRevision;
 };
+
+/**
+ * A version of FCompositeFont that should be used when it's not being embedded within another UObject
+ * This derives from FGCObject to ensure that the bulk data objects are referenced correctly 
+ */
+struct SLATECORE_API FStandaloneCompositeFont : public FCompositeFont, public FGCObject
+{
+	/** Default constructor */
+	FStandaloneCompositeFont()
+	{
+	}
+
+	/** Convenience constructor for when your composite font only contains a single font */
+	FStandaloneCompositeFont(const FName& InFontName, FString InFontFilename, const UFontBulkData* const InBulkData, const EFontHinting InHinting = EFontHinting::Default)
+		: FCompositeFont(InFontName, MoveTemp(InFontFilename), InBulkData, InHinting)
+	{
+	}
+
+	// FGCObject interface
+	virtual void AddReferencedObjects(FReferenceCollector& Collector) override;
+};
diff --git a/Engine/Source/Runtime/SlateCore/Public/Fonts/FontBulkData.h b/Engine/Source/Runtime/SlateCore/Public/Fonts/FontBulkData.h
new file mode 100644
--- /dev/null
+++ b/Engine/Source/Runtime/SlateCore/Public/Fonts/FontBulkData.h
@@ -0,0 +1,34 @@
+// Copyright 1998-2014 Epic Games, Inc. All Rights Reserved.
+
+#pragma once
+
+#include "FontBulkData.generated.h"
+
+UCLASS()
+class SLATECORE_API UFontBulkData : public UObject
+{
+	GENERATED_BODY()
+
+public:
+	/** Default constructor */
+	UFontBulkData();
+
+	/** Construct the bulk font data from the given file */
+	UFontBulkData(const FString& InFontFilename);
+
+	/** Construct the bulk font data from the given data */
+	UFontBulkData(const void* const InFontData, const int32 InFontDataSizeBytes);
+
+	/** Locks the bulk font data and returns a read-only pointer to it */
+	const void* Lock(int32& OutFontDataSizeBytes) const;
+
+	/** Unlock the bulk font data, after which point the pointer returned by Lock no longer is valid */
+	void Unlock() const;
+
+	// UObject interface
+	virtual void Serialize(FArchive& Ar) override;
+
+private:
+	/** Internal bulk data */
+	FByteBulkData BulkData;
+};
