SWindow: Added two new options to SWindow to allow developers to create windows without a layout border about the edge of the window while still allowing for a user resize area.

diff --git a/Engine/Source/Runtime/SlateCore/Private/Widgets/SWindow.cpp b/Engine/Source/Runtime/SlateCore/Private/Widgets/SWindow.cpp
--- a/Engine/Source/Runtime/SlateCore/Private/Widgets/SWindow.cpp
+++ b/Engine/Source/Runtime/SlateCore/Private/Widgets/SWindow.cpp
@@ -10,12 +10,6 @@ namespace SWindowDefs
 	/** Height of a Slate window title bar, in pixels */
 	static const float DefaultTitleBarSize = 24.0f;
 
-	/** Size of the hit result border for the window borders */
-	static const FSlateRect HitResultBorderSize(10,10,10,10);
-
-	/** Actual size of the window borders */
-	static const FMargin WindowBorderSize(5,5,5,5);
-
 	/** Size of the corner rounding radius.  Used for regular, non-maximized windows only (not tool-tips or decorators.) */
 	static const int32 CornerRadius = 6;
 }
@@ -204,6 +198,8 @@ void SWindow::Construct(const FArguments& InArgs)
 	this->bHasMinimizeButton = InArgs._SupportsMinimize;
 	this->bHasMaximizeButton = InArgs._SupportsMaximize;
 	this->bHasSizingFrame = !InArgs._IsPopupWindow && InArgs._SizingRule == ESizingRule::UserSized;
+	this->LayoutBorder = InArgs._LayoutBorder;
+	this->UserResizeBorder = InArgs._UserResizeBorder;
 	
 	// calculate window size from client size
 	const bool bCreateTitleBar = InArgs._CreateTitleBar && !bIsPopupWindow && !bIsCursorDecoratorWindow && !bHasOSWindowBorder;
@@ -367,7 +363,7 @@ FVector2D SWindow::ComputeWindowSizeForContent( FVector2D ContentSize )
 	// @todo mainframe: This code should be updated to handle the case where we're spawning a window that doesn't have 
 	//                  a traditional title bar, such as a window that contains a primary SDockingArea.  Currently, the
 	//                  size reported here will be too large!
-	return ContentSize + FVector2D(0, SWindowDefs::DefaultTitleBarSize) + SWindowDefs::WindowBorderSize.GetDesiredSize();
+	return ContentSize + FVector2D(0, SWindowDefs::DefaultTitleBarSize);
 }
 
 void SWindow::ConstructWindowInternals( const bool bCreateTitleBar )
@@ -690,10 +686,11 @@ FMargin SWindow::GetWindowBorderSize( bool bIncTitleBar ) const
 			// Add title bar size (whether it's visible or not)
 			BorderSize.Top += NativeWindow->GetWindowTitleBarSize();
 		}
+
 		return BorderSize;
 	}
 #if PLATFORM_WINDOWS // || PLATFORM_LINUX
-	return SWindowDefs::WindowBorderSize;
+	return LayoutBorder;
 #else
 	return FMargin();
 #endif
@@ -1618,29 +1615,30 @@ EWindowZone::Type SWindow::GetCurrentWindowZone(FVector2D LocalMousePosition)
 		int32 Col = 1;
 		if (SizingRule == ESizingRule::UserSized)
 		{
-			if (LocalMousePosition.X < SWindowDefs::HitResultBorderSize.Left)
+			if (LocalMousePosition.X < (UserResizeBorder.Left + 5))
 			{
 				Col = 0;
 			}
-			else if (LocalMousePosition.X >= Size.X - SWindowDefs::HitResultBorderSize.Right)
+			else if (LocalMousePosition.X >= Size.X - (UserResizeBorder.Right + 5))
 			{
 				Col = 2;
 			}
-			if (LocalMousePosition.Y < SWindowDefs::HitResultBorderSize.Top)
+
+			if (LocalMousePosition.Y < (UserResizeBorder.Top + 5))
 			{
 				Row = 0;
 			}
-			else if (LocalMousePosition.Y >= Size.Y - SWindowDefs::HitResultBorderSize.Bottom)
+			else if (LocalMousePosition.Y >= Size.Y - (UserResizeBorder.Bottom + 5))
 			{
 				Row = 2;
 			}
 
 			// The actual border is smaller than the hit result zones
 			// This grants larger corner areas to grab onto
-			bool bInBorder =	LocalMousePosition.X < SWindowDefs::WindowBorderSize.Left ||
-								LocalMousePosition.X >= Size.X - SWindowDefs::WindowBorderSize.Right ||
-								LocalMousePosition.Y < SWindowDefs::WindowBorderSize.Top ||
-								LocalMousePosition.Y >= Size.Y - SWindowDefs::WindowBorderSize.Bottom;
+			bool bInBorder =	LocalMousePosition.X < UserResizeBorder.Left ||
+								LocalMousePosition.X >= Size.X - UserResizeBorder.Right ||
+								LocalMousePosition.Y < UserResizeBorder.Top ||
+								LocalMousePosition.Y >= Size.Y - UserResizeBorder.Bottom;
 
 			if (!bInBorder)
 			{
@@ -1677,7 +1675,7 @@ EWindowZone::Type SWindow::GetCurrentWindowZone(FVector2D LocalMousePosition)
 				else if( HitTestResults.Widgets.Last().Widget == AsShared() )
 				{
 					// The window itself was hit, so check for a traditional title bar
-					if( (LocalMousePosition.Y - SWindowDefs::WindowBorderSize.Top) < TitleBarSize )
+					if ((LocalMousePosition.Y - UserResizeBorder.Top) < TitleBarSize)
 					{
 						InZone = EWindowZone::TitleBar;
 					}
diff --git a/Engine/Source/Runtime/SlateCore/Public/Widgets/SWindow.h b/Engine/Source/Runtime/SlateCore/Public/Widgets/SWindow.h
--- a/Engine/Source/Runtime/SlateCore/Public/Widgets/SWindow.h
+++ b/Engine/Source/Runtime/SlateCore/Public/Widgets/SWindow.h
@@ -83,6 +83,8 @@ public:
 		, _SupportsMinimize( true )
 		, _CreateTitleBar( true )
 		, _SaneWindowPlacement( true )
+		, _LayoutBorder(FMargin(5, 5, 5, 5))
+		, _UserResizeBorder(FMargin(5, 5, 5, 5))
 	{ }
 
 		/** Style used to draw this window */
@@ -143,6 +145,12 @@ public:
 		    constraints on the window and bring it back into view. */
 		SLATE_ARGUMENT( bool, SaneWindowPlacement )
 
+		/** The padding around the edges of the window applied to it's content. */
+		SLATE_ARGUMENT(FMargin, LayoutBorder)
+
+		/** The margin around the edges of the window that will be detected as places the user can grab to resize the window. */
+		SLATE_ARGUMENT(FMargin, UserResizeBorder)
+
 		SLATE_DEFAULT_SLOT( FArguments, Content )
 
 	SLATE_END_ARGS()
@@ -859,6 +867,12 @@ private:
 	// The window title bar.
 	TSharedPtr<IWindowTitleBar> TitleBar;
 
+	// The padding for between the edges of the window and it's content
+	FMargin LayoutBorder;
+
+	// The margin around the edges of the window that will be detected as places the user can grab to resize the window. 
+	FMargin UserResizeBorder;
+
 private:
 	
 	void ConstructWindowInternals( const bool bCreateTitleBar );
