Refactor to unify keyboard and controller input and focus.

diff --git a/Engine/Source/Runtime/SlateCore/Private/Layout/WidgetPath.cpp b/Engine/Source/Runtime/SlateCore/Private/Layout/WidgetPath.cpp
--- a/Engine/Source/Runtime/SlateCore/Private/Layout/WidgetPath.cpp
+++ b/Engine/Source/Runtime/SlateCore/Private/Layout/WidgetPath.cpp
@@ -155,9 +155,11 @@ struct FFocusableWidgetMatcher
  *
  * @return true if the focus moved successfully, false if we were unable to move focus
  */
-bool FWidgetPath::MoveFocus(int32 PathLevel, EFocusMoveDirection::Type MoveDirection)
+bool FWidgetPath::MoveFocus(int32 PathLevel, EUINavigation NavigationType)
 {
-	const int32 MoveDirectionAsInt = (MoveDirection == EFocusMoveDirection::Next)
+	check(NavigationType == EUINavigation::Next || NavigationType == EUINavigation::Previous);
+
+	const int32 MoveDirectionAsInt = (NavigationType == EUINavigation::Next)
 		? +1
 		: -1;
 
@@ -165,7 +167,7 @@ bool FWidgetPath::MoveFocus(int32 PathLevel, EFocusMoveDirection::Type MoveDirec
 	if ( PathLevel == Widgets.Num()-1 )
 	{
 		// We are the currently focused widget because we are at the very bottom of focus path.
-		if (MoveDirection == EFocusMoveDirection::Next)
+		if (NavigationType == EUINavigation::Next)
 		{
 			// EFocusMoveDirection::Next implies descend, so try to find a focusable descendant.
 			return ExtendPathTo( FFocusableWidgetMatcher() );
@@ -197,7 +199,7 @@ bool FWidgetPath::MoveFocus(int32 PathLevel, EFocusMoveDirection::Type MoveDirec
 			if ( ArrangedChildren[FocusedChildIndex].Widget->IsEnabled() )
 			{
 				// Look for a focusable descendant.
-				FArrangedChildren PathToFocusableChild = GeneratePathToWidget( FFocusableWidgetMatcher(), ArrangedChildren[FocusedChildIndex], MoveDirection );
+				FArrangedChildren PathToFocusableChild = GeneratePathToWidget(FFocusableWidgetMatcher(), ArrangedChildren[FocusedChildIndex], NavigationType);
 				// Either we found a focusable descendant, or an immediate child that is focusable.
 				const bool bFoundNextFocusable = ( PathToFocusableChild.Num() > 0 ) || ArrangedChildren[FocusedChildIndex].Widget->SupportsKeyboardFocus();
 				if ( bFoundNextFocusable )
@@ -330,8 +332,10 @@ bool FWeakWidgetPath::ContainsWidget( const TSharedRef< const SWidget >& SomeWid
  * 
  * @return The new focus path.
  */
-FWidgetPath FWeakWidgetPath::ToNextFocusedPath(EFocusMoveDirection::Type MoveDirection)
+FWidgetPath FWeakWidgetPath::ToNextFocusedPath(EUINavigation NavigationType)
 {
+	check(NavigationType == EUINavigation::Next || NavigationType == EUINavigation::Previous);
+
 	// Make a copy of the focus path. We will mutate it until it meets the necessary requirements.
 	FWidgetPath NewFocusPath = this->ToWidgetPath();
 	TSharedPtr<SWidget> CurrentlyFocusedWidget = this->Widgets.Last().Pin();
@@ -340,7 +344,7 @@ FWidgetPath FWeakWidgetPath::ToNextFocusedPath(EFocusMoveDirection::Type MoveDir
 	// Attempt to move the focus starting at the leafmost widget and bubbling up to the root (i.e. the window)
 	for (int32 FocusNodeIndex=NewFocusPath.Widgets.Num()-1; !bMovedFocus && FocusNodeIndex >= 0; --FocusNodeIndex)
 	{
-		bMovedFocus = NewFocusPath.MoveFocus( FocusNodeIndex, MoveDirection );
+		bMovedFocus = NewFocusPath.MoveFocus(FocusNodeIndex, NavigationType);
 	}
 
 	return NewFocusPath;
diff --git a/Engine/Source/Runtime/SlateCore/Private/Widgets/SWidget.cpp b/Engine/Source/Runtime/SlateCore/Private/Widgets/SWidget.cpp
--- a/Engine/Source/Runtime/SlateCore/Private/Widgets/SWidget.cpp
+++ b/Engine/Source/Runtime/SlateCore/Private/Widgets/SWidget.cpp
@@ -61,16 +61,29 @@ void SWidget::Construct(
 }
 
 
-FReply SWidget::OnKeyboardFocusReceived( const FGeometry& MyGeometry, const FKeyboardFocusEvent& InKeyboardFocusEvent )
+FReply SWidget::OnFocusReceived(const FGeometry& MyGeometry, const FFocusEvent& InFocusEvent)
 {
 	return FReply::Unhandled();
 }
 
-void SWidget::OnKeyboardFocusLost( const FKeyboardFocusEvent& InKeyboardFocusEvent )
+FReply SWidget::OnKeyboardFocusReceived(const FGeometry& MyGeometry, const FKeyboardFocusEvent& InFocusEvent)
+{
+	return FReply::Unhandled();
+}
+
+void SWidget::OnFocusLost(const FFocusEvent& InFocusEvent)
+{
+}
+
+void SWidget::OnKeyboardFocusLost(const FKeyboardFocusEvent& InFocusEvent)
+{
+}
+
+void SWidget::OnFocusChanging(const FWeakWidgetPath& PreviousFocusPath, const FWidgetPath& NewWidgetPath)
 {
 }
 
-void SWidget::OnKeyboardFocusChanging( const FWeakWidgetPath& PreviousFocusPath, const FWidgetPath& NewWidgetPath )
+void SWidget::OnKeyboardFocusChanging(const FWeakWidgetPath& PreviousFocusPath, const FWidgetPath& NewWidgetPath)
 {
 }
 
@@ -80,19 +93,53 @@ FReply SWidget::OnKeyChar( const FGeometry& MyGeometry, const FCharacterEvent& I
 }
 
 
-FReply SWidget::OnPreviewKeyDown( const FGeometry& MyGeometry, const FKeyboardEvent& InKeyboardEvent )
+FReply SWidget::OnPreviewKeyDown( const FGeometry& MyGeometry, const FKeyEvent& InKeyEvent )
 {
 	return FReply::Unhandled();
 }
 
 
-FReply SWidget::OnKeyDown( const FGeometry& MyGeometry, const FKeyboardEvent& InKeyboardEvent )
+FReply SWidget::OnKeyDown( const FGeometry& MyGeometry, const FKeyEvent& InKeyEvent )
+{
+	if (SupportsKeyboardFocus())
+	{
+		// It's the left stick return a navigation request of the correct direction
+		if (InKeyEvent.GetKey() == EKeys::Right || InKeyEvent.GetKey() == EKeys::Gamepad_DPad_Right || InKeyEvent.GetKey() == EKeys::Gamepad_LeftStick_Right)
+		{
+			return FReply::Handled().SetNavigation(EUINavigation::Right);
+		}
+		else if (InKeyEvent.GetKey() == EKeys::Left || InKeyEvent.GetKey() == EKeys::Gamepad_DPad_Left || InKeyEvent.GetKey() == EKeys::Gamepad_LeftStick_Left)
+		{
+			return FReply::Handled().SetNavigation(EUINavigation::Left);
+		}
+		else if (InKeyEvent.GetKey() == EKeys::Up || InKeyEvent.GetKey() == EKeys::Gamepad_DPad_Up || InKeyEvent.GetKey() == EKeys::Gamepad_LeftStick_Up)
+		{
+			return FReply::Handled().SetNavigation(EUINavigation::Up);
+		}
+		else if (InKeyEvent.GetKey() == EKeys::Down || InKeyEvent.GetKey() == EKeys::Gamepad_DPad_Down || InKeyEvent.GetKey() == EKeys::Gamepad_LeftStick_Down)
+		{
+			return FReply::Handled().SetNavigation(EUINavigation::Down);
+		}
+		// If the key was Tab, interpret as an attempt to move focus.
+		else if (InKeyEvent.GetKey() == EKeys::Tab)
+		{
+			EUINavigation MoveDirection = (InKeyEvent.IsShiftDown())
+				? EUINavigation::Previous
+				: EUINavigation::Next;
+			return FReply::Handled().SetNavigation(MoveDirection);
+		}
+	}
+	return FReply::Unhandled();
+}
+
+
+FReply SWidget::OnKeyUp( const FGeometry& MyGeometry, const FKeyEvent& InKeyEvent )
 {
 	return FReply::Unhandled();
 }
 
 
-FReply SWidget::OnKeyUp( const FGeometry& MyGeometry, const FKeyboardEvent& InKeyboardEvent )
+FReply SWidget::OnAnalogValueChanged( const FGeometry& MyGeometry, const FAnalogInputEvent& InAnalogInputEvent )
 {
 	return FReply::Unhandled();
 }
@@ -213,40 +260,55 @@ FReply SWidget::OnTouchGesture( const FGeometry& MyGeometry, const FPointerEvent
 }
 
 
-FReply SWidget::OnTouchStarted( const FGeometry& MyGeometry, const FPointerEvent& TouchEvent )
+FReply SWidget::OnTouchStarted( const FGeometry& MyGeometry, const FPointerEvent& InTouchEvent )
 {
 	return FReply::Unhandled();
 }
 
 
-FReply SWidget::OnTouchMoved( const FGeometry& MyGeometry, const FPointerEvent& TouchEvent )
+FReply SWidget::OnTouchMoved( const FGeometry& MyGeometry, const FPointerEvent& InTouchEvent )
 {
 	return FReply::Unhandled();
 }
 
 
-FReply SWidget::OnTouchEnded( const FGeometry& MyGeometry, const FPointerEvent& TouchEvent )
+FReply SWidget::OnTouchEnded( const FGeometry& MyGeometry, const FPointerEvent& InTouchEvent )
 {
 	return FReply::Unhandled();
 }
 
 
-FReply SWidget::OnMotionDetected( const FGeometry& MyGeometry, const FMotionEvent& MotionEvent )
+FReply SWidget::OnMotionDetected( const FGeometry& MyGeometry, const FMotionEvent& InMotionEvent )
 {
 	return FReply::Unhandled();
 }
 
+
 void SWidget::OnFinishedPointerInput()
 {
 
 }
 
+
 void SWidget::OnFinishedKeyInput()
 {
 
 }
 
 
+FNavigationReply SWidget::OnNavigation(const FGeometry& MyGeometry, const FNavigationEvent& InNavigationEvent)
+{
+	EUINavigation Type = InNavigationEvent.GetNavigationType();
+	TSharedPtr<FNavigationMetaData> MetaData = GetMetaData<FNavigationMetaData>();
+	if (MetaData.IsValid())
+	{
+		TSharedPtr<SWidget> Widget = MetaData->GetFocusRecipient(Type).Pin();
+		return FNavigationReply(MetaData->GetBoundaryRule(Type), Widget, MetaData->GetFocusDelegate(Type));
+	}
+	return FNavigationReply::Escape();
+}
+
+
 EWindowZone::Type SWidget::GetWindowZoneOverride() const
 {
 	// No special behavior.  Override this in derived widgets, if needed.
@@ -344,7 +406,6 @@ void SWidget::OnMouseCaptureLost()
 	
 }
 
-
 bool SWidget::FindChildGeometries( const FGeometry& MyGeometry, const TSet< TSharedRef<SWidget> >& WidgetsToFind, TMap<TSharedRef<SWidget>, FArrangedWidget>& OutResult ) const
 {
 	FindChildGeometries_Helper(MyGeometry, WidgetsToFind, OutResult);
diff --git a/Engine/Source/Runtime/SlateCore/Private/Widgets/SWindow.cpp b/Engine/Source/Runtime/SlateCore/Private/Widgets/SWindow.cpp
--- a/Engine/Source/Runtime/SlateCore/Private/Widgets/SWindow.cpp
+++ b/Engine/Source/Runtime/SlateCore/Private/Widgets/SWindow.cpp
@@ -1317,9 +1317,9 @@ bool SWindow::SupportsKeyboardFocus() const
 	return !bIsToolTipWindow && !bIsCursorDecoratorWindow;
 }
 
-FReply SWindow::OnKeyboardFocusReceived( const FGeometry& MyGeometry, const FKeyboardFocusEvent& InKeyboardFocusEvent )
+FReply SWindow::OnFocusReceived(const FGeometry& MyGeometry, const FFocusEvent& InFocusEvent)
 {
-	if( InKeyboardFocusEvent.GetCause() == EKeyboardFocusCause::WindowActivate  || FSlateApplicationBase::Get().IsExternalUIOpened() )
+	if (InFocusEvent.GetCause() == EFocusCause::WindowActivate || FSlateApplicationBase::Get().IsExternalUIOpened())
 	{
 		TArray< TSharedRef<SWindow> > JustThisWindow;
 		{
@@ -1335,7 +1335,7 @@ FReply SWindow::OnKeyboardFocusReceived( const FGeometry& MyGeometry, const FKey
 			FWidgetPath WidgetToFocusPath;
 			if( FSlateWindowHelper::FindPathToWidget( JustThisWindow, PinnedWidgetToFocusOnActivate.ToSharedRef(), WidgetToFocusPath ) )
 			{
-				FSlateApplicationBase::Get().SetKeyboardFocus( WidgetToFocusPath, EKeyboardFocusCause::SetDirectly );
+				FSlateApplicationBase::Get().SetKeyboardFocus( WidgetToFocusPath, EFocusCause::SetDirectly );
 			}
 		}
 		else
@@ -1344,7 +1344,7 @@ FReply SWindow::OnKeyboardFocusReceived( const FGeometry& MyGeometry, const FKey
 			if( FSlateWindowHelper::FindPathToWidget( JustThisWindow, AsShared(), WindowWidgetPath ) )
 			{
 				FWeakWidgetPath WeakWindowPath(WindowWidgetPath);
-				FSlateApplicationBase::Get().SetKeyboardFocus( WeakWindowPath.ToNextFocusedPath(EFocusMoveDirection::Next), EKeyboardFocusCause::SetDirectly );
+				FSlateApplicationBase::Get().SetKeyboardFocus( WeakWindowPath.ToNextFocusedPath(EUINavigation::Next), EFocusCause::SetDirectly );
 			}
 		}
 	}
diff --git a/Engine/Source/Runtime/SlateCore/Public/Application/SlateApplicationBase.h b/Engine/Source/Runtime/SlateCore/Public/Application/SlateApplicationBase.h
--- a/Engine/Source/Runtime/SlateCore/Public/Application/SlateApplicationBase.h
+++ b/Engine/Source/Runtime/SlateCore/Public/Application/SlateApplicationBase.h
@@ -266,7 +266,17 @@ public:
 	 * @param InCause The reason that keyboard focus is changing
 	 * @return true if the widget is now focused, false otherwise.
 	 */
-	virtual bool SetKeyboardFocus( const FWidgetPath& InFocusPath, const EKeyboardFocusCause::Type InCause ) = 0;
+	virtual bool SetKeyboardFocus( const FWidgetPath& InFocusPath, const EFocusCause InCause ) = 0;
+
+	/**
+	 * Sets user focus to the specified widget.  The widget must be allowed to receive focus.
+	 *
+	 * @param InUserIndex Index of the user that we want to change the focus of.
+	 * @param InWidget WidgetPath to the Widget to being focused.
+	 * @param InCause The reason that focus is changing.
+	 * @return true if the widget is now focused, false otherwise.
+	 */
+	virtual bool SetUserFocus(const uint32 InUserIndex, const FWidgetPath& InFocusPath, const EFocusCause InCause) = 0;
 
 public:
 	const static uint32 CursorPointerIndex;
diff --git a/Engine/Source/Runtime/SlateCore/Public/Input/Events.h b/Engine/Source/Runtime/SlateCore/Public/Input/Events.h
--- a/Engine/Source/Runtime/SlateCore/Public/Input/Events.h
+++ b/Engine/Source/Runtime/SlateCore/Public/Input/Events.h
@@ -9,40 +9,51 @@ class SWindow;
 
 
 /**
- * Context for keyboard focus change
- */
+* Context for focus change
+*/
 UENUM()
-namespace EKeyboardFocusCause
+enum class EFocusCause
 {
-	enum Type
-	{
-		/** Keyboard focus was changed because of a mouse action */
-		Mouse,
+	/** Focus was changed because of a mouse action */
+	Mouse,
 
-		/** Keyboard focus was changed in response to a keystroke, such as the TAB key. */
-		Keyboard,
+	/** Focus was changed in response to a navigation, such as the arrow keys, TAB key, controller DPad, ... */
+	Navigation,
 
-		/** Keyboard focus was changed because someone asked the application to change it. */
-		SetDirectly,
+	/** Focus was changed because someone asked the application to change it. */
+	SetDirectly,
 
-		/** Keyboard focus was explicitly cleared via the escape key or other similar action */
-		Cleared,
+	/** Focus was explicitly cleared via the escape key or other similar action */
+	Cleared,
 
-		/** Keyboard focus was changed because another widget lost focus, and focus moved to a new widget */
-		OtherWidgetLostFocus,
+	/** Focus was changed because another widget lost focus, and focus moved to a new widget */
+	OtherWidgetLostFocus,
 
-		/** Keyboard focus was set in response to the owning window being activated */
-		WindowActivate,
-	};
-}
+	/** Focus was set in response to the owning window being activated */
+	WindowActivate,
+};
 
+/** DEPRECATED 4.6 - Do not use */
+//@Todo slate: Remove this as soon as the 4.6 deprecated API is Removed.
+namespace EKeyboardFocusCause
+{
+	DEPRECATED(4.6, "EFocusCause::Type is deprecated and was renamed to EFocusCause. Please use that type instead.")
+	typedef ::EFocusCause Type;
+
+	const EFocusCause Mouse = EFocusCause::Mouse;
+	const EFocusCause Keyboard = EFocusCause::Navigation;
+	const EFocusCause SetDirectly = EFocusCause::SetDirectly;
+	const EFocusCause Cleared = EFocusCause::Cleared;
+	const EFocusCause OtherWidgetLostFocus = EFocusCause::OtherWidgetLostFocus;
+	const EFocusCause WindowActivate = EFocusCause::WindowActivate;
+}
 
 /**
- * FKeyboardFocusEvent is used when notifying widgets about keyboard focus changes
+ * FFocusEvent is used when notifying widgets about keyboard focus changes
  * It is passed to event handlers dealing with keyboard focus
  */
 USTRUCT()
-struct FKeyboardFocusEvent
+struct FFocusEvent
 {
 	GENERATED_USTRUCT_BODY()
 
@@ -51,8 +62,8 @@ public:
 	/**
 	 * UStruct Constructor.  Not meant for normal usage.
 	 */
-	FKeyboardFocusEvent()
-		: Cause(EKeyboardFocusCause::SetDirectly)
+	FFocusEvent()
+		: Cause(EFocusCause::SetDirectly)
 	{ }
 
 	/**
@@ -60,7 +71,7 @@ public:
 	 *
 	 * @param  InCause  The cause of the focus event
 	 */
-	FKeyboardFocusEvent( const EKeyboardFocusCause::Type InCause )
+	FFocusEvent(const EFocusCause InCause)
 		: Cause(InCause)
 	{ }
 
@@ -68,9 +79,9 @@ public:
 	/**
 	 * Queries the reason for the focus change
 	 *
-	 * @return  The cause of the keyboard focus change
+	 * @return  The cause of the focus change
 	 */
-	EKeyboardFocusCause::Type GetCause() const
+	EFocusCause GetCause() const
 	{
 		return Cause;
 	}
@@ -78,7 +89,18 @@ public:
 private:
 
 	/** The cause of the focus change */
-	EKeyboardFocusCause::Type Cause;
+	EFocusCause Cause;
+};
+
+//@Todo slate: Remove this as soon as the 4.6 deprecated API is Removed.
+USTRUCT()
+//struct DEPRECATED(4.6, "Use FFocusEvent") FKeyboardFocusEvent : public FFocusEvent
+struct FKeyboardFocusEvent : public FFocusEvent
+{
+	GENERATED_USTRUCT_BODY()
+public:
+	FKeyboardFocusEvent() { }
+	FKeyboardFocusEvent(const EFocusCause InCause) : FFocusEvent(InCause) { }
 };
 
 /**
@@ -95,7 +117,7 @@ struct FVirtualPointerPosition
 };
 
 /**
- * Base class for all mouse and keyboard events.
+ * Base class for all mouse and keyevents.
  */
 USTRUCT(BlueprintType)
 struct FInputEvent
@@ -109,6 +131,7 @@ public:
 	 */
 	FInputEvent()
 		: ModifierKeys(FModifierKeysState())
+		, UserIndex(0)
 		, bIsRepeat(false)
 		, EventPath(nullptr)
 	{ }
@@ -119,8 +142,9 @@ public:
 	 * @param  InModifierKeys  Modifier key state for this event
 	 * @param  bInIsRepeat  True if this key is an auto-repeated keystroke
 	 */
-	FInputEvent( const FModifierKeysState& InModifierKeys, const bool bInIsRepeat )
+	FInputEvent(const FModifierKeysState& InModifierKeys, const int32 InUserIndex, const bool bInIsRepeat)
 		: ModifierKeys(InModifierKeys)
+		, UserIndex(InUserIndex)
 		, bIsRepeat(bInIsRepeat)
 		, EventPath(nullptr)
 	{ }
@@ -272,6 +296,16 @@ public:
 		return ModifierKeys.AreCapsLocked();
 	}
 
+	/**
+	* Returns the index of the user that generated this event.
+	*
+	* @return The index of the user that caused the event
+	*/
+	uint32 GetUserIndex() const
+	{
+		return UserIndex;
+	}
+
 	/** The event path provides additional context for handling */
 	FGeometry FindGeometry( const TSharedRef<SWidget>& WidgetToFind ) const
 	{
@@ -296,6 +330,9 @@ protected:
 	// State of modifier keys when this event happened.
 	FModifierKeysState ModifierKeys;
 
+	// The index of the user that caused the event.
+	uint32 UserIndex;
+
 	// True if this key was auto-repeated.
 	bool bIsRepeat;
 
@@ -305,11 +342,11 @@ protected:
 
 
 /**
- * FKeyboardEvent describes a keyboard action (key pressed or released.)
- * It is passed to event handlers dealing with keyboard input.
+ * FKeyEvent describes a key action (keyboard/controller key/button pressed or released.)
+ * It is passed to event handlers dealing with key input.
  */
 USTRUCT(BlueprintType)
-struct FKeyboardEvent
+struct FKeyEvent
 	: public FInputEvent
 {
 	GENERATED_USTRUCT_BODY()
@@ -318,8 +355,8 @@ public:
 	/**
 	 * UStruct Constructor.  Not meant for normal usage.
 	 */
-	FKeyboardEvent()
-		: FInputEvent(FModifierKeysState(), false)
+	FKeyEvent()
+		: FInputEvent(FModifierKeysState(), 0, false)
 		, Key(EKeys::SpaceBar)
 		, CharacterCode(0)
 		, KeyCode(0)
@@ -333,18 +370,19 @@ public:
 	 * @param  InModifierKeys  Modifier key state for this event
 	 * @param  bInIsRepeat  True if this key is an auto-repeated keystroke
 	 */
-	FKeyboardEvent( const FKey InKey,
-					const FModifierKeysState& InModifierKeys, 
-					const bool bInIsRepeat,
-					const uint32 InCharacterCode,
-					const uint32 InKeyCode
+	FKeyEvent(	const FKey InKey,
+				const FModifierKeysState& InModifierKeys, 
+				const uint32 InUserIndex,
+				const bool bInIsRepeat,
+				const uint32 InCharacterCode,
+				const uint32 InKeyCode
 	)
-		: FInputEvent(InModifierKeys, bInIsRepeat)
+		: FInputEvent(InModifierKeys, InUserIndex, bInIsRepeat)
 		, Key(InKey)
 		, CharacterCode(InCharacterCode)
 		, KeyCode(InKeyCode)
 	{ }
-	
+
 	/**
 	 * Returns the name of the key for this event
 	 *
@@ -356,7 +394,7 @@ public:
 	}
 
 	/**
-	 * Returns the character code for this event
+	 * Returns the character code for this event.
 	 *
 	 * @return  Character code or 0 if this event was not a character key press
 	 */
@@ -365,6 +403,11 @@ public:
 		return CharacterCode;
 	}
 
+	/**
+	 * Returns the key code received from hardware before any conversion/mapping.
+	 *
+	 * @return  Key code received from hardware
+	 */
 	uint32 GetKeyCode() const
 	{
 		return KeyCode;
@@ -373,7 +416,6 @@ public:
 	SLATECORE_API virtual FText ToText() const override;
 
 private:
-
 	// Name of the key that was pressed.
 	FKey Key;
 
@@ -384,6 +426,78 @@ private:
 	uint32 KeyCode;
 };
 
+/** DEPRECATED 4.6 - Do not use */
+//@Todo slate: Remove this as soon as the 4.6 deprecated API is Removed.
+#define FKeyboardEvent \
+	FKeyEvent \
+	EMIT_DEPRECATED_WARNING_MESSAGE("FKeyboardEvent is deprecated in 4.6 and was renamed to FKeyEvent. Please use that type instead.")
+/*
+USTRUCT(BlueprintType)
+struct FKeyboardEvent : public FKeyEvent
+//struct DEPRECATED(4.6, "Use FKeyEvent") FKeyboardEvent : public FKeyEvent
+{
+	GENERATED_USTRUCT_BODY()
+public:
+	FKeyEvent() {}
+
+	FKeyEvent(const FKey InKey, const FModifierKeysState& InModifierKeys, const uint32 InUserIndex, const bool bInIsRepeat, const uint32 InCharacterCode, const uint32 InKeyCode)
+		: FKeyEvent(InKey, InModifierKeys, InUserIndex, bInIsRepeat, InCharacterCode, InKeyCode) {}
+}*/
+
+/**
+ * FAnalogEvent describes a analog key value.
+ * It is passed to event handlers dealing with analog keys.
+ */
+USTRUCT(BlueprintType)
+struct FAnalogInputEvent
+	: public FKeyEvent
+{
+	GENERATED_USTRUCT_BODY()
+
+public:
+	/**
+	* UStruct Constructor.  Not meant for normal usage.
+	*/
+	FAnalogInputEvent()
+		: FKeyEvent(FKey(EKeys::SpaceBar), FModifierKeysState(), false, 0, 0, 0)
+		, AnalogValue(0.0f)
+	{
+	}
+
+	/**
+	 * Constructor.  Events are immutable once constructed.
+	 *
+	 * @param  InKeyName  Character name
+	 * @param  InModifierKeys  Modifier key state for this event
+	 * @param  bInIsRepeat  True if this key is an auto-repeated keystroke
+	 */
+	FAnalogInputEvent(const FKey InKey,
+		const FModifierKeysState& InModifierKeys,
+		const uint32 InUserIndex,
+		const bool bInIsRepeat,
+		const uint32 InCharacterCode,
+		const uint32 InKeyCode,
+		const float InAnalogValue
+		)
+		: FKeyEvent(InKey, InModifierKeys, InUserIndex, bInIsRepeat, InCharacterCode, InKeyCode)
+		, AnalogValue(InAnalogValue)
+	{ }
+
+	/**
+	 * Returns the analog value between 0 and 1.
+	 * 0 being not pressed at all, 1 being fully pressed.
+	 * Non analog keys will only be 0 or 1.
+	 *
+	 * @return Analog value between 0 and 1.  1 being fully pressed, 0 being not pressed at all
+	 */
+	float GetAnalogValue() const { return AnalogValue; }
+
+	SLATECORE_API virtual FText ToText() const override;
+
+private:
+	//  Analog value between 0 and 1 (0 being not pressed at all, 1 being fully pressed).
+	float AnalogValue;
+};
 
 /**
  * FCharacterEvent describes a keyboard action where the utf-16 code is given.  Used for OnKeyChar messages
@@ -399,13 +513,13 @@ public:
 	 * UStruct Constructor.  Not meant for normal usage.
 	 */
 	FCharacterEvent()
-		: FInputEvent(FModifierKeysState(), false)
+		: FInputEvent(FModifierKeysState(), 0, false)
 		, Character(0)
 	{
 	}
 
-	FCharacterEvent( const TCHAR InCharacter, const FModifierKeysState& InModifierKeys, const bool bInIsRepeat )
-		: FInputEvent(InModifierKeys, bInIsRepeat)
+	FCharacterEvent(const TCHAR InCharacter, const FModifierKeysState& InModifierKeys, const uint32 InUserIndex, const bool bInIsRepeat)
+		: FInputEvent(InModifierKeys, InUserIndex, bInIsRepeat)
 		, Character(InCharacter)
 	{ }
 
@@ -476,7 +590,6 @@ public:
 		, CursorDelta(FVector2D(0, 0))
 		, PressedButtons(FTouchKeySet::EmptySet)
 		, EffectingButton(EKeys::LeftMouseButton)
-		, UserIndex(0)
 		, PointerIndex(0)
 		, TouchpadIndex(0)
 		, bIsTouchEvent(false)
@@ -494,13 +607,12 @@ public:
 		float InWheelDelta,
 		const FModifierKeysState& InModifierKeys
 	)
-		: FInputEvent(InModifierKeys, false)
+		: FInputEvent(InModifierKeys, 0, false)
 		, ScreenSpacePosition(InScreenSpacePosition)
 		, LastScreenSpacePosition(InLastScreenSpacePosition)
 		, CursorDelta(InScreenSpacePosition - InLastScreenSpacePosition)
 		, PressedButtons(InPressedButtons)
 		, EffectingButton(InEffectingButton)
-		, UserIndex(0)
 		, PointerIndex(InPointerIndex)
 		, TouchpadIndex(0)
 		, bIsTouchEvent(false)
@@ -517,12 +629,11 @@ public:
 		const TSet<FKey>& InPressedButtons,
 		const FModifierKeysState& InModifierKeys
 	)
-		: FInputEvent(InModifierKeys, false)
+		: FInputEvent(InModifierKeys, 0, false)
 		, ScreenSpacePosition(InScreenSpacePosition)
 		, LastScreenSpacePosition(InLastScreenSpacePosition)
 		, CursorDelta(InDelta)
 		, PressedButtons(InPressedButtons)
-		, UserIndex(0)
 		, PointerIndex(InPointerIndex)
 		, TouchpadIndex(0)
 		, bIsTouchEvent(false)
@@ -540,13 +651,12 @@ public:
 		const FModifierKeysState& InModifierKeys = FModifierKeysState(),
 		uint32 InTouchpadIndex=0
 	)
-		: FInputEvent(InModifierKeys, false)
+	: FInputEvent(InModifierKeys, InUserIndex, false)
 		, ScreenSpacePosition(InScreenSpacePosition)
 		, LastScreenSpacePosition(InLastScreenSpacePosition)
 		, CursorDelta(InScreenSpacePosition - InLastScreenSpacePosition)
 		, PressedButtons(bPressLeftMouseButton ? FTouchKeySet::StandardSet : FTouchKeySet::EmptySet)
 		, EffectingButton(EKeys::LeftMouseButton)
-		, UserIndex(InUserIndex)
 		, PointerIndex(InPointerIndex)
 		, TouchpadIndex(InTouchpadIndex)
 		, bIsTouchEvent(true)
@@ -563,12 +673,11 @@ public:
 		EGestureEvent::Type InGestureType,
 		const FVector2D& InGestureDelta
 	)
-		: FInputEvent(InModifierKeys, false)
+		: FInputEvent(InModifierKeys, 0, false)
 		, ScreenSpacePosition(InScreenSpacePosition)
 		, LastScreenSpacePosition(InLastScreenSpacePosition)
 		, CursorDelta(LastScreenSpacePosition - ScreenSpacePosition)
 		, PressedButtons(InPressedButtons)
-		, UserIndex(0)
 		, PointerIndex(0)
 		, bIsTouchEvent(false)
 		, GestureType(InGestureType)
@@ -648,7 +757,6 @@ private:
 	FVector2D CursorDelta;
 	const TSet<FKey>& PressedButtons;
 	FKey EffectingButton;
-	uint32 UserIndex;
 	uint32 PointerIndex;
 	uint32 TouchpadIndex;
 	bool bIsTouchEvent;
@@ -657,10 +765,8 @@ private:
 };
 
 
-/**
- * FControllerEvent describes a controller action (e.g. Button Press, Release, Analog stick move, etc).
- * It is passed to event handlers dealing with controller input.
- */
+/** DEPRECATED 4.6 - Do not use */
+//@Todo slate: Remove this as soon as the 4.6 deprecated API is Removed.
 USTRUCT(BlueprintType)
 struct FControllerEvent
 	: public FInputEvent
@@ -668,44 +774,25 @@ struct FControllerEvent
 	GENERATED_USTRUCT_BODY()
 public:
 
-	/**
-	 * UStruct Constructor.  Not meant for normal usage.
-	 */
 	FControllerEvent()
 		: EffectingButton(EKeys::Gamepad_RightTrigger)
-		, UserIndex(0)
 		, AnalogValue(0)
 	{ }
 
 	FControllerEvent( FKey InEffectingButton, int32 InUserIndex, float InAnalogValue, bool bIsRepeat )
-		: FInputEvent(FModifierKeysState(), bIsRepeat)
+		: FInputEvent(FModifierKeysState(), InUserIndex, bIsRepeat)
 		, EffectingButton(InEffectingButton)
-		, UserIndex(InUserIndex)
 		, AnalogValue(InAnalogValue)
 	{ }
 
 public:
-
-	/** @return The controller button that caused this event */
 	FKey GetEffectingButton() const { return EffectingButton; }
-
-	/** @return The index of the user that caused the event */
-	uint32 GetUserIndex() const { return UserIndex; }
-
-	/** @return Analog value between 0 and 1.  1 being fully pressed, 0 being not pressed at all */
 	float GetAnalogValue() const { return AnalogValue; }
 
 	SLATECORE_API virtual FText ToText() const override;
 
 private:
-
-	// The controller button that caused this event.
 	FKey EffectingButton;
-
-	// The index of the user that caused the event.
-	uint32 UserIndex;
-	
-	//  Analog value between 0 and 1 (1 being fully pressed, 0 being not pressed at all).
 	float AnalogValue;
 };
 
@@ -725,8 +812,7 @@ public:
 	* UStruct Constructor.  Not meant for normal usage.
 	*/
 	FMotionEvent()
-		: UserIndex(0)
-		, Tilt(FVector(0, 0, 0))
+		: Tilt(FVector(0, 0, 0))
 		, RotationRate(FVector(0, 0, 0))
 		, Gravity(FVector(0, 0, 0))
 		, Acceleration(FVector(0, 0, 0))
@@ -739,8 +825,7 @@ public:
 		const FVector& InGravity, 
 		const FVector& InAcceleration
 	)
-		: FInputEvent(FModifierKeysState(), false)
-		, UserIndex(InUserIndex)
+		: FInputEvent(FModifierKeysState(), InUserIndex, false)
 		, Tilt(InTilt)
 		, RotationRate(InRotationRate)
 		, Gravity(InGravity)
@@ -766,9 +851,6 @@ public:
 
 private:
 
-	// The index of the user that caused the event.
-	uint32 UserIndex;
-
 	// The current tilt of the device/controller.
 	FVector Tilt;
 
@@ -783,6 +865,40 @@ private:
 };
 
 
+/**
+* FNavigationEvent describes a navigation action (Left, Right, Up, Down)
+* It is passed to event handlers dealing with navigation.
+*/
+USTRUCT(BlueprintType)
+struct FNavigationEvent
+	: public FInputEvent
+{
+	GENERATED_USTRUCT_BODY()
+
+public:
+	/**
+	* UStruct Constructor.  Not meant for normal usage.
+	*/
+	FNavigationEvent()
+		: NavigationType(EUINavigation::Invalid)
+	{ }
+
+	FNavigationEvent(uint32 InUserIndex, EUINavigation InNavigationType)
+		: FInputEvent(FModifierKeysState(), InUserIndex, false)
+		, NavigationType(InNavigationType)
+	{ }
+
+public:
+
+	/** @return the type of navigation request (Left, Right, Up, Down) */
+	EUINavigation GetNavigationType() const { return NavigationType; }
+
+private:
+
+	// The navigation type
+	EUINavigation NavigationType;
+};
+
 /**
  * FWindowActivateEvent describes a window being activated or deactivated.
  * (i.e. brought to the foreground or moved to the background)
diff --git a/Engine/Source/Runtime/SlateCore/Public/Input/HittestGrid.h b/Engine/Source/Runtime/SlateCore/Public/Input/HittestGrid.h
--- a/Engine/Source/Runtime/SlateCore/Public/Input/HittestGrid.h
+++ b/Engine/Source/Runtime/SlateCore/Public/Input/HittestGrid.h
@@ -41,6 +41,17 @@ public:
 	int32 InsertWidget( const int32 ParentHittestIndex, const EVisibility& Visibility, const FArrangedWidget& Widget, const FVector2D InWindowOffset, const FSlateRect& InClippingRect );
 
 	void InsertCustomHitTestPath( TSharedRef<ICustomHitTestPath> CustomHitTestPath, int32 WidgetIndex );
+
+	/**
+	 * Finds the next focusable widget by searching through the hit test grid
+	 *
+	 * @param StartingWidget  This is the widget we are starting at, and navigating from.
+	 * @param Direction       The direction we should search in.
+	 * @param NavigationReply The Navigation Reply to specify a boundary rule for the search.
+	 * @param RuleWidget      The Widget that is applying the boundary rule, used to get the bounds of the Rule.
+	 */
+	TSharedPtr<SWidget> FindNextFocusableWidget(const FArrangedWidget& StartingWidget, const EUINavigation Direction, const FNavigationReply& NavigationReply, const FArrangedWidget& RuleWidget);
+
 private:
 
 	/**
@@ -66,6 +77,16 @@ private:
 	
 	friend class SWidgetReflector;
 
+	/** @returns true if Child is a descendant of Parent. */
+	bool IsDecendantOf(const TSharedRef<SWidget> Parent, const FCachedWidget& Child);
+
+	/** Utility function for searching for the next focusable widget. */
+	template<typename TCompareFunc, typename TSourceSideFunc, typename TDestSideFunc>
+	TSharedPtr<SWidget> FindFocusableWidget(const FSlateRect WidgetRect, const FSlateRect SweptRect, int32 AxisIndex, int32 Increment, const EUINavigation Direction, const FNavigationReply& NavigationReply, TCompareFunc CompareFunc, TSourceSideFunc SourceSideFunc, TDestSideFunc DestSideFunc);
+
+	/** Constrains a float position into the grid coordinate. */
+	FIntPoint GetCellCoordinate(FVector2D Position);
+
 	/** Access a cell at coordinates X, Y. Coordinates are row and column indexes. */
 	FCell& CellAt( const int32 X, const int32 Y );
 	const FCell& CellAt( const int32 X, const int32 Y ) const;
@@ -85,4 +106,4 @@ private:
 	void LogGrid() const;
 
 	static void LogChildren(int32 Index, int32 IndentLevel, const TArray<FCachedWidget>& WidgetsCachedThisFrame);
-};
\ No newline at end of file
+};
diff --git a/Engine/Source/Runtime/SlateCore/Public/Input/Reply.h b/Engine/Source/Runtime/SlateCore/Public/Input/Reply.h
--- a/Engine/Source/Runtime/SlateCore/Public/Input/Reply.h
+++ b/Engine/Source/Runtime/SlateCore/Public/Input/Reply.h
@@ -24,13 +24,6 @@ public:
 		return Me();
 	}
 
-	FReply& CaptureJoystick( TSharedRef<SWidget> InJoystickCaptor, bool bInAllJoysticks = false )
-	{
-		this->JoystickCaptor = InJoystickCaptor;
-		this->bAllJoysticks = bInAllJoysticks;
-		return Me();
-	}
-
 	/**
 	 * Enables the use of high precision (raw input) mouse movement, for more accurate mouse movement without mouse ballistics
 	 * NOTE: This implies mouse capture and hidden mouse movement.  Releasing mouse capture deactivates this mode.
@@ -51,14 +44,53 @@ public:
 		return Me();
 	}
 
-	/** An event should return FReply::Handled().SetKeyboardFocus( SomeWidget ) as a means of asking the system to set keyboard focus to the provided widget*/
-	FReply& SetKeyboardFocus( TSharedRef<SWidget> GiveMeFocus, EKeyboardFocusCause::Type ReasonFocusIsChanging )
+	/** An event should return FReply::Handled().SetUserFocus( SomeWidget ) as a means of asking the system to set users focus to the provided widget*/
+	FReply& SetUserFocus(TSharedRef<SWidget> GiveMeFocus, EFocusCause ReasonFocusIsChanging = EFocusCause::SetDirectly, bool bInAllUsers = false)
 	{
+		this->bSetUserFocus = true;
 		this->FocusRecipient = GiveMeFocus;
 		this->FocusChangeReason = ReasonFocusIsChanging;
-		return *this;
+		this->bReleaseUserFocus = false;
+		this->bAllUsers = bInAllUsers;
+		return Me();
 	}
-		
+	
+	DEPRECATED(4.6, "FReply::CaptureJoystick() is deprecated, use FReply::SetUserFocus() instead.")
+	FReply& CaptureJoystick(TSharedRef<SWidget> InJoystickCaptor, bool bInAllJoysticks = false)
+	{
+		return SetUserFocus(InJoystickCaptor, EFocusCause::SetDirectly, bInAllJoysticks);
+	}
+
+	DEPRECATED(4.6, "FReply::SetKeyboardFocus() is deprecated, use FReply::SetUserFocus() instead.")
+	FReply& SetKeyboardFocus(TSharedRef<SWidget> GiveMeFocus, EFocusCause ReasonFocusIsChanging)
+	{
+		return SetUserFocus(GiveMeFocus, ReasonFocusIsChanging, false);
+	}
+
+	/** An event should return a FReply::Handled().ClearUserFocus() to ask the system to clear user focus*/
+	FReply& ClearUserFocus(bool bInAllUsers = false)
+	{
+		this->FocusRecipient = nullptr;
+		this->FocusChangeReason = EFocusCause::Cleared;
+		this->bReleaseUserFocus = true;
+		this->bSetUserFocus = false;
+		this->bAllUsers = bInAllUsers;
+		return Me();
+	}
+
+	DEPRECATED(4.6, "FReply::ReleaseJoystickCapture() is deprecated, use FReply::ClearUserFocus() instead.")
+	FReply& ReleaseJoystickCapture(bool bInAllJoysticks = false)
+	{
+		return ClearUserFocus(bInAllJoysticks);
+	}
+
+	/** An event should return FReply::Handled().SetNavigation( NavigationType ) as a means of asking the system to attempt a navigation*/
+	FReply& SetNavigation(EUINavigation InNavigationType)
+	{
+		this->NavigationType = InNavigationType;
+		return Me();
+	}
+
 	/**
 	 * An event should return FReply::Handled().LockMouseToWidget( SomeWidget ) as a means of asking the system to 
 	 * Lock the mouse so it cannot move out of the bounds of the widget.
@@ -77,7 +109,7 @@ public:
 	{
 		this->bShouldReleaseMouseLock = true;
 		MouseLockWidget.Reset();
-		return *this;
+		return Me();
 	}
 
 	/** 
@@ -91,16 +123,6 @@ public:
 		return Me();
 	}
 
-	/**
-	 * An event should return a FReply::Handled().ReleaseJoystickCapture() to ask the system to release joystick capture
-	 */
-	FReply& ReleaseJoystickCapture(bool bInAllJoysticks = false)
-	{
-		this->bReleaseJoystickCapture = true;
-		this->bAllJoysticks = bInAllJoysticks;
-		return Me();
-	}
-
 	/**
 	 * Ask Slate to detect if a user started dragging in this widget.
 	 * If a drag is detected, Slate will send an OnDragDetected event.
@@ -147,11 +169,17 @@ public:
 	/** True if this reply indicated that we should release mouse capture as a result of the event being handled */
 	bool ShouldReleaseMouse() const { return bReleaseMouseCapture; }
 
-	/** true if this reply indicated that we should release joystick capture as a result of the event being handled */
-	bool ShouldReleaseJoystick() const { return bReleaseJoystickCapture; }
+	/** true if this reply indicated that we should set focus as a result of the event being handled */
+	bool ShouldSetUserFocus() const { return bSetUserFocus; }
 
-	/** If the event replied with a request to capture or release the joystick whether it should do it for all joysticks or just the one tied to the UserIndex */
-	bool AffectsAllJoysticks() const { return bAllJoysticks; }
+	/** true if this reply indicated that we should release focus as a result of the event being handled */
+	bool ShouldReleaseUserFocus() const { return bReleaseUserFocus; }
+
+	DEPRECATED(4.6, "FReply::ShouldReleaseJoystick() is deprecated, use FReply::ShouldReleaseUserFocus() instead.")
+	bool ShouldReleaseJoystick() const { return ShouldReleaseUserFocus(); }
+
+	/** If the event replied with a request to change the user focus whether it should do it for all users or just the current UserIndex */
+	bool AffectsAllUsers() const { return bAllUsers; }
 
 	/** True if this reply indicated that we should use high precision mouse movement */
 	bool ShouldUseHighPrecisionMouse() const { return bUseHighPrecisionMouse; }
@@ -165,17 +193,23 @@ public:
 	/** Returns the widget that the mouse should be locked to (if any) */
 	const TSharedPtr<SWidget>& GetMouseLockWidget() const { return MouseLockWidget; }
 
-	/** When not nullptr, keyboard focus has been requested to be set on the FocusRecipient. */
-	const TSharedPtr<SWidget>& GetFocusRecepient() const { return FocusRecipient; }
+	/** When not nullptr, user focus has been requested to be set on the FocusRecipient. */
+	const TSharedPtr<SWidget>& GetUserFocusRecepient() const { return FocusRecipient; }
+
+	DEPRECATED(4.6, "FReply::GetFocusRecepient() is deprecated, use FReply::GetUserFocusRecepient() instead.")
+	const TSharedPtr<SWidget>& GetFocusRecepient() const { return GetUserFocusRecepient(); }
+
+	DEPRECATED(4.6, "FReply::GetJoystickCaptor() is deprecated, use FReply::GetUserFocusRecepient() instead.")
+	const TSharedPtr<SWidget>& GetJoystickCaptor() const { return GetUserFocusRecepient(); }
 
 	/** Get the reason that a focus change is being requested. */
-	EKeyboardFocusCause::Type GetFocusCause() const { return FocusChangeReason; }
+	EFocusCause GetFocusCause() const { return FocusChangeReason; }
 
 	/** If the event replied with a request to capture the mouse, this returns the desired mouse captor. Otherwise returns an invalid pointer. */
 	const TSharedPtr<SWidget>& GetMouseCaptor() const { return MouseCaptor; }	
 
-	/** If the event replied with a request to capture the joystick, this returns the desired joystick captor. Otherwise returns an invalid pointer. */
-	const TSharedPtr<SWidget>& GetJoystickCaptor() const { return JoystickCaptor; }
+	/** Get the navigation type (Invalid if no navigation is requested). */
+	EUINavigation GetNavigationType() const { return NavigationType; }
 
 	/** @return the Content that we should use for the Drag and Drop operation; Invalid SharedPtr if a drag and drop operation is not requested*/
 	const TSharedPtr<FDragDropOperation>& GetDragDropContent() const { return DragDropContent; }
@@ -220,14 +254,15 @@ private:
 		, RequestedMousePos()
 		, EventHandler(nullptr)
 		, MouseCaptor(nullptr)
-		, JoystickCaptor(nullptr)
 		, FocusRecipient(nullptr)
 		, MouseLockWidget(nullptr)
 		, DragDropContent(nullptr)
-		, FocusChangeReason(EKeyboardFocusCause::SetDirectly)
+		, FocusChangeReason(EFocusCause::SetDirectly)
+		, NavigationType(EUINavigation::Invalid)
 		, bReleaseMouseCapture(false)
-		, bReleaseJoystickCapture(false)
-		, bAllJoysticks(false)
+		, bSetUserFocus(false)
+		, bReleaseUserFocus(false)
+		, bAllUsers(false)
 		, bShouldReleaseMouseLock(false)
 		, bUseHighPrecisionMouse(false)
 		, bPreventThrottling(false)
@@ -243,16 +278,17 @@ private:
 	TOptional<FIntPoint> RequestedMousePos;
 	TSharedPtr<SWidget> EventHandler;
 	TSharedPtr<SWidget> MouseCaptor;
-	TSharedPtr<SWidget> JoystickCaptor;
 	TSharedPtr<SWidget> FocusRecipient;
 	TSharedPtr<SWidget> MouseLockWidget;
 	TSharedPtr<SWidget> DetectDragForWidget;
 	FKey DetectDragForMouseButton;
 	TSharedPtr<FDragDropOperation> DragDropContent;
-	EKeyboardFocusCause::Type FocusChangeReason;
+	EFocusCause FocusChangeReason;
+	EUINavigation NavigationType;
 	uint32 bReleaseMouseCapture:1;
-	uint32 bReleaseJoystickCapture:1;
-	uint32 bAllJoysticks:1;
+	uint32 bSetUserFocus:1;
+	uint32 bReleaseUserFocus : 1;
+	uint32 bAllUsers:1;
 	uint32 bShouldReleaseMouseLock:1;
 	uint32 bUseHighPrecisionMouse:1;
 	uint32 bPreventThrottling:1;
diff --git a/Engine/Source/Runtime/SlateCore/Public/Layout/WidgetPath.h b/Engine/Source/Runtime/SlateCore/Public/Layout/WidgetPath.h
--- a/Engine/Source/Runtime/SlateCore/Public/Layout/WidgetPath.h
+++ b/Engine/Source/Runtime/SlateCore/Public/Layout/WidgetPath.h
@@ -87,7 +87,7 @@ public:
 	{
 		const FArrangedWidget& LastWidget = Widgets.Last();
 		
-		FArrangedChildren Extension = GeneratePathToWidget( Matcher, LastWidget, EFocusMoveDirection::Next, VisibilityFilter );
+		FArrangedChildren Extension = GeneratePathToWidget(Matcher, LastWidget, EUINavigation::Next, VisibilityFilter);
 
 		for( int32 WidgetIndex=0; WidgetIndex < Extension.Num(); ++WidgetIndex )
 		{
@@ -107,11 +107,11 @@ public:
 	 * @return A path from FromWidget to WidgetToFind; will not include FromWidget.
 	 */
 	template<typename MatcherType>
-	FArrangedChildren GeneratePathToWidget( const MatcherType& Matcher, const FArrangedWidget& FromWidget, EFocusMoveDirection::Type Direction = EFocusMoveDirection::Next, EVisibility VisibilityFilter = EVisibility::Visible )
+	FArrangedChildren GeneratePathToWidget(const MatcherType& Matcher, const FArrangedWidget& FromWidget, EUINavigation NavigationType = EUINavigation::Next, EVisibility VisibilityFilter = EVisibility::Visible)
 	{
 		FArrangedChildren PathResult(VisibilityFilter);
 
-		if ( Direction == EFocusMoveDirection::Next )
+		if (NavigationType == EUINavigation::Next)
 		{
 			SearchForWidgetRecursively( Matcher, FromWidget, PathResult, VisibilityFilter );
 		}
@@ -136,7 +136,7 @@ public:
 	 *
 	 * @return true if the focus moved successfully, false if we were unable to move focus
 	 */
-	bool MoveFocus(int32 PathLevel, EFocusMoveDirection::Type MoveDirection);
+	bool MoveFocus(int32 PathLevel, EUINavigation NavigationType);
 
 
 	FArrangedChildren Widgets;
@@ -216,7 +216,7 @@ public:
 	 * 
 	 * @return The new focus path.
 	 */
-	FWidgetPath ToNextFocusedPath(EFocusMoveDirection::Type MoveDirection);
+	FWidgetPath ToNextFocusedPath(EUINavigation NavigationType);
 	
 	/** Get the last (leaf-most) widget in this path; assumes path is valid */
 	TWeakPtr< SWidget > GetLastWidget() const
diff --git a/Engine/Source/Runtime/SlateCore/Public/Logging/IEventLogger.h b/Engine/Source/Runtime/SlateCore/Public/Logging/IEventLogger.h
--- a/Engine/Source/Runtime/SlateCore/Public/Logging/IEventLogger.h
+++ b/Engine/Source/Runtime/SlateCore/Public/Logging/IEventLogger.h
@@ -26,6 +26,7 @@ namespace EEventLog
 		KeyDown,
 		KeyUp,
 		KeyChar,
+		AnalogInput,
 		UICommand,
 		BeginTransaction,
 		EndTransaction,
@@ -59,6 +60,7 @@ namespace EEventLog
 		case KeyDown: return FString("KeyDown");
 		case KeyUp: return FString("KeyUp");
 		case KeyChar: return FString("KeyChar");
+		case AnalogInput: return FString("AnalogInput");
 		case UICommand: return FString("UICommand");
 		case BeginTransaction: return FString("BeginTransaction");
 		case EndTransaction: return FString("EndTransaction");
diff --git a/Engine/Source/Runtime/SlateCore/Public/Rendering/RenderingCommon.h b/Engine/Source/Runtime/SlateCore/Public/Rendering/RenderingCommon.h
--- a/Engine/Source/Runtime/SlateCore/Public/Rendering/RenderingCommon.h
+++ b/Engine/Source/Runtime/SlateCore/Public/Rendering/RenderingCommon.h
@@ -347,11 +347,11 @@ public:
 	 * Called by Slate when a key is pressed inside the viewport
 	 *
 	 * @param MyGeometry	Information about the location and size of the viewport
-	 * @param MouseEvent	Information about the keyboard event
+	 * @param MouseEvent	Information about the key event
 	 *
 	 * @return Whether the event was handled along with possible requests for the system to take action.
 	 */
-	virtual FReply OnKeyDown( const FGeometry& MyGeometry, const FKeyboardEvent& InKeyboardEvent )
+	virtual FReply OnKeyDown( const FGeometry& MyGeometry, const FKeyEvent& InKeyEvent )
 	{
 		return FReply::Unhandled();
 	}
@@ -360,64 +360,46 @@ public:
 	 * Called by Slate when a key is released inside the viewport
 	 *
 	 * @param MyGeometry	Information about the location and size of the viewport
-	 * @param MouseEvent	Information about the keyboard event
+	 * @param MouseEvent	Information about the key event
 	 *
 	 * @return Whether the event was handled along with possible requests for the system to take action.
 	 */
-	virtual FReply OnKeyUp( const FGeometry& MyGeometry, const FKeyboardEvent& InKeyboardEvent )
+	virtual FReply OnKeyUp( const FGeometry& MyGeometry, const FKeyEvent& InKeyEvent )
 	{
 		return FReply::Unhandled();
 	}
 
 	/**
-	 * Called by Slate when a character key is pressed while the viewport has keyboard focus
+	 * Called when an analog value changes on a button that supports analog
 	 *
-	 * @param MyGeometry	Information about the location and size of the viewport
-	 * @param MouseEvent	Information about the character that was pressed
-	 *
-	 * @return Whether the event was handled along with possible requests for the system to take action.
+	 * @param MyGeometry The Geometry of the widget receiving the event
+	 * @param InAnalogInputEvent Analog input event
+	 * @return Returns whether the event was handled, along with other possible actions
 	 */
-	virtual FReply OnKeyChar( const FGeometry& MyGeometry, const FCharacterEvent& InCharacterEvent )
+	virtual FReply OnAnalogValueChanged( const FGeometry& MyGeometry, const FAnalogInputEvent& InAnalogInputEvent )
 	{
 		return FReply::Unhandled();
 	}
 
 	/**
-	 * Called when the viewport gains keyboard focus.  
+	 * Called by Slate when a character key is pressed while the viewport has focus
 	 *
-	 * @param InKeyboardFocusEvent	Information about what caused the viewport to gain focus
-	 */
-	virtual FReply OnKeyboardFocusReceived( const FKeyboardFocusEvent& InKeyboardFocusEvent )
-	{
-		return FReply::Unhandled();
-	}
-
-	/**
-	 * Called when a controller button is pressed
-	 * 
-	 * @param ControllerEvent	The controller event generated
-	 */
-	virtual FReply OnControllerButtonPressed( const FGeometry& MyGeometry, const FControllerEvent& ControllerEvent )
-	{
-		return FReply::Unhandled();
-	}
-
-	/**
-	 * Called when a controller button is released
-	 * 
-	 * @param ControllerEvent	The controller event generated
+	 * @param MyGeometry	Information about the location and size of the viewport
+	 * @param MouseEvent	Information about the character that was pressed
+	 *
+	 * @return Whether the event was handled along with possible requests for the system to take action.
 	 */
-	virtual FReply OnControllerButtonReleased( const FGeometry& MyGeometry, const FControllerEvent& ControllerEvent )
+	virtual FReply OnKeyChar( const FGeometry& MyGeometry, const FCharacterEvent& InCharacterEvent )
 	{
 		return FReply::Unhandled();
 	}
 
 	/**
-	 * Called when an analog value on a controller changes
-	 * 
-	 * @param ControllerEvent	The controller event generated
+	 * Called when the viewport gains keyboard focus.  
+	 *
+	 * @param InFocusEvent	Information about what caused the viewport to gain focus
 	 */
-	virtual FReply OnControllerAnalogValueChanged( const FGeometry& MyGeometry, const FControllerEvent& ControllerEvent )
+	virtual FReply OnFocusReceived( const FFocusEvent& InFocusEvent )
 	{
 		return FReply::Unhandled();
 	}
@@ -464,10 +446,11 @@ public:
 	
 	/**
 	 * Called when motion is detected (controller or device)
+	 * e.g. Someone tilts or shakes their controller.
 	 * 
-	 * @param ControllerEvent	The controller event generated
+	 * @param InMotionEvent	The motion event generated
 	 */
-	virtual FReply OnMotionDetected( const FGeometry& MyGeometry, const FMotionEvent& MotionEvent )
+	virtual FReply OnMotionDetected( const FGeometry& MyGeometry, const FMotionEvent& InMotionEvent )
 	{
 		return FReply::Unhandled();
 	}
@@ -479,12 +462,23 @@ public:
 	{
 	}
 
+	/**
+	 * Called when navigation is requested
+	 * e.g. Left Joystick, Direction Pad, Arrow Keys can generate navigation events.
+	 * 
+	 * @param InNavigationEvent	The navigation event generated
+	 */
+	virtual FNavigationReply OnNavigation( const FGeometry& MyGeometry, const FNavigationEvent& InNavigationEvent )
+	{
+		return FNavigationReply::Stop();
+	}
+
 	/**
 	 * Called when the viewport loses keyboard focus.  
 	 *
-	 * @param InKeyboardFocusEvent	Information about what caused the viewport to lose focus
+	 * @param InFocusEvent	Information about what caused the viewport to lose focus
 	 */
-	virtual void OnKeyboardFocusLost( const FKeyboardFocusEvent& InKeyboardFocusEvent )
+	virtual void OnFocusLost( const FFocusEvent& InFocusEvent )
 	{
 	}
 
@@ -512,4 +506,4 @@ public:
 	 * @param RenderTarget	handle to the platform specific render target implementation.  Note this is already bound by Slate initially 
 	 */
 	virtual void DrawRenderThread(class FRHICommandListImmediate& RHICmdList, const void* RenderTarget) = 0;
-};
\ No newline at end of file
+};
diff --git a/Engine/Source/Runtime/SlateCore/Public/SlateCore.h b/Engine/Source/Runtime/SlateCore/Public/SlateCore.h
--- a/Engine/Source/Runtime/SlateCore/Public/SlateCore.h
+++ b/Engine/Source/Runtime/SlateCore/Public/SlateCore.h
@@ -132,6 +132,8 @@ DECLARE_MEMORY_STAT_EXTERN(TEXT("Index Buffer Memory"), STAT_SlateIndexBufferMem
 #include "Events.h"
 #include "DragAndDrop.h"
 #include "Reply.h"
+#include "NavigationReply.h"
+#include "NavigationMetaData.h"
 
 // Rendering
 #include "RenderingCommon.h"
diff --git a/Engine/Source/Runtime/SlateCore/Public/Types/NavigationMetaData.h b/Engine/Source/Runtime/SlateCore/Public/Types/NavigationMetaData.h
--- a/Engine/Source/Runtime/SlateCore/Public/Types/NavigationMetaData.h
+++ b/Engine/Source/Runtime/SlateCore/Public/Types/NavigationMetaData.h
@@ -14,15 +14,103 @@ public:
 
 	FNavigationMetaData()
 	{
+		for (auto Rule : Rules)
+		{
+			Rule.BoundaryRule = EUINavigationRule::Escape;
+			Rule.FocusRecipient = nullptr;
+		}
 	}
 
+	/**
+	 * Get the boundary rule  for the provided navigation type
+	 * 
+	 * @param InNavigation the type of navigation to get the boundary rule for
+	 * @return the boundary rule for the provided navigation type
+	 */
+	EUINavigationRule GetBoundaryRule(EUINavigation InNavigation) const
+	{
+		return Rules[(uint8)InNavigation].BoundaryRule;
+	}
+
+	/**
+	 * Get the focus recipient for the provided navigation type
+	 * 
+	 * @param InNavigation the type of navigation to get the focus recipient for
+	 * @return the focus recipient for the provided navigation type
+	 */
+	const TWeakPtr<SWidget>& GetFocusRecipient(EUINavigation InNavigation) const
+	{
+		return Rules[(uint8)InNavigation].FocusRecipient;
+	}
+
+	/**
+	 * Get the focus recipient delegate for the provided navigation type
+	 * 
+	 * @param InNavigation the type of navigation to get the focus recipient delegate for
+	 * @return the focus recipient delegate for the provided navigation type
+	 */
+	const FNavigationDelegate& GetFocusDelegate(EUINavigation InNavigation) const
+	{
+		return Rules[(uint8)InNavigation].FocusDelegate;
+	}
+
+	/**
+	 * Set the navigation behavior for the provided navigation type to be explicit, using a constant widget
+	 * 
+	 * @param InNavigation the type of navigation to set explicit
+	 * @param InFocusRecipient the widget used when the system requests the destination for the explicit navigation
+	 */
+	void SetNavigationExplicit(EUINavigation InNavigation, TSharedPtr<SWidget> InFocusRecipient)
+	{
+		Rules[(uint8)InNavigation].BoundaryRule = EUINavigationRule::Explicit;
+		Rules[(uint8)InNavigation].FocusRecipient = InFocusRecipient;
+	}
+
+	/**
+	 * Set the navigation behavior for the provided navigation type to be a custom delegate
+	 * 
+	 * @param InNavigation the type of navigation to set custom
+	 * @param InFocusRecipient the delegate used when the system requests the destination for the custom navigation
+	 */
+	void SetNavigationCustom(EUINavigation InNavigation, FNavigationDelegate InFocusDelegate)
+	{
+		Rules[(uint8)InNavigation].BoundaryRule = EUINavigationRule::Custom;
+		Rules[(uint8)InNavigation].FocusDelegate = InFocusDelegate;
+	}
+
+	/**
+	* Set the navigation behavior for the provided navigation type to be wrap
+	*
+	* @param InNavigation the type of navigation to set wrapping 
+	*/
+	void SetNavigationWrap(EUINavigation InNavigation)
+	{
+		Rules[(uint8)InNavigation].BoundaryRule = EUINavigationRule::Wrap;
+	}
+
+	/**
+	 * An event should return a FNavigationReply::Explicit() to let the system know to stop at the bounds of this widget.
+	 */
+	void SetNavigationStop(EUINavigation InNavigation)
+	{
+		Rules[(uint8)InNavigation].BoundaryRule = EUINavigationRule::Stop;
+	}
+
+	/**
+	 * An event should return a FNavigationReply::Escape() to let the system know that a navigation can escape the bounds of this widget.
+	 */
+	void SetNavigationEscape(EUINavigation InNavigation)
+	{
+		Rules[(uint8)InNavigation].BoundaryRule = EUINavigationRule::Escape;
+	}
 private:
-	/*
+
 	struct SNavData
 	{
 		EUINavigationRule BoundaryRule;
-		TAttribute<TSharedPtr<SWidget>> FocusRecipient;
+		TWeakPtr<SWidget> FocusRecipient;
+		FNavigationDelegate FocusDelegate;
 	};
-	SNavData Rules[EUINavigation::Num];
-	*/
+	SNavData Rules[(uint8)EUINavigation::Num];
+
 };
diff --git a/Engine/Source/Runtime/SlateCore/Public/Types/SlateEnums.h b/Engine/Source/Runtime/SlateCore/Public/Types/SlateEnums.h
--- a/Engine/Source/Runtime/SlateCore/Public/Types/SlateEnums.h
+++ b/Engine/Source/Runtime/SlateCore/Public/Types/SlateEnums.h
@@ -56,23 +56,54 @@ namespace EButtonTouchMethod
 	};
 }
 
-
 UENUM(BlueprintType)
-namespace EFocusMoveDirection
+namespace EButtonPressMethod
 {
 	/**
-	 * Enumerates directions in which the focus can move.
-	 */
+	* Enumerates different methods that a button can be triggered with keyboard/controller. Normally, DownAndUp is appropriate.
+	*/
 	enum Type
 	{
-		/** Move to next widget. */
-		Next,
+		/**
+		* User must press the button, then release while the button has focus to trigger the click.
+		* This is the most common type of button.
+		*/
+		DownAndUp,
+
+		/**
+		* Click will be triggered immediately on button press.
+		*/
+		ButtonPress,
 
-		/** Move focus to previous widget. */
-		Previous
+		/**
+		* Click will always be triggered when a button release occurs on the focused button,
+		* even if the button wasn't pressed while focused.
+		*/
+		ButtonRelease,
 	};
 }
 
+/**
+ * Navigation context for event
+ */
+UENUM(BlueprintType)
+enum class EUINavigation : uint8
+{
+	/** Four cardinal directions*/
+	Left,
+	Right,
+	Up,
+	Down,
+
+	/** Conceptual next and previous*/
+	Next,
+	Previous,
+
+	/** Number of navigation types*/
+	Num,
+	/** Denotes an invalid navigation, more important used to denote no specified navigation*/
+	Invalid
+};
 
 /**
  * Enumerates horizontal alignment options, i.e. for widget slots.
diff --git a/Engine/Source/Runtime/SlateCore/Public/Widgets/SWidget.h b/Engine/Source/Runtime/SlateCore/Public/Widgets/SWidget.h
--- a/Engine/Source/Runtime/SlateCore/Public/Widgets/SWidget.h
+++ b/Engine/Source/Runtime/SlateCore/Public/Widgets/SWidget.h
@@ -131,27 +131,36 @@ public:
 	virtual bool OnHitTest( const FGeometry& MyGeometry, FVector2D InAbsoluteCursorPosition ){return false;}
 
 	//
-	// KEYBOARD INPUT
+	// KEY INPUT
 	//
 
 	/**
-	 * Called when keyboard focus is given to this widget.  This event does not bubble.
+	 * Called when focus is given to this widget.  This event does not bubble.
 	 *
 	 * @param MyGeometry The Geometry of the widget receiving the event
-	 * @param  InKeyboardFocusEvent  KeyboardFocusEvent
+	 * @param  InFocusEvent  The FocusEvent
 	 * @return  Returns whether the event was handled, along with other possible actions
 	 */
-	virtual FReply OnKeyboardFocusReceived( const FGeometry& MyGeometry, const FKeyboardFocusEvent& InKeyboardFocusEvent );
+	virtual FReply OnFocusReceived(const FGeometry& MyGeometry, const FFocusEvent& InFocusEvent);
+
+	DEPRECATED(4.6, "SWidget::OnKeyboardFocusReceived() is deprecated, implement SWidget::OnFocusReceived() instead.")
+	virtual FReply OnKeyboardFocusReceived(const FGeometry& MyGeometry, const FKeyboardFocusEvent& InFocusEvent);
 
 	/**
-	 * Called when this widget loses the keyboard focus.  This event does not bubble.
+	 * Called when this widget loses focus.  This event does not bubble.
 	 *
-	 * @param  InKeyboardFocusEvent  KeyboardFocusEvent
+	 * @param InFocusEvent The FocusEvent
 	 */
-	virtual void OnKeyboardFocusLost( const FKeyboardFocusEvent& InKeyboardFocusEvent );
+	virtual void OnFocusLost(const FFocusEvent& InFocusEvent);
+
+	DEPRECATED(4.6, "SWidget::OnKeyboardFocusLost() is deprecated, implement SWidget::OnFocusLost() instead.")
+	virtual void OnKeyboardFocusLost(const FKeyboardFocusEvent& InFocusEvent);
 
 	/** Called whenever a focus path is changing on all the widgets within the old and new focus paths */
-	virtual void OnKeyboardFocusChanging( const FWeakWidgetPath& PreviousFocusPath, const FWidgetPath& NewWidgetPath );
+	virtual void OnFocusChanging(const FWeakWidgetPath& PreviousFocusPath, const FWidgetPath& NewWidgetPath);
+
+	DEPRECATED(4.6, "SWidget::OnKeyboardFocusChanging() is deprecated, implement SWidget::OnFocusChanging() instead.")
+	virtual void OnKeyboardFocusChanging(const FWeakWidgetPath& PreviousFocusPath, const FWidgetPath& NewWidgetPath);
 
 	/**
 	 * Called after a character is entered while this widget has keyboard focus
@@ -163,35 +172,53 @@ public:
 	virtual FReply OnKeyChar( const FGeometry& MyGeometry,const FCharacterEvent& InCharacterEvent );
 
 	/**
-	 * Called after a key is pressed when this widget or a child of this widget has keyboard focus
+	 * Called after a key is pressed when this widget or a child of this widget has focus
 	 * If a widget handles this event, OnKeyDown will *not* be passed to the focused widget.
 	 *
 	 * This event is primarily to allow parent widgets to consume an event before a child widget processes
 	 * it and it should be used only when there is no better design alternative.
 	 *
 	 * @param MyGeometry The Geometry of the widget receiving the event
-	 * @param  InKeyboardEvent  Keyboard event
-	 * @return  Returns whether the event was handled, along with other possible actions
+	 * @param InKeyEvent  Key event
+	 * @return Returns whether the event was handled, along with other possible actions
 	 */
-	virtual FReply OnPreviewKeyDown( const FGeometry& MyGeometry, const FKeyboardEvent& InKeyboardEvent );
+	virtual FReply OnPreviewKeyDown( const FGeometry& MyGeometry, const FKeyEvent& InKeyEvent );
 
 	/**
-	 * Called after a key is pressed when this widget has keyboard focus (this event bubbles if not handled)
+	 * Called after a key is pressed when this widget has focus (this event bubbles if not handled)
 	 *
 	 * @param MyGeometry The Geometry of the widget receiving the event
-	 * @param  InKeyboardEvent  Keyboard event
-	 * @return  Returns whether the event was handled, along with other possible actions
+	 * @param InKeyEvent  Key event
+	 * @return Returns whether the event was handled, along with other possible actions
 	 */
-	virtual FReply OnKeyDown( const FGeometry& MyGeometry, const FKeyboardEvent& InKeyboardEvent );
+	virtual FReply OnKeyDown( const FGeometry& MyGeometry, const FKeyEvent& InKeyEvent );
+
+	DEPRECATED(4.6, "SWidget::OnControllerButtonPressed() is deprecated, SWidget::OnKeyDown() now handles controller input as well as controller.")
+	virtual FReply OnControllerButtonPressed(const FGeometry& MyGeometry, const FControllerEvent& ControllerEvent);
 
 	/**
-	 * Called after a key is released when this widget has keyboard focus
+	 * Called after a key is released when this widget has focus
 	 *
 	 * @param MyGeometry The Geometry of the widget receiving the event
-	 * @param  InKeyboardEvent  Keyboard event
-	 * @return  Returns whether the event was handled, along with other possible actions
+	 * @param InKeyEvent  Key event
+	 * @return Returns whether the event was handled, along with other possible actions
+	 */
+	virtual FReply OnKeyUp( const FGeometry& MyGeometry, const FKeyEvent& InKeyEvent );
+
+	DEPRECATED(4.6, "SWidget::OnControllerButtonReleased() is deprecated, SWidget::OnKeyUp() now handles controller input as well as controller.")
+	virtual FReply OnControllerButtonReleased(const FGeometry& MyGeometry, const FControllerEvent& ControllerEvent);
+
+	/**
+	 * Called when an analog value changes on a button that supports analog
+	 *
+	 * @param MyGeometry The Geometry of the widget receiving the event
+	 * @param InAnalogInputEvent Analog input event
+	 * @return Returns whether the event was handled, along with other possible actions
 	 */
-	virtual FReply OnKeyUp( const FGeometry& MyGeometry, const FKeyboardEvent& InKeyboardEvent );
+	virtual FReply OnAnalogValueChanged( const FGeometry& MyGeometry, const FAnalogInputEvent& InAnalogInputEvent );
+
+	DEPRECATED(4.6, "SWidget::OnControllerAnalogValueChanged() is deprecated, SWidget::OnAnalogValueChanged() handles keyboard and controller input.")
+	virtual FReply OnControllerAnalogValueChanged(const FGeometry& MyGeometry, const FControllerEvent& ControllerEvent);
 
 	//
 	// MOUSE INPUT
@@ -344,27 +371,6 @@ public:
 	 * @return A reply that indicated whether this event was handled.
 	 */
 	virtual FReply OnDrop( const FGeometry& MyGeometry, const FDragDropEvent& DragDropEvent );
-	
-	/**
-	 * Called when a controller button is pressed
-	 * 
-	 * @param ControllerEvent	The controller event generated
-	 */
-	virtual FReply OnControllerButtonPressed( const FGeometry& MyGeometry, const FControllerEvent& ControllerEvent );
-
-	/**
-	 * Called when a controller button is released
-	 * 
-	 * @param ControllerEvent	The controller event generated
-	 */
-	virtual FReply OnControllerButtonReleased( const FGeometry& MyGeometry, const FControllerEvent& ControllerEvent );
-
-	/**
-	 * Called when an analog value on a controller changes
-	 * 
-	 * @param ControllerEvent	The controller event generated
-	 */
-	virtual FReply OnControllerAnalogValueChanged( const FGeometry& MyGeometry, const FControllerEvent& ControllerEvent );
 
 	/**
 	 * Called when the user performs a gesture on trackpad. This event is bubbled.
@@ -399,7 +405,7 @@ public:
 	 * Called when motion is detected (controller or device)
 	 * e.g. Someone tilts or shakes their controller.
 	 * 
-	 * @param MotionEvent	The motion event generated
+	 * @param InMotionEvent	The motion event generated
 	 */
 	virtual FReply OnMotionDetected( const FGeometry& MyGeometry, const FMotionEvent& InMotionEvent );
 
@@ -420,6 +426,14 @@ public:
 	 */
 	virtual void OnFinishedKeyInput();
 
+	/**
+	 * Called when navigation is requested
+	 * e.g. Left Joystick, Direction Pad, Arrow Keys can generate navigation events.
+	 * 
+	 * @param InNavigationEvent	The navigation event generated
+	 */
+	virtual FNavigationReply OnNavigation(const FGeometry& MyGeometry, const FNavigationEvent& InNavigationEvent);
+
 	/**
 	 * Called when the mouse is moved over the widget's window, to determine if we should report whether
 	 * OS-specific features should be active at this location (such as a title bar grip, system menu, etc.)
@@ -668,6 +682,16 @@ public:
 		return FoundMetaData;
 	}
 
+	/**
+	 * Add metadata to this widget.
+	 * @param AddMe the metadata to add to the widget.
+	 */
+	template<typename MetaDataType>
+	void AddMetadata(const TSharedRef<MetaDataType>& AddMe)
+	{
+		MetaData.Add(AddMe);
+	}
+
 public:
 
 	// Widget Inspector and debugging methods
diff --git a/Engine/Source/Runtime/SlateCore/Public/Widgets/SWindow.h b/Engine/Source/Runtime/SlateCore/Public/Widgets/SWindow.h
--- a/Engine/Source/Runtime/SlateCore/Public/Widgets/SWindow.h
+++ b/Engine/Source/Runtime/SlateCore/Public/Widgets/SWindow.h
@@ -582,7 +582,7 @@ public:
 
 	virtual bool SupportsKeyboardFocus() const override;
 
-	virtual FReply OnKeyboardFocusReceived( const FGeometry& MyGeometry, const FKeyboardFocusEvent& InKeyboardFocusEvent );
+	virtual FReply OnFocusReceived( const FGeometry& MyGeometry, const FFocusEvent& InFocusEvent );
 
 	virtual FReply OnMouseButtonDown( const FGeometry& MyGeometry, const FPointerEvent& MouseEvent ) override;
 	virtual FReply OnMouseButtonUp( const FGeometry& MyGeometry, const FPointerEvent& MouseEvent ) override;
diff --git a/Engine/Source/Runtime/UMG/Private/Components/Widget.cpp b/Engine/Source/Runtime/UMG/Private/Components/Widget.cpp
--- a/Engine/Source/Runtime/UMG/Private/Components/Widget.cpp
+++ b/Engine/Source/Runtime/UMG/Private/Components/Widget.cpp
@@ -466,6 +466,18 @@ void UWidget::SynchronizeProperties()
 	{
 		SafeWidget->SetToolTipText(OPTIONAL_BINDING(FText, ToolTipText));
 	}
+
+	if (Navigation != nullptr)
+	{
+		TSharedPtr<FNavigationMetaData> MetaData = SafeWidget->GetMetaData<FNavigationMetaData>();
+		if (!MetaData.IsValid())
+		{
+			MetaData = MakeShareable(new FNavigationMetaData());
+			SafeWidget->AddMetadata(MetaData.ToSharedRef());
+		}
+
+		Navigation->UpdateMetaData(MetaData.ToSharedRef());
+	}
 }
 
 bool UWidget::IsDesignTime() const
diff --git a/Engine/Source/Runtime/UMG/Private/Slate/SObjectWidget.cpp b/Engine/Source/Runtime/UMG/Private/Slate/SObjectWidget.cpp
--- a/Engine/Source/Runtime/UMG/Private/Slate/SObjectWidget.cpp
+++ b/Engine/Source/Runtime/UMG/Private/Slate/SObjectWidget.cpp
@@ -65,25 +65,25 @@ bool SObjectWidget::SupportsKeyboardFocus() const
 	return false;
 }
 
-FReply SObjectWidget::OnKeyboardFocusReceived(const FGeometry& MyGeometry, const FKeyboardFocusEvent& InKeyboardFocusEvent)
+FReply SObjectWidget::OnFocusReceived(const FGeometry& MyGeometry, const FFocusEvent& InFocusEvent)
 {
 	if ( WidgetObject && !WidgetObject->IsDesignTime() )
 	{
-		return WidgetObject->OnKeyboardFocusReceived(MyGeometry, InKeyboardFocusEvent).NativeReply;
+		return WidgetObject->OnFocusReceived(MyGeometry, InFocusEvent).NativeReply;
 	}
 
 	return FReply::Unhandled();
 }
 
-void SObjectWidget::OnKeyboardFocusLost(const FKeyboardFocusEvent& InKeyboardFocusEvent)
+void SObjectWidget::OnFocusLost(const FFocusEvent& InFocusEvent)
 {
 	if ( WidgetObject && !WidgetObject->IsDesignTime() )
 	{
-		WidgetObject->OnKeyboardFocusLost(InKeyboardFocusEvent);
+		WidgetObject->OnFocusLost(InFocusEvent);
 	}
 }
 
-void SObjectWidget::OnKeyboardFocusChanging(const FWeakWidgetPath& PreviousFocusPath, const FWidgetPath& NewWidgetPath)
+void SObjectWidget::OnFocusChanging(const FWeakWidgetPath& PreviousFocusPath, const FWidgetPath& NewWidgetPath)
 {
 	// TODO UMG
 }
@@ -98,31 +98,41 @@ FReply SObjectWidget::OnKeyChar(const FGeometry& MyGeometry, const FCharacterEve
 	return FReply::Unhandled();
 }
 
-FReply SObjectWidget::OnPreviewKeyDown(const FGeometry& MyGeometry, const FKeyboardEvent& InKeyboardEvent)
+FReply SObjectWidget::OnPreviewKeyDown(const FGeometry& MyGeometry, const FKeyEvent& InKeyEvent)
 {
 	if ( WidgetObject && !WidgetObject->IsDesignTime() )
 	{
-		return WidgetObject->OnPreviewKeyDown(MyGeometry, InKeyboardEvent).NativeReply;
+		return WidgetObject->OnPreviewKeyDown(MyGeometry, InKeyEvent).NativeReply;
 	}
 
 	return FReply::Unhandled();
 }
 
-FReply SObjectWidget::OnKeyDown(const FGeometry& MyGeometry, const FKeyboardEvent& InKeyboardEvent)
+FReply SObjectWidget::OnKeyDown(const FGeometry& MyGeometry, const FKeyEvent& InKeyEvent)
 {
 	if ( WidgetObject && !WidgetObject->IsDesignTime() )
 	{
-		return WidgetObject->OnKeyDown(MyGeometry, InKeyboardEvent).NativeReply;
+		return WidgetObject->OnKeyDown(MyGeometry, InKeyEvent).NativeReply;
 	}
 
 	return FReply::Unhandled();
 }
 
-FReply SObjectWidget::OnKeyUp(const FGeometry& MyGeometry, const FKeyboardEvent& InKeyboardEvent)
+FReply SObjectWidget::OnKeyUp(const FGeometry& MyGeometry, const FKeyEvent& InKeyEvent)
 {
 	if ( WidgetObject && !WidgetObject->IsDesignTime() )
 	{
-		return WidgetObject->OnKeyUp(MyGeometry, InKeyboardEvent).NativeReply;
+		return WidgetObject->OnKeyUp(MyGeometry, InKeyEvent).NativeReply;
+	}
+
+	return FReply::Unhandled();
+}
+
+FReply SObjectWidget::OnAnalogValueChanged(const FGeometry& MyGeometry, const FAnalogInputEvent& InAnalogInputEvent)
+{
+	if (WidgetObject && !WidgetObject->IsDesignTime())
+	{
+		return WidgetObject->OnAnalogValueChanged(MyGeometry, InAnalogInputEvent).NativeReply;
 	}
 
 	return FReply::Unhandled();
@@ -281,36 +291,6 @@ void SObjectWidget::OnDragCancelled(const FDragDropEvent& DragDropEvent, UDragDr
 	}
 }
 
-FReply SObjectWidget::OnControllerButtonPressed(const FGeometry& MyGeometry, const FControllerEvent& ControllerEvent)
-{
-	if ( WidgetObject && !WidgetObject->IsDesignTime() )
-	{
-		return WidgetObject->OnControllerButtonPressed(MyGeometry, ControllerEvent).NativeReply;
-	}
-
-	return FReply::Unhandled();
-}
-
-FReply SObjectWidget::OnControllerButtonReleased(const FGeometry& MyGeometry, const FControllerEvent& ControllerEvent)
-{
-	if ( WidgetObject && !WidgetObject->IsDesignTime() )
-	{
-		return WidgetObject->OnControllerButtonReleased(MyGeometry, ControllerEvent).NativeReply;
-	}
-
-	return FReply::Unhandled();
-}
-
-FReply SObjectWidget::OnControllerAnalogValueChanged(const FGeometry& MyGeometry, const FControllerEvent& ControllerEvent)
-{
-	if ( WidgetObject && !WidgetObject->IsDesignTime() )
-	{
-		return WidgetObject->OnControllerAnalogValueChanged(MyGeometry, ControllerEvent).NativeReply;
-	}
-
-	return FReply::Unhandled();
-}
-
 FReply SObjectWidget::OnTouchGesture(const FGeometry& MyGeometry, const FPointerEvent& GestureEvent)
 {
 	if ( WidgetObject && !WidgetObject->IsDesignTime() )
diff --git a/Engine/Source/Runtime/UMG/Private/UserWidget.cpp b/Engine/Source/Runtime/UMG/Private/UserWidget.cpp
--- a/Engine/Source/Runtime/UMG/Private/UserWidget.cpp
+++ b/Engine/Source/Runtime/UMG/Private/UserWidget.cpp
@@ -173,11 +173,21 @@ void UUserWidget::OnPaint_Implementation(FPaintContext& Context) const
 
 }
 
+FEventReply UUserWidget::OnFocusReceived_Implementation(FGeometry MyGeometry, FFocusEvent InFocusEvent)
+{
+	return UWidgetBlueprintLibrary::Unhandled();
+}
+
 FEventReply UUserWidget::OnKeyboardFocusReceived_Implementation(FGeometry MyGeometry, FKeyboardFocusEvent InKeyboardFocusEvent)
 {
 	return UWidgetBlueprintLibrary::Unhandled();
 }
 
+void UUserWidget::OnFocusLost_Implementation(FFocusEvent InFocusEvent)
+{
+
+}
+
 void UUserWidget::OnKeyboardFocusLost_Implementation(FKeyboardFocusEvent InKeyboardFocusEvent)
 {
 
@@ -188,17 +198,22 @@ FEventReply UUserWidget::OnKeyChar_Implementation(FGeometry MyGeometry, FCharact
 	return UWidgetBlueprintLibrary::Unhandled();
 }
 
-FEventReply UUserWidget::OnPreviewKeyDown_Implementation(FGeometry MyGeometry, FKeyboardEvent InKeyboardEvent)
+FEventReply UUserWidget::OnPreviewKeyDown_Implementation(FGeometry MyGeometry, FKeyEvent InKeyEvent)
+{
+	return UWidgetBlueprintLibrary::Unhandled();
+}
+
+FEventReply UUserWidget::OnKeyDown_Implementation(FGeometry MyGeometry, FKeyEvent InKeyEvent)
 {
 	return UWidgetBlueprintLibrary::Unhandled();
 }
 
-FEventReply UUserWidget::OnKeyDown_Implementation(FGeometry MyGeometry, FKeyboardEvent InKeyboardEvent)
+FEventReply UUserWidget::OnKeyUp_Implementation(FGeometry MyGeometry, FKeyEvent InKeyEvent)
 {
 	return UWidgetBlueprintLibrary::Unhandled();
 }
 
-FEventReply UUserWidget::OnKeyUp_Implementation(FGeometry MyGeometry, FKeyboardEvent InKeyboardEvent)
+FEventReply UUserWidget::OnAnalogValueChanged_Implementation(FGeometry MyGeometry, FAnalogInputEvent InAnalogInputEvent)
 {
 	return UWidgetBlueprintLibrary::Unhandled();
 }
diff --git a/Engine/Source/Runtime/UMG/Private/WidgetBlueprintLibrary.cpp b/Engine/Source/Runtime/UMG/Private/WidgetBlueprintLibrary.cpp
--- a/Engine/Source/Runtime/UMG/Private/WidgetBlueprintLibrary.cpp
+++ b/Engine/Source/Runtime/UMG/Private/WidgetBlueprintLibrary.cpp
@@ -40,7 +40,7 @@ UDragDropOperation* UWidgetBlueprintLibrary::CreateDragDropOperation(TSubclassOf
 	}
 }
 
-void UWidgetBlueprintLibrary::SetInputMode_UIOnly(APlayerController* Target, UWidget* InWidgetToFocus, bool bLockMouseToViewport)
+void UWidgetBlueprintLibrary::SetInputMode_UIOnly(APlayerController*const& Target, UWidget* InWidgetToFocus, bool bLockMouseToViewport)
 {
 	if (Target != nullptr)
 	{
@@ -55,7 +55,7 @@ void UWidgetBlueprintLibrary::SetInputMode_UIOnly(APlayerController* Target, UWi
 	}
 }
 
-void UWidgetBlueprintLibrary::SetInputMode_GameAndUI(APlayerController* Target, UWidget* InWidgetToFocus, bool bLockMouseToViewport, bool bHideCursorDuringCapture)
+void UWidgetBlueprintLibrary::SetInputMode_GameAndUI(APlayerController*const& Target, UWidget* InWidgetToFocus, bool bLockMouseToViewport, bool bHideCursorDuringCapture)
 {
 	if (Target != nullptr)
 	{
@@ -71,7 +71,7 @@ void UWidgetBlueprintLibrary::SetInputMode_GameAndUI(APlayerController* Target,
 	}
 }
 
-void UWidgetBlueprintLibrary::SetInputMode_GameOnly(APlayerController* Target)
+void UWidgetBlueprintLibrary::SetInputMode_GameOnly(APlayerController*const& Target)
 {
 	if (Target != nullptr)
 	{
@@ -82,7 +82,7 @@ void UWidgetBlueprintLibrary::SetInputMode_GameOnly(APlayerController* Target)
 
 void UWidgetBlueprintLibrary::SetFocusToGameViewport()
 {
-	FSlateApplication::Get().SetFocusToGameViewport();
+	FSlateApplication::Get().SetAllUserFocusToGameViewport();
 }
 
 void UWidgetBlueprintLibrary::DrawBox(UPARAM(ref) FPaintContext& Context, FVector2D Position, FVector2D Size, USlateBrushAsset* Brush, FLinearColor Tint)
@@ -176,27 +176,37 @@ FEventReply UWidgetBlueprintLibrary::ReleaseMouseCapture(UPARAM(ref) FEventReply
 	return Reply;
 }
 
-FEventReply UWidgetBlueprintLibrary::CaptureJoystick(UPARAM(ref) FEventReply& Reply, UWidget* CapturingWidget, bool bInAllJoysticks/* = false*/)
+FEventReply UWidgetBlueprintLibrary::SetUserFocus(UPARAM(ref) FEventReply& Reply, UWidget* FocusWidget, bool bInAllUsers/* = false*/)
 {
-	if ( CapturingWidget )
+	if (FocusWidget)
 	{
-		TSharedPtr<SWidget> CapturingSlateWidget = CapturingWidget->GetCachedWidget();
-		if ( CapturingSlateWidget.IsValid() )
+		TSharedPtr<SWidget> CapturingSlateWidget = FocusWidget->GetCachedWidget();
+		if (CapturingSlateWidget.IsValid())
 		{
-			Reply.NativeReply = Reply.NativeReply.CaptureJoystick(CapturingSlateWidget.ToSharedRef(), bInAllJoysticks);
+			Reply.NativeReply = Reply.NativeReply.SetUserFocus(CapturingSlateWidget.ToSharedRef(), EFocusCause::SetDirectly, bInAllUsers);
 		}
 	}
 
 	return Reply;
 }
 
-FEventReply UWidgetBlueprintLibrary::ReleaseJoystickCapture(UPARAM(ref) FEventReply& Reply, bool bInAllJoysticks /*= false*/)
+FEventReply UWidgetBlueprintLibrary::CaptureJoystick(UPARAM(ref) FEventReply& Reply, UWidget* CapturingWidget, bool bInAllJoysticks/* = false*/)
 {
-	Reply.NativeReply = Reply.NativeReply.ReleaseJoystickCapture(bInAllJoysticks);
+	return SetUserFocus(Reply, CapturingWidget, bInAllJoysticks);
+}
+
+FEventReply UWidgetBlueprintLibrary::ClearUserFocus(UPARAM(ref) FEventReply& Reply, bool bInAllUsers /*= false*/)
+{
+	Reply.NativeReply = Reply.NativeReply.ClearUserFocus(bInAllUsers);
 
 	return Reply;
 }
 
+FEventReply UWidgetBlueprintLibrary::ReleaseJoystickCapture(UPARAM(ref) FEventReply& Reply, bool bInAllJoysticks /*= false*/)
+{
+	return ClearUserFocus(Reply, bInAllJoysticks);
+}
+
 FEventReply UWidgetBlueprintLibrary::SetMousePosition(UPARAM(ref) FEventReply& Reply, FVector2D NewMousePosition)
 {
 	FIntPoint NewPoint((int32)NewMousePosition.X, (int32)NewMousePosition.Y);
diff --git a/Engine/Source/Runtime/UMG/Public/Blueprint/UserWidget.h b/Engine/Source/Runtime/UMG/Public/Blueprint/UserWidget.h
--- a/Engine/Source/Runtime/UMG/Public/Blueprint/UserWidget.h
+++ b/Engine/Source/Runtime/UMG/Public/Blueprint/UserWidget.h
@@ -232,63 +232,81 @@ public:
 	 * Called when keyboard focus is given to this widget.  This event does not bubble.
 	 *
 	 * @param MyGeometry The Geometry of the widget receiving the event
-	 * @param InKeyboardFocusEvent  KeyboardFocusEvent
+	 * @param InFocusEvent  FocusEvent
 	 * @return  Returns whether the event was handled, along with other possible actions
 	 */
-	UFUNCTION(BlueprintNativeEvent, Category="Keyboard")
+	UFUNCTION(BlueprintNativeEvent, Category="Input")
+	FEventReply OnFocusReceived(FGeometry MyGeometry, FFocusEvent InFocusEvent);
+
+	UFUNCTION(BlueprintNativeEvent, meta = (DeprecatedFunction, DeprecationMessage = "Use OnFocusReceived() instead"), Category = "Keyboard")
 	FEventReply OnKeyboardFocusReceived(FGeometry MyGeometry, FKeyboardFocusEvent InKeyboardFocusEvent);
 
+
 	/**
-	 * Called when this widget loses the keyboard focus.  This event does not bubble.
+	 * Called when this widget loses focus.  This event does not bubble.
 	 *
-	 * @param  InKeyboardFocusEvent  KeyboardFocusEvent
+	 * @param  InFocusEvent  FocusEvent
 	 */
-	UFUNCTION(BlueprintNativeEvent, Category="Keyboard")
+	UFUNCTION(BlueprintNativeEvent, Category="Input")
+	void OnFocusLost(FFocusEvent InFocusEvent);
+
+	UFUNCTION(BlueprintNativeEvent, meta = (DeprecatedFunction, DeprecationMessage = "Use OnFocusLost() instead"), Category = "Keyboard")
 	void OnKeyboardFocusLost(FKeyboardFocusEvent InKeyboardFocusEvent);
 
+
 	/**
-	 * Called after a character is entered while this widget has keyboard focus
+	 * Called after a character is entered while this widget has focus
 	 *
 	 * @param MyGeometry The Geometry of the widget receiving the event
 	 * @param  InCharacterEvent  Character event
 	 * @return  Returns whether the event was handled, along with other possible actions
 	 */
-	UFUNCTION(BlueprintNativeEvent, Category="Keyboard")
+	UFUNCTION(BlueprintNativeEvent, Category="Input")
 	FEventReply OnKeyChar(FGeometry MyGeometry, FCharacterEvent InCharacterEvent);
 
 	/**
-	 * Called after a key is pressed when this widget or a child of this widget has keyboard focus
+	 * Called after a key (keyboard, controller, ...) is pressed when this widget or a child of this widget has focus
 	 * If a widget handles this event, OnKeyDown will *not* be passed to the focused widget.
 	 *
 	 * This event is primarily to allow parent widgets to consume an event before a child widget processes
 	 * it and it should be used only when there is no better design alternative.
 	 *
 	 * @param MyGeometry The Geometry of the widget receiving the event
-	 * @param  InKeyboardEvent  Keyboard event
+	 * @param  InKeyEvent  Key event
 	 * @return  Returns whether the event was handled, along with other possible actions
 	 */
-	UFUNCTION(BlueprintNativeEvent, Category="Keyboard")
-	FEventReply OnPreviewKeyDown(FGeometry MyGeometry, FKeyboardEvent InKeyboardEvent);
+	UFUNCTION(BlueprintNativeEvent, Category="Input")
+	FEventReply OnPreviewKeyDown(FGeometry MyGeometry, FKeyEvent InKeyEvent);
 
 	/**
-	 * Called after a key is pressed when this widget has keyboard focus (this event bubbles if not handled)
+	 * Called after a key (keyboard, controller, ...) is pressed when this widget has focus (this event bubbles if not handled)
 	 *
 	 * @param MyGeometry The Geometry of the widget receiving the event
-	 * @param  InKeyboardEvent  Keyboard event
+	 * @param  InKeyEvent  Key event
 	 * @return  Returns whether the event was handled, along with other possible actions
 	 */
-	UFUNCTION(BlueprintNativeEvent, Category="Keyboard")
-	FEventReply OnKeyDown(FGeometry MyGeometry, FKeyboardEvent InKeyboardEvent);
+	UFUNCTION(BlueprintNativeEvent, Category="Input")
+	FEventReply OnKeyDown(FGeometry MyGeometry, FKeyEvent InKeyEvent);
 
 	/**
-	 * Called after a key is released when this widget has keyboard focus
+	 * Called after a key (keyboard, controller, ...) is released when this widget has focus
 	 *
 	 * @param MyGeometry The Geometry of the widget receiving the event
-	 * @param  InKeyboardEvent  Keyboard event
+	 * @param  InKeyEvent  Key event
 	 * @return  Returns whether the event was handled, along with other possible actions
 	 */
-	UFUNCTION(BlueprintNativeEvent, Category="Keyboard")
-	FEventReply OnKeyUp(FGeometry MyGeometry, FKeyboardEvent InKeyboardEvent);
+	UFUNCTION(BlueprintNativeEvent, Category="Input")
+	FEventReply OnKeyUp(FGeometry MyGeometry, FKeyEvent InKeyEvent);
+
+	/**
+	* Called when an analog value changes on a button that supports analog
+	*
+	* @param MyGeometry The Geometry of the widget receiving the event
+	* @param  InAnalogInputEvent  Analog Event
+	* @return  Returns whether the event was handled, along with other possible actions
+	*/
+	UFUNCTION(BlueprintNativeEvent, Category = "Input")
+	FEventReply OnAnalogValueChanged(FGeometry MyGeometry, FAnalogInputEvent InAnalogInputEvent);
 
 	/**
 	 * The system calls this method to notify the widget that a mouse button was pressed within it. This event is bubbled.
@@ -440,31 +458,14 @@ public:
 	UFUNCTION(BlueprintNativeEvent, Category="Drag and Drop")
 	bool OnDrop(FGeometry MyGeometry, FPointerEvent PointerEvent, UDragDropOperation* Operation);
 
-	/**
-	 * Called when a controller button is pressed
-	 * 
-	 * @param MyGeometry        The geometry of the widget receiving the event.
-	 * @param ControllerEvent	The controller event generated
-	 */
-	UFUNCTION(BlueprintNativeEvent, Category="Gamepad Input")
+
+	UFUNCTION(BlueprintNativeEvent, meta = (DeprecatedFunction, DeprecationMessage = "Use OnKeyDown() instead"), Category = "Gamepad Input")
 	FEventReply OnControllerButtonPressed(FGeometry MyGeometry, FControllerEvent ControllerEvent);
 
-	/**
-	 * Called when a controller button is released
-	 * 
-	 * @param MyGeometry        The geometry of the widget receiving the event.
-	 * @param ControllerEvent	The controller event generated
-	 */
-	UFUNCTION(BlueprintNativeEvent, Category="Gamepad Input")
+	UFUNCTION(BlueprintNativeEvent, meta = (DeprecatedFunction, DeprecationMessage = "Use OnKeyUp() instead"), Category = "Gamepad Input")
 	FEventReply OnControllerButtonReleased(FGeometry MyGeometry, FControllerEvent ControllerEvent);
 
-	/**
-	 * Called when an analog value on a controller changes, like a joystick.
-	 * 
-	 * @param MyGeometry        The geometry of the widget receiving the event.
-	 * @param ControllerEvent	The controller event generated
-	 */
-	UFUNCTION(BlueprintNativeEvent, Category="Gamepad Input")
+	UFUNCTION(BlueprintNativeEvent, meta = (DeprecatedFunction, DeprecationMessage = "Use OnAnalogValueChanged() instead"), Category = "Gamepad Input")
 	FEventReply OnControllerAnalogValueChanged(FGeometry MyGeometry, FControllerEvent ControllerEvent);
 
 	/**
diff --git a/Engine/Source/Runtime/UMG/Public/Blueprint/WidgetBlueprintLibrary.h b/Engine/Source/Runtime/UMG/Public/Blueprint/WidgetBlueprintLibrary.h
--- a/Engine/Source/Runtime/UMG/Public/Blueprint/WidgetBlueprintLibrary.h
+++ b/Engine/Source/Runtime/UMG/Public/Blueprint/WidgetBlueprintLibrary.h
@@ -28,15 +28,15 @@ class UWidgetBlueprintLibrary : public UBlueprintFunctionLibrary
 	
 	/** Setup an input mode that allows only the UI to respond to user input. */
 	UFUNCTION(BlueprintCallable, Category = "Input")
-	static void SetInputMode_UIOnly(APlayerController* Target, UWidget* InWidgetToFocus = nullptr, bool bLockMouseToViewport = false);
+	static void SetInputMode_UIOnly(APlayerController*const& Target, UWidget* InWidgetToFocus = nullptr, bool bLockMouseToViewport = false);
 
 	/** Setup an input mode that allows only the UI to respond to user input, and if the UI doesn't handle it player input / player controller gets a chance. */
 	UFUNCTION(BlueprintCallable, Category = "Input")
-	static void SetInputMode_GameAndUI(APlayerController* Target, UWidget* InWidgetToFocus = nullptr, bool bLockMouseToViewport = false, bool bHideCursorDuringCapture = true);
+	static void SetInputMode_GameAndUI(APlayerController*const& Target, UWidget* InWidgetToFocus = nullptr, bool bLockMouseToViewport = false, bool bHideCursorDuringCapture = true);
 
 	/** Setup an input mode that allows only player input / player controller to respond to user input. */
 	UFUNCTION(BlueprintCallable, Category = "Input")
-	static void SetInputMode_GameOnly(APlayerController* Target);
+	static void SetInputMode_GameOnly(APlayerController*const& Target);
 
 
 	UFUNCTION(BlueprintCallable, Category="Focus")
@@ -86,11 +86,17 @@ class UWidgetBlueprintLibrary : public UBlueprintFunctionLibrary
 	static FEventReply ReleaseMouseCapture(UPARAM(ref) FEventReply& Reply);
 
 	/**  */
-	UFUNCTION(BlueprintPure, meta=( HidePin="CapturingWidget", DefaultToSelf="CapturingWidget" ), Category="Widget|Event Reply")
+	UFUNCTION(BlueprintPure, meta= (HidePin="CapturingWidget", DefaultToSelf="CapturingWidget"), Category="Widget|Event Reply")
+	static FEventReply SetUserFocus(UPARAM(ref) FEventReply& Reply, UWidget* FocusWidget, bool bInAllUsers = false);
+
+	UFUNCTION(BlueprintPure, meta = (DeprecatedFunction, DeprecationMessage = "Use SetUserFocus() instead"), Category = "Widget|Event Reply")
 	static FEventReply CaptureJoystick(UPARAM(ref) FEventReply& Reply, UWidget* CapturingWidget, bool bInAllJoysticks = false);
 
 	/**  */
-	UFUNCTION(BlueprintPure, Category="Widget|Event Reply")
+	UFUNCTION(BlueprintPure, meta = (HidePin = "CapturingWidget", DefaultToSelf = "CapturingWidget"), Category = "Widget|Event Reply")
+	static FEventReply ClearUserFocus(UPARAM(ref) FEventReply& Reply, bool bInAllUsers = false);
+
+	UFUNCTION(BlueprintPure, meta = (DeprecatedFunction, DeprecationMessage = "Use ClearUserFocus() instead"), Category = "Widget|Event Reply")
 	static FEventReply ReleaseJoystickCapture(UPARAM(ref) FEventReply& Reply, bool bInAllJoysticks = false);
 
 	/**  */
diff --git a/Engine/Source/Runtime/UMG/Public/Blueprint/WidgetNavigation.h b/Engine/Source/Runtime/UMG/Public/Blueprint/WidgetNavigation.h
--- a/Engine/Source/Runtime/UMG/Public/Blueprint/WidgetNavigation.h
+++ b/Engine/Source/Runtime/UMG/Public/Blueprint/WidgetNavigation.h
@@ -2,19 +2,22 @@
 
 #pragma once
 
+#include "NavigationReply.h"
+#include "NavigationMetaData.h"
+
 #include "WidgetNavigation.generated.h"
 
 /**
  *
  */
 USTRUCT()
-struct FWidgetNavigationData
+struct UMG_API FWidgetNavigationData
 {
 	GENERATED_USTRUCT_BODY()
 
 public:
-	//UPROPERTY(EditAnywhere, BlueprintReadOnly, Category="Navigation")
-	//TEnumAsByte<> Rule;
+	UPROPERTY(EditAnywhere, BlueprintReadOnly, Category="Navigation")
+	TEnumAsByte<EUINavigationRule> Rule;
 
 	UPROPERTY(EditAnywhere, BlueprintReadOnly, Category="Navigation")
 	FName WidgetToFocus;
@@ -44,4 +47,46 @@ public:
 	/**  */
 	UPROPERTY(EditAnywhere, BlueprintReadOnly, Category="Navigation")
 	FWidgetNavigationData Right;
+
+	EUINavigationRule GetNavigationRule(EUINavigation Nav)
+	{
+		switch (Nav)
+		{
+		case EUINavigation::Up:
+			return Up.Rule;
+		case EUINavigation::Down:
+			return Down.Rule;
+		case EUINavigation::Left:
+			return Left.Rule;
+		case EUINavigation::Right:
+			return Right.Rule;
+		}
+		return EUINavigationRule::Escape;
+	}
+
+	void UpdateMetaData(TSharedRef<FNavigationMetaData> MetaData)
+	{
+		UpdateMetaDataEntry(MetaData, Up, EUINavigation::Up);
+		UpdateMetaDataEntry(MetaData, Down, EUINavigation::Down);
+		UpdateMetaDataEntry(MetaData, Left, EUINavigation::Left);
+		UpdateMetaDataEntry(MetaData, Right, EUINavigation::Right);
+	}
+
+private:
+
+	void UpdateMetaDataEntry(TSharedRef<FNavigationMetaData> MetaData, const FWidgetNavigationData & NavData, EUINavigation Nav)
+	{
+		switch (NavData.Rule)
+		{
+			case EUINavigationRule::Escape:
+				MetaData->SetNavigationEscape(Nav);
+				break;
+			case EUINavigationRule::Stop:
+				MetaData->SetNavigationStop(Nav);
+				break;
+			case EUINavigationRule::Wrap:
+				MetaData->SetNavigationWrap(Nav);
+				break;
+		}
+	}
 };
diff --git a/Engine/Source/Runtime/UMG/Public/Slate/SObjectWidget.h b/Engine/Source/Runtime/UMG/Public/Slate/SObjectWidget.h
--- a/Engine/Source/Runtime/UMG/Public/Slate/SObjectWidget.h
+++ b/Engine/Source/Runtime/UMG/Public/Slate/SObjectWidget.h
@@ -33,14 +33,15 @@ class SObjectWidget : public SCompoundWidget, public FGCObject
 
 	virtual bool SupportsKeyboardFocus() const override;
 
-	virtual FReply OnKeyboardFocusReceived(const FGeometry& MyGeometry, const FKeyboardFocusEvent& InKeyboardFocusEvent) override;
-	virtual void OnKeyboardFocusLost(const FKeyboardFocusEvent& InKeyboardFocusEvent) override;
-	virtual void OnKeyboardFocusChanging(const FWeakWidgetPath& PreviousFocusPath, const FWidgetPath& NewWidgetPath) override;
+	virtual FReply OnFocusReceived(const FGeometry& MyGeometry, const FFocusEvent& InFocusEvent) override;
+	virtual void OnFocusLost(const FFocusEvent& InFocusEvent) override;
+	virtual void OnFocusChanging(const FWeakWidgetPath& PreviousFocusPath, const FWidgetPath& NewWidgetPath) override;
 	
 	virtual FReply OnKeyChar(const FGeometry& MyGeometry, const FCharacterEvent& InCharacterEvent) override;
-	virtual FReply OnPreviewKeyDown(const FGeometry& MyGeometry, const FKeyboardEvent& InKeyboardEvent) override;
-	virtual FReply OnKeyDown(const FGeometry& MyGeometry, const FKeyboardEvent& InKeyboardEvent) override;
-	virtual FReply OnKeyUp(const FGeometry& MyGeometry, const FKeyboardEvent& InKeyboardEvent) override;
+	virtual FReply OnPreviewKeyDown(const FGeometry& MyGeometry, const FKeyEvent& InKeyEvent) override;
+	virtual FReply OnKeyDown(const FGeometry& MyGeometry, const FKeyEvent& InKeyEvent) override;
+	virtual FReply OnKeyUp(const FGeometry& MyGeometry, const FKeyEvent& InKeyEvent) override;
+	virtual FReply OnAnalogValueChanged(const FGeometry& MyGeometry, const FAnalogInputEvent& InAnalogInputEvent) override;
 	
 	virtual FReply OnMouseButtonDown(const FGeometry& MyGeometry, const FPointerEvent& MouseEvent) override;
 	virtual FReply OnPreviewMouseButtonDown(const FGeometry& MyGeometry, const FPointerEvent& MouseEvent) override;
@@ -60,10 +61,6 @@ class SObjectWidget : public SCompoundWidget, public FGCObject
 
 	void OnDragCancelled(const FDragDropEvent& DragDropEvent, UDragDropOperation* Operation);
 	
-	virtual FReply OnControllerButtonPressed(const FGeometry& MyGeometry, const FControllerEvent& ControllerEvent) override;
-	virtual FReply OnControllerButtonReleased(const FGeometry& MyGeometry, const FControllerEvent& ControllerEvent) override;
-	virtual FReply OnControllerAnalogValueChanged(const FGeometry& MyGeometry, const FControllerEvent& ControllerEvent) override;
-	
 	virtual FReply OnTouchGesture(const FGeometry& MyGeometry, const FPointerEvent& GestureEvent) override;
 	virtual FReply OnTouchStarted(const FGeometry& MyGeometry, const FPointerEvent& InTouchEvent) override;
 	virtual FReply OnTouchMoved(const FGeometry& MyGeometry, const FPointerEvent& InTouchEvent) override;
