Fixed error in TSlotlessChildren::AsArrayCopy which was generating an array copy with duplicate elements.

diff --git a/Engine/Source/Runtime/SlateCore/Public/Layout/Children.h b/Engine/Source/Runtime/SlateCore/Public/Layout/Children.h
--- a/Engine/Source/Runtime/SlateCore/Public/Layout/Children.h
+++ b/Engine/Source/Runtime/SlateCore/Public/Layout/Children.h
@@ -311,10 +311,6 @@ public:
 	{
 		const int32 NumChildren = this->Num();
 		TArray< TSharedRef< ChildType > > Copy(*this);
-		for (int32 Index=0; Index < NumChildren; ++Index)
-		{
-			Copy.Add( (*this)[Index] );
-		}
 		return Copy;
 	}
 
