Slate : Exposed a globally accessible monospace font.

diff --git a/Engine/Source/Runtime/SlateCore/Private/Styling/CoreStyle.cpp b/Engine/Source/Runtime/SlateCore/Private/Styling/CoreStyle.cpp
--- a/Engine/Source/Runtime/SlateCore/Private/Styling/CoreStyle.cpp
+++ b/Engine/Source/Runtime/SlateCore/Private/Styling/CoreStyle.cpp
@@ -157,6 +157,18 @@ TSharedRef<ISlateStyle> FCoreStyle::Create( const FName& InStyleSetName )
 		.SetHighlightColor(FLinearColor(0.02f, 0.3f, 0.0f))
 		.SetHighlightShape(BOX_BRUSH("Common/TextBlockHighlightShape", FMargin(3.f /8.f)));
 
+	// Monospaced Text
+	const FTextBlockStyle MonospacedText = FTextBlockStyle()
+			.SetFont(TTF_FONT("Fonts/DroidSansMono", 10))
+			.SetColorAndOpacity(FSlateColor::UseForeground())
+			.SetShadowOffset(FVector2D::ZeroVector)
+			.SetShadowColorAndOpacity(FLinearColor::Black)
+			.SetHighlightColor(FLinearColor(0.02f, 0.3f, 0.0f))
+			.SetHighlightShape(BOX_BRUSH("Common/TextBlockHighlightShape", FMargin(3.f/8.f))
+			);
+
+	Style->Set("MonospacedText", MonospacedText);
+
 	// Small Text
 	FTextBlockStyle SmallText = FTextBlockStyle(NormalText)
 		.SetFont(TTF_FONT("Fonts/Roboto-Regular", 8));
@@ -1169,20 +1181,14 @@ TSharedRef<ISlateStyle> FCoreStyle::Create( const FName& InStyleSetName )
 
 	// Syntax highlighting
 	{
-		const FTextBlockStyle MonospacedText = FTextBlockStyle()
-			.SetFont(TTF_FONT("Fonts/DroidSansMono", 9))
-			.SetColorAndOpacity(FSlateColor::UseForeground())
-			.SetShadowOffset(FVector2D::ZeroVector)
-			.SetShadowColorAndOpacity(FLinearColor::Black)
-			.SetHighlightColor(FLinearColor(0.02f, 0.3f, 0.0f))
-			.SetHighlightShape(BOX_BRUSH("Common/TextBlockHighlightShape", FMargin(3.f/8.f))
-			);
+		const FTextBlockStyle SmallMonospacedText = FTextBlockStyle(MonospacedText)
+			.SetFont(TTF_FONT("Fonts/DroidSansMono", 9));
 
-		Style->Set("SyntaxHighlight.Normal", MonospacedText);
-		Style->Set("SyntaxHighlight.Node", FTextBlockStyle(MonospacedText).SetColorAndOpacity(FLinearColor(FColor(0xff006ab4)))); // blue
-		Style->Set("SyntaxHighlight.NodeAttributeKey", FTextBlockStyle(MonospacedText).SetColorAndOpacity(FLinearColor(FColor(0xffb40000)))); // red
-		Style->Set("SyntaxHighlight.NodeAttribueAssignment", FTextBlockStyle(MonospacedText).SetColorAndOpacity(FLinearColor(FColor(0xffb2b400)))); // yellow
-		Style->Set("SyntaxHighlight.NodeAttributeValue", FTextBlockStyle(MonospacedText).SetColorAndOpacity(FLinearColor(FColor(0xffb46100)))); // orange
+		Style->Set("SyntaxHighlight.Normal", SmallMonospacedText);
+		Style->Set("SyntaxHighlight.Node", FTextBlockStyle(SmallMonospacedText).SetColorAndOpacity(FLinearColor(FColor(0xff006ab4)))); // blue
+		Style->Set("SyntaxHighlight.NodeAttributeKey", FTextBlockStyle(SmallMonospacedText).SetColorAndOpacity(FLinearColor(FColor(0xffb40000)))); // red
+		Style->Set("SyntaxHighlight.NodeAttribueAssignment", FTextBlockStyle(SmallMonospacedText).SetColorAndOpacity(FLinearColor(FColor(0xffb2b400)))); // yellow
+		Style->Set("SyntaxHighlight.NodeAttributeValue", FTextBlockStyle(SmallMonospacedText).SetColorAndOpacity(FLinearColor(FColor(0xffb46100)))); // orange
 	}
 
 	return Style;
