SlateDynamicImageBrush - Add a static method for creating a dynamic bruch from image data.

diff --git a/Engine/Plugins/Editor/EpicSurvey/Source/EpicSurvey/Private/EpicSurvey.cpp b/Engine/Plugins/Editor/EpicSurvey/Source/EpicSurvey/Private/EpicSurvey.cpp
--- a/Engine/Plugins/Editor/EpicSurvey/Source/EpicSurvey/Private/EpicSurvey.cpp
+++ b/Engine/Plugins/Editor/EpicSurvey/Source/EpicSurvey/Private/EpicSurvey.cpp
@@ -93,9 +93,9 @@ TSharedPtr< FSlateDynamicImageBrush > FEpicSurvey::LoadRawDataAsBrush( FName Res
 		}
 	}
 
-	if ( bSucceeded && FSlateApplication::Get().GetRenderer()->GenerateDynamicImageResource( ResourceName, ImageWrapper->GetWidth(), ImageWrapper->GetHeight(), DecodedImage ) )
+	if ( bSucceeded )
 	{
-		Brush = MakeShareable(new FSlateDynamicImageBrush( ResourceName, FVector2D(ImageWrapper->GetWidth(), ImageWrapper->GetHeight()) ) );
+		Brush = FSlateDynamicImageBrush::CreateWithImageData( ResourceName, FVector2D(ImageWrapper->GetWidth(), ImageWrapper->GetHeight()), DecodedImage );
 	}
 
 	return Brush;
diff --git a/Engine/Source/Runtime/SlateCore/Private/Brushes/SlateDynamicImageBrush.cpp b/Engine/Source/Runtime/SlateCore/Private/Brushes/SlateDynamicImageBrush.cpp
--- a/Engine/Source/Runtime/SlateCore/Private/Brushes/SlateDynamicImageBrush.cpp
+++ b/Engine/Source/Runtime/SlateCore/Private/Brushes/SlateDynamicImageBrush.cpp
@@ -6,6 +6,29 @@
 /* FSlateDynamicImageBrush structors
  *****************************************************************************/
 
+
+TSharedPtr<FSlateDynamicImageBrush> FSlateDynamicImageBrush::CreateWithImageData(
+	const FName InTextureName,
+	const FVector2D& InImageSize,
+	const TArray<uint8>& InImageData,
+	const FLinearColor& InTint,
+	ESlateBrushTileType::Type InTiling,
+	ESlateBrushImageType::Type InImageType)
+{
+	TSharedPtr<FSlateDynamicImageBrush> Brush;
+	if (FSlateApplicationBase::IsInitialized() &&
+		FSlateApplicationBase::Get().GetRenderer()->GenerateDynamicImageResource(InTextureName, InImageSize.X, InImageSize.Y, InImageData))
+	{
+		Brush = MakeShareable(new FSlateDynamicImageBrush(
+			InTextureName,
+			InImageSize,
+			InTint,
+			InTiling,
+			InImageType));
+	}
+	return Brush;
+}
+
 FSlateDynamicImageBrush::~FSlateDynamicImageBrush( )
 {
 	if (FSlateApplicationBase::IsInitialized())
diff --git a/Engine/Source/Runtime/SlateCore/Public/Brushes/SlateDynamicImageBrush.h b/Engine/Source/Runtime/SlateCore/Public/Brushes/SlateDynamicImageBrush.h
--- a/Engine/Source/Runtime/SlateCore/Public/Brushes/SlateDynamicImageBrush.h
+++ b/Engine/Source/Runtime/SlateCore/Public/Brushes/SlateDynamicImageBrush.h
@@ -57,6 +57,22 @@ struct SLATECORE_API FSlateDynamicImageBrush
 		bIsDynamicallyLoaded = true;
 	}
 
+	/**
+	* @param InTextureName		The name to use when registering the image data as a texture.
+	* @param InImageSize		How large should the image be (not necessarily the image size on disk)
+	* @param InImageData		The raw image data formatted as BGRA
+	* @param InTint				The tint of the image.
+	* @param InTiling			How do we tile if at all?
+	* @param InImageType		The type of image this this is
+	*/
+	static TSharedPtr<FSlateDynamicImageBrush> CreateWithImageData(
+		const FName InTextureName,
+		const FVector2D& InImageSize,
+		const TArray<uint8>& InImageData,
+		const FLinearColor& InTint = FLinearColor(1.0f, 1.0f, 1.0f, 1.0f),
+		ESlateBrushTileType::Type InTiling = ESlateBrushTileType::NoTile,
+		ESlateBrushImageType::Type InImageType = ESlateBrushImageType::FullColor);
+
 	/** Destructor. */
 	virtual ~FSlateDynamicImageBrush();
 };
