Added missing Padding overloads to OverlaySlots.

diff --git a/Engine/Source/Runtime/SlateCore/Public/Widgets/SOverlay.h b/Engine/Source/Runtime/SlateCore/Public/Widgets/SOverlay.h
--- a/Engine/Source/Runtime/SlateCore/Public/Widgets/SOverlay.h
+++ b/Engine/Source/Runtime/SlateCore/Public/Widgets/SOverlay.h
@@ -47,6 +47,24 @@ public:
 			return *this;
 		}
 
+		FOverlaySlot& Padding(float Uniform)
+		{
+			SlotPadding = FMargin(Uniform);
+			return *this;
+		}
+
+		FOverlaySlot& Padding(float Horizontal, float Vertical)
+		{
+			SlotPadding = FMargin(Horizontal, Vertical);
+			return *this;
+		}
+
+		FOverlaySlot& Padding(float Left, float Top, float Right, float Bottom)
+		{
+			SlotPadding = FMargin(Left, Top, Right, Bottom);
+			return *this;
+		}
+
 		FOverlaySlot& Padding( const TAttribute<FMargin> InPadding )
 		{
 			SlotPadding = InPadding;
