Slate : changed HittestGrid check to an ensure. So that it does not block folks from working, as it is not always fatal.

diff --git a/Engine/Source/Runtime/SlateCore/Private/Input/HittestGrid.cpp b/Engine/Source/Runtime/SlateCore/Private/Input/HittestGrid.cpp
--- a/Engine/Source/Runtime/SlateCore/Private/Input/HittestGrid.cpp
+++ b/Engine/Source/Runtime/SlateCore/Private/Input/HittestGrid.cpp
@@ -164,7 +164,7 @@ void FHittestGrid::BeginFrame( const FSlateRect& HittestArea )
 
 int32 FHittestGrid::InsertWidget( const int32 ParentHittestIndex, const EVisibility& Visibility, const FArrangedWidget& Widget, const FVector2D InWindowOffset, const FSlateRect& InClippingRect )
 {
-	check( ParentHittestIndex < WidgetsCachedThisFrame->Num() );
+	ensureMsgf( ParentHittestIndex < WidgetsCachedThisFrame->Num(), TEXT("Widget '%s' being drawn before its parent."), *Widget.ToString() );
 
 	// Update the FGeometry to transform into desktop space.
 	FArrangedWidget WindowAdjustedWidget(Widget);
