Fixed composite sub-fonts using the incorrect baseline

All fonts within a composite font need to use the baseline from the default font. This ensures that characters rendered from different fonts line up correctly.

This also implements a way to get the baseline and max character height directly from FFreeTypeInterface. These functions do the same thing as the FSlateFontCache versions used to, however are more efficient since they only look at the default font when finding a typeface match, rather than also checking all sub-font character ranges.

diff --git a/Engine/Source/Runtime/SlateCore/Private/Fonts/FontCache.cpp b/Engine/Source/Runtime/SlateCore/Private/Fonts/FontCache.cpp
--- a/Engine/Source/Runtime/SlateCore/Private/Fonts/FontCache.cpp
+++ b/Engine/Source/Runtime/SlateCore/Private/Fonts/FontCache.cpp
@@ -325,6 +325,48 @@ public:
 #endif // WITH_FREETYPE
 	}
 
+	/**
+	 * @return The global max height for any character in the default font
+	 */
+	uint16 GetMaxHeight( const FSlateFontInfo& InFontInfo, const float InScale )
+	{
+#if WITH_FREETYPE
+		const FFontData& FontData = GetDefaultFontData(InFontInfo);
+
+		// Just render the null character 
+		TCHAR Char = 0;
+
+		// Render the character 
+		FCharacterRenderData NewRenderData;
+		GetRenderData(FontData, InFontInfo.Size, Char, NewRenderData, InScale);
+
+		return NewRenderData.MaxHeight;
+#else
+		return 0;
+#endif // WITH_FREETYPE
+	}
+
+	/** 
+	 * @return the baseline for any character in the default font
+	 */
+	int16 GetBaseline( const FSlateFontInfo& InFontInfo, const float InScale )
+	{
+#if WITH_FREETYPE
+		const FFontData& FontData = GetDefaultFontData(InFontInfo);
+
+		// Just render the null character 
+		TCHAR Char = 0;
+
+		// Render the character 
+		FCharacterRenderData NewRenderData;
+		GetRenderData(FontData, InFontInfo.Size, Char, NewRenderData, InScale);
+
+		return NewRenderData.MeasureInfo.GlobalDescender;
+#else
+		return 0;
+#endif // WITH_FREETYPE
+	}
+
 	/** 
 	 * Creates render data for a specific character 
 	 * 
@@ -337,13 +379,23 @@ public:
 #if WITH_FREETYPE
 		float SubFontScalingFactor = 1.0f;
 		const FFontData& FontData = GetFontDataForCharacter(InFontInfo, Char, SubFontScalingFactor);
+		return GetRenderData(FontData, InFontInfo.Size, Char, OutRenderData, InScale * SubFontScalingFactor);
+#endif // WITH_FREETYPE
+	}
 
-		// Apply the sub-font scale
-		const float FinalScale = InScale * SubFontScalingFactor;
-
+	/** 
+	 * Creates render data for a specific character 
+	 * 
+	 * @param InFontData	Raw font data to render the character with
+	 * @param Char			The character to render
+	 * @param OutCharInfo	Will contain the created render data
+	 */
+	void GetRenderData( const FFontData& InFontData, const int32 InSize, TCHAR Char, FCharacterRenderData& OutRenderData, const float InScale )
+	{
+#if WITH_FREETYPE
 		// Find or load the face if needed
 		FT_UInt GlyphIndex = 0;
-		FT_Face FontFace = GetFontFace( FontData );
+		FT_Face FontFace = GetFontFace( InFontData );
 
 		if ( FontFace != nullptr ) 
 		{
@@ -353,7 +405,7 @@ public:
 
 		uint32 LocalGlyphFlags = GlyphFlags;
 
-		switch(FontData.Hinting)
+		switch(InFontData.Hinting)
 		{
 		case EFontHinting::Auto:		LocalGlyphFlags |= FT_LOAD_FORCE_AUTOHINT; break;
 		case EFontHinting::AutoLight:	LocalGlyphFlags |= FT_LOAD_TARGET_LIGHT; break;
@@ -384,15 +436,15 @@ public:
 		}
 
 		// Set the character size to render at (needs to be in 1/64 of a "point")
-		FT_Error Error = FT_Set_Char_Size( FontFace, 0, InFontInfo.Size*64, FontCacheConstants::HorizontalDPI, FontCacheConstants::VerticalDPI );
+		FT_Error Error = FT_Set_Char_Size( FontFace, 0, InSize*64, FontCacheConstants::HorizontalDPI, FontCacheConstants::VerticalDPI );
 		check(Error==0);
 
-		if( FinalScale != 1.0f )
+		if( InScale != 1.0f )
 		{
 			FT_Matrix ScaleMatrix;
 			ScaleMatrix.xy = 0;
-			ScaleMatrix.xx = (FT_Fixed)(FinalScale * 65536);
-			ScaleMatrix.yy = (FT_Fixed)(FinalScale * 65536);
+			ScaleMatrix.xx = (FT_Fixed)(InScale * 65536);
+			ScaleMatrix.yy = (FT_Fixed)(InScale * 65536);
 			ScaleMatrix.yx = 0;
 			FT_Set_Transform( FontFace, &ScaleMatrix, nullptr );
 		}
@@ -466,7 +518,7 @@ public:
 		FT_Get_Glyph( Slot, &Glyph );
 		FT_Glyph_Get_CBox( Glyph, FT_GLYPH_BBOX_PIXELS, &GlyphBox );
 
-		int32 Height = (FT_MulFix( FontFace->height, FontFace->size->metrics.y_scale ) / 64) * FinalScale;
+		int32 Height = (FT_MulFix( FontFace->height, FontFace->size->metrics.y_scale ) / 64) * InScale;
 
 		// Set measurement info for this character
 		OutRenderData.Char = Char;
@@ -476,9 +528,9 @@ public:
 
 		// Need to divide by 64 to get pixels;
 		// Ascender is not scaled by freetype.  Scale it now. 
-		OutRenderData.MeasureInfo.GlobalAscender = ( FontFace->size->metrics.ascender / 64 ) * InScale; // Don't add the sub-font scale, as we want a consistent ascender
+		OutRenderData.MeasureInfo.GlobalAscender = ( FontFace->size->metrics.ascender / 64 ) * InScale;
 		// Descender is not scaled by freetype.  Scale it now. 
-		OutRenderData.MeasureInfo.GlobalDescender = ( FontFace->size->metrics.descender / 64 ) * InScale; // Don't add the sub-font scale, as we want a consistent descender
+		OutRenderData.MeasureInfo.GlobalDescender = ( FontFace->size->metrics.descender / 64 ) * InScale;
 		// Note we use Slot->advance instead of Slot->metrics.horiAdvance because Slot->Advance contains transformed position (needed if we scale)
 		OutRenderData.MeasureInfo.XAdvance =  Slot->advance.x / 64;
 		OutRenderData.MeasureInfo.HorizontalOffset = Slot->bitmap_left;
@@ -1036,7 +1088,7 @@ bool FSlateFontCache::AddNewEntry( TCHAR Character, const FSlateFontKey& InKey,
 		OutCharacterEntry.TextureIndex = 0;
 		OutCharacterEntry.XAdvance = RenderData.MeasureInfo.XAdvance;
 		OutCharacterEntry.VerticalOffset = RenderData.MeasureInfo.VerticalOffset;
-		OutCharacterEntry.GlobalDescender = RenderData.MeasureInfo.GlobalDescender;
+		OutCharacterEntry.GlobalDescender = GetBaseline(InKey.FontInfo, InKey.Scale); // All fonts within a composite font need to use the baseline of the default font
 		OutCharacterEntry.HorizontalOffset = RenderData.MeasureInfo.HorizontalOffset;
 		OutCharacterEntry.TextureIndex = AtlasIndex;
 		OutCharacterEntry.Valid = 1;
@@ -1071,28 +1123,12 @@ FCharacterList& FSlateFontCache::GetCharacterList( const FSlateFontInfo &InFontI
 
 uint16 FSlateFontCache::GetMaxCharacterHeight( const FSlateFontInfo& InFontInfo, float FontScale ) const
 {
-	FCharacterRenderData NewRenderData;
-
-	// Just render the null character 
-	TCHAR Char = 0;
-	// Render the character 
-	FTInterface->GetRenderData( InFontInfo, Char, NewRenderData, FontScale );	
-
-	//return NewRenderData.MeasureInfo.GlobalAscender - NewRenderData.MeasureInfo.GlobalDescender;
-
-	return NewRenderData.MaxHeight;
+	return FTInterface->GetMaxHeight(InFontInfo, FontScale);
 }
 
 int16 FSlateFontCache::GetBaseline( const FSlateFontInfo& InFontInfo, float FontScale ) const
 {
-	FCharacterRenderData NewRenderData;
-
-	// Just render the null character 
-	TCHAR Char = 0;
-	// Render the character 
-	FTInterface->GetRenderData( InFontInfo, Char, NewRenderData, FontScale );	
-
-	return NewRenderData.MeasureInfo.GlobalDescender;
+	return FTInterface->GetBaseline(InFontInfo, FontScale);
 }
 
 int8 FSlateFontCache::GetKerning( TCHAR First, TCHAR Second, const FSlateFontInfo& InFontInfo, float Scale ) const 
