Fixed FKerningTable causing all sub-fonts to be loaded into memory

FKerningTable used to check whether any font contained within a composite font contained kerning information. This check caused all of the sub-fonts to be loaded into memory, even if they were never used.

This change defers this check until we actually need to get the kerning information for the characters in question.

diff --git a/Engine/Source/Runtime/SlateCore/Private/Fonts/FontCache.cpp b/Engine/Source/Runtime/SlateCore/Private/Fonts/FontCache.cpp
--- a/Engine/Source/Runtime/SlateCore/Private/Fonts/FontCache.cpp
+++ b/Engine/Source/Runtime/SlateCore/Private/Fonts/FontCache.cpp
@@ -549,10 +549,13 @@ public:
 	/**
 	 * @param Whether or not the font has kerning
 	 */
-	bool HasKerning( const FFontData& InFontData )
+	bool HasKerning( const FSlateFontInfo& InFontInfo, TCHAR Char )
 	{
 #if WITH_FREETYPE
-		FT_Face FontFace = GetFontFace( InFontData );
+		float SubFontScalingFactor = 1.0f;
+		const FFontData& FontData = GetFontDataForCharacter(InFontInfo, Char, SubFontScalingFactor);
+
+		FT_Face FontFace = GetFontFace( FontData );
 
 		if ( FontFace == nullptr )
 		{
@@ -877,7 +880,6 @@ FKerningTable::FKerningTable( const FSlateFontKey& InFont, const FSlateFontCache
 	, FontCache( InFontCache )
 	, FontKey( InFont )
 {
-	bHasKerning = InFontCache.HasKerning( InFont.FontInfo );
 }
 
 FKerningTable::~FKerningTable()
@@ -895,9 +897,15 @@ FKerningTable::~FKerningTable()
 
 int8 FKerningTable::GetKerning( TCHAR FirstChar, TCHAR SecondChar )
 {
+	const bool bHasKerning = FontCache.HasKerning( FontKey.FontInfo, FirstChar ) && FontCache.HasKerning( FontKey.FontInfo, SecondChar );
+	if( !bHasKerning )
+	{
+		return 0;
+	}
+
 	int8 OutKerning = 0;
 
-	if( bHasKerning && FirstChar < FontCacheConstants::DirectAccessSize && SecondChar < FontCacheConstants::DirectAccessSize )
+	if( FirstChar < FontCacheConstants::DirectAccessSize && SecondChar < FontCacheConstants::DirectAccessSize )
 	{
 		// This character can be directly indexed
 
@@ -918,7 +926,7 @@ int8 FKerningTable::GetKerning( TCHAR FirstChar, TCHAR SecondChar )
 			DirectAccessTable[Index] = OutKerning;
 		}
 	}
-	else if( bHasKerning )
+	else
 	{
 		// Kerning is mapped
 		FKerningPair KerningPair( FirstChar, SecondChar );
@@ -1136,44 +1144,9 @@ int8 FSlateFontCache::GetKerning( TCHAR First, TCHAR Second, const FSlateFontInf
 	return FTInterface->GetKerning( First, Second, InFontInfo, Scale );
 }
 
-bool FSlateFontCache::HasKerning( const FSlateFontInfo& InFontInfo ) const
-{
-	const FCompositeFont* const CompositeFont = InFontInfo.GetCompositeFont();
-	if(CompositeFont)
-	{
-		if (HasKerning(CompositeFont->DefaultTypeface))
-		{
-			return true;
-		}
-
-		for (const FCompositeSubFont& SubTypeface : CompositeFont->SubTypefaces)
-		{
-			if (HasKerning(SubTypeface.Typeface))
-			{
-				return true;
-			}
-		}
-	}
-
-	return false;
-}
-
-bool FSlateFontCache::HasKerning( const FTypeface& InTypeface ) const
-{
-	for (const FTypefaceEntry& TypefaceEntry : InTypeface.Fonts)
-	{
-		if (HasKerning(TypefaceEntry.Font))
-		{
-			return true;
-		}
-	}
-
-	return false;
-}
-
-bool FSlateFontCache::HasKerning( const FFontData& InFontData ) const
+bool FSlateFontCache::HasKerning( const FSlateFontInfo& InFontInfo, TCHAR Char ) const
 {
-	return FTInterface->HasKerning( InFontData );
+	return FTInterface->HasKerning( InFontInfo, Char );
 }
 
 const TSet<FName>& FSlateFontCache::GetFontAttributes( const FFontData& InFontData ) const
diff --git a/Engine/Source/Runtime/SlateCore/Public/Fonts/FontCache.h b/Engine/Source/Runtime/SlateCore/Public/Fonts/FontCache.h
--- a/Engine/Source/Runtime/SlateCore/Public/Fonts/FontCache.h
+++ b/Engine/Source/Runtime/SlateCore/Public/Fonts/FontCache.h
@@ -101,8 +101,6 @@ private:
 	const class FSlateFontCache& FontCache;
 	/** Font this kerning table is for */
 	FSlateFontKey FontKey;
-	/** Whether the font has kerning at all */
-	bool bHasKerning;
 };
 
 
@@ -298,19 +296,9 @@ public:
 	int8 GetKerning( TCHAR First, TCHAR Second, const FSlateFontInfo& InFontInfo, float Scale ) const;
 
 	/**
-	 * @return Whether or not the font has kerning information (this checks both the default font as well as any sub-fonts)
+	 * @return Whether or not the font used by the given character has kerning information
 	 */
-	bool HasKerning( const FSlateFontInfo& InFontInfo ) const;
-
-	/**
-	 * @return Whether or not the typeface has kerning information
-	 */
-	bool HasKerning( const FTypeface& InTypeface ) const;
-
-	/**
-	 * @return Whether or not the font has kerning information
-	 */
-	bool HasKerning( const FFontData& InFontData ) const;
+	bool HasKerning( const FSlateFontInfo& InFontInfo, TCHAR Char ) const;
 
 	/**
 	 * Returns the height of the largest character in the font. 
