Slate : There is now a dedicated HittestGrid for each SWindow .

diff --git a/Engine/Source/Runtime/Slate/Private/Framework/Application/SlateApplication.cpp b/Engine/Source/Runtime/Slate/Private/Framework/Application/SlateApplication.cpp
--- a/Engine/Source/Runtime/Slate/Private/Framework/Application/SlateApplication.cpp
+++ b/Engine/Source/Runtime/Slate/Private/Framework/Application/SlateApplication.cpp
@@ -693,7 +693,6 @@ FSlateApplication::FSlateApplication()
 	, bMenuAnimationsEnabled( true )
 	, AppIcon( FCoreStyle::Get().GetBrush("DefaultAppIcon") )
 	, VirtualDesktopRect( 0,0,0,0 )
-	, HittestGrid( MakeShareable( new FHittestGrid() ) )
 {
 #if WITH_UNREAL_DEVELOPER_TOOLS
 	FModuleManager::Get().LoadModule(TEXT("Settings"));
@@ -815,7 +814,7 @@ FWidgetPath FSlateApplication::LocateWindowUnderMouse( FVector2D ScreenspaceMous
 
 		if ( Window->IsVisible() && AcceptsInput && Window->IsScreenspaceMouseWithin(ScreenspaceMouseCoordinate) && !bPrevWindowWasModal )
 		{
-			const TArray<FWidgetAndPointer> WidgetsAndCursors = HittestGrid->GetBubblePath( ScreenspaceMouseCoordinate, bIgnoreEnabledStatus );
+			const TArray<FWidgetAndPointer> WidgetsAndCursors = Window->GetHittestGrid()->GetBubblePath(ScreenspaceMouseCoordinate, bIgnoreEnabledStatus);
 			return FWidgetPath( WidgetsAndCursors );
 		}
 	}
@@ -908,8 +907,9 @@ void FSlateApplication::DrawWindowAndChildren( const TSharedRef<SWindow>& Window
 		FGeometry WindowGeometry = WindowToDraw->GetWindowGeometryInWindow();
 		int32 MaxLayerId = 0;
 		{
+			WindowToDraw->GetHittestGrid()->ClearGridForNewFrame( VirtualDesktopRect );
 			MaxLayerId = WindowToDraw->PaintWindow(
-				FPaintArgs(WindowToDraw, *HittestGrid, WindowToDraw->GetPositionInScreen(), GetCurrentTime(), GetDeltaTime()),
+				FPaintArgs(WindowToDraw, *WindowToDraw->GetHittestGrid(), WindowToDraw->GetPositionInScreen(), GetCurrentTime(), GetDeltaTime()),
 				WindowGeometry, WindowToDraw->GetClippingRectangleInWindow(),
 				WindowElementList,
 				0,
@@ -1034,12 +1034,6 @@ void FSlateApplication::PrivateDrawWindows( TSharedPtr<SWindow> DrawOnlyThisWind
 	{
 		SCOPE_CYCLE_COUNTER( STAT_SlateDrawWindowTime );
 
-		const bool bClearHittestGrid = !DrawOnlyThisWindow.IsValid();
-		if ( bClearHittestGrid )
-		{
-			HittestGrid->BeginFrame( VirtualDesktopRect );
-		}
-
 		TSharedPtr<SWindow> ActiveModalWindow = GetActiveModalWindow(); 
 
 		if (ActiveModalWindow.IsValid())
@@ -1477,10 +1471,7 @@ void FSlateApplication::AddModalWindow( TSharedRef<SWindow> InSlateWindow, const
 		return;
 	}
 
-	// If we are spawning a modal window, we don't want to allow the possibility of input events being routed to non-modal geometry that existed on the hittest grid last frame.
-	// To this end we clear the grid now. This means that input events that occur on the first frame will be missed, however since the user couldn't even see the window yet, this is ok.
-	// @todo slate : switch to FHittestGrid per window
-	HittestGrid->BeginFrame( VirtualDesktopRect );
+
 
 	// Push the active modal window onto the stack.  
 	ActiveModalWindows.AddUnique( InSlateWindow );
@@ -4419,7 +4410,7 @@ bool FSlateApplication::AttemptNavigation(const FNavigationEvent& NavigationEven
 				// Switch worlds for widgets in the current path 
 				FScopedSwitchWorldHack SwitchWorld(FocusedWidgetPath);
 
-				NewFocusedWidget = HittestGrid->FindNextFocusableWidget(FocusedArrangedWidget, NavigationType, NavigationReply, BoundaryWidget);
+				NewFocusedWidget = FocusedWidgetPath.GetWindow()->GetHittestGrid()->FindNextFocusableWidget(FocusedArrangedWidget, NavigationType, NavigationReply, BoundaryWidget);
 			}
 		}
 	}
diff --git a/Engine/Source/Runtime/Slate/Public/Framework/Application/SlateApplication.h b/Engine/Source/Runtime/Slate/Public/Framework/Application/SlateApplication.h
--- a/Engine/Source/Runtime/Slate/Public/Framework/Application/SlateApplication.h
+++ b/Engine/Source/Runtime/Slate/Public/Framework/Application/SlateApplication.h
@@ -1476,6 +1476,4 @@ private:
 	// but there could be another monitor on any of the sides.
 	FSlateRect VirtualDesktopRect;
 
-	// Hittest acceleration structure. Filled out during Paint.
-	TSharedRef<FHittestGrid> HittestGrid;
 };
diff --git a/Engine/Source/Runtime/SlateCore/Private/Input/HittestGrid.cpp b/Engine/Source/Runtime/SlateCore/Private/Input/HittestGrid.cpp
--- a/Engine/Source/Runtime/SlateCore/Private/Input/HittestGrid.cpp
+++ b/Engine/Source/Runtime/SlateCore/Private/Input/HittestGrid.cpp
@@ -150,7 +150,7 @@ TArray<FWidgetAndPointer> FHittestGrid::GetBubblePath( FVector2D DesktopSpaceCoo
 	}
 }
 
-void FHittestGrid::BeginFrame( const FSlateRect& HittestArea )
+void FHittestGrid::ClearGridForNewFrame( const FSlateRect& HittestArea )
 {
 	//LogGrid();
 
@@ -164,57 +164,62 @@ void FHittestGrid::BeginFrame( const FSlateRect& HittestArea )
 
 int32 FHittestGrid::InsertWidget( const int32 ParentHittestIndex, const EVisibility& Visibility, const FArrangedWidget& Widget, const FVector2D InWindowOffset, const FSlateRect& InClippingRect )
 {
-	ensureMsgf( ParentHittestIndex < WidgetsCachedThisFrame->Num(), TEXT("Widget '%s' being drawn before its parent."), *Widget.ToString() );
+	if ( ensureMsgf( ParentHittestIndex < WidgetsCachedThisFrame->Num(), TEXT("Widget '%s' being drawn before its parent."), *Widget.ToString() ) )
+	{
+		// Update the FGeometry to transform into desktop space.
+		FArrangedWidget WindowAdjustedWidget(Widget);
+		WindowAdjustedWidget.Geometry.AppendTransform(FSlateLayoutTransform(InWindowOffset));
+		const FSlateRect WindowAdjustedRect = InClippingRect.OffsetBy(InWindowOffset);
 
-	// Update the FGeometry to transform into desktop space.
-	FArrangedWidget WindowAdjustedWidget(Widget);
-	WindowAdjustedWidget.Geometry.AppendTransform(FSlateLayoutTransform(InWindowOffset));
-	const FSlateRect WindowAdjustedRect = InClippingRect.OffsetBy(InWindowOffset);
 
-	// Remember this widget, its geometry, and its place in the logical hierarchy.
-	const int32 WidgetIndex = WidgetsCachedThisFrame->Add( FCachedWidget( ParentHittestIndex, WindowAdjustedWidget, WindowAdjustedRect ) );
-	check( WidgetIndex < WidgetsCachedThisFrame->Num() ); 
-	if (ParentHittestIndex != INDEX_NONE)
-	{
-		(*WidgetsCachedThisFrame)[ParentHittestIndex].AddChild( WidgetIndex );
-	}
+		// Remember this widget, its geometry, and its place in the logical hierarchy.
+		const int32 WidgetIndex = WidgetsCachedThisFrame->Add( FCachedWidget( ParentHittestIndex, WindowAdjustedWidget, WindowAdjustedRect ) );
+		check( WidgetIndex < WidgetsCachedThisFrame->Num() ); 
+		if (ParentHittestIndex != INDEX_NONE)
+		{
+			(*WidgetsCachedThisFrame)[ParentHittestIndex].AddChild( WidgetIndex );
+		}
 	
-	if (Visibility.IsHitTestVisible())
-	{
-		// Mark any cell that is overlapped by this widget.
-
-		// Compute the render space clipping rect, and compute it's aligned bounds so we can insert conservatively into the hit test grid.
-		FSlateRect GridRelativeBoundingClipRect = 
-			TransformRect(
-				Concatenate(
-					Inverse(WindowAdjustedWidget.Geometry.GetAccumulatedLayoutTransform()),
-					WindowAdjustedWidget.Geometry.GetAccumulatedRenderTransform()
-				),
-				FSlateRotatedRect(WindowAdjustedWidget.Geometry.GetClippingRect().IntersectionWith(WindowAdjustedRect))
-			)
-			.ToBoundingRect()
-			.OffsetBy(-GridOrigin);
-
-		// Starting and ending cells covered by this widget.	
-		const FIntPoint UpperLeftCell = FIntPoint(
-			FMath::Max(0, FMath::FloorToInt(GridRelativeBoundingClipRect.Left / CellSize.X)),
-			FMath::Max(0, FMath::FloorToInt(GridRelativeBoundingClipRect.Top / CellSize.Y)));
-
-		const FIntPoint LowerRightCell = FIntPoint(
-			FMath::Min( NumCells.X-1, FMath::FloorToInt(GridRelativeBoundingClipRect.Right / CellSize.X)),
-			FMath::Min( NumCells.Y-1, FMath::FloorToInt(GridRelativeBoundingClipRect.Bottom / CellSize.Y)));
-
-		for (int32 XIndex=UpperLeftCell.X; XIndex <= LowerRightCell.X; ++ XIndex )
+		if (Visibility.IsHitTestVisible())
 		{
-			for(int32 YIndex=UpperLeftCell.Y; YIndex <= LowerRightCell.Y; ++YIndex)
+			// Mark any cell that is overlapped by this widget.
+
+			// Compute the render space clipping rect, and compute it's aligned bounds so we can insert conservatively into the hit test grid.
+			FSlateRect GridRelativeBoundingClipRect = 
+				TransformRect(
+					Concatenate(
+						Inverse(WindowAdjustedWidget.Geometry.GetAccumulatedLayoutTransform()),
+						WindowAdjustedWidget.Geometry.GetAccumulatedRenderTransform()
+					),
+					FSlateRotatedRect(WindowAdjustedWidget.Geometry.GetClippingRect().IntersectionWith(WindowAdjustedRect))
+				)
+				.ToBoundingRect()
+				.OffsetBy(-GridOrigin);
+
+			// Starting and ending cells covered by this widget.	
+			const FIntPoint UpperLeftCell = FIntPoint(
+				FMath::Max(0, FMath::FloorToInt(GridRelativeBoundingClipRect.Left / CellSize.X)),
+				FMath::Max(0, FMath::FloorToInt(GridRelativeBoundingClipRect.Top / CellSize.Y)));
+
+			const FIntPoint LowerRightCell = FIntPoint(
+				FMath::Min( NumCells.X-1, FMath::FloorToInt(GridRelativeBoundingClipRect.Right / CellSize.X)),
+				FMath::Min( NumCells.Y-1, FMath::FloorToInt(GridRelativeBoundingClipRect.Bottom / CellSize.Y)));
+
+			for (int32 XIndex=UpperLeftCell.X; XIndex <= LowerRightCell.X; ++ XIndex )
 			{
-				CellAt(XIndex, YIndex).CachedWidgetIndexes.Add( WidgetIndex );
+				for(int32 YIndex=UpperLeftCell.Y; YIndex <= LowerRightCell.Y; ++YIndex)
+				{
+					CellAt(XIndex, YIndex).CachedWidgetIndexes.Add( WidgetIndex );
+				}
 			}
 		}
-	}
-	
-	return WidgetIndex;
 
+		return WidgetIndex;
+	}
+	else
+	{
+		return INDEX_NONE;
+	}
 }
 
 void FHittestGrid::InsertCustomHitTestPath( TSharedRef<ICustomHitTestPath> CustomHitTestPath, int32 WidgetIndex )
diff --git a/Engine/Source/Runtime/SlateCore/Private/Widgets/SWindow.cpp b/Engine/Source/Runtime/SlateCore/Private/Widgets/SWindow.cpp
--- a/Engine/Source/Runtime/SlateCore/Private/Widgets/SWindow.cpp
+++ b/Engine/Source/Runtime/SlateCore/Private/Widgets/SWindow.cpp
@@ -511,6 +511,10 @@ bool SWindow::HasActiveChildren() const
 	return false;
 }
 
+TSharedRef<FHittestGrid> SWindow::GetHittestGrid()
+{
+	return HittestGrid;
+}
 
 void SWindow::Tick( const FGeometry& AllottedGeometry, const double InCurrentTime, const float InDeltaTime )
 {
@@ -1696,7 +1700,9 @@ EWindowZone::Type SWindow::GetCurrentWindowZone(FVector2D LocalMousePosition)
  * Default constructor. Protected because SWindows must always be used via TSharedPtr. Instead, use FSlateApplication::MakeWindow()
  */
 SWindow::SWindow()
-	: Opacity( 1.0f )
+	: HittestGrid( MakeShareable(new FHittestGrid()) )
+	, bDragAnywhere( false )
+	, Opacity( 1.0f )
 	, SizingRule( ESizingRule::UserSized )
 	, bIsTransparent( false )
 	, bIsPopupWindow( false )
diff --git a/Engine/Source/Runtime/SlateCore/Public/Input/HittestGrid.h b/Engine/Source/Runtime/SlateCore/Public/Input/HittestGrid.h
--- a/Engine/Source/Runtime/SlateCore/Public/Input/HittestGrid.h
+++ b/Engine/Source/Runtime/SlateCore/Public/Input/HittestGrid.h
@@ -35,7 +35,7 @@ public:
 	 *
 	 * @param HittestArea       Size in Slate Units of the area we are considering for hittesting.
 	 */
-	void BeginFrame(const FSlateRect& HittestArea);
+	void ClearGridForNewFrame(const FSlateRect& HittestArea);
 
 	/** Add Widget into the hittest data structure so that we can later make queries about it. */
 	int32 InsertWidget( const int32 ParentHittestIndex, const EVisibility& Visibility, const FArrangedWidget& Widget, const FVector2D InWindowOffset, const FSlateRect& InClippingRect );
diff --git a/Engine/Source/Runtime/SlateCore/Public/Widgets/SWindow.h b/Engine/Source/Runtime/SlateCore/Public/Widgets/SWindow.h
--- a/Engine/Source/Runtime/SlateCore/Public/Widgets/SWindow.h
+++ b/Engine/Source/Runtime/SlateCore/Public/Widgets/SWindow.h
@@ -2,6 +2,8 @@
 
 #pragma once
 
+class FHittestGrid;
+
 
 /** Notification that a window has been activated */
 DECLARE_DELEGATE( FOnWindowActivated );
@@ -668,6 +670,14 @@ public:
 		ViewportSize = VP;
 	}
 
+	/**
+	 * Access the hittest acceleration data structure for this window.
+	 * The grid is filled out every time the window is painted.
+	 *
+	 * @see FHittestGrid for more details.
+	 */
+	TSharedRef<FHittestGrid> GetHittestGrid();
+
 public:
 
 	// SWidget overrides
@@ -833,6 +843,9 @@ private:
 
 	/** The native window that is backing this Slate Window */
 	TSharedPtr<FGenericWindow> NativeWindow;
+
+	/** Each window has its own hittest grid for accelerated widget picking. */
+	TSharedRef<FHittestGrid> HittestGrid;
 	
 	/** Invoked when the window has been activated. */
 	FOnWindowActivated OnWindowActivated;
diff --git a/Engine/Source/Runtime/UMG/Private/Components/WidgetComponent.cpp b/Engine/Source/Runtime/UMG/Private/Components/WidgetComponent.cpp
--- a/Engine/Source/Runtime/UMG/Private/Components/WidgetComponent.cpp
+++ b/Engine/Source/Runtime/UMG/Private/Components/WidgetComponent.cpp
@@ -466,12 +466,12 @@ void UWidgetComponent::TickComponent(float DeltaTime, enum ELevelTick TickType,
 		FSlateDrawBuffer& DrawBuffer = Renderer->GetDrawBuffer();
 		FSlateWindowElementList& WindowElementList = DrawBuffer.AddWindowElementList(SlateWidget.ToSharedRef());
 
-		// Prepare the test grid 
-		HitTestGrid->BeginFrame(WindowGeometry.GetClippingRect());
-
-		// Paint the window
 		int32 MaxLayerId = 0;
 		{
+			// Prepare the test grid 
+			HitTestGrid->ClearGridForNewFrame(WindowGeometry.GetClippingRect());
+
+			// Paint the window
 			MaxLayerId = SlateWidget->Paint(
 				FPaintArgs(SlateWidget.ToSharedRef(), *HitTestGrid, FVector2D::ZeroVector, FSlateApplication::Get().GetCurrentTime(), FSlateApplication::Get().GetDeltaTime()),
 				WindowGeometry, WindowGeometry.GetClippingRect(),
