Slate.FoldTick is now enabled by default.

diff --git a/Engine/Source/Runtime/SlateCore/Private/Widgets/SWidget.cpp b/Engine/Source/Runtime/SlateCore/Private/Widgets/SWidget.cpp
--- a/Engine/Source/Runtime/SlateCore/Private/Widgets/SWidget.cpp
+++ b/Engine/Source/Runtime/SlateCore/Private/Widgets/SWidget.cpp
@@ -586,7 +586,7 @@ void SWidget::SetDebugInfo( const ANSICHAR* InType, const ANSICHAR* InFile, int3
 	this->CreatedOnLine = OnLine;
 }
 
-SLATECORE_API int32 bFoldTick = 0;
+SLATECORE_API int32 bFoldTick = 1;
 FAutoConsoleVariableRef FoldTick(TEXT("Slate.FoldTick"), bFoldTick, TEXT("When folding, call Tick as part of the paint pass instead of a separate tick pass."));
 
 int32 SWidget::Paint(const FPaintArgs& Args, const FGeometry& AllottedGeometry, const FSlateRect& MyClippingRect, FSlateWindowElementList& OutDrawElements, int32 LayerId, const FWidgetStyle& InWidgetStyle, bool bParentEnabled) const
