Fixed initialization order.

diff --git a/Engine/Source/Runtime/SlateCore/Private/Widgets/SWindow.cpp b/Engine/Source/Runtime/SlateCore/Private/Widgets/SWindow.cpp
--- a/Engine/Source/Runtime/SlateCore/Private/Widgets/SWindow.cpp
+++ b/Engine/Source/Runtime/SlateCore/Private/Widgets/SWindow.cpp
@@ -1700,8 +1700,8 @@ EWindowZone::Type SWindow::GetCurrentWindowZone(FVector2D LocalMousePosition)
  * Default constructor. Protected because SWindows must always be used via TSharedPtr. Instead, use FSlateApplication::MakeWindow()
  */
 SWindow::SWindow()
-	: HittestGrid( MakeShareable(new FHittestGrid()) )
-	, bDragAnywhere( false )
+	: bDragAnywhere( false )
+	, HittestGrid( MakeShareable(new FHittestGrid()) )
 	, Opacity( 1.0f )
 	, SizingRule( ESizingRule::UserSized )
 	, bIsTransparent( false )
