Missing include for "HittestGrid.h" in SWindow.cpp.

diff --git a/Engine/Source/Runtime/SlateCore/Private/Widgets/SWindow.cpp b/Engine/Source/Runtime/SlateCore/Private/Widgets/SWindow.cpp
--- a/Engine/Source/Runtime/SlateCore/Private/Widgets/SWindow.cpp
+++ b/Engine/Source/Runtime/SlateCore/Private/Widgets/SWindow.cpp
@@ -1,6 +1,7 @@
 // Copyright 1998-2014 Epic Games, Inc. All Rights Reserved.
 
 #include "SlateCorePrivatePCH.h"
+#include "HittestGrid.h"
 
 // this define is the first step to complete removal of the below hack
 #define PLATFORM_SPECIFIC_HACK 	0
