Fixed incorrect assert in FSlateTextureAtlas

UE-5321 - Regression: Crash using VisualizeTexture from FSlateFontCache

CL# 2350206 made it possible for FSlateTextureAtlas to be created and managed by the render thread as well as the game thread, and the assert in FSlateTextureAtlas::MarkTextureDirty was updated in an attempt to deal with that.

The assert was naive though, as the render thread ID can change if the render thread is suspended and resumed. It now stores the type of thread that owns the FSlateTextureAtlas (rather than its raw thread ID), and then uses that to check that FSlateTextureAtlas::MarkTextureDirty is being called from the correct thread.

diff --git a/Engine/Source/Runtime/SlateCore/Private/Textures/TextureAtlas.cpp b/Engine/Source/Runtime/SlateCore/Private/Textures/TextureAtlas.cpp
--- a/Engine/Source/Runtime/SlateCore/Private/Textures/TextureAtlas.cpp
+++ b/Engine/Source/Runtime/SlateCore/Private/Textures/TextureAtlas.cpp
@@ -48,13 +48,18 @@ const FAtlasedTextureSlot* FSlateTextureAtlas::AddTexture( uint32 TextureWidth,
 	return NewSlot;
 }
 
+
 void FSlateTextureAtlas::MarkTextureDirty()
 {
-	//check( IsThreadSafeForSlateRendering() );
-	check( (GSlateLoadingThreadId != 0) || FPlatformTLS::GetCurrentThreadId() == AtlasOwnerThreadId );
+	check( 
+		(GSlateLoadingThreadId != 0) || 
+		(AtlasOwnerThread == ESlateTextureAtlasOwnerThread::Game && IsInGameThread()) || 
+		(AtlasOwnerThread == ESlateTextureAtlasOwnerThread::Render && IsInRenderingThread()) 
+		);
+
 	bNeedsUpdate = true;
 }
-	
+
 
 const FAtlasedTextureSlot* FSlateTextureAtlas::FindSlotForTexture( uint32 InWidth, uint32 InHeight )
 {
@@ -85,7 +90,8 @@ void FSlateTextureAtlas::InitAtlasData()
 	AtlasData.Reserve(AtlasWidth * AtlasHeight * BytesPerPixel);
 	AtlasData.AddZeroed(AtlasWidth * AtlasHeight * BytesPerPixel);
 
-	AtlasOwnerThreadId = FPlatformTLS::GetCurrentThreadId();
+	check(IsInGameThread() || IsInRenderingThread());
+	AtlasOwnerThread = (IsInGameThread()) ? ESlateTextureAtlasOwnerThread::Game : ESlateTextureAtlasOwnerThread::Render;
 
 	INC_MEMORY_STAT_BY(STAT_SlateTextureAtlasMemory, AtlasData.GetAllocatedSize());
 }
diff --git a/Engine/Source/Runtime/SlateCore/Public/Textures/TextureAtlas.h b/Engine/Source/Runtime/SlateCore/Public/Textures/TextureAtlas.h
--- a/Engine/Source/Runtime/SlateCore/Public/Textures/TextureAtlas.h
+++ b/Engine/Source/Runtime/SlateCore/Public/Textures/TextureAtlas.h
@@ -16,6 +16,19 @@ enum ESlateTextureAtlasPaddingStyle
 	PadWithZero,
 };
 
+/** 
+ * The type of thread that owns a texture atlas - this is the only thread that can safely update it 
+ */
+enum class ESlateTextureAtlasOwnerThread : uint8
+{
+	/** Owner thread is currently unknown */
+	Unknown,
+	/** Atlas is owned by the game thread */
+	Game,
+	/** Atlas is owned by the render thread */
+	Render,
+};
+
 /**
  * Structure holding information about where a texture is located in the atlas
  * Pointers to left and right children build a tree of texture rectangles so we can easily find optimal slots for textures
@@ -64,6 +77,7 @@ public:
 		, BytesPerPixel( InBytesPerPixel )
 		, PaddingStyle( InPaddingStyle )
 		, bNeedsUpdate( false )
+		, AtlasOwnerThread( ESlateTextureAtlasOwnerThread::Unknown )
 	{
 		InitAtlasData();
 	}
@@ -194,6 +208,9 @@ protected:
 	/** True if this texture needs to have its rendering resources updated */
 	bool bNeedsUpdate;
 
-	/** ID of the thread that owns this atlas - this is the only thread that can safely update it */
-	uint32 AtlasOwnerThreadId;
+	/** 
+	 * The type of thread that owns this atlas - this is the only thread that can safely update it 
+	 * NOTE: We don't use the thread ID here, as the render thread can be recreated if it gets suspended and resumed, giving it a new ID
+	 */
+	ESlateTextureAtlasOwnerThread AtlasOwnerThread;
 };
