Text rendering perf. improvements

UETOOL-129 - Investigate SOutputLog performance degradation

This improves the caching of the kerning data used by the Slate font cache, as profiles were showing this to be a hot spot for text rendering.

FKerningTable has been changed to work on a single raw font (FFontData), rather than an entire composite font. This lets FCharacterList only generate a kerning table for those fonts that have kerning, rather than asking each font every time kerning was requested for a character.

This has the side effect of cleaning up the API of FSlateFontCache, as all functions that take a FSlateFontInfo now operate on the composite font as a whole (GetMaxCharacterHeight, GetBaseline), and the functions that take FFontData now operate on a single raw font (GetKerning, HasKerning, GetFontAttributes). This also exposes the GetDefaultFontData and GetFontDataForCharacter so that you can perform more efficient FFontData lookups from outside the FreeType interface.

diff --git a/Engine/Source/Runtime/SlateCore/Private/Fonts/CompositeFont.cpp b/Engine/Source/Runtime/SlateCore/Private/Fonts/CompositeFont.cpp
--- a/Engine/Source/Runtime/SlateCore/Private/Fonts/CompositeFont.cpp
+++ b/Engine/Source/Runtime/SlateCore/Private/Fonts/CompositeFont.cpp
@@ -25,6 +25,16 @@ void FFontData::SetFont(FString InFontFilename, const UFontBulkData* const InBul
 	BulkDataPtr = InBulkData;
 }
 
+bool FFontData::operator==(const FFontData& Other) const
+{
+	return BulkDataPtr == Other.BulkDataPtr && Hinting == Other.Hinting;
+}
+
+bool FFontData::operator!=(const FFontData& Other) const
+{
+	return !(*this == Other);
+}
+
 void FStandaloneCompositeFont::AddReferencedObjects(FReferenceCollector& Collector)
 {
 	for(FTypefaceEntry& TypefaceEntry : DefaultTypeface.Fonts)
diff --git a/Engine/Source/Runtime/SlateCore/Private/Fonts/FontCache.cpp b/Engine/Source/Runtime/SlateCore/Private/Fonts/FontCache.cpp
--- a/Engine/Source/Runtime/SlateCore/Private/Fonts/FontCache.cpp
+++ b/Engine/Source/Runtime/SlateCore/Private/Fonts/FontCache.cpp
@@ -202,13 +202,13 @@ public:
 	}
 
 	/** Get the default typeface for this composite font */
-	TSharedPtr<const FCachedTypefaceData> GetDefaultTypeface() const
+	const FCachedTypefaceData* GetDefaultTypeface() const
 	{
-		return CachedTypefaces[0];
+		return CachedTypefaces[0].Get();
 	}
 
 	/** Get the typeface that should be used for the given character */
-	TSharedPtr<const FCachedTypefaceData> GetTypefaceForCharacter(const TCHAR InChar) const
+	const FCachedTypefaceData* GetTypefaceForCharacter(const TCHAR InChar) const
 	{
 		const int32 CharIndex = static_cast<int32>(InChar);
 
@@ -227,11 +227,11 @@ public:
 
 			if(CachedRange.Range.Contains(CharIndex))
 			{
-				return CachedRange.CachedTypeface;
+				return CachedRange.CachedTypeface.Get();
 			}
 		}
 
-		return CachedTypefaces[0];
+		return CachedTypefaces[0].Get();
 	}
 
 private:
@@ -246,9 +246,9 @@ private:
 		}
 
 		/** Construct from the given range and typeface */
-		FCachedFontRange(const FInt32Range& InRange, const TSharedPtr<FCachedTypefaceData>& InCachedTypeface)
+		FCachedFontRange(const FInt32Range& InRange, TSharedPtr<FCachedTypefaceData> InCachedTypeface)
 			: Range(InRange)
-			, CachedTypeface(InCachedTypeface)
+			, CachedTypeface(MoveTemp(InCachedTypeface))
 		{
 		}
 
@@ -367,26 +367,11 @@ public:
 #endif // WITH_FREETYPE
 	}
 
-	/** 
-	 * Creates render data for a specific character 
-	 * 
-	 * @param InFontInfo	Information about the font to render the character with
-	 * @param Char			The character to render
-	 * @param OutCharInfo	Will contain the created render data
-	 */
-	void GetRenderData( const FSlateFontInfo& InFontInfo, TCHAR Char, FCharacterRenderData& OutRenderData, const float InScale )
-	{
-#if WITH_FREETYPE
-		float SubFontScalingFactor = 1.0f;
-		const FFontData& FontData = GetFontDataForCharacter(InFontInfo, Char, SubFontScalingFactor);
-		return GetRenderData(FontData, InFontInfo.Size, Char, OutRenderData, InScale * SubFontScalingFactor);
-#endif // WITH_FREETYPE
-	}
-
 	/** 
 	 * Creates render data for a specific character 
 	 * 
 	 * @param InFontData	Raw font data to render the character with
+	 * @param InSize		The size of the font to draw
 	 * @param Char			The character to render
 	 * @param OutCharInfo	Will contain the created render data
 	 */
@@ -546,16 +531,90 @@ public:
 #endif // WITH_FREETYPE
 	}
 
+	/** Get the default font data to use for the given font info */
+	const FFontData& GetDefaultFontData(const FSlateFontInfo& InFontInfo)
+	{
+		static const FFontData DummyFontData;
+
+#if WITH_FREETYPE
+		const FCompositeFont* const ResolvedCompositeFont = InFontInfo.GetCompositeFont();
+		const FCachedTypefaceData* const CachedTypefaceData = GetDefaultCachedTypeface(ResolvedCompositeFont);
+		if(CachedTypefaceData)
+		{
+			// Try to find the correct font from the typeface
+			const FFontData* FoundFontData = CachedTypefaceData->GetFontData(InFontInfo.TypefaceFontName);
+			if(FoundFontData)
+			{
+				return *FoundFontData;
+			}
+
+			// Failing that, return the first font available (the "None" font)
+			FoundFontData = CachedTypefaceData->GetFontData(NAME_None);
+			if(FoundFontData)
+			{
+				return *FoundFontData;
+			}
+		}
+#endif // WITH_FREETYPE
+
+		return DummyFontData;
+	}
+
+	/** Get the font data to use for the given font info and character */
+	const FFontData& GetFontDataForCharacter(const FSlateFontInfo& InFontInfo, const TCHAR InChar, float& OutScalingFactor)
+	{
+		static const FFontData DummyFontData;
+
+#if WITH_FREETYPE
+		const FCompositeFont* const ResolvedCompositeFont = InFontInfo.GetCompositeFont();
+		const FCachedTypefaceData* const CachedTypefaceData = GetCachedTypefaceForCharacter(ResolvedCompositeFont, InChar);
+		if(CachedTypefaceData)
+		{
+			OutScalingFactor = CachedTypefaceData->GetScalingFactor();
+
+			// Try to find the correct font from the typeface
+			const FFontData* FoundFontData = CachedTypefaceData->GetFontData(InFontInfo.TypefaceFontName);
+			if(FoundFontData)
+			{
+				return *FoundFontData;
+			}
+
+			// Failing that, try and find a font by the attributes of the default font with the given name
+			const FCachedTypefaceData* const CachedDefaultTypefaceData = GetDefaultCachedTypeface(ResolvedCompositeFont);
+			if(CachedDefaultTypefaceData && CachedTypefaceData != CachedDefaultTypefaceData)
+			{
+				const FFontData* const FoundDefaultFontData = CachedDefaultTypefaceData->GetFontData(InFontInfo.TypefaceFontName);
+				if(FoundDefaultFontData)
+				{
+					const TSet<FName>& DefaultFontAttributes = GetFontAttributes(*FoundDefaultFontData);
+					FoundFontData = GetBestMatchFontForAttributes(CachedTypefaceData, DefaultFontAttributes);
+					if(FoundFontData)
+					{
+						return *FoundFontData;
+					}
+				}
+			}
+
+			// Failing that, return the first font available (the "None" font)
+			FoundFontData = CachedTypefaceData->GetFontData(NAME_None);
+			if(FoundFontData)
+			{
+				return *FoundFontData;
+			}
+		}
+#endif // WITH_FREETYPE
+
+		OutScalingFactor = 1.0f;
+		return DummyFontData;
+	}
+
 	/**
 	 * @param Whether or not the font has kerning
 	 */
-	bool HasKerning( const FSlateFontInfo& InFontInfo, TCHAR Char )
+	bool HasKerning( const FFontData& InFontData )
 	{
 #if WITH_FREETYPE
-		float SubFontScalingFactor = 1.0f;
-		const FFontData& FontData = GetFontDataForCharacter(InFontInfo, Char, SubFontScalingFactor);
-
-		FT_Face FontFace = GetFontFace( FontData );
+		FT_Face FontFace = GetFontFace( InFontData );
 
 		if ( FontFace == nullptr )
 		{
@@ -571,39 +630,32 @@ public:
 	/**
 	 * Calculates the kerning amount for a pair of characters
 	 *
+	 * @param InFontData	The font that used to draw the string with the first and second characters
+	 * @param InSize		The size of the font to draw
 	 * @param First			The first character in the pair
 	 * @param Second		The second character in the pair
-	 * @param InFontInfo	Information about the font that used to draw the string with the first and second characters
 	 * @return The kerning amount, 0 if no kerning
 	 */
-	int8 GetKerning( TCHAR First, TCHAR Second, const FSlateFontInfo& InFontInfo, float InScale )
+	int8 GetKerning( const FFontData& InFontData, const int32 InSize, TCHAR First, TCHAR Second, const float InScale )
 	{
 #if WITH_FREETYPE
-		float SubFontScalingFactor = 1.0f; // this is fine to reuse below, as we only apply kerning for characters using the same font
-		const FFontData& FirstFontData = GetFontDataForCharacter(InFontInfo, First, SubFontScalingFactor);
-		const FFontData& SecondFontData = GetFontDataForCharacter(InFontInfo, Second, SubFontScalingFactor);
-
-		// Apply the sub-font scale
-		const float FinalScale = InScale * SubFontScalingFactor;
-
 		int32 Kerning = 0;
 		int32* FoundKerning = nullptr;
 
-		FT_Face FontFace = GetFontFace( FirstFontData );
-		FT_Face SecondFontFace = GetFontFace( SecondFontData );
+		FT_Face FontFace = GetFontFace( InFontData );
 
 		// Check if this font has kerning as not all fonts do.
 		// We also can't perform kerning between two separate font faces
-		if( FontFace != nullptr && FontFace == SecondFontFace && FT_HAS_KERNING( FontFace ) )
+		if( FontFace != nullptr && FT_HAS_KERNING( FontFace ) )
 		{
-			FT_Error Error = FT_Set_Char_Size( FontFace, 0, InFontInfo.Size*64, FontCacheConstants::HorizontalDPI, FontCacheConstants::VerticalDPI  );
+			FT_Error Error = FT_Set_Char_Size( FontFace, 0, InSize*64, FontCacheConstants::HorizontalDPI, FontCacheConstants::VerticalDPI  );
 
-			if( FinalScale != 1.0f )
+			if( InScale != 1.0f )
 			{
 				FT_Matrix ScaleMatrix;
 				ScaleMatrix.xy = 0;
-				ScaleMatrix.xx = (FT_Fixed)(FinalScale * 65536);
-				ScaleMatrix.yy = (FT_Fixed)(FinalScale * 65536);
+				ScaleMatrix.xx = (FT_Fixed)(InScale * 65536);
+				ScaleMatrix.yy = (FT_Fixed)(InScale * 65536);
 				ScaleMatrix.yx = 0;
 				FT_Set_Transform( FontFace, &ScaleMatrix, nullptr );
 			}
@@ -623,12 +675,7 @@ public:
 			FT_Get_Kerning( FontFace, FirstIndex, SecondIndex, FT_KERNING_DEFAULT, &KerningVec );
 
 			// Return pixel sizes
-			int32 NewKerning = KerningVec.x / 64;
-
-			// Set this pair's kerning value
-			//KeringPairMap->Add( Pair, NewKerning );
-
-			Kerning = NewKerning;
+			Kerning = KerningVec.x / 64;
 		}
 
 		return Kerning;
@@ -660,110 +707,37 @@ private:
 
 #if WITH_FREETYPE
 	/** Get the cached composite font data for the given composite font */
-	TSharedPtr<const FCachedCompositeFontData> GetCachedCompositeFont(const FCompositeFont* const InCompositeFont)
+	const FCachedCompositeFontData* GetCachedCompositeFont(const FCompositeFont* const InCompositeFont)
 	{
 		if(!InCompositeFont)
 		{
 			return nullptr;
 		}
 
-		TSharedPtr<FCachedCompositeFontData> FoundCompositeFontData = CompositeFontToCachedDataMap.FindRef(InCompositeFont);
-		if(FoundCompositeFontData.IsValid())
+		TSharedPtr<FCachedCompositeFontData>* const FoundCompositeFontData = CompositeFontToCachedDataMap.Find(InCompositeFont);
+		if(FoundCompositeFontData)
 		{
-			return FoundCompositeFontData;
+			return FoundCompositeFontData->Get();
 		}
 
-		return CompositeFontToCachedDataMap.Add(InCompositeFont, MakeShareable(new FCachedCompositeFontData(*InCompositeFont)));
+		return CompositeFontToCachedDataMap.Add(InCompositeFont, MakeShareable(new FCachedCompositeFontData(*InCompositeFont))).Get();
 	}
 
 	/** Get the default typeface for the given composite font */
-	TSharedPtr<const FCachedTypefaceData> GetDefaultCachedTypeface(const FCompositeFont* const InCompositeFont)
+	const FCachedTypefaceData* GetDefaultCachedTypeface(const FCompositeFont* const InCompositeFont)
 	{
-		TSharedPtr<const FCachedCompositeFontData> CachedCompositeFont = GetCachedCompositeFont(InCompositeFont);
-		return (CachedCompositeFont.IsValid()) ? CachedCompositeFont->GetDefaultTypeface() : nullptr;
+		const FCachedCompositeFontData* const CachedCompositeFont = GetCachedCompositeFont(InCompositeFont);
+		return (CachedCompositeFont) ? CachedCompositeFont->GetDefaultTypeface() : nullptr;
 	}
 
 	/** Get the typeface that should be used for the given character */
-	TSharedPtr<const FCachedTypefaceData> GetCachedTypefaceForCharacter(const FCompositeFont* const InCompositeFont, const TCHAR InChar)
+	const FCachedTypefaceData* GetCachedTypefaceForCharacter(const FCompositeFont* const InCompositeFont, const TCHAR InChar)
 	{
-		TSharedPtr<const FCachedCompositeFontData> CachedCompositeFont = GetCachedCompositeFont(InCompositeFont);
-		return (CachedCompositeFont.IsValid()) ? CachedCompositeFont->GetTypefaceForCharacter(InChar) : nullptr;
+		const FCachedCompositeFontData* const CachedCompositeFont = GetCachedCompositeFont(InCompositeFont);
+		return (CachedCompositeFont) ? CachedCompositeFont->GetTypefaceForCharacter(InChar) : nullptr;
 	}
-	
-	/** Get the default font data to use for the given font info */
-	const FFontData& GetDefaultFontData(const FSlateFontInfo& InFontInfo)
-	{
-		static const FFontData DummyFontData;
-
-		const FCompositeFont* const ResolvedCompositeFont = InFontInfo.GetCompositeFont();
-		TSharedPtr<const FCachedTypefaceData> CachedTypefaceData = GetDefaultCachedTypeface(ResolvedCompositeFont);
-		if(CachedTypefaceData.IsValid())
-		{
-			// Try to find the correct font from the typeface
-			const FFontData* FoundFontData = CachedTypefaceData->GetFontData(InFontInfo.TypefaceFontName);
-			if(FoundFontData)
-			{
-				return *FoundFontData;
-			}
 
-			// Failing that, return the first font available (the "None" font)
-			FoundFontData = CachedTypefaceData->GetFontData(NAME_None);
-			if(FoundFontData)
-			{
-				return *FoundFontData;
-			}
-		}
-
-		return DummyFontData;
-	}
-
-	/** Get the font data to use for the given font info and character */
-	const FFontData& GetFontDataForCharacter(const FSlateFontInfo& InFontInfo, const TCHAR InChar, float& OutScalingFactor)
-	{
-		static const FFontData DummyFontData;
-
-		const FCompositeFont* const ResolvedCompositeFont = InFontInfo.GetCompositeFont();
-		TSharedPtr<const FCachedTypefaceData> CachedTypefaceData = GetCachedTypefaceForCharacter(ResolvedCompositeFont, InChar);
-		if(CachedTypefaceData.IsValid())
-		{
-			OutScalingFactor = CachedTypefaceData->GetScalingFactor();
-
-			// Try to find the correct font from the typeface
-			const FFontData* FoundFontData = CachedTypefaceData->GetFontData(InFontInfo.TypefaceFontName);
-			if(FoundFontData)
-			{
-				return *FoundFontData;
-			}
-
-			// Failing that, try and find a font by the attributes of the default font with the given name
-			TSharedPtr<const FCachedTypefaceData> CachedDefaultTypefaceData = GetDefaultCachedTypeface(ResolvedCompositeFont);
-			if(CachedDefaultTypefaceData.IsValid() && CachedTypefaceData != CachedDefaultTypefaceData)
-			{
-				const FFontData* const FoundDefaultFontData = CachedDefaultTypefaceData->GetFontData(InFontInfo.TypefaceFontName);
-				if(FoundDefaultFontData)
-				{
-					const TSet<FName>& DefaultFontAttributes = GetFontAttributes(*FoundDefaultFontData);
-					FoundFontData = GetBestMatchFontForAttributes(CachedTypefaceData, DefaultFontAttributes);
-					if(FoundFontData)
-					{
-						return *FoundFontData;
-					}
-				}
-			}
-
-			// Failing that, return the first font available (the "None" font)
-			FoundFontData = CachedTypefaceData->GetFontData(NAME_None);
-			if(FoundFontData)
-			{
-				return *FoundFontData;
-			}
-		}
-
-		OutScalingFactor = 1.0f;
-		return DummyFontData;
-	}
-
-	const FFontData* GetBestMatchFontForAttributes(const TSharedPtr<const FCachedTypefaceData>& InCachedTypefaceData, const TSet<FName>& InFontAttributes)
+	const FFontData* GetBestMatchFontForAttributes(const FCachedTypefaceData* const InCachedTypefaceData, const TSet<FName>& InFontAttributes)
 	{
 		const FFontData* BestMatchFont = nullptr;
 		int32 BestMatchCount = 0;
@@ -875,10 +849,12 @@ private:
 #endif // WITH_FREETYPE
 };
 
-FKerningTable::FKerningTable( const FSlateFontKey& InFont, const FSlateFontCache& InFontCache )
+FKerningTable::FKerningTable( const FFontData& InFontData, const int32 InFontSize, const float InFontScale, const FSlateFontCache& InFontCache )
 	: DirectAccessTable( nullptr )
 	, FontCache( InFontCache )
-	, FontKey( InFont )
+	, FontData( &InFontData )
+	, FontSize( InFontSize )
+	, FontScale( InFontScale )
 {
 }
 
@@ -897,12 +873,6 @@ FKerningTable::~FKerningTable()
 
 int8 FKerningTable::GetKerning( TCHAR FirstChar, TCHAR SecondChar )
 {
-	const bool bHasKerning = FontCache.HasKerning( FontKey.FontInfo, FirstChar ) && FontCache.HasKerning( FontKey.FontInfo, SecondChar );
-	if( !bHasKerning )
-	{
-		return 0;
-	}
-
 	int8 OutKerning = 0;
 
 	if( FirstChar < FontCacheConstants::DirectAccessSize && SecondChar < FontCacheConstants::DirectAccessSize )
@@ -922,7 +892,7 @@ int8 FKerningTable::GetKerning( TCHAR FirstChar, TCHAR SecondChar )
 		// If the kerning value hasn't been accessed yet, get the value from the cache now 
 		if( OutKerning == MAX_int8 )
 		{
-			OutKerning = FontCache.GetKerning( FirstChar, SecondChar, FontKey.FontInfo, FontKey.Scale );
+			OutKerning = FontCache.GetKerning( *FontData, FontSize, FirstChar, SecondChar, FontScale );
 			DirectAccessTable[Index] = OutKerning;
 		}
 	}
@@ -934,7 +904,7 @@ int8 FKerningTable::GetKerning( TCHAR FirstChar, TCHAR SecondChar )
 		int8* FoundKerning = MappedKerningPairs.Find( KerningPair );
 		if( !FoundKerning )
 		{
-			OutKerning = FontCache.GetKerning( FirstChar, SecondChar, FontKey.FontInfo, FontKey.Scale );
+			OutKerning = FontCache.GetKerning( *FontData, FontSize, FirstChar, SecondChar, FontScale );
 
 #if STATS
 			const uint32 CurrentMemoryUsage = MappedKerningPairs.GetAllocatedSize();
@@ -966,8 +936,7 @@ void FKerningTable::CreateDirectTable()
 }
 
 FCharacterList::FCharacterList( const FSlateFontKey& InFontKey, const FSlateFontCache& InFontCache )
-	: KerningTable( InFontKey, InFontCache )
-	, FontKey( InFontKey )
+	: FontKey( InFontKey )
 	, FontCache( InFontCache )
 	, CompositeFontHistoryRevision( 0 )
 	, MaxDirectIndexedEntries( FontCacheConstants::DirectAccessSize )
@@ -989,7 +958,41 @@ bool FCharacterList::IsStale() const
 
 int8 FCharacterList::GetKerning( TCHAR First, TCHAR Second )
 {
-	return KerningTable.GetKerning( First, Second );
+	const FCharacterEntry& FirstCharacterEntry = GetCharacter( First );
+	const FCharacterEntry& SecondCharacterEntry = GetCharacter( Second );
+
+	// We can only get kerning if both characters are using the same font
+	if( !FirstCharacterEntry.FontData || !FirstCharacterEntry.FontData->BulkDataPtr || *FirstCharacterEntry.FontData != *SecondCharacterEntry.FontData )
+	{
+		return 0;
+	}
+
+	FKerningTable* KerningTablePtr = nullptr;
+
+	// Find or add the kerning table for this font
+	TSharedPtr<FKerningTable>* const FoundKerningTable = KerningTables.Find( FirstCharacterEntry.FontData );
+	if( FoundKerningTable )
+	{
+		KerningTablePtr = FoundKerningTable->Get();
+	}
+	else
+	{
+		// We add a null pointer (rather than a valid kerning table) if the font doesn't have kerning information
+		TSharedPtr<FKerningTable> NewKerningTable;
+		if( FontCache.HasKerning( *FirstCharacterEntry.FontData ) )
+		{
+			NewKerningTable = MakeShareable(new FKerningTable(*FirstCharacterEntry.FontData, FontKey.FontInfo.Size, FirstCharacterEntry.FontScale, FontCache));
+		}
+		KerningTablePtr = KerningTables.Add( FirstCharacterEntry.FontData, MoveTemp(NewKerningTable) ).Get();
+	}
+
+	// A null kerning table pointer means that the font doesn't have kerning information
+	if( KerningTablePtr )
+	{
+		return KerningTablePtr->GetKerning( First, Second );
+	}
+
+	return 0;
 }
 
 uint16 FCharacterList::GetMaxHeight() const
@@ -1012,6 +1015,42 @@ int16 FCharacterList::GetBaseline() const
 	return Baseline;
 }
 
+const FCharacterEntry& FCharacterList::GetCharacter( TCHAR Character )
+{
+	if( Character < MaxDirectIndexedEntries )
+	{
+		// The character can be indexed directly
+
+		int32 NumToAdd = (Character - DirectIndexEntries.Num()) + 1;
+		if( NumToAdd > 0 )
+		{
+			// The character doesn't exist yet in the index and there is not enough space
+			// Resize the array now
+			DirectIndexEntries.AddZeroed( NumToAdd );
+		}
+			
+		FCharacterEntry& CharacterEntry = DirectIndexEntries[ Character ];
+		if( !CharacterEntry.IsValidEntry() )
+		{
+			// Character has not been cached yet
+			CharacterEntry = CacheCharacter( Character );
+		}
+
+		return CharacterEntry;
+	}
+	else
+	{
+		FCharacterEntry& CharacterEntry = MappedEntries.FindOrAdd( Character );
+		if( !CharacterEntry.IsValidEntry() )
+		{
+			// Character has not been cached yet
+			CharacterEntry = CacheCharacter( Character );
+		}
+
+		return CharacterEntry;
+	}
+}
+
 FCharacterEntry& FCharacterList::CacheCharacter( TCHAR Character )
 {
 	SCOPE_CYCLE_COUNTER( STAT_SlateFontCachingTime );
@@ -1047,9 +1086,13 @@ FSlateFontCache::~FSlateFontCache()
 
 bool FSlateFontCache::AddNewEntry( TCHAR Character, const FSlateFontKey& InKey, FCharacterEntry& OutCharacterEntry ) const
 {
-	FCharacterRenderData RenderData;
+	float SubFontScalingFactor = 1.0f;
+	const FFontData& FontData = FTInterface->GetFontDataForCharacter( InKey.FontInfo, Character, SubFontScalingFactor );
+
 	// Render the character 
-	FTInterface->GetRenderData( InKey.FontInfo, Character, RenderData, InKey.Scale );	
+	FCharacterRenderData RenderData;
+	const float FontScale = InKey.Scale * SubFontScalingFactor;
+	FTInterface->GetRenderData( FontData, InKey.FontInfo.Size, Character, RenderData, FontScale );
 
 	// the index of the atlas where the character is stored
 	int32 AtlasIndex = 0;
@@ -1089,6 +1132,8 @@ bool FSlateFontCache::AddNewEntry( TCHAR Character, const FSlateFontKey& InKey,
 
 	if( bSuccess )
 	{
+		OutCharacterEntry.FontData = &FontData;
+		OutCharacterEntry.FontScale = FontScale;
 		OutCharacterEntry.StartU = NewSlot->X + NewSlot->Padding;
 		OutCharacterEntry.StartV = NewSlot->Y + NewSlot->Padding;
 		OutCharacterEntry.USize = NewSlot->Width - 2 * NewSlot->Padding;
@@ -1129,6 +1174,16 @@ FCharacterList& FSlateFontCache::GetCharacterList( const FSlateFontInfo &InFontI
 	return FontToCharacterListCache.Add( FontKey, MakeShareable( new FCharacterList( FontKey, *this ) ) ).Get();
 }
 
+const FFontData& FSlateFontCache::GetDefaultFontData( const FSlateFontInfo& InFontInfo ) const
+{
+	return FTInterface->GetDefaultFontData(InFontInfo);
+}
+
+const FFontData& FSlateFontCache::GetFontDataForCharacter( const FSlateFontInfo& InFontInfo, const TCHAR InChar, float& OutScalingFactor ) const
+{
+	return FTInterface->GetFontDataForCharacter(InFontInfo, InChar, OutScalingFactor);
+}
+
 uint16 FSlateFontCache::GetMaxCharacterHeight( const FSlateFontInfo& InFontInfo, float FontScale ) const
 {
 	return FTInterface->GetMaxHeight(InFontInfo, FontScale);
@@ -1139,19 +1194,19 @@ int16 FSlateFontCache::GetBaseline( const FSlateFontInfo& InFontInfo, float Font
 	return FTInterface->GetBaseline(InFontInfo, FontScale);
 }
 
-int8 FSlateFontCache::GetKerning( TCHAR First, TCHAR Second, const FSlateFontInfo& InFontInfo, float Scale ) const 
+int8 FSlateFontCache::GetKerning( const FFontData& InFontData, const int32 InSize, TCHAR First, TCHAR Second, float Scale ) const
 {
-	return FTInterface->GetKerning( First, Second, InFontInfo, Scale );
+	return FTInterface->GetKerning(InFontData, InSize, First, Second, Scale);
 }
 
-bool FSlateFontCache::HasKerning( const FSlateFontInfo& InFontInfo, TCHAR Char ) const
+bool FSlateFontCache::HasKerning( const FFontData& InFontData ) const
 {
-	return FTInterface->HasKerning( InFontInfo, Char );
+	return FTInterface->HasKerning(InFontData);
 }
 
 const TSet<FName>& FSlateFontCache::GetFontAttributes( const FFontData& InFontData ) const
 {
-	return FTInterface->GetFontAttributes( InFontData );
+	return FTInterface->GetFontAttributes(InFontData);
 }
 
 void FSlateFontCache::FlushObject( const UObject* const InObject )
diff --git a/Engine/Source/Runtime/SlateCore/Public/Fonts/CompositeFont.h b/Engine/Source/Runtime/SlateCore/Public/Fonts/CompositeFont.h
--- a/Engine/Source/Runtime/SlateCore/Public/Fonts/CompositeFont.h
+++ b/Engine/Source/Runtime/SlateCore/Public/Fonts/CompositeFont.h
@@ -35,6 +35,12 @@ struct SLATECORE_API FFontData
 	/** Set the new font data from a filename, and the data associated with that file */
 	void SetFont(FString InFontFilename, const UFontBulkData* const InBulkData);
 
+	/** Check to see whether this font data is equal to the other font data */
+	bool operator==(const FFontData& Other) const;
+
+	/** Check to see whether this font data is not equal to the other font data */
+	bool operator!=(const FFontData& Other) const;
+
 	/** The filename of the font to use - this may not always exist on disk, as we may have previously loaded and cached the font data inside an asset */
 	UPROPERTY(EditAnywhere, BlueprintReadWrite, Category=Font)
 	FString FontFilename;
diff --git a/Engine/Source/Runtime/SlateCore/Public/Fonts/FontCache.h b/Engine/Source/Runtime/SlateCore/Public/Fonts/FontCache.h
--- a/Engine/Source/Runtime/SlateCore/Public/Fonts/FontCache.h
+++ b/Engine/Source/Runtime/SlateCore/Public/Fonts/FontCache.h
@@ -9,6 +9,10 @@ class FFreeTypeInterface;
 /** Information for rendering one character */
 struct SLATECORE_API FCharacterEntry
 {
+	/** The raw font data this character was rendered with */
+	const FFontData* FontData;
+	/** Scale that was applied when rendering this character */
+	float FontScale;
 	/** Start X location of the character in the texture */
 	float StartU;
 	/** Start Y location of the character in the texture */
@@ -27,7 +31,7 @@ struct SLATECORE_API FCharacterEntry
 	int16 XAdvance;
 	/** Index to a specific texture in the font cache. */
 	uint8 TextureIndex;
-	/** 1 if this entry is valid, 0 otherwize. */
+	/** 1 if this entry is valid, 0 otherwise. */
 	uint8 Valid;
 
 	FCharacterEntry()
@@ -71,12 +75,12 @@ struct SLATECORE_API FKerningPair
 class FSlateFontCache;
 
 /**
- * A Kerning table for a single font key
+ * A Kerning table for a single raw font
  */
 class SLATECORE_API FKerningTable
 {
 public:
-	FKerningTable( const FSlateFontKey& InFont, const FSlateFontCache& InFontCache );
+	FKerningTable( const FFontData& InFontData, const int32 InFontSize, const float InFontScale, const FSlateFontCache& InFontCache );
 	~FKerningTable();
 
 	/**
@@ -87,20 +91,25 @@ public:
 	 * @return The kerning value
 	 */
 	int8 GetKerning( TCHAR FirstChar, TCHAR SecondChar );
-private:
+
 	/**
 	 * Allocated memory for the directly indexed kerning table
 	 */
 	void CreateDirectTable();
+
 private:
 	/** Extended kerning Pairs which are mapped to save memory */
 	TMap<FKerningPair,int8> MappedKerningPairs;
-	/** Direct access kerning table for ascii chars.  Note its very important that this stays small.  Kerni */
+	/** Direct access kerning table for ascii chars.  Note its very important that this stays small. */
 	int8* DirectAccessTable;
 	/** Interface to freetype for accessing new kerning values */
 	const class FSlateFontCache& FontCache;
-	/** Font this kerning table is for */
-	FSlateFontKey FontKey;
+	/** Font for this kerning table */
+	const FFontData* FontData;
+	/** Size of the font used by this kerning table */
+	int32 FontSize;
+	/** Scale of the font used by this kerning table */
+	float FontScale;
 };
 
 
@@ -130,38 +139,7 @@ public:
 	 */
 	const FCharacterEntry& operator[]( TCHAR Character )
 	{
-		if( Character < MaxDirectIndexedEntries )
-		{
-			// The character can be indexed directly
-
-			int32 NumToAdd = (Character - DirectIndexEntries.Num()) + 1;
-			if( NumToAdd > 0 )
-			{
-				// The character doesnt exist yet in the index and there is not enough space
-				// Resize the array now
-				DirectIndexEntries.AddZeroed( NumToAdd );
-			}
-			
-			FCharacterEntry& CharacterEntry = DirectIndexEntries[ Character ];
-			if( !CharacterEntry.IsValidEntry() )
-			{
-				// Character has not been cached yet
-				CharacterEntry = CacheCharacter( Character );
-			}
-
-			return CharacterEntry;
-		}
-		else
-		{
-			FCharacterEntry& CharacterEntry = MappedEntries.FindOrAdd( Character );
-			if( !CharacterEntry.IsValidEntry() )
-			{
-				// Character has not been cached yet
-				CharacterEntry = CacheCharacter( Character );
-			}
-
-			return CharacterEntry;
-		}
+		return GetCharacter( Character );
 	}
 
 	/** Check to see if our cached data is potentially stale for our font */
@@ -190,6 +168,15 @@ public:
 
 private:
 
+	/**
+	 * Gets data about how to render and measure a character 
+	 * Caching and atlasing it if needed
+	 *
+	 * @param Character	The character to get
+	 * @return Data about the character
+	 */
+	const FCharacterEntry& GetCharacter( TCHAR Character );
+
 	/**
 	 * Caches a new character
 	 * 
@@ -204,8 +191,8 @@ private:
 	TMap<TCHAR, FCharacterEntry> MappedEntries; 
 	/** Directly indexed entries for fast lookup */
 	TArray<FCharacterEntry> DirectIndexEntries;
-	/** Table of kerning values for this font */
-	FKerningTable KerningTable;
+	/** Tables of kerning values for this font */
+	TMap<const FFontData*, TSharedPtr<FKerningTable>> KerningTables;
 	/** Font for this character list */
 	FSlateFontKey FontKey;
 	/** Reference to the font cache for accessing new unseen characters */
@@ -286,19 +273,24 @@ public:
 	class FTextureResource* GetEngineTextureResource( uint32 Index ) { return FontAtlases[Index]->GetEngineTexture(); }
 
 	/**
-	 * Calculates the kerning amount for a pair of characters
+	 * Returns the font to use from the default typeface
 	 *
-	 * @param First			The first character in the pair
-	 * @param Second		The second character in the pair
-	 * @param InFontInfo	Information about the font that used to draw the string with the first and second characters
-	 * @return The kerning amount, 0 if no kerning
+	 * @param InFontInfo	A descriptor of the font to get the default typeface for
+	 * 
+	 * @return The raw font data
 	 */
-	int8 GetKerning( TCHAR First, TCHAR Second, const FSlateFontInfo& InFontInfo, float Scale ) const;
+	const FFontData& GetDefaultFontData( const FSlateFontInfo& InFontInfo ) const;
 
 	/**
-	 * @return Whether or not the font used by the given character has kerning information
+	 * Returns the font to use from the typeface associated with the given character
+	 *
+	 * @param InFontInfo		A descriptor of the font to get the typeface for
+	 * @param InChar			The character to get the typeface associated with
+	 * @param OutScalingFactor	The scaling factor applied to characters rendered with the given font
+	 * 
+	 * @return The raw font data
 	 */
-	bool HasKerning( const FSlateFontInfo& InFontInfo, TCHAR Char ) const;
+	const FFontData& GetFontDataForCharacter( const FSlateFontInfo& InFontInfo, const TCHAR InChar, float& OutScalingFactor ) const;
 
 	/**
 	 * Returns the height of the largest character in the font. 
@@ -320,6 +312,22 @@ public:
 	 */
 	int16 GetBaseline( const FSlateFontInfo& InFontInfo, float FontScale ) const;
 
+	/**
+	 * Calculates the kerning amount for a pair of characters
+	 *
+	 * @param InFontData	The font that used to draw the string with the first and second characters
+	 * @param InSize		The size of the font to draw
+	 * @param First			The first character in the pair
+	 * @param Second		The second character in the pair
+	 * @return The kerning amount, 0 if no kerning
+	 */
+	int8 GetKerning( const FFontData& InFontData, const int32 InSize, TCHAR First, TCHAR Second, float Scale ) const;
+
+	/**
+	 * @return Whether or not the font used has kerning information
+	 */
+	bool HasKerning( const FFontData& InFontData ) const;
+
 	/**
 	 * Returns the font attributes for the specified font.
 	 *
