Text rendering perf. improvements

UETOOL-129 - Investigate SOutputLog performance degradation

This further improves the performance of GetKerning by removing some map lookups.

FCharacterList doesn't really need multiple kerning tables, as each character in a composite font will only ever map to a single raw font for a given FSlateFontInfo. This means we can use a single kerning table, *however* we need to know whether the font a character is using has kerning *without* having to query it from the FreeType interface (as this involves a face lookup, which is relatively slow).

FCharacterList now caches this information in its character entries, so that we can immediately know whether a character will need kerning or not. Additionally, GetKerning can now optionally take a pair of FCharacterEntrys (rather than a pair of TCHARs), which avoids the need to lookup the character entry from the FCharacterList (which is handy when you already have the correct FCharacterEntry to hand, as you do in the case of text batching and measuring).

diff --git a/Engine/Source/Runtime/Engine/Private/UserInterface/CanvasItem.cpp b/Engine/Source/Runtime/Engine/Private/UserInterface/CanvasItem.cpp
--- a/Engine/Source/Runtime/Engine/Private/UserInterface/CanvasItem.cpp
+++ b/Engine/Source/Runtime/Engine/Private/UserInterface/CanvasItem.cpp
@@ -1050,7 +1050,7 @@ void FCanvasTextItem::DrawStringInternal_RuntimeCache( FCanvas* InCanvas, const
 
 	float LineX = 0;
 
-	TCHAR PreviousChar = 0;
+	const FCharacterEntry* PreviousCharEntry = nullptr;
 
 	int32 Kerning = 0;
 
@@ -1107,7 +1107,7 @@ void FCanvasTextItem::DrawStringInternal_RuntimeCache( FCanvas* InCanvas, const
 
 			if( !bIsWhitespace && CharIndex > 0 )
 			{
-				Kerning = CharacterList.GetKerning( PreviousChar, CurrentChar );
+				Kerning = CharacterList.GetKerning( *PreviousCharEntry, Entry );
 			}
 			else
 			{
@@ -1115,7 +1115,7 @@ void FCanvasTextItem::DrawStringInternal_RuntimeCache( FCanvas* InCanvas, const
 			}
 
 			LineX += Kerning;
-			PreviousChar = CurrentChar;
+			PreviousCharEntry = &Entry;
 
 			if( !bIsWhitespace )
 			{
diff --git a/Engine/Source/Runtime/SlateCore/Private/Fonts/FontCache.cpp b/Engine/Source/Runtime/SlateCore/Private/Fonts/FontCache.cpp
--- a/Engine/Source/Runtime/SlateCore/Private/Fonts/FontCache.cpp
+++ b/Engine/Source/Runtime/SlateCore/Private/Fonts/FontCache.cpp
@@ -507,6 +507,7 @@ public:
 
 		// Set measurement info for this character
 		OutRenderData.Char = Char;
+		OutRenderData.HasKerning = FT_HAS_KERNING( FontFace ) != 0;
 		OutRenderData.MeasureInfo.SizeX = Bitmap->width;
 		OutRenderData.MeasureInfo.SizeY = Bitmap->rows;
 		OutRenderData.MaxHeight = Height;
@@ -849,12 +850,9 @@ private:
 #endif // WITH_FREETYPE
 };
 
-FKerningTable::FKerningTable( const FFontData& InFontData, const int32 InFontSize, const float InFontScale, const FSlateFontCache& InFontCache )
+FKerningTable::FKerningTable( const FSlateFontCache& InFontCache )
 	: DirectAccessTable( nullptr )
 	, FontCache( InFontCache )
-	, FontData( &InFontData )
-	, FontSize( InFontSize )
-	, FontScale( InFontScale )
 {
 }
 
@@ -871,7 +869,7 @@ FKerningTable::~FKerningTable()
 	DEC_MEMORY_STAT_BY( STAT_SlateFontKerningTableMemory, MappedKerningPairs.GetAllocatedSize() );
 }
 
-int8 FKerningTable::GetKerning( TCHAR FirstChar, TCHAR SecondChar )
+int8 FKerningTable::GetKerning( const FFontData& InFontData, const int32 InSize, TCHAR FirstChar, TCHAR SecondChar, float InScale )
 {
 	int8 OutKerning = 0;
 
@@ -892,7 +890,7 @@ int8 FKerningTable::GetKerning( TCHAR FirstChar, TCHAR SecondChar )
 		// If the kerning value hasn't been accessed yet, get the value from the cache now 
 		if( OutKerning == MAX_int8 )
 		{
-			OutKerning = FontCache.GetKerning( *FontData, FontSize, FirstChar, SecondChar, FontScale );
+			OutKerning = FontCache.GetKerning( InFontData, InSize, FirstChar, SecondChar, InScale );
 			DirectAccessTable[Index] = OutKerning;
 		}
 	}
@@ -904,7 +902,7 @@ int8 FKerningTable::GetKerning( TCHAR FirstChar, TCHAR SecondChar )
 		int8* FoundKerning = MappedKerningPairs.Find( KerningPair );
 		if( !FoundKerning )
 		{
-			OutKerning = FontCache.GetKerning( *FontData, FontSize, FirstChar, SecondChar, FontScale );
+			OutKerning = FontCache.GetKerning( InFontData, InSize, FirstChar, SecondChar, InScale );
 
 #if STATS
 			const uint32 CurrentMemoryUsage = MappedKerningPairs.GetAllocatedSize();
@@ -936,7 +934,8 @@ void FKerningTable::CreateDirectTable()
 }
 
 FCharacterList::FCharacterList( const FSlateFontKey& InFontKey, const FSlateFontCache& InFontCache )
-	: FontKey( InFontKey )
+	: KerningTable( InFontCache )
+	, FontKey( InFontKey )
 	, FontCache( InFontCache )
 	, CompositeFontHistoryRevision( 0 )
 	, MaxDirectIndexedEntries( FontCacheConstants::DirectAccessSize )
@@ -956,40 +955,26 @@ bool FCharacterList::IsStale() const
 	return !CompositeFont || CompositeFontHistoryRevision != CompositeFont->HistoryRevision;
 }
 
-int8 FCharacterList::GetKerning( TCHAR First, TCHAR Second )
+int8 FCharacterList::GetKerning( TCHAR FirstChar, TCHAR SecondChar )
 {
-	const FCharacterEntry& FirstCharacterEntry = GetCharacter( First );
-	const FCharacterEntry& SecondCharacterEntry = GetCharacter( Second );
+	return GetKerning( GetCharacter( FirstChar ), GetCharacter( SecondChar ) );
+}
 
+int8 FCharacterList::GetKerning( const FCharacterEntry& FirstCharacterEntry, const FCharacterEntry& SecondCharacterEntry )
+{
 	// We can only get kerning if both characters are using the same font
-	if( !FirstCharacterEntry.FontData || !FirstCharacterEntry.FontData->BulkDataPtr || *FirstCharacterEntry.FontData != *SecondCharacterEntry.FontData )
-	{
-		return 0;
-	}
-
-	FKerningTable* KerningTablePtr = nullptr;
-
-	// Find or add the kerning table for this font
-	TSharedPtr<FKerningTable>* const FoundKerningTable = KerningTables.Find( FirstCharacterEntry.FontData );
-	if( FoundKerningTable )
-	{
-		KerningTablePtr = FoundKerningTable->Get();
-	}
-	else
-	{
-		// We add a null pointer (rather than a valid kerning table) if the font doesn't have kerning information
-		TSharedPtr<FKerningTable> NewKerningTable;
-		if( FontCache.HasKerning( *FirstCharacterEntry.FontData ) )
-		{
-			NewKerningTable = MakeShareable(new FKerningTable(*FirstCharacterEntry.FontData, FontKey.FontInfo.Size, FirstCharacterEntry.FontScale, FontCache));
-		}
-		KerningTablePtr = KerningTables.Add( FirstCharacterEntry.FontData, MoveTemp(NewKerningTable) ).Get();
-	}
-
-	// A null kerning table pointer means that the font doesn't have kerning information
-	if( KerningTablePtr )
+	if( FirstCharacterEntry.FontData && 
+		FirstCharacterEntry.FontData->BulkDataPtr && 
+		FirstCharacterEntry.HasKerning && 
+		*FirstCharacterEntry.FontData == *SecondCharacterEntry.FontData )
 	{
-		return KerningTablePtr->GetKerning( First, Second );
+		return KerningTable.GetKerning( 
+			*FirstCharacterEntry.FontData, 
+			FontKey.FontInfo.Size, 
+			FirstCharacterEntry.Character, 
+			SecondCharacterEntry.Character, 
+			FirstCharacterEntry.FontScale 
+			);
 	}
 
 	return 0;
@@ -1132,6 +1117,7 @@ bool FSlateFontCache::AddNewEntry( TCHAR Character, const FSlateFontKey& InKey,
 
 	if( bSuccess )
 	{
+		OutCharacterEntry.Character = Character;
 		OutCharacterEntry.FontData = &FontData;
 		OutCharacterEntry.FontScale = FontScale;
 		OutCharacterEntry.StartU = NewSlot->X + NewSlot->Padding;
@@ -1144,6 +1130,7 @@ bool FSlateFontCache::AddNewEntry( TCHAR Character, const FSlateFontKey& InKey,
 		OutCharacterEntry.GlobalDescender = GetBaseline(InKey.FontInfo, InKey.Scale); // All fonts within a composite font need to use the baseline of the default font
 		OutCharacterEntry.HorizontalOffset = RenderData.MeasureInfo.HorizontalOffset;
 		OutCharacterEntry.TextureIndex = AtlasIndex;
+		OutCharacterEntry.HasKerning = RenderData.HasKerning;
 		OutCharacterEntry.Valid = 1;
 	}
 
diff --git a/Engine/Source/Runtime/SlateCore/Private/Fonts/FontMeasure.cpp b/Engine/Source/Runtime/SlateCore/Private/Fonts/FontMeasure.cpp
--- a/Engine/Source/Runtime/SlateCore/Private/Fonts/FontMeasure.cpp
+++ b/Engine/Source/Runtime/SlateCore/Private/Fonts/FontMeasure.cpp
@@ -168,12 +168,12 @@ FVector2D FSlateFontMeasure::MeasureStringInternal( const FString& Text, int32 S
 	// Accumulated height of this block of text
 	int32 StringSizeY = MaxHeight;
 	// The previous char (for kerning)
-	TCHAR PreviousChar = 0;
+	const FCharacterEntry* PreviousCharEntry = nullptr;
 
 	//If we are measuring a range then we should take into account the kerning with the character before the start of the range
 	if ( !DoesStartAtBeginning && IncludeKerningWithPrecedingChar )
 	{
-		PreviousChar = Text[ StartIndex - 1 ];
+		PreviousCharEntry = &CharacterList[ Text[ StartIndex - 1 ] ];
 	}
 
 	int32 FinalPosX = 0;
@@ -202,14 +202,14 @@ FVector2D FSlateFontMeasure::MeasureStringInternal( const FString& Text, int32 S
 			int32 Kerning = 0;
 			if( CharIndex > 0 )
 			{
-				Kerning = CharacterList.GetKerning( PreviousChar, CurrentChar );
+				Kerning = CharacterList.GetKerning( *PreviousCharEntry, Entry );
 			}
 
 			const int32 TotalCharSpacing = 
 				Kerning + Entry.HorizontalOffset +		// Width is any kerning plus how much to advance the position when drawing a new character
 				Entry.XAdvance;	// How far we advance 
 
-			PreviousChar = CurrentChar;
+			PreviousCharEntry = &Entry;
 
 			CurrentX += Kerning + Entry.XAdvance;
 
diff --git a/Engine/Source/Runtime/SlateCore/Private/Rendering/ElementBatcher.cpp b/Engine/Source/Runtime/SlateCore/Private/Rendering/ElementBatcher.cpp
--- a/Engine/Source/Runtime/SlateCore/Private/Rendering/ElementBatcher.cpp
+++ b/Engine/Source/Runtime/SlateCore/Private/Rendering/ElementBatcher.cpp
@@ -525,7 +525,7 @@ void FSlateElementBatcher::AddTextElement(const FSlateDrawElement& DrawElement)
 
 	float LineX = 0;
 
-	TCHAR PreviousChar = 0;
+	const FCharacterEntry* PreviousCharEntry = nullptr;
 
 	int32 Kerning = 0;
 
@@ -540,7 +540,7 @@ void FSlateElementBatcher::AddTextElement(const FSlateDrawElement& DrawElement)
 
 	for( int32 CharIndex = 0; CharIndex < Text.Len(); ++CharIndex )
 	{
-		TCHAR CurrentChar = Text[ CharIndex ];
+		const TCHAR CurrentChar = Text[ CharIndex ];
 
 		const bool IsNewline = (CurrentChar == '\n');
 
@@ -577,7 +577,7 @@ void FSlateElementBatcher::AddTextElement(const FSlateDrawElement& DrawElement)
 
 			if( !bIsWhitespace && CharIndex > 0 )
 			{
-				Kerning = CharacterList.GetKerning( PreviousChar, CurrentChar );
+				Kerning = CharacterList.GetKerning( *PreviousCharEntry, Entry );
 			}
 			else
 			{
@@ -585,7 +585,7 @@ void FSlateElementBatcher::AddTextElement(const FSlateDrawElement& DrawElement)
 			}
 
 			LineX += Kerning;
-			PreviousChar = CurrentChar;
+			PreviousCharEntry = &Entry;
 
 			if( !bIsWhitespace )
 			{
diff --git a/Engine/Source/Runtime/SlateCore/Public/Fonts/FontCache.h b/Engine/Source/Runtime/SlateCore/Public/Fonts/FontCache.h
--- a/Engine/Source/Runtime/SlateCore/Public/Fonts/FontCache.h
+++ b/Engine/Source/Runtime/SlateCore/Public/Fonts/FontCache.h
@@ -9,6 +9,8 @@ class FFreeTypeInterface;
 /** Information for rendering one character */
 struct SLATECORE_API FCharacterEntry
 {
+	/** The character this entry is for */
+	TCHAR Character;
 	/** The raw font data this character was rendered with */
 	const FFontData* FontData;
 	/** Scale that was applied when rendering this character */
@@ -31,8 +33,10 @@ struct SLATECORE_API FCharacterEntry
 	int16 XAdvance;
 	/** Index to a specific texture in the font cache. */
 	uint8 TextureIndex;
+	/** 1 if this entry has kerning, 0 otherwise. */
+	bool HasKerning;
 	/** 1 if this entry is valid, 0 otherwise. */
-	uint8 Valid;
+	bool Valid;
 
 	FCharacterEntry()
 	{
@@ -46,7 +50,7 @@ struct SLATECORE_API FCharacterEntry
 	 */
 	FORCEINLINE bool IsValidEntry() const
 	{
-		return Valid == 1;
+		return Valid;
 	}
 };
 
@@ -75,12 +79,12 @@ struct SLATECORE_API FKerningPair
 class FSlateFontCache;
 
 /**
- * A Kerning table for a single raw font
+ * A Kerning table for a single font key
  */
 class SLATECORE_API FKerningTable
 {
 public:
-	FKerningTable( const FFontData& InFontData, const int32 InFontSize, const float InFontScale, const FSlateFontCache& InFontCache );
+	FKerningTable( const FSlateFontCache& InFontCache );
 	~FKerningTable();
 
 	/**
@@ -90,7 +94,7 @@ public:
 	 * @param SecondChar	The second character in the pair
 	 * @return The kerning value
 	 */
-	int8 GetKerning( TCHAR FirstChar, TCHAR SecondChar );
+	int8 GetKerning( const FFontData& InFontData, const int32 InSize, TCHAR FirstChar, TCHAR SecondChar, float InScale );
 
 	/**
 	 * Allocated memory for the directly indexed kerning table
@@ -104,12 +108,6 @@ private:
 	int8* DirectAccessTable;
 	/** Interface to freetype for accessing new kerning values */
 	const class FSlateFontCache& FontCache;
-	/** Font for this kerning table */
-	const FFontData* FontData;
-	/** Size of the font used by this kerning table */
-	int32 FontSize;
-	/** Scale of the font used by this kerning table */
-	float FontScale;
 };
 
 
@@ -152,7 +150,16 @@ public:
 	 * @param SecondChar	The second character in the pair
 	 * @return The kerning value
 	 */
-	int8 GetKerning( TCHAR First, TCHAR Second );
+	int8 GetKerning( TCHAR FirstChar, TCHAR SecondChar );
+
+	/**
+	 * Gets a kerning value for a pair of character entries
+	 *
+	 * @param FirstCharacterEntry	The first character entry in the pair
+	 * @param SecondCharacterEntry	The second character entry in the pair
+	 * @return The kerning value
+	 */
+	int8 GetKerning( const FCharacterEntry& FirstCharacterEntry, const FCharacterEntry& SecondCharacterEntry );
 
 	/**
 	 * @return The global max height for any character in this font
@@ -191,8 +198,8 @@ private:
 	TMap<TCHAR, FCharacterEntry> MappedEntries; 
 	/** Directly indexed entries for fast lookup */
 	TArray<FCharacterEntry> DirectIndexEntries;
-	/** Tables of kerning values for this font */
-	TMap<const FFontData*, TSharedPtr<FKerningTable>> KerningTables;
+	/** Table of kerning values for this font */
+	FKerningTable KerningTable;
 	/** Font for this character list */
 	FSlateFontKey FontKey;
 	/** Reference to the font cache for accessing new unseen characters */
diff --git a/Engine/Source/Runtime/SlateCore/Public/Fonts/FontTypes.h b/Engine/Source/Runtime/SlateCore/Public/Fonts/FontTypes.h
--- a/Engine/Source/Runtime/SlateCore/Public/Fonts/FontTypes.h
+++ b/Engine/Source/Runtime/SlateCore/Public/Fonts/FontTypes.h
@@ -61,6 +61,8 @@ struct FCharacterRenderData
 	uint16 MaxHeight;
 	/** The character that was rendered */
 	TCHAR Char;	
+	/** Whether or not the character has kerning */
+	bool HasKerning;
 };
 
 /** 
