Fixed crash in FSlateFontMeasure::MeasureStringInternal

It was possible for PreviousCharEntry to be null even when CharIndex was greater than zero, as we can measure a sub-range of a string without setting IncludeKerningWithPrecedingChar to true.

diff --git a/Engine/Source/Runtime/SlateCore/Private/Fonts/FontMeasure.cpp b/Engine/Source/Runtime/SlateCore/Private/Fonts/FontMeasure.cpp
--- a/Engine/Source/Runtime/SlateCore/Private/Fonts/FontMeasure.cpp
+++ b/Engine/Source/Runtime/SlateCore/Private/Fonts/FontMeasure.cpp
@@ -200,7 +200,7 @@ FVector2D FSlateFontMeasure::MeasureStringInternal( const FString& Text, int32 S
 			const FCharacterEntry& Entry = CharacterList[CurrentChar];
 
 			int32 Kerning = 0;
-			if( CharIndex > 0 )
+			if( PreviousCharEntry )
 			{
 				Kerning = CharacterList.GetKerning( *PreviousCharEntry, Entry );
 			}
