Removed duplicate hashing of FSlateFontKey

UETOOL-129 - Investigate SOutputLog performance degradation

This caches the hash of FSlateFontKey to avoid some duplicate hashes when looking up items from a map (these keys are immutable), as it can happen a lot during text layout flowing.

diff --git a/Engine/Source/Runtime/SlateCore/Private/Fonts/FontCache.cpp b/Engine/Source/Runtime/SlateCore/Private/Fonts/FontCache.cpp
--- a/Engine/Source/Runtime/SlateCore/Private/Fonts/FontCache.cpp
+++ b/Engine/Source/Runtime/SlateCore/Private/Fonts/FontCache.cpp
@@ -942,7 +942,7 @@ FCharacterList::FCharacterList( const FSlateFontKey& InFontKey, const FSlateFont
 	, MaxHeight( 0 )
 	, Baseline( 0 )
 {
-	const FCompositeFont* const CompositeFont = InFontKey.FontInfo.GetCompositeFont();
+	const FCompositeFont* const CompositeFont = InFontKey.GetFontInfo().GetCompositeFont();
 	if( CompositeFont )
 	{
 		CompositeFontHistoryRevision = CompositeFont->HistoryRevision;
@@ -951,7 +951,7 @@ FCharacterList::FCharacterList( const FSlateFontKey& InFontKey, const FSlateFont
 
 bool FCharacterList::IsStale() const
 {
-	const FCompositeFont* const CompositeFont = FontKey.FontInfo.GetCompositeFont();
+	const FCompositeFont* const CompositeFont = FontKey.GetFontInfo().GetCompositeFont();
 	return !CompositeFont || CompositeFontHistoryRevision != CompositeFont->HistoryRevision;
 }
 
@@ -970,7 +970,7 @@ int8 FCharacterList::GetKerning( const FCharacterEntry& FirstCharacterEntry, con
 	{
 		return KerningTable.GetKerning( 
 			*FirstCharacterEntry.FontData, 
-			FontKey.FontInfo.Size, 
+			FontKey.GetFontInfo().Size, 
 			FirstCharacterEntry.Character, 
 			SecondCharacterEntry.Character, 
 			FirstCharacterEntry.FontScale 
@@ -984,7 +984,7 @@ uint16 FCharacterList::GetMaxHeight() const
 {
 	if( MaxHeight == 0 )
 	{
-		MaxHeight = FontCache.GetMaxCharacterHeight( FontKey.FontInfo, FontKey.Scale );
+		MaxHeight = FontCache.GetMaxCharacterHeight( FontKey.GetFontInfo(), FontKey.GetScale() );
 	}
 
 	return MaxHeight;
@@ -994,7 +994,7 @@ int16 FCharacterList::GetBaseline() const
 {
 	if( Baseline == 0 )
 	{
-		Baseline = FontCache.GetBaseline( FontKey.FontInfo, FontKey.Scale );
+		Baseline = FontCache.GetBaseline( FontKey.GetFontInfo(), FontKey.GetScale() );
 	}
 
 	return Baseline;
@@ -1072,12 +1072,12 @@ FSlateFontCache::~FSlateFontCache()
 bool FSlateFontCache::AddNewEntry( TCHAR Character, const FSlateFontKey& InKey, FCharacterEntry& OutCharacterEntry ) const
 {
 	float SubFontScalingFactor = 1.0f;
-	const FFontData& FontData = FTInterface->GetFontDataForCharacter( InKey.FontInfo, Character, SubFontScalingFactor );
+	const FFontData& FontData = FTInterface->GetFontDataForCharacter( InKey.GetFontInfo(), Character, SubFontScalingFactor );
 
 	// Render the character 
 	FCharacterRenderData RenderData;
-	const float FontScale = InKey.Scale * SubFontScalingFactor;
-	FTInterface->GetRenderData( FontData, InKey.FontInfo.Size, Character, RenderData, FontScale );
+	const float FontScale = InKey.GetScale() * SubFontScalingFactor;
+	FTInterface->GetRenderData( FontData, InKey.GetFontInfo().Size, Character, RenderData, FontScale );
 
 	// the index of the atlas where the character is stored
 	int32 AtlasIndex = 0;
@@ -1127,7 +1127,7 @@ bool FSlateFontCache::AddNewEntry( TCHAR Character, const FSlateFontKey& InKey,
 		OutCharacterEntry.TextureIndex = 0;
 		OutCharacterEntry.XAdvance = RenderData.MeasureInfo.XAdvance;
 		OutCharacterEntry.VerticalOffset = RenderData.MeasureInfo.VerticalOffset;
-		OutCharacterEntry.GlobalDescender = GetBaseline(InKey.FontInfo, InKey.Scale); // All fonts within a composite font need to use the baseline of the default font
+		OutCharacterEntry.GlobalDescender = GetBaseline(InKey.GetFontInfo(), InKey.GetScale()); // All fonts within a composite font need to use the baseline of the default font
 		OutCharacterEntry.HorizontalOffset = RenderData.MeasureInfo.HorizontalOffset;
 		OutCharacterEntry.TextureIndex = AtlasIndex;
 		OutCharacterEntry.HasKerning = RenderData.HasKerning;
@@ -1140,7 +1140,7 @@ bool FSlateFontCache::AddNewEntry( TCHAR Character, const FSlateFontKey& InKey,
 FCharacterList& FSlateFontCache::GetCharacterList( const FSlateFontInfo &InFontInfo, float FontScale ) const
 {
 	// Create a key for looking up each character
-	FSlateFontKey FontKey( InFontInfo, FontScale );
+	const FSlateFontKey FontKey( InFontInfo, FontScale );
 
 	TSharedRef< class FCharacterList >* CachedCharacterList = FontToCharacterListCache.Find( FontKey );
 
@@ -1206,7 +1206,7 @@ void FSlateFontCache::FlushObject( const UObject* const InObject )
 	bool bHasRemovedEntries = false;
 	for( auto It = FontToCharacterListCache.CreateIterator(); It; ++It )
 	{
-		if( It.Key().FontInfo.FontObject == InObject)
+		if( It.Key().GetFontInfo().FontObject == InObject)
 		{
 			bHasRemovedEntries = true;
 			It.RemoveCurrent();
diff --git a/Engine/Source/Runtime/SlateCore/Public/Fonts/FontTypes.h b/Engine/Source/Runtime/SlateCore/Public/Fonts/FontTypes.h
--- a/Engine/Source/Runtime/SlateCore/Public/Fonts/FontTypes.h
+++ b/Engine/Source/Runtime/SlateCore/Public/Fonts/FontTypes.h
@@ -3,28 +3,45 @@
 
 struct SLATECORE_API FSlateFontKey
 {
-	FSlateFontInfo FontInfo;
-	float Scale;
-
-	FSlateFontKey( const FSlateFontInfo& InInfo, float InScale )
+public:
+	FSlateFontKey( const FSlateFontInfo& InInfo, const float InScale )
 		: FontInfo( InInfo )
 		, Scale( InScale )
+		, KeyHash( 0 )
 	{
+		KeyHash = HashCombine(KeyHash, GetTypeHash(FontInfo));
+		KeyHash = HashCombine(KeyHash, GetTypeHash(Scale));
+	}
 
+	FORCEINLINE const FSlateFontInfo& GetFontInfo() const
+	{
+		return FontInfo;
 	}
 
-	bool operator==(const FSlateFontKey& Other ) const
+	FORCEINLINE float GetScale() const
+	{
+		return Scale;
+	}
+
+	FORCEINLINE bool operator==(const FSlateFontKey& Other ) const
 	{
 		return FontInfo == Other.FontInfo && Scale == Other.Scale;
 	}
 
+	FORCEINLINE bool operator!=(const FSlateFontKey& Other ) const
+	{
+		return !(*this == Other);
+	}
+
 	friend inline uint32 GetTypeHash( const FSlateFontKey& Key )
 	{
-		uint32 Hash = 0;
-		Hash = HashCombine(Hash, GetTypeHash(Key.FontInfo));
-		Hash = HashCombine(Hash, GetTypeHash(Key.Scale));
-		return Hash;
+		return Key.KeyHash;
 	}
+
+private:
+	FSlateFontInfo FontInfo;
+	float Scale;
+	uint32 KeyHash;
 };
 
 /** Measurement details for a specific character */
