Slate : UseCurrentWindow vs. CreateNewWindow is now automatically determined by SMenuAnchor based on Context. SViewport currently asks everything within it to EPopupMenu::UseCurrentWindow.

diff --git a/Engine/Source/Runtime/AppFramework/Private/Framework/Testing/SWidgetGallery.cpp b/Engine/Source/Runtime/AppFramework/Private/Framework/Testing/SWidgetGallery.cpp
--- a/Engine/Source/Runtime/AppFramework/Private/Framework/Testing/SWidgetGallery.cpp
+++ b/Engine/Source/Runtime/AppFramework/Private/Framework/Testing/SWidgetGallery.cpp
@@ -50,6 +50,8 @@ public:
 	 */
 	void Construct( const FArguments& InArgs )
 	{
+		PopupMethod = EPopupMethod::CreateNewWindow;
+
 		// example of tab activation registration
 		{
 			IsActiveTabVisibility = EVisibility::Visible;
@@ -257,38 +259,44 @@ public:
 						.HAlign(HAlign_Left)
 						.Padding(0.0f, 5.0f)
 						[
-							SNew(SFxWidget).IgnoreClipping(false)
-							[
 							SNew(SVerticalBox)
 
 							+ SVerticalBox::Slot()
-								.AutoHeight()
+							[
+								SNew(SCheckBox)
+								.IsChecked(this, &SWidgetGallery::IsCreateNewWindowChecked)
+								.OnCheckStateChanged( this, &SWidgetGallery::CreateNewWindowToggled )
 								[
-									// selector combo box
-									SAssignNew(SelectorComboBox, SComboBox<TSharedPtr<FString> >)
-										.Method(SMenuAnchor::UseCurrentWindow)
-										.OptionsSource(&SelectorComboBoxOptions)
-										.OnSelectionChanged(this, &SWidgetGallery::HandleSelectorComboBoxSelectionChanged)
-										.OnGenerateWidget(this, &SWidgetGallery::HandleComboBoxGenerateWidget)
-										[
-											SNew(STextBlock)
-												.Text(this, &SWidgetGallery::HandleSelectorComboBoxText)
-										]
+									SNew(STextBlock).Text(LOCTEXT("CreateNewWindowCheckbox", "Create new window for popups."))
 								]
+							]
 
 							+ SVerticalBox::Slot()
-								.AutoHeight()
-								[
-									// second combo box
-									SAssignNew(SecondComboBox, SComboBox<TSharedPtr<FString> >)
-										.OptionsSource(&SecondComboBoxOptions)
-										.OnSelectionChanged(this, &SWidgetGallery::HandleSecondComboBoxSelectionChanged)
-										.OnGenerateWidget(this, &SWidgetGallery::HandleComboBoxGenerateWidget)
-										[
-											SNew(STextBlock)
-												.Text(this, &SWidgetGallery::HandleSecondComboBoxText)
-										]
-								]
+							.AutoHeight()
+							[
+								// selector combo box
+								SAssignNew(SelectorComboBox, SComboBox<TSharedPtr<FString> >)
+									.OptionsSource(&SelectorComboBoxOptions)
+									.OnSelectionChanged(this, &SWidgetGallery::HandleSelectorComboBoxSelectionChanged)
+									.OnGenerateWidget(this, &SWidgetGallery::HandleComboBoxGenerateWidget)
+									[
+										SNew(STextBlock)
+											.Text(this, &SWidgetGallery::HandleSelectorComboBoxText)
+									]
+							]
+
+							+ SVerticalBox::Slot()
+							.AutoHeight()
+							[
+								// second combo box
+								SAssignNew(SecondComboBox, SComboBox<TSharedPtr<FString> >)
+									.OptionsSource(&SecondComboBoxOptions)
+									.OnSelectionChanged(this, &SWidgetGallery::HandleSecondComboBoxSelectionChanged)
+									.OnGenerateWidget(this, &SWidgetGallery::HandleComboBoxGenerateWidget)
+									[
+										SNew(STextBlock)
+											.Text(this, &SWidgetGallery::HandleSecondComboBoxText)
+									]
 							]
 						]
 
@@ -306,7 +314,6 @@ public:
 							SNew(SFxWidget).IgnoreClipping(false)
 							[
 								SNew(SComboButton)
-								.Method(SMenuAnchor::UseCurrentWindow)
 								.ButtonContent()
 								[
 									SNew(STextBlock)
@@ -934,6 +941,21 @@ protected:
 
 private:
 
+	EPopupMethod PopupMethod;
+
+	ESlateCheckBoxState::Type IsCreateNewWindowChecked() const
+	{
+		return this->PopupMethod == EPopupMethod::CreateNewWindow ? ESlateCheckBoxState::Checked : ESlateCheckBoxState::Unchecked;
+	}
+
+	void CreateNewWindowToggled(ESlateCheckBoxState::Type NewState)
+	{
+		this->PopupMethod = (PopupMethod == EPopupMethod::CreateNewWindow)
+			? EPopupMethod::UseCurrentWindow
+			: EPopupMethod::CreateNewWindow;
+	}
+
+
 	// Callback for clicking the Add button in the SBreadcrumbTrail example.
 	FReply HandleBreadcrumbTrailAddButtonClicked( )
 	{
@@ -1074,6 +1096,11 @@ private:
 		FSlateApplication::Get().DismissAllMenus();
 	}
 
+	virtual TOptional<EPopupMethod> OnQueryPopupMethod() const override
+	{
+		return PopupMethod;
+	}
+
 private:
 
 	// Holds the bread crumb trail widget.
diff --git a/Engine/Source/Runtime/Slate/Private/Widgets/Input/SMenuAnchor.cpp b/Engine/Source/Runtime/Slate/Private/Widgets/Input/SMenuAnchor.cpp
--- a/Engine/Source/Runtime/Slate/Private/Widgets/Input/SMenuAnchor.cpp
+++ b/Engine/Source/Runtime/Slate/Private/Widgets/Input/SMenuAnchor.cpp
@@ -32,9 +32,9 @@ void SMenuAnchor::Tick( const FGeometry& AllottedGeometry, const double InCurren
 	SPanel::Tick(AllottedGeometry, InCurrentTime, InDeltaTime);
 
 	TSharedPtr<SWindow> PopupWindow = PopupWindowPtr.Pin();
-	if ( PopupWindow.IsValid() && Method == CreateNewWindow )
+	if ( IsOpenViaCreatedWindow() )
 	{
-		// Figure out where our attached pop-up window should be placed.				
+		// Figure out where our attached pop-up window should be placed.
 		const FVector2D PopupContentDesiredSize = PopupWindow->GetContent()->GetDesiredSize();
 		FGeometry PopupGeometry = ComputeMenuPlacement( AllottedGeometry, PopupContentDesiredSize, Placement.Get( ) );
 		const FVector2D NewPosition = PopupGeometry.AbsolutePosition;
@@ -190,7 +190,7 @@ FGeometry SMenuAnchor::ComputeMenuPlacement( const FGeometry& AllottedGeometry,
 void SMenuAnchor::RequestClosePopupWindow( const TSharedRef<SWindow>& PopupWindow )
 {
 	bDismissedThisTick = true;
-	if (ensure(Method == CreateNewWindow))
+	if (ensure(IsOpenViaCreatedWindow()))
 	{
 		FSlateApplication::Get().RequestDestroyWindow(PopupWindow);
 		
@@ -204,11 +204,18 @@ void SMenuAnchor::RequestClosePopupWindow( const TSharedRef<SWindow>& PopupWindo
 void SMenuAnchor::OnClickedOutsidePopup()
 {
 	bDismissedThisTick = true;
-	if (ensure(Method == UseCurrentWindow))
-	{
-		FSlateApplication::Get().GetPopupSupport().UnregisterClickNotification( FOnClickedOutside::CreateSP(this, &SMenuAnchor::OnClickedOutsidePopup) );
-		SetIsOpen(false);	
-	}
+	FSlateApplication::Get().GetPopupSupport().UnregisterClickNotification( FOnClickedOutside::CreateSP(this, &SMenuAnchor::OnClickedOutsidePopup) );
+	SetIsOpen(false);	
+}
+
+bool SMenuAnchor::IsOpenAndReusingWindow() const
+{
+	return Children[1].GetWidget() != SNullWidget::NullWidget;
+}
+
+bool SMenuAnchor::IsOpenViaCreatedWindow() const
+{
+	return PopupWindowPtr.Pin().IsValid() && Children[1].GetWidget() == SNullWidget::NullWidget;
 }
 
 void SMenuAnchor::SetContent(TSharedRef<SWidget> InContent)
@@ -225,6 +232,21 @@ void SMenuAnchor::SetMenuContent(TSharedRef<SWidget> InMenuContent)
 	MenuContent = InMenuContent;
 }
 
+EPopupMethod QueryPopupMethod(const FWidgetPath& PathToQuery)
+{
+	for (int32 WidgetIndex = PathToQuery.Widgets.Num() - 1; WidgetIndex >= 0; --WidgetIndex)
+	{
+		TOptional<EPopupMethod> PopupMethod = PathToQuery.Widgets[WidgetIndex].Widget->OnQueryPopupMethod();
+		if (PopupMethod.IsSet())
+		{
+			return PopupMethod.GetValue();
+		}
+	}
+
+	return EPopupMethod::CreateNewWindow;
+}
+
+
 /**
  * Open or close the popup
  *
@@ -291,7 +313,11 @@ void SMenuAnchor::SetIsOpen( bool InIsOpen, const bool bFocusMenu )
 						TransitionEffect = FPopupTransitionEffect( FPopupTransitionEffect::SubMenu );
 					}
 
-					if ( Method == CreateNewWindow )
+					const EPopupMethod PopupMethod = Method.IsSet()
+						? Method.GetValue()
+						: QueryPopupMethod(MyWidgetPath);
+
+					if (PopupMethod == EPopupMethod::CreateNewWindow)
 					{
 						// Open the pop-up
 						TSharedRef<SWindow> PopupWindow = FSlateApplication::Get().PushMenu( AsShared(), MenuContentRef, NewPosition, TransitionEffect, bFocusMenu, false, NewWindowSize, SummonLocationSize );
@@ -302,11 +328,11 @@ void SMenuAnchor::SetIsOpen( bool InIsOpen, const bool bFocusMenu )
 					{
 						// We are re-using the current window instead of creating a new one.
 						// The popup will be presented via an overlay service.
-						ensure( Method == UseCurrentWindow );
+						ensure(PopupMethod == EPopupMethod::UseCurrentWindow);
 						Children[1]
-							[
-								MenuContentRef
-							];
+						[
+							MenuContentRef
+						];
 
 						// We want to support dismissing the popup widget when the user clicks outside it.
 						FSlateApplication::Get().GetPopupSupport().RegisterClickNotification( MenuContentRef, FOnClickedOutside::CreateSP( this, &SMenuAnchor::OnClickedOutsidePopup ) );
@@ -323,22 +349,22 @@ void SMenuAnchor::SetIsOpen( bool InIsOpen, const bool bFocusMenu )
 		{
 			// CLOSE POPUP
 
-			if (Method == CreateNewWindow)
+			if (IsOpenAndReusingWindow())
 			{
-				// Close the Popup window.
+				// Window reuse case; remove child widget.
+				Children[1].DetachWidget();
+			}
+			else
+			{
+				// New window case
+				ensure(IsOpenViaCreatedWindow());
 				TSharedPtr<SWindow> PopupWindow = PopupWindowPtr.Pin();
-				if ( PopupWindow.IsValid() )
+				if (PopupWindow.IsValid())
 				{
 					// Request that the popup be closed.
 					PopupWindow->RequestDestroyWindow();
 				}
 			}
-			else
-			{
-				// Close the popup overlay
-				ensure(Method==UseCurrentWindow);
-				Children[1].DetachWidget();
-			}
 
 			if (OnMenuOpenChanged.IsBound())
 			{
@@ -350,7 +376,7 @@ void SMenuAnchor::SetIsOpen( bool InIsOpen, const bool bFocusMenu )
 
 bool SMenuAnchor::IsOpen() const
 {
-	return PopupWindowPtr.Pin().IsValid() || Children[1].GetWidget() != SNullWidget::NullWidget;
+	return PopupWindowPtr.Pin().IsValid() || IsOpenAndReusingWindow();
 }
 
 bool SMenuAnchor::ShouldOpenDueToClick() const
@@ -392,7 +418,7 @@ SMenuAnchor::~SMenuAnchor()
 	if (PopupWindow.IsValid())
 	{
 		// Close the Popup window.
-		if (Method == CreateNewWindow)
+		if (IsOpenViaCreatedWindow())
 		{
 			// Request that the popup be closed.
 			PopupWindow->RequestDestroyWindow();
diff --git a/Engine/Source/Runtime/Slate/Private/Widgets/Notifications/SPopUpErrorText.cpp b/Engine/Source/Runtime/Slate/Private/Widgets/Notifications/SPopUpErrorText.cpp
--- a/Engine/Source/Runtime/Slate/Private/Widgets/Notifications/SPopUpErrorText.cpp
+++ b/Engine/Source/Runtime/Slate/Private/Widgets/Notifications/SPopUpErrorText.cpp
@@ -7,7 +7,6 @@ void SPopupErrorText::Construct( const FArguments& InArgs )
 {
 	SComboButton::Construct( SComboButton::FArguments()
 		.ComboButtonStyle( FCoreStyle::Get(), "MessageLogListingComboButton" )
-		.Method( InArgs._ShowInNewWindow ? SMenuAnchor::CreateNewWindow : SMenuAnchor::UseCurrentWindow )
 		.HasDownArrow(false)
 		.ContentPadding(0)
 		.ButtonContent()
diff --git a/Engine/Source/Runtime/Slate/Private/Widgets/SViewport.cpp b/Engine/Source/Runtime/Slate/Private/Widgets/SViewport.cpp
--- a/Engine/Source/Runtime/Slate/Private/Widgets/SViewport.cpp
+++ b/Engine/Source/Runtime/Slate/Private/Widgets/SViewport.cpp
@@ -287,6 +287,11 @@ FReply SViewport::OnMotionDetected( const FGeometry& MyGeometry, const FMotionEv
 	return ViewportInterface.IsValid() ? ViewportInterface.Pin()->OnMotionDetected(MyGeometry, InMotionEvent) : FReply::Unhandled();
 }
 
+TOptional<EPopupMethod> SViewport::OnQueryPopupMethod() const
+{
+	return EPopupMethod::UseCurrentWindow;
+}
+
 void SViewport::OnFinishedPointerInput()
 {
 	TSharedPtr<ISlateViewport> PinnedInterface = ViewportInterface.Pin();
diff --git a/Engine/Source/Runtime/Slate/Public/Widgets/Input/SComboBox.h b/Engine/Source/Runtime/Slate/Public/Widgets/Input/SComboBox.h
--- a/Engine/Source/Runtime/Slate/Public/Widgets/Input/SComboBox.h
+++ b/Engine/Source/Runtime/Slate/Public/Widgets/Input/SComboBox.h
@@ -85,7 +85,7 @@ public:
 		, _OnSelectionChanged()
 		, _OnGenerateWidget()
 		, _InitiallySelectedItem( nullptr)
-		, _Method( SMenuAnchor::CreateNewWindow )
+		, _Method()
 		, _MaxListHeight(450.0f)
 		, _HasDownArrow( true )
 		{}
@@ -111,7 +111,7 @@ public:
 		/** The option that should be selected when the combo box is first created */
 		SLATE_ARGUMENT( OptionType, InitiallySelectedItem )
 
-		SLATE_ARGUMENT( SMenuAnchor::EMethod, Method )
+		SLATE_ARGUMENT( TOptional<EPopupMethod>, Method )
 
 		/** The max height of the combo box menu */
 		SLATE_ARGUMENT(float, MaxListHeight)
diff --git a/Engine/Source/Runtime/Slate/Public/Widgets/Input/SComboButton.h b/Engine/Source/Runtime/Slate/Public/Widgets/Input/SComboButton.h
--- a/Engine/Source/Runtime/Slate/Public/Widgets/Input/SComboButton.h
+++ b/Engine/Source/Runtime/Slate/Public/Widgets/Input/SComboButton.h
@@ -24,7 +24,7 @@ public:
 		, _MenuPlacement(MenuPlacement_ComboBox)
 		, _HAlign(HAlign_Fill)
 		, _VAlign(VAlign_Center)
-		, _Method( SMenuAnchor::CreateNewWindow )
+		, _Method()
 		{}
 
 		SLATE_STYLE_ARGUMENT( FComboButtonStyle, ComboButtonStyle )
@@ -52,7 +52,7 @@ public:
 		SLATE_ARGUMENT( EVerticalAlignment, VAlign )
 
 		/** Spawn a new window or reuse current window for this combo*/
-		SLATE_ARGUMENT( SMenuAnchor::EMethod, Method )
+		SLATE_ARGUMENT( TOptional<EPopupMethod>, Method )
 	SLATE_END_ARGS()
 	
 	/**
diff --git a/Engine/Source/Runtime/Slate/Public/Widgets/Input/SMenuAnchor.h b/Engine/Source/Runtime/Slate/Public/Widgets/Input/SMenuAnchor.h
--- a/Engine/Source/Runtime/Slate/Public/Widgets/Input/SMenuAnchor.h
+++ b/Engine/Source/Runtime/Slate/Public/Widgets/Input/SMenuAnchor.h
@@ -6,6 +6,17 @@
 /** Notification when popup is opened/closed. */
 DECLARE_DELEGATE_OneParam(FOnIsOpenChanged, bool)
 
+
+enum class EPopupMethod : uint8
+{
+	/** Creating a new window allows us to place popups outside of the window in which the menu anchor resides. */
+	CreateNewWindow,
+	/** Place the popup into the current window. Applications that intend to run in fullscreen cannot create new windows, so they must use this method.. */
+	UseCurrentWindow
+};
+
+;
+
 /**
  * A PopupAnchor summons a Popup relative to its content.
  * Summoning a popup relative to the cursor is accomplished via the application.
@@ -13,13 +24,10 @@ DECLARE_DELEGATE_OneParam(FOnIsOpenChanged, bool)
 class SLATE_API SMenuAnchor : public SPanel
 {
 public:
-	enum EMethod
-	{
-		/** Creating a new window allows us to place popups outside of the window in which the menu anchor resides. */
-		CreateNewWindow,
-		/** Applications that intend to run in fullscreen cannot create new windows. */
-		UseCurrentWindow
-	};
+	DEPRECATED(4.7, "You probably do not need this setting any more. See OnQueryPopupMethod() in SWidget.h.")
+	typedef ::EPopupMethod EMethod;
+	static const EPopupMethod CreateNewWindow = EPopupMethod::CreateNewWindow;
+	static const EPopupMethod UseCurrentWindow = EPopupMethod::UseCurrentWindow;
 	
 public:
 	SLATE_BEGIN_ARGS( SMenuAnchor )
@@ -27,7 +35,7 @@ public:
 		, _MenuContent( SNew(STextBlock) .Text( NSLOCTEXT("SMenuAnchor", "NoMenuContent", "No Menu Content Assigned; use .MenuContent") ) )
 		, _OnGetMenuContent()
 		, _Placement( MenuPlacement_BelowAnchor )
-		, _Method( CreateNewWindow )
+		, _Method()
 		{}
 		
 		SLATE_DEFAULT_SLOT( FArguments, Content )
@@ -40,7 +48,8 @@ public:
 
 		SLATE_ATTRIBUTE( EMenuPlacement, Placement )
 
-		SLATE_ARGUMENT( EMethod, Method )
+		DEPRECATED(4.7, "Popup method now determined automatically; prefer not setting it. Use EPopupMethod when setting explicitly.")
+		SLATE_ARGUMENT(TOptional<EPopupMethod>, Method)
 
 	SLATE_END_ARGS()
 
@@ -108,6 +117,12 @@ protected:
 	/** Invoked when a click occurred outside the widget that is our popup; used for popup auto-dismissal */
 	void OnClickedOutsidePopup();
 
+	/** @return true if the popup is currently open and reusing an existing window */
+	bool IsOpenAndReusingWindow() const;
+
+	/** @return true if the popup is currently open and we created a dedicated window for it */
+	bool IsOpenViaCreatedWindow() const;
+
 protected:
 	/** A pointer to the window presenting this popup. Pointer is null when the popup is not visible */
 	TWeakPtr<SWindow> PopupWindowPtr;
@@ -128,7 +143,7 @@ protected:
 	bool bDismissedThisTick;
 
 	/** Should we summon a new window for this popup or  */
-	EMethod Method;
+	TOptional<EPopupMethod> Method;
 
 	TPanelChildren<FSimpleSlot> Children;
 };
diff --git a/Engine/Source/Runtime/Slate/Public/Widgets/SViewport.h b/Engine/Source/Runtime/Slate/Public/Widgets/SViewport.h
--- a/Engine/Source/Runtime/Slate/Public/Widgets/SViewport.h
+++ b/Engine/Source/Runtime/Slate/Public/Widgets/SViewport.h
@@ -151,6 +151,7 @@ public:
 	virtual FReply OnFocusReceived( const FGeometry& MyGeometry, const FFocusEvent& InFocusEvent ) override;
 	virtual void OnFocusLost( const FFocusEvent& InFocusEvent ) override;
 	virtual FReply OnMotionDetected( const FGeometry& MyGeometry, const FMotionEvent& InMotionEvent ) override;
+	virtual TOptional<EPopupMethod> OnQueryPopupMethod() const override;
 	virtual void OnFinishedPointerInput() override;
 	virtual void OnArrangeChildren( const FGeometry& AllottedGeometry, FArrangedChildren& ArrangedChildren ) const override;
 	virtual TSharedPtr<struct FVirtualPointerPosition> TranslateMouseCoordinateFor3DChild(const TSharedRef<SWidget>& ChildWidget, const FGeometry& MyGeometry, const FVector2D& ScreenSpaceMouseCoordinate, const FVector2D& LastScreenSpaceMouseCoordinate) const;
diff --git a/Engine/Source/Runtime/SlateCore/Private/Widgets/SWidget.cpp b/Engine/Source/Runtime/SlateCore/Private/Widgets/SWidget.cpp
--- a/Engine/Source/Runtime/SlateCore/Private/Widgets/SWidget.cpp
+++ b/Engine/Source/Runtime/SlateCore/Private/Widgets/SWidget.cpp
@@ -285,6 +285,16 @@ FReply SWidget::OnMotionDetected( const FGeometry& MyGeometry, const FMotionEven
 	return FReply::Unhandled();
 }
 
+TOptional<EPopupMethod> SWidget::OnQueryPopupMethod() const
+{
+	return TOptional<EPopupMethod>();
+}
+
+TSharedPtr<struct FVirtualPointerPosition> SWidget::TranslateMouseCoordinateFor3DChild(const TSharedRef<SWidget>& ChildWidget, const FGeometry& MyGeometry, const FVector2D& ScreenSpaceMouseCoordinate, const FVector2D& LastScreenSpaceMouseCoordinate) const
+{
+	return nullptr;
+}
+
 
 void SWidget::OnFinishedPointerInput()
 {
diff --git a/Engine/Source/Runtime/SlateCore/Public/Widgets/SWidget.h b/Engine/Source/Runtime/SlateCore/Public/Widgets/SWidget.h
--- a/Engine/Source/Runtime/SlateCore/Public/Widgets/SWidget.h
+++ b/Engine/Source/Runtime/SlateCore/Public/Widgets/SWidget.h
@@ -4,6 +4,8 @@
 
 #include "SlateRenderTransform.h"
 
+enum class EPopupMethod : uint8;
+
 class SLATECORE_API FSlateControlledConstruction
 {
 public:
@@ -109,7 +111,7 @@ public:
 	 * @param bParentEnabled	True if the parent of this widget is enabled.
 	 * @return The maximum layer ID attained by this widget or any of its children.
 	 */
-	int32 Paint( const FPaintArgs& Args, const FGeometry& AllottedGeometry, const FSlateRect& MyClippingRect, FSlateWindowElementList& OutDrawElements, int32 LayerId, const FWidgetStyle& InWidgetStyle, bool bParentEnabled ) const;
+	int32 Paint(const FPaintArgs& Args, const FGeometry& AllottedGeometry, const FSlateRect& MyClippingRect, FSlateWindowElementList& OutDrawElements, int32 LayerId, const FWidgetStyle& InWidgetStyle, bool bParentEnabled) const;
 
 	DEPRECATED(4.4, "Paint() now requires an extra FPaintArgs parameter. When calling paint on a child widget, use Args.WithNewParent(this).")
 	int32 Paint(const FGeometry& AllottedGeometry, const FSlateRect& MyClippingRect, FSlateWindowElementList& OutDrawElements, int32 LayerId, const FWidgetStyle& InWidgetStyle, bool bParentEnabled) const { return 0; }
@@ -169,7 +171,7 @@ public:
 	 * @param  InCharacterEvent  Character event
 	 * @return  Returns whether the event was handled, along with other possible actions
 	 */
-	virtual FReply OnKeyChar( const FGeometry& MyGeometry,const FCharacterEvent& InCharacterEvent );
+	virtual FReply OnKeyChar(const FGeometry& MyGeometry, const FCharacterEvent& InCharacterEvent);
 
 	/**
 	 * Called after a key is pressed when this widget or a child of this widget has focus
@@ -182,7 +184,7 @@ public:
 	 * @param InKeyEvent  Key event
 	 * @return Returns whether the event was handled, along with other possible actions
 	 */
-	virtual FReply OnPreviewKeyDown( const FGeometry& MyGeometry, const FKeyEvent& InKeyEvent );
+	virtual FReply OnPreviewKeyDown(const FGeometry& MyGeometry, const FKeyEvent& InKeyEvent);
 
 	/**
 	 * Called after a key is pressed when this widget has focus (this event bubbles if not handled)
@@ -191,7 +193,7 @@ public:
 	 * @param InKeyEvent  Key event
 	 * @return Returns whether the event was handled, along with other possible actions
 	 */
-	virtual FReply OnKeyDown( const FGeometry& MyGeometry, const FKeyEvent& InKeyEvent );
+	virtual FReply OnKeyDown(const FGeometry& MyGeometry, const FKeyEvent& InKeyEvent);
 
 	DEPRECATED(4.6, "SWidget::OnControllerButtonPressed() is deprecated, SWidget::OnKeyDown() now handles controller input as well as controller.")
 	virtual FReply OnControllerButtonPressed(const FGeometry& MyGeometry, const FControllerEvent& ControllerEvent);
@@ -231,7 +233,7 @@ public:
 	 * @param MouseEvent Information about the input event
 	 * @return Whether the event was handled along with possible requests for the system to take action.
 	 */
-	virtual FReply OnMouseButtonDown( const FGeometry& MyGeometry, const FPointerEvent& MouseEvent );
+	virtual FReply OnMouseButtonDown(const FGeometry& MyGeometry, const FPointerEvent& MouseEvent);
 
 	/**
 	 * Just like OnMouseButtonDown, but tunnels instead of bubbling.
@@ -240,7 +242,7 @@ public:
 	 * Use this event sparingly as preview events generally make UIs more
 	 * difficult to reason about.
 	 */
-	virtual FReply OnPreviewMouseButtonDown( const FGeometry& MyGeometry, const FPointerEvent& MouseEvent );
+	virtual FReply OnPreviewMouseButtonDown(const FGeometry& MyGeometry, const FPointerEvent& MouseEvent);
 	
 	/**
 	 * The system calls this method to notify the widget that a mouse button was release within it. This event is bubbled.
@@ -249,7 +251,7 @@ public:
 	 * @param MouseEvent Information about the input event
 	 * @return Whether the event was handled along with possible requests for the system to take action.
 	 */
-	virtual FReply OnMouseButtonUp( const FGeometry& MyGeometry, const FPointerEvent& MouseEvent );
+	virtual FReply OnMouseButtonUp(const FGeometry& MyGeometry, const FPointerEvent& MouseEvent);
 	
 	/**
 	 * The system calls this method to notify the widget that a mouse moved within it. This event is bubbled.
@@ -258,7 +260,7 @@ public:
 	 * @param MouseEvent Information about the input event
 	 * @return Whether the event was handled along with possible requests for the system to take action.
 	 */
-	virtual FReply OnMouseMove( const FGeometry& MyGeometry, const FPointerEvent& MouseEvent );
+	virtual FReply OnMouseMove(const FGeometry& MyGeometry, const FPointerEvent& MouseEvent);
 	
 	/**
 	 * The system will use this event to notify a widget that the cursor has entered it. This event is NOT bubbled.
@@ -266,14 +268,14 @@ public:
 	 * @param MyGeometry The Geometry of the widget receiving the event
 	 * @param MouseEvent Information about the input event
 	 */
-	virtual void OnMouseEnter( const FGeometry& MyGeometry, const FPointerEvent& MouseEvent );
+	virtual void OnMouseEnter(const FGeometry& MyGeometry, const FPointerEvent& MouseEvent);
 	
 	/**
 	 * The system will use this event to notify a widget that the cursor has left it. This event is NOT bubbled.
 	 *
 	 * @param MouseEvent Information about the input event
 	 */
-	virtual void OnMouseLeave( const FPointerEvent& MouseEvent );
+	virtual void OnMouseLeave(const FPointerEvent& MouseEvent);
 	
 	/**
 	 * Called when the mouse wheel is spun. This event is bubbled.
@@ -281,14 +283,14 @@ public:
 	 * @param  MouseEvent  Mouse event
 	 * @return  Returns whether the event was handled, along with other possible actions
 	 */
-	virtual FReply OnMouseWheel( const FGeometry& MyGeometry, const FPointerEvent& MouseEvent );
+	virtual FReply OnMouseWheel(const FGeometry& MyGeometry, const FPointerEvent& MouseEvent);
 	
 	/**
 	 * The system asks each widget under the mouse to provide a cursor. This event is bubbled.
 	 * 
 	 * @return FCursorReply::Unhandled() if the event is not handled; return FCursorReply::Cursor() otherwise.
 	 */
-	virtual FCursorReply OnCursorQuery( const FGeometry& MyGeometry, const FPointerEvent& CursorEvent ) const;
+	virtual FCursorReply OnCursorQuery(const FGeometry& MyGeometry, const FPointerEvent& CursorEvent) const;
 
 	/**
 	 * Called when a mouse button is double clicked.  Override this in derived classes.
@@ -297,7 +299,7 @@ public:
 	 * @param  InMouseEvent  Mouse button event
 	 * @return  Returns whether the event was handled, along with other possible actions
 	 */
-	virtual FReply OnMouseButtonDoubleClick( const FGeometry& InMyGeometry, const FPointerEvent& InMouseEvent );
+	virtual FReply OnMouseButtonDoubleClick(const FGeometry& InMyGeometry, const FPointerEvent& InMouseEvent);
 
 	/**
 	 * Called when Slate wants to visualize tooltip.
@@ -307,7 +309,7 @@ public:
 	 * @param  TooltipContent    The TooltipContent that I may want to visualize.
 	 * @return true if this widget visualized the tooltip content; i.e., the event is handled.
 	 */
-	virtual bool OnVisualizeTooltip( const TSharedPtr<SWidget>& TooltipContent );
+	virtual bool OnVisualizeTooltip(const TSharedPtr<SWidget>& TooltipContent);
 
 	/**
 	 * Called when Slate detects that a widget started to be dragged.
@@ -322,7 +324,7 @@ public:
 	 * @param  InMyGeometry  Widget geometry
 	 * @param  InMouseEvent  MouseMove that triggered the drag
 	 */
-	virtual FReply OnDragDetected( const FGeometry& MyGeometry, const FPointerEvent& MouseEvent );
+	virtual FReply OnDragDetected(const FGeometry& MyGeometry, const FPointerEvent& MouseEvent);
 
 	//
 	// DRAG AND DROP (DragDrop)
@@ -345,14 +347,14 @@ public:
 	 *
 	 * @return A reply that indicated whether the contents of the DragDropEvent can potentially be processed by this widget.
 	 */
-	virtual void OnDragEnter( const FGeometry& MyGeometry, const FDragDropEvent& DragDropEvent );
+	virtual void OnDragEnter(const FGeometry& MyGeometry, const FDragDropEvent& DragDropEvent);
 	
 	/**
 	 * Called during drag and drop when the drag leaves a widget.
 	 *
 	 * @param DragDropEvent   The drag and drop event.
 	 */
-	virtual void OnDragLeave( const FDragDropEvent& DragDropEvent );
+	virtual void OnDragLeave(const FDragDropEvent& DragDropEvent);
 
 	/**
 	 * Called during drag and drop when the the mouse is being dragged over a widget.
@@ -361,7 +363,7 @@ public:
 	 * @param DragDropEvent   The drag and drop event.
 	 * @return A reply that indicated whether this event was handled.
 	 */
-	virtual FReply OnDragOver( const FGeometry& MyGeometry, const FDragDropEvent& DragDropEvent );
+	virtual FReply OnDragOver(const FGeometry& MyGeometry, const FDragDropEvent& DragDropEvent);
 	
 	/**
 	 * Called when the user is dropping something onto a widget; terminates drag and drop.
@@ -370,7 +372,11 @@ public:
 	 * @param DragDropEvent   The drag and drop event.
 	 * @return A reply that indicated whether this event was handled.
 	 */
-	virtual FReply OnDrop( const FGeometry& MyGeometry, const FDragDropEvent& DragDropEvent );
+	virtual FReply OnDrop(const FGeometry& MyGeometry, const FDragDropEvent& DragDropEvent);
+
+	//
+	// TOUCH and GESTURES
+	//
 
 	/**
 	 * Called when the user performs a gesture on trackpad. This event is bubbled.
@@ -378,28 +384,28 @@ public:
 	 * @param  GestureEvent  gesture event
 	 * @return  Returns whether the event was handled, along with other possible actions
 	 */
-	virtual FReply OnTouchGesture( const FGeometry& MyGeometry, const FPointerEvent& GestureEvent );
+	virtual FReply OnTouchGesture(const FGeometry& MyGeometry, const FPointerEvent& GestureEvent);
 	
 	/**
 	 * Called when a touchpad touch is started (finger down)
 	 * 
 	 * @param InTouchEvent	The touch event generated
 	 */
-	virtual FReply OnTouchStarted( const FGeometry& MyGeometry, const FPointerEvent& InTouchEvent );
+	virtual FReply OnTouchStarted(const FGeometry& MyGeometry, const FPointerEvent& InTouchEvent);
 
 	/**
 	 * Called when a touchpad touch is moved  (finger moved)
 	 * 
 	 * @param InTouchEvent	The touch event generated
 	 */
-	virtual FReply OnTouchMoved( const FGeometry& MyGeometry, const FPointerEvent& InTouchEvent );
+	virtual FReply OnTouchMoved(const FGeometry& MyGeometry, const FPointerEvent& InTouchEvent);
 
 	/**
 	 * Called when a touchpad touch is ended (finger lifted)
 	 * 
 	 * @param InTouchEvent	The touch event generated
 	 */
-	virtual FReply OnTouchEnded( const FGeometry& MyGeometry, const FPointerEvent& InTouchEvent );
+	virtual FReply OnTouchEnded(const FGeometry& MyGeometry, const FPointerEvent& InTouchEvent);
 
 	/**
 	 * Called when motion is detected (controller or device)
@@ -407,16 +413,25 @@ public:
 	 * 
 	 * @param InMotionEvent	The motion event generated
 	 */
-	virtual FReply OnMotionDetected( const FGeometry& MyGeometry, const FMotionEvent& InMotionEvent );
-
-	virtual TSharedPtr<struct FVirtualPointerPosition> TranslateMouseCoordinateFor3DChild( const TSharedRef<SWidget>& ChildWidget, const FGeometry& MyGeometry, const FVector2D& ScreenSpaceMouseCoordinate, const FVector2D& LastScreenSpaceMouseCoordinate ) const { return nullptr; }
+	virtual FReply OnMotionDetected(const FGeometry& MyGeometry, const FMotionEvent& InMotionEvent);
 
+	/**
+	 * Popups can manifest in a NEW OS WINDOW or via an OVERLAY in an existing window.
+	 * This can be set explicitly on SMenuAnchor, or can be determined by a scoping widget.
+	 * A scoping widget can reply to OnQueryPopupMethod() to drive all its descendants'
+	 * poup methods.
+	 *
+	 * e.g. Fullscreen games cannot summon a new window, so game SViewports will reply with
+	 *      EPopupMethod::UserCurrentWindow. This makes all the menu anchors within them
+	 *      use the current window.
+	 */
+	virtual TOptional<EPopupMethod> OnQueryPopupMethod() const;
 
+	virtual TSharedPtr<struct FVirtualPointerPosition> TranslateMouseCoordinateFor3DChild(const TSharedRef<SWidget>& ChildWidget, const FGeometry& MyGeometry, const FVector2D& ScreenSpaceMouseCoordinate, const FVector2D& LastScreenSpaceMouseCoordinate) const;
+	
 	/**
 	 * All the pointer (mouse, touch, stylus, etc.) events from this frame have been routed.
 	 * This is a widget's chance to act on any accumulated data.
-	 *
-	 * 
 	 */
 	virtual void OnFinishedPointerInput();
 
