Speculatively set window cached screen position.

- Helps with issues caused by Slate expecting it to immediately (i.e. within a single frame) match the position where the windows was told to move.
- The underlying problem of such a synchronous Slate behavior is tracked as UE-1308.

diff --git a/Engine/Source/Runtime/SlateCore/Private/Widgets/SWindow.cpp b/Engine/Source/Runtime/SlateCore/Private/Widgets/SWindow.cpp
--- a/Engine/Source/Runtime/SlateCore/Private/Widgets/SWindow.cpp
+++ b/Engine/Source/Runtime/SlateCore/Private/Widgets/SWindow.cpp
@@ -706,6 +706,15 @@ void SWindow::MoveWindowTo( FVector2D NewPosition )
 {
 	if (NativeWindow.IsValid())
 	{
+#if PLATFORM_LINUX
+		// Slate code often expects cached screen position to be accurate immediately after the move.
+		// This expectation is generally invalid (see UE-1308) as there may be a delay before the OS reports it back.
+		// This hack sets the position speculatively, keeping Slate happy while also giving the OS chance to report it
+		// correctly after or even during the actual call.
+		FVector2D SpeculativeScreenPosition(FMath::TruncToInt(NewPosition.X), FMath::TruncToInt(NewPosition.Y));
+		SetCachedScreenPosition(SpeculativeScreenPosition);
+#endif // PLATFORM_LINUX
+
 		NativeWindow->MoveWindowTo( FMath::TruncToInt(NewPosition.X), FMath::TruncToInt(NewPosition.Y) );
 	}
 	else
@@ -718,6 +727,15 @@ void SWindow::ReshapeWindow( FVector2D NewPosition, FVector2D NewSize )
 {
 	if (NativeWindow.IsValid())
 	{
+#if PLATFORM_LINUX
+		// Slate code often expects cached screen position to be accurate immediately after the move.
+		// This expectation is generally invalid (see UE-1308) as there may be a delay before the OS reports it back.
+		// This hack sets the position speculatively, keeping Slate happy while also giving the OS chance to report it
+		// correctly after or even during the actual call.
+		FVector2D SpeculativeScreenPosition(FMath::TruncToInt(NewPosition.X), FMath::TruncToInt(NewPosition.Y));
+		SetCachedScreenPosition(SpeculativeScreenPosition);
+#endif // PLATFORM_LINUX
+
 		NativeWindow->ReshapeWindow( FMath::TruncToInt(NewPosition.X), FMath::TruncToInt(NewPosition.Y), FMath::TruncToInt(NewSize.X), FMath::TruncToInt(NewSize.Y) );
 	}
 	else
