Fixed crash in FSlateFontMeasure::MeasureStringInternal. PreviousCharEntry is now a copy of the struct instead of a pointer.

const FCharacterEntry& Entry = CharacterList[CurrentChar]; can realloc memory in TMaps in CharacterList and if realloc moved the TMap data to a different address, PreviousCharEntry was no longer valid.

diff --git a/Engine/Source/Runtime/SlateCore/Private/Fonts/FontMeasure.cpp b/Engine/Source/Runtime/SlateCore/Private/Fonts/FontMeasure.cpp
--- a/Engine/Source/Runtime/SlateCore/Private/Fonts/FontMeasure.cpp
+++ b/Engine/Source/Runtime/SlateCore/Private/Fonts/FontMeasure.cpp
@@ -168,12 +168,12 @@ FVector2D FSlateFontMeasure::MeasureStringInternal( const FString& Text, int32 S
 	// Accumulated height of this block of text
 	int32 StringSizeY = MaxHeight;
 	// The previous char (for kerning)
-	const FCharacterEntry* PreviousCharEntry = nullptr;
+	FCharacterEntry PreviousCharEntry;
 
 	//If we are measuring a range then we should take into account the kerning with the character before the start of the range
 	if ( !DoesStartAtBeginning && IncludeKerningWithPrecedingChar )
 	{
-		PreviousCharEntry = &CharacterList[ Text[ StartIndex - 1 ] ];
+		PreviousCharEntry = CharacterList[ Text[ StartIndex - 1 ] ];
 	}
 
 	int32 FinalPosX = 0;
@@ -200,16 +200,16 @@ FVector2D FSlateFontMeasure::MeasureStringInternal( const FString& Text, int32 S
 			const FCharacterEntry& Entry = CharacterList[CurrentChar];
 
 			int32 Kerning = 0;
-			if( PreviousCharEntry )
+			if( PreviousCharEntry.IsValidEntry() )
 			{
-				Kerning = CharacterList.GetKerning( *PreviousCharEntry, Entry );
+				Kerning = CharacterList.GetKerning( PreviousCharEntry, Entry );
 			}
 
 			const int32 TotalCharSpacing = 
 				Kerning + Entry.HorizontalOffset +		// Width is any kerning plus how much to advance the position when drawing a new character
 				Entry.XAdvance;	// How far we advance 
 
-			PreviousCharEntry = &Entry;
+			PreviousCharEntry = Entry;
 
 			CurrentX += Kerning + Entry.XAdvance;
 
