Fixed crash in FCharacterList::GetKerning caused by FCharacterList::GetCharacter and FCharacterList::CacheCharacter reallocating TMap memory

diff --git a/Engine/Source/Runtime/SlateCore/Private/Fonts/FontCache.cpp b/Engine/Source/Runtime/SlateCore/Private/Fonts/FontCache.cpp
--- a/Engine/Source/Runtime/SlateCore/Private/Fonts/FontCache.cpp
+++ b/Engine/Source/Runtime/SlateCore/Private/Fonts/FontCache.cpp
@@ -957,7 +957,9 @@ bool FCharacterList::IsStale() const
 
 int8 FCharacterList::GetKerning( TCHAR FirstChar, TCHAR SecondChar )
 {
-	return GetKerning( GetCharacter( FirstChar ), GetCharacter( SecondChar ) );
+	const FCharacterEntry First = GetCharacter( FirstChar );
+	const FCharacterEntry Second = GetCharacter( SecondChar );
+	return GetKerning( First, Second );
 }
 
 int8 FCharacterList::GetKerning( const FCharacterEntry& FirstCharacterEntry, const FCharacterEntry& SecondCharacterEntry )
