Fixed crash in FSlateElementBatcher::AddTextElement and FCanvasTextItem::DrawStringInternal_RuntimeCache.

Continuation of CL# 2370014.

PreviousCharEntry is now a copy of the struct instead of a pointer.

const FCharacterEntry& Entry = CharacterList[CurrentChar]; can realloc memory in TMaps in CharacterList and if realloc moved the TMap data to a different address, PreviousCharEntry was no longer valid.

diff --git a/Engine/Source/Runtime/Engine/Private/UserInterface/CanvasItem.cpp b/Engine/Source/Runtime/Engine/Private/UserInterface/CanvasItem.cpp
--- a/Engine/Source/Runtime/Engine/Private/UserInterface/CanvasItem.cpp
+++ b/Engine/Source/Runtime/Engine/Private/UserInterface/CanvasItem.cpp
@@ -1050,7 +1050,7 @@ void FCanvasTextItem::DrawStringInternal_RuntimeCache( FCanvas* InCanvas, const
 
 	float LineX = 0;
 
-	const FCharacterEntry* PreviousCharEntry = nullptr;
+	FCharacterEntry PreviousCharEntry;
 
 	int32 Kerning = 0;
 
@@ -1105,9 +1105,9 @@ void FCanvasTextItem::DrawStringInternal_RuntimeCache( FCanvas* InCanvas, const
 
 			const bool bIsWhitespace = FChar::IsWhitespace(CurrentChar);
 
-			if( !bIsWhitespace && CharIndex > 0 )
+			if( !bIsWhitespace && PreviousCharEntry.IsValidEntry() )
 			{
-				Kerning = CharacterList.GetKerning( *PreviousCharEntry, Entry );
+				Kerning = CharacterList.GetKerning( PreviousCharEntry, Entry );
 			}
 			else
 			{
@@ -1115,7 +1115,7 @@ void FCanvasTextItem::DrawStringInternal_RuntimeCache( FCanvas* InCanvas, const
 			}
 
 			LineX += Kerning;
-			PreviousCharEntry = &Entry;
+			PreviousCharEntry = Entry;
 
 			if( !bIsWhitespace )
 			{
diff --git a/Engine/Source/Runtime/SlateCore/Private/Rendering/ElementBatcher.cpp b/Engine/Source/Runtime/SlateCore/Private/Rendering/ElementBatcher.cpp
--- a/Engine/Source/Runtime/SlateCore/Private/Rendering/ElementBatcher.cpp
+++ b/Engine/Source/Runtime/SlateCore/Private/Rendering/ElementBatcher.cpp
@@ -525,7 +525,7 @@ void FSlateElementBatcher::AddTextElement(const FSlateDrawElement& DrawElement)
 
 	float LineX = 0;
 
-	const FCharacterEntry* PreviousCharEntry = nullptr;
+	FCharacterEntry PreviousCharEntry;
 
 	int32 Kerning = 0;
 
@@ -575,9 +575,9 @@ void FSlateElementBatcher::AddTextElement(const FSlateDrawElement& DrawElement)
 
 			const bool bIsWhitespace = FChar::IsWhitespace(CurrentChar);
 
-			if( !bIsWhitespace && CharIndex > 0 )
+			if( !bIsWhitespace && PreviousCharEntry.IsValidEntry() )
 			{
-				Kerning = CharacterList.GetKerning( *PreviousCharEntry, Entry );
+				Kerning = CharacterList.GetKerning( PreviousCharEntry, Entry );
 			}
 			else
 			{
@@ -585,7 +585,7 @@ void FSlateElementBatcher::AddTextElement(const FSlateDrawElement& DrawElement)
 			}
 
 			LineX += Kerning;
-			PreviousCharEntry = &Entry;
+			PreviousCharEntry = Entry;
 
 			if( !bIsWhitespace )
 			{
