Only show focus highlight for navigation causes.

diff --git a/Engine/Source/Runtime/SlateCore/Private/Widgets/SWidget.cpp b/Engine/Source/Runtime/SlateCore/Private/Widgets/SWidget.cpp
--- a/Engine/Source/Runtime/SlateCore/Private/Widgets/SWidget.cpp
+++ b/Engine/Source/Runtime/SlateCore/Private/Widgets/SWidget.cpp
@@ -631,7 +631,7 @@ int32 SWidget::Paint(const FPaintArgs& Args, const FGeometry& AllottedGeometry,
 	if (SupportsKeyboardFocus())
 	{
 		TOptional<EFocusCause> FocusCause = HasAnyUserFocus();
-		if (FocusCause.IsSet() && FocusCause.GetValue() != EFocusCause::Mouse)
+		if (FocusCause.IsSet() && FocusCause.GetValue() == EFocusCause::Navigation)
 		{
 			const FSlateBrush* BrushResource = FocusBrush.IsSet() ? FocusBrush.GetValue() : FCoreStyle::Get().GetBrush("FocusRectangle");
 
