Refactored Focus brush to be more hidden.
We would prefer it to be used a little as possible for the time being.

diff --git a/Engine/Source/Runtime/Slate/Private/Widgets/SViewport.cpp b/Engine/Source/Runtime/Slate/Private/Widgets/SViewport.cpp
--- a/Engine/Source/Runtime/Slate/Private/Widgets/SViewport.cpp
+++ b/Engine/Source/Runtime/Slate/Private/Widgets/SViewport.cpp
@@ -27,9 +27,6 @@ void SViewport::Construct( const FArguments& InArgs )
 	ViewportInterface = InArgs._ViewportInterface;
 	ViewportSize = InArgs._ViewportSize;
 
-	// Viewports should show focus
-	SetFocusBrush(nullptr);
-
 	this->ChildSlot
 	[
 		InArgs._Content.Widget
diff --git a/Engine/Source/Runtime/Slate/Public/Widgets/SViewport.h b/Engine/Source/Runtime/Slate/Public/Widgets/SViewport.h
--- a/Engine/Source/Runtime/Slate/Public/Widgets/SViewport.h
+++ b/Engine/Source/Runtime/Slate/Public/Widgets/SViewport.h
@@ -164,6 +164,12 @@ private:
 		return ChildSlot;
 	}
 
+	// Viewports shouldn't show focus
+	virtual const FSlateBrush* GetFocusBrush() const override
+	{
+		return nullptr;
+	}
+
 protected:
 
 	/** Interface to the rendering and I/O implementation of the viewport. */
diff --git a/Engine/Source/Runtime/SlateCore/Private/Widgets/SWidget.cpp b/Engine/Source/Runtime/SlateCore/Private/Widgets/SWidget.cpp
--- a/Engine/Source/Runtime/SlateCore/Private/Widgets/SWidget.cpp
+++ b/Engine/Source/Runtime/SlateCore/Private/Widgets/SWidget.cpp
@@ -414,9 +414,9 @@ bool SWidget::HasFocusedDescendants() const
 	return FSlateApplicationBase::Get().HasFocusedDescendants(SharedThis(this));
 }
 
-void SWidget::SetFocusBrush(TOptional<FSlateBrush*> InFocusBrush)
+const FSlateBrush* SWidget::GetFocusBrush() const
 {
-	FocusBrush = InFocusBrush;
+	return FCoreStyle::Get().GetBrush("FocusRectangle");
 }
 
 bool SWidget::HasMouseCapture() const
@@ -633,7 +633,7 @@ int32 SWidget::Paint(const FPaintArgs& Args, const FGeometry& AllottedGeometry,
 		TOptional<EFocusCause> FocusCause = HasAnyUserFocus();
 		if (FocusCause.IsSet() && FocusCause.GetValue() == EFocusCause::Navigation)
 		{
-			const FSlateBrush* BrushResource = FocusBrush.IsSet() ? FocusBrush.GetValue() : FCoreStyle::Get().GetBrush("FocusRectangle");
+			const FSlateBrush* BrushResource = GetFocusBrush();
 
 			if (BrushResource != nullptr)
 			{
diff --git a/Engine/Source/Runtime/SlateCore/Public/Widgets/SWidget.h b/Engine/Source/Runtime/SlateCore/Public/Widgets/SWidget.h
--- a/Engine/Source/Runtime/SlateCore/Public/Widgets/SWidget.h
+++ b/Engine/Source/Runtime/SlateCore/Public/Widgets/SWidget.h
@@ -551,9 +551,6 @@ public:
 	 */
 	bool HasFocusedDescendants() const;
 
-	/** @param InFocusBrush Sets the brush used to render focus. */
-	void SetFocusBrush(TOptional<FSlateBrush*> InFocusBrush);
-
 	/**
 	 * Checks to see if this widget is the current mouse captor
 	 *
@@ -807,6 +804,9 @@ protected:
 		return IsEnabled() && InParentEnabled;
 	}
 
+	/** @return a brush to draw focus, nullptr if no focus drawing is desired */
+	virtual const FSlateBrush* GetFocusBrush() const;
+
 private:
 
 	/**
@@ -870,9 +870,6 @@ protected:
 	/** Render transform pivot of this widget (in normalized local space) */
 	TAttribute< FVector2D > RenderTransformPivot;
 
-	/** The brush used to draw focus, if unset it will use system default. */
-	TOptional<FSlateBrush*> FocusBrush;
-
 	/** Is this widget hovered? */
 	bool bIsHovered;
 
