Slate : Windows now have a MinWidth, MinHeight, MaxWidth, and MaxHeight.
Known issue: Currently not enforcing these constraints upon construction; only when resizing.
Known issue on Linux : Linux sets cached window size before the OS has actually resized the window; this code does not respect the size constraints.

https://jira.ol.epicgames.net/browse/UETOOL-222

diff --git a/Engine/Source/Runtime/Core/Private/Windows/WindowsApplication.cpp b/Engine/Source/Runtime/Core/Private/Windows/WindowsApplication.cpp
--- a/Engine/Source/Runtime/Core/Private/Windows/WindowsApplication.cpp
+++ b/Engine/Source/Runtime/Core/Private/Windows/WindowsApplication.cpp
@@ -993,6 +993,18 @@ int32 FWindowsApplication::ProcessMessage( HWND hwnd, uint32 msg, WPARAM wParam,
 				}
 			}
 			break;
+
+		case WM_GETMINMAXINFO:
+			{
+				MINMAXINFO* MinMaxInfo = (MINMAXINFO*)lParam;
+				FWindowSizeLimits SizeLimits = MessageHandler->GetSizeLimitsForWindow(CurrentNativeEventWindow);
+				MinMaxInfo->ptMinTrackSize.x = FMath::RoundToInt( SizeLimits.GetMinWidth().Get(MinMaxInfo->ptMinTrackSize.x) );
+				MinMaxInfo->ptMinTrackSize.y = FMath::RoundToInt( SizeLimits.GetMinHeight().Get(MinMaxInfo->ptMinTrackSize.y) );
+				MinMaxInfo->ptMaxTrackSize.x = FMath::RoundToInt( SizeLimits.GetMaxWidth().Get(MinMaxInfo->ptMaxTrackSize.x) );
+				MinMaxInfo->ptMaxTrackSize.y = FMath::RoundToInt( SizeLimits.GetMaxHeight().Get(MinMaxInfo->ptMaxTrackSize.y) );
+				return 0;
+			}
+			break;
 			
 		case WM_NCLBUTTONDOWN:
 		case WM_NCRBUTTONDOWN:
diff --git a/Engine/Source/Runtime/Core/Public/GenericPlatform/GenericApplicationMessageHandler.h b/Engine/Source/Runtime/Core/Public/GenericPlatform/GenericApplicationMessageHandler.h
--- a/Engine/Source/Runtime/Core/Public/GenericPlatform/GenericApplicationMessageHandler.h
+++ b/Engine/Source/Runtime/Core/Public/GenericPlatform/GenericApplicationMessageHandler.h
@@ -3,6 +3,7 @@
 #pragma once
 #include "HAL/Platform.h"
 #include "Templates/SharedPointer.h"
+#include "Misc/Optional.h"
 
 class FGenericWindow;
 
@@ -168,6 +169,29 @@ enum class EConvertibleLaptopModes
 	Laptop
 };
 
+/** Defines the minimum and maximum dimensions that a window can take on. */
+struct FWindowSizeLimits
+{
+public:
+	FWindowSizeLimits& SetMinWidth(TOptional<float> InValue){ MinWidth = InValue; return *this; }
+	const TOptional<float>& GetMinWidth() const { return MinWidth; }
+
+	FWindowSizeLimits& SetMinHeight(TOptional<float> InValue){ MinHeight = InValue; return *this; }
+	const TOptional<float>& GetMinHeight() const { return MinHeight; }
+
+	FWindowSizeLimits& SetMaxWidth(TOptional<float> InValue){ MaxWidth = InValue; return *this; }
+	const TOptional<float>& GetMaxWidth() const { return MaxWidth; }
+
+	FWindowSizeLimits& SetMaxHeight(TOptional<float> InValue){ MaxHeight = InValue; return *this; }
+	const TOptional<float>& GetMaxHeight() const { return MaxHeight; }
+
+private:
+	TOptional<float> MinWidth;
+	TOptional<float> MinHeight;
+	TOptional<float> MaxWidth;
+	TOptional<float> MaxHeight;
+};
+
 
 class FGenericApplicationMessageHandler
 {
@@ -288,6 +312,11 @@ public:
 	
 	}
 
+	virtual FWindowSizeLimits GetSizeLimitsForWindow( const TSharedRef<FGenericWindow>& Window ) const
+	{
+		return FWindowSizeLimits();
+	}
+
 	virtual void OnResizingWindow( const TSharedRef< FGenericWindow >& Window )
 	{
 
diff --git a/Engine/Source/Runtime/Slate/Private/Framework/Application/SlateApplication.cpp b/Engine/Source/Runtime/Slate/Private/Framework/Application/SlateApplication.cpp
--- a/Engine/Source/Runtime/Slate/Private/Framework/Application/SlateApplication.cpp
+++ b/Engine/Source/Runtime/Slate/Private/Framework/Application/SlateApplication.cpp
@@ -4698,6 +4698,20 @@ void FSlateApplication::OnOSPaint( const TSharedRef< FGenericWindow >& PlatformW
 	Renderer->FlushCommands();
 }
 
+FWindowSizeLimits FSlateApplication::GetSizeLimitsForWindow(const TSharedRef<FGenericWindow>& Window) const
+{
+	TSharedPtr<SWindow> SlateWindow = FSlateWindowHelper::FindWindowByPlatformWindow(SlateWindows, Window);
+	if (SlateWindow.IsValid())
+	{
+		return SlateWindow->GetSizeLimits();
+	}
+	else
+	{
+		return FWindowSizeLimits();
+	}
+
+}
+
 void FSlateApplication::OnResizingWindow( const TSharedRef< FGenericWindow >& PlatformWindow )
 {
 	// Flush the rendering command queue to ensure that there aren't pending viewport draw commands for the old viewport size.
diff --git a/Engine/Source/Runtime/Slate/Public/Framework/Application/SlateApplication.h b/Engine/Source/Runtime/Slate/Public/Framework/Application/SlateApplication.h
--- a/Engine/Source/Runtime/Slate/Public/Framework/Application/SlateApplication.h
+++ b/Engine/Source/Runtime/Slate/Public/Framework/Application/SlateApplication.h
@@ -1077,6 +1077,8 @@ public:
 
 	virtual void OnOSPaint( const TSharedRef< FGenericWindow >& PlatformWindow ) override;
 
+	virtual FWindowSizeLimits GetSizeLimitsForWindow(const TSharedRef<FGenericWindow>& Window) const override;
+
 	virtual void OnResizingWindow( const TSharedRef< FGenericWindow >& PlatformWindow ) override;
 
 	virtual bool BeginReshapingWindow( const TSharedRef< FGenericWindow >& PlatformWindow ) override;
diff --git a/Engine/Source/Runtime/SlateCore/Private/Widgets/SWindow.cpp b/Engine/Source/Runtime/SlateCore/Private/Widgets/SWindow.cpp
--- a/Engine/Source/Runtime/SlateCore/Private/Widgets/SWindow.cpp
+++ b/Engine/Source/Runtime/SlateCore/Private/Widgets/SWindow.cpp
@@ -201,6 +201,11 @@ void SWindow::Construct(const FArguments& InArgs)
 	this->bHasSizingFrame = !InArgs._IsPopupWindow && InArgs._SizingRule == ESizingRule::UserSized;
 	this->LayoutBorder = InArgs._LayoutBorder;
 	this->UserResizeBorder = InArgs._UserResizeBorder;
+	this->SizeLimits = FWindowSizeLimits()
+		.SetMinWidth(InArgs._MinWidth)
+		.SetMinHeight(InArgs._MinHeight)
+		.SetMaxWidth(InArgs._MaxWidth)
+		.SetMaxHeight(InArgs._MaxHeight);
 	
 	// calculate window size from client size
 	const bool bCreateTitleBar = InArgs._CreateTitleBar && !bIsPopupWindow && !bIsCursorDecoratorWindow && !bHasOSWindowBorder;
@@ -517,6 +522,11 @@ TSharedRef<FHittestGrid> SWindow::GetHittestGrid()
 	return HittestGrid;
 }
 
+FWindowSizeLimits SWindow::GetSizeLimits() const
+{
+	return SizeLimits;
+}
+
 void SWindow::Tick( const FGeometry& AllottedGeometry, const double InCurrentTime, const float InDeltaTime )
 {
 	if( Morpher.bIsActive )
diff --git a/Engine/Source/Runtime/SlateCore/Public/Widgets/SWindow.h b/Engine/Source/Runtime/SlateCore/Public/Widgets/SWindow.h
--- a/Engine/Source/Runtime/SlateCore/Public/Widgets/SWindow.h
+++ b/Engine/Source/Runtime/SlateCore/Public/Widgets/SWindow.h
@@ -138,6 +138,18 @@ public:
 		/** Can this window be minimized? */
 		SLATE_ARGUMENT( bool, SupportsMinimize )
 
+		/** The smallest width this window can be in Desktop Pixel Units. */
+		SLATE_ARGUMENT( TOptional<float>, MinWidth )
+		
+		/** The smallest height this window can be in Desktop Pixel Units. */
+		SLATE_ARGUMENT( TOptional<float>, MinHeight )
+		
+		/** The biggest width this window can be in Desktop Pixel Units. */
+		SLATE_ARGUMENT( TOptional<float>, MaxWidth )
+
+		/** The biggest height this window can be in Desktop Pixel Units. */
+		SLATE_ARGUMENT( TOptional<float>, MaxHeight )
+
 		/** True if we should initially create a traditional title bar area.  If false, the user must embed the title
 		    area content into the window manually, taking into account platform-specific considerations!  Has no
 			effect for certain types of windows (popups, tool-tips, etc.) */
@@ -655,7 +667,7 @@ public:
 	}
 
 	/**
-	 * Returns the viewport size, taking into consideration if the window size should drive the viewport size
+	 * Returns the viewport size, taking into consideration if the window size should drive the viewport size 
 	 */
 	inline FVector2D GetViewportSize() const
 	{
@@ -678,6 +690,9 @@ public:
 	 */
 	TSharedRef<FHittestGrid> GetHittestGrid();
 
+	/** Optional constraints on min and max sizes that this window can be. */
+	FWindowSizeLimits GetSizeLimits() const;
+
 public:
 
 	// SWidget overrides
@@ -841,6 +856,9 @@ protected:
 
 private:
 
+	/** Min and Max values for Width and Height; all optional. */
+	FWindowSizeLimits SizeLimits;
+
 	/** The native window that is backing this Slate Window */
 	TSharedPtr<FGenericWindow> NativeWindow;
 
