Slate - removing commented out code.

diff --git a/Engine/Source/Runtime/SlateCore/Public/Styling/SlateBrush.h b/Engine/Source/Runtime/SlateCore/Public/Styling/SlateBrush.h
--- a/Engine/Source/Runtime/SlateCore/Public/Styling/SlateBrush.h
+++ b/Engine/Source/Runtime/SlateCore/Public/Styling/SlateBrush.h
@@ -17,7 +17,7 @@ namespace ESlateBrushDrawType
 	enum Type
 	{
 		/** Don't do anything */
-		NoDrawType UMETA(DisplayName="None"), // UMETA(Hidden),
+		NoDrawType UMETA(DisplayName="None"),
 
 		/** Draw a 3x3 box, where the sides and the middle stretch based on the Margin */
 		Box,
