Slate - Fixing a bug in the element batcher, images if they had a margin were still being treated as a box, instead of it properly ignoring the field.  Adding an explicit check for DrawAs Image to avoid drawing images that happen to have margins - with the 9-box.

diff --git a/Engine/Source/Runtime/SlateCore/Private/Rendering/ElementBatcher.cpp b/Engine/Source/Runtime/SlateCore/Private/Rendering/ElementBatcher.cpp
--- a/Engine/Source/Runtime/SlateCore/Private/Rendering/ElementBatcher.cpp
+++ b/Engine/Source/Runtime/SlateCore/Private/Rendering/ElementBatcher.cpp
@@ -197,9 +197,10 @@ void FSlateElementBatcher::AddBoxElement( const FSlateDrawElement& DrawElement )
 	// The clip rect is NOT subject to the rotations specified by MakeRotatedBox.
 	FSlateRotatedClipRectType RenderClipRect = ToSnappedRotatedRect(InClippingRect, InverseLayoutTransform, ElementRenderTransform);
 
-	check( InPayload.BrushResource );
+	check(InPayload.BrushResource);
+	const FSlateBrush* BrushResource = InPayload.BrushResource;
 
-	if(InPayload.BrushResource->DrawAs != ESlateBrushDrawType::NoDrawType)
+	if ( BrushResource->DrawAs != ESlateBrushDrawType::NoDrawType )
 	{
 		// Do pixel snapping
 		FVector2D TopLeft(0,0);
@@ -215,7 +216,7 @@ void FSlateElementBatcher::AddBoxElement( const FSlateDrawElement& DrawElement )
 
 		FVector2D HalfTexel;
 
-		FSlateShaderResourceProxy* ResourceProxy = ResourceManager.GetShaderResource( *InPayload.BrushResource );
+		FSlateShaderResourceProxy* ResourceProxy = ResourceManager.GetShaderResource(*BrushResource);
 		FSlateShaderResource* Resource = nullptr;
 		if( ResourceProxy )
 		{
@@ -240,13 +241,13 @@ void FSlateElementBatcher::AddBoxElement( const FSlateDrawElement& DrawElement )
 		}
 
 
-		FColor Tint = GetElementColor( InPayload.Tint, InPayload.BrushResource );
+		FColor Tint = GetElementColor(InPayload.Tint, BrushResource);
 
-		const ESlateBrushTileType::Type TilingRule = InPayload.BrushResource->Tiling;
+		const ESlateBrushTileType::Type TilingRule = BrushResource->Tiling;
 		const bool bTileHorizontal = (TilingRule == ESlateBrushTileType::Both || TilingRule == ESlateBrushTileType::Horizontal);
 		const bool bTileVertical = (TilingRule == ESlateBrushTileType::Both || TilingRule == ESlateBrushTileType::Vertical);
 
-		const ESlateBrushMirrorType::Type MirroringRule = InPayload.BrushResource->Mirroring;
+		const ESlateBrushMirrorType::Type MirroringRule = BrushResource->Mirroring;
 		const bool bMirrorHorizontal = (MirroringRule == ESlateBrushMirrorType::Both || MirroringRule == ESlateBrushMirrorType::Horizontal);
 		const bool bMirrorVertical = (MirroringRule == ESlateBrushMirrorType::Both || MirroringRule == ESlateBrushMirrorType::Vertical);
 
@@ -267,10 +268,10 @@ void FSlateElementBatcher::AddBoxElement( const FSlateDrawElement& DrawElement )
 		// The offset into the index buffer where this elements indices start
 		uint32 IndexOffsetStart = BatchIndices.Num();
 
-		if( InPayload.BrushResource->Margin.Left!= 0.0f || 
-			InPayload.BrushResource->Margin.Top != 0.0f ||
-			InPayload.BrushResource->Margin.Right != 0.0f ||
-			InPayload.BrushResource->Margin.Bottom != 0.0f )
+		const FMargin& Margin = BrushResource->Margin;
+
+		if ( BrushResource->DrawAs != ESlateBrushDrawType::Image &&
+			( Margin.Left != 0.0f || Margin.Top != 0.0f || Margin.Right != 0.0f || Margin.Bottom != 0.0f ) )
 		{
 			// Create 9 quads for the box element based on the following diagram
 			//     ___LeftMargin    ___RightMargin
@@ -284,8 +285,6 @@ void FSlateElementBatcher::AddBoxElement( const FSlateDrawElement& DrawElement )
 			//  |  |             |  | ___BottomMargin
 			//  +--+-------------+--+
 
-			const FMargin& Margin = InPayload.BrushResource->Margin;
-
 
 			// Determine the texture coordinates for each quad
 			// These are not scaled.
