#UE4

- GetBubblePath() now takes in a radius
- Added IsInteractable() to SWidget (currently only returns true for SButton)
- Added Cursor Radius to Slate App

diff --git a/Engine/Source/Runtime/Slate/Private/Framework/Application/SlateApplication.cpp b/Engine/Source/Runtime/Slate/Private/Framework/Application/SlateApplication.cpp
--- a/Engine/Source/Runtime/Slate/Private/Framework/Application/SlateApplication.cpp
+++ b/Engine/Source/Runtime/Slate/Private/Framework/Application/SlateApplication.cpp
@@ -653,6 +653,7 @@ FSlateApplication::FSlateApplication()
 	, bSlateWindowActive(true)
 	, Scale( 1.0f )
 	, DragTriggerDistnace( 5.0f )
+	, CursorRadius(0.0f)
 	, LastUserInteractionTime( 0.0 )
 	, LastUserInteractionTimeForThrottling( 0.0 )
 	, SlateSoundDevice( MakeShareable(new FNullSlateSoundDevice()) )
@@ -812,7 +813,7 @@ FWidgetPath FSlateApplication::LocateWindowUnderMouse( FVector2D ScreenspaceMous
 
 		if ( Window->IsVisible() && AcceptsInput && Window->IsScreenspaceMouseWithin(ScreenspaceMouseCoordinate) && !bPrevWindowWasModal )
 		{
-			const TArray<FWidgetAndPointer> WidgetsAndCursors = Window->GetHittestGrid()->GetBubblePath(ScreenspaceMouseCoordinate, bIgnoreEnabledStatus);
+			const TArray<FWidgetAndPointer> WidgetsAndCursors = Window->GetHittestGrid()->GetBubblePath(ScreenspaceMouseCoordinate, GetCursorRadius(), bIgnoreEnabledStatus);
 			return FWidgetPath( WidgetsAndCursors );
 		}
 	}
@@ -3011,6 +3012,15 @@ void FSlateApplication::SetAnalogCursorEnable(bool bEnable, TSharedPtr<class FAn
 	}
 }
 
+void FSlateApplication::SetCursorRadius(float NewRadius)
+{
+	CursorRadius = FMath::Max<float>(0.0f, NewRadius);
+}
+
+float FSlateApplication::GetCursorRadius() const
+{
+	return CursorRadius;
+}
 
 FVector2D FSlateApplication::CalculatePopupWindowPosition( const FSlateRect& InAnchor, const FVector2D& InSize, const EOrientation Orientation ) const
 {
@@ -5138,3 +5148,5 @@ void FSlateApplication::SetWidgetReflector(const TSharedRef<IWidgetReflector>& W
 
 	WidgetReflectorPtr = WidgetReflector;
 }
+
+
diff --git a/Engine/Source/Runtime/Slate/Private/Widgets/Input/SButton.cpp b/Engine/Source/Runtime/Slate/Private/Widgets/Input/SButton.cpp
--- a/Engine/Source/Runtime/Slate/Private/Widgets/Input/SButton.cpp
+++ b/Engine/Source/Runtime/Slate/Private/Widgets/Input/SButton.cpp
@@ -281,6 +281,11 @@ void SButton::OnMouseLeave( const FPointerEvent& MouseEvent )
 	}
 }
 
+bool SButton::IsInteractable() const
+{
+	return IsEnabled() && SupportsKeyboardFocus();
+}
+
 bool SButton::IsPreciseTapOrClick(const FPointerEvent& MouseEvent) const
 {
 	return ( TouchMethod == EButtonTouchMethod::PreciseTap && MouseEvent.IsTouchEvent() ) ||
diff --git a/Engine/Source/Runtime/Slate/Public/Framework/Application/SlateApplication.h b/Engine/Source/Runtime/Slate/Public/Framework/Application/SlateApplication.h
--- a/Engine/Source/Runtime/Slate/Public/Framework/Application/SlateApplication.h
+++ b/Engine/Source/Runtime/Slate/Public/Framework/Application/SlateApplication.h
@@ -949,6 +949,12 @@ public:
 	/** Set the analog cursor to be enabled or disabled. */
 	void SetAnalogCursorEnable(bool bEnable, TSharedPtr<class FAnalogCursor> OptionalNewAnalogCursor = nullptr);
 
+	/** Sets the hit detection radius of the cursor */
+	void SetCursorRadius(float NewRadius);
+
+	/** Getter for the cursor radius */
+	float GetCursorRadius() const;
+
 public:
 
 	// Begin FSlateApplicationBase interface
@@ -1268,6 +1274,8 @@ private:
 
 	TSharedPtr<FAnalogCursor> AnalogCursor;
 
+	/** The hit-test radius of the cursor. Default value is 0. */
+	float CursorRadius;
 
 	struct FUserFocusEntry
 	{
diff --git a/Engine/Source/Runtime/Slate/Public/Widgets/Input/SButton.h b/Engine/Source/Runtime/Slate/Public/Widgets/Input/SButton.h
--- a/Engine/Source/Runtime/Slate/Public/Widgets/Input/SButton.h
+++ b/Engine/Source/Runtime/Slate/Public/Widgets/Input/SButton.h
@@ -141,6 +141,8 @@ public:
 
 	virtual void OnMouseLeave( const FPointerEvent& MouseEvent ) override;
 
+	virtual bool IsInteractable() const override;
+
 protected:
 
 	/** @return combines the user-specified margin and the button's internal margin. */
diff --git a/Engine/Source/Runtime/SlateCore/Private/Input/HittestGrid.cpp b/Engine/Source/Runtime/SlateCore/Private/Input/HittestGrid.cpp
--- a/Engine/Source/Runtime/SlateCore/Private/Input/HittestGrid.cpp
+++ b/Engine/Source/Runtime/SlateCore/Private/Input/HittestGrid.cpp
@@ -5,22 +5,110 @@
 
 DEFINE_LOG_CATEGORY_STATIC(LogHittestDebug, Display, All);
 
-static const FVector2D CellSize(128.0f, 128.0f);
+//
+// Helper Functions
+//
+
+FVector2D ClosestPointOnSlateRotatedRect(const FVector2D &Point, const FSlateRotatedRect& RotatedRect)
+{
+	//no need to do any testing if we are inside of the rect
+	if (RotatedRect.IsUnderLocation(Point))
+	{
+		return Point;
+	}
+
+	const static int32 NumOfCorners = 4;
+	FVector2D Corners[NumOfCorners];
+	Corners[0] = RotatedRect.TopLeft;
+	Corners[1] = Corners[0] + RotatedRect.ExtentX;
+	Corners[2] = Corners[1] + RotatedRect.ExtentY;
+	Corners[3] = Corners[0] + RotatedRect.ExtentY;
+
+	FVector2D RetPoint;
+	float ClosestDistSq = FLT_MAX;
+	for (int32 i = 0; i < NumOfCorners; ++i)
+	{
+		//grab the closest point along the line segment
+		const FVector2D ClosestPoint = FMath::ClosestPointOnSegment2D(Point, Corners[i], Corners[(i + 1) % NumOfCorners]);
+
+		//get the distance between the two
+		const float TestDist = FVector2D::DistSquared(Point, ClosestPoint);
+
+		//if the distance is smaller than the current smallest, update our closest
+		if (TestDist < ClosestDistSq)
+		{
+			RetPoint = ClosestPoint;
+			ClosestDistSq = TestDist;
+		}
+	}
+
+	return RetPoint;
+}
+
+
+float DistanceSqToSlateRotatedRect(const FVector2D &Point, const FSlateRotatedRect& RotatedRect)
+{
+	return FVector2D::DistSquared(ClosestPointOnSlateRotatedRect(Point, RotatedRect), Point);
+}
+
+
+bool IsOverlappingSlateRotatedRect(const FVector2D& Point, const float Radius, const FSlateRotatedRect& RotatedRect)
+{
+	return DistanceSqToSlateRotatedRect(Point, RotatedRect) <= (Radius * Radius);
+}
+
+
+bool ContainsInteractableWidget(const TArray<FWidgetAndPointer>& PathToTest)
+{
+	for (int32 i = PathToTest.Num() - 1; i >= 0; --i)
+	{
+		const FWidgetAndPointer& WidgetAndPointer = PathToTest[i];
+		if (WidgetAndPointer.Widget->IsInteractable())
+		{
+			return true;
+		}
+	}
+	return false;
+}
+
+
+
+//
+// FHittestGrid
+//
+
+const FVector2D CellSize(128.0f, 128.0f);
+
+struct FHittestGrid::FGridTestingParams
+{
+	/** Ctor */
+	FGridTestingParams()
+	: CellCoord(-1, -1)
+	, CursorPositionInGrid(FVector2D::ZeroVector)
+	, Radius(-1.0f)
+	, bTestWidgetIsInteractive(false)
+	{}
+
+	FIntPoint CellCoord;
+	FVector2D CursorPositionInGrid;
+	float Radius;
+	bool bTestWidgetIsInteractive;
+};
+
 
 struct FHittestGrid::FCachedWidget
 {
-	FCachedWidget( int32 InParentIndex, const FArrangedWidget& InWidget, const FSlateRect& InClippingRect )
+	FCachedWidget(int32 InParentIndex, const FArrangedWidget& InWidget, const FSlateRect& InClippingRect)
 	: WidgetPtr(InWidget.Widget)
-	, CachedGeometry( InWidget.Geometry )
-	, ClippingRect( InClippingRect )
+	, CachedGeometry(InWidget.Geometry)
+	, ClippingRect(InClippingRect)
 	, Children()
-	, ParentIndex( InParentIndex )
-	{
-	}
+	, ParentIndex(InParentIndex)
+	{}
 
-	void AddChild( const int32 ChildIndex )
+	void AddChild(const int32 ChildIndex)
 	{
-		Children.Add( ChildIndex );
+		Children.Add(ChildIndex);
 	}
 
 	TWeakPtr<SWidget> WidgetPtr;
@@ -39,113 +127,96 @@ FHittestGrid::FHittestGrid()
 {
 }
 
-TArray<FWidgetAndPointer> FHittestGrid::GetBubblePath( FVector2D DesktopSpaceCoordinate, bool bIgnoreEnabledStatus )
+
+TArray<FWidgetAndPointer> FHittestGrid::GetBubblePath(FVector2D DesktopSpaceCoordinate, float CursorRadius, bool bIgnoreEnabledStatus)
 {
+	//grab the radius, and if we are non-zero
+	const bool bDirectTestingOnly = (CursorRadius <= 0.0f);
+
+	//calculate the cursor position in the grid
+	const FVector2D CursorPositionInGrid = DesktopSpaceCoordinate - GridOrigin;
+
 	if (WidgetsCachedThisFrame->Num() > 0 && Cells.Num() > 0)
 	{
-		const FVector2D CursorPositionInGrid = DesktopSpaceCoordinate - GridOrigin;
-		const FIntPoint CellCoordinate = FIntPoint(
-			FMath::Min( FMath::Max(FMath::FloorToInt(CursorPositionInGrid.X / CellSize.X), 0), NumCells.X-1),
-			FMath::Min( FMath::Max(FMath::FloorToInt(CursorPositionInGrid.Y / CellSize.Y), 0), NumCells.Y-1 ) );
-		
-		static FVector2D LastCoordinate = FVector2D::ZeroVector;
-		if ( LastCoordinate != CursorPositionInGrid )
-		{
-			LastCoordinate = CursorPositionInGrid;
-		}
+		//grab the path for direct testing first
+		FGridTestingParams DirectTestingParams;
+		DirectTestingParams.CursorPositionInGrid = CursorPositionInGrid;
+		DirectTestingParams.CellCoord = GetCellCoordinate(CursorPositionInGrid);
+		DirectTestingParams.Radius = 0.0f;
+		DirectTestingParams.bTestWidgetIsInteractive = false;
 
-		checkf( (CellCoordinate.Y*NumCells.X + CellCoordinate.X) < Cells.Num(), TEXT("Index out of range, CellCoordinate is: %d %d CursorPosition is: %f %f"), CellCoordinate.X, CellCoordinate.Y, CursorPositionInGrid.X, CursorPositionInGrid.Y );
+		const FWidgetPathAndDist DirectBubblePathInfo = GetWidgetPathAndDist(DirectTestingParams, bIgnoreEnabledStatus);
 
-		const TArray<int32>& IndexesInCell = CellAt( CellCoordinate.X, CellCoordinate.Y ).CachedWidgetIndexes;
-		int32 HitWidgetIndex = INDEX_NONE;
-	
-		// Consider front-most widgets first for hittesting.
-		for ( int32 i = IndexesInCell.Num()-1; i>=0 && HitWidgetIndex==INDEX_NONE; --i )
+		//if we aren't doing a radius check, or we already have a direct path, use that
+		if (bDirectTestingOnly || ContainsInteractableWidget(DirectBubblePathInfo.BubblePath))
 		{
-			check( IndexesInCell[i] < WidgetsCachedThisFrame->Num() ); 
+			return DirectBubblePathInfo.BubblePath;
+		}
 
-			const FCachedWidget& TestCandidate = (*WidgetsCachedThisFrame)[IndexesInCell[i]];
+		//if we are here, we need to check other cells
+		const FVector2D RadiusVector(CursorRadius, CursorRadius);
+		const FIntPoint ULIndex = GetCellCoordinate(CursorPositionInGrid - RadiusVector);
+		const FIntPoint LRIndex = GetCellCoordinate(CursorPositionInGrid + RadiusVector);
 
-			// Compute the render space clipping rect (FGeometry exposes a layout space clipping rect).
-			FSlateRotatedRect DesktopOrientedClipRect = 
-				TransformRect(
-					Concatenate(
-						Inverse(TestCandidate.CachedGeometry.GetAccumulatedLayoutTransform()), 
-						TestCandidate.CachedGeometry.GetAccumulatedRenderTransform()
-					), 
-					FSlateRotatedRect(TestCandidate.CachedGeometry.GetClippingRect().IntersectionWith(TestCandidate.ClippingRect))
-				);
+		//first, find all the overlapping cells
+		TArray<FIntPoint> CellIndexes;
+		CellIndexes.Reserve(16);
 
-			if (DesktopOrientedClipRect.IsUnderLocation(DesktopSpaceCoordinate) && TestCandidate.WidgetPtr.IsValid())
-			{
-				HitWidgetIndex = IndexesInCell[i];
-			}
-		}
-		
-		TArray<FWidgetAndPointer> BubblePath =[&]()
+		for (int32 YIndex = ULIndex.Y; YIndex <= LRIndex.Y; ++YIndex)
 		{
-			if( HitWidgetIndex != INDEX_NONE )
+			for (int32 XIndex = ULIndex.X; XIndex <= LRIndex.X; ++XIndex)
 			{
-				const FCachedWidget& PhysicallyHitWidget = (*WidgetsCachedThisFrame)[HitWidgetIndex];
-				if(PhysicallyHitWidget.CustomPath.IsValid())
+				const FIntPoint PointToTest(XIndex, YIndex);
+				if (IsValidCellCoord(PointToTest))
 				{
-					return PhysicallyHitWidget.CustomPath.Pin()->GetBubblePathAndVirtualCursors(PhysicallyHitWidget.CachedGeometry, DesktopSpaceCoordinate, bIgnoreEnabledStatus);
+					CellIndexes.Add(PointToTest);
 				}
 			}
+		}
 
-			return TArray<FWidgetAndPointer>();
-
-		}();
+		//setup the radius testing params
+		FGridTestingParams RadiusTestingParams;
+		RadiusTestingParams.CursorPositionInGrid = CursorPositionInGrid;
+		RadiusTestingParams.Radius = CursorRadius;
+		RadiusTestingParams.bTestWidgetIsInteractive = true;
 
-		if (HitWidgetIndex != INDEX_NONE)
+		//next, collect valid paths from all those cells
+		TArray<FWidgetPathAndDist> PathsAndDistances;
+		for (const FIntPoint& CellCoord : CellIndexes)
 		{
-			int32 CurWidgetIndex=HitWidgetIndex;
-			bool bPathUninterrupted = false;
-			do
-			{
-				check( CurWidgetIndex < WidgetsCachedThisFrame->Num() );
-				const FCachedWidget& CurCachedWidget = (*WidgetsCachedThisFrame)[CurWidgetIndex];
-				const TSharedPtr<SWidget> CachedWidgetPtr = CurCachedWidget.WidgetPtr.Pin();
+			//update the cell coord property
+			RadiusTestingParams.CellCoord = CellCoord;
 
-		
-				bPathUninterrupted = CachedWidgetPtr.IsValid();
-				if (bPathUninterrupted)
-				{
-					BubblePath.Insert( FWidgetAndPointer( FArrangedWidget(CachedWidgetPtr.ToSharedRef(), CurCachedWidget.CachedGeometry), TSharedPtr<FVirtualPointerPosition>() ), 0 );
-					CurWidgetIndex = CurCachedWidget.ParentIndex;
-				}
-			}
-			while (CurWidgetIndex != INDEX_NONE && bPathUninterrupted);
-			
-			if (!bPathUninterrupted)
+			const FWidgetPathAndDist TestPath = GetWidgetPathAndDist(RadiusTestingParams, bIgnoreEnabledStatus);
+			if ( TestPath.IsValidPath() )
 			{
-				// A widget in the path to the root has been removed, so anything
-				// we thought we had hittest is no longer actually there.
-				// Pretend we didn't hit anything.
-				BubblePath = TArray<FWidgetAndPointer>();
+				PathsAndDistances.Add(TestPath);
 			}
+		}
 
-			// Disabling a widget disables all of its logical children
-			// This effect is achieved by truncating the path to the
-			// root-most enabled widget.
-			if ( !bIgnoreEnabledStatus )
+		// We have paths from multiple cells; use the one that is closest to the cursor's center.
+		if (PathsAndDistances.Num() > 0)
+		{
+			//sort the paths, and get the closest one that is valid
+			PathsAndDistances.Sort([](const FWidgetPathAndDist& A, const FWidgetPathAndDist& B)
 			{
-				const int32 DisabledWidgetIndex = BubblePath.IndexOfByPredicate( []( const FArrangedWidget& SomeWidget ){ return !SomeWidget.Widget->IsEnabled( ); } );
-				if (DisabledWidgetIndex != INDEX_NONE)
+				return A.DistToTopWidgetSq < B.DistToTopWidgetSq;
+			});
+
+			for (const FWidgetPathAndDist& TestPath : PathsAndDistances)
+			{
+				if (ContainsInteractableWidget(TestPath.BubblePath))
 				{
-					BubblePath.RemoveAt( DisabledWidgetIndex, BubblePath.Num() - DisabledWidgetIndex );
+					return TestPath.BubblePath;
 				}
 			}
-			
-			return BubblePath;
-		}
-		else
-		{
-			return TArray<FWidgetAndPointer>();
 		}
+
+		return 	DirectBubblePathInfo.BubblePath;
 	}
 	else
 	{
+		// We didn't hit anything.
 		return TArray<FWidgetAndPointer>();
 	}
 }
@@ -477,6 +548,16 @@ const FHittestGrid::FCell& FHittestGrid::CellAt( const int32 X, const int32 Y )
 	return Cells[ Y*NumCells.X + X ];
 }
 
+bool FHittestGrid::IsValidCellCoord(const FIntPoint& CellCoord) const
+{
+	return IsValidCellCoord(CellCoord.X, CellCoord.Y);
+}
+
+bool FHittestGrid::IsValidCellCoord(const int32 XCoord, const int32 YCoord) const
+{
+	const int32 Index = (YCoord * NumCells.X) + XCoord;
+	return Cells.IsValidIndex(Index);
+}
 
 void FHittestGrid::LogGrid() const
 {
@@ -530,3 +611,120 @@ void FHittestGrid::LogChildren(int32 Index, int32 IndentLevel, const TArray<FCac
 		LogChildren( CachedWidget.Children[i], IndentLevel+1, WidgetsCachedThisFrame );
 	}
 }
+
+FHittestGrid::FWidgetPathAndDist FHittestGrid::GetWidgetPathAndDist(const FGridTestingParams& Params, const bool bIgnoreEnabledStatus) const
+{
+	//Also grab the Hit Index, and the distance to the top hit
+	const FIndexAndDistance HitIndex = GetHitIndexFromCellIndex(Params);
+	
+	//if we get here, we want to do the testing for 3D Widgets
+	if (HitIndex.WidgetIndex != INDEX_NONE)
+	{
+		const FCachedWidget& PhysicallyHitWidget = (*WidgetsCachedThisFrame)[HitIndex.WidgetIndex];
+		if (PhysicallyHitWidget.CustomPath.IsValid())
+		{
+			const FVector2D DesktopSpaceCoordinate = Params.CursorPositionInGrid + GridOrigin;
+			const TArray<FWidgetAndPointer> BubblePath = PhysicallyHitWidget.CustomPath.Pin()->GetBubblePathAndVirtualCursors(PhysicallyHitWidget.CachedGeometry, DesktopSpaceCoordinate, bIgnoreEnabledStatus);
+			return FWidgetPathAndDist(BubblePath,0.0);
+		}
+		else
+		{
+			//get the path from the hit index, check if anything came back
+			const TArray<FWidgetAndPointer> BubblePath = GetBubblePathFromHitIndex(HitIndex.WidgetIndex, bIgnoreEnabledStatus);
+			return FWidgetPathAndDist(BubblePath, BubblePath.Num() > 0 ? HitIndex.DistanceSqToWidget : -1.0f);
+		}
+	}
+
+	//return if we have a valid path or not
+	return FWidgetPathAndDist();
+}
+
+FHittestGrid::FIndexAndDistance FHittestGrid::GetHitIndexFromCellIndex(const FGridTestingParams& Params) const
+{
+	//check if the cell coord 
+	if (IsValidCellCoord(Params.CellCoord))
+	{
+		const TArray<int32>& IndexesInCell = CellAt(Params.CellCoord.X, Params.CellCoord.Y).CachedWidgetIndexes;
+
+		// Consider front-most widgets first for hittesting.
+		for (int32 i = IndexesInCell.Num() - 1; i >= 0; --i)
+		{
+			const int32 WidgetIndex = IndexesInCell[i];
+
+			check(WidgetsCachedThisFrame->IsValidIndex(WidgetIndex));
+
+			const FHittestGrid::FCachedWidget& TestCandidate = (*WidgetsCachedThisFrame)[WidgetIndex];
+
+			// Compute the render space clipping rect (FGeometry exposes a layout space clipping rect).
+			FSlateRotatedRect DesktopOrientedClipRect =
+				TransformRect(
+				Concatenate(
+				Inverse(TestCandidate.CachedGeometry.GetAccumulatedLayoutTransform()),
+				TestCandidate.CachedGeometry.GetAccumulatedRenderTransform()
+				),
+				FSlateRotatedRect(TestCandidate.CachedGeometry.GetClippingRect().IntersectionWith(TestCandidate.ClippingRect))
+				);
+
+			//test if it is a valid index or not
+			const bool bIsValidWidget = !Params.bTestWidgetIsInteractive || (TestCandidate.WidgetPtr.IsValid() && TestCandidate.WidgetPtr.Pin()->IsInteractable());
+
+			//now do the overlap test
+			if (bIsValidWidget && IsOverlappingSlateRotatedRect(Params.CursorPositionInGrid, Params.Radius, DesktopOrientedClipRect))
+			{
+				//fill in the distance param, if it isn't NULL
+				const bool bNeedsDistanceSearch = Params.Radius > 0.0f;
+				const float DistSq = (bNeedsDistanceSearch) ? DistanceSqToSlateRotatedRect(Params.CursorPositionInGrid, DesktopOrientedClipRect) : 0.0f;
+				return FIndexAndDistance(WidgetIndex, DistSq);
+			}
+		}
+	}
+
+	return FIndexAndDistance();
+}
+
+TArray<FWidgetAndPointer> FHittestGrid::GetBubblePathFromHitIndex(const int32 HitIndex, const bool bIgnoreEnabledStatus) const
+{
+	TArray<FWidgetAndPointer> BubblePath;
+
+	if (WidgetsCachedThisFrame->IsValidIndex(HitIndex))
+	{
+		int32 CurWidgetIndex = HitIndex;
+		bool bPathUninterrupted = false;
+		do
+		{
+			check(CurWidgetIndex < WidgetsCachedThisFrame->Num());
+			const FHittestGrid::FCachedWidget& CurCachedWidget = (*WidgetsCachedThisFrame)[CurWidgetIndex];
+			const TSharedPtr<SWidget> CachedWidgetPtr = CurCachedWidget.WidgetPtr.Pin();
+
+
+			bPathUninterrupted = CachedWidgetPtr.IsValid();
+			if (bPathUninterrupted)
+			{
+				BubblePath.Insert(FWidgetAndPointer(FArrangedWidget(CachedWidgetPtr.ToSharedRef(), CurCachedWidget.CachedGeometry), TSharedPtr<FVirtualPointerPosition>()), 0);
+				CurWidgetIndex = CurCachedWidget.ParentIndex;
+			}
+		} while (CurWidgetIndex != INDEX_NONE && bPathUninterrupted);
+
+		if (!bPathUninterrupted)
+		{
+			// A widget in the path to the root has been removed, so anything
+			// we thought we had hittest is no longer actually there.
+			// Pretend we didn't hit anything.
+			BubblePath.Reset();
+		}
+
+		// Disabling a widget disables all of its logical children
+		// This effect is achieved by truncating the path to the
+		// root-most enabled widget.
+		if (!bIgnoreEnabledStatus)
+		{
+			const int32 DisabledWidgetIndex = BubblePath.IndexOfByPredicate([](const FArrangedWidget& SomeWidget){ return !SomeWidget.Widget->IsEnabled(); });
+			if (DisabledWidgetIndex != INDEX_NONE)
+			{
+				BubblePath.RemoveAt(DisabledWidgetIndex, BubblePath.Num() - DisabledWidgetIndex);
+			}
+		}
+	}
+
+	return BubblePath;
+}
diff --git a/Engine/Source/Runtime/SlateCore/Public/Input/HittestGrid.h b/Engine/Source/Runtime/SlateCore/Public/Input/HittestGrid.h
--- a/Engine/Source/Runtime/SlateCore/Public/Input/HittestGrid.h
+++ b/Engine/Source/Runtime/SlateCore/Public/Input/HittestGrid.h
@@ -26,7 +26,7 @@ public:
 	 * Given a Slate Units coordinate in virtual desktop space, perform a hittest
 	 * and return the path along which the corresponding event would be bubbled.
 	 */
-	TArray<FWidgetAndPointer> GetBubblePath( FVector2D DesktopSpaceCoordinate, bool bIgnoreEnabledStatus );
+	TArray<FWidgetAndPointer> GetBubblePath( FVector2D DesktopSpaceCoordinate, float CursorRadius, bool bIgnoreEnabledStatus );
 
 	/**
 	 * Clear the hittesting area and prepare to execute a new frame.
@@ -73,7 +73,56 @@ private:
 	 */
 	struct FCachedWidget;
 
-private:
+	/** Shared arguments to helper functions. */
+	struct FGridTestingParams;
+
+	/** Helper functions */
+	bool IsValidCellCoord(const FIntPoint& CellCoord) const;
+	bool IsValidCellCoord(const int32 XCoord, const int32 YCoord) const;
+
+	/** Bubble path and distance to leafmost widget from point of picking. */
+	struct FWidgetPathAndDist
+	{
+		/** Ctor */
+		FWidgetPathAndDist( float InDistanceSq = -1.0f )
+		: DistToTopWidgetSq(InDistanceSq)
+		{}
+
+		FWidgetPathAndDist(const TArray<FWidgetAndPointer>& InPath, float InDistanceSq)
+		: BubblePath(InPath)
+		, DistToTopWidgetSq(InDistanceSq)
+		{}
+
+		void Clear()
+		{
+			BubblePath.Reset();
+			DistToTopWidgetSq = -1.0f;
+		}
+
+		bool IsValidPath() const
+		{
+			return DistToTopWidgetSq >= 0.0f && BubblePath.Num() > 0;
+		}
+
+		TArray<FWidgetAndPointer> BubblePath;
+		float DistToTopWidgetSq;
+	};
+
+	FWidgetPathAndDist GetWidgetPathAndDist(const FGridTestingParams& Params, const bool bIgnoreEnabledStatus) const;
+
+	/** Return Value for GetHitIndexFromCellIndex */
+	struct FIndexAndDistance
+	{
+		FIndexAndDistance( int32 InIndex = INDEX_NONE, float InDistanceSq = 0 )
+		: WidgetIndex(InIndex)
+		, DistanceSqToWidget(InDistanceSq)
+		{}
+		int32 WidgetIndex;
+		float DistanceSqToWidget;
+	};
+	FIndexAndDistance GetHitIndexFromCellIndex(const FGridTestingParams& Params) const;
+
+	TArray<FWidgetAndPointer> GetBubblePathFromHitIndex(const int32 HitIndex, const bool bIgnoreEnabledStatus) const;
 	
 	friend class SWidgetReflector;
 
diff --git a/Engine/Source/Runtime/SlateCore/Public/Widgets/SWidget.h b/Engine/Source/Runtime/SlateCore/Public/Widgets/SWidget.h
--- a/Engine/Source/Runtime/SlateCore/Public/Widgets/SWidget.h
+++ b/Engine/Source/Runtime/SlateCore/Public/Widgets/SWidget.h
@@ -593,6 +593,12 @@ public:
 		return EnabledState.Get();
 	}
 
+	/** @return Is this widget interactable or not? Defaults to false */
+	virtual bool IsInteractable() const
+	{
+		return false;
+	}
+
 	/** @return The tool tip associated with this widget; Invalid reference if there is not one */
 	virtual TSharedPtr<IToolTip> GetToolTip();
 
diff --git a/Engine/Source/Runtime/UMG/Private/Components/WidgetComponent.cpp b/Engine/Source/Runtime/UMG/Private/Components/WidgetComponent.cpp
--- a/Engine/Source/Runtime/UMG/Private/Components/WidgetComponent.cpp
+++ b/Engine/Source/Runtime/UMG/Private/Components/WidgetComponent.cpp
@@ -786,7 +786,7 @@ TArray<FWidgetAndPointer> UWidgetComponent::GetHitWidgetPath( const FHitResult&
 	// Cache the location of the hit
 	LastLocalHitLocation = LocalHitLocation;
 
-	TArray<FWidgetAndPointer> ArrangedWidgets = HitTestGrid->GetBubblePath( LocalHitLocation, bIgnoreEnabledStatus );
+	TArray<FWidgetAndPointer> ArrangedWidgets = HitTestGrid->GetBubblePath( LocalHitLocation, 0.0f, bIgnoreEnabledStatus );
 
 	for( FWidgetAndPointer& Widget : ArrangedWidgets )
 	{
