Slate : Fixed fallout from radial hittest support change.

diff --git a/Engine/Source/Runtime/SlateCore/Private/Input/HittestGrid.cpp b/Engine/Source/Runtime/SlateCore/Private/Input/HittestGrid.cpp
--- a/Engine/Source/Runtime/SlateCore/Private/Input/HittestGrid.cpp
+++ b/Engine/Source/Runtime/SlateCore/Private/Input/HittestGrid.cpp
@@ -54,7 +54,7 @@ float DistanceSqToSlateRotatedRect(const FVector2D &Point, const FSlateRotatedRe
 
 bool IsOverlappingSlateRotatedRect(const FVector2D& Point, const float Radius, const FSlateRotatedRect& RotatedRect)
 {
-	return DistanceSqToSlateRotatedRect(Point, RotatedRect) <= (Radius * Radius);
+	return DistanceSqToSlateRotatedRect( Point, RotatedRect ) <= (Radius * Radius);
 }
 
 
@@ -656,24 +656,24 @@ FHittestGrid::FIndexAndDistance FHittestGrid::GetHitIndexFromCellIndex(const FGr
 			const FHittestGrid::FCachedWidget& TestCandidate = (*WidgetsCachedThisFrame)[WidgetIndex];
 
 			// Compute the render space clipping rect (FGeometry exposes a layout space clipping rect).
-			FSlateRotatedRect DesktopOrientedClipRect =
-				TransformRect(
+			const FSlateRotatedRect DesktopOrientedClipRect = TransformRect(
 				Concatenate(
-				Inverse(TestCandidate.CachedGeometry.GetAccumulatedLayoutTransform()),
-				TestCandidate.CachedGeometry.GetAccumulatedRenderTransform()
-				),
+					Inverse(TestCandidate.CachedGeometry.GetAccumulatedLayoutTransform()),
+					TestCandidate.CachedGeometry.GetAccumulatedRenderTransform()),
 				FSlateRotatedRect(TestCandidate.CachedGeometry.GetClippingRect().IntersectionWith(TestCandidate.ClippingRect))
-				);
+			);
 
-			//test if it is a valid index or not
+			// When performing a point hittest, accept all hittestable widgets.
+			// When performing a hittest with a radius, only grab interactive widgets.
 			const bool bIsValidWidget = !Params.bTestWidgetIsInteractive || (TestCandidate.WidgetPtr.IsValid() && TestCandidate.WidgetPtr.Pin()->IsInteractable());
 
-			//now do the overlap test
-			if (bIsValidWidget && IsOverlappingSlateRotatedRect(Params.CursorPositionInGrid, Params.Radius, DesktopOrientedClipRect))
+			const FVector2D DesktopSpaceCoordinate = Params.CursorPositionInGrid + GridOrigin;
+			if( bIsValidWidget && IsOverlappingSlateRotatedRect( DesktopSpaceCoordinate, Params.Radius, DesktopOrientedClipRect ) )
 			{
-				//fill in the distance param, if it isn't NULL
+				// We are within the search radius!
 				const bool bNeedsDistanceSearch = Params.Radius > 0.0f;
-				const float DistSq = (bNeedsDistanceSearch) ? DistanceSqToSlateRotatedRect(Params.CursorPositionInGrid, DesktopOrientedClipRect) : 0.0f;
+				// For non-0 radii also record the distance to cursor's center so that we can pick the closest hit from the results.
+				const float DistSq = (bNeedsDistanceSearch) ? DistanceSqToSlateRotatedRect( DesktopSpaceCoordinate, DesktopOrientedClipRect ) : 0.0f;
 				return FIndexAndDistance(WidgetIndex, DistSq);
 			}
 		}
