Fix for Input routing when 3D Widgets are used.

diff --git a/Engine/Source/Runtime/SlateCore/Private/Input/HittestGrid.cpp b/Engine/Source/Runtime/SlateCore/Private/Input/HittestGrid.cpp
--- a/Engine/Source/Runtime/SlateCore/Private/Input/HittestGrid.cpp
+++ b/Engine/Source/Runtime/SlateCore/Private/Input/HittestGrid.cpp
@@ -617,15 +617,17 @@ FHittestGrid::FWidgetPathAndDist FHittestGrid::GetWidgetPathAndDist(const FGridT
 	//Also grab the Hit Index, and the distance to the top hit
 	const FIndexAndDistance HitIndex = GetHitIndexFromCellIndex(Params);
 	
-	//if we get here, we want to do the testing for 3D Widgets
 	if (HitIndex.WidgetIndex != INDEX_NONE)
 	{
+		// if we have a custom path, we want to do the testing for 3D Widgets
 		const FCachedWidget& PhysicallyHitWidget = (*WidgetsCachedThisFrame)[HitIndex.WidgetIndex];
 		if (PhysicallyHitWidget.CustomPath.IsValid())
 		{
 			const FVector2D DesktopSpaceCoordinate = Params.CursorPositionInGrid + GridOrigin;
-			const TArray<FWidgetAndPointer> BubblePath = PhysicallyHitWidget.CustomPath.Pin()->GetBubblePathAndVirtualCursors(PhysicallyHitWidget.CachedGeometry, DesktopSpaceCoordinate, bIgnoreEnabledStatus);
-			return FWidgetPathAndDist(BubblePath,0.0);
+			TArray<FWidgetAndPointer> LogicalBubblePath = GetBubblePathFromHitIndex(HitIndex.WidgetIndex, bIgnoreEnabledStatus);
+			const TArray<FWidgetAndPointer> BubblePathExtension = PhysicallyHitWidget.CustomPath.Pin()->GetBubblePathAndVirtualCursors(PhysicallyHitWidget.CachedGeometry, DesktopSpaceCoordinate, bIgnoreEnabledStatus);
+			LogicalBubblePath.Append(BubblePathExtension);
+			return FWidgetPathAndDist(LogicalBubblePath, 0.0);
 		}
 		else
 		{
