UE-5495 - Slate Application now bubbles motion detection events.  This prevents the commonly reported bug if focusing UMG widget (Or any slate widget) except for the game viewport suddenly causes all motion input to stop being received by the game.

diff --git a/Engine/Source/Runtime/Slate/Private/Framework/Application/SlateApplication.cpp b/Engine/Source/Runtime/Slate/Private/Framework/Application/SlateApplication.cpp
--- a/Engine/Source/Runtime/Slate/Private/Framework/Application/SlateApplication.cpp
+++ b/Engine/Source/Runtime/Slate/Private/Framework/Application/SlateApplication.cpp
@@ -4660,15 +4660,14 @@ void FSlateApplication::ProcessMotionDetectedEvent( FMotionEvent& MotionEvent )
 	if (MotionEvent.GetUserIndex() < ARRAY_COUNT(UserFocusEntries) && UserFocusEntries[MotionEvent.GetUserIndex()].WidgetPath.IsValid())
 	{
 		/* Get the controller focus target for this user */
-		FWidgetPath PathToWidget = UserFocusEntries[MotionEvent.GetUserIndex()].WidgetPath.ToWidgetPath();
-		FArrangedWidget& ArrangedWidget = PathToWidget.Widgets.Last();
-
-		/* Switch worlds for widgets in the current path */
+		const FWidgetPath PathToWidget = UserFocusEntries[MotionEvent.GetUserIndex()].WidgetPath.ToWidgetPath();
+		
 		FScopedSwitchWorldHack SwitchWorld(PathToWidget);
 
-		/* Send the message to the widget */
-		FReply Reply = ArrangedWidget.Widget->OnMotionDetected(ArrangedWidget.Geometry, MotionEvent).SetHandler(ArrangedWidget.Widget);
-		ProcessReply(PathToWidget, Reply, nullptr, nullptr, MotionEvent.GetUserIndex());
+		FReply Reply = FEventRouter::Route<FReply>(this, FEventRouter::FBubblePolicy(PathToWidget), MotionEvent, [] (const FArrangedWidget& SomeWidget, const FMotionEvent& MotionEvent)
+		{
+			return SomeWidget.Widget->OnMotionDetected(SomeWidget.Geometry, MotionEvent);
+		});
 	}
 }
 
diff --git a/Engine/Source/Runtime/SlateCore/Public/Widgets/SWindow.h b/Engine/Source/Runtime/SlateCore/Public/Widgets/SWindow.h
--- a/Engine/Source/Runtime/SlateCore/Public/Widgets/SWindow.h
+++ b/Engine/Source/Runtime/SlateCore/Public/Widgets/SWindow.h
@@ -99,7 +99,7 @@ public:
 		SLATE_ARGUMENT( bool, bDragAnywhere )
 
 		/** The windows auto-centering mode. If set to anything other than None, then the
-		    ScreenPosition value will be ignored */
+			ScreenPosition value will be ignored */
 		SLATE_ARGUMENT( EAutoCenter::Type, AutoCenter )
 
 		/** Screen-space position where the window should be initially located. */
@@ -151,12 +151,12 @@ public:
 		SLATE_ARGUMENT( TOptional<float>, MaxHeight )
 
 		/** True if we should initially create a traditional title bar area.  If false, the user must embed the title
-		    area content into the window manually, taking into account platform-specific considerations!  Has no
+			area content into the window manually, taking into account platform-specific considerations!  Has no
 			effect for certain types of windows (popups, tool-tips, etc.) */
 		SLATE_ARGUMENT( bool, CreateTitleBar )
 
 		/** If the window appears off screen or is too large to safely fit this flag will force realistic 
-		    constraints on the window and bring it back into view. */
+			constraints on the window and bring it back into view. */
 		SLATE_ARGUMENT( bool, SaneWindowPlacement )
 
 		/** The padding around the edges of the window applied to it's content. */
@@ -557,15 +557,15 @@ public:
 	/** Set modal window related flags - called by Slate app code during FSlateApplication::AddModalWindow() */
 	void SetAsModalWindow()
 	{
-        bIsModalWindow = true;
+		bIsModalWindow = true;
 		bHasMaximizeButton = false;
 		bHasMinimizeButton = false;
 	}
 
 	bool IsModalWindow()
-    {
-        return bIsModalWindow;
-    }
+	{
+		return bIsModalWindow;
+	}
 
 	void SetTitleBar( const TSharedPtr<IWindowTitleBar> InTitleBar )
 	{
@@ -749,7 +749,7 @@ protected:
 	bool bIsTopmostWindow : 1;
 
 	/** True if we expect the size of this window to change often, such as if its animated, or if it recycled for tool-tips,
-	    and we'd like to avoid costly GPU buffer resizes when that happens.  Enabling this may incur memory overhead or
+		and we'd like to avoid costly GPU buffer resizes when that happens.  Enabling this may incur memory overhead or
 		other platform-specific side effects */
 	bool bSizeWillChangeOften : 1;
 
@@ -779,9 +779,9 @@ protected:
 
 	/** True if this window displays thick edge that can be used to resize the window */
 	bool bHasSizingFrame : 1;
-    
-    /** True if the window is modal */
-    bool bIsModalWindow : 1;
+	
+	/** True if the window is modal */
+	bool bIsModalWindow : 1;
 
 	/** Initial desired position of the window's content in screen space */
 	FVector2D InitialDesiredScreenPosition;
@@ -847,7 +847,7 @@ protected:
 	SVerticalBox::FSlot* ContentSlot;
 
 	/** Widget to transfer keyboard focus to when this window becomes active, if any.  This is used to
-	    restore focus to a widget after a popup has been dismissed. */
+		restore focus to a widget after a popup has been dismissed. */
 	TWeakPtr< SWidget > WidgetToFocusOnActivate;
 
 	/** Style used to draw this window */
@@ -1017,7 +1017,7 @@ private:
  */
 struct FScopedSwitchWorldHack
 {
-	FScopedSwitchWorldHack( FWidgetPath& WidgetPath )
+	FScopedSwitchWorldHack( const FWidgetPath& WidgetPath )
 		: Window( WidgetPath.TopLevelWindow )
 		, WorldId( -1 )
 	{
