Slate : SWindow now only considers its own widgets when trying to respond to GetCurrentWindowZone().

diff --git a/Engine/Source/Runtime/Slate/Private/Framework/Application/SlateApplication.cpp b/Engine/Source/Runtime/Slate/Private/Framework/Application/SlateApplication.cpp
--- a/Engine/Source/Runtime/Slate/Private/Framework/Application/SlateApplication.cpp
+++ b/Engine/Source/Runtime/Slate/Private/Framework/Application/SlateApplication.cpp
@@ -16,9 +16,15 @@ class FEventRouter
 {
 
 // @todo slate : making too many event copies when translating events( i.e. Translate<EventType>::PointerEvent ).
+
+
+
 // @todo slate : Widget Reflector should log: (1) Every process reply (2) Every time the event is handled and by who.
+
 // @todo slate : Remove remaining [&]-style mass captures.
+
 // @todo slate : Eliminate all ad-hoc uses of SetEventPath()
+
 // @todo slate : Remove CALL_WIDGET_FUNCTION
 
 public:
@@ -808,12 +814,13 @@ FWidgetPath FSlateApplication::LocateWindowUnderMouse( FVector2D ScreenspaceMous
 
 		// Only accept input if the current window accepts input and the current window is not under a modal window or an interactive tooltip
 		
-		const bool AcceptsInput = Window->AcceptsInput() || IsWindowHousingInteractiveTooltip(Window);
-
-		if ( Window->IsVisible() && AcceptsInput && Window->IsScreenspaceMouseWithin(ScreenspaceMouseCoordinate) && !bPrevWindowWasModal )
+		if (!bPrevWindowWasModal)
 		{
-			const TArray<FWidgetAndPointer> WidgetsAndCursors = Window->GetHittestGrid()->GetBubblePath(ScreenspaceMouseCoordinate, GetCursorRadius(), bIgnoreEnabledStatus);
-			return FWidgetPath( WidgetsAndCursors );
+			FWidgetPath PathToLocatedWidget = LocateWidgetInWindow(ScreenspaceMouseCoordinate, Window, bIgnoreEnabledStatus);
+			if (PathToLocatedWidget.IsValid())
+			{
+				return PathToLocatedWidget;
+			}
 		}
 	}
 
@@ -1393,6 +1400,20 @@ void FSlateApplication::ThrottleApplicationBasedOnMouseMovement()
 	}
 }
 
+FWidgetPath FSlateApplication::LocateWidgetInWindow(FVector2D ScreenspaceMouseCoordinate, const TSharedRef<SWindow>& Window, bool bIgnoreEnabledStatus) const
+{
+	const bool bAcceptsInput = Window->IsVisible() && (Window->AcceptsInput() || IsWindowHousingInteractiveTooltip(Window));
+	if (bAcceptsInput && Window->IsScreenspaceMouseWithin(ScreenspaceMouseCoordinate))
+	{
+		const TArray<FWidgetAndPointer> WidgetsAndCursors = Window->GetHittestGrid()->GetBubblePath(ScreenspaceMouseCoordinate, GetCursorRadius(), bIgnoreEnabledStatus);
+		return FWidgetPath(WidgetsAndCursors);
+	}
+	else
+	{
+		return FWidgetPath();
+	}
+}
+
 
 TSharedRef<SWindow> FSlateApplication::AddWindow( TSharedRef<SWindow> InSlateWindow, const bool bShowImmediately )
 {
diff --git a/Engine/Source/Runtime/Slate/Public/Framework/Application/SlateApplication.h b/Engine/Source/Runtime/Slate/Public/Framework/Application/SlateApplication.h
--- a/Engine/Source/Runtime/Slate/Public/Framework/Application/SlateApplication.h
+++ b/Engine/Source/Runtime/Slate/Public/Framework/Application/SlateApplication.h
@@ -705,6 +705,8 @@ protected:
 	/** Engages or disengages application throttling based on user behavior */
 	void ThrottleApplicationBasedOnMouseMovement();
 
+	virtual FWidgetPath LocateWidgetInWindow(FVector2D ScreenspaceMouseCoordinate, const TSharedRef<SWindow>& Window, bool bIgnoreEnabledStatus) const override;
+
 public:
 
 	/**
diff --git a/Engine/Source/Runtime/SlateCore/Private/Application/SlateApplicationBase.cpp b/Engine/Source/Runtime/SlateCore/Private/Application/SlateApplicationBase.cpp
--- a/Engine/Source/Runtime/SlateCore/Private/Application/SlateApplicationBase.cpp
+++ b/Engine/Source/Runtime/SlateCore/Private/Application/SlateApplicationBase.cpp
@@ -1,6 +1,7 @@
 // Copyright 1998-2015 Epic Games, Inc. All Rights Reserved.
 
 #include "SlateCorePrivatePCH.h"
+#include "HittestGrid.h"
 
 
 /* Static initialization
@@ -10,3 +11,21 @@ TSharedPtr<FSlateApplicationBase> FSlateApplicationBase::CurrentBaseApplication
 TSharedPtr<GenericApplication> FSlateApplicationBase::PlatformApplication = nullptr;
 // TODO: Identifier the cursor index in a smarter way.
 const uint32 FSlateApplicationBase::CursorPointerIndex = EKeys::NUM_TOUCH_KEYS - 1;
+
+FWidgetPath FHitTesting::LocateWidgetInWindow(FVector2D ScreenspaceMouseCoordinate, const TSharedRef<SWindow>& Window, bool bIgnoreEnabledStatus) const
+{
+	return SlateApp->LocateWidgetInWindow(ScreenspaceMouseCoordinate, Window, bIgnoreEnabledStatus);
+}
+
+
+FSlateApplicationBase::FSlateApplicationBase()
+: Renderer()
+, HitTesting(this)
+{
+
+}
+
+const FHitTesting& FSlateApplicationBase::GetHitTesting() const
+{
+	return HitTesting;
+}
\ No newline at end of file
diff --git a/Engine/Source/Runtime/SlateCore/Private/Widgets/SWindow.cpp b/Engine/Source/Runtime/SlateCore/Private/Widgets/SWindow.cpp
--- a/Engine/Source/Runtime/SlateCore/Private/Widgets/SWindow.cpp
+++ b/Engine/Source/Runtime/SlateCore/Private/Widgets/SWindow.cpp
@@ -1693,12 +1693,7 @@ EWindowZone::Type SWindow::GetCurrentWindowZone(FVector2D LocalMousePosition)
 		if (InZone == EWindowZone::ClientArea)
 		{
 			// Hittest to see if the widget under the mouse should be treated as a title bar (i.e. should move the window)
-			TArray< TSharedRef<SWindow> > ThisWindow;
-			{
-				ThisWindow.Add( SharedThis(this) );
-			}
-
-			FWidgetPath HitTestResults = FSlateApplicationBase::Get().LocateWindowUnderMouse(FSlateApplicationBase::Get().GetCursorPos(), ThisWindow);
+			FWidgetPath HitTestResults = FSlateApplicationBase::Get().GetHitTesting().LocateWidgetInWindow(FSlateApplicationBase::Get().GetCursorPos(), SharedThis(this), false);
 			if( HitTestResults.Widgets.Num() > 0 )
 			{
 				const EWindowZone::Type ZoneOverride = HitTestResults.Widgets.Last().Widget->GetWindowZoneOverride();
diff --git a/Engine/Source/Runtime/SlateCore/Public/Application/SlateApplicationBase.h b/Engine/Source/Runtime/SlateCore/Public/Application/SlateApplicationBase.h
--- a/Engine/Source/Runtime/SlateCore/Public/Application/SlateApplicationBase.h
+++ b/Engine/Source/Runtime/SlateCore/Public/Application/SlateApplicationBase.h
@@ -3,8 +3,7 @@
 #pragma once
 
 
-class FHittestGrid;
-
+class FHitTesting;
 
 /**
  * Interface for window title bars.
@@ -16,6 +15,29 @@ public:
 	virtual void Flash( ) = 0;
 };
 
+class FSlateApplicationBase;
+
+
+/**
+ * Private interface to control which classes are allowed to perform hit-testing.
+ */
+class FHitTesting
+{
+public:
+	FHitTesting(FSlateApplicationBase* InSlateApplication)
+		: SlateApp(InSlateApplication)
+	{
+	}
+
+private:
+	// SWindow must be able to test which part of the window is being moused-over
+	friend class SWindow;
+	
+private:	
+	FSlateApplicationBase* SlateApp;
+	// @see FSlateApplicationBase::LocateWidgetInWindow
+	FWidgetPath LocateWidgetInWindow(FVector2D ScreenspaceMouseCoordinate, const TSharedRef<SWindow>& Window, bool bIgnoreEnabledStatus) const;
+};
 
 /**
  * Base class for Slate applications.
@@ -28,6 +50,8 @@ class SLATECORE_API FSlateApplicationBase
 	friend class SWidget;
 public:
 
+	FSlateApplicationBase();
+
 	/**
 	 * Gets the renderer being used to draw this application.
 	 *
@@ -208,6 +232,10 @@ public:
 	 */
 	virtual bool IsExternalUIOpened( ) = 0;
 
+	/** @return a hittesting object that can perform hittests agains widgets. Only certain classes can make use of FHitTesting */
+	friend class FHitTesting;
+	const FHitTesting& GetHitTesting() const;
+
 	/** 
 	 * Given the screen-space coordinate of the mouse cursor, searches for a string of widgets that are under the mouse.
 	 *
@@ -347,11 +375,17 @@ protected:
 	 */
 	virtual bool ShowUserFocus(const TSharedPtr<const SWidget> Widget) const = 0;
 
+	/** Given a window, locate a widget under the cursor in it; returns an invalid path if cursor is not over this window. */
+	virtual FWidgetPath LocateWidgetInWindow(FVector2D ScreenspaceMouseCoordinate, const TSharedRef<SWindow>& Window, bool bIgnoreEnabledStatus) const = 0;
+
 protected:
 
 	// Holds the Slate renderer used to render this application.
 	TSharedPtr<FSlateRenderer> Renderer;
 
+	// Private interface for select entities that are allowed to perform hittesting
+	FHitTesting HitTesting;
+
 protected:
 
 	// Holds a pointer to the current application.
@@ -360,3 +394,4 @@ protected:
 	// Holds a pointer to the platform application.
 	static TSharedPtr<class GenericApplication> PlatformApplication;
 };
+
