Fix bad call to ToolTip delegate in constructor

diff --git a/Engine/Source/Runtime/SlateCore/Private/Widgets/SWidget.cpp b/Engine/Source/Runtime/SlateCore/Private/Widgets/SWidget.cpp
--- a/Engine/Source/Runtime/SlateCore/Private/Widgets/SWidget.cpp
+++ b/Engine/Source/Runtime/SlateCore/Private/Widgets/SWidget.cpp
@@ -42,7 +42,7 @@ void SWidget::Construct(
 		// If someone specified a fancy widget tooltip, use it.
 		ToolTip = InToolTip;
 	}
-	else if ( InToolTipText.IsBound() || !(InToolTipText.Get().IsEmpty()) )
+	else if ( InToolTipText.IsSet() )
 	{
 		// If someone specified a text binding, make a tooltip out of it
 		ToolTip = FSlateApplicationBase::Get().MakeToolTip(InToolTipText);
