---- Merging with SlateDev branch ----

diff --git a/Engine/Source/Runtime/SlateCore/Private/Animation/CurveSequence.cpp b/Engine/Source/Runtime/SlateCore/Private/Animation/CurveSequence.cpp
--- a/Engine/Source/Runtime/SlateCore/Private/Animation/CurveSequence.cpp
+++ b/Engine/Source/Runtime/SlateCore/Private/Animation/CurveSequence.cpp
@@ -3,27 +3,38 @@
 #include "SlateCorePrivatePCH.h"
 
 
-/* FCurveSequence structors
+/* FCurveSequence constructors
  *****************************************************************************/
 
 FCurveSequence::FCurveSequence( )
 	: StartTime(0)
 	, TotalDuration(0)
 	, bInReverse(true)
+	, bIsLooping(false)
+	, bIsPaused(false)
 { }
 
-
-/* FCurveSequence interface
- *****************************************************************************/
-
 FCurveSequence::FCurveSequence( const float InStartTimeSeconds, const float InDurationSeconds, const ECurveEaseFunction::Type InEaseFunction )
 	: StartTime(0)
 	, TotalDuration(0)
 	, bInReverse(true)
+	, bIsLooping( false )
+	, bIsPaused( false )
 {
 	const FCurveHandle IgnoredCurveHandle = AddCurve( InStartTimeSeconds, InDurationSeconds, InEaseFunction );
 }
 
+/* FCurveSequence interface
+ *****************************************************************************/
+
+FCurveSequence::~FCurveSequence()
+{
+	// If the curve sequence is destroyed before the owning widget, unregister the active tick
+	if ( OwnerWidget.IsValid() && ActiveTickHandle.IsValid() )
+	{
+		OwnerWidget.Pin()->UnRegisterActiveTick(ActiveTickHandle.Pin().ToSharedRef());
+	}
+}
 
 FCurveHandle FCurveSequence::AddCurve( const float InStartTimeSeconds, const float InDurationSeconds, const ECurveEaseFunction::Type InEaseFunction )
 {
@@ -54,18 +65,30 @@ void FCurveSequence::Play( const float StartAtTime )
 	SetStartTime(FSlateApplicationBase::Get().GetCurrentTime() - StartAtTime);
 }
 
+void FCurveSequence::Play( const TSharedRef<SWidget>& InOwnerWidget, bool bPlayLooped, const float StartAtTime )
+{
+	RegisterActiveTickIfNeeded(InOwnerWidget);
+	bIsLooping = bPlayLooped;
+	bIsPaused = false;
+	
+	// Playing forward
+	bInReverse = false;
+
+	// We start playing NOW.
+	SetStartTime(FSlateApplicationBase::Get().GetCurrentTime() - StartAtTime);
+}
 
 void FCurveSequence::Reverse( )
-{
+{	
+	// We've played this far into the animation.
+	const float FractionCompleted = FMath::Clamp(GetSequenceTime() / TotalDuration, 0.0f, 1.0f);
+
 	// We're going the other way now.
-	bInReverse = !bInReverse;	
+	bInReverse = !bInReverse;
 
 	// CurTime is now; we cannot change that, so everything happens relative to CurTime.
 	const double CurTime = FSlateApplicationBase::Get().GetCurrentTime();
-	
-	// We've played this far into the animation.
-	const float FractionCompleted = FMath::Clamp((CurTime - StartTime) / TotalDuration, 0.0, 1.0);
-	
+
 	// Assume CurTime is constant (now).
 	// Figure out when the animation would need to have started in order to keep
 	// its place if playing in reverse.
@@ -82,10 +105,48 @@ void FCurveSequence::PlayReverse( const float StartAtTime )
 	SetStartTime(FSlateApplicationBase::Get().GetCurrentTime() - StartAtTime);
 }
 
+void FCurveSequence::PlayReverse( const TSharedRef<SWidget>& InOwnerWidget, bool bPlayLooped, const float StartAtTime )
+{
+	RegisterActiveTickIfNeeded(InOwnerWidget);
+	bIsLooping = bPlayLooped;
+	bIsPaused = false;
+
+	bInReverse = true;
+
+	// We start reversing NOW.
+	SetStartTime(FSlateApplicationBase::Get().GetCurrentTime() - StartAtTime);
+}
+
+void FCurveSequence::Pause()
+{
+	if ( IsPlaying() )
+	{
+		bIsPaused = true;
+		PauseTime = FSlateApplicationBase::Get().GetCurrentTime();
+	}
+}
+
+void FCurveSequence::Resume()
+{
+	if ( bIsPaused )
+	{
+		// Make sure the widget that owns the sequence is still valid
+		auto PinnedOwner = OwnerWidget.Pin();
+		if ( PinnedOwner.IsValid() )
+		{
+			bIsPaused = false;
+			RegisterActiveTickIfNeeded( PinnedOwner.ToSharedRef() );
+
+			// Update the start time to be the same relative to the current time as it was when paused
+			const double NewStartTime = FSlateApplicationBase::Get().GetCurrentTime() - ( PauseTime - StartTime );
+			SetStartTime( NewStartTime );
+		}
+	}
+}
 
 bool FCurveSequence::IsPlaying( ) const
 {
-	return (FSlateApplicationBase::Get().GetCurrentTime() - StartTime) <= TotalDuration;
+	return !bIsPaused && ( bIsLooping || ( FSlateApplicationBase::Get().GetCurrentTime() - StartTime ) <= TotalDuration );
 }
 
 
@@ -94,22 +155,34 @@ void FCurveSequence::SetStartTime( double InStartTime )
 	StartTime = InStartTime;
 }
 
-
 float FCurveSequence::GetSequenceTime( ) const
 {
-	const double CurrentTime = FSlateApplicationBase::Get().GetCurrentTime();
+	const double CurrentTime = bIsPaused ? PauseTime : FSlateApplicationBase::Get().GetCurrentTime();
+	float SequenceTime = IsInReverse() ? TotalDuration - ( CurrentTime - StartTime ) : CurrentTime - StartTime;
 
-	return IsInReverse()
-		? TotalDuration - (CurrentTime - StartTime)
-		: CurrentTime - StartTime;
+	return bIsLooping ? FMath::Fmod( SequenceTime, TotalDuration ) : SequenceTime;
 }
 
-
 float FCurveSequence::GetSequenceTimeLooping( ) const
 {
-	return FMath::Fmod(GetSequenceTime(), TotalDuration);
+	return DEPRECATED_GetSequenceTimeLooping();
 }
+float FCurveSequence::DEPRECATED_GetSequenceTimeLooping() const
+{
+	if (!bIsLooping)
+	{
+		auto MutableThis = const_cast<FCurveSequence*>( this );
+
+		// Fake that we're looping to get the lerp
+		MutableThis->bIsLooping = true;
+		float SequenceTime = GetSequenceTime();
+		MutableThis->bIsLooping = false;
 
+		return SequenceTime;
+	}
+
+	return GetSequenceTime();
+}
 
 bool FCurveSequence::IsInReverse( ) const
 {
@@ -126,6 +199,8 @@ bool FCurveSequence::IsForward( ) const
 void FCurveSequence::JumpToStart( )
 {
 	bInReverse = true;
+	bIsLooping = false;
+	bIsPaused = false;
 	SetStartTime(FSlateApplicationBase::Get().GetCurrentTime() - TotalDuration);
 }
 
@@ -133,21 +208,27 @@ void FCurveSequence::JumpToStart( )
 void FCurveSequence::JumpToEnd( )
 {
 	bInReverse = false;
+	bIsLooping = false;
+	bIsPaused = false;
 	SetStartTime(FSlateApplicationBase::Get().GetCurrentTime() - TotalDuration);
 }
 
 
 bool FCurveSequence::IsAtStart( ) const
 {
-	return (IsInReverse() == true && IsPlaying() == false);
+	return ( IsInReverse() == true && IsPlaying() == false && !bIsLooping );
 }
 
 
 bool FCurveSequence::IsAtEnd( ) const
 {
-	return (IsForward() == true && IsPlaying() == false);
+	return ( IsForward() == true && IsPlaying() == false && !bIsLooping );
 }
 
+bool FCurveSequence::IsLooping() const
+{
+	return bIsLooping;
+}
 
 float FCurveSequence::GetLerp( ) const
 {
@@ -155,21 +236,41 @@ float FCurveSequence::GetLerp( ) const
 	// interpolation alpha values.
 	checkSlow(Curves.Num() == 1);
 
-	return FCurveHandle(this, 0).GetLerp();
+	return FCurveHandle( this, 0 ).GetLerp();
 }
 
-
-float FCurveSequence::GetLerpLooping( ) const
+float FCurveSequence::GetLerpLooping() const
+{
+	return DEPRECATED_GetLerpLooping();
+}
+float FCurveSequence::DEPRECATED_GetLerpLooping() const
 {
 	// Only supported for sequences with a single curve.  If you have multiple curves, use your FCurveHandle to compute
 	// interpolation alpha values.
 	checkSlow(Curves.Num() == 1);
 
-	return FCurveHandle(this, 0).GetLerpLooping();
+	return bIsLooping ? GetLerp() : FCurveHandle(this, 0).DEPRECATED_GetLerpLooping();
 }
 
-
 const FCurveSequence::FSlateCurve& FCurveSequence::GetCurve( int32 CurveIndex ) const
 {
 	return Curves[CurveIndex];
 }
+
+void FCurveSequence::RegisterActiveTickIfNeeded(TSharedRef<SWidget> InOwnerWidget)
+{
+	// Register the active tick
+	if ( !ActiveTickHandle.IsValid() )
+	{
+		ActiveTickHandle = InOwnerWidget->RegisterActiveTick(0.f, FWidgetActiveTickDelegate::CreateRaw(this, &FCurveSequence::EnsureSlateTickDuringAnimation));
+
+		// Save the reference in case we need to take care of unregistering the tick
+		OwnerWidget = InOwnerWidget;
+	}
+}
+
+EActiveTickReturnType FCurveSequence::EnsureSlateTickDuringAnimation( double InCurrentTime, float InDeltaTime )
+{
+	// Keep ticking until play stops
+	return IsPlaying() ? EActiveTickReturnType::KeepTicking : EActiveTickReturnType::StopTicking;
+}
\ No newline at end of file
diff --git a/Engine/Source/Runtime/SlateCore/Private/Application/ActiveTickHandle.cpp b/Engine/Source/Runtime/SlateCore/Private/Application/ActiveTickHandle.cpp
new file mode 100644
--- /dev/null
+++ b/Engine/Source/Runtime/SlateCore/Private/Application/ActiveTickHandle.cpp
@@ -0,0 +1,58 @@
+// Copyright 1998-2014 Epic Games, Inc. All Rights Reserved.
+
+#include "SlateCorePrivatePCH.h"
+#include "ActiveTickHandle.h"
+
+
+FActiveTickHandle::FActiveTickHandle(float InTickFrequency, FWidgetActiveTickDelegate InTickFunction, double InNextTickTime) 
+	: TickFrequency(InTickFrequency)
+	, TickFunction(InTickFunction)
+	, NextTickTime(InNextTickTime)
+	, bExecutionPending(false)
+{
+}
+
+bool FActiveTickHandle::IsReady(double CurrentTime) const
+{
+	return CurrentTime >= NextTickTime;
+}
+
+bool FActiveTickHandle::IsPendingExecution() const
+{
+	return bExecutionPending;
+}
+
+bool FActiveTickHandle::UpdateExecutionPendingState( double CurrentTime )
+{
+	bExecutionPending = IsReady( CurrentTime );
+	return bExecutionPending;
+}
+
+EActiveTickReturnType FActiveTickHandle::ExecuteIfPending( double CurrentTime, float DeltaTime )
+{
+	// check if the handle timer is ready to fire.
+	if ( bExecutionPending )
+	{
+		bExecutionPending = false;
+
+		// before we do anything, check validity of the delegate.
+		if ( !TickFunction.IsBound() )
+		{
+			// if the handle is no longer valid, we must assume it is invalid (say, owning widget was destroyed) and needs to be removed.
+			return EActiveTickReturnType::StopTicking;
+		}
+
+		// update time in a while loop so we skip any times that we may have missed due to slow frames.
+		if ( TickFrequency > 0.0 )
+		{
+			while ( CurrentTime >= NextTickTime )
+			{
+				NextTickTime += TickFrequency;
+			}
+		}
+		// timer is updated, now call the delegate
+		return TickFunction.Execute( CurrentTime, DeltaTime );
+	}
+	// if the handle is not ready, then we need to keep ticking it.
+	return EActiveTickReturnType::KeepTicking;
+}
\ No newline at end of file
diff --git a/Engine/Source/Runtime/SlateCore/Private/Application/ActiveTickHandle.h b/Engine/Source/Runtime/SlateCore/Private/Application/ActiveTickHandle.h
new file mode 100644
--- /dev/null
+++ b/Engine/Source/Runtime/SlateCore/Private/Application/ActiveTickHandle.h
@@ -0,0 +1,48 @@
+// Copyright 1998-2014 Epic Games, Inc. All Rights Reserved.
+
+#pragma once
+
+#include "SlateCore.h"
+
+/** Stores info about an active tick delegate for a widget. */
+class FActiveTickHandle
+{
+public:
+	/** Public ctor that initializes a new tick handle. Not intended to be called by user code. */
+	FActiveTickHandle(float InTickFrequency, FWidgetActiveTickDelegate InTickFunction, double InNextTickTime);
+
+	/** @return True if the active tick is pending execution */
+	bool IsPendingExecution() const;
+
+	/**
+	* Updates the pending state of the active tick based on the current time
+	* @return True if the active tick is pending execution
+	*/
+	bool UpdateExecutionPendingState( double CurrentTime );
+
+	/**
+	 * Execute the bound delegate if the active tick is bPendingExecution.
+	 * @return ETickWidgetReturnType::StopTicking if the handle should be auto-unregistered
+	 *         (either return value indicated this or the handle has become invalid).
+	 */
+	EActiveTickReturnType ExecuteIfPending( double CurrentTime, float DeltaTime );
+
+private:
+	/** @return true if the tick handle is ready to have its delegate executed. */
+	bool IsReady( double CurrentTime ) const;
+
+private:
+	/** The frequency with which this active tick should execute */
+	float TickFrequency;
+
+	/** The delegate to the active tick function */
+	FWidgetActiveTickDelegate TickFunction;
+	
+	/** The next time at which to execute the active tick */
+	double NextTickTime;
+
+	/** True if execution of the active tick is pending */
+	bool bExecutionPending;
+};
+
+
diff --git a/Engine/Source/Runtime/SlateCore/Private/Application/SlateApplicationBase.cpp b/Engine/Source/Runtime/SlateCore/Private/Application/SlateApplicationBase.cpp
--- a/Engine/Source/Runtime/SlateCore/Private/Application/SlateApplicationBase.cpp
+++ b/Engine/Source/Runtime/SlateCore/Private/Application/SlateApplicationBase.cpp
@@ -2,6 +2,7 @@
 
 #include "SlateCorePrivatePCH.h"
 #include "HittestGrid.h"
+#include "ActiveTickHandle.h"
 
 
 /* Static initialization
@@ -28,4 +29,49 @@ FSlateApplicationBase::FSlateApplicationBase()
 const FHitTesting& FSlateApplicationBase::GetHitTesting() const
 {
 	return HitTesting;
-}
\ No newline at end of file
+}
+
+void FSlateApplicationBase::RegisterActiveTick( const TSharedRef<FActiveTickHandle>& ActiveTickHandle )
+{
+	ActiveTickHandles.Add(ActiveTickHandle);
+}
+
+void FSlateApplicationBase::UnRegisterActiveTick( const TSharedRef<FActiveTickHandle>& ActiveTickHandle )
+{
+	ActiveTickHandles.Remove(ActiveTickHandle);
+}
+
+bool FSlateApplicationBase::AnyActiveTicksArePending()
+{
+	// first remove any tick handles that may have become invalid.
+	// If we didn't remove invalid handles here, they would never get removed because
+	// we don't force widgets to UnRegister before they are destroyed.
+	ActiveTickHandles.RemoveAll([](const TWeakPtr<FActiveTickHandle>& ActiveTickHandle)
+	{
+		// only check the weak pointer to the handle. Just want to make sure to clear out any widgets that have since been deleted.
+		return !ActiveTickHandle.IsValid();
+	});
+
+	// The rest are valid. Update their pending status and see if any are ready.
+	const double CurrentTime = GetCurrentTime();
+	bool bAnyTickReady = false;
+	for ( auto& ActiveTickInfo : ActiveTickHandles )
+	{
+		auto ActiveTickInfoPinned = ActiveTickInfo.Pin();
+		check( ActiveTickInfoPinned.IsValid() );
+
+		// If an active tick is still pending execution from last frame, it is collapsed 
+		// or otherwise blocked from ticking. Disregard until it executes.
+		if ( ActiveTickInfoPinned->IsPendingExecution() )
+		{
+			continue;
+		}
+
+		if ( ActiveTickInfoPinned->UpdateExecutionPendingState( CurrentTime ) )
+		{
+			bAnyTickReady = true;
+		}
+	}
+
+	return bAnyTickReady;
+}
diff --git a/Engine/Source/Runtime/SlateCore/Private/Widgets/SWidget.cpp b/Engine/Source/Runtime/SlateCore/Private/Widgets/SWidget.cpp
--- a/Engine/Source/Runtime/SlateCore/Private/Widgets/SWidget.cpp
+++ b/Engine/Source/Runtime/SlateCore/Private/Widgets/SWidget.cpp
@@ -3,6 +3,7 @@
 #include "SlateCorePrivatePCH.h"
 #include "Widgets/SWidget.h"
 #include "Input/Events.h"
+#include "ActiveTickHandle.h"
 
 DECLARE_CYCLE_STAT(TEXT("OnPaint"), STAT_SlateOnPaint, STATGROUP_Slate);
 DECLARE_CYCLE_STAT(TEXT("ArrangeChildren"), STAT_SlateArrangeChildren, STATGROUP_Slate);
@@ -341,13 +342,14 @@ void SWidget::Tick( const FGeometry& AllottedGeometry, const double InCurrentTim
 {
 }
 
-TAutoConsoleVariable<int32> TickInvisibleWidgets(TEXT("Slate.TickInvisibleWidgets"), 1, TEXT("Controls whether invisible widgets are ticked."));
+TAutoConsoleVariable<int32> TickInvisibleWidgets(TEXT("Slate.TickInvisibleWidgets"), 0, TEXT("Controls whether invisible widgets are ticked."));
 
 void SWidget::TickWidgetsRecursively( const FGeometry& AllottedGeometry, const double InCurrentTime, const float InDeltaTime )
 {
 	INC_DWORD_STAT(STAT_SlateNumTickedWidgets);
 
-	// Tick this widget
+	// Execute any pending active ticks for this widget, followed by the passive tick
+	ExecuteActiveTicks( InCurrentTime, InDeltaTime );
 	Tick( AllottedGeometry, InCurrentTime, InDeltaTime );
 
 	// Gather all children, whether they're visible or not.  We need to allow invisible widgets to
@@ -632,6 +634,7 @@ int32 SWidget::Paint(const FPaintArgs& Args, const FGeometry& AllottedGeometry,
 		TickGeometry.AppendTransform( FSlateLayoutTransform(Args.GetWindowToDesktopTransform()) );
 
 		SWidget* MutableThis = const_cast<SWidget*>(this);
+		MutableThis->ExecuteActiveTicks( Args.GetCurrentTime(), Args.GetDeltaTime() );
 		MutableThis->Tick( TickGeometry, Args.GetCurrentTime(), Args.GetDeltaTime() );
 	}
 
@@ -669,3 +672,43 @@ void SWidget::ArrangeChildren(const FGeometry& AllottedGeometry, FArrangedChildr
 	OnArrangeChildren(AllottedGeometry, ArrangedChildren);
 }
 
+TSharedRef<FActiveTickHandle> SWidget::RegisterActiveTick(float TickPeriod, FWidgetActiveTickDelegate TickFunction)
+{
+	TSharedRef<FActiveTickHandle> ActiveTickHandle = MakeShareable(new FActiveTickHandle(TickPeriod, TickFunction, FSlateApplicationBase::Get().GetCurrentTime() + TickPeriod));
+	FSlateApplicationBase::Get().RegisterActiveTick(ActiveTickHandle);
+	ActiveTicks.Add(ActiveTickHandle);
+	return ActiveTickHandle;
+}
+
+void SWidget::UnRegisterActiveTick(const TSharedRef<FActiveTickHandle>& ActiveTickHandle)
+{
+	FSlateApplicationBase::Get().UnRegisterActiveTick(ActiveTickHandle);
+	ActiveTicks.Remove(ActiveTickHandle);
+}
+
+void SWidget::ExecuteActiveTicks(double CurrentTime, float DeltaTime)
+{
+	// loop over the registered tick handles and execute them, removing them if necessary.
+	for (int32 i = 0; i < ActiveTicks.Num();)
+	{
+		EActiveTickReturnType Result = ActiveTicks[i]->ExecuteIfPending( CurrentTime, DeltaTime );
+		if (Result == EActiveTickReturnType::KeepTicking)
+		{
+			++i;
+		}
+		else
+		{
+			FSlateApplicationBase::Get().UnRegisterActiveTick(ActiveTicks[i]);
+			ActiveTicks.RemoveAt(i);
+		}
+	}
+}
+
+SWidget::~SWidget()
+{
+	// Unregister all ActiveTicks so they aren't left stranded in the Application's list.
+	for (const auto& ActiveTickHandle : ActiveTicks)
+	{
+		FSlateApplicationBase::Get().UnRegisterActiveTick(ActiveTickHandle);
+	}
+}
diff --git a/Engine/Source/Runtime/SlateCore/Private/Widgets/SWindow.cpp b/Engine/Source/Runtime/SlateCore/Private/Widgets/SWindow.cpp
--- a/Engine/Source/Runtime/SlateCore/Private/Widgets/SWindow.cpp
+++ b/Engine/Source/Runtime/SlateCore/Private/Widgets/SWindow.cpp
@@ -531,11 +531,6 @@ void SWindow::Tick( const FGeometry& AllottedGeometry, const double InCurrentTim
 {
 	if( Morpher.bIsActive )
 	{
-		if(Morpher.bIsPendingPlay)
-		{
-			Morpher.Sequence.Play();
-			Morpher.bIsPendingPlay = false;
-		}
 		if ( Morpher.Sequence.IsPlaying() )
 		{
 			const float InterpAlpha = Morpher.Sequence.GetLerp();
@@ -549,7 +544,7 @@ void SWindow::Tick( const FGeometry& AllottedGeometry, const double InCurrentTim
 					this->ReshapeWindow( WindowRect );
 				}
 			}
-			else
+			else // if animating position
 			{
 				const FVector2D StartPosition( Morpher.StartingMorphShape.Left, Morpher.StartingMorphShape.Top );
 				const FVector2D TargetPosition( Morpher.TargetMorphShape.Left, Morpher.TargetMorphShape.Top );
@@ -565,6 +560,7 @@ void SWindow::Tick( const FGeometry& AllottedGeometry, const double InCurrentTim
 		}
 		else
 		{
+			// The animation is complete, so just make sure the target size/position and opacity are reached
 			if( Morpher.bIsAnimatingWindowSize )
 			{
 				if( Morpher.TargetMorphShape != GetRectInScreen() )
@@ -573,7 +569,7 @@ void SWindow::Tick( const FGeometry& AllottedGeometry, const double InCurrentTim
 					this->ReshapeWindow( Morpher.TargetMorphShape );
 				}
 			}
-			else
+			else // if animating position
 			{
 				const FVector2D TargetPosition( Morpher.TargetMorphShape.Left, Morpher.TargetMorphShape.Top );
 				if( TargetPosition != this->GetPositionInScreen() )
@@ -846,9 +842,13 @@ void SWindow::StartMorph()
 {
 	Morpher.StartingOpacity = GetOpacity();
 	Morpher.StartingMorphShape = FSlateRect( this->ScreenPosition.X, this->ScreenPosition.Y, this->ScreenPosition.X + this->Size.X, this->ScreenPosition.Y + this->Size.Y );
-	Morpher.bIsPendingPlay = true;
 	Morpher.bIsActive = true;
 	Morpher.Sequence.JumpToStart();
+
+	if ( !ActiveTickHandle.IsValid() )
+	{
+		ActiveTickHandle = RegisterActiveTick( 0.f, FWidgetActiveTickDelegate::CreateSP( this, &SWindow::TriggerPlayMorphSequence ) );
+	}
 }
 
 const FSlateBrush* SWindow::GetWindowBackground() const
@@ -1856,3 +1856,9 @@ EVisibility SWindow::GetWindowContentVisibility() const
 	// in which case the full window overlay content is visible but nothing under it
 	return (bShouldShowWindowContentDuringOverlay == true || !FullWindowOverlayWidget.IsValid()) ? EVisibility::SelfHitTestInvisible : EVisibility::Hidden;
 };
+
+EActiveTickReturnType SWindow::TriggerPlayMorphSequence( double InCurrentTime, float InDeltaTime )
+{
+	Morpher.Sequence.Play( this->AsShared() );
+	return EActiveTickReturnType::StopTicking;
+}
\ No newline at end of file
diff --git a/Engine/Source/Runtime/SlateCore/Public/Animation/CurveHandle.h b/Engine/Source/Runtime/SlateCore/Public/Animation/CurveHandle.h
--- a/Engine/Source/Runtime/SlateCore/Public/Animation/CurveHandle.h
+++ b/Engine/Source/Runtime/SlateCore/Public/Animation/CurveHandle.h
@@ -63,7 +63,10 @@ public:
 	 * @return Looped lerp value.
 	 * @see GetLerp
 	 */
+	DEPRECATED(4.7, "FCurveHandle::GetLerpLooping() is deprecated. Use GetLerp() instead. To play a sequence on a loop, pass \"true\" as the second parameter to FCurveSequence::Play / ::PlayReverse.")
 	float GetLerpLooping( ) const;
+	/** Shell to avoid generating internal deprecated warnings. Do not call. Will be removed when GetLerpLooping() is removed. */
+	float DEPRECATED_GetLerpLooping() const;
 
 	/**
 	 * Checks whether this handle is initialized.
diff --git a/Engine/Source/Runtime/SlateCore/Public/Animation/CurveSequence.h b/Engine/Source/Runtime/SlateCore/Public/Animation/CurveSequence.h
--- a/Engine/Source/Runtime/SlateCore/Public/Animation/CurveSequence.h
+++ b/Engine/Source/Runtime/SlateCore/Public/Animation/CurveSequence.h
@@ -4,16 +4,18 @@
 
 
 /**
- * A sequence of curves that can be used to drive animations.
+ * A sequence of curves that can be used to drive animations for slate widgets.
+ * Active tick registration is handled for the widget being animated when calling play.
+ *
  * Each curve within the sequence has a time offset and a duration.
  * This makes FCurveSequence convenient for crating staggered animations.
  * e.g.
  *   // We want to zoom in a widget, and then fade in its contents.
  *   FCurveHandle ZoomCurve = Sequence.AddCurve( 0, 0.15f );
  *   FCurveHandle FadeCurve = Sequence.AddCurve( 0.15f, 0.1f );
- *	 Sequence.Play();
+ *	 Sequence.Play( this->AsShared() );
  */
-struct SLATECORE_API FCurveSequence
+struct SLATECORE_API FCurveSequence : public TSharedFromThis<FCurveSequence>
 {
 public:
 
@@ -44,6 +46,9 @@ public:
 	/** Default constructor */
 	FCurveSequence( );
 
+	/** Makes sure the active tick is unregistered */
+	~FCurveSequence();
+
 	/**
 	 * Construct by adding a single animation curve to this sequence.  Does not provide access to the curve though.
 	 *
@@ -78,24 +83,50 @@ public:
 	 * @param InEaseFunction       Easing function to use for this curve.  Defaults to Linear.  Use this to smooth out your animation transitions.
 	 */
 	FCurveHandle AddCurveRelative( const float InOffset, const float InDurationSecond, const ECurveEaseFunction::Type InEaseFunction = ECurveEaseFunction::Linear );
-
+	
 	/**
 	 * Start playing this curve sequence
 	 *
 	 * @param	StartAtTime		Specifies a time offset relative to the animation to start at.  Defaults to zero (the actual start of the sequence.)
 	 */
+	DEPRECATED(4.7, "FCurveSequence::Play(const float StartAtTime) is deprecated. Curve sequences must be associated with a specific widget when played. Use Play(const TSharedRef<SWidget>& InOwnerWidget, bool bPlayLooped, const float StartAtTime) instead.")
 	void Play( const float StartAtTime = 0.0f );
 
-	/** Reverse the direction of an in-progress animation */
-	void Reverse( );
-
 	/**
 	 * Start playing this curve sequence in reverse
 	 *
 	 * @param	StartAtTime		Specifies a time offset relative to the animation to start at.  Defaults to zero (the actual start of the sequence.)
 	 */
+	DEPRECATED(4.7, "FCurveSequence::PlayReverse(const float StartAtTime) is deprecated. Curve sequences must be associated with a specific widget when played. Use PlayReverse(const TSharedRef<SWidget>& InOwnerWidget, bool bPlayLooped, const float StartAtTime) instead.")
 	void PlayReverse( const float StartAtTime = 0.0f );
 
+	/**
+	 * Start playing this curve sequence. Registers an active tick with the widget being animated.
+	 *
+	 * @param InOwnerWidget The widget that is being animated by this sequence.
+	 * @param bPlayLooped True if the curve sequence should play continually on a loop. Note that the active tick will persist until this sequence is paused or jumped to the start/end.
+	 * @param StartAtTime The relative time within the animation at which to begin playing (i.e. 0.0f is the beginning).
+	 */
+	void Play( const TSharedRef<SWidget>& InOwnerWidget, bool bPlayLooped = false, const float StartAtTime = 0.0f );
+
+	/**
+	 * Start playing this curve sequence in reverse. Registers an active tick for the widget using the sequence.
+	 *
+	 * @param InOwnerWidget The widget that is being animated by this sequence.
+	 * @param bPlayLooped True if the curve sequence should play continually on a loop. Note that the active tick will persist until this sequence is paused or jumped to the start/end.
+	 * @param StartAtTime The relative time within the animation at which to begin playing (i.e. 0.0f is the beginning).
+	 */
+	void PlayReverse( const TSharedRef<SWidget>& InOwnerWidget, bool bPlayLooped = false, const float StartAtTime = 0.0f );
+
+	/** Reverse the direction of an in-progress animation */
+	void Reverse( );
+
+	/** Pause this curve sequence. */
+	void Pause();
+
+	/** Unpause this curve sequence to resume play. */
+	void Resume( );
+
 	/**
 	 * Checks whether the sequence is currently playing.
 	 *
@@ -107,7 +138,10 @@ public:
 	float GetSequenceTime( ) const;
 
 	/** @return the current time relative to the beginning of the sequence as if the animation were a looping one. */
+	DEPRECATED(4.7, "FCurveSequence::GetSequenceTimeLooping() is deprecated. Use GetSequenceTime() instead. To play a sequence on a loop, pass \"true\" as the second parameter to Play or PlayReverse.")
 	float GetSequenceTimeLooping( ) const;
+	/** Shell to avoid generating internal deprecated warnings. Do not call. Will be removed when GetSequenceTimeLooping() is removed. */
+	float DEPRECATED_GetSequenceTimeLooping() const;
 
 	/** @return true if the animation is in reverse */
 	bool IsInReverse( ) const;
@@ -127,13 +161,16 @@ public:
 	/** Is the sequence at the end? */
 	bool IsAtEnd( ) const;
 
+	/** Is the sequence looping? */
+	bool IsLooping() const;
+
 	/**
 	 * For single-curve animations, returns the interpolation alpha for the animation.  If you call this function
 	 * on a sequence with multiple curves, an assertion will trigger.
 	 *
 	 * @return A linearly interpolated value between 0 and 1 for this curve.
 	 */
-	float GetLerp( ) const;
+	float GetLerp() const;
 
 	/**
 	 * For single-curve animations, returns the looping interpolation alpha for the animation.  If you call this
@@ -141,7 +178,10 @@ public:
 	 *
 	 * @return A linearly interpolated value between 0 and 1 for this curve.
 	 */
+	DEPRECATED(4.7, "FCurveSequence::GetLerpLooping() is deprecated. Use GetLerp() instead. To play a sequence on a loop, pass \"true\" as the second parameter to Play or PlayReverse.")
 	float GetLerpLooping( ) const;
+	/** Shell to avoid generating internal deprecated warnings. Do not call. Will be removed when GetLerpLooping() is removed. */
+	float DEPRECATED_GetLerpLooping() const;
 
 	/**
 	 * @param CurveIndex  Index of a curve in the curves array.
@@ -156,6 +196,22 @@ protected:
 	void SetStartTime( double InStartTime );
 
 private:
+	/** Helper to take care of registering the active tick */
+	void RegisterActiveTickIfNeeded(TSharedRef<SWidget> InOwnerWidget);
+
+	/** Hollow active tick to ensure a Slate Tick/Paint pass while the sequence is playing */
+	EActiveTickReturnType EnsureSlateTickDuringAnimation( double InCurrentTime, float InDeltaTime );
+
+private:
+
+	/** 
+	 * Weak reference to the owner widget that is being animated by this curve sequence.
+	 * Necessary to ensure the active tick is unregistered if the sequence is destroyed before/by the owner.
+	 */
+	TWeakPtr<SWidget> OwnerWidget;
+
+	/** The handle to the active tick */
+	TWeakPtr<FActiveTickHandle> ActiveTickHandle;
 
 	/** All the curves in this sequence. */
 	TArray<FSlateCurve> Curves;
@@ -163,9 +219,18 @@ private:
 	/** When the curve started playing. */
 	double StartTime;
 
+	/** When the curve was paused */
+	double PauseTime;
+
 	/** How long the entire sequence lasts. */
 	float TotalDuration;
 
 	/** Are we playing the animation in reverse */
-	bool bInReverse;
+	uint8 bInReverse : 1;
+
+	/** Is the sequence playing on a loop? */
+	uint8 bIsLooping : 1;
+
+	/** Is the sequence currently paused? */
+	uint8 bIsPaused : 1;
 };
diff --git a/Engine/Source/Runtime/SlateCore/Public/Animation/SlateSprings.h b/Engine/Source/Runtime/SlateCore/Public/Animation/SlateSprings.h
--- a/Engine/Source/Runtime/SlateCore/Public/Animation/SlateSprings.h
+++ b/Engine/Source/Runtime/SlateCore/Public/Animation/SlateSprings.h
@@ -128,6 +128,11 @@ public:
 		return Target;
 	}
 
+	/** @return True if the spring is at rest (i.e. at its target position) */
+	bool IsAtRest()
+	{
+		return Target == Position;
+	}
 
 	/**
 	 * Updates the simulation.  Should be called every tick!
@@ -179,7 +184,6 @@ public:
 		PreviousTarget = Target;
 	}
 
-
 private:
 
 	/** Configuration */
diff --git a/Engine/Source/Runtime/SlateCore/Public/Application/SlateApplicationBase.h b/Engine/Source/Runtime/SlateCore/Public/Application/SlateApplicationBase.h
--- a/Engine/Source/Runtime/SlateCore/Public/Application/SlateApplicationBase.h
+++ b/Engine/Source/Runtime/SlateCore/Public/Application/SlateApplicationBase.h
@@ -320,6 +320,27 @@ public:
 	 */
 	virtual bool SetUserFocus(const uint32 InUserIndex, const FWidgetPath& InFocusPath, const EFocusCause InCause) = 0;
 
+private:
+	/**
+	 * Implementation for active tick registration. See SWidget::RegisterActiveTick.
+	 */
+	void RegisterActiveTick( const TSharedRef<FActiveTickHandle>& ActiveTickHandle );
+
+	/**
+	 * Implementation for active tick registration. See SWidget::UnRegisterActiveTick.
+	 */
+	void UnRegisterActiveTick( const TSharedRef<FActiveTickHandle>& ActiveTickHandle );
+
+	/** The list of active tick handles. */
+	TArray<TWeakPtr<FActiveTickHandle>> ActiveTickHandles;
+
+protected:
+	/**
+	 * Used to determine if any active tick handles are ready to fire.
+	 * Means we need to tick slate even if no user interaction.
+	 */
+	bool AnyActiveTicksArePending();
+
 public:
 	const static uint32 CursorPointerIndex;
 
diff --git a/Engine/Source/Runtime/SlateCore/Public/SlateCore.h b/Engine/Source/Runtime/SlateCore/Public/SlateCore.h
--- a/Engine/Source/Runtime/SlateCore/Public/SlateCore.h
+++ b/Engine/Source/Runtime/SlateCore/Public/SlateCore.h
@@ -58,6 +58,16 @@ DECLARE_MEMORY_STAT_EXTERN(TEXT("Index Buffer Memory"), STAT_SlateIndexBufferMem
 // Compile all the RichText and MultiLine editable text?
 #define WITH_FANCY_TEXT 1
 
+/* Forward declarations
+ *****************************************************************************/
+class FActiveTickHandle;
+enum EActiveTickReturnType : uint8;
+
+/* Delegates
+ *****************************************************************************/
+
+/** Delegate used for registering a widget for tick services. */
+DECLARE_DELEGATE_RetVal_TwoParams(EActiveTickReturnType, FWidgetActiveTickDelegate, double /*InCurrentTime*/, float /*InDeltaTime*/);
 
 /* Includes
  *****************************************************************************/
diff --git a/Engine/Source/Runtime/SlateCore/Public/Types/SlateEnums.h b/Engine/Source/Runtime/SlateCore/Public/Types/SlateEnums.h
--- a/Engine/Source/Runtime/SlateCore/Public/Types/SlateEnums.h
+++ b/Engine/Source/Runtime/SlateCore/Public/Types/SlateEnums.h
@@ -240,3 +240,16 @@ namespace ESelectInfo
 		Direct
 	};
 }
+
+/** 
+ * Return type for FWidgetActiveTickDelegate.
+ * Don't expose to blueprints.
+ */
+enum EActiveTickReturnType : uint8
+{
+	/** If this value is returned, the widget will be unregistered for tick automatically. No need to call UnRegisterActiveTick. */
+	StopTicking,
+	/** If this value is returns, the widget will continue to have it's tick delegate called on it. */
+	KeepTicking,
+};
+
diff --git a/Engine/Source/Runtime/SlateCore/Public/Widgets/SWidget.h b/Engine/Source/Runtime/SlateCore/Public/Widgets/SWidget.h
--- a/Engine/Source/Runtime/SlateCore/Public/Widgets/SWidget.h
+++ b/Engine/Source/Runtime/SlateCore/Public/Widgets/SWidget.h
@@ -35,7 +35,6 @@ public:
 	}
 };
 
-
 /**
  * Abstract base class for Slate widgets.
  *
@@ -63,6 +62,8 @@ class SLATECORE_API SWidget
 	: public FSlateControlledConstruction,
 	  public TSharedFromThis<SWidget>		// Enables 'this->AsShared()'
 {
+	friend struct FCurveSequence;
+
 public:
 	
 	/**
@@ -117,7 +118,7 @@ public:
 	int32 Paint(const FGeometry& AllottedGeometry, const FSlateRect& MyClippingRect, FSlateWindowElementList& OutDrawElements, int32 LayerId, const FWidgetStyle& InWidgetStyle, bool bParentEnabled) const { return 0; }
 
 	/**
-	 * Ticks this widget.  Override in derived classes, but always call the parent implementation.
+	 * Ticks this widget with Geometry.  Override in derived classes, but always call the parent implementation.
 	 *
 	 * @param  AllottedGeometry The space allotted for this widget
 	 * @param  InCurrentTime  Current absolute real time
@@ -858,7 +859,48 @@ private:
 	virtual void OnArrangeChildren(const FGeometry& AllottedGeometry, FArrangedChildren& ArrangedChildren) const = 0;
 
 protected:
-	
+	/**
+	 * Registers a delegate for "active tick" at the specified frequency. Will not call TickFunction until the specified interval has elapsed once.
+	 * A widget can register as many delegates as it needs. Be careful when registering to avoid duplicate active ticks.
+	 * 
+	 * An active tick can be UnRegistered in one of three ways:
+	 *   1. Call UnRegisterActiveTick using the active tick handle that is returned here.
+	 *   2. Have your delegate return ETickWidgetReturnType::StopTicking.
+	 *   3. Destroying the widget
+	 * 
+	 * Active Ticking
+	 * --------------
+	 * Slate may go to sleep when there is no user interaction for some time to save power.
+	 * However, some UI elements may need to drive the UI even when the user is not providing any input
+	 * (ie, animations, viewport rendering, async polling, etc). A widget notifies Slate of this by
+	 * registering for "Active Tick", basically saying it needs to be ticked at some frequency to drive the UI.
+	 * In this way, slate can go to sleep, but still tick any widgets that are driving the UI at the interval they specify.
+	 * When any active tick needs to fire, all of Slate will do a Tick and Paint pass.
+	 * 
+	 * @param TickPeriod maximum frequency to tick the widget in seconds. Pass zero to tick once per frame. 
+	 *                      If an interval is missed, the delegate is NOT called more than once.
+	 * @param TickFunction delegate to call every TickFrequency seconds.
+	 * @return an active tick handle that can be used to UnRegister later.
+	 */
+	TSharedRef<FActiveTickHandle> RegisterActiveTick( float TickPeriod, FWidgetActiveTickDelegate TickFunction );
+
+	/**
+	 * Unregisters an active tick handle. This is optional, as the delegate can UnRegister itself by returning ETickWidgetReturnType::StopTicking.
+	 */
+	void UnRegisterActiveTick( const TSharedRef<FActiveTickHandle>& ActiveTickHandle );
+
+private:
+
+	/** Iterates over the active tick handles on the widget and executes them if their interval has elapsed. */
+	void ExecuteActiveTicks(double CurrentTime, float DeltaTime);
+
+	/** The list of active tick handles for this widget. */
+	TArray<TSharedRef<FActiveTickHandle>> ActiveTicks;
+
+protected:
+	/** Dtor ensures that active tick handles are UnRegistered with the SlateApplication. */
+	~SWidget();
+
 	//	DEBUG INFORMATION
 	// @todo Slate: Should compile out in final release builds?
 	FName TypeOfWidget;
diff --git a/Engine/Source/Runtime/SlateCore/Public/Widgets/SWindow.h b/Engine/Source/Runtime/SlateCore/Public/Widgets/SWindow.h
--- a/Engine/Source/Runtime/SlateCore/Public/Widgets/SWindow.h
+++ b/Engine/Source/Runtime/SlateCore/Public/Widgets/SWindow.h
@@ -811,7 +811,6 @@ protected:
 			: StartingMorphShape( FSlateRect(0,0,100,100) )
 			, TargetMorphShape( FSlateRect(0,0,100,100) )
 			, bIsActive(false)
-			, bIsPendingPlay(false)
 			, bIsAnimatingWindowSize(false)
 		{ }
 
@@ -831,9 +830,6 @@ protected:
 		/** True if this morph is currently active */
 		bool bIsActive : 1;
 
-		/** True if this morph sequence is pending play (i.e. will start on next Tick) */
-		bool bIsPendingPlay : 1;
-
 		/** True if we're morphing size as well as position.  False if we're just morphing position */
 		bool bIsAnimatingWindowSize : 1;
 
@@ -935,6 +931,12 @@ private:
 	 * @return EVisibility::Visible if the window is flashing. Used to show/hide the white flash in the title area
 	 */
 	EVisibility GetWindowFlashVisibility() const;
+
+	/** One-off active tick to trigger a the morph sequence to play */
+	EActiveTickReturnType TriggerPlayMorphSequence( double InCurrentTime, float InDeltaTime );
+
+	/** The handle to the active tick */
+	TWeakPtr<FActiveTickHandle> ActiveTickHandle;
 };
 
 
diff --git a/Engine/Source/Runtime/UMG/Private/Components/Viewport.cpp b/Engine/Source/Runtime/UMG/Private/Components/Viewport.cpp
--- a/Engine/Source/Runtime/UMG/Private/Components/Viewport.cpp
+++ b/Engine/Source/Runtime/UMG/Private/Components/Viewport.cpp
@@ -14,7 +14,7 @@ namespace FocusConstants
 }
 
 FUMGViewportCameraTransform::FUMGViewportCameraTransform()
-	: TransitionCurve(new FCurveSequence(0.0f, FocusConstants::TransitionTime, ECurveEaseFunction::CubicOut))
+	: TransitionStartTime(0)
 	, ViewLocation(FVector::ZeroVector)
 	, ViewRotation(FRotator::ZeroRotator)
 	, DesiredLocation(FVector::ZeroVector)
@@ -34,25 +34,27 @@ void FUMGViewportCameraTransform::TransitionToLocation(const FVector& InDesiredL
 	if ( bInstant )
 	{
 		SetLocation(InDesiredLocation);
-		TransitionCurve->JumpToEnd();
+		TransitionStartTime = FSlateApplication::Get().GetCurrentTime() - FocusConstants::TransitionTime;
 	}
 	else
 	{
 		DesiredLocation = InDesiredLocation;
 		StartLocation = ViewLocation;
 
-		TransitionCurve->Play();
+		TransitionStartTime = FSlateApplication::Get().GetCurrentTime();
 	}
 }
 
 bool FUMGViewportCameraTransform::UpdateTransition()
 {
 	bool bIsAnimating = false;
-	if ( TransitionCurve->IsPlaying() || ViewLocation != DesiredLocation )
+	double TransitionProgress = FMath::Clamp(( FSlateApplication::Get().GetCurrentTime() - TransitionStartTime ) / FocusConstants::TransitionTime, 0.0, 1.0);
+	if (TransitionProgress < 1.0 || ViewLocation != DesiredLocation)
 	{
-		float LerpWeight = TransitionCurve->GetLerp();
+		const float Offset = (float)TransitionProgress - 1.0f;
+		float LerpWeight = Offset * Offset * Offset + 1.0f;
 
-		if ( LerpWeight == 1.0f )
+		if (LerpWeight == 1.0f)
 		{
 			// Failsafe for the value not being exact on lerps
 			ViewLocation = DesiredLocation;
@@ -464,13 +466,12 @@ class SAutoRefreshViewport : public SViewport
 		SetViewportInterface(Viewport.ToSharedRef());
 	}
 
-	void Tick(const FGeometry& AllottedGeometry, const double InCurrentTime, const float InDeltaTime) override
+	virtual void Tick( const FGeometry& AllottedGeometry, double InCurrentTime, float InDeltaTime ) override
 	{
 		Viewport->Invalidate();
 		Viewport->InvalidateDisplay();
 
 		Viewport->Tick(AllottedGeometry, InCurrentTime, InDeltaTime);
-
 		ViewportClient->Tick(InDeltaTime);
 	}
 
diff --git a/Engine/Source/Runtime/UMG/Public/Components/Viewport.h b/Engine/Source/Runtime/UMG/Public/Components/Viewport.h
--- a/Engine/Source/Runtime/UMG/Public/Components/Viewport.h
+++ b/Engine/Source/Runtime/UMG/Public/Components/Viewport.h
@@ -70,8 +70,8 @@ public:
 	*/
 	FMatrix ComputeOrbitMatrix() const;
 private:
-	/** Curve for animating between locations */
-	TSharedPtr<struct FCurveSequence> TransitionCurve;
+	/** The time when a transition to the desired location began */
+	double TransitionStartTime;
 	/** Current viewport Position. */
 	FVector	ViewLocation;
 	/** Current Viewport orientation; valid only for perspective projections. */
