Corrected the version on deprecated functions to 4.8

diff --git a/Engine/Source/Runtime/SlateCore/Public/Animation/CurveHandle.h b/Engine/Source/Runtime/SlateCore/Public/Animation/CurveHandle.h
--- a/Engine/Source/Runtime/SlateCore/Public/Animation/CurveHandle.h
+++ b/Engine/Source/Runtime/SlateCore/Public/Animation/CurveHandle.h
@@ -63,7 +63,7 @@ public:
 	 * @return Looped lerp value.
 	 * @see GetLerp
 	 */
-	DEPRECATED(4.7, "FCurveHandle::GetLerpLooping() is deprecated. Use GetLerp() instead. To play a sequence on a loop, pass \"true\" as the second parameter to FCurveSequence::Play / ::PlayReverse.")
+	DEPRECATED(4.8, "FCurveHandle::GetLerpLooping() is deprecated. Use GetLerp() instead. To play a sequence on a loop, pass \"true\" as the second parameter to FCurveSequence::Play / ::PlayReverse.")
 	float GetLerpLooping( ) const;
 	/** Shell to avoid generating internal deprecated warnings. Do not call. Will be removed when GetLerpLooping() is removed. */
 	float DEPRECATED_GetLerpLooping() const;
diff --git a/Engine/Source/Runtime/SlateCore/Public/Animation/CurveSequence.h b/Engine/Source/Runtime/SlateCore/Public/Animation/CurveSequence.h
--- a/Engine/Source/Runtime/SlateCore/Public/Animation/CurveSequence.h
+++ b/Engine/Source/Runtime/SlateCore/Public/Animation/CurveSequence.h
@@ -89,7 +89,7 @@ public:
 	 *
 	 * @param	StartAtTime		Specifies a time offset relative to the animation to start at.  Defaults to zero (the actual start of the sequence.)
 	 */
-	DEPRECATED(4.7, "FCurveSequence::Play(const float StartAtTime) is deprecated. Curve sequences must be associated with a specific widget when played. Use Play(const TSharedRef<SWidget>& InOwnerWidget, bool bPlayLooped, const float StartAtTime) instead.")
+	DEPRECATED(4.8, "FCurveSequence::Play(const float StartAtTime) is deprecated. Curve sequences must be associated with a specific widget when played. Use Play(const TSharedRef<SWidget>& InOwnerWidget, bool bPlayLooped, const float StartAtTime) instead.")
 	void Play( const float StartAtTime = 0.0f );
 
 	/**
@@ -97,7 +97,7 @@ public:
 	 *
 	 * @param	StartAtTime		Specifies a time offset relative to the animation to start at.  Defaults to zero (the actual start of the sequence.)
 	 */
-	DEPRECATED(4.7, "FCurveSequence::PlayReverse(const float StartAtTime) is deprecated. Curve sequences must be associated with a specific widget when played. Use PlayReverse(const TSharedRef<SWidget>& InOwnerWidget, bool bPlayLooped, const float StartAtTime) instead.")
+	DEPRECATED(4.8, "FCurveSequence::PlayReverse(const float StartAtTime) is deprecated. Curve sequences must be associated with a specific widget when played. Use PlayReverse(const TSharedRef<SWidget>& InOwnerWidget, bool bPlayLooped, const float StartAtTime) instead.")
 	void PlayReverse( const float StartAtTime = 0.0f );
 
 	/**
@@ -138,7 +138,7 @@ public:
 	float GetSequenceTime( ) const;
 
 	/** @return the current time relative to the beginning of the sequence as if the animation were a looping one. */
-	DEPRECATED(4.7, "FCurveSequence::GetSequenceTimeLooping() is deprecated. Use GetSequenceTime() instead. To play a sequence on a loop, pass \"true\" as the second parameter to Play or PlayReverse.")
+	DEPRECATED(4.8, "FCurveSequence::GetSequenceTimeLooping() is deprecated. Use GetSequenceTime() instead. To play a sequence on a loop, pass \"true\" as the second parameter to Play or PlayReverse.")
 	float GetSequenceTimeLooping( ) const;
 	/** Shell to avoid generating internal deprecated warnings. Do not call. Will be removed when GetSequenceTimeLooping() is removed. */
 	float DEPRECATED_GetSequenceTimeLooping() const;
@@ -178,7 +178,7 @@ public:
 	 *
 	 * @return A linearly interpolated value between 0 and 1 for this curve.
 	 */
-	DEPRECATED(4.7, "FCurveSequence::GetLerpLooping() is deprecated. Use GetLerp() instead. To play a sequence on a loop, pass \"true\" as the second parameter to Play or PlayReverse.")
+	DEPRECATED(4.8, "FCurveSequence::GetLerpLooping() is deprecated. Use GetLerp() instead. To play a sequence on a loop, pass \"true\" as the second parameter to Play or PlayReverse.")
 	float GetLerpLooping( ) const;
 	/** Shell to avoid generating internal deprecated warnings. Do not call. Will be removed when GetLerpLooping() is removed. */
 	float DEPRECATED_GetLerpLooping() const;
