Remove old Linux hack from Slate code.

- Has been disabled for some time and seems not to be needed anymore.
- Should resolve UE-4844.

diff --git a/Engine/Source/Runtime/SlateCore/Private/Widgets/SWindow.cpp b/Engine/Source/Runtime/SlateCore/Private/Widgets/SWindow.cpp
--- a/Engine/Source/Runtime/SlateCore/Private/Widgets/SWindow.cpp
+++ b/Engine/Source/Runtime/SlateCore/Private/Widgets/SWindow.cpp
@@ -3,9 +3,6 @@
 #include "SlateCorePrivatePCH.h"
 #include "HittestGrid.h"
 
-// this define is the first step to complete removal of the below hack
-#define PLATFORM_SPECIFIC_HACK 	0
-
 namespace SWindowDefs
 {
 	/** Height of a Slate window title bar, in pixels */
@@ -1417,13 +1414,6 @@ FReply SWindow::OnMouseButtonDown( const FGeometry& MyGeometry, const FPointerEv
 
 FReply SWindow::OnMouseButtonUp( const FGeometry& MyGeometry, const FPointerEvent& MouseEvent )
 {
-#if PLATFORM_SPECIFIC_HACK && PLATFORM_LINUX
-	if (MoveResizeZone != EWindowZone::Unspecified && MouseEvent.GetEffectingButton() == EKeys::LeftMouseButton)
-	{
-		MoveResizeZone =  EWindowZone::Unspecified;
-		return FReply::Handled().ReleaseMouseCapture();
-	}
-#endif
 	if (bDragAnywhere &&  this->HasMouseCapture() && MouseEvent.GetEffectingButton() == EKeys::LeftMouseButton)
 	{
 		MoveResizeZone =  EWindowZone::Unspecified;
@@ -1437,110 +1427,6 @@ FReply SWindow::OnMouseButtonUp( const FGeometry& MyGeometry, const FPointerEven
 
 FReply SWindow::OnMouseMove( const FGeometry& MyGeometry, const FPointerEvent& MouseEvent )
 {
-#if PLATFORM_SPECIFIC_HACK && PLATFORM_LINUX
-	if (MoveResizeZone == EWindowZone::TopLeftBorder)
-	{
-		FVector2D MoveResizeOffset = MouseEvent.GetScreenSpacePosition() - MoveResizeStart;
-		if (NativeWindow.IsValid())
-		{
-			NativeWindow->ReshapeWindow(
-				FMath::TruncToInt(MoveResizeRect.Left + MoveResizeOffset.X), FMath::TruncToInt(MoveResizeRect.Top + MoveResizeOffset.Y),
-				FMath::TruncToInt(MoveResizeRect.Right - MoveResizeRect.Left - MoveResizeOffset.X), FMath::TruncToInt(MoveResizeRect.Bottom - MoveResizeRect.Top - MoveResizeOffset.Y)
-				);
-		}
-		return FReply::Handled();
-	}
-	if (MoveResizeZone == EWindowZone::BottomRightBorder)
-	{
-		FVector2D MoveResizeOffset = MouseEvent.GetScreenSpacePosition() - MoveResizeStart;
-		if (NativeWindow.IsValid())
-		{
-			NativeWindow->ReshapeWindow(
-				FMath::TruncToInt(MoveResizeRect.Left), FMath::TruncToInt(MoveResizeRect.Top),
-				FMath::TruncToInt(MoveResizeRect.Right - MoveResizeRect.Left + MoveResizeOffset.X), FMath::TruncToInt(MoveResizeRect.Bottom - MoveResizeRect.Top + MoveResizeOffset.Y)
-				);
-		}
-		return FReply::Handled();
-	}
-	if (MoveResizeZone == EWindowZone::BottomLeftBorder)
-	{
-		FVector2D MoveResizeOffset = MouseEvent.GetScreenSpacePosition() - MoveResizeStart;
-		if (NativeWindow.IsValid())
-		{
-			NativeWindow->ReshapeWindow(
-				FMath::TruncToInt(MoveResizeRect.Left + MoveResizeOffset.X), FMath::TruncToInt(MoveResizeRect.Top),
-				FMath::TruncToInt(MoveResizeRect.Right - MoveResizeRect.Left - MoveResizeOffset.X), FMath::TruncToInt(MoveResizeRect.Bottom - MoveResizeRect.Top + MoveResizeOffset.Y)
-				);
-		}
-		return FReply::Handled();
-	}
-	if (MoveResizeZone == EWindowZone::TopRightBorder)
-	{
-		FVector2D MoveResizeOffset = MouseEvent.GetScreenSpacePosition() - MoveResizeStart;
-		if (NativeWindow.IsValid())
-		{
-			NativeWindow->ReshapeWindow(
-				FMath::TruncToInt(MoveResizeRect.Left), FMath::TruncToInt(MoveResizeRect.Top + MoveResizeOffset.Y),
-				FMath::TruncToInt(MoveResizeRect.Right - MoveResizeRect.Left + MoveResizeOffset.X), FMath::TruncToInt(MoveResizeRect.Bottom - MoveResizeRect.Top - MoveResizeOffset.Y)
-				);
-		}
-		return FReply::Handled();
-	}
-	if (MoveResizeZone == EWindowZone::TopBorder)
-	{
-		FVector2D MoveResizeOffset = MouseEvent.GetScreenSpacePosition() - MoveResizeStart;
-		if (NativeWindow.IsValid())
-		{
-			NativeWindow->ReshapeWindow(
-				FMath::TruncToInt(MoveResizeRect.Left), FMath::TruncToInt(MoveResizeRect.Top + MoveResizeOffset.Y),
-				FMath::TruncToInt(MoveResizeRect.Right - MoveResizeRect.Left), FMath::TruncToInt(MoveResizeRect.Bottom - MoveResizeRect.Top - MoveResizeOffset.Y)
-				);
-		}
-		return FReply::Handled();
-	}
-	if (MoveResizeZone == EWindowZone::BottomBorder)
-	{
-		FVector2D MoveResizeOffset = MouseEvent.GetScreenSpacePosition() - MoveResizeStart;
-		if (NativeWindow.IsValid())
-		{
-			NativeWindow->ReshapeWindow(
-				FMath::TruncToInt(MoveResizeRect.Left), FMath::TruncToInt(MoveResizeRect.Top),
-				FMath::TruncToInt(MoveResizeRect.Right - MoveResizeRect.Left), FMath::TruncToInt(MoveResizeRect.Bottom - MoveResizeRect.Top + MoveResizeOffset.Y)
-				);
-		}
-		return FReply::Handled();
-	}
-	if (MoveResizeZone == EWindowZone::LeftBorder)
-	{
-		FVector2D MoveResizeOffset = MouseEvent.GetScreenSpacePosition() - MoveResizeStart;
-		if (NativeWindow.IsValid())
-		{
-			NativeWindow->ReshapeWindow(
-				FMath::TruncToInt(MoveResizeRect.Left + MoveResizeOffset.X), FMath::TruncToInt(MoveResizeRect.Top),
-				FMath::TruncToInt(MoveResizeRect.Right - MoveResizeRect.Left - MoveResizeOffset.X), FMath::TruncToInt(MoveResizeRect.Bottom - MoveResizeRect.Top)
-				);
-		}
-		return FReply::Handled();
-	}
-	if (MoveResizeZone == EWindowZone::RightBorder)
-	{
-		FVector2D MoveResizeOffset = MouseEvent.GetScreenSpacePosition() - MoveResizeStart;
-		if (NativeWindow.IsValid())
-		{
-			NativeWindow->ReshapeWindow(
-				FMath::TruncToInt(MoveResizeRect.Left), FMath::TruncToInt(MoveResizeRect.Top),
-				FMath::TruncToInt(MoveResizeRect.Right - MoveResizeRect.Left + MoveResizeOffset.X), FMath::TruncToInt(MoveResizeRect.Bottom - MoveResizeRect.Top)
-				);
-		}
-		return FReply::Handled();
-	}
-	if (MoveResizeZone == EWindowZone::TitleBar)
-	{
-		FVector2D MoveResizeOffset = MouseEvent.GetScreenSpacePosition() - MoveResizeStart;
-		this->MoveWindowTo( FVector2D(MoveResizeRect.Left, MoveResizeRect.Top) + MoveResizeOffset );
-		return FReply::Handled();
-	}
-#endif
 	if ( bDragAnywhere && this->HasMouseCapture() && MouseEvent.IsMouseButtonDown(EKeys::LeftMouseButton) && MoveResizeZone != EWindowZone::TitleBar )
 	{
 		this->MoveWindowTo( ScreenPosition + MouseEvent.GetCursorDelta() );
