[UE-6663] and [UE-5186] Fix some international characters not getting rendered with the TextBlock of UMG. Changed Slate to use FText (i.e. ICU) to determine whitespace as the performance issue with FText was resolved.

diff --git a/Engine/Source/Runtime/SlateCore/Private/Rendering/ElementBatcher.cpp b/Engine/Source/Runtime/SlateCore/Private/Rendering/ElementBatcher.cpp
--- a/Engine/Source/Runtime/SlateCore/Private/Rendering/ElementBatcher.cpp
+++ b/Engine/Source/Runtime/SlateCore/Private/Rendering/ElementBatcher.cpp
@@ -3,6 +3,7 @@
 #include "SlateCorePrivatePCH.h"
 #include "ElementBatcher.h"
 #include "SlateRenderTransform.h"
+#include "Internationalization/Text.h"
 
 // Super-hacky way of storing the scissor rect so we don't have to change all the FSlateDrawElement APIs for this hacky support.
 SLATECORE_API TOptional<FShortRect> GSlateScissorRect;
@@ -572,7 +573,7 @@ void FSlateElementBatcher::AddTextElement(const FSlateDrawElement& DrawElement)
 				InvTextureSizeY = 1.0f/FontTexture->GetHeight();
 			}
 
-			const bool bIsWhitespace = FChar::IsWhitespace(CurrentChar);
+			const bool bIsWhitespace = FText::IsWhitespace(CurrentChar);
 
 			if( !bIsWhitespace && PreviousCharEntry.IsValidEntry() )
 			{
