Fixing crash that when trying to unregister an active tick on an uninitialized slate application

diff --git a/Engine/Source/Runtime/SlateCore/Private/Widgets/SWidget.cpp b/Engine/Source/Runtime/SlateCore/Private/Widgets/SWidget.cpp
--- a/Engine/Source/Runtime/SlateCore/Private/Widgets/SWidget.cpp
+++ b/Engine/Source/Runtime/SlateCore/Private/Widgets/SWidget.cpp
@@ -707,8 +707,11 @@ void SWidget::ExecuteActiveTicks(double CurrentTime, float DeltaTime)
 SWidget::~SWidget()
 {
 	// Unregister all ActiveTicks so they aren't left stranded in the Application's list.
-	for (const auto& ActiveTickHandle : ActiveTicks)
+	if (FSlateApplicationBase::IsInitialized())
 	{
-		FSlateApplicationBase::Get().UnRegisterActiveTick(ActiveTickHandle);
+		for (const auto& ActiveTickHandle : ActiveTicks)
+		{
+			FSlateApplicationBase::Get().UnRegisterActiveTick(ActiveTickHandle);
+		}
 	}
 }
