Full rename of the Slate "active tick" system to "active timer"

diff --git a/Engine/Plugins/Editor/EpicSurvey/Source/EpicSurvey/Private/SSurvey.cpp b/Engine/Plugins/Editor/EpicSurvey/Source/EpicSurvey/Private/SSurvey.cpp
--- a/Engine/Plugins/Editor/EpicSurvey/Source/EpicSurvey/Private/SSurvey.cpp
+++ b/Engine/Plugins/Editor/EpicSurvey/Source/EpicSurvey/Private/SSurvey.cpp
@@ -70,11 +70,11 @@ void SSurvey::Construct( const FArguments& Args, const TSharedRef< FEpicSurvey >
 	if (!bIsLoaded)
 	{
 		//Construct the proper layout when finished loading
-		RegisterActiveTick( 0.1f, FWidgetActiveTickDelegate::CreateSP( this, &SSurvey::MonitorLoadStatePostConstruct ) );
+		RegisterActiveTimer( 0.1f, FWidgetActiveTimerDelegate::CreateSP( this, &SSurvey::MonitorLoadStatePostConstruct ) );
 	}
 }
 
-EActiveTickReturnType SSurvey::MonitorLoadStatePostConstruct( double InCurrentTime, float InDeltaTime )
+EActiveTimerReturnType SSurvey::MonitorLoadStatePostConstruct( double InCurrentTime, float InDeltaTime )
 {
 	EContentInitializationState::Type State = Survey->GetInitializationState();
 
@@ -91,7 +91,7 @@ EActiveTickReturnType SSurvey::MonitorLoadStatePostConstruct( double InCurrentTi
 		ConstructFailureLayout();
 	}
 
-	return bIsLoaded ? EActiveTickReturnType::StopTicking : EActiveTickReturnType::KeepTicking;
+	return bIsLoaded ? EActiveTimerReturnType::Stop : EActiveTimerReturnType::Continue;
 }
 
 void SSurvey::ConstructLoadingLayout()
diff --git a/Engine/Plugins/Editor/EpicSurvey/Source/EpicSurvey/Private/SSurvey.h b/Engine/Plugins/Editor/EpicSurvey/Source/EpicSurvey/Private/SSurvey.h
--- a/Engine/Plugins/Editor/EpicSurvey/Source/EpicSurvey/Private/SSurvey.h
+++ b/Engine/Plugins/Editor/EpicSurvey/Source/EpicSurvey/Private/SSurvey.h
@@ -41,7 +41,7 @@ private:
 
 private:
 	/** Initializes the survey once its done loading */
-	EActiveTickReturnType MonitorLoadStatePostConstruct( double InCurrentTime, float InDeltaTime );
+	EActiveTimerReturnType MonitorLoadStatePostConstruct( double InCurrentTime, float InDeltaTime );
 
 	TSharedPtr< FEpicSurvey > EpicSurvey;
 	TSharedPtr< FSurvey > Survey;
diff --git a/Engine/Plugins/Editor/PluginsEditor/Source/PluginsEditor/Private/SPluginList.cpp b/Engine/Plugins/Editor/PluginsEditor/Source/PluginsEditor/Private/SPluginList.cpp
--- a/Engine/Plugins/Editor/PluginsEditor/Source/PluginsEditor/Private/SPluginList.cpp
+++ b/Engine/Plugins/Editor/PluginsEditor/Source/PluginsEditor/Private/SPluginList.cpp
@@ -19,7 +19,7 @@ void SPluginList::Construct( const FArguments& Args, const TSharedRef< SPluginsE
 	// Find out when the plugin text filter changes
 	Owner->GetPluginTextFilter().OnChanged().AddSP( this, &SPluginList::OnPluginTextFilterChanged );
 
-	bIsActiveTickRegistered = false;
+	bIsActiveTimerRegistered = false;
 	RebuildAndFilterPluginList();
 
 	// @todo plugedit: Have optional compact version with only plugin icon + name + version?  Only expand selected?
@@ -131,12 +131,12 @@ void SPluginList::RebuildAndFilterPluginList()
 	}
 }
 
-EActiveTickReturnType SPluginList::TriggerListRebuild(double InCurrentTime, float InDeltaTime)
+EActiveTimerReturnType SPluginList::TriggerListRebuild(double InCurrentTime, float InDeltaTime)
 {
 	RebuildAndFilterPluginList();
 
-	bIsActiveTickRegistered = false;
-	return EActiveTickReturnType::StopTicking;
+	bIsActiveTimerRegistered = false;
+	return EActiveTimerReturnType::Stop;
 }
 
 void SPluginList::OnPluginTextFilterChanged()
@@ -147,10 +147,10 @@ void SPluginList::OnPluginTextFilterChanged()
 
 void SPluginList::SetNeedsRefresh()
 {
-	if (!bIsActiveTickRegistered)
+	if (!bIsActiveTimerRegistered)
 	{
-		bIsActiveTickRegistered = true;
-		RegisterActiveTick(0.f, FWidgetActiveTickDelegate::CreateSP(this, &SPluginList::TriggerListRebuild));
+		bIsActiveTimerRegistered = true;
+		RegisterActiveTimer(0.f, FWidgetActiveTimerDelegate::CreateSP(this, &SPluginList::TriggerListRebuild));
 	}
 }
 
diff --git a/Engine/Plugins/Editor/PluginsEditor/Source/PluginsEditor/Private/SPluginList.h b/Engine/Plugins/Editor/PluginsEditor/Source/PluginsEditor/Private/SPluginList.h
--- a/Engine/Plugins/Editor/PluginsEditor/Source/PluginsEditor/Private/SPluginList.h
+++ b/Engine/Plugins/Editor/PluginsEditor/Source/PluginsEditor/Private/SPluginList.h
@@ -45,8 +45,8 @@ private:
 	/** Rebuilds the list of plugins from scratch and applies filtering. */
 	void RebuildAndFilterPluginList();
 
-	/** One-off active tick to trigger a full refresh when something has changed with either our filtering or the loaded plugin set */
-	EActiveTickReturnType TriggerListRebuild(double InCurrentTime, float InDeltaTime);
+	/** One-off active timer to trigger a full refresh when something has changed with either our filtering or the loaded plugin set */
+	EActiveTimerReturnType TriggerListRebuild(double InCurrentTime, float InDeltaTime);
 
 private:
 
@@ -59,7 +59,7 @@ private:
 	/** List of everything that we want to display in the plugin list */
 	TArray< FPluginListItemPtr > PluginListItems;
 
-	/** Whether the active tick to refresh the list is registered */
-	bool bIsActiveTickRegistered;
+	/** Whether the active timer to refresh the list is registered */
+	bool bIsActiveTimerRegistered;
 };
 
diff --git a/Engine/Plugins/Editor/PluginsEditor/Source/PluginsEditor/Private/SPluginsEditor.cpp b/Engine/Plugins/Editor/PluginsEditor/Source/PluginsEditor/Private/SPluginsEditor.cpp
--- a/Engine/Plugins/Editor/PluginsEditor/Source/PluginsEditor/Private/SPluginsEditor.cpp
+++ b/Engine/Plugins/Editor/PluginsEditor/Source/PluginsEditor/Private/SPluginsEditor.cpp
@@ -26,7 +26,7 @@ void SPluginsEditor::Construct( const FArguments& Args )
 
 	// @todo plugedit: Should we save/restore category selection and other view settings?  Splitter position, etc.
 
-	RegisterActiveTick (0.f, FWidgetActiveTickDelegate::CreateSP (this, &SPluginsEditor::TriggerBreadcrumbRefresh));
+	RegisterActiveTimer (0.f, FWidgetActiveTimerDelegate::CreateSP (this, &SPluginsEditor::TriggerBreadcrumbRefresh));
 
 	// Setup text filtering
 	PluginTextFilter = MakeShareable( new FPluginTextFilter( FPluginTextFilter::FItemToStringArray::CreateStatic( &Local::PluginStatusToStringArray ) ) );
@@ -159,13 +159,13 @@ void SPluginsEditor::OnCategorySelectionChanged()
 	}
 
 	// Breadcrumbs will need to be refreshed
-	RegisterActiveTick (0.f, FWidgetActiveTickDelegate::CreateSP (this, &SPluginsEditor::TriggerBreadcrumbRefresh));
+	RegisterActiveTimer (0.f, FWidgetActiveTimerDelegate::CreateSP (this, &SPluginsEditor::TriggerBreadcrumbRefresh));
 }
 
-EActiveTickReturnType SPluginsEditor::TriggerBreadcrumbRefresh(double InCurrentTime, float InDeltaTime)
+EActiveTimerReturnType SPluginsEditor::TriggerBreadcrumbRefresh(double InCurrentTime, float InDeltaTime)
 {
 	RefreshBreadcrumbTrail();
-	return EActiveTickReturnType::StopTicking;
+	return EActiveTimerReturnType::Stop;
 }
 
 void SPluginsEditor::RefreshBreadcrumbTrail()
diff --git a/Engine/Plugins/Editor/PluginsEditor/Source/PluginsEditor/Private/SPluginsEditor.h b/Engine/Plugins/Editor/PluginsEditor/Source/PluginsEditor/Private/SPluginsEditor.h
--- a/Engine/Plugins/Editor/PluginsEditor/Source/PluginsEditor/Private/SPluginsEditor.h
+++ b/Engine/Plugins/Editor/PluginsEditor/Source/PluginsEditor/Private/SPluginsEditor.h
@@ -82,8 +82,8 @@ private:
 	/** Called to refresh the breadcrumb trail immediately */
 	void RefreshBreadcrumbTrail();
 
-	/** One-off active tick to trigger a refresh of the breadcrumb trail as needed */
-	EActiveTickReturnType TriggerBreadcrumbRefresh(double InCurrentTime, float InDeltaTime);
+	/** One-off active timer to trigger a refresh of the breadcrumb trail as needed */
+	EActiveTimerReturnType TriggerBreadcrumbRefresh(double InCurrentTime, float InDeltaTime);
 
 private:
 
diff --git a/Engine/Source/Developer/DeviceManager/Private/Widgets/Processes/SDeviceProcesses.cpp b/Engine/Source/Developer/DeviceManager/Private/Widgets/Processes/SDeviceProcesses.cpp
--- a/Engine/Source/Developer/DeviceManager/Private/Widgets/Processes/SDeviceProcesses.cpp
+++ b/Engine/Source/Developer/DeviceManager/Private/Widgets/Processes/SDeviceProcesses.cpp
@@ -133,14 +133,14 @@ void SDeviceProcesses::Construct( const FArguments& InArgs, const FDeviceManager
 	ReloadProcessList(true);
 
 	// Register for an active update every 2.5 seconds
-	RegisterActiveTick( 2.5f, FWidgetActiveTickDelegate::CreateSP( this, &SDeviceProcesses::UpdateProcessList ) );
+	RegisterActiveTimer( 2.5f, FWidgetActiveTimerDelegate::CreateSP( this, &SDeviceProcesses::UpdateProcessList ) );
 }
 
-EActiveTickReturnType SDeviceProcesses::UpdateProcessList( double InCurrentTime, float InDeltaTime )
+EActiveTimerReturnType SDeviceProcesses::UpdateProcessList( double InCurrentTime, float InDeltaTime )
 {
 	ReloadProcessList( true );
 
-	return EActiveTickReturnType::KeepTicking;
+	return EActiveTimerReturnType::Continue;
 }
 
 /* SDeviceDetails implementation
diff --git a/Engine/Source/Developer/DeviceManager/Private/Widgets/Processes/SDeviceProcesses.h b/Engine/Source/Developer/DeviceManager/Private/Widgets/Processes/SDeviceProcesses.h
--- a/Engine/Source/Developer/DeviceManager/Private/Widgets/Processes/SDeviceProcesses.h
+++ b/Engine/Source/Developer/DeviceManager/Private/Widgets/Processes/SDeviceProcesses.h
@@ -40,7 +40,7 @@ protected:
 
 private:
 	// Periodically refreshes the process list
-	EActiveTickReturnType UpdateProcessList( double InCurrentTime, float InDeltaTime );
+	EActiveTimerReturnType UpdateProcessList( double InCurrentTime, float InDeltaTime );
 
 	// Callback for getting the text of the message overlay.
 	FText HandleMessageOverlayText( ) const;
diff --git a/Engine/Source/Developer/Profiler/Private/Widgets/SFiltersAndPresets.cpp b/Engine/Source/Developer/Profiler/Private/Widgets/SFiltersAndPresets.cpp
--- a/Engine/Source/Developer/Profiler/Private/Widgets/SFiltersAndPresets.cpp
+++ b/Engine/Source/Developer/Profiler/Private/Widgets/SFiltersAndPresets.cpp
@@ -572,7 +572,7 @@ SFiltersAndPresets::SFiltersAndPresets()
 	: GroupingMode( EStatGroupingOrSortingMode::GroupName )
 	, SortingMode( EStatGroupingOrSortingMode::StatName )
 	, bExpansionSaved( false )
-	, bIsActiveTickRegistered( false )
+	, bIsActiveTimerRegistered( false )
 {
 	FMemory::MemSet( bStatTypeIsVisible, 1 );
 }
@@ -760,19 +760,19 @@ void SFiltersAndPresets::ProfilerManager_OnRequestFilterAndPresetsUpdate()
 
 void SFiltersAndPresets::RegisterForUpdate()
 {
-	if (!bIsActiveTickRegistered)
+	if (!bIsActiveTimerRegistered)
 	{
-		bIsActiveTickRegistered = true;
-		RegisterActiveTick(0.f, FWidgetActiveTickDelegate::CreateSP(this, &SFiltersAndPresets::UpdateData));
+		bIsActiveTimerRegistered = true;
+		RegisterActiveTimer(0.f, FWidgetActiveTimerDelegate::CreateSP(this, &SFiltersAndPresets::UpdateData));
 	}
 }
 
-EActiveTickReturnType SFiltersAndPresets::UpdateData(double InCurrentTime, float InDeltaTime)
+EActiveTimerReturnType SFiltersAndPresets::UpdateData(double InCurrentTime, float InDeltaTime)
 {
 	auto It = FProfilerManager::Get()->GetProfilerInstancesIterator();
 	UpdateGroupAndStatTree(It.Value());
 
-	return EActiveTickReturnType::StopTicking;
+	return EActiveTimerReturnType::Stop;
 }
 
 void SFiltersAndPresets::UpdateGroupAndStatTree( const FProfilerSessionPtr InProfilerSession )
diff --git a/Engine/Source/Developer/Profiler/Private/Widgets/SFiltersAndPresets.h b/Engine/Source/Developer/Profiler/Private/Widgets/SFiltersAndPresets.h
--- a/Engine/Source/Developer/Profiler/Private/Widgets/SFiltersAndPresets.h
+++ b/Engine/Source/Developer/Profiler/Private/Widgets/SFiltersAndPresets.h
@@ -413,12 +413,12 @@ protected:
 	bool bStatTypeIsVisible[EProfilerSampleTypes::InvalidOrMax];
 
 private:
-	/** Registers the active tick to trigger an update if not done already */
+	/** Registers the active timer to trigger an update if not done already */
 	void RegisterForUpdate();
 
-	/** Simple one-off active tick to update the group and stat trees */
-	EActiveTickReturnType UpdateData(double InCurrentTime, float InDeltaTime);
+	/** Simple one-off active timer to update the group and stat trees */
+	EActiveTimerReturnType UpdateData(double InCurrentTime, float InDeltaTime);
 
-	/** Whether the active tick has already been registered this frame */
-	bool bIsActiveTickRegistered;
+	/** Whether the active timer has already been registered this frame */
+	bool bIsActiveTimerRegistered;
 };
diff --git a/Engine/Source/Developer/Profiler/Private/Widgets/SProfilerMiniView.cpp b/Engine/Source/Developer/Profiler/Private/Widgets/SProfilerMiniView.cpp
--- a/Engine/Source/Developer/Profiler/Private/Widgets/SProfilerMiniView.cpp
+++ b/Engine/Source/Developer/Profiler/Private/Widgets/SProfilerMiniView.cpp
@@ -8,7 +8,7 @@
 #undef LOCTEXT_NAMESPACE
 
 SProfilerMiniView::SProfilerMiniView()
-	: bIsActiveTickRegistered( false )
+	: bIsActiveTimerRegistered( false )
 {
 	Reset();
 }
@@ -52,16 +52,16 @@ void SProfilerMiniView::Reset()
 	CursorType = EMiniviewCursor::Default;
 }
 
-EActiveTickReturnType SProfilerMiniView::EnsureDataUpdateDuringPreview(double InCurrentTime, float InDeltaTime)
+EActiveTimerReturnType SProfilerMiniView::EnsureDataUpdateDuringPreview(double InCurrentTime, float InDeltaTime)
 {
 	if (RecentlyAddedFrames.Num() > 0)
 	{
 		bUpdateData = true;
-		return EActiveTickReturnType::KeepTicking;
+		return EActiveTimerReturnType::Continue;
 	}
 	
-	bIsActiveTickRegistered = false;
-	return EActiveTickReturnType::StopTicking;
+	bIsActiveTimerRegistered = false;
+	return EActiveTimerReturnType::Stop;
 }
 
 void SProfilerMiniView::Tick( const FGeometry& AllottedGeometry, const double InCurrentTime, const float InDeltaTime )
@@ -586,10 +586,10 @@ void SProfilerMiniView::AddThreadTime( int32 InFrameIndex, const TMap<uint32, fl
 	FrameThreadTimes.ThreadTimes = InThreadMS;
 	RecentlyAddedFrames.Add( FrameThreadTimes );
 
-	if (!bIsActiveTickRegistered)
+	if (!bIsActiveTimerRegistered)
 	{
-		bIsActiveTickRegistered = true;
-		RegisterActiveTick(0.f, FWidgetActiveTickDelegate::CreateSP(this, &SProfilerMiniView::EnsureDataUpdateDuringPreview));
+		bIsActiveTimerRegistered = true;
+		RegisterActiveTimer(0.f, FWidgetActiveTimerDelegate::CreateSP(this, &SProfilerMiniView::EnsureDataUpdateDuringPreview));
 	}
 
 	StatMetadata = InStatMetaData;
diff --git a/Engine/Source/Developer/Profiler/Private/Widgets/SProfilerMiniView.h b/Engine/Source/Developer/Profiler/Private/Widgets/SProfilerMiniView.h
--- a/Engine/Source/Developer/Profiler/Private/Widgets/SProfilerMiniView.h
+++ b/Engine/Source/Developer/Profiler/Private/Widgets/SProfilerMiniView.h
@@ -398,8 +398,8 @@ protected:
 
 private:
 	/** Actively ticks the widget to update to process new frame data as long as any frames have been added recently */
-	EActiveTickReturnType EnsureDataUpdateDuringPreview( double InCurrentTime, float InDeltaTime );
+	EActiveTimerReturnType EnsureDataUpdateDuringPreview( double InCurrentTime, float InDeltaTime );
 
-	/** True if the active tick is currently registered */
-	bool bIsActiveTickRegistered;
+	/** True if the active timer is currently registered */
+	bool bIsActiveTimerRegistered;
 };
\ No newline at end of file
diff --git a/Engine/Source/Developer/Profiler/Private/Widgets/SProfilerWindow.cpp b/Engine/Source/Developer/Profiler/Private/Widgets/SProfilerWindow.cpp
--- a/Engine/Source/Developer/Profiler/Private/Widgets/SProfilerWindow.cpp
+++ b/Engine/Source/Developer/Profiler/Private/Widgets/SProfilerWindow.cpp
@@ -490,21 +490,21 @@ void SProfilerWindow::SendingServiceSideCapture_Load( const FString Filename )
 	}
 }
 
-EActiveTickReturnType SProfilerWindow::UpdateActiveDuration( double InCurrentTime, float InDeltaTime )
+EActiveTimerReturnType SProfilerWindow::UpdateActiveDuration( double InCurrentTime, float InDeltaTime )
 {
 	DurationActive += InDeltaTime;
 
-	// The profiler window will explicitly unregister this active tick when the mouse leaves
-	return EActiveTickReturnType::KeepTicking;
+	// The profiler window will explicitly unregister this active timer when the mouse leaves
+	return EActiveTimerReturnType::Continue;
 }
 
 void SProfilerWindow::OnMouseEnter(const FGeometry& MyGeometry, const FPointerEvent& MouseEvent)
 {
 	SCompoundWidget::OnMouseEnter(MyGeometry, MouseEvent);
 
-	if (!ActiveTickHandle.IsValid())
+	if (!ActiveTimerHandle.IsValid())
 	{
-		ActiveTickHandle = RegisterActiveTick(0.f, FWidgetActiveTickDelegate::CreateSP(this, &SProfilerWindow::UpdateActiveDuration));
+		ActiveTimerHandle = RegisterActiveTimer(0.f, FWidgetActiveTimerDelegate::CreateSP(this, &SProfilerWindow::UpdateActiveDuration));
 	}
 }
 
@@ -512,10 +512,10 @@ void SProfilerWindow::OnMouseLeave(const FPointerEvent& MouseEvent)
 {
 	SCompoundWidget::OnMouseLeave(MouseEvent);
 
-	auto PinnedActiveTickHandle = ActiveTickHandle.Pin();
-	if (PinnedActiveTickHandle.IsValid())
+	auto PinnedActiveTimerHandle = ActiveTimerHandle.Pin();
+	if (PinnedActiveTimerHandle.IsValid())
 	{
-		UnRegisterActiveTick(PinnedActiveTickHandle.ToSharedRef());
+		UnRegisterActiveTimer(PinnedActiveTimerHandle.ToSharedRef());
 	}
 	
 }
diff --git a/Engine/Source/Developer/Profiler/Private/Widgets/SProfilerWindow.h b/Engine/Source/Developer/Profiler/Private/Widgets/SProfilerWindow.h
--- a/Engine/Source/Developer/Profiler/Private/Widgets/SProfilerWindow.h
+++ b/Engine/Source/Developer/Profiler/Private/Widgets/SProfilerWindow.h
@@ -116,10 +116,10 @@ private:
 	virtual FReply OnDragOver( const FGeometry& MyGeometry, const FDragDropEvent& DragDropEvent )  override;
 
 	/** Updates the amount of time the profiler has been active */
-	EActiveTickReturnType UpdateActiveDuration( double InCurrentTime, float InDeltaTime );
+	EActiveTimerReturnType UpdateActiveDuration( double InCurrentTime, float InDeltaTime );
 
 	/** The handle to the active update duration tick */
-	TWeakPtr<FActiveTickHandle> ActiveTickHandle;
+	TWeakPtr<FActiveTimerHandle> ActiveTimerHandle;
 
 	/** Holds all widgets for the profiler window like menu bar, toolbar and tabs. */
 	TSharedPtr<SVerticalBox> MainContentPanel;
diff --git a/Engine/Source/Developer/SharedSettingsWidgets/Private/SSettingsEditorCheckoutNotice.cpp b/Engine/Source/Developer/SharedSettingsWidgets/Private/SSettingsEditorCheckoutNotice.cpp
--- a/Engine/Source/Developer/SharedSettingsWidgets/Private/SSettingsEditorCheckoutNotice.cpp
+++ b/Engine/Source/Developer/SharedSettingsWidgets/Private/SSettingsEditorCheckoutNotice.cpp
@@ -228,7 +228,7 @@ FSlateColor SSettingsEditorCheckoutNotice::HandleBorderBackgroundColor() const
 
 void SSettingsEditorCheckoutNotice::Tick( const FGeometry& AllottedGeometry, const double InCurrentTime, const float InDeltaTime )
 {
-	// cache selected settings object's configuration file state (purposely not in an active tick to make sure it appears while the user is active)
+	// cache selected settings object's configuration file state (purposely not in an active timer to make sure it appears while the user is active)
 	if (InCurrentTime - LastDefaultConfigCheckOutTime >= 1.0f)
 	{
 		bool NewCheckOutNeeded = false;
diff --git a/Engine/Source/Developer/SourceControl/Private/SSourceControlLogin.cpp b/Engine/Source/Developer/SourceControl/Private/SSourceControlLogin.cpp
--- a/Engine/Source/Developer/SourceControl/Private/SSourceControlLogin.cpp
+++ b/Engine/Source/Developer/SourceControl/Private/SSourceControlLogin.cpp
@@ -216,7 +216,7 @@ FReply SSourceControlLogin::OnAcceptSettings()
 	//Increase the tick frequency during login if needed
 	if ( ParentWindowPtr.IsValid() && ( FSlateApplication::Get().GetActiveModalWindow() == ParentWindowPtr.Pin() ) )
 	{
-		ActiveTickHandle = RegisterActiveTick( 0.f, FWidgetActiveTickDelegate::CreateSP( this, &SSourceControlLogin::TickSourceControlModule ) );
+		ActiveTimerHandle = RegisterActiveTimer( 0.f, FWidgetActiveTimerDelegate::CreateSP( this, &SSourceControlLogin::TickSourceControlModule ) );
 	}
 
 	FSourceControlModule& SourceControlModule = FSourceControlModule::Get();
@@ -277,11 +277,11 @@ void SSourceControlLogin::DisplayConnectionError(const FText& InErrorText)
 	
 	SourceControlLog.Notify();
 
-	// Suspend the active tick until there's another login attempt
-	auto PinnedActiveTick = ActiveTickHandle.Pin();
-	if ( PinnedActiveTick.IsValid() )
+	// Suspend the active timer until there's another login attempt
+	auto PinnedActiveTimer = ActiveTimerHandle.Pin();
+	if ( PinnedActiveTimer.IsValid() )
 	{
-		UnRegisterActiveTick( PinnedActiveTick.ToSharedRef() );
+		UnRegisterActiveTimer( PinnedActiveTimer.ToSharedRef() );
 	}
 }
 
@@ -322,10 +322,10 @@ EVisibility SSourceControlLogin::GetDisabledTextVisibility() const
 	return SourceControlModule.GetProvider().GetName() == "None" ? EVisibility::Visible : EVisibility::Collapsed;
 }
 
-EActiveTickReturnType SSourceControlLogin::TickSourceControlModule( double InCurrentTime, float InDeltaTime )
+EActiveTimerReturnType SSourceControlLogin::TickSourceControlModule( double InCurrentTime, float InDeltaTime )
 {
 	FSourceControlModule::Get().Tick();
-	return EActiveTickReturnType::KeepTicking;
+	return EActiveTimerReturnType::Continue;
 }
 
 #undef LOCTEXT_NAMESPACE
diff --git a/Engine/Source/Developer/SourceControl/Private/SSourceControlLogin.h b/Engine/Source/Developer/SourceControl/Private/SSourceControlLogin.h
--- a/Engine/Source/Developer/SourceControl/Private/SSourceControlLogin.h
+++ b/Engine/Source/Developer/SourceControl/Private/SSourceControlLogin.h
@@ -42,7 +42,7 @@ public:
 private:
 
 	/** Ticks the source control module (only necessary if the login window is modal) */
-	EActiveTickReturnType TickSourceControlModule( double InCurrentTime, float InDeltaTime );
+	EActiveTimerReturnType TickSourceControlModule( double InCurrentTime, float InDeltaTime );
 
 	/** Delegate called when the user clicks the 'Accept Settings' button */
 	FReply OnAcceptSettings();
@@ -94,7 +94,7 @@ private:
 	TSharedPtr<class SBorder> SettingsBorder;
 
 	/** The handle to the active scc module tick */
-	TWeakPtr<FActiveTickHandle> ActiveTickHandle;
+	TWeakPtr<FActiveTimerHandle> ActiveTimerHandle;
 };
 
 #endif // SOURCE_CONTROL_WITH_SLATE
diff --git a/Engine/Source/Editor/ContentBrowser/Private/SAssetDialog.cpp b/Engine/Source/Editor/ContentBrowser/Private/SAssetDialog.cpp
--- a/Engine/Source/Editor/ContentBrowser/Private/SAssetDialog.cpp
+++ b/Engine/Source/Editor/ContentBrowser/Private/SAssetDialog.cpp
@@ -31,7 +31,7 @@ void SAssetDialog::Construct(const FArguments& InArgs, const FSharedAssetDialogC
 
 	const FString DefaultPath = InConfig.DefaultPath;
 
-	RegisterActiveTick( 0.f, FWidgetActiveTickDelegate::CreateSP( this, &SAssetDialog::SetFocusPostConstruct ) );
+	RegisterActiveTimer( 0.f, FWidgetActiveTimerDelegate::CreateSP( this, &SAssetDialog::SetFocusPostConstruct ) );
 
 	FPathPickerConfig PathPickerConfig;
 	PathPickerConfig.DefaultPath = DefaultPath;
@@ -249,10 +249,10 @@ FReply SAssetDialog::OnKeyDown( const FGeometry& MyGeometry, const FKeyEvent& In
 	return SCompoundWidget::OnKeyDown(MyGeometry, InKeyEvent);
 }
 
-EActiveTickReturnType SAssetDialog::SetFocusPostConstruct( double InCurrentTime, float InDeltaTime )
+EActiveTimerReturnType SAssetDialog::SetFocusPostConstruct( double InCurrentTime, float InDeltaTime )
 {
 	FocusNameBox();
-	return EActiveTickReturnType::StopTicking;
+	return EActiveTimerReturnType::Stop;
 }
 
 void SAssetDialog::SetOnAssetsChosenForOpen(const FOnAssetsChosenForOpen& InOnAssetsChosenForOpen)
diff --git a/Engine/Source/Editor/ContentBrowser/Private/SAssetDialog.h b/Engine/Source/Editor/ContentBrowser/Private/SAssetDialog.h
--- a/Engine/Source/Editor/ContentBrowser/Private/SAssetDialog.h
+++ b/Engine/Source/Editor/ContentBrowser/Private/SAssetDialog.h
@@ -28,7 +28,7 @@ public:
 private:
 
 	/** Used to focus the name box immediately following construction */
-	EActiveTickReturnType SetFocusPostConstruct( double InCurrentTime, float InDeltaTime );
+	EActiveTimerReturnType SetFocusPostConstruct( double InCurrentTime, float InDeltaTime );
 
 	/** Moves keyboard focus to the name box if this is a save dialog */
 	void FocusNameBox();
diff --git a/Engine/Source/Editor/ContentBrowser/Private/SAssetPicker.cpp b/Engine/Source/Editor/ContentBrowser/Private/SAssetPicker.cpp
--- a/Engine/Source/Editor/ContentBrowser/Private/SAssetPicker.cpp
+++ b/Engine/Source/Editor/ContentBrowser/Private/SAssetPicker.cpp
@@ -25,7 +25,7 @@ void SAssetPicker::Construct( const FArguments& InArgs )
 
 	if ( InArgs._AssetPickerConfig.bFocusSearchBoxWhenOpened )
 	{
-		RegisterActiveTick( 0.f, FWidgetActiveTickDelegate::CreateSP( this, &SAssetPicker::SetFocusPostConstruct ) );
+		RegisterActiveTimer( 0.f, FWidgetActiveTimerDelegate::CreateSP( this, &SAssetPicker::SetFocusPostConstruct ) );
 	}
 
 	for (auto DelegateIt = InArgs._AssetPickerConfig.GetCurrentSelectionDelegates.CreateConstIterator(); DelegateIt; ++DelegateIt)
@@ -236,7 +236,7 @@ void SAssetPicker::Construct( const FArguments& InArgs )
 	AssetViewPtr->RequestSlowFullListRefresh();
 }
 
-EActiveTickReturnType SAssetPicker::SetFocusPostConstruct( double InCurrentTime, float InDeltaTime )
+EActiveTimerReturnType SAssetPicker::SetFocusPostConstruct( double InCurrentTime, float InDeltaTime )
 {
 	if ( SearchBoxPtr.IsValid() )
 	{
@@ -244,10 +244,10 @@ EActiveTickReturnType SAssetPicker::SetFocusPostConstruct( double InCurrentTime,
 		FSlateApplication::Get().GeneratePathToWidgetUnchecked( SearchBoxPtr.ToSharedRef(), WidgetToFocusPath );
 		FSlateApplication::Get().SetKeyboardFocus( WidgetToFocusPath, EFocusCause::SetDirectly );
 
-		return EActiveTickReturnType::StopTicking;
+		return EActiveTimerReturnType::Stop;
 	}
 
-	return EActiveTickReturnType::KeepTicking;
+	return EActiveTimerReturnType::Continue;
 }
 
 FReply SAssetPicker::OnKeyDown(const FGeometry& MyGeometry, const FKeyEvent& InKeyEvent)
diff --git a/Engine/Source/Editor/ContentBrowser/Private/SAssetPicker.h b/Engine/Source/Editor/ContentBrowser/Private/SAssetPicker.h
--- a/Engine/Source/Editor/ContentBrowser/Private/SAssetPicker.h
+++ b/Engine/Source/Editor/ContentBrowser/Private/SAssetPicker.h
@@ -26,7 +26,7 @@ public:
 
 private:
 	/** Focuses the search box post-construct */
-	EActiveTickReturnType SetFocusPostConstruct( double InCurrentTime, float InDeltaTime );
+	EActiveTimerReturnType SetFocusPostConstruct( double InCurrentTime, float InDeltaTime );
 
 	void FolderEntered(const FString& FolderPath);
 
diff --git a/Engine/Source/Editor/ContentBrowser/Private/SPathView.cpp b/Engine/Source/Editor/ContentBrowser/Private/SPathView.cpp
--- a/Engine/Source/Editor/ContentBrowser/Private/SPathView.cpp
+++ b/Engine/Source/Editor/ContentBrowser/Private/SPathView.cpp
@@ -28,7 +28,7 @@ SPathView::~SPathView()
 
 void SPathView::Construct( const FArguments& InArgs )
 {
-	RegisterActiveTick(0.f, FWidgetActiveTickDelegate::CreateSP(this, &SPathView::TriggerRepopulate));
+	RegisterActiveTimer(0.f, FWidgetActiveTimerDelegate::CreateSP(this, &SPathView::TriggerRepopulate));
 
 	OnPathSelected = InArgs._OnPathSelected;
 	bAllowContextMenu = InArgs._AllowContextMenu;
@@ -39,7 +39,7 @@ void SPathView::Construct( const FArguments& InArgs )
 
 	if ( InArgs._FocusSearchBoxWhenOpened )
 	{
-		RegisterActiveTick( 0.f, FWidgetActiveTickDelegate::CreateSP( this, &SPathView::SetFocusPostConstruct ) );
+		RegisterActiveTimer( 0.f, FWidgetActiveTimerDelegate::CreateSP( this, &SPathView::SetFocusPostConstruct ) );
 	}
 
 	// Listen for when view settings are changed
@@ -620,19 +620,19 @@ void SPathView::LoadSettings(const FString& IniFilename, const FString& IniSecti
 	}
 }
 
-EActiveTickReturnType SPathView::SetFocusPostConstruct( double InCurrentTime, float InDeltaTime )
+EActiveTimerReturnType SPathView::SetFocusPostConstruct( double InCurrentTime, float InDeltaTime )
 {
 	FWidgetPath WidgetToFocusPath;
 	FSlateApplication::Get().GeneratePathToWidgetUnchecked( SearchBoxPtr.ToSharedRef(), WidgetToFocusPath );
 	FSlateApplication::Get().SetKeyboardFocus( WidgetToFocusPath, EFocusCause::SetDirectly );
 
-	return EActiveTickReturnType::StopTicking;
+	return EActiveTimerReturnType::Stop;
 }
 
-EActiveTickReturnType SPathView::TriggerRepopulate(double InCurrentTime, float InDeltaTime)
+EActiveTimerReturnType SPathView::TriggerRepopulate(double InCurrentTime, float InDeltaTime)
 {
 	Populate();
-	return EActiveTickReturnType::StopTicking;
+	return EActiveTimerReturnType::Stop;
 }
 
 TSharedPtr<SWidget> SPathView::MakePathViewContextMenu()
@@ -1336,7 +1336,7 @@ void SPathView::OnAssetRegistrySearchCompleted()
 void SPathView::OnContentPathMountedOrDismounted( const FString& AssetPath, const FString& FilesystemPath )
 {
 	// A new content path has appeared, so we should refresh out root set of paths
-	RegisterActiveTick(0.f, FWidgetActiveTickDelegate::CreateSP(this, &SPathView::TriggerRepopulate));
+	RegisterActiveTimer(0.f, FWidgetActiveTimerDelegate::CreateSP(this, &SPathView::TriggerRepopulate));
 }
 
 void SPathView::HandleSettingChanged(FName PropertyName)
diff --git a/Engine/Source/Editor/ContentBrowser/Private/SPathView.h b/Engine/Source/Editor/ContentBrowser/Private/SPathView.h
--- a/Engine/Source/Editor/ContentBrowser/Private/SPathView.h
+++ b/Engine/Source/Editor/ContentBrowser/Private/SPathView.h
@@ -221,11 +221,11 @@ private:
 	/** Delegate called when an editor setting is changed */
 	void HandleSettingChanged(FName PropertyName);
 
-	/** One-off active tick to focus the widget post-construct */
-	EActiveTickReturnType SetFocusPostConstruct(double InCurrentTime, float InDeltaTime);
+	/** One-off active timer to focus the widget post-construct */
+	EActiveTimerReturnType SetFocusPostConstruct(double InCurrentTime, float InDeltaTime);
 
-	/** One-off active tick to repopulate the path view */
-	EActiveTickReturnType TriggerRepopulate(double InCurrentTime, float InDeltaTime);
+	/** One-off active timer to repopulate the path view */
+	EActiveTimerReturnType TriggerRepopulate(double InCurrentTime, float InDeltaTime);
 
 private:
 	/** A helper class to manage PreventTreeItemChangedDelegateCount by incrementing it when constructed (on the stack) and decrementing when destroyed */
diff --git a/Engine/Source/Editor/EditorWidgets/Private/SObjectNameEditableTextBox.cpp b/Engine/Source/Editor/EditorWidgets/Private/SObjectNameEditableTextBox.cpp
--- a/Engine/Source/Editor/EditorWidgets/Private/SObjectNameEditableTextBox.cpp
+++ b/Engine/Source/Editor/EditorWidgets/Private/SObjectNameEditableTextBox.cpp
@@ -34,7 +34,7 @@ void SObjectNameEditableTextBox::Construct( const FArguments& InArgs )
 		];
 }
 
-EActiveTickReturnType SObjectNameEditableTextBox::UpdateHighlightSpringState( double InCurrentTime, float InDeltaTime )
+EActiveTimerReturnType SObjectNameEditableTextBox::UpdateHighlightSpringState( double InCurrentTime, float InDeltaTime )
 {
 	if ( (float)(InCurrentTime - LastCommittedTime) <= HighlightTargetEffectDuration )
 	{
@@ -45,7 +45,7 @@ EActiveTickReturnType SObjectNameEditableTextBox::UpdateHighlightSpringState( do
 		bUpdateHighlightSpring = false;
 	}
 
-	return bUpdateHighlightSpring ? EActiveTickReturnType::KeepTicking : EActiveTickReturnType::StopTicking;
+	return bUpdateHighlightSpring ? EActiveTimerReturnType::Continue : EActiveTimerReturnType::Stop;
 }
 
 void SObjectNameEditableTextBox::Tick( const FGeometry& AllottedGeometry, const double InCurrentTime, const float InDeltaTime )
@@ -256,7 +256,7 @@ void SObjectNameEditableTextBox::OnNameTextCommitted(const FText& NewText, EText
 					{
 						Actor->SetActorLabel(TrimmedText.ToString());
 						LastCommittedTime = FSlateApplication::Get().GetCurrentTime();
-						RegisterActiveTick( 0.f, FWidgetActiveTickDelegate::CreateSP( this, &SObjectNameEditableTextBox::UpdateHighlightSpringState ) );
+						RegisterActiveTimer( 0.f, FWidgetActiveTimerDelegate::CreateSP( this, &SObjectNameEditableTextBox::UpdateHighlightSpringState ) );
 					}
 				}
 			}
@@ -276,7 +276,7 @@ void SObjectNameEditableTextBox::OnNameTextCommitted(const FText& NewText, EText
 						{
 							Actor->SetActorLabel(TrimmedText.ToString());
 							LastCommittedTime = FSlateApplication::Get().GetCurrentTime();
-							RegisterActiveTick( 0.f, FWidgetActiveTickDelegate::CreateSP( this, &SObjectNameEditableTextBox::UpdateHighlightSpringState ) );
+							RegisterActiveTimer( 0.f, FWidgetActiveTimerDelegate::CreateSP( this, &SObjectNameEditableTextBox::UpdateHighlightSpringState ) );
 						}
 					}
 				}
diff --git a/Engine/Source/Editor/EditorWidgets/Private/SObjectNameEditableTextBox.h b/Engine/Source/Editor/EditorWidgets/Private/SObjectNameEditableTextBox.h
--- a/Engine/Source/Editor/EditorWidgets/Private/SObjectNameEditableTextBox.h
+++ b/Engine/Source/Editor/EditorWidgets/Private/SObjectNameEditableTextBox.h
@@ -27,7 +27,7 @@ protected:
 
 private:
 	/** Updates whether the highlight spring effect is happening */
-	EActiveTickReturnType UpdateHighlightSpringState( double InCurrentTime, float InDeltaTime );
+	EActiveTimerReturnType UpdateHighlightSpringState( double InCurrentTime, float InDeltaTime );
 
 	/** Getter for the Text attribute of the editable text inside this widget */
 	FText GetNameText() const;
diff --git a/Engine/Source/Editor/EditorWidgets/Private/STransportControl.cpp b/Engine/Source/Editor/EditorWidgets/Private/STransportControl.cpp
--- a/Engine/Source/Editor/EditorWidgets/Private/STransportControl.cpp
+++ b/Engine/Source/Editor/EditorWidgets/Private/STransportControl.cpp
@@ -148,7 +148,7 @@ END_SLATE_FUNCTION_BUILD_OPTIMIZATION
 
 bool STransportControl::IsTickable() const
 {
-	// Only bother if an active tick delegate was provided
+	// Only bother if an active timer delegate was provided
 	return TransportControlArgs.OnTickPlayback.IsBound() && TransportControlArgs.OnGetPlaybackMode.IsBound();
 }
 
@@ -157,13 +157,13 @@ void STransportControl::Tick( float DeltaTime )
 	const auto PlaybackMode = TransportControlArgs.OnGetPlaybackMode.Execute();
 	const bool bIsPlaying = PlaybackMode == EPlaybackMode::PlayingForward || PlaybackMode == EPlaybackMode::PlayingReverse || PlaybackMode == EPlaybackMode::Recording;
 
-	if ( bIsPlaying && !ActiveTickHandle.IsValid() )
+	if ( bIsPlaying && !ActiveTimerHandle.IsValid() )
 	{
-		ActiveTickHandle = RegisterActiveTick( 0.f, FWidgetActiveTickDelegate::CreateSP( this, &STransportControl::TickPlayback ) );
+		ActiveTimerHandle = RegisterActiveTimer( 0.f, FWidgetActiveTimerDelegate::CreateSP( this, &STransportControl::TickPlayback ) );
 	}
-	else if ( !bIsPlaying && ActiveTickHandle.IsValid() )
+	else if ( !bIsPlaying && ActiveTimerHandle.IsValid() )
 	{
-		UnRegisterActiveTick( ActiveTickHandle.Pin().ToSharedRef() );
+		UnRegisterActiveTimer( ActiveTimerHandle.Pin().ToSharedRef() );
 	}
 }
 
@@ -223,10 +223,10 @@ FSlateColor STransportControl::GetLoopStatusColor() const
 	return FSlateColor(FLinearColor(1.f, 1.f, 1.f));
 }
 
-EActiveTickReturnType STransportControl::TickPlayback( double InCurrentTime, float InDeltaTime )
+EActiveTimerReturnType STransportControl::TickPlayback( double InCurrentTime, float InDeltaTime )
 {
 	TransportControlArgs.OnTickPlayback.Execute( InCurrentTime, InDeltaTime );
-	return EActiveTickReturnType::KeepTicking;
+	return EActiveTimerReturnType::Continue;
 }
 
 #undef LOCTEXT_NAMESPACE
diff --git a/Engine/Source/Editor/EditorWidgets/Private/STransportControl.h b/Engine/Source/Editor/EditorWidgets/Private/STransportControl.h
--- a/Engine/Source/Editor/EditorWidgets/Private/STransportControl.h
+++ b/Engine/Source/Editor/EditorWidgets/Private/STransportControl.h
@@ -36,16 +36,16 @@ private:
 	FSlateColor GetLoopStatusColor() const;
 
 	/** Executes the OnTickPlayback delegate */
-	EActiveTickReturnType TickPlayback( double InCurrentTime, float InDeltaTime );
+	EActiveTimerReturnType TickPlayback( double InCurrentTime, float InDeltaTime );
 
 	FReply OnToggleLooping();
 
 private:
-	/** The handle to the active tick */
-	TWeakPtr<FActiveTickHandle> ActiveTickHandle;
+	/** The handle to the active timer */
+	TWeakPtr<FActiveTimerHandle> ActiveTimerHandle;
 
-	/** Whether the active tick is currently registered */
-	bool bIsActiveTickRegistered;
+	/** Whether the active timer is currently registered */
+	bool bIsActiveTimerRegistered;
 
 	FTransportControlArgs TransportControlArgs;
 };
diff --git a/Engine/Source/Editor/GameProjectGeneration/Private/SGameProjectDialog.cpp b/Engine/Source/Editor/GameProjectGeneration/Private/SGameProjectDialog.cpp
--- a/Engine/Source/Editor/GameProjectGeneration/Private/SGameProjectDialog.cpp
+++ b/Engine/Source/Editor/GameProjectGeneration/Private/SGameProjectDialog.cpp
@@ -32,7 +32,7 @@ void SGameProjectDialog::Construct( const FArguments& InArgs )
 	}
 
 	FadeAnimation.AddCurve( 0.f, 0.5f, ECurveEaseFunction::QuadOut );
-	RegisterActiveTick( 0.f, FWidgetActiveTickDelegate::CreateSP( this, &SGameProjectDialog::TriggerFadeInPostConstruct ) );
+	RegisterActiveTimer( 0.f, FWidgetActiveTimerDelegate::CreateSP( this, &SGameProjectDialog::TriggerFadeInPostConstruct ) );
 
 	TSharedPtr<SWidget> Content;
 
@@ -189,14 +189,14 @@ void SGameProjectDialog::Construct( const FArguments& InArgs )
 	}
 }
 
-EActiveTickReturnType SGameProjectDialog::TriggerFadeInPostConstruct(double InCurrentTime, float InDeltaTime)
+EActiveTimerReturnType SGameProjectDialog::TriggerFadeInPostConstruct(double InCurrentTime, float InDeltaTime)
 {
 	// Play the intro fade in the first frame after the widget is created.
 	// We start it now instead of Construct because there is a lot of elapsed time between Construct and when we
 	// see the dialog and the beginning of the animation is cut off.
 	FadeAnimation.Play(this->AsShared());
 
-	return EActiveTickReturnType::StopTicking;
+	return EActiveTimerReturnType::Stop;
 }
 
 /* SGameProjectDialog implementation
diff --git a/Engine/Source/Editor/GameProjectGeneration/Private/SGameProjectDialog.h b/Engine/Source/Editor/GameProjectGeneration/Private/SGameProjectDialog.h
--- a/Engine/Source/Editor/GameProjectGeneration/Private/SGameProjectDialog.h
+++ b/Engine/Source/Editor/GameProjectGeneration/Private/SGameProjectDialog.h
@@ -49,7 +49,7 @@ protected:
 
 private:
 	/** Ensures the fade-in animation is played post-construct */
-	EActiveTickReturnType TriggerFadeInPostConstruct( double InCurrentTime, float InDeltaTime );
+	EActiveTimerReturnType TriggerFadeInPostConstruct( double InCurrentTime, float InDeltaTime );
 
 	// Callback for getting the color of the custom content area.
 	FLinearColor HandleCustomContentColorAndOpacity() const;
diff --git a/Engine/Source/Editor/GraphEditor/Private/SGraphEditorImpl.cpp b/Engine/Source/Editor/GraphEditor/Private/SGraphEditorImpl.cpp
--- a/Engine/Source/Editor/GraphEditor/Private/SGraphEditorImpl.cpp
+++ b/Engine/Source/Editor/GraphEditor/Private/SGraphEditorImpl.cpp
@@ -87,14 +87,14 @@ void SGraphEditorImpl::NotifyGraphChanged()
 
 void SGraphEditorImpl::OnGraphChanged(const FEdGraphEditAction& InAction)
 {
-	if ( !bIsActiveTickRegistered )
+	if ( !bIsActiveTimerRegistered )
 	{
 		// Remove the old user interface nodes
 		GraphPanel->PurgeVisualRepresentation();
 
 		// Trigger the refresh
-		bIsActiveTickRegistered = true;
-		RegisterActiveTick( 0.f, FWidgetActiveTickDelegate::CreateSP( this, &SGraphEditorImpl::TriggerRefresh ) );
+		bIsActiveTimerRegistered = true;
+		RegisterActiveTimer( 0.f, FWidgetActiveTimerDelegate::CreateSP( this, &SGraphEditorImpl::TriggerRefresh ) );
 	}
 }
 
@@ -158,7 +158,7 @@ void SGraphEditorImpl::Construct( const FArguments& InArgs )
 	OnNavigateHistoryForward = InArgs._OnNavigateHistoryForward;
 	OnNodeSpawnedByKeymap = InArgs._GraphEvents.OnNodeSpawnedByKeymap;
 
-	bIsActiveTickRegistered = false;
+	bIsActiveTimerRegistered = false;
 
 	// Make sure that the editor knows about what kinds
 	// of commands GraphEditor can do.
@@ -385,12 +385,12 @@ void SGraphEditorImpl::Tick( const FGeometry& AllottedGeometry, const double InC
 	}
 }
 
-EActiveTickReturnType SGraphEditorImpl::TriggerRefresh( double InCurrentTime, float InDeltaTime )
+EActiveTimerReturnType SGraphEditorImpl::TriggerRefresh( double InCurrentTime, float InDeltaTime )
 {
 	GraphPanel->Update();
 
-	bIsActiveTickRegistered = false;
-	return EActiveTickReturnType::StopTicking;
+	bIsActiveTimerRegistered = false;
+	return EActiveTimerReturnType::Stop;
 }
 
 void SGraphEditorImpl::OnClosedActionMenu()
diff --git a/Engine/Source/Editor/GraphEditor/Private/SGraphEditorImpl.h b/Engine/Source/Editor/GraphEditor/Private/SGraphEditorImpl.h
--- a/Engine/Source/Editor/GraphEditor/Private/SGraphEditorImpl.h
+++ b/Engine/Source/Editor/GraphEditor/Private/SGraphEditorImpl.h
@@ -156,8 +156,8 @@ protected:
 	// End of SGraphEditorInterface
 
 private:
-	/** One-off active tick to ensure the graph is refreshes as needed */
-	EActiveTickReturnType TriggerRefresh( double InCurrentTime, float InDeltaTime );
+	/** One-off active timer to ensure the graph is refreshes as needed */
+	EActiveTimerReturnType TriggerRefresh( double InCurrentTime, float InDeltaTime );
 
 	FText GetZoomText() const;
 
@@ -168,6 +168,6 @@ private:
 	bool IsLocked() const;
 
 private:
-	bool bIsActiveTickRegistered;
+	bool bIsActiveTimerRegistered;
 };
 
diff --git a/Engine/Source/Editor/GraphEditor/Private/SNodePanel.cpp b/Engine/Source/Editor/GraphEditor/Private/SNodePanel.cpp
--- a/Engine/Source/Editor/GraphEditor/Private/SNodePanel.cpp
+++ b/Engine/Source/Editor/GraphEditor/Private/SNodePanel.cpp
@@ -408,7 +408,7 @@ void SNodePanel::OnEndNodeInteraction(const TSharedRef<SNode>& InNodeToDrag)
 {
 }
 
-EActiveTickReturnType SNodePanel::HandleZoomToFit(double InCurrentTime, float InDeltaTime)
+EActiveTimerReturnType SNodePanel::HandleZoomToFit(double InCurrentTime, float InDeltaTime)
 {
 	const FVector2D DesiredViewCenter = ( ZoomTargetTopLeft + ZoomTargetBottomRight ) * 0.5f;
 	const bool bDoneScrolling = ScrollToLocation(CachedGeometry, DesiredViewCenter, bTeleportInsteadOfScrollingWhenZoomingToFit ? 1000.0f : InDeltaTime);
@@ -426,16 +426,16 @@ EActiveTickReturnType SNodePanel::HandleZoomToFit(double InCurrentTime, float In
 
 		DeferredMovementTargetObject = nullptr;
 
-		return EActiveTickReturnType::StopTicking;
+		return EActiveTimerReturnType::Stop;
 	}
 
-	return EActiveTickReturnType::KeepTicking;
+	return EActiveTimerReturnType::Continue;
 }
 
 void SNodePanel::Tick( const FGeometry& AllottedGeometry, const double InCurrentTime, const float InDeltaTime )
 {
 	CachedGeometry = AllottedGeometry;
-	bool bWasActiveTickRegisteredThisFrame = false;
+	bool bWasActiveTimerRegisteredThisFrame = false;
 
 	if(DeferredSelectionTargetObjects.Num() > 0)
 	{
@@ -464,10 +464,10 @@ void SNodePanel::Tick( const FGeometry& AllottedGeometry, const double InCurrent
 	else if(DeferredMovementTargetObject != nullptr && GetBoundsForNodes(true, ZoomTargetTopLeft, ZoomTargetBottomRight, ZoomPadding))
 	{
 		DeferredMovementTargetObject = nullptr;
-		if (!ActiveTickHandle.IsValid())
+		if (!ActiveTimerHandle.IsValid())
 		{
-			ActiveTickHandle = RegisterActiveTick(0.f, FWidgetActiveTickDelegate::CreateSP(this, &SNodePanel::HandleZoomToFit));
-			bWasActiveTickRegisteredThisFrame = true;
+			ActiveTimerHandle = RegisterActiveTimer(0.f, FWidgetActiveTimerDelegate::CreateSP(this, &SNodePanel::HandleZoomToFit));
+			bWasActiveTimerRegisteredThisFrame = true;
 		}
 	}
 	
@@ -479,10 +479,10 @@ void SNodePanel::Tick( const FGeometry& AllottedGeometry, const double InCurrent
 		if( GetBoundsForNodes(bDeferredZoomToSelection, ZoomTargetTopLeft, ZoomTargetBottomRight, ZoomPadding))
 		{
 			bDeferredZoomToSelection = false;
-			if (!ActiveTickHandle.IsValid())
+			if (!ActiveTimerHandle.IsValid())
 			{
-				ActiveTickHandle = RegisterActiveTick(0.f, FWidgetActiveTickDelegate::CreateSP(this, &SNodePanel::HandleZoomToFit));
-				bWasActiveTickRegisteredThisFrame = true;
+				ActiveTimerHandle = RegisterActiveTimer(0.f, FWidgetActiveTimerDelegate::CreateSP(this, &SNodePanel::HandleZoomToFit));
+				bWasActiveTimerRegisteredThisFrame = true;
 			}
 		}
 	}
@@ -1070,11 +1070,11 @@ void SNodePanel::SelectAndCenterObject(const UObject* ObjectToSelect, bool bCent
 		DeferredMovementTargetObject = ObjectToSelect;
 	}
 
-	auto PinnedActiveTickHandle = ActiveTickHandle.Pin();
-	if (PinnedActiveTickHandle.IsValid())
+	auto PinnedActiveTimerHandle = ActiveTimerHandle.Pin();
+	if (PinnedActiveTimerHandle.IsValid())
 	{
 		// Stop zooming for now and let the a new zoom target be established
-		UnRegisterActiveTick(PinnedActiveTickHandle.ToSharedRef());
+		UnRegisterActiveTimer(PinnedActiveTimerHandle.ToSharedRef());
 	}
 }
 
@@ -1136,10 +1136,10 @@ void SNodePanel::RestoreViewSettings(const FVector2D& InViewOffset, float InZoom
 		ZoomLevel = ZoomLevels->GetNearestZoomLevel(InZoomAmount);
 		bDeferredZoomToNodeExtents = false;
 
-		auto PinnedActiveTickHandle = ActiveTickHandle.Pin();
-		if (PinnedActiveTickHandle.IsValid())
+		auto PinnedActiveTimerHandle = ActiveTimerHandle.Pin();
+		if (PinnedActiveTimerHandle.IsValid())
 		{
-			UnRegisterActiveTick(PinnedActiveTickHandle.ToSharedRef());
+			UnRegisterActiveTimer(PinnedActiveTimerHandle.ToSharedRef());
 		}
 	}
 
@@ -1522,9 +1522,9 @@ void SNodePanel::ZoomToTarget(const FVector2D& TopLeft, const FVector2D& BottomR
 	ZoomTargetTopLeft = TopLeft;
 	ZoomTargetBottomRight = BottomRight;
 
-	if (!ActiveTickHandle.IsValid())
+	if (!ActiveTimerHandle.IsValid())
 	{
-		ActiveTickHandle = RegisterActiveTick(0.f, FWidgetActiveTickDelegate::CreateSP(this, &SNodePanel::HandleZoomToFit));
+		ActiveTimerHandle = RegisterActiveTimer(0.f, FWidgetActiveTimerDelegate::CreateSP(this, &SNodePanel::HandleZoomToFit));
 	}
 }
 
diff --git a/Engine/Source/Editor/GraphEditor/Public/SNodePanel.h b/Engine/Source/Editor/GraphEditor/Public/SNodePanel.h
--- a/Engine/Source/Editor/GraphEditor/Public/SNodePanel.h
+++ b/Engine/Source/Editor/GraphEditor/Public/SNodePanel.h
@@ -848,14 +848,14 @@ protected:
 	TSharedPtr<FScopedTransaction> ScopedTransactionPtr;
 
 private:
-	/** Active tick that handles deferred zooming until the target zoom is reached */
-	EActiveTickReturnType HandleZoomToFit(double InCurrentTime, float InDeltaTime);
+	/** Active timer that handles deferred zooming until the target zoom is reached */
+	EActiveTimerReturnType HandleZoomToFit(double InCurrentTime, float InDeltaTime);
 
 private:
-	/** The handle to the active tick */
-	TWeakPtr<FActiveTickHandle> ActiveTickHandle;
+	/** The handle to the active timer */
+	TWeakPtr<FActiveTimerHandle> ActiveTimerHandle;
 
-	/** Cached geometry for use within the active tick */
+	/** Cached geometry for use within the active timer */
 	FGeometry CachedGeometry;
 
 	/** Zoom target rectangle */
diff --git a/Engine/Source/Editor/IntroTutorials/Private/STutorialButton.cpp b/Engine/Source/Editor/IntroTutorials/Private/STutorialButton.cpp
--- a/Engine/Source/Editor/IntroTutorials/Private/STutorialButton.cpp
+++ b/Engine/Source/Editor/IntroTutorials/Private/STutorialButton.cpp
@@ -28,7 +28,7 @@ void STutorialButton::Construct(const FArguments& InArgs)
 	AlertStartTime = 0.0f;
 
 	PulseAnimation.AddCurve(0.0f, TutorialButtonConstants::PulseAnimationLength, ECurveEaseFunction::Linear);
-	RegisterActiveTick( 0.f, FWidgetActiveTickDelegate::CreateSP( this, &STutorialButton::OpenTutorialPostConstruct ) );
+	RegisterActiveTimer( 0.f, FWidgetActiveTimerDelegate::CreateSP( this, &STutorialButton::OpenTutorialPostConstruct ) );
 
 	ChildSlot
 	[
@@ -46,7 +46,7 @@ void STutorialButton::Construct(const FArguments& InArgs)
 	];
 }
 
-EActiveTickReturnType STutorialButton::OpenTutorialPostConstruct( double InCurrentTime, float InDeltaTime )
+EActiveTimerReturnType STutorialButton::OpenTutorialPostConstruct( double InCurrentTime, float InDeltaTime )
 {
 	// Begin playing the pulse animation on a loop
 	PulseAnimation.Play(this->AsShared(), true);
@@ -71,7 +71,7 @@ EActiveTickReturnType STutorialButton::OpenTutorialPostConstruct( double InCurre
 		TutorialTitle = CachedLaunchTutorial->Title;
 	}
 
-	return EActiveTickReturnType::StopTicking;
+	return EActiveTimerReturnType::Stop;
 }
 
 static void GetAnimationValues(float InAnimationProgress, float& OutAlphaFactor0, float& OutPulseFactor0, float& OutAlphaFactor1, float& OutPulseFactor1)
diff --git a/Engine/Source/Editor/IntroTutorials/Private/STutorialButton.h b/Engine/Source/Editor/IntroTutorials/Private/STutorialButton.h
--- a/Engine/Source/Editor/IntroTutorials/Private/STutorialButton.h
+++ b/Engine/Source/Editor/IntroTutorials/Private/STutorialButton.h
@@ -20,7 +20,7 @@ class STutorialButton : public SCompoundWidget
 	
 private:
 	/** Opens the tutorial post-construct */
-	EActiveTickReturnType OpenTutorialPostConstruct( double InCurrentTime, float InDeltaTime );
+	EActiveTimerReturnType OpenTutorialPostConstruct( double InCurrentTime, float InDeltaTime );
 
 	/** Handle clicking the tutorial button */
 	FReply HandleButtonClicked();
diff --git a/Engine/Source/Editor/IntroTutorials/Private/STutorialsBrowser.cpp b/Engine/Source/Editor/IntroTutorials/Private/STutorialsBrowser.cpp
--- a/Engine/Source/Editor/IntroTutorials/Private/STutorialsBrowser.cpp
+++ b/Engine/Source/Editor/IntroTutorials/Private/STutorialsBrowser.cpp
@@ -497,7 +497,7 @@ void STutorialsBrowser::Construct(const FArguments& InArgs)
 	FAssetRegistryModule& AssetRegistryModule = FModuleManager::LoadModuleChecked<FAssetRegistryModule>(TEXT("AssetRegistry"));
 	AssetRegistryModule.Get().OnAssetAdded().AddSP(this, &STutorialsBrowser::HandleAssetAdded);
 
-	RegisterActiveTick( TutorialBrowserConstants::RefreshTimerInterval, FWidgetActiveTickDelegate::CreateSP( this, &STutorialsBrowser::TriggerReloadTutorials ) );
+	RegisterActiveTimer( TutorialBrowserConstants::RefreshTimerInterval, FWidgetActiveTimerDelegate::CreateSP( this, &STutorialsBrowser::TriggerReloadTutorials ) );
 
 	ChildSlot
 	[
@@ -579,7 +579,7 @@ void STutorialsBrowser::Construct(const FArguments& InArgs)
 	RebuildCrumbs();
 }
 
-EActiveTickReturnType STutorialsBrowser::TriggerReloadTutorials( double InCurrentTime, float InDeltaTime )
+EActiveTimerReturnType STutorialsBrowser::TriggerReloadTutorials( double InCurrentTime, float InDeltaTime )
 {
 	if (bNeedsRefresh)
 	{
@@ -587,7 +587,7 @@ EActiveTickReturnType STutorialsBrowser::TriggerReloadTutorials( double InCurren
 		ReloadTutorials();
 	}
 	
-	return EActiveTickReturnType::KeepTicking;
+	return EActiveTimerReturnType::Continue;
 }
 
 void STutorialsBrowser::SetFilter(const FString& InFilter)
diff --git a/Engine/Source/Editor/IntroTutorials/Private/STutorialsBrowser.h b/Engine/Source/Editor/IntroTutorials/Private/STutorialsBrowser.h
--- a/Engine/Source/Editor/IntroTutorials/Private/STutorialsBrowser.h
+++ b/Engine/Source/Editor/IntroTutorials/Private/STutorialsBrowser.h
@@ -100,7 +100,7 @@ protected:
 private:
 
 	/** Triggers a reload of the tutorials */
-	EActiveTickReturnType TriggerReloadTutorials( double InCurrentTime, float InDeltaTime );
+	EActiveTimerReturnType TriggerReloadTutorials( double InCurrentTime, float InDeltaTime );
 
 	/** Root entry of the tutorials tree */
 	TSharedPtr<FTutorialListEntry_Category> RootEntry;
diff --git a/Engine/Source/Editor/Kismet/Private/FindInBlueprints.cpp b/Engine/Source/Editor/Kismet/Private/FindInBlueprints.cpp
--- a/Engine/Source/Editor/Kismet/Private/FindInBlueprints.cpp
+++ b/Engine/Source/Editor/Kismet/Private/FindInBlueprints.cpp
@@ -1092,7 +1092,7 @@ SFindInBlueprints::~SFindInBlueprints()
 	}
 }
 
-EActiveTickReturnType SFindInBlueprints::UpdateSearchResults( double InCurrentTime, float InDeltaTime )
+EActiveTimerReturnType SFindInBlueprints::UpdateSearchResults( double InCurrentTime, float InDeltaTime )
 {
 	if ( StreamSearch.IsValid() )
 	{
@@ -1127,7 +1127,7 @@ EActiveTickReturnType SFindInBlueprints::UpdateSearchResults( double InCurrentTi
 		}
 	}
 
-	return StreamSearch.IsValid() ? EActiveTickReturnType::KeepTicking : EActiveTickReturnType::StopTicking;
+	return StreamSearch.IsValid() ? EActiveTimerReturnType::Continue : EActiveTimerReturnType::Stop;
 }
 
 void SFindInBlueprints::RegisterCommands()
@@ -1302,8 +1302,8 @@ void SFindInBlueprints::LaunchStreamThread(const TArray<FString>& InTokens)
 	}
 	else
 	{
-		// If the stream search wasn't already running, register the active tick
-		RegisterActiveTick( 0.f, FWidgetActiveTickDelegate::CreateSP( this, &SFindInBlueprints::UpdateSearchResults ) );
+		// If the stream search wasn't already running, register the active timer
+		RegisterActiveTimer( 0.f, FWidgetActiveTimerDelegate::CreateSP( this, &SFindInBlueprints::UpdateSearchResults ) );
 	}
 
 	StreamSearch = MakeShareable(new FStreamSearch(InTokens));
diff --git a/Engine/Source/Editor/Kismet/Private/SBlueprintSubPalette.cpp b/Engine/Source/Editor/Kismet/Private/SBlueprintSubPalette.cpp
--- a/Engine/Source/Editor/Kismet/Private/SBlueprintSubPalette.cpp
+++ b/Engine/Source/Editor/Kismet/Private/SBlueprintSubPalette.cpp
@@ -256,10 +256,10 @@ void SBlueprintSubPalette::Construct(FArguments const& InArgs, TWeakPtr<FBluepri
 }
 
 //------------------------------------------------------------------------------
-EActiveTickReturnType SBlueprintSubPalette::TriggerRefreshActionsList(double InCurrentTime, float InDeltaTime)
+EActiveTimerReturnType SBlueprintSubPalette::TriggerRefreshActionsList(double InCurrentTime, float InDeltaTime)
 {
 	RefreshActionsList(true);
-	return EActiveTickReturnType::StopTicking;
+	return EActiveTimerReturnType::Stop;
 }
 
 //------------------------------------------------------------------------------
@@ -390,7 +390,7 @@ void SBlueprintSubPalette::GenerateContextMenuEntries(FMenuBuilder& MenuBuilder)
 //------------------------------------------------------------------------------
 void SBlueprintSubPalette::RequestRefreshActionsList()
 {
-	RegisterActiveTick(0.f, FWidgetActiveTickDelegate::CreateSP(this, &SBlueprintSubPalette::TriggerRefreshActionsList));
+	RegisterActiveTimer(0.f, FWidgetActiveTimerDelegate::CreateSP(this, &SBlueprintSubPalette::TriggerRefreshActionsList));
 }
 
 //------------------------------------------------------------------------------
diff --git a/Engine/Source/Editor/Kismet/Private/SBlueprintSubPalette.h b/Engine/Source/Editor/Kismet/Private/SBlueprintSubPalette.h
--- a/Engine/Source/Editor/Kismet/Private/SBlueprintSubPalette.h
+++ b/Engine/Source/Editor/Kismet/Private/SBlueprintSubPalette.h
@@ -97,8 +97,8 @@ private:
 	/** Immediately calls RefreshActionsList(), doesn't not defer until Tick() like RequestRefreshActionsList() does. */
 	void ForceRefreshActionList();
 
-	/** One-off active tick to trigger a refresh of the action list */
-	EActiveTickReturnType TriggerRefreshActionsList(double InCurrentTime, float InDeltaTime);
+	/** One-off active timer to trigger a refresh of the action list */
+	EActiveTimerReturnType TriggerRefreshActionsList(double InCurrentTime, float InDeltaTime);
 
 	/**
 	 * Constructs a slate widget for the right-click context menu in this 
diff --git a/Engine/Source/Editor/Kismet/Private/SSCSEditorViewport.cpp b/Engine/Source/Editor/Kismet/Private/SSCSEditorViewport.cpp
--- a/Engine/Source/Editor/Kismet/Private/SSCSEditorViewport.cpp
+++ b/Engine/Source/Editor/Kismet/Private/SSCSEditorViewport.cpp
@@ -187,7 +187,7 @@ private:
 
 void SSCSEditorViewport::Construct(const FArguments& InArgs)
 {
-	bIsActiveTickRegistered = false;
+	bIsActiveTimerRegistered = false;
 
 	// Save off the Blueprint editor reference, we'll need this later
 	BlueprintEditorPtr = InArgs._BlueprintEditor;
@@ -304,10 +304,10 @@ void SSCSEditorViewport::RequestRefresh(bool bResetCamera, bool bRefreshNow)
 	else
 	{
 		// Defer the update until the next tick. This way we don't accidentally spawn the preview actor in the middle of a transaction, for example.
-		if (!bIsActiveTickRegistered)
+		if (!bIsActiveTimerRegistered)
 		{
-			bIsActiveTickRegistered = true;
-			RegisterActiveTick(0.f, FWidgetActiveTickDelegate::CreateSP(this, &SSCSEditorViewport::DeferredUpdatePreview, bResetCamera));
+			bIsActiveTimerRegistered = true;
+			RegisterActiveTimer(0.f, FWidgetActiveTimerDelegate::CreateSP(this, &SSCSEditorViewport::DeferredUpdatePreview, bResetCamera));
 		}
 	}
 }
@@ -328,15 +328,15 @@ bool SSCSEditorViewport::GetIsSimulateEnabled()
 	return ViewportClient->GetIsSimulateEnabled();
 }
 
-EActiveTickReturnType SSCSEditorViewport::DeferredUpdatePreview(double InCurrentTime, float InDeltaTime, bool bResetCamera)
+EActiveTimerReturnType SSCSEditorViewport::DeferredUpdatePreview(double InCurrentTime, float InDeltaTime, bool bResetCamera)
 {
 	if (ViewportClient.IsValid())
 	{
 		ViewportClient->InvalidatePreview(bResetCamera);
 	}
 
-	bIsActiveTickRegistered = false;
-	return EActiveTickReturnType::StopTicking;
+	bIsActiveTimerRegistered = false;
+	return EActiveTimerReturnType::Stop;
 }
 
 AActor* SSCSEditorViewport::GetPreviewActor() const
diff --git a/Engine/Source/Editor/Kismet/Private/SSCSEditorViewport.h b/Engine/Source/Editor/Kismet/Private/SSCSEditorViewport.h
--- a/Engine/Source/Editor/Kismet/Private/SSCSEditorViewport.h
+++ b/Engine/Source/Editor/Kismet/Private/SSCSEditorViewport.h
@@ -90,8 +90,8 @@ protected:
 	virtual void BindCommands();
 
 private:
-	/** One-off active tick to update the preview */
-	EActiveTickReturnType DeferredUpdatePreview(double InCurrentTime, float InDeltaTime, bool bResetCamera);
+	/** One-off active timer to update the preview */
+	EActiveTimerReturnType DeferredUpdatePreview(double InCurrentTime, float InDeltaTime, bool bResetCamera);
 
 private:
 	/** Pointer back to editor tool (owner) */
@@ -103,6 +103,6 @@ private:
 	/** Viewport client */
 	TSharedPtr<class FSCSEditorViewportClient> ViewportClient;
 
-	/** Whether the active tick (for updating the preview) is registered */
-	bool bIsActiveTickRegistered;
+	/** Whether the active timer (for updating the preview) is registered */
+	bool bIsActiveTimerRegistered;
 };
diff --git a/Engine/Source/Editor/Kismet/Public/FindInBlueprints.h b/Engine/Source/Editor/Kismet/Public/FindInBlueprints.h
--- a/Engine/Source/Editor/Kismet/Public/FindInBlueprints.h
+++ b/Engine/Source/Editor/Kismet/Public/FindInBlueprints.h
@@ -232,7 +232,7 @@ public:
 
 private:
 	/** Processes results of the ongoing async stream search */
-	EActiveTickReturnType UpdateSearchResults( double InCurrentTime, float InDeltaTime );
+	EActiveTimerReturnType UpdateSearchResults( double InCurrentTime, float InDeltaTime );
 
 	/** Register any Find-in-Blueprint commands */
 	void RegisterCommands();
diff --git a/Engine/Source/Editor/LevelEditor/Private/LevelViewportLayout.cpp b/Engine/Source/Editor/LevelEditor/Private/LevelViewportLayout.cpp
--- a/Engine/Source/Editor/LevelEditor/Private/LevelViewportLayout.cpp
+++ b/Engine/Source/Editor/LevelEditor/Private/LevelViewportLayout.cpp
@@ -72,7 +72,7 @@ public:
 private:
 	
 	/** Ticks the level viewport tab */
-	EActiveTickReturnType ActiveTick( double InCurrentTime, float InDeltaTime );
+	EActiveTimerReturnType ActiveTimer( double InCurrentTime, float InDeltaTime );
 
 	/** Reference to the owning level viewport tab */
 	TSharedPtr<FLevelViewportTabContent> LevelViewportTab;
@@ -90,7 +90,7 @@ void SViewportsOverlay::Construct( const FArguments& InArgs )
 	const TSharedRef<SWidget>& ContentWidget = InArgs._Content.Widget;
 	LevelViewportTab = InArgs._LevelViewportTab;
 
-	//RegisterActiveTick( 0.f, FTickWidgetDelegate::CreateSP( this, &SViewportsOverlay::ActiveTick ) );
+	//RegisterActiveTimer( 0.f, FTickWidgetDelegate::CreateSP( this, &SViewportsOverlay::ActiveTimer ) );
 
 	ChildSlot
 		[
@@ -102,10 +102,10 @@ void SViewportsOverlay::Construct( const FArguments& InArgs )
 		];
 }
 
-EActiveTickReturnType SViewportsOverlay::ActiveTick( double InCurrentTime, float InDeltaTime )
+EActiveTimerReturnType SViewportsOverlay::ActiveTimer( double InCurrentTime, float InDeltaTime )
 {
 	// Exists to ensure slate is ticked
-	return EActiveTickReturnType::KeepTicking;
+	return EActiveTimerReturnType::Continue;
 }
 
 void SViewportsOverlay::Tick( const FGeometry& AllottedGeometry, const double InCurrentTime, const float InDeltaTime )
diff --git a/Engine/Source/Editor/PListEditor/Private/SPlistEditor.cpp b/Engine/Source/Editor/PListEditor/Private/SPlistEditor.cpp
--- a/Engine/Source/Editor/PListEditor/Private/SPlistEditor.cpp
+++ b/Engine/Source/Editor/PListEditor/Private/SPlistEditor.cpp
@@ -2070,7 +2070,7 @@ void SPListEditorPanel::OnPopupTextChosen(const FString& ChosenText)
 	FSlateApplication::Get().DismissAllMenus();
 }
 
-EActiveTickReturnType SPListEditorPanel::DisplayDeferredNotifications( double InCurrentTime, float InDeltaTime )
+EActiveTimerReturnType SPListEditorPanel::DisplayDeferredNotifications( double InCurrentTime, float InDeltaTime )
 {
 	if ( --FramesToSkip == 0 )
 	{
@@ -2096,19 +2096,19 @@ EActiveTickReturnType SPListEditorPanel::DisplayDeferredNotifications( double In
 
 		QueuedNotifications.Empty();
 
-		return EActiveTickReturnType::StopTicking;
+		return EActiveTimerReturnType::Stop;
 	}
 
-	return EActiveTickReturnType::KeepTicking;
+	return EActiveTimerReturnType::Continue;
 }
 
 /** Helper function to display notifications in the current tab */
 void SPListEditorPanel::DisplayNotification(const FText& ToDisplay, const ENTF_Types NotificationType)
 {
-	// Register the active tick if it isn't already
+	// Register the active timer if it isn't already
 	if ( FramesToSkip == 0 )
 	{
-		RegisterActiveTick( 0.f, FWidgetActiveTickDelegate::CreateSP( this, &SPListEditorPanel::DisplayDeferredNotifications ) );
+		RegisterActiveTimer( 0.f, FWidgetActiveTimerDelegate::CreateSP( this, &SPListEditorPanel::DisplayDeferredNotifications ) );
 	}
 
 	QueuedNotifications.Add( FQueuedNotification( ToDisplay, NotificationType ) );
diff --git a/Engine/Source/Editor/PListEditor/Public/SPlistEditor.h b/Engine/Source/Editor/PListEditor/Public/SPlistEditor.h
--- a/Engine/Source/Editor/PListEditor/Public/SPlistEditor.h
+++ b/Engine/Source/Editor/PListEditor/Public/SPlistEditor.h
@@ -233,7 +233,7 @@ public:
 
 private:
 	/** Displays queued notifications */
-	EActiveTickReturnType DisplayDeferredNotifications( double InCurrentTime, float InDeltaTime );
+	EActiveTimerReturnType DisplayDeferredNotifications( double InCurrentTime, float InDeltaTime );
 	/** Helper method to display queued notifications */
 	void PerformDisplayNotifications();
 
diff --git a/Engine/Source/Editor/Persona/Private/SAnimCompositeEditor.cpp b/Engine/Source/Editor/Persona/Private/SAnimCompositeEditor.cpp
--- a/Engine/Source/Editor/Persona/Private/SAnimCompositeEditor.cpp
+++ b/Engine/Source/Editor/Persona/Private/SAnimCompositeEditor.cpp
@@ -22,7 +22,7 @@ TSharedRef<SWidget> SAnimCompositeEditor::CreateDocumentAnchor()
 
 void SAnimCompositeEditor::Construct(const FArguments& InArgs)
 {
-	bIsActiveTickRegistered = false;
+	bIsActiveTimerRegistered = false;
 	PersonaPtr = InArgs._Persona;
 	CompositeObj = InArgs._Composite;
 	check(CompositeObj);
@@ -112,10 +112,10 @@ void SAnimCompositeEditor::OnCompositeChange(class UObject *EditorAnimBaseObj, b
 {
 	if ( CompositeObj != nullptr )
 	{
-		if(bRebuild && !bIsActiveTickRegistered)
+		if(bRebuild && !bIsActiveTimerRegistered)
 		{
-			bIsActiveTickRegistered = true;
-			RegisterActiveTick(0.f, FWidgetActiveTickDelegate::CreateSP(this, &SAnimCompositeEditor::TriggerRebuildPanel));
+			bIsActiveTimerRegistered = true;
+			RegisterActiveTimer(0.f, FWidgetActiveTimerDelegate::CreateSP(this, &SAnimCompositeEditor::TriggerRebuildPanel));
 		} 
 		else
 		{
@@ -140,10 +140,10 @@ void SAnimCompositeEditor::CollapseComposite()
 
 void SAnimCompositeEditor::PostUndo()
 {
-	if (!bIsActiveTickRegistered)
+	if (!bIsActiveTimerRegistered)
 	{
-		bIsActiveTickRegistered = true;
-		RegisterActiveTick(0.f, FWidgetActiveTickDelegate::CreateSP(this, &SAnimCompositeEditor::TriggerRebuildPanel));
+		bIsActiveTimerRegistered = true;
+		RegisterActiveTimer(0.f, FWidgetActiveTimerDelegate::CreateSP(this, &SAnimCompositeEditor::TriggerRebuildPanel));
 	}
 
 	// when undo or redo happens, we still have to recalculate length, so we can't rely on sequence length changes or not
@@ -158,14 +158,14 @@ void SAnimCompositeEditor::InitDetailsViewEditorObject(UEditorAnimBaseObj* EdObj
 	EdObj->InitFromAnim(CompositeObj, FOnAnimObjectChange::CreateSP( SharedThis(this), &SAnimCompositeEditor::OnCompositeChange ));
 }
 
-EActiveTickReturnType SAnimCompositeEditor::TriggerRebuildPanel(double InCurrentTime, float InDeltaTime)
+EActiveTimerReturnType SAnimCompositeEditor::TriggerRebuildPanel(double InCurrentTime, float InDeltaTime)
 {
 	// we should not update any property related within PostEditChange, 
 	// so this is deferred to Tick, when it needs to rebuild, just mark it and this will update in the next tick
 	RebuildPanel();
 
-	bIsActiveTickRegistered = false;
-	return EActiveTickReturnType::StopTicking;
+	bIsActiveTimerRegistered = false;
+	return EActiveTimerReturnType::Stop;
 }
 
 float SAnimCompositeEditor::CalculateSequenceLengthOfEditorObject() const
diff --git a/Engine/Source/Editor/Persona/Private/SAnimCompositeEditor.h b/Engine/Source/Editor/Persona/Private/SAnimCompositeEditor.h
--- a/Engine/Source/Editor/Persona/Private/SAnimCompositeEditor.h
+++ b/Engine/Source/Editor/Persona/Private/SAnimCompositeEditor.h
@@ -69,11 +69,11 @@ private:
 	/** This will sort all components of the montage and update (recreate) the UI */
 	void SortAndUpdateComposite();
 
-	/** One-off active tick to trigger a panel rebuild */
-	EActiveTickReturnType TriggerRebuildPanel(double InCurrentTime, float InDeltaTime);
+	/** One-off active timer to trigger a panel rebuild */
+	EActiveTimerReturnType TriggerRebuildPanel(double InCurrentTime, float InDeltaTime);
 
-	/** Whether the active tick to trigger a panel rebuild is currently registered */
-	bool bIsActiveTickRegistered;
+	/** Whether the active timer to trigger a panel rebuild is currently registered */
+	bool bIsActiveTimerRegistered;
 
 public:
 
diff --git a/Engine/Source/Editor/Persona/Private/SAnimationBlendSpaceBase.cpp b/Engine/Source/Editor/Persona/Private/SAnimationBlendSpaceBase.cpp
--- a/Engine/Source/Editor/Persona/Private/SAnimationBlendSpaceBase.cpp
+++ b/Engine/Source/Editor/Persona/Private/SAnimationBlendSpaceBase.cpp
@@ -1180,23 +1180,23 @@ SBlendSpaceEditorBase::~SBlendSpaceEditorBase()
 
 void SBlendSpaceEditorBase::Construct(const FArguments& InArgs)
 {
-	bIsActiveTickRegistered = false;
+	bIsActiveTimerRegistered = false;
 	BlendSpace = InArgs._BlendSpace;
 	PersonaPtr = InArgs._Persona;
 	PersonaPtr.Pin()->RegisterOnPostUndo(FPersona::FOnPostUndo::CreateSP( this, &SBlendSpaceEditorBase::PostUndo ) );
 }
 
-EActiveTickReturnType SBlendSpaceEditorBase::UpdatePreview( double InCurrentTime, float InDeltaTime )
+EActiveTimerReturnType SBlendSpaceEditorBase::UpdatePreview( double InCurrentTime, float InDeltaTime )
 {
 	// Update the preview as long as its enabled
 	if (BlendSpaceWidget->bPreviewOn)
 	{
 		UpdatePreviewParameter();
-		return EActiveTickReturnType::KeepTicking;
+		return EActiveTimerReturnType::Continue;
 	}
 
-	bIsActiveTickRegistered = false;
-	return EActiveTickReturnType::StopTicking;
+	bIsActiveTimerRegistered = false;
+	return EActiveTimerReturnType::Stop;
 }
 
 void SBlendSpaceEditorBase::PostUndo()
@@ -1294,10 +1294,10 @@ void SBlendSpaceEditorBase::ShowPreview_OnIsCheckedChanged( ECheckBoxState NewVa
 
 		BlendSpaceWidget->bPreviewOn = bPreviewEnabled;
 
-		if ( bPreviewEnabled && !bIsActiveTickRegistered )
+		if ( bPreviewEnabled && !bIsActiveTimerRegistered )
 		{
-			bIsActiveTickRegistered = true;
-			RegisterActiveTick( 0.f, FWidgetActiveTickDelegate::CreateSP( this, &SBlendSpaceEditorBase::UpdatePreview ) );
+			bIsActiveTimerRegistered = true;
+			RegisterActiveTimer( 0.f, FWidgetActiveTimerDelegate::CreateSP( this, &SBlendSpaceEditorBase::UpdatePreview ) );
 		}
 	}
 }
diff --git a/Engine/Source/Editor/Persona/Private/SAnimationBlendSpaceBase.h b/Engine/Source/Editor/Persona/Private/SAnimationBlendSpaceBase.h
--- a/Engine/Source/Editor/Persona/Private/SAnimationBlendSpaceBase.h
+++ b/Engine/Source/Editor/Persona/Private/SAnimationBlendSpaceBase.h
@@ -361,7 +361,7 @@ public:
 
 protected:
 	/** Updates the preview */
-	EActiveTickReturnType UpdatePreview( double InCurrentTime, float InDeltaTime );
+	EActiveTimerReturnType UpdatePreview( double InCurrentTime, float InDeltaTime );
 
 	/** Creates the editor heading UI */
 	TSharedRef<SWidget> MakeEditorHeader() const;
@@ -412,7 +412,7 @@ protected:
 	/** This is used to tell the user when they do something wrong */
 	TSharedPtr<SNotificationList> NotificationListPtr;
 
-	/** Whether the active tick to update the preview is registered */
-	bool bIsActiveTickRegistered;
+	/** Whether the active timer to update the preview is registered */
+	bool bIsActiveTimerRegistered;
 };
 
diff --git a/Engine/Source/Editor/Persona/Private/SAnimationSequenceBrowser.cpp b/Engine/Source/Editor/Persona/Private/SAnimationSequenceBrowser.cpp
--- a/Engine/Source/Editor/Persona/Private/SAnimationSequenceBrowser.cpp
+++ b/Engine/Source/Editor/Persona/Private/SAnimationSequenceBrowser.cpp
@@ -401,7 +401,7 @@ void SAnimationSequenceBrowser::Construct(const FArguments& InArgs)
 	CurrentAssetHistoryIndex = INDEX_NONE;
 	bTriedToCacheOrginalAsset = false;
 
-	bIsActiveTickRegistered = false;
+	bIsActiveTimerRegistered = false;
 	bToolTipVisualizedThisFrame = false;
 	bToolTipClosedThisFrame = false;
 
@@ -920,10 +920,10 @@ bool SAnimationSequenceBrowser::OnVisualizeAssetToolTip(const TSharedPtr<SWidget
 			ViewportWidget->SetVisibility(EVisibility::Visible);
 			
 			// Update the preview as long as the tooltip is visible
-			if ( !bIsActiveTickRegistered )
+			if ( !bIsActiveTimerRegistered )
 			{
-				bIsActiveTickRegistered = true;
-				RegisterActiveTick(0.f, FWidgetActiveTickDelegate::CreateSP(this, &SAnimationSequenceBrowser::UpdateTootipPreview));
+				bIsActiveTimerRegistered = true;
+				RegisterActiveTimer(0.f, FWidgetActiveTimerDelegate::CreateSP(this, &SAnimationSequenceBrowser::UpdateTootipPreview));
 			}
 			bToolTipVisualizedThisFrame = true;
 		}
@@ -961,7 +961,7 @@ void SAnimationSequenceBrowser::CleanupPreviewSceneComponent(USceneComponent* Co
 	}
 }
 
-EActiveTickReturnType SAnimationSequenceBrowser::UpdateTootipPreview( double InCurrentTime, float InDeltaTime )
+EActiveTimerReturnType SAnimationSequenceBrowser::UpdateTootipPreview( double InCurrentTime, float InDeltaTime )
 {
 	bToolTipVisualizedThisFrame = false;
 	if ( PreviewComponent && IsToolTipPreviewVisible() )
@@ -971,11 +971,11 @@ EActiveTickReturnType SAnimationSequenceBrowser::UpdateTootipPreview( double InC
 	}
 	else
 	{
-		bIsActiveTickRegistered = false;
-		return EActiveTickReturnType::StopTicking;
+		bIsActiveTimerRegistered = false;
+		return EActiveTimerReturnType::Stop;
 	}
 
-	return EActiveTickReturnType::KeepTicking;
+	return EActiveTimerReturnType::Continue;
 }
 
 bool SAnimationSequenceBrowser::IsToolTipPreviewVisible()
diff --git a/Engine/Source/Editor/Persona/Private/SAnimationSequenceBrowser.h b/Engine/Source/Editor/Persona/Private/SAnimationSequenceBrowser.h
--- a/Engine/Source/Editor/Persona/Private/SAnimationSequenceBrowser.h
+++ b/Engine/Source/Editor/Persona/Private/SAnimationSequenceBrowser.h
@@ -201,14 +201,14 @@ protected:
 
 private:
 	/** Updates the animation preview in the tooltip */
-	EActiveTickReturnType UpdateTootipPreview( double InCurrentTime, float InDeltaTime );
+	EActiveTimerReturnType UpdateTootipPreview( double InCurrentTime, float InDeltaTime );
 
 	/** Whether the tooltip was visualized this frame */
 	uint8 bToolTipClosedThisFrame : 1;
 
-	/** Whether the active tick is currently registered */
-	uint8 bIsActiveTickRegistered : 1;
+	/** Whether the active timer is currently registered */
+	uint8 bIsActiveTimerRegistered : 1;
 
-	/** Whether the active tick should stop */
+	/** Whether the active timer should stop */
 	uint8 bToolTipVisualizedThisFrame : 1;
 };
\ No newline at end of file
diff --git a/Engine/Source/Editor/Persona/Private/SMontageEditor.cpp b/Engine/Source/Editor/Persona/Private/SMontageEditor.cpp
--- a/Engine/Source/Editor/Persona/Private/SMontageEditor.cpp
+++ b/Engine/Source/Editor/Persona/Private/SMontageEditor.cpp
@@ -31,7 +31,7 @@ void SMontageEditor::Construct(const FArguments& InArgs)
 	check(MontageObj);
 
 	bDragging = false;
-	bIsActiveTickRegistered = false;
+	bIsActiveTimerRegistered = false;
 
 	SAnimEditorBase::Construct( SAnimEditorBase::FArguments()
 		.Persona(InArgs._Persona)
@@ -299,12 +299,12 @@ void SMontageEditor::RebuildMontagePanel()
 	AnimMontageSectionsPanel->Update();
 }
 
-EActiveTickReturnType SMontageEditor::TriggerRebuildMontagePanel(double InCurrentTime, float InDeltaTime)
+EActiveTimerReturnType SMontageEditor::TriggerRebuildMontagePanel(double InCurrentTime, float InDeltaTime)
 {
 	RebuildMontagePanel();
 
-	bIsActiveTickRegistered = false;
-	return EActiveTickReturnType::StopTicking;
+	bIsActiveTimerRegistered = false;
+	return EActiveTimerReturnType::Stop;
 }
 
 void SMontageEditor::OnMontageChange(class UObject *EditorAnimBaseObj, bool Rebuild)
@@ -313,10 +313,10 @@ void SMontageEditor::OnMontageChange(class UObject *EditorAnimBaseObj, bool Rebu
 
 	if ( MontageObj != nullptr )
 	{
-		if(Rebuild && !bIsActiveTickRegistered)
+		if(Rebuild && !bIsActiveTimerRegistered)
 		{
-			bIsActiveTickRegistered = true;
-			RegisterActiveTick(0.f, FWidgetActiveTickDelegate::CreateSP(this, &SMontageEditor::TriggerRebuildMontagePanel));
+			bIsActiveTimerRegistered = true;
+			RegisterActiveTimer(0.f, FWidgetActiveTimerDelegate::CreateSP(this, &SMontageEditor::TriggerRebuildMontagePanel));
 		} 
 		else
 		{
diff --git a/Engine/Source/Editor/Persona/Private/SMontageEditor.h b/Engine/Source/Editor/Persona/Private/SMontageEditor.h
--- a/Engine/Source/Editor/Persona/Private/SMontageEditor.h
+++ b/Engine/Source/Editor/Persona/Private/SMontageEditor.h
@@ -71,8 +71,8 @@ private:
 	/** If user is currently dragging an item */
 	bool bDragging : 1;
 
-	/** If the active tick to trigger a montage panel rebuild is currently registered */
-	bool bIsActiveTickRegistered : 1;
+	/** If the active timer to trigger a montage panel rebuild is currently registered */
+	bool bIsActiveTimerRegistered : 1;
 
 	virtual float CalculateSequenceLengthOfEditorObject() const override;
 	void SortAndUpdateMontage();
@@ -92,8 +92,8 @@ private:
 	/** Updates Notify trigger offsets to take into account current montage state */
 	void RefreshNotifyTriggerOffsets();
 
-	/** One-off active tick to trigger a montage panel rebuild */
-	EActiveTickReturnType TriggerRebuildMontagePanel(double InCurrentTime, float InDeltaTime);
+	/** One-off active timer to trigger a montage panel rebuild */
+	EActiveTimerReturnType TriggerRebuildMontagePanel(double InCurrentTime, float InDeltaTime);
 
 	/** Rebuilds the montage panel */
 	void RebuildMontagePanel();
diff --git a/Engine/Source/Editor/PropertyEditor/Private/UserInterface/PropertyTable/SPropertyTableCell.cpp b/Engine/Source/Editor/PropertyEditor/Private/UserInterface/PropertyTable/SPropertyTableCell.cpp
--- a/Engine/Source/Editor/PropertyEditor/Private/UserInterface/PropertyTable/SPropertyTableCell.cpp
+++ b/Engine/Source/Editor/PropertyEditor/Private/UserInterface/PropertyTable/SPropertyTableCell.cpp
@@ -92,7 +92,7 @@ void SPropertyTableCell::EnteredEditMode()
 	if (Cell->IsValid())
 	{
 		// We delay the activation of editing mode till Tick due to mouse related input replies stomping on the focus
-		RegisterActiveTick(0.f, FWidgetActiveTickDelegate::CreateSP(this, &SPropertyTableCell::TriggerEnterEditingMode));
+		RegisterActiveTimer(0.f, FWidgetActiveTimerDelegate::CreateSP(this, &SPropertyTableCell::TriggerEnterEditingMode));
 	}
 }
 
@@ -110,7 +110,7 @@ void SPropertyTableCell::ExitedEditMode()
 	}
 }
 
-EActiveTickReturnType SPropertyTableCell::TriggerEnterEditingMode(double InCurrentTime, float InDeltaTime)
+EActiveTimerReturnType SPropertyTableCell::TriggerEnterEditingMode(double InCurrentTime, float InDeltaTime)
 {
 	if (Cell->GetTable()->GetCurrentCell() == Cell)
 	{
@@ -131,7 +131,7 @@ EActiveTickReturnType SPropertyTableCell::TriggerEnterEditingMode(double InCurre
 		}
 	}
 
-	return EActiveTickReturnType::StopTicking;
+	return EActiveTimerReturnType::Stop;
 }
 
 int32 SPropertyTableCell::OnPaint( const FPaintArgs& Args, const FGeometry& AllottedGeometry, const FSlateRect& MyClippingRect, FSlateWindowElementList& OutDrawElements, int32 LayerId, const FWidgetStyle& InWidgetStyle, bool bParentEnabled ) const
diff --git a/Engine/Source/Editor/PropertyEditor/Private/UserInterface/PropertyTable/SPropertyTableCell.h b/Engine/Source/Editor/PropertyEditor/Private/UserInterface/PropertyTable/SPropertyTableCell.h
--- a/Engine/Source/Editor/PropertyEditor/Private/UserInterface/PropertyTable/SPropertyTableCell.h
+++ b/Engine/Source/Editor/PropertyEditor/Private/UserInterface/PropertyTable/SPropertyTableCell.h
@@ -61,8 +61,8 @@ private:
 
 	void OnCellValueChanged( UObject* Object, FPropertyChangedEvent& PropertyChangedEvent );
 
-	/** One-off active tick to trigger entering the editing mode */
-	EActiveTickReturnType TriggerEnterEditingMode(double InCurrentTime, float InDeltaTime);
+	/** One-off active timer to trigger entering the editing mode */
+	EActiveTimerReturnType TriggerEnterEditingMode(double InCurrentTime, float InDeltaTime);
 
 private:
 
diff --git a/Engine/Source/Editor/Sequencer/Private/SSequencer.cpp b/Engine/Source/Editor/Sequencer/Private/SSequencer.cpp
--- a/Engine/Source/Editor/Sequencer/Private/SSequencer.cpp
+++ b/Engine/Source/Editor/Sequencer/Private/SSequencer.cpp
@@ -143,7 +143,7 @@ private:
 void SSequencer::Construct( const FArguments& InArgs, TSharedRef< class FSequencer > InSequencer )
 {
 	Sequencer = InSequencer;
-	bIsActiveTickRegistered = false;
+	bIsActiveTimerRegistered = false;
 
 	// Create a node tree which contains a tree of movie scene data to display in the sequence
 	SequencerNodeTree = MakeShareable( new FSequencerNodeTree( InSequencer.Get() ) );
@@ -374,28 +374,28 @@ SSequencer::~SSequencer()
 {
 }
 
-void SSequencer::RegisterActiveTickForPlayback()
+void SSequencer::RegisterActiveTimerForPlayback()
 {
-	if (!bIsActiveTickRegistered)
+	if (!bIsActiveTimerRegistered)
 	{
-		bIsActiveTickRegistered = true;
-		RegisterActiveTick(0.f, FWidgetActiveTickDelegate::CreateSP(this, &SSequencer::EnsureSlateTickDuringPlayback));
+		bIsActiveTimerRegistered = true;
+		RegisterActiveTimer(0.f, FWidgetActiveTimerDelegate::CreateSP(this, &SSequencer::EnsureSlateTickDuringPlayback));
 	}
 }
 
-EActiveTickReturnType SSequencer::EnsureSlateTickDuringPlayback(double InCurrentTime, float InDeltaTime)
+EActiveTimerReturnType SSequencer::EnsureSlateTickDuringPlayback(double InCurrentTime, float InDeltaTime)
 {
 	if (Sequencer.IsValid())
 	{
 		auto PlaybackStatus = Sequencer.Pin()->GetPlaybackStatus();
 		if (PlaybackStatus == EMovieScenePlayerStatus::Playing || PlaybackStatus == EMovieScenePlayerStatus::Recording)
 		{
-			return EActiveTickReturnType::KeepTicking;
+			return EActiveTimerReturnType::Continue;
 		}
 	}
 
-	bIsActiveTickRegistered = false;
-	return EActiveTickReturnType::StopTicking;
+	bIsActiveTimerRegistered = false;
+	return EActiveTimerReturnType::Stop;
 }
 
 void SSequencer::UpdateLayoutTree()
diff --git a/Engine/Source/Editor/Sequencer/Private/SSequencer.h b/Engine/Source/Editor/Sequencer/Private/SSequencer.h
--- a/Engine/Source/Editor/Sequencer/Private/SSequencer.h
+++ b/Engine/Source/Editor/Sequencer/Private/SSequencer.h
@@ -77,8 +77,8 @@ public:
 	/** Updates the layout node tree from movie scene data */
 	void UpdateLayoutTree();
 
-	/** Causes the widget to register an empty active tick that persists until Sequencer playback stops */
-	void RegisterActiveTickForPlayback();
+	/** Causes the widget to register an empty active timer that persists until Sequencer playback stops */
+	void RegisterActiveTimerForPlayback();
 
 	/**
 	 * Updates the breadcrumbs from a change in the shot filter state
@@ -92,8 +92,8 @@ public:
 	void DeleteSelectedNodes();
 
 private:
-	/** Empty active tick to ensure Slate ticks during Sequencer playback */
-	EActiveTickReturnType EnsureSlateTickDuringPlayback(double InCurrentTime, float InDeltaTime);	
+	/** Empty active timer to ensure Slate ticks during Sequencer playback */
+	EActiveTimerReturnType EnsureSlateTickDuringPlayback(double InCurrentTime, float InDeltaTime);	
 
 	/** Makes the toolbar for the outline section. */
 	TSharedRef<SWidget> MakeToolBar();
@@ -208,6 +208,6 @@ private:
 	/** The main sequencer interface */
 	TWeakPtr<FSequencer> Sequencer;
 
-	/** Whether the active tick is currently registered */
-	bool bIsActiveTickRegistered;
+	/** Whether the active timer is currently registered */
+	bool bIsActiveTimerRegistered;
 };
diff --git a/Engine/Source/Editor/Sequencer/Private/Sequencer.cpp b/Engine/Source/Editor/Sequencer/Private/Sequencer.cpp
--- a/Engine/Source/Editor/Sequencer/Private/Sequencer.cpp
+++ b/Engine/Source/Editor/Sequencer/Private/Sequencer.cpp
@@ -721,7 +721,7 @@ FReply FSequencer::OnPlay()
 			PlaybackState = EMovieScenePlayerStatus::Playing;
 			
 			// Make sure Slate ticks during playback
-			SequencerWidget->RegisterActiveTickForPlayback();
+			SequencerWidget->RegisterActiveTimerForPlayback();
 		}
 	}
 
@@ -735,7 +735,7 @@ FReply FSequencer::OnRecord()
 		PlaybackState = EMovieScenePlayerStatus::Recording;
 		
 		// Make sure Slate ticks during playback
-		SequencerWidget->RegisterActiveTickForPlayback();
+		SequencerWidget->RegisterActiveTimerForPlayback();
 
 		// @todo sequencer livecapture: Ideally we would support fixed timestep capture from simulation
 		//			Basically we need to run the PIE world at a fixed time step, capturing key frames every frame. 
diff --git a/Engine/Source/Editor/UMGEditor/Private/Customizations/CanvasSlotCustomization.cpp b/Engine/Source/Editor/UMGEditor/Private/Customizations/CanvasSlotCustomization.cpp
--- a/Engine/Source/Editor/UMGEditor/Private/Customizations/CanvasSlotCustomization.cpp
+++ b/Engine/Source/Editor/UMGEditor/Private/Customizations/CanvasSlotCustomization.cpp
@@ -164,7 +164,7 @@ private:
 
 private:
 	FCurveSequence ResizeCurve;
-	TWeakPtr<FActiveTickHandle> ActiveTickHandle;
+	TWeakPtr<FActiveTimerHandle> ActiveTimerHandle;
 	bool bIsHovered;
 };
 
diff --git a/Engine/Source/Editor/UMGEditor/Private/Designer/SDesignSurface.cpp b/Engine/Source/Editor/UMGEditor/Private/Designer/SDesignSurface.cpp
--- a/Engine/Source/Editor/UMGEditor/Private/Designer/SDesignSurface.cpp
+++ b/Engine/Source/Editor/UMGEditor/Private/Designer/SDesignSurface.cpp
@@ -142,7 +142,7 @@ void SDesignSurface::Construct(const FArguments& InArgs)
 	];
 }
 
-EActiveTickReturnType SDesignSurface::HandleZoomToFit( double InCurrentTime, float InDeltaTime )
+EActiveTimerReturnType SDesignSurface::HandleZoomToFit( double InCurrentTime, float InDeltaTime )
 {
 	const FVector2D DesiredViewCenter = ( ZoomTargetTopLeft + ZoomTargetBottomRight ) * 0.5f;
 	const bool bDoneScrolling = ScrollToLocation(CachedGeometry, DesiredViewCenter, bTeleportInsteadOfScrollingWhenZoomingToFit ? 1000.0f : InDeltaTime);
@@ -156,10 +156,10 @@ EActiveTickReturnType SDesignSurface::HandleZoomToFit( double InCurrentTime, flo
 		ZoomTargetTopLeft = FVector2D::ZeroVector;
 		ZoomTargetBottomRight = FVector2D::ZeroVector;
 
-		return EActiveTickReturnType::StopTicking;
+		return EActiveTimerReturnType::Stop;
 	}
 	
-	return EActiveTickReturnType::KeepTicking;
+	return EActiveTimerReturnType::Continue;
 }
 
 void SDesignSurface::Tick( const FGeometry& AllottedGeometry, const double InCurrentTime, const float InDeltaTime )
@@ -174,9 +174,9 @@ void SDesignSurface::Tick( const FGeometry& AllottedGeometry, const double InCur
 		ZoomTargetTopLeft = FVector2D(Bounds.Left, Bounds.Top);
 		ZoomTargetBottomRight = FVector2D(Bounds.Right, Bounds.Bottom);
 
-		if (!ActiveTickHandle.IsValid())
+		if (!ActiveTimerHandle.IsValid())
 		{
-			RegisterActiveTick(0.f, FWidgetActiveTickDelegate::CreateSP(this, &SDesignSurface::HandleZoomToFit));
+			RegisterActiveTimer(0.f, FWidgetActiveTimerDelegate::CreateSP(this, &SDesignSurface::HandleZoomToFit));
 		}
 	}
 }
diff --git a/Engine/Source/Editor/UMGEditor/Private/Designer/SDesignSurface.h b/Engine/Source/Editor/UMGEditor/Private/Designer/SDesignSurface.h
--- a/Engine/Source/Editor/UMGEditor/Private/Designer/SDesignSurface.h
+++ b/Engine/Source/Editor/UMGEditor/Private/Designer/SDesignSurface.h
@@ -99,15 +99,15 @@ protected:
 	/** Does the user need to press Control in order to over-zoom. */
 	bool bRequireControlToOverZoom;
 
-	/** Cached geometry for use within the active tick */
+	/** Cached geometry for use within the active timer */
 	FGeometry CachedGeometry;
 
 private:
-	/** Active tick that handles deferred zooming until the target zoom is reached */
-	EActiveTickReturnType HandleZoomToFit(double InCurrentTime, float InDeltaTime);
+	/** Active timer that handles deferred zooming until the target zoom is reached */
+	EActiveTimerReturnType HandleZoomToFit(double InCurrentTime, float InDeltaTime);
 
-	/** The handle to the active tick */
-	TWeakPtr<FActiveTickHandle> ActiveTickHandle;
+	/** The handle to the active timer */
+	TWeakPtr<FActiveTimerHandle> ActiveTimerHandle;
 
 	// A flag noting if we have a pending zoom to extents operation to perform next tick.
 	bool bDeferredZoomToExtents;
diff --git a/Engine/Source/Editor/UnrealEd/Private/CreditsScreen.cpp b/Engine/Source/Editor/UnrealEd/Private/CreditsScreen.cpp
--- a/Engine/Source/Editor/UnrealEd/Private/CreditsScreen.cpp
+++ b/Engine/Source/Editor/UnrealEd/Private/CreditsScreen.cpp
@@ -14,7 +14,7 @@ void SCreditsScreen::Construct(const FArguments& InArgs)
 	ScrollPixelsPerSecond = 50.0f;
 	bIsPlaying = true;
 
-	ActiveTickHandle = RegisterActiveTick( 0.f, FWidgetActiveTickDelegate::CreateSP( this, &SCreditsScreen::RollCredits ) );
+	ActiveTimerHandle = RegisterActiveTimer( 0.f, FWidgetActiveTimerDelegate::CreateSP( this, &SCreditsScreen::RollCredits ) );
 
 	const FString Version = GEngineVersion.ToString(FEngineBuildSettings::IsPerforceBuild() ? EVersionComponent::Branch : EVersionComponent::Patch);
 
@@ -69,13 +69,13 @@ void SCreditsScreen::Construct(const FArguments& InArgs)
 	];
 }
 
-EActiveTickReturnType SCreditsScreen::RollCredits( double InCurrentTime, float InDeltaTime )
+EActiveTimerReturnType SCreditsScreen::RollCredits( double InCurrentTime, float InDeltaTime )
 {
 	float NewPixelOffset = ( ScrollPixelsPerSecond * InDeltaTime );
 	ScrollBox->SetScrollOffset( ScrollBox->GetScrollOffset() + NewPixelOffset );
 	PreviousScrollPosition = ScrollBox->GetScrollOffset();
 
-	return EActiveTickReturnType::KeepTicking;
+	return EActiveTimerReturnType::Continue;
 }
 
 FReply SCreditsScreen::HandleTogglePlayPause()
@@ -83,17 +83,17 @@ FReply SCreditsScreen::HandleTogglePlayPause()
 	if ( bIsPlaying )
 	{
 		bIsPlaying = false;
-		if ( ActiveTickHandle.IsValid() )
+		if ( ActiveTimerHandle.IsValid() )
 		{
-			UnRegisterActiveTick( ActiveTickHandle.Pin().ToSharedRef() );
+			UnRegisterActiveTimer( ActiveTimerHandle.Pin().ToSharedRef() );
 		}
 	}
 	else
 	{
 		bIsPlaying = true;
-		if ( !ActiveTickHandle.IsValid() )
+		if ( !ActiveTimerHandle.IsValid() )
 		{
-			ActiveTickHandle = RegisterActiveTick( 0.f, FWidgetActiveTickDelegate::CreateSP( this, &SCreditsScreen::RollCredits ) );
+			ActiveTimerHandle = RegisterActiveTimer( 0.f, FWidgetActiveTimerDelegate::CreateSP( this, &SCreditsScreen::RollCredits ) );
 		}
 	}
 
@@ -106,9 +106,9 @@ void SCreditsScreen::HandleUserScrolled(float ScrollOffset)
 	if ( bIsPlaying && ScrollOffset < PreviousScrollPosition )
 	{
 		bIsPlaying = false;
-		if ( ActiveTickHandle.IsValid() )
+		if ( ActiveTimerHandle.IsValid() )
 		{
-			UnRegisterActiveTick( ActiveTickHandle.Pin().ToSharedRef() );
+			UnRegisterActiveTimer( ActiveTimerHandle.Pin().ToSharedRef() );
 		}
 	}
 
diff --git a/Engine/Source/Editor/UnrealEd/Private/Dialogs/SDeleteAssetsDialog.cpp b/Engine/Source/Editor/UnrealEd/Private/Dialogs/SDeleteAssetsDialog.cpp
--- a/Engine/Source/Editor/UnrealEd/Private/Dialogs/SDeleteAssetsDialog.cpp
+++ b/Engine/Source/Editor/UnrealEd/Private/Dialogs/SDeleteAssetsDialog.cpp
@@ -101,9 +101,9 @@ SDeleteAssetsDialog::~SDeleteAssetsDialog()
 
 void SDeleteAssetsDialog::Construct( const FArguments& InArgs, TSharedRef<FAssetDeleteModel> InDeleteModel )
 {
-	//bIsActiveTickRegistered = false;
-	bIsActiveTickRegistered = true;
-	RegisterActiveTick( 0.f, FWidgetActiveTickDelegate::CreateSP( this, &SDeleteAssetsDialog::TickDeleteModel ) );
+	//bIsActiveTimerRegistered = false;
+	bIsActiveTimerRegistered = true;
+	RegisterActiveTimer( 0.f, FWidgetActiveTimerDelegate::CreateSP( this, &SDeleteAssetsDialog::TickDeleteModel ) );
 
 	DeleteModel = InDeleteModel;
 
@@ -511,17 +511,17 @@ FReply SDeleteAssetsDialog::OnKeyDown( const FGeometry& MyGeometry, const FKeyEv
 	return FReply::Unhandled();
 }
 
-EActiveTickReturnType SDeleteAssetsDialog::TickDeleteModel( double InCurrentTime, float InDeltaTime )
+EActiveTimerReturnType SDeleteAssetsDialog::TickDeleteModel( double InCurrentTime, float InDeltaTime )
 {
 	DeleteModel->Tick( InDeltaTime );
 
 	if ( DeleteModel->GetState() == FAssetDeleteModel::EState::Finished )
 	{
-		bIsActiveTickRegistered = false;
-		return EActiveTickReturnType::StopTicking;
+		bIsActiveTimerRegistered = false;
+		return EActiveTimerReturnType::Stop;
 	}
 
-	return EActiveTickReturnType::KeepTicking;
+	return EActiveTimerReturnType::Continue;
 }
 
 void SDeleteAssetsDialog::HandleDeleteModelStateChanged(FAssetDeleteModel::EState NewState)
@@ -838,10 +838,10 @@ void SDeleteAssetsDialog::ExecuteDeleteReferencers()
 		UObject* ObjectToDelete = SelectedAsset.GetAsset();
 		
 		DeleteModel->AddObjectToDelete( ObjectToDelete );
-		if ( !bIsActiveTickRegistered )
+		if ( !bIsActiveTimerRegistered )
 		{
-			bIsActiveTickRegistered = true;
-			RegisterActiveTick( 0.f, FWidgetActiveTickDelegate::CreateSP( this, &SDeleteAssetsDialog::TickDeleteModel ) );
+			bIsActiveTimerRegistered = true;
+			RegisterActiveTimer( 0.f, FWidgetActiveTimerDelegate::CreateSP( this, &SDeleteAssetsDialog::TickDeleteModel ) );
 		}
 	}
 }
diff --git a/Engine/Source/Editor/UnrealEd/Private/FeedbackContextEditor.cpp b/Engine/Source/Editor/UnrealEd/Private/FeedbackContextEditor.cpp
--- a/Engine/Source/Editor/UnrealEd/Private/FeedbackContextEditor.cpp
+++ b/Engine/Source/Editor/UnrealEd/Private/FeedbackContextEditor.cpp
@@ -42,8 +42,8 @@ public:
 		OnCancelClickedDelegate = InArgs._OnCancelClickedDelegate;
 		ScopeStack = InArgs._ScopeStack;
 
-		// This is a temporary widget that needs to be updated over its entire lifespan => has an active tick registered for its entire lifespan
-		RegisterActiveTick( 0.f, FWidgetActiveTickDelegate::CreateSP( this, &SSlowTaskWidget::UpdateProgress ) );
+		// This is a temporary widget that needs to be updated over its entire lifespan => has an active timer registered for its entire lifespan
+		RegisterActiveTimer( 0.f, FWidgetActiveTimerDelegate::CreateSP( this, &SSlowTaskWidget::UpdateProgress ) );
 
 		TSharedRef<SVerticalBox> VerticalBox = SNew(SVerticalBox)
 
@@ -136,12 +136,12 @@ public:
 
 private:
 
-	/** Active tick to update the progress bars */
-	EActiveTickReturnType UpdateProgress(double InCurrentTime, float InDeltaTime)
+	/** Active timer to update the progress bars */
+	EActiveTimerReturnType UpdateProgress(double InCurrentTime, float InDeltaTime)
 	{
 		UpdateDynamicProgressBars();
 
-		return EActiveTickReturnType::KeepTicking;
+		return EActiveTimerReturnType::Continue;
 	}
 
 	/** Updates the dynamic progress bars for this widget */
diff --git a/Engine/Source/Editor/UnrealEd/Private/SEditorViewport.cpp b/Engine/Source/Editor/UnrealEd/Private/SEditorViewport.cpp
--- a/Engine/Source/Editor/UnrealEd/Private/SEditorViewport.cpp
+++ b/Engine/Source/Editor/UnrealEd/Private/SEditorViewport.cpp
@@ -63,7 +63,7 @@ void SEditorViewport::Construct( const FArguments& InArgs )
 
 	if ( Client->IsRealtime() )
 	{
-		ActiveTickHandle = RegisterActiveTick( 0.f, FWidgetActiveTickDelegate::CreateSP( this, &SEditorViewport::EnsureTickWhileRealtime ) );
+		ActiveTimerHandle = RegisterActiveTimer( 0.f, FWidgetActiveTimerDelegate::CreateSP( this, &SEditorViewport::EnsureTickWhileRealtime ) );
 	}
 
 	CommandList = MakeShareable( new FUICommandList );
@@ -331,16 +331,16 @@ void SEditorViewport::OnToggleRealtime()
 	if (Client->IsRealtime())
 	{
 		Client->SetRealtime( false );
-		if ( ActiveTickHandle.IsValid() )
+		if ( ActiveTimerHandle.IsValid() )
 		{
-			UnRegisterActiveTick( ActiveTickHandle.Pin().ToSharedRef() );
+			UnRegisterActiveTimer( ActiveTimerHandle.Pin().ToSharedRef() );
 		}
 		
 	}
 	else
 	{
 		Client->SetRealtime( true );
-		ActiveTickHandle = RegisterActiveTick( 0.f, FWidgetActiveTickDelegate::CreateSP( this, &SEditorViewport::EnsureTickWhileRealtime ) );
+		ActiveTimerHandle = RegisterActiveTimer( 0.f, FWidgetActiveTimerDelegate::CreateSP( this, &SEditorViewport::EnsureTickWhileRealtime ) );
 	}
 }
 
@@ -355,7 +355,7 @@ void SEditorViewport::OnToggleStats()
 		if ( !Client->IsRealtime() )
 		{
 			Client->SetRealtime( true );
-			ActiveTickHandle = RegisterActiveTick( 0.f, FWidgetActiveTickDelegate::CreateSP( this, &SEditorViewport::EnsureTickWhileRealtime ) );
+			ActiveTimerHandle = RegisterActiveTimer( 0.f, FWidgetActiveTimerDelegate::CreateSP( this, &SEditorViewport::EnsureTickWhileRealtime ) );
 		}
 
 		 // let the user know how they can enable stats via the console
@@ -518,9 +518,9 @@ bool SEditorViewport::OnIsSurfaceSnapEnabled()
 	return GetDefault<ULevelEditorViewportSettings>()->SnapToSurface.bEnabled;
 }
 
-EActiveTickReturnType SEditorViewport::EnsureTickWhileRealtime( double InCurrentTime, float InDeltaTime )
+EActiveTimerReturnType SEditorViewport::EnsureTickWhileRealtime( double InCurrentTime, float InDeltaTime )
 {
-	return EActiveTickReturnType::KeepTicking;
+	return EActiveTimerReturnType::Continue;
 }
 
 #undef LOCTEXT_NAMESPACE
diff --git a/Engine/Source/Editor/UnrealEd/Private/SSkeletonWidget.cpp b/Engine/Source/Editor/UnrealEd/Private/SSkeletonWidget.cpp
--- a/Engine/Source/Editor/UnrealEd/Private/SSkeletonWidget.cpp
+++ b/Engine/Source/Editor/UnrealEd/Private/SSkeletonWidget.cpp
@@ -1129,7 +1129,7 @@ public:
 // SBasePoseViewport
 void SBasePoseViewport::Construct(const FArguments& InArgs)
 {
-	//RegisterActiveTick( 0.f, FTickWidgetDelegate::CreateSP( this, &SBasePoseViewport::TickViewport ) );
+	//RegisterActiveTimer( 0.f, FTickWidgetDelegate::CreateSP( this, &SBasePoseViewport::TickViewport ) );
 
 	this->ChildSlot
 	[
diff --git a/Engine/Source/Editor/UnrealEd/Public/CreditsScreen.h b/Engine/Source/Editor/UnrealEd/Public/CreditsScreen.h
--- a/Engine/Source/Editor/UnrealEd/Public/CreditsScreen.h
+++ b/Engine/Source/Editor/UnrealEd/Public/CreditsScreen.h
@@ -19,7 +19,7 @@ public:
 
 private:
 	/** Animates the credits during play */
-	EActiveTickReturnType RollCredits( double InCurrentTime, float InDeltaTime );
+	EActiveTimerReturnType RollCredits( double InCurrentTime, float InDeltaTime );
 
 	/** Handles the user clicking the play/pause toggle button. */
 	FReply HandleTogglePlayPause();
@@ -37,7 +37,7 @@ private:
 	TSharedPtr<SScrollBox> ScrollBox;
 
 	/** The handle to the active roll credits tick */
-	TWeakPtr<FActiveTickHandle> ActiveTickHandle;
+	TWeakPtr<FActiveTimerHandle> ActiveTimerHandle;
 
 	/** The auto scroll rate in pixels per second */
 	float ScrollPixelsPerSecond;
diff --git a/Engine/Source/Editor/UnrealEd/Public/Dialogs/SDeleteAssetsDialog.h b/Engine/Source/Editor/UnrealEd/Public/Dialogs/SDeleteAssetsDialog.h
--- a/Engine/Source/Editor/UnrealEd/Public/Dialogs/SDeleteAssetsDialog.h
+++ b/Engine/Source/Editor/UnrealEd/Public/Dialogs/SDeleteAssetsDialog.h
@@ -51,8 +51,8 @@ private:
 	TSharedRef<SWidget> BuildDeleteDialog();
 
 private:
-	/** Active tick to tick the delete model until it reaches a "Finished" state */
-	EActiveTickReturnType TickDeleteModel( double InCurrentTime, float InDeltaTime );
+	/** Active timer to tick the delete model until it reaches a "Finished" state */
+	EActiveTimerReturnType TickDeleteModel( double InCurrentTime, float InDeltaTime );
 
 	void HandleDeleteModelStateChanged( FAssetDeleteModel::EState NewState );
 
@@ -120,8 +120,8 @@ private:
 
 private:
 
-	/** Whether the active tick is currently registered */
-	bool bIsActiveTickRegistered;
+	/** Whether the active timer is currently registered */
+	bool bIsActiveTimerRegistered;
 
 	/** The model used for deleting assets */
 	TSharedPtr<FAssetDeleteModel> DeleteModel;
diff --git a/Engine/Source/Editor/UnrealEd/Public/SEditorViewport.h b/Engine/Source/Editor/UnrealEd/Public/SEditorViewport.h
--- a/Engine/Source/Editor/UnrealEd/Public/SEditorViewport.h
+++ b/Engine/Source/Editor/UnrealEd/Public/SEditorViewport.h
@@ -173,8 +173,8 @@ protected:
 
 private:
 	/** Ensures the viewport is ticked actively when realtime is enabled */
-	EActiveTickReturnType EnsureTickWhileRealtime( double InCurrentTime, float InDeltaTime );
+	EActiveTimerReturnType EnsureTickWhileRealtime( double InCurrentTime, float InDeltaTime );
 
-	/** The handle to the active tick */
-	TWeakPtr<FActiveTickHandle> ActiveTickHandle;
+	/** The handle to the active timer */
+	TWeakPtr<FActiveTimerHandle> ActiveTimerHandle;
 };
diff --git a/Engine/Source/Runtime/AppFramework/Private/Widgets/Colors/SColorPicker.cpp b/Engine/Source/Runtime/AppFramework/Private/Widgets/Colors/SColorPicker.cpp
--- a/Engine/Source/Runtime/AppFramework/Private/Widgets/Colors/SColorPicker.cpp
+++ b/Engine/Source/Runtime/AppFramework/Private/Widgets/Colors/SColorPicker.cpp
@@ -63,7 +63,7 @@ void SColorPicker::Construct( const FArguments& InArgs )
 	DisplayGamma = InArgs._DisplayGamma;
 	bClosedViaOkOrCancel = false;
 
-	RegisterActiveTick( 0.f, FWidgetActiveTickDelegate::CreateSP( this, &SColorPicker::AnimatePostConstruct ) );
+	RegisterActiveTimer( 0.f, FWidgetActiveTimerDelegate::CreateSP( this, &SColorPicker::AnimatePostConstruct ) );
 
 	// We need a parent window to set the close callback
 	if (ParentWindowPtr.IsValid())
@@ -573,11 +573,11 @@ void SColorPicker::GenerateDefaultColorPickerContent( bool bAdvancedSectionExpan
 	HideSmallTrash();
 }
 
-EActiveTickReturnType SColorPicker::AnimatePostConstruct( double InCurrentTime, float InDeltaTime )
+EActiveTimerReturnType SColorPicker::AnimatePostConstruct( double InCurrentTime, float InDeltaTime )
 {
 	static const float AnimationTime = 0.25f;
 
-	EActiveTickReturnType TickReturnVal = EActiveTickReturnType::KeepTicking;
+	EActiveTimerReturnType TickReturnVal = EActiveTimerReturnType::Continue;
 	if ( CurrentTime < AnimationTime )
 	{
 		CurrentColorHSV = FMath::Lerp( ColorBegin, ColorEnd, CurrentTime / AnimationTime );
@@ -594,7 +594,7 @@ EActiveTickReturnType SColorPicker::AnimatePostConstruct( double InCurrentTime,
 		if ( CurrentTime >= AnimationTime )
 		{
 			CurrentColorHSV = ColorEnd;
-			TickReturnVal = EActiveTickReturnType::StopTicking;
+			TickReturnVal = EActiveTimerReturnType::Stop;
 		}
 
 		CurrentColorRGB = CurrentColorHSV.HSVToLinearRGB();
diff --git a/Engine/Source/Runtime/AppFramework/Public/Widgets/Colors/SColorPicker.h b/Engine/Source/Runtime/AppFramework/Public/Widgets/Colors/SColorPicker.h
--- a/Engine/Source/Runtime/AppFramework/Public/Widgets/Colors/SColorPicker.h
+++ b/Engine/Source/Runtime/AppFramework/Public/Widgets/Colors/SColorPicker.h
@@ -222,8 +222,8 @@ protected:
 
 private:
 
-	// Callback for the active tick to animate the color post-construct
-	EActiveTickReturnType AnimatePostConstruct( double InCurrentTime, float InDeltaTime );
+	// Callback for the active timer to animate the color post-construct
+	EActiveTimerReturnType AnimatePostConstruct( double InCurrentTime, float InDeltaTime );
 
 	// Callback for getting the end color of a color spin box gradient.
 	FLinearColor GetGradientEndColor( EColorPickerChannels Channel ) const;
diff --git a/Engine/Source/Runtime/Slate/Private/Framework/Application/SlateApplication.cpp b/Engine/Source/Runtime/Slate/Private/Framework/Application/SlateApplication.cpp
--- a/Engine/Source/Runtime/Slate/Private/Framework/Application/SlateApplication.cpp
+++ b/Engine/Source/Runtime/Slate/Private/Framework/Application/SlateApplication.cpp
@@ -339,17 +339,17 @@ TAutoConsoleVariable<int32> SkipSecondPrepass(
 	0,
 	TEXT("Whether to skip the second Slate PrePass call (the one right before rendering)."));
 
-/** Whether Slate should go to sleep when there are no active ticks and the user is idle */
+/** Whether Slate should go to sleep when there are no active timers and the user is idle */
 TAutoConsoleVariable<int32> AllowSlateToSleep(
 	TEXT("Slate.AllowSlateToSleep"),
 	false,
-	TEXT("Whether Slate should go to sleep when there are no active ticks and the user is idle"));
+	TEXT("Whether Slate should go to sleep when there are no active timers and the user is idle"));
 
-/** The amount of time that must pass without any user action before Slate is put to sleep (provided that there are no active ticks). */
+/** The amount of time that must pass without any user action before Slate is put to sleep (provided that there are no active timers). */
 TAutoConsoleVariable<float> SleepBufferPostInput(
 	TEXT("Slate.SleepBufferPostInput"),
 	0.0f,
-	TEXT("The amount of time that must pass without any user action before Slate is put to sleep (provided that there are no active ticks)."));
+	TEXT("The amount of time that must pass without any user action before Slate is put to sleep (provided that there are no active timers)."));
 
 //////////////////////////////////////////////////////////////////////////
 bool FSlateApplication::MouseCaptorHelper::HasCapture() const
@@ -1296,19 +1296,19 @@ void FSlateApplication::Tick()
 	const double TimeSinceMouseMove = LastTickTime - LastMouseMoveTime;
 	
 	const bool bIsUserIdle = (TimeSinceInput > SleepThreshold) && (TimeSinceMouseMove > SleepThreshold);
-	const bool bAnyActiveTicksPending = AnyActiveTicksArePending();
-	if (bAnyActiveTicksPending)
+	const bool bAnyActiveTimersPending = AnyActiveTimersArePending();
+	if (bAnyActiveTimersPending)
 	{
 		// Some UI might slide under the cursor. To a widget, this is
 		// as if the cursor moved over it.
 		QueueSynthesizedMouseMove();
 	}
 
-	// skip tick/draw if we are idle and there are no active ticks registered that we need to drive slate for.
+	// skip tick/draw if we are idle and there are no active timers registered that we need to drive slate for.
 	// This effectively means the slate application is totally idle and we don't need to update the UI.
-	// This relies on Widgets properly registering for Active tick when they need something to happen even
+	// This relies on Widgets properly registering for Active timer when they need something to happen even
 	// when the user is not providing any input (ie, animations, viewport rendering, async polling, etc).
-	if (!AllowSlateToSleep.GetValueOnGameThread() || bAnyActiveTicksPending || !bIsUserIdle || bNeedsSyntheticMouseMouse)
+	if (!AllowSlateToSleep.GetValueOnGameThread() || bAnyActiveTimersPending || !bIsUserIdle || bNeedsSyntheticMouseMouse)
 	{
 		if (!bFoldTick)
 		{
diff --git a/Engine/Source/Runtime/Slate/Private/Framework/MultiBox/SMenuEntryBlock.cpp b/Engine/Source/Runtime/Slate/Private/Framework/MultiBox/SMenuEntryBlock.cpp
--- a/Engine/Source/Runtime/Slate/Private/Framework/MultiBox/SMenuEntryBlock.cpp
+++ b/Engine/Source/Runtime/Slate/Private/Framework/MultiBox/SMenuEntryBlock.cpp
@@ -813,9 +813,9 @@ void SMenuEntryBlock::RequestSubMenuToggle( bool bOpenMenu, const bool bClobber
 	// Reset the time before the menu opens
 	float TimeToSubMenuOpen = bClobber ? MultiBoxConstants::SubMenuClobberTime : MultiBoxConstants::SubMenuOpenTime;
 	//SubMenuRequestState = bOpenMenu ? WantOpen : WantClose;
-	if (!ActiveTickHandle.IsValid())
+	if (!ActiveTimerHandle.IsValid())
 	{
-		ActiveTickHandle = RegisterActiveTick(TimeToSubMenuOpen, FWidgetActiveTickDelegate::CreateSP(this, &SMenuEntryBlock::UpdateSubMenuState, bOpenMenu));
+		ActiveTimerHandle = RegisterActiveTimer(TimeToSubMenuOpen, FWidgetActiveTimerDelegate::CreateSP(this, &SMenuEntryBlock::UpdateSubMenuState, bOpenMenu));
 	}
 	//UpdateSubMenuState();
 }
@@ -825,10 +825,10 @@ void SMenuEntryBlock::CancelPendingSubMenu()
 	// Reset any pending sub-menu openings
 	//SubMenuRequestState = Idle;
 	
-	auto PinnedActiveTickHandle = ActiveTickHandle.Pin();
-	if (PinnedActiveTickHandle.IsValid())
+	auto PinnedActiveTimerHandle = ActiveTimerHandle.Pin();
+	if (PinnedActiveTimerHandle.IsValid())
 	{
-		UnRegisterActiveTick(PinnedActiveTickHandle.ToSharedRef());
+		UnRegisterActiveTimer(PinnedActiveTimerHandle.ToSharedRef());
 	}
 }
 
@@ -1140,7 +1140,7 @@ FReply SMenuEntryBlock::OnKeyDown( const FGeometry& MyGeometry, const FKeyEvent&
 //	}
 //}
 
-EActiveTickReturnType SMenuEntryBlock::UpdateSubMenuState(double InCurrentTime, float InDeltaTime, bool bWantsOpen)
+EActiveTimerReturnType SMenuEntryBlock::UpdateSubMenuState(double InCurrentTime, float InDeltaTime, bool bWantsOpen)
 {
 	//const bool bSubMenuNeedsToOpen = ( SubMenuRequestState == WantOpen );
 	//SubMenuRequestState = Idle;
@@ -1179,7 +1179,7 @@ EActiveTickReturnType SMenuEntryBlock::UpdateSubMenuState(double InCurrentTime,
 		PinnedOwnerMultiBoxWidget->CloseSummonedMenus();
 	}
 
-	return EActiveTickReturnType::StopTicking;
+	return EActiveTimerReturnType::Stop;
 }
 
 //void SMenuEntryBlock::Tick( const FGeometry& AllottedGeometry, const double InCurrentTime, const float InDeltaTime )
diff --git a/Engine/Source/Runtime/Slate/Private/Framework/MultiBox/SMenuEntryBlock.h b/Engine/Source/Runtime/Slate/Private/Framework/MultiBox/SMenuEntryBlock.h
--- a/Engine/Source/Runtime/Slate/Private/Framework/MultiBox/SMenuEntryBlock.h
+++ b/Engine/Source/Runtime/Slate/Private/Framework/MultiBox/SMenuEntryBlock.h
@@ -316,11 +316,11 @@ protected:
 private:
 	const FSlateBrush* GetCheckBoxImageBrushFromStyle(const FCheckBoxStyle* Style) const;
 
-	/** One-off delayed active tick to update the open/closed state of the sub menu. */
-	EActiveTickReturnType UpdateSubMenuState(double InCurrentTime, float InDeltaTime, bool bWantsOpen);
+	/** One-off delayed active timer to update the open/closed state of the sub menu. */
+	EActiveTimerReturnType UpdateSubMenuState(double InCurrentTime, float InDeltaTime, bool bWantsOpen);
 
-	/** The handle to the active tick to update the sub-menu state */
-	TWeakPtr<FActiveTickHandle> ActiveTickHandle;
+	/** The handle to the active timer to update the sub-menu state */
+	TWeakPtr<FActiveTimerHandle> ActiveTimerHandle;
 
 	/** The brush to use when an item should appear checked */
 	const FSlateBrush* CheckedImage;
diff --git a/Engine/Source/Runtime/Slate/Private/Widgets/Docking/SDockTab.cpp b/Engine/Source/Runtime/Slate/Private/Widgets/Docking/SDockTab.cpp
--- a/Engine/Source/Runtime/Slate/Private/Widgets/Docking/SDockTab.cpp
+++ b/Engine/Source/Runtime/Slate/Private/Widgets/Docking/SDockTab.cpp
@@ -96,9 +96,9 @@ FReply SDockTab::OnMouseButtonUp( const FGeometry& MyGeometry, const FPointerEve
 void SDockTab::OnDragEnter( const FGeometry& MyGeometry, const FDragDropEvent& DragDropEvent )
 {
 	// Register to activate the tab after a delay
-	if ( !ActiveTickHandle.IsValid() )
+	if ( !ActiveTimerHandle.IsValid() )
 	{
-		ActiveTickHandle = RegisterActiveTick( SDockTabDefs::DragTimerActivate, FWidgetActiveTickDelegate::CreateSP( this, &SDockTab::TriggerActivateTab ) );
+		ActiveTimerHandle = RegisterActiveTimer( SDockTabDefs::DragTimerActivate, FWidgetActiveTimerDelegate::CreateSP( this, &SDockTab::TriggerActivateTab ) );
 	}
 
 	SBorder::OnDragEnter( MyGeometry, DragDropEvent );
@@ -107,9 +107,9 @@ void SDockTab::OnDragEnter( const FGeometry& MyGeometry, const FDragDropEvent& D
 void SDockTab::OnDragLeave( const FDragDropEvent& DragDropEvent )
 {
 	// Unregister the activation timer if it hasn't fired yet
-	if ( ActiveTickHandle.IsValid() )
+	if ( ActiveTimerHandle.IsValid() )
 	{
-		UnRegisterActiveTick( ActiveTickHandle.Pin().ToSharedRef() );
+		UnRegisterActiveTimer( ActiveTimerHandle.Pin().ToSharedRef() );
 	}
 
 	SBorder::OnDragLeave( DragDropEvent );
@@ -118,9 +118,9 @@ void SDockTab::OnDragLeave( const FDragDropEvent& DragDropEvent )
 FReply SDockTab::OnDrop( const FGeometry& MyGeometry, const FDragDropEvent& DragDropEvent )
 {
 	// Unregister the activation timer if it hasn't fired yet
-	if ( ActiveTickHandle.IsValid() )
+	if ( ActiveTimerHandle.IsValid() )
 	{
-		UnRegisterActiveTick( ActiveTickHandle.Pin().ToSharedRef() );
+		UnRegisterActiveTimer( ActiveTimerHandle.Pin().ToSharedRef() );
 	}
 
 	return SBorder::OnDrop( MyGeometry, DragDropEvent );
@@ -537,10 +537,10 @@ void SDockTab::Construct( const FArguments& InArgs )
 	);
 }
 
-EActiveTickReturnType SDockTab::TriggerActivateTab( double InCurrentTime, float InDeltaTime )
+EActiveTimerReturnType SDockTab::TriggerActivateTab( double InCurrentTime, float InDeltaTime )
 {
 	ActivateInParent( ETabActivationCause::UserClickedOnTab );
-	return EActiveTickReturnType::StopTicking;
+	return EActiveTimerReturnType::Stop;
 }
 
 const FDockTabStyle& SDockTab::GetCurrentStyle() const
diff --git a/Engine/Source/Runtime/Slate/Private/Widgets/Input/SEditableText.cpp b/Engine/Source/Runtime/Slate/Private/Widgets/Input/SEditableText.cpp
--- a/Engine/Source/Runtime/Slate/Private/Widgets/Input/SEditableText.cpp
+++ b/Engine/Source/Runtime/Slate/Private/Widgets/Input/SEditableText.cpp
@@ -308,7 +308,7 @@ void SEditableText::Tick( const FGeometry& AllottedGeometry, const double InCurr
 	}
 }
 
-EActiveTickReturnType SEditableText::AnimateSpringsWhileFocused(double InCurrentTime, float InDeltaTime)
+EActiveTimerReturnType SEditableText::AnimateSpringsWhileFocused(double InCurrentTime, float InDeltaTime)
 {
 	// Keep ticking as long as we should appear focused or the caret is in transit
 	const bool bShouldAppearFocused = ShouldAppearFocused();
@@ -324,10 +324,10 @@ EActiveTickReturnType SEditableText::AnimateSpringsWhileFocused(double InCurrent
 			SelectionTargetRightSpring.Tick(InDeltaTime);
 		}
 
-		return EActiveTickReturnType::KeepTicking;
+		return EActiveTimerReturnType::Continue;
 	}
 
-	return EActiveTickReturnType::StopTicking;
+	return EActiveTimerReturnType::Stop;
 }
 
 bool SEditableText::FTextInputMethodContext::IsReadOnly()
@@ -1501,7 +1501,7 @@ FReply SEditableText::OnFocusReceived( const FGeometry& MyGeometry, const FFocus
 	// Skip the rest of the focus received code if it's due to the context menu closing
 	if ( !ContextMenuWindow.IsValid() )
 	{
-		RegisterActiveTick(0.f, FWidgetActiveTickDelegate::CreateSP(this, &SEditableText::AnimateSpringsWhileFocused));
+		RegisterActiveTimer(0.f, FWidgetActiveTimerDelegate::CreateSP(this, &SEditableText::AnimateSpringsWhileFocused));
 
 		// Don't unselect text when focus is set because another widget lost focus, as that may be in response to a
 		// dismissed context menu where the user clicked an option to select text.
diff --git a/Engine/Source/Runtime/Slate/Private/Widgets/Input/SSearchBox.cpp b/Engine/Source/Runtime/Slate/Private/Widgets/Input/SSearchBox.cpp
--- a/Engine/Source/Runtime/Slate/Private/Widgets/Input/SSearchBox.cpp
+++ b/Engine/Source/Runtime/Slate/Private/Widgets/Input/SSearchBox.cpp
@@ -9,7 +9,7 @@ void SSearchBox::Construct( const FArguments& InArgs )
 {
 	check(InArgs._Style);
 
-	bIsActiveTickRegistered = false;
+	bIsActiveTimerRegistered = false;
 	OnSearchDelegate = InArgs._OnSearch;
 	OnTextChangedDelegate = InArgs._OnTextChanged;
 	OnTextCommittedDelegate = InArgs._OnTextCommitted;
@@ -109,13 +109,13 @@ void SSearchBox::Construct( const FArguments& InArgs )
 	];
 }
 
-EActiveTickReturnType SSearchBox::TriggerOnTextChanged( double InCurrentTime, float InDeltaTime )
+EActiveTimerReturnType SSearchBox::TriggerOnTextChanged( double InCurrentTime, float InDeltaTime )
 {
 	// Reset the flag first in case the delegate winds up triggering HandleTextChanged
-	bIsActiveTickRegistered = false;
+	bIsActiveTimerRegistered = false;
 
 	OnTextChangedDelegate.ExecuteIfBound( LastPendingTextChangedValue );
-	return EActiveTickReturnType::StopTicking;
+	return EActiveTimerReturnType::Stop;
 }
 
 void SSearchBox::HandleTextChanged(const FText& NewText)
@@ -125,12 +125,12 @@ void SSearchBox::HandleTextChanged(const FText& NewText)
 		LastPendingTextChangedValue = NewText;
 
 		// Remove the existing registered tick if necessary
-		if ( ActiveTickHandle.IsValid() )
+		if ( ActiveTimerHandle.IsValid() )
 		{
-			UnRegisterActiveTick( ActiveTickHandle.Pin().ToSharedRef() );
+			UnRegisterActiveTimer( ActiveTimerHandle.Pin().ToSharedRef() );
 		}
-		bIsActiveTickRegistered = true;
-		ActiveTickHandle = RegisterActiveTick( FilterDelayAfterTyping, FWidgetActiveTickDelegate::CreateSP( this, &SSearchBox::TriggerOnTextChanged ) );
+		bIsActiveTimerRegistered = true;
+		ActiveTimerHandle = RegisterActiveTimer( FilterDelayAfterTyping, FWidgetActiveTimerDelegate::CreateSP( this, &SSearchBox::TriggerOnTextChanged ) );
 	}
 	else
 	{
@@ -140,10 +140,10 @@ void SSearchBox::HandleTextChanged(const FText& NewText)
 
 void SSearchBox::HandleTextCommitted(const FText& NewText, ETextCommit::Type CommitType)
 {
-	if ( bIsActiveTickRegistered && ActiveTickHandle.IsValid() )
+	if ( bIsActiveTimerRegistered && ActiveTimerHandle.IsValid() )
 	{
-		bIsActiveTickRegistered = false;
-		UnRegisterActiveTick( ActiveTickHandle.Pin().ToSharedRef() );
+		bIsActiveTimerRegistered = false;
+		UnRegisterActiveTimer( ActiveTimerHandle.Pin().ToSharedRef() );
 	}
 
 	OnTextCommittedDelegate.ExecuteIfBound( NewText, CommitType );
diff --git a/Engine/Source/Runtime/Slate/Private/Widgets/Layout/SScrollBox.cpp b/Engine/Source/Runtime/Slate/Private/Widgets/Layout/SScrollBox.cpp
--- a/Engine/Source/Runtime/Slate/Private/Widgets/Layout/SScrollBox.cpp
+++ b/Engine/Source/Runtime/Slate/Private/Widgets/Layout/SScrollBox.cpp
@@ -160,7 +160,7 @@ void SScrollBox::Construct( const FArguments& InArgs )
 	OnUserScrolled = InArgs._OnUserScrolled;
 	Orientation = InArgs._Orientation;
 	bScrollToEnd = false;
-	bIsActiveTickRegistered = false;
+	bIsActiveTimerRegistered = false;
 
 	if (InArgs._ExternalScrollbar.IsValid())
 	{
@@ -433,7 +433,7 @@ void SScrollBox::SetScrollBarThickness(FVector2D InThickness)
 	ScrollBar->SetThickness(InThickness);
 }
 
-EActiveTickReturnType SScrollBox::UpdateInertialScroll(double InCurrentTime, float InDeltaTime)
+EActiveTimerReturnType SScrollBox::UpdateInertialScroll(double InCurrentTime, float InDeltaTime)
 {
 	InertialScrollManager.UpdateScrollVelocity(InDeltaTime);
 
@@ -445,14 +445,14 @@ EActiveTickReturnType SScrollBox::UpdateInertialScroll(double InCurrentTime, flo
 			ScrollBy(CachedGeometry, ScrollVelocity * InDeltaTime, true);
 		}
 
-		return EActiveTickReturnType::KeepTicking;
+		return EActiveTimerReturnType::Continue;
 	}
 
 	// Clear the scroll velocity so there isn't any delayed minuscule movement
 	InertialScrollManager.ClearScrollVelocity();
 
-	bIsActiveTickRegistered = false;
-	return EActiveTickReturnType::StopTicking;
+	bIsActiveTimerRegistered = false;
+	return EActiveTimerReturnType::Stop;
 }
 
 void SScrollBox::Tick( const FGeometry& AllottedGeometry, const double InCurrentTime, const float InDeltaTime )
@@ -485,11 +485,11 @@ void SScrollBox::Tick( const FGeometry& AllottedGeometry, const double InCurrent
 		// We cannot scroll, so ensure that there is no offset.
 		ScrollPanel->PhysicalOffset = 0.0f;
 	}
-	else if (bIsScrolling && !bIsActiveTickRegistered)
+	else if (bIsScrolling && !bIsActiveTimerRegistered)
 	{
-		// If scrolling and the scrollbar is needed, make sure the active tick is registered (possible b/c we may need to scroll when our geometry changes)
-		bIsActiveTickRegistered = true;
-		RegisterActiveTick(0.f, FWidgetActiveTickDelegate::CreateSP(this, &SScrollBox::UpdateInertialScroll));
+		// If scrolling and the scrollbar is needed, make sure the active timer is registered (possible b/c we may need to scroll when our geometry changes)
+		bIsActiveTimerRegistered = true;
+		RegisterActiveTimer(0.f, FWidgetActiveTimerDelegate::CreateSP(this, &SScrollBox::UpdateInertialScroll));
 	}
 }
 
@@ -517,13 +517,13 @@ FReply SScrollBox::OnMouseButtonUp( const FGeometry& MyGeometry, const FPointerE
 {
 	if ( MouseEvent.GetEffectingButton() == EKeys::RightMouseButton )
 	{
-		if (!bIsActiveTickRegistered && IsRightClickScrolling())
+		if (!bIsActiveTimerRegistered && IsRightClickScrolling())
 		{
-			// Register the active tick to handle the inertial scrolling
+			// Register the active timer to handle the inertial scrolling
 			CachedGeometry = MyGeometry;
 			bIsScrolling = true;
-			bIsActiveTickRegistered = true;
-			RegisterActiveTick(0.f, FWidgetActiveTickDelegate::CreateSP(this, &SScrollBox::UpdateInertialScroll));
+			bIsActiveTimerRegistered = true;
+			RegisterActiveTimer(0.f, FWidgetActiveTimerDelegate::CreateSP(this, &SScrollBox::UpdateInertialScroll));
 		}
 
 		AmountScrolledWhileRightMouseDown = 0;
@@ -607,13 +607,13 @@ FReply SScrollBox::OnMouseWheel( const FGeometry& MyGeometry, const FPointerEven
 
 		const bool bScrollWasHandled = this->ScrollBy( MyGeometry, -MouseEvent.GetWheelDelta()*WheelScrollAmount );
 
-		if (bScrollWasHandled && !bIsActiveTickRegistered)
+		if (bScrollWasHandled && !bIsActiveTimerRegistered)
 		{
-			// Register the active tick to handle the inertial scrolling
+			// Register the active timer to handle the inertial scrolling
 			CachedGeometry = MyGeometry;
 			bIsScrolling = true;
-			bIsActiveTickRegistered = true;
-			RegisterActiveTick(0.f, FWidgetActiveTickDelegate::CreateSP(this, &SScrollBox::UpdateInertialScroll));
+			bIsActiveTimerRegistered = true;
+			RegisterActiveTimer(0.f, FWidgetActiveTimerDelegate::CreateSP(this, &SScrollBox::UpdateInertialScroll));
 		}
 
 		return bScrollWasHandled ? FReply::Handled() : FReply::Unhandled();
diff --git a/Engine/Source/Runtime/Slate/Private/Widgets/Text/SInlineEditableTextBlock.cpp b/Engine/Source/Runtime/Slate/Private/Widgets/Text/SInlineEditableTextBlock.cpp
--- a/Engine/Source/Runtime/Slate/Private/Widgets/Text/SInlineEditableTextBlock.cpp
+++ b/Engine/Source/Runtime/Slate/Private/Widgets/Text/SInlineEditableTextBlock.cpp
@@ -134,9 +134,9 @@ FReply SInlineEditableTextBlock::OnMouseButtonDown( const FGeometry& MyGeometry,
 
 	if(IsSelected.IsBound())
 	{
-		if(IsSelected.Execute() && !bIsReadOnly.Get() && !ActiveTickHandle.IsValid())
+		if(IsSelected.Execute() && !bIsReadOnly.Get() && !ActiveTimerHandle.IsValid())
 		{
-			RegisterActiveTick(0.5f, FWidgetActiveTickDelegate::CreateSP(this, &SInlineEditableTextBlock::TriggerEditMode));
+			RegisterActiveTimer(0.5f, FWidgetActiveTimerDelegate::CreateSP(this, &SInlineEditableTextBlock::TriggerEditMode));
 		}
 	}
 	else
@@ -156,28 +156,28 @@ FReply SInlineEditableTextBlock::OnMouseButtonDown( const FGeometry& MyGeometry,
 FReply SInlineEditableTextBlock::OnDragOver( const FGeometry& MyGeometry, const FDragDropEvent& DragDropEvent )
 {
 	// Cancel during a drag over event, otherwise the widget may enter edit mode.
-	auto PinnedActiveTickHandle = ActiveTickHandle.Pin();
-	if (PinnedActiveTickHandle.IsValid())
+	auto PinnedActiveTimerHandle = ActiveTimerHandle.Pin();
+	if (PinnedActiveTimerHandle.IsValid())
 	{
-		UnRegisterActiveTick(PinnedActiveTickHandle.ToSharedRef());
+		UnRegisterActiveTimer(PinnedActiveTimerHandle.ToSharedRef());
 	}
 	return FReply::Unhandled();
 }
 
 FReply SInlineEditableTextBlock::OnMouseButtonDoubleClick( const FGeometry& InMyGeometry, const FPointerEvent& InMouseEvent )
 {
-	auto PinnedActiveTickHandle = ActiveTickHandle.Pin();
-	if (PinnedActiveTickHandle.IsValid())
+	auto PinnedActiveTimerHandle = ActiveTimerHandle.Pin();
+	if (PinnedActiveTimerHandle.IsValid())
 	{
-		UnRegisterActiveTick(PinnedActiveTickHandle.ToSharedRef());
+		UnRegisterActiveTimer(PinnedActiveTimerHandle.ToSharedRef());
 	}
 	return FReply::Unhandled();
 }
 
-EActiveTickReturnType SInlineEditableTextBlock::TriggerEditMode(double InCurrentTime, float InDeltaTime)
+EActiveTimerReturnType SInlineEditableTextBlock::TriggerEditMode(double InCurrentTime, float InDeltaTime)
 {
 	EnterEditingMode();
-	return EActiveTickReturnType::StopTicking;
+	return EActiveTimerReturnType::Stop;
 }
 
 FReply SInlineEditableTextBlock::OnKeyDown( const FGeometry& MyGeometry, const FKeyEvent& InKeyEvent )
diff --git a/Engine/Source/Runtime/Slate/Private/Widgets/Views/STableViewBase.cpp b/Engine/Source/Runtime/Slate/Private/Widgets/Views/STableViewBase.cpp
--- a/Engine/Source/Runtime/Slate/Private/Widgets/Views/STableViewBase.cpp
+++ b/Engine/Source/Runtime/Slate/Private/Widgets/Views/STableViewBase.cpp
@@ -12,7 +12,7 @@ namespace ListConstants
 void STableViewBase::ConstructChildren( const TAttribute<float>& InItemWidth, const TAttribute<float>& InItemHeight, const TAttribute<EListItemAlignment>& InItemAlignment, const TSharedPtr<SHeaderRow>& InHeaderRow, const TSharedPtr<SScrollBar>& InScrollBar )
 {
 	bItemsNeedRefresh = true;
-	//RegisterActiveTick(0.f, FWidgetActiveTickDelegate::CreateSP(this, &STableViewBase::EnsureTickToRefresh));
+	//RegisterActiveTimer(0.f, FWidgetActiveTimerDelegate::CreateSP(this, &STableViewBase::EnsureTickToRefresh));
 
 	HeaderRow = InHeaderRow;
 
@@ -151,7 +151,7 @@ static FEndOfListResult ComputeOffsetForEndOfList( const FGeometry& ListPanelGeo
 	return FEndOfListResult( OffsetFromEndOfList, ItemsAboveView );
 }
 
-EActiveTickReturnType STableViewBase::UpdateInertialScroll(double InCurrentTime, float InDeltaTime)
+EActiveTimerReturnType STableViewBase::UpdateInertialScroll(double InCurrentTime, float InDeltaTime)
 {
 	bool bKeepTicking = false;
 	if (ItemsPanel.IsValid())
@@ -202,16 +202,16 @@ EActiveTickReturnType STableViewBase::UpdateInertialScroll(double InCurrentTime,
 		TickScrollDelta = 0.f;
 	}
 
-	bIsScrollingActiveTickRegistered = bKeepTicking;
-	return bKeepTicking ? EActiveTickReturnType::KeepTicking : EActiveTickReturnType::StopTicking;
+	bIsScrollingActiveTimerRegistered = bKeepTicking;
+	return bKeepTicking ? EActiveTimerReturnType::Continue : EActiveTimerReturnType::Stop;
 }
 
-EActiveTickReturnType STableViewBase::EnsureTickToRefresh(double InCurrentTime, float InDeltaTime)
+EActiveTimerReturnType STableViewBase::EnsureTickToRefresh(double InCurrentTime, float InDeltaTime)
 {
 	// Actual refresh isn't implemented here as it can be needed in response to changes in the panel geometry.
 	// Since that isn't known until Tick (called after this when registered), refreshing here could result in two refreshes in one frame.
 
-	return EActiveTickReturnType::StopTicking;
+	return EActiveTimerReturnType::Stop;
 }
 
 void STableViewBase::Tick( const FGeometry& AllottedGeometry, const double InCurrentTime, const float InDeltaTime )
@@ -382,11 +382,11 @@ FReply STableViewBase::OnMouseMove( const FGeometry& MyGeometry, const FPointerE
 		// the mouse and dragging the view?
 		if( IsRightClickScrolling() )
 		{
-			// Make sure the active tick is registered to update the inertial scroll
-			if (!bIsScrollingActiveTickRegistered)
+			// Make sure the active timer is registered to update the inertial scroll
+			if (!bIsScrollingActiveTimerRegistered)
 			{
-				bIsScrollingActiveTickRegistered = true;
-				RegisterActiveTick(0.f, FWidgetActiveTickDelegate::CreateSP(this, &STableViewBase::UpdateInertialScroll));
+				bIsScrollingActiveTimerRegistered = true;
+				RegisterActiveTimer(0.f, FWidgetActiveTimerDelegate::CreateSP(this, &STableViewBase::UpdateInertialScroll));
 			}
 
 			TickScrollDelta -= ScrollByAmount;
@@ -552,7 +552,7 @@ void STableViewBase::RequestListRefresh()
 	if (!bItemsNeedRefresh)
 	{
 		bItemsNeedRefresh = true;
-		RegisterActiveTick(0.f, FWidgetActiveTickDelegate::CreateSP(this, &STableViewBase::EnsureTickToRefresh));
+		RegisterActiveTimer(0.f, FWidgetActiveTimerDelegate::CreateSP(this, &STableViewBase::EnsureTickToRefresh));
 	}
 	ItemsPanel->SetRefreshPending(true);
 }
@@ -595,7 +595,7 @@ STableViewBase::STableViewBase( ETableViewMode::Type InTableViewMode )
 	, SelectionMode( ESelectionMode::Multi )
 	, SoftwareCursorPosition( ForceInitToZero )
 	, bShowSoftwareCursor( false )
-	, bIsScrollingActiveTickRegistered( false )
+	, bIsScrollingActiveTimerRegistered( false )
 	, Overscroll()
 	, AllowOverscroll(EAllowOverscroll::Yes)
 	, bItemsNeedRefresh( false )
diff --git a/Engine/Source/Runtime/Slate/Public/Framework/Application/SlateApplication.h b/Engine/Source/Runtime/Slate/Public/Framework/Application/SlateApplication.h
--- a/Engine/Source/Runtime/Slate/Public/Framework/Application/SlateApplication.h
+++ b/Engine/Source/Runtime/Slate/Public/Framework/Application/SlateApplication.h
@@ -1119,7 +1119,7 @@ private:
 
 	TSet<FKey> PressedMouseButtons;
 
-	/** After processing an event or performing an active tick, we need to synthesize a mouse move. @see SynthesizeMouseMove */
+	/** After processing an event or performing an active timer, we need to synthesize a mouse move. @see SynthesizeMouseMove */
 	int32 SynthesizeMouseMovePending;
 
 	/** true when the slate app is active; i.e. the current foreground window is from our Slate app*/
diff --git a/Engine/Source/Runtime/Slate/Public/Widgets/Docking/SDockTab.h b/Engine/Source/Runtime/Slate/Public/Widgets/Docking/SDockTab.h
--- a/Engine/Source/Runtime/Slate/Public/Widgets/Docking/SDockTab.h
+++ b/Engine/Source/Runtime/Slate/Public/Widgets/Docking/SDockTab.h
@@ -275,10 +275,10 @@ protected:
 
 private:
 	/** Activates the tab in its tab well */
-	EActiveTickReturnType TriggerActivateTab( double InCurrentTime, float InDeltaTime );
+	EActiveTimerReturnType TriggerActivateTab( double InCurrentTime, float InDeltaTime );
 
 	/** The handle to the active tab activation tick */
-	TWeakPtr<FActiveTickHandle> ActiveTickHandle;
+	TWeakPtr<FActiveTimerHandle> ActiveTimerHandle;
 
 protected:
 
diff --git a/Engine/Source/Runtime/Slate/Public/Widgets/Input/SEditableText.h b/Engine/Source/Runtime/Slate/Public/Widgets/Input/SEditableText.h
--- a/Engine/Source/Runtime/Slate/Public/Widgets/Input/SEditableText.h
+++ b/Engine/Source/Runtime/Slate/Public/Widgets/Input/SEditableText.h
@@ -464,7 +464,7 @@ protected:
 
 private:
 	/** Animates the caret and highlight selection springs */
-	EActiveTickReturnType AnimateSpringsWhileFocused(double InCurrentTime, float InDeltaTime);
+	EActiveTimerReturnType AnimateSpringsWhileFocused(double InCurrentTime, float InDeltaTime);
 	
 	/** @return Whether the editable text should appear focused */
 	bool ShouldAppearFocused() const;
diff --git a/Engine/Source/Runtime/Slate/Public/Widgets/Input/SSearchBox.h b/Engine/Source/Runtime/Slate/Public/Widgets/Input/SSearchBox.h
--- a/Engine/Source/Runtime/Slate/Public/Widgets/Input/SSearchBox.h
+++ b/Engine/Source/Runtime/Slate/Public/Widgets/Input/SSearchBox.h
@@ -73,7 +73,7 @@ private:
 private:
 
 	/** Fires the text changed delegate on a delay */
-	EActiveTickReturnType TriggerOnTextChanged( double InCurrentTime, float InDeltaTime );
+	EActiveTimerReturnType TriggerOnTextChanged( double InCurrentTime, float InDeltaTime );
 
 	/** @return should we show the X to clear search? */
 	EVisibility GetXVisibility() const;
@@ -94,7 +94,7 @@ private:
 	static const double FilterDelayAfterTyping;
 
 	/** Handle to the active trigger text changed timer */
-	TWeakPtr<FActiveTickHandle> ActiveTickHandle;
+	TWeakPtr<FActiveTimerHandle> ActiveTimerHandle;
 
 	/** Delegate that is invoked when the user does next or previous */
 	FOnSearch OnSearchDelegate;
@@ -113,7 +113,7 @@ private:
 
 	FText LastPendingTextChangedValue;
 
-	bool bIsActiveTickRegistered : 1;
+	bool bIsActiveTimerRegistered : 1;
 };
 
 
diff --git a/Engine/Source/Runtime/Slate/Public/Widgets/Layout/SScrollBox.h b/Engine/Source/Runtime/Slate/Public/Widgets/Layout/SScrollBox.h
--- a/Engine/Source/Runtime/Slate/Public/Widgets/Layout/SScrollBox.h
+++ b/Engine/Source/Runtime/Slate/Public/Widgets/Layout/SScrollBox.h
@@ -189,8 +189,8 @@ private:
 	/** Does the user need a hint that they can scroll to the end of the list? */
 	FSlateColor GetEndShadowOpacity() const;
 
-	/** Active tick to update inertial scrolling as needed */
-	EActiveTickReturnType UpdateInertialScroll(double InCurrentTime, float InDeltaTime);
+	/** Active timer to update inertial scrolling as needed */
+	EActiveTimerReturnType UpdateInertialScroll(double InCurrentTime, float InDeltaTime);
 
 private:
 
@@ -221,7 +221,7 @@ private:
 	/** Style resource for the scrollbar */
 	const FScrollBarStyle* ScrollBarStyle;
 
-	/** Cached geometry for use with the active tick */
+	/** Cached geometry for use with the active timer */
 	FGeometry CachedGeometry;
 
 	/**	Whether the software cursor should be drawn in the viewport */
@@ -239,8 +239,8 @@ private:
 	/** If true, will scroll to the end next Tick */
 	bool bScrollToEnd : 1;
 
-	/** Whether the active tick to update the inertial scroll is registered */
-	bool bIsActiveTickRegistered : 1;
+	/** Whether the active timer to update the inertial scroll is registered */
+	bool bIsActiveTimerRegistered : 1;
 };
 
 
diff --git a/Engine/Source/Runtime/Slate/Public/Widgets/Text/SInlineEditableTextBlock.h b/Engine/Source/Runtime/Slate/Public/Widgets/Text/SInlineEditableTextBlock.h
--- a/Engine/Source/Runtime/Slate/Public/Widgets/Text/SInlineEditableTextBlock.h
+++ b/Engine/Source/Runtime/Slate/Public/Widgets/Text/SInlineEditableTextBlock.h
@@ -158,9 +158,9 @@ protected:
 	TAttribute< bool > bIsReadOnly;
 
 private:
-	/** Active tick to trigger entry into edit mode after a delay */
-	EActiveTickReturnType TriggerEditMode(double InCurrentTime, float InDeltaTime);
+	/** Active timer to trigger entry into edit mode after a delay */
+	EActiveTimerReturnType TriggerEditMode(double InCurrentTime, float InDeltaTime);
 
-	/** The handle to the active tick */
-	TWeakPtr<FActiveTickHandle> ActiveTickHandle;
+	/** The handle to the active timer */
+	TWeakPtr<FActiveTimerHandle> ActiveTimerHandle;
 };
diff --git a/Engine/Source/Runtime/Slate/Public/Widgets/Views/STableViewBase.h b/Engine/Source/Runtime/Slate/Public/Widgets/Views/STableViewBase.h
--- a/Engine/Source/Runtime/Slate/Public/Widgets/Views/STableViewBase.h
+++ b/Engine/Source/Runtime/Slate/Public/Widgets/Views/STableViewBase.h
@@ -277,16 +277,16 @@ private:
 	/** Check whether the current state of the table warrants inertial scroll by the specified amount */
 	bool CanUseInertialScroll( float ScrollAmount ) const;
 
-	/** Active tick to update the inertial scroll */
-	EActiveTickReturnType UpdateInertialScroll(double InCurrentTime, float InDeltaTime);
+	/** Active timer to update the inertial scroll */
+	EActiveTimerReturnType UpdateInertialScroll(double InCurrentTime, float InDeltaTime);
 
-	/** One-off active tick to refresh the contents of the table as needed */
-	EActiveTickReturnType EnsureTickToRefresh(double InCurrentTime, float InDeltaTime);
+	/** One-off active timer to refresh the contents of the table as needed */
+	EActiveTimerReturnType EnsureTickToRefresh(double InCurrentTime, float InDeltaTime);
 
-	/** Whether the active tick to update the inertial scrolling is currently registered */
-	bool bIsScrollingActiveTickRegistered;
+	/** Whether the active timer to update the inertial scrolling is currently registered */
+	bool bIsScrollingActiveTimerRegistered;
 
-	/** Cached geometry for use by the active tick */
+	/** Cached geometry for use by the active timer */
 	FGeometry CachedGeometry;
 
 protected:
diff --git a/Engine/Source/Runtime/SlateCore/Private/Animation/CurveSequence.cpp b/Engine/Source/Runtime/SlateCore/Private/Animation/CurveSequence.cpp
--- a/Engine/Source/Runtime/SlateCore/Private/Animation/CurveSequence.cpp
+++ b/Engine/Source/Runtime/SlateCore/Private/Animation/CurveSequence.cpp
@@ -29,10 +29,10 @@ FCurveSequence::FCurveSequence( const float InStartTimeSeconds, const float InDu
 
 FCurveSequence::~FCurveSequence()
 {
-	// If the curve sequence is destroyed before the owning widget, unregister the active tick
-	if ( OwnerWidget.IsValid() && ActiveTickHandle.IsValid() )
+	// If the curve sequence is destroyed before the owning widget, unregister the active timer
+	if ( OwnerWidget.IsValid() && ActiveTimerHandle.IsValid() )
 	{
-		OwnerWidget.Pin()->UnRegisterActiveTick(ActiveTickHandle.Pin().ToSharedRef());
+		OwnerWidget.Pin()->UnRegisterActiveTimer(ActiveTimerHandle.Pin().ToSharedRef());
 	}
 }
 
@@ -67,7 +67,7 @@ void FCurveSequence::Play( const float StartAtTime )
 
 void FCurveSequence::Play( const TSharedRef<SWidget>& InOwnerWidget, bool bPlayLooped, const float StartAtTime )
 {
-	RegisterActiveTickIfNeeded(InOwnerWidget);
+	RegisterActiveTimerIfNeeded(InOwnerWidget);
 	bIsLooping = bPlayLooped;
 	bIsPaused = false;
 	
@@ -107,7 +107,7 @@ void FCurveSequence::PlayReverse( const float StartAtTime )
 
 void FCurveSequence::PlayReverse( const TSharedRef<SWidget>& InOwnerWidget, bool bPlayLooped, const float StartAtTime )
 {
-	RegisterActiveTickIfNeeded(InOwnerWidget);
+	RegisterActiveTimerIfNeeded(InOwnerWidget);
 	bIsLooping = bPlayLooped;
 	bIsPaused = false;
 
@@ -135,7 +135,7 @@ void FCurveSequence::Resume()
 		if ( PinnedOwner.IsValid() )
 		{
 			bIsPaused = false;
-			RegisterActiveTickIfNeeded( PinnedOwner.ToSharedRef() );
+			RegisterActiveTimerIfNeeded( PinnedOwner.ToSharedRef() );
 
 			// Update the start time to be the same relative to the current time as it was when paused
 			const double NewStartTime = FSlateApplicationBase::Get().GetCurrentTime() - ( PauseTime - StartTime );
@@ -257,20 +257,20 @@ const FCurveSequence::FSlateCurve& FCurveSequence::GetCurve( int32 CurveIndex )
 	return Curves[CurveIndex];
 }
 
-void FCurveSequence::RegisterActiveTickIfNeeded(TSharedRef<SWidget> InOwnerWidget)
+void FCurveSequence::RegisterActiveTimerIfNeeded(TSharedRef<SWidget> InOwnerWidget)
 {
-	// Register the active tick
-	if ( !ActiveTickHandle.IsValid() )
+	// Register the active timer
+	if ( !ActiveTimerHandle.IsValid() )
 	{
-		ActiveTickHandle = InOwnerWidget->RegisterActiveTick(0.f, FWidgetActiveTickDelegate::CreateRaw(this, &FCurveSequence::EnsureSlateTickDuringAnimation));
+		ActiveTimerHandle = InOwnerWidget->RegisterActiveTimer(0.f, FWidgetActiveTimerDelegate::CreateRaw(this, &FCurveSequence::EnsureSlateTickDuringAnimation));
 
 		// Save the reference in case we need to take care of unregistering the tick
 		OwnerWidget = InOwnerWidget;
 	}
 }
 
-EActiveTickReturnType FCurveSequence::EnsureSlateTickDuringAnimation( double InCurrentTime, float InDeltaTime )
+EActiveTimerReturnType FCurveSequence::EnsureSlateTickDuringAnimation( double InCurrentTime, float InDeltaTime )
 {
-	// Keep ticking until play stops
-	return IsPlaying() ? EActiveTickReturnType::KeepTicking : EActiveTickReturnType::StopTicking;
+	// Keep Slate ticking until play stops
+	return IsPlaying() ? EActiveTimerReturnType::Continue : EActiveTimerReturnType::Stop;
 }
\ No newline at end of file
diff --git a/Engine/Source/Runtime/SlateCore/Private/Application/ActiveTickHandle.cpp b/Engine/Source/Runtime/SlateCore/Private/Application/ActiveTickHandle.cpp
deleted file mode 100644
diff --git a/Engine/Source/Runtime/SlateCore/Private/Application/ActiveTickHandle.h b/Engine/Source/Runtime/SlateCore/Private/Application/ActiveTickHandle.h
deleted file mode 100644
diff --git a/Engine/Source/Runtime/SlateCore/Private/Application/ActiveTimerHandle.cpp b/Engine/Source/Runtime/SlateCore/Private/Application/ActiveTimerHandle.cpp
new file mode 100644
--- /dev/null
+++ b/Engine/Source/Runtime/SlateCore/Private/Application/ActiveTimerHandle.cpp
@@ -0,0 +1,57 @@
+// Copyright 1998-2014 Epic Games, Inc. All Rights Reserved.
+
+#include "SlateCorePrivatePCH.h"
+#include "ActiveTimerHandle.h"
+
+
+FActiveTimerHandle::FActiveTimerHandle(float InExecutionPeriod, FWidgetActiveTimerDelegate InTimerFunction, double InNextExecutionTime)
+	: ExecutionPeriod(InExecutionPeriod)
+	, TimerFunction(InTimerFunction)
+	, NextExecutionTime(InNextExecutionTime)
+	, bExecutionPending(false)
+{
+}
+
+bool FActiveTimerHandle::IsPendingExecution() const
+{
+	return bExecutionPending;
+}
+
+bool FActiveTimerHandle::UpdateExecutionPendingState( double CurrentTime )
+{
+	bExecutionPending = IsReady( CurrentTime );
+	return bExecutionPending;
+}
+
+EActiveTimerReturnType FActiveTimerHandle::ExecuteIfPending( double CurrentTime, float DeltaTime )
+{
+	if ( bExecutionPending )
+	{
+		bExecutionPending = false;
+
+		// before we do anything, check validity of the delegate.
+		if ( !TimerFunction.IsBound() )
+		{
+			// if the handle is no longer valid, we must assume it is invalid (say, owning widget was destroyed) and needs to be removed.
+			return EActiveTimerReturnType::Stop;
+		}
+
+		// update time in a while loop so we skip any times that we may have missed due to slow frames.
+		if ( ExecutionPeriod > 0.0 )
+		{
+			while ( CurrentTime >= NextExecutionTime )
+			{
+				NextExecutionTime += ExecutionPeriod;
+			}
+		}
+		// timer is updated, now call the delegate
+		return TimerFunction.Execute( CurrentTime, DeltaTime );
+	}
+	// if execution is not pending, then we need to continue until it is
+	return EActiveTimerReturnType::Continue;
+}
+
+bool FActiveTimerHandle::IsReady(double CurrentTime) const
+{
+	return CurrentTime >= NextExecutionTime;
+}
\ No newline at end of file
diff --git a/Engine/Source/Runtime/SlateCore/Private/Application/ActiveTimerHandle.h b/Engine/Source/Runtime/SlateCore/Private/Application/ActiveTimerHandle.h
new file mode 100644
--- /dev/null
+++ b/Engine/Source/Runtime/SlateCore/Private/Application/ActiveTimerHandle.h
@@ -0,0 +1,48 @@
+// Copyright 1998-2014 Epic Games, Inc. All Rights Reserved.
+
+#pragma once
+
+#include "SlateCore.h"
+
+/** Stores info about an active timer delegate for a widget. */
+class FActiveTimerHandle
+{
+public:
+	/** Public ctor that initializes a new active timer handle. Not intended to be called by user code. */
+	FActiveTimerHandle(float InExecutionPeriod, FWidgetActiveTimerDelegate InTimerFunction, double InNextExecutionTime);
+
+	/** @return True if the active timer is pending execution */
+	bool IsPendingExecution() const;
+
+	/**
+	* Updates the pending state of the active timer based on the current time
+	* @return True if the active timer is pending execution
+	*/
+	bool UpdateExecutionPendingState( double CurrentTime );
+
+	/**
+	 * Execute the bound delegate if the active timer is bPendingExecution.
+	 * @return EActiveTimerReturnType::Stop if the handle should be auto-unregistered
+	 *         (either the TimerFunction indicated this or the TimerFunction is unbound).
+	 */
+	EActiveTimerReturnType ExecuteIfPending( double CurrentTime, float DeltaTime );
+
+private:
+	/** @return true if the tick handle is ready to have its delegate executed. */
+	bool IsReady( double CurrentTime ) const;
+
+private:
+	/** The period between executions of the timer delegate */
+	float ExecutionPeriod;
+
+	/** The delegate to the active timer function */
+	FWidgetActiveTimerDelegate TimerFunction;
+	
+	/** The next time at which the TimerFunction will execute */
+	double NextExecutionTime;
+
+	/** True if execution of the TimerFunction is pending */
+	bool bExecutionPending;
+};
+
+
diff --git a/Engine/Source/Runtime/SlateCore/Private/Application/SlateApplicationBase.cpp b/Engine/Source/Runtime/SlateCore/Private/Application/SlateApplicationBase.cpp
--- a/Engine/Source/Runtime/SlateCore/Private/Application/SlateApplicationBase.cpp
+++ b/Engine/Source/Runtime/SlateCore/Private/Application/SlateApplicationBase.cpp
@@ -2,7 +2,7 @@
 
 #include "SlateCorePrivatePCH.h"
 #include "HittestGrid.h"
-#include "ActiveTickHandle.h"
+#include "ActiveTimerHandle.h"
 
 
 /* Static initialization
@@ -31,43 +31,43 @@ const FHitTesting& FSlateApplicationBase::GetHitTesting() const
 	return HitTesting;
 }
 
-void FSlateApplicationBase::RegisterActiveTick( const TSharedRef<FActiveTickHandle>& ActiveTickHandle )
+void FSlateApplicationBase::RegisterActiveTimer( const TSharedRef<FActiveTimerHandle>& ActiveTimerHandle )
 {
-	ActiveTickHandles.Add(ActiveTickHandle);
+	ActiveTimerHandles.Add(ActiveTimerHandle);
 }
 
-void FSlateApplicationBase::UnRegisterActiveTick( const TSharedRef<FActiveTickHandle>& ActiveTickHandle )
+void FSlateApplicationBase::UnRegisterActiveTimer( const TSharedRef<FActiveTimerHandle>& ActiveTimerHandle )
 {
-	ActiveTickHandles.Remove(ActiveTickHandle);
+	ActiveTimerHandles.Remove(ActiveTimerHandle);
 }
 
-bool FSlateApplicationBase::AnyActiveTicksArePending()
+bool FSlateApplicationBase::AnyActiveTimersArePending()
 {
 	// first remove any tick handles that may have become invalid.
 	// If we didn't remove invalid handles here, they would never get removed because
 	// we don't force widgets to UnRegister before they are destroyed.
-	ActiveTickHandles.RemoveAll([](const TWeakPtr<FActiveTickHandle>& ActiveTickHandle)
+	ActiveTimerHandles.RemoveAll([](const TWeakPtr<FActiveTimerHandle>& ActiveTimerHandle)
 	{
 		// only check the weak pointer to the handle. Just want to make sure to clear out any widgets that have since been deleted.
-		return !ActiveTickHandle.IsValid();
+		return !ActiveTimerHandle.IsValid();
 	});
 
 	// The rest are valid. Update their pending status and see if any are ready.
 	const double CurrentTime = GetCurrentTime();
 	bool bAnyTickReady = false;
-	for ( auto& ActiveTickInfo : ActiveTickHandles )
+	for ( auto& ActiveTimerInfo : ActiveTimerHandles )
 	{
-		auto ActiveTickInfoPinned = ActiveTickInfo.Pin();
-		check( ActiveTickInfoPinned.IsValid() );
+		auto ActiveTimerInfoPinned = ActiveTimerInfo.Pin();
+		check( ActiveTimerInfoPinned.IsValid() );
 
-		// If an active tick is still pending execution from last frame, it is collapsed 
+		// If an active timer is still pending execution from last frame, it is collapsed 
 		// or otherwise blocked from ticking. Disregard until it executes.
-		if ( ActiveTickInfoPinned->IsPendingExecution() )
+		if ( ActiveTimerInfoPinned->IsPendingExecution() )
 		{
 			continue;
 		}
 
-		if ( ActiveTickInfoPinned->UpdateExecutionPendingState( CurrentTime ) )
+		if ( ActiveTimerInfoPinned->UpdateExecutionPendingState( CurrentTime ) )
 		{
 			bAnyTickReady = true;
 		}
diff --git a/Engine/Source/Runtime/SlateCore/Private/Widgets/SWidget.cpp b/Engine/Source/Runtime/SlateCore/Private/Widgets/SWidget.cpp
--- a/Engine/Source/Runtime/SlateCore/Private/Widgets/SWidget.cpp
+++ b/Engine/Source/Runtime/SlateCore/Private/Widgets/SWidget.cpp
@@ -3,7 +3,7 @@
 #include "SlateCorePrivatePCH.h"
 #include "Widgets/SWidget.h"
 #include "Input/Events.h"
-#include "ActiveTickHandle.h"
+#include "ActiveTimerHandle.h"
 
 DECLARE_CYCLE_STAT(TEXT("OnPaint"), STAT_SlateOnPaint, STATGROUP_Slate);
 DECLARE_CYCLE_STAT(TEXT("ArrangeChildren"), STAT_SlateArrangeChildren, STATGROUP_Slate);
@@ -348,8 +348,8 @@ void SWidget::TickWidgetsRecursively( const FGeometry& AllottedGeometry, const d
 {
 	INC_DWORD_STAT(STAT_SlateNumTickedWidgets);
 
-	// Execute any pending active ticks for this widget, followed by the passive tick
-	ExecuteActiveTicks( InCurrentTime, InDeltaTime );
+	// Execute any pending active timers for this widget, followed by the passive tick
+	ExecuteActiveTimers( InCurrentTime, InDeltaTime );
 	Tick( AllottedGeometry, InCurrentTime, InDeltaTime );
 
 	// Gather all children, whether they're visible or not.  We need to allow invisible widgets to
@@ -634,7 +634,7 @@ int32 SWidget::Paint(const FPaintArgs& Args, const FGeometry& AllottedGeometry,
 		TickGeometry.AppendTransform( FSlateLayoutTransform(Args.GetWindowToDesktopTransform()) );
 
 		SWidget* MutableThis = const_cast<SWidget*>(this);
-		MutableThis->ExecuteActiveTicks( Args.GetCurrentTime(), Args.GetDeltaTime() );
+		MutableThis->ExecuteActiveTimers( Args.GetCurrentTime(), Args.GetDeltaTime() );
 		MutableThis->Tick( TickGeometry, Args.GetCurrentTime(), Args.GetDeltaTime() );
 	}
 
@@ -672,46 +672,46 @@ void SWidget::ArrangeChildren(const FGeometry& AllottedGeometry, FArrangedChildr
 	OnArrangeChildren(AllottedGeometry, ArrangedChildren);
 }
 
-TSharedRef<FActiveTickHandle> SWidget::RegisterActiveTick(float TickPeriod, FWidgetActiveTickDelegate TickFunction)
+TSharedRef<FActiveTimerHandle> SWidget::RegisterActiveTimer(float TickPeriod, FWidgetActiveTimerDelegate TickFunction)
 {
-	TSharedRef<FActiveTickHandle> ActiveTickHandle = MakeShareable(new FActiveTickHandle(TickPeriod, TickFunction, FSlateApplicationBase::Get().GetCurrentTime() + TickPeriod));
-	FSlateApplicationBase::Get().RegisterActiveTick(ActiveTickHandle);
-	ActiveTicks.Add(ActiveTickHandle);
-	return ActiveTickHandle;
+	TSharedRef<FActiveTimerHandle> ActiveTimerHandle = MakeShareable(new FActiveTimerHandle(TickPeriod, TickFunction, FSlateApplicationBase::Get().GetCurrentTime() + TickPeriod));
+	FSlateApplicationBase::Get().RegisterActiveTimer(ActiveTimerHandle);
+	ActiveTimers.Add(ActiveTimerHandle);
+	return ActiveTimerHandle;
 }
 
-void SWidget::UnRegisterActiveTick(const TSharedRef<FActiveTickHandle>& ActiveTickHandle)
+void SWidget::UnRegisterActiveTimer(const TSharedRef<FActiveTimerHandle>& ActiveTimerHandle)
 {
-	FSlateApplicationBase::Get().UnRegisterActiveTick(ActiveTickHandle);
-	ActiveTicks.Remove(ActiveTickHandle);
+	FSlateApplicationBase::Get().UnRegisterActiveTimer(ActiveTimerHandle);
+	ActiveTimers.Remove(ActiveTimerHandle);
 }
 
-void SWidget::ExecuteActiveTicks(double CurrentTime, float DeltaTime)
+void SWidget::ExecuteActiveTimers(double CurrentTime, float DeltaTime)
 {
 	// loop over the registered tick handles and execute them, removing them if necessary.
-	for (int32 i = 0; i < ActiveTicks.Num();)
+	for (int32 i = 0; i < ActiveTimers.Num();)
 	{
-		EActiveTickReturnType Result = ActiveTicks[i]->ExecuteIfPending( CurrentTime, DeltaTime );
-		if (Result == EActiveTickReturnType::KeepTicking)
+		EActiveTimerReturnType Result = ActiveTimers[i]->ExecuteIfPending( CurrentTime, DeltaTime );
+		if (Result == EActiveTimerReturnType::Continue)
 		{
 			++i;
 		}
 		else
 		{
-			FSlateApplicationBase::Get().UnRegisterActiveTick(ActiveTicks[i]);
-			ActiveTicks.RemoveAt(i);
+			FSlateApplicationBase::Get().UnRegisterActiveTimer(ActiveTimers[i]);
+			ActiveTimers.RemoveAt(i);
 		}
 	}
 }
 
 SWidget::~SWidget()
 {
-	// Unregister all ActiveTicks so they aren't left stranded in the Application's list.
+	// Unregister all ActiveTimers so they aren't left stranded in the Application's list.
 	if (FSlateApplicationBase::IsInitialized())
 	{
-		for (const auto& ActiveTickHandle : ActiveTicks)
+		for (const auto& ActiveTimerHandle : ActiveTimers)
 		{
-			FSlateApplicationBase::Get().UnRegisterActiveTick(ActiveTickHandle);
+			FSlateApplicationBase::Get().UnRegisterActiveTimer(ActiveTimerHandle);
 		}
 	}
 }
diff --git a/Engine/Source/Runtime/SlateCore/Private/Widgets/SWindow.cpp b/Engine/Source/Runtime/SlateCore/Private/Widgets/SWindow.cpp
--- a/Engine/Source/Runtime/SlateCore/Private/Widgets/SWindow.cpp
+++ b/Engine/Source/Runtime/SlateCore/Private/Widgets/SWindow.cpp
@@ -842,9 +842,9 @@ void SWindow::StartMorph()
 	Morpher.bIsActive = true;
 	Morpher.Sequence.JumpToStart();
 
-	if ( !ActiveTickHandle.IsValid() )
+	if ( !ActiveTimerHandle.IsValid() )
 	{
-		ActiveTickHandle = RegisterActiveTick( 0.f, FWidgetActiveTickDelegate::CreateSP( this, &SWindow::TriggerPlayMorphSequence ) );
+		ActiveTimerHandle = RegisterActiveTimer( 0.f, FWidgetActiveTimerDelegate::CreateSP( this, &SWindow::TriggerPlayMorphSequence ) );
 	}
 }
 
@@ -1743,8 +1743,8 @@ EVisibility SWindow::GetWindowContentVisibility() const
 	return (bShouldShowWindowContentDuringOverlay == true || !FullWindowOverlayWidget.IsValid()) ? EVisibility::SelfHitTestInvisible : EVisibility::Hidden;
 };
 
-EActiveTickReturnType SWindow::TriggerPlayMorphSequence( double InCurrentTime, float InDeltaTime )
+EActiveTimerReturnType SWindow::TriggerPlayMorphSequence( double InCurrentTime, float InDeltaTime )
 {
 	Morpher.Sequence.Play( this->AsShared() );
-	return EActiveTickReturnType::StopTicking;
+	return EActiveTimerReturnType::Stop;
 }
\ No newline at end of file
diff --git a/Engine/Source/Runtime/SlateCore/Public/Animation/CurveSequence.h b/Engine/Source/Runtime/SlateCore/Public/Animation/CurveSequence.h
--- a/Engine/Source/Runtime/SlateCore/Public/Animation/CurveSequence.h
+++ b/Engine/Source/Runtime/SlateCore/Public/Animation/CurveSequence.h
@@ -5,7 +5,7 @@
 
 /**
  * A sequence of curves that can be used to drive animations for slate widgets.
- * Active tick registration is handled for the widget being animated when calling play.
+ * Active timer registration is handled for the widget being animated when calling play.
  *
  * Each curve within the sequence has a time offset and a duration.
  * This makes FCurveSequence convenient for crating staggered animations.
@@ -46,7 +46,7 @@ public:
 	/** Default constructor */
 	FCurveSequence( );
 
-	/** Makes sure the active tick is unregistered */
+	/** Makes sure the active timer is unregistered */
 	~FCurveSequence();
 
 	/**
@@ -101,19 +101,19 @@ public:
 	void PlayReverse( const float StartAtTime = 0.0f );
 
 	/**
-	 * Start playing this curve sequence. Registers an active tick with the widget being animated.
+	 * Start playing this curve sequence. Registers an active timer with the widget being animated.
 	 *
 	 * @param InOwnerWidget The widget that is being animated by this sequence.
-	 * @param bPlayLooped True if the curve sequence should play continually on a loop. Note that the active tick will persist until this sequence is paused or jumped to the start/end.
+	 * @param bPlayLooped True if the curve sequence should play continually on a loop. Note that the active timer will persist until this sequence is paused or jumped to the start/end.
 	 * @param StartAtTime The relative time within the animation at which to begin playing (i.e. 0.0f is the beginning).
 	 */
 	void Play( const TSharedRef<SWidget>& InOwnerWidget, bool bPlayLooped = false, const float StartAtTime = 0.0f );
 
 	/**
-	 * Start playing this curve sequence in reverse. Registers an active tick for the widget using the sequence.
+	 * Start playing this curve sequence in reverse. Registers an active timer for the widget using the sequence.
 	 *
 	 * @param InOwnerWidget The widget that is being animated by this sequence.
-	 * @param bPlayLooped True if the curve sequence should play continually on a loop. Note that the active tick will persist until this sequence is paused or jumped to the start/end.
+	 * @param bPlayLooped True if the curve sequence should play continually on a loop. Note that the active timer will persist until this sequence is paused or jumped to the start/end.
 	 * @param StartAtTime The relative time within the animation at which to begin playing (i.e. 0.0f is the beginning).
 	 */
 	void PlayReverse( const TSharedRef<SWidget>& InOwnerWidget, bool bPlayLooped = false, const float StartAtTime = 0.0f );
@@ -196,22 +196,22 @@ protected:
 	void SetStartTime( double InStartTime );
 
 private:
-	/** Helper to take care of registering the active tick */
-	void RegisterActiveTickIfNeeded(TSharedRef<SWidget> InOwnerWidget);
+	/** Helper to take care of registering the active timer */
+	void RegisterActiveTimerIfNeeded(TSharedRef<SWidget> InOwnerWidget);
 
-	/** Hollow active tick to ensure a Slate Tick/Paint pass while the sequence is playing */
-	EActiveTickReturnType EnsureSlateTickDuringAnimation( double InCurrentTime, float InDeltaTime );
+	/** Hollow active timer to ensure a Slate Tick/Paint pass while the sequence is playing */
+	EActiveTimerReturnType EnsureSlateTickDuringAnimation( double InCurrentTime, float InDeltaTime );
 
 private:
 
 	/** 
 	 * Weak reference to the owner widget that is being animated by this curve sequence.
-	 * Necessary to ensure the active tick is unregistered if the sequence is destroyed before/by the owner.
+	 * Necessary to ensure the active timer is unregistered if the sequence is destroyed before/by the owner.
 	 */
 	TWeakPtr<SWidget> OwnerWidget;
 
-	/** The handle to the active tick */
-	TWeakPtr<FActiveTickHandle> ActiveTickHandle;
+	/** The handle to the active timer */
+	TWeakPtr<FActiveTimerHandle> ActiveTimerHandle;
 
 	/** All the curves in this sequence. */
 	TArray<FSlateCurve> Curves;
diff --git a/Engine/Source/Runtime/SlateCore/Public/Application/SlateApplicationBase.h b/Engine/Source/Runtime/SlateCore/Public/Application/SlateApplicationBase.h
--- a/Engine/Source/Runtime/SlateCore/Public/Application/SlateApplicationBase.h
+++ b/Engine/Source/Runtime/SlateCore/Public/Application/SlateApplicationBase.h
@@ -322,24 +322,24 @@ public:
 
 private:
 	/**
-	 * Implementation for active tick registration. See SWidget::RegisterActiveTick.
+	 * Implementation for active timer registration. See SWidget::RegisterActiveTimer.
 	 */
-	void RegisterActiveTick( const TSharedRef<FActiveTickHandle>& ActiveTickHandle );
+	void RegisterActiveTimer( const TSharedRef<FActiveTimerHandle>& ActiveTimerHandle );
 
 	/**
-	 * Implementation for active tick registration. See SWidget::UnRegisterActiveTick.
+	 * Implementation for active timer registration. See SWidget::UnRegisterActiveTimer.
 	 */
-	void UnRegisterActiveTick( const TSharedRef<FActiveTickHandle>& ActiveTickHandle );
+	void UnRegisterActiveTimer( const TSharedRef<FActiveTimerHandle>& ActiveTimerHandle );
 
-	/** The list of active tick handles. */
-	TArray<TWeakPtr<FActiveTickHandle>> ActiveTickHandles;
+	/** The list of active timer handles. */
+	TArray<TWeakPtr<FActiveTimerHandle>> ActiveTimerHandles;
 
 protected:
 	/**
-	 * Used to determine if any active tick handles are ready to fire.
+	 * Used to determine if any active timer handles are ready to fire.
 	 * Means we need to tick slate even if no user interaction.
 	 */
-	bool AnyActiveTicksArePending();
+	bool AnyActiveTimersArePending();
 
 public:
 	const static uint32 CursorPointerIndex;
diff --git a/Engine/Source/Runtime/SlateCore/Public/SlateCore.h b/Engine/Source/Runtime/SlateCore/Public/SlateCore.h
--- a/Engine/Source/Runtime/SlateCore/Public/SlateCore.h
+++ b/Engine/Source/Runtime/SlateCore/Public/SlateCore.h
@@ -60,14 +60,14 @@ DECLARE_MEMORY_STAT_EXTERN(TEXT("Index Buffer Memory"), STAT_SlateIndexBufferMem
 
 /* Forward declarations
  *****************************************************************************/
-class FActiveTickHandle;
-enum EActiveTickReturnType : uint8;
+class FActiveTimerHandle;
+enum class EActiveTimerReturnType : uint8;
 
 /* Delegates
  *****************************************************************************/
 
 /** Delegate used for registering a widget for tick services. */
-DECLARE_DELEGATE_RetVal_TwoParams(EActiveTickReturnType, FWidgetActiveTickDelegate, double /*InCurrentTime*/, float /*InDeltaTime*/);
+DECLARE_DELEGATE_RetVal_TwoParams(EActiveTimerReturnType, FWidgetActiveTimerDelegate, double /*InCurrentTime*/, float /*InDeltaTime*/);
 
 /* Includes
  *****************************************************************************/
diff --git a/Engine/Source/Runtime/SlateCore/Public/Types/SlateEnums.h b/Engine/Source/Runtime/SlateCore/Public/Types/SlateEnums.h
--- a/Engine/Source/Runtime/SlateCore/Public/Types/SlateEnums.h
+++ b/Engine/Source/Runtime/SlateCore/Public/Types/SlateEnums.h
@@ -241,15 +241,15 @@ namespace ESelectInfo
 	};
 }
 
-/** 
- * Return type for FWidgetActiveTickDelegate.
- * Don't expose to blueprints.
- */
-enum EActiveTickReturnType : uint8
-{
-	/** If this value is returned, the widget will be unregistered for tick automatically. No need to call UnRegisterActiveTick. */
-	StopTicking,
-	/** If this value is returns, the widget will continue to have it's tick delegate called on it. */
-	KeepTicking,
-};
 
+/**
+* Return type for FWidgetActiveTimerDelegate.
+* Don't expose to blueprints.
+*/
+enum class EActiveTimerReturnType : uint8
+{
+	/** If this value is returned, the widget's active timer will be unregistered automatically. No need to call UnRegisterActiveTimer. */
+	Stop,
+	/** If this value is returned, the widget will continue to have its timer delegate called on it. */
+	Continue,
+};
\ No newline at end of file
diff --git a/Engine/Source/Runtime/SlateCore/Public/Widgets/SWidget.h b/Engine/Source/Runtime/SlateCore/Public/Widgets/SWidget.h
--- a/Engine/Source/Runtime/SlateCore/Public/Widgets/SWidget.h
+++ b/Engine/Source/Runtime/SlateCore/Public/Widgets/SWidget.h
@@ -860,45 +860,45 @@ private:
 
 protected:
 	/**
-	 * Registers a delegate for "active tick" at the specified frequency. Will not call TickFunction until the specified interval has elapsed once.
-	 * A widget can register as many delegates as it needs. Be careful when registering to avoid duplicate active ticks.
+	 * Registers an "active timer" delegate that will execute at some regular interval. TickFunction will not be called until the specified interval has elapsed once.
+	 * A widget can register as many delegates as it needs. Be careful when registering to avoid duplicate active timers.
 	 * 
-	 * An active tick can be UnRegistered in one of three ways:
-	 *   1. Call UnRegisterActiveTick using the active tick handle that is returned here.
+	 * An active timer can be UnRegistered in one of three ways:
+	 *   1. Call UnRegisterActiveTimer using the active timer handle that is returned here.
 	 *   2. Have your delegate return ETickWidgetReturnType::StopTicking.
 	 *   3. Destroying the widget
 	 * 
-	 * Active Ticking
+	 * Active Timers
 	 * --------------
 	 * Slate may go to sleep when there is no user interaction for some time to save power.
-	 * However, some UI elements may need to drive the UI even when the user is not providing any input
+	 * However, some UI elements may need to "drive" the UI even when the user is not providing any input
 	 * (ie, animations, viewport rendering, async polling, etc). A widget notifies Slate of this by
-	 * registering for "Active Tick", basically saying it needs to be ticked at some frequency to drive the UI.
-	 * In this way, slate can go to sleep, but still tick any widgets that are driving the UI at the interval they specify.
-	 * When any active tick needs to fire, all of Slate will do a Tick and Paint pass.
+	 * registering an "Active Timer" that is executed at a specified frequency to drive the UI.
+	 * In this way, slate can go to sleep when there is no input and no active timer needs to fire.
+	 * When any active timer needs to fire, all of Slate will do a Tick and Paint pass.
 	 * 
-	 * @param TickPeriod maximum frequency to tick the widget in seconds. Pass zero to tick once per frame. 
+	 * @param Period The time period to wait between each execution of the timer. Pass zero to fire the timer once per frame.
 	 *                      If an interval is missed, the delegate is NOT called more than once.
-	 * @param TickFunction delegate to call every TickFrequency seconds.
-	 * @return an active tick handle that can be used to UnRegister later.
+	 * @param TimerFunction The active timer delegate to call every Period seconds.
+	 * @return An active timer handle that can be used to UnRegister later.
 	 */
-	TSharedRef<FActiveTickHandle> RegisterActiveTick( float TickPeriod, FWidgetActiveTickDelegate TickFunction );
+	TSharedRef<FActiveTimerHandle> RegisterActiveTimer( float TickPeriod, FWidgetActiveTimerDelegate TickFunction );
 
 	/**
-	 * Unregisters an active tick handle. This is optional, as the delegate can UnRegister itself by returning ETickWidgetReturnType::StopTicking.
+	 * Unregisters an active timer handle. This is optional, as the delegate can UnRegister itself by returning EActiveTimerReturnType::Stop.
 	 */
-	void UnRegisterActiveTick( const TSharedRef<FActiveTickHandle>& ActiveTickHandle );
+	void UnRegisterActiveTimer( const TSharedRef<FActiveTimerHandle>& ActiveTimerHandle );
 
 private:
 
-	/** Iterates over the active tick handles on the widget and executes them if their interval has elapsed. */
-	void ExecuteActiveTicks(double CurrentTime, float DeltaTime);
+	/** Iterates over the active timer handles on the widget and executes them if their interval has elapsed. */
+	void ExecuteActiveTimers(double CurrentTime, float DeltaTime);
 
-	/** The list of active tick handles for this widget. */
-	TArray<TSharedRef<FActiveTickHandle>> ActiveTicks;
+	/** The list of active timer handles for this widget. */
+	TArray<TSharedRef<FActiveTimerHandle>> ActiveTimers;
 
 protected:
-	/** Dtor ensures that active tick handles are UnRegistered with the SlateApplication. */
+	/** Dtor ensures that active timer handles are UnRegistered with the SlateApplication. */
 	~SWidget();
 
 	//	DEBUG INFORMATION
diff --git a/Engine/Source/Runtime/SlateCore/Public/Widgets/SWindow.h b/Engine/Source/Runtime/SlateCore/Public/Widgets/SWindow.h
--- a/Engine/Source/Runtime/SlateCore/Public/Widgets/SWindow.h
+++ b/Engine/Source/Runtime/SlateCore/Public/Widgets/SWindow.h
@@ -932,11 +932,11 @@ private:
 	 */
 	EVisibility GetWindowFlashVisibility() const;
 
-	/** One-off active tick to trigger a the morph sequence to play */
-	EActiveTickReturnType TriggerPlayMorphSequence( double InCurrentTime, float InDeltaTime );
+	/** One-off active timer to trigger a the morph sequence to play */
+	EActiveTimerReturnType TriggerPlayMorphSequence( double InCurrentTime, float InDeltaTime );
 
-	/** The handle to the active tick */
-	TWeakPtr<FActiveTickHandle> ActiveTickHandle;
+	/** The handle to the active timer */
+	TWeakPtr<FActiveTimerHandle> ActiveTimerHandle;
 };
 
 
