InsertsSlots funcs default to AddSlot if no index is specified

#nonjira

#branch UE4

#change Requested from CL#2394969 review

diff --git a/Engine/Source/Runtime/SlateCore/Public/Widgets/SBoxPanel.h b/Engine/Source/Runtime/SlateCore/Public/Widgets/SBoxPanel.h
--- a/Engine/Source/Runtime/SlateCore/Public/Widgets/SBoxPanel.h
+++ b/Engine/Source/Runtime/SlateCore/Public/Widgets/SBoxPanel.h
@@ -200,8 +200,12 @@ public:
 		return NewSlot;
 	}
 
-	FSlot& InsertSlot( int32 Index )
+	FSlot& InsertSlot(int32 Index = INDEX_NONE)
 	{
+		if (Index == INDEX_NONE)
+		{
+			return AddSlot();
+		}
 		SHorizontalBox::FSlot& NewSlot = *new SHorizontalBox::FSlot();
 		this->Children.Insert(&NewSlot, Index);
 		return NewSlot;
@@ -339,8 +343,12 @@ public:
 		return NewSlot;
 	}
 
-	FSlot& InsertSlot(int32 Index)
+	FSlot& InsertSlot(int32 Index = INDEX_NONE)
 	{
+		if (Index == INDEX_NONE)
+		{
+			return AddSlot();
+		}
 		SVerticalBox::FSlot& NewSlot = *new SVerticalBox::FSlot();
 		this->Children.Insert(&NewSlot, Index);
 		return NewSlot;
