Fixed crash on restart when importing/resetting key bindings

diff --git a/Engine/Source/Editor/InputBindingEditor/Private/Widgets/SInputBindingEditorPanel.cpp b/Engine/Source/Editor/InputBindingEditor/Private/Widgets/SInputBindingEditorPanel.cpp
--- a/Engine/Source/Editor/InputBindingEditor/Private/Widgets/SInputBindingEditorPanel.cpp
+++ b/Engine/Source/Editor/InputBindingEditor/Private/Widgets/SInputBindingEditorPanel.cpp
@@ -240,7 +240,12 @@ void SInputBindingEditorPanel::OnCommandsChanged()
 {
 	UpdateContextMasterList();
 
-	FilterVisibleContextList();
+	// Make sure we aren't restarting after importing/resetting to default
+	if (FSlateApplication::IsInitialized())
+	{
+		FilterVisibleContextList();
+	}
+	
 }
 
 
diff --git a/Engine/Source/Runtime/SlateCore/Private/Widgets/SWidget.cpp b/Engine/Source/Runtime/SlateCore/Private/Widgets/SWidget.cpp
--- a/Engine/Source/Runtime/SlateCore/Private/Widgets/SWidget.cpp
+++ b/Engine/Source/Runtime/SlateCore/Private/Widgets/SWidget.cpp
@@ -682,8 +682,11 @@ TSharedRef<FActiveTimerHandle> SWidget::RegisterActiveTimer(float TickPeriod, FW
 
 void SWidget::UnRegisterActiveTimer(const TSharedRef<FActiveTimerHandle>& ActiveTimerHandle)
 {
-	FSlateApplicationBase::Get().UnRegisterActiveTimer(ActiveTimerHandle);
-	ActiveTimers.Remove(ActiveTimerHandle);
+	if (FSlateApplicationBase::IsInitialized())
+	{
+		FSlateApplicationBase::Get().UnRegisterActiveTimer(ActiveTimerHandle);
+		ActiveTimers.Remove(ActiveTimerHandle);
+	}
 }
 
 void SWidget::ExecuteActiveTimers(double CurrentTime, float DeltaTime)
@@ -698,7 +701,11 @@ void SWidget::ExecuteActiveTimers(double CurrentTime, float DeltaTime)
 		}
 		else
 		{
-			FSlateApplicationBase::Get().UnRegisterActiveTimer(ActiveTimers[i]);
+			if (FSlateApplicationBase::IsInitialized())
+			{
+				FSlateApplicationBase::Get().UnRegisterActiveTimer(ActiveTimers[i]);
+			}
+			
 			ActiveTimers.RemoveAt(i);
 		}
 	}
