Updated deprecated SWindow transparency construction across the codebase to use new syntax

diff --git a/Engine/Source/Editor/LevelEditor/Private/SLevelViewportControlsPopup.cpp b/Engine/Source/Editor/LevelEditor/Private/SLevelViewportControlsPopup.cpp
--- a/Engine/Source/Editor/LevelEditor/Private/SLevelViewportControlsPopup.cpp
+++ b/Engine/Source/Editor/LevelEditor/Private/SLevelViewportControlsPopup.cpp
@@ -86,7 +86,7 @@ TSharedRef<SWidget> SLevelViewportControlsPopup::OnGetMenuContent()
 				SNew(SWebBrowser)
 				.InitialURL(PopupPath)
 				.ShowControls(false)
-				.SupportsTransparency(true)
+				.SupportsTransparency(EWindowTransparency::PerWindow)
 				.ViewportSize(PopupSize + 1)
 			];
 	}
diff --git a/Engine/Source/Runtime/SlateCore/Public/Widgets/SWindow.h b/Engine/Source/Runtime/SlateCore/Public/Widgets/SWindow.h
--- a/Engine/Source/Runtime/SlateCore/Public/Widgets/SWindow.h
+++ b/Engine/Source/Runtime/SlateCore/Public/Widgets/SWindow.h
@@ -60,7 +60,7 @@ namespace ESizingRule
 /** Proxy structure to handle deprecated construction from bool */
 struct FWindowTransparency
 {
-	DEPRECATED(4.7, "Please now use the static initializers defined below.")
+	DEPRECATED(4.8, "Please specify an EWindowTransparency value instead.")
 	FWindowTransparency(bool bSupportsTransparency) : Value(bSupportsTransparency ? EWindowTransparency::PerWindow : EWindowTransparency::None) {}
 
 	FWindowTransparency(EWindowTransparency In) : Value(In) {}
