Reduced padding on toolbar borders

diff --git a/Engine/Source/Runtime/SlateCore/Private/Styling/CoreStyle.cpp b/Engine/Source/Runtime/SlateCore/Private/Styling/CoreStyle.cpp
--- a/Engine/Source/Runtime/SlateCore/Private/Styling/CoreStyle.cpp
+++ b/Engine/Source/Runtime/SlateCore/Private/Styling/CoreStyle.cpp
@@ -718,10 +718,10 @@ TSharedRef<ISlateStyle> FCoreStyle::Create( const FName& InStyleSetName )
 		Style->Set( "ToolBar.Icon", new IMAGE_BRUSH( "Icons/icon_tab_toolbar_16px", Icon16x16 ) );
 		Style->Set( "ToolBar.Expand", new IMAGE_BRUSH( "Icons/toolbar_expand_16x", Icon16x16) );
 		Style->Set( "ToolBar.SubMenuIndicator", new IMAGE_BRUSH( "Common/SubmenuArrow", Icon8x8 ) );
-		Style->Set( "ToolBar.SToolBarComboButtonBlock.Padding", FMargin(4.0f));
-		Style->Set( "ToolBar.SToolBarButtonBlock.Padding", FMargin(4.0f));
-		Style->Set( "ToolBar.SToolBarCheckComboButtonBlock.Padding", FMargin(4.0f));
-		Style->Set( "ToolBar.SToolBarButtonBlock.CheckBox.Padding", FMargin(4.0f) );
+		Style->Set( "ToolBar.SToolBarComboButtonBlock.Padding", FMargin(4.0f,0.0f));
+		Style->Set( "ToolBar.SToolBarButtonBlock.Padding", FMargin(4.0f,0.0f));
+		Style->Set( "ToolBar.SToolBarCheckComboButtonBlock.Padding", FMargin(4.0f,0.0f));
+		Style->Set( "ToolBar.SToolBarButtonBlock.CheckBox.Padding", FMargin(4.0f,0.0f) );
 		Style->Set( "ToolBar.SToolBarComboButtonBlock.ComboButton.Color", DefaultForeground );
 
 		Style->Set( "ToolBar.Block.IndentedPadding", FMargin( 18.0f, 2.0f, 4.0f, 4.0f ) );
