Slate : RelativeLayoutScale is now computed in the SlatePrepass(). ComputeDesiredSize() now takes a float parameter. API-breaking change!

diff --git a/Engine/Source/Runtime/SlateCore/Public/Widgets/SWindow.h b/Engine/Source/Runtime/SlateCore/Public/Widgets/SWindow.h
--- a/Engine/Source/Runtime/SlateCore/Public/Widgets/SWindow.h
+++ b/Engine/Source/Runtime/SlateCore/Public/Widgets/SWindow.h
@@ -606,6 +606,7 @@ public:
 
 	virtual bool SupportsKeyboardFocus() const override;
 
+private:
 	virtual FReply OnFocusReceived( const FGeometry& MyGeometry, const FFocusEvent& InFocusEvent );
 
 	virtual FReply OnMouseButtonDown( const FGeometry& MyGeometry, const FPointerEvent& MouseEvent ) override;
@@ -613,8 +614,11 @@ public:
 	virtual FReply OnMouseMove( const FGeometry& MyGeometry, const FPointerEvent& MouseEvent ) override;
 
 	/** The window's desired size takes into account the ratio between the slate units and the pixel size */
-	virtual FVector2D ComputeDesiredSize() const override;
+	virtual FVector2D ComputeDesiredSize(float) const override;
+
+	virtual float GetRelativeLayoutScale(const FSlotBase& Child) const override;
 
+public:
 	/** @return true if this window will be focused when it is first shown */
 	inline bool IsFocusedInitially() const
 	{
@@ -930,6 +934,8 @@ private:
 
 private:
 	
+	virtual
+
 	void ConstructWindowInternals( const bool bCreateTitleBar );
 
 	/**
diff --git a/Engine/Source/Runtime/UMG/Private/Slate/SConstraintCanvas.cpp b/Engine/Source/Runtime/UMG/Private/Slate/SConstraintCanvas.cpp
--- a/Engine/Source/Runtime/UMG/Private/Slate/SConstraintCanvas.cpp
+++ b/Engine/Source/Runtime/UMG/Private/Slate/SConstraintCanvas.cpp
@@ -168,7 +168,7 @@ int32 SConstraintCanvas::OnPaint( const FPaintArgs& Args, const FGeometry& Allot
 	return MaxLayerId;
 }
 
-FVector2D SConstraintCanvas::ComputeDesiredSize() const
+FVector2D SConstraintCanvas::ComputeDesiredSize( float ) const
 {
 	FVector2D DesiredSize(0,0);
 
