Forgot to include SlateStats.h in SlateStats.cpp. Unity compilation masked the error most of hte time. :-/

diff --git a/Engine/Source/Runtime/SlateCore/Private/Stats/SlateStats.cpp b/Engine/Source/Runtime/SlateCore/Private/Stats/SlateStats.cpp
--- a/Engine/Source/Runtime/SlateCore/Private/Stats/SlateStats.cpp
+++ b/Engine/Source/Runtime/SlateCore/Private/Stats/SlateStats.cpp
@@ -2,6 +2,8 @@
 
 #include "SlateCorePrivatePCH.h"
 
+#include "SlateStats.h"
+
 #if SLATE_STATS
 
 SLATE_DECLARE_CYCLE_COUNTER(GSlateStatsOverhead, "Stats Overhead");
