Fix for clang pedantic error about NULL comparison.

diff --git a/Engine/Source/Runtime/SlateCore/Private/Stats/SlateStats.cpp b/Engine/Source/Runtime/SlateCore/Private/Stats/SlateStats.cpp
--- a/Engine/Source/Runtime/SlateCore/Private/Stats/SlateStats.cpp
+++ b/Engine/Source/Runtime/SlateCore/Private/Stats/SlateStats.cpp
@@ -111,7 +111,7 @@ void FSlateStatCycleCounter::EndFrame(double CurrentTime)
 					// put a dot after the current depth and track the new position to put new depths.
 					if (PathCurrentDepth > 0)
 					{
-						while (*PathStrCurrent != NULL) ++PathStrCurrent;
+						while (*PathStrCurrent != 0) ++PathStrCurrent;
 						*PathStrCurrent++ = '.';
 						*PathStrCurrent = 0;
 					}
