Fix for Unity compile error.

diff --git a/Engine/Source/Runtime/SlateCore/Private/Rendering/ElementBatcher.cpp b/Engine/Source/Runtime/SlateCore/Private/Rendering/ElementBatcher.cpp
--- a/Engine/Source/Runtime/SlateCore/Private/Rendering/ElementBatcher.cpp
+++ b/Engine/Source/Runtime/SlateCore/Private/Rendering/ElementBatcher.cpp
@@ -4,6 +4,7 @@
 #include "ElementBatcher.h"
 #include "SlateRenderTransform.h"
 #include "Internationalization/Text.h"
+#include "SlateStats.h"
 
 DECLARE_DWORD_COUNTER_STAT(TEXT("Num Layers"), STAT_SlateNumLayers, STATGROUP_Slate);
 DECLARE_DWORD_COUNTER_STAT(TEXT("Num Batches"), STAT_SlateNumBatches, STATGROUP_Slate);
