Fixed an incorrect comment that referred to an outdated enumerated type.

diff --git a/Engine/Source/Runtime/SlateCore/Public/Widgets/SWidget.h b/Engine/Source/Runtime/SlateCore/Public/Widgets/SWidget.h
--- a/Engine/Source/Runtime/SlateCore/Public/Widgets/SWidget.h
+++ b/Engine/Source/Runtime/SlateCore/Public/Widgets/SWidget.h
@@ -876,7 +876,7 @@ protected:
 	 * 
 	 * An active timer can be UnRegistered in one of three ways:
 	 *   1. Call UnRegisterActiveTimer using the active timer handle that is returned here.
-	 *   2. Have your delegate return ETickWidgetReturnType::StopTicking.
+	 *   2. Have your delegate return EActiveTimerReturnType::Stop.
 	 *   3. Destroying the widget
 	 * 
 	 * Active Timers
