Slate: Fixed widget reflector VS integration (files didn't open due to wrong file names)

diff --git a/Engine/Source/Runtime/SlateCore/Private/Widgets/SWidget.cpp b/Engine/Source/Runtime/SlateCore/Private/Widgets/SWidget.cpp
--- a/Engine/Source/Runtime/SlateCore/Private/Widgets/SWidget.cpp
+++ b/Engine/Source/Runtime/SlateCore/Private/Widgets/SWidget.cpp
@@ -553,13 +553,13 @@ FName SWidget::GetType() const
 
 FString SWidget::GetReadableLocation() const
 {
-	return FString::Printf(TEXT("%s(%d)"), *this->CreatedInFile.ToString(), this->CreatedOnLine );
+	return FString::Printf(TEXT("%s(%d)"), *FPaths::GetCleanFilename(this->GetCreatedInFile()), this->CreatedOnLine);
 }
 
 
 FString SWidget::GetCreatedInFile() const
 {
-	return this->CreatedInFileFullPath.ToString();
+	return this->CreatedInFileFullPath.GetPlainNameString();
 }
 
 FName SWidget::GetCreatedInFileFName() const
