Fixed WidgetReflector source navigation functionality. #BUN
* Removed SWidget:: GetCreatedInFile, GetCreatedInFileFName, GetCreatedInLineNumber.
* Added GetCreatedInLocation() -> FName as a non-redundant replacement, updated all code usages.
* This shrinks SWidget size by 16 bytes and reduces ambiguity on which ones to use.
* Implemented ToString() in terms of GetReadableLocation().

diff --git a/Engine/Source/Developer/SlateReflector/Private/Widgets/SWidgetReflectorTreeWidgetItem.cpp b/Engine/Source/Developer/SlateReflector/Private/Widgets/SWidgetReflectorTreeWidgetItem.cpp
--- a/Engine/Source/Developer/SlateReflector/Private/Widgets/SWidgetReflectorTreeWidgetItem.cpp
+++ b/Engine/Source/Developer/SlateReflector/Private/Widgets/SWidgetReflectorTreeWidgetItem.cpp
@@ -131,6 +131,6 @@ void SReflectorTreeWidgetItem::HandleHyperlinkNavigate()
 
 	if ( OnAccessSourceCode.IsBound() )
 	{
-		OnAccessSourceCode.Execute(GetWidgetFile().ToString(), GetWidgetLineNumber(), 0);
+		OnAccessSourceCode.Execute(GetWidgetFile(), GetWidgetLineNumber(), 0);
 	}
 }
diff --git a/Engine/Source/Developer/SlateReflector/Private/Widgets/SWidgetReflectorTreeWidgetItem.h b/Engine/Source/Developer/SlateReflector/Private/Widgets/SWidgetReflectorTreeWidgetItem.h
--- a/Engine/Source/Developer/SlateReflector/Private/Widgets/SWidgetReflectorTreeWidgetItem.h
+++ b/Engine/Source/Developer/SlateReflector/Private/Widgets/SWidgetReflectorTreeWidgetItem.h
@@ -40,8 +40,8 @@ public:
 		if ( Widget.IsValid() )
 		{
 			CachedWidgetType = FText::FromString(Widget->GetTypeAsString());
-			CachedWidgetFile = FText::FromString(Widget->GetCreatedInFile());
-			CachedWidgetLineNumber =  Widget->GetCreatedInLineNumber();
+			CachedWidgetFile = Widget->GetCreatedInLocation().GetPlainNameString();
+			CachedWidgetLineNumber = Widget->GetCreatedInLocation().GetNumber();
 			CachedWidgetVisibility = FText::FromString(Widget->GetVisibility().ToString());
 
 			CachedReadableLocation = Widget->GetReadableLocation();
@@ -71,7 +71,7 @@ protected:
 		return FText::FromString(GetReadableLocation());
 	}
 
-	FText GetWidgetFile() const
+	FString GetWidgetFile() const
 	{
 		return CachedWidgetFile;
 	}
@@ -100,7 +100,7 @@ private:
 	TAttribute< TSharedPtr<FReflectorNode> > WidgetInfo;
 
 	FText CachedWidgetType;
-	FText CachedWidgetFile;
+	FString CachedWidgetFile;
 	int32 CachedWidgetLineNumber;
 	FText CachedWidgetVisibility;
 
diff --git a/Engine/Source/Developer/StandaloneRenderer/Private/Windows/D3D/SlateD3DRenderer.cpp b/Engine/Source/Developer/StandaloneRenderer/Private/Windows/D3D/SlateD3DRenderer.cpp
--- a/Engine/Source/Developer/StandaloneRenderer/Private/Windows/D3D/SlateD3DRenderer.cpp
+++ b/Engine/Source/Developer/StandaloneRenderer/Private/Windows/D3D/SlateD3DRenderer.cpp
@@ -337,7 +337,7 @@ void FSlateD3DRenderer::DrawWindows( FSlateDrawBuffer& InWindowDrawBuffer )
 	for( int32 ListIndex = 0; ListIndex < WindowElementLists.Num(); ++ListIndex )
 	{
 		FSlateWindowElementList& ElementList = WindowElementLists[ListIndex];
-		SLATE_CYCLE_COUNTER_SCOPE_CUSTOM_DETAILED(SLATE_STATS_DETAIL_LEVEL_MED, GRendererDrawElementList, ElementList.GetWindow()->GetCreatedInFileFName());
+		SLATE_CYCLE_COUNTER_SCOPE_CUSTOM_DETAILED(SLATE_STATS_DETAIL_LEVEL_MED, GRendererDrawElementList, ElementList.GetWindow()->GetCreatedInLocation());
 
 		if ( ElementList.GetWindow().IsValid() )
 		{
diff --git a/Engine/Source/Runtime/Slate/Private/Framework/Application/SlateApplication.cpp b/Engine/Source/Runtime/Slate/Private/Framework/Application/SlateApplication.cpp
--- a/Engine/Source/Runtime/Slate/Private/Framework/Application/SlateApplication.cpp
+++ b/Engine/Source/Runtime/Slate/Private/Framework/Application/SlateApplication.cpp
@@ -886,7 +886,7 @@ bool FSlateApplication::IsWindowHousingInteractiveTooltip(const TSharedRef<const
 
 void FSlateApplication::TickWindowAndChildren( TSharedRef<SWindow> WindowToTick )
 {
-	SLATE_CYCLE_COUNTER_SCOPE_CUSTOM(GTickWindowAndChildrenTime, WindowToTick->GetCreatedInFileFName());
+	SLATE_CYCLE_COUNTER_SCOPE_CUSTOM(GTickWindowAndChildrenTime, WindowToTick->GetCreatedInLocation());
 	if ( WindowToTick->IsVisible() && !WindowToTick->IsWindowMinimized() )
 	{
 		// Switch to the appropriate world for ticking
@@ -943,7 +943,7 @@ struct FDrawWindowArgs
 
 void FSlateApplication::DrawWindowAndChildren( const TSharedRef<SWindow>& WindowToDraw, FDrawWindowArgs& DrawWindowArgs )
 {
-	SLATE_CYCLE_COUNTER_SCOPE_CUSTOM(GSlateDrawWindowAndChildren, WindowToDraw->GetCreatedInFileFName());
+	SLATE_CYCLE_COUNTER_SCOPE_CUSTOM(GSlateDrawWindowAndChildren, WindowToDraw->GetCreatedInLocation());
 	// Only draw visible windows
 	if( WindowToDraw->IsVisible() && !WindowToDraw->IsWindowMinimized() )
 	{
@@ -1023,7 +1023,7 @@ void FSlateApplication::DrawWindowAndChildren( const TSharedRef<SWindow>& Window
 
 static void PrepassWindowAndChildren( TSharedRef<SWindow> WindowToPrepass )
 {
-	SLATE_CYCLE_COUNTER_SCOPE_CUSTOM(GSlatePrepassWindowAndChildren, WindowToPrepass->GetCreatedInFileFName());
+	SLATE_CYCLE_COUNTER_SCOPE_CUSTOM(GSlatePrepassWindowAndChildren, WindowToPrepass->GetCreatedInLocation());
 	FScopedSwitchWorldHack SwitchWorld( WindowToPrepass );
 
 	{
diff --git a/Engine/Source/Runtime/SlateCore/Private/Widgets/SWidget.cpp b/Engine/Source/Runtime/SlateCore/Private/Widgets/SWidget.cpp
--- a/Engine/Source/Runtime/SlateCore/Private/Widgets/SWidget.cpp
+++ b/Engine/Source/Runtime/SlateCore/Private/Widgets/SWidget.cpp
@@ -18,9 +18,7 @@ SLATE_DECLARE_CYCLE_COUNTER(GSlateGetVisibility, "GetVisibility");
 TAutoConsoleVariable<int32> TickInvisibleWidgets(TEXT("Slate.TickInvisibleWidgets"), 0, TEXT("Controls whether invisible widgets are ticked."));
 
 SWidget::SWidget()
-	: CreatedInFile( TEXT("") )
-	, CreatedOnLine( -1 )
-	, Cursor( TOptional<EMouseCursor::Type>() )
+	: Cursor( TOptional<EMouseCursor::Type>() )
 	, EnabledState( true )
 	, Visibility( EVisibility::Visible )
 	, RenderTransform( )
@@ -535,7 +533,7 @@ int32 SWidget::FindChildUnderMouse( const FArrangedChildren& Children, const FPo
 
 FString SWidget::ToString() const
 {
-	return FString::Printf(TEXT("%s [%s(%d)]"), *this->TypeOfWidget.ToString(), *this->CreatedInFile.ToString(), this->CreatedOnLine );
+	return FString::Printf(TEXT("%s [%s]"), *this->TypeOfWidget.ToString(), *this->GetReadableLocation() );
 }
 
 
@@ -553,26 +551,15 @@ FName SWidget::GetType() const
 
 FString SWidget::GetReadableLocation() const
 {
-	return FString::Printf(TEXT("%s(%d)"), *FPaths::GetCleanFilename(this->GetCreatedInFile()), this->CreatedOnLine);
+	return FString::Printf(TEXT("%s(%d)"), *FPaths::GetCleanFilename(this->CreatedInLocation.GetPlainNameString()), this->CreatedInLocation.GetNumber());
 }
 
 
-FString SWidget::GetCreatedInFile() const
+FName SWidget::GetCreatedInLocation() const
 {
-	return this->CreatedInFileFullPath.GetPlainNameString();
+	return this->CreatedInLocation;
 }
 
-FName SWidget::GetCreatedInFileFName() const
-{
-	return this->CreatedInFile;
-}
-
-int32 SWidget::GetCreatedInLineNumber() const
-{
-	return this->CreatedOnLine;
-}
-
-
 FName SWidget::GetTag() const
 {
 	return this->Tag;
@@ -634,11 +621,9 @@ void SWidget::SetCursor( const TAttribute< TOptional<EMouseCursor::Type> >& InCu
 void SWidget::SetDebugInfo( const ANSICHAR* InType, const ANSICHAR* InFile, int32 OnLine )
 {
 	this->TypeOfWidget = InType;
-	this->CreatedInFileFullPath = FName( InFile );
-	this->CreatedInFileFullPath.SetNumber(OnLine);
-	this->CreatedInFile = FName( *FPaths::GetCleanFilename(InFile) );
-	this->CreatedInFile.SetNumber(OnLine);
-	this->CreatedOnLine = OnLine;
+	this->
+		CreatedInLocation = FName( InFile );
+	this->CreatedInLocation.SetNumber(OnLine);
 }
 
 SLATECORE_API int32 bFoldTick = 0;
diff --git a/Engine/Source/Runtime/SlateCore/Public/Widgets/SWidget.h b/Engine/Source/Runtime/SlateCore/Public/Widgets/SWidget.h
--- a/Engine/Source/Runtime/SlateCore/Public/Widgets/SWidget.h
+++ b/Engine/Source/Runtime/SlateCore/Public/Widgets/SWidget.h
@@ -763,17 +763,11 @@ public:
 	/** @return The widget's type as an FName ID */
 	FName GetType() const;
 
-	/** @return A String of the widget's code location in readable format */
+	/** @return A String of the widget's code location in readable format "BaseFileName(LineNumber)" */
 	virtual FString GetReadableLocation() const;
 
-	/** @return A String of the widget's code location */
-	virtual FString GetCreatedInFile() const;
-
-	/** @return A String of the widget's code location */
-	virtual FName GetCreatedInFileFName() const;
-
-	/** @return The line number of the widgets location */
-	virtual int32 GetCreatedInLineNumber() const;
+	/** @return An FName of the widget's code location (full path with number == line number of the file) */
+	FName GetCreatedInLocation() const;
 
 	/** @return The name this widget was tagged with */
 	virtual FName GetTag() const;
@@ -915,12 +909,8 @@ protected:
 	//	DEBUG INFORMATION
 	// @todo Slate: Should compile out in final release builds?
 	FName TypeOfWidget;
-	/** Full file path in which this widget was created */
-	FName CreatedInFileFullPath;
-	/** Filename in which this widget was created */
-	FName CreatedInFile;
-	/** Line number on which this widget was created */
-	int32 CreatedOnLine;
+	/** Full file path (and line) in which this widget was created */
+	FName CreatedInLocation;
 
 	/** Tag for this widget */
 	FName Tag;
