#UE4 Leaving slack in Vert/Index batch lists to reduce per-frame heap allocations

diff --git a/Engine/Source/Runtime/SlateCore/Private/Rendering/ElementBatcher.cpp b/Engine/Source/Runtime/SlateCore/Private/Rendering/ElementBatcher.cpp
--- a/Engine/Source/Runtime/SlateCore/Private/Rendering/ElementBatcher.cpp
+++ b/Engine/Source/Runtime/SlateCore/Private/Rendering/ElementBatcher.cpp
@@ -1606,8 +1606,8 @@ void FSlateElementBatcher::FillBatchBuffers( FSlateWindowElementList& WindowElem
 					OutRenderBatches.Add( FSlateRenderBatch( ElementBatch ) );
 
 					// Done with the batch
-					BatchVertices.Empty();
-					BatchIndices.Empty();
+					BatchVertices.Empty(BatchVertices.Num());
+					BatchIndices.Empty(BatchIndices.Num());
 				}
 
 			}
