[INTEGRATE] Change 2436873 by [REDACTED_PERSON][REDACTED_HANDLE] on 2015/02/07 15:05:03

	Slate: Add debug-only warnings for missing float and vector constants in a style to match how the other Get methods work

diff --git a/Engine/Source/Runtime/SlateCore/Private/Styling/SlateStyleSet.cpp b/Engine/Source/Runtime/SlateCore/Private/Styling/SlateStyleSet.cpp
--- a/Engine/Source/Runtime/SlateCore/Private/Styling/SlateStyleSet.cpp
+++ b/Engine/Source/Runtime/SlateCore/Private/Styling/SlateStyleSet.cpp
@@ -119,6 +119,14 @@ float FSlateStyleSet::GetFloat(const FName PropertyName, const ANSICHAR* Specifi
 {
 	const float* Result = FloatValues.Find(Join(PropertyName, Specifier));
 
+#if DO_GUARD_SLOW
+	if ((Result == nullptr) && !MissingResources.Contains(PropertyName))
+	{
+		MissingResources.Add(PropertyName);
+		Log(ISlateStyle::Warning, FText::Format(NSLOCTEXT("SlateStyleSet", "UknownSlateFloat", "Unable to find float property '{0}' in style."), FText::FromName(PropertyName)));
+	}
+#endif
+
 	return Result ? *Result : FStyleDefaults::GetFloat();
 }
 
@@ -126,6 +134,14 @@ FVector2D FSlateStyleSet::GetVector(const FName PropertyName, const ANSICHAR* Sp
 {
 	const FVector2D* const Result = Vector2DValues.Find(Join(PropertyName, Specifier));
 
+#if DO_GUARD_SLOW
+	if ((Result == nullptr) && !MissingResources.Contains(PropertyName))
+	{
+		MissingResources.Add(PropertyName);
+		Log(ISlateStyle::Warning, FText::Format(NSLOCTEXT("SlateStyleSet", "UknownSlateVector", "Unable to find vector property '{0}' in style."), FText::FromName(PropertyName)));
+	}
+#endif
+
 	return Result ? *Result : FStyleDefaults::GetVector2D();
 }
 
