Adding to 2448848, missing file.  Adding missing code for windowless drag drop.

diff --git a/Engine/Source/Runtime/SlateCore/Private/Input/DragAndDrop.cpp b/Engine/Source/Runtime/SlateCore/Private/Input/DragAndDrop.cpp
--- a/Engine/Source/Runtime/SlateCore/Private/Input/DragAndDrop.cpp
+++ b/Engine/Source/Runtime/SlateCore/Private/Input/DragAndDrop.cpp
@@ -6,6 +6,11 @@
 /* FDragDropOperation structors
  *****************************************************************************/
 
+FDragDropOperation::FDragDropOperation()
+	: bCreateNewWindow(true)
+{
+}
+
 FDragDropOperation::~FDragDropOperation()
 {
 	DestroyCursorDecoratorWindow();
@@ -56,6 +61,22 @@ void FDragDropOperation::SetDecoratorVisibility(bool bVisible)
 			CursorDecoratorWindow->HideWindow();
 		}
 	}
+	else if ( bCreateNewWindow == false )
+	{
+		TSharedPtr<SWidget> DecoratorToUse = GetDefaultDecorator();
+
+		if ( DecoratorToUse.IsValid() )
+		{
+			if ( bVisible )
+			{
+				DecoratorToUse->SetVisibility(EVisibility::HitTestInvisible);
+			}
+			else
+			{
+				DecoratorToUse->SetVisibility(EVisibility::Hidden);
+			}
+		}
+	}
 }
 
 void FDragDropOperation::SetCursorOverride( TOptional<EMouseCursor::Type> CursorType )
@@ -68,10 +89,18 @@ void FDragDropOperation::SetCursorOverride( TOptional<EMouseCursor::Type> Cursor
  *****************************************************************************/
 
 void FDragDropOperation::Construct()
+{
+	if ( bCreateNewWindow )
+	{
+		CreateCursorDecoratorWindow();
+	}
+}
+
+void FDragDropOperation::CreateCursorDecoratorWindow()
 {
 	TSharedPtr<SWidget> DecoratorToUse = GetDefaultDecorator();
 
-	if (DecoratorToUse.IsValid())
+	if ( DecoratorToUse.IsValid() )
 	{
 		CursorDecoratorWindow = SWindow::MakeCursorDecorator();
 		CursorDecoratorWindow->SetContent(DecoratorToUse.ToSharedRef());
@@ -90,7 +119,7 @@ void FDragDropOperation::DestroyCursorDecoratorWindow()
 }
 
 
-/* FDragDropOperation implementation
+/* FExternalDragOperation implementation
  *****************************************************************************/
 
 TSharedRef<FExternalDragOperation> FExternalDragOperation::NewText( const FString& InText )
@@ -110,3 +139,18 @@ TSharedRef<FExternalDragOperation> FExternalDragOperation::NewFiles( const TArra
 	Operation->Construct();
 	return Operation;
 }
+
+
+/* FGameDragDropOperation implementation
+*****************************************************************************/
+
+FGameDragDropOperation::FGameDragDropOperation()
+	: DecoratorPosition(0,0)
+{
+	bCreateNewWindow = false;
+}
+
+FVector2D FGameDragDropOperation::GetDecoratorPosition() const
+{
+	return DecoratorPosition;
+}
