Improved the performance when UMG had a null font set on a widget

UE-10394 - UMG Textblock with no font drops FPS by ~30

Couple of things here:
1) We used to maintain a map of fallback fonts. This isn't needed as there can only be one culture specific fallback font active at once, and we now detect a culture change rather than maintain a map.
2) We would get the fallback font per-character in the font cache. This was really slow as it involved taking a lock per-character. We now set the fallback font once when getting the composite font from the FSlateFontInfo instead.

diff --git a/Engine/Source/Runtime/SlateCore/Private/Fonts/FontCache.cpp b/Engine/Source/Runtime/SlateCore/Private/Fonts/FontCache.cpp
--- a/Engine/Source/Runtime/SlateCore/Private/Fonts/FontCache.cpp
+++ b/Engine/Source/Runtime/SlateCore/Private/Fonts/FontCache.cpp
@@ -406,7 +406,7 @@ public:
 		// If the requested glyph doesn't exist, use the localization fallback font.
 		if ( FontFace == nullptr || (Char != 0 && GlyphIndex == 0) )
 		{
-			FontFace = GetFontFace( FLegacySlateFontInfoCache::Get().GetFallbackFont() );
+			FontFace = GetFontFace( FLegacySlateFontInfoCache::Get().GetFallbackFontData() );
 			if (FontFace != nullptr)
 			{					
 				GlyphIndex = FT_Get_Char_Index( FontFace, Char );
@@ -417,7 +417,7 @@ public:
 		// If the requested glyph doesn't exist, use the last resort fallback font.
 		if ( FontFace == nullptr || ( Char != 0 && GlyphIndex == 0 ) )
 		{
-			FontFace = GetFontFace( FLegacySlateFontInfoCache::Get().GetLastResortFont() );
+			FontFace = GetFontFace( FLegacySlateFontInfoCache::Get().GetLastResortFontData() );
 			check( FontFace );
 			GlyphIndex = FT_Get_Char_Index( FontFace, Char );
 			LocalGlyphFlags |= FT_LOAD_FORCE_AUTOHINT;
diff --git a/Engine/Source/Runtime/SlateCore/Private/Fonts/LegacySlateFontInfoCache.cpp b/Engine/Source/Runtime/SlateCore/Private/Fonts/LegacySlateFontInfoCache.cpp
--- a/Engine/Source/Runtime/SlateCore/Private/Fonts/LegacySlateFontInfoCache.cpp
+++ b/Engine/Source/Runtime/SlateCore/Private/Fonts/LegacySlateFontInfoCache.cpp
@@ -68,56 +68,75 @@ TSharedPtr<const FCompositeFont> FLegacySlateFontInfoCache::GetSystemFont()
 	return SystemFont;
 }
 
-const FFontData& FLegacySlateFontInfoCache::GetFallbackFont()
+TSharedPtr<const FCompositeFont> FLegacySlateFontInfoCache::GetFallbackFont()
 {
-	const FName FallbackFontName = *(FPaths::EngineContentDir() / TEXT("Slate/Fonts/") / (NSLOCTEXT("Slate", "FallbackFont", "DroidSansFallback").ToString() + TEXT(".ttf")));
-
-	// GetFallbackFont is called directly from the font cache, so may be called from multiple threads at once
+	// GetFallbackFont may be called from multiple threads at once
 	FScopeLock Lock(&FallbackFontCS);
 
+	// The fallback font can change if the active culture is changed
+	const int32 CurrentHistoryVersion = FTextLocalizationManager::Get().GetHeadCultureRevision();
+
+	if (!FallbackFont.IsValid() || FallbackFontHistoryVersion != CurrentHistoryVersion)
 	{
-		TSharedPtr<const FFontData>* const ExistingFallbackFont = FallbackFonts.Find(FallbackFontName);
-		if(ExistingFallbackFont)
-		{
-			return **ExistingFallbackFont;
-		}
+		FallbackFontHistoryVersion = CurrentHistoryVersion;
+
+		const FString FallbackFontPath = FPaths::EngineContentDir() / TEXT("Slate/Fonts/") / (NSLOCTEXT("Slate", "FallbackFont", "DroidSansFallback").ToString() + TEXT(".ttf"));
+		UFontBulkData* FontBulkData = NewObject<UFontBulkData>();
+		FontBulkData->Initialize(FallbackFontPath);
+		FallbackFont = MakeShareable(new FStandaloneCompositeFont(NAME_None, FallbackFontPath, FontBulkData, EFontHinting::Default));
 	}
 
-	const FString FallbackFontPath = FallbackFontName.ToString();
-	UFontBulkData* FontBulkData = NewObject<UFontBulkData>();
-	FontBulkData->Initialize(FallbackFontPath);
-	TSharedRef<const FFontData> NewFallbackFont = MakeShareable(new FFontData(FallbackFontPath, FontBulkData, EFontHinting::Default));
-	FallbackFonts.Add(FallbackFontName, NewFallbackFont);
-	return *NewFallbackFont;
+	return FallbackFont;
+}
+
+const FFontData& FLegacySlateFontInfoCache::GetFallbackFontData()
+{
+	// GetFallbackFontData is called directly from the font cache, so may be called from multiple threads at once
+	FScopeLock Lock(&FallbackFontDataCS);
+
+	// The fallback font can change if the active culture is changed
+	const int32 CurrentHistoryVersion = FTextLocalizationManager::Get().GetHeadCultureRevision();
+
+	if (!FallbackFontData.IsValid() || FallbackFontDataHistoryVersion != CurrentHistoryVersion)
+	{
+		FallbackFontDataHistoryVersion = CurrentHistoryVersion;
+
+		const FString FallbackFontPath = FPaths::EngineContentDir() / TEXT("Slate/Fonts/") / (NSLOCTEXT("Slate", "FallbackFont", "DroidSansFallback").ToString() + TEXT(".ttf"));
+		UFontBulkData* FontBulkData = NewObject<UFontBulkData>();
+		FontBulkData->Initialize(FallbackFontPath);
+		FallbackFontData = MakeShareable(new FFontData(FallbackFontPath, FontBulkData, EFontHinting::Default));
+	}
+
+	return *FallbackFontData;
 }
 
-const FFontData& FLegacySlateFontInfoCache::GetLastResortFont()
+const FFontData& FLegacySlateFontInfoCache::GetLastResortFontData()
 {
-	// GetLastResortFont is called directly from the font cache, so may be called from multiple threads at once
-	FScopeLock Lock(&LastResortFontCS);
+	// GetLastResortFontData is called directly from the font cache, so may be called from multiple threads at once
+	FScopeLock Lock(&LastResortFontDataCS);
 
-	if (!LastResortFont.IsValid())
+	if (!LastResortFontData.IsValid())
 	{
 		const FString LastResortFontPath = FPaths::EngineContentDir() / TEXT("Slate/Fonts/LastResort.ttf");
 		UFontBulkData* FontBulkData = NewObject<UFontBulkData>();
 		FontBulkData->Initialize(LastResortFontPath);
-		LastResortFont = MakeShareable(new FFontData(LastResortFontPath, FontBulkData, EFontHinting::Default));
+		LastResortFontData = MakeShareable(new FFontData(LastResortFontPath, FontBulkData, EFontHinting::Default));
 	}
 
-	return *LastResortFont;
+	return *LastResortFontData;
 }
 
 void FLegacySlateFontInfoCache::AddReferencedObjects(FReferenceCollector& Collector)
 {
-	for(auto& FallbackFontEntry : FallbackFonts)
+	if(FallbackFontData.IsValid())
 	{
-		const UFontBulkData* TmpPtr = FallbackFontEntry.Value->BulkDataPtr;
+		const UFontBulkData* TmpPtr = FallbackFontData->BulkDataPtr;
 		Collector.AddReferencedObject(TmpPtr);
 	}
 
-	if(LastResortFont.IsValid())
+	if(LastResortFontData.IsValid())
 	{
-		const UFontBulkData* TmpPtr = LastResortFont->BulkDataPtr;
+		const UFontBulkData* TmpPtr = LastResortFontData->BulkDataPtr;
 		Collector.AddReferencedObject(TmpPtr);
 	}
 }
diff --git a/Engine/Source/Runtime/SlateCore/Private/Fonts/LegacySlateFontInfoCache.h b/Engine/Source/Runtime/SlateCore/Private/Fonts/LegacySlateFontInfoCache.h
--- a/Engine/Source/Runtime/SlateCore/Private/Fonts/LegacySlateFontInfoCache.h
+++ b/Engine/Source/Runtime/SlateCore/Private/Fonts/LegacySlateFontInfoCache.h
@@ -27,12 +27,17 @@ public:
 	/**
 	 * Get (or create) the culture specific fallback font
 	 */
-	const FFontData& GetFallbackFont();
+	TSharedPtr<const FCompositeFont> GetFallbackFont();
+
+	/**
+	 * Get (or create) the culture specific fallback font
+	 */
+	const FFontData& GetFallbackFontData();
 
 	/**
 	 * Get (or create) the last resort fallback font
 	 */
-	const FFontData& GetLastResortFont();
+	const FFontData& GetLastResortFontData();
 
 	// FGCObject interface
 	virtual void AddReferencedObjects(FReferenceCollector& Collector) override;
@@ -76,12 +81,17 @@ private:
 
 	TMap<FLegacyFontKey, TSharedPtr<const FCompositeFont>> LegacyFontNameToCompositeFont;
 	TSharedPtr<const FCompositeFont> SystemFont;
+	TSharedPtr<const FCompositeFont> FallbackFont;
 	
-	TMap<FName, TSharedPtr<const FFontData>> FallbackFonts;
-	TSharedPtr<const FFontData> LastResortFont;
+	TSharedPtr<const FFontData> FallbackFontData;
+	TSharedPtr<const FFontData> LastResortFontData;
+
+	int32 FallbackFontHistoryVersion;
+	int32 FallbackFontDataHistoryVersion;
 
 	FCriticalSection FallbackFontCS;
-	FCriticalSection LastResortFontCS;
+	FCriticalSection FallbackFontDataCS;
+	FCriticalSection LastResortFontDataCS;
 
 	static TSharedPtr<FLegacySlateFontInfoCache> Singleton;
 };
diff --git a/Engine/Source/Runtime/SlateCore/Private/Fonts/SlateFontInfo.cpp b/Engine/Source/Runtime/SlateCore/Private/Fonts/SlateFontInfo.cpp
--- a/Engine/Source/Runtime/SlateCore/Private/Fonts/SlateFontInfo.cpp
+++ b/Engine/Source/Runtime/SlateCore/Private/Fonts/SlateFontInfo.cpp
@@ -120,7 +120,7 @@ const FCompositeFont* FSlateFontInfo::GetCompositeFont() const
 		return CompositeFont.Get();
 	}
 
-	return nullptr;
+	return FLegacySlateFontInfoCache::Get().GetFallbackFont().Get();
 }
 
 
diff --git a/Engine/Source/Runtime/SlateCore/Public/Fonts/SlateFontInfo.h b/Engine/Source/Runtime/SlateCore/Public/Fonts/SlateFontInfo.h
--- a/Engine/Source/Runtime/SlateCore/Public/Fonts/SlateFontInfo.h
+++ b/Engine/Source/Runtime/SlateCore/Public/Fonts/SlateFontInfo.h
@@ -133,6 +133,7 @@ public:
 
 	/**
 	 * Get the composite font pointer associated with this font info (either directly or via a UFont)
+	 * @note This function will return the fallback font if this font info itself does not contain a valid font. If you want to test whether this font info is empty, use HasValidFont
 	 */
 	const FCompositeFont* GetCompositeFont() const;
 
