Deprecated the FString passthrough support of SLATE_TEXT_ATTRIBUTE and SLATE_TEXT_ARGUMENT

UETOOL-213 - Minimize Slate FString -> FText conversion (remove SLATE_TEXT_ATTRIBUTE)

This will now generate a warning if you try and pass text to Slate as an FString. This will catch any new code which does this, as well as inform any users outside of the Engine and Editor (which have already been converted to exclusively use FText) that they need to update their code to use FText.

diff --git a/Engine/Source/Runtime/Slate/Private/Framework/Application/SlateApplication.cpp b/Engine/Source/Runtime/Slate/Private/Framework/Application/SlateApplication.cpp
--- a/Engine/Source/Runtime/Slate/Private/Framework/Application/SlateApplication.cpp
+++ b/Engine/Source/Runtime/Slate/Private/Framework/Application/SlateApplication.cpp
@@ -3575,8 +3575,12 @@ TSharedRef<SWidget> FSlateApplication::MakeWindowTitleBar( const TSharedRef<SWin
 
 TSharedRef<IToolTip> FSlateApplication::MakeToolTip( const TAttribute<FString>& ToolTipString )
 {
+	PRAGMA_DISABLE_DEPRECATION_WARNINGS
+
 	return SNew(SToolTip)
 		.Text(ToolTipString);
+
+	PRAGMA_ENABLE_DEPRECATION_WARNINGS
 }
 
 
diff --git a/Engine/Source/Runtime/Slate/Public/Framework/Application/SlateApplication.h b/Engine/Source/Runtime/Slate/Public/Framework/Application/SlateApplication.h
--- a/Engine/Source/Runtime/Slate/Public/Framework/Application/SlateApplication.h
+++ b/Engine/Source/Runtime/Slate/Public/Framework/Application/SlateApplication.h
@@ -999,6 +999,7 @@ public:
 	virtual bool IsWindowHousingInteractiveTooltip(const TSharedRef<const SWindow>& WindowToTest) const override;
 	virtual TSharedRef<SWidget> MakeImage( const TAttribute<const FSlateBrush*>& Image, const TAttribute<FSlateColor>& Color, const TAttribute<EVisibility>& Visibility ) const override;
 	virtual TSharedRef<SWidget> MakeWindowTitleBar( const TSharedRef<SWindow>& Window, const TSharedPtr<SWidget>& CenterContent, EHorizontalAlignment CenterContentAlignment, TSharedPtr<IWindowTitleBar>& OutTitleBar ) const override;
+	DEPRECATED(4.8, "Passing text to Slate as FString is deprecated, please use FText instead (likely via a LOCTEXT).")
 	virtual TSharedRef<IToolTip> MakeToolTip( const TAttribute<FString>& ToolTipString ) override;
 	virtual TSharedRef<IToolTip> MakeToolTip( const TAttribute<FText>& ToolTipText ) override;
 	virtual TSharedRef<IToolTip> MakeToolTip( const FText& ToolTipText ) override;
diff --git a/Engine/Source/Runtime/SlateCore/Public/Widgets/DeclarativeSyntaxSupport.h b/Engine/Source/Runtime/SlateCore/Public/Widgets/DeclarativeSyntaxSupport.h
--- a/Engine/Source/Runtime/SlateCore/Public/Widgets/DeclarativeSyntaxSupport.h
+++ b/Engine/Source/Runtime/SlateCore/Public/Widgets/DeclarativeSyntaxSupport.h
@@ -283,178 +283,6 @@
 		}
 
 
-/**
- * Internal macro to declare overloads for setting attributes
- *
- * This macro expects 'WidgetArgsType' and 'AttrName##_Local' 
- * to be defined in the scope where it is used.
- */
-#define INTERNAL_SLATE_DECL_ATTRIBUTE_PASSTHROUGH_COMMON( TargetAttrType, SourceAttrType, Var, AttrName ) \
-		/* Bind attribute with delegate to a global function \
-		 * NOTE: We use a template here to avoid 'typename' issues when hosting attributes inside templated classes */ \
-		WidgetArgsType& AttrName##_Static( SourceAttrType::FGetter::FStaticDelegate::FFuncPtr InFunc )	\
-		{ \
-			Var = TargetAttrType::Create( TargetAttrType::FGetter::CreateStatic( &AttrName##_Local::PassThroughDelegate, SourceAttrType::FGetter::CreateStatic( InFunc ) ) ); \
-			return this->Me(); \
-		} \
-		template< typename Var1Type > \
-		WidgetArgsType& AttrName##_Static( typename SourceAttrType::FGetter::template TStaticDelegate_OneVar< Var1Type >::FFuncPtr InFunc, Var1Type Var1 )	\
-		{ \
-			Var = TargetAttrType::Create( TargetAttrType::FGetter::CreateStatic( &AttrName##_Local::PassThroughDelegate, SourceAttrType::FGetter::CreateStatic( InFunc, Var1 ) ) ); \
-			return this->Me(); \
-		} \
-		template< typename Var1Type, typename Var2Type > \
-		WidgetArgsType& AttrName##_Static( typename SourceAttrType::FGetter::template TStaticDelegate_TwoVars< Var1Type, Var2Type >::FFuncPtr InFunc, Var1Type Var1, Var2Type Var2 )	\
-		{ \
-			Var = TargetAttrType::Create( TargetAttrType::FGetter::CreateStatic( &AttrName##_Local::PassThroughDelegate, SourceAttrType::FGetter::CreateStatic( InFunc, Var1, Var2 ) ) ); \
-			return this->Me(); \
-		} \
-		template< typename Var1Type, typename Var2Type, typename Var3Type > \
-		WidgetArgsType& AttrName##_Static( typename SourceAttrType::FGetter::template TStaticDelegate_ThreeVars< Var1Type, Var2Type, Var3Type >::FFuncPtr InFunc, Var1Type Var1, Var2Type Var2, Var3Type Var3 )	\
-		{ \
-			Var = TargetAttrType::Create( TargetAttrType::FGetter::CreateStatic( &AttrName##_Local::PassThroughDelegate, SourceAttrType::FGetter::CreateStatic( InFunc, Var1, Var2, Var3 ) ) ); \
-			return this->Me(); \
-		} \
-		template< typename Var1Type, typename Var2Type, typename Var3Type, typename Var4Type > \
-		WidgetArgsType& AttrName##_Static( typename SourceAttrType::FGetter::template TStaticDelegate_FourVars< Var1Type, Var2Type, Var3Type, Var4Type >::FFuncPtr InFunc, Var1Type Var1, Var2Type Var2, Var3Type Var3, Var4Type Var4 )	\
-		{ \
-			Var = TargetAttrType::Create( TargetAttrType::FGetter::CreateStatic( &AttrName##_Local::PassThroughDelegate, SourceAttrType::FGetter::CreateStatic( InFunc, Var1, Var2, Var3, Var4 ) ) ); \
-			return this->Me(); \
-		} \
-	\
-		/* We don't support binding FString lambdas to FText attributes because we currently have no way to disambiguate them from FText lambdas \
-		template<typename FunctorType> \
-		WidgetArgsType& AttrName##_Lambda(FunctorType&& InFunctor) */ \
-	\
-		/* Bind attribute with delegate to a raw C++ class method */ \
-		template< class UserClass >	\
-		WidgetArgsType& AttrName##_Raw( UserClass* InUserObject, typename SourceAttrType::FGetter::template TRawMethodDelegate_Const< UserClass >::FMethodPtr InFunc )	\
-		{ \
-			Var = TargetAttrType::Create( TargetAttrType::FGetter::CreateStatic( &AttrName##_Local::PassThroughDelegate, SourceAttrType::FGetter::CreateRaw( InUserObject, InFunc ) ) ); \
-			return this->Me(); \
-		} \
-		template< class UserClass, typename Var1Type >	\
-		WidgetArgsType& AttrName##_Raw( UserClass* InUserObject, typename SourceAttrType::FGetter::template TRawMethodDelegate_OneVar_Const< UserClass, Var1Type  >::FMethodPtr InFunc, Var1Type Var1 )	\
-		{ \
-			Var = TargetAttrType::Create( TargetAttrType::FGetter::CreateStatic( &AttrName##_Local::PassThroughDelegate, SourceAttrType::FGetter::CreateRaw( InUserObject, InFunc, Var1 ) ) ); \
-			return this->Me(); \
-		} \
-		template< class UserClass, typename Var1Type, typename Var2Type >	\
-		WidgetArgsType& AttrName##_Raw( UserClass* InUserObject, typename SourceAttrType::FGetter::template TRawMethodDelegate_TwoVars_Const< UserClass, Var1Type, Var2Type >::FMethodPtr InFunc, Var1Type Var1, Var2Type Var2 )	\
-		{ \
-			Var = TargetAttrType::Create( TargetAttrType::FGetter::CreateStatic( &AttrName##_Local::PassThroughDelegate, SourceAttrType::FGetter::CreateRaw( InUserObject, InFunc, Var1, Var2 ) ) ); \
-			return this->Me(); \
-		} \
-		template< class UserClass, typename Var1Type, typename Var2Type, typename Var3Type >	\
-		WidgetArgsType& AttrName##_Raw( UserClass* InUserObject, typename SourceAttrType::FGetter::template TRawMethodDelegate_ThreeVars_Const< UserClass, Var1Type, Var2Type, Var3Type >::FMethodPtr InFunc, Var1Type Var1, Var2Type Var2, Var3Type Var3 )	\
-		{ \
-			Var = TargetAttrType::Create( TargetAttrType::FGetter::CreateStatic( &AttrName##_Local::PassThroughDelegate, SourceAttrType::FGetter::CreateRaw( InUserObject, InFunc, Var1, Var2, Var3 ) ) ); \
-			return this->Me(); \
-		} \
-		template< class UserClass, typename Var1Type, typename Var2Type, typename Var3Type, typename Var4Type >	\
-		WidgetArgsType& AttrName##_Raw( UserClass* InUserObject, typename SourceAttrType::FGetter::template TRawMethodDelegate_FourVars_Const< UserClass, Var1Type, Var2Type, Var3Type, Var4Type >::FMethodPtr InFunc, Var1Type Var1, Var2Type Var2, Var3Type Var3, Var4Type Var4 )	\
-		{ \
-			Var = TargetAttrType::Create( TargetAttrType::FGetter::CreateStatic( &AttrName##_Local::PassThroughDelegate, SourceAttrType::FGetter::CreateRaw( InUserObject, InFunc, Var1, Var2, Var3, Var4 ) ) ); \
-			return this->Me(); \
-		} \
-	\
-		/* Bind attribute with delegate to a shared pointer-based class method.  Slate mostly uses shared pointers so we use an overload for this type of binding. */ \
-		template< class UserClass >	\
-		WidgetArgsType& AttrName( TSharedRef< UserClass > InUserObjectRef, typename SourceAttrType::FGetter::template TSPMethodDelegate_Const< UserClass >::FMethodPtr InFunc )	\
-		{ \
-			Var = TargetAttrType::Create( TargetAttrType::FGetter::CreateStatic( &AttrName##_Local::PassThroughDelegate, SourceAttrType::FGetter::CreateSP( InUserObjectRef, InFunc ) ) ); \
-			return this->Me(); \
-		} \
-		template< class UserClass, typename Var1Type >	\
-		WidgetArgsType& AttrName( TSharedRef< UserClass > InUserObjectRef, typename SourceAttrType::FGetter::template TSPMethodDelegate_OneVar_Const< UserClass, Var1Type >::FMethodPtr InFunc, Var1Type Var1 )	\
-		{ \
-			Var = TargetAttrType::Create( TargetAttrType::FGetter::CreateStatic( &AttrName##_Local::PassThroughDelegate, SourceAttrType::FGetter::CreateSP( InUserObjectRef, InFunc, Var1 ) ) ); \
-			return this->Me(); \
-		} \
-		template< class UserClass, typename Var1Type, typename Var2Type >	\
-		WidgetArgsType& AttrName( TSharedRef< UserClass > InUserObjectRef, typename SourceAttrType::FGetter::template TSPMethodDelegate_TwoVars_Const< UserClass, Var1Type, Var2Type >::FMethodPtr InFunc, Var1Type Var1, Var2Type Var2 )	\
-		{ \
-			Var = TargetAttrType::Create( TargetAttrType::FGetter::CreateStatic( &AttrName##_Local::PassThroughDelegate, SourceAttrType::FGetter::CreateSP( InUserObjectRef, InFunc, Var1, Var2 ) ) ); \
-			return this->Me(); \
-		} \
-		template< class UserClass, typename Var1Type, typename Var2Type, typename Var3Type >	\
-		WidgetArgsType& AttrName( TSharedRef< UserClass > InUserObjectRef, typename SourceAttrType::FGetter::template TSPMethodDelegate_ThreeVars_Const< UserClass, Var1Type, Var2Type, Var3Type >::FMethodPtr InFunc, Var1Type Var1, Var2Type Var2, Var3Type Var3 )	\
-		{ \
-			Var = TargetAttrType::Create( TargetAttrType::FGetter::CreateStatic( &AttrName##_Local::PassThroughDelegate, SourceAttrType::FGetter::CreateSP( InUserObjectRef, InFunc, Var1, Var2, Var3 ) ) ); \
-			return this->Me(); \
-		} \
-		template< class UserClass, typename Var1Type, typename Var2Type, typename Var3Type, typename Var4Type >	\
-		WidgetArgsType& AttrName( TSharedRef< UserClass > InUserObjectRef, typename SourceAttrType::FGetter::template TSPMethodDelegate_FourVars_Const< UserClass, Var1Type, Var2Type, Var3Type, Var4Type >::FMethodPtr InFunc, Var1Type Var1, Var2Type Var2, Var3Type Var3, Var4Type Var4 )	\
-		{ \
-			Var = TargetAttrType::Create( TargetAttrType::FGetter::CreateStatic( &AttrName##_Local::PassThroughDelegate, SourceAttrType::FGetter::CreateSP( InUserObjectRef, InFunc, Var1, Var2, Var3, Var4 ) ) ); \
-			return this->Me(); \
-		} \
-	\
-		/* Bind attribute with delegate to a shared pointer-based class method.  Slate mostly uses shared pointers so we use an overload for this type of binding. */ \
-		template< class UserClass >	\
-		WidgetArgsType& AttrName( UserClass* InUserObject, typename SourceAttrType::FGetter::template TSPMethodDelegate_Const< UserClass >::FMethodPtr InFunc )	\
-		{ \
-			Var = TargetAttrType::Create( TargetAttrType::FGetter::CreateStatic( &AttrName##_Local::PassThroughDelegate, SourceAttrType::FGetter::CreateSP( InUserObject, InFunc ) ) ); \
-			return this->Me(); \
-		} \
-		template< class UserClass, typename Var1Type >	\
-		WidgetArgsType& AttrName( UserClass* InUserObject, typename SourceAttrType::FGetter::template TSPMethodDelegate_OneVar_Const< UserClass, Var1Type >::FMethodPtr InFunc, Var1Type Var1 )	\
-		{ \
-			Var = TargetAttrType::Create( TargetAttrType::FGetter::CreateStatic( &AttrName##_Local::PassThroughDelegate, SourceAttrType::FGetter::CreateSP( InUserObject, InFunc, Var1 ) ) ); \
-			return this->Me(); \
-		} \
-		template< class UserClass, typename Var1Type, typename Var2Type >	\
-		WidgetArgsType& AttrName( UserClass* InUserObject, typename SourceAttrType::FGetter::template TSPMethodDelegate_TwoVars_Const< UserClass, Var1Type, Var2Type >::FMethodPtr InFunc, Var1Type Var1, Var2Type Var2 )	\
-		{ \
-			Var = TargetAttrType::Create( TargetAttrType::FGetter::CreateStatic( &AttrName##_Local::PassThroughDelegate, SourceAttrType::FGetter::CreateSP( InUserObject, InFunc, Var1, Var2 ) ) ); \
-			return this->Me(); \
-		} \
-		template< class UserClass, typename Var1Type, typename Var2Type, typename Var3Type >	\
-		WidgetArgsType& AttrName( UserClass* InUserObject, typename SourceAttrType::FGetter::template TSPMethodDelegate_ThreeVars_Const< UserClass, Var1Type, Var2Type, Var3Type >::FMethodPtr InFunc, Var1Type Var1, Var2Type Var2, Var3Type Var3 )	\
-		{ \
-			Var = TargetAttrType::Create( TargetAttrType::FGetter::CreateStatic( &AttrName##_Local::PassThroughDelegate, SourceAttrType::FGetter::CreateSP( InUserObject, InFunc, Var1, Var2, Var3 ) ) ); \
-			return this->Me(); \
-		} \
-		template< class UserClass, typename Var1Type, typename Var2Type, typename Var3Type, typename Var4Type >	\
-		WidgetArgsType& AttrName( UserClass* InUserObject, typename SourceAttrType::FGetter::template TSPMethodDelegate_FourVars_Const< UserClass, Var1Type, Var2Type, Var3Type, Var4Type >::FMethodPtr InFunc, Var1Type Var1, Var2Type Var2, Var3Type Var3, Var4Type Var4 )	\
-		{ \
-			Var = TargetAttrType::Create( TargetAttrType::FGetter::CreateStatic( &AttrName##_Local::PassThroughDelegate, SourceAttrType::FGetter::CreateSP( InUserObject, InFunc, Var1, Var2, Var3, Var4 ) ) ); \
-			return this->Me(); \
-		} \
-	\
-		/* Bind attribute with delegate to a UObject-based class method */ \
-		template< class UserClass >	\
-		WidgetArgsType& AttrName##_UObject( UserClass* InUserObject, typename SourceAttrType::FGetter::template TUObjectMethodDelegate_Const< UserClass >::FMethodPtr InFunc )	\
-		{ \
-			Var = TargetAttrType::Create( TargetAttrType::FGetter::CreateStatic( &AttrName##_Local::PassThroughDelegate, SourceAttrType::FGetter::CreateUObject( InUserObject, InFunc ) ) ); \
-			return this->Me(); \
-		} \
-		template< class UserClass, typename Var1Type >	\
-		WidgetArgsType& AttrName##_UObject( UserClass* InUserObject, typename SourceAttrType::FGetter::template TUObjectMethodDelegate_OneVar_Const< UserClass, Var1Type >::FMethodPtr InFunc, Var1Type Var1 )	\
-		{ \
-			Var = TargetAttrType::Create( TargetAttrType::FGetter::CreateStatic( &AttrName##_Local::PassThroughDelegate, SourceAttrType::FGetter::CreateUObject( InUserObject, InFunc, Var1 ) ) ); \
-			return this->Me(); \
-		} \
-		template< class UserClass, typename Var1Type, typename Var2Type  >	\
-		WidgetArgsType& AttrName##_UObject( UserClass* InUserObject, typename SourceAttrType::FGetter::template TUObjectMethodDelegate_TwoVars_Const< UserClass, Var1Type, Var2Type >::FMethodPtr InFunc, Var1Type Var1, Var2Type Var2 )	\
-		{ \
-			Var = TargetAttrType::Create( TargetAttrType::FGetter::CreateStatic( &AttrName##_Local::PassThroughDelegate, SourceAttrType::FGetter::CreateUObject( InUserObject, InFunc, Var1, Var2 ) ) ); \
-			return this->Me(); \
-		} \
-		template< class UserClass, typename Var1Type, typename Var2Type, typename Var3Type >	\
-		WidgetArgsType& AttrName##_UObject( UserClass* InUserObject, typename SourceAttrType::FGetter::template TUObjectMethodDelegate_ThreeVars_Const< UserClass, Var1Type, Var2Type, Var3Type >::FMethodPtr InFunc, Var1Type Var1, Var2Type Var2, Var3Type Var3 )	\
-		{ \
-			Var = TargetAttrType::Create( TargetAttrType::FGetter::CreateStatic( &AttrName##_Local::PassThroughDelegate, SourceAttrType::FGetter::CreateUObject( InUserObject, InFunc, Var1, Var2, Var3 ) ) ); \
-			return this->Me(); \
-		} \
-		template< class UserClass, typename Var1Type, typename Var2Type, typename Var3Type, typename Var4Type >	\
-		WidgetArgsType& AttrName##_UObject( UserClass* InUserObject, typename SourceAttrType::FGetter::template TUObjectMethodDelegate_FourVars_Const< UserClass, Var1Type, Var2Type, Var3Type, Var4Type >::FMethodPtr InFunc, Var1Type Var1, Var2Type Var2, Var3Type Var3, Var4Type Var4 )	\
-		{ \
-			Var = TargetAttrType::Create( TargetAttrType::FGetter::CreateStatic( &AttrName##_Local::PassThroughDelegate, SourceAttrType::FGetter::CreateUObject( InUserObject, InFunc, Var1, Var2, Var3, Var4 ) ) ); \
-			return this->Me(); \
-		}
-		
-
 /**
  * Use this macro to add a attribute to the declaration of your widget.
  * An attribute can be a value or a function.
@@ -464,48 +292,6 @@
 	INTERNAL_SLATE_DECL_ATTRIBUTE( _##AttrName, TAttribute< AttrType >, AttrName )
 
 
-#define SLATE_STRING_ATTRIBUTE_IMPL( Var, AttrName ) \
-	TAttribute< FString > Var; \
-	struct AttrName##_Local\
-	{\
-		static FString PassThroughDelegate( TAttribute< FText >::FGetter TextDelegate )\
-		{\
-			return (TextDelegate.IsBound()) ? TextDelegate.Execute().ToString() : FString();\
-		}\
-		\
-		static FString PassThroughAttribute( TAttribute< FText > TextAttribute )\
-		{\
-			return TextAttribute.Get( FText::GetEmpty() ).ToString();\
-		}\
-	};\
-	WidgetArgsType& AttrName( const FString& String ) \
-	{ \
-		Var = TAttribute< FString >( String ); \
-		return this->Me(); \
-	} \
-	WidgetArgsType& AttrName( const FText& Text ) \
-	{ \
-		Var = TAttribute< FString >( Text.ToString() ); \
-		return this->Me(); \
-	} \
-	WidgetArgsType& AttrName( const TAttribute< FText >& InAttribute ) \
-	{ \
-		if ( InAttribute.IsBound() ) {\
-			Var = TAttribute< FString >::Create( TAttribute< FString >::FGetter::CreateStatic( &AttrName##_Local::PassThroughAttribute, InAttribute ) ); \
-		} else {\
-			Var = InAttribute.Get().ToString();\
-		}\
-		return this->Me(); \
-	} \
-	INTERNAL_SLATE_DECL_ATTRIBUTE_PASSTHROUGH_COMMON( TAttribute< FString >, TAttribute< FText >, Var, AttrName )\
-	INTERNAL_SLATE_DECL_ATTRIBUTE( Var, TAttribute< FString >, AttrName )
-
-/**
- * Declare an FString attribute which also allows content passthrough from FText
- */
-#define SLATE_STRING_ATTRIBUTE( AttrName ) SLATE_STRING_ATTRIBUTE_IMPL( _##AttrName, AttrName )
-
-
 #define SLATE_TEXT_ATTRIBUTE_IMPL( Var, AttrName ) \
 	TAttribute< FText > Var; \
 	struct AttrName##_Local\
@@ -525,11 +311,13 @@
 		Var = TAttribute< FText >( Text ); \
 		return this->Me(); \
 	} \
+	DEPRECATED(4.8, "Passing text to Slate as FString is deprecated, please use FText instead (likely via a LOCTEXT).")\
 	WidgetArgsType& AttrName( const FString& String ) \
 	{ \
 		Var = TAttribute< FText >( FText::FromString( String ) ); \
 		return this->Me(); \
 	} \
+	DEPRECATED(4.8, "Passing text to Slate as FString is deprecated, please use FText instead (likely via a LOCTEXT).")\
 	WidgetArgsType& AttrName( const TAttribute< FString >& InAttribute ) \
 	{ \
 		if ( InAttribute.IsBound() ) {\
@@ -539,7 +327,194 @@
 		}\
 		return this->Me(); \
 	} \
-	INTERNAL_SLATE_DECL_ATTRIBUTE_PASSTHROUGH_COMMON( TAttribute< FText >, TAttribute< FString >, Var, AttrName )\
+	/* Bind attribute with delegate to a global function \
+	* NOTE: We use a template here to avoid 'typename' issues when hosting attributes inside templated classes */ \
+	DEPRECATED(4.8, "Passing text to Slate as FString is deprecated, please use FText instead (likely via a LOCTEXT).")\
+	WidgetArgsType& AttrName##_Static( TAttribute< FString >::FGetter::FStaticDelegate::FFuncPtr InFunc )	\
+	{ \
+		Var = TAttribute< FText >::Create( TAttribute< FText >::FGetter::CreateStatic( &AttrName##_Local::PassThroughDelegate, TAttribute< FString >::FGetter::CreateStatic( InFunc ) ) ); \
+		return this->Me(); \
+	} \
+	template< typename Var1Type > \
+	DEPRECATED(4.8, "Passing text to Slate as FString is deprecated, please use FText instead (likely via a LOCTEXT).")\
+	WidgetArgsType& AttrName##_Static( typename TAttribute< FString >::FGetter::template TStaticDelegate_OneVar< Var1Type >::FFuncPtr InFunc, Var1Type Var1 )	\
+	{ \
+		Var = TAttribute< FText >::Create( TAttribute< FText >::FGetter::CreateStatic( &AttrName##_Local::PassThroughDelegate, TAttribute< FString >::FGetter::CreateStatic( InFunc, Var1 ) ) ); \
+		return this->Me(); \
+	} \
+	template< typename Var1Type, typename Var2Type > \
+	DEPRECATED(4.8, "Passing text to Slate as FString is deprecated, please use FText instead (likely via a LOCTEXT).")\
+	WidgetArgsType& AttrName##_Static( typename TAttribute< FString >::FGetter::template TStaticDelegate_TwoVars< Var1Type, Var2Type >::FFuncPtr InFunc, Var1Type Var1, Var2Type Var2 )	\
+	{ \
+		Var = TAttribute< FText >::Create( TAttribute< FText >::FGetter::CreateStatic( &AttrName##_Local::PassThroughDelegate, TAttribute< FString >::FGetter::CreateStatic( InFunc, Var1, Var2 ) ) ); \
+		return this->Me(); \
+	} \
+	template< typename Var1Type, typename Var2Type, typename Var3Type > \
+	DEPRECATED(4.8, "Passing text to Slate as FString is deprecated, please use FText instead (likely via a LOCTEXT).")\
+	WidgetArgsType& AttrName##_Static( typename TAttribute< FString >::FGetter::template TStaticDelegate_ThreeVars< Var1Type, Var2Type, Var3Type >::FFuncPtr InFunc, Var1Type Var1, Var2Type Var2, Var3Type Var3 )	\
+	{ \
+		Var = TAttribute< FText >::Create( TAttribute< FText >::FGetter::CreateStatic( &AttrName##_Local::PassThroughDelegate, TAttribute< FString >::FGetter::CreateStatic( InFunc, Var1, Var2, Var3 ) ) ); \
+		return this->Me(); \
+	} \
+	template< typename Var1Type, typename Var2Type, typename Var3Type, typename Var4Type > \
+	DEPRECATED(4.8, "Passing text to Slate as FString is deprecated, please use FText instead (likely via a LOCTEXT).")\
+	WidgetArgsType& AttrName##_Static( typename TAttribute< FString >::FGetter::template TStaticDelegate_FourVars< Var1Type, Var2Type, Var3Type, Var4Type >::FFuncPtr InFunc, Var1Type Var1, Var2Type Var2, Var3Type Var3, Var4Type Var4 )	\
+	{ \
+		Var = TAttribute< FText >::Create( TAttribute< FText >::FGetter::CreateStatic( &AttrName##_Local::PassThroughDelegate, TAttribute< FString >::FGetter::CreateStatic( InFunc, Var1, Var2, Var3, Var4 ) ) ); \
+		return this->Me(); \
+	} \
+\
+	/* We don't support binding FString lambdas to FText attributes because we currently have no way to disambiguate them from FText lambdas \
+	template<typename FunctorType> \
+	WidgetArgsType& AttrName##_Lambda(FunctorType&& InFunctor) */ \
+\
+	/* Bind attribute with delegate to a raw C++ class method */ \
+	template< class UserClass >	\
+	DEPRECATED(4.8, "Passing text to Slate as FString is deprecated, please use FText instead (likely via a LOCTEXT).")\
+	WidgetArgsType& AttrName##_Raw( UserClass* InUserObject, typename TAttribute< FString >::FGetter::template TRawMethodDelegate_Const< UserClass >::FMethodPtr InFunc )	\
+	{ \
+		Var = TAttribute< FText >::Create( TAttribute< FText >::FGetter::CreateStatic( &AttrName##_Local::PassThroughDelegate, TAttribute< FString >::FGetter::CreateRaw( InUserObject, InFunc ) ) ); \
+		return this->Me(); \
+	} \
+	template< class UserClass, typename Var1Type >	\
+	DEPRECATED(4.8, "Passing text to Slate as FString is deprecated, please use FText instead (likely via a LOCTEXT).")\
+	WidgetArgsType& AttrName##_Raw( UserClass* InUserObject, typename TAttribute< FString >::FGetter::template TRawMethodDelegate_OneVar_Const< UserClass, Var1Type  >::FMethodPtr InFunc, Var1Type Var1 )	\
+	{ \
+		Var = TAttribute< FText >::Create( TAttribute< FText >::FGetter::CreateStatic( &AttrName##_Local::PassThroughDelegate, TAttribute< FString >::FGetter::CreateRaw( InUserObject, InFunc, Var1 ) ) ); \
+		return this->Me(); \
+	} \
+	template< class UserClass, typename Var1Type, typename Var2Type >	\
+	DEPRECATED(4.8, "Passing text to Slate as FString is deprecated, please use FText instead (likely via a LOCTEXT).")\
+	WidgetArgsType& AttrName##_Raw( UserClass* InUserObject, typename TAttribute< FString >::FGetter::template TRawMethodDelegate_TwoVars_Const< UserClass, Var1Type, Var2Type >::FMethodPtr InFunc, Var1Type Var1, Var2Type Var2 )	\
+	{ \
+		Var = TAttribute< FText >::Create( TAttribute< FText >::FGetter::CreateStatic( &AttrName##_Local::PassThroughDelegate, TAttribute< FString >::FGetter::CreateRaw( InUserObject, InFunc, Var1, Var2 ) ) ); \
+		return this->Me(); \
+	} \
+	template< class UserClass, typename Var1Type, typename Var2Type, typename Var3Type >	\
+	DEPRECATED(4.8, "Passing text to Slate as FString is deprecated, please use FText instead (likely via a LOCTEXT).")\
+	WidgetArgsType& AttrName##_Raw( UserClass* InUserObject, typename TAttribute< FString >::FGetter::template TRawMethodDelegate_ThreeVars_Const< UserClass, Var1Type, Var2Type, Var3Type >::FMethodPtr InFunc, Var1Type Var1, Var2Type Var2, Var3Type Var3 )	\
+	{ \
+		Var = TAttribute< FText >::Create( TAttribute< FText >::FGetter::CreateStatic( &AttrName##_Local::PassThroughDelegate, TAttribute< FString >::FGetter::CreateRaw( InUserObject, InFunc, Var1, Var2, Var3 ) ) ); \
+		return this->Me(); \
+	} \
+	template< class UserClass, typename Var1Type, typename Var2Type, typename Var3Type, typename Var4Type >	\
+	DEPRECATED(4.8, "Passing text to Slate as FString is deprecated, please use FText instead (likely via a LOCTEXT).")\
+	WidgetArgsType& AttrName##_Raw( UserClass* InUserObject, typename TAttribute< FString >::FGetter::template TRawMethodDelegate_FourVars_Const< UserClass, Var1Type, Var2Type, Var3Type, Var4Type >::FMethodPtr InFunc, Var1Type Var1, Var2Type Var2, Var3Type Var3, Var4Type Var4 )	\
+	{ \
+		Var = TAttribute< FText >::Create( TAttribute< FText >::FGetter::CreateStatic( &AttrName##_Local::PassThroughDelegate, TAttribute< FString >::FGetter::CreateRaw( InUserObject, InFunc, Var1, Var2, Var3, Var4 ) ) ); \
+		return this->Me(); \
+	} \
+\
+	/* Bind attribute with delegate to a shared pointer-based class method.  Slate mostly uses shared pointers so we use an overload for this type of binding. */ \
+	template< class UserClass >	\
+	DEPRECATED(4.8, "Passing text to Slate as FString is deprecated, please use FText instead (likely via a LOCTEXT).")\
+	WidgetArgsType& AttrName( TSharedRef< UserClass > InUserObjectRef, typename TAttribute< FString >::FGetter::template TSPMethodDelegate_Const< UserClass >::FMethodPtr InFunc )	\
+	{ \
+		Var = TAttribute< FText >::Create( TAttribute< FText >::FGetter::CreateStatic( &AttrName##_Local::PassThroughDelegate, TAttribute< FString >::FGetter::CreateSP( InUserObjectRef, InFunc ) ) ); \
+		return this->Me(); \
+	} \
+	template< class UserClass, typename Var1Type >	\
+	DEPRECATED(4.8, "Passing text to Slate as FString is deprecated, please use FText instead (likely via a LOCTEXT).")\
+	WidgetArgsType& AttrName( TSharedRef< UserClass > InUserObjectRef, typename TAttribute< FString >::FGetter::template TSPMethodDelegate_OneVar_Const< UserClass, Var1Type >::FMethodPtr InFunc, Var1Type Var1 )	\
+	{ \
+		Var = TAttribute< FText >::Create( TAttribute< FText >::FGetter::CreateStatic( &AttrName##_Local::PassThroughDelegate, TAttribute< FString >::FGetter::CreateSP( InUserObjectRef, InFunc, Var1 ) ) ); \
+		return this->Me(); \
+	} \
+	template< class UserClass, typename Var1Type, typename Var2Type >	\
+	DEPRECATED(4.8, "Passing text to Slate as FString is deprecated, please use FText instead (likely via a LOCTEXT).")\
+	WidgetArgsType& AttrName( TSharedRef< UserClass > InUserObjectRef, typename TAttribute< FString >::FGetter::template TSPMethodDelegate_TwoVars_Const< UserClass, Var1Type, Var2Type >::FMethodPtr InFunc, Var1Type Var1, Var2Type Var2 )	\
+	{ \
+		Var = TAttribute< FText >::Create( TAttribute< FText >::FGetter::CreateStatic( &AttrName##_Local::PassThroughDelegate, TAttribute< FString >::FGetter::CreateSP( InUserObjectRef, InFunc, Var1, Var2 ) ) ); \
+		return this->Me(); \
+	} \
+	template< class UserClass, typename Var1Type, typename Var2Type, typename Var3Type >	\
+	DEPRECATED(4.8, "Passing text to Slate as FString is deprecated, please use FText instead (likely via a LOCTEXT).")\
+	WidgetArgsType& AttrName( TSharedRef< UserClass > InUserObjectRef, typename TAttribute< FString >::FGetter::template TSPMethodDelegate_ThreeVars_Const< UserClass, Var1Type, Var2Type, Var3Type >::FMethodPtr InFunc, Var1Type Var1, Var2Type Var2, Var3Type Var3 )	\
+	{ \
+		Var = TAttribute< FText >::Create( TAttribute< FText >::FGetter::CreateStatic( &AttrName##_Local::PassThroughDelegate, TAttribute< FString >::FGetter::CreateSP( InUserObjectRef, InFunc, Var1, Var2, Var3 ) ) ); \
+		return this->Me(); \
+	} \
+	template< class UserClass, typename Var1Type, typename Var2Type, typename Var3Type, typename Var4Type >	\
+	DEPRECATED(4.8, "Passing text to Slate as FString is deprecated, please use FText instead (likely via a LOCTEXT).")\
+	WidgetArgsType& AttrName( TSharedRef< UserClass > InUserObjectRef, typename TAttribute< FString >::FGetter::template TSPMethodDelegate_FourVars_Const< UserClass, Var1Type, Var2Type, Var3Type, Var4Type >::FMethodPtr InFunc, Var1Type Var1, Var2Type Var2, Var3Type Var3, Var4Type Var4 )	\
+	{ \
+		Var = TAttribute< FText >::Create( TAttribute< FText >::FGetter::CreateStatic( &AttrName##_Local::PassThroughDelegate, TAttribute< FString >::FGetter::CreateSP( InUserObjectRef, InFunc, Var1, Var2, Var3, Var4 ) ) ); \
+		return this->Me(); \
+	} \
+\
+	/* Bind attribute with delegate to a shared pointer-based class method.  Slate mostly uses shared pointers so we use an overload for this type of binding. */ \
+	template< class UserClass >	\
+	DEPRECATED(4.8, "Passing text to Slate as FString is deprecated, please use FText instead (likely via a LOCTEXT).")\
+	WidgetArgsType& AttrName( UserClass* InUserObject, typename TAttribute< FString >::FGetter::template TSPMethodDelegate_Const< UserClass >::FMethodPtr InFunc )	\
+	{ \
+		Var = TAttribute< FText >::Create( TAttribute< FText >::FGetter::CreateStatic( &AttrName##_Local::PassThroughDelegate, TAttribute< FString >::FGetter::CreateSP( InUserObject, InFunc ) ) ); \
+		return this->Me(); \
+	} \
+	template< class UserClass, typename Var1Type >	\
+	DEPRECATED(4.8, "Passing text to Slate as FString is deprecated, please use FText instead (likely via a LOCTEXT).")\
+	WidgetArgsType& AttrName( UserClass* InUserObject, typename TAttribute< FString >::FGetter::template TSPMethodDelegate_OneVar_Const< UserClass, Var1Type >::FMethodPtr InFunc, Var1Type Var1 )	\
+	{ \
+		Var = TAttribute< FText >::Create( TAttribute< FText >::FGetter::CreateStatic( &AttrName##_Local::PassThroughDelegate, TAttribute< FString >::FGetter::CreateSP( InUserObject, InFunc, Var1 ) ) ); \
+		return this->Me(); \
+	} \
+	template< class UserClass, typename Var1Type, typename Var2Type >	\
+	DEPRECATED(4.8, "Passing text to Slate as FString is deprecated, please use FText instead (likely via a LOCTEXT).")\
+	WidgetArgsType& AttrName( UserClass* InUserObject, typename TAttribute< FString >::FGetter::template TSPMethodDelegate_TwoVars_Const< UserClass, Var1Type, Var2Type >::FMethodPtr InFunc, Var1Type Var1, Var2Type Var2 )	\
+	{ \
+		Var = TAttribute< FText >::Create( TAttribute< FText >::FGetter::CreateStatic( &AttrName##_Local::PassThroughDelegate, TAttribute< FString >::FGetter::CreateSP( InUserObject, InFunc, Var1, Var2 ) ) ); \
+		return this->Me(); \
+	} \
+	template< class UserClass, typename Var1Type, typename Var2Type, typename Var3Type >	\
+	DEPRECATED(4.8, "Passing text to Slate as FString is deprecated, please use FText instead (likely via a LOCTEXT).")\
+	WidgetArgsType& AttrName( UserClass* InUserObject, typename TAttribute< FString >::FGetter::template TSPMethodDelegate_ThreeVars_Const< UserClass, Var1Type, Var2Type, Var3Type >::FMethodPtr InFunc, Var1Type Var1, Var2Type Var2, Var3Type Var3 )	\
+	{ \
+		Var = TAttribute< FText >::Create( TAttribute< FText >::FGetter::CreateStatic( &AttrName##_Local::PassThroughDelegate, TAttribute< FString >::FGetter::CreateSP( InUserObject, InFunc, Var1, Var2, Var3 ) ) ); \
+		return this->Me(); \
+	} \
+	template< class UserClass, typename Var1Type, typename Var2Type, typename Var3Type, typename Var4Type >	\
+	DEPRECATED(4.8, "Passing text to Slate as FString is deprecated, please use FText instead (likely via a LOCTEXT).")\
+	WidgetArgsType& AttrName( UserClass* InUserObject, typename TAttribute< FString >::FGetter::template TSPMethodDelegate_FourVars_Const< UserClass, Var1Type, Var2Type, Var3Type, Var4Type >::FMethodPtr InFunc, Var1Type Var1, Var2Type Var2, Var3Type Var3, Var4Type Var4 )	\
+	{ \
+		Var = TAttribute< FText >::Create( TAttribute< FText >::FGetter::CreateStatic( &AttrName##_Local::PassThroughDelegate, TAttribute< FString >::FGetter::CreateSP( InUserObject, InFunc, Var1, Var2, Var3, Var4 ) ) ); \
+		return this->Me(); \
+	} \
+\
+	/* Bind attribute with delegate to a UObject-based class method */ \
+	template< class UserClass >	\
+	DEPRECATED(4.8, "Passing text to Slate as FString is deprecated, please use FText instead (likely via a LOCTEXT).")\
+	WidgetArgsType& AttrName##_UObject( UserClass* InUserObject, typename TAttribute< FString >::FGetter::template TUObjectMethodDelegate_Const< UserClass >::FMethodPtr InFunc )	\
+	{ \
+		Var = TAttribute< FText >::Create( TAttribute< FText >::FGetter::CreateStatic( &AttrName##_Local::PassThroughDelegate, TAttribute< FString >::FGetter::CreateUObject( InUserObject, InFunc ) ) ); \
+		return this->Me(); \
+	} \
+	template< class UserClass, typename Var1Type >	\
+	DEPRECATED(4.8, "Passing text to Slate as FString is deprecated, please use FText instead (likely via a LOCTEXT).")\
+	WidgetArgsType& AttrName##_UObject( UserClass* InUserObject, typename TAttribute< FString >::FGetter::template TUObjectMethodDelegate_OneVar_Const< UserClass, Var1Type >::FMethodPtr InFunc, Var1Type Var1 )	\
+	{ \
+		Var = TAttribute< FText >::Create( TAttribute< FText >::FGetter::CreateStatic( &AttrName##_Local::PassThroughDelegate, TAttribute< FString >::FGetter::CreateUObject( InUserObject, InFunc, Var1 ) ) ); \
+		return this->Me(); \
+	} \
+	template< class UserClass, typename Var1Type, typename Var2Type  >	\
+	DEPRECATED(4.8, "Passing text to Slate as FString is deprecated, please use FText instead (likely via a LOCTEXT).")\
+	WidgetArgsType& AttrName##_UObject( UserClass* InUserObject, typename TAttribute< FString >::FGetter::template TUObjectMethodDelegate_TwoVars_Const< UserClass, Var1Type, Var2Type >::FMethodPtr InFunc, Var1Type Var1, Var2Type Var2 )	\
+	{ \
+		Var = TAttribute< FText >::Create( TAttribute< FText >::FGetter::CreateStatic( &AttrName##_Local::PassThroughDelegate, TAttribute< FString >::FGetter::CreateUObject( InUserObject, InFunc, Var1, Var2 ) ) ); \
+		return this->Me(); \
+	} \
+	template< class UserClass, typename Var1Type, typename Var2Type, typename Var3Type >	\
+	DEPRECATED(4.8, "Passing text to Slate as FString is deprecated, please use FText instead (likely via a LOCTEXT).")\
+	WidgetArgsType& AttrName##_UObject( UserClass* InUserObject, typename TAttribute< FString >::FGetter::template TUObjectMethodDelegate_ThreeVars_Const< UserClass, Var1Type, Var2Type, Var3Type >::FMethodPtr InFunc, Var1Type Var1, Var2Type Var2, Var3Type Var3 )	\
+	{ \
+		Var = TAttribute< FText >::Create( TAttribute< FText >::FGetter::CreateStatic( &AttrName##_Local::PassThroughDelegate, TAttribute< FString >::FGetter::CreateUObject( InUserObject, InFunc, Var1, Var2, Var3 ) ) ); \
+		return this->Me(); \
+	} \
+	template< class UserClass, typename Var1Type, typename Var2Type, typename Var3Type, typename Var4Type >	\
+	DEPRECATED(4.8, "Passing text to Slate as FString is deprecated, please use FText instead (likely via a LOCTEXT).")\
+	WidgetArgsType& AttrName##_UObject( UserClass* InUserObject, typename TAttribute< FString >::FGetter::template TUObjectMethodDelegate_FourVars_Const< UserClass, Var1Type, Var2Type, Var3Type, Var4Type >::FMethodPtr InFunc, Var1Type Var1, Var2Type Var2, Var3Type Var3, Var4Type Var4 )	\
+	{ \
+		Var = TAttribute< FText >::Create( TAttribute< FText >::FGetter::CreateStatic( &AttrName##_Local::PassThroughDelegate, TAttribute< FString >::FGetter::CreateUObject( InUserObject, InFunc, Var1, Var2, Var3, Var4 ) ) ); \
+		return this->Me(); \
+	}\
 	INTERNAL_SLATE_DECL_ATTRIBUTE( Var, TAttribute< FText >, AttrName )
 
 /**
@@ -562,6 +537,7 @@
 
 #define SLATE_TEXT_ARGUMENT( ArgName ) \
 	FText _##ArgName; \
+	DEPRECATED(4.8, "Passing text to Slate as FString is deprecated, please use FText instead (likely via a LOCTEXT).")\
 	WidgetArgsType& ArgName( FString InArg ) \
 	{ \
 		_##ArgName = FText::FromString(InArg); \
diff --git a/Engine/Source/Runtime/SlateCore/Public/Widgets/SWidget.h b/Engine/Source/Runtime/SlateCore/Public/Widgets/SWidget.h
--- a/Engine/Source/Runtime/SlateCore/Public/Widgets/SWidget.h
+++ b/Engine/Source/Runtime/SlateCore/Public/Widgets/SWidget.h
@@ -682,6 +682,7 @@ public:
 	 *
 	 * @param InToolTipText  the text that should appear in the tool tip
 	 */
+	DEPRECATED(4.8, "Passing text to Slate as FString is deprecated, please use FText instead (likely via a LOCTEXT).")
 	void SetToolTipText( const TAttribute<FString>& InToolTipText );
 
 	void SetToolTipText(const TAttribute<FText>& ToolTipText);
