Added user inxed to FFocusEvent.

diff --git a/Engine/Source/Runtime/SlateCore/Private/Widgets/SWidget.cpp b/Engine/Source/Runtime/SlateCore/Private/Widgets/SWidget.cpp
--- a/Engine/Source/Runtime/SlateCore/Private/Widgets/SWidget.cpp
+++ b/Engine/Source/Runtime/SlateCore/Private/Widgets/SWidget.cpp
@@ -74,19 +74,10 @@ FReply SWidget::OnFocusReceived(const FGeometry& MyGeometry, const FFocusEvent&
 	return FReply::Unhandled();
 }
 
-FReply SWidget::OnKeyboardFocusReceived(const FGeometry& MyGeometry, const FKeyboardFocusEvent& InFocusEvent)
-{
-	return FReply::Unhandled();
-}
-
 void SWidget::OnFocusLost(const FFocusEvent& InFocusEvent)
 {
 }
 
-void SWidget::OnKeyboardFocusLost(const FKeyboardFocusEvent& InFocusEvent)
-{
-}
-
 void SWidget::OnFocusChanging(const FWeakWidgetPath& PreviousFocusPath, const FWidgetPath& NewWidgetPath)
 {
 }
diff --git a/Engine/Source/Runtime/SlateCore/Public/Input/Events.h b/Engine/Source/Runtime/SlateCore/Public/Input/Events.h
--- a/Engine/Source/Runtime/SlateCore/Public/Input/Events.h
+++ b/Engine/Source/Runtime/SlateCore/Public/Input/Events.h
@@ -61,49 +61,54 @@ struct FFocusEvent
 public:
 
 	/**
-	 * UStruct Constructor.  Not meant for normal usage.
-	 */
+	* UStruct Constructor.  Not meant for normal usage.
+	*/
 	FFocusEvent()
-		: Cause(EFocusCause::SetDirectly)
+		: Cause(EFocusCause::SetDirectly),
+		UserIndex(0)
 	{ }
 
 	/**
-	 * Constructor.  Events are immutable once constructed.
-	 *
-	 * @param  InCause  The cause of the focus event
-	 */
-	FFocusEvent(const EFocusCause InCause)
-		: Cause(InCause)
+	* Constructor.  Events are immutable once constructed.
+	*
+	* @param  InCause  The cause of the focus event
+	*/
+	FFocusEvent(const EFocusCause InCause, uint32 InUserIndex)
+		: Cause(InCause),
+		UserIndex(InUserIndex)
 	{ }
 
 
 	/**
-	 * Queries the reason for the focus change
-	 *
-	 * @return  The cause of the focus change
-	 */
+	* Queries the reason for the focus change
+	*
+	* @return  The cause of the focus change
+	*/
 	EFocusCause GetCause() const
 	{
 		return Cause;
 	}
 
+	/**
+	* Queries the user that is changing focus
+	*
+	* @return  The user that is changing focus
+	*/
+	uint32 GetUser() const
+	{
+		return UserIndex;
+	}
+
 private:
 
 	/** The cause of the focus change */
 	EFocusCause Cause;
-};
 
-//@Todo slate: Remove this as soon as the 4.6 deprecated API is Removed.
-USTRUCT()
-//struct DEPRECATED(4.6, "Use FFocusEvent") FKeyboardFocusEvent : public FFocusEvent
-struct FKeyboardFocusEvent : public FFocusEvent
-{
-	GENERATED_USTRUCT_BODY()
-public:
-	FKeyboardFocusEvent() { }
-	FKeyboardFocusEvent(const EFocusCause InCause) : FFocusEvent(InCause) { }
+	/** User that is changing focus*/
+	uint32 UserIndex;
 };
 
+
 /**
  * Represents the current and last cursor position in a "virtual window" for events that are routed to widgets transformed in a 3D scene.
  */
diff --git a/Engine/Source/Runtime/SlateCore/Public/Widgets/SWidget.h b/Engine/Source/Runtime/SlateCore/Public/Widgets/SWidget.h
--- a/Engine/Source/Runtime/SlateCore/Public/Widgets/SWidget.h
+++ b/Engine/Source/Runtime/SlateCore/Public/Widgets/SWidget.h
@@ -148,9 +148,6 @@ public:
 	 */
 	virtual FReply OnFocusReceived(const FGeometry& MyGeometry, const FFocusEvent& InFocusEvent);
 
-	DEPRECATED(4.6, "SWidget::OnKeyboardFocusReceived() is deprecated, implement SWidget::OnFocusReceived() instead.")
-	virtual FReply OnKeyboardFocusReceived(const FGeometry& MyGeometry, const FKeyboardFocusEvent& InFocusEvent);
-
 	/**
 	 * Called when this widget loses focus.  This event does not bubble.
 	 *
@@ -158,15 +155,9 @@ public:
 	 */
 	virtual void OnFocusLost(const FFocusEvent& InFocusEvent);
 
-	DEPRECATED(4.6, "SWidget::OnKeyboardFocusLost() is deprecated, implement SWidget::OnFocusLost() instead.")
-	virtual void OnKeyboardFocusLost(const FKeyboardFocusEvent& InFocusEvent);
-
 	/** Called whenever a focus path is changing on all the widgets within the old and new focus paths */
 	virtual void OnFocusChanging(const FWeakWidgetPath& PreviousFocusPath, const FWidgetPath& NewWidgetPath);
 
-	DEPRECATED(4.6, "SWidget::OnKeyboardFocusChanging() is deprecated, implement SWidget::OnFocusChanging() instead.")
-	virtual void OnKeyboardFocusChanging(const FWeakWidgetPath& PreviousFocusPath, const FWidgetPath& NewWidgetPath);
-
 	/**
 	 * Called after a character is entered while this widget has keyboard focus
 	 *
diff --git a/Engine/Source/Runtime/UMG/Private/UserWidget.cpp b/Engine/Source/Runtime/UMG/Private/UserWidget.cpp
--- a/Engine/Source/Runtime/UMG/Private/UserWidget.cpp
+++ b/Engine/Source/Runtime/UMG/Private/UserWidget.cpp
@@ -221,21 +221,11 @@ FEventReply UUserWidget::OnFocusReceived_Implementation(FGeometry MyGeometry, FF
 	return UWidgetBlueprintLibrary::Unhandled();
 }
 
-FEventReply UUserWidget::OnKeyboardFocusReceived_Implementation(FGeometry MyGeometry, FKeyboardFocusEvent InKeyboardFocusEvent)
-{
-	return UWidgetBlueprintLibrary::Unhandled();
-}
-
 void UUserWidget::OnFocusLost_Implementation(FFocusEvent InFocusEvent)
 {
 
 }
 
-void UUserWidget::OnKeyboardFocusLost_Implementation(FKeyboardFocusEvent InKeyboardFocusEvent)
-{
-
-}
-
 FEventReply UUserWidget::OnKeyChar_Implementation(FGeometry MyGeometry, FCharacterEvent InCharacterEvent)
 {
 	return UWidgetBlueprintLibrary::Unhandled();
diff --git a/Engine/Source/Runtime/UMG/Public/Blueprint/UserWidget.h b/Engine/Source/Runtime/UMG/Public/Blueprint/UserWidget.h
--- a/Engine/Source/Runtime/UMG/Public/Blueprint/UserWidget.h
+++ b/Engine/Source/Runtime/UMG/Public/Blueprint/UserWidget.h
@@ -263,10 +263,6 @@ public:
 	UFUNCTION(BlueprintNativeEvent, BlueprintCosmetic, Category="Input")
 	FEventReply OnFocusReceived(FGeometry MyGeometry, FFocusEvent InFocusEvent);
 
-	UFUNCTION(BlueprintNativeEvent, BlueprintCosmetic, meta = (DeprecatedFunction, DeprecationMessage = "Use OnFocusReceived() instead"), Category = "Keyboard")
-	FEventReply OnKeyboardFocusReceived(FGeometry MyGeometry, FKeyboardFocusEvent InKeyboardFocusEvent);
-
-
 	/**
 	 * Called when this widget loses focus.  This event does not bubble.
 	 *
@@ -275,10 +271,6 @@ public:
 	UFUNCTION(BlueprintNativeEvent, BlueprintCosmetic, Category="Input")
 	void OnFocusLost(FFocusEvent InFocusEvent);
 
-	UFUNCTION(BlueprintNativeEvent, BlueprintCosmetic, meta = (DeprecatedFunction, DeprecationMessage = "Use OnFocusLost() instead"), Category = "Keyboard")
-	void OnKeyboardFocusLost(FKeyboardFocusEvent InKeyboardFocusEvent);
-
-
 	/**
 	 * Called after a character is entered while this widget has focus
 	 *
