Another fix.

diff --git a/Engine/Source/Runtime/SlateCore/Private/Widgets/SWidget.cpp b/Engine/Source/Runtime/SlateCore/Private/Widgets/SWidget.cpp
--- a/Engine/Source/Runtime/SlateCore/Private/Widgets/SWidget.cpp
+++ b/Engine/Source/Runtime/SlateCore/Private/Widgets/SWidget.cpp
@@ -82,10 +82,6 @@ void SWidget::OnFocusChanging(const FWeakWidgetPath& PreviousFocusPath, const FW
 {
 }
 
-void SWidget::OnKeyboardFocusChanging(const FWeakWidgetPath& PreviousFocusPath, const FWidgetPath& NewWidgetPath)
-{
-}
-
 FReply SWidget::OnKeyChar( const FGeometry& MyGeometry, const FCharacterEvent& InCharacterEvent )
 {
 	return FReply::Unhandled();
