Fix memory error spotted by GuardMalloc - the font-cache code was assuming that MappedEntries references would be valid after modifying the TMap but this isn't true.

diff --git a/Engine/Source/Runtime/SlateCore/Private/Fonts/FontCache.cpp b/Engine/Source/Runtime/SlateCore/Private/Fonts/FontCache.cpp
--- a/Engine/Source/Runtime/SlateCore/Private/Fonts/FontCache.cpp
+++ b/Engine/Source/Runtime/SlateCore/Private/Fonts/FontCache.cpp
@@ -1018,24 +1018,28 @@ const FCharacterEntry& FCharacterList::GetCharacter( TCHAR Character )
 		}
 			
 		FCharacterEntry& CharacterEntry = DirectIndexEntries[ Character ];
-		if( !CharacterEntry.IsValidEntry() )
+		if( CharacterEntry.IsValidEntry() )
+		{
+			return CharacterEntry;
+		}
+		else
 		{
 			// Character has not been cached yet
-			CharacterEntry = CacheCharacter( Character );
+			return CacheCharacter( Character );
 		}
-
-		return CharacterEntry;
 	}
 	else
 	{
 		FCharacterEntry& CharacterEntry = MappedEntries.FindOrAdd( Character );
-		if( !CharacterEntry.IsValidEntry() )
+		if( CharacterEntry.IsValidEntry() )
+		{
+			return CharacterEntry;
+		}
+		else
 		{
 			// Character has not been cached yet
-			CharacterEntry = CacheCharacter( Character );
+			return CacheCharacter( Character );
 		}
-
-		return CharacterEntry;
 	}
 }
 
