Fixed Slate Hierarchical stats no longer assuming the first stat is the global root (which is is not when events are being processed).

diff --git a/Engine/Source/Runtime/SlateCore/Public/Stats/SlateStats.h b/Engine/Source/Runtime/SlateCore/Public/Stats/SlateStats.h
--- a/Engine/Source/Runtime/SlateCore/Public/Stats/SlateStats.h
+++ b/Engine/Source/Runtime/SlateCore/Public/Stats/SlateStats.h
@@ -303,7 +303,7 @@ public:
 			const int32 ParentDepth = ParentEntry.StackDepth;
 			// look at our current depth.
 			const int32 Depth = StatEntries[Index].StackDepth;
-			// we are immediately below our parent. Subtract off the entrie's inclusive time
+			// we are immediately below our parent. Subtract off the entry's inclusive time
 			// because it isn't part of the parent's exclusive time.
 			if (Depth == ParentDepth + 1)
 			{
@@ -311,11 +311,17 @@ public:
 				ParentEntry.ExclusiveTime -= InclusiveTime;
 				++Index;
 			}
-			// we have descended on level in the stack. Update the parent and go again.
+			// we have descended one level in the stack. Update the parent and go again.
 			else if (Depth == ParentDepth + 2)
 			{
 				Index = ComputeExclusiveTimes(Index, Index - 1);
 			}
+			// There may be many roots, so we just have to treat all level 0's as siblings.
+			else if (Depth == 0 && ParentDepth == 0)
+			{
+				ParentIndex = Index;
+				++Index;
+			}
 			// we are ascending out of the stack, so just return.
 			else
 			{
