Merging 2477112

//depot/UE4-UT/...

to //depot/UE4/...

diff --git a/Engine/Source/Runtime/SlateCore/Private/Sound/SlateSound.cpp b/Engine/Source/Runtime/SlateCore/Private/Sound/SlateSound.cpp
--- a/Engine/Source/Runtime/SlateCore/Private/Sound/SlateSound.cpp
+++ b/Engine/Source/Runtime/SlateCore/Private/Sound/SlateSound.cpp
@@ -16,16 +16,23 @@ UObject* FSlateSound::GetResourceObject( ) const
 
 		if (!LegacyResourceObject)
 		{
-			// We can't check the object type against USoundBase as we don't have access to it here
-			// The user is required to cast the result of FSlateSound::GetResourceObject so we should be fine
-			LegacyResourceObject = StaticFindObject(UObject::StaticClass(), nullptr, *LegacyResourceName_DEPRECATED.ToString());
-			if (!ResourceObject)
+			if (!IsInGameThread())
 			{
-				LegacyResourceObject = StaticLoadObject(UObject::StaticClass(), nullptr, *LegacyResourceName_DEPRECATED.ToString());
+				UE_LOG(LogSlate, Warning, TEXT("Can't find/load sound %s because Slate is being updated in another thread! (loading screen?)"), *LegacyResourceName_DEPRECATED.ToString());
+			}
+			else
+			{
+				// We can't check the object type against USoundBase as we don't have access to it here
+				// The user is required to cast the result of FSlateSound::GetResourceObject so we should be fine
+				LegacyResourceObject = StaticFindObject(UObject::StaticClass(), nullptr, *LegacyResourceName_DEPRECATED.ToString());
+				if (!ResourceObject)
+				{
+					LegacyResourceObject = StaticLoadObject(UObject::StaticClass(), nullptr, *LegacyResourceName_DEPRECATED.ToString());
+				}
+
+				// Cache this in the weak-ptr to try and avoid having to load it all the time
+				LegacyResourceObject_DEPRECATED = LegacyResourceObject;
 			}
-
-			// Cache this in the weak-ptr to try and avoid having to load it all the time
-			LegacyResourceObject_DEPRECATED = LegacyResourceObject;
 		}
 
 		return LegacyResourceObject;
diff --git a/Engine/Source/Runtime/SlateRHIRenderer/Private/SlateRHIResourceManager.cpp b/Engine/Source/Runtime/SlateRHIRenderer/Private/SlateRHIResourceManager.cpp
--- a/Engine/Source/Runtime/SlateRHIRenderer/Private/SlateRHIResourceManager.cpp
+++ b/Engine/Source/Runtime/SlateRHIRenderer/Private/SlateRHIResourceManager.cpp
@@ -348,7 +348,7 @@ static void LoadUObjectForBrush( const FSlateBrush& InBrush )
 	// Load the utexture
 	FString Path = InBrush.GetResourceName().ToString();
 
-	if( !Path.IsEmpty() )
+	if (!Path.IsEmpty() && Path.StartsWith(FSlateBrush::UTextureIdentifier()))
 	{
 		FString NewPath = Path.RightChop(FSlateBrush::UTextureIdentifier().Len());
 		UObject* TextureObject = LoadObject<UTexture2D>(NULL, *NewPath, NULL, LOAD_None, NULL);
