Slate: line filtering no longer relies on texture resource.

diff --git a/Engine/Shaders/SlateElementPixelShader.usf b/Engine/Shaders/SlateElementPixelShader.usf
--- a/Engine/Shaders/SlateElementPixelShader.usf
+++ b/Engine/Shaders/SlateElementPixelShader.usf
@@ -244,7 +244,7 @@ float4 GetLineSegmentElementColor( VertexOut InVertex )
 	// samples the first and we get white pixels 
 	const float HalfPixelOffset = 1/32.f;
 
-	Color.a *= Texture2DSample(ElementTexture, ElementTextureSampler, float2(Index-HalfPixelOffset,0)).x;
+	Color.a *= smoothstep(0.3, 1.0f, Index);
 	
 	// Simple alpha test to help with overlapping segments
 	if( Color.a < 0.05f )
diff --git a/Engine/Shaders/StandaloneRenderer/D3D/SlateElementPixelShader.hlsl b/Engine/Shaders/StandaloneRenderer/D3D/SlateElementPixelShader.hlsl
--- a/Engine/Shaders/StandaloneRenderer/D3D/SlateElementPixelShader.hlsl
+++ b/Engine/Shaders/StandaloneRenderer/D3D/SlateElementPixelShader.hlsl
@@ -140,7 +140,7 @@ float4 GetSplineElementColor( VertexOut InVertex )
 	// samples the first and we get white pixels 
 	const float HalfPixelOffset = 1/32.f;
 
-	Color.a *= ElementTexture.Sample(ElementTextureSampler, float2(Index-HalfPixelOffset,0)).x;
+	Color.a *= smoothstep(0.3, 1.0f, Index);
 
 	if( Color.a < 0.05f )
 	{
diff --git a/Engine/Shaders/StandaloneRenderer/OpenGL/SlateElementPixelShader.glsl b/Engine/Shaders/StandaloneRenderer/OpenGL/SlateElementPixelShader.glsl
--- a/Engine/Shaders/StandaloneRenderer/OpenGL/SlateElementPixelShader.glsl
+++ b/Engine/Shaders/StandaloneRenderer/OpenGL/SlateElementPixelShader.glsl
@@ -137,7 +137,7 @@ vec4 GetSplineElementColor()
 	// samples the first and we get white pixels
 	const float HalfPixelOffset = 1.0/32.0;
 	
-	InColor.a *= texture2D(ElementTexture, vec2(Index-HalfPixelOffset,0.0)).x;
+	InColor.a *= smoothstep(0.3, 1.0f, Index);
 	
 	if( InColor.a < 0.05 )
 	{
diff --git a/Engine/Source/Developer/StandaloneRenderer/Private/OpenGL/SlateOpenGLRenderingPolicy.cpp b/Engine/Source/Developer/StandaloneRenderer/Private/OpenGL/SlateOpenGLRenderingPolicy.cpp
--- a/Engine/Source/Developer/StandaloneRenderer/Private/OpenGL/SlateOpenGLRenderingPolicy.cpp
+++ b/Engine/Source/Developer/StandaloneRenderer/Private/OpenGL/SlateOpenGLRenderingPolicy.cpp
@@ -203,8 +203,6 @@ void FSlateOpenGLRenderingPolicy::DrawElements( const FMatrix& ViewProjectionMat
 		{
 			// Test that the pixels in the line segment have only been drawn once
 			glEnable(GL_STENCIL_TEST);
-			
-			ElementProgram.SetTexture( ((FSlateOpenGLTexture*)Texture)->GetTypedResource(), GL_CLAMP_TO_EDGE, GL_CLAMP_TO_EDGE );
 		}
 		else if( Texture )
 		{
diff --git a/Engine/Source/Developer/StandaloneRenderer/Private/Windows/D3D/SlateD3DRenderingPolicy.cpp b/Engine/Source/Developer/StandaloneRenderer/Private/Windows/D3D/SlateD3DRenderingPolicy.cpp
--- a/Engine/Source/Developer/StandaloneRenderer/Private/Windows/D3D/SlateD3DRenderingPolicy.cpp
+++ b/Engine/Source/Developer/StandaloneRenderer/Private/Windows/D3D/SlateD3DRenderingPolicy.cpp
@@ -292,15 +292,7 @@ void FSlateD3D11RenderingPolicy::DrawElements( const FMatrix& ViewProjectionMatr
 		// Disable stenciling and depth testing by default 
 		GD3DDeviceContext->OMSetDepthStencilState( DSStateOff, 0x00);
 
-		if( RenderBatch.ShaderType == ESlateShader::LineSegment )
-		{
-			// debugging
-			//GD3DDeviceContext->RSSetState( WireframeRasterState );
-			//PixelShader->SetShaderType( ESlateShader::Default );
-
-			PixelShader->SetTexture( ((FSlateD3DTexture*)Texture)->GetTypedResource(), BilinearSamplerState_Clamp );
-		}
-		else if( Texture )
+		if( Texture )
 		{
 			TRefCountPtr<ID3D11SamplerState> SamplerState;
 			if( DrawFlags & ESlateBatchDrawFlag::TileU || DrawFlags & ESlateBatchDrawFlag::TileV )
diff --git a/Engine/Source/Editor/EditorStyle/Private/SlateEditorStyle.cpp b/Engine/Source/Editor/EditorStyle/Private/SlateEditorStyle.cpp
--- a/Engine/Source/Editor/EditorStyle/Private/SlateEditorStyle.cpp
+++ b/Engine/Source/Editor/EditorStyle/Private/SlateEditorStyle.cpp
@@ -238,7 +238,6 @@ void FSlateEditorStyle::FStyle::SetupGeneralStyles()
 
 	// Rendering resources that never change
 	{
-		Set( "SplineFilterTable", new IMAGE_BRUSH( "SplineFilterTable", FVector2D(32,1), FLinearColor::White, ESlateBrushTileType::Both, ESlateBrushImageType::Linear ) );
 		Set( "None", new FSlateNoResource() );
 	}
 
diff --git a/Engine/Source/Runtime/SlateCore/Private/Rendering/ElementBatcher.cpp b/Engine/Source/Runtime/SlateCore/Private/Rendering/ElementBatcher.cpp
--- a/Engine/Source/Runtime/SlateCore/Private/Rendering/ElementBatcher.cpp
+++ b/Engine/Source/Runtime/SlateCore/Private/Rendering/ElementBatcher.cpp
@@ -37,13 +37,6 @@ FSlateRotatedClipRectType ToSnappedRotatedRect(const FSlateRect& ClipRectInLayou
 		RoundToInt(RotatedRect.ExtentY));
 }
 
-static const FSlateImageBrush& GetSplineFilterTable()
-{
-	static const FSlateImageBrush SplineFilterTable( FPaths::EngineContentDir() / TEXT("Slate/SplineFilterTable.png"), FVector2D(32,1) );
-	return SplineFilterTable;
-}
-
-
 /**
  * Computes the element tint color based in the user specified color and brush being used.
  * Note: The color could be in RGB or HSV
@@ -837,12 +830,8 @@ void FSlateElementBatcher::AddSplineElement( const FSlateDrawElement& DrawElemen
 	// The amount we increase each side of the line to generate enough pixels
 	const float HalfThickness = LineThickness * .5f + Radius;
 
-	// Currently splines are not atlased because they are tiled.  So we just assume the texture proxy holds the actual texture
-	FSlateShaderResourceProxy* ResourceProxy = ResourceManager.GetShaderResource( GetSplineFilterTable() );
-	check(ResourceProxy && ResourceProxy->Resource);
-
 	// Find a batch for the element
-	FSlateElementBatch& ElementBatch = FindBatchForElement( Layer, FShaderParams::MakePixelShaderParams( FVector4( InPayload.Thickness,Radius,0,0) ), ResourceProxy->Resource, ESlateDrawPrimitive::TriangleList, ESlateShader::LineSegment, InDrawEffects, ESlateBatchDrawFlag::None, DrawElement.GetScissorRect() );
+	FSlateElementBatch& ElementBatch = FindBatchForElement( Layer, FShaderParams::MakePixelShaderParams( FVector4( InPayload.Thickness,Radius,0,0) ), nullptr, ESlateDrawPrimitive::TriangleList, ESlateShader::LineSegment, InDrawEffects, ESlateBatchDrawFlag::None, DrawElement.GetScissorRect() );
 	TArray<FSlateVertex>& BatchVertices = BatchVertexArrays[ElementBatch.VertexArrayIndex];
 	TArray<SlateIndex>& BatchIndices = BatchIndexArrays[ElementBatch.IndexArrayIndex];
 
@@ -969,10 +958,6 @@ void FSlateElementBatcher::AddLineElement( const FSlateDrawElement& DrawElement
 
 	if( InPayload.bAntialias )
 	{
-		// Currently splines are not atlased because they are tiled.  So we just assume the texture proxy holds the actual texture
-		FSlateShaderResourceProxy* ResourceProxy = ResourceManager.GetShaderResource( GetSplineFilterTable() );
-		check(ResourceProxy && ResourceProxy->Resource);
-
 		// The radius to use when checking the distance of pixels to the actual line.  Arbitrary value based on what looks the best
 		const float Radius = 1.5f;
 
@@ -987,7 +972,7 @@ void FSlateElementBatcher::AddLineElement( const FSlateDrawElement& DrawElement
 		const float HalfThickness = LineThickness * .5f + Radius;
 
 		// Find a batch for the element
-		FSlateElementBatch& ElementBatch = FindBatchForElement( Layer, FShaderParams::MakePixelShaderParams( FVector4(RequestedThickness,Radius,0,0) ), ResourceProxy->Resource, ESlateDrawPrimitive::TriangleList, ESlateShader::LineSegment, DrawEffects, ESlateBatchDrawFlag::None, DrawElement.GetScissorRect() );
+		FSlateElementBatch& ElementBatch = FindBatchForElement( Layer, FShaderParams::MakePixelShaderParams( FVector4(RequestedThickness,Radius,0,0) ), nullptr, ESlateDrawPrimitive::TriangleList, ESlateShader::LineSegment, DrawEffects, ESlateBatchDrawFlag::None, DrawElement.GetScissorRect() );
 		TArray<FSlateVertex>& BatchVertices = BatchVertexArrays[ElementBatch.VertexArrayIndex];
 		TArray<SlateIndex>& BatchIndices = BatchIndexArrays[ElementBatch.IndexArrayIndex];
 
diff --git a/Engine/Source/Runtime/SlateCore/Private/Styling/CoreStyle.cpp b/Engine/Source/Runtime/SlateCore/Private/Styling/CoreStyle.cpp
--- a/Engine/Source/Runtime/SlateCore/Private/Styling/CoreStyle.cpp
+++ b/Engine/Source/Runtime/SlateCore/Private/Styling/CoreStyle.cpp
@@ -185,7 +185,6 @@ TSharedRef<ISlateStyle> FCoreStyle::Create( const FName& InStyleSetName )
 	FSlateBrush* GenericWhiteBox = new IMAGE_BRUSH("Old/White", Icon16x16);
 	{
 		Style->Set("Checkerboard", new IMAGE_BRUSH("Checkerboard", Icon16x16, FLinearColor::White, ESlateBrushTileType::Both));
-		Style->Set("SplineFilterTable", new IMAGE_BRUSH("SplineFilterTable", FVector2D(32.0f, 1.0f), FLinearColor::White, ESlateBrushTileType::Both, ESlateBrushImageType::Linear));
 
 		Style->Set("GenericWhiteBox", GenericWhiteBox);
 
diff --git a/Engine/Source/Runtime/SlateRHIRenderer/Private/SlateRHIRenderingPolicy.cpp b/Engine/Source/Runtime/SlateRHIRenderer/Private/SlateRHIRenderingPolicy.cpp
--- a/Engine/Source/Runtime/SlateRHIRenderer/Private/SlateRHIRenderingPolicy.cpp
+++ b/Engine/Source/Runtime/SlateRHIRenderer/Private/SlateRHIRenderingPolicy.cpp
@@ -374,12 +374,7 @@ void FSlateRHIRenderingPolicy::DrawElements(FRHICommandListImmediate& RHICmdList
 					TextureRHI = ((TSlateTexture<FTexture2DRHIRef>*)ShaderResource)->GetTypedResource();
 				}
 
-				if (ShaderType == ESlateShader::LineSegment)
-				{
-					// The lookup table used for splines should clamp when sampling otherwise the results are wrong
-					PixelShader->SetTexture(RHICmdList, TextureRHI, BilinearClamp );
-				}
-				else if( ShaderResource && IsValidRef( TextureRHI ) )
+				if( ShaderResource && IsValidRef( TextureRHI ) )
 				{
 					FSamplerStateRHIRef SamplerState; 
 
