Slate: Partially prevent SWidget::OnKeyDown from handling Tab (for focus traversal) when some modifier keys are held, allowing a limited amount of binding (Ctrl+Tab, Alt+Tab, Cmd+Tab; Shift+Tab is still always consumed)
- The full solution to handle the focus traversal only if the focused widget returns Unhandled is filed as UE-11841

diff --git a/Engine/Source/Runtime/SlateCore/Private/Widgets/SWidget.cpp b/Engine/Source/Runtime/SlateCore/Private/Widgets/SWidget.cpp
--- a/Engine/Source/Runtime/SlateCore/Private/Widgets/SWidget.cpp
+++ b/Engine/Source/Runtime/SlateCore/Private/Widgets/SWidget.cpp
@@ -118,10 +118,17 @@ FReply SWidget::OnKeyDown( const FGeometry& MyGeometry, const FKeyEvent& InKeyEv
 		// If the key was Tab, interpret as an attempt to move focus.
 		else if (InKeyEvent.GetKey() == EKeys::Tab)
 		{
-			EUINavigation MoveDirection = (InKeyEvent.IsShiftDown())
-				? EUINavigation::Previous
-				: EUINavigation::Next;
-			return FReply::Handled().SetNavigation(MoveDirection);
+			//@TODO: Really these uses of input should be at a lower priority, only occurring if nothing else handled them
+			// For now this code prevents consuming them when some modifiers are held down, allowing some limited binding
+			const bool bAllowEatingKeyEvents = !InKeyEvent.IsControlDown() && !InKeyEvent.IsAltDown() && !InKeyEvent.IsCommandDown();
+
+			if (bAllowEatingKeyEvents)
+			{
+				EUINavigation MoveDirection = (InKeyEvent.IsShiftDown())
+					? EUINavigation::Previous
+					: EUINavigation::Next;
+				return FReply::Handled().SetNavigation(MoveDirection);
+			}
 		}
 	}
 	return FReply::Unhandled();
