Back out changelist 2481333

diff --git a/Engine/Source/Runtime/Landscape/Classes/LandscapeHeightfieldCollisionComponent.h b/Engine/Source/Runtime/Landscape/Classes/LandscapeHeightfieldCollisionComponent.h
--- a/Engine/Source/Runtime/Landscape/Classes/LandscapeHeightfieldCollisionComponent.h
+++ b/Engine/Source/Runtime/Landscape/Classes/LandscapeHeightfieldCollisionComponent.h
@@ -20,9 +20,7 @@ namespace physx
 UCLASS(MinimalAPI)
 class ULandscapeHeightfieldCollisionComponent : public UPrimitiveComponent
 {
-	GENERATED_BODY()
-public:
-	LANDSCAPE_API ULandscapeHeightfieldCollisionComponent(const FObjectInitializer& ObjectInitializer = FObjectInitializer::Get());
+	GENERATED_UCLASS_BODY()
 
 	/** List of layers painted on this component. Matches the WeightmapLayerAllocations array in the LandscapeComponent. */
 	UPROPERTY()
diff --git a/Engine/Source/Runtime/Landscape/Classes/LandscapeInfo.h b/Engine/Source/Runtime/Landscape/Classes/LandscapeInfo.h
--- a/Engine/Source/Runtime/Landscape/Classes/LandscapeInfo.h
+++ b/Engine/Source/Runtime/Landscape/Classes/LandscapeInfo.h
@@ -16,7 +16,7 @@ struct FLandscapeDataInterface;
 USTRUCT()
 struct FLandscapeAddCollision
 {
-	GENERATED_BODY()
+	GENERATED_USTRUCT_BODY()
 
 #if WITH_EDITORONLY_DATA
 	UPROPERTY()
@@ -37,7 +37,7 @@ struct FLandscapeAddCollision
 USTRUCT()
 struct FLandscapeInfoLayerSettings
 {
-	GENERATED_BODY()
+	GENERATED_USTRUCT_BODY()
 
 	UPROPERTY()
 	ULandscapeLayerInfoObject* LayerInfoObj;
@@ -95,9 +95,7 @@ struct FLandscapeInfoLayerSettings
 UCLASS()
 class ULandscapeInfo : public UObject
 {
-	GENERATED_BODY()
-public:
-	ULandscapeInfo(const FObjectInitializer& ObjectInitializer = FObjectInitializer::Get());
+	GENERATED_UCLASS_BODY()
 
 	UPROPERTY()
 	TLazyObjectPtr<ALandscape> LandscapeActor;
diff --git a/Engine/Source/Runtime/Landscape/Classes/LandscapeLayerInfoObject.h b/Engine/Source/Runtime/Landscape/Classes/LandscapeLayerInfoObject.h
--- a/Engine/Source/Runtime/Landscape/Classes/LandscapeLayerInfoObject.h
+++ b/Engine/Source/Runtime/Landscape/Classes/LandscapeLayerInfoObject.h
@@ -7,9 +7,7 @@
 UCLASS(MinimalAPI)
 class ULandscapeLayerInfoObject : public UObject
 {
-	GENERATED_BODY()
-public:
-	LANDSCAPE_API ULandscapeLayerInfoObject(const FObjectInitializer& ObjectInitializer = FObjectInitializer::Get());
+	GENERATED_UCLASS_BODY()
 
 	UPROPERTY(VisibleAnywhere, Category=LandscapeLayerInfoObject, AssetRegistrySearchable)
 	FName LayerName;
diff --git a/Engine/Source/Runtime/Landscape/Classes/LandscapeMaterialInstanceConstant.h b/Engine/Source/Runtime/Landscape/Classes/LandscapeMaterialInstanceConstant.h
--- a/Engine/Source/Runtime/Landscape/Classes/LandscapeMaterialInstanceConstant.h
+++ b/Engine/Source/Runtime/Landscape/Classes/LandscapeMaterialInstanceConstant.h
@@ -6,9 +6,7 @@
 UCLASS(MinimalAPI)
 class ULandscapeMaterialInstanceConstant : public UMaterialInstanceConstant
 {
-	GENERATED_BODY()
-public:
-	LANDSCAPE_API ULandscapeMaterialInstanceConstant(const FObjectInitializer& ObjectInitializer = FObjectInitializer::Get());
+	GENERATED_UCLASS_BODY()
 
 	UPROPERTY()
 	uint32 bIsLayerThumbnail:1;
diff --git a/Engine/Source/Runtime/Landscape/Classes/LandscapeMeshCollisionComponent.h b/Engine/Source/Runtime/Landscape/Classes/LandscapeMeshCollisionComponent.h
--- a/Engine/Source/Runtime/Landscape/Classes/LandscapeMeshCollisionComponent.h
+++ b/Engine/Source/Runtime/Landscape/Classes/LandscapeMeshCollisionComponent.h
@@ -18,9 +18,7 @@ namespace physx
 UCLASS()
 class ULandscapeMeshCollisionComponent : public ULandscapeHeightfieldCollisionComponent
 {
-	GENERATED_BODY()
-public:
-	ULandscapeMeshCollisionComponent(const FObjectInitializer& ObjectInitializer = FObjectInitializer::Get());
+	GENERATED_UCLASS_BODY()
 
 	// Keep the possibility to share projected height field PhysX object with editor mesh collision objects...
 
diff --git a/Engine/Source/Runtime/Landscape/Classes/LandscapeMeshProxyActor.h b/Engine/Source/Runtime/Landscape/Classes/LandscapeMeshProxyActor.h
--- a/Engine/Source/Runtime/Landscape/Classes/LandscapeMeshProxyActor.h
+++ b/Engine/Source/Runtime/Landscape/Classes/LandscapeMeshProxyActor.h
@@ -11,9 +11,7 @@
 UCLASS(MinimalAPI)
 class ALandscapeMeshProxyActor : public AActor
 {
-	GENERATED_BODY()
-public:
-	LANDSCAPE_API ALandscapeMeshProxyActor(const FObjectInitializer& ObjectInitializer = FObjectInitializer::Get());
+	GENERATED_UCLASS_BODY()
 
 private:
 	UPROPERTY(Category = LandscapeMeshProxyActor, VisibleAnywhere, BlueprintReadOnly, meta = (ExposeFunctionCategories = "Mesh,Rendering,Physics,Components|StaticMesh", AllowPrivateAccess = "true"))
diff --git a/Engine/Source/Runtime/Landscape/Classes/LandscapeProxy.h b/Engine/Source/Runtime/Landscape/Classes/LandscapeProxy.h
--- a/Engine/Source/Runtime/Landscape/Classes/LandscapeProxy.h
+++ b/Engine/Source/Runtime/Landscape/Classes/LandscapeProxy.h
@@ -35,7 +35,7 @@ class USplineComponent;
 USTRUCT()
 struct FLandscapeWeightmapUsage
 {
-	GENERATED_BODY()
+	GENERATED_USTRUCT_BODY()
 
 	UPROPERTY()
 	ULandscapeComponent* ChannelUsage[4];
@@ -60,7 +60,7 @@ struct FLandscapeWeightmapUsage
 USTRUCT()
 struct FLandscapeEditorLayerSettings
 {
-	GENERATED_BODY()
+	GENERATED_USTRUCT_BODY()
 
 #if WITH_EDITORONLY_DATA
 	UPROPERTY()
@@ -91,7 +91,7 @@ struct FLandscapeEditorLayerSettings
 USTRUCT()
 struct FLandscapeLayerStruct
 {
-	GENERATED_BODY()
+	GENERATED_USTRUCT_BODY()
 
 	UPROPERTY()
 	ULandscapeLayerInfoObject* LayerInfoObj;
@@ -130,7 +130,7 @@ struct FLandscapeLayerStruct
 USTRUCT()
 struct FLandscapeImportLayerInfo
 {
-	GENERATED_BODY()
+	GENERATED_USTRUCT_BODY()
 
 #if WITH_EDITORONLY_DATA
 	UPROPERTY(Category="Import", VisibleAnywhere)
@@ -308,9 +308,7 @@ public:
 UCLASS(NotPlaceable, NotBlueprintable, hidecategories=(Display, Attachment, Physics, Debug, Lighting, LOD), showcategories=(Lighting, Rendering, "Utilities|Transformation"), MinimalAPI)
 class ALandscapeProxy : public AActor, public FTickableGameObject
 {
-	GENERATED_BODY()
-public:
-	LANDSCAPE_API ALandscapeProxy(const FObjectInitializer& ObjectInitializer = FObjectInitializer::Get());
+	GENERATED_UCLASS_BODY()
 
 	virtual ~ALandscapeProxy();
 
diff --git a/Engine/Source/Runtime/Landscape/Classes/LandscapeSplineControlPoint.h b/Engine/Source/Runtime/Landscape/Classes/LandscapeSplineControlPoint.h
--- a/Engine/Source/Runtime/Landscape/Classes/LandscapeSplineControlPoint.h
+++ b/Engine/Source/Runtime/Landscape/Classes/LandscapeSplineControlPoint.h
@@ -11,7 +11,7 @@ struct FLandscapeSplineInterpPoint;
 USTRUCT()
 struct FLandscapeSplineConnection
 {
-	GENERATED_BODY()
+	GENERATED_USTRUCT_BODY()
 
 	FLandscapeSplineConnection() {}
 
@@ -41,9 +41,7 @@ struct FLandscapeSplineConnection
 UCLASS(Within=LandscapeSplinesComponent,autoExpandCategories=LandscapeSplineControlPoint,MinimalAPI)
 class ULandscapeSplineControlPoint : public UObject
 {
-	GENERATED_BODY()
-public:
-	LANDSCAPE_API ULandscapeSplineControlPoint(const FObjectInitializer& ObjectInitializer = FObjectInitializer::Get());
+	GENERATED_UCLASS_BODY()
 
 // Directly editable data:
 	/** Location in Landscape-space */
diff --git a/Engine/Source/Runtime/Landscape/Classes/LandscapeSplineSegment.h b/Engine/Source/Runtime/Landscape/Classes/LandscapeSplineSegment.h
--- a/Engine/Source/Runtime/Landscape/Classes/LandscapeSplineSegment.h
+++ b/Engine/Source/Runtime/Landscape/Classes/LandscapeSplineSegment.h
@@ -7,7 +7,7 @@
 USTRUCT()
 struct FLandscapeSplineInterpPoint
 {
-	GENERATED_BODY()
+	GENERATED_USTRUCT_BODY()
 
 	/** Center Point */
 	UPROPERTY()
@@ -51,7 +51,7 @@ struct FLandscapeSplineInterpPoint
 USTRUCT()
 struct FLandscapeSplineSegmentConnection
 {
-	GENERATED_BODY()
+	GENERATED_USTRUCT_BODY()
 
 	// Control point connected to this end of the segment
 	UPROPERTY()
@@ -85,7 +85,7 @@ enum LandscapeSplineMeshOrientation
 USTRUCT()
 struct FLandscapeSplineMeshEntry
 {
-	GENERATED_BODY()
+	GENERATED_USTRUCT_BODY()
 
 	/** Mesh to use on the spline */
 	UPROPERTY(EditAnywhere, Category=LandscapeSplineMeshEntry)
@@ -143,9 +143,7 @@ struct FLandscapeSplineMeshEntry
 UCLASS(Within=LandscapeSplinesComponent,autoExpandCategories=(LandscapeSplineSegment,LandscapeSplineMeshes),MinimalAPI)
 class ULandscapeSplineSegment : public UObject
 {
-	GENERATED_BODY()
-public:
-	LANDSCAPE_API ULandscapeSplineSegment(const FObjectInitializer& ObjectInitializer = FObjectInitializer::Get());
+	GENERATED_UCLASS_BODY()
 
 // Directly editable data:
 	UPROPERTY(EditAnywhere, EditFixedSize, Category=LandscapeSplineSegment)
diff --git a/Engine/Source/Runtime/Landscape/Classes/LandscapeSplinesComponent.h b/Engine/Source/Runtime/Landscape/Classes/LandscapeSplinesComponent.h
--- a/Engine/Source/Runtime/Landscape/Classes/LandscapeSplinesComponent.h
+++ b/Engine/Source/Runtime/Landscape/Classes/LandscapeSplinesComponent.h
@@ -7,9 +7,7 @@
 UCLASS(MinimalAPI)
 class ULandscapeSplinesComponent : public UPrimitiveComponent
 {
-	GENERATED_BODY()
-public:
-	LANDSCAPE_API ULandscapeSplinesComponent(const FObjectInitializer& ObjectInitializer = FObjectInitializer::Get());
+	GENERATED_UCLASS_BODY()
 
 #if WITH_EDITORONLY_DATA
 	/** Resolution of the spline, in distance per point */
diff --git a/Engine/Source/Runtime/Landscape/Classes/Materials/MaterialExpressionLandscapeGrassOutput.h b/Engine/Source/Runtime/Landscape/Classes/Materials/MaterialExpressionLandscapeGrassOutput.h
--- a/Engine/Source/Runtime/Landscape/Classes/Materials/MaterialExpressionLandscapeGrassOutput.h
+++ b/Engine/Source/Runtime/Landscape/Classes/Materials/MaterialExpressionLandscapeGrassOutput.h
@@ -10,7 +10,7 @@
 USTRUCT()
 struct FGrassInput
 {
-	GENERATED_BODY()
+	GENERATED_USTRUCT_BODY()
 
 	UPROPERTY(EditAnywhere, Category = Grass)
 	FName Name;
@@ -34,9 +34,7 @@ struct FGrassInput
 UCLASS(collapsecategories, hidecategories=Object, MinimalAPI)
 class UMaterialExpressionLandscapeGrassOutput : public UMaterialExpressionCustomOutput
 {
-	GENERATED_BODY()
-public:
-	LANDSCAPE_API UMaterialExpressionLandscapeGrassOutput(const FObjectInitializer& ObjectInitializer = FObjectInitializer::Get());
+	GENERATED_UCLASS_BODY()
 	virtual int32 Compile(class FMaterialCompiler* Compiler, int32 OutputIndex, int32 MultiplexIndex) override;
 	virtual void GetCaption(TArray<FString>& OutCaptions) const override;
 	virtual const TArray<FExpressionInput*> GetInputs() override;
diff --git a/Engine/Source/Runtime/Landscape/Classes/Materials/MaterialExpressionLandscapeLayerCoords.h b/Engine/Source/Runtime/Landscape/Classes/Materials/MaterialExpressionLandscapeLayerCoords.h
--- a/Engine/Source/Runtime/Landscape/Classes/Materials/MaterialExpressionLandscapeLayerCoords.h
+++ b/Engine/Source/Runtime/Landscape/Classes/Materials/MaterialExpressionLandscapeLayerCoords.h
@@ -29,9 +29,7 @@ enum ELandscapeCustomizedCoordType
 UCLASS(collapsecategories, hidecategories=Object)
 class UMaterialExpressionLandscapeLayerCoords : public UMaterialExpression
 {
-	GENERATED_BODY()
-public:
-	UMaterialExpressionLandscapeLayerCoords(const FObjectInitializer& ObjectInitializer = FObjectInitializer::Get());
+	GENERATED_UCLASS_BODY()
 
 	/** Determines the mapping place to use on the terrain. */
 	UPROPERTY(EditAnywhere, Category=MaterialExpressionLandscapeLayerCoords)
diff --git a/Engine/Source/Runtime/Landscape/Classes/Materials/MaterialExpressionLandscapeLayerSample.h b/Engine/Source/Runtime/Landscape/Classes/Materials/MaterialExpressionLandscapeLayerSample.h
--- a/Engine/Source/Runtime/Landscape/Classes/Materials/MaterialExpressionLandscapeLayerSample.h
+++ b/Engine/Source/Runtime/Landscape/Classes/Materials/MaterialExpressionLandscapeLayerSample.h
@@ -9,9 +9,7 @@
 UCLASS(collapsecategories, hidecategories=Object)
 class UMaterialExpressionLandscapeLayerSample : public UMaterialExpression
 {
-	GENERATED_BODY()
-public:
-	UMaterialExpressionLandscapeLayerSample(const FObjectInitializer& ObjectInitializer = FObjectInitializer::Get());
+	GENERATED_UCLASS_BODY()
 
 	UPROPERTY(EditAnywhere, Category=MaterialExpressionLandscapeLayerWeight)
 	FName ParameterName;
diff --git a/Engine/Source/Runtime/Landscape/Classes/Materials/MaterialExpressionLandscapeLayerSwitch.h b/Engine/Source/Runtime/Landscape/Classes/Materials/MaterialExpressionLandscapeLayerSwitch.h
--- a/Engine/Source/Runtime/Landscape/Classes/Materials/MaterialExpressionLandscapeLayerSwitch.h
+++ b/Engine/Source/Runtime/Landscape/Classes/Materials/MaterialExpressionLandscapeLayerSwitch.h
@@ -8,9 +8,7 @@
 UCLASS(collapsecategories, hidecategories=Object)
 class UMaterialExpressionLandscapeLayerSwitch : public UMaterialExpression
 {
-	GENERATED_BODY()
-public:
-	UMaterialExpressionLandscapeLayerSwitch(const FObjectInitializer& ObjectInitializer = FObjectInitializer::Get());
+	GENERATED_UCLASS_BODY()
 
 	UPROPERTY()
 	FExpressionInput LayerUsed;
diff --git a/Engine/Source/Runtime/Landscape/Classes/Materials/MaterialExpressionLandscapeLayerWeight.h b/Engine/Source/Runtime/Landscape/Classes/Materials/MaterialExpressionLandscapeLayerWeight.h
--- a/Engine/Source/Runtime/Landscape/Classes/Materials/MaterialExpressionLandscapeLayerWeight.h
+++ b/Engine/Source/Runtime/Landscape/Classes/Materials/MaterialExpressionLandscapeLayerWeight.h
@@ -8,9 +8,7 @@
 UCLASS(collapsecategories, hidecategories=Object)
 class UMaterialExpressionLandscapeLayerWeight : public UMaterialExpression
 {
-	GENERATED_BODY()
-public:
-	UMaterialExpressionLandscapeLayerWeight(const FObjectInitializer& ObjectInitializer = FObjectInitializer::Get());
+	GENERATED_UCLASS_BODY()
 
 	UPROPERTY(meta = (RequiredInput = "false", ToolTip = "Defaults to 'ConstBase' if not specified"))
 	FExpressionInput Base;
diff --git a/Engine/Source/Runtime/MediaAssets/Public/MediaPlayer.h b/Engine/Source/Runtime/MediaAssets/Public/MediaPlayer.h
--- a/Engine/Source/Runtime/MediaAssets/Public/MediaPlayer.h
+++ b/Engine/Source/Runtime/MediaAssets/Public/MediaPlayer.h
@@ -41,9 +41,7 @@ UCLASS(BlueprintType, hidecategories=(Object))
 class MEDIAASSETS_API UMediaPlayer
 	: public UObject
 {
-	GENERATED_BODY()
-public:
-	UMediaPlayer(const FObjectInitializer& ObjectInitializer = FObjectInitializer::Get());
+	GENERATED_UCLASS_BODY()
 
 public:
 
diff --git a/Engine/Source/Runtime/MediaAssets/Public/MediaSoundWave.h b/Engine/Source/Runtime/MediaAssets/Public/MediaSoundWave.h
--- a/Engine/Source/Runtime/MediaAssets/Public/MediaSoundWave.h
+++ b/Engine/Source/Runtime/MediaAssets/Public/MediaSoundWave.h
@@ -20,9 +20,7 @@ UCLASS(hidecategories=(Compression, Sound, SoundWave, Subtitles))
 class MEDIAASSETS_API UMediaSoundWave
 	: public USoundWave
 {
-	GENERATED_BODY()
-public:
-	UMediaSoundWave(const FObjectInitializer& ObjectInitializer = FObjectInitializer::Get());
+	GENERATED_UCLASS_BODY()
 
 	/** The index of the MediaPlayer's audio track to get the wave data from. */
 	UPROPERTY(EditAnywhere, BlueprintReadWrite, Category=MediaPlayer)
diff --git a/Engine/Source/Runtime/MediaAssets/Public/MediaTexture.h b/Engine/Source/Runtime/MediaAssets/Public/MediaTexture.h
--- a/Engine/Source/Runtime/MediaAssets/Public/MediaTexture.h
+++ b/Engine/Source/Runtime/MediaAssets/Public/MediaTexture.h
@@ -21,9 +21,7 @@ UCLASS(hidecategories=(Compression, LevelOfDetail, Object, Texture))
 class MEDIAASSETS_API UMediaTexture
 	: public UTexture
 {
-	GENERATED_BODY()
-public:
-	UMediaTexture(const FObjectInitializer& ObjectInitializer = FObjectInitializer::Get());
+	GENERATED_UCLASS_BODY()
 
 	/** The addressing mode to use for the X axis. */
 	UPROPERTY(EditAnywhere, BlueprintReadWrite, Category=MediaTexture, AssetRegistrySearchable)
diff --git a/Engine/Source/Runtime/MoviePlayer/Public/MoviePlayerSettings.h b/Engine/Source/Runtime/MoviePlayer/Public/MoviePlayerSettings.h
--- a/Engine/Source/Runtime/MoviePlayer/Public/MoviePlayerSettings.h
+++ b/Engine/Source/Runtime/MoviePlayer/Public/MoviePlayerSettings.h
@@ -12,11 +12,9 @@ UCLASS(config=Game, defaultconfig)
 class MOVIEPLAYER_API UMoviePlayerSettings
 	: public UObject
 {
-	GENERATED_BODY()
-public:
-	UMoviePlayerSettings(const FObjectInitializer& ObjectInitializer = FObjectInitializer::Get());
 public:
 
+	GENERATED_UCLASS_BODY()
 
 	/** If enabled the game waits for startup movies to complete even if loading has finished */
 	UPROPERTY(globalconfig, EditAnywhere, Category="Movies")
diff --git a/Engine/Source/Runtime/MovieSceneCore/Classes/MovieSceneBindings.h b/Engine/Source/Runtime/MovieSceneCore/Classes/MovieSceneBindings.h
--- a/Engine/Source/Runtime/MovieSceneCore/Classes/MovieSceneBindings.h
+++ b/Engine/Source/Runtime/MovieSceneCore/Classes/MovieSceneBindings.h
@@ -11,7 +11,7 @@
 USTRUCT()
 struct FMovieSceneBoundObject
 {
-	GENERATED_BODY()
+	GENERATED_USTRUCT_BODY( FMovieSceneBoundObject )
 
 public:
 
@@ -87,9 +87,7 @@ private:
 UCLASS( MinimalAPI )
 class UMovieSceneBindings : public UObject
 {
-	GENERATED_BODY()
-public:
-	MOVIESCENECORE_API UMovieSceneBindings(const FObjectInitializer& ObjectInitializer = FObjectInitializer::Get());
+	GENERATED_UCLASS_BODY()
 
 public:
 
diff --git a/Engine/Source/Runtime/MovieSceneCore/Classes/MovieScenePropertyTrack.h b/Engine/Source/Runtime/MovieSceneCore/Classes/MovieScenePropertyTrack.h
--- a/Engine/Source/Runtime/MovieSceneCore/Classes/MovieScenePropertyTrack.h
+++ b/Engine/Source/Runtime/MovieSceneCore/Classes/MovieScenePropertyTrack.h
@@ -11,9 +11,7 @@
 UCLASS( abstract )
 class MOVIESCENECORE_API UMovieScenePropertyTrack : public UMovieSceneTrack
 {
-	GENERATED_BODY()
-public:
-	UMovieScenePropertyTrack(const FObjectInitializer& ObjectInitializer = FObjectInitializer::Get());
+	GENERATED_UCLASS_BODY()
 
 public:
 	/** UMovieSceneTrack interface */
diff --git a/Engine/Source/Runtime/MovieSceneCore/Classes/MovieSceneTrack.h b/Engine/Source/Runtime/MovieSceneCore/Classes/MovieSceneTrack.h
--- a/Engine/Source/Runtime/MovieSceneCore/Classes/MovieSceneTrack.h
+++ b/Engine/Source/Runtime/MovieSceneCore/Classes/MovieSceneTrack.h
@@ -12,9 +12,7 @@ class UMovieSceneSection;
 UCLASS( abstract, MinimalAPI )
 class UMovieSceneTrack : public UObject
 {
-	GENERATED_BODY()
-public:
-	MOVIESCENECORE_API UMovieSceneTrack(const FObjectInitializer& ObjectInitializer = FObjectInitializer::Get());
+	GENERATED_UCLASS_BODY()
 
 public:
 	/**
diff --git a/Engine/Source/Runtime/MovieSceneCore/Classes/RuntimeMovieScenePlayer.h b/Engine/Source/Runtime/MovieSceneCore/Classes/RuntimeMovieScenePlayer.h
--- a/Engine/Source/Runtime/MovieSceneCore/Classes/RuntimeMovieScenePlayer.h
+++ b/Engine/Source/Runtime/MovieSceneCore/Classes/RuntimeMovieScenePlayer.h
@@ -17,9 +17,7 @@
 UCLASS( MinimalAPI )
 class URuntimeMovieScenePlayer : public UObject, public IMovieScenePlayer, public IRuntimeMovieScenePlayerInterface
 {
-	GENERATED_BODY()
-public:
-	MOVIESCENECORE_API URuntimeMovieScenePlayer(const FObjectInitializer& ObjectInitializer = FObjectInitializer::Get());
+	GENERATED_UCLASS_BODY()
 
 public:
 
diff --git a/Engine/Source/Runtime/MovieSceneCoreTypes/Classes/MovieScene3DTransformSection.h b/Engine/Source/Runtime/MovieSceneCoreTypes/Classes/MovieScene3DTransformSection.h
--- a/Engine/Source/Runtime/MovieSceneCoreTypes/Classes/MovieScene3DTransformSection.h
+++ b/Engine/Source/Runtime/MovieSceneCoreTypes/Classes/MovieScene3DTransformSection.h
@@ -12,9 +12,7 @@
 UCLASS(MinimalAPI)
 class UMovieScene3DTransformSection : public UMovieSceneSection
 {
-	GENERATED_BODY()
-public:
-	MOVIESCENECORETYPES_API UMovieScene3DTransformSection(const FObjectInitializer& ObjectInitializer = FObjectInitializer::Get());
+	GENERATED_UCLASS_BODY()
 public:
 
 	/** MovieSceneSection interface */
@@ -113,4 +111,4 @@ private:
 	/** Scale curves */
 	UPROPERTY()
 	FRichCurve Scale[3];
-};
+};
\ No newline at end of file
diff --git a/Engine/Source/Runtime/MovieSceneCoreTypes/Classes/MovieScene3DTransformTrack.h b/Engine/Source/Runtime/MovieSceneCoreTypes/Classes/MovieScene3DTransformTrack.h
--- a/Engine/Source/Runtime/MovieSceneCoreTypes/Classes/MovieScene3DTransformTrack.h
+++ b/Engine/Source/Runtime/MovieSceneCoreTypes/Classes/MovieScene3DTransformTrack.h
@@ -129,9 +129,7 @@ private:
 UCLASS( MinimalAPI )
 class UMovieScene3DTransformTrack : public UMovieScenePropertyTrack
 {
-	GENERATED_BODY()
-public:
-	MOVIESCENECORETYPES_API UMovieScene3DTransformTrack(const FObjectInitializer& ObjectInitializer = FObjectInitializer::Get());
+	GENERATED_UCLASS_BODY()
 public:
 	/** UMovieSceneTrack interface */
 	virtual UMovieSceneSection* CreateNewSection() override;
diff --git a/Engine/Source/Runtime/MovieSceneCoreTypes/Classes/MovieSceneAnimationSection.h b/Engine/Source/Runtime/MovieSceneCoreTypes/Classes/MovieSceneAnimationSection.h
--- a/Engine/Source/Runtime/MovieSceneCoreTypes/Classes/MovieSceneAnimationSection.h
+++ b/Engine/Source/Runtime/MovieSceneCoreTypes/Classes/MovieSceneAnimationSection.h
@@ -11,9 +11,7 @@
 UCLASS( MinimalAPI )
 class UMovieSceneAnimationSection : public UMovieSceneSection
 {
-	GENERATED_BODY()
-public:
-	MOVIESCENECORETYPES_API UMovieSceneAnimationSection(const FObjectInitializer& ObjectInitializer = FObjectInitializer::Get());
+	GENERATED_UCLASS_BODY()
 public:
 	/** Sets the animation sequence for this section */
 	void SetAnimSequence(class UAnimSequence* InAnimSequence) {AnimSequence = InAnimSequence;}
diff --git a/Engine/Source/Runtime/MovieSceneCoreTypes/Classes/MovieSceneAnimationTrack.h b/Engine/Source/Runtime/MovieSceneCoreTypes/Classes/MovieSceneAnimationTrack.h
--- a/Engine/Source/Runtime/MovieSceneCoreTypes/Classes/MovieSceneAnimationTrack.h
+++ b/Engine/Source/Runtime/MovieSceneCoreTypes/Classes/MovieSceneAnimationTrack.h
@@ -13,9 +13,7 @@
 UCLASS( MinimalAPI )
 class UMovieSceneAnimationTrack : public UMovieSceneTrack
 {
-	GENERATED_BODY()
-public:
-	MOVIESCENECORETYPES_API UMovieSceneAnimationTrack(const FObjectInitializer& ObjectInitializer = FObjectInitializer::Get());
+	GENERATED_UCLASS_BODY()
 
 public:
 	/** UMovieSceneTrack interface */
diff --git a/Engine/Source/Runtime/MovieSceneCoreTypes/Classes/MovieSceneAudioSection.h b/Engine/Source/Runtime/MovieSceneCoreTypes/Classes/MovieSceneAudioSection.h
--- a/Engine/Source/Runtime/MovieSceneCoreTypes/Classes/MovieSceneAudioSection.h
+++ b/Engine/Source/Runtime/MovieSceneCoreTypes/Classes/MovieSceneAudioSection.h
@@ -11,9 +11,7 @@
 UCLASS( MinimalAPI )
 class UMovieSceneAudioSection : public UMovieSceneSection
 {
-	GENERATED_BODY()
-public:
-	MOVIESCENECORETYPES_API UMovieSceneAudioSection(const FObjectInitializer& ObjectInitializer = FObjectInitializer::Get());
+	GENERATED_UCLASS_BODY()
 public:
 	/** Sets this section's sound */
 	void SetSound(class USoundBase* InSound) {Sound = InSound;}
diff --git a/Engine/Source/Runtime/MovieSceneCoreTypes/Classes/MovieSceneAudioTrack.h b/Engine/Source/Runtime/MovieSceneCoreTypes/Classes/MovieSceneAudioTrack.h
--- a/Engine/Source/Runtime/MovieSceneCoreTypes/Classes/MovieSceneAudioTrack.h
+++ b/Engine/Source/Runtime/MovieSceneCoreTypes/Classes/MovieSceneAudioTrack.h
@@ -18,9 +18,7 @@ namespace AudioTrackConstants
 UCLASS( MinimalAPI )
 class UMovieSceneAudioTrack : public UMovieSceneTrack
 {
-	GENERATED_BODY()
-public:
-	MOVIESCENECORETYPES_API UMovieSceneAudioTrack(const FObjectInitializer& ObjectInitializer = FObjectInitializer::Get());
+	GENERATED_UCLASS_BODY()
 
 public:
 	/** UMovieSceneTrack interface */
diff --git a/Engine/Source/Runtime/MovieSceneCoreTypes/Classes/MovieSceneBoolSection.h b/Engine/Source/Runtime/MovieSceneCoreTypes/Classes/MovieSceneBoolSection.h
--- a/Engine/Source/Runtime/MovieSceneCoreTypes/Classes/MovieSceneBoolSection.h
+++ b/Engine/Source/Runtime/MovieSceneCoreTypes/Classes/MovieSceneBoolSection.h
@@ -10,9 +10,7 @@
 UCLASS( MinimalAPI )
 class UMovieSceneBoolSection : public UMovieSceneSection
 {
-	GENERATED_BODY()
-public:
-	MOVIESCENECORETYPES_API UMovieSceneBoolSection(const FObjectInitializer& ObjectInitializer = FObjectInitializer::Get());
+	GENERATED_UCLASS_BODY()
 public:
 	/**
 	 * Updates this section
diff --git a/Engine/Source/Runtime/MovieSceneCoreTypes/Classes/MovieSceneBoolTrack.h b/Engine/Source/Runtime/MovieSceneCoreTypes/Classes/MovieSceneBoolTrack.h
--- a/Engine/Source/Runtime/MovieSceneCoreTypes/Classes/MovieSceneBoolTrack.h
+++ b/Engine/Source/Runtime/MovieSceneCoreTypes/Classes/MovieSceneBoolTrack.h
@@ -13,9 +13,7 @@
 UCLASS( MinimalAPI )
 class UMovieSceneBoolTrack : public UMovieScenePropertyTrack
 {
-	GENERATED_BODY()
-public:
-	MOVIESCENECORETYPES_API UMovieSceneBoolTrack(const FObjectInitializer& ObjectInitializer = FObjectInitializer::Get());
+	GENERATED_UCLASS_BODY()
 
 public:
 	/** UMovieSceneTrack interface */
diff --git a/Engine/Source/Runtime/MovieSceneCoreTypes/Classes/MovieSceneColorTrack.h b/Engine/Source/Runtime/MovieSceneCoreTypes/Classes/MovieSceneColorTrack.h
--- a/Engine/Source/Runtime/MovieSceneCoreTypes/Classes/MovieSceneColorTrack.h
+++ b/Engine/Source/Runtime/MovieSceneCoreTypes/Classes/MovieSceneColorTrack.h
@@ -21,9 +21,7 @@ struct FColorKey
 UCLASS( MinimalAPI )
 class UMovieSceneColorTrack : public UMovieScenePropertyTrack
 {
-	GENERATED_BODY()
-public:
-	MOVIESCENECORETYPES_API UMovieSceneColorTrack(const FObjectInitializer& ObjectInitializer = FObjectInitializer::Get());
+	GENERATED_UCLASS_BODY()
 
 public:
 	/** UMovieSceneTrack interface */
diff --git a/Engine/Source/Runtime/MovieSceneCoreTypes/Classes/MovieSceneDirectorTrack.h b/Engine/Source/Runtime/MovieSceneCoreTypes/Classes/MovieSceneDirectorTrack.h
--- a/Engine/Source/Runtime/MovieSceneCoreTypes/Classes/MovieSceneDirectorTrack.h
+++ b/Engine/Source/Runtime/MovieSceneCoreTypes/Classes/MovieSceneDirectorTrack.h
@@ -12,9 +12,7 @@
 UCLASS( MinimalAPI )
 class UMovieSceneDirectorTrack : public UMovieSceneTrack
 {
-	GENERATED_BODY()
-public:
-	MOVIESCENECORETYPES_API UMovieSceneDirectorTrack(const FObjectInitializer& ObjectInitializer = FObjectInitializer::Get());
+	GENERATED_UCLASS_BODY()
 
 public:
 	/** UMovieSceneTrack interface */
diff --git a/Engine/Source/Runtime/MovieSceneCoreTypes/Classes/MovieSceneFloatSection.h b/Engine/Source/Runtime/MovieSceneCoreTypes/Classes/MovieSceneFloatSection.h
--- a/Engine/Source/Runtime/MovieSceneCoreTypes/Classes/MovieSceneFloatSection.h
+++ b/Engine/Source/Runtime/MovieSceneCoreTypes/Classes/MovieSceneFloatSection.h
@@ -10,9 +10,7 @@
 UCLASS( MinimalAPI )
 class UMovieSceneFloatSection : public UMovieSceneSection
 {
-	GENERATED_BODY()
-public:
-	MOVIESCENECORETYPES_API UMovieSceneFloatSection(const FObjectInitializer& ObjectInitializer = FObjectInitializer::Get());
+	GENERATED_UCLASS_BODY()
 public:
 	/**
 	 * Updates this section
@@ -52,4 +50,4 @@ private:
 	/** Curve data */
 	UPROPERTY()
 	FRichCurve FloatCurve;
-};
+};
\ No newline at end of file
diff --git a/Engine/Source/Runtime/MovieSceneCoreTypes/Classes/MovieSceneFloatTrack.h b/Engine/Source/Runtime/MovieSceneCoreTypes/Classes/MovieSceneFloatTrack.h
--- a/Engine/Source/Runtime/MovieSceneCoreTypes/Classes/MovieSceneFloatTrack.h
+++ b/Engine/Source/Runtime/MovieSceneCoreTypes/Classes/MovieSceneFloatTrack.h
@@ -13,9 +13,7 @@
 UCLASS( MinimalAPI )
 class UMovieSceneFloatTrack : public UMovieScenePropertyTrack
 {
-	GENERATED_BODY()
-public:
-	MOVIESCENECORETYPES_API UMovieSceneFloatTrack(const FObjectInitializer& ObjectInitializer = FObjectInitializer::Get());
+	GENERATED_UCLASS_BODY()
 
 public:
 	/** UMovieSceneTrack interface */
diff --git a/Engine/Source/Runtime/MovieSceneCoreTypes/Classes/MovieSceneParticleSection.h b/Engine/Source/Runtime/MovieSceneCoreTypes/Classes/MovieSceneParticleSection.h
--- a/Engine/Source/Runtime/MovieSceneCoreTypes/Classes/MovieSceneParticleSection.h
+++ b/Engine/Source/Runtime/MovieSceneCoreTypes/Classes/MovieSceneParticleSection.h
@@ -24,9 +24,7 @@ namespace EParticleKey
 UCLASS( MinimalAPI )
 class UMovieSceneParticleSection : public UMovieSceneSection
 {
-	GENERATED_BODY()
-public:
-	MOVIESCENECORETYPES_API UMovieSceneParticleSection(const FObjectInitializer& ObjectInitializer = FObjectInitializer::Get());
+	GENERATED_UCLASS_BODY()
 public:
 	/** Sets this section's key type*/
 	void SetKeyType(EParticleKey::Type InKeyType) {KeyType = InKeyType;}
diff --git a/Engine/Source/Runtime/MovieSceneCoreTypes/Classes/MovieSceneParticleTrack.h b/Engine/Source/Runtime/MovieSceneCoreTypes/Classes/MovieSceneParticleTrack.h
--- a/Engine/Source/Runtime/MovieSceneCoreTypes/Classes/MovieSceneParticleTrack.h
+++ b/Engine/Source/Runtime/MovieSceneCoreTypes/Classes/MovieSceneParticleTrack.h
@@ -13,9 +13,7 @@
 UCLASS( MinimalAPI )
 class UMovieSceneParticleTrack : public UMovieSceneTrack
 {
-	GENERATED_BODY()
-public:
-	MOVIESCENECORETYPES_API UMovieSceneParticleTrack(const FObjectInitializer& ObjectInitializer = FObjectInitializer::Get());
+	GENERATED_UCLASS_BODY()
 
 public:
 	/** UMovieSceneTrack interface */
diff --git a/Engine/Source/Runtime/MovieSceneCoreTypes/Classes/MovieSceneShotSection.h b/Engine/Source/Runtime/MovieSceneCoreTypes/Classes/MovieSceneShotSection.h
--- a/Engine/Source/Runtime/MovieSceneCoreTypes/Classes/MovieSceneShotSection.h
+++ b/Engine/Source/Runtime/MovieSceneCoreTypes/Classes/MovieSceneShotSection.h
@@ -11,9 +11,7 @@
 UCLASS( MinimalAPI )
 class UMovieSceneShotSection : public UMovieSceneSection
 {
-	GENERATED_BODY()
-public:
-	MOVIESCENECORETYPES_API UMovieSceneShotSection(const FObjectInitializer& ObjectInitializer = FObjectInitializer::Get());
+	GENERATED_UCLASS_BODY()
 public:
 	/** Sets the camera guid for this shot section */
 	void SetCameraGuid(const FGuid& InGuid);
diff --git a/Engine/Source/Runtime/MovieSceneCoreTypes/Classes/MovieSceneVectorSection.h b/Engine/Source/Runtime/MovieSceneCoreTypes/Classes/MovieSceneVectorSection.h
--- a/Engine/Source/Runtime/MovieSceneCoreTypes/Classes/MovieSceneVectorSection.h
+++ b/Engine/Source/Runtime/MovieSceneCoreTypes/Classes/MovieSceneVectorSection.h
@@ -12,9 +12,7 @@
 UCLASS(MinimalAPI )
 class UMovieSceneVectorSection : public UMovieSceneSection
 {
-	GENERATED_BODY()
-public:
-	MOVIESCENECORETYPES_API UMovieSceneVectorSection(const FObjectInitializer& ObjectInitializer = FObjectInitializer::Get());
+	GENERATED_UCLASS_BODY()
 public:
 	
 	/**
diff --git a/Engine/Source/Runtime/MovieSceneCoreTypes/Classes/MovieSceneVectorTrack.h b/Engine/Source/Runtime/MovieSceneCoreTypes/Classes/MovieSceneVectorTrack.h
--- a/Engine/Source/Runtime/MovieSceneCoreTypes/Classes/MovieSceneVectorTrack.h
+++ b/Engine/Source/Runtime/MovieSceneCoreTypes/Classes/MovieSceneVectorTrack.h
@@ -23,9 +23,7 @@ struct FVectorKey
 UCLASS(MinimalAPI)
 class UMovieSceneVectorTrack : public UMovieScenePropertyTrack
 {
-	GENERATED_BODY()
-public:
-	MOVIESCENECORETYPES_API UMovieSceneVectorTrack(const FObjectInitializer& ObjectInitializer = FObjectInitializer::Get());
+	GENERATED_UCLASS_BODY()
 public:
 	/** UMovieSceneTrack interface */
 	virtual UMovieSceneSection* CreateNewSection() override;
diff --git a/Engine/Source/Runtime/MovieSceneCoreTypes/Classes/SubMovieSceneSection.h b/Engine/Source/Runtime/MovieSceneCoreTypes/Classes/SubMovieSceneSection.h
--- a/Engine/Source/Runtime/MovieSceneCoreTypes/Classes/SubMovieSceneSection.h
+++ b/Engine/Source/Runtime/MovieSceneCoreTypes/Classes/SubMovieSceneSection.h
@@ -10,9 +10,7 @@
 UCLASS( MinimalAPI )
 class USubMovieSceneSection : public UMovieSceneSection
 {
-	GENERATED_BODY()
-public:
-	MOVIESCENECORETYPES_API USubMovieSceneSection(const FObjectInitializer& ObjectInitializer = FObjectInitializer::Get());
+	GENERATED_UCLASS_BODY()
 public:
 	/**
 	 * Sets the movie scene played by this section
diff --git a/Engine/Source/Runtime/MovieSceneCoreTypes/Classes/SubMovieSceneTrack.h b/Engine/Source/Runtime/MovieSceneCoreTypes/Classes/SubMovieSceneTrack.h
--- a/Engine/Source/Runtime/MovieSceneCoreTypes/Classes/SubMovieSceneTrack.h
+++ b/Engine/Source/Runtime/MovieSceneCoreTypes/Classes/SubMovieSceneTrack.h
@@ -11,9 +11,7 @@
 UCLASS( MinimalAPI )
 class USubMovieSceneTrack : public UMovieSceneTrack
 {
-	GENERATED_BODY()
-public:
-	MOVIESCENECORETYPES_API USubMovieSceneTrack(const FObjectInitializer& ObjectInitializer = FObjectInitializer::Get());
+	GENERATED_UCLASS_BODY()
 public:
 	/** UMovieSceneTrack interface */
 	virtual UMovieSceneSection* CreateNewSection() override;
diff --git a/Engine/Source/Runtime/Online/BuildPatchServices/Private/BuildPatchManifest.h b/Engine/Source/Runtime/Online/BuildPatchServices/Private/BuildPatchManifest.h
--- a/Engine/Source/Runtime/Online/BuildPatchServices/Private/BuildPatchManifest.h
+++ b/Engine/Source/Runtime/Online/BuildPatchServices/Private/BuildPatchManifest.h
@@ -102,7 +102,7 @@ namespace EManifestFileHeader
 USTRUCT()
 struct FCustomFieldData
 {
-	GENERATED_BODY()
+	GENERATED_USTRUCT_BODY()
 
 	UPROPERTY()
 	FString Key;
@@ -120,7 +120,7 @@ struct FCustomFieldData
 USTRUCT()
 struct FChunkInfoData
 {
-	GENERATED_BODY()
+	GENERATED_USTRUCT_BODY()
 
 	UPROPERTY()
 	FGuid Guid;
@@ -143,7 +143,7 @@ struct FChunkInfoData
 USTRUCT()
 struct FChunkPartData
 {
-	GENERATED_BODY()
+	GENERATED_USTRUCT_BODY()
 
 	// The GUID of the chunk containing this part
 	UPROPERTY()
@@ -166,7 +166,7 @@ struct FChunkPartData
 USTRUCT()
 struct FFileManifestData
 {
-	GENERATED_BODY()
+	GENERATED_USTRUCT_BODY()
 
 	UPROPERTY()
 	FString Filename;
@@ -207,9 +207,7 @@ private:
 UCLASS()
 class UBuildPatchManifest : public UObject
 {
-	GENERATED_BODY()
-public:
-	UBuildPatchManifest(const FObjectInitializer& ObjectInitializer = FObjectInitializer::Get());
+	GENERATED_UCLASS_BODY()
 
 	UPROPERTY()
 	uint8 ManifestFileVersion;
diff --git a/Engine/Source/Runtime/Online/OnlineSubsystem/Public/Interfaces/TurnBasedMatchInterface.h b/Engine/Source/Runtime/Online/OnlineSubsystem/Public/Interfaces/TurnBasedMatchInterface.h
--- a/Engine/Source/Runtime/Online/OnlineSubsystem/Public/Interfaces/TurnBasedMatchInterface.h
+++ b/Engine/Source/Runtime/Online/OnlineSubsystem/Public/Interfaces/TurnBasedMatchInterface.h
@@ -9,15 +9,12 @@ class IOnlineSubsystem;
 UINTERFACE(Blueprintable)
 class UTurnBasedMatchInterface : public UInterface
 {
-	GENERATED_BODY()
-public:
-	UTurnBasedMatchInterface(const FObjectInitializer& ObjectInitializer = FObjectInitializer::Get());
+	GENERATED_UINTERFACE_BODY()
 };
 
 class ITurnBasedMatchInterface
 {
-	GENERATED_BODY()
-public:
+	GENERATED_IINTERFACE_BODY()
 
 public:
 	UFUNCTION(BlueprintImplementableEvent, Category = "Online|TurnBased")
diff --git a/Engine/Source/Runtime/Online/OnlineSubsystem/Public/NamedInterfaces.h b/Engine/Source/Runtime/Online/OnlineSubsystem/Public/NamedInterfaces.h
--- a/Engine/Source/Runtime/Online/OnlineSubsystem/Public/NamedInterfaces.h
+++ b/Engine/Source/Runtime/Online/OnlineSubsystem/Public/NamedInterfaces.h
@@ -11,7 +11,7 @@
 USTRUCT()
 struct FNamedInterface
 {
-	GENERATED_BODY()
+	GENERATED_USTRUCT_BODY()
 
 	FNamedInterface() :
 		InterfaceName(NAME_None),
@@ -31,7 +31,7 @@ struct FNamedInterface
 USTRUCT()
 struct FNamedInterfaceDef
 {
-	GENERATED_BODY()
+	GENERATED_USTRUCT_BODY()
 
 	FNamedInterfaceDef() : 
 		InterfaceName(NAME_None)
@@ -52,9 +52,7 @@ struct FNamedInterfaceDef
 UCLASS(transient, config=Engine)
 class UNamedInterfaces: public UObject
 {
-	GENERATED_BODY()
-public:
-	UNamedInterfaces(const FObjectInitializer& ObjectInitializer = FObjectInitializer::Get());
+	GENERATED_UCLASS_BODY()
 
 	void Initialize();
 	class UObject* GetNamedInterface(FName InterfaceName) const;
diff --git a/Engine/Source/Runtime/Online/OnlineSubsystemSteam/Classes/SteamNetConnection.h b/Engine/Source/Runtime/Online/OnlineSubsystemSteam/Classes/SteamNetConnection.h
--- a/Engine/Source/Runtime/Online/OnlineSubsystemSteam/Classes/SteamNetConnection.h
+++ b/Engine/Source/Runtime/Online/OnlineSubsystemSteam/Classes/SteamNetConnection.h
@@ -10,9 +10,7 @@
 UCLASS(transient, config=Engine)
 class USteamNetConnection : public UIpConnection
 {
-    GENERATED_BODY()
-public:
-    USteamNetConnection(const FObjectInitializer& ObjectInitializer = FObjectInitializer::Get());
+    GENERATED_UCLASS_BODY()
 
 	// Begin UIpConnection Interface
 	virtual void InitRemoteConnection(class UNetDriver* InDriver, class FSocket* InSocket, const FURL& InURL, const class FInternetAddr& InRemoteAddr, EConnectionState InState, int32 InMaxPacket = 0, int32 InPacketOverhead = 0) override;
diff --git a/Engine/Source/Runtime/Online/OnlineSubsystemSteam/Classes/SteamNetDriver.h b/Engine/Source/Runtime/Online/OnlineSubsystemSteam/Classes/SteamNetDriver.h
--- a/Engine/Source/Runtime/Online/OnlineSubsystemSteam/Classes/SteamNetDriver.h
+++ b/Engine/Source/Runtime/Online/OnlineSubsystemSteam/Classes/SteamNetDriver.h
@@ -10,9 +10,7 @@
 UCLASS(transient, config=Engine)
 class USteamNetDriver : public UIpNetDriver
 {
-	GENERATED_BODY()
-public:
-	USteamNetDriver(const FObjectInitializer& ObjectInitializer = FObjectInitializer::Get());
+	GENERATED_UCLASS_BODY()
 
 	/** Time between connection detail output */
 	double ConnectionDumpInterval;
diff --git a/Engine/Source/Runtime/Online/OnlineSubsystemUtils/Classes/AchievementBlueprintLibrary.h b/Engine/Source/Runtime/Online/OnlineSubsystemUtils/Classes/AchievementBlueprintLibrary.h
--- a/Engine/Source/Runtime/Online/OnlineSubsystemUtils/Classes/AchievementBlueprintLibrary.h
+++ b/Engine/Source/Runtime/Online/OnlineSubsystemUtils/Classes/AchievementBlueprintLibrary.h
@@ -8,9 +8,7 @@
 UCLASS()
 class ONLINESUBSYSTEMUTILS_API UAchievementBlueprintLibrary : public UBlueprintFunctionLibrary
 {
-	GENERATED_BODY()
-public:
-	UAchievementBlueprintLibrary(const FObjectInitializer& ObjectInitializer = FObjectInitializer::Get());
+	GENERATED_UCLASS_BODY()
 
 	/**
 	 * Gets the status of an achievement ID (you must call CacheAchievements first to cache them)
diff --git a/Engine/Source/Runtime/Online/OnlineSubsystemUtils/Classes/AchievementQueryCallbackProxy.h b/Engine/Source/Runtime/Online/OnlineSubsystemUtils/Classes/AchievementQueryCallbackProxy.h
--- a/Engine/Source/Runtime/Online/OnlineSubsystemUtils/Classes/AchievementQueryCallbackProxy.h
+++ b/Engine/Source/Runtime/Online/OnlineSubsystemUtils/Classes/AchievementQueryCallbackProxy.h
@@ -7,9 +7,7 @@
 UCLASS(MinimalAPI)
 class UAchievementQueryCallbackProxy : public UOnlineBlueprintCallProxyBase
 {
-	GENERATED_BODY()
-public:
-	ONLINESUBSYSTEMUTILS_API UAchievementQueryCallbackProxy(const FObjectInitializer& ObjectInitializer = FObjectInitializer::Get());
+	GENERATED_UCLASS_BODY()
 
 	// Called when there is a successful query
 	UPROPERTY(BlueprintAssignable)
diff --git a/Engine/Source/Runtime/Online/OnlineSubsystemUtils/Classes/AchievementWriteCallbackProxy.h b/Engine/Source/Runtime/Online/OnlineSubsystemUtils/Classes/AchievementWriteCallbackProxy.h
--- a/Engine/Source/Runtime/Online/OnlineSubsystemUtils/Classes/AchievementWriteCallbackProxy.h
+++ b/Engine/Source/Runtime/Online/OnlineSubsystemUtils/Classes/AchievementWriteCallbackProxy.h
@@ -9,9 +9,7 @@ DECLARE_DYNAMIC_MULTICAST_DELEGATE_ThreeParams(FAchievementWriteDelegate, FName,
 UCLASS(MinimalAPI)
 class UAchievementWriteCallbackProxy : public UOnlineBlueprintCallProxyBase
 {
-	GENERATED_BODY()
-public:
-	ONLINESUBSYSTEMUTILS_API UAchievementWriteCallbackProxy(const FObjectInitializer& ObjectInitializer = FObjectInitializer::Get());
+	GENERATED_UCLASS_BODY()
 
 	// Called when there is a successful achievement write
 	UPROPERTY(BlueprintAssignable)
diff --git a/Engine/Source/Runtime/Online/OnlineSubsystemUtils/Classes/ConnectionCallbackProxy.h b/Engine/Source/Runtime/Online/OnlineSubsystemUtils/Classes/ConnectionCallbackProxy.h
--- a/Engine/Source/Runtime/Online/OnlineSubsystemUtils/Classes/ConnectionCallbackProxy.h
+++ b/Engine/Source/Runtime/Online/OnlineSubsystemUtils/Classes/ConnectionCallbackProxy.h
@@ -11,9 +11,7 @@ DECLARE_DYNAMIC_MULTICAST_DELEGATE_OneParam(FOnlineConnectionResult, int32, Erro
 UCLASS(MinimalAPI)
 class UConnectionCallbackProxy : public UOnlineBlueprintCallProxyBase
 {
-	GENERATED_BODY()
-public:
-	ONLINESUBSYSTEMUTILS_API UConnectionCallbackProxy(const FObjectInitializer& ObjectInitializer = FObjectInitializer::Get());
+	GENERATED_UCLASS_BODY()
 
 	// Called when there is a successful query
 	UPROPERTY(BlueprintAssignable)
@@ -47,4 +45,4 @@ private:
 
 	// The world context object in which this call is taking place
 	UObject* WorldContextObject;
-};
+};
\ No newline at end of file
diff --git a/Engine/Source/Runtime/Online/OnlineSubsystemUtils/Classes/CreateSessionCallbackProxy.h b/Engine/Source/Runtime/Online/OnlineSubsystemUtils/Classes/CreateSessionCallbackProxy.h
--- a/Engine/Source/Runtime/Online/OnlineSubsystemUtils/Classes/CreateSessionCallbackProxy.h
+++ b/Engine/Source/Runtime/Online/OnlineSubsystemUtils/Classes/CreateSessionCallbackProxy.h
@@ -7,9 +7,7 @@
 UCLASS(MinimalAPI)
 class UCreateSessionCallbackProxy : public UOnlineBlueprintCallProxyBase
 {
-	GENERATED_BODY()
-public:
-	ONLINESUBSYSTEMUTILS_API UCreateSessionCallbackProxy(const FObjectInitializer& ObjectInitializer = FObjectInitializer::Get());
+	GENERATED_UCLASS_BODY()
 
 	// Called when the session was created successfully
 	UPROPERTY(BlueprintAssignable)
diff --git a/Engine/Source/Runtime/Online/OnlineSubsystemUtils/Classes/DestroySessionCallbackProxy.h b/Engine/Source/Runtime/Online/OnlineSubsystemUtils/Classes/DestroySessionCallbackProxy.h
--- a/Engine/Source/Runtime/Online/OnlineSubsystemUtils/Classes/DestroySessionCallbackProxy.h
+++ b/Engine/Source/Runtime/Online/OnlineSubsystemUtils/Classes/DestroySessionCallbackProxy.h
@@ -7,9 +7,7 @@
 UCLASS(MinimalAPI)
 class UDestroySessionCallbackProxy : public UOnlineBlueprintCallProxyBase
 {
-	GENERATED_BODY()
-public:
-	ONLINESUBSYSTEMUTILS_API UDestroySessionCallbackProxy(const FObjectInitializer& ObjectInitializer = FObjectInitializer::Get());
+	GENERATED_UCLASS_BODY()
 
 	// Called when there is a successful destroy
 	UPROPERTY(BlueprintAssignable)
diff --git a/Engine/Source/Runtime/Online/OnlineSubsystemUtils/Classes/EndTurnCallbackProxy.h b/Engine/Source/Runtime/Online/OnlineSubsystemUtils/Classes/EndTurnCallbackProxy.h
--- a/Engine/Source/Runtime/Online/OnlineSubsystemUtils/Classes/EndTurnCallbackProxy.h
+++ b/Engine/Source/Runtime/Online/OnlineSubsystemUtils/Classes/EndTurnCallbackProxy.h
@@ -11,9 +11,7 @@ class AActor;
 UCLASS(MinimalAPI)
 class UEndTurnCallbackProxy : public UOnlineBlueprintCallProxyBase
 {
-	GENERATED_BODY()
-public:
-	ONLINESUBSYSTEMUTILS_API UEndTurnCallbackProxy(const FObjectInitializer& ObjectInitializer = FObjectInitializer::Get());
+	GENERATED_UCLASS_BODY()
 
 	virtual ~UEndTurnCallbackProxy();
 
@@ -47,4 +45,4 @@ private:
 
 	// Object that is inheriting from UTurnBasedMatchInterface. Replicated properties will be serialize into the platform specific match data object.
 	UTurnBasedMatchInterface* TurnBasedMatchInterface;
-};
+};
\ No newline at end of file
diff --git a/Engine/Source/Runtime/Online/OnlineSubsystemUtils/Classes/FindSessionsCallbackProxy.h b/Engine/Source/Runtime/Online/OnlineSubsystemUtils/Classes/FindSessionsCallbackProxy.h
--- a/Engine/Source/Runtime/Online/OnlineSubsystemUtils/Classes/FindSessionsCallbackProxy.h
+++ b/Engine/Source/Runtime/Online/OnlineSubsystemUtils/Classes/FindSessionsCallbackProxy.h
@@ -7,7 +7,7 @@
 USTRUCT(BlueprintType)
 struct FBlueprintSessionResult
 {
-	GENERATED_BODY()
+	GENERATED_USTRUCT_BODY()
 
 	FOnlineSessionSearchResult OnlineResult;
 };
@@ -17,9 +17,7 @@ DECLARE_DYNAMIC_MULTICAST_DELEGATE_OneParam(FBlueprintFindSessionsResultDelegate
 UCLASS(MinimalAPI)
 class UFindSessionsCallbackProxy : public UOnlineBlueprintCallProxyBase
 {
-	GENERATED_BODY()
-public:
-	ONLINESUBSYSTEMUTILS_API UFindSessionsCallbackProxy(const FObjectInitializer& ObjectInitializer = FObjectInitializer::Get());
+	GENERATED_UCLASS_BODY()
 
 	// Called when there is a successful query
 	UPROPERTY(BlueprintAssignable)
diff --git a/Engine/Source/Runtime/Online/OnlineSubsystemUtils/Classes/InAppPurchaseCallbackProxy.h b/Engine/Source/Runtime/Online/OnlineSubsystemUtils/Classes/InAppPurchaseCallbackProxy.h
--- a/Engine/Source/Runtime/Online/OnlineSubsystemUtils/Classes/InAppPurchaseCallbackProxy.h
+++ b/Engine/Source/Runtime/Online/OnlineSubsystemUtils/Classes/InAppPurchaseCallbackProxy.h
@@ -9,9 +9,7 @@ DECLARE_DYNAMIC_MULTICAST_DELEGATE_TwoParams(FInAppPurchaseResult, EInAppPurchas
 UCLASS(MinimalAPI)
 class UInAppPurchaseCallbackProxy : public UObject
 {
-	GENERATED_BODY()
-public:
-	ONLINESUBSYSTEMUTILS_API UInAppPurchaseCallbackProxy(const FObjectInitializer& ObjectInitializer = FObjectInitializer::Get());
+	GENERATED_UCLASS_BODY()
 
 	// Called when there is a successful In-App Purchase transaction
 	UPROPERTY(BlueprintAssignable)
diff --git a/Engine/Source/Runtime/Online/OnlineSubsystemUtils/Classes/InAppPurchaseQueryCallbackProxy.h b/Engine/Source/Runtime/Online/OnlineSubsystemUtils/Classes/InAppPurchaseQueryCallbackProxy.h
--- a/Engine/Source/Runtime/Online/OnlineSubsystemUtils/Classes/InAppPurchaseQueryCallbackProxy.h
+++ b/Engine/Source/Runtime/Online/OnlineSubsystemUtils/Classes/InAppPurchaseQueryCallbackProxy.h
@@ -9,9 +9,7 @@ DECLARE_DYNAMIC_MULTICAST_DELEGATE_OneParam(FInAppPurchaseQueryResult, const TAr
 UCLASS(MinimalAPI)
 class UInAppPurchaseQueryCallbackProxy : public UObject
 {
-	GENERATED_BODY()
-public:
-	ONLINESUBSYSTEMUTILS_API UInAppPurchaseQueryCallbackProxy(const FObjectInitializer& ObjectInitializer = FObjectInitializer::Get());
+	GENERATED_UCLASS_BODY()
 
 	// Called when there is a successful InAppPurchase query
 	UPROPERTY(BlueprintAssignable)
diff --git a/Engine/Source/Runtime/Online/OnlineSubsystemUtils/Classes/IpConnection.h b/Engine/Source/Runtime/Online/OnlineSubsystemUtils/Classes/IpConnection.h
--- a/Engine/Source/Runtime/Online/OnlineSubsystemUtils/Classes/IpConnection.h
+++ b/Engine/Source/Runtime/Online/OnlineSubsystemUtils/Classes/IpConnection.h
@@ -10,9 +10,7 @@
 UCLASS(transient, config=Engine)
 class ONLINESUBSYSTEMUTILS_API UIpConnection : public UNetConnection
 {
-    GENERATED_BODY()
-public:
-    UIpConnection(const FObjectInitializer& ObjectInitializer = FObjectInitializer::Get());
+    GENERATED_UCLASS_BODY()
 	// Variables.
 	TSharedPtr<FInternetAddr>	RemoteAddr;
 	class FSocket*				Socket;
diff --git a/Engine/Source/Runtime/Online/OnlineSubsystemUtils/Classes/IpNetDriver.h b/Engine/Source/Runtime/Online/OnlineSubsystemUtils/Classes/IpNetDriver.h
--- a/Engine/Source/Runtime/Online/OnlineSubsystemUtils/Classes/IpNetDriver.h
+++ b/Engine/Source/Runtime/Online/OnlineSubsystemUtils/Classes/IpNetDriver.h
@@ -10,9 +10,7 @@
 UCLASS(transient, config=Engine)
 class ONLINESUBSYSTEMUTILS_API UIpNetDriver : public UNetDriver
 {
-    GENERATED_BODY()
-public:
-    UIpNetDriver(const FObjectInitializer& ObjectInitializer = FObjectInitializer::Get());
+    GENERATED_UCLASS_BODY()
 
 	/** Should port unreachable messages be logged */
     UPROPERTY(Config)
diff --git a/Engine/Source/Runtime/Online/OnlineSubsystemUtils/Classes/JoinSessionCallbackProxy.h b/Engine/Source/Runtime/Online/OnlineSubsystemUtils/Classes/JoinSessionCallbackProxy.h
--- a/Engine/Source/Runtime/Online/OnlineSubsystemUtils/Classes/JoinSessionCallbackProxy.h
+++ b/Engine/Source/Runtime/Online/OnlineSubsystemUtils/Classes/JoinSessionCallbackProxy.h
@@ -7,9 +7,7 @@
 UCLASS(MinimalAPI)
 class UJoinSessionCallbackProxy : public UOnlineBlueprintCallProxyBase
 {
-	GENERATED_BODY()
-public:
-	ONLINESUBSYSTEMUTILS_API UJoinSessionCallbackProxy(const FObjectInitializer& ObjectInitializer = FObjectInitializer::Get());
+	GENERATED_UCLASS_BODY()
 
 	// Called when there is a successful join
 	UPROPERTY(BlueprintAssignable)
diff --git a/Engine/Source/Runtime/Online/OnlineSubsystemUtils/Classes/LeaderboardBlueprintLibrary.h b/Engine/Source/Runtime/Online/OnlineSubsystemUtils/Classes/LeaderboardBlueprintLibrary.h
--- a/Engine/Source/Runtime/Online/OnlineSubsystemUtils/Classes/LeaderboardBlueprintLibrary.h
+++ b/Engine/Source/Runtime/Online/OnlineSubsystemUtils/Classes/LeaderboardBlueprintLibrary.h
@@ -10,9 +10,7 @@
 UCLASS()
 class ONLINESUBSYSTEMUTILS_API ULeaderboardBlueprintLibrary : public UBlueprintFunctionLibrary
 {
-	GENERATED_BODY()
-public:
-	ULeaderboardBlueprintLibrary(const FObjectInitializer& ObjectInitializer = FObjectInitializer::Get());
+	GENERATED_UCLASS_BODY()
 
 	/** Writes an integer value to the specified leaderboard */
 	UFUNCTION(BlueprintCallable, Category = "Online|Leaderboard")
diff --git a/Engine/Source/Runtime/Online/OnlineSubsystemUtils/Classes/LeaderboardFlushCallbackProxy.h b/Engine/Source/Runtime/Online/OnlineSubsystemUtils/Classes/LeaderboardFlushCallbackProxy.h
--- a/Engine/Source/Runtime/Online/OnlineSubsystemUtils/Classes/LeaderboardFlushCallbackProxy.h
+++ b/Engine/Source/Runtime/Online/OnlineSubsystemUtils/Classes/LeaderboardFlushCallbackProxy.h
@@ -9,9 +9,7 @@ DECLARE_DYNAMIC_MULTICAST_DELEGATE_OneParam(FOnLeaderboardFlushed, FName, Sessio
 UCLASS(MinimalAPI)
 class ULeaderboardFlushCallbackProxy : public UObject
 {
-	GENERATED_BODY()
-public:
-	ONLINESUBSYSTEMUTILS_API ULeaderboardFlushCallbackProxy(const FObjectInitializer& ObjectInitializer = FObjectInitializer::Get());
+	GENERATED_UCLASS_BODY()
 
 	// Called when there is a successful leaderboard flush
 	UPROPERTY(BlueprintAssignable)
diff --git a/Engine/Source/Runtime/Online/OnlineSubsystemUtils/Classes/LeaderboardQueryCallbackProxy.h b/Engine/Source/Runtime/Online/OnlineSubsystemUtils/Classes/LeaderboardQueryCallbackProxy.h
--- a/Engine/Source/Runtime/Online/OnlineSubsystemUtils/Classes/LeaderboardQueryCallbackProxy.h
+++ b/Engine/Source/Runtime/Online/OnlineSubsystemUtils/Classes/LeaderboardQueryCallbackProxy.h
@@ -9,9 +9,7 @@ DECLARE_DYNAMIC_MULTICAST_DELEGATE_OneParam(FLeaderboardQueryResult, int32, Lead
 UCLASS(MinimalAPI)
 class ULeaderboardQueryCallbackProxy : public UObject
 {
-	GENERATED_BODY()
-public:
-	ONLINESUBSYSTEMUTILS_API ULeaderboardQueryCallbackProxy(const FObjectInitializer& ObjectInitializer = FObjectInitializer::Get());
+	GENERATED_UCLASS_BODY()
 
 	// Called when there is a successful leaderboard query
 	UPROPERTY(BlueprintAssignable)
diff --git a/Engine/Source/Runtime/Online/OnlineSubsystemUtils/Classes/OnlineBeacon.h b/Engine/Source/Runtime/Online/OnlineSubsystemUtils/Classes/OnlineBeacon.h
--- a/Engine/Source/Runtime/Online/OnlineSubsystemUtils/Classes/OnlineBeacon.h
+++ b/Engine/Source/Runtime/Online/OnlineSubsystemUtils/Classes/OnlineBeacon.h
@@ -24,9 +24,7 @@ namespace EBeaconState
 UCLASS(transient, config=Engine, notplaceable)
 class ONLINESUBSYSTEMUTILS_API AOnlineBeacon : public AActor, public FNetworkNotify
 {
-	GENERATED_BODY()
-public:
-	AOnlineBeacon(const FObjectInitializer& ObjectInitializer = FObjectInitializer::Get());
+	GENERATED_UCLASS_BODY()
 
 	// Begin AActor Interface
 	virtual void OnActorChannelOpen(class FInBunch& InBunch, class UNetConnection* Connection) override;
diff --git a/Engine/Source/Runtime/Online/OnlineSubsystemUtils/Classes/OnlineBeaconClient.h b/Engine/Source/Runtime/Online/OnlineSubsystemUtils/Classes/OnlineBeaconClient.h
--- a/Engine/Source/Runtime/Online/OnlineSubsystemUtils/Classes/OnlineBeaconClient.h
+++ b/Engine/Source/Runtime/Online/OnlineSubsystemUtils/Classes/OnlineBeaconClient.h
@@ -19,9 +19,7 @@ DECLARE_DELEGATE(FOnHostConnectionFailure);
 UCLASS(transient, notplaceable, config=Engine)
 class ONLINESUBSYSTEMUTILS_API AOnlineBeaconClient : public AOnlineBeacon
 {
-	GENERATED_BODY()
-public:
-	AOnlineBeaconClient(const FObjectInitializer& ObjectInitializer = FObjectInitializer::Get());
+	GENERATED_UCLASS_BODY()
 
 	// Begin AActor Interface
 	void OnNetCleanup(class UNetConnection* Connection);
@@ -89,10 +87,9 @@ private:
 	/**
 	 * Called on the server side to open up the actor channel that will allow RPCs to occur
 	 */
-	UFUNCTION(client="ClientOnConnected_Implementation", reliable)
+	UFUNCTION(client, reliable)
 	virtual void ClientOnConnected();
-	virtual void ClientOnConnected_Implementation();
 
 	friend class AOnlineBeaconHost;
 	friend class AOnlineBeaconHostObject;
-};
+};
\ No newline at end of file
diff --git a/Engine/Source/Runtime/Online/OnlineSubsystemUtils/Classes/OnlineBeaconHost.h b/Engine/Source/Runtime/Online/OnlineSubsystemUtils/Classes/OnlineBeaconHost.h
--- a/Engine/Source/Runtime/Online/OnlineSubsystemUtils/Classes/OnlineBeaconHost.h
+++ b/Engine/Source/Runtime/Online/OnlineSubsystemUtils/Classes/OnlineBeaconHost.h
@@ -19,9 +19,7 @@ class AOnlineBeaconClient;
 UCLASS(transient, notplaceable, config=Engine)
 class ONLINESUBSYSTEMUTILS_API AOnlineBeaconHost : public AOnlineBeacon
 {
-	GENERATED_BODY()
-public:
-	AOnlineBeaconHost(const FObjectInitializer& ObjectInitializer = FObjectInitializer::Get());
+	GENERATED_UCLASS_BODY()
 
 	/** Configured listen port for this beacon host */
 	UPROPERTY(Config)
@@ -118,4 +116,4 @@ private:
 
 	/** Mapping of beacon types to the OnBeaconConnected delegates */
 	TMap<FString, FOnBeaconConnected> OnBeaconConnectedMapping;
-};
+};
\ No newline at end of file
diff --git a/Engine/Source/Runtime/Online/OnlineSubsystemUtils/Classes/OnlineBeaconHostObject.h b/Engine/Source/Runtime/Online/OnlineSubsystemUtils/Classes/OnlineBeaconHostObject.h
--- a/Engine/Source/Runtime/Online/OnlineSubsystemUtils/Classes/OnlineBeaconHostObject.h
+++ b/Engine/Source/Runtime/Online/OnlineSubsystemUtils/Classes/OnlineBeaconHostObject.h
@@ -19,9 +19,7 @@ class UNetConnection;
 UCLASS(transient, config=Engine, notplaceable)
 class ONLINESUBSYSTEMUTILS_API AOnlineBeaconHostObject : public AActor
 {
-	GENERATED_BODY()
-public:
-	AOnlineBeaconHostObject(const FObjectInitializer& ObjectInitializer = FObjectInitializer::Get());
+	GENERATED_UCLASS_BODY()
 
 	/** Custom name for this beacon */
 	UPROPERTY(Transient)
@@ -54,4 +52,4 @@ public:
 	 * @param ClientActor client beacon actor to remove
 	 */
 	virtual void RemoveClientActor(AOnlineBeaconClient* ClientActor);
-};
+};
\ No newline at end of file
diff --git a/Engine/Source/Runtime/Online/OnlineSubsystemUtils/Classes/OnlineBlueprintCallProxyBase.h b/Engine/Source/Runtime/Online/OnlineSubsystemUtils/Classes/OnlineBlueprintCallProxyBase.h
--- a/Engine/Source/Runtime/Online/OnlineSubsystemUtils/Classes/OnlineBlueprintCallProxyBase.h
+++ b/Engine/Source/Runtime/Online/OnlineSubsystemUtils/Classes/OnlineBlueprintCallProxyBase.h
@@ -8,9 +8,7 @@ DECLARE_DYNAMIC_MULTICAST_DELEGATE(FEmptyOnlineDelegate);
 UCLASS(MinimalAPI)
 class UOnlineBlueprintCallProxyBase : public UObject
 {
-	GENERATED_BODY()
-public:
-	ONLINESUBSYSTEMUTILS_API UOnlineBlueprintCallProxyBase(const FObjectInitializer& ObjectInitializer = FObjectInitializer::Get());
+	GENERATED_UCLASS_BODY()
 
 	// Called to trigger the actual online action once the delegates have been bound
 	UFUNCTION(BlueprintCallable, meta=(BlueprintInternalUseOnly="true"), Category="Online")
diff --git a/Engine/Source/Runtime/Online/OnlineSubsystemUtils/Classes/OnlineSessionClient.h b/Engine/Source/Runtime/Online/OnlineSubsystemUtils/Classes/OnlineSessionClient.h
--- a/Engine/Source/Runtime/Online/OnlineSubsystemUtils/Classes/OnlineSessionClient.h
+++ b/Engine/Source/Runtime/Online/OnlineSubsystemUtils/Classes/OnlineSessionClient.h
@@ -16,9 +16,7 @@ class APlayerController;
 UCLASS(config=Game)
 class ONLINESUBSYSTEMUTILS_API UOnlineSessionClient : public UOnlineSession
 {
-	GENERATED_BODY()
-public:
-	UOnlineSessionClient(const FObjectInitializer& ObjectInitializer = FObjectInitializer::Get());
+	GENERATED_UCLASS_BODY()
 
 protected:
 
diff --git a/Engine/Source/Runtime/Online/OnlineSubsystemUtils/Classes/PartyBeaconClient.h b/Engine/Source/Runtime/Online/OnlineSubsystemUtils/Classes/PartyBeaconClient.h
--- a/Engine/Source/Runtime/Online/OnlineSubsystemUtils/Classes/PartyBeaconClient.h
+++ b/Engine/Source/Runtime/Online/OnlineSubsystemUtils/Classes/PartyBeaconClient.h
@@ -22,9 +22,7 @@ DECLARE_DELEGATE_OneParam(FOnReservationRequestComplete, EPartyReservationResult
 UCLASS(transient, notplaceable, config=Engine)
 class ONLINESUBSYSTEMUTILS_API APartyBeaconClient : public AOnlineBeaconClient
 {
-	GENERATED_BODY()
-public:
-	APartyBeaconClient(const FObjectInitializer& ObjectInitializer = FObjectInitializer::Get());
+	GENERATED_UCLASS_BODY()
 
 	// Begin AOnlineBeacon Interface
 	virtual FString GetBeaconType() override { return PARTY_BEACON_TYPE; }
@@ -57,9 +55,8 @@ public:
 	 *
 	 * @param ReservationResponse response from server
 	 */
- 	UFUNCTION(client="ClientReservationResponse_Implementation", reliable)
+ 	UFUNCTION(client, reliable)
 	virtual void ClientReservationResponse(EPartyReservationResult::Type ReservationResponse);
- 	virtual void ClientReservationResponse_Implementation(EPartyReservationResult::Type ReservationResponse);
 
 	/**
 	 * Delegate triggered when a response from the party beacon host has been received
@@ -93,13 +90,9 @@ protected:
 	 *
 	 * @param Reservation pending reservation request to make with server
 	 */
-	UFUNCTION(server="ServerReservationRequest_Implementation", reliable, WithValidation="ServerReservationRequest_Validate")
+	UFUNCTION(server, reliable, WithValidation)
 	virtual void ServerReservationRequest(const FString& SessionId, struct FPartyReservation Reservation);
-	virtual void ServerReservationRequest_Implementation(const FString& SessionId, FPartyReservation Reservation);
-	virtual bool ServerReservationRequest_Validate(const FString& , FPartyReservation );
 
-	UFUNCTION(server="ServerCancelReservationRequest_Implementation", reliable, WithValidation="ServerCancelReservationRequest_Validate")
+	UFUNCTION(server, reliable, WithValidation)
 	virtual void ServerCancelReservationRequest(struct FUniqueNetIdRepl PartyLeader);
-	virtual void ServerCancelReservationRequest_Implementation(FUniqueNetIdRepl PartyLeader);
-	virtual bool ServerCancelReservationRequest_Validate(FUniqueNetIdRepl );
 };
diff --git a/Engine/Source/Runtime/Online/OnlineSubsystemUtils/Classes/PartyBeaconHost.h b/Engine/Source/Runtime/Online/OnlineSubsystemUtils/Classes/PartyBeaconHost.h
--- a/Engine/Source/Runtime/Online/OnlineSubsystemUtils/Classes/PartyBeaconHost.h
+++ b/Engine/Source/Runtime/Online/OnlineSubsystemUtils/Classes/PartyBeaconHost.h
@@ -49,9 +49,7 @@ DECLARE_DELEGATE(FOnReservationsFull);
 UCLASS(transient, notplaceable, config=Engine)
 class ONLINESUBSYSTEMUTILS_API APartyBeaconHost : public AOnlineBeaconHostObject
 {
-	GENERATED_BODY()
-public:
-	APartyBeaconHost(const FObjectInitializer& ObjectInitializer = FObjectInitializer::Get());
+	GENERATED_UCLASS_BODY()
 
 	// Begin AActor Interface
 	virtual void Tick(float DeltaTime) override;
diff --git a/Engine/Source/Runtime/Online/OnlineSubsystemUtils/Classes/PartyBeaconState.h b/Engine/Source/Runtime/Online/OnlineSubsystemUtils/Classes/PartyBeaconState.h
--- a/Engine/Source/Runtime/Online/OnlineSubsystemUtils/Classes/PartyBeaconState.h
+++ b/Engine/Source/Runtime/Online/OnlineSubsystemUtils/Classes/PartyBeaconState.h
@@ -127,7 +127,7 @@ namespace EPartyReservationResult
 USTRUCT()
 struct FPlayerReservation
 {
-	GENERATED_BODY()
+	GENERATED_USTRUCT_BODY()
 	
 	#if CPP
 	FPlayerReservation() :
@@ -152,7 +152,7 @@ struct FPlayerReservation
 USTRUCT()
 struct ONLINESUBSYSTEMUTILS_API FPartyReservation
 {
-	GENERATED_BODY()
+	GENERATED_USTRUCT_BODY()
 
 	/** Team assigned to this party */
 	UPROPERTY(Transient)
@@ -176,9 +176,7 @@ struct ONLINESUBSYSTEMUTILS_API FPartyReservation
 UCLASS(transient, notplaceable, config=Engine)
 class ONLINESUBSYSTEMUTILS_API UPartyBeaconState : public UObject
 {
-	GENERATED_BODY()
-public:
-	UPartyBeaconState(const FObjectInitializer& ObjectInitializer = FObjectInitializer::Get());
+	GENERATED_UCLASS_BODY()
 
 	/**
 	 * Initialize this state object 
diff --git a/Engine/Source/Runtime/Online/OnlineSubsystemUtils/Classes/QuitMatchCallbackProxy.h b/Engine/Source/Runtime/Online/OnlineSubsystemUtils/Classes/QuitMatchCallbackProxy.h
--- a/Engine/Source/Runtime/Online/OnlineSubsystemUtils/Classes/QuitMatchCallbackProxy.h
+++ b/Engine/Source/Runtime/Online/OnlineSubsystemUtils/Classes/QuitMatchCallbackProxy.h
@@ -8,10 +8,8 @@
 UCLASS(MinimalAPI)
 class UQuitMatchCallbackProxy : public UOnlineBlueprintCallProxyBase
 {
-	GENERATED_BODY()
-public:
-	ONLINESUBSYSTEMUTILS_API UQuitMatchCallbackProxy(const FObjectInitializer& ObjectInitializer = FObjectInitializer::Get());
 public:
+	GENERATED_UCLASS_BODY()
 
 	virtual ~UQuitMatchCallbackProxy();
 
@@ -48,4 +46,4 @@ private:
 
 	// If the match isn't over, this will be how much time the next player will have to complete their turn
 	int32 TurnTimeoutInSeconds;
-};
+};
\ No newline at end of file
diff --git a/Engine/Source/Runtime/Online/OnlineSubsystemUtils/Classes/TestBeaconClient.h b/Engine/Source/Runtime/Online/OnlineSubsystemUtils/Classes/TestBeaconClient.h
--- a/Engine/Source/Runtime/Online/OnlineSubsystemUtils/Classes/TestBeaconClient.h
+++ b/Engine/Source/Runtime/Online/OnlineSubsystemUtils/Classes/TestBeaconClient.h
@@ -10,9 +10,7 @@
 UCLASS(transient, notplaceable, config=Engine)
 class ONLINESUBSYSTEMUTILS_API ATestBeaconClient : public AOnlineBeaconClient
 {
-	GENERATED_BODY()
-public:
-	ATestBeaconClient(const FObjectInitializer& ObjectInitializer = FObjectInitializer::Get());
+	GENERATED_UCLASS_BODY()
 
 	// Begin AOnlineBeacon Interface
 	virtual FString GetBeaconType() override { return TEXT("TestBeacon"); }
@@ -23,13 +21,10 @@ public:
 	// End AOnlineBeaconClient Interface
 
 	/** Send a ping RPC to the client */
-	UFUNCTION(client="ClientPing_Implementation", reliable)
+	UFUNCTION(client, reliable)
 	virtual void ClientPing();
-	virtual void ClientPing_Implementation();
 
 	/** Send a pong RPC to the host */
-	UFUNCTION(server="ServerPong_Implementation", reliable, WithValidation="ServerPong_Validate")
+	UFUNCTION(server, reliable, WithValidation)
 	virtual void ServerPong();
-	virtual void ServerPong_Implementation();
-	virtual bool ServerPong_Validate();
 };
diff --git a/Engine/Source/Runtime/Online/OnlineSubsystemUtils/Classes/TestBeaconHost.h b/Engine/Source/Runtime/Online/OnlineSubsystemUtils/Classes/TestBeaconHost.h
--- a/Engine/Source/Runtime/Online/OnlineSubsystemUtils/Classes/TestBeaconHost.h
+++ b/Engine/Source/Runtime/Online/OnlineSubsystemUtils/Classes/TestBeaconHost.h
@@ -10,9 +10,7 @@
 UCLASS(transient, notplaceable, config=Engine)
 class ONLINESUBSYSTEMUTILS_API ATestBeaconHost : public AOnlineBeaconHostObject
 {
-	GENERATED_BODY()
-public:
-	ATestBeaconHost(const FObjectInitializer& ObjectInitializer = FObjectInitializer::Get());
+	GENERATED_UCLASS_BODY()
 
 	// Begin AOnlineBeaconHost Interface 
 	virtual AOnlineBeaconClient* SpawnBeaconActor(class UNetConnection* ClientConnection) override;
diff --git a/Engine/Source/Runtime/Online/OnlineSubsystemUtils/Classes/TurnBasedBlueprintLibrary.h b/Engine/Source/Runtime/Online/OnlineSubsystemUtils/Classes/TurnBasedBlueprintLibrary.h
--- a/Engine/Source/Runtime/Online/OnlineSubsystemUtils/Classes/TurnBasedBlueprintLibrary.h
+++ b/Engine/Source/Runtime/Online/OnlineSubsystemUtils/Classes/TurnBasedBlueprintLibrary.h
@@ -9,9 +9,7 @@
 UCLASS()
 class ONLINESUBSYSTEMUTILS_API UTurnBasedBlueprintLibrary : public UBlueprintFunctionLibrary
 {
-	GENERATED_BODY()
-public:
-	UTurnBasedBlueprintLibrary(const FObjectInitializer& ObjectInitializer = FObjectInitializer::Get());
+	GENERATED_UCLASS_BODY()
 
 	UFUNCTION(BlueprintCallable, meta = (WorldContext = "WorldContextObject"), Category = "Online|TurnBased")
 	static void GetIsMyTurn(UObject* WorldContextObject, APlayerController* PlayerController, FString MatchID, /*out*/ bool& bIsMyTurn);
diff --git a/Engine/Source/Runtime/Slate/Public/Framework/Styling/ButtonWidgetStyle.h b/Engine/Source/Runtime/Slate/Public/Framework/Styling/ButtonWidgetStyle.h
--- a/Engine/Source/Runtime/Slate/Public/Framework/Styling/ButtonWidgetStyle.h
+++ b/Engine/Source/Runtime/Slate/Public/Framework/Styling/ButtonWidgetStyle.h
@@ -10,9 +10,10 @@
 UCLASS(BlueprintType, hidecategories=Object, MinimalAPI)
 class UButtonWidgetStyle : public USlateWidgetStyleContainerBase
 {
-	GENERATED_BODY()
 public:
+	GENERATED_BODY()
 
+public:
 	/** The actual data describing the button's appearance. */
 	UPROPERTY(Category=Appearance, EditAnywhere, BlueprintReadWrite, meta=( ShowOnlyInnerProperties ))
 	FButtonStyle ButtonStyle;
diff --git a/Engine/Source/Runtime/SlateCore/Public/Fonts/CompositeFont.h b/Engine/Source/Runtime/SlateCore/Public/Fonts/CompositeFont.h
--- a/Engine/Source/Runtime/SlateCore/Public/Fonts/CompositeFont.h
+++ b/Engine/Source/Runtime/SlateCore/Public/Fonts/CompositeFont.h
@@ -24,7 +24,7 @@ enum class EFontHinting : uint8
 USTRUCT(BlueprintType)
 struct SLATECORE_API FFontData
 {
-	GENERATED_BODY()
+	GENERATED_USTRUCT_BODY()
 
 	/** Default constructor */
 	FFontData();
@@ -62,7 +62,7 @@ struct SLATECORE_API FFontData
 USTRUCT(BlueprintType)
 struct SLATECORE_API FTypefaceEntry
 {
-	GENERATED_BODY()
+	GENERATED_USTRUCT_BODY()
 
 	/** Default constructor */
 	FTypefaceEntry()
@@ -95,7 +95,7 @@ struct SLATECORE_API FTypefaceEntry
 USTRUCT(BlueprintType)
 struct SLATECORE_API FTypeface
 {
-	GENERATED_BODY()
+	GENERATED_USTRUCT_BODY()
 
 	/** Default constructor */
 	FTypeface()
@@ -123,7 +123,7 @@ struct SLATECORE_API FTypeface
 USTRUCT(BlueprintType)
 struct SLATECORE_API FCompositeSubFont
 {
-	GENERATED_BODY()
+	GENERATED_USTRUCT_BODY()
 
 	/** Default constructor */
 	FCompositeSubFont()
@@ -155,7 +155,7 @@ struct SLATECORE_API FCompositeSubFont
 USTRUCT(BlueprintType)
 struct SLATECORE_API FCompositeFont
 {
-	GENERATED_BODY()
+	GENERATED_USTRUCT_BODY()
 
 	/** Default constructor */
 	FCompositeFont()
diff --git a/Engine/Source/Runtime/SlateCore/Public/Fonts/FontProviderInterface.h b/Engine/Source/Runtime/SlateCore/Public/Fonts/FontProviderInterface.h
--- a/Engine/Source/Runtime/SlateCore/Public/Fonts/FontProviderInterface.h
+++ b/Engine/Source/Runtime/SlateCore/Public/Fonts/FontProviderInterface.h
@@ -8,15 +8,12 @@
 UINTERFACE(MinimalAPI, meta=(CannotImplementInterfaceInBlueprint))
 class UFontProviderInterface : public UInterface
 {
-	GENERATED_BODY()
-public:
-	SLATECORE_API UFontProviderInterface(const FObjectInitializer& ObjectInitializer = FObjectInitializer::Get());
+	GENERATED_UINTERFACE_BODY()
 };
 
 class IFontProviderInterface
 {
-	GENERATED_BODY()
-public:
+	GENERATED_IINTERFACE_BODY()
 
 	virtual const FCompositeFont* GetCompositeFont() const = 0;
 };
diff --git a/Engine/Source/Runtime/SlateCore/Public/Styling/SlateTypes.h b/Engine/Source/Runtime/SlateCore/Public/Styling/SlateTypes.h
--- a/Engine/Source/Runtime/SlateCore/Public/Styling/SlateTypes.h
+++ b/Engine/Source/Runtime/SlateCore/Public/Styling/SlateTypes.h
@@ -29,7 +29,7 @@ namespace ESlateCheckBoxType
 USTRUCT(BlueprintType)
 struct SLATECORE_API FCheckBoxStyle : public FSlateWidgetStyle
 {
-	GENERATED_BODY()
+	GENERATED_USTRUCT_BODY()
 
 	FCheckBoxStyle();
 
@@ -156,7 +156,7 @@ struct TStructOpsTypeTraits<FCheckBoxStyle> : public TStructOpsTypeTraitsBase
 USTRUCT(BlueprintType)
 struct SLATECORE_API FTextBlockStyle : public FSlateWidgetStyle
 {
-	GENERATED_BODY()
+	GENERATED_USTRUCT_BODY()
 
 	FTextBlockStyle();
 
@@ -221,7 +221,7 @@ struct SLATECORE_API FTextBlockStyle : public FSlateWidgetStyle
 USTRUCT(BlueprintType)
 struct SLATECORE_API FButtonStyle : public FSlateWidgetStyle
 {
-	GENERATED_BODY()
+	GENERATED_USTRUCT_BODY()
 
 	FButtonStyle();
 
@@ -311,7 +311,7 @@ struct TStructOpsTypeTraits<FButtonStyle> : public TStructOpsTypeTraitsBase
 USTRUCT(BlueprintType)
 struct SLATECORE_API FComboButtonStyle : public FSlateWidgetStyle
 {
-	GENERATED_BODY()
+	GENERATED_USTRUCT_BODY()
 
 	FComboButtonStyle();
 
@@ -360,7 +360,7 @@ struct SLATECORE_API FComboButtonStyle : public FSlateWidgetStyle
 USTRUCT(BlueprintType)
 struct SLATECORE_API FComboBoxStyle : public FSlateWidgetStyle
 {
-	GENERATED_BODY()
+	GENERATED_USTRUCT_BODY()
 
 	FComboBoxStyle();
 
@@ -420,7 +420,7 @@ struct TStructOpsTypeTraits<FComboBoxStyle> : public TStructOpsTypeTraitsBase
 USTRUCT(BlueprintType)
 struct SLATECORE_API FHyperlinkStyle : public FSlateWidgetStyle
 {
-	GENERATED_BODY()
+	GENERATED_USTRUCT_BODY()
 
 	FHyperlinkStyle();
 
@@ -455,7 +455,7 @@ struct SLATECORE_API FHyperlinkStyle : public FSlateWidgetStyle
 USTRUCT(BlueprintType)
 struct SLATECORE_API FEditableTextStyle : public FSlateWidgetStyle
 {
-	GENERATED_BODY()
+	GENERATED_USTRUCT_BODY()
 
 	FEditableTextStyle();
 
@@ -507,7 +507,7 @@ struct SLATECORE_API FEditableTextStyle : public FSlateWidgetStyle
 USTRUCT(BlueprintType)
 struct SLATECORE_API FScrollBarStyle : public FSlateWidgetStyle
 {
-	GENERATED_BODY()
+	GENERATED_USTRUCT_BODY()
 
 	FScrollBarStyle();
 
@@ -573,7 +573,7 @@ struct SLATECORE_API FScrollBarStyle : public FSlateWidgetStyle
 USTRUCT(BlueprintType)
 struct SLATECORE_API FEditableTextBoxStyle : public FSlateWidgetStyle
 {
-	GENERATED_BODY()
+	GENERATED_USTRUCT_BODY()
 
 	FEditableTextBoxStyle();
 
@@ -655,7 +655,7 @@ struct SLATECORE_API FEditableTextBoxStyle : public FSlateWidgetStyle
 USTRUCT(BlueprintType)
 struct SLATECORE_API FInlineEditableTextBlockStyle : public FSlateWidgetStyle
 {
-	GENERATED_BODY()
+	GENERATED_USTRUCT_BODY()
 
 	FInlineEditableTextBlockStyle();
 
@@ -686,7 +686,7 @@ struct SLATECORE_API FInlineEditableTextBlockStyle : public FSlateWidgetStyle
 USTRUCT(BlueprintType)
 struct SLATECORE_API FProgressBarStyle : public FSlateWidgetStyle
 {
-	GENERATED_BODY()
+	GENERATED_USTRUCT_BODY()
 
 	FProgressBarStyle();
 
@@ -722,7 +722,7 @@ struct SLATECORE_API FProgressBarStyle : public FSlateWidgetStyle
 USTRUCT()
 struct SLATECORE_API FExpandableAreaStyle : public FSlateWidgetStyle
 {
-	GENERATED_BODY()
+	GENERATED_USTRUCT_BODY()
 
 	FExpandableAreaStyle();
 
@@ -753,7 +753,7 @@ struct SLATECORE_API FExpandableAreaStyle : public FSlateWidgetStyle
 USTRUCT()
 struct SLATECORE_API FSearchBoxStyle : public FSlateWidgetStyle
 {
-	GENERATED_BODY()
+	GENERATED_USTRUCT_BODY()
 
 	FSearchBoxStyle();
 
@@ -814,7 +814,7 @@ struct SLATECORE_API FSearchBoxStyle : public FSlateWidgetStyle
 USTRUCT(BlueprintType)
 struct SLATECORE_API FSliderStyle : public FSlateWidgetStyle
 {
-	GENERATED_BODY()
+	GENERATED_USTRUCT_BODY()
 
 	FSliderStyle();
 
@@ -855,7 +855,7 @@ struct SLATECORE_API FSliderStyle : public FSlateWidgetStyle
 USTRUCT()
 struct SLATECORE_API FVolumeControlStyle : public FSlateWidgetStyle
 {
-	GENERATED_BODY()
+	GENERATED_USTRUCT_BODY()
 
 	FVolumeControlStyle();
 
@@ -905,7 +905,7 @@ struct SLATECORE_API FVolumeControlStyle : public FSlateWidgetStyle
 USTRUCT()
 struct SLATECORE_API FInlineTextImageStyle : public FSlateWidgetStyle
 {
-	GENERATED_BODY()
+	GENERATED_USTRUCT_BODY()
 
 	FInlineTextImageStyle();
 
@@ -935,7 +935,7 @@ struct SLATECORE_API FInlineTextImageStyle : public FSlateWidgetStyle
 USTRUCT(BlueprintType)
 struct SLATECORE_API FSpinBoxStyle : public FSlateWidgetStyle
 {
-	GENERATED_BODY()
+	GENERATED_USTRUCT_BODY()
 
 	FSpinBoxStyle();
 
@@ -991,7 +991,7 @@ struct SLATECORE_API FSpinBoxStyle : public FSlateWidgetStyle
 USTRUCT(BlueprintType)
 struct SLATECORE_API FSplitterStyle : public FSlateWidgetStyle
 {
-	GENERATED_BODY()
+	GENERATED_USTRUCT_BODY()
 
 	FSplitterStyle();
 
@@ -1022,7 +1022,7 @@ struct SLATECORE_API FSplitterStyle : public FSlateWidgetStyle
 USTRUCT(BlueprintType)
 struct SLATECORE_API FTableRowStyle : public FSlateWidgetStyle
 {
-	GENERATED_BODY()
+	GENERATED_USTRUCT_BODY()
 
 	FTableRowStyle();
 
@@ -1118,7 +1118,7 @@ struct SLATECORE_API FTableRowStyle : public FSlateWidgetStyle
 USTRUCT()
 struct SLATECORE_API FTableColumnHeaderStyle : public FSlateWidgetStyle
 {
-	GENERATED_BODY()
+	GENERATED_USTRUCT_BODY()
 
 	FTableColumnHeaderStyle();
 
@@ -1184,7 +1184,7 @@ struct SLATECORE_API FTableColumnHeaderStyle : public FSlateWidgetStyle
 USTRUCT()
 struct SLATECORE_API FHeaderRowStyle : public FSlateWidgetStyle
 {
-	GENERATED_BODY()
+	GENERATED_USTRUCT_BODY()
 
 	FHeaderRowStyle();
 
@@ -1230,7 +1230,7 @@ struct SLATECORE_API FHeaderRowStyle : public FSlateWidgetStyle
 USTRUCT()
 struct SLATECORE_API FDockTabStyle : public FSlateWidgetStyle
 {
-	GENERATED_BODY()
+	GENERATED_USTRUCT_BODY()
 
 	FDockTabStyle();
 
@@ -1306,7 +1306,7 @@ struct SLATECORE_API FDockTabStyle : public FSlateWidgetStyle
 USTRUCT(BlueprintType)
 struct SLATECORE_API FScrollBoxStyle : public FSlateWidgetStyle
 {
-	GENERATED_BODY()
+	GENERATED_USTRUCT_BODY()
 
 	FScrollBoxStyle();
 
@@ -1355,7 +1355,7 @@ struct SLATECORE_API FScrollBoxStyle : public FSlateWidgetStyle
 USTRUCT(BlueprintType)
 struct SLATECORE_API FScrollBorderStyle : public FSlateWidgetStyle
 {
-	GENERATED_BODY()
+	GENERATED_USTRUCT_BODY()
 
 	FScrollBorderStyle();
 
@@ -1386,7 +1386,7 @@ struct SLATECORE_API FScrollBorderStyle : public FSlateWidgetStyle
 USTRUCT(BlueprintType)
 struct SLATECORE_API FWindowStyle : public FSlateWidgetStyle
 {
-	GENERATED_BODY()
+	GENERATED_USTRUCT_BODY()
 
 	FWindowStyle();
 
@@ -1470,5 +1470,7 @@ struct SLATECORE_API FWindowStyle : public FSlateWidgetStyle
 UCLASS()
 class USlateTypes : public UObject
 {
+public:
 	GENERATED_BODY()
+
 };
diff --git a/Engine/Source/Runtime/SlateCore/Public/Styling/SlateWidgetStyleContainerInterface.h b/Engine/Source/Runtime/SlateCore/Public/Styling/SlateWidgetStyleContainerInterface.h
--- a/Engine/Source/Runtime/SlateCore/Public/Styling/SlateWidgetStyleContainerInterface.h
+++ b/Engine/Source/Runtime/SlateCore/Public/Styling/SlateWidgetStyleContainerInterface.h
@@ -8,17 +8,14 @@
 UINTERFACE()
 class SLATECORE_API USlateWidgetStyleContainerInterface : public UInterface
 {
-	GENERATED_BODY()
-public:
-	USlateWidgetStyleContainerInterface(const FObjectInitializer& ObjectInitializer = FObjectInitializer::Get());
+	GENERATED_UINTERFACE_BODY()
 };
 
 class ISlateWidgetStyleContainerInterface
 {
-	GENERATED_BODY()
-public:
+	GENERATED_IINTERFACE_BODY()
 
 public:
 
 	virtual const struct FSlateWidgetStyle* const GetStyle() const = 0;
-};
+};
\ No newline at end of file
diff --git a/Engine/Source/Runtime/UMG/Public/Animation/MovieScene2DTransformSection.h b/Engine/Source/Runtime/UMG/Public/Animation/MovieScene2DTransformSection.h
--- a/Engine/Source/Runtime/UMG/Public/Animation/MovieScene2DTransformSection.h
+++ b/Engine/Source/Runtime/UMG/Public/Animation/MovieScene2DTransformSection.h
@@ -14,9 +14,7 @@ struct FWidgetTransform;
 UCLASS(MinimalAPI)
 class UMovieScene2DTransformSection : public UMovieSceneSection
 {
-	GENERATED_BODY()
-public:
-	UMG_API UMovieScene2DTransformSection(const FObjectInitializer& ObjectInitializer = FObjectInitializer::Get());
+	GENERATED_UCLASS_BODY()
 public:
 
 	/** MovieSceneSection interface */
@@ -54,4 +52,4 @@ private:
 	/** Shear curve */
 	UPROPERTY()
 	FRichCurve Shear[2];
-};
+};
\ No newline at end of file
diff --git a/Engine/Source/Runtime/UMG/Public/Animation/MovieSceneMarginTrack.h b/Engine/Source/Runtime/UMG/Public/Animation/MovieSceneMarginTrack.h
--- a/Engine/Source/Runtime/UMG/Public/Animation/MovieSceneMarginTrack.h
+++ b/Engine/Source/Runtime/UMG/Public/Animation/MovieSceneMarginTrack.h
@@ -18,9 +18,7 @@ struct FMarginKey
 UCLASS( MinimalAPI )
 class UMovieSceneMarginTrack : public UMovieScenePropertyTrack
 {
-	GENERATED_BODY()
-public:
-	UMG_API UMovieSceneMarginTrack(const FObjectInitializer& ObjectInitializer = FObjectInitializer::Get());
+	GENERATED_UCLASS_BODY()
 
 public:
 	/** UMovieSceneTrack interface */
diff --git a/Engine/Source/Runtime/UMG/Public/Animation/UMGSequencePlayer.h b/Engine/Source/Runtime/UMG/Public/Animation/UMGSequencePlayer.h
--- a/Engine/Source/Runtime/UMG/Public/Animation/UMGSequencePlayer.h
+++ b/Engine/Source/Runtime/UMG/Public/Animation/UMGSequencePlayer.h
@@ -11,9 +11,7 @@ class UMovieSceneBindings;
 UCLASS(transient)
 class UMG_API UUMGSequencePlayer : public UObject, public IMovieScenePlayer
 {
-	GENERATED_BODY()
-public:
-	UUMGSequencePlayer(const FObjectInitializer& ObjectInitializer = FObjectInitializer::Get());
+	GENERATED_UCLASS_BODY()
 
 public:
 	void InitSequencePlayer( const UWidgetAnimation& InAnimation, UUserWidget& UserWidget );
diff --git a/Engine/Source/Runtime/UMG/Public/Animation/WidgetAnimation.h b/Engine/Source/Runtime/UMG/Public/Animation/WidgetAnimation.h
--- a/Engine/Source/Runtime/UMG/Public/Animation/WidgetAnimation.h
+++ b/Engine/Source/Runtime/UMG/Public/Animation/WidgetAnimation.h
@@ -12,7 +12,7 @@ class UWidgetTree;
 USTRUCT()
 struct FWidgetAnimationBinding
 {
-	GENERATED_BODY()
+	GENERATED_USTRUCT_BODY()
 
 	UPROPERTY()
 	FName WidgetName;
@@ -48,9 +48,7 @@ public:
 UCLASS(MinimalAPI)
 class UWidgetAnimation : public UObject
 {
-	GENERATED_BODY()
-public:
-	UMG_API UWidgetAnimation(const FObjectInitializer& ObjectInitializer = FObjectInitializer::Get());
+	GENERATED_UCLASS_BODY()
 
 public:
 
diff --git a/Engine/Source/Runtime/UMG/Public/Binding/DynamicPropertyPath.h b/Engine/Source/Runtime/UMG/Public/Binding/DynamicPropertyPath.h
--- a/Engine/Source/Runtime/UMG/Public/Binding/DynamicPropertyPath.h
+++ b/Engine/Source/Runtime/UMG/Public/Binding/DynamicPropertyPath.h
@@ -1,4 +1,4 @@
-// Copyright 1998-2015 Epic Games, Inc. All Rights Reserved.
+﻿// Copyright 1998-2015 Epic Games, Inc. All Rights Reserved.
 
 #pragma once
 
@@ -123,7 +123,7 @@ inline bool IsConcreteTypeCompatibleWithReflectedType<UObject*>(UProperty* Prope
 USTRUCT()
 struct UMG_API FPropertyPathSegment
 {
-	GENERATED_BODY()
+	GENERATED_USTRUCT_BODY()
 
 public:
 	/** Implementation detail, don't use this constructor. */
@@ -166,7 +166,7 @@ private:
 USTRUCT()
 struct UMG_API FDynamicPropertyPath
 {
-	GENERATED_BODY()
+	GENERATED_USTRUCT_BODY()
 
 public:
 
diff --git a/Engine/Source/Runtime/UMG/Public/Blueprint/DragDropOperation.h b/Engine/Source/Runtime/UMG/Public/Blueprint/DragDropOperation.h
--- a/Engine/Source/Runtime/UMG/Public/Blueprint/DragDropOperation.h
+++ b/Engine/Source/Runtime/UMG/Public/Blueprint/DragDropOperation.h
@@ -33,9 +33,7 @@ enum class EDragPivot : uint8
 UCLASS(BlueprintType, Blueprintable)
 class UMG_API UDragDropOperation : public UObject
 {
-	GENERATED_BODY()
-public:
-	UDragDropOperation(const FObjectInitializer& ObjectInitializer = FObjectInitializer::Get());
+	GENERATED_UCLASS_BODY()
 	
 public:
 	/** A simple string tag you can optionally use to provide extra metadata about the operation. */
@@ -82,15 +80,12 @@ public:
 	/**  */
 	UFUNCTION(BlueprintNativeEvent, Category="Drag and Drop")
 	void Drop(const FPointerEvent& PointerEvent);
-	virtual void Drop_Implementation(const FPointerEvent& PointerEvent);
 
 	/**  */
 	UFUNCTION(BlueprintNativeEvent, Category="Drag and Drop")
 	void DragCancelled(const FPointerEvent& PointerEvent);
-	virtual void DragCancelled_Implementation(const FPointerEvent& PointerEvent);
 
 	/**  */
 	UFUNCTION(BlueprintNativeEvent, Category="Drag and Drop")
 	void Dragged(const FPointerEvent& PointerEvent);
-	virtual void Dragged_Implementation(const FPointerEvent& PointerEvent);
 };
diff --git a/Engine/Source/Runtime/UMG/Public/Blueprint/SlateBlueprintLibrary.h b/Engine/Source/Runtime/UMG/Public/Blueprint/SlateBlueprintLibrary.h
--- a/Engine/Source/Runtime/UMG/Public/Blueprint/SlateBlueprintLibrary.h
+++ b/Engine/Source/Runtime/UMG/Public/Blueprint/SlateBlueprintLibrary.h
@@ -7,9 +7,7 @@
 UCLASS(MinimalAPI)
 class USlateBlueprintLibrary : public UBlueprintFunctionLibrary
 {
-	GENERATED_BODY()
-public:
-	UMG_API USlateBlueprintLibrary(const FObjectInitializer& ObjectInitializer = FObjectInitializer::Get());
+	GENERATED_UCLASS_BODY()
 
 public:
 
diff --git a/Engine/Source/Runtime/UMG/Public/Blueprint/WidgetBlueprintGeneratedClass.h b/Engine/Source/Runtime/UMG/Public/Blueprint/WidgetBlueprintGeneratedClass.h
--- a/Engine/Source/Runtime/UMG/Public/Blueprint/WidgetBlueprintGeneratedClass.h
+++ b/Engine/Source/Runtime/UMG/Public/Blueprint/WidgetBlueprintGeneratedClass.h
@@ -22,7 +22,7 @@ enum class EBindingKind : uint8
 USTRUCT()
 struct FDelegateRuntimeBinding
 {
-	GENERATED_BODY()
+	GENERATED_USTRUCT_BODY()
 
 	/** The widget that will be bound to the live data. */
 	UPROPERTY()
@@ -55,9 +55,7 @@ struct FDelegateRuntimeBinding
 UCLASS()
 class UMG_API UWidgetBlueprintGeneratedClass : public UBlueprintGeneratedClass
 {
-	GENERATED_BODY()
-public:
-	UWidgetBlueprintGeneratedClass(const FObjectInitializer& ObjectInitializer = FObjectInitializer::Get());
+	GENERATED_UCLASS_BODY()
 
 public:
 
@@ -84,4 +82,4 @@ public:
 	 * binding and wiring necessary to have the user's widget perform as desired.
 	 */
 	void InitializeWidget(UUserWidget* UserWidget) const;
-};
+};
\ No newline at end of file
diff --git a/Engine/Source/Runtime/UMG/Public/Blueprint/WidgetBlueprintLibrary.h b/Engine/Source/Runtime/UMG/Public/Blueprint/WidgetBlueprintLibrary.h
--- a/Engine/Source/Runtime/UMG/Public/Blueprint/WidgetBlueprintLibrary.h
+++ b/Engine/Source/Runtime/UMG/Public/Blueprint/WidgetBlueprintLibrary.h
@@ -13,9 +13,7 @@ struct FPaintContext;
 UCLASS(MinimalAPI)
 class UWidgetBlueprintLibrary : public UBlueprintFunctionLibrary
 {
-	GENERATED_BODY()
-public:
-	UMG_API UWidgetBlueprintLibrary(const FObjectInitializer& ObjectInitializer = FObjectInitializer::Get());
+	GENERATED_UCLASS_BODY()
 
 public:
 
diff --git a/Engine/Source/Runtime/UMG/Public/Blueprint/WidgetNavigation.h b/Engine/Source/Runtime/UMG/Public/Blueprint/WidgetNavigation.h
--- a/Engine/Source/Runtime/UMG/Public/Blueprint/WidgetNavigation.h
+++ b/Engine/Source/Runtime/UMG/Public/Blueprint/WidgetNavigation.h
@@ -13,7 +13,7 @@
 USTRUCT()
 struct UMG_API FWidgetNavigationData
 {
-	GENERATED_BODY()
+	GENERATED_USTRUCT_BODY()
 
 public:
 	UPROPERTY(EditAnywhere, BlueprintReadOnly, Category="Navigation")
@@ -29,9 +29,7 @@ public:
 UCLASS()
 class UMG_API UWidgetNavigation : public UObject
 {
-	GENERATED_BODY()
-public:
-	UWidgetNavigation(const FObjectInitializer& ObjectInitializer = FObjectInitializer::Get());
+	GENERATED_UCLASS_BODY()
 	
 public:
 	/**  */
diff --git a/Engine/Source/Runtime/UMG/Public/Blueprint/WidgetTree.h b/Engine/Source/Runtime/UMG/Public/Blueprint/WidgetTree.h
--- a/Engine/Source/Runtime/UMG/Public/Blueprint/WidgetTree.h
+++ b/Engine/Source/Runtime/UMG/Public/Blueprint/WidgetTree.h
@@ -10,9 +10,7 @@
 UCLASS()
 class UMG_API UWidgetTree : public UObject
 {
-	GENERATED_BODY()
-public:
-	UWidgetTree(const FObjectInitializer& ObjectInitializer = FObjectInitializer::Get());
+	GENERATED_UCLASS_BODY()
 
 public:
 
diff --git a/Engine/Source/Runtime/UMG/Public/Components/Border.h b/Engine/Source/Runtime/UMG/Public/Components/Border.h
--- a/Engine/Source/Runtime/UMG/Public/Components/Border.h
+++ b/Engine/Source/Runtime/UMG/Public/Components/Border.h
@@ -1,4 +1,4 @@
-// Copyright 1998-2015 Epic Games, Inc. All Rights Reserved.
+﻿// Copyright 1998-2015 Epic Games, Inc. All Rights Reserved.
 
 #pragma once
 
@@ -16,9 +16,7 @@ class USlateBrushAsset;
 UCLASS()
 class UMG_API UBorder : public UContentWidget
 {
-	GENERATED_BODY()
-public:
-	UBorder(const FObjectInitializer& ObjectInitializer = FObjectInitializer::Get());
+	GENERATED_UCLASS_BODY()
 
 public:
 
diff --git a/Engine/Source/Runtime/UMG/Public/Components/BorderSlot.h b/Engine/Source/Runtime/UMG/Public/Components/BorderSlot.h
--- a/Engine/Source/Runtime/UMG/Public/Components/BorderSlot.h
+++ b/Engine/Source/Runtime/UMG/Public/Components/BorderSlot.h
@@ -14,9 +14,7 @@ class UBorder;
 UCLASS()
 class UMG_API UBorderSlot : public UPanelSlot
 {
-	GENERATED_BODY()
-public:
-	UBorderSlot(const FObjectInitializer& ObjectInitializer = FObjectInitializer::Get());
+	GENERATED_UCLASS_BODY()
 
 public:
 
diff --git a/Engine/Source/Runtime/UMG/Public/Components/Button.h b/Engine/Source/Runtime/UMG/Public/Components/Button.h
--- a/Engine/Source/Runtime/UMG/Public/Components/Button.h
+++ b/Engine/Source/Runtime/UMG/Public/Components/Button.h
@@ -1,4 +1,4 @@
-// Copyright 1998-2015 Epic Games, Inc. All Rights Reserved.
+﻿// Copyright 1998-2015 Epic Games, Inc. All Rights Reserved.
 
 #pragma once
 
@@ -21,9 +21,7 @@ DECLARE_DYNAMIC_MULTICAST_DELEGATE(FOnButtonReleasedEvent);
 UCLASS()
 class UMG_API UButton : public UContentWidget
 {
-	GENERATED_BODY()
-public:
-	UButton(const FObjectInitializer& ObjectInitializer = FObjectInitializer::Get());
+	GENERATED_UCLASS_BODY()
 
 public:
 	/** The template style asset, used to seed the mutable instance of the style. */
diff --git a/Engine/Source/Runtime/UMG/Public/Components/ButtonSlot.h b/Engine/Source/Runtime/UMG/Public/Components/ButtonSlot.h
--- a/Engine/Source/Runtime/UMG/Public/Components/ButtonSlot.h
+++ b/Engine/Source/Runtime/UMG/Public/Components/ButtonSlot.h
@@ -10,9 +10,7 @@
 UCLASS()
 class UMG_API UButtonSlot : public UPanelSlot
 {
-	GENERATED_BODY()
-public:
-	UButtonSlot(const FObjectInitializer& ObjectInitializer = FObjectInitializer::Get());
+	GENERATED_UCLASS_BODY()
 
 public:
 	
diff --git a/Engine/Source/Runtime/UMG/Public/Components/CanvasPanel.h b/Engine/Source/Runtime/UMG/Public/Components/CanvasPanel.h
--- a/Engine/Source/Runtime/UMG/Public/Components/CanvasPanel.h
+++ b/Engine/Source/Runtime/UMG/Public/Components/CanvasPanel.h
@@ -1,4 +1,4 @@
-// Copyright 1998-2015 Epic Games, Inc. All Rights Reserved.
+﻿// Copyright 1998-2015 Epic Games, Inc. All Rights Reserved.
 
 #pragma once
 #include "Components/PanelWidget.h"
@@ -19,9 +19,7 @@ class UCanvasPanelSlot;
 UCLASS()
 class UMG_API UCanvasPanel : public UPanelWidget
 {
-	GENERATED_BODY()
-public:
-	UCanvasPanel(const FObjectInitializer& ObjectInitializer = FObjectInitializer::Get());
+	GENERATED_UCLASS_BODY()
 
 public:
 
diff --git a/Engine/Source/Runtime/UMG/Public/Components/CanvasPanelSlot.h b/Engine/Source/Runtime/UMG/Public/Components/CanvasPanelSlot.h
--- a/Engine/Source/Runtime/UMG/Public/Components/CanvasPanelSlot.h
+++ b/Engine/Source/Runtime/UMG/Public/Components/CanvasPanelSlot.h
@@ -10,8 +10,8 @@
 USTRUCT(BlueprintType)
 struct FAnchorData
 {
-	GENERATED_BODY()
 public:
+	GENERATED_USTRUCT_BODY()
 
 public:
 
@@ -31,9 +31,7 @@ public:
 UCLASS()
 class UMG_API UCanvasPanelSlot : public UPanelSlot
 {
-	GENERATED_BODY()
-public:
-	UCanvasPanelSlot(const FObjectInitializer& ObjectInitializer = FObjectInitializer::Get());
+	GENERATED_UCLASS_BODY()
 
 public:
 
@@ -158,4 +156,4 @@ private:
 
 	TOptional<FVector2D> DesiredPosition;
 #endif
-};
+};
\ No newline at end of file
diff --git a/Engine/Source/Runtime/UMG/Public/Components/CheckBox.h b/Engine/Source/Runtime/UMG/Public/Components/CheckBox.h
--- a/Engine/Source/Runtime/UMG/Public/Components/CheckBox.h
+++ b/Engine/Source/Runtime/UMG/Public/Components/CheckBox.h
@@ -1,4 +1,4 @@
-// Copyright 1998-2015 Epic Games, Inc. All Rights Reserved.
+﻿// Copyright 1998-2015 Epic Games, Inc. All Rights Reserved.
 
 #pragma once
 
@@ -19,9 +19,7 @@ DECLARE_DYNAMIC_MULTICAST_DELEGATE_OneParam( FOnCheckBoxComponentStateChanged, b
 UCLASS()
 class UMG_API UCheckBox : public UContentWidget
 {
-	GENERATED_BODY()
-public:
-	UCheckBox(const FObjectInitializer& ObjectInitializer = FObjectInitializer::Get());
+	GENERATED_UCLASS_BODY()
 
 public:
 	/** Whether the check box is currently in a checked state */
diff --git a/Engine/Source/Runtime/UMG/Public/Components/CircularThrobber.h b/Engine/Source/Runtime/UMG/Public/Components/CircularThrobber.h
--- a/Engine/Source/Runtime/UMG/Public/Components/CircularThrobber.h
+++ b/Engine/Source/Runtime/UMG/Public/Components/CircularThrobber.h
@@ -1,4 +1,4 @@
-// Copyright 1998-2015 Epic Games, Inc. All Rights Reserved.
+﻿// Copyright 1998-2015 Epic Games, Inc. All Rights Reserved.
 
 #pragma once
 
@@ -15,9 +15,7 @@ class USlateBrushAsset;
 UCLASS()
 class UMG_API UCircularThrobber : public UWidget
 {
-	GENERATED_BODY()
-public:
-	UCircularThrobber(const FObjectInitializer& ObjectInitializer = FObjectInitializer::Get());
+	GENERATED_UCLASS_BODY()
 
 public:
 
diff --git a/Engine/Source/Runtime/UMG/Public/Components/ComboBox.h b/Engine/Source/Runtime/UMG/Public/Components/ComboBox.h
--- a/Engine/Source/Runtime/UMG/Public/Components/ComboBox.h
+++ b/Engine/Source/Runtime/UMG/Public/Components/ComboBox.h
@@ -10,9 +10,7 @@
 UCLASS(Experimental, meta=( DisplayName="ComboBox (Object)" ))
 class UMG_API UComboBox : public UWidget
 {
-	GENERATED_BODY()
-public:
-	UComboBox(const FObjectInitializer& ObjectInitializer = FObjectInitializer::Get());
+	GENERATED_UCLASS_BODY()
 
 public:
 
diff --git a/Engine/Source/Runtime/UMG/Public/Components/ComboBoxString.h b/Engine/Source/Runtime/UMG/Public/Components/ComboBoxString.h
--- a/Engine/Source/Runtime/UMG/Public/Components/ComboBoxString.h
+++ b/Engine/Source/Runtime/UMG/Public/Components/ComboBoxString.h
@@ -12,9 +12,7 @@
 UCLASS( meta=( DisplayName="ComboBox (String)"), ClassGroup=UserInterface)
 class UMG_API UComboBoxString : public UWidget
 {
-	GENERATED_BODY()
-public:
-	UComboBoxString(const FObjectInitializer& ObjectInitializer = FObjectInitializer::Get());
+	GENERATED_UCLASS_BODY()
 
 	DECLARE_DYNAMIC_MULTICAST_DELEGATE_TwoParams(FOnSelectionChangedEvent, FString, SelectedItem, ESelectInfo::Type, SelectionType);
 	DECLARE_DYNAMIC_MULTICAST_DELEGATE(FOnOpeningEvent);
diff --git a/Engine/Source/Runtime/UMG/Public/Components/ContentWidget.h b/Engine/Source/Runtime/UMG/Public/Components/ContentWidget.h
--- a/Engine/Source/Runtime/UMG/Public/Components/ContentWidget.h
+++ b/Engine/Source/Runtime/UMG/Public/Components/ContentWidget.h
@@ -8,9 +8,7 @@
 UCLASS(Abstract)
 class UMG_API UContentWidget : public UPanelWidget
 {
-	GENERATED_BODY()
-public:
-	UContentWidget(const FObjectInitializer& ObjectInitializer = FObjectInitializer::Get());
+	GENERATED_UCLASS_BODY()
 
 public:
 	/**  */
diff --git a/Engine/Source/Runtime/UMG/Public/Components/EditableText.h b/Engine/Source/Runtime/UMG/Public/Components/EditableText.h
--- a/Engine/Source/Runtime/UMG/Public/Components/EditableText.h
+++ b/Engine/Source/Runtime/UMG/Public/Components/EditableText.h
@@ -12,9 +12,7 @@
 UCLASS()
 class UMG_API UEditableText : public UWidget
 {
-	GENERATED_BODY()
-public:
-	UEditableText(const FObjectInitializer& ObjectInitializer = FObjectInitializer::Get());
+	GENERATED_UCLASS_BODY()
 
 public:
 
diff --git a/Engine/Source/Runtime/UMG/Public/Components/EditableTextBox.h b/Engine/Source/Runtime/UMG/Public/Components/EditableTextBox.h
--- a/Engine/Source/Runtime/UMG/Public/Components/EditableTextBox.h
+++ b/Engine/Source/Runtime/UMG/Public/Components/EditableTextBox.h
@@ -1,4 +1,4 @@
-// Copyright 1998-2015 Epic Games, Inc. All Rights Reserved.
+﻿// Copyright 1998-2015 Epic Games, Inc. All Rights Reserved.
 
 #pragma once
 
@@ -13,9 +13,7 @@
 UCLASS(meta=(DisplayName="Text Box"))
 class UMG_API UEditableTextBox : public UWidget
 {
-	GENERATED_BODY()
-public:
-	UEditableTextBox(const FObjectInitializer& ObjectInitializer = FObjectInitializer::Get());
+	GENERATED_UCLASS_BODY()
 
 public:
 
diff --git a/Engine/Source/Runtime/UMG/Public/Components/ExpandableArea.h b/Engine/Source/Runtime/UMG/Public/Components/ExpandableArea.h
--- a/Engine/Source/Runtime/UMG/Public/Components/ExpandableArea.h
+++ b/Engine/Source/Runtime/UMG/Public/Components/ExpandableArea.h
@@ -1,4 +1,4 @@
-// Copyright 1998-2015 Epic Games, Inc. All Rights Reserved.
+﻿// Copyright 1998-2015 Epic Games, Inc. All Rights Reserved.
 
 #pragma once
 
@@ -12,9 +12,7 @@ DECLARE_DYNAMIC_MULTICAST_DELEGATE_OneParam(FOnExpandableAreaExpansionChanged, b
 UCLASS(Experimental)
 class UMG_API UExpandableArea : public UWidget, public INamedSlotInterface
 {
-	GENERATED_BODY()
-public:
-	UExpandableArea(const FObjectInitializer& ObjectInitializer = FObjectInitializer::Get());
+	GENERATED_UCLASS_BODY()
 
 public:
 	/**  */
diff --git a/Engine/Source/Runtime/UMG/Public/Components/GridPanel.h b/Engine/Source/Runtime/UMG/Public/Components/GridPanel.h
--- a/Engine/Source/Runtime/UMG/Public/Components/GridPanel.h
+++ b/Engine/Source/Runtime/UMG/Public/Components/GridPanel.h
@@ -1,4 +1,4 @@
-// Copyright 1998-2015 Epic Games, Inc. All Rights Reserved.
+﻿// Copyright 1998-2015 Epic Games, Inc. All Rights Reserved.
 
 #pragma once
 
@@ -12,9 +12,7 @@
 UCLASS()
 class UMG_API UGridPanel : public UPanelWidget
 {
-	GENERATED_BODY()
-public:
-	UGridPanel(const FObjectInitializer& ObjectInitializer = FObjectInitializer::Get());
+	GENERATED_UCLASS_BODY()
 
 public:
 
diff --git a/Engine/Source/Runtime/UMG/Public/Components/GridSlot.h b/Engine/Source/Runtime/UMG/Public/Components/GridSlot.h
--- a/Engine/Source/Runtime/UMG/Public/Components/GridSlot.h
+++ b/Engine/Source/Runtime/UMG/Public/Components/GridSlot.h
@@ -13,9 +13,7 @@
 UCLASS()
 class UMG_API UGridSlot : public UPanelSlot
 {
-	GENERATED_BODY()
-public:
-	UGridSlot(const FObjectInitializer& ObjectInitializer = FObjectInitializer::Get());
+	GENERATED_UCLASS_BODY()
 
 public:
 
@@ -89,4 +87,4 @@ private:
 
 	/** A raw pointer to the slot to allow us to adjust the size, padding...etc at runtime. */
 	SGridPanel::FSlot* Slot;
-};
+};
\ No newline at end of file
diff --git a/Engine/Source/Runtime/UMG/Public/Components/HorizontalBox.h b/Engine/Source/Runtime/UMG/Public/Components/HorizontalBox.h
--- a/Engine/Source/Runtime/UMG/Public/Components/HorizontalBox.h
+++ b/Engine/Source/Runtime/UMG/Public/Components/HorizontalBox.h
@@ -1,4 +1,4 @@
-// Copyright 1998-2015 Epic Games, Inc. All Rights Reserved.
+﻿// Copyright 1998-2015 Epic Games, Inc. All Rights Reserved.
 
 #pragma once
 #include "Components/PanelWidget.h"
@@ -15,9 +15,7 @@ class UHorizontalBoxSlot;
 UCLASS(ClassGroup=UserInterface)
 class UMG_API UHorizontalBox : public UPanelWidget
 {
-	GENERATED_BODY()
-public:
-	UHorizontalBox(const FObjectInitializer& ObjectInitializer = FObjectInitializer::Get());
+	GENERATED_UCLASS_BODY()
 
 	/**  */
 	UFUNCTION(BlueprintCallable, Category="Widget")
diff --git a/Engine/Source/Runtime/UMG/Public/Components/HorizontalBoxSlot.h b/Engine/Source/Runtime/UMG/Public/Components/HorizontalBoxSlot.h
--- a/Engine/Source/Runtime/UMG/Public/Components/HorizontalBoxSlot.h
+++ b/Engine/Source/Runtime/UMG/Public/Components/HorizontalBoxSlot.h
@@ -9,9 +9,7 @@
 UCLASS()
 class UMG_API UHorizontalBoxSlot : public UPanelSlot
 {
-	GENERATED_BODY()
-public:
-	UHorizontalBoxSlot(const FObjectInitializer& ObjectInitializer = FObjectInitializer::Get());
+	GENERATED_UCLASS_BODY()
 
 public:
 	
@@ -53,4 +51,4 @@ public:
 
 private:
 	SHorizontalBox::FSlot* Slot;
-};
+};
\ No newline at end of file
diff --git a/Engine/Source/Runtime/UMG/Public/Components/Image.h b/Engine/Source/Runtime/UMG/Public/Components/Image.h
--- a/Engine/Source/Runtime/UMG/Public/Components/Image.h
+++ b/Engine/Source/Runtime/UMG/Public/Components/Image.h
@@ -1,4 +1,4 @@
-// Copyright 1998-2015 Epic Games, Inc. All Rights Reserved.
+﻿// Copyright 1998-2015 Epic Games, Inc. All Rights Reserved.
 
 #pragma once
 
@@ -14,9 +14,7 @@ class USlateBrushAsset;
 UCLASS()
 class UMG_API UImage : public UWidget
 {
-	GENERATED_BODY()
-public:
-	UImage(const FObjectInitializer& ObjectInitializer = FObjectInitializer::Get());
+	GENERATED_UCLASS_BODY()
 
 public:
 
diff --git a/Engine/Source/Runtime/UMG/Public/Components/ListView.h b/Engine/Source/Runtime/UMG/Public/Components/ListView.h
--- a/Engine/Source/Runtime/UMG/Public/Components/ListView.h
+++ b/Engine/Source/Runtime/UMG/Public/Components/ListView.h
@@ -10,9 +10,7 @@
 UCLASS(Experimental, ClassGroup=UserInterface)
 class UMG_API UListView : public UTableViewBase
 {
-	GENERATED_BODY()
-public:
-	UListView(const FObjectInitializer& ObjectInitializer = FObjectInitializer::Get());
+	GENERATED_UCLASS_BODY()
 
 public:
 
diff --git a/Engine/Source/Runtime/UMG/Public/Components/MenuAnchor.h b/Engine/Source/Runtime/UMG/Public/Components/MenuAnchor.h
--- a/Engine/Source/Runtime/UMG/Public/Components/MenuAnchor.h
+++ b/Engine/Source/Runtime/UMG/Public/Components/MenuAnchor.h
@@ -15,9 +15,7 @@ DECLARE_DYNAMIC_MULTICAST_DELEGATE_OneParam(FOnMenuOpenChangedEvent, bool, bIsOp
 UCLASS()
 class UMG_API UMenuAnchor : public UContentWidget
 {
-	GENERATED_BODY()
-public:
-	UMenuAnchor(const FObjectInitializer& ObjectInitializer = FObjectInitializer::Get());
+	GENERATED_UCLASS_BODY()
 
 public:
 
diff --git a/Engine/Source/Runtime/UMG/Public/Components/MultiLineEditableText.h b/Engine/Source/Runtime/UMG/Public/Components/MultiLineEditableText.h
--- a/Engine/Source/Runtime/UMG/Public/Components/MultiLineEditableText.h
+++ b/Engine/Source/Runtime/UMG/Public/Components/MultiLineEditableText.h
@@ -10,9 +10,7 @@
 UCLASS(meta=( DisplayName="Editable Text (Multi-Line)" ))
 class UMG_API UMultiLineEditableText : public UWidget
 {
-	GENERATED_BODY()
-public:
-	UMultiLineEditableText(const FObjectInitializer& ObjectInitializer = FObjectInitializer::Get());
+	GENERATED_UCLASS_BODY()
 
 public:
 
diff --git a/Engine/Source/Runtime/UMG/Public/Components/MultiLineEditableTextBox.h b/Engine/Source/Runtime/UMG/Public/Components/MultiLineEditableTextBox.h
--- a/Engine/Source/Runtime/UMG/Public/Components/MultiLineEditableTextBox.h
+++ b/Engine/Source/Runtime/UMG/Public/Components/MultiLineEditableTextBox.h
@@ -10,9 +10,7 @@
 UCLASS(meta=(DisplayName="Text Box (Multi-Line)"))
 class UMG_API UMultiLineEditableTextBox : public UWidget
 {
-	GENERATED_BODY()
-public:
-	UMultiLineEditableTextBox(const FObjectInitializer& ObjectInitializer = FObjectInitializer::Get());
+	GENERATED_UCLASS_BODY()
 
 public:
 
diff --git a/Engine/Source/Runtime/UMG/Public/Components/NamedSlot.h b/Engine/Source/Runtime/UMG/Public/Components/NamedSlot.h
--- a/Engine/Source/Runtime/UMG/Public/Components/NamedSlot.h
+++ b/Engine/Source/Runtime/UMG/Public/Components/NamedSlot.h
@@ -11,9 +11,7 @@
 UCLASS()
 class UMG_API UNamedSlot : public UContentWidget
 {
-	GENERATED_BODY()
-public:
-	UNamedSlot(const FObjectInitializer& ObjectInitializer = FObjectInitializer::Get());
+	GENERATED_UCLASS_BODY()
 
 public:
 
diff --git a/Engine/Source/Runtime/UMG/Public/Components/NamedSlotInterface.h b/Engine/Source/Runtime/UMG/Public/Components/NamedSlotInterface.h
--- a/Engine/Source/Runtime/UMG/Public/Components/NamedSlotInterface.h
+++ b/Engine/Source/Runtime/UMG/Public/Components/NamedSlotInterface.h
@@ -10,15 +10,12 @@ class UWidget;
 UINTERFACE(MinimalAPI, meta=( CannotImplementInterfaceInBlueprint ))
 class UNamedSlotInterface : public UInterface
 {
-	GENERATED_BODY()
-public:
-	UMG_API UNamedSlotInterface(const FObjectInitializer& ObjectInitializer = FObjectInitializer::Get());
+	GENERATED_UINTERFACE_BODY()
 };
 
 class UMG_API INamedSlotInterface
 {
-	GENERATED_BODY()
-public:
+	GENERATED_IINTERFACE_BODY()
 
 	/**  */
 	virtual void GetSlotNames(TArray<FName>& SlotNames) const = 0;
diff --git a/Engine/Source/Runtime/UMG/Public/Components/NativeWidgetHost.h b/Engine/Source/Runtime/UMG/Public/Components/NativeWidgetHost.h
--- a/Engine/Source/Runtime/UMG/Public/Components/NativeWidgetHost.h
+++ b/Engine/Source/Runtime/UMG/Public/Components/NativeWidgetHost.h
@@ -11,9 +11,7 @@
 UCLASS()
 class UMG_API UNativeWidgetHost : public UWidget
 {
-	GENERATED_BODY()
-public:
-	UNativeWidgetHost(const FObjectInitializer& ObjectInitializer = FObjectInitializer::Get());
+	GENERATED_UCLASS_BODY()
 
 	void SetContent(TSharedRef<SWidget> InContent);
 
diff --git a/Engine/Source/Runtime/UMG/Public/Components/Overlay.h b/Engine/Source/Runtime/UMG/Public/Components/Overlay.h
--- a/Engine/Source/Runtime/UMG/Public/Components/Overlay.h
+++ b/Engine/Source/Runtime/UMG/Public/Components/Overlay.h
@@ -10,9 +10,7 @@
 UCLASS()
 class UMG_API UOverlay : public UPanelWidget
 {
-	GENERATED_BODY()
-public:
-	UOverlay(const FObjectInitializer& ObjectInitializer = FObjectInitializer::Get());
+	GENERATED_UCLASS_BODY()
 
 	virtual void ReleaseSlateResources(bool bReleaseChildren) override;
 
diff --git a/Engine/Source/Runtime/UMG/Public/Components/OverlaySlot.h b/Engine/Source/Runtime/UMG/Public/Components/OverlaySlot.h
--- a/Engine/Source/Runtime/UMG/Public/Components/OverlaySlot.h
+++ b/Engine/Source/Runtime/UMG/Public/Components/OverlaySlot.h
@@ -12,9 +12,7 @@
 UCLASS()
 class UMG_API UOverlaySlot : public UPanelSlot
 {
-	GENERATED_BODY()
-public:
-	UOverlaySlot(const FObjectInitializer& ObjectInitializer = FObjectInitializer::Get());
+	GENERATED_UCLASS_BODY()
 
 public:
 	
