Added warnings for when FSlateFontInfo is created with a null font, or a font that can't provide a valid composite font

This also stops offline (bitmap) cached fonts from reporting that they're able to provide a composite font (they can't).

This makes it clearer when Slate will be forced to use the fallback font path, which may not always be as fast as the standard font path.

diff --git a/Engine/Source/Editor/DetailCustomizations/Private/SlateFontInfoCustomization.cpp b/Engine/Source/Editor/DetailCustomizations/Private/SlateFontInfoCustomization.cpp
--- a/Engine/Source/Editor/DetailCustomizations/Private/SlateFontInfoCustomization.cpp
+++ b/Engine/Source/Editor/DetailCustomizations/Private/SlateFontInfoCustomization.cpp
@@ -102,6 +102,17 @@ void FSlateFontInfoStructCustomization::OnFontChanged(const FAssetData& InAssetD
 		// We've changed (or cleared) the font asset, so make sure and update the typeface entry name being used by the font info
 		TypefaceFontNameProperty->SetValue(FirstFontName);
 	}
+
+	if(!FontAsset)
+	{
+		const FString PropertyPath = FontObjectProperty->GeneratePathToProperty();
+		TArray<UObject*> PropertyOuterObjects;
+		FontObjectProperty->GetOuterObjects(PropertyOuterObjects);
+		for(const UObject* OuterObject : PropertyOuterObjects)
+		{
+			UE_LOG(LogSlate, Warning, TEXT("FSlateFontInfo property '%s' on object '%s' was set to use a null UFont. Slate will be forced to use the fallback font path which may be slower."), *PropertyPath, *OuterObject->GetPathName());
+		}
+	}
 }
 
 bool FSlateFontInfoStructCustomization::IsFontEntryComboEnabled() const
diff --git a/Engine/Source/Runtime/Engine/Classes/Engine/Font.h b/Engine/Source/Runtime/Engine/Classes/Engine/Font.h
--- a/Engine/Source/Runtime/Engine/Classes/Engine/Font.h
+++ b/Engine/Source/Runtime/Engine/Classes/Engine/Font.h
@@ -182,7 +182,7 @@ public:
 	/** IFontProviderInterface */
 	virtual const FCompositeFont* GetCompositeFont() const override
 	{
-		return &CompositeFont;
+		return (FontCacheType == EFontCacheType::Runtime) ? &CompositeFont : nullptr;
 	}
 
 	/** Get the info needed to use this UFont with Slate, using the fallback data for legacy Canvas APIs */
diff --git a/Engine/Source/Runtime/SlateCore/Private/Fonts/SlateFontInfo.cpp b/Engine/Source/Runtime/SlateCore/Private/Fonts/SlateFontInfo.cpp
--- a/Engine/Source/Runtime/SlateCore/Private/Fonts/SlateFontInfo.cpp
+++ b/Engine/Source/Runtime/SlateCore/Private/Fonts/SlateFontInfo.cpp
@@ -27,6 +27,10 @@ FSlateFontInfo::FSlateFontInfo( TSharedPtr<const FCompositeFont> InCompositeFont
 	, FontName_DEPRECATED()
 	, Hinting_DEPRECATED(EFontHinting::Default)
 {
+	if (!InCompositeFont.IsValid())
+	{
+		UE_LOG(LogSlate, Warning, TEXT("FSlateFontInfo was constructed with a null FCompositeFont. Slate will be forced to use the fallback font path which may be slower."));
+	}
 }
 
 
@@ -38,6 +42,18 @@ FSlateFontInfo::FSlateFontInfo( const UObject* InFontObject, const int32 InSize,
 	, FontName_DEPRECATED()
 	, Hinting_DEPRECATED(EFontHinting::Default)
 {
+	if (InFontObject)
+	{
+		const IFontProviderInterface* FontProvider = Cast<const IFontProviderInterface>(InFontObject);
+		if (!FontProvider || !FontProvider->GetCompositeFont())
+		{
+			UE_LOG(LogSlate, Warning, TEXT("'%s' does not provide a composite font that can be used with Slate. Slate will be forced to use the fallback font path which may be slower."), *InFontObject->GetName());
+		}
+	}
+	else
+	{
+		UE_LOG(LogSlate, Warning, TEXT("FSlateFontInfo was constructed with a null UFont. Slate will be forced to use the fallback font path which may be slower."));
+	}
 }
 
 
@@ -112,7 +128,8 @@ const FCompositeFont* FSlateFontInfo::GetCompositeFont() const
 	const IFontProviderInterface* FontProvider = Cast<const IFontProviderInterface>(FontObject);
 	if (FontProvider)
 	{
-		return FontProvider->GetCompositeFont();
+		const FCompositeFont* const ProvidedCompositeFont = FontProvider->GetCompositeFont();
+		return (ProvidedCompositeFont) ? ProvidedCompositeFont : FLegacySlateFontInfoCache::Get().GetFallbackFont().Get();
 	}
 
 	if (CompositeFont.IsValid())
