Slate: Expose a variant of SWidget::FindChildUnderMouse that takes the position to check directly instead of the pointer event

diff --git a/Engine/Source/Runtime/SlateCore/Private/Widgets/SWidget.cpp b/Engine/Source/Runtime/SlateCore/Private/Widgets/SWidget.cpp
--- a/Engine/Source/Runtime/SlateCore/Private/Widgets/SWidget.cpp
+++ b/Engine/Source/Runtime/SlateCore/Private/Widgets/SWidget.cpp
@@ -507,13 +507,19 @@ FGeometry SWidget::FindChildGeometry( const FGeometry& MyGeometry, TSharedRef<SW
 int32 SWidget::FindChildUnderMouse( const FArrangedChildren& Children, const FPointerEvent& MouseEvent )
 {
 	const FVector2D& AbsoluteCursorLocation = MouseEvent.GetScreenSpacePosition();
+	return SWidget::FindChildUnderPosition( Children, AbsoluteCursorLocation );
+}
+
+
+int32 SWidget::FindChildUnderPosition( const FArrangedChildren& Children, const FVector2D& ArrangedSpacePosition )
+{
 	const int32 NumChildren = Children.Num();
 	for( int32 ChildIndex=NumChildren-1; ChildIndex >= 0; --ChildIndex )
 	{
 		const FArrangedWidget& Candidate = Children[ChildIndex];
 		const bool bCandidateUnderCursor = 
 			// Candidate is physically under the cursor
-			Candidate.Geometry.IsUnderLocation( AbsoluteCursorLocation );
+			Candidate.Geometry.IsUnderLocation( ArrangedSpacePosition );
 
 		if (bCandidateUnderCursor)
 		{
diff --git a/Engine/Source/Runtime/SlateCore/Public/Widgets/SWidget.h b/Engine/Source/Runtime/SlateCore/Public/Widgets/SWidget.h
--- a/Engine/Source/Runtime/SlateCore/Public/Widgets/SWidget.h
+++ b/Engine/Source/Runtime/SlateCore/Public/Widgets/SWidget.h
@@ -841,6 +841,9 @@ protected:
 	/** @return The index of the child that the mouse is currently hovering */
 	static int32 FindChildUnderMouse( const FArrangedChildren& Children, const FPointerEvent& MouseEvent );
 
+	/** @return The index of the child that is under the specified position */
+	static int32 FindChildUnderPosition(const FArrangedChildren& Children, const FVector2D& ArrangedSpacePosition);
+
 	/** 
 	 * Determines if this widget should be enabled.
 	 * 
