Temp fix for a mismatch between the order of windows on the Slate's list and the actual order of windows on screen caused by activating a panel window.

diff --git a/Engine/Source/Runtime/SlateCore/Private/Application/SlateWindowHelper.cpp b/Engine/Source/Runtime/SlateCore/Private/Application/SlateWindowHelper.cpp
--- a/Engine/Source/Runtime/SlateCore/Private/Application/SlateWindowHelper.cpp
+++ b/Engine/Source/Runtime/SlateCore/Private/Application/SlateWindowHelper.cpp
@@ -43,7 +43,10 @@ void FSlateWindowHelper::ArrangeWindowToFront( TArray< TSharedRef<SWindow> >& Wi
 void FSlateWindowHelper::BringWindowToFront( TArray<TSharedRef<SWindow>>& Windows, const TSharedRef<SWindow>& BringMeToFront )
 {
 	const TSharedRef<SWindow> TopLeveWindowToReorder = BringToFrontInParent(BringMeToFront);
-	FSlateWindowHelper::ArrangeWindowToFront(Windows, TopLeveWindowToReorder);
+	if (!PLATFORM_MAC || TopLeveWindowToReorder == BringMeToFront)
+	{
+		FSlateWindowHelper::ArrangeWindowToFront(Windows, TopLeveWindowToReorder);
+	}
 }
 
 
