Fix for hit test grid not navigating correctly when the primary monitor was to the right of a secondary monitor.
Resulted in an non zero grid offset which wasn't handled properly in navigations clip rect intersection code.
Fixes Jira UE-8052

diff --git a/Engine/Source/Runtime/SlateCore/Private/Input/HittestGrid.cpp b/Engine/Source/Runtime/SlateCore/Private/Input/HittestGrid.cpp
--- a/Engine/Source/Runtime/SlateCore/Private/Input/HittestGrid.cpp
+++ b/Engine/Source/Runtime/SlateCore/Private/Input/HittestGrid.cpp
@@ -385,7 +385,7 @@ TSharedPtr<SWidget> FHittestGrid::FindFocusableWidget(FSlateRect WidgetRect, con
 				check(CurrentIndex < WidgetsCachedThisFrame->Num());
 
 				const FCachedWidget& TestCandidate = (*WidgetsCachedThisFrame)[CurrentIndex];
-				FSlateRect TestCandidateRect = TestCandidate.ClippingRect;
+				FSlateRect TestCandidateRect = TestCandidate.ClippingRect.OffsetBy(-GridOrigin);
 
 				if (CompareFunc(DestSideFunc(TestCandidateRect), CurrentSourceSide) && FSlateRect::DoRectanglesIntersect(SweptRect, TestCandidateRect))
 				{
