UMG/Slate - No longer exposing Disabled slate brush on the FButtonStyle.  This was adding a lot of confusion.  We want to support custom look and feel disabled buttons, but this particular property wasn't actually being used by any of the core slate widgets and so users were becoming confused why it didn't work as expected.  Hiding from view until we have time to refactor disabled styling for buttons and more.

diff --git a/Engine/Source/Runtime/SlateCore/Public/Styling/SlateTypes.h b/Engine/Source/Runtime/SlateCore/Public/Styling/SlateTypes.h
--- a/Engine/Source/Runtime/SlateCore/Public/Styling/SlateTypes.h
+++ b/Engine/Source/Runtime/SlateCore/Public/Styling/SlateTypes.h
@@ -261,8 +261,8 @@ struct SLATECORE_API FButtonStyle : public FSlateWidgetStyle
 	FSlateBrush Pressed;
 	FButtonStyle& SetPressed( const FSlateBrush& InPressed ){ Pressed = InPressed; return *this; }
 
-	/** Button appearance when disabled */
-	UPROPERTY(EditAnywhere, BlueprintReadWrite, Category=Appearance)
+	/** Button appearance when disabled, not exposed as a UProperty, because they don't have an effect in the general case. */
+	UPROPERTY()
 	FSlateBrush Disabled;
 	FButtonStyle& SetDisabled( const FSlateBrush& InDisabled ){ Disabled = InDisabled; return *this; }
 
