Slate : Hittest grid no longer discards entire event path when a single widget in that path is broken. Instead, return a truncated path from root.
Fixes UE-13401.

diff --git a/Engine/Source/Runtime/SlateCore/Private/Input/HittestGrid.cpp b/Engine/Source/Runtime/SlateCore/Private/Input/HittestGrid.cpp
--- a/Engine/Source/Runtime/SlateCore/Private/Input/HittestGrid.cpp
+++ b/Engine/Source/Runtime/SlateCore/Private/Input/HittestGrid.cpp
@@ -705,7 +705,6 @@ TArray<FWidgetAndPointer> FHittestGrid::GetBubblePathFromHitIndex(const int32 Hi
 	if (WidgetsCachedThisFrame->IsValidIndex(HitIndex))
 	{
 		int32 CurWidgetIndex = HitIndex;
-		bool bPathUninterrupted = false;
 		do
 		{
 			check(CurWidgetIndex < WidgetsCachedThisFrame->Num());
@@ -713,21 +712,23 @@ TArray<FWidgetAndPointer> FHittestGrid::GetBubblePathFromHitIndex(const int32 Hi
 			const TSharedPtr<SWidget> CachedWidgetPtr = CurCachedWidget.WidgetPtr.Pin();
 
 
-			bPathUninterrupted = CachedWidgetPtr.IsValid();
-			if (bPathUninterrupted)
+			const bool bPathInterrupted = CachedWidgetPtr.IsValid();
+			if ( bPathInterrupted )
+			{
+				// A widget in the path to the root has been removed, so anything
+				// we thought we had hittest so far is no longer actually in the hierarchy.
+				// Continue bubbling to the root of the hirarchy to find an unbroken chain to root.
+				// The leafmost widget in that chain will get first shot at the events.
+				BubblePath.Reset();
+			}
+			else
 			{
 				BubblePath.Insert(FWidgetAndPointer(FArrangedWidget(CachedWidgetPtr.ToSharedRef(), CurCachedWidget.CachedGeometry), TSharedPtr<FVirtualPointerPosition>()), 0);
-				CurWidgetIndex = CurCachedWidget.ParentIndex;
 			}
-		} while (CurWidgetIndex != INDEX_NONE && bPathUninterrupted);
-
-		if (!bPathUninterrupted)
-		{
-			// A widget in the path to the root has been removed, so anything
-			// we thought we had hittest is no longer actually there.
-			// Pretend we didn't hit anything.
-			BubblePath.Reset();
+			CurWidgetIndex = CurCachedWidget.ParentIndex;
 		}
+		while (CurWidgetIndex != INDEX_NONE);
+
 
 		// Disabling a widget disables all of its logical children
 		// This effect is achieved by truncating the path to the
