Slate : fixed typo.

diff --git a/Engine/Source/Runtime/SlateCore/Private/Input/HittestGrid.cpp b/Engine/Source/Runtime/SlateCore/Private/Input/HittestGrid.cpp
--- a/Engine/Source/Runtime/SlateCore/Private/Input/HittestGrid.cpp
+++ b/Engine/Source/Runtime/SlateCore/Private/Input/HittestGrid.cpp
@@ -712,7 +712,7 @@ TArray<FWidgetAndPointer> FHittestGrid::GetBubblePathFromHitIndex(const int32 Hi
 			const TSharedPtr<SWidget> CachedWidgetPtr = CurCachedWidget.WidgetPtr.Pin();
 
 
-			const bool bPathInterrupted = CachedWidgetPtr.IsValid();
+			const bool bPathInterrupted = !CachedWidgetPtr.IsValid();
 			if ( bPathInterrupted )
 			{
 				// A widget in the path to the root has been removed, so anything
