Added constructor to FSlateTextureData that takes a pointer to texture data

diff --git a/Engine/Source/Runtime/SlateCore/Public/Textures/SlateTextureData.h b/Engine/Source/Runtime/SlateCore/Public/Textures/SlateTextureData.h
--- a/Engine/Source/Runtime/SlateCore/Public/Textures/SlateTextureData.h
+++ b/Engine/Source/Runtime/SlateCore/Public/Textures/SlateTextureData.h
@@ -18,6 +18,27 @@ struct SLATECORE_API FSlateTextureData
 		INC_MEMORY_STAT_BY( STAT_SlateTextureDataMemory, Bytes.GetAllocatedSize() );
 	}
 
+	/**
+	 * Constructor to create texture data by copying from a pointer instead of an array
+	 * @param InBuffer Pointer to Texture data (must contain InWidth*InHeight*InBytesPerPixel bytes).
+	 * @param InWidth Width of the Texture.
+	 * @param InHeight Height of the Texture.
+	 * @param InBytesPerPixel Bytes per pixel of the Texture.
+	 */
+	FSlateTextureData( const uint8* InBuffer, uint32 InWidth, uint32 InHeight, uint32 InBytesPerPixel )
+		: Width(InWidth)
+		, Height(InHeight)
+		, BytesPerPixel(InBytesPerPixel)
+	{
+		const uint32 BufferSize = Width*Height*BytesPerPixel;
+		Bytes.SetNumUninitialized(BufferSize);
+		if (InBuffer != nullptr)
+		{
+			FMemory::Memcpy(Bytes.GetData(), InBuffer, BufferSize);
+		}
+		INC_MEMORY_STAT_BY(STAT_SlateTextureDataMemory, Bytes.GetAllocatedSize());
+	}
+
 	FSlateTextureData( const FSlateTextureData &Other )
 		: Bytes( Other.Bytes )
 		, Width( Other.Width )
