Fix the inconsistency in the ownership hierarchy of 3D widgets

#change Allow ArrangeChildren() calls inside FWidgetPath::GeneratePathToWidget() to include 3D widgets.

diff --git a/Engine/Source/Runtime/SlateCore/Public/Layout/WidgetPath.inl b/Engine/Source/Runtime/SlateCore/Public/Layout/WidgetPath.inl
--- a/Engine/Source/Runtime/SlateCore/Public/Layout/WidgetPath.inl
+++ b/Engine/Source/Runtime/SlateCore/Public/Layout/WidgetPath.inl
@@ -17,7 +17,8 @@
 template<typename MatchRuleType>
 bool FWidgetPath::SearchForWidgetRecursively( const MatchRuleType& MatchRule, const FArrangedWidget& InCandidate, FArrangedChildren& OutReversedPath, EVisibility VisibilityFilter )
 {
-	FArrangedChildren ArrangedChildren(VisibilityFilter);
+	const bool bAllow3DWidgets = true;
+	FArrangedChildren ArrangedChildren(VisibilityFilter, bAllow3DWidgets);
 	InCandidate.Widget->ArrangeChildren( InCandidate.Geometry, ArrangedChildren );
 
 	for( int32 ChildIndex = 0; ChildIndex < ArrangedChildren.Num(); ++ChildIndex )
@@ -42,7 +43,8 @@ bool FWidgetPath::SearchForWidgetRecursively( const MatchRuleType& MatchRule, co
 template<typename MatchRuleType>
 bool FWidgetPath::SearchForWidgetRecursively_Reverse( const MatchRuleType& MatchRule, const FArrangedWidget& InCandidate, FArrangedChildren& OutReversedPath, EVisibility VisibilityFilter )
 {
-	FArrangedChildren ArrangedChildren(VisibilityFilter);
+	const bool bAllow3DWidgets = true;
+	FArrangedChildren ArrangedChildren(VisibilityFilter, bAllow3DWidgets);
 	InCandidate.Widget->ArrangeChildren( InCandidate.Geometry, ArrangedChildren );
 
 	for( int32 ChildIndex = ArrangedChildren.Num()-1; ChildIndex >=0 ; --ChildIndex )
