Slate mouse input now works with Application Scales < 1.0. Previously mouse input was erroneously clipped.

diff --git a/Engine/Source/Runtime/SlateCore/Private/Widgets/SWindow.cpp b/Engine/Source/Runtime/SlateCore/Private/Widgets/SWindow.cpp
--- a/Engine/Source/Runtime/SlateCore/Private/Widgets/SWindow.cpp
+++ b/Engine/Source/Runtime/SlateCore/Private/Widgets/SWindow.cpp
@@ -1238,7 +1238,7 @@ bool SWindow::IsTopmostWindow() const
 /** @return true if mouse coordinates is within this window */
 bool SWindow::IsScreenspaceMouseWithin(FVector2D ScreenspaceMouseCoordinate) const
 {
-	const FVector2D LocalMouseCoordinate = GetWindowGeometryInScreen().AbsoluteToLocal(ScreenspaceMouseCoordinate);
+	const FVector2D LocalMouseCoordinate = ScreenspaceMouseCoordinate - GetWindowGeometryInScreen().AbsolutePosition;
 	return NativeWindow->IsPointInWindow(FMath::TruncToInt(LocalMouseCoordinate.X), FMath::TruncToInt(LocalMouseCoordinate.Y));
 }
 
