Viewports can now request not to be scaled by overriding ISlateViewport::AllowScaling.

Merging CL2531963+2534462 from LauncherDev to Main (using UE4-To-UE4-LauncherDev)

diff --git a/Engine/Source/Runtime/SlateCore/Private/Rendering/ElementBatcher.cpp b/Engine/Source/Runtime/SlateCore/Private/Rendering/ElementBatcher.cpp
--- a/Engine/Source/Runtime/SlateCore/Private/Rendering/ElementBatcher.cpp
+++ b/Engine/Source/Runtime/SlateCore/Private/Rendering/ElementBatcher.cpp
@@ -1166,6 +1166,12 @@ void FSlateElementBatcher::AddViewportElement( const FSlateDrawElement& DrawElem
 	FVector2D TopLeft(0,0);
 	FVector2D BotRight(LocalSize);
 
+	// If the viewport disallows scaling, force size to current texture size.
+	if (ViewportResource != nullptr && !ViewportPin->AllowScaling())
+	{
+		BotRight = FVector2D(ViewportResource->GetWidth(), ViewportResource->GetHeight());
+	}
+
 	FVector2D TopRight = FVector2D( BotRight.X, TopLeft.Y);
 	FVector2D BotLeft =	 FVector2D( TopLeft.X, BotRight.Y);
 
diff --git a/Engine/Source/Runtime/SlateCore/Public/Rendering/RenderingCommon.h b/Engine/Source/Runtime/SlateCore/Public/Rendering/RenderingCommon.h
--- a/Engine/Source/Runtime/SlateCore/Public/Rendering/RenderingCommon.h
+++ b/Engine/Source/Runtime/SlateCore/Public/Rendering/RenderingCommon.h
@@ -251,6 +251,14 @@ public:
 	 */
 	virtual bool RequiresVsync() const = 0;
 
+	/**
+	 * Whether the viewport contents should be scaled or not. Defaults to true.
+	 */
+	virtual bool AllowScaling() const
+	{
+		return true;
+	}
+
 	/**
 	 * Called when Slate needs to know what the mouse cursor should be.
 	 * 
