Slate windows now take into account required border/tiitle bar sizes when resizing the window. This ensures that the client area stays at the requested size

diff --git a/Engine/Source/Editor/UnrealEd/Private/PlayLevel.cpp b/Engine/Source/Editor/UnrealEd/Private/PlayLevel.cpp
--- a/Engine/Source/Editor/UnrealEd/Private/PlayLevel.cpp
+++ b/Engine/Source/Editor/UnrealEd/Private/PlayLevel.cpp
@@ -2882,7 +2882,10 @@ UGameInstance* UEditorEngine::CreatePIEGameInstance(int32 PIEInstance, bool bInS
 				ViewportClient->Viewport->SetPlayInEditorViewport( ViewportClient->bIsPlayInEditorViewport );
 
 				// Ensure the window has a valid size before calling BeginPlay
-				SlatePlayInEditorSession.SlatePlayInEditorWindowViewport->ResizeFrame( PieWindow->GetSizeInScreen().X, PieWindow->GetSizeInScreen().Y, EWindowMode::Windowed, PieWindow->GetPositionInScreen().X, PieWindow->GetPositionInScreen().Y );
+				SlatePlayInEditorSession.SlatePlayInEditorWindowViewport->ResizeFrame( NewWindowWidth, NewWindowHeight, EWindowMode::Windowed, PieWindow->GetPositionInScreen().X, PieWindow->GetPositionInScreen().Y );
+
+				// Change the system resolution to match our window, to make sure game and slate window are kept syncronised
+				FSystemResolution::RequestResolutionChange(NewWindowWidth, NewWindowHeight, EWindowMode::Windowed);
 
 				if (bUseVRPreviewForPlayWorld && GEngine->HMDDevice.IsValid())
 				{
diff --git a/Engine/Source/Runtime/SlateCore/Private/Widgets/SWindow.cpp b/Engine/Source/Runtime/SlateCore/Private/Widgets/SWindow.cpp
--- a/Engine/Source/Runtime/SlateCore/Private/Widgets/SWindow.cpp
+++ b/Engine/Source/Runtime/SlateCore/Private/Widgets/SWindow.cpp
@@ -176,6 +176,23 @@ private:
 	TWeakPtr<SWindow> OwnerWindow;
 };
 
+FVector2D SWindow::GetWindowSizeFromClientSize(FVector2D InClientSize)
+{
+	if (!HasOSWindowBorder())
+	{
+		const FMargin BorderSize = GetWindowBorderSize();
+
+		InClientSize.X += BorderSize.Left + BorderSize.Right;
+		InClientSize.Y += BorderSize.Bottom + BorderSize.Top;
+
+		if (bCreateTitleBar)
+		{
+			InClientSize.Y += SWindowDefs::DefaultTitleBarSize;
+		}
+	}
+
+	return InClientSize;
+}
 
 void SWindow::Construct(const FArguments& InArgs)
 {
@@ -206,22 +223,10 @@ void SWindow::Construct(const FArguments& InArgs)
 		.SetMaxHeight(InArgs._MaxHeight);
 	
 	// calculate window size from client size
-	const bool bCreateTitleBar = InArgs._CreateTitleBar && !bIsPopupWindow && !bIsCursorDecoratorWindow && !bHasOSWindowBorder;
-	FVector2D WindowSize = InArgs._ClientSize;
-
+	bCreateTitleBar = InArgs._CreateTitleBar && !bIsPopupWindow && !bIsCursorDecoratorWindow && !bHasOSWindowBorder;
+	
 	// If the window has no OS border, simulate it ourselves, enlarging window by the size that OS border would have.
-	if (!HasOSWindowBorder())
-	{
-		const FMargin BorderSize = GetWindowBorderSize();
-
-		WindowSize.X += BorderSize.Left + BorderSize.Right;
-		WindowSize.Y += BorderSize.Bottom + BorderSize.Top;
-
-		if (bCreateTitleBar)
-		{
-			WindowSize.Y += SWindowDefs::DefaultTitleBarSize;
-		}
-	}
+	FVector2D WindowSize = GetWindowSizeFromClientSize(InArgs._ClientSize);
 
 	// calculate initial window position
 	FVector2D WindowPosition = InArgs._ScreenPosition;
@@ -759,6 +764,9 @@ void SWindow::ReshapeWindow( const FSlateRect& InNewShape )
 void SWindow::Resize( FVector2D NewSize )
 {
 	Morpher.Sequence.JumpToEnd();
+
+	NewSize = GetWindowSizeFromClientSize(NewSize);
+
 	if ( Size != NewSize )
 	{
 		if (NativeWindow.IsValid())
diff --git a/Engine/Source/Runtime/SlateCore/Public/Widgets/SWindow.h b/Engine/Source/Runtime/SlateCore/Public/Widgets/SWindow.h
--- a/Engine/Source/Runtime/SlateCore/Public/Widgets/SWindow.h
+++ b/Engine/Source/Runtime/SlateCore/Public/Widgets/SWindow.h
@@ -624,6 +624,9 @@ private:
 
 	virtual float GetRelativeLayoutScale(const FSlotBase& Child) const override;
 
+	// For a given client size, calculate the window size required to accomodate any potential non-OS borders and tilebars
+	FVector2D GetWindowSizeFromClientSize(FVector2D InClientSize);
+
 public:
 	/** @return true if this window will be focused when it is first shown */
 	inline bool IsFocusedInitially() const
@@ -759,6 +762,9 @@ protected:
 	/** Transparency setting for this window */
 	EWindowTransparency TransparencySupport;
 
+	/** True if this window has a title bar */
+	bool bCreateTitleBar : 1;
+
 	/** True if this is a pop up window */
 	bool bIsPopupWindow : 1;
 
