Fix build error about shadowing bCreateTitleBar in SWindow

diff --git a/Engine/Source/Runtime/SlateCore/Private/Widgets/SWindow.cpp b/Engine/Source/Runtime/SlateCore/Private/Widgets/SWindow.cpp
--- a/Engine/Source/Runtime/SlateCore/Private/Widgets/SWindow.cpp
+++ b/Engine/Source/Runtime/SlateCore/Private/Widgets/SWindow.cpp
@@ -302,7 +302,7 @@ void SWindow::Construct(const FArguments& InArgs)
 	// Window visibility is currently driven by whether the window is interactive.
 	this->Visibility = TAttribute<EVisibility>::Create( TAttribute<EVisibility>::FGetter::CreateRaw(this, &SWindow::GetWindowVisibility) );
 
-	this->ConstructWindowInternals( bCreateTitleBar );
+	this->ConstructWindowInternals();
 	this->SetContent( InArgs._Content.Widget );
 }
 
@@ -375,7 +375,7 @@ FVector2D SWindow::ComputeWindowSizeForContent( FVector2D ContentSize )
 	return ContentSize + FVector2D(0, SWindowDefs::DefaultTitleBarSize);
 }
 
-void SWindow::ConstructWindowInternals( const bool bCreateTitleBar )
+void SWindow::ConstructWindowInternals()
 {
 	ForegroundColor = FCoreStyle::Get().GetSlateColor("DefaultForeground");
 
diff --git a/Engine/Source/Runtime/SlateCore/Public/Widgets/SWindow.h b/Engine/Source/Runtime/SlateCore/Public/Widgets/SWindow.h
--- a/Engine/Source/Runtime/SlateCore/Public/Widgets/SWindow.h
+++ b/Engine/Source/Runtime/SlateCore/Public/Widgets/SWindow.h
@@ -951,7 +951,7 @@ private:
 	
 	virtual
 
-	void ConstructWindowInternals( const bool bCreateTitleBar );
+	void ConstructWindowInternals();
 
 	/**
 	 * @return EVisibility::Visible if we are showing this viewports content.  EVisibility::Hidden otherwise (we hide the content during full screen overlays)
