Skip adding/removing Resource object in FSlateDynamicImageBrush to/from root set if object is already there.

diff --git a/Engine/Source/Runtime/SlateCore/Private/Brushes/SlateDynamicImageBrush.cpp b/Engine/Source/Runtime/SlateCore/Private/Brushes/SlateDynamicImageBrush.cpp
--- a/Engine/Source/Runtime/SlateCore/Private/Brushes/SlateDynamicImageBrush.cpp
+++ b/Engine/Source/Runtime/SlateCore/Private/Brushes/SlateDynamicImageBrush.cpp
@@ -34,7 +34,7 @@ FSlateDynamicImageBrush::~FSlateDynamicImageBrush( )
 	if (FSlateApplicationBase::IsInitialized())
 	{
 		// Brush resource is no longer referenced by this object
-		if( ResourceObject )
+		if (ResourceObject && bRemoveResourceFromRootSet)
 		{
 			ResourceObject->RemoveFromRoot();
 		}
diff --git a/Engine/Source/Runtime/SlateCore/Public/Brushes/SlateDynamicImageBrush.h b/Engine/Source/Runtime/SlateCore/Public/Brushes/SlateDynamicImageBrush.h
--- a/Engine/Source/Runtime/SlateCore/Public/Brushes/SlateDynamicImageBrush.h
+++ b/Engine/Source/Runtime/SlateCore/Public/Brushes/SlateDynamicImageBrush.h
@@ -25,6 +25,7 @@ struct SLATECORE_API FSlateDynamicImageBrush
 		ESlateBrushImageType::Type InImageType = ESlateBrushImageType::FullColor
 	)
 		: FSlateBrush(ESlateBrushDrawType::Image, FName(TEXT("None")), FMargin(0.0f), InTiling, InImageType, InImageSize, InTint, (UObject*)InTexture)
+		, bRemoveResourceFromRootSet(false)
 	{
 		bIsDynamicallyLoaded = true;
 
@@ -32,8 +33,14 @@ struct SLATECORE_API FSlateDynamicImageBrush
 		if (ResourceObject != nullptr)
 		{
 			// @todo Slate - Hack:  This is to address an issue where the brush created and a GC occurs before the brush resource object becomes referenced
-			// by the Slate resource manager.  
-			ResourceObject->AddToRoot();
+			// by the Slate resource manager. Don't add objects that are already in root set (and mark them as such) to avoid incorrect removing objects
+			// from root set in destructor.
+			if (!ResourceObject->HasAllFlags(RF_RootSet))
+			{
+				ResourceObject->AddToRoot();
+				bRemoveResourceFromRootSet = true;
+			}
+
 			ResourceName = InTextureName;
 		}
 	}
@@ -53,6 +60,7 @@ struct SLATECORE_API FSlateDynamicImageBrush
 		ESlateBrushImageType::Type InImageType = ESlateBrushImageType::FullColor
 	)
 		: FSlateBrush(ESlateBrushDrawType::Image, InTextureName, FMargin(0.0f), InTiling, InImageType, InImageSize, InTint, nullptr, true)
+		, bRemoveResourceFromRootSet(false)
 	{
 		bIsDynamicallyLoaded = true;
 	}
@@ -75,4 +83,8 @@ struct SLATECORE_API FSlateDynamicImageBrush
 
 	/** Destructor. */
 	virtual ~FSlateDynamicImageBrush();
+
+private:
+	// Tracks if Resource was in root set to avoid unnecessary removing it from there.
+	bool bRemoveResourceFromRootSet;
 };
