Fixed slate application focus not being applied correctly when activating windows

There are 2 issues at play here:
    * When setting the keyboard focus to a particular widget, it was possible for a widget to forward this focus onto another widget (inside SWidget::OnFocusReceived, called from FSlateApplication::SetUserFocus). If this handler returned an FReply that also set the user focus, the window's WidgetToFocusOnActivate would end up not being set to the correct widget. We now set WidgetToFocusOnActivate *before* invoking the focus received handler to ensure that the widget is set in the correct order.

    * WidgetToFocusOnActivate is now entirely managed by SWindow, and handled only through window activation, rather than forcing keyboard focus to the window, and having the OnReceivedFocus handler focussing WidgetToFocusOnActivate. This prevents widgets temporarily losing focus when windows are activated (particularly things like editable text blocks that will close edit operations when focus is lost).

This addresses UE-14872 - Unable to name Variables upon creation

diff --git a/Engine/Source/Runtime/Slate/Private/Framework/Application/SlateApplication.cpp b/Engine/Source/Runtime/Slate/Private/Framework/Application/SlateApplication.cpp
--- a/Engine/Source/Runtime/Slate/Private/Framework/Application/SlateApplication.cpp
+++ b/Engine/Source/Runtime/Slate/Private/Framework/Application/SlateApplication.cpp
@@ -2204,6 +2204,15 @@ bool FSlateApplication::SetUserFocus(const uint32 InUserIndex, const FWidgetPath
 	{
 		TSharedPtr<SWindow> FocusedWindow = NewFocusedWidgetPath.GetWindow();
 		
+		// Set the windows restore state - we do this before calling OnFocusReceived to ensure that if
+		// OnFocusReceived returns a new UserFocus, the WidgetToFocusOnActivate will get set correctly
+		//@Todo Slate: We need to store the full focus state
+		//@Todo Slate: Why are we checking FocusedWindow != NewFocusedWidget?
+		if (FocusedWindow.IsValid() && FocusedWindow != NewFocusedWidget)
+		{
+			FocusedWindow->SetWidgetToFocusOnActivate(NewFocusedWidget);
+		}
+
 		// Switch worlds for widgets in the new path
 		FScopedSwitchWorldHack SwitchWorld(FocusedWindow);
 		
@@ -2214,14 +2223,6 @@ bool FSlateApplication::SetUserFocus(const uint32 InUserIndex, const FWidgetPath
 		{
 			ProcessReply(InFocusPath, Reply, nullptr, nullptr);
 		}
-		
-		// Set the windows restore state.
-		//@Todo Slate: We need to store the full focus state 
-		//@Todo Slate: Why are we checking FocusedWindow != NewFocusedWidget?
-		if (FocusedWindow.IsValid() && FocusedWindow != NewFocusedWidget)
-		{
-			FocusedWindow->SetWidgetToFocusOnActivate(NewFocusedWidget);
-		}
 	}
 
 	return true;
@@ -5064,24 +5065,6 @@ bool FSlateApplication::ProcessWindowActivatedEvent( const FWindowActivateEvent&
 			// A Slate window was activated
 			bSlateWindowActive = true;
 
-			if (ActivateEvent.GetAffectedWindow()->IsFocusedInitially() && ActivateEvent.GetAffectedWindow()->SupportsKeyboardFocus() )
-			{
-				// Set keyboard focus on the window being activated.
-				{
-					FWidgetPath PathToWindowBeingActivated;
-					GeneratePathToWidgetChecked( ActivateEvent.GetAffectedWindow(), PathToWindowBeingActivated );
-
-					if( ActivateEvent.GetActivationType() == FWindowActivateEvent::EA_ActivateByMouse )
-					{
-						SetKeyboardFocus(PathToWindowBeingActivated, EFocusCause::Mouse);
-					}
-					else
-					{
-						SetKeyboardFocus(PathToWindowBeingActivated, EFocusCause::WindowActivate);
-					}
-				}
-
-			}
 
 			{
 				FScopedSwitchWorldHack SwitchWorld( ActivateEvent.GetAffectedWindow() );
diff --git a/Engine/Source/Runtime/SlateCore/Private/Widgets/SWindow.cpp b/Engine/Source/Runtime/SlateCore/Private/Widgets/SWindow.cpp
--- a/Engine/Source/Runtime/SlateCore/Private/Widgets/SWindow.cpp
+++ b/Engine/Source/Runtime/SlateCore/Private/Widgets/SWindow.cpp
@@ -1321,6 +1321,34 @@ bool SWindow::OnIsActiveChanged( const FWindowActivateEvent& ActivateEvent )
 	}
 	else
 	{
+		if (SupportsKeyboardFocus())
+		{
+			TArray< TSharedRef<SWindow> > JustThisWindow;
+			JustThisWindow.Add( SharedThis(this) );
+			
+			// If we're becoming active and we were set to restore keyboard focus to a specific widget
+			// after reactivating, then do so now
+			TSharedPtr< SWidget > PinnedWidgetToFocus( WidgetToFocusOnActivate.Pin() );
+			
+			if (PinnedWidgetToFocus.IsValid())
+			{
+				FWidgetPath WidgetToFocusPath;
+				if( FSlateWindowHelper::FindPathToWidget( JustThisWindow, PinnedWidgetToFocus.ToSharedRef(), WidgetToFocusPath ) )
+				{
+					FSlateApplicationBase::Get().SetKeyboardFocus( WidgetToFocusPath, EFocusCause::SetDirectly );
+				}
+			}
+			else
+			{
+				FWidgetPath WindowWidgetPath;
+				if( FSlateWindowHelper::FindPathToWidget( JustThisWindow, AsShared(), WindowWidgetPath ) )
+				{
+					FWeakWidgetPath WeakWindowPath(WindowWidgetPath);
+					FSlateApplicationBase::Get().SetKeyboardFocus( WeakWindowPath.ToNextFocusedPath(EUINavigation::Next), EFocusCause::SetDirectly );
+				}
+			}
+		}
+
 		OnWindowActivated.ExecuteIfBound();
 	}
 
@@ -1364,36 +1392,6 @@ bool SWindow::SupportsKeyboardFocus() const
 
 FReply SWindow::OnFocusReceived(const FGeometry& MyGeometry, const FFocusEvent& InFocusEvent)
 {
-	if (InFocusEvent.GetCause() == EFocusCause::WindowActivate || FSlateApplicationBase::Get().IsExternalUIOpened())
-	{
-		TArray< TSharedRef<SWindow> > JustThisWindow;
-		{
-			JustThisWindow.Add( SharedThis(this) );
-		}
-		
-		// If we're becoming active and we were set to restore keyboard focus to a specific widget
-		// after reactivating, then do so now
-		TSharedPtr< SWidget > PinnedWidgetToFocusOnActivate( WidgetToFocusOnActivate.Pin() );
-		
-		if(PinnedWidgetToFocusOnActivate.IsValid())
-		{
-			FWidgetPath WidgetToFocusPath;
-			if( FSlateWindowHelper::FindPathToWidget( JustThisWindow, PinnedWidgetToFocusOnActivate.ToSharedRef(), WidgetToFocusPath ) )
-			{
-				FSlateApplicationBase::Get().SetKeyboardFocus( WidgetToFocusPath, EFocusCause::SetDirectly );
-			}
-		}
-		else
-		{
-			FWidgetPath WindowWidgetPath;
-			if( FSlateWindowHelper::FindPathToWidget( JustThisWindow, AsShared(), WindowWidgetPath ) )
-			{
-				FWeakWidgetPath WeakWindowPath(WindowWidgetPath);
-				FSlateApplicationBase::Get().SetKeyboardFocus( WeakWindowPath.ToNextFocusedPath(EUINavigation::Next), EFocusCause::SetDirectly );
-			}
-		}
-	}
-
 	return FReply::Handled();
 }
 
