Fixed a regression caused by CL# 2537219

UE-15050 - black edge padding around message windows

We now only adjust the client size given to SWindow::Resize if we're using a window to which we've added window decorations (ie, a regular window). This avoids us inflating the size of tooltip and popup windows.

diff --git a/Engine/Source/Runtime/SlateCore/Private/Widgets/SWindow.cpp b/Engine/Source/Runtime/SlateCore/Private/Widgets/SWindow.cpp
--- a/Engine/Source/Runtime/SlateCore/Private/Widgets/SWindow.cpp
+++ b/Engine/Source/Runtime/SlateCore/Private/Widgets/SWindow.cpp
@@ -178,7 +178,9 @@ private:
 
 FVector2D SWindow::GetWindowSizeFromClientSize(FVector2D InClientSize)
 {
-	if (!HasOSWindowBorder())
+	// If this is a regular non-OS window, we need to compensate for the border and title bar area that we will add
+	// Note: Windows with an OS border do this in ReshapeWindow
+	if (IsRegularWindow() && !HasOSWindowBorder())
 	{
 		const FMargin BorderSize = GetWindowBorderSize();
 
