Added some missing slate stats and renamed a few slate stats

diff --git a/Engine/Source/Runtime/Slate/Private/Framework/Application/SlateApplication.cpp b/Engine/Source/Runtime/Slate/Private/Framework/Application/SlateApplication.cpp
--- a/Engine/Source/Runtime/Slate/Private/Framework/Application/SlateApplication.cpp
+++ b/Engine/Source/Runtime/Slate/Private/Framework/Application/SlateApplication.cpp
@@ -272,7 +272,7 @@ DECLARE_CYCLE_STAT( TEXT("Update Tooltip Time"), STAT_SlateUpdateTooltip, STATGR
 DECLARE_CYCLE_STAT( TEXT("Tick Window And Children Time"), STAT_SlateTickWindowAndChildren, STATGROUP_Slate );
 DECLARE_CYCLE_STAT( TEXT("Total Slate Tick Time"), STAT_SlateTickTime, STATGROUP_Slate );
 DECLARE_CYCLE_STAT( TEXT("SlatePrepass"), STAT_SlatePrepass, STATGROUP_Slate );
-DECLARE_CYCLE_STAT( TEXT("DrawWindows"), STAT_SlateDrawWindowTime, STATGROUP_Slate );
+DECLARE_CYCLE_STAT( TEXT("Draw Window And Children Time"), STAT_SlateDrawWindowTime, STATGROUP_Slate );
 DECLARE_CYCLE_STAT( TEXT("TickWidgets"), STAT_SlateTickWidgets, STATGROUP_Slate );
 DECLARE_CYCLE_STAT( TEXT("TickRegisteredWidgets"), STAT_SlateTickRegisteredWidgets, STATGROUP_Slate );
 
diff --git a/Engine/Source/Runtime/SlateCore/Private/Rendering/ElementBatcher.cpp b/Engine/Source/Runtime/SlateCore/Private/Rendering/ElementBatcher.cpp
--- a/Engine/Source/Runtime/SlateCore/Private/Rendering/ElementBatcher.cpp
+++ b/Engine/Source/Runtime/SlateCore/Private/Rendering/ElementBatcher.cpp
@@ -11,6 +11,8 @@ DECLARE_DWORD_COUNTER_STAT(TEXT("Num Batches"), STAT_SlateNumBatches, STATGROUP_
 DECLARE_DWORD_COUNTER_STAT(TEXT("Num Vertices"), STAT_SlateVertexCount, STATGROUP_Slate);
 DECLARE_MEMORY_STAT(TEXT("Batch Vertex Memory"), STAT_SlateVertexBatchMemory, STATGROUP_SlateMemory);
 DECLARE_MEMORY_STAT(TEXT("Batch Index Memory"), STAT_SlateIndexBatchMemory, STATGROUP_SlateMemory);
+DECLARE_CYCLE_STAT(TEXT("FillBatchBuffers Time"), STAT_SlateFillBatchBuffers, STATGROUP_Slate);
+DECLARE_CYCLE_STAT(TEXT("AddElements Time"), STAT_SlateAddElements, STATGROUP_Slate);
 
 SLATE_DECLARE_CYCLE_COUNTER(GSlateAddElements, "Add Elements");
 SLATE_DECLARE_CYCLE_COUNTER(GSlateFindBatchTime, "FindElementForBatch");
@@ -65,6 +67,7 @@ FSlateElementBatcher::~FSlateElementBatcher()
 void FSlateElementBatcher::AddElements( const TArray<FSlateDrawElement>& DrawElements )
 {
 	SLATE_CYCLE_COUNTER_SCOPE(GSlateAddElements);
+	SCOPE_CYCLE_COUNTER(STAT_SlateAddElements);
 	// This stuff is just for the counters. Could be scoped by an #ifdef if necessary.
 	static_assert(
 		FSlateDrawElement::EElementType::ET_Box == 0 &&
@@ -1545,6 +1548,7 @@ void FSlateElementBatcher::AddIndices( TArray<SlateIndex>& OutIndices, FSlateEle
 void FSlateElementBatcher::FillBatchBuffers( FSlateWindowElementList& WindowElementList, bool& bRequiresStencilTest )
 {
 	SLATE_CYCLE_COUNTER_SCOPE(GSlateFillBatchBuffers);
+	SCOPE_CYCLE_COUNTER(STAT_SlateFillBatchBuffers);
 
 	TArray<FSlateRenderBatch>& OutRenderBatches = WindowElementList.RenderBatches;
 	TArray<FSlateVertex>& OutBatchedVertices = WindowElementList.BatchedVertices;
