Added logic to enforce the min window size when calling Resize on the window.

Merging using UE4-To-UE4-LauncherDev

diff --git a/Engine/Source/Runtime/SlateCore/Private/Widgets/SWindow.cpp b/Engine/Source/Runtime/SlateCore/Private/Widgets/SWindow.cpp
--- a/Engine/Source/Runtime/SlateCore/Private/Widgets/SWindow.cpp
+++ b/Engine/Source/Runtime/SlateCore/Private/Widgets/SWindow.cpp
@@ -771,6 +771,12 @@ void SWindow::Resize( FVector2D NewSize )
 
 	if ( Size != NewSize )
 	{
+		NewSize.X = FMath::Max(SizeLimits.GetMinWidth().Get(NewSize.X), NewSize.X);
+		NewSize.X = FMath::Min(SizeLimits.GetMaxWidth().Get(NewSize.X), NewSize.X);
+		
+		NewSize.Y = FMath::Max(SizeLimits.GetMinHeight().Get(NewSize.Y), NewSize.Y);
+		NewSize.Y = FMath::Min(SizeLimits.GetMaxHeight().Get(NewSize.Y), NewSize.Y);
+		
 		if (NativeWindow.IsValid())
 		{
 			NativeWindow->ReshapeWindow( FMath::TruncToInt(ScreenPosition.X), FMath::TruncToInt(ScreenPosition.Y), FMath::TruncToInt(NewSize.X), FMath::TruncToInt(NewSize.Y) );
