Fixed a discrepancy between the scale passed to SlatePrepass and the scale used when rendering.

This was caused by SWindow applying the application scale twice.

diff --git a/Engine/Source/Runtime/SlateCore/Private/Widgets/SWindow.cpp b/Engine/Source/Runtime/SlateCore/Private/Widgets/SWindow.cpp
--- a/Engine/Source/Runtime/SlateCore/Private/Widgets/SWindow.cpp
+++ b/Engine/Source/Runtime/SlateCore/Private/Widgets/SWindow.cpp
@@ -1448,11 +1448,6 @@ FVector2D SWindow::ComputeDesiredSize(float LayoutScaleMultiplier) const
 	return SCompoundWidget::ComputeDesiredSize(LayoutScaleMultiplier) * LayoutScaleMultiplier;
 }
 
-float SWindow::GetRelativeLayoutScale(const FSlotBase& Child) const
-{
-	return FSlateApplicationBase::Get().GetApplicationScale();
-}
-
 const TArray< TSharedRef<SWindow> >& SWindow::GetChildWindows() const
 {
 	return ChildWindows;
diff --git a/Engine/Source/Runtime/SlateCore/Public/Widgets/SWindow.h b/Engine/Source/Runtime/SlateCore/Public/Widgets/SWindow.h
--- a/Engine/Source/Runtime/SlateCore/Public/Widgets/SWindow.h
+++ b/Engine/Source/Runtime/SlateCore/Public/Widgets/SWindow.h
@@ -622,8 +622,6 @@ private:
 	/** The window's desired size takes into account the ratio between the slate units and the pixel size */
 	virtual FVector2D ComputeDesiredSize(float) const override;
 
-	virtual float GetRelativeLayoutScale(const FSlotBase& Child) const override;
-
 	// For a given client size, calculate the window size required to accomodate any potential non-OS borders and tilebars
 	FVector2D GetWindowSizeFromClientSize(FVector2D InClientSize);
 
