#Slate Removed unnecessary use of Geometry.AbsolutePosition

diff --git a/Engine/Source/Runtime/SlateCore/Private/Widgets/SWindow.cpp b/Engine/Source/Runtime/SlateCore/Private/Widgets/SWindow.cpp
--- a/Engine/Source/Runtime/SlateCore/Private/Widgets/SWindow.cpp
+++ b/Engine/Source/Runtime/SlateCore/Private/Widgets/SWindow.cpp
@@ -1254,7 +1254,7 @@ bool SWindow::IsTopmostWindow() const
 /** @return true if mouse coordinates is within this window */
 bool SWindow::IsScreenspaceMouseWithin(FVector2D ScreenspaceMouseCoordinate) const
 {
-	const FVector2D LocalMouseCoordinate = ScreenspaceMouseCoordinate - GetWindowGeometryInScreen().AbsolutePosition;
+	const FVector2D LocalMouseCoordinate = ScreenspaceMouseCoordinate - ScreenPosition;
 	return NativeWindow->IsPointInWindow(FMath::TruncToInt(LocalMouseCoordinate.X), FMath::TruncToInt(LocalMouseCoordinate.Y));
 }
 
