Optimized the element batcher to perform less memory allocations per frame and pass off some processing to the render thead earlier to avoid additional memcopy's
- The Slate RHI renderer now uses a larger upfront pool for vertex buffers to avoid many upfront recreations of the vertex buffers
- The window element lists are no longer destroyed each frame and are reused as long as the window that they reference exists.  This reduces memory allocations because arrays are no longer destroyed
- Creation of the final rendering batches has now been moved to the rendering thread
- The mapping of layers to element batch arrays has been simplified and is no longer emptied each frame thus avoiding memory allocations

diff --git a/Engine/Config/BaseEngine.ini b/Engine/Config/BaseEngine.ini
--- a/Engine/Config/BaseEngine.ini
+++ b/Engine/Config/BaseEngine.ini
@@ -2087,6 +2087,9 @@ DeviceProfileSelectionModule="ExampleDeviceProfileSelector"
 [SlateRenderer]
 TextureAtlasSize=1024
 FontAtlasSize=1024
+; This is the number of preallocated vertices to create for use in Slate/UMG UI. A smaller number would indicate more resizes of the vertex buffer which could cause hitches.  A larger number increases memory
+; Use stat slatememory to see the cost of these buffers
+NumPreallocatedVertices=50000
 
 [MobileSlateUI]
 bTouchFallbackToMouse=true
diff --git a/Engine/Source/Developer/StandaloneRenderer/Private/OpenGL/SlateOpenGLRenderer.cpp b/Engine/Source/Developer/StandaloneRenderer/Private/OpenGL/SlateOpenGLRenderer.cpp
--- a/Engine/Source/Developer/StandaloneRenderer/Private/OpenGL/SlateOpenGLRenderer.cpp
+++ b/Engine/Source/Developer/StandaloneRenderer/Private/OpenGL/SlateOpenGLRenderer.cpp
@@ -95,11 +95,11 @@ void FSlateOpenGLRenderer::DrawWindows( FSlateDrawBuffer& InWindowDrawBuffer )
 	FontCache->UpdateCache();
 
 	// Draw each window.  For performance.  All elements are batched before anything is rendered
-	TArray<FSlateWindowElementList>& WindowElementLists = InWindowDrawBuffer.GetWindowElementLists();
+	TArray< TSharedPtr<FSlateWindowElementList> >& WindowElementLists = InWindowDrawBuffer.GetWindowElementLists();
 
 	for( int32 ListIndex = 0; ListIndex < WindowElementLists.Num(); ++ListIndex )
 	{
-		FSlateWindowElementList& ElementList = WindowElementLists[ListIndex];
+		FSlateWindowElementList& ElementList = *WindowElementLists[ListIndex];
 
 		if ( ElementList.GetWindow().IsValid() )
 		{
@@ -120,22 +120,25 @@ void FSlateOpenGLRenderer::DrawWindows( FSlateDrawBuffer& InWindowDrawBuffer )
 			Viewport->MakeCurrent();
 
 			// Batch elements.  Note that we must set the current viewport before doing this so we have a valid rendering context when calling OpenGL functions
-			ElementBatcher->AddElements( ElementList.GetDrawElements() );
+			ElementBatcher->AddElements( ElementList );
 
 			//@ todo Slate: implement for opengl
 			bool bRequiresStencilTest = false;
-			ElementBatcher->FillBatchBuffers( ElementList, bRequiresStencilTest );
 
 			ElementBatcher->ResetBatches();
+			
+			FSlateBatchData& BatchData = ElementList.GetBatchData();
+
+			BatchData.CreateRenderBatches();
 
-			RenderingPolicy->UpdateBuffers( ElementList );
+			RenderingPolicy->UpdateVertexAndIndexBuffers( BatchData );
 
 			check(Viewport);
 
 			glViewport( Viewport->ViewportRect.Left, Viewport->ViewportRect.Top, Viewport->ViewportRect.Right, Viewport->ViewportRect.Bottom );
 
 			// Draw all elements
-			RenderingPolicy->DrawElements( ViewMatrix*Viewport->ProjectionMatrix, ElementList.GetRenderBatches() );
+			RenderingPolicy->DrawElements( ViewMatrix*Viewport->ProjectionMatrix, BatchData.GetRenderBatches() );
 
 			Viewport->SwapBuffers();
 
diff --git a/Engine/Source/Developer/StandaloneRenderer/Private/OpenGL/SlateOpenGLRenderingPolicy.cpp b/Engine/Source/Developer/StandaloneRenderer/Private/OpenGL/SlateOpenGLRenderingPolicy.cpp
--- a/Engine/Source/Developer/StandaloneRenderer/Private/OpenGL/SlateOpenGLRenderingPolicy.cpp
+++ b/Engine/Source/Developer/StandaloneRenderer/Private/OpenGL/SlateOpenGLRenderingPolicy.cpp
@@ -86,45 +86,43 @@ void FSlateOpenGLRenderingPolicy::ReleaseResources()
  * @param InVertices	The vertices to copy to the vertex buffer
  * @param InIndices		The indices to copy to the index buffer
  */
-void FSlateOpenGLRenderingPolicy::UpdateBuffers( const FSlateWindowElementList& InElementList )
+void FSlateOpenGLRenderingPolicy::UpdateVertexAndIndexBuffers(FSlateBatchData& InBatchData)
 {
-	const TArray<FSlateVertex>& Vertices = InElementList.GetBatchedVertices();
-	const TArray<SlateIndex>& Indices = InElementList.GetBatchedIndices();
-
-	if( Vertices.Num() )
+	if( InBatchData.GetRenderBatches().Num() > 0 )
 	{
-		uint32 NumVertices = Vertices.Num();
-	
-		// resize if needed
-		if( NumVertices*sizeof(FSlateVertex) > VertexBuffer.GetBufferSize() )
+		if( InBatchData.GetNumBatchedVertices() > 0 )
 		{
-			uint32 NumBytesNeeded = NumVertices*sizeof(FSlateVertex);
-			// increase by a static size.
-			// @todo make this better
-			VertexBuffer.ResizeBuffer( NumBytesNeeded + 200*sizeof(FSlateVertex) );
+			uint32 NumVertices = InBatchData.GetNumBatchedVertices();
+	
+			// resize if needed
+			if( NumVertices*sizeof(FSlateVertex) > VertexBuffer.GetBufferSize() )
+			{
+				uint32 NumBytesNeeded = NumVertices*sizeof(FSlateVertex);
+				// increase by a static size.
+				// @todo make this better
+				VertexBuffer.ResizeBuffer( NumBytesNeeded + 200*sizeof(FSlateVertex) );
+			}
 		}
 
-		void* VerticesPtr = VertexBuffer.Lock(0);
-		FMemory::Memcpy( VerticesPtr, Vertices.GetData(), sizeof(FSlateVertex)*NumVertices );
-
-		VertexBuffer.Unlock();
-	}
-
-	if( Indices.Num() )
-	{
-		uint32 NumIndices = Indices.Num();
-
-		// resize if needed
-		if( NumIndices > IndexBuffer.GetMaxNumIndices() )
+		if( InBatchData.GetNumBatchedIndices() > 0 )
 		{
-			// increase by a static size.
-			// @todo make this better
-			IndexBuffer.ResizeBuffer( NumIndices + 100 );
+			uint32 NumIndices = InBatchData.GetNumBatchedIndices();
+
+			// resize if needed
+			if( NumIndices > IndexBuffer.GetMaxNumIndices() )
+			{
+				// increase by a static size.
+				// @todo make this better
+				IndexBuffer.ResizeBuffer( NumIndices + 100 );
+			}
 		}
 
-		void* IndicesPtr = IndexBuffer.Lock(0);
-		FMemory::Memcpy( IndicesPtr, Indices.GetData(), sizeof(SlateIndex)*NumIndices );
+		uint8* VerticesPtr = (uint8*)VertexBuffer.Lock(0);
+		uint8* IndicesPtr = (uint8*)IndexBuffer.Lock(0);
 
+		InBatchData.FillVertexAndIndexBuffer(VerticesPtr, IndicesPtr);
+
+		VertexBuffer.Unlock();
 		IndexBuffer.Unlock();
 	}
 }
diff --git a/Engine/Source/Developer/StandaloneRenderer/Private/OpenGL/SlateOpenGLRenderingPolicy.h b/Engine/Source/Developer/StandaloneRenderer/Private/OpenGL/SlateOpenGLRenderingPolicy.h
--- a/Engine/Source/Developer/StandaloneRenderer/Private/OpenGL/SlateOpenGLRenderingPolicy.h
+++ b/Engine/Source/Developer/StandaloneRenderer/Private/OpenGL/SlateOpenGLRenderingPolicy.h
@@ -14,10 +14,10 @@ public:
 	/**
 	 * Updates vertex and index buffers used in drawing
 	 *
-	 * @param InVertices	The vertices to copy to the vertex buffer
-	 * @param InIndices		The indices to copy to the index buffer
+	 * @param BatchData	The batch data that contains rendering data we need to upload to the buffers
 	 */
-	void UpdateBuffers( const FSlateWindowElementList& InElementList );
+	void UpdateVertexAndIndexBuffers(FSlateBatchData& BatchData);
+
 	
 	/**
 	 * Draws Slate elements
diff --git a/Engine/Source/Developer/StandaloneRenderer/Private/Windows/D3D/SlateD3DRenderer.cpp b/Engine/Source/Developer/StandaloneRenderer/Private/Windows/D3D/SlateD3DRenderer.cpp
--- a/Engine/Source/Developer/StandaloneRenderer/Private/Windows/D3D/SlateD3DRenderer.cpp
+++ b/Engine/Source/Developer/StandaloneRenderer/Private/Windows/D3D/SlateD3DRenderer.cpp
@@ -334,10 +334,10 @@ void FSlateD3DRenderer::DrawWindows( FSlateDrawBuffer& InWindowDrawBuffer )
 	FontCache->UpdateCache();
 
 	// Iterate through each element list and set up an RHI window for it if needed
-	TArray<FSlateWindowElementList>& WindowElementLists = InWindowDrawBuffer.GetWindowElementLists();
+	TArray< TSharedPtr<FSlateWindowElementList> >& WindowElementLists = InWindowDrawBuffer.GetWindowElementLists();
 	for( int32 ListIndex = 0; ListIndex < WindowElementLists.Num(); ++ListIndex )
 	{
-		FSlateWindowElementList& ElementList = WindowElementLists[ListIndex];
+		FSlateWindowElementList& ElementList = *WindowElementLists[ListIndex];
 		SLATE_CYCLE_COUNTER_SCOPE_CUSTOM_DETAILED(SLATE_STATS_DETAIL_LEVEL_MED, GRendererDrawElementList, ElementList.GetWindow()->GetCreatedInLocation());
 
 		if ( ElementList.GetWindow().IsValid() )
@@ -345,22 +345,18 @@ void FSlateD3DRenderer::DrawWindows( FSlateDrawBuffer& InWindowDrawBuffer )
 			TSharedRef<SWindow> WindowToDraw = ElementList.GetWindow().ToSharedRef();
 
 			// Add all elements for this window to the element batcher
-			ElementBatcher->AddElements( ElementList.GetDrawElements() );
-
-			bool bUnused = false;
-			ElementBatcher->FillBatchBuffers( ElementList, bUnused );
-
-			// All elements for this window have been batched and rendering data updated
-			ElementBatcher->ResetBatches();
+			ElementBatcher->AddElements( ElementList );
 
 			FVector2D WindowSize = WindowToDraw->GetSizeInScreen();
 
 			FSlateD3DViewport* Viewport = WindowToViewportMap.Find( &WindowToDraw.Get() );
 			check(Viewport);
 
+			FSlateBatchData& BatchData = ElementList.GetBatchData();
 			{
 				SLATE_CYCLE_COUNTER_SCOPE(GRendererUpdateBuffers);
-				RenderingPolicy->UpdateBuffers(ElementList);
+				BatchData.CreateRenderBatches();
+				RenderingPolicy->UpdateVertexAndIndexBuffers(BatchData);
 			}
 
 			check(Viewport);
@@ -379,7 +375,7 @@ void FSlateD3DRenderer::DrawWindows( FSlateDrawBuffer& InWindowDrawBuffer )
 
 			{
 				SLATE_CYCLE_COUNTER_SCOPE(GRendererDrawElements);
-				RenderingPolicy->DrawElements(ViewMatrix*Viewport->ProjectionMatrix, ElementList.GetRenderBatches());
+				RenderingPolicy->DrawElements(ViewMatrix*Viewport->ProjectionMatrix, BatchData.GetRenderBatches());
 			}
 
 			GD3DDeviceContext->OMSetRenderTargets(0, NULL, NULL);
diff --git a/Engine/Source/Developer/StandaloneRenderer/Private/Windows/D3D/SlateD3DRenderingPolicy.cpp b/Engine/Source/Developer/StandaloneRenderer/Private/Windows/D3D/SlateD3DRenderingPolicy.cpp
--- a/Engine/Source/Developer/StandaloneRenderer/Private/Windows/D3D/SlateD3DRenderingPolicy.cpp
+++ b/Engine/Source/Developer/StandaloneRenderer/Private/Windows/D3D/SlateD3DRenderingPolicy.cpp
@@ -153,68 +153,62 @@ void FSlateD3D11RenderingPolicy::ReleaseResources()
 	delete PixelShader;
 }
 
-void FSlateD3D11RenderingPolicy::UpdateBuffers( const FSlateWindowElementList& InElementList )
+void FSlateD3D11RenderingPolicy::UpdateVertexAndIndexBuffers( FSlateBatchData& InBatchData )
 {	
-	const TArray<FSlateVertex>& Vertices = InElementList.GetBatchedVertices();
-	const TArray<SlateIndex>& Indices = InElementList.GetBatchedIndices();
-
-	if( Vertices.Num() )
+	if( InBatchData.GetRenderBatches().Num() > 0 )
 	{
-		uint32 NumVertices = Vertices.Num();
-	
-		// resize if needed
-		if( NumVertices*sizeof(FSlateVertex) > VertexBuffer.GetBufferSize() )
+		if( InBatchData.GetNumBatchedVertices() > 0 )
 		{
-			SLATE_CYCLE_COUNTER_SCOPE_DETAILED(SLATE_STATS_DETAIL_LEVEL_FULL, GSlateResizeRenderBuffers);
-			uint32 NumBytesNeeded = NumVertices*sizeof(FSlateVertex);
-			// increase by a static size.
-			// @todo make this better
-			VertexBuffer.ResizeBuffer( NumBytesNeeded + 200*sizeof(FSlateVertex) );
-		}
+			uint32 NumVertices = InBatchData.GetNumBatchedVertices();
+	
+			// resize if needed
+			if( NumVertices*sizeof(FSlateVertex) > VertexBuffer.GetBufferSize() )
+			{
+				SLATE_CYCLE_COUNTER_SCOPE_DETAILED(SLATE_STATS_DETAIL_LEVEL_FULL, GSlateResizeRenderBuffers);
+				uint32 NumBytesNeeded = NumVertices*sizeof(FSlateVertex);
+				// increase by a static size.
+				// @todo make this better
+				VertexBuffer.ResizeBuffer( NumBytesNeeded + 200*sizeof(FSlateVertex) );
+			}
 
-		void* VerticesPtr = nullptr;
-		{
-			SLATE_CYCLE_COUNTER_SCOPE_DETAILED(SLATE_STATS_DETAIL_LEVEL_FULL, GSlateLockRenderBuffers);
-			VerticesPtr = VertexBuffer.Lock(0);
-		}
-		{
-			SLATE_CYCLE_COUNTER_SCOPE(GSlateMemCopyRenderBuffers);
-			FMemory::Memcpy(VerticesPtr, Vertices.GetData(), sizeof(FSlateVertex)*NumVertices);
+		
 		}
 
+		if( InBatchData.GetNumBatchedIndices() > 0 )
 		{
-			SLATE_CYCLE_COUNTER_SCOPE_DETAILED(SLATE_STATS_DETAIL_LEVEL_FULL, GSlateUnlockRenderBuffers);
-			VertexBuffer.Unlock();
-		}
-	}
+			uint32 NumIndices = InBatchData.GetNumBatchedIndices();
 
-	if( Indices.Num() )
-	{
-		uint32 NumIndices = Indices.Num();
+			// resize if needed
+			if( NumIndices > IndexBuffer.GetMaxNumIndices() )
+			{
+				SLATE_CYCLE_COUNTER_SCOPE_DETAILED(SLATE_STATS_DETAIL_LEVEL_FULL, GSlateResizeRenderBuffers);
+				// increase by a static size.
+				// @todo make this better
+				IndexBuffer.ResizeBuffer( NumIndices + 100 );
+			}
 
-		// resize if needed
-		if( NumIndices > IndexBuffer.GetMaxNumIndices() )
-		{
-			SLATE_CYCLE_COUNTER_SCOPE_DETAILED(SLATE_STATS_DETAIL_LEVEL_FULL, GSlateResizeRenderBuffers);
-			// increase by a static size.
-			// @todo make this better
-			IndexBuffer.ResizeBuffer( NumIndices + 100 );
+	
 		}
 
-		void* IndicesPtr = nullptr;
+		uint8* VerticesPtr = nullptr;
+		uint8* IndicesPtr = nullptr;
 		{
 			SLATE_CYCLE_COUNTER_SCOPE_DETAILED(SLATE_STATS_DETAIL_LEVEL_FULL, GSlateLockRenderBuffers);
-			IndicesPtr = IndexBuffer.Lock(0);
+			VerticesPtr = (uint8*)VertexBuffer.Lock(0);
+			IndicesPtr = (uint8*)IndexBuffer.Lock(0);
 		}
+	
 		{
-			SLATE_CYCLE_COUNTER_SCOPE_DETAILED(SLATE_STATS_DETAIL_LEVEL_FULL, GSlateMemCopyRenderBuffers);
-			FMemory::Memcpy(IndicesPtr, Indices.GetData(), sizeof(SlateIndex)*NumIndices);
+			SLATE_CYCLE_COUNTER_SCOPE(GSlateMemCopyRenderBuffers);
+			InBatchData.FillVertexAndIndexBuffer( VerticesPtr, IndicesPtr );
 		}
 
 		{
 			SLATE_CYCLE_COUNTER_SCOPE_DETAILED(SLATE_STATS_DETAIL_LEVEL_FULL, GSlateUnlockRenderBuffers);
+			VertexBuffer.Unlock();
 			IndexBuffer.Unlock();
 		}
+
 	}
 }
 
diff --git a/Engine/Source/Developer/StandaloneRenderer/Private/Windows/D3D/SlateD3DRenderingPolicy.h b/Engine/Source/Developer/StandaloneRenderer/Private/Windows/D3D/SlateD3DRenderingPolicy.h
--- a/Engine/Source/Developer/StandaloneRenderer/Private/Windows/D3D/SlateD3DRenderingPolicy.h
+++ b/Engine/Source/Developer/StandaloneRenderer/Private/Windows/D3D/SlateD3DRenderingPolicy.h
@@ -9,7 +9,7 @@ public:
 	FSlateD3D11RenderingPolicy( TSharedPtr<FSlateFontCache> InFontCache, TSharedRef<FSlateD3DTextureManager> InTextureManager );
 	~FSlateD3D11RenderingPolicy();
 
-	void UpdateBuffers( const FSlateWindowElementList& InElementList );
+	void UpdateVertexAndIndexBuffers( FSlateBatchData& InBatchData );
 	void DrawElements( const FMatrix& ViewProjectionMatrix, const TArray<FSlateRenderBatch>& RenderBatches );
 
 	virtual TSharedRef<FSlateShaderResourceManager> GetResourceManager() override;
diff --git a/Engine/Source/Runtime/Slate/Private/Framework/Application/SlateApplication.cpp b/Engine/Source/Runtime/Slate/Private/Framework/Application/SlateApplication.cpp
--- a/Engine/Source/Runtime/Slate/Private/Framework/Application/SlateApplication.cpp
+++ b/Engine/Source/Runtime/Slate/Private/Framework/Application/SlateApplication.cpp
@@ -1195,9 +1195,9 @@ void FSlateApplication::PrivateDrawWindows( TSharedPtr<SWindow> DrawOnlyThisWind
 	// Some windows may have been destroyed/removed.
 	// Do not attempt to draw any windows that have been removed.
 	TArray<TSharedRef<SWindow>> AllWindows = GatherAllDescendants(SlateWindows);
-	DrawWindowArgs.OutDrawBuffer.GetWindowElementLists().RemoveAll([&]( FSlateWindowElementList& Candidate )
+	DrawWindowArgs.OutDrawBuffer.GetWindowElementLists().RemoveAll([&]( TSharedPtr<FSlateWindowElementList>& Candidate )
 	{
-		TSharedPtr<SWindow> CandidateWindow = Candidate.GetWindow();
+		TSharedPtr<SWindow> CandidateWindow = Candidate->GetWindow();
 		return !CandidateWindow.IsValid() || !AllWindows.Contains(CandidateWindow.ToSharedRef());
 	});
 
diff --git a/Engine/Source/Runtime/SlateCore/Private/Rendering/DrawElements.cpp b/Engine/Source/Runtime/SlateCore/Private/Rendering/DrawElements.cpp
--- a/Engine/Source/Runtime/SlateCore/Private/Rendering/DrawElements.cpp
+++ b/Engine/Source/Runtime/SlateCore/Private/Rendering/DrawElements.cpp
@@ -174,6 +174,139 @@ FVector2D FSlateDrawElement::GetRotationPoint(const FPaintGeometry& PaintGeometr
 	return RotationPoint;
 }
 
+void FSlateBatchData::Reset()
+{
+	RenderBatches.Reset();
+	
+	// note: LayerToElementBatches is not reset here as the same layers are 
+	// more than likely reused and we can save memory allocations by not resetting the map every frame
+
+	NumBatchedVertices = 0;
+	NumBatchedIndices = 0;
+	NumLayers = 0;
+}
+
+void FSlateBatchData::AssignVertexArrayToBatch( FSlateElementBatch& Batch )
+{
+	// Get a free vertex array
+	if (VertexArrayFreeList.Num() > 0)
+	{
+		Batch.VertexArrayIndex = VertexArrayFreeList.Pop(/*bAllowShrinking=*/ false);
+		BatchVertexArrays[Batch.VertexArrayIndex].Reserve(200);
+	}
+	else
+	{
+		// There are no free vertex arrays so we must add one		
+		uint32 NewIndex = BatchVertexArrays.Add(TArray<FSlateVertex>());
+		BatchVertexArrays[NewIndex].Reserve(200);
+
+		Batch.VertexArrayIndex = NewIndex;
+	}
+}
+
+void FSlateBatchData::AssignIndexArrayToBatch( FSlateElementBatch& Batch )
+{
+	// Get a free index array
+	if (IndexArrayFreeList.Num() > 0)
+	{
+		Batch.IndexArrayIndex = IndexArrayFreeList.Pop(/*bAllowShrinking=*/ false);
+		BatchIndexArrays[Batch.IndexArrayIndex].Reserve(200);
+	}
+	else
+	{
+		// There are no free index arrays so we must add one
+		uint32 NewIndex = BatchIndexArrays.Add(TArray<SlateIndex>());
+		BatchIndexArrays[NewIndex].Reserve(500);
+
+		Batch.IndexArrayIndex = NewIndex;
+	}
+
+}
+
+void FSlateBatchData::FillVertexAndIndexBuffer( uint8* VertexBuffer, uint8* IndexBuffer )
+{
+	int32 IndexOffset = 0;
+	int32 VertexOffset = 0;
+	for( const FSlateRenderBatch& Batch : RenderBatches )
+	{
+		if( Batch.VertexArrayIndex != INDEX_NONE && Batch.IndexArrayIndex != INDEX_NONE )
+		{
+			TArray<FSlateVertex>& Vertices = BatchVertexArrays[Batch.VertexArrayIndex];
+			TArray<SlateIndex>& Indices = BatchIndexArrays[Batch.IndexArrayIndex];
+
+			if(Vertices.Num() && Indices.Num())
+			{
+				uint32 RequiredVertexSize = Vertices.Num() * Vertices.GetTypeSize();
+				uint32 RequiredIndexSize = Indices.Num() * Indices.GetTypeSize();
+
+				FMemory::Memcpy(VertexBuffer+VertexOffset, Vertices.GetData(), RequiredVertexSize);
+				FMemory::Memcpy(IndexBuffer+IndexOffset, Indices.GetData(), RequiredIndexSize);
+
+				VertexArrayFreeList.Add(Batch.VertexArrayIndex);
+				IndexArrayFreeList.Add(Batch.IndexArrayIndex);
+
+				IndexOffset += (Indices.Num()*sizeof(SlateIndex));
+				VertexOffset += (Vertices.Num()*sizeof(FSlateVertex));
+
+				Vertices.Empty(Vertices.Num());
+				Indices.Empty(Indices.Num());
+			}
+		}
+	}
+}
+
+void FSlateBatchData::CreateRenderBatches()
+{
+	checkSlow( IsInRenderingThread() );
+
+	LayerToElementBatches.KeySort( TLess<uint32>() );
+
+	uint32 VertexOffset = 0;
+	uint32 IndexOffset = 0;
+	// For each element batch add its vertices and indices to the bulk lists.
+	for (FElementBatchMap::TIterator It(LayerToElementBatches); It; ++It)
+	{
+		FElementBatchArray& ElementBatches = It.Value();
+
+		if( ElementBatches.Num() > 0 )
+		{
+			++NumLayers;
+			for (FElementBatchArray::TIterator BatchIt(ElementBatches); BatchIt; ++BatchIt)
+			{
+				FSlateElementBatch& ElementBatch = *BatchIt;
+
+				bool bIsMaterial = ElementBatch.GetShaderResource() && ElementBatch.GetShaderResource()->GetType() == ESlateShaderResource::Material;
+
+				if (!ElementBatch.GetCustomDrawer().IsValid())
+				{
+					TArray<FSlateVertex>& BatchVertices = GetBatchVertexList(ElementBatch);
+					TArray<SlateIndex>& BatchIndices = GetBatchIndexList(ElementBatch);
+	
+					// We should have at least some vertices and indices in the batch or none at all
+					check(BatchVertices.Num() > 0 && BatchIndices.Num() > 0 || BatchVertices.Num() == 0 && BatchIndices.Num() == 0);
+
+					if (BatchVertices.Num() > 0 && BatchIndices.Num() > 0)
+					{
+						int32 NumVertices = BatchVertices.Num();
+						int32 NumIndices = BatchIndices.Num();
+
+						AddRenderBatch( ElementBatch, NumVertices, NumIndices, VertexOffset, IndexOffset );
+				
+						VertexOffset += BatchVertices.Num();
+						IndexOffset += BatchIndices.Num();
+					}
+				}
+				else
+				{
+					AddRenderBatch( ElementBatch, 0, 0, 0, 0 );
+				}
+			}
+		}
+		ElementBatches.Empty( ElementBatches.Num() );
+	}
+
+}
+
 FSlateWindowElementList::FDeferredPaint::FDeferredPaint( const TSharedRef<const SWidget>& InWidgetToPaint, const FPaintArgs& InArgs, const FGeometry InAllottedGeometry, const FSlateRect InMyClippingRect, const FWidgetStyle& InWidgetStyle, bool InParentEnabled )
 : WidgetToPaintPtr( InWidgetToPaint )
 , Args( InArgs )
diff --git a/Engine/Source/Runtime/SlateCore/Private/Rendering/ElementBatcher.cpp b/Engine/Source/Runtime/SlateCore/Private/Rendering/ElementBatcher.cpp
--- a/Engine/Source/Runtime/SlateCore/Private/Rendering/ElementBatcher.cpp
+++ b/Engine/Source/Runtime/SlateCore/Private/Rendering/ElementBatcher.cpp
@@ -6,13 +6,9 @@
 #include "Internationalization/Text.h"
 #include "SlateStats.h"
 
-DECLARE_DWORD_COUNTER_STAT(TEXT("Num Layers"), STAT_SlateNumLayers, STATGROUP_Slate);
-DECLARE_DWORD_COUNTER_STAT(TEXT("Num Batches"), STAT_SlateNumBatches, STATGROUP_Slate);
-DECLARE_DWORD_COUNTER_STAT(TEXT("Num Vertices"), STAT_SlateVertexCount, STATGROUP_Slate);
-DECLARE_MEMORY_STAT(TEXT("Batch Vertex Memory"), STAT_SlateVertexBatchMemory, STATGROUP_SlateMemory);
-DECLARE_MEMORY_STAT(TEXT("Batch Index Memory"), STAT_SlateIndexBatchMemory, STATGROUP_SlateMemory);
-DECLARE_CYCLE_STAT(TEXT("FillBatchBuffers Time"), STAT_SlateFillBatchBuffers, STATGROUP_Slate);
-DECLARE_CYCLE_STAT(TEXT("AddElements Time"), STAT_SlateAddElements, STATGROUP_Slate);
+//DECLARE_CYCLE_STAT(TEXT("Find Batch For Element Time"), STAT_SlateFindBatchForElement, STATGROUP_Slate);
+DECLARE_DWORD_COUNTER_STAT(TEXT("Num Elements (Prebatch)"), STAT_SlateNumPrebatchElements, STATGROUP_Slate);
+DECLARE_CYCLE_STAT(TEXT("Add Elements Time"), STAT_SlateAddElements, STATGROUP_Slate);
 
 SLATE_DECLARE_CYCLE_COUNTER(GSlateAddElements, "Add Elements");
 SLATE_DECLARE_CYCLE_COUNTER(GSlateFindBatchTime, "FindElementForBatch");
@@ -53,6 +49,7 @@ static FColor GetElementColor( const FLinearColor& InColor, const FSlateBrush* I
 FSlateElementBatcher::FSlateElementBatcher( TSharedRef<FSlateRenderingPolicy> InRenderingPolicy )
 	: ResourceManager( *InRenderingPolicy->GetResourceManager() )
 	, FontCache( *InRenderingPolicy->GetFontCache() )
+	, BatchData( nullptr )
 	, PixelCenterOffset( InRenderingPolicy->GetPixelCenterOffset() )
 {
 
@@ -64,10 +61,11 @@ FSlateElementBatcher::~FSlateElementBatcher()
 }
 
 
-void FSlateElementBatcher::AddElements( const TArray<FSlateDrawElement>& DrawElements )
+void FSlateElementBatcher::AddElements( FSlateWindowElementList& WindowElementList )
 {
 	SLATE_CYCLE_COUNTER_SCOPE(GSlateAddElements);
-	SCOPE_CYCLE_COUNTER(STAT_SlateAddElements);
+
+	SCOPE_CYCLE_COUNTER( STAT_SlateAddElements );
 	// This stuff is just for the counters. Could be scoped by an #ifdef if necessary.
 	static_assert(
 		FSlateDrawElement::EElementType::ET_Box == 0 &&
@@ -95,6 +93,11 @@ void FSlateElementBatcher::AddElements( const TArray<FSlateDrawElement>& DrawEle
 		FName(TEXT("Custom")),
 	};
 
+	BatchData = &WindowElementList.GetBatchData();
+
+	const TArray<FSlateDrawElement>& DrawElements = WindowElementList.GetDrawElements();
+
+	int32 NumDrawnBatches = 0;
 
 	for( int32 DrawElementIndex = 0; DrawElementIndex < DrawElements.Num(); ++DrawElementIndex )
 	{
@@ -111,6 +114,8 @@ void FSlateElementBatcher::AddElements( const TArray<FSlateDrawElement>& DrawEle
 			// scissor rects are sort of a low level hack, so no one konws to clip against them. Instead we do it here to make sure the element is not actually rendered.
 			if (!bIsScissored)
 			{
+				++NumDrawnBatches;
+
 				// time just the adding of the element. The clipping stuff will be counted in exclusive time for the non-typed timer.
 				SLATE_CYCLE_COUNTER_SCOPE_CUSTOM_DETAILED(SLATE_STATS_DETAIL_LEVEL_MED, GSlateAddElements, ElementFNames[DrawElement.GetElementType()]);
 				// Determine what type of element to add
@@ -150,6 +155,11 @@ void FSlateElementBatcher::AddElements( const TArray<FSlateDrawElement>& DrawEle
 			}
 		}
 	}
+
+	SET_DWORD_STAT( STAT_SlateNumPrebatchElements, NumDrawnBatches );
+
+	// Done with the element list
+	BatchData = nullptr;
 }
 
 void FSlateElementBatcher::AddQuadElement( const FSlateDrawElement& DrawElement, FColor Color )
@@ -169,8 +179,8 @@ void FSlateElementBatcher::AddQuadElement( const FSlateDrawElement& DrawElement,
 	FSlateRotatedClipRectType RenderClipRect = ToSnappedRotatedRect(InClippingRect, InverseLayoutTransform, RenderTransform);
 
 	FSlateElementBatch& ElementBatch = FindBatchForElement( Layer, FShaderParams(), nullptr, ESlateDrawPrimitive::TriangleList, ESlateShader::Default, ESlateDrawEffect::None, ESlateBatchDrawFlag::Wireframe|ESlateBatchDrawFlag::NoBlending, DrawElement.GetScissorRect() );
-	TArray<FSlateVertex>& BatchVertices = BatchVertexArrays[ElementBatch.VertexArrayIndex];
-	TArray<SlateIndex>& BatchIndices = BatchIndexArrays[ElementBatch.IndexArrayIndex];
+	TArray<FSlateVertex>& BatchVertices = BatchData->GetBatchVertexList(ElementBatch);
+	TArray<SlateIndex>& BatchIndices = BatchData->GetBatchIndexList(ElementBatch);
 
 	// Determine the four corners of the quad
 	FVector2D TopLeft = FVector2D::ZeroVector;
@@ -290,8 +300,8 @@ void FSlateElementBatcher::AddBoxElement( const FSlateDrawElement& DrawElement )
 		ESlateBatchDrawFlag::Type DrawFlags = ( ( bTileHorizontal ? ESlateBatchDrawFlag::TileU : 0 ) | ( bTileVertical ? ESlateBatchDrawFlag::TileV : 0 ) );
 
 		FSlateElementBatch& ElementBatch = FindBatchForElement( Layer, FShaderParams(), Resource, ESlateDrawPrimitive::TriangleList, ESlateShader::Default, InDrawEffects, DrawFlags, DrawElement.GetScissorRect() );
-		TArray<FSlateVertex>& BatchVertices = BatchVertexArrays[ElementBatch.VertexArrayIndex];
-		TArray<SlateIndex>& BatchIndices = BatchIndexArrays[ElementBatch.IndexArrayIndex];
+		TArray<FSlateVertex>& BatchVertices = BatchData->GetBatchVertexList(ElementBatch);
+		TArray<SlateIndex>& BatchIndices = BatchData->GetBatchIndexList(ElementBatch);
 
 		float HorizontalTiling = bTileHorizontal ? LocalSize.X/TextureWidth : 1.0f;
 		float VerticalTiling = bTileVertical ? LocalSize.Y/TextureHeight : 1.0f;
@@ -594,8 +604,8 @@ void FSlateElementBatcher::AddTextElement(const FSlateDrawElement& DrawElement)
 				FontTexture = FontCache.GetSlateTextureResource( FontTextureIndex );
 				ElementBatch = &FindBatchForElement( Layer, FShaderParams(), FontTexture, ESlateDrawPrimitive::TriangleList, ESlateShader::Font, InDrawEffects, ESlateBatchDrawFlag::None, DrawElement.GetScissorRect() );
 
-				BatchVertices = &BatchVertexArrays[ElementBatch->VertexArrayIndex];
-				BatchIndices = &BatchIndexArrays[ElementBatch->IndexArrayIndex];
+				BatchVertices = &BatchData->GetBatchVertexList(*ElementBatch);
+				BatchIndices = &BatchData->GetBatchIndexList(*ElementBatch);
 
 				VertexOffset = BatchVertices->Num();
 				IndexOffset = BatchIndices->Num();
@@ -691,8 +701,8 @@ void FSlateElementBatcher::AddGradientElement( const FSlateDrawElement& DrawElem
 	check( InPayload.GradientStops.Num() > 0 );
 
 	FSlateElementBatch& ElementBatch = FindBatchForElement( Layer, FShaderParams(), nullptr, ESlateDrawPrimitive::TriangleList, ESlateShader::Default, InDrawEffects, InPayload.bGammaCorrect == false ? ESlateBatchDrawFlag::NoGamma : ESlateBatchDrawFlag::None, DrawElement.GetScissorRect() );
-	TArray<FSlateVertex>& BatchVertices = BatchVertexArrays[ElementBatch.VertexArrayIndex];
-	TArray<SlateIndex>& BatchIndices = BatchIndexArrays[ElementBatch.IndexArrayIndex];
+	TArray<FSlateVertex>& BatchVertices = BatchData->GetBatchVertexList(ElementBatch);
+	TArray<SlateIndex>& BatchIndices = BatchData->GetBatchIndexList(ElementBatch);
 
 	// Determine the four corners of the quad containing the gradient
 	FVector2D TopLeft = FVector2D::ZeroVector;
@@ -835,9 +845,8 @@ void FSlateElementBatcher::AddSplineElement( const FSlateDrawElement& DrawElemen
 
 	// Find a batch for the element
 	FSlateElementBatch& ElementBatch = FindBatchForElement( Layer, FShaderParams::MakePixelShaderParams( FVector4( InPayload.Thickness,Radius,0,0) ), nullptr, ESlateDrawPrimitive::TriangleList, ESlateShader::LineSegment, InDrawEffects, ESlateBatchDrawFlag::None, DrawElement.GetScissorRect() );
-	TArray<FSlateVertex>& BatchVertices = BatchVertexArrays[ElementBatch.VertexArrayIndex];
-	TArray<SlateIndex>& BatchIndices = BatchIndexArrays[ElementBatch.IndexArrayIndex];
-
+	TArray<FSlateVertex>& BatchVertices = BatchData->GetBatchVertexList(ElementBatch);
+	TArray<SlateIndex>& BatchIndices = BatchData->GetBatchIndexList(ElementBatch);
 
 	const FVector2D StartPt = InPayload.StartPt;
 	const FVector2D StartDir = InPayload.StartDir;
@@ -976,8 +985,8 @@ void FSlateElementBatcher::AddLineElement( const FSlateDrawElement& DrawElement
 
 		// Find a batch for the element
 		FSlateElementBatch& ElementBatch = FindBatchForElement( Layer, FShaderParams::MakePixelShaderParams( FVector4(RequestedThickness,Radius,0,0) ), nullptr, ESlateDrawPrimitive::TriangleList, ESlateShader::LineSegment, DrawEffects, ESlateBatchDrawFlag::None, DrawElement.GetScissorRect() );
-		TArray<FSlateVertex>& BatchVertices = BatchVertexArrays[ElementBatch.VertexArrayIndex];
-		TArray<SlateIndex>& BatchIndices = BatchIndexArrays[ElementBatch.IndexArrayIndex];
+		TArray<FSlateVertex>& BatchVertices = BatchData->GetBatchVertexList(ElementBatch);
+		TArray<SlateIndex>& BatchIndices = BatchData->GetBatchIndexList(ElementBatch);
 
 		const TArray<FVector2D>& Points = InPayload.Points;
 
@@ -1087,8 +1096,8 @@ void FSlateElementBatcher::AddLineElement( const FSlateDrawElement& DrawElement
 	{
 		// Find a batch for the element
 		FSlateElementBatch& ElementBatch = FindBatchForElement( Layer, FShaderParams(), nullptr, ESlateDrawPrimitive::LineList, ESlateShader::Default, DrawEffects, ESlateBatchDrawFlag::None, DrawElement.GetScissorRect() );
-		TArray<FSlateVertex>& BatchVertices = BatchVertexArrays[ElementBatch.VertexArrayIndex];
-		TArray<SlateIndex>& BatchIndices = BatchIndexArrays[ElementBatch.IndexArrayIndex];
+		TArray<FSlateVertex>& BatchVertices = BatchData->GetBatchVertexList(ElementBatch);
+		TArray<SlateIndex>& BatchIndices = BatchData->GetBatchIndexList(ElementBatch);
 
 		// Generate the rest of the segments
 		for( int32 Point = 0; Point < InPayload.Points.Num()-1; ++Point )
@@ -1156,8 +1165,8 @@ void FSlateElementBatcher::AddViewportElement( const FSlateDrawElement& DrawElem
 	}
 
 	FSlateElementBatch& ElementBatch = FindBatchForElement( Layer, FShaderParams(), ViewportResource, ESlateDrawPrimitive::TriangleList, ShaderType, InDrawEffects, DrawFlags, DrawElement.GetScissorRect() );
-	TArray<FSlateVertex>& BatchVertices = BatchVertexArrays[ElementBatch.VertexArrayIndex];
-	TArray<SlateIndex>& BatchIndices = BatchIndexArrays[ElementBatch.IndexArrayIndex];
+	TArray<FSlateVertex>& BatchVertices = BatchData->GetBatchVertexList(ElementBatch);
+	TArray<SlateIndex>& BatchIndices = BatchData->GetBatchIndexList(ElementBatch);
 
 	// Tag this batch as requiring vsync if the viewport requires it.
 	if( ViewportPin.IsValid() )
@@ -1198,7 +1207,6 @@ void FSlateElementBatcher::AddViewportElement( const FSlateDrawElement& DrawElem
 	BatchIndices.Add( IndexStart + 2 );
 	BatchIndices.Add( IndexStart + 1 );
 	BatchIndices.Add( IndexStart + 3 );
-
 }
 
 
@@ -1291,8 +1299,8 @@ void FSlateElementBatcher::AddBorderElement( const FSlateDrawElement& DrawElemen
 	ESlateBatchDrawFlag::Type DrawFlags = (ESlateBatchDrawFlag::TileU|ESlateBatchDrawFlag::TileV);
 
 	FSlateElementBatch& ElementBatch = FindBatchForElement( Layer, ShaderParams, Resource, ESlateDrawPrimitive::TriangleList, ESlateShader::Border, InDrawEffects, DrawFlags, DrawElement.GetScissorRect() );
-	TArray<FSlateVertex>& BatchVertices = BatchVertexArrays[ ElementBatch.VertexArrayIndex ];
-	TArray<SlateIndex>& BatchIndices = BatchIndexArrays[ ElementBatch.IndexArrayIndex ];
+	TArray<FSlateVertex>& BatchVertices = BatchData->GetBatchVertexList(ElementBatch);
+	TArray<SlateIndex>& BatchIndices = BatchData->GetBatchIndexList(ElementBatch);
 
 	// Ensure tiling of at least 1.  
 	TopTiling = TopTiling >= 1.0f ? TopTiling : 1.0f;
@@ -1414,20 +1422,21 @@ void FSlateElementBatcher::AddBorderElement( const FSlateDrawElement& DrawElemen
 	BatchIndices.Add( IndexStart + 31 );
 }
 
-
 void FSlateElementBatcher::AddCustomElement( const FSlateDrawElement& DrawElement )
 {
+	FElementBatchMap& LayerToElementBatches = BatchData->GetElementBatchMap();
+
 	const FSlateDataPayload& InPayload = DrawElement.GetDataPayload();
 	uint32 Layer = DrawElement.GetLayer();
 
 	if( InPayload.CustomDrawer.IsValid() )
 	{
 		// See if the layer already exists.
-		TSet<FSlateElementBatch>* ElementBatches = LayerToElementBatches.Find( Layer );
+		FElementBatchArray* ElementBatches = LayerToElementBatches.Find( Layer );
 		if( !ElementBatches )
 		{
 			// The layer doesn't exist so make it now
-			ElementBatches = &LayerToElementBatches.Add( Layer, TSet<FSlateElementBatch>() );
+			ElementBatches = &LayerToElementBatches.Add( Layer, FElementBatchArray() );
 		}
 		check( ElementBatches );
 
@@ -1450,57 +1459,31 @@ FSlateElementBatch& FSlateElementBatcher::FindBatchForElement(
 {
 	SLATE_CYCLE_COUNTER_SCOPE_DETAILED(SLATE_STATS_DETAIL_LEVEL_HI, GSlateFindBatchTime);
 
+//	SCOPE_CYCLE_COUNTER( STAT_SlateFindBatchForElement );
+	FElementBatchMap& LayerToElementBatches = BatchData->GetElementBatchMap();
+
 	// See if the layer already exists.
-	TSet<FSlateElementBatch>* ElementBatches = LayerToElementBatches.Find( Layer );
+	FElementBatchArray* ElementBatches = LayerToElementBatches.Find( Layer );
 	if( !ElementBatches )
 	{
 		// The layer doesn't exist so make it now
-		ElementBatches = &LayerToElementBatches.Add( Layer, TSet<FSlateElementBatch>() );
+		ElementBatches = &LayerToElementBatches.Add( Layer, FElementBatchArray() );
 	}
-	check( ElementBatches );
+
+	checkSlow( ElementBatches );
 
 	// Create a temp batch so we can use it as our key to find if the same batch already exists
 	FSlateElementBatch TempBatch( InTexture, ShaderParams, ShaderType, PrimitiveType, DrawEffects, DrawFlags, ScissorRect );
 
-	FSlateElementBatch* ElementBatch = ElementBatches->Find( TempBatch );
+	FSlateElementBatch* ElementBatch = ElementBatches->FindByKey( TempBatch );
 	if( !ElementBatch )
 	{
 		// No batch with the specified parameter exists.  Create it from the temp batch.
-		FSetElementId ID = ElementBatches->Add( TempBatch );
-		ElementBatch = &(*ElementBatches)[ID];
-
-		// Get a free vertex array
-		if( VertexArrayFreeList.Num() > 0 )
-		{
-			ElementBatch->VertexArrayIndex = VertexArrayFreeList.Pop(/*bAllowShrinking=*/ false);
-			BatchVertexArrays[ElementBatch->VertexArrayIndex].Reserve(200);		
-		}
-		else
-		{
-			// There are no free vertex arrays so we must add one		
-			uint32 NewIndex = BatchVertexArrays.Add( TArray<FSlateVertex>() );
-			BatchVertexArrays[NewIndex].Reserve(200);		
-
-			ElementBatch->VertexArrayIndex = NewIndex;
-		}
+		int32 Index = ElementBatches->Add( TempBatch );
+		ElementBatch = &(*ElementBatches)[Index];
 
-		// Get a free index array
-		if( IndexArrayFreeList.Num() > 0 )
-		{
-			ElementBatch->IndexArrayIndex = IndexArrayFreeList.Pop(/*bAllowShrinking=*/ false);
-			BatchIndexArrays[ElementBatch->IndexArrayIndex].Reserve(200);		
-		}
-		else
-		{
-			// There are no free index arrays so we must add one
-			uint32 NewIndex = BatchIndexArrays.Add( TArray<SlateIndex>() );
-			BatchIndexArrays[NewIndex].Reserve(500);	
-
-			ElementBatch->IndexArrayIndex = NewIndex;
-
-			check( BatchIndexArrays.IsValidIndex( ElementBatch->IndexArrayIndex ) );
-		}
-		
+		BatchData->AssignVertexArrayToBatch( *ElementBatch );
+		BatchData->AssignIndexArrayToBatch( *ElementBatch );
 	}
 	check( ElementBatch );
 
@@ -1509,133 +1492,8 @@ FSlateElementBatch& FSlateElementBatcher::FindBatchForElement(
 	return *ElementBatch;
 }
 
-
-void FSlateElementBatcher::AddVertices( TArray<FSlateVertex>& OutVertices, FSlateElementBatch& ElementBatch, const TArray<FSlateVertex>& VertexBatch )
-{
-	uint32 FirstIndex = OutVertices.Num();
-	OutVertices.AddUninitialized( VertexBatch.Num() );
-
-	uint32 RequiredSize =  VertexBatch.Num() * VertexBatch.GetTypeSize();
-
-	FMemory::Memcpy( &OutVertices[FirstIndex], VertexBatch.GetData(), RequiredSize );
-
-	RequiredVertexMemory += RequiredSize;
-	TotalVertexMemory += VertexBatch.GetAllocatedSize();
-	NumVertices += VertexBatch.Num();
-
-	ElementBatch.VertexOffset = FirstIndex;
-	ElementBatch.NumVertices = VertexBatch.Num();
-}
-
-
-void FSlateElementBatcher::AddIndices( TArray<SlateIndex>& OutIndices, FSlateElementBatch& ElementBatch, const TArray<SlateIndex>& IndexBatch )
-{
-	uint32 FirstIndex = OutIndices.Num();
-	OutIndices.AddUninitialized( IndexBatch.Num() );
-
-	uint32 RequiredSize =  IndexBatch.Num() * IndexBatch.GetTypeSize();
-
-	FMemory::Memcpy( &OutIndices[FirstIndex], IndexBatch.GetData(), RequiredSize );
-
-	RequiredIndexMemory += RequiredSize;
-	TotalIndexMemory += IndexBatch.GetAllocatedSize();
-
-	ElementBatch.IndexOffset = FirstIndex;
-	ElementBatch.NumIndices = IndexBatch.Num();
-}
-
-
-void FSlateElementBatcher::FillBatchBuffers( FSlateWindowElementList& WindowElementList, bool& bRequiresStencilTest )
-{
-	SLATE_CYCLE_COUNTER_SCOPE(GSlateFillBatchBuffers);
-	SCOPE_CYCLE_COUNTER(STAT_SlateFillBatchBuffers);
-
-	TArray<FSlateRenderBatch>& OutRenderBatches = WindowElementList.RenderBatches;
-	TArray<FSlateVertex>& OutBatchedVertices = WindowElementList.BatchedVertices;
-	TArray<SlateIndex>& OutBatchedIndices = WindowElementList.BatchedIndices;
-
-	check( OutRenderBatches.Num() == 0 && OutBatchedVertices.Num() == 0 && OutBatchedIndices.Num() == 0 );
-
-	// Sort by layer
-	LayerToElementBatches.KeySort( TLess<uint32>() );
-
-
-	STAT(NumLayers = FMath::Max<uint32>(NumLayers, LayerToElementBatches.Num() ));
-
-	bRequiresStencilTest = false;
-	// For each element batch add its vertices and indices to the bulk lists.
-	for( TMap< uint32, TSet<FSlateElementBatch> >::TIterator It( LayerToElementBatches ); It; ++It )
-	{
-		TSet<FSlateElementBatch>& ElementBatches = It.Value();
-
-		STAT(NumBatches += ElementBatches.Num());
-		for( TSet<FSlateElementBatch>::TIterator BatchIt(ElementBatches); BatchIt; ++BatchIt )
-		{
-			FSlateElementBatch& ElementBatch = *BatchIt;
-
-			bool bIsMaterial = ElementBatch.GetShaderResource() && ElementBatch.GetShaderResource()->GetType() == ESlateShaderResource::Material;
-
-			if( !ElementBatch.GetCustomDrawer().IsValid() )
-			{
-				if( ElementBatch.GetShaderType() == ESlateShader::LineSegment )
-				{
-					bRequiresStencilTest = true;
-				}
-
-				TArray<FSlateVertex>& BatchVertices = BatchVertexArrays[ ElementBatch.VertexArrayIndex ];
-				// after this loop, we'll be done with the array, so put it back on the free list.
-				VertexArrayFreeList.Add( ElementBatch.VertexArrayIndex );
-
-				TArray<SlateIndex>& BatchIndices = BatchIndexArrays[ ElementBatch.IndexArrayIndex ];
-				// after this loop, we'll be done with the array, so put it back on the free list.
-				IndexArrayFreeList.Add( ElementBatch.IndexArrayIndex );
-
-				// We should have at least some vertices and indices in the batch or none at all
-				check( BatchVertices.Num() > 0 && BatchIndices.Num() > 0  || BatchVertices.Num() == 0 && BatchIndices.Num() == 0 );
-
-				if( BatchVertices.Num() > 0 && BatchIndices.Num() > 0  )
-				{
-					AddVertices( OutBatchedVertices, ElementBatch, BatchVertices );
-					AddIndices( OutBatchedIndices, ElementBatch, BatchIndices );
-
-					OutRenderBatches.Add( FSlateRenderBatch( ElementBatch ) );
-
-					// Done with the batch
-					BatchVertices.Empty(BatchVertices.Num());
-					BatchIndices.Empty(BatchIndices.Num());
-				}
-
-			}
-			else
-			{
-				OutRenderBatches.Add( FSlateRenderBatch( ElementBatch ) );
-			}
-		}
-	}
-
-}
-
-
 void FSlateElementBatcher::ResetBatches()
 {
-	LayerToElementBatches.Reset();
 	bRequiresVsync = false;
 }
 
-
-void FSlateElementBatcher::ResetStats()
-{
-	SET_DWORD_STAT( STAT_SlateNumLayers, NumLayers );
-	SET_DWORD_STAT( STAT_SlateNumBatches, NumBatches );
-	SET_DWORD_STAT( STAT_SlateVertexCount, NumVertices );
-	SET_MEMORY_STAT( STAT_SlateVertexBatchMemory, TotalVertexMemory );
-	SET_MEMORY_STAT( STAT_SlateIndexBatchMemory, TotalIndexMemory );
-
-	NumLayers = 0;
-	NumBatches = 0;
-	NumVertices = 0;
-	RequiredIndexMemory = 0;
-	RequiredVertexMemory = 0;
-	TotalVertexMemory = 0;
-	TotalIndexMemory = 0;
-}
diff --git a/Engine/Source/Runtime/SlateCore/Private/Rendering/SlateDrawBuffer.cpp b/Engine/Source/Runtime/SlateCore/Private/Rendering/SlateDrawBuffer.cpp
--- a/Engine/Source/Runtime/SlateCore/Private/Rendering/SlateDrawBuffer.cpp
+++ b/Engine/Source/Runtime/SlateCore/Private/Rendering/SlateDrawBuffer.cpp
@@ -6,21 +6,58 @@
 /* FSlateDrawBuffer interface
  *****************************************************************************/
 
-FSlateWindowElementList& FSlateDrawBuffer::AddWindowElementList( TSharedRef<SWindow> ForWindow )
+FSlateWindowElementList& FSlateDrawBuffer::AddWindowElementList(TSharedRef<SWindow> ForWindow)
 {
-	int32 Index = WindowElementLists.Add(FSlateWindowElementList(ForWindow));
+	TSharedPtr<FSlateWindowElementList> WindowElements;
 
-	return WindowElementLists[Index];
-}
+	for ( int32 WindowIndex = 0; WindowIndex < WindowElementListsPool.Num(); ++WindowIndex )
+	{
+		WindowElements = WindowElementListsPool[WindowIndex];
+
+		if ( WindowElements->GetWindow() == ForWindow )
+		{
+			WindowElementLists.Add(WindowElements);
+			WindowElementListsPool.RemoveAtSwap(WindowIndex);
+
+			WindowElements->Reset();
+
+			return *WindowElements;
+		}
+	}
+
+	WindowElements = MakeShareable(new FSlateWindowElementList(ForWindow));
+	WindowElementLists.Add(WindowElements);
 
+	return *WindowElements;
+}
 
-bool FSlateDrawBuffer::Lock( )
+bool FSlateDrawBuffer::Lock()
 {
 	return FPlatformAtomics::InterlockedCompareExchange(&Locked, 1, 0) == 0;
 }
 
-
-void FSlateDrawBuffer::Unlock( )
+void FSlateDrawBuffer::Unlock()
 {
 	FPlatformAtomics::InterlockedExchange(&Locked, 0);
 }
+
+void FSlateDrawBuffer::ClearBuffer()
+{
+	// Remove any window elements that are no longer valid.
+	for ( int32 WindowIndex = 0; WindowIndex < WindowElementListsPool.Num(); ++WindowIndex )
+	{
+		if ( WindowElementListsPool[WindowIndex]->GetWindow().IsValid() == false )
+		{
+			WindowElementListsPool.RemoveAtSwap(WindowIndex);
+			--WindowIndex;
+		}
+	}
+
+	// Move all the window elements back into the pool.
+	for ( TSharedPtr<FSlateWindowElementList> ExistingList : WindowElementLists )
+	{
+		WindowElementListsPool.Add(ExistingList);
+	}
+
+	WindowElementLists.Reset();
+}
diff --git a/Engine/Source/Runtime/SlateCore/Public/Rendering/DrawElements.h b/Engine/Source/Runtime/SlateCore/Public/Rendering/DrawElements.h
--- a/Engine/Source/Runtime/SlateCore/Public/Rendering/DrawElements.h
+++ b/Engine/Source/Runtime/SlateCore/Public/Rendering/DrawElements.h
@@ -380,10 +380,6 @@ public:
 	FSlateElementBatch( const FSlateShaderResource* InShaderResource, const FShaderParams& InShaderParams, ESlateShader::Type ShaderType, ESlateDrawPrimitive::Type PrimitiveType, ESlateDrawEffect::Type DrawEffects, ESlateBatchDrawFlag::Type DrawFlags, const TOptional<FShortRect>& ScissorRect )
 		: BatchKey( InShaderParams, ShaderType, PrimitiveType, DrawEffects, DrawFlags, ScissorRect )
 		, ShaderResource(InShaderResource)
-		, VertexOffset(0)
-		, IndexOffset(0)
-		, NumVertices(0)
-		, NumIndices(0)
 		, NumElementsInBatch(0)
 		, VertexArrayIndex(INDEX_NONE)
 		, IndexArrayIndex(INDEX_NONE)
@@ -392,10 +388,6 @@ public:
 	FSlateElementBatch( TWeakPtr<ICustomSlateElement, ESPMode::ThreadSafe> InCustomDrawer, const TOptional<FShortRect>& ScissorRect )
 		: BatchKey( InCustomDrawer, ScissorRect )
 		, ShaderResource( nullptr )
-		, VertexOffset(0)
-		, IndexOffset(0)
-		, NumVertices(0)
-		, NumIndices(0)
 		, NumElementsInBatch(0)
 		, VertexArrayIndex(INDEX_NONE)
 		, IndexArrayIndex(INDEX_NONE)
@@ -483,14 +475,6 @@ private:
 	const FSlateShaderResource* ShaderResource;
 
 public:
-	/** How far into the vertex buffer is this batch*/
-	uint32 VertexOffset;
-	/** How far into the index buffer this batch is*/	
-	uint32 IndexOffset;
-	/** Number of vertices in the batch */
-	uint32 NumVertices;
-	/** Number of indices in the batch */
-	uint32 NumIndices;
 	/** Number of elements in the batch */
 	uint32 NumElementsInBatch;
 	/** Index into an array of vertex arrays where this batches vertices are found (before submitting to the vertex buffer)*/
@@ -502,7 +486,7 @@ public:
 class FSlateRenderBatch
 {
 public:
-	FSlateRenderBatch( const FSlateElementBatch& InBatch )
+	FSlateRenderBatch( const FSlateElementBatch& InBatch, int32 InNumVertices, int32 InNumIndices, int32 InVertexOffset, int32 InIndexOffset )
 		: ShaderParams( InBatch.GetShaderParams() )
 		, Texture( InBatch.GetShaderResource() )
 		, CustomDrawer( InBatch.GetCustomDrawer() )
@@ -511,10 +495,12 @@ public:
 		, DrawPrimitiveType( InBatch.GetPrimitiveType() )
 		, DrawEffects( InBatch.GetDrawEffects() )
 		, ScissorRect( InBatch.GetScissorRect() )
-		, VertexOffset( InBatch.VertexOffset )
-		, IndexOffset( InBatch.IndexOffset )
-		, NumVertices( InBatch.NumVertices )
-		, NumIndices( InBatch.NumIndices )
+		, VertexArrayIndex( InBatch.VertexArrayIndex )
+		, IndexArrayIndex( InBatch.IndexArrayIndex )
+		, VertexOffset( InVertexOffset )
+		, IndexOffset( InIndexOffset )
+		, NumVertices( InNumVertices )
+		, NumIndices( InNumIndices )
 	{}
 
 public:
@@ -534,7 +520,11 @@ public:
 	const ESlateDrawEffect::Type DrawEffects;
 
 	const TOptional<FShortRect> ScissorRect;
-
+	
+	/** Index into the vertex array pool */
+	const int32 VertexArrayIndex;
+	/** Index into the index array pool */
+	const int32 IndexArrayIndex;
 	/** How far into the vertex buffer is this batch*/
 	const uint32 VertexOffset;
 	/** How far into the index buffer this batch is*/	
@@ -545,6 +535,103 @@ public:
 	const uint32 NumIndices;
 };
 
+
+typedef TArray<FSlateElementBatch, TInlineAllocator<1>> FElementBatchArray;
+
+typedef TMap<uint32,FElementBatchArray> FElementBatchMap ;
+
+class FSlateBatchData
+{
+public:
+	FSlateBatchData()
+		: NumBatchedVertices(0)
+		, NumBatchedIndices(0)
+		, NumLayers(0)
+	{}
+
+	void Reset();
+
+	FElementBatchMap& GetElementBatchMap() { return LayerToElementBatches; }
+
+	/**
+	 * Returns a list of element batches for this window
+	 */
+	const TArray<FSlateRenderBatch>& GetRenderBatches() const { return RenderBatches; }
+
+	/**
+	 * Assigns a vertex array from the pool which is appropriate for the batch.  Creates a new array if needed
+	 */
+	void AssignVertexArrayToBatch( FSlateElementBatch& Batch );
+
+	/**
+	 * Assigns an index array from the pool which is appropriate for the batch.  Creates a new array if needed
+	 */
+	void AssignIndexArrayToBatch( FSlateElementBatch& Batch );
+
+	/** @return the list of vertices for a batch */
+	TArray<FSlateVertex>& GetBatchVertexList( FSlateElementBatch& Batch ) { return BatchVertexArrays[Batch.VertexArrayIndex]; }
+
+	/** @return the list of indices for a batch */
+	TArray<SlateIndex>& GetBatchIndexList( FSlateElementBatch& Batch ) { return BatchIndexArrays[Batch.IndexArrayIndex]; }
+
+	/** @return The total number of batched vertices */
+	int32 GetNumBatchedVertices() const { return NumBatchedVertices; }
+
+	/** @return The total number of batched indices */
+	int32 GetNumBatchedIndices() const { return NumBatchedIndices; }
+
+	/** @return Total number of batched layers */
+	int32 GetNumLayers() const { return NumLayers; }
+
+	/** 
+	 * Fills batch data into the actual vertex and index buffer
+	 *
+	 * @param VertexBuffer	Pointer to the actual memory for the vertex buffer 
+	 * @param IndexBuffer	Pointer to the actual memory for an index buffer
+	 */
+	SLATECORE_API void FillVertexAndIndexBuffer( uint8* VertexBuffer, uint8* IndexBuffer );
+
+	/** 
+	 * Creates rendering data from batched elements
+	 */
+	SLATECORE_API void CreateRenderBatches();
+
+private:
+	void AddRenderBatch( const FSlateElementBatch& InElementBatch, int32 InNumVertices, int32 InNumIndices, int32 InVertexOffset, int32 InIndexOffset )
+	{
+		NumBatchedVertices += InNumVertices;
+		NumBatchedIndices += InNumIndices;
+
+		RenderBatches.Add( FSlateRenderBatch(InElementBatch, InNumVertices, InNumIndices, InVertexOffset, InIndexOffset) );
+	}
+private:
+
+	// Element batch maps sorted by layer.
+	FElementBatchMap LayerToElementBatches;
+
+	// Array of vertex lists that are currently free (have no elements in them).
+	TArray<uint32> VertexArrayFreeList;
+
+	// Array of index lists that are currently free (have no elements in them).
+	TArray<uint32> IndexArrayFreeList;
+
+	// Array of vertex lists for batching vertices. We use this method for quickly resetting the arrays without deleting memory.
+	TArray<TArray<FSlateVertex>> BatchVertexArrays;
+
+	// Array of vertex lists for batching indices. We use this method for quickly resetting the arrays without deleting memory.
+	TArray<TArray<SlateIndex>> BatchIndexArrays;
+
+	/** List of element batches sorted by later for use in rendering (for threaded renderers, can only be accessed from the render thread)*/
+	TArray<FSlateRenderBatch> RenderBatches;
+
+	int32 NumBatchedVertices;
+
+	int32 NumBatchedIndices;
+
+	int32 NumLayers;
+
+};
+
 /**
  * Represents a top level window and its draw elements.
  */
@@ -624,45 +711,25 @@ public:
 	/**
 	 * Remove all the elements from this draw list.
 	 */
-	void Empty()
+	void Reset()
 	{
-		DrawElements.Empty();
-		RenderBatches.Empty();
-		BatchedVertices.Empty();
-		BatchedIndices.Empty();
+		DrawElements.Reset();
+		BatchData.Reset();
+		DeferredPaintList.Reset();
 	}
-	
-	/**
-	 * Returns a list of element batches for this window
-	 */
-	const TArray<FSlateRenderBatch>& GetRenderBatches() const { return RenderBatches; }
 
-	/**
-	 * Returns all the vertices for every draw element in this window
-	 */
-	const TArray<FSlateVertex>& GetBatchedVertices() const { return BatchedVertices; }
+	FSlateBatchData& GetBatchData() { return BatchData; }
 
-	/**
-	 * Returns all the indices for every draw element in this window
-	 */
-	const TArray<SlateIndex>& GetBatchedIndices() const { return BatchedIndices; }
+private:
+	/** Batched data used for rendering */
+	FSlateBatchData BatchData;
 
-protected:
 	/** The top level window which these elements are being drawn on */
 	TWeakPtr<SWindow> TopLevelWindow;
 	
 	/** List of draw elements for the window */
 	TArray<FSlateDrawElement> DrawElements;
 
-	/** List of element batches sorted by later for use in rendering (for threaded renderers, can only be accessed from the render thread)*/
-	TArray<FSlateRenderBatch> RenderBatches;
-
-	/** All the vertices for every batch. */
-	TArray<FSlateVertex> BatchedVertices;
-
-	/** All the indices for every batch. */
-	TArray<SlateIndex> BatchedIndices;
-
 	/**
 	 * Some widgets want their logical children to appear at a different "layer" in the physical hierarchy.
 	 * We accomplish this by deferring their painting.
diff --git a/Engine/Source/Runtime/SlateCore/Public/Rendering/ElementBatcher.h b/Engine/Source/Runtime/SlateCore/Public/Rendering/ElementBatcher.h
--- a/Engine/Source/Runtime/SlateCore/Public/Rendering/ElementBatcher.h
+++ b/Engine/Source/Runtime/SlateCore/Public/Rendering/ElementBatcher.h
@@ -5,7 +5,6 @@
 
 class FSlateShaderResource;
 
-
 /**
  * A class which batches Slate elements for rendering
  */
@@ -21,15 +20,7 @@ public:
 	 * 
 	 * @param DrawElements	The elements to batch
 	 */
-	void AddElements( const TArray<FSlateDrawElement>& DrawElements );
-
-	/**
-	 * Populates the window element list with batched data for use in rendering
-	 *
-	 * @param WindowElementList		The window element list to update
-	 * @param bRequiresStencilTest	Will be set to true if stencil testing should be enabled when drawing this viewport (currently for anti-aliased lines)
-	 */
-	void FillBatchBuffers(FSlateWindowElementList& WindowElementList, bool& bRequiresStencilTest);
+	void AddElements( FSlateWindowElementList& ElementList );
 
 	/**
 	 * Returns true if the elements in this batcher require v-sync.
@@ -40,15 +31,12 @@ public:
 	 * Resets all stored data accumulated during the batching process
 	 */
 	void ResetBatches();
-
-	void ResetStats();
-
 private:
 
 	/** 
 	 * Creates vertices necessary to draw a Quad element 
 	 */
-	void AddQuadElement( const FSlateDrawElement& DrawElement, FColor Color = FColor(255,255,255));
+	void AddQuadElement( const FSlateDrawElement& DrawElement, FColor Color = FColor::White);
 
 	/** 
 	 * Creates vertices necessary to draw a 3x3 element
@@ -116,26 +104,9 @@ private:
 											 ESlateDrawEffect::Type DrawEffects, 
 											 ESlateBatchDrawFlag::Type DrawFlags,
 											 const TOptional<FShortRect>& ScissorRect);
-
-	void AddVertices( TArray<FSlateVertex>& OutVertices, FSlateElementBatch& ElementBatch, const TArray<FSlateVertex>& VertexBatch );
-
-	void AddIndices( TArray<SlateIndex>& OutIndices, FSlateElementBatch& ElementBatch, const TArray<SlateIndex>& IndexBatch );
-
 private:
-	// Element batch maps sorted by layer.
-	TMap<uint32, TSet<FSlateElementBatch>> LayerToElementBatches;
-
-	// Array of vertex lists that are currently free (have no elements in them).
-	TArray<uint32> VertexArrayFreeList;
-
-	// Array of index lists that are currently free (have no elements in them).
-	TArray<uint32> IndexArrayFreeList;
-
-	// Array of vertex lists for batching vertices. We use this method for quickly resetting the arrays without deleting memory.
-	TArray<TArray<FSlateVertex>> BatchVertexArrays;
-
-	// Array of vertex lists for batching indices. We use this method for quickly resetting the arrays without deleting memory.
-	TArray<TArray<SlateIndex>> BatchIndexArrays;
+	/** Batch data currently being filled in */
+	FSlateBatchData* BatchData;
 
 	/** Resource manager for accessing shader resources */
 	FSlateShaderResourceManager& ResourceManager;
@@ -148,12 +119,4 @@ private:
 
 	// true if any element in the batch requires vsync.
 	bool bRequiresVsync;
-
-	uint32 NumVertices;
-	uint32 NumBatches;
-	uint32 NumLayers;
-	uint32 TotalVertexMemory;
-	uint32 RequiredVertexMemory;
-	uint32 TotalIndexMemory;
-	uint32 RequiredIndexMemory;
 };
diff --git a/Engine/Source/Runtime/SlateCore/Public/Rendering/RenderingPolicy.h b/Engine/Source/Runtime/SlateCore/Public/Rendering/RenderingPolicy.h
--- a/Engine/Source/Runtime/SlateCore/Public/Rendering/RenderingPolicy.h
+++ b/Engine/Source/Runtime/SlateCore/Public/Rendering/RenderingPolicy.h
@@ -27,13 +27,11 @@ public:
 	/** Virtual constructor. */
 	virtual ~FSlateRenderingPolicy( ) { }
 
-	virtual void UpdateBuffers( const FSlateWindowElementList& WindowElementList ) = 0;
-
-	virtual TSharedRef<class FSlateFontCache> GetFontCache( ) = 0;
+	virtual TSharedRef<class FSlateFontCache> GetFontCache() = 0;
 
 	virtual TSharedRef<class FSlateShaderResourceManager> GetResourceManager() = 0;
 
-	float GetPixelCenterOffset( ) const
+	float GetPixelCenterOffset() const
 	{
 		return PixelCenterOffset;
 	}
diff --git a/Engine/Source/Runtime/SlateCore/Public/Rendering/SlateDrawBuffer.h b/Engine/Source/Runtime/SlateCore/Public/Rendering/SlateDrawBuffer.h
--- a/Engine/Source/Runtime/SlateCore/Public/Rendering/SlateDrawBuffer.h
+++ b/Engine/Source/Runtime/SlateCore/Public/Rendering/SlateDrawBuffer.h
@@ -18,22 +18,19 @@ public:
 public:
 
 	/** Removes all data from the buffer. */
-	void ClearBuffer( )
-	{
-		WindowElementLists.Empty();
-	}
+	void ClearBuffer();
 
 	/**
 	 * Creates a new FSlateWindowElementList and returns a reference to it so it can have draw elements added to it
 	 *
 	 * @param ForWindow    The window for which we are creating a list of paint elements.
 	 */
-	FSlateWindowElementList& AddWindowElementList( TSharedRef<SWindow> ForWindow );
+	FSlateWindowElementList& AddWindowElementList(TSharedRef<SWindow> ForWindow);
 
 	/**
 	 * Gets all window element lists in this buffer.
 	 */
-	TArray<FSlateWindowElementList>& GetWindowElementLists( )
+	TArray< TSharedPtr<FSlateWindowElementList> >& GetWindowElementLists()
 	{
 		return WindowElementLists;
 	}
@@ -56,7 +53,11 @@ public:
 protected:
 
 	// List of window element lists.
-	TArray<FSlateWindowElementList> WindowElementLists;
+	TArray< TSharedPtr<FSlateWindowElementList> > WindowElementLists;
+
+	// List of window element lists that we store from the previous frame 
+	// that we restore if they're requested again.
+	TArray< TSharedPtr<FSlateWindowElementList> > WindowElementListsPool;
 
 	// 1 if this buffer is locked, 0 otherwise.
 	int32 Locked;
diff --git a/Engine/Source/Runtime/SlateRHIRenderer/Private/Slate3DRenderer.cpp b/Engine/Source/Runtime/SlateRHIRenderer/Private/Slate3DRenderer.cpp
--- a/Engine/Source/Runtime/SlateRHIRenderer/Private/Slate3DRenderer.cpp
+++ b/Engine/Source/Runtime/SlateRHIRenderer/Private/Slate3DRenderer.cpp
@@ -51,12 +51,12 @@ void FSlate3DRenderer::DrawWindow_GameThread(FSlateDrawBuffer& DrawBuffer)
 {
 	check( IsInGameThread() );
 
-	TArray<FSlateWindowElementList>& WindowElementLists = DrawBuffer.GetWindowElementLists();
+	TArray<TSharedPtr<FSlateWindowElementList>>& WindowElementLists = DrawBuffer.GetWindowElementLists();
 
 	// Only one window element list for now.
 	check( WindowElementLists.Num() == 1);
 
-	FSlateWindowElementList& ElementList = WindowElementLists[0];
+	FSlateWindowElementList& ElementList = *WindowElementLists[0];
 
 	TSharedPtr<SWindow> Window = ElementList.GetWindow();
 
@@ -66,14 +66,12 @@ void FSlate3DRenderer::DrawWindow_GameThread(FSlateDrawBuffer& DrawBuffer)
 		if (WindowSize.X > 0 && WindowSize.Y > 0)
 		{
 			// Add all elements for this window to the element batcher
-			ElementBatcher->AddElements(ElementList.GetDrawElements());
+			ElementBatcher->AddElements(ElementList);
 
 			// Update the font cache with new text after elements are batched
 			FontCache->UpdateCache();
 
 			bool temp = false;
-			// Populate the element list with batched vertices and indices
-			ElementBatcher->FillBatchBuffers(ElementList, temp);
 
 
 			// All elements for this window have been batched and rendering data updated
@@ -82,25 +80,23 @@ void FSlate3DRenderer::DrawWindow_GameThread(FSlateDrawBuffer& DrawBuffer)
 	}
 
 	FontCache->ConditionalFlushCache();
-
-	ElementBatcher->ResetStats();
-
 }
 
 void FSlate3DRenderer::DrawWindowToTarget_RenderThread( FRHICommandListImmediate& RHICmdList, UTextureRenderTarget2D* RenderTarget, FSlateDrawBuffer& InDrawBuffer )
 {
-	const FSlateWindowElementList& WindowElementList = InDrawBuffer.GetWindowElementLists()[0];
+	FSlateBatchData& BatchData = InDrawBuffer.GetWindowElementLists()[0]->GetBatchData();
 
+	BatchData.CreateRenderBatches();
 	RenderTargetPolicy->BeginDrawingWindows();
 
-	RenderTargetPolicy->UpdateBuffers( WindowElementList );
+	RenderTargetPolicy->UpdateVertexAndIndexBuffers( RHICmdList, BatchData );
 	FTextureRenderTarget2DResource* RenderTargetResource = static_cast<FTextureRenderTarget2DResource*>( RenderTarget->GetRenderTargetResource() );
 	
 	SetRenderTarget( RHICmdList, RenderTargetResource->GetTextureRHI(), FTextureRHIParamRef() );
 	RHICmdList.Clear( true, RenderTarget->ClearColor, false, 0.0f, false, 0x00, FIntRect() );
 
 	FMatrix ProjectionMatrix = FSlateRHIRenderer::CreateProjectionMatrix( RenderTarget->SizeX, RenderTarget->SizeY );
-	if (WindowElementList.GetRenderBatches().Num() > 0)
+	if (BatchData.GetRenderBatches().Num() > 0)
 	{
 		FSlateBackBuffer BackBufferTarget(RenderTargetResource->GetTextureRHI(), FIntPoint(RenderTarget->SizeX, RenderTarget->SizeY ) );
 
@@ -112,9 +108,9 @@ void FSlate3DRenderer::DrawWindowToTarget_RenderThread( FRHICommandListImmediate
 			RHICmdList,
 			BackBufferTarget,
 			ProjectionMatrix,
-			WindowElementList.GetRenderBatches(),
+			BatchData.GetRenderBatches(),
 			bAllowSwitchVerticalAxis
-			);
+		);
 	}
 
 	InDrawBuffer.Unlock();
diff --git a/Engine/Source/Runtime/SlateRHIRenderer/Private/SlateRHIRenderer.cpp b/Engine/Source/Runtime/SlateRHIRenderer/Private/SlateRHIRenderer.cpp
--- a/Engine/Source/Runtime/SlateRHIRenderer/Private/SlateRHIRenderer.cpp
+++ b/Engine/Source/Runtime/SlateRHIRenderer/Private/SlateRHIRenderer.cpp
@@ -390,7 +390,7 @@ void FSlateRHIRenderer::OnWindowDestroyed( const TSharedRef<SWindow>& InWindow )
 }
 
 /** Draws windows from a FSlateDrawBuffer on the render thread */
-void FSlateRHIRenderer::DrawWindow_RenderThread(FRHICommandListImmediate& RHICmdList, const FViewportInfo& ViewportInfo, const FSlateWindowElementList& WindowElementList, bool bLockToVsync, bool bClear)
+void FSlateRHIRenderer::DrawWindow_RenderThread(FRHICommandListImmediate& RHICmdList, const FViewportInfo& ViewportInfo, FSlateWindowElementList& WindowElementList, bool bLockToVsync, bool bClear)
 {
 	SCOPED_DRAW_EVENT(RHICmdList, SlateUI);
 
@@ -400,8 +400,11 @@ void FSlateRHIRenderer::DrawWindow_RenderThread(FRHICommandListImmediate& RHICmd
 	{
 		SCOPE_CYCLE_COUNTER( STAT_SlateRenderingRTTime );
 
-		// Update the vertex and index buffer		
-		RenderingPolicy->UpdateVertexAndIndexBuffers(RHICmdList, WindowElementList);
+		FSlateBatchData& BatchData = WindowElementList.GetBatchData();
+
+		// Update the vertex and index buffer	
+		BatchData.CreateRenderBatches();	
+		RenderingPolicy->UpdateVertexAndIndexBuffers(RHICmdList, BatchData);
 		// should have been created by the game thread
 		check( IsValidRef(ViewportInfo.ViewportRHI) );
 
@@ -437,7 +440,7 @@ void FSlateRHIRenderer::DrawWindow_RenderThread(FRHICommandListImmediate& RHICmd
 #if DEBUG_OVERDRAW
 		RHIClear(true, FLinearColor::Black, false, 0.0f, true, 0x00, FIntRect());
 #endif
-		if( WindowElementList.GetRenderBatches().Num() > 0 )
+		if( BatchData.GetRenderBatches().Num() > 0 )
 		{
 			FSlateBackBuffer BackBufferTarget( BackBuffer, FIntPoint( ViewportWidth, ViewportHeight ) );
 
@@ -446,7 +449,7 @@ void FSlateRHIRenderer::DrawWindow_RenderThread(FRHICommandListImmediate& RHICmd
 				RHICmdList,
 				BackBufferTarget,
 				ViewMatrix*ViewportInfo.ProjectionMatrix,
-				WindowElementList.GetRenderBatches()
+				BatchData.GetRenderBatches()
 			);
 		}
 	}
@@ -568,10 +571,10 @@ void FSlateRHIRenderer::DrawWindows_Private( FSlateDrawBuffer& WindowDrawBuffer
 	ResourceManager->UpdateTextureAtlases();
 
 	// Iterate through each element list and set up an RHI window for it if needed
-	TArray<FSlateWindowElementList>& WindowElementLists = WindowDrawBuffer.GetWindowElementLists();
+	TArray<TSharedPtr<FSlateWindowElementList>>& WindowElementLists = WindowDrawBuffer.GetWindowElementLists();
 	for( int32 ListIndex = 0; ListIndex < WindowElementLists.Num(); ++ListIndex )
 	{
-		FSlateWindowElementList& ElementList = WindowElementLists[ListIndex];
+		FSlateWindowElementList& ElementList = *WindowElementLists[ListIndex];
 
 		TSharedPtr<SWindow> Window = ElementList.GetWindow();
 
@@ -581,16 +584,14 @@ void FSlateRHIRenderer::DrawWindows_Private( FSlateDrawBuffer& WindowDrawBuffer
 			if ( WindowSize.X > 0 && WindowSize.Y > 0 )
 			{
 				// Add all elements for this window to the element batcher
-				ElementBatcher->AddElements( ElementList.GetDrawElements() );
+				ElementBatcher->AddElements( ElementList );
 
 				// Update the font cache with new text after elements are batched
 				FontCache->UpdateCache();
 
 				bool bRequiresStencilTest = false;
 				bool bLockToVsync = false;
-				bool temp = false;
-				// Populate the element list with batched vertices and indicies
-				ElementBatcher->FillBatchBuffers( ElementList, temp );
+
 				bLockToVsync = ElementBatcher->RequiresVsync();
 
 				if( !GIsEditor )
@@ -688,8 +689,6 @@ void FSlateRHIRenderer::DrawWindows_Private( FSlateDrawBuffer& WindowDrawBuffer
 
 	// flush the cache if needed
 	FontCache->ConditionalFlushCache();
-
-	ElementBatcher->ResetStats();
 }
 
 
@@ -923,10 +922,7 @@ void FSlateRHIRenderer::CopyWindowsToVirtualScreenBuffer(const TArray<FString>&
 			FSlateRect(0, 0, VirtualScreenSize.X, VirtualScreenSize.Y));
 	}
 	
-	ElementBatcher->AddElements(WindowElementList->GetDrawElements());
-	bool bRequiresStencilTest = false;
-	ElementBatcher->FillBatchBuffers(*WindowElementList, bRequiresStencilTest );
-	check(!bRequiresStencilTest);
+	ElementBatcher->AddElements(*WindowElementList);
 	ElementBatcher->ResetBatches();
 	
 	struct FWriteMouseCursorAndKeyPressesContext
@@ -951,13 +947,17 @@ void FSlateRHIRenderer::CopyWindowsToVirtualScreenBuffer(const TArray<FString>&
 	{
 		RHICmdList.SetBlendState(TStaticBlendState<CW_RGBA, BO_Add, BF_SourceAlpha, BF_InverseSourceAlpha, BO_Add, BF_Zero, BF_One>::GetRHI());
 		
-		Context.RenderPolicy->UpdateBuffers(*Context.SlateElementList);
-		if( Context.SlateElementList->GetRenderBatches().Num() > 0 )
+		FSlateBatchData& BatchData = Context.SlateElementList->GetBatchData();
+
+		BatchData.CreateRenderBatches();
+
+		Context.RenderPolicy->UpdateVertexAndIndexBuffers(RHICmdList, BatchData);
+		if( BatchData.GetRenderBatches().Num() > 0 )
 		{
 			FTexture2DRHIRef UnusedTargetTexture;
 			FSlateBackBuffer UnusedTarget( UnusedTargetTexture, FIntPoint::ZeroValue );
 
-			Context.RenderPolicy->DrawElements(RHICmdList, UnusedTarget, CreateProjectionMatrix(Context.ViewportSize.X, Context.ViewportSize.Y), Context.SlateElementList->GetRenderBatches());
+			Context.RenderPolicy->DrawElements(RHICmdList, UnusedTarget, CreateProjectionMatrix(Context.ViewportSize.X, Context.ViewportSize.Y),BatchData.GetRenderBatches());
 		}
 	});
 
diff --git a/Engine/Source/Runtime/SlateRHIRenderer/Private/SlateRHIRenderer.h b/Engine/Source/Runtime/SlateRHIRenderer/Private/SlateRHIRenderer.h
--- a/Engine/Source/Runtime/SlateRHIRenderer/Private/SlateRHIRenderer.h
+++ b/Engine/Source/Runtime/SlateRHIRenderer/Private/SlateRHIRenderer.h
@@ -188,7 +188,7 @@ public:
 	virtual ISlateAtlasProvider* GetTextureAtlasProvider() override;
 
 	/** Draws windows from a FSlateDrawBuffer on the render thread */
-	void DrawWindow_RenderThread(FRHICommandListImmediate& RHICmdList, const FSlateRHIRenderer::FViewportInfo& ViewportInfo, const FSlateWindowElementList& WindowElementList, bool bLockToVsync, bool bClear);
+	void DrawWindow_RenderThread(FRHICommandListImmediate& RHICmdList, const FSlateRHIRenderer::FViewportInfo& ViewportInfo, FSlateWindowElementList& WindowElementList, bool bLockToVsync, bool bClear);
 
 
 	/**
diff --git a/Engine/Source/Runtime/SlateRHIRenderer/Private/SlateRHIRenderingPolicy.cpp b/Engine/Source/Runtime/SlateRHIRenderer/Private/SlateRHIRenderingPolicy.cpp
--- a/Engine/Source/Runtime/SlateRHIRenderer/Private/SlateRHIRenderingPolicy.cpp
+++ b/Engine/Source/Runtime/SlateRHIRenderer/Private/SlateRHIRenderingPolicy.cpp
@@ -12,13 +12,20 @@
 #include "SlateMaterialResource.h"
 
 DECLARE_CYCLE_STAT(TEXT("Update Buffers RT"), STAT_SlateUpdateBufferRTTime, STATGROUP_Slate);
+DECLARE_CYCLE_STAT(TEXT("PreFill Buffers RT"), STAT_SlatePreFullBufferRTTime, STATGROUP_Slate);
 DECLARE_CYCLE_STAT(TEXT("Draw Time RT"), STAT_SlateDrawTime, STATGROUP_Slate);
+DECLARE_DWORD_COUNTER_STAT(TEXT("Num Layers"), STAT_SlateNumLayers, STATGROUP_Slate);
+DECLARE_DWORD_COUNTER_STAT(TEXT("Num Batches"), STAT_SlateNumBatches, STATGROUP_Slate);
+DECLARE_DWORD_COUNTER_STAT(TEXT("Num Vertices"), STAT_SlateVertexCount, STATGROUP_Slate);
 
+DECLARE_MEMORY_STAT(TEXT("Batch Vertex Memory"), STAT_SlateVertexBatchMemory, STATGROUP_SlateMemory);
+DECLARE_MEMORY_STAT(TEXT("Batch Index Memory"), STAT_SlateIndexBatchMemory, STATGROUP_SlateMemory);
 DECLARE_MEMORY_STAT(TEXT("Vertex Buffer Memory"), STAT_SlateVertexBufferMemory, STATGROUP_SlateMemory);
 DECLARE_MEMORY_STAT(TEXT("Index Buffer Memory"), STAT_SlateIndexBufferMemory, STATGROUP_SlateMemory);
 
 FSlateElementIndexBuffer::FSlateElementIndexBuffer()
 	: BufferSize(0)	 
+	, MinBufferSize(0)
 	, BufferUsageSize(0)
 {
 
@@ -29,91 +36,92 @@ FSlateElementIndexBuffer::~FSlateElementIndexBuffer()
 
 }
 
+void FSlateElementIndexBuffer::Init( int32 MinNumIndices )
+{
+	MinBufferSize = sizeof(SlateIndex) * FMath::Max( MinNumIndices, 200 );
+
+	BeginInitResource(this);
+}
+
+void FSlateElementIndexBuffer::Destroy()
+{
+	BeginReleaseResource(this);
+}
+
+
 /** Initializes the index buffers RHI resource. */
 void FSlateElementIndexBuffer::InitDynamicRHI()
 {
-	check( IsInRenderingThread() );
+	checkSlow( IsInRenderingThread() );
 
-	if( BufferSize == 0 )
-	{
-		BufferSize = 200 * sizeof(SlateIndex);
-	}
+	check( MinBufferSize > 0 );
+
+	BufferSize = MinBufferSize;
 
 	FRHIResourceCreateInfo CreateInfo;
-	IndexBufferRHI = RHICreateIndexBuffer( sizeof(SlateIndex), BufferSize, BUF_Dynamic, CreateInfo );
+	IndexBufferRHI = RHICreateIndexBuffer( sizeof(SlateIndex), MinBufferSize, BUF_Dynamic, CreateInfo );
 	check( IsValidRef(IndexBufferRHI) );
 }
 
 /** Resizes the buffer to the passed in size.  Preserves internal data */
-void FSlateElementIndexBuffer::ResizeBuffer( uint32 NewSizeBytes )
+void FSlateElementIndexBuffer::ResizeBuffer( int32 NewSizeBytes )
 {
-	check( IsInRenderingThread() );
+	checkSlow( IsInRenderingThread() );
 
-	if( NewSizeBytes != 0 && NewSizeBytes != BufferSize)
+	int32 FinalSize = FMath::Max( NewSizeBytes, MinBufferSize );
+
+	if( FinalSize != 0 && FinalSize != BufferSize )
 	{
 		IndexBufferRHI.SafeRelease();
 		FRHIResourceCreateInfo CreateInfo;
-		IndexBufferRHI = RHICreateIndexBuffer( sizeof(SlateIndex), NewSizeBytes, BUF_Dynamic, CreateInfo );
+		IndexBufferRHI = RHICreateIndexBuffer( sizeof(SlateIndex), FinalSize, BUF_Dynamic, CreateInfo );
 		check(IsValidRef(IndexBufferRHI));
 
-		BufferSize = NewSizeBytes;
+		BufferSize = FinalSize;
 	}
 }
 
-void FSlateElementIndexBuffer::PreFillBuffer(const TArray<SlateIndex>& InIndices, bool bShrinkToFit)
+void FSlateElementIndexBuffer::PreFillBuffer(int32 RequiredIndexCount, bool bShrinkToMinSize)
 {
-	check(IsInRenderingThread());
+	//SCOPE_CYCLE_COUNTER( STAT_SlatePreFullBufferRTTime );
 
-	if (InIndices.Num())
-	{
-		uint32 NumIndices = InIndices.Num();
+	checkSlow(IsInRenderingThread());
 
-		uint32 RequiredBufferSize = NumIndices*sizeof(SlateIndex);
+	if (RequiredIndexCount > 0)
+	{
+		int32 RequiredBufferSize = RequiredIndexCount*sizeof(SlateIndex);
 
 		// resize if needed
-		if (RequiredBufferSize > GetBufferSize() || bShrinkToFit)
+		if (RequiredBufferSize > GetBufferSize() || bShrinkToMinSize)
 		{
 			// Use array resize techniques for the vertex buffer
-			ResizeBuffer(InIndices.GetAllocatedSize());
+			ResizeBuffer(RequiredBufferSize);
 		}
 
-		BufferUsageSize += RequiredBufferSize;		
+		BufferUsageSize = RequiredBufferSize;		
 	}
 }
 
-void FSlateElementIndexBuffer::FillBuffer_RenderThread( const TArray<SlateIndex>& InIndices )
+void* FSlateElementIndexBuffer::LockBuffer_RenderThread(int32 NumIndices)
 {
-	check( IsInRenderingThread() );
-
-	if( InIndices.Num() )
-	{
-		uint32 NumIndices = InIndices.Num();
-
-		uint32 RequiredBufferSize = NumIndices*sizeof(SlateIndex);		
-
-		void* IndicesPtr = RHILockIndexBuffer( IndexBufferRHI, 0, RequiredBufferSize, RLM_WriteOnly );
-
-		FMemory::Memcpy( IndicesPtr, InIndices.GetData(), RequiredBufferSize );
-
-		RHIUnlockIndexBuffer(IndexBufferRHI);
-	}
+	uint32 RequiredBufferSize = NumIndices*sizeof(SlateIndex);		
+	return RHILockIndexBuffer( IndexBufferRHI, 0, RequiredBufferSize, RLM_WriteOnly );
 }
 
-void FSlateElementIndexBuffer::FillBuffer_RHIThread(const TArray<SlateIndex>& InIndices )
+void FSlateElementIndexBuffer::UnlockBuffer_RenderThread()
 {
-	check(IsInRHIThread() || (!IsInRenderingThread() && !IsInGameThread()));
-
-	if (InIndices.Num())
-	{
-		uint32 NumIndices = InIndices.Num();
+	RHIUnlockIndexBuffer( IndexBufferRHI );
+}
 
-		uint32 RequiredBufferSize = NumIndices*sizeof(SlateIndex);		
+void* FSlateElementIndexBuffer::LockBuffer_RHIThread(int32 NumIndices)
+{
+	uint32 RequiredBufferSize = NumIndices*sizeof(SlateIndex);		
+	return GDynamicRHI->RHILockIndexBuffer( IndexBufferRHI, 0, RequiredBufferSize, RLM_WriteOnly );
+}
 
-		//use non-flushing DynamicRHI version when on RHIThread.
-		void* IndicesPtr = GDynamicRHI->RHILockIndexBuffer(IndexBufferRHI, 0, RequiredBufferSize, RLM_WriteOnly);
-		FMemory::Memcpy(IndicesPtr, InIndices.GetData(), RequiredBufferSize);
-		GDynamicRHI->RHIUnlockIndexBuffer(IndexBufferRHI);		
-	}
+void FSlateElementIndexBuffer::UnlockBuffer_RHIThread()
+{
+	GDynamicRHI->RHIUnlockIndexBuffer( IndexBufferRHI );
 }
 
 /** Releases the index buffers RHI resource. */
@@ -140,46 +148,57 @@ FSlateRHIRenderingPolicy::~FSlateRHIRenderingPolicy()
 
 void FSlateRHIRenderingPolicy::InitResources()
 {
-	for( int32 BufferIndex = 0; BufferIndex < 2; ++BufferIndex )
+	int32 NumVertices = 200;
+	if( GConfig )
+	{
+		int32 NumVertsInConfig = 0;
+		GConfig->GetInt( TEXT("SlateRenderer"), TEXT("NumPreallocatedVertices"), NumVertsInConfig, GEngineIni);
+		// Always create a little space but never allow it to get too high
+
+#if !SLATE_USE_32BIT_INDICES
+		NumVertices = FMath::Clamp(NumVertsInConfig, 200, 65535);
+#else
+		NumVertices = FMath::Clamp(NumVertsInConfig, 200, 1000000);
+#endif
+
+		UE_LOG(LogSlate, Verbose, TEXT("Allocating space for %d vertices"), NumVertices);
+	}
+
+	for( int32 BufferIndex = 0; BufferIndex < SlateRHIConstants::NumBuffers; ++BufferIndex )
 	{
-		BeginInitResource(&VertexBuffers[BufferIndex]);
-		BeginInitResource(&IndexBuffers[BufferIndex]);
+		VertexBuffers[BufferIndex].Init(NumVertices);
+		IndexBuffers[BufferIndex].Init(NumVertices);
 	}
 }
 
 void FSlateRHIRenderingPolicy::ReleaseResources()
 {
-	for( int32 BufferIndex = 0; BufferIndex < 2; ++BufferIndex )
+	for( int32 BufferIndex = 0; BufferIndex < SlateRHIConstants::NumBuffers; ++BufferIndex )
 	{
-		BeginReleaseResource(&VertexBuffers[BufferIndex]);
-		BeginReleaseResource(&IndexBuffers[BufferIndex]);
+		VertexBuffers[BufferIndex].Destroy();
+		IndexBuffers[BufferIndex].Destroy();
 	}
 }
 
 void FSlateRHIRenderingPolicy::BeginDrawingWindows()
 {
 	check( IsInRenderingThread() );
-
-	for( int32 BufferIndex = 0; BufferIndex < 2; ++BufferIndex )
-	{
-		VertexBuffers[CurrentBufferIndex].ResetBufferUsage();
-		IndexBuffers[CurrentBufferIndex].ResetBufferUsage();
-	}
 }
 
 void FSlateRHIRenderingPolicy::EndDrawingWindows()
 {
 	check( IsInRenderingThread() );
 
-	bShouldShrinkResources = false;
-
 	uint32 TotalVertexBufferMemory = 0;
 	uint32 TotalIndexBufferMemory = 0;
 
 	uint32 TotalVertexBufferUsage = 0;
 	uint32 TotalIndexBufferUsage = 0;
 
-	for( int32 BufferIndex = 0; BufferIndex < 2; ++BufferIndex )
+	uint32 MinVertexBufferSize = VertexBuffers[0].GetMinBufferSize();
+	uint32 MinIndexBufferSize = IndexBuffers[0].GetMinBufferSize();
+
+	for( int32 BufferIndex = 0; BufferIndex < SlateRHIConstants::NumBuffers; ++BufferIndex )
 	{
 		TotalVertexBufferMemory += VertexBuffers[BufferIndex].GetBufferSize();
 		TotalVertexBufferUsage += VertexBuffers[BufferIndex].GetBufferUsageSize();
@@ -188,16 +207,6 @@ void FSlateRHIRenderingPolicy::EndDrawingWindows()
 		TotalIndexBufferUsage += IndexBuffers[BufferIndex].GetBufferUsageSize();
 	}
 
-	// How much larger the buffers can be than the required size. 
-	// @todo Slate probably could be more intelligent about this
-	const uint32 MaxSizeMultiplier = 2;
-
-	if( TotalVertexBufferMemory > TotalVertexBufferUsage*MaxSizeMultiplier || TotalIndexBufferMemory > TotalIndexBufferUsage*MaxSizeMultiplier )
-	{
-		// The vertex buffer or index is more than twice the size of what is required.  Shrink it
-		bShouldShrinkResources = true;
-	}
-
 
 	SET_MEMORY_STAT( STAT_SlateVertexBufferMemory, TotalVertexBufferMemory );
 	SET_MEMORY_STAT( STAT_SlateIndexBufferMemory, TotalIndexBufferMemory );
@@ -207,69 +216,67 @@ struct FSlateUpdateVertexAndIndexBuffers : public FRHICommand<FSlateUpdateVertex
 {
 	TSlateElementVertexBuffer<FSlateVertex>& VertexBuffer;
 	FSlateElementIndexBuffer& IndexBuffer;
-	const TArray<FSlateVertex>& Vertices;
-	const TArray<SlateIndex>& Indices;	
+	FSlateBatchData& BatchData;
 
-	FSlateUpdateVertexAndIndexBuffers(TSlateElementVertexBuffer<FSlateVertex>& InVertexBuffer, const TArray<FSlateVertex>& InVertices, FSlateElementIndexBuffer& InIndexBuffer, const TArray<SlateIndex>& InIndices)
+	FSlateUpdateVertexAndIndexBuffers(TSlateElementVertexBuffer<FSlateVertex>& InVertexBuffer, FSlateElementIndexBuffer& InIndexBuffer, FSlateBatchData& InBatchData)
 		: VertexBuffer(InVertexBuffer)
 		, IndexBuffer(InIndexBuffer)
-		, Vertices(InVertices)
-		, Indices(InIndices)		
+		, BatchData(InBatchData)
 	{}
 
 	void Execute(FRHICommandListBase& CmdList)
 	{
-		VertexBuffer.FillBuffer_RHIThread(Vertices);
-		IndexBuffer.FillBuffer_RHIThread(Indices);
+		SCOPE_CYCLE_COUNTER( STAT_SlateUpdateBufferRTTime );
+
+		uint8* VertexBufferData = (uint8*)VertexBuffer.LockBuffer_RHIThread(BatchData.GetNumBatchedVertices());
+		uint8* IndexBufferData = (uint8*)IndexBuffer.LockBuffer_RHIThread(BatchData.GetNumBatchedIndices());
+
+		BatchData.FillVertexAndIndexBuffer( VertexBufferData, IndexBufferData );
+	
+		VertexBuffer.UnlockBuffer_RHIThread();
+		IndexBuffer.UnlockBuffer_RHIThread();
 	}
 };
 
-void FSlateRHIRenderingPolicy::UpdateVertexAndIndexBuffers(FRHICommandListImmediate& RHICmdList, const FSlateWindowElementList& WindowElementList)
+void FSlateRHIRenderingPolicy::UpdateVertexAndIndexBuffers(FRHICommandListImmediate& RHICmdList, FSlateBatchData& InBatchData)
 {
 	SCOPE_CYCLE_COUNTER( STAT_SlateUpdateBufferRTTime );
 
 	// Should only be called by the rendering thread
 	check(IsInRenderingThread());
 
-	const TArray<FSlateVertex>& Vertices = WindowElementList.GetBatchedVertices();
-	const TArray<SlateIndex>& Indices = WindowElementList.GetBatchedIndices();
+	const int32 NumVertices = InBatchData.GetNumBatchedVertices();
+	const int32 NumIndices = InBatchData.GetNumBatchedIndices();
 
-	TSlateElementVertexBuffer<FSlateVertex>& VertexBuffer = VertexBuffers[CurrentBufferIndex];
-	FSlateElementIndexBuffer& IndexBuffer = IndexBuffers[CurrentBufferIndex];
-	VertexBuffer.PreFillBuffer(Vertices, bShouldShrinkResources);
-	IndexBuffer.PreFillBuffer(Indices, bShouldShrinkResources);
-
-	if (!GRHIThread || RHICmdList.Bypass())
+	if( InBatchData.GetRenderBatches().Num() > 0  && NumVertices > 0 && NumIndices > 0)
 	{
-		VertexBuffer.FillBuffer_RenderThread(Vertices);
-		IndexBuffer.FillBuffer_RenderThread(Indices);		
-	}
-	else
-	{
-		new (RHICmdList.AllocCommand<FSlateUpdateVertexAndIndexBuffers>()) FSlateUpdateVertexAndIndexBuffers(VertexBuffer, Vertices, IndexBuffer, Indices);
-	}
-
-}
+		TSlateElementVertexBuffer<FSlateVertex>& VertexBuffer = VertexBuffers[CurrentBufferIndex];
+		FSlateElementIndexBuffer& IndexBuffer = IndexBuffers[CurrentBufferIndex];
 
-void FSlateRHIRenderingPolicy::UpdateBuffers( const FSlateWindowElementList& WindowElementList )
-{
-	SCOPE_CYCLE_COUNTER( STAT_SlateUpdateBufferRTTime );
-	// Should only be called by the rendering thread
-	check(IsInRenderingThread());
+		bool bShouldShrinkResources = false;
 
-	{
-		const TArray<FSlateVertex>& Vertices = WindowElementList.GetBatchedVertices();
-		const TArray<SlateIndex>& Indices = WindowElementList.GetBatchedIndices();
+		VertexBuffer.PreFillBuffer(NumVertices, bShouldShrinkResources);
+		IndexBuffer.PreFillBuffer(NumIndices, bShouldShrinkResources);
 
-		TSlateElementVertexBuffer<FSlateVertex>& VertexBuffer = VertexBuffers[CurrentBufferIndex];
-		FSlateElementIndexBuffer& IndexBuffer = IndexBuffers[CurrentBufferIndex];
-		
-		VertexBuffer.PreFillBuffer(Vertices, bShouldShrinkResources);
-		IndexBuffer.PreFillBuffer(Indices, bShouldShrinkResources);
+		if(!GRHIThread || RHICmdList.Bypass())
+		{
+			uint8* VertexBufferData = (uint8*)VertexBuffer.LockBuffer_RenderThread(NumVertices);
+			uint8* IndexBufferData =  (uint8*)IndexBuffer.LockBuffer_RenderThread(NumIndices);
 
-		VertexBuffer.FillBuffer_RenderThread(Vertices);
-		IndexBuffer.FillBuffer_RenderThread(Indices);
+			InBatchData.FillVertexAndIndexBuffer( VertexBufferData, IndexBufferData );
+	
+			VertexBuffer.UnlockBuffer_RenderThread();
+			IndexBuffer.UnlockBuffer_RenderThread();
+		}
+		else
+		{
+			new (RHICmdList.AllocCommand<FSlateUpdateVertexAndIndexBuffers>()) FSlateUpdateVertexAndIndexBuffers(VertexBuffer, IndexBuffer, InBatchData);
+		}
 	}
+
+	SET_DWORD_STAT( STAT_SlateNumLayers, InBatchData.GetNumLayers() );
+	SET_DWORD_STAT( STAT_SlateNumBatches, InBatchData.GetRenderBatches().Num() );
+	SET_DWORD_STAT( STAT_SlateVertexCount, InBatchData.GetNumBatchedVertices() );
 }
 
 static FSceneView& CreateSceneView( FSceneViewFamilyContext& ViewFamilyContext, FSlateBackBuffer& BackBuffer, const FMatrix& ViewProjectionMatrix )
@@ -565,7 +572,7 @@ void FSlateRHIRenderingPolicy::DrawElements(FRHICommandListImmediate& RHICmdList
 
 	}
 
-	CurrentBufferIndex = (CurrentBufferIndex + 1) % 2;
+	CurrentBufferIndex = (CurrentBufferIndex + 1) % SlateRHIConstants::NumBuffers;
 }
 
 
@@ -702,5 +709,3 @@ EPrimitiveType FSlateRHIRenderingPolicy::GetRHIPrimitiveType(ESlateDrawPrimitive
 
 };
 
-
-
diff --git a/Engine/Source/Runtime/SlateRHIRenderer/Private/SlateRHIRenderingPolicy.h b/Engine/Source/Runtime/SlateRHIRenderer/Private/SlateRHIRenderingPolicy.h
--- a/Engine/Source/Runtime/SlateRHIRenderer/Private/SlateRHIRenderingPolicy.h
+++ b/Engine/Source/Runtime/SlateRHIRenderer/Private/SlateRHIRenderingPolicy.h
@@ -7,9 +7,14 @@ class FSlateRHIResourceManager;
 
 #include "RenderingPolicy.h"
 
+namespace SlateRHIConstants
+{
+	// Number of vertex and index buffers we swap between when drawing windows
+	const int32 NumBuffers = 2;
+}
+
 /** 
  * Vertex buffer containing all Slate vertices
- * All vertices are added through RHILockVertexBuffer
  */
 template <typename VertexType>
 class TSlateElementVertexBuffer : public FVertexBuffer
@@ -17,28 +22,35 @@ class TSlateElementVertexBuffer : public FVertexBuffer
 public:
 	TSlateElementVertexBuffer()
 		: BufferSize(0)
-		, BufferUsageSize(0)
-	{}
-
-	TSlateElementVertexBuffer( uint32 InBufferSize )
-		: BufferSize(InBufferSize)
+		, MinBufferSize(0)
 		, BufferUsageSize(0)
 	{}
 
 	~TSlateElementVertexBuffer() {};
 
+	void Init( int32 MinNumVertices )
+	{
+		MinBufferSize = sizeof(VertexType) * FMath::Max( MinNumVertices, 200 );
+
+		BeginInitResource(this);
+	}
+
+	void Destroy()
+	{
+		BeginReleaseResource(this);
+	}
+
 	/** Initializes the vertex buffers RHI resource. */
 	virtual void InitDynamicRHI()
 	{
 		if( !IsValidRef(VertexBufferRHI) )
 		{
-			if( BufferSize == 0 )
-			{
-				BufferSize = 200 * sizeof(VertexType);
-			}
+			check( MinBufferSize > 0 )
+	
+			BufferSize = MinBufferSize;
 
 			FRHIResourceCreateInfo CreateInfo;
-			VertexBufferRHI = RHICreateVertexBuffer( BufferSize, BUF_Dynamic, CreateInfo );
+			VertexBufferRHI = RHICreateVertexBuffer( MinBufferSize, BUF_Dynamic, CreateInfo );
 
 			// Ensure the vertex buffer could be created
 			check(IsValidRef(VertexBufferRHI));
@@ -56,104 +68,95 @@ public:
 	virtual FString GetFriendlyName() const { return TEXT("SlateElementVertices"); }
 
 	/** Returns the size of the buffer in bytes. */
-	uint32 GetBufferSize() const { return BufferSize; }
+	int32 GetBufferSize() const { return BufferSize; }
 
 	/** Returns the used size of this buffer */
-	uint32 GetBufferUsageSize() const { return BufferUsageSize; }
+	int32 GetBufferUsageSize() const { return BufferUsageSize; }
 
 	/** Resets the usage of the buffer */
 	void ResetBufferUsage() { BufferUsageSize = 0; }
 
 	/** Resizes buffer, accumulates states safely on render thread */
-	void PreFillBuffer(const TArray<VertexType>& InVertices, bool bShrinkToFit)
+	void PreFillBuffer(int32 RequiredVertexCount, bool bShrinkToMinSize)
 	{
-		check(IsInRenderingThread());
-		check(IsInRenderingThread());
+		checkSlow(IsInRenderingThread());
 
-		if (InVertices.Num())
+		if (RequiredVertexCount > 0 )
 		{
-			uint32 NumVertices = InVertices.Num();
-
 #if !SLATE_USE_32BIT_INDICES
 			// make sure our index buffer can handle this
-			checkf(NumVertices < 0xFFFF, TEXT("Slate vertex buffer is too large (%d) to work with uint16 indices"), NumVertices);
+			checkf(RequiredVertexCount < 0xFFFF, TEXT("Slate vertex buffer is too large (%d) to work with uint16 indices"), NumVertices);
 #endif
-
-			uint32 RequiredBufferSize = NumVertices*sizeof(VertexType);
+			int32 RequiredBufferSize = RequiredVertexCount*sizeof(VertexType);
 
 			// resize if needed
-			if (RequiredBufferSize > GetBufferSize() || bShrinkToFit)
+			if(RequiredBufferSize > GetBufferSize() || bShrinkToMinSize)
 			{
-				// Use array resize techniques for the vertex buffer
-				ResizeBuffer(InVertices.GetAllocatedSize());
+				ResizeBuffer(RequiredBufferSize);
 			}
 
-			BufferUsageSize += RequiredBufferSize;
+			BufferUsageSize = RequiredBufferSize;
 		}
 
 	}
 
-	/** Fills the buffer with slate vertices on renderthread. */
-	void FillBuffer_RenderThread( const TArray<VertexType>& InVertices)
-	{
-		check( IsInRenderingThread() );
-
-		if( InVertices.Num() )
-		{
-			uint32 NumVertices = InVertices.Num();
-			uint32 RequiredBufferSize = NumVertices*sizeof(VertexType);			
-
-			void* VerticesPtr = RHILockVertexBuffer( VertexBufferRHI, 0, RequiredBufferSize, RLM_WriteOnly );
+	int32 GetMinBufferSize() const { return MinBufferSize; }
 
-			FMemory::Memcpy( VerticesPtr, InVertices.GetData(), RequiredBufferSize );
+	void* LockBuffer_RenderThread(int32 NumVertices)
+	{
+		uint32 RequiredBufferSize = NumVertices*sizeof(VertexType);	
 
-			RHIUnlockVertexBuffer( VertexBufferRHI );
-		}
+		return RHILockVertexBuffer( VertexBufferRHI, 0, RequiredBufferSize, RLM_WriteOnly );
 	}
 
-	
-	void FillBuffer_RHIThread(const TArray<VertexType>& InVertices)
+	void UnlockBuffer_RenderThread()
 	{
-		check(IsInRHIThread() || (!IsInRenderingThread() && !IsInGameThread()));
+		RHIUnlockVertexBuffer( VertexBufferRHI );
+	}
 
-		if (InVertices.Num())
-		{
-			uint32 NumVertices = InVertices.Num();
+	void* LockBuffer_RHIThread(int32 NumVertices)
+	{
+		int32 RequiredBufferSize = NumVertices*sizeof(VertexType);
 
-			uint32 RequiredBufferSize = NumVertices*sizeof(VertexType);
+		//use non-flushing DynamicRHI version when on RHIThread.
+		return GDynamicRHI->RHILockVertexBuffer(VertexBufferRHI, 0, RequiredBufferSize, RLM_WriteOnly);
+	}
 
-			//use non-flushing DynamicRHI version when on RHIThread.
-			void* VerticesPtr = GDynamicRHI->RHILockVertexBuffer(VertexBufferRHI, 0, RequiredBufferSize, RLM_WriteOnly);
-			FMemory::Memcpy(VerticesPtr, InVertices.GetData(), RequiredBufferSize);
-			GDynamicRHI->RHIUnlockVertexBuffer(VertexBufferRHI);			
-		}
+	void UnlockBuffer_RHIThread()
+	{
+		GDynamicRHI->RHIUnlockVertexBuffer(VertexBufferRHI);		
 	}
 
 private:
 	/** Resizes the buffer to the passed in size.  Preserves internal data*/
-	void ResizeBuffer( uint32 NewSizeBytes )
+	void ResizeBuffer( int32 NewSizeBytes )
 	{
-		check( IsInRenderingThread() );
+		checkSlow( IsInRenderingThread() );
+
+		int32 FinalSize = FMath::Max( NewSizeBytes, MinBufferSize );
 
-		if( NewSizeBytes != 0 && NewSizeBytes != BufferSize)
+		if( FinalSize != 0 && FinalSize != BufferSize )
 		{
 			VertexBufferRHI.SafeRelease();
 
 			FRHIResourceCreateInfo CreateInfo;
-			VertexBufferRHI = RHICreateVertexBuffer(NewSizeBytes, BUF_Dynamic, CreateInfo);
+			VertexBufferRHI = RHICreateVertexBuffer(FinalSize, BUF_Dynamic, CreateInfo);
 
 			check(IsValidRef(VertexBufferRHI));
 
-			BufferSize = NewSizeBytes;
+			BufferSize = FinalSize;
 		}
 	}
 
 private:
 	/** The size of the buffer in bytes. */
-	uint32 BufferSize;
+	int32 BufferSize;
 	
+	/** The minimum size the buffer should always be */
+	int32 MinBufferSize;
+
 	/** The size of the used portion of the buffer */
-	uint32 BufferUsageSize;
+	int32 BufferUsageSize;
 
 	/** Hidden copy methods. */
 	TSlateElementVertexBuffer( const TSlateElementVertexBuffer& );
@@ -163,9 +166,17 @@ private:
 class FSlateElementIndexBuffer : public FIndexBuffer
 {
 public:
+	/**
+	 * Constructor
+	 * 
+	 * @param MinNumIndices	The minimum number of indices this buffer should always support
+	 */
 	FSlateElementIndexBuffer();
 	~FSlateElementIndexBuffer();
 
+	void Init( int32 MinNumIndices );
+	void Destroy();
+
 	/** Initializes the index buffers RHI resource. */
 	virtual void InitDynamicRHI() override;
 
@@ -176,32 +187,37 @@ public:
 	virtual FString GetFriendlyName() const override { return TEXT("SlateElementIndices"); }
 
 	/** Returns the size of this buffer */
-	uint32 GetBufferSize() const { return BufferSize; }
+	int32 GetBufferSize() const { return BufferSize; }
 
 	/** Returns the used size of this buffer */
-	uint32 GetBufferUsageSize() const { return BufferUsageSize; }
+	int32 GetBufferUsageSize() const { return BufferUsageSize; }
 
 	/** Resets the usage of the buffer */
 	void ResetBufferUsage() { BufferUsageSize = 0; }
 
 	/** Resizes buffer, accumulates states safely on render thread */
-	void PreFillBuffer(const TArray<SlateIndex>& InIndices, bool bShrinkToFit);
+	void PreFillBuffer(int32 RequiredIndexCount, bool bShrinkToMinSize);
 
-	/** Fills the buffer with slate vertices */
-	void FillBuffer_RenderThread( const TArray<SlateIndex>& InIndices );
+	void* LockBuffer_RenderThread(int32 NumIndices);
+	void UnlockBuffer_RenderThread();
 
-	/** Fills the buffer with slate vertices */
-	void FillBuffer_RHIThread(const TArray<SlateIndex>& InIndices);
+	void* LockBuffer_RHIThread(int32 NumIndices);
+	void UnlockBuffer_RHIThread();
+
+	int32 GetMinBufferSize() const { return MinBufferSize; }
 
 private:
 	/** Resizes the buffer to the passed in size.  Preserves internal data */
-	void ResizeBuffer( uint32 NewSizeBytes );
+	void ResizeBuffer( int32 NewSizeBytes );
 private:
-	/** The size of the buffer in bytes. */
-	uint32 BufferSize;
+	/** The current size of the buffer in bytes. */
+	int32 BufferSize;
 	
-	/** The size of the used portion of the buffer */
-	uint32 BufferUsageSize;
+	/** The minimum size the buffer should always be */
+	int32 MinBufferSize;
+
+	/** The size of the currently used portion of the buffer */
+	int32 BufferUsageSize;
 
 	/** Hidden copy methods. */
 	FSlateElementIndexBuffer( const FSlateElementIndexBuffer& );
@@ -215,8 +231,8 @@ public:
 	FSlateRHIRenderingPolicy( TSharedPtr<FSlateFontCache> InFontCache, TSharedRef<FSlateRHIResourceManager> InResourceManager );
 	~FSlateRHIRenderingPolicy();
 
-	virtual void UpdateBuffers( const FSlateWindowElementList& WindowElementList ) override;
-	void UpdateVertexAndIndexBuffers(FRHICommandListImmediate& RHICmdList, const FSlateWindowElementList& WindowElementList);
+	void UpdateVertexAndIndexBuffers(FRHICommandListImmediate& RHICmdList, FSlateBatchData& BatchData);
+
 	virtual void DrawElements(FRHICommandListImmediate& RHICmdList, class FSlateBackBuffer& BackBuffer, const FMatrix& ViewProjectionMatrix, const TArray<FSlateRenderBatch>& RenderBatches, bool bAllowSwtichVerticalAxis=true);
 
 	virtual TSharedRef<FSlateFontCache> GetFontCache() override { return FontCache.ToSharedRef(); }
@@ -244,8 +260,8 @@ private:
 	EPrimitiveType GetRHIPrimitiveType(ESlateDrawPrimitive::Type SlateType);
 private:
 	/** Buffers used for rendering */
-	TSlateElementVertexBuffer<FSlateVertex> VertexBuffers[2];
-	FSlateElementIndexBuffer IndexBuffers[2];
+	TSlateElementVertexBuffer<FSlateVertex> VertexBuffers[SlateRHIConstants::NumBuffers];
+	FSlateElementIndexBuffer IndexBuffers[SlateRHIConstants::NumBuffers];
 
 	TSharedRef<FSlateRHIResourceManager> ResourceManager;
 	TSharedPtr<FSlateFontCache> FontCache;
