Fix initalization order

diff --git a/Engine/Source/Runtime/SlateCore/Private/Rendering/ElementBatcher.cpp b/Engine/Source/Runtime/SlateCore/Private/Rendering/ElementBatcher.cpp
--- a/Engine/Source/Runtime/SlateCore/Private/Rendering/ElementBatcher.cpp
+++ b/Engine/Source/Runtime/SlateCore/Private/Rendering/ElementBatcher.cpp
@@ -47,9 +47,9 @@ static FColor GetElementColor( const FLinearColor& InColor, const FSlateBrush* I
 }
 
 FSlateElementBatcher::FSlateElementBatcher( TSharedRef<FSlateRenderingPolicy> InRenderingPolicy )
-	: ResourceManager( *InRenderingPolicy->GetResourceManager() )
+	: BatchData( nullptr )
+	, ResourceManager( *InRenderingPolicy->GetResourceManager() )
 	, FontCache( *InRenderingPolicy->GetFontCache() )
-	, BatchData( nullptr )
 	, PixelCenterOffset( InRenderingPolicy->GetPixelCenterOffset() )
 {
 
