Fixed font measuring crashing when attempting to measure an empty string

diff --git a/Engine/Source/Runtime/SlateCore/Private/Fonts/FontMeasure.cpp b/Engine/Source/Runtime/SlateCore/Private/Fonts/FontMeasure.cpp
--- a/Engine/Source/Runtime/SlateCore/Private/Fonts/FontMeasure.cpp
+++ b/Engine/Source/Runtime/SlateCore/Private/Fonts/FontMeasure.cpp
@@ -139,7 +139,7 @@ FVector2D FSlateFontMeasure::MeasureStringInternal( const FString& Text, int32 S
 	const bool DoesStartAtBeginning = StartIndex == 0;
 	const bool DoesFinishAtEnd = EndIndex == Text.Len();
 	const int32 TextRangeLength = EndIndex - StartIndex;
-	if ( EndIndex - StartIndex <= 0 || EndIndex <= 0 || StartIndex < 0 || EndIndex <= StartIndex )
+	if ( Text.IsEmpty() || EndIndex - StartIndex <= 0 || EndIndex <= 0 || StartIndex < 0 || EndIndex <= StartIndex )
 	{
 		return FVector2D( 0, MaxHeight );
 	}
