[INTEGRATE] merge ue4<-ue4-launcherDev cl 2581195

diff --git a/Engine/Source/Runtime/SlateCore/Public/Layout/SlateRect.h b/Engine/Source/Runtime/SlateCore/Public/Layout/SlateRect.h
--- a/Engine/Source/Runtime/SlateCore/Public/Layout/SlateRect.h
+++ b/Engine/Source/Runtime/SlateCore/Public/Layout/SlateRect.h
@@ -162,14 +162,13 @@ public:
 
 	/**
 	 * Returns whether or not a point is inside the rectangle
-	 * Note: The lower right and bottom points of the rect are not inside the rectangle due to rendering API clipping rules.
 	 * 
 	 * @param Point	The point to check
 	 * @return True if the point is inside the rectangle
 	 */
 	FORCEINLINE bool ContainsPoint( const FVector2D& Point ) const
 	{
-		return Point.X >= Left && Point.X < Right && Point.Y >= Top && Point.Y < Bottom;
+		return Point.X >= Left && Point.X <= Right && Point.Y >= Top && Point.Y <= Bottom;
 	}
 
 	bool operator==( const FSlateRect& Other ) const
