Slate - Continuing to refine the fixes to clipping in slate, no longer pixel snapping twice.  Now snapping the points of the clip rect after transforming them.  Then calculating the extents, but no longer snapping them as that would artifically inflate or even potentially shrink the bounding box.

diff --git a/Engine/Source/Runtime/SlateCore/Private/Rendering/ElementBatcher.cpp b/Engine/Source/Runtime/SlateCore/Private/Rendering/ElementBatcher.cpp
--- a/Engine/Source/Runtime/SlateCore/Private/Rendering/ElementBatcher.cpp
+++ b/Engine/Source/Runtime/SlateCore/Private/Rendering/ElementBatcher.cpp
@@ -27,12 +27,23 @@ FVector2D RoundToInt(const FVector2D& Vec)
  */
 FSlateRotatedClipRectType ToSnappedRotatedRect(const FSlateRect& ClipRectInLayoutWindowSpace, const FSlateLayoutTransform& InverseLayoutTransform, const FSlateRenderTransform& RenderTransform)
 {
-	FSlateRotatedRect RotatedRect = TransformRect(Concatenate(InverseLayoutTransform, RenderTransform), FSlateRotatedRect(ClipRectInLayoutWindowSpace));
-	// Pixel snapping is done here by rounding the resulting floats to ints.
-	return FSlateRotatedClipRectType(
-		RoundToInt(RotatedRect.TopLeft), 
-		RoundToInt(RotatedRect.ExtentX), 
-		RoundToInt(RotatedRect.ExtentY));
+	auto Transform = Concatenate(InverseLayoutTransform, RenderTransform);
+
+	// Pixel snapping is done here by rounding the resulting floats to ints, we do this before
+	// calculating the extent of the clip box otherwise we'll get a smaller clip rect than a visual
+	// rect where each point is individually snapped.
+	FVector2D TopLeft = RoundToInt(TransformPoint(Transform, ClipRectInLayoutWindowSpace.GetTopLeft()));
+	FVector2D TopRight = RoundToInt(TransformPoint(Transform, ClipRectInLayoutWindowSpace.GetTopRight()));
+	FVector2D BottomLeft = RoundToInt(TransformPoint(Transform, ClipRectInLayoutWindowSpace.GetBottomLeft()));
+
+	// Calculate the distance between the vectors coming from the clip origin.
+	FVector2D ExtentX = TopRight - TopLeft;
+	FVector2D ExtentY = BottomLeft - TopLeft;
+
+	//NOTE: We explicitly do not re-snap the extent x/y, it wouldn't be correct to snap again in distance space
+	// even if two points are snapped, their distance wont necessarily be a whole number if those points are not
+	// axis aligned.
+	return FSlateRotatedClipRectType(TopLeft, ExtentX, ExtentY);
 }
 
 /**
diff --git a/Engine/Source/Runtime/SlateCore/Private/Rendering/RenderingCommon.cpp b/Engine/Source/Runtime/SlateCore/Private/Rendering/RenderingCommon.cpp
--- a/Engine/Source/Runtime/SlateCore/Private/Rendering/RenderingCommon.cpp
+++ b/Engine/Source/Runtime/SlateCore/Private/Rendering/RenderingCommon.cpp
@@ -6,9 +6,9 @@
 FSlateRotatedRect::FSlateRotatedRect() {}
 
 FSlateRotatedRect::FSlateRotatedRect(const FSlateRect& AlignedRect)
-	: TopLeft(FMath::RoundToInt(AlignedRect.Left), FMath::RoundToInt(AlignedRect.Top))
-	, ExtentX(FMath::RoundToInt(AlignedRect.Right) - FMath::RoundToInt(AlignedRect.Left), 0.0f)
-	, ExtentY(0.0f, FMath::RoundToInt(AlignedRect.Bottom) - FMath::RoundToInt(AlignedRect.Top))
+	: TopLeft(AlignedRect.Left, AlignedRect.Top)
+	, ExtentX(AlignedRect.Right - AlignedRect.Left, 0.0f)
+	, ExtentY(0.0f, AlignedRect.Bottom - AlignedRect.Top)
 {
 }
 
diff --git a/Engine/Source/Runtime/SlateCore/Public/Layout/SlateRect.h b/Engine/Source/Runtime/SlateCore/Public/Layout/SlateRect.h
--- a/Engine/Source/Runtime/SlateCore/Public/Layout/SlateRect.h
+++ b/Engine/Source/Runtime/SlateCore/Public/Layout/SlateRect.h
@@ -97,6 +97,16 @@ public:
 		return FVector2D( Left, Top );
 	}
 
+	/**
+	 * Returns the top-right position of the rectangle
+	 *
+	 * @return The top-right position.
+	 */
+	FVector2D GetTopRight() const
+	{
+		return FVector2D(Right, Top);
+	}
+
 	/**
 	 * Returns the bottom-right position of the rectangle
 	 * 
@@ -107,6 +117,16 @@ public:
 		return FVector2D( Right, Bottom );
 	}
 
+	/**
+	 * Returns the bottom-left position of the rectangle
+	 * 
+	 * @return The bottom-left position.
+	 */
+	FVector2D GetBottomLeft() const
+	{
+		return FVector2D( Left, Bottom );
+	}
+
 	/**
 	 * Return a rectangle that is contracted on each side by the amount specified in each margin.
 	 *
