Slate - Small perf boost to clip rect transform code so that we're doing as few transforms as possible.

diff --git a/Engine/Source/Runtime/SlateCore/Private/Rendering/ElementBatcher.cpp b/Engine/Source/Runtime/SlateCore/Private/Rendering/ElementBatcher.cpp
--- a/Engine/Source/Runtime/SlateCore/Private/Rendering/ElementBatcher.cpp
+++ b/Engine/Source/Runtime/SlateCore/Private/Rendering/ElementBatcher.cpp
@@ -27,23 +27,24 @@ FVector2D RoundToInt(const FVector2D& Vec)
  */
 FSlateRotatedClipRectType ToSnappedRotatedRect(const FSlateRect& ClipRectInLayoutWindowSpace, const FSlateLayoutTransform& InverseLayoutTransform, const FSlateRenderTransform& RenderTransform)
 {
-	auto Transform = Concatenate(InverseLayoutTransform, RenderTransform);
+	FSlateRotatedRect RotatedRect = TransformRect(
+		Concatenate(InverseLayoutTransform, RenderTransform),
+		FSlateRotatedRect(ClipRectInLayoutWindowSpace));
 
 	// Pixel snapping is done here by rounding the resulting floats to ints, we do this before
-	// calculating the extent of the clip box otherwise we'll get a smaller clip rect than a visual
+	// calculating the final extents of the clip box otherwise we'll get a smaller clip rect than a visual
 	// rect where each point is individually snapped.
-	FVector2D TopLeft = RoundToInt(TransformPoint(Transform, ClipRectInLayoutWindowSpace.GetTopLeft()));
-	FVector2D TopRight = RoundToInt(TransformPoint(Transform, ClipRectInLayoutWindowSpace.GetTopRight()));
-	FVector2D BottomLeft = RoundToInt(TransformPoint(Transform, ClipRectInLayoutWindowSpace.GetBottomLeft()));
-
-	// Calculate the distance between the vectors coming from the clip origin.
-	FVector2D ExtentX = TopRight - TopLeft;
-	FVector2D ExtentY = BottomLeft - TopLeft;
+	FVector2D SnappedTopLeft = RoundToInt(RotatedRect.TopLeft);
+	FVector2D SnappedTopRight = RoundToInt(RotatedRect.TopLeft + RotatedRect.ExtentX);
+	FVector2D SnappedBottomLeft = RoundToInt(RotatedRect.TopLeft + RotatedRect.ExtentY);
 
 	//NOTE: We explicitly do not re-snap the extent x/y, it wouldn't be correct to snap again in distance space
 	// even if two points are snapped, their distance wont necessarily be a whole number if those points are not
 	// axis aligned.
-	return FSlateRotatedClipRectType(TopLeft, ExtentX, ExtentY);
+	return FSlateRotatedClipRectType(
+		SnappedTopLeft,
+		SnappedTopRight - SnappedTopLeft,
+		SnappedBottomLeft - SnappedTopLeft);
 }
 
 /**
