Added support for materials to be used on Slate based fonts.

diff --git a/Engine/Shaders/SlateElementPixelShader.usf b/Engine/Shaders/SlateElementPixelShader.usf
--- a/Engine/Shaders/SlateElementPixelShader.usf
+++ b/Engine/Shaders/SlateElementPixelShader.usf
@@ -33,12 +33,12 @@ float2 GammaValues;
 /** Per element params */
 float4 ShaderParams;
 
-#if !USE_MATERIALS
+//#if !USE_MATERIALS
 /** Texture for the element */
 Texture2D ElementTexture;
 SamplerState ElementTextureSampler;
 
-#endif
+//#endif
 
 
 float3 GammaCorrect(float3 InColor)
@@ -63,23 +63,33 @@ struct VertexOut
 	float4 TextureCoordinates : TEXCOORD0;
 	float4 ClipOriginAndPos : TEXCOORD1;
 	float4 ClipExtents : TEXCOORD2;
+	float2 MaterialTexCoords : TEXCOORD3;
 };
 	
 #if USE_MATERIALS
-float4 GetMaterialColor( VertexOut InVertex, float2 UV )
+float4 GetMaterialColor( VertexOut InVertex, float2 CustomUV )
 {
 	FMaterialPixelParameters Parameters = MakeInitializedMaterialPixelParameters();
 	CalcMaterialParameters(Parameters, InVertex.Position, true, InVertex.Position);
 
-#if NUM_MATERIAL_TEXCOORDS == 1
-	Parameters.TexCoords[0] = UV;
-#elif NUM_MATERIAL_TEXCOORDS == 2
-	Parameters.TexCoords[0] = UV;
-	Parameters.TexCoords[1] = InVertex.TextureCoordinates.xy;
-#elif NUM_MATERIAL_TEXCOORDS == 3
-	Parameters.TexCoords[0] = UV;
-	Parameters.TexCoords[1] = InVertex.TextureCoordinates.xy;
-	Parameters.TexCoords[2] = InVertex.TextureCoordinates.zw;
+#if SHADER_TYPE == ST_Font
+	#if NUM_MATERIAL_TEXCOORDS == 1
+		Parameters.TexCoords[0] = InVertex.TextureCoordinates.zw;
+	#elif NUM_MATERIAL_TEXCOORDS == 2
+		Parameters.TexCoords[0] = InVertex.TextureCoordinates.zw;
+		Parameters.TexCoords[1] = InVertex.MaterialTexCoords;
+	#endif
+#else
+	#if NUM_MATERIAL_TEXCOORDS == 1
+		Parameters.TexCoords[0] = CustomUV;
+	#elif NUM_MATERIAL_TEXCOORDS == 2
+		Parameters.TexCoords[0] = CustomUV;
+		Parameters.TexCoords[1] = InVertex.TextureCoordinates.xy;
+	#elif NUM_MATERIAL_TEXCOORDS == 3
+		Parameters.TexCoords[0] = CustomUV;
+		Parameters.TexCoords[1] = InVertex.TextureCoordinates.xy;
+		Parameters.TexCoords[2] = InVertex.TextureCoordinates.zw;
+	#endif
 #endif
 
 	Parameters.VertexColor = InVertex.Color;
@@ -133,8 +143,15 @@ float4 GetColor( VertexOut InVertex, float2 UV )
 #if SHADER_TYPE == ST_Font
 float4 GetFontElementColor( VertexOut InVertex )
 {
-	float4 OutColor = InVertex.Color;
+#if USE_MATERIALS
+	float4 BaseColor = GetMaterialColor( InVertex, float2(0,0) )*InVertex.Color;
+#else
+	float4 BaseColor = InVertex.Color;
+#endif
+	float4 OutColor = BaseColor;
+
 	OutColor.a *= Texture2DSample_A8(ElementTexture, ElementTextureSampler, InVertex.TextureCoordinates.xy);
+
 	return OutColor;
 }
 #endif
diff --git a/Engine/Shaders/SlateVertexShader.usf b/Engine/Shaders/SlateVertexShader.usf
--- a/Engine/Shaders/SlateVertexShader.usf
+++ b/Engine/Shaders/SlateVertexShader.usf
@@ -14,18 +14,17 @@ struct VertexOut
 	float4 TextureCoordinates : TEXCOORD0;
 	float4 ClipOriginAndPos : TEXCOORD1;
 	float4 ClipExtents : TEXCOORD2;
+	float2 MaterialTexCoords : TEXCOORD3;
 };
 
 
 VertexOut Main(
 	in float4 InTextureCoordinates : ATTRIBUTE0,
-/* Don't change these int streams to float, otherwise it WILL fail since the vertex declaration is defined as int on D3D11:
-		D3D11: WARNING: ID3D11Device::CreateInputLayout: The provided input signature expects to read an element with SemanticName/Index: 'ATTRIBUTE'/1 and component(s) of the type 'float32'.  However, the matching entry in the Input Layout declaration, element[1], specifies mismatched format: 'R16G16B16A16_SINT'.  This is not an error, since behavior is well defined: The element format determines what data conversion algorithm gets applied before it shows up in a shader register. Independently, the shader input signature defines how the shader will interpret the data that has been placed in its input registers, with no change in the bits stored.  It is valid for the application to reinterpret data as a different type once it is in the vertex shader, so this warning is issued just in case reinterpretation was not intended by the author. [ STATE_CREATION WARNING #391: CREATEINPUTLAYOUT_TYPE_MISMATCH ]
-*/
-	in float2 InPosition : ATTRIBUTE1,
-	in float2 InClipOrigin : ATTRIBUTE2,
-	in float4 InClipExtents : ATTRIBUTE3,
-	in float4 InColor : ATTRIBUTE4
+	in float2 InMaterialTextureCoordinates : ATTRIBUTE1,
+	in float2 InPosition : ATTRIBUTE2,
+	in float2 InClipOrigin : ATTRIBUTE3,
+	in float4 InClipExtents : ATTRIBUTE4,
+	in float4 InColor : ATTRIBUTE5
 	)
 {
 	VertexOut VOut;
@@ -38,6 +37,7 @@ VertexOut Main(
 
 	// TextureCoordinates contains both the first and second texture coordinates in xy and zw respectively.
 	VOut.TextureCoordinates = InTextureCoordinates;
+	VOut.MaterialTexCoords = InMaterialTextureCoordinates;
 	VOut.ClipOriginAndPos = float4(InClipOrigin, InPosition.xy);
 	VOut.ClipExtents = InClipExtents;
 
diff --git a/Engine/Shaders/StandaloneRenderer/D3D/SlateDefaultVertexShader.hlsl b/Engine/Shaders/StandaloneRenderer/D3D/SlateDefaultVertexShader.hlsl
--- a/Engine/Shaders/StandaloneRenderer/D3D/SlateDefaultVertexShader.hlsl
+++ b/Engine/Shaders/StandaloneRenderer/D3D/SlateDefaultVertexShader.hlsl
@@ -19,6 +19,7 @@ struct VertexOut
 VertexOut Main(
 	in float2 InPosition : POSITION,
 	in float4 InTextureCoordinates : TEXCOORD0,
+	in float2 MaterialTexCoords : TEXCOORD3,
 	in float2 InClipOrigin : TEXCOORD1,
 	in float4 InClipExtents : TEXCOORD2,
 	in float4 InColor : COLOR0
diff --git a/Engine/Source/Developer/StandaloneRenderer/Private/OpenGL/SlateOpenGLShaders.cpp b/Engine/Source/Developer/StandaloneRenderer/Private/OpenGL/SlateOpenGLShaders.cpp
--- a/Engine/Source/Developer/StandaloneRenderer/Private/OpenGL/SlateOpenGLShaders.cpp
+++ b/Engine/Source/Developer/StandaloneRenderer/Private/OpenGL/SlateOpenGLShaders.cpp
@@ -218,6 +218,7 @@ void FSlateOpenGLShaderProgram::LinkShaders( const FSlateOpenGLVS& VertexShader,
 	glBindAttribLocation(ProgramID, 2, "InClipOrigin");
 	glBindAttribLocation(ProgramID, 3, "InClipExtents");
 	glBindAttribLocation(ProgramID, 4, "InColor");
+	glBindAttribLocation(ProgramID, 5, "InMaterialTexCoords");
 
 	// Link the shaders
 	glLinkProgram( ProgramID );
diff --git a/Engine/Source/Developer/StandaloneRenderer/Private/Windows/D3D/SlateD3DShaders.cpp b/Engine/Source/Developer/StandaloneRenderer/Private/Windows/D3D/SlateD3DShaders.cpp
--- a/Engine/Source/Developer/StandaloneRenderer/Private/Windows/D3D/SlateD3DShaders.cpp
+++ b/Engine/Source/Developer/StandaloneRenderer/Private/Windows/D3D/SlateD3DShaders.cpp
@@ -241,6 +241,7 @@ FSlateDefaultVS::FSlateDefaultVS()
 	D3D11_INPUT_ELEMENT_DESC Layout[] = 
 	{
 		{ "TEXCOORD",	0, DXGI_FORMAT_R32G32B32A32_FLOAT,	0, 0,								D3D11_INPUT_PER_VERTEX_DATA,	0 },
+		{ "TEXCOORD",	3, DXGI_FORMAT_R32G32_FLOAT,		0, D3D11_APPEND_ALIGNED_ELEMENT,	D3D11_INPUT_PER_VERTEX_DATA,	0 },
 		{ "POSITION",	0, DXGI_FORMAT_R32G32_FLOAT,		0, D3D11_APPEND_ALIGNED_ELEMENT,	D3D11_INPUT_PER_VERTEX_DATA,	0 },
 		{ "TEXCOORD",	1, DXGI_FORMAT_R32G32_FLOAT,		0, D3D11_APPEND_ALIGNED_ELEMENT,	D3D11_INPUT_PER_VERTEX_DATA,	0 },
 		{ "TEXCOORD",	2, DXGI_FORMAT_R32G32B32A32_FLOAT,	0, D3D11_APPEND_ALIGNED_ELEMENT,	D3D11_INPUT_PER_VERTEX_DATA,	0 },
diff --git a/Engine/Source/Editor/DetailCustomizations/Private/SlateFontInfoCustomization.cpp b/Engine/Source/Editor/DetailCustomizations/Private/SlateFontInfoCustomization.cpp
--- a/Engine/Source/Editor/DetailCustomizations/Private/SlateFontInfoCustomization.cpp
+++ b/Engine/Source/Editor/DetailCustomizations/Private/SlateFontInfoCustomization.cpp
@@ -84,6 +84,7 @@ void FSlateFontInfoStructCustomization::CustomizeHeader(TSharedRef<IPropertyHand
 
 void FSlateFontInfoStructCustomization::CustomizeChildren(TSharedRef<IPropertyHandle> InStructPropertyHandle, IDetailChildrenBuilder& InStructBuilder, IPropertyTypeCustomizationUtils& InStructCustomizationUtils)
 {
+	InStructBuilder.AddChildProperty( InStructPropertyHandle->GetChildHandle( TEXT("FontMaterial") ).ToSharedRef() );
 }
 
 bool FSlateFontInfoStructCustomization::OnFilterFontAsset(const FAssetData& InAssetData)
diff --git a/Engine/Source/Runtime/SlateCore/Private/Fonts/FontCache.cpp b/Engine/Source/Runtime/SlateCore/Private/Fonts/FontCache.cpp
--- a/Engine/Source/Runtime/SlateCore/Private/Fonts/FontCache.cpp
+++ b/Engine/Source/Runtime/SlateCore/Private/Fonts/FontCache.cpp
@@ -1243,13 +1243,17 @@ void FSlateFontCache::FlushObject( const UObject* const InObject )
 	}
 }
 
-void FSlateFontCache::ConditionalFlushCache()
+bool FSlateFontCache::ConditionalFlushCache()
 {
+	bool bFlushed = false;
 	if( bFlushRequested )
 	{
 		FlushCache();
+		bFlushed = true;
 		bFlushRequested = false;
 	}
+
+	return bFlushed;
 }
 
 void FSlateFontCache::UpdateCache()
diff --git a/Engine/Source/Runtime/SlateCore/Private/Fonts/SlateFontInfo.cpp b/Engine/Source/Runtime/SlateCore/Private/Fonts/SlateFontInfo.cpp
--- a/Engine/Source/Runtime/SlateCore/Private/Fonts/SlateFontInfo.cpp
+++ b/Engine/Source/Runtime/SlateCore/Private/Fonts/SlateFontInfo.cpp
@@ -10,6 +10,7 @@
 
 FSlateFontInfo::FSlateFontInfo( )
 	: FontObject(nullptr)
+	, FontMaterial(nullptr)
 	, CompositeFont()
 	, TypefaceFontName()
 	, Size(0)
@@ -21,6 +22,7 @@ FSlateFontInfo::FSlateFontInfo( )
 
 FSlateFontInfo::FSlateFontInfo( TSharedPtr<const FCompositeFont> InCompositeFont, const int32 InSize, const FName& InTypefaceFontName )
 	: FontObject(nullptr)
+	, FontMaterial(nullptr)
 	, CompositeFont(InCompositeFont)
 	, TypefaceFontName(InTypefaceFontName)
 	, Size(InSize)
@@ -36,6 +38,7 @@ FSlateFontInfo::FSlateFontInfo( TSharedPtr<const FCompositeFont> InCompositeFont
 
 FSlateFontInfo::FSlateFontInfo( const UObject* InFontObject, const int32 InSize, const FName& InTypefaceFontName )
 	: FontObject(InFontObject)
+	, FontMaterial(nullptr)
 	, CompositeFont()
 	, TypefaceFontName(InTypefaceFontName)
 	, Size(InSize)
@@ -59,6 +62,7 @@ FSlateFontInfo::FSlateFontInfo( const UObject* InFontObject, const int32 InSize,
 
 FSlateFontInfo::FSlateFontInfo( const FString& InFontName, uint16 InSize, EFontHinting InHinting )
 	: FontObject(nullptr)
+	, FontMaterial(nullptr)
 	, CompositeFont()
 	, TypefaceFontName()
 	, Size(InSize)
@@ -74,6 +78,7 @@ FSlateFontInfo::FSlateFontInfo( const FString& InFontName, uint16 InSize, EFontH
 
 FSlateFontInfo::FSlateFontInfo( const FName& InFontName, uint16 InSize, EFontHinting InHinting )
 	: FontObject(nullptr)
+	, FontMaterial(nullptr)
 	, CompositeFont()
 	, TypefaceFontName()
 	, Size(InSize)
@@ -89,6 +94,7 @@ FSlateFontInfo::FSlateFontInfo( const FName& InFontName, uint16 InSize, EFontHin
 
 FSlateFontInfo::FSlateFontInfo( const ANSICHAR* InFontName, uint16 InSize, EFontHinting InHinting )
 	: FontObject(nullptr)
+	, FontMaterial(nullptr)
 	, CompositeFont()
 	, TypefaceFontName()
 	, Size(InSize)
@@ -104,6 +110,7 @@ FSlateFontInfo::FSlateFontInfo( const ANSICHAR* InFontName, uint16 InSize, EFont
 
 FSlateFontInfo::FSlateFontInfo( const WIDECHAR* InFontName, uint16 InSize, EFontHinting InHinting )
 	: FontObject(nullptr)
+	, FontMaterial(nullptr)
 	, CompositeFont()
 	, TypefaceFontName()
 	, Size(InSize)
diff --git a/Engine/Source/Runtime/SlateCore/Private/Rendering/ElementBatcher.cpp b/Engine/Source/Runtime/SlateCore/Private/Rendering/ElementBatcher.cpp
--- a/Engine/Source/Runtime/SlateCore/Private/Rendering/ElementBatcher.cpp
+++ b/Engine/Source/Runtime/SlateCore/Private/Rendering/ElementBatcher.cpp
@@ -564,7 +564,8 @@ void FSlateElementBatcher::AddTextElement(const FSlateDrawElement& DrawElement)
 	float MaxHeight = CharacterList.GetMaxHeight();
 
 	uint32 FontTextureIndex = 0;
-	FSlateShaderResource* FontTexture = nullptr;
+	FSlateShaderResource* FontAtlasTexture = nullptr;
+	FSlateShaderResource* FontShaderResource = nullptr;
 
 	FSlateElementBatch* ElementBatch = nullptr;
 	TArray<FSlateVertex>* BatchVertices = nullptr;
@@ -589,9 +590,14 @@ void FSlateElementBatcher::AddTextElement(const FSlateDrawElement& DrawElement)
 
 	LineX = PosX;
 	
+	const bool bIsFontMaterial = InPayload.FontInfo.FontMaterial != nullptr;
+
 	FColor FinalColor = GetElementColor( InPayload.Tint, nullptr );
 
-	for( int32 CharIndex = 0; CharIndex < Text.Len(); ++CharIndex )
+	uint32 NumChars = Text.Len();
+
+	uint32 NumLines = 1;
+	for( uint32 CharIndex = 0; CharIndex < NumChars; ++CharIndex )
 	{
 		const TCHAR CurrentChar = Text[ CharIndex ];
 
@@ -603,18 +609,23 @@ void FSlateElementBatcher::AddTextElement(const FSlateDrawElement& DrawElement)
 			PosY += MaxHeight;
 			// Carriage return 
 			LineX = PosX;
+
+			++NumLines;
+
 		}
 		else
 		{
 			const FCharacterEntry& Entry = CharacterList[ CurrentChar ];
 
-			if( FontTexture == nullptr || Entry.TextureIndex != FontTextureIndex )
+			if( FontAtlasTexture == nullptr || Entry.TextureIndex != FontTextureIndex )
 			{
 				// Font has a new texture for this glyph. Refresh the batch we use and the index we are currently using
 				FontTextureIndex = Entry.TextureIndex;
 
-				FontTexture = FontCache.GetSlateTextureResource( FontTextureIndex );
-				ElementBatch = &FindBatchForElement( Layer, FShaderParams(), FontTexture, ESlateDrawPrimitive::TriangleList, ESlateShader::Font, InDrawEffects, ESlateBatchDrawFlag::None, DrawElement.GetScissorRect() );
+				FontAtlasTexture = FontCache.GetSlateTextureResource( FontTextureIndex );
+				FontShaderResource = ResourceManager.GetFontShaderResource( FontTextureIndex, FontAtlasTexture, InPayload.FontInfo.FontMaterial );
+
+				ElementBatch = &FindBatchForElement( Layer, FShaderParams(), FontShaderResource, ESlateDrawPrimitive::TriangleList, ESlateShader::Font, InDrawEffects, ESlateBatchDrawFlag::None, DrawElement.GetScissorRect() );
 
 				BatchVertices = &BatchData->GetBatchVertexList(*ElementBatch);
 				BatchIndices = &BatchData->GetBatchIndexList(*ElementBatch);
@@ -622,8 +633,8 @@ void FSlateElementBatcher::AddTextElement(const FSlateDrawElement& DrawElement)
 				VertexOffset = BatchVertices->Num();
 				IndexOffset = BatchIndices->Num();
 				
-				InvTextureSizeX = 1.0f/FontTexture->GetWidth();
-				InvTextureSizeY = 1.0f/FontTexture->GetHeight();
+				InvTextureSizeX = 1.0f/FontAtlasTexture->GetWidth();
+				InvTextureSizeY = 1.0f/FontAtlasTexture->GetHeight();
 			}
 
 			const bool bIsWhitespace = FText::IsWhitespace(CurrentChar);
@@ -673,11 +684,25 @@ void FSlateElementBatcher::AddTextElement(const FSlateDrawElement& DrawElement)
 					// The start index of these vertices in the index buffer
 					uint32 IndexStart = VertexOffset;
 
+					float Ut = 0.0f, Vt = 0.0f, UtMax = 0.0f, VtMax = 0.0f;
+					if( bIsFontMaterial )
+					{
+						float DistAlpha = (float)CharIndex/NumChars;
+						float DistAlphaNext = (float)(CharIndex+1)/NumChars;
+
+						// This creates a set of UV's that goes from 0-1, from left to right of the string in U and 0-1 baseline to baseline top to bottom in V
+						Ut = FMath::Lerp(0.0f, 1.0f, DistAlpha);
+						Vt = FMath::Lerp(0.0f, 1.0f, UpperLeft.Y/(MaxHeight*NumLines));
+
+						UtMax = FMath::Lerp(0.0f, 1.0f, DistAlphaNext);
+						VtMax = FMath::Lerp(0.0f, 1.0f, LowerLeft.Y/(MaxHeight*NumLines));
+					}
+
 					// Add four vertices to the list of verts to be added to the vertex buffer
-					BatchVerticesRef[ VertexOffset++ ] = FSlateVertex( RenderTransform, UpperLeft,								FVector2D(U,V),				FinalColor, RenderClipRect );
-					BatchVerticesRef[ VertexOffset++ ] = FSlateVertex( RenderTransform, FVector2D(LowerRight.X,UpperLeft.Y),	FVector2D(U+SizeU, V),		FinalColor, RenderClipRect );
-					BatchVerticesRef[ VertexOffset++ ] = FSlateVertex( RenderTransform, FVector2D(UpperLeft.X,LowerRight.Y),	FVector2D(U, V+SizeV),		FinalColor, RenderClipRect );
-					BatchVerticesRef[ VertexOffset++ ] = FSlateVertex( RenderTransform, LowerRight,								FVector2D(U+SizeU, V+SizeV),FinalColor, RenderClipRect );
+					BatchVerticesRef[ VertexOffset++ ] = FSlateVertex( RenderTransform, UpperLeft,								FVector4(U,V,				Ut,Vt),			FVector2D(0.0f,0.0f), FinalColor, RenderClipRect );
+					BatchVerticesRef[ VertexOffset++ ] = FSlateVertex( RenderTransform, FVector2D(LowerRight.X,UpperLeft.Y),	FVector4(U+SizeU, V,		UtMax,Vt),		FVector2D(1.0f,0.0f), FinalColor, RenderClipRect );
+					BatchVerticesRef[ VertexOffset++ ] = FSlateVertex( RenderTransform, FVector2D(UpperLeft.X,LowerRight.Y),	FVector4(U, V+SizeV,		Ut,VtMax),		FVector2D(0.0f,1.0f), FinalColor, RenderClipRect );
+					BatchVerticesRef[ VertexOffset++ ] = FSlateVertex( RenderTransform, LowerRight,								FVector4(U+SizeU, V+SizeV,	UtMax,VtMax),	FVector2D(1.0f,1.0f), FinalColor, RenderClipRect );
 
 					BatchIndicesRef[IndexOffset++] = IndexStart + 0;
 					BatchIndicesRef[IndexOffset++] = IndexStart + 1;
diff --git a/Engine/Source/Runtime/SlateCore/Private/Rendering/RenderingCommon.cpp b/Engine/Source/Runtime/SlateCore/Private/Rendering/RenderingCommon.cpp
--- a/Engine/Source/Runtime/SlateCore/Private/Rendering/RenderingCommon.cpp
+++ b/Engine/Source/Runtime/SlateCore/Private/Rendering/RenderingCommon.cpp
@@ -88,3 +88,23 @@ FSlateVertex::FSlateVertex( const FSlateRenderTransform& RenderTransform, const
 	Position[1] = FMath::RoundToInt(WindowPosition.Y);
 }
 
+
+FSlateVertex::FSlateVertex( const FSlateRenderTransform& RenderTransform, const FVector2D& InLocalPosition, const FVector4& InTexCoord, const FVector2D& InMaterialTexCoords, const FColor& InColor, const FSlateRotatedClipRectType& InClipRect )
+	: ClipRect( InClipRect )
+	, Color( InColor )
+{
+
+	TexCoords[0] = InTexCoord.X;
+	TexCoords[1] = InTexCoord.Y;
+	TexCoords[2] = InTexCoord.Z;
+	TexCoords[3] = InTexCoord.W;
+
+	MaterialTexCoords[0] = InMaterialTexCoords.X;
+	MaterialTexCoords[1] = InMaterialTexCoords.Y;
+
+	const FVector2D WindowPosition = TransformPoint(RenderTransform, InLocalPosition);
+
+	// Pixel snapping here.
+	Position[0] = FMath::RoundToInt(WindowPosition.X);
+	Position[1] = FMath::RoundToInt(WindowPosition.Y);
+}
diff --git a/Engine/Source/Runtime/SlateCore/Public/Fonts/FontCache.h b/Engine/Source/Runtime/SlateCore/Public/Fonts/FontCache.h
--- a/Engine/Source/Runtime/SlateCore/Public/Fonts/FontCache.h
+++ b/Engine/Source/Runtime/SlateCore/Public/Fonts/FontCache.h
@@ -265,7 +265,7 @@ public:
 	/** 
 	 * Flush the cache if needed
 	 */
-	void ConditionalFlushCache();
+	bool ConditionalFlushCache();
 
 	/**
 	 * Updates the texture used for rendering
diff --git a/Engine/Source/Runtime/SlateCore/Public/Fonts/SlateFontInfo.h b/Engine/Source/Runtime/SlateCore/Public/Fonts/SlateFontInfo.h
--- a/Engine/Source/Runtime/SlateCore/Public/Fonts/SlateFontInfo.h
+++ b/Engine/Source/Runtime/SlateCore/Public/Fonts/SlateFontInfo.h
@@ -18,6 +18,10 @@ struct SLATECORE_API FSlateFontInfo
 	UPROPERTY(EditAnywhere, BlueprintReadWrite, Category=SlateStyleRules, meta=(AllowedClasses="Font", DisplayName="Font Family"))
 	const UObject* FontObject;
 
+	/** The material to use when rendering this font */
+	UPROPERTY(EditAnywhere, BlueprintReadWrite, Category=SlateStyleRules, meta=(AllowedClasses="MaterialInterface"))
+	const UObject* FontMaterial;
+
 	/** The composite font data to use (valid when used with a Slate style set in C++) */
 	TSharedPtr<const FCompositeFont> CompositeFont;
 
@@ -109,6 +113,7 @@ public:
 	bool operator==( const FSlateFontInfo& Other ) const 
 	{
 		return FontObject == Other.FontObject
+			&& FontMaterial == Other.FontMaterial
 			&& CompositeFont == Other.CompositeFont 
 			&& TypefaceFontName == Other.TypefaceFontName
 			&& Size == Other.Size;
@@ -149,6 +154,7 @@ public:
 	{
 		uint32 Hash = 0;
 		Hash = HashCombine(Hash, GetTypeHash(FontInfo.FontObject));
+		Hash = HashCombine(Hash, GetTypeHash(FontInfo.FontMaterial));
 		Hash = HashCombine(Hash, GetTypeHash(FontInfo.CompositeFont));
 		Hash = HashCombine(Hash, GetTypeHash(FontInfo.TypefaceFontName));
 		Hash = HashCombine(Hash, GetTypeHash(FontInfo.Size));
diff --git a/Engine/Source/Runtime/SlateCore/Public/Rendering/RenderingCommon.h b/Engine/Source/Runtime/SlateCore/Public/Rendering/RenderingCommon.h
--- a/Engine/Source/Runtime/SlateCore/Public/Rendering/RenderingCommon.h
+++ b/Engine/Source/Runtime/SlateCore/Public/Rendering/RenderingCommon.h
@@ -139,6 +139,9 @@ struct FSlateVertex
 	/** Texture coordinates.  The first 2 are in xy and the 2nd are in zw */
 	float TexCoords[4]; 
 
+	/** Texture coordinates used as pass through to materials for custom texturing. */
+	float MaterialTexCoords[2];
+
 	/** Position of the vertex in window space */
 	float Position[2];
 
@@ -151,6 +154,7 @@ struct FSlateVertex
 	FSlateVertex();
 	FSlateVertex( const FSlateRenderTransform& RenderTransform, const FVector2D& InLocalPosition, const FVector2D& InTexCoord, const FVector2D& InTexCoord2, const FColor& InColor, const FSlateRotatedClipRectType& InClipRect );
 	FSlateVertex( const FSlateRenderTransform& RenderTransform, const FVector2D& InLocalPosition, const FVector2D& InTexCoord, const FColor& InColor, const FSlateRotatedClipRectType& InClipRect );
+	FSlateVertex( const FSlateRenderTransform& RenderTransform, const FVector2D& InLocalPosition, const FVector4& InTexCoords, const FVector2D& InMaterialTexCoords, const FColor& InColor, const FSlateRotatedClipRectType& InClipRect );
 };
 
 template<> struct TIsPODType<FSlateVertex> { enum { Value = true }; };
diff --git a/Engine/Source/Runtime/SlateCore/Public/Textures/TextureManager.h b/Engine/Source/Runtime/SlateCore/Public/Textures/TextureManager.h
--- a/Engine/Source/Runtime/SlateCore/Public/Textures/TextureManager.h
+++ b/Engine/Source/Runtime/SlateCore/Public/Textures/TextureManager.h
@@ -50,6 +50,8 @@ public:
 	 */
 	virtual FSlateShaderResourceProxy* GetShaderResource( const FSlateBrush& InBrush ) = 0;
 
+	virtual FSlateShaderResource* GetFontShaderResource( uint32 FontAtlasIndex, FSlateShaderResource* FontTextureAtlas, const class UObject* FontMaterial ) { return FontTextureAtlas; }
+
 	/**
 	 * Returns the way to access the texture atlas information from this resource manager
 	 */
diff --git a/Engine/Source/Runtime/SlateRHIRenderer/Private/SlateMaterialResource.cpp b/Engine/Source/Runtime/SlateRHIRenderer/Private/SlateMaterialResource.cpp
--- a/Engine/Source/Runtime/SlateRHIRenderer/Private/SlateMaterialResource.cpp
+++ b/Engine/Source/Runtime/SlateRHIRenderer/Private/SlateMaterialResource.cpp
@@ -4,10 +4,11 @@
 #include "SlateMaterialResource.h"
 
 
-FSlateMaterialResource::FSlateMaterialResource(UMaterialInterface& InMaterial, const FVector2D& InImageSize)
+FSlateMaterialResource::FSlateMaterialResource(const UMaterialInterface& InMaterial, const FVector2D& InImageSize, FSlateShaderResource* InTextureMask )
 	: MaterialObject( &InMaterial )
 	, RenderProxy( InMaterial.GetRenderProxy(false) )
 	, SlateProxy( new FSlateShaderResourceProxy )
+	, TextureMaskResource( InTextureMask )
 	, Width(FMath::RoundToInt(InImageSize.X))
 	, Height(FMath::RoundToInt(InImageSize.Y))
 {
@@ -23,7 +24,7 @@ FSlateMaterialResource::~FSlateMaterialResource()
 	}
 }
 
-void FSlateMaterialResource::UpdateMaterial(UMaterialInterface& InMaterialResource, const FVector2D& InImageSize)
+void FSlateMaterialResource::UpdateMaterial(const UMaterialInterface& InMaterialResource, const FVector2D& InImageSize, FSlateShaderResource* InTextureMask )
 {
 	MaterialObject = &InMaterialResource;
 	RenderProxy = InMaterialResource.GetRenderProxy(false);
@@ -32,6 +33,8 @@ void FSlateMaterialResource::UpdateMaterial(UMaterialInterface& InMaterialResour
 		SlateProxy = new FSlateShaderResourceProxy;
 	}
 
+	TextureMaskResource = InTextureMask;
+
 	SlateProxy->ActualSize = InImageSize.IntPoint();
 	SlateProxy->Resource = this;
 
diff --git a/Engine/Source/Runtime/SlateRHIRenderer/Private/SlateMaterialResource.h b/Engine/Source/Runtime/SlateRHIRenderer/Private/SlateMaterialResource.h
--- a/Engine/Source/Runtime/SlateRHIRenderer/Private/SlateMaterialResource.h
+++ b/Engine/Source/Runtime/SlateRHIRenderer/Private/SlateMaterialResource.h
@@ -8,14 +8,14 @@
 class FSlateMaterialResource : public FSlateShaderResource
 {
 public:
-	FSlateMaterialResource(UMaterialInterface& InMaterialResource, const FVector2D& InImageSize);
+	FSlateMaterialResource(const UMaterialInterface& InMaterialResource, const FVector2D& InImageSize, FSlateShaderResource* InTextureMask = nullptr );
 	~FSlateMaterialResource();
 
 	virtual uint32 GetWidth() const override { return Width; }
 	virtual uint32 GetHeight() const override { return Height; }
 	virtual ESlateShaderResource::Type GetType() const override { return ESlateShaderResource::Material; }
 
-	void UpdateMaterial(UMaterialInterface& InMaterialResource, const FVector2D& InImageSize);
+	void UpdateMaterial(const UMaterialInterface& InMaterialResource, const FVector2D& InImageSize, FSlateShaderResource* InTextureMask );
 
 	/** Updates the rendering resource */
 	void UpdateRenderResource(FMaterialRenderProxy* Proxy) { RenderProxy = Proxy; }
@@ -24,15 +24,17 @@ public:
 	FMaterialRenderProxy* GetRenderProxy() const { return RenderProxy; }
 
 	/** @return the material object */
-	UMaterialInterface* GetMaterialObject() const { return MaterialObject; }
+	const UMaterialInterface* GetMaterialObject() const { return MaterialObject; }
 
+	FSlateShaderResource* GetTextureMaskResource() const { return TextureMaskResource; }
 public:
-	class UMaterialInterface* MaterialObject;
+	const class UMaterialInterface* MaterialObject;
 	/** Rendering thread version of the material object */
 	FMaterialRenderProxy* RenderProxy;
 	/** Slate proxy used for batching the material */
 	FSlateShaderResourceProxy* SlateProxy;
 
+	FSlateShaderResource* TextureMaskResource;
 	uint32 Width;
 	uint32 Height;
 };
diff --git a/Engine/Source/Runtime/SlateRHIRenderer/Private/SlateMaterialShader.cpp b/Engine/Source/Runtime/SlateRHIRenderer/Private/SlateMaterialShader.cpp
--- a/Engine/Source/Runtime/SlateRHIRenderer/Private/SlateMaterialShader.cpp
+++ b/Engine/Source/Runtime/SlateRHIRenderer/Private/SlateMaterialShader.cpp
@@ -23,6 +23,8 @@ FSlateMaterialShaderPS::FSlateMaterialShaderPS(const FMaterialShaderType::Compil
 {
 	ShaderParams.Bind(Initializer.ParameterMap, TEXT("ShaderParams"));
 	GammaValues.Bind(Initializer.ParameterMap, TEXT("GammaValues"));
+	AdditionalTextureParameter.Bind(Initializer.ParameterMap, TEXT("ElementTexture"));
+	TextureParameterSampler.Bind(Initializer.ParameterMap, TEXT("ElementTextureSampler"));
 }
 
 void FSlateMaterialShaderPS::SetParameters(FRHICommandList& RHICmdList, const FSceneView& View, const FMaterialRenderProxy* MaterialRenderProxy, const FMaterial* Material, float InDisplayGamma, const FVector4& InShaderParams )
@@ -60,6 +62,10 @@ void FSlateMaterialShaderPS::SetParameters(FRHICommandList& RHICmdList, const FS
 
 }
 
+void FSlateMaterialShaderPS::SetAdditionalTexture( FRHICommandList& RHICmdList, const FTextureRHIParamRef InTexture, const FSamplerStateRHIRef SamplerState )
+{
+	SetTextureParameter(RHICmdList, GetPixelShader(), AdditionalTextureParameter, TextureParameterSampler, SamplerState, InTexture );
+}
 
 void FSlateMaterialShaderPS::SetDisplayGamma(FRHICommandList& RHICmdList, float InDisplayGamma)
 {
@@ -73,6 +79,8 @@ bool FSlateMaterialShaderPS::Serialize(FArchive& Ar)
 	bool bShaderHasOutdatedParameters = FMaterialShader::Serialize(Ar);
 	Ar << GammaValues;
 	Ar << ShaderParams;
+	Ar << TextureParameterSampler;
+	Ar << AdditionalTextureParameter;
 	return bShaderHasOutdatedParameters;
 }
 
@@ -85,3 +93,5 @@ IMPLEMENT_SLATE_MATERIALSHADER_TYPE(Default,true);
 IMPLEMENT_SLATE_MATERIALSHADER_TYPE(Default,false);
 IMPLEMENT_SLATE_MATERIALSHADER_TYPE(Border,true);
 IMPLEMENT_SLATE_MATERIALSHADER_TYPE(Border,false);
+IMPLEMENT_SLATE_MATERIALSHADER_TYPE(Font, true);
+IMPLEMENT_SLATE_MATERIALSHADER_TYPE(Font, false);
diff --git a/Engine/Source/Runtime/SlateRHIRenderer/Private/SlateMaterialShader.h b/Engine/Source/Runtime/SlateRHIRenderer/Private/SlateMaterialShader.h
--- a/Engine/Source/Runtime/SlateRHIRenderer/Private/SlateMaterialShader.h
+++ b/Engine/Source/Runtime/SlateRHIRenderer/Private/SlateMaterialShader.h
@@ -18,16 +18,20 @@ public:
 	FSlateMaterialShaderPS() {}
 	FSlateMaterialShaderPS(const FMaterialShaderType::CompiledShaderInitializerType& Initializer);
 
-	void SetParameters(FRHICommandList& RHICmdList, const FSceneView& View, const FMaterialRenderProxy* MaterialRenderProxy, const FMaterial* Material, float InDisplayGamma, const FVector4& InShaderParams );
+	void SetParameters(FRHICommandList& RHICmdList, const FSceneView& View, const FMaterialRenderProxy* MaterialRenderProxy, const FMaterial* Material, float InDisplayGamma, const FVector4& InShaderParams);
 
 	void SetDisplayGamma(FRHICommandList& RHICmdList, float InDisplayGamma);
 
+	void SetAdditionalTexture( FRHICommandList& RHICmdList, const FTextureRHIParamRef InTexture, const FSamplerStateRHIRef SamplerState );
+
 	virtual bool Serialize(FArchive& Ar) override;
 
 private:
-
 	FShaderParameter GammaValues;
 	FShaderParameter ShaderParams;
+	/** Extra texture (like a font atlas) to be used in addition to any material textures */
+	FShaderResourceParameter TextureParameterSampler;
+	FShaderResourceParameter AdditionalTextureParameter;
 };
 
 
diff --git a/Engine/Source/Runtime/SlateRHIRenderer/Private/SlateRHIRenderingPolicy.cpp b/Engine/Source/Runtime/SlateRHIRenderer/Private/SlateRHIRenderingPolicy.cpp
--- a/Engine/Source/Runtime/SlateRHIRenderer/Private/SlateRHIRenderingPolicy.cpp
+++ b/Engine/Source/Runtime/SlateRHIRenderer/Private/SlateRHIRenderingPolicy.cpp
@@ -407,6 +407,8 @@ void FSlateRHIRenderingPolicy::DrawElements(FRHICommandListImmediate& RHICmdList
 
 		if( !RenderBatch.CustomDrawer.IsValid() )
 		{
+			check(RenderBatch.NumIndices > 0);
+
 			if( !ShaderResource || ShaderResource->GetType() != ESlateShaderResource::Material ) 
 			{
 				FSlateElementPS* PixelShader = GetTexturePixelShader(ShaderType, DrawEffects);
@@ -490,9 +492,6 @@ void FSlateRHIRenderingPolicy::DrawElements(FRHICommandListImmediate& RHICmdList
 				PixelShader->SetShaderParams(RHICmdList, ShaderParams.PixelParams);
 				PixelShader->SetDisplayGamma(RHICmdList, (DrawFlags & ESlateBatchDrawFlag::NoGamma) ? 1.0f : DisplayGamma);
 
-
-				check(RenderBatch.NumIndices > 0);
-
 				uint32 PrimitiveCount = RenderBatch.DrawPrimitiveType == ESlateDrawPrimitive::LineList ? RenderBatch.NumIndices / 2 : RenderBatch.NumIndices / 3;
 
 				// for RHIs that can't handle VertexOffset, we need to offset the stream source each time
@@ -515,7 +514,8 @@ void FSlateRHIRenderingPolicy::DrawElements(FRHICommandListImmediate& RHICmdList
 					SceneView = &CreateSceneView(SceneViewContext, BackBuffer, ViewProjectionMatrix);
 				}
 
-				FMaterialRenderProxy* MaterialRenderProxy = ((FSlateMaterialResource*)ShaderResource)->GetRenderProxy();
+				FSlateMaterialResource* MaterialShaderResource = (FSlateMaterialResource*)ShaderResource;
+				FMaterialRenderProxy* MaterialRenderProxy = MaterialShaderResource->GetRenderProxy();
 
 				const FMaterial* Material = MaterialRenderProxy->GetMaterial(SceneView->GetFeatureLevel());
 
@@ -533,10 +533,21 @@ void FSlateRHIRenderingPolicy::DrawElements(FRHICommandListImmediate& RHICmdList
 
 					PixelShader->SetParameters(RHICmdList, *SceneView, MaterialRenderProxy, Material, DisplayGamma, ShaderParams.PixelParams);
 					PixelShader->SetDisplayGamma(RHICmdList, (DrawFlags & ESlateBatchDrawFlag::NoGamma) ? 1.0f : DisplayGamma);
+					
+					FSlateShaderResource* MaskResource = MaterialShaderResource->GetTextureMaskResource();
+					if( MaskResource )
+					{
+						RHICmdList.SetBlendState(TStaticBlendState<CW_RGBA, BO_Add, BF_SourceAlpha, BF_InverseSourceAlpha, BO_Add, BF_InverseDestAlpha, BF_One>::GetRHI());
+
+						FTexture2DRHIRef TextureRHI;
+						TextureRHI = ((TSlateTexture<FTexture2DRHIRef>*)MaskResource)->GetTypedResource();
+
+			
+						PixelShader->SetAdditionalTexture(RHICmdList, TextureRHI, BilinearClamp);
+					}
 
 					VertexShader->SetViewProjection( RHICmdList, ViewProjectionMatrix );
 
-					check(RenderBatch.NumIndices > 0);
 
 					uint32 PrimitiveCount = RenderBatch.DrawPrimitiveType == ESlateDrawPrimitive::LineList ? RenderBatch.NumIndices / 2 : RenderBatch.NumIndices / 3;
 
@@ -670,21 +681,31 @@ FSlateMaterialShaderPS* FSlateRHIRenderingPolicy::GetMaterialPixelShader( const
 	case ESlateShader::Default:
 		if (bDrawDisabled)
 		{
-			FoundShader = MaterialShaderMap->GetShader(&TSlateMaterialShaderPS<0, true>::StaticType);
+			FoundShader = MaterialShaderMap->GetShader(&TSlateMaterialShaderPS<ESlateShader::Default, true>::StaticType);
 		}
 		else
 		{
-			FoundShader = MaterialShaderMap->GetShader(&TSlateMaterialShaderPS<0, false>::StaticType);
+			FoundShader = MaterialShaderMap->GetShader(&TSlateMaterialShaderPS<ESlateShader::Default, false>::StaticType);
 		}
 		break;
 	case ESlateShader::Border:
 		if (bDrawDisabled)
 		{
-			FoundShader = MaterialShaderMap->GetShader(&TSlateMaterialShaderPS<1, true>::StaticType);
+			FoundShader = MaterialShaderMap->GetShader(&TSlateMaterialShaderPS<ESlateShader::Border, true>::StaticType);
+		}
+		else
+		{
+			FoundShader = MaterialShaderMap->GetShader(&TSlateMaterialShaderPS<ESlateShader::Border, false>::StaticType);
+		}
+		break;
+	case ESlateShader::Font:
+		if(bDrawDisabled)
+		{
+			FoundShader = MaterialShaderMap->GetShader(&TSlateMaterialShaderPS<ESlateShader::Font, true>::StaticType);
 		}
 		else
 		{
-			FoundShader = MaterialShaderMap->GetShader(&TSlateMaterialShaderPS<1, false>::StaticType);
+			FoundShader = MaterialShaderMap->GetShader(&TSlateMaterialShaderPS<ESlateShader::Font, false>::StaticType);
 		}
 		break;
 	default:
diff --git a/Engine/Source/Runtime/SlateRHIRenderer/Private/SlateRHIResourceManager.cpp b/Engine/Source/Runtime/SlateRHIRenderer/Private/SlateRHIResourceManager.cpp
--- a/Engine/Source/Runtime/SlateRHIRenderer/Private/SlateRHIResourceManager.cpp
+++ b/Engine/Source/Runtime/SlateRHIRenderer/Private/SlateRHIResourceManager.cpp
@@ -33,7 +33,7 @@ TSharedPtr<FSlateUTextureResource> FDynamicResourceMap::GetUTextureResource( UTe
 	return nullptr;
 }
 
-TSharedPtr<FSlateMaterialResource> FDynamicResourceMap::GetMaterialResource( UMaterialInterface* Material ) const
+TSharedPtr<FSlateMaterialResource> FDynamicResourceMap::GetMaterialResource( const UMaterialInterface* Material ) const
 {
 	return MaterialMap.FindRef( Material );
 }
@@ -56,7 +56,7 @@ void FDynamicResourceMap::AddUTextureResource( UTexture2D* TextureObject, TShare
 	}
 }
 
-void FDynamicResourceMap::AddMaterialResource( UMaterialInterface* Material, TSharedRef<FSlateMaterialResource> InMaterialResource )
+void FDynamicResourceMap::AddMaterialResource( const UMaterialInterface* Material, TSharedRef<FSlateMaterialResource> InMaterialResource )
 {
 	check( Material == InMaterialResource->GetMaterialObject() );
 	MaterialMap.Add(Material, InMaterialResource);
@@ -78,7 +78,7 @@ void FDynamicResourceMap::RemoveUTextureResource( UTexture2D* TextureObject )
 	}
 }
 
-void FDynamicResourceMap::RemoveMaterialResource( UMaterialInterface* Material )
+void FDynamicResourceMap::RemoveMaterialResource( const UMaterialInterface* Material )
 {
 	MaterialMap.Remove(Material);
 }
@@ -433,7 +433,8 @@ FSlateShaderResourceProxy* FSlateRHIResourceManager::GetShaderResource( const FS
 	}
 	else if (InBrush.GetResourceObject() && InBrush.GetResourceObject()->IsA<UMaterialInterface>())
 	{
-		Texture = GetMaterialResource(InBrush);
+		FSlateMaterialResource* Resource = GetMaterialResource(InBrush.GetResourceObject(), InBrush.ImageSize, nullptr);
+		Texture = Resource->SlateProxy;
 	}
 	else if( InBrush.IsDynamicallyLoaded() || ( InBrush.HasUObject() ) )
 	{
@@ -449,6 +450,18 @@ FSlateShaderResourceProxy* FSlateRHIResourceManager::GetShaderResource( const FS
 	return Texture;
 }
 
+FSlateShaderResource* FSlateRHIResourceManager::GetFontShaderResource( uint32 FontAtlasIndex, FSlateShaderResource* FontTextureAtlas, const class UObject* FontMaterial )
+{
+	if( FontMaterial == nullptr )
+	{
+		return FontTextureAtlas;
+	}
+	else
+	{
+		return GetMaterialResource( FontMaterial, FVector2D::ZeroVector, FontTextureAtlas );
+	}
+}
+
 ISlateAtlasProvider* FSlateRHIResourceManager::GetTextureAtlasProvider()
 {
 	return this;
@@ -649,27 +662,25 @@ FSlateShaderResourceProxy* FSlateRHIResourceManager::FindOrCreateDynamicTextureR
 	return nullptr;
 }
 
-FSlateShaderResourceProxy* FSlateRHIResourceManager::GetMaterialResource(const FSlateBrush& InBrush)
+FSlateMaterialResource* FSlateRHIResourceManager::GetMaterialResource(const UObject* InMaterial, FVector2D ImageSize, FSlateShaderResource* TextureMask )
 {
 	check(IsThreadSafeForSlateRendering());
 
-	const FName ResourceName = InBrush.GetResourceName();
-
-	UMaterialInterface* Material = CastChecked<UMaterialInterface>(InBrush.GetResourceObject());
+	const UMaterialInterface* Material = CastChecked<UMaterialInterface>(InMaterial);
 
 	TSharedPtr<FSlateMaterialResource> MaterialResource = DynamicResourceMap.GetMaterialResource(Material);
-	if (!MaterialResource.IsValid())
+	if (!MaterialResource.IsValid() || MaterialResource->GetTextureMaskResource() != TextureMask )
 	{
 		// Get a resource from the free list if possible
 		if(MaterialResourceFreeList.Num() > 0)
 		{
 			MaterialResource = MaterialResourceFreeList.Pop();
 		
-			MaterialResource->UpdateMaterial( *Material, InBrush.ImageSize );
+			MaterialResource->UpdateMaterial( *Material, ImageSize, TextureMask );
 		}
 		else
 		{
-			MaterialResource = MakeShareable(new FSlateMaterialResource(*Material, InBrush.ImageSize));
+			MaterialResource = MakeShareable(new FSlateMaterialResource(*Material, ImageSize, TextureMask));
 		}
 		
 		DynamicResourceMap.AddMaterialResource(Material, MaterialResource.ToSharedRef());
@@ -678,10 +689,10 @@ FSlateShaderResourceProxy* FSlateRHIResourceManager::GetMaterialResource(const F
 	{
 		// Keep the resource up to date
 		MaterialResource->UpdateRenderResource(Material->GetRenderProxy(false, false));
-		MaterialResource->SlateProxy->ActualSize = InBrush.ImageSize.IntPoint();
+		MaterialResource->SlateProxy->ActualSize = ImageSize.IntPoint();
 	}
 
-	return MaterialResource->SlateProxy;
+	return MaterialResource.Get();
 }
 
 void FSlateRHIResourceManager::OnAppExit()
diff --git a/Engine/Source/Runtime/SlateRHIRenderer/Private/SlateRHIResourceManager.h b/Engine/Source/Runtime/SlateRHIRenderer/Private/SlateRHIResourceManager.h
--- a/Engine/Source/Runtime/SlateRHIRenderer/Private/SlateRHIResourceManager.h
+++ b/Engine/Source/Runtime/SlateRHIRenderer/Private/SlateRHIResourceManager.h
@@ -17,7 +17,7 @@ public:
 
 	TSharedPtr<FSlateUTextureResource> GetUTextureResource( UTexture2D* TextureObject ) const;
 
-	TSharedPtr<FSlateMaterialResource> GetMaterialResource( UMaterialInterface* Material ) const;
+	TSharedPtr<FSlateMaterialResource> GetMaterialResource( const UMaterialInterface* Material ) const;
 
 	void AddUTextureResource( UTexture2D* TextureObject, TSharedRef<FSlateUTextureResource> InResource );
 	void RemoveUTextureResource( UTexture2D* TextureObject );
@@ -25,8 +25,8 @@ public:
 	void AddDynamicTextureResource( FName ResourceName, TSharedRef<FSlateDynamicTextureResource> InResource);
 	void RemoveDynamicTextureResource( FName ResourceName );
 
-	void AddMaterialResource( UMaterialInterface* Material, TSharedRef<FSlateMaterialResource> InResource );
-	void RemoveMaterialResource( UMaterialInterface* Material );
+	void AddMaterialResource( const UMaterialInterface* Material, TSharedRef<FSlateMaterialResource> InResource );
+	void RemoveMaterialResource( const UMaterialInterface* Material );
 
 	void Empty();
 
@@ -98,6 +98,7 @@ public:
 
 	/** FSlateShaderResourceManager interface */
 	virtual FSlateShaderResourceProxy* GetShaderResource( const FSlateBrush& InBrush ) override;
+	virtual FSlateShaderResource* GetFontShaderResource( uint32 FontAtlasIndex, FSlateShaderResource* FontTextureAtlas, const class UObject* FontMaterial ) override;
 	virtual ISlateAtlasProvider* GetTextureAtlasProvider() override;
 
 	/**
@@ -184,9 +185,9 @@ private:
 	/**
 	 * Returns a rendering resource for a material
 	 *
-	 * @param InBrush	Slate brush for the material
+	 * @param InMaterial	The material object
 	 */
-	FSlateShaderResourceProxy* GetMaterialResource( const FSlateBrush& InBrush );
+	FSlateMaterialResource* GetMaterialResource( const UObject* InMaterial, FVector2D ImageSize, FSlateShaderResource* TextureMask );
 
 	/**
 	 * Called when the application exists before the UObject system shuts down so we can free object resources
diff --git a/Engine/Source/Runtime/SlateRHIRenderer/Private/SlateShaders.cpp b/Engine/Source/Runtime/SlateRHIRenderer/Private/SlateShaders.cpp
--- a/Engine/Source/Runtime/SlateRHIRenderer/Private/SlateShaders.cpp
+++ b/Engine/Source/Runtime/SlateRHIRenderer/Private/SlateShaders.cpp
@@ -44,10 +44,11 @@ void FSlateVertexDeclaration::InitRHI()
 	FVertexDeclarationElementList Elements;
 	uint32 Stride = sizeof(FSlateVertex);
 	Elements.Add(FVertexElement(0, STRUCT_OFFSET(FSlateVertex, TexCoords), VET_Float4, 0, Stride));
-	Elements.Add(FVertexElement(0, STRUCT_OFFSET(FSlateVertex, Position), VET_Float2, 1, Stride));
-	Elements.Add(FVertexElement(0, STRUCT_OFFSET(FSlateVertex, ClipRect) + STRUCT_OFFSET(FSlateRotatedClipRectType, TopLeft), VET_Float2, 2, Stride));
-	Elements.Add(FVertexElement(0, STRUCT_OFFSET(FSlateVertex, ClipRect) + STRUCT_OFFSET(FSlateRotatedClipRectType, ExtentX), VET_Float4, 3, Stride));
-	Elements.Add(FVertexElement(0, STRUCT_OFFSET(FSlateVertex, Color), VET_Color, 4, Stride));
+	Elements.Add(FVertexElement(0, STRUCT_OFFSET(FSlateVertex, MaterialTexCoords), VET_Float2, 1, Stride));
+	Elements.Add(FVertexElement(0, STRUCT_OFFSET(FSlateVertex, Position), VET_Float2, 2, Stride));
+	Elements.Add(FVertexElement(0, STRUCT_OFFSET(FSlateVertex, ClipRect) + STRUCT_OFFSET(FSlateRotatedClipRectType, TopLeft), VET_Float2, 3, Stride));
+	Elements.Add(FVertexElement(0, STRUCT_OFFSET(FSlateVertex, ClipRect) + STRUCT_OFFSET(FSlateRotatedClipRectType, ExtentX), VET_Float4, 4, Stride));
+	Elements.Add(FVertexElement(0, STRUCT_OFFSET(FSlateVertex, Color), VET_Color, 5, Stride));
 
 	VertexDeclarationRHI = RHICreateVertexDeclaration(Elements);
 }
