UMG - Fixing a bug where the modifier keys and additional base class data were never passed to new copies for mouse pointer events.  Only affected UMG and not Slate, as the VM makes copies of the struct.

diff --git a/Engine/Source/Runtime/SlateCore/Public/Input/Events.h b/Engine/Source/Runtime/SlateCore/Public/Input/Events.h
--- a/Engine/Source/Runtime/SlateCore/Public/Input/Events.h
+++ b/Engine/Source/Runtime/SlateCore/Public/Input/Events.h
@@ -61,39 +61,38 @@ struct FFocusEvent
 public:
 
 	/**
-	* UStruct Constructor.  Not meant for normal usage.
-	*/
+	 * UStruct Constructor.  Not meant for normal usage.
+	 */
 	FFocusEvent()
-		: Cause(EFocusCause::SetDirectly),
-		UserIndex(0)
+		: Cause(EFocusCause::SetDirectly)
+		, UserIndex(0)
 	{ }
 
 	/**
-	* Constructor.  Events are immutable once constructed.
-	*
-	* @param  InCause  The cause of the focus event
-	*/
+	 * Constructor.  Events are immutable once constructed.
+	 *
+	 * @param  InCause  The cause of the focus event
+	 */
 	FFocusEvent(const EFocusCause InCause, uint32 InUserIndex)
-		: Cause(InCause),
-		UserIndex(InUserIndex)
+		: Cause(InCause)
+		, UserIndex(InUserIndex)
 	{ }
 
-
 	/**
-	* Queries the reason for the focus change
-	*
-	* @return  The cause of the focus change
-	*/
+	 * Queries the reason for the focus change
+	 *
+	 * @return  The cause of the focus change
+	 */
 	EFocusCause GetCause() const
 	{
 		return Cause;
 	}
 
 	/**
-	* Queries the user that is changing focus
-	*
-	* @return  The user that is changing focus
-	*/
+	 * Queries the user that is changing focus
+	 *
+	 * @return  The user that is changing focus
+	 */
 	uint32 GetUser() const
 	{
 		return UserIndex;
@@ -354,14 +353,21 @@ protected:
 	const FWidgetPath* EventPath;
 };
 
+template<>
+struct TStructOpsTypeTraits<FInputEvent> : public TStructOpsTypeTraitsBase
+{
+	enum
+	{
+		WithCopy = true,
+	};
+};
 
 /**
  * FKeyEvent describes a key action (keyboard/controller key/button pressed or released.)
  * It is passed to event handlers dealing with key input.
  */
 USTRUCT(BlueprintType)
-struct FKeyEvent
-	: public FInputEvent
+struct FKeyEvent : public FInputEvent
 {
 	GENERATED_USTRUCT_BODY()
 
@@ -371,7 +377,7 @@ public:
 	 */
 	FKeyEvent()
 		: FInputEvent(FModifierKeysState(), 0, false)
-		, Key(EKeys::SpaceBar)
+		, Key()
 		, CharacterCode(0)
 		, KeyCode(0)
 	{
@@ -440,6 +446,15 @@ private:
 	uint32 KeyCode;
 };
 
+template<>
+struct TStructOpsTypeTraits<FKeyEvent> : public TStructOpsTypeTraitsBase
+{
+	enum
+	{
+		WithCopy = true,
+	};
+};
+
 /** DEPRECATED 4.6 - Do not use */
 //@Todo slate: Remove this as soon as the 4.6 deprecated API is Removed.
 #define FKeyboardEvent \
@@ -473,7 +488,7 @@ public:
 	* UStruct Constructor.  Not meant for normal usage.
 	*/
 	FAnalogInputEvent()
-		: FKeyEvent(FKey(EKeys::SpaceBar), FModifierKeysState(), false, 0, 0, 0)
+		: FKeyEvent(FKey(), FModifierKeysState(), false, 0, 0, 0)
 		, AnalogValue(0.0f)
 	{
 	}
@@ -513,6 +528,15 @@ private:
 	float AnalogValue;
 };
 
+template<>
+struct TStructOpsTypeTraits<FAnalogInputEvent> : public TStructOpsTypeTraitsBase
+{
+	enum
+	{
+		WithCopy = true,
+	};
+};
+
 /**
  * FCharacterEvent describes a keyboard action where the utf-16 code is given.  Used for OnKeyChar messages
  */
@@ -556,6 +580,15 @@ private:
 };
 
 
+template<>
+struct TStructOpsTypeTraits<FCharacterEvent> : public TStructOpsTypeTraitsBase
+{
+	enum
+	{
+		WithCopy = true,
+	};
+};
+
 /**
  * Helper class to auto-populate a set with the set of "keys" a touch represents
  */
@@ -584,6 +617,7 @@ public:
 };
 
 
+
 /**
  * FPointerEvent describes a mouse or touch action (e.g. Press, Release, Move, etc).
  * It is passed to event handlers dealing with pointer-based input.
@@ -603,7 +637,7 @@ public:
 		, LastScreenSpacePosition(FVector2D(0, 0))
 		, CursorDelta(FVector2D(0, 0))
 		, PressedButtons(FTouchKeySet::EmptySet)
-		, EffectingButton(EKeys::LeftMouseButton)
+		, EffectingButton()
 		, PointerIndex(0)
 		, TouchpadIndex(0)
 		, bIsTouchEvent(false)
@@ -670,7 +704,7 @@ public:
 		, LastScreenSpacePosition(InLastScreenSpacePosition)
 		, CursorDelta(InScreenSpacePosition - InLastScreenSpacePosition)
 		, PressedButtons(bPressLeftMouseButton ? FTouchKeySet::StandardSet : FTouchKeySet::EmptySet)
-		, EffectingButton(EKeys::LeftMouseButton)
+		, EffectingButton()
 		, PointerIndex(InPointerIndex)
 		, TouchpadIndex(InTouchpadIndex)
 		, bIsTouchEvent(true)
@@ -739,6 +773,9 @@ public:
 	/** We override the assignment operator to allow generated code to compile with the const ref member. */
 	void operator=( const FPointerEvent& Other )
 	{
+		FInputEvent::operator=( Other );
+
+		// Pointer
 		ScreenSpacePosition = Other.ScreenSpacePosition;
 		LastScreenSpacePosition = Other.LastScreenSpacePosition;
 		CursorDelta = Other.CursorDelta;
@@ -778,9 +815,19 @@ private:
 	bool bIsTouchEvent;
 	EGestureEvent::Type GestureType;
 	FVector2D WheelOrGestureDelta;
+	// NOTE: If you add a new member, make sure you add it to the assignment operator.
 };
 
 
+template<>
+struct TStructOpsTypeTraits<FPointerEvent> : public TStructOpsTypeTraitsBase
+{
+	enum
+	{
+		WithCopy = true,
+	};
+};
+
 /** DEPRECATED 4.6 - Do not use */
 //@Todo slate: Remove this as soon as the 4.6 deprecated API is Removed.
 USTRUCT(BlueprintType)
@@ -812,6 +859,15 @@ private:
 	float AnalogValue;
 };
 
+template<>
+struct TStructOpsTypeTraits<FControllerEvent> : public TStructOpsTypeTraitsBase
+{
+	enum
+	{
+		WithCopy = true,
+	};
+};
+
 
 /**
  * FMotionEvent describes a touch pad action (press, move, lift)
@@ -881,6 +937,15 @@ private:
 };
 
 
+template<>
+struct TStructOpsTypeTraits<FMotionEvent> : public TStructOpsTypeTraitsBase
+{
+	enum
+	{
+		WithCopy = true,
+	};
+};
+
 /**
 * FNavigationEvent describes a navigation action (Left, Right, Up, Down)
 * It is passed to event handlers dealing with navigation.
@@ -915,6 +980,16 @@ private:
 	EUINavigation NavigationType;
 };
 
+
+template<>
+struct TStructOpsTypeTraits<FNavigationEvent> : public TStructOpsTypeTraitsBase
+{
+	enum
+	{
+		WithCopy = true,
+	};
+};
+
 /**
  * FWindowActivateEvent describes a window being activated or deactivated.
  * (i.e. brought to the foreground or moved to the background)
