Work around not having the new version of freetype on some platforms yet

diff --git a/Engine/Source/Runtime/SlateCore/Private/Fonts/FontCache.cpp b/Engine/Source/Runtime/SlateCore/Private/Fonts/FontCache.cpp
--- a/Engine/Source/Runtime/SlateCore/Private/Fonts/FontCache.cpp
+++ b/Engine/Source/Runtime/SlateCore/Private/Fonts/FontCache.cpp
@@ -479,7 +479,7 @@ public:
 
 			if( Slot->bitmap.pixel_mode != FT_PIXEL_MODE_MONO )
 			{
-				for (uint32 Row = 0; Row < Bitmap->rows; ++Row)
+				for (uint32 Row = 0; Row < (uint32)Bitmap->rows; ++Row)
 				{
 					// Copy a single row. Note Bitmap.pitch contains the offset (in bytes) between rows.  Not always equal to Bitmap.width!
 					FMemory::Memcpy(&OutRenderData.RawPixels[Row*Bitmap->width], &Bitmap->buffer[Row*Bitmap->pitch], Bitmap->width*GlyphPixelSize);
@@ -490,7 +490,7 @@ public:
 			{
 				// In Mono a value of 1 means the pixel is drawn and a value of zero means it is not. 
 				// So we must check each pixel and convert it to a color.
-				for( uint32 Height = 0; Height < Bitmap->rows; ++Height )
+				for(uint32 Height = 0; Height < (uint32)Bitmap->rows; ++Height )
 				{
 					for( uint32 Width = 0; Width < Bitmap->width; ++Width )
 					{
