Fixed signed/unsigned mismatch warning.

diff --git a/Engine/Source/Runtime/SlateCore/Private/Fonts/FontCache.cpp b/Engine/Source/Runtime/SlateCore/Private/Fonts/FontCache.cpp
--- a/Engine/Source/Runtime/SlateCore/Private/Fonts/FontCache.cpp
+++ b/Engine/Source/Runtime/SlateCore/Private/Fonts/FontCache.cpp
@@ -492,7 +492,7 @@ public:
 				// So we must check each pixel and convert it to a color.
 				for(uint32 Height = 0; Height < (uint32)Bitmap->rows; ++Height )
 				{
-					for( uint32 Width = 0; Width < Bitmap->width; ++Width )
+					for( uint32 Width = 0; Width < (uint32)Bitmap->width; ++Width )
 					{
 						OutRenderData.RawPixels[Height*Bitmap->width+Width] = Bitmap->buffer[Height*Bitmap->pitch+Width] == 1 ? 255 : 0;
 					}
