Slate - Fixing a problem with Slate quantizing the lower range of black colors.  When we send linear colors down to the GPU we were first converting to byte linear space, in byte linear space several sRGB values all become quantized under a single number.  To avoid this problem, without increasing our vertex buffer size, we're now converting the Tint to sRGB Byte Space, then in the vertex shader convert the Tint back into Linear Space.  This adds a little more computation to both the GPU and CPU.

diff --git a/Engine/Shaders/SlateVertexShader.usf b/Engine/Shaders/SlateVertexShader.usf
--- a/Engine/Shaders/SlateVertexShader.usf
+++ b/Engine/Shaders/SlateVertexShader.usf
@@ -2,6 +2,7 @@
 
 
 #include "Common.usf"
+#include "GammaCorrectionCommon.usf"
 
 float4x4 ViewProjection;
 
@@ -41,6 +42,7 @@ VertexOut Main(
 	VOut.ClipOriginAndPos = float4(InClipOrigin, InPosition.xy);
 	VOut.ClipExtents = InClipExtents;
 
+	InColor.rgb = sRGBToLinear(InColor.rgb);
 	VOut.Color = InColor FCOLOR_COMPONENT_SWIZZLE;
 
 	return VOut;
diff --git a/Engine/Source/Runtime/SlateCore/Private/Rendering/ElementBatcher.cpp b/Engine/Source/Runtime/SlateCore/Private/Rendering/ElementBatcher.cpp
--- a/Engine/Source/Runtime/SlateCore/Private/Rendering/ElementBatcher.cpp
+++ b/Engine/Source/Runtime/SlateCore/Private/Rendering/ElementBatcher.cpp
@@ -55,7 +55,7 @@ FSlateRotatedClipRectType ToSnappedRotatedRect(const FSlateRect& ClipRectInLayou
 static FColor GetElementColor( const FLinearColor& InColor, const FSlateBrush* InBrush )
 {
 	// Pass the color through
-	return InColor.ToFColor(false);
+	return InColor.ToFColor(true);
 }
 
 FSlateElementBatcher::FSlateElementBatcher( TSharedRef<FSlateRenderingPolicy> InRenderingPolicy )
