Slate - Fixing the standalone slate renderers, that still need the vertex colors to come down in linear space.

diff --git a/Engine/Source/Developer/StandaloneRenderer/Private/OpenGL/SlateOpenGLRenderingPolicy.h b/Engine/Source/Developer/StandaloneRenderer/Private/OpenGL/SlateOpenGLRenderingPolicy.h
--- a/Engine/Source/Developer/StandaloneRenderer/Private/OpenGL/SlateOpenGLRenderingPolicy.h
+++ b/Engine/Source/Developer/StandaloneRenderer/Private/OpenGL/SlateOpenGLRenderingPolicy.h
@@ -18,7 +18,6 @@ public:
 	 */
 	void UpdateVertexAndIndexBuffers(FSlateBatchData& BatchData);
 
-	
 	/**
 	 * Draws Slate elements
 	 *
@@ -27,13 +26,15 @@ public:
 	 */
 	void DrawElements( const FMatrix& ViewProjectionMatrix, const TArray<FSlateRenderBatch>& RenderBatches );
 
-
 	virtual TSharedRef<FSlateShaderResourceManager> GetResourceManager() override;
+
 	/**
 	 * Returns the font cache used when the OpenGL rendering policy is active
- 	 */
+	 */
 	virtual TSharedRef<FSlateFontCache> GetFontCache() override { return FontCache.ToSharedRef(); }
 
+	virtual bool IsVertexColorInLinearSpace() const override { return false; }
+
 	/** 
 	 * Initializes resources if needed
 	 */
diff --git a/Engine/Source/Developer/StandaloneRenderer/Private/Windows/D3D/SlateD3DRenderingPolicy.h b/Engine/Source/Developer/StandaloneRenderer/Private/Windows/D3D/SlateD3DRenderingPolicy.h
--- a/Engine/Source/Developer/StandaloneRenderer/Private/Windows/D3D/SlateD3DRenderingPolicy.h
+++ b/Engine/Source/Developer/StandaloneRenderer/Private/Windows/D3D/SlateD3DRenderingPolicy.h
@@ -14,6 +14,7 @@ public:
 
 	virtual TSharedRef<FSlateShaderResourceManager> GetResourceManager() override;
 	virtual TSharedRef<FSlateFontCache> GetFontCache() override { return FontCache.ToSharedRef(); }
+	virtual bool IsVertexColorInLinearSpace() const override { return true; }
 
 private:
 	void InitResources();
diff --git a/Engine/Source/Runtime/SlateCore/Private/Rendering/ElementBatcher.cpp b/Engine/Source/Runtime/SlateCore/Private/Rendering/ElementBatcher.cpp
--- a/Engine/Source/Runtime/SlateCore/Private/Rendering/ElementBatcher.cpp
+++ b/Engine/Source/Runtime/SlateCore/Private/Rendering/ElementBatcher.cpp
@@ -47,22 +47,12 @@ FSlateRotatedClipRectType ToSnappedRotatedRect(const FSlateRect& ClipRectInLayou
 		SnappedBottomLeft - SnappedTopLeft);
 }
 
-/**
- * Computes the element tint color based in the user specified color and brush being used.
- * Note: The color could be in RGB or HSV
- * @return The final color to be passed per vertex
- */
-static FColor GetElementColor( const FLinearColor& InColor, const FSlateBrush* InBrush )
-{
-	// Pass the color through
-	return InColor.ToFColor(true);
-}
-
 FSlateElementBatcher::FSlateElementBatcher( TSharedRef<FSlateRenderingPolicy> InRenderingPolicy )
 	: BatchData( nullptr )
 	, ResourceManager( *InRenderingPolicy->GetResourceManager() )
 	, FontCache( *InRenderingPolicy->GetFontCache() )
 	, PixelCenterOffset( InRenderingPolicy->GetPixelCenterOffset() )
+	, bSRGBVertexColor( !InRenderingPolicy->IsVertexColorInLinearSpace() )
 {
 
 }
@@ -309,7 +299,7 @@ void FSlateElementBatcher::AddBoxElement( const FSlateDrawElement& DrawElement )
 		}
 
 
-		FColor Tint = GetElementColor(InPayload.Tint, BrushResource);
+		FColor Tint = InPayload.Tint.ToFColor(bSRGBVertexColor);
 
 		const ESlateBrushTileType::Type TilingRule = BrushResource->Tiling;
 		const bool bTileHorizontal = (TilingRule == ESlateBrushTileType::Both || TilingRule == ESlateBrushTileType::Horizontal);
@@ -603,7 +593,7 @@ void FSlateElementBatcher::AddTextElement(const FSlateDrawElement& DrawElement)
 	
 	const bool bIsFontMaterial = InPayload.FontInfo.FontMaterial != nullptr;
 
-	FColor FinalColor = GetElementColor( InPayload.Tint, nullptr );
+	FColor FinalColor = InPayload.Tint.ToFColor(bSRGBVertexColor);
 
 	uint32 NumChars = Text.Len();
 
@@ -911,7 +901,7 @@ void FSlateElementBatcher::AddSplineElement( const FSlateDrawElement& DrawElemen
 	FVector2D StartPos = StartPt;
 	FVector2D EndPos = FVector2D( FMath::CubicInterp( StartPt, StartDir, EndPt, EndDir, Alpha ) );
 
-	FColor FinalColor = GetElementColor( InPayload.Tint, nullptr );
+	FColor FinalColor = InPayload.Tint.ToFColor(bSRGBVertexColor);
 
 	BatchVertices.Add( FSlateVertex( RenderTransform, StartPos + Up, TransformPoint(RenderTransform, StartPos), TransformPoint(RenderTransform, EndPos), FinalColor, RenderClipRect ) );
 	BatchVertices.Add( FSlateVertex( RenderTransform, StartPos - Up, TransformPoint(RenderTransform, StartPos), TransformPoint(RenderTransform, EndPos), FinalColor, RenderClipRect ) );
@@ -1014,7 +1004,7 @@ void FSlateElementBatcher::AddLineElement( const FSlateDrawElement& DrawElement
 		return;
 	}
 
-	FColor FinalColor = GetElementColor( InPayload.Tint, nullptr );
+	FColor FinalColor = InPayload.Tint.ToFColor(bSRGBVertexColor);
 
 	if( InPayload.bAntialias )
 	{
@@ -1181,7 +1171,7 @@ void FSlateElementBatcher::AddViewportElement( const FSlateDrawElement& DrawElem
 	FSlateLayoutTransform InverseLayoutTransform(Inverse(FSlateLayoutTransform(DrawElement.GetScale(), DrawElement.GetPosition())));
 	FSlateRotatedClipRectType RenderClipRect = ToSnappedRotatedRect(InClippingRect, InverseLayoutTransform, RenderTransform);
 
-	const FColor FinalColor = GetElementColor( InPayload.Tint, nullptr );
+	const FColor FinalColor = InPayload.Tint.ToFColor(bSRGBVertexColor);
 
 	ESlateBatchDrawFlag::Type DrawFlags = ESlateBatchDrawFlag::None;
 	
@@ -1341,7 +1331,7 @@ void FSlateElementBatcher::AddBorderElement( const FSlateDrawElement& DrawElemen
 	FShaderParams ShaderParams = FShaderParams::MakePixelShaderParams( FVector4(LeftMarginU,RightMarginU,TopMarginV,BottomMarginV) );
 
 	// The tint color applies to all brushes and is passed per vertex
-	FColor Tint = GetElementColor( InPayload.Tint, InPayload.BrushResource );
+	FColor Tint = InPayload.Tint.ToFColor(bSRGBVertexColor);
 
 	// Pass the tiling information as a flag so we can pick the correct texture addressing mode
 	ESlateBatchDrawFlag::Type DrawFlags = (ESlateBatchDrawFlag::TileU|ESlateBatchDrawFlag::TileV);
diff --git a/Engine/Source/Runtime/SlateCore/Public/Rendering/ElementBatcher.h b/Engine/Source/Runtime/SlateCore/Public/Rendering/ElementBatcher.h
--- a/Engine/Source/Runtime/SlateCore/Public/Rendering/ElementBatcher.h
+++ b/Engine/Source/Runtime/SlateCore/Public/Rendering/ElementBatcher.h
@@ -117,6 +117,9 @@ private:
 	/** Offset to use when supporting 1:1 texture to pixel snapping */
 	const float PixelCenterOffset;
 
+	/** Are the vertex colors expected to be in sRGB space? */
+	const bool bSRGBVertexColor;
+
 	// true if any element in the batch requires vsync.
 	bool bRequiresVsync;
 };
diff --git a/Engine/Source/Runtime/SlateCore/Public/Rendering/RenderingPolicy.h b/Engine/Source/Runtime/SlateCore/Public/Rendering/RenderingPolicy.h
--- a/Engine/Source/Runtime/SlateCore/Public/Rendering/RenderingPolicy.h
+++ b/Engine/Source/Runtime/SlateCore/Public/Rendering/RenderingPolicy.h
@@ -31,6 +31,8 @@ public:
 
 	virtual TSharedRef<class FSlateShaderResourceManager> GetResourceManager() = 0;
 
+	virtual bool IsVertexColorInLinearSpace() const = 0;
+
 	float GetPixelCenterOffset() const
 	{
 		return PixelCenterOffset;
diff --git a/Engine/Source/Runtime/SlateRHIRenderer/Private/SlateRHIRenderingPolicy.h b/Engine/Source/Runtime/SlateRHIRenderer/Private/SlateRHIRenderingPolicy.h
--- a/Engine/Source/Runtime/SlateRHIRenderer/Private/SlateRHIRenderingPolicy.h
+++ b/Engine/Source/Runtime/SlateRHIRenderer/Private/SlateRHIRenderingPolicy.h
@@ -237,6 +237,7 @@ public:
 
 	virtual TSharedRef<FSlateFontCache> GetFontCache() override { return FontCache.ToSharedRef(); }
 	virtual TSharedRef<FSlateShaderResourceManager> GetResourceManager() override { return ResourceManager; }
+	virtual bool IsVertexColorInLinearSpace() const override { return false; }
 	
 	void InitResources();
 	void ReleaseResources();
