UMG/Slate - Initial version of SInvalidationPanel!  Has several limitations, but permits UMG and Slate to effectively "Freeze" UI, caching a great deal of the processing, computation...etc required every time slate paints.  Cached widgets do not get prepass, tick, or paint until their parent SInvalidationPanel is invalidated.  Widgets that change often, can be marked as Volatile and get treated special when used in conjunction with the SInvalidationPanel, they get stored seperately and get prepass/tick/paint like normal, but using the cached geometry of their container they origionally were computed to have.  Changes to their desired size will not change the invalidation of the rest of the panel, a manual invocation would be required for that.

When used effectively the SInvalidationPanel can shave milliseconds off your Slate time on a complicated hierarchy of widgets.  Provided those widgets don't need hit testing or input, this can be a very useful widget to use.  Scenarios like a HUD are the primary target of this panel, future advances will hopefully make it possible to be used in conjunction with a UI that needs input.

diff --git a/Engine/Source/Developer/SlateReflector/Private/Widgets/SWidgetReflector.cpp b/Engine/Source/Developer/SlateReflector/Private/Widgets/SWidgetReflector.cpp
--- a/Engine/Source/Developer/SlateReflector/Private/Widgets/SWidgetReflector.cpp
+++ b/Engine/Source/Developer/SlateReflector/Private/Widgets/SWidgetReflector.cpp
@@ -3,6 +3,7 @@
 #include "SlateReflectorPrivatePCH.h"
 #include "ISlateReflectorModule.h"
 #include "SlateStats.h"
+#include "SInvalidationPanel.h"
 
 #define LOCTEXT_NAMESPACE "SWidgetReflector"
 #define WITH_EVENT_LOGGING 0
@@ -432,6 +433,33 @@ void SWidgetReflector::Construct( const FArguments& InArgs )
 							]
 						]
 					]
+#if !UE_BUILD_SHIPPING
+					+ SHorizontalBox::Slot()
+					.AutoWidth()
+					.Padding(5.0f)
+					[
+						SNew(SCheckBox)
+						.Style(FCoreStyle::Get(), "ToggleButtonCheckbox")
+						.IsChecked_Lambda([&]()
+						{
+							return SInvalidationPanel::IsInvalidationDebuggingEnabled() ? ECheckBoxState::Checked : ECheckBoxState::Unchecked;
+						})
+						.OnCheckStateChanged_Lambda([&](const ECheckBoxState NewState)
+						{
+							SInvalidationPanel::EnableInvalidationDebugging(( NewState == ECheckBoxState::Checked ) ? true : false);
+						})
+						[
+							SNew(SBox)
+							.VAlign(VAlign_Center)
+							.HAlign(HAlign_Center)
+							[
+								SNew(STextBlock)
+								.Text(LOCTEXT("InvalidationDebugging", "Invalidation Debugging"))
+							]
+						]
+					]
+#endif
+
 #if SLATE_STATS
 					+SHorizontalBox::Slot()
 					.AutoWidth()
diff --git a/Engine/Source/Runtime/Slate/Private/Framework/Application/SlateApplication.cpp b/Engine/Source/Runtime/Slate/Private/Framework/Application/SlateApplication.cpp
--- a/Engine/Source/Runtime/Slate/Private/Framework/Application/SlateApplication.cpp
+++ b/Engine/Source/Runtime/Slate/Private/Framework/Application/SlateApplication.cpp
@@ -1027,6 +1027,8 @@ void FSlateApplication::DrawWindowAndChildren( const TSharedRef<SWindow>& Window
 			}
 		}
 
+#if SLATE_HAS_WIDGET_REFLECTOR
+
 		// The widget reflector may want to paint some additional stuff as part of the Widget introspection that it performs.
 		// For example: it may draw layout rectangles for hovered widgets.
 		const bool bVisualizeLayoutUnderCursor = DrawWindowArgs.WidgetsUnderCursor.IsValid();
@@ -1044,6 +1046,8 @@ void FSlateApplication::DrawWindowAndChildren( const TSharedRef<SWindow>& Window
 			MaxLayerId = WidgetReflector->VisualizeCursorAndKeys( WindowElementList, MaxLayerId );
 		}
 
+#endif
+
 		// Draw the child windows
 		const TArray< TSharedRef<SWindow> >& WindowChildren = WindowToDraw->GetChildWindows();
 		for (int32 ChildIndex=0; ChildIndex < WindowChildren.Num(); ++ChildIndex)
@@ -1133,8 +1137,12 @@ void FSlateApplication::PrivateDrawWindows( TSharedPtr<SWindow> DrawOnlyThisWind
 {
 	check(Renderer.IsValid());
 
+#if SLATE_HAS_WIDGET_REFLECTOR
 	// Is user expecting visual feedback from the Widget Reflector?
 	const bool bVisualizeLayoutUnderCursor = WidgetReflectorPtr.IsValid() && WidgetReflectorPtr.Pin()->IsVisualizingLayoutUnderCursor();
+#else
+	const bool bVisualizeLayoutUnderCursor = false;
+#endif
 
 	FWidgetPath WidgetsUnderCursor = bVisualizeLayoutUnderCursor
 		? WidgetsUnderCursorLastEvent.ToWidgetPath()
@@ -3841,6 +3849,7 @@ bool FSlateApplication::ProcessKeyDownEvent( FKeyEvent& InKeyEvent )
 	{
 		LastUserInteractionTimeForThrottling = LastUserInteractionTime;
 
+#if SLATE_HAS_WIDGET_REFLECTOR
 		// If we are inspecting, pressing ESC exits inspection mode.
 		if ( InKeyEvent.GetKey() == EKeys::Escape )
 		{
@@ -3854,6 +3863,7 @@ bool FSlateApplication::ProcessKeyDownEvent( FKeyEvent& InKeyEvent )
 					return Reply.IsEventHandled();
 			}
 		}
+#endif
 
 #if !(UE_BUILD_SHIPPING || UE_BUILD_TEST)
 		// Ctrl+Shift+~ summons the Toolbox.
diff --git a/Engine/Source/Runtime/Slate/Private/Widgets/Images/SImage.cpp b/Engine/Source/Runtime/Slate/Private/Widgets/Images/SImage.cpp
--- a/Engine/Source/Runtime/Slate/Private/Widgets/Images/SImage.cpp
+++ b/Engine/Source/Runtime/Slate/Private/Widgets/Images/SImage.cpp
@@ -68,17 +68,26 @@ FVector2D SImage::ComputeDesiredSize( float ) const
 
 void SImage::SetColorAndOpacity( const TAttribute<FSlateColor>& InColorAndOpacity )
 {
-	ColorAndOpacity = InColorAndOpacity;
+	if ( !ColorAndOpacity.IdenticalTo(InColorAndOpacity) )
+	{
+		ColorAndOpacity = InColorAndOpacity;
+		InvalidateLayout();
+	}
 }
 
 void SImage::SetColorAndOpacity( FLinearColor InColorAndOpacity )
 {
-	ColorAndOpacity = InColorAndOpacity;
+	if ( ColorAndOpacity.IsBound() || ColorAndOpacity.Get() != InColorAndOpacity )
+	{
+		ColorAndOpacity = InColorAndOpacity;
+		InvalidateLayout();
+	}
 }
 
 void SImage::SetImage(TAttribute<const FSlateBrush*> InImage)
 {
 	Image = InImage;
+	InvalidateLayout();
 }
 
 void SImage::SetOnMouseButtonDown(FPointerEventHandler EventHandler)
diff --git a/Engine/Source/Runtime/Slate/Private/Widgets/Notifications/SProgressBar.cpp b/Engine/Source/Runtime/Slate/Private/Widgets/Notifications/SProgressBar.cpp
--- a/Engine/Source/Runtime/Slate/Private/Widgets/Notifications/SProgressBar.cpp
+++ b/Engine/Source/Runtime/Slate/Private/Widgets/Notifications/SProgressBar.cpp
@@ -29,42 +29,53 @@ void SProgressBar::Construct( const FArguments& InArgs )
 
 void SProgressBar::SetPercent(TAttribute< TOptional<float> > InPercent)
 {
-	Percent = InPercent;
+	if ( !Percent.IdenticalTo(InPercent) )
+	{
+		Percent = InPercent;
+		InvalidateLayout();
+	}
 }
 
 void SProgressBar::SetStyle(const FProgressBarStyle* InStyle)
 {
 	Style = InStyle;
+	InvalidateLayout();
 }
 
 void SProgressBar::SetBarFillType(EProgressBarFillType::Type InBarFillType)
 {
 	BarFillType = InBarFillType;
+	InvalidateLayout();
 }
 
 void SProgressBar::SetFillColorAndOpacity(TAttribute< FSlateColor > InFillColorAndOpacity)
 {
 	FillColorAndOpacity = InFillColorAndOpacity;
+	InvalidateLayout();
 }
 
 void SProgressBar::SetBorderPadding(TAttribute< FVector2D > InBorderPadding)
 {
 	BorderPadding = InBorderPadding;
+	InvalidateLayout();
 }
 
 void SProgressBar::SetBackgroundImage(const FSlateBrush* InBackgroundImage)
 {
 	BackgroundImage = InBackgroundImage;
+	InvalidateLayout();
 }
 
 void SProgressBar::SetFillImage(const FSlateBrush* InFillImage)
 {
 	FillImage = InFillImage;
+	InvalidateLayout();
 }
 
 void SProgressBar::SetMarqueeImage(const FSlateBrush* InMarqueeImage)
 {
 	MarqueeImage = InMarqueeImage;
+	InvalidateLayout();
 }
 
 const FSlateBrush* SProgressBar::GetBackgroundImage() const
@@ -240,16 +251,16 @@ int32 SProgressBar::OnPaint( const FPaintArgs& Args, const FGeometry& AllottedGe
 	return RetLayerId - 1;
 }
 
-/**
- * Computes the desired size of this widget (SWidget)
- *
- * @return  The widget's desired size
- */
 FVector2D SProgressBar::ComputeDesiredSize( float ) const
 {
 	return GetMarqueeImage()->ImageSize;
 }
 
+bool SProgressBar::ComputeVolatility() const
+{
+	return SLeafWidget::ComputeVolatility() || Percent.IsBound();
+}
+
 void SProgressBar::SetActiveTimerTickRate(float TickRate)
 {
 	if (CurrentTickRate != TickRate || !ActiveTimerHandle.IsValid())
diff --git a/Engine/Source/Runtime/Slate/Private/Widgets/SInvalidationPanel.cpp b/Engine/Source/Runtime/Slate/Private/Widgets/SInvalidationPanel.cpp
new file mode 100644
--- /dev/null
+++ b/Engine/Source/Runtime/Slate/Private/Widgets/SInvalidationPanel.cpp
@@ -0,0 +1,261 @@
+// Copyright 1998-2015 Epic Games, Inc. All Rights Reserved.
+
+#include "SlatePrivatePCH.h"
+#include "LayoutUtils.h"
+#include "SInvalidationPanel.h"
+
+#if !UE_BUILD_SHIPPING
+
+/** True if we should allow widgets to be cached in the UI at all. */
+TAutoConsoleVariable<int32> InvalidationDebugging(
+	TEXT("Slate.InvalidationDebugging"),
+	false,
+	TEXT("Whether to show invalidation debugging visualization"));
+
+bool SInvalidationPanel::IsInvalidationDebuggingEnabled()
+{
+	return InvalidationDebugging.GetValueOnGameThread() == 1;
+}
+
+void SInvalidationPanel::EnableInvalidationDebugging(bool bEnable)
+{
+	InvalidationDebugging.AsVariable()->Set(bEnable);
+}
+
+/** True if we should allow widgets to be cached in the UI at all. */
+TAutoConsoleVariable<int32> EnableWidgetCaching(
+	TEXT("Slate.EnableWidgetCaching"),
+	true,
+	TEXT("Whether to attempt to cache any widgets through invalidation panels."));
+
+#endif
+
+void SInvalidationPanel::Construct( const FArguments& InArgs )
+{
+	ChildSlot
+	[
+		InArgs._Content.Widget
+	];
+
+	bNeedsCaching = true;
+	bIsInvalidating = false;
+	bCanCache = true;
+}
+
+bool SInvalidationPanel::GetCanCache() const
+{
+#if !UE_BUILD_SHIPPING
+	return bCanCache && EnableWidgetCaching.GetValueOnGameThread() == 1;
+#else
+	return bCanCache;
+#endif
+}
+
+void SInvalidationPanel::SetCanCache(bool InCanCache)
+{
+	bCanCache = InCanCache;
+	Visibility = InCanCache ? EVisibility::HitTestInvisible : EVisibility::SelfHitTestInvisible;
+
+	Invalidate();
+	InvalidateLayout();
+}
+
+void SInvalidationPanel::Tick( const FGeometry& AllottedGeometry, const double InCurrentTime, const float InDeltaTime )
+{
+	if ( GetCanCache() )
+	{
+		if ( bNeedsCaching == false )
+		{
+			// If the container we're in has changed in any way we need to invalidate for sure.
+			if ( AllottedGeometry.GetAccumulatedLayoutTransform() != LastAllottedGeometry.GetAccumulatedLayoutTransform() ||
+				AllottedGeometry.GetAccumulatedRenderTransform() != LastAllottedGeometry.GetAccumulatedRenderTransform() )
+			{
+				Invalidate();
+			}
+		}
+
+		LastAllottedGeometry = AllottedGeometry;
+
+		if ( bNeedsCaching )
+		{
+			SlatePrepass(LastAllottedGeometry.Scale);
+			CachePrepass(SharedThis(this));
+		}
+	}
+}
+
+void SInvalidationPanel::OnArrangeChildren( const FGeometry& AllottedGeometry, FArrangedChildren& ArrangedChildren ) const
+{
+	SCompoundWidget::OnArrangeChildren( AllottedGeometry, ArrangedChildren );
+}
+
+FChildren* SInvalidationPanel::GetChildren()
+{
+	if ( GetCanCache() == false || bNeedsCaching )
+	{
+		return SCompoundWidget::GetChildren();
+	}
+	else
+	{
+		return &EmptyChildSlot;
+	}
+}
+
+void SInvalidationPanel::InvalidateWidget(SWidget* InvalidateWidget)
+{
+	bNeedsCaching = true;
+
+#if !UE_BUILD_SHIPPING
+	if ( InvalidateWidget != nullptr && IsInvalidationDebuggingEnabled() )
+	{
+		InvalidatorWidgets.Add(InvalidateWidget->AsShared());
+	}
+#endif
+}
+
+int32 SInvalidationPanel::OnPaint( const FPaintArgs& Args, const FGeometry& AllottedGeometry, const FSlateRect& MyClippingRect, FSlateWindowElementList& OutDrawElements, int32 LayerId, const FWidgetStyle& InWidgetStyle, bool bParentEnabled ) const
+{
+	if ( GetCanCache() )
+	{
+		const bool bWasCachingNeeded = bNeedsCaching;
+
+		if ( bNeedsCaching )
+		{
+			// Always set the caching flag to false first, during the paint / tick pass we may change something
+			// to volatile and need to re-cache.
+			bNeedsCaching = false;
+
+			bIsInvalidating = true;
+
+			if ( !CachedWindowElements.IsValid() || CachedWindowElements->GetWindow() != OutDrawElements.GetWindow() )
+			{
+				CachedWindowElements = MakeShareable(new FSlateWindowElementList(OutDrawElements.GetWindow()));
+			}
+			else
+			{
+				CachedWindowElements->Reset();
+			}
+
+			//TODO: When SWidget::Paint is called don't drag self if volatile, and we're doing a cache pass.
+			CachedMaxChildLayer = SCompoundWidget::OnPaint(
+				Args.EnableCaching(SharedThis(const_cast<SInvalidationPanel*>( this )), true, false),
+				AllottedGeometry,
+				MyClippingRect,
+				*CachedWindowElements.Get(),
+				LayerId,
+				InWidgetStyle,
+				bParentEnabled);
+
+			bIsInvalidating = false;
+		}
+
+		OutDrawElements.AppendDrawElements(CachedWindowElements->GetDrawElements());
+
+		int32 OutMaxChildLayer = CachedMaxChildLayer;
+
+		// Paint the volatile elements
+		if ( CachedWindowElements.IsValid() )
+		{
+			OutMaxChildLayer = FMath::Max(CachedMaxChildLayer, CachedWindowElements->PaintVolatile(OutDrawElements));
+		}
+
+#if !UE_BUILD_SHIPPING
+
+		if ( IsInvalidationDebuggingEnabled() )
+		{
+			// Draw a green or red border depending on if we were invalidated this frame.
+			{
+				check(Args.IsCaching() == false);
+				//const bool bShowOutlineAsCached = Args.IsCaching() || bWasCachingNeeded == false;
+				const FLinearColor DebugTint = bWasCachingNeeded ? FLinearColor::Red : FLinearColor::Green;
+
+				FGeometry ScaledOutline = AllottedGeometry.MakeChild(FVector2D(0, 0), AllottedGeometry.GetLocalSize() * AllottedGeometry.Scale, Inverse(AllottedGeometry.Scale));
+
+				FSlateDrawElement::MakeBox(
+					OutDrawElements,
+					++OutMaxChildLayer,
+					ScaledOutline.ToPaintGeometry(),
+					FCoreStyle::Get().GetBrush(TEXT("Debug.Border")),
+					MyClippingRect,
+					ESlateDrawEffect::None,
+					DebugTint
+					);
+			}
+
+			// Draw a yellow outline around any volatile elements.
+			const TArray< TSharedRef<FSlateWindowElementList::FVolatilePaint> >& VolatileElements = CachedWindowElements->GetVolatileElements();
+			for ( const TSharedRef<FSlateWindowElementList::FVolatilePaint>& VolatileElement : VolatileElements )
+			{
+				FSlateDrawElement::MakeBox(
+					OutDrawElements,
+					++OutMaxChildLayer,
+					VolatileElement->GetGeometry().ToPaintGeometry(),
+					FCoreStyle::Get().GetBrush(TEXT("FocusRectangle")),
+					MyClippingRect,
+					ESlateDrawEffect::None,
+					FLinearColor::Yellow
+				);
+			}
+
+			// Draw a white flash for any widget that invalidated us this frame.
+			for ( TWeakPtr<SWidget> Invalidator : InvalidatorWidgets )
+			{
+				TSharedPtr<SWidget> SafeInvalidator = Invalidator.Pin();
+				if ( SafeInvalidator.IsValid() )
+				{
+					FWidgetPath WidgetPath;
+					if ( FSlateApplication::Get().GeneratePathToWidgetUnchecked(SafeInvalidator.ToSharedRef(), WidgetPath, EVisibility::All) )
+					{
+						FArrangedWidget ArrangedWidget = WidgetPath.FindArrangedWidget(SafeInvalidator.ToSharedRef()).Get(FArrangedWidget::NullWidget);
+						ArrangedWidget.Geometry.AppendTransform( FSlateLayoutTransform(Inverse(Args.GetWindowToDesktopTransform())) );
+
+						FSlateDrawElement::MakeBox(
+							OutDrawElements,
+							++OutMaxChildLayer,
+							ArrangedWidget.Geometry.ToPaintGeometry(),
+							FCoreStyle::Get().GetBrush(TEXT("WhiteBrush")),
+							MyClippingRect,
+							ESlateDrawEffect::None,
+							FLinearColor::White.CopyWithNewOpacity(0.6f)
+						);
+					}
+				}
+			}
+
+			InvalidatorWidgets.Reset();
+		}
+
+#endif
+
+		return OutMaxChildLayer;
+	}
+	else
+	{
+		return SCompoundWidget::OnPaint(Args, AllottedGeometry, MyClippingRect, OutDrawElements, LayerId, InWidgetStyle, bParentEnabled);
+	}
+}
+
+void SInvalidationPanel::SetContent(const TSharedRef< SWidget >& InContent)
+{
+	Invalidate();
+
+	ChildSlot
+	[
+		InContent
+	];
+}
+
+void SInvalidationPanel::CacheDesiredSize(float LayoutScaleMultiplier)
+{
+	SCompoundWidget::CacheDesiredSize(LayoutScaleMultiplier);
+}
+
+void SInvalidationPanel::Invalidate()
+{
+	bNeedsCaching = true;
+}
+
+bool SInvalidationPanel::ComputeVolatility() const
+{
+	return true;
+}
diff --git a/Engine/Source/Runtime/Slate/Private/Widgets/Text/STextBlock.cpp b/Engine/Source/Runtime/Slate/Private/Widgets/Text/STextBlock.cpp
--- a/Engine/Source/Runtime/Slate/Private/Widgets/Text/STextBlock.cpp
+++ b/Engine/Source/Runtime/Slate/Private/Widgets/Text/STextBlock.cpp
@@ -76,21 +76,29 @@ void STextBlock::SetText( const TAttribute< FString >& InText )
 	};
 
 	BoundText = TAttribute< FText >::Create(TAttribute<FText>::FGetter::CreateStatic( &Local::PassThroughAttribute, InText) );
+
+	InvalidateLayout();
 }
 
 void STextBlock::SetText( const FString& InText )
 {
 	BoundText = FText::FromString( InText );
+	InvalidateLayout();
 }
 
 void STextBlock::SetText( const TAttribute< FText >& InText )
 {
 	BoundText = InText;
+	InvalidateLayout();
 }
 
 void STextBlock::SetText( const FText& InText )
 {
-	BoundText = InText;
+	if ( BoundText.IsBound() || BoundText.Get().ToString() != InText.ToString() )
+	{
+		BoundText = InText;
+		InvalidateLayout();
+	}
 }
 
 int32 STextBlock::OnPaint( const FPaintArgs& Args, const FGeometry& AllottedGeometry, const FSlateRect& MyClippingRect, FSlateWindowElementList& OutDrawElements, int32 LayerId, const FWidgetStyle& InWidgetStyle, bool bParentEnabled ) const
@@ -183,59 +191,78 @@ FVector2D STextBlock::ComputeDesiredSize(float LayoutScaleMultiplier) const
 	return FVector2D(FMath::Max(MinDesiredWidth.Get(0.0f), TextSize.X), TextSize.Y);
 }
 
+bool STextBlock::ComputeVolatility() const
+{
+	return SLeafWidget::ComputeVolatility() || BoundText.IsBound();
+}
+
 void STextBlock::SetFont(const TAttribute< FSlateFontInfo >& InFont)
 {
 	Font = InFont;
+	InvalidateLayout();
 }
 
 void STextBlock::SetColorAndOpacity(const TAttribute<FSlateColor>& InColorAndOpacity)
 {
-	ColorAndOpacity = InColorAndOpacity;
+	if ( !ColorAndOpacity.IdenticalTo(InColorAndOpacity) )
+	{
+		ColorAndOpacity = InColorAndOpacity;
+		InvalidateLayout();
+	}
 }
 
 void STextBlock::SetTextStyle(const FTextBlockStyle* InTextStyle)
 {
 	TextStyle = InTextStyle;
+	InvalidateLayout();
 }
 
 void STextBlock::SetWrapTextAt(const TAttribute<float>& InWrapTextAt)
 {
 	WrapTextAt = InWrapTextAt;
+	InvalidateLayout();
 }
 
 void STextBlock::SetAutoWrapText(const TAttribute<bool>& InAutoWrapText)
 {
 	AutoWrapText = InAutoWrapText;
+	InvalidateLayout();
 }
 
 void STextBlock::SetShadowOffset(const TAttribute<FVector2D>& InShadowOffset)
 {
 	ShadowOffset = InShadowOffset;
+	InvalidateLayout();
 }
 
 void STextBlock::SetShadowColorAndOpacity(const TAttribute<FLinearColor>& InShadowColorAndOpacity)
 {
 	ShadowColorAndOpacity = InShadowColorAndOpacity;
+	InvalidateLayout();
 }
 
 void STextBlock::SetMinDesiredWidth(const TAttribute<float>& InMinDesiredWidth)
 {
 	MinDesiredWidth = InMinDesiredWidth;
+	InvalidateLayout();
 }
 
 void STextBlock::SetLineHeightPercentage(const TAttribute<float>& InLineHeightPercentage)
 {
 	LineHeightPercentage = InLineHeightPercentage;
+	InvalidateLayout();
 }
 
 void STextBlock::SetMargin(const TAttribute<FMargin>& InMargin)
 {
 	Margin = InMargin;
+	InvalidateLayout();
 }
 
 void STextBlock::SetJustification(const TAttribute<ETextJustify::Type>& InJustification)
 {
 	Justification = InJustification;
+	InvalidateLayout();
 }
 
 FTextBlockStyle STextBlock::GetComputedTextStyle() const
@@ -249,4 +276,3 @@ FTextBlockStyle STextBlock::GetComputedTextStyle() const
 	ComputedStyle.SetHighlightShape( *GetHighlightShape() );
 	return ComputedStyle;
 }
-
diff --git a/Engine/Source/Runtime/Slate/Public/Slate.h b/Engine/Source/Runtime/Slate/Public/Slate.h
--- a/Engine/Source/Runtime/Slate/Public/Slate.h
+++ b/Engine/Source/Runtime/Slate/Public/Slate.h
@@ -4,4 +4,6 @@
 
 EMIT_DEPRECATED_WARNING_MESSAGE("including Slate.h is deprecated. Please include SlateBasics.h and then individual widget headers instead.");
 #include "SlateBasics.h"
-#include "SlateExtras.h"
\ No newline at end of file
+#include "SlateExtras.h"
+
+#define SLATE_HAS_WIDGET_REFLECTOR !UE_BUILD_SHIPPING || PLATFORM_DESKTOP
diff --git a/Engine/Source/Runtime/Slate/Public/Widgets/Notifications/SProgressBar.h b/Engine/Source/Runtime/Slate/Public/Widgets/Notifications/SProgressBar.h
--- a/Engine/Source/Runtime/Slate/Public/Widgets/Notifications/SProgressBar.h
+++ b/Engine/Source/Runtime/Slate/Public/Widgets/Notifications/SProgressBar.h
@@ -79,8 +79,8 @@ public:
 	void Construct(const FArguments& InArgs);
 
 	virtual int32 OnPaint( const FPaintArgs& Args, const FGeometry& AllottedGeometry, const FSlateRect& MyClippingRect, FSlateWindowElementList& OutDrawElements, int32 LayerId, const FWidgetStyle& InWidgetStyle, bool bParentEnabled ) const override;
-
 	virtual FVector2D ComputeDesiredSize(float) const override;
+	virtual bool ComputeVolatility() const override;
 
 	/** See attribute Percent */
 	void SetPercent(TAttribute< TOptional<float> > InPercent);
diff --git a/Engine/Source/Runtime/Slate/Public/Widgets/SInvalidationPanel.h b/Engine/Source/Runtime/Slate/Public/Widgets/SInvalidationPanel.h
new file mode 100644
--- /dev/null
+++ b/Engine/Source/Runtime/Slate/Public/Widgets/SInvalidationPanel.h
@@ -0,0 +1,62 @@
+// Copyright 1998-2015 Epic Games, Inc. All Rights Reserved.
+
+#pragma once
+
+class SLATE_API SInvalidationPanel : public SCompoundWidget, public ILayoutCache
+{
+public:
+	SLATE_BEGIN_ARGS( SInvalidationPanel )
+	{
+		_Visibility = EVisibility::HitTestInvisible;
+	}
+		SLATE_DEFAULT_SLOT( FArguments, Content )
+	SLATE_END_ARGS()
+
+	void Construct( const FArguments& InArgs );
+
+	bool GetCanCache() const;
+	void SetCanCache(bool InCanCache);
+
+	void Invalidate();
+
+	// ILayoutCache overrides
+	virtual void InvalidateWidget(SWidget* InvalidateWidget) override;
+	// End ILayoutCache
+
+public:
+
+	// SWidget overrides
+	virtual void OnArrangeChildren( const FGeometry& AllottedGeometry, FArrangedChildren& ArrangedChildren ) const override;
+	virtual int32 OnPaint( const FPaintArgs& Args, const FGeometry& AllottedGeometry, const FSlateRect& MyClippingRect, FSlateWindowElementList& OutDrawElements, int32 LayerId, const FWidgetStyle& InWidgetStyle, bool bParentEnabled ) const override;
+	virtual FChildren* GetChildren() override;
+	virtual void Tick( const FGeometry& AllottedGeometry, const double InCurrentTime, const float InDeltaTime );
+	virtual bool ComputeVolatility() const override;
+	// End SWidget
+
+	void SetContent(const TSharedRef< SWidget >& InContent);
+
+#if !UE_BUILD_SHIPPING
+	static bool IsInvalidationDebuggingEnabled();
+	static void EnableInvalidationDebugging(bool bEnable);
+#endif
+
+protected:
+	void CacheDesiredSize(float LayoutScaleMultiplier) override;
+
+private:
+	FGeometry LastAllottedGeometry;
+
+	FSimpleSlot EmptyChildSlot;
+	FVector2D CachedDesiredSize;
+
+#if !UE_BUILD_SHIPPING
+	mutable TSet<TWeakPtr<SWidget>> InvalidatorWidgets;
+#endif
+
+	mutable TSharedPtr< FSlateWindowElementList > CachedWindowElements;
+
+	mutable int32 CachedMaxChildLayer;
+	mutable bool bNeedsCaching;
+	mutable bool bIsInvalidating;
+	bool bCanCache;
+};
diff --git a/Engine/Source/Runtime/Slate/Public/Widgets/Text/STextBlock.h b/Engine/Source/Runtime/Slate/Public/Widgets/Text/STextBlock.h
--- a/Engine/Source/Runtime/Slate/Public/Widgets/Text/STextBlock.h
+++ b/Engine/Source/Runtime/Slate/Public/Widgets/Text/STextBlock.h
@@ -172,6 +172,11 @@ public:
 	virtual FVector2D ComputeDesiredSize(float) const override;
 	// End of SWidget interface
 
+protected:
+	// SWidget interface
+	virtual bool ComputeVolatility() const override;
+	// End of SWidget interface
+
 private:
 	/** Get the computed text style to use with the text marshaller */
 	FTextBlockStyle GetComputedTextStyle() const;
diff --git a/Engine/Source/Runtime/SlateCore/Private/Rendering/DrawElements.cpp b/Engine/Source/Runtime/SlateCore/Private/Rendering/DrawElements.cpp
--- a/Engine/Source/Runtime/SlateCore/Private/Rendering/DrawElements.cpp
+++ b/Engine/Source/Runtime/SlateCore/Private/Rendering/DrawElements.cpp
@@ -2,6 +2,8 @@
 
 #include "SlateCorePrivatePCH.h"
 
+DECLARE_CYCLE_STAT(TEXT("FSlateDrawElement::Make Time"), STAT_SlateDrawElementMakeTime, STATGROUP_Slate);
+
 void FSlateDrawElement::Init(uint32 InLayer, const FPaintGeometry& PaintGeometry, const FSlateRect& InClippingRect, ESlateDrawEffect::Type InDrawEffects)
 {
 	RenderTransform = PaintGeometry.GetAccumulatedRenderTransform();
@@ -17,6 +19,7 @@ void FSlateDrawElement::Init(uint32 InLayer, const FPaintGeometry& PaintGeometry
 
 void FSlateDrawElement::MakeDebugQuad( FSlateWindowElementList& ElementList, uint32 InLayer, const FPaintGeometry& PaintGeometry, const FSlateRect& InClippingRect)
 {
+	SCOPE_CYCLE_COUNTER( STAT_SlateDrawElementMakeTime )
 	PaintGeometry.CommitTransformsIfUsingLegacyConstructor();
 	FSlateDrawElement& DrawElt = ElementList.AddUninitialized();
 	DrawElt.Init(InLayer, PaintGeometry, InClippingRect, ESlateDrawEffect::None);
@@ -33,6 +36,7 @@ void FSlateDrawElement::MakeBox(
 	ESlateDrawEffect::Type InDrawEffects, 
 	const FLinearColor& InTint )
 {
+	SCOPE_CYCLE_COUNTER( STAT_SlateDrawElementMakeTime )
 	PaintGeometry.CommitTransformsIfUsingLegacyConstructor();
 	FSlateDrawElement& DrawElt = ElementList.AddUninitialized();
 	DrawElt.Init(InLayer, PaintGeometry, InClippingRect, InDrawEffects);
@@ -52,6 +56,7 @@ void FSlateDrawElement::MakeRotatedBox(
 	ERotationSpace RotationSpace,
 	const FLinearColor& InTint )
 {
+	SCOPE_CYCLE_COUNTER( STAT_SlateDrawElementMakeTime )
 	PaintGeometry.CommitTransformsIfUsingLegacyConstructor();
 	FSlateDrawElement& DrawElt = ElementList.AddUninitialized();
 	DrawElt.Init(InLayer, PaintGeometry, InClippingRect, InDrawEffects);
@@ -64,6 +69,7 @@ void FSlateDrawElement::MakeRotatedBox(
 
 void FSlateDrawElement::MakeText( FSlateWindowElementList& ElementList, uint32 InLayer, const FPaintGeometry& PaintGeometry, const FString& InText, const int32 StartIndex, const int32 EndIndex, const FSlateFontInfo& InFontInfo, const FSlateRect& InClippingRect,ESlateDrawEffect::Type InDrawEffects, const FLinearColor& InTint )
 {
+	SCOPE_CYCLE_COUNTER( STAT_SlateDrawElementMakeTime )
 	PaintGeometry.CommitTransformsIfUsingLegacyConstructor();
 	FSlateDrawElement& DrawElt = ElementList.AddUninitialized();
 	DrawElt.Init(InLayer, PaintGeometry, InClippingRect, InDrawEffects);
@@ -74,6 +80,7 @@ void FSlateDrawElement::MakeText( FSlateWindowElementList& ElementList, uint32 I
 
 void FSlateDrawElement::MakeText( FSlateWindowElementList& ElementList, uint32 InLayer, const FPaintGeometry& PaintGeometry, const FString& InText, const FSlateFontInfo& InFontInfo, const FSlateRect& InClippingRect,ESlateDrawEffect::Type InDrawEffects, const FLinearColor& InTint )
 {
+	SCOPE_CYCLE_COUNTER( STAT_SlateDrawElementMakeTime )
 	PaintGeometry.CommitTransformsIfUsingLegacyConstructor();
 	FSlateDrawElement& DrawElt = ElementList.AddUninitialized();
 	DrawElt.Init(InLayer, PaintGeometry, InClippingRect, InDrawEffects);
@@ -84,6 +91,7 @@ void FSlateDrawElement::MakeText( FSlateWindowElementList& ElementList, uint32 I
 
 void FSlateDrawElement::MakeText( FSlateWindowElementList& ElementList, uint32 InLayer, const FPaintGeometry& PaintGeometry, const FText& InText, const FSlateFontInfo& InFontInfo, const FSlateRect& InClippingRect,ESlateDrawEffect::Type InDrawEffects, const FLinearColor& InTint )
 {
+	SCOPE_CYCLE_COUNTER( STAT_SlateDrawElementMakeTime )
 	PaintGeometry.CommitTransformsIfUsingLegacyConstructor();
 	FSlateDrawElement& DrawElt = ElementList.AddUninitialized();
 	DrawElt.Init(InLayer, PaintGeometry, InClippingRect, InDrawEffects);
@@ -93,6 +101,7 @@ void FSlateDrawElement::MakeText( FSlateWindowElementList& ElementList, uint32 I
 
 void FSlateDrawElement::MakeGradient( FSlateWindowElementList& ElementList, uint32 InLayer, const FPaintGeometry& PaintGeometry, TArray<FSlateGradientStop> InGradientStops, EOrientation InGradientType, const FSlateRect& InClippingRect, ESlateDrawEffect::Type InDrawEffects, bool bGammaCorrect )
 {
+	SCOPE_CYCLE_COUNTER( STAT_SlateDrawElementMakeTime )
 	PaintGeometry.CommitTransformsIfUsingLegacyConstructor();
 	FSlateDrawElement& DrawElt = ElementList.AddUninitialized();
 	DrawElt.Init(InLayer, PaintGeometry, InClippingRect, InDrawEffects);
@@ -103,6 +112,7 @@ void FSlateDrawElement::MakeGradient( FSlateWindowElementList& ElementList, uint
 
 void FSlateDrawElement::MakeSpline( FSlateWindowElementList& ElementList, uint32 InLayer, const FPaintGeometry& PaintGeometry, const FVector2D& InStart, const FVector2D& InStartDir, const FVector2D& InEnd, const FVector2D& InEndDir, const FSlateRect InClippingRect, float InThickness, ESlateDrawEffect::Type InDrawEffects, const FLinearColor& InTint )
 {
+	SCOPE_CYCLE_COUNTER( STAT_SlateDrawElementMakeTime )
 	PaintGeometry.CommitTransformsIfUsingLegacyConstructor();
 	FSlateDrawElement& DrawElt = ElementList.AddUninitialized();
 	DrawElt.Init(InLayer, PaintGeometry, InClippingRect, InDrawEffects);
@@ -119,6 +129,7 @@ void FSlateDrawElement::MakeDrawSpaceSpline( FSlateWindowElementList& ElementLis
 
 void FSlateDrawElement::MakeLines( FSlateWindowElementList& ElementList, uint32 InLayer, const FPaintGeometry& PaintGeometry, const TArray<FVector2D>& Points, const FSlateRect InClippingRect, ESlateDrawEffect::Type InDrawEffects, const FLinearColor& InTint, bool bAntialias )
 {
+	SCOPE_CYCLE_COUNTER( STAT_SlateDrawElementMakeTime )
 	PaintGeometry.CommitTransformsIfUsingLegacyConstructor();
 	FSlateDrawElement& DrawElt = ElementList.AddUninitialized();
 	DrawElt.Init(InLayer, PaintGeometry, InClippingRect, InDrawEffects);
@@ -129,6 +140,7 @@ void FSlateDrawElement::MakeLines( FSlateWindowElementList& ElementList, uint32
 
 void FSlateDrawElement::MakeViewport( FSlateWindowElementList& ElementList, uint32 InLayer, const FPaintGeometry& PaintGeometry, TSharedPtr<const ISlateViewport> Viewport, const FSlateRect& InClippingRect, bool bGammaCorrect, bool bAllowBlending, ESlateDrawEffect::Type InDrawEffects, const FLinearColor& InTint )
 {
+	SCOPE_CYCLE_COUNTER( STAT_SlateDrawElementMakeTime )
 	PaintGeometry.CommitTransformsIfUsingLegacyConstructor();
 	FSlateDrawElement& DrawElt = ElementList.AddUninitialized();
 	DrawElt.Init(InLayer, PaintGeometry, InClippingRect, InDrawEffects);
@@ -139,6 +151,7 @@ void FSlateDrawElement::MakeViewport( FSlateWindowElementList& ElementList, uint
 
 void FSlateDrawElement::MakeCustom( FSlateWindowElementList& ElementList, uint32 InLayer, TSharedPtr<ICustomSlateElement, ESPMode::ThreadSafe> CustomDrawer )
 {
+	SCOPE_CYCLE_COUNTER( STAT_SlateDrawElementMakeTime )
 	FSlateDrawElement& DrawElt = ElementList.AddUninitialized();
 	DrawElt.Init(InLayer, FPaintGeometry(), FSlateRect(1,1,1,1), ESlateDrawEffect::None);
 	DrawElt.RenderTransform = FSlateRenderTransform();
@@ -335,15 +348,69 @@ void FSlateWindowElementList::QueueDeferredPainting( const FDeferredPaint& InDef
 	DeferredPaintList.Add( MakeShareable( new FDeferredPaint( InDeferredPaint ) ) );
 }
 
-
-int32 FSlateWindowElementList::PaintDeferred( int32 LayerId )
+int32 FSlateWindowElementList::PaintDeferred(int32 LayerId)
 {
 	for ( int32 i = 0; i < DeferredPaintList.Num(); ++i )
 	{
-		const TSharedRef< FDeferredPaint >& Args = DeferredPaintList[ i ];
+		const TSharedRef< FDeferredPaint >& Args = DeferredPaintList[i];
+
+		LayerId = Args->ExecutePaint(LayerId, *this);
+	}
+
+	return LayerId;
+}
+
+
+FSlateWindowElementList::FVolatilePaint::FVolatilePaint(const TSharedRef<const SWidget>& InWidgetToPaint, const FPaintArgs& InArgs, const FGeometry InAllottedGeometry, const FSlateRect InMyClippingRect, int32 InLayerId, const FWidgetStyle& InWidgetStyle, bool InParentEnabled)
+	: WidgetToPaintPtr(InWidgetToPaint)
+	, Args(InArgs.EnableCaching(InArgs.GetLayoutCache().Pin(), false, true))
+	, AllottedGeometry(InAllottedGeometry)
+	, MyClippingRect(InMyClippingRect)
+	, LayerId(InLayerId)
+	, WidgetStyle(InWidgetStyle)
+	, bParentEnabled(InParentEnabled)
+{
+}
+
+int32 FSlateWindowElementList::FVolatilePaint::ExecutePaint(FSlateWindowElementList& OutDrawElements) const
+{
+	TSharedPtr<const SWidget> WidgetToPaint = WidgetToPaintPtr.Pin();
+	if ( WidgetToPaint.IsValid() )
+	{
+		// Have to run a slate pre-pass for all volatile elements, some widgets cache information like 
+		// the STextBlock.  This may be all kinds of terrible an idea to do during paint.
+		SWidget* MutableWidget = const_cast<SWidget*>( WidgetToPaint.Get() );
+		MutableWidget->SlatePrepass(AllottedGeometry.Scale);
 
-		LayerId = Args->ExecutePaint( LayerId, *this );
+		return WidgetToPaint->Paint(Args, AllottedGeometry, MyClippingRect, OutDrawElements, LayerId, WidgetStyle, bParentEnabled);
 	}
 
 	return LayerId;
 }
+
+void FSlateWindowElementList::QueueVolatilePainting(const FVolatilePaint& InDeferredPaint)
+{
+	VolatilePaintList.Add(MakeShareable(new FVolatilePaint(InDeferredPaint)));
+}
+
+int32 FSlateWindowElementList::PaintVolatile(FSlateWindowElementList& OutElementList)
+{
+	int32 MaxLayerId = 0;
+
+	for ( int32 VolatileIndex = 0; VolatileIndex < VolatilePaintList.Num(); ++VolatileIndex )
+	{
+		const TSharedRef< FVolatilePaint >& Args = VolatilePaintList[VolatileIndex];
+
+		MaxLayerId = FMath::Max(MaxLayerId, Args->ExecutePaint(OutElementList));
+	}
+
+	return MaxLayerId;
+}
+
+void FSlateWindowElementList::Reset()
+{
+	BatchData.Reset();
+	DrawElements.Reset();
+	DeferredPaintList.Reset();
+	VolatilePaintList.Reset();
+}
diff --git a/Engine/Source/Runtime/SlateCore/Private/Textures/TextureAtlas.cpp b/Engine/Source/Runtime/SlateCore/Private/Textures/TextureAtlas.cpp
--- a/Engine/Source/Runtime/SlateCore/Private/Textures/TextureAtlas.cpp
+++ b/Engine/Source/Runtime/SlateCore/Private/Textures/TextureAtlas.cpp
@@ -56,7 +56,10 @@ void FSlateTextureAtlas::MarkTextureDirty()
 {
 	check( 
 		(GSlateLoadingThreadId != 0) || 
-		(AtlasOwnerThread == ESlateTextureAtlasOwnerThread::Game && IsInGameThread()) || 
+		// Note: For Game thread ownership, there is a point at which multiple worker threads operate on text simulatenously while the game thread is blocked
+		// Access to the font cache is controlled through mutexes so we simply need to check that we are not accessing it from the render thread
+		// Game thread access is also allowed when the game thread and render thread are the same
+		(AtlasOwnerThread == ESlateTextureAtlasOwnerThread::Game && !IsInActualRenderingThread()) || 
 		(AtlasOwnerThread == ESlateTextureAtlasOwnerThread::Render && IsInRenderingThread()) 
 		);
 
diff --git a/Engine/Source/Runtime/SlateCore/Private/Types/PaintArgs.cpp b/Engine/Source/Runtime/SlateCore/Private/Types/PaintArgs.cpp
--- a/Engine/Source/Runtime/SlateCore/Private/Types/PaintArgs.cpp
+++ b/Engine/Source/Runtime/SlateCore/Private/Types/PaintArgs.cpp
@@ -11,30 +11,48 @@ FPaintArgs::FPaintArgs( const TSharedRef<SWidget>& Parent, FHittestGrid& InHitte
 , WindowOffset(InWindowOffset)
 , CurrentTime(InCurrentTime)
 , DeltaTime(InDeltaTime)
+, bIsCaching(false)
 {
 }
 
-
 FPaintArgs FPaintArgs::WithNewParent( const SWidget* Parent ) const
 {
 	FPaintArgs Args = FPaintArgs( const_cast<SWidget*>(Parent)->AsShared(), this->Grid, this->WindowOffset, this->CurrentTime, this->DeltaTime );
 	Args.LastHittestIndex = this->LastHittestIndex;
 	Args.LastRecordedVisibility = this->LastRecordedVisibility;
+	Args.LayoutCache = this->LayoutCache;
+	Args.bIsCaching = this->bIsCaching;
+	Args.bIsVolatilityPass = this->bIsVolatilityPass;
+
 	return Args;
 }
 
+FPaintArgs FPaintArgs::EnableCaching(const TSharedPtr<ILayoutCache>& InLayoutCache, bool bEnableCaching, bool bEnableVolatile) const
+{
+	FPaintArgs UpdatedArgs(*this);
+	UpdatedArgs.LayoutCache = InLayoutCache;
+	UpdatedArgs.bIsCaching = bEnableCaching;
+	UpdatedArgs.bIsVolatilityPass = bEnableVolatile;
+
+	return UpdatedArgs;
+}
 
 FPaintArgs FPaintArgs::RecordHittestGeometry(const SWidget* Widget, const FGeometry& WidgetGeometry, const FSlateRect& InClippingRect) const
 {
-	const EVisibility RecordedVisibility = (LastRecordedVisibility.AreChildrenHitTestVisible())
-		? Widget->GetVisibility()
-		: LastRecordedVisibility;
+	FPaintArgs UpdatedArgs(*this);
 
-	const int32 RecordedHittestIndex = Grid.InsertWidget( LastHittestIndex, RecordedVisibility, FArrangedWidget(const_cast<SWidget*>(Widget)->AsShared(), WidgetGeometry), WindowOffset, InClippingRect );
+	if ( LastRecordedVisibility.AreChildrenHitTestVisible() )
+	{
+		const EVisibility RecordedVisibility = Widget->GetVisibility();
+		const int32 RecordedHittestIndex = Grid.InsertWidget(LastHittestIndex, RecordedVisibility, FArrangedWidget(const_cast<SWidget*>( Widget )->AsShared(), WidgetGeometry), WindowOffset, InClippingRect);
+		UpdatedArgs.LastHittestIndex = RecordedHittestIndex;
+		UpdatedArgs.LastRecordedVisibility = RecordedVisibility;
+	}
+	else
+	{
+		UpdatedArgs.LastRecordedVisibility = LastRecordedVisibility;
+	}
 
-	FPaintArgs UpdatedArgs(*this);
-	UpdatedArgs.LastHittestIndex = RecordedHittestIndex;
-	UpdatedArgs.LastRecordedVisibility = RecordedVisibility;
 	return UpdatedArgs;
 }
 
@@ -42,4 +60,4 @@ FPaintArgs FPaintArgs::InsertCustomHitTestPath( TSharedRef<ICustomHitTestPath> C
 {
 	const_cast<FHittestGrid&>(Grid).InsertCustomHitTestPath( CustomHitTestPath, InLastHittestIndex );
 	return *this;
-}
\ No newline at end of file
+}
diff --git a/Engine/Source/Runtime/SlateCore/Private/Widgets/SWidget.cpp b/Engine/Source/Runtime/SlateCore/Private/Widgets/SWidget.cpp
--- a/Engine/Source/Runtime/SlateCore/Private/Widgets/SWidget.cpp
+++ b/Engine/Source/Runtime/SlateCore/Private/Widgets/SWidget.cpp
@@ -23,10 +23,13 @@ SWidget::SWidget()
 	, Visibility( EVisibility::Visible )
 	, RenderTransform( )
 	, RenderTransformPivot( FVector2D::ZeroVector )
-	, bIsHovered(false)
 	, DesiredSize(FVector2D::ZeroVector)
 	, ToolTip()
+	, bIsHovered(false)
 	, bToolTipForceFieldEnabled( false )
+	, bForceVolatile(false)
+	, bCachedVolatile(false)
+	, bInheritedVolatility(false)
 {
 
 }
@@ -366,7 +369,6 @@ void SWidget::TickWidgetsRecursively( const FGeometry& AllottedGeometry, const d
 	}
 }
 
-
 void SWidget::SlatePrepass()
 {
 	SlatePrepass( FSlateApplicationBase::Get().GetApplicationScale() );
@@ -379,12 +381,13 @@ void SWidget::SlatePrepass(float LayoutScaleMultiplier)
 	// a function of its children's sizes.
 	FChildren* MyChildren = this->GetChildren();
 	int32 NumChildren = MyChildren->Num();
-	for(int32 ChildIndex=0; ChildIndex < NumChildren; ++ChildIndex)
+	for ( int32 ChildIndex=0; ChildIndex < NumChildren; ++ChildIndex )
 	{
 		const TSharedRef<SWidget>& Child = MyChildren->GetChildAt(ChildIndex);
+
 		if ( Child->Visibility.Get() != EVisibility::Collapsed )
 		{
-			const float ChildLayoutScaleMultiplier = GetRelativeLayoutScale( MyChildren->GetSlotAt(ChildIndex) );
+			const float ChildLayoutScaleMultiplier = GetRelativeLayoutScale(MyChildren->GetSlotAt(ChildIndex));
 			// Recur: Descend down the widget tree.
 			Child->SlatePrepass(LayoutScaleMultiplier*ChildLayoutScaleMultiplier);
 		}
@@ -394,19 +397,34 @@ void SWidget::SlatePrepass(float LayoutScaleMultiplier)
 	CacheDesiredSize(LayoutScaleMultiplier);
 }
 
-
 void SWidget::CacheDesiredSize(float LayoutScaleMultiplier)
 {
 	// Cache this widget's desired size.
 	this->Advanced_SetDesiredSize(this->ComputeDesiredSize(LayoutScaleMultiplier));
 }
 
-
 const FVector2D& SWidget::GetDesiredSize() const
 {
 	return DesiredSize;
 }
 
+void SWidget::CachePrepass(TWeakPtr<ILayoutCache> InLayoutCache)
+{
+	FChildren* MyChildren = this->GetChildren();
+	int32 NumChildren = MyChildren->Num();
+	for ( int32 ChildIndex=0; ChildIndex < NumChildren; ++ChildIndex )
+	{
+		const TSharedRef<SWidget>& Child = MyChildren->GetChildAt(ChildIndex);
+		if ( Child->GetVisibility().IsVisible() == false )
+		{
+			Child->LayoutCache = InLayoutCache;
+		}
+		else
+		{
+			Child->CachePrepass(InLayoutCache);
+		}
+	}
+}
 
 bool SWidget::SupportsKeyboardFocus() const
 {
@@ -643,6 +661,29 @@ int32 SWidget::Paint(const FPaintArgs& Args, const FGeometry& AllottedGeometry,
 	INC_DWORD_STAT(STAT_SlateNumPaintedWidgets);
 	SLATE_CYCLE_COUNTER_SCOPE_CUSTOM_DETAILED(SLATE_STATS_DETAIL_LEVEL_MED, GSlateOnPaint, GetType());
 
+	// Save the current layout cache we're associated with (if any)
+	LayoutCache = Args.GetLayoutCache();
+
+	// Record if we're part of a volatility pass, this is critical for ensuring we don't report a child
+	// of a volatile widget as non-volatile, causing the invalidation panel to do work that's not required.
+	bInheritedVolatility = Args.IsVolatilityPass();
+
+	// If this paint pass is to cache off our geometry, but we're a volatile widget,
+	// record this widget as volatile in the draw elements so that we get our own tick/paint 
+	// pass later when the layout cache draws.
+	if ( Args.IsCaching() && IsVolatile() )
+	{
+		// Volatile widgets don't have an associated layout cache.
+		LayoutCache.Reset();
+
+		OutDrawElements.QueueVolatilePainting(
+			FSlateWindowElementList::FVolatilePaint(SharedThis(this), Args, AllottedGeometry, MyClippingRect, LayerId, InWidgetStyle, bParentEnabled));
+
+		// Add a layer padding so that cached elements drawn on top of this volatile widget have some airspace room and not be drawn on top
+		// of when the volatile element draws.
+		return LayerId + 1000;
+	}
+
 	if ( bFoldTick )
 	{
 		FGeometry TickGeometry = AllottedGeometry;
diff --git a/Engine/Source/Runtime/SlateCore/Public/Rendering/DrawElements.h b/Engine/Source/Runtime/SlateCore/Public/Rendering/DrawElements.h
--- a/Engine/Source/Runtime/SlateCore/Public/Rendering/DrawElements.h
+++ b/Engine/Source/Runtime/SlateCore/Public/Rendering/DrawElements.h
@@ -651,19 +651,19 @@ public:
 	}
 	
 	/** @return Get the window that we will be painting */
-	const TSharedPtr<SWindow> GetWindow() const
+	FORCEINLINE const TSharedPtr<SWindow> GetWindow() const
 	{
 		return TopLevelWindow.Pin();
 	}
 	
 	/** @return Get the window that we will be painting */
-	TSharedPtr<SWindow> GetWindow()
+	FORCEINLINE TSharedPtr<SWindow> GetWindow()
 	{
 		return TopLevelWindow.Pin();
 	}
 	
 	/** @return Get the draw elements that we want to render into this window */
-	const TArray<FSlateDrawElement>& GetDrawElements() const
+	FORCEINLINE const TArray<FSlateDrawElement>& GetDrawElements() const
 	{
 		return DrawElements;
 	}
@@ -673,21 +673,32 @@ public:
 	 *
 	 * @param InDrawElement  The draw element to add
 	 */
-	void AddItem( const FSlateDrawElement& InDrawElement )
+	FORCEINLINE void AddItem(const FSlateDrawElement& InDrawElement)
 	{
 		DrawElements.Add( InDrawElement );
 	}
 
-	FSlateDrawElement& AddUninitialized()
+	/**
+	 * Creates an uninitialized draw element
+	 */
+	FORCEINLINE FSlateDrawElement& AddUninitialized()
+	{
+		const int32 InsertIdx = DrawElements.AddDefaulted();
+		return DrawElements[InsertIdx];
+	}
+
+	/**
+	 * Append draw elements to the list of draw elements
+	 */
+	FORCEINLINE void AppendDrawElements(const TArray<FSlateDrawElement>& InDrawElements)
 	{
-		const int32 InsertIdx = DrawElements.AddUninitialized();
-		return *(new(DrawElements.GetData() + InsertIdx) FSlateDrawElement());
+		DrawElements.Append(InDrawElements);
 	}
 
 	/**
-	* Some widgets may want to paint their children after after another, loosely-related widget finished painting.
-	* Or they may want to paint "after everyone".
-	*/
+	 * Some widgets may want to paint their children after after another, loosely-related widget finished painting.
+	 * Or they may want to paint "after everyone".
+	 */
 	struct FDeferredPaint
 	{
 	public:
@@ -706,17 +717,37 @@ public:
 
 	SLATECORE_API void QueueDeferredPainting( const FDeferredPaint& InDeferredPaint );
 
-	int32 PaintDeferred( int32 LayerId );
+	int32 PaintDeferred(int32 LayerId);
+
+	struct FVolatilePaint
+	{
+	public:
+		SLATECORE_API FVolatilePaint(const TSharedRef<const SWidget>& InWidgetToPaint, const FPaintArgs& InArgs, const FGeometry InAllottedGeometry, const FSlateRect InMyClippingRect, int32 LayerId, const FWidgetStyle& InWidgetStyle, bool InParentEnabled);
+
+		int32 ExecutePaint( FSlateWindowElementList& OutDrawElements ) const;
+
+		FORCEINLINE FGeometry GetGeometry() const { return AllottedGeometry; }
+	 
+	private:
+		const TWeakPtr<const SWidget> WidgetToPaintPtr;
+		const FPaintArgs Args;
+		const FGeometry AllottedGeometry;
+		const FSlateRect MyClippingRect;
+		const int32 LayerId;
+		const FWidgetStyle WidgetStyle;
+		const bool bParentEnabled;
+	};
+
+	SLATECORE_API void QueueVolatilePainting( const FVolatilePaint& InVolatilePaint );
+
+	SLATECORE_API int32 PaintVolatile(FSlateWindowElementList& OutElementList);
+
+	SLATECORE_API const TArray< TSharedRef<FVolatilePaint> >& GetVolatileElements() const { return VolatilePaintList; }
 	
 	/**
 	 * Remove all the elements from this draw list.
 	 */
-	void Reset()
-	{
-		DrawElements.Reset();
-		BatchData.Reset();
-		DeferredPaintList.Reset();
-	}
+	SLATECORE_API void Reset();
 
 	FSlateBatchData& GetBatchData() { return BatchData; }
 
@@ -735,4 +766,7 @@ private:
 	 * We accomplish this by deferring their painting.
 	 */
 	TArray< TSharedRef<FDeferredPaint> > DeferredPaintList;
+
+	/** The widgets be cached for a later paint pass when the invalidation host paints. */
+	TArray< TSharedRef<FVolatilePaint> > VolatilePaintList;
 };
diff --git a/Engine/Source/Runtime/SlateCore/Public/Types/PaintArgs.h b/Engine/Source/Runtime/SlateCore/Public/Types/PaintArgs.h
--- a/Engine/Source/Runtime/SlateCore/Public/Types/PaintArgs.h
+++ b/Engine/Source/Runtime/SlateCore/Public/Types/PaintArgs.h
@@ -9,6 +9,7 @@ struct FGeometry;
 class FHittestGrid;
 class FSlateRect;
 class ICustomHitTestPath;
+class ILayoutCache;
 
 /**
  * SWidget::OnPaint and SWidget::Paint use FPaintArgs as their
@@ -20,12 +21,17 @@ class SLATECORE_API FPaintArgs
 public:
 	FPaintArgs( const TSharedRef<SWidget>& Parent, FHittestGrid& InHittestGrid, FVector2D InWindowOffset, double InCurrentTime, float InDeltaTime );
 	FPaintArgs WithNewParent( const SWidget* Parent ) const;
+	FPaintArgs EnableCaching(const TSharedPtr<ILayoutCache>& LayoutCache, bool bEnableCaching, bool bEnableVolatility) const;
 	FPaintArgs RecordHittestGeometry(const SWidget* Widget, const FGeometry& WidgetGeometry, const FSlateRect& InClippingRect) const;
 	FPaintArgs InsertCustomHitTestPath( TSharedRef<ICustomHitTestPath> CustomHitTestPath, int32 HitTestIndex ) const;
+
 	int32 GetLastHitTestIndex() const { return LastHittestIndex; }
 	FVector2D GetWindowToDesktopTransform() const { return WindowOffset; }
 	double GetCurrentTime() const { return CurrentTime; }
 	float GetDeltaTime() const { return DeltaTime; }
+	bool IsCaching() const { return bIsCaching; }
+	bool IsVolatilityPass() const { return bIsVolatilityPass; }
+	TWeakPtr<ILayoutCache> GetLayoutCache() const { return LayoutCache; }
 
 private:
 	const TSharedRef<SWidget>& ParentPtr;
@@ -35,4 +41,7 @@ private:
 	FVector2D WindowOffset;
 	double CurrentTime;
 	float DeltaTime;
+	bool bIsCaching;
+	bool bIsVolatilityPass;
+	TWeakPtr<ILayoutCache> LayoutCache;
 };
\ No newline at end of file
diff --git a/Engine/Source/Runtime/SlateCore/Public/Widgets/SWidget.h b/Engine/Source/Runtime/SlateCore/Public/Widgets/SWidget.h
--- a/Engine/Source/Runtime/SlateCore/Public/Widgets/SWidget.h
+++ b/Engine/Source/Runtime/SlateCore/Public/Widgets/SWidget.h
@@ -66,6 +66,14 @@ public:
 	}
 };
 
+
+class ILayoutCache
+{
+public:
+	virtual void InvalidateWidget(class SWidget* InvalidateWidget) = 0;
+};
+
+
 class IToolTip;
 
 /**
@@ -572,7 +580,6 @@ public:
 	 */
 	virtual bool SupportsKeyboardFocus() const;
 
-
 	/**
 	 * Checks to see if this widget currently has the keyboard focus
 	 *
@@ -633,6 +640,7 @@ public:
 	void SetEnabled( const TAttribute<bool>& InEnabledState )
 	{
 		EnabledState = InEnabledState;
+		InvalidateLayout();
 	}
 
 	/** @return Whether or not this widget is enabled */
@@ -682,29 +690,81 @@ public:
 	/** @param InVisibility  should this widget be */
 	virtual void SetVisibility( TAttribute<EVisibility> InVisibility )
 	{
-		Visibility = InVisibility;
+		if ( !Visibility.IdenticalTo(InVisibility) )
+		{
+			Visibility = InVisibility;
+			InvalidateLayout();
+		}
+	}
+
+	/**
+	 * When performing a caching pass, volatile widgets are not cached as part of everything
+	 * else, instead they and their children are drawn as normal standard widgets and excluded
+	 * from the cache.  This is extremely useful for things like timers and text that change
+	 * every frame.
+	 */
+	FORCEINLINE bool IsVolatile() const { return bCachedVolatile; }
+
+	/**
+	 * Was this widget painted as part of a volatile pass previously.  This may mean it was the
+	 * widget directly responsible for making a hierarchy volatile, or it may mean it was simply
+	 * a child of a volatile panel.
+	 */
+	FORCEINLINE bool IsVolatileIndirectly() const { return bInheritedVolatility; }
+
+	/**
+	 * Should this widget always appear as volatile for any layout caching host widget.  A volatile
+	 * widget's geometry and layout data will never be cached, and neither will any children.
+	 * @param bForce should we force the widget to be volatile?
+	 */
+	FORCEINLINE void ForceVolatile(bool bForce)
+	{
+		bForceVolatile = bForce;
+		InvalidateLayout();
+	}
+
+	/**
+	 * Invalidates the widget from the view of a layout caching widget that may own this widget.
+	 * will force the owning widget to redraw and cache children on the next paint pass.
+	 */
+	FORCEINLINE void InvalidateLayout()
+	{
+		const bool bWasVolatile = IsVolatileIndirectly() || IsVolatile();
+		const bool bWasDirectlyVolatile = IsVolatile();
+		CacheVolatility();
+		const bool bVolatilityChanged = bWasDirectlyVolatile != IsVolatile();
+
+		if ( bWasVolatile == false || bVolatilityChanged )
+		{
+			TSharedPtr<ILayoutCache> Cache = LayoutCache.Pin();
+			if ( Cache.IsValid() )
+			{
+				Cache->InvalidateWidget(this);
+			}
+		}
 	}
 
 	/** @return the render transform of the widget. */
-	const TOptional<FSlateRenderTransform>& GetRenderTransform() const
+	FORCEINLINE const TOptional<FSlateRenderTransform>& GetRenderTransform() const
 	{
 		return RenderTransform.Get();
 	}
 
-	/** @param InTransform the render transform to set for the widget (tranforms from widget's local space). TOptional<> to allow code to skip expensive overhead if there is no render transform applied. */
-	void SetRenderTransform( TAttribute<TOptional<FSlateRenderTransform>> InTransform )
+	/** @param InTransform the render transform to set for the widget (transforms from widget's local space). TOptional<> to allow code to skip expensive overhead if there is no render transform applied. */
+	FORCEINLINE void SetRenderTransform(TAttribute<TOptional<FSlateRenderTransform>> InTransform)
 	{
 		RenderTransform = InTransform;
+		InvalidateLayout();
 	}
 
 	/** @return the pivot point of the render transform. */
-	const FVector2D& GetRenderTransformPivot() const
+	FORCEINLINE const FVector2D& GetRenderTransformPivot() const
 	{
 		return RenderTransformPivot.Get();
 	}
 
 	/** @param InTransformPivot Sets the pivot point of the widget's render transform (in normalized local space). */
-	void SetRenderTransformPivot( TAttribute<FVector2D> InTransformPivot )
+	FORCEINLINE void SetRenderTransformPivot(TAttribute<FVector2D> InTransformPivot)
 	{
 		RenderTransformPivot = InTransformPivot;
 	}
@@ -869,6 +929,24 @@ protected:
 	/** @return a brush to draw focus, nullptr if no focus drawing is desired */
 	virtual const FSlateBrush* GetFocusBrush() const;
 
+	/**
+	 * Recalculates volatility of the widget and caches the result.  Should be called any time 
+	 * anything examined by your implementation of ComputeVolatility is changed.
+	 */
+	FORCEINLINE void CacheVolatility()
+	{
+		bCachedVolatile = bForceVolatile || ComputeVolatility();
+	}
+
+	/**
+	 * Recomputes the volatility of the widget.  If you have additional state you automatically want to make
+	 * the widget volatile, you should sample that information here.
+	 */
+	virtual bool ComputeVolatility() const
+	{
+		return Visibility.IsBound();
+	}
+
 private:
 
 	/**
@@ -899,6 +977,15 @@ private:
 	 */
 	virtual void OnArrangeChildren(const FGeometry& AllottedGeometry, FArrangedChildren& ArrangedChildren) const = 0;
 
+protected:
+
+	/**
+	 * Don't call this directly unless you're a layout cache, this is used to recursively set the layout cache on
+	 * on invisible children that never get the opportunity to paint and receive the layout cache through the normal
+	 * means.  That way if an invisible widget becomes visible, we still properly invalidate the hierarchy.
+	 */
+	void CachePrepass(TWeakPtr<ILayoutCache> LayoutCache);
+
 protected:
 	/**
 	 * Registers an "active timer" delegate that will execute at some regular interval. TickFunction will not be called until the specified interval has elapsed once.
@@ -969,9 +1056,6 @@ protected:
 	/** Render transform pivot of this widget (in normalized local space) */
 	TAttribute< FVector2D > RenderTransformPivot;
 
-	/** Is this widget hovered? */
-	bool bIsHovered;
-
 private:
 
 	/**
@@ -984,7 +1068,27 @@ private:
 	/** Tool tip content for this widget */
 	TSharedPtr<IToolTip> ToolTip;
 
-	// Whether this widget is a "tool tip force field".  That is, tool-tips should never spawn over the area
-	// occupied by this widget, and will instead be repelled to an outside edge
+	/** The current layout cache that may need to invalidated by changes to this widget. */
+	mutable TWeakPtr<ILayoutCache> LayoutCache;
+
+protected:
+	/** Is this widget hovered? */
+	bool bIsHovered;
+
+private:
+
+	/**
+	 * Whether this widget is a "tool tip force field".  That is, tool-tips should never spawn over the area
+	 * occupied by this widget, and will instead be repelled to an outside edge
+	 */
 	bool bToolTipForceFieldEnabled;
+
+	/** Should we be forcing this widget to be volatile at all times and redrawn every frame? */
+	bool bForceVolatile;
+
+	/** The last cached volatility of this widget.  Cached so that we don't need to recompute volatility every frame. */
+	bool bCachedVolatile;
+
+	/** If we're owned by a volatile widget, we need inherit that volatility and use as part of our volatility, but don't cache it. */
+	mutable bool bInheritedVolatility;
 };
diff --git a/Engine/Source/Runtime/SlateCore/Public/Widgets/SWindow.h b/Engine/Source/Runtime/SlateCore/Public/Widgets/SWindow.h
--- a/Engine/Source/Runtime/SlateCore/Public/Widgets/SWindow.h
+++ b/Engine/Source/Runtime/SlateCore/Public/Widgets/SWindow.h
@@ -1055,6 +1055,8 @@ private:
 };
 
 
+#if WITH_EDITOR
+
 /**
  * Hack to switch wolds in a scope and switch back when we fall out of scope                   
  */
@@ -1096,3 +1098,14 @@ private:
 	// The worldID serves as identification to the user about the world.  It can be anything although -1 is assumed to be always invalid.
 	int32 WorldId;
 };
+
+#else
+
+struct FScopedSwitchWorldHack
+{
+	FORCEINLINE FScopedSwitchWorldHack(const FWidgetPath& WidgetPath) { }
+	FORCEINLINE FScopedSwitchWorldHack(TSharedPtr<SWindow> InWindow) { }
+	FORCEINLINE ~FScopedSwitchWorldHack() { }
+};
+
+#endif
\ No newline at end of file
diff --git a/Engine/Source/Runtime/SlateRHIRenderer/Private/SlateRHIRenderer.cpp b/Engine/Source/Runtime/SlateRHIRenderer/Private/SlateRHIRenderer.cpp
--- a/Engine/Source/Runtime/SlateRHIRenderer/Private/SlateRHIRenderer.cpp
+++ b/Engine/Source/Runtime/SlateRHIRenderer/Private/SlateRHIRenderer.cpp
@@ -18,7 +18,6 @@ DECLARE_CYCLE_STAT(TEXT("Slate RT Present Time"), STAT_SlatePresentRTTime, STATG
 #define MAX_VIEWPORT_SIZE 16384
 
 
-
 void FSlateCrashReportResource::InitDynamicRHI()
 {
 	FRHIResourceCreateInfo CreateInfo;
@@ -152,7 +151,7 @@ void FSlateRHIRenderer::Initialize()
 {
 	LoadUsedTextures();
 
-	RenderingPolicy = MakeShareable( new FSlateRHIRenderingPolicy( FontCache, ResourceManager.ToSharedRef() ) ); 
+	RenderingPolicy = MakeShareable( new FSlateRHIRenderingPolicy( FontCache, ResourceManager.ToSharedRef() ) );
 
 	ElementBatcher = MakeShareable( new FSlateElementBatcher( RenderingPolicy.ToSharedRef() ) );
 }
@@ -565,8 +564,8 @@ void FSlateRHIRenderer::DrawWindows_Private( FSlateDrawBuffer& WindowDrawBuffer
 		Policy.BeginDrawingWindows();
 	});
 
-	// Clear accessed UTexture objects from the previous frame
-	ResourceManager->ClearAccessedUTextures();
+	// Clear accessed UTexture and Material objects from the previous frame
+	ResourceManager->ClearAccessedResources();
 
 	// Update texture atlases if needed
 	ResourceManager->UpdateTextureAtlases();
diff --git a/Engine/Source/Runtime/SlateRHIRenderer/Private/SlateRHIResourceManager.cpp b/Engine/Source/Runtime/SlateRHIRenderer/Private/SlateRHIResourceManager.cpp
--- a/Engine/Source/Runtime/SlateRHIRenderer/Private/SlateRHIResourceManager.cpp
+++ b/Engine/Source/Runtime/SlateRHIRenderer/Private/SlateRHIResourceManager.cpp
@@ -217,6 +217,12 @@ bool FSlateRHIResourceManager::IsAtlasPageResourceAlphaOnly() const
 	return false;
 }
 
+void FSlateRHIResourceManager::AddReferencedObjects(FReferenceCollector& Collector)
+{
+	Collector.AddReferencedObjects(AccessedUTextures);
+	Collector.AddReferencedObjects(AccessedMaterials);
+}
+
 void FSlateRHIResourceManager::CreateTextures( const TArray< const FSlateBrush* >& Resources )
 {
 	TMap<FName,FNewTextureInfo> TextureInfoMap;
@@ -696,6 +702,8 @@ FSlateMaterialResource* FSlateRHIResourceManager::GetMaterialResource(const UObj
 		MaterialResource->SlateProxy->ActualSize = ImageSize.IntPoint();
 	}
 
+	AccessedMaterials.Add(const_cast<UMaterialInterface*>( Material ));
+
 	return MaterialResource.Get();
 }
 
@@ -790,9 +798,10 @@ void FSlateRHIResourceManager::LoadStyleResources( const ISlateStyle& Style )
 	CreateTextures( Resources );
 }
 
-void FSlateRHIResourceManager::ClearAccessedUTextures()
+void FSlateRHIResourceManager::ClearAccessedResources()
 {
-	AccessedUTextures.Empty();
+	AccessedUTextures.Reset();
+	AccessedMaterials.Reset();
 }
 
 void FSlateRHIResourceManager::UpdateTextureAtlases()
@@ -838,6 +847,7 @@ void FSlateRHIResourceManager::DeleteResources()
 	SET_DWORD_STAT(STAT_SlateNumDynamicTextures, 0);
 
 	AccessedUTextures.Empty();
+	AccessedMaterials.Empty();
 	DynamicResourceMap.Empty();
 	TextureAtlases.Empty();
 	NonAtlasedTextures.Empty();
diff --git a/Engine/Source/Runtime/SlateRHIRenderer/Private/SlateRHIResourceManager.h b/Engine/Source/Runtime/SlateRHIRenderer/Private/SlateRHIResourceManager.h
--- a/Engine/Source/Runtime/SlateRHIRenderer/Private/SlateRHIResourceManager.h
+++ b/Engine/Source/Runtime/SlateRHIRenderer/Private/SlateRHIResourceManager.h
@@ -64,7 +64,7 @@ private:
 /**
  * Stores a mapping of texture names to their RHI texture resource               
  */
-class FSlateRHIResourceManager : public ISlateAtlasProvider, public FSlateShaderResourceManager
+class FSlateRHIResourceManager : public ISlateAtlasProvider, public FSlateShaderResourceManager, public FGCObject
 {
 public:
 	FSlateRHIResourceManager();
@@ -76,6 +76,8 @@ public:
 	virtual FSlateShaderResource* GetAtlasPageResource(const int32 InIndex) const override;
 	virtual bool IsAtlasPageResourceAlphaOnly() const override;
 
+	virtual void AddReferencedObjects(FReferenceCollector& Collector) override;
+
 	/**
 	 * Loads and creates rendering resources for all used textures.  
 	 * In this implementation all textures must be known at startup time or they will not be found
@@ -85,11 +87,11 @@ public:
 	void LoadStyleResources( const ISlateStyle& Style );
 
 	/**
-	 * Clears accessed UTexture resources from the previous frame
+	 * Clears accessed UTexture and Material resources from the previous frame
 	 * The accessed textures is used to determine which textures need be updated on the render thread
 	 * so they can be used by slate
 	 */
-	void ClearAccessedUTextures();
+	void ClearAccessedResources();
 
 	/**
 	 * Updates texture atlases if needed
@@ -203,7 +205,9 @@ private:
 	/** Map of all active dynamic resources being used by brushes */
 	FDynamicResourceMap DynamicResourceMap;
 	/** Set of dynamic textures that are currently being accessed */
-	TSet< TWeakObjectPtr<UTexture> > AccessedUTextures;
+	TSet<UTexture*> AccessedUTextures;
+	/** Set of dynamic materials that are current being accessed */
+	TSet<UMaterialInterface*> AccessedMaterials;
 	/** List of old utexture resources that are free to use as new resources */
 	TArray< TSharedPtr<FSlateUTextureResource> > UTextureFreeList;
 	/** List of old dynamic resources that are free to use as new resources */
diff --git a/Engine/Source/Runtime/UMG/Private/Components/InvalidationBox.cpp b/Engine/Source/Runtime/UMG/Private/Components/InvalidationBox.cpp
new file mode 100644
--- /dev/null
+++ b/Engine/Source/Runtime/UMG/Private/Components/InvalidationBox.cpp
@@ -0,0 +1,84 @@
+// Copyright 1998-2015 Epic Games, Inc. All Rights Reserved.
+
+#include "UMGPrivatePCH.h"
+#include "InvalidationBox.h"
+#include "SInvalidationPanel.h"
+
+#define LOCTEXT_NAMESPACE "UMG"
+
+/////////////////////////////////////////////////////
+// UInvalidationBox
+
+UInvalidationBox::UInvalidationBox(const FObjectInitializer& ObjectInitializer)
+	: Super(ObjectInitializer)
+{
+	Visibility = ESlateVisibility::HitTestInvisible;
+}
+
+void UInvalidationBox::ReleaseSlateResources(bool bReleaseChildren)
+{
+	Super::ReleaseSlateResources(bReleaseChildren);
+
+	MyInvalidationPanel.Reset();
+}
+
+TSharedRef<SWidget> UInvalidationBox::RebuildWidget()
+{
+	MyInvalidationPanel = SNew(SInvalidationPanel);
+
+	if ( IsDesignTime() )
+	{
+		MyInvalidationPanel->SetCanCache(false);
+	}
+
+	if ( GetChildrenCount() > 0 )
+	{
+		MyInvalidationPanel->SetContent(GetContentSlot()->Content ? GetContentSlot()->Content->TakeWidget() : SNullWidget::NullWidget);
+	}
+	
+	return BuildDesignTimeWidget(MyInvalidationPanel.ToSharedRef());
+}
+
+void UInvalidationBox::OnSlotAdded(UPanelSlot* Slot)
+{
+	// Add the child to the live slot if it already exists
+	if ( MyInvalidationPanel.IsValid() )
+	{
+		MyInvalidationPanel->SetContent(Slot->Content ? Slot->Content->TakeWidget() : SNullWidget::NullWidget);
+	}
+}
+
+void UInvalidationBox::OnSlotRemoved(UPanelSlot* Slot)
+{
+	// Remove the widget from the live slot if it exists.
+	if ( MyInvalidationPanel.IsValid() )
+	{
+		MyInvalidationPanel->SetContent(SNullWidget::NullWidget);
+	}
+}
+
+void UInvalidationBox::Invalidate()
+{
+	if ( MyInvalidationPanel.IsValid() )
+	{
+		return MyInvalidationPanel->Invalidate();
+	}
+}
+
+#if WITH_EDITOR
+
+const FSlateBrush* UInvalidationBox::GetEditorIcon()
+{
+	return FUMGStyle::Get().GetBrush("Widget.MenuAnchor");
+}
+
+const FText UInvalidationBox::GetPaletteCategory()
+{
+	return LOCTEXT("Optimization", "Optimization");
+}
+
+#endif
+
+/////////////////////////////////////////////////////
+
+#undef LOCTEXT_NAMESPACE
diff --git a/Engine/Source/Runtime/UMG/Private/Components/Widget.cpp b/Engine/Source/Runtime/UMG/Private/Components/Widget.cpp
--- a/Engine/Source/Runtime/UMG/Private/Components/Widget.cpp
+++ b/Engine/Source/Runtime/UMG/Private/Components/Widget.cpp
@@ -217,6 +217,16 @@ void UWidget::SetVisibility(ESlateVisibility InVisibility)
 	}
 }
 
+void UWidget::ForceVolatile(bool bForce)
+{
+	bIsVolatile = bForce;
+	TSharedPtr<SWidget> SafeWidget = GetCachedWidget();
+	if ( SafeWidget.IsValid() )
+	{
+		SafeWidget->ForceVolatile(bForce);
+	}
+}
+
 void UWidget::SetToolTipText(const FText& InToolTipText)
 {
 	ToolTipText = InToolTipText;
@@ -398,6 +408,15 @@ void UWidget::ForceLayoutPrepass()
 	}
 }
 
+void UWidget::InvalidateLayout()
+{
+	TSharedPtr<SWidget> SafeWidget = GetCachedWidget();
+	if ( SafeWidget.IsValid() )
+	{
+		SafeWidget->InvalidateLayout();
+	}
+}
+
 FVector2D UWidget::GetDesiredSize() const
 {
 	TSharedPtr<SWidget> SafeWidget = GetCachedWidget();
@@ -720,6 +739,8 @@ void UWidget::SynchronizeProperties()
 		SafeWidget->SetVisibility(OPTIONAL_BINDING_CONVERT(ESlateVisibility, Visibility, EVisibility, ConvertVisibility));
 	}
 
+	SafeWidget->ForceVolatile(bIsVolatile);
+
 	UpdateRenderTransform();
 	SafeWidget->SetRenderTransformPivot(RenderTransformPivot);
 
diff --git a/Engine/Source/Runtime/UMG/Private/Slate/SObjectWidget.cpp b/Engine/Source/Runtime/UMG/Private/Slate/SObjectWidget.cpp
--- a/Engine/Source/Runtime/UMG/Private/Slate/SObjectWidget.cpp
+++ b/Engine/Source/Runtime/UMG/Private/Slate/SObjectWidget.cpp
@@ -74,6 +74,16 @@ int32 SObjectWidget::OnPaint(const FPaintArgs& Args, const FGeometry& AllottedGe
 	return MaxLayer;
 }
 
+bool SObjectWidget::ComputeVolatility() const
+{
+	if ( CanRouteEvent() )
+	{
+		return SCompoundWidget::ComputeVolatility() || WidgetObject->IsPlayingAnimation();
+	}
+
+	return SCompoundWidget::ComputeVolatility();
+}
+
 bool SObjectWidget::IsInteractable() const
 {
 	if ( CanRouteEvent() )
diff --git a/Engine/Source/Runtime/UMG/Private/UserWidget.cpp b/Engine/Source/Runtime/UMG/Private/UserWidget.cpp
--- a/Engine/Source/Runtime/UMG/Private/UserWidget.cpp
+++ b/Engine/Source/Runtime/UMG/Private/UserWidget.cpp
@@ -222,6 +222,12 @@ void UUserWidget::PlayAnimation( const UWidgetAnimation* InAnimation, float Star
 			( *FoundPlayer )->Play( StartAtTime, NumberOfLoops, PlayMode );
 		}
 
+		TSharedPtr<SWidget> CachedWidget = GetCachedWidget();
+		if ( CachedWidget.IsValid() )
+		{
+			CachedWidget->InvalidateLayout();
+		}
+
 		OnAnimationStarted( InAnimation );
 	}
 }
@@ -733,14 +739,26 @@ void UUserWidget::TickActionsAndAnimation(const FGeometry& MyGeometry, float InD
 		Player->Tick(InDeltaTime);
 	}
 
+	const bool bWasPlayingAnimation = IsPlayingAnimation();
+
 	// The process of ticking the players above can stop them so we remove them after all players have ticked
 	for ( UUMGSequencePlayer* StoppedPlayer : StoppedSequencePlayers )
 	{
-		ActiveSequencePlayers.Remove(StoppedPlayer);
+		ActiveSequencePlayers.RemoveSwap(StoppedPlayer);
 	}
 
 	StoppedSequencePlayers.Empty();
 
+	// If we're no longer playing animations invalidate layout so that we recache the volatility of the widget.
+	if ( bWasPlayingAnimation && IsPlayingAnimation() == false )
+	{
+		TSharedPtr<SWidget> CachedWidget = GetCachedWidget();
+		if ( CachedWidget.IsValid() )
+		{
+			CachedWidget->InvalidateLayout();
+		}
+	}
+
 	UWorld* World = GetWorld();
 	if ( World )
 	{
diff --git a/Engine/Source/Runtime/UMG/Public/Blueprint/UserWidget.h b/Engine/Source/Runtime/UMG/Public/Blueprint/UserWidget.h
--- a/Engine/Source/Runtime/UMG/Public/Blueprint/UserWidget.h
+++ b/Engine/Source/Runtime/UMG/Public/Blueprint/UserWidget.h
@@ -760,6 +760,9 @@ public:
 	virtual void PreSave() override;
 	// End UObject interface
 
+	/** Are we currently playing any animations? */
+	bool IsPlayingAnimation() const { return ActiveSequencePlayers.Num() > 0; }
+
 #if WITH_EDITOR
 	// UWidget interface
 	virtual const FSlateBrush* GetEditorIcon() override;
diff --git a/Engine/Source/Runtime/UMG/Public/Components/InvalidationBox.h b/Engine/Source/Runtime/UMG/Public/Components/InvalidationBox.h
new file mode 100644
--- /dev/null
+++ b/Engine/Source/Runtime/UMG/Public/Components/InvalidationBox.h
@@ -0,0 +1,45 @@
+﻿// Copyright 1998-2015 Epic Games, Inc. All Rights Reserved.
+
+#pragma once
+
+#include "InvalidationBox.generated.h"
+
+/**
+ * Invalidate
+ * ● Single Child
+ * ● Popup
+ */
+UCLASS()
+class UMG_API UInvalidationBox : public UContentWidget
+{
+	GENERATED_UCLASS_BODY()
+
+public:
+
+	/**
+	 * 
+	 */
+	UFUNCTION(BlueprintCallable, Category="Invalidate")
+	void Invalidate();
+
+	virtual void ReleaseSlateResources(bool bReleaseChildren) override;
+
+#if WITH_EDITOR
+	virtual const FSlateBrush* GetEditorIcon() override;
+	virtual const FText GetPaletteCategory() override;
+#endif
+
+protected:
+
+	// UPanelWidget
+	virtual void OnSlotAdded(UPanelSlot* Slot) override;
+	virtual void OnSlotRemoved(UPanelSlot* Slot) override;
+	// End UPanelWidget
+
+	// UWidget interface
+	virtual TSharedRef<SWidget> RebuildWidget() override;
+	// End of UWidget interface
+
+protected:
+	TSharedPtr<class SInvalidationPanel> MyInvalidationPanel;
+};
diff --git a/Engine/Source/Runtime/UMG/Public/Components/Widget.h b/Engine/Source/Runtime/UMG/Public/Components/Widget.h
--- a/Engine/Source/Runtime/UMG/Public/Components/Widget.h
+++ b/Engine/Source/Runtime/UMG/Public/Components/Widget.h
@@ -177,6 +177,19 @@ public:
 	UPROPERTY()
 	FGetMouseCursor CursorDelegate;
 
+protected:
+
+	/**
+	 * If true prevents the widget or its child's geometry or layout information from being cached.  If this widget
+	 * changes every frame, but you want it to still be in an invalidation panel you should make it as volatile
+	 * instead of invalidating it every frame, which would prevent the invalidation panel from actually
+	 * ever caching anything.
+	 */
+	UPROPERTY(EditAnywhere, BlueprintReadOnly, Category="Performance")
+	bool bIsVolatile;
+
+public:
+
 	/** The render transform of the widget allows for arbitrary 2D transforms to be applied to the widget. */
 	UPROPERTY(EditAnywhere, BlueprintReadOnly, Category="Render Transform", meta=( DisplayName="Transform" ))
 	FWidgetTransform RenderTransform;
@@ -266,6 +279,10 @@ public:
 	UFUNCTION(BlueprintCallable, Category="Widget")
 	void SetVisibility(ESlateVisibility InVisibility);
 
+	/** Sets the forced volatility of the widget. */
+	UFUNCTION(BlueprintCallable, Category="Widget")
+	void ForceVolatile(bool bForce);
+
 	/** @return true if the widget is currently being hovered by a pointer device */
 	UFUNCTION(BlueprintCallable, Category="Widget")
 	bool IsHovered() const;
@@ -317,6 +334,12 @@ public:
 	UFUNCTION(BlueprintCallable, Category="Widget")
 	void ForceLayoutPrepass();
 
+	/**
+	 *
+	 */
+	UFUNCTION(BlueprintCallable, Category="Widget")
+	void InvalidateLayout();
+
 	/**
 	 * Gets the widgets desired size.
 	 * NOTE: The underlying Slate widget must exist and be valid, also at least one pre-pass must
diff --git a/Engine/Source/Runtime/UMG/Public/Slate/SObjectWidget.h b/Engine/Source/Runtime/UMG/Public/Slate/SObjectWidget.h
--- a/Engine/Source/Runtime/UMG/Public/Slate/SObjectWidget.h
+++ b/Engine/Source/Runtime/UMG/Public/Slate/SObjectWidget.h
@@ -35,6 +35,8 @@ class UMG_API SObjectWidget : public SCompoundWidget, public FGCObject
 	virtual void Tick( const FGeometry& AllottedGeometry, const double InCurrentTime, const float InDeltaTime ) override;
 	virtual int32 OnPaint(const FPaintArgs& Args, const FGeometry& AllottedGeometry, const FSlateRect& MyClippingRect, FSlateWindowElementList& OutDrawElements, int32 LayerId, const FWidgetStyle& InWidgetStyle, bool bParentEnabled) const override;
 
+	virtual bool ComputeVolatility() const override;
+
 	virtual bool IsInteractable() const override;
 	virtual bool SupportsKeyboardFocus() const override;
 
