[INTEGRATE] merge ue4<-ue4-launcherDev cl 2612250

diff --git a/Engine/Source/Runtime/SlateCore/Private/Styling/SlateStyleSet.cpp b/Engine/Source/Runtime/SlateCore/Private/Styling/SlateStyleSet.cpp
--- a/Engine/Source/Runtime/SlateCore/Private/Styling/SlateStyleSet.cpp
+++ b/Engine/Source/Runtime/SlateCore/Private/Styling/SlateStyleSet.cpp
@@ -421,6 +421,15 @@ bool FSlateStyleSet::IsBrushFromFile(const FString& FilePath, const FSlateBrush*
 	{
 		return true;
 	}
+	else
+	{
+		const FString FullFilePath = FPaths::ConvertRelativePathToFull(FilePath);
+		const FString FullPath = FPaths::ConvertRelativePathToFull(Path);
+		if (FullPath.Compare(FullFilePath, ESearchCase::IgnoreCase) == 0)
+		{
+			return true;
+		}
+	}
 
 	return false;
 }
