UMG - Continuing to work on the invalidation panel, there's now a caching mechanism it uses to allow aspects of the hit testing to be recreated every frame.  Additionally, started working on a test suite for it in the Slate Viewer.  A simple proof of concept of invalidating a hit testable control can be seen on button.

diff --git a/Engine/Source/Runtime/AppFramework/Private/Widgets/Testing/STestSuite.cpp b/Engine/Source/Runtime/AppFramework/Private/Widgets/Testing/STestSuite.cpp
--- a/Engine/Source/Runtime/AppFramework/Private/Widgets/Testing/STestSuite.cpp
+++ b/Engine/Source/Runtime/AppFramework/Private/Widgets/Testing/STestSuite.cpp
@@ -24,6 +24,7 @@
 #include "SColorPicker.h"
 #include "INotificationWidget.h"
 #include "IMenu.h"
+#include "SInvalidationPanel.h"
 
 #define LOCTEXT_NAMESPACE "STestSuite"
 
@@ -3841,6 +3842,97 @@ class SDPIScalingTest : public SCompoundWidget
 	SVerticalBox::FSlot* ScalerSlot;
 };
 
+class SInvalidationTest : public SCompoundWidget
+{
+	SLATE_BEGIN_ARGS(SInvalidationTest)
+	{}
+	SLATE_END_ARGS()
+
+	void Construct(const FArguments& InArgs)
+	{
+		ChildSlot
+		.Padding(10)
+		[
+			SNew(SVerticalBox)
+
+			+ SVerticalBox::Slot()
+			.AutoHeight()
+			[
+				SAssignNew(CachePanel1, SInvalidationPanel)
+				.Visibility(EVisibility::Visible)
+				[
+					SNew(SVerticalBox)
+
+					+ SVerticalBox::Slot()
+					.AutoHeight()
+					[
+						SNew(STextBlock)
+						.Text(LOCTEXT("StaticText", "This text is static and cached."))
+					]
+				]
+			]
+
+			+ SVerticalBox::Slot()
+			.AutoHeight()
+			[
+				SNew(SButton)
+				.Text(LOCTEXT("InvalidateManually", "Manually Invalidate"))
+				.OnClicked(this, &SInvalidationTest::ManuallyInvalidatePanel1)
+			]
+
+			+ SVerticalBox::Slot()
+			.AutoHeight()
+			[
+				SNew(SInvalidationPanel)
+				.Visibility(EVisibility::Visible)
+				[
+					SNew(SVerticalBox)
+
+					+ SVerticalBox::Slot()
+					.AutoHeight()
+					[
+						SNew(STextBlock)
+						.Text(LOCTEXT("StaticText", "Support Input"))
+					]
+
+					+ SVerticalBox::Slot()
+					.AutoHeight()
+					[
+						SNew(SBorder)
+						.Padding(10)
+						[
+							SNew(SButton)
+							.Text(LOCTEXT("CachedClickable", "I'm Cached - But Clickable"))
+						]
+					]
+
+					+ SVerticalBox::Slot()
+					.AutoHeight()
+					[
+						SNew(SBorder)
+						.Padding(10)
+						[
+							SNew(SButton)
+							.ForceVolatile(true)
+							.Text(LOCTEXT("VolatileClickable", "Volatile - But Clickable"))
+						]
+					]
+				]
+			]
+		];
+	}
+
+private:
+	FReply ManuallyInvalidatePanel1()
+	{
+		CachePanel1->InvalidateCache();
+		return FReply::Handled();
+	}
+
+private:
+	TSharedPtr<SInvalidationPanel> CachePanel1;
+};
+
 class SColorPickerTest : public SCompoundWidget
 {
 public:
@@ -5210,6 +5302,18 @@ TSharedRef<SDockTab> SpawnTab(const FSpawnTabArgs& Args, FName TabIdentifier)
 			]
 		];
 	}
+	else if ( TabIdentifier == FName(TEXT("InvalidationTest")) )
+	{
+		return SNew(SDockTab)
+			[
+				SNew(SScissorRectBox)
+				[
+					SNew(SInvalidationTest)
+					.RenderTransform_Static(&::GetTestRenderTransform)
+					.RenderTransformPivot_Static(&::GetTestRenderTransformPivot)
+				]
+			];
+	}
 	else if (TabIdentifier == FName(TEXT("NotificationListTestTab")))
 	{
 		return SNew(SDockTab)
@@ -5431,6 +5535,7 @@ TSharedRef<SDockTab> SpawnTestSuite2( const FSpawnTabArgs& Args )
 			->AddTab("NotificationListTestTab", ETabState::OpenedTab)
 			->AddTab("GridPanelTest", ETabState::OpenedTab)
 			->AddTab("DPIScalingTest", ETabState::OpenedTab)
+			->AddTab("InvalidationTest", ETabState::OpenedTab)
 		)
 	);
 
@@ -5464,8 +5569,11 @@ TSharedRef<SDockTab> SpawnTestSuite2( const FSpawnTabArgs& Args )
 
 		TestSuite2TabManager->RegisterTabSpawner( "DPIScalingTest", FOnSpawnTab::CreateStatic( &SpawnTab, FName("DPIScalingTest") )  )
 			.SetDisplayName( NSLOCTEXT("TestSuite1", "DPIScalingTest", "DPI Scaling") )
-			.SetGroup(TestSuiteMenu::SuiteTabs
-			);
+			.SetGroup(TestSuiteMenu::SuiteTabs);
+
+		TestSuite2TabManager->RegisterTabSpawner("InvalidationTest", FOnSpawnTab::CreateStatic(&SpawnTab, FName("InvalidationTest")))
+			.SetDisplayName(NSLOCTEXT("TestSuite1", "InvalidationTest", "Invalidtion"))
+			.SetGroup(TestSuiteMenu::SuiteTabs);
 	}
 
 	FMenuBarBuilder MenuBarBuilder = FMenuBarBuilder( TSharedPtr<FUICommandList>() );
diff --git a/Engine/Source/Runtime/Slate/Private/Widgets/Images/SImage.cpp b/Engine/Source/Runtime/Slate/Private/Widgets/Images/SImage.cpp
--- a/Engine/Source/Runtime/Slate/Private/Widgets/Images/SImage.cpp
+++ b/Engine/Source/Runtime/Slate/Private/Widgets/Images/SImage.cpp
@@ -71,7 +71,7 @@ void SImage::SetColorAndOpacity( const TAttribute<FSlateColor>& InColorAndOpacit
 	if ( !ColorAndOpacity.IdenticalTo(InColorAndOpacity) )
 	{
 		ColorAndOpacity = InColorAndOpacity;
-		InvalidateLayout();
+		Invalidate(EInvalidateWidget::Layout);
 	}
 }
 
@@ -80,14 +80,14 @@ void SImage::SetColorAndOpacity( FLinearColor InColorAndOpacity )
 	if ( ColorAndOpacity.IsBound() || ColorAndOpacity.Get() != InColorAndOpacity )
 	{
 		ColorAndOpacity = InColorAndOpacity;
-		InvalidateLayout();
+		Invalidate(EInvalidateWidget::Layout);
 	}
 }
 
 void SImage::SetImage(TAttribute<const FSlateBrush*> InImage)
 {
 	Image = InImage;
-	InvalidateLayout();
+	Invalidate(EInvalidateWidget::Layout);
 }
 
 void SImage::SetOnMouseButtonDown(FPointerEventHandler EventHandler)
diff --git a/Engine/Source/Runtime/Slate/Private/Widgets/Input/SButton.cpp b/Engine/Source/Runtime/Slate/Private/Widgets/Input/SButton.cpp
--- a/Engine/Source/Runtime/Slate/Private/Widgets/Input/SButton.cpp
+++ b/Engine/Source/Runtime/Slate/Private/Widgets/Input/SButton.cpp
@@ -194,6 +194,8 @@ FReply SButton::OnMouseButtonDown( const FGeometry& MyGeometry, const FPointerEv
 		}
 	}
 
+	Invalidate(EInvalidateWidget::Layout);
+
 	//return the constructed reply
 	return Reply;
 }
@@ -251,6 +253,8 @@ FReply SButton::OnMouseButtonUp( const FGeometry& MyGeometry, const FPointerEven
 		}
 	}
 
+	Invalidate(EInvalidateWidget::Layout);
+
 	return Reply;
 }
 
@@ -272,6 +276,8 @@ void SButton::OnMouseEnter( const FGeometry& MyGeometry, const FPointerEvent& Mo
 	}
 	
 	SBorder::OnMouseEnter( MyGeometry, MouseEvent );
+
+	Invalidate(EInvalidateWidget::Layout);
 }
 
 void SButton::OnMouseLeave( const FPointerEvent& MouseEvent )
@@ -285,6 +291,8 @@ void SButton::OnMouseLeave( const FPointerEvent& MouseEvent )
 	{
 		Release();
 	}
+
+	Invalidate(EInvalidateWidget::Layout);
 }
 
 void SButton::Press()
diff --git a/Engine/Source/Runtime/Slate/Private/Widgets/Layout/SSpacer.cpp b/Engine/Source/Runtime/Slate/Private/Widgets/Layout/SSpacer.cpp
--- a/Engine/Source/Runtime/Slate/Private/Widgets/Layout/SSpacer.cpp
+++ b/Engine/Source/Runtime/Slate/Private/Widgets/Layout/SSpacer.cpp
@@ -21,6 +21,7 @@ void SSpacer::Construct( const FArguments& InArgs )
 		InArgs._RenderTransform,
 		InArgs._RenderTransformPivot,
 		InArgs._Tag,
+		InArgs._ForceVolatile,
 		InArgs.MetaData);
 }
 
diff --git a/Engine/Source/Runtime/Slate/Private/Widgets/Notifications/SProgressBar.cpp b/Engine/Source/Runtime/Slate/Private/Widgets/Notifications/SProgressBar.cpp
--- a/Engine/Source/Runtime/Slate/Private/Widgets/Notifications/SProgressBar.cpp
+++ b/Engine/Source/Runtime/Slate/Private/Widgets/Notifications/SProgressBar.cpp
@@ -32,50 +32,50 @@ void SProgressBar::SetPercent(TAttribute< TOptional<float> > InPercent)
 	if ( !Percent.IdenticalTo(InPercent) )
 	{
 		Percent = InPercent;
-		InvalidateLayout();
+		Invalidate(EInvalidateWidget::LayoutAndVolatility);
 	}
 }
 
 void SProgressBar::SetStyle(const FProgressBarStyle* InStyle)
 {
 	Style = InStyle;
-	InvalidateLayout();
+	Invalidate(EInvalidateWidget::Layout);
 }
 
 void SProgressBar::SetBarFillType(EProgressBarFillType::Type InBarFillType)
 {
 	BarFillType = InBarFillType;
-	InvalidateLayout();
+	Invalidate(EInvalidateWidget::Layout);
 }
 
 void SProgressBar::SetFillColorAndOpacity(TAttribute< FSlateColor > InFillColorAndOpacity)
 {
 	FillColorAndOpacity = InFillColorAndOpacity;
-	InvalidateLayout();
+	Invalidate(EInvalidateWidget::Layout);
 }
 
 void SProgressBar::SetBorderPadding(TAttribute< FVector2D > InBorderPadding)
 {
 	BorderPadding = InBorderPadding;
-	InvalidateLayout();
+	Invalidate(EInvalidateWidget::Layout);
 }
 
 void SProgressBar::SetBackgroundImage(const FSlateBrush* InBackgroundImage)
 {
 	BackgroundImage = InBackgroundImage;
-	InvalidateLayout();
+	Invalidate(EInvalidateWidget::Layout);
 }
 
 void SProgressBar::SetFillImage(const FSlateBrush* InFillImage)
 {
 	FillImage = InFillImage;
-	InvalidateLayout();
+	Invalidate(EInvalidateWidget::Layout);
 }
 
 void SProgressBar::SetMarqueeImage(const FSlateBrush* InMarqueeImage)
 {
 	MarqueeImage = InMarqueeImage;
-	InvalidateLayout();
+	Invalidate(EInvalidateWidget::Layout);
 }
 
 const FSlateBrush* SProgressBar::GetBackgroundImage() const
diff --git a/Engine/Source/Runtime/Slate/Private/Widgets/SInvalidationPanel.cpp b/Engine/Source/Runtime/Slate/Private/Widgets/SInvalidationPanel.cpp
--- a/Engine/Source/Runtime/Slate/Private/Widgets/SInvalidationPanel.cpp
+++ b/Engine/Source/Runtime/Slate/Private/Widgets/SInvalidationPanel.cpp
@@ -3,6 +3,7 @@
 #include "SlatePrivatePCH.h"
 #include "LayoutUtils.h"
 #include "SInvalidationPanel.h"
+#include "WidgetCaching.h"
 
 #if !UE_BUILD_SHIPPING
 
@@ -40,6 +41,16 @@ void SInvalidationPanel::Construct( const FArguments& InArgs )
 	bNeedsCaching = true;
 	bIsInvalidating = false;
 	bCanCache = true;
+	RootCacheNode = nullptr;
+	LastUsedCachedNodeIndex = 0;
+}
+
+SInvalidationPanel::~SInvalidationPanel()
+{
+	for ( int32 i = 0; i < NodePool.Num(); i++ )
+	{
+		delete NodePool[i];
+	}
 }
 
 bool SInvalidationPanel::GetCanCache() const
@@ -54,10 +65,8 @@ bool SInvalidationPanel::GetCanCache() const
 void SInvalidationPanel::SetCanCache(bool InCanCache)
 {
 	bCanCache = InCanCache;
-	Visibility = InCanCache ? EVisibility::HitTestInvisible : EVisibility::SelfHitTestInvisible;
 
-	Invalidate();
-	InvalidateLayout();
+	InvalidateCache();
 }
 
 void SInvalidationPanel::Tick( const FGeometry& AllottedGeometry, const double InCurrentTime, const float InDeltaTime )
@@ -70,7 +79,7 @@ void SInvalidationPanel::Tick( const FGeometry& AllottedGeometry, const double I
 			if ( AllottedGeometry.GetAccumulatedLayoutTransform() != LastAllottedGeometry.GetAccumulatedLayoutTransform() ||
 				AllottedGeometry.GetAccumulatedRenderTransform() != LastAllottedGeometry.GetAccumulatedRenderTransform() )
 			{
-				Invalidate();
+				InvalidateCache();
 			}
 		}
 
@@ -84,11 +93,6 @@ void SInvalidationPanel::Tick( const FGeometry& AllottedGeometry, const double I
 	}
 }
 
-void SInvalidationPanel::OnArrangeChildren( const FGeometry& AllottedGeometry, FArrangedChildren& ArrangedChildren ) const
-{
-	SCompoundWidget::OnArrangeChildren( AllottedGeometry, ArrangedChildren );
-}
-
 FChildren* SInvalidationPanel::GetChildren()
 {
 	if ( GetCanCache() == false || bNeedsCaching )
@@ -113,8 +117,29 @@ void SInvalidationPanel::InvalidateWidget(SWidget* InvalidateWidget)
 #endif
 }
 
+FCachedWidgetNode* SInvalidationPanel::CreateCacheNode() const
+{
+	// If the node pool is empty, allocate a few
+	if ( LastUsedCachedNodeIndex >= NodePool.Num() )
+	{
+		for ( int32 i = 0; i < 10; i++ )
+		{
+			NodePool.Add(new FCachedWidgetNode());
+		}
+	}
+
+	// Return one of the preallocated nodes and increment the next node index.
+	FCachedWidgetNode* NewNode = NodePool[LastUsedCachedNodeIndex];
+	++LastUsedCachedNodeIndex;
+
+	return NewNode;
+}
+
 int32 SInvalidationPanel::OnPaint( const FPaintArgs& Args, const FGeometry& AllottedGeometry, const FSlateRect& MyClippingRect, FSlateWindowElementList& OutDrawElements, int32 LayerId, const FWidgetStyle& InWidgetStyle, bool bParentEnabled ) const
 {
+	//TODO If the LayerId changes, or other properties, we should probably invalidate, but
+	// we can't invalidate here, because we haven't done a pre-pass yet, need to do it next frame.
+
 	if ( GetCanCache() )
 	{
 		const bool bWasCachingNeeded = bNeedsCaching;
@@ -136,9 +161,15 @@ int32 SInvalidationPanel::OnPaint( const FPaintArgs& Args, const FGeometry& Allo
 				CachedWindowElements->Reset();
 			}
 
+			// Reset the cached node pool index so that we effectively reset the pool.
+			LastUsedCachedNodeIndex = 0;
+
+			RootCacheNode = CreateCacheNode();
+			RootCacheNode->Initialize(Args, SharedThis(const_cast<SInvalidationPanel*>( this )), AllottedGeometry, MyClippingRect);
+
 			//TODO: When SWidget::Paint is called don't drag self if volatile, and we're doing a cache pass.
 			CachedMaxChildLayer = SCompoundWidget::OnPaint(
-				Args.EnableCaching(SharedThis(const_cast<SInvalidationPanel*>( this )), true, false),
+				Args.EnableCaching(SharedThis(const_cast<SInvalidationPanel*>( this )), RootCacheNode, true, false),
 				AllottedGeometry,
 				MyClippingRect,
 				*CachedWindowElements.Get(),
@@ -149,6 +180,8 @@ int32 SInvalidationPanel::OnPaint( const FPaintArgs& Args, const FGeometry& Allo
 			bIsInvalidating = false;
 		}
 
+		RootCacheNode->RecordHittestGeometry(Args.GetGrid(), Args.GetLastHitTestIndex());
+
 		OutDrawElements.AppendDrawElements(CachedWindowElements->GetDrawElements());
 
 		int32 OutMaxChildLayer = CachedMaxChildLayer;
@@ -237,7 +270,7 @@ int32 SInvalidationPanel::OnPaint( const FPaintArgs& Args, const FGeometry& Allo
 
 void SInvalidationPanel::SetContent(const TSharedRef< SWidget >& InContent)
 {
-	Invalidate();
+	InvalidateCache();
 
 	ChildSlot
 	[
@@ -245,12 +278,7 @@ void SInvalidationPanel::SetContent(const TSharedRef< SWidget >& InContent)
 	];
 }
 
-void SInvalidationPanel::CacheDesiredSize(float LayoutScaleMultiplier)
-{
-	SCompoundWidget::CacheDesiredSize(LayoutScaleMultiplier);
-}
-
-void SInvalidationPanel::Invalidate()
+void SInvalidationPanel::InvalidateCache()
 {
 	bNeedsCaching = true;
 }
diff --git a/Engine/Source/Runtime/Slate/Private/Widgets/Text/STextBlock.cpp b/Engine/Source/Runtime/Slate/Private/Widgets/Text/STextBlock.cpp
--- a/Engine/Source/Runtime/Slate/Private/Widgets/Text/STextBlock.cpp
+++ b/Engine/Source/Runtime/Slate/Private/Widgets/Text/STextBlock.cpp
@@ -77,19 +77,19 @@ void STextBlock::SetText( const TAttribute< FString >& InText )
 
 	BoundText = TAttribute< FText >::Create(TAttribute<FText>::FGetter::CreateStatic( &Local::PassThroughAttribute, InText) );
 
-	InvalidateLayout();
+	Invalidate(EInvalidateWidget::LayoutAndVolatility);
 }
 
 void STextBlock::SetText( const FString& InText )
 {
 	BoundText = FText::FromString( InText );
-	InvalidateLayout();
+	Invalidate(EInvalidateWidget::LayoutAndVolatility);
 }
 
 void STextBlock::SetText( const TAttribute< FText >& InText )
 {
 	BoundText = InText;
-	InvalidateLayout();
+	Invalidate(EInvalidateWidget::LayoutAndVolatility);
 }
 
 void STextBlock::SetText( const FText& InText )
@@ -97,7 +97,7 @@ void STextBlock::SetText( const FText& InText )
 	if ( BoundText.IsBound() || BoundText.Get().ToString() != InText.ToString() )
 	{
 		BoundText = InText;
-		InvalidateLayout();
+		Invalidate(EInvalidateWidget::LayoutAndVolatility);
 	}
 }
 
@@ -199,7 +199,7 @@ bool STextBlock::ComputeVolatility() const
 void STextBlock::SetFont(const TAttribute< FSlateFontInfo >& InFont)
 {
 	Font = InFont;
-	InvalidateLayout();
+	Invalidate(EInvalidateWidget::Layout);
 }
 
 void STextBlock::SetColorAndOpacity(const TAttribute<FSlateColor>& InColorAndOpacity)
@@ -207,62 +207,62 @@ void STextBlock::SetColorAndOpacity(const TAttribute<FSlateColor>& InColorAndOpa
 	if ( !ColorAndOpacity.IdenticalTo(InColorAndOpacity) )
 	{
 		ColorAndOpacity = InColorAndOpacity;
-		InvalidateLayout();
+		Invalidate(EInvalidateWidget::Layout);
 	}
 }
 
 void STextBlock::SetTextStyle(const FTextBlockStyle* InTextStyle)
 {
 	TextStyle = InTextStyle;
-	InvalidateLayout();
+	Invalidate(EInvalidateWidget::Layout);
 }
 
 void STextBlock::SetWrapTextAt(const TAttribute<float>& InWrapTextAt)
 {
 	WrapTextAt = InWrapTextAt;
-	InvalidateLayout();
+	Invalidate(EInvalidateWidget::Layout);
 }
 
 void STextBlock::SetAutoWrapText(const TAttribute<bool>& InAutoWrapText)
 {
 	AutoWrapText = InAutoWrapText;
-	InvalidateLayout();
+	Invalidate(EInvalidateWidget::Layout);
 }
 
 void STextBlock::SetShadowOffset(const TAttribute<FVector2D>& InShadowOffset)
 {
 	ShadowOffset = InShadowOffset;
-	InvalidateLayout();
+	Invalidate(EInvalidateWidget::Layout);
 }
 
 void STextBlock::SetShadowColorAndOpacity(const TAttribute<FLinearColor>& InShadowColorAndOpacity)
 {
 	ShadowColorAndOpacity = InShadowColorAndOpacity;
-	InvalidateLayout();
+	Invalidate(EInvalidateWidget::Layout);
 }
 
 void STextBlock::SetMinDesiredWidth(const TAttribute<float>& InMinDesiredWidth)
 {
 	MinDesiredWidth = InMinDesiredWidth;
-	InvalidateLayout();
+	Invalidate(EInvalidateWidget::Layout);
 }
 
 void STextBlock::SetLineHeightPercentage(const TAttribute<float>& InLineHeightPercentage)
 {
 	LineHeightPercentage = InLineHeightPercentage;
-	InvalidateLayout();
+	Invalidate(EInvalidateWidget::Layout);
 }
 
 void STextBlock::SetMargin(const TAttribute<FMargin>& InMargin)
 {
 	Margin = InMargin;
-	InvalidateLayout();
+	Invalidate(EInvalidateWidget::Layout);
 }
 
 void STextBlock::SetJustification(const TAttribute<ETextJustify::Type>& InJustification)
 {
 	Justification = InJustification;
-	InvalidateLayout();
+	Invalidate(EInvalidateWidget::Layout);
 }
 
 FTextBlockStyle STextBlock::GetComputedTextStyle() const
diff --git a/Engine/Source/Runtime/Slate/Private/Widgets/Views/SHeaderRow.cpp b/Engine/Source/Runtime/Slate/Private/Widgets/Views/SHeaderRow.cpp
--- a/Engine/Source/Runtime/Slate/Private/Widgets/Views/SHeaderRow.cpp
+++ b/Engine/Source/Runtime/Slate/Private/Widgets/Views/SHeaderRow.cpp
@@ -34,7 +34,7 @@ public:
 	{
 		check(InArgs._Style);
 
-		SWidget::Construct( InArgs._ToolTipText, InArgs._ToolTip, InArgs._Cursor, InArgs._IsEnabled, InArgs._Visibility, InArgs._RenderTransform, InArgs._RenderTransformPivot, InArgs._Tag, InArgs.MetaData );
+		SWidget::Construct( InArgs._ToolTipText, InArgs._ToolTip, InArgs._Cursor, InArgs._IsEnabled, InArgs._Visibility, InArgs._RenderTransform, InArgs._RenderTransformPivot, InArgs._Tag, InArgs._ForceVolatile, InArgs.MetaData );
 
 		Style = InArgs._Style;
 		ColumnId = Column.ColumnId;
@@ -396,7 +396,7 @@ void SHeaderRow::Construct( const FArguments& InArgs )
 {
 	check(InArgs._Style);
 
-	SWidget::Construct( InArgs._ToolTipText, InArgs._ToolTip, InArgs._Cursor, InArgs._IsEnabled, InArgs._Visibility, InArgs._RenderTransform, InArgs._RenderTransformPivot, InArgs._Tag, InArgs.MetaData );
+	SWidget::Construct( InArgs._ToolTipText, InArgs._ToolTip, InArgs._Cursor, InArgs._IsEnabled, InArgs._Visibility, InArgs._RenderTransform, InArgs._RenderTransformPivot, InArgs._Tag, InArgs._ForceVolatile, InArgs.MetaData );
 
 	ScrollBarThickness = FVector2D::ZeroVector;
 	ScrollBarVisibility = EVisibility::Collapsed;
diff --git a/Engine/Source/Runtime/Slate/Public/Widgets/SInvalidationPanel.h b/Engine/Source/Runtime/Slate/Public/Widgets/SInvalidationPanel.h
--- a/Engine/Source/Runtime/Slate/Public/Widgets/SInvalidationPanel.h
+++ b/Engine/Source/Runtime/Slate/Public/Widgets/SInvalidationPanel.h
@@ -2,31 +2,34 @@
 
 #pragma once
 
+#include "Input/HittestGrid.h"
+
 class SLATE_API SInvalidationPanel : public SCompoundWidget, public ILayoutCache
 {
 public:
 	SLATE_BEGIN_ARGS( SInvalidationPanel )
 	{
-		_Visibility = EVisibility::HitTestInvisible;
+		_Visibility = EVisibility::SelfHitTestInvisible;
 	}
 		SLATE_DEFAULT_SLOT( FArguments, Content )
 	SLATE_END_ARGS()
 
 	void Construct( const FArguments& InArgs );
+	~SInvalidationPanel();
 
 	bool GetCanCache() const;
 	void SetCanCache(bool InCanCache);
 
-	void Invalidate();
+	void InvalidateCache();
 
 	// ILayoutCache overrides
 	virtual void InvalidateWidget(SWidget* InvalidateWidget) override;
+	virtual FCachedWidgetNode* CreateCacheNode() const override;
 	// End ILayoutCache
 
 public:
 
 	// SWidget overrides
-	virtual void OnArrangeChildren( const FGeometry& AllottedGeometry, FArrangedChildren& ArrangedChildren ) const override;
 	virtual int32 OnPaint( const FPaintArgs& Args, const FGeometry& AllottedGeometry, const FSlateRect& MyClippingRect, FSlateWindowElementList& OutDrawElements, int32 LayerId, const FWidgetStyle& InWidgetStyle, bool bParentEnabled ) const override;
 	virtual FChildren* GetChildren() override;
 	virtual void Tick( const FGeometry& AllottedGeometry, const double InCurrentTime, const float InDeltaTime );
@@ -40,9 +43,6 @@ public:
 	static void EnableInvalidationDebugging(bool bEnable);
 #endif
 
-protected:
-	void CacheDesiredSize(float LayoutScaleMultiplier) override;
-
 private:
 	FGeometry LastAllottedGeometry;
 
@@ -53,7 +53,10 @@ private:
 	mutable TSet<TWeakPtr<SWidget>> InvalidatorWidgets;
 #endif
 
+	mutable FCachedWidgetNode* RootCacheNode;
 	mutable TSharedPtr< FSlateWindowElementList > CachedWindowElements;
+	mutable TArray< FCachedWidgetNode* > NodePool;
+	mutable int32 LastUsedCachedNodeIndex;
 
 	mutable int32 CachedMaxChildLayer;
 	mutable bool bNeedsCaching;
diff --git a/Engine/Source/Runtime/SlateCore/Private/Layout/WidgetCaching.cpp b/Engine/Source/Runtime/SlateCore/Private/Layout/WidgetCaching.cpp
new file mode 100644
--- /dev/null
+++ b/Engine/Source/Runtime/SlateCore/Private/Layout/WidgetCaching.cpp
@@ -0,0 +1,37 @@
+// Copyright 1998-2015 Epic Games, Inc. All Rights Reserved.
+
+#include "SlateCorePrivatePCH.h"
+#include "WidgetCaching.h"
+
+// FCachedWidgetNode
+//-------------------------------------------------------------
+
+void FCachedWidgetNode::Initialize(const FPaintArgs& Args, TSharedRef<SWidget> InWidget, const FGeometry& InGeometry, const FSlateRect& InClippingRect)
+{
+	Widget = InWidget;
+	Geometry = InGeometry;
+	ClippingRect = InClippingRect;
+	WindowOffset = Args.GetWindowToDesktopTransform();
+	RecordedVisibility = Args.GetLastRecordedVisibility();
+
+	if ( RecordedVisibility.AreChildrenHitTestVisible() )
+	{
+		RecordedVisibility = InWidget->GetVisibility();
+	}
+
+	Children.Reset();
+}
+
+void FCachedWidgetNode::RecordHittestGeometry(FHittestGrid& Grid, int32 LastHittestIndex)
+{
+	if ( RecordedVisibility.AreChildrenHitTestVisible() )
+	{
+		const int32 RecordedHittestIndex = Grid.InsertWidget(LastHittestIndex, RecordedVisibility, FArrangedWidget(Widget.Pin()->AsShared(), Geometry), WindowOffset, ClippingRect);
+
+		int32 ChildCount = Children.Num();
+		for ( int32 i = 0; i < ChildCount; i++ )
+		{
+			Children[i]->RecordHittestGeometry(Grid, RecordedHittestIndex);
+		}
+	}
+}
diff --git a/Engine/Source/Runtime/SlateCore/Private/Rendering/DrawElements.cpp b/Engine/Source/Runtime/SlateCore/Private/Rendering/DrawElements.cpp
--- a/Engine/Source/Runtime/SlateCore/Private/Rendering/DrawElements.cpp
+++ b/Engine/Source/Runtime/SlateCore/Private/Rendering/DrawElements.cpp
@@ -363,7 +363,7 @@ int32 FSlateWindowElementList::PaintDeferred(int32 LayerId)
 
 FSlateWindowElementList::FVolatilePaint::FVolatilePaint(const TSharedRef<const SWidget>& InWidgetToPaint, const FPaintArgs& InArgs, const FGeometry InAllottedGeometry, const FSlateRect InMyClippingRect, int32 InLayerId, const FWidgetStyle& InWidgetStyle, bool InParentEnabled)
 	: WidgetToPaintPtr(InWidgetToPaint)
-	, Args(InArgs.EnableCaching(InArgs.GetLayoutCache().Pin(), false, true))
+	, Args(InArgs.EnableCaching(InArgs.GetLayoutCache(), nullptr, false, true))
 	, AllottedGeometry(InAllottedGeometry)
 	, MyClippingRect(InMyClippingRect)
 	, LayerId(InLayerId)
diff --git a/Engine/Source/Runtime/SlateCore/Private/Types/PaintArgs.cpp b/Engine/Source/Runtime/SlateCore/Private/Types/PaintArgs.cpp
--- a/Engine/Source/Runtime/SlateCore/Private/Types/PaintArgs.cpp
+++ b/Engine/Source/Runtime/SlateCore/Private/Types/PaintArgs.cpp
@@ -2,6 +2,7 @@
 
 #include "SlateCorePrivatePCH.h"
 #include "HittestGrid.h"
+#include "WidgetCaching.h"
 
 FPaintArgs::FPaintArgs( const TSharedRef<SWidget>& Parent, FHittestGrid& InHittestGrid, FVector2D InWindowOffset, double InCurrentTime, float InDeltaTime )
 : ParentPtr(Parent)
@@ -12,6 +13,7 @@ FPaintArgs::FPaintArgs( const TSharedRef<SWidget>& Parent, FHittestGrid& InHitte
 , CurrentTime(InCurrentTime)
 , DeltaTime(InDeltaTime)
 , bIsCaching(false)
+, ParentCacheNode(nullptr)
 {
 }
 
@@ -21,16 +23,18 @@ FPaintArgs FPaintArgs::WithNewParent( const SWidget* Parent ) const
 	Args.LastHittestIndex = this->LastHittestIndex;
 	Args.LastRecordedVisibility = this->LastRecordedVisibility;
 	Args.LayoutCache = this->LayoutCache;
+	Args.ParentCacheNode = this->ParentCacheNode;
 	Args.bIsCaching = this->bIsCaching;
 	Args.bIsVolatilityPass = this->bIsVolatilityPass;
 
 	return Args;
 }
 
-FPaintArgs FPaintArgs::EnableCaching(const TSharedPtr<ILayoutCache>& InLayoutCache, bool bEnableCaching, bool bEnableVolatile) const
+FPaintArgs FPaintArgs::EnableCaching(const TSharedPtr<ILayoutCache>& InLayoutCache, FCachedWidgetNode* InParentCacheNode, bool bEnableCaching, bool bEnableVolatile) const
 {
 	FPaintArgs UpdatedArgs(*this);
 	UpdatedArgs.LayoutCache = InLayoutCache;
+	UpdatedArgs.ParentCacheNode = InParentCacheNode;
 	UpdatedArgs.bIsCaching = bEnableCaching;
 	UpdatedArgs.bIsVolatilityPass = bEnableVolatile;
 
@@ -43,10 +47,21 @@ FPaintArgs FPaintArgs::RecordHittestGeometry(const SWidget* Widget, const FGeome
 
 	if ( LastRecordedVisibility.AreChildrenHitTestVisible() )
 	{
-		const EVisibility RecordedVisibility = Widget->GetVisibility();
-		const int32 RecordedHittestIndex = Grid.InsertWidget(LastHittestIndex, RecordedVisibility, FArrangedWidget(const_cast<SWidget*>( Widget )->AsShared(), WidgetGeometry), WindowOffset, InClippingRect);
-		UpdatedArgs.LastHittestIndex = RecordedHittestIndex;
-		UpdatedArgs.LastRecordedVisibility = RecordedVisibility;
+		if ( bIsCaching )
+		{
+			FCachedWidgetNode* CacheNode = LayoutCache->CreateCacheNode();
+			CacheNode->Initialize(UpdatedArgs, const_cast<SWidget*>( Widget )->AsShared(), WidgetGeometry, InClippingRect);
+			UpdatedArgs.ParentCacheNode->Children.Add(CacheNode);
+
+			UpdatedArgs.ParentCacheNode = CacheNode;
+		}
+		else
+		{
+			const EVisibility RecordedVisibility = Widget->GetVisibility();
+			const int32 RecordedHittestIndex = Grid.InsertWidget(LastHittestIndex, RecordedVisibility, FArrangedWidget(const_cast<SWidget*>( Widget )->AsShared(), WidgetGeometry), WindowOffset, InClippingRect);
+			UpdatedArgs.LastHittestIndex = RecordedHittestIndex;
+			UpdatedArgs.LastRecordedVisibility = RecordedVisibility;
+		}
 	}
 	else
 	{
@@ -58,6 +73,6 @@ FPaintArgs FPaintArgs::RecordHittestGeometry(const SWidget* Widget, const FGeome
 
 FPaintArgs FPaintArgs::InsertCustomHitTestPath( TSharedRef<ICustomHitTestPath> CustomHitTestPath, int32 InLastHittestIndex ) const
 {
-	const_cast<FHittestGrid&>(Grid).InsertCustomHitTestPath( CustomHitTestPath, InLastHittestIndex );
+	const_cast<FHittestGrid&>(Grid).InsertCustomHitTestPath(CustomHitTestPath, InLastHittestIndex);
 	return *this;
 }
diff --git a/Engine/Source/Runtime/SlateCore/Private/Widgets/SNullWidget.cpp b/Engine/Source/Runtime/SlateCore/Private/Widgets/SNullWidget.cpp
--- a/Engine/Source/Runtime/SlateCore/Private/Widgets/SNullWidget.cpp
+++ b/Engine/Source/Runtime/SlateCore/Private/Widgets/SNullWidget.cpp
@@ -27,6 +27,7 @@ public:
 			InArgs._RenderTransform,
 			InArgs._RenderTransformPivot,
 			InArgs._Tag,
+			InArgs._ForceVolatile,
 			InArgs.MetaData
 		);
 	}
diff --git a/Engine/Source/Runtime/SlateCore/Private/Widgets/SWidget.cpp b/Engine/Source/Runtime/SlateCore/Private/Widgets/SWidget.cpp
--- a/Engine/Source/Runtime/SlateCore/Private/Widgets/SWidget.cpp
+++ b/Engine/Source/Runtime/SlateCore/Private/Widgets/SWidget.cpp
@@ -43,6 +43,7 @@ void SWidget::Construct(
 	const TAttribute<TOptional<FSlateRenderTransform>>& InTransform,
 	const TAttribute<FVector2D>& InTransformPivot,
 	const FName& InTag,
+	const bool InForceVolatile,
 	const TArray<TSharedRef<ISlateMetaData>>& InMetaData
 )
 {
@@ -68,6 +69,7 @@ void SWidget::Construct(
 	RenderTransform = InTransform;
 	RenderTransformPivot = InTransformPivot;
 	Tag = InTag;
+	bForceVolatile = InForceVolatile;
 	MetaData = InMetaData;
 }
 
@@ -694,9 +696,14 @@ int32 SWidget::Paint(const FPaintArgs& Args, const FGeometry& AllottedGeometry,
 		MutableThis->Tick( TickGeometry, Args.GetCurrentTime(), Args.GetDeltaTime() );
 	}
 
-	const FPaintArgs UpdatedArgs = Args.RecordHittestGeometry( this, AllottedGeometry, MyClippingRect );
+	// Record hit test geometry, but only if we're not caching.
+	const FPaintArgs UpdatedArgs = Args.RecordHittestGeometry(this, AllottedGeometry, MyClippingRect);
+
+	// Paint the geometry of this widget.
 	int32 NewLayerID = OnPaint(UpdatedArgs, AllottedGeometry, MyClippingRect, OutDrawElements, LayerId, InWidgetStyle, bParentEnabled);
 
+	// HACK Doing this on every widget is expensive to call this virtual on everyone, do we really need this?
+	// it's not even all that useful on a shipped game you'd never use the built in keyboard focusing system.
 	if (SupportsKeyboardFocus())
 	{
 		bool bShowUserFocus = FSlateApplicationBase::Get().ShowUserFocus(SharedThis(this));
diff --git a/Engine/Source/Runtime/SlateCore/Public/Layout/WidgetCaching.h b/Engine/Source/Runtime/SlateCore/Public/Layout/WidgetCaching.h
new file mode 100644
--- /dev/null
+++ b/Engine/Source/Runtime/SlateCore/Public/Layout/WidgetCaching.h
@@ -0,0 +1,30 @@
+// Copyright 1998-2015 Epic Games, Inc. All Rights Reserved.
+
+#pragma once
+
+/**
+ * Used by the implementation of ILayoutCache to store additional information about every widget drawn,
+ * so that the hit test grid can be replicated during subsequent frames without the widgets actually doing 
+ * any work.
+ */
+class SLATECORE_API FCachedWidgetNode
+{
+public:
+	void Initialize(const FPaintArgs& Args, TSharedRef<SWidget> InWidget, const FGeometry& InGeometry, const FSlateRect& InClippingRect);
+
+	void RecordHittestGeometry(FHittestGrid& Grid, int32 LastHittestIndex);
+
+public:
+
+	TArray< FCachedWidgetNode*, TInlineAllocator<4> > Children;
+
+	TWeakPtr<SWidget> Widget;
+	FGeometry Geometry;
+	FSlateRect ClippingRect;
+	FVector2D WindowOffset;
+
+	EVisibility RecordedVisibility;
+
+	// TODO Do we want to cache information about the elements drawn by each widget?
+	//TArray< FSlateWindowElementList > Elements;
+};
diff --git a/Engine/Source/Runtime/SlateCore/Public/Types/PaintArgs.h b/Engine/Source/Runtime/SlateCore/Public/Types/PaintArgs.h
--- a/Engine/Source/Runtime/SlateCore/Public/Types/PaintArgs.h
+++ b/Engine/Source/Runtime/SlateCore/Public/Types/PaintArgs.h
@@ -10,6 +10,8 @@ class FHittestGrid;
 class FSlateRect;
 class ICustomHitTestPath;
 class ILayoutCache;
+class FCachedWidgetNode;
+
 
 /**
  * SWidget::OnPaint and SWidget::Paint use FPaintArgs as their
@@ -21,17 +23,19 @@ class SLATECORE_API FPaintArgs
 public:
 	FPaintArgs( const TSharedRef<SWidget>& Parent, FHittestGrid& InHittestGrid, FVector2D InWindowOffset, double InCurrentTime, float InDeltaTime );
 	FPaintArgs WithNewParent( const SWidget* Parent ) const;
-	FPaintArgs EnableCaching(const TSharedPtr<ILayoutCache>& LayoutCache, bool bEnableCaching, bool bEnableVolatility) const;
+	FPaintArgs EnableCaching(const TSharedPtr<ILayoutCache>& InLayoutCache, FCachedWidgetNode* InParentCacheNode, bool bEnableCaching, bool bEnableVolatility) const;
 	FPaintArgs RecordHittestGeometry(const SWidget* Widget, const FGeometry& WidgetGeometry, const FSlateRect& InClippingRect) const;
 	FPaintArgs InsertCustomHitTestPath( TSharedRef<ICustomHitTestPath> CustomHitTestPath, int32 HitTestIndex ) const;
 
+	FHittestGrid& GetGrid() const { return Grid; }
 	int32 GetLastHitTestIndex() const { return LastHittestIndex; }
+	EVisibility GetLastRecordedVisibility() const { return LastRecordedVisibility; }
 	FVector2D GetWindowToDesktopTransform() const { return WindowOffset; }
 	double GetCurrentTime() const { return CurrentTime; }
 	float GetDeltaTime() const { return DeltaTime; }
 	bool IsCaching() const { return bIsCaching; }
 	bool IsVolatilityPass() const { return bIsVolatilityPass; }
-	TWeakPtr<ILayoutCache> GetLayoutCache() const { return LayoutCache; }
+	TSharedPtr<ILayoutCache> GetLayoutCache() const { return LayoutCache; }
 
 private:
 	const TSharedRef<SWidget>& ParentPtr;
@@ -43,5 +47,6 @@ private:
 	float DeltaTime;
 	bool bIsCaching;
 	bool bIsVolatilityPass;
-	TWeakPtr<ILayoutCache> LayoutCache;
+	TSharedPtr<ILayoutCache> LayoutCache;
+	FCachedWidgetNode* ParentCacheNode;
 };
\ No newline at end of file
diff --git a/Engine/Source/Runtime/SlateCore/Public/Widgets/DeclarativeSyntaxSupport.h b/Engine/Source/Runtime/SlateCore/Public/Widgets/DeclarativeSyntaxSupport.h
--- a/Engine/Source/Runtime/SlateCore/Public/Widgets/DeclarativeSyntaxSupport.h
+++ b/Engine/Source/Runtime/SlateCore/Public/Widgets/DeclarativeSyntaxSupport.h
@@ -1013,6 +1013,7 @@ struct TSlateBaseNamedArgs
 	, _Visibility( EVisibility::Visible )
 	, _RenderTransform( )
 	, _RenderTransformPivot( FVector2D::ZeroVector )
+	, _ForceVolatile( false )
 	{
 	}
 
@@ -1053,6 +1054,7 @@ struct TSlateBaseNamedArgs
 	SLATE_ATTRIBUTE( TOptional<FSlateRenderTransform>, RenderTransform )
 	SLATE_ATTRIBUTE( FVector2D, RenderTransformPivot )
 	SLATE_ARGUMENT( FName, Tag )
+	SLATE_ARGUMENT( bool, ForceVolatile )
 
 	TArray<TSharedRef<ISlateMetaData>> MetaData;
 };
@@ -1070,6 +1072,7 @@ namespace RequiredArgs
 		{
 			// YOUR WIDGET MUST IMPLEMENT void Construct(const FArguments& InArgs)
 			OnWidget->Construct(WithNamedArgs);
+			OnWidget->CacheVolatility();
 		}
 	};
 
@@ -1086,6 +1089,7 @@ namespace RequiredArgs
 		{
 			// YOUR WIDGET MUST IMPLEMENT void Construct(const FArguments& InArgs)
 			OnWidget->Construct(WithNamedArgs, Forward<Arg0Type>(Arg0));
+			OnWidget->CacheVolatility();
 		}
 
 		Arg0Type& Arg0;
@@ -1105,6 +1109,7 @@ namespace RequiredArgs
 		{
 			// YOUR WIDGET MUST IMPLEMENT Construct(const FArguments& InArgs)
 			OnWidget->Construct(WithNamedArgs, Forward<Arg0Type>(Arg0), Forward<Arg1Type>(Arg1));
+			OnWidget->CacheVolatility();
 		}
 
 		Arg0Type& Arg0;
@@ -1126,6 +1131,7 @@ namespace RequiredArgs
 		{
 			// YOUR WIDGET MUST IMPLEMENT Construct(const FArguments& InArgs)
 			OnWidget->Construct(WithNamedArgs, Forward<Arg0Type>(Arg0), Forward<Arg1Type>(Arg1), Forward<Arg2Type>(Arg2));
+			OnWidget->CacheVolatility();
 		}
 
 		Arg0Type& Arg0;
@@ -1149,6 +1155,7 @@ namespace RequiredArgs
 		{
 			// YOUR WIDGET MUST IMPLEMENT Construct(const FArguments& InArgs)
 			OnWidget->Construct(WithNamedArgs, Forward<Arg0Type>(Arg0), Forward<Arg1Type>(Arg1), Forward<Arg2Type>(Arg2), Forward<Arg3Type>(Arg3));
+			OnWidget->CacheVolatility();
 		}
 
 		Arg0Type& Arg0;
@@ -1174,6 +1181,7 @@ namespace RequiredArgs
 		{
 			// YOUR WIDGET MUST IMPLEMENT Construct(const FArguments& InArgs)
 			OnWidget->Construct(WithNamedArgs, Forward<Arg0Type>(Arg0), Forward<Arg1Type>(Arg1), Forward<Arg2Type>(Arg2), Forward<Arg3Type>(Arg3), Forward<Arg4Type>(Arg4));
+			OnWidget->CacheVolatility();
 		}
 
 		Arg0Type& Arg0;
@@ -1316,6 +1324,7 @@ struct TDecl
 			InArgs._RenderTransform,
 			InArgs._RenderTransformPivot,
 			InArgs._Tag,
+			InArgs._ForceVolatile,
 			InArgs.MetaData );
 
 		_RequiredArgs.CallConstruct(_Widget, InArgs);
diff --git a/Engine/Source/Runtime/SlateCore/Public/Widgets/SWidget.h b/Engine/Source/Runtime/SlateCore/Public/Widgets/SWidget.h
--- a/Engine/Source/Runtime/SlateCore/Public/Widgets/SWidget.h
+++ b/Engine/Source/Runtime/SlateCore/Public/Widgets/SWidget.h
@@ -67,10 +67,24 @@ public:
 };
 
 
+/**
+ * 
+ */
+enum class EInvalidateWidget
+{
+	Layout,
+	LayoutAndVolatility
+};
+
+
+/**
+ * 
+ */
 class ILayoutCache
 {
 public:
 	virtual void InvalidateWidget(class SWidget* InvalidateWidget) = 0;
+	virtual FCachedWidgetNode* CreateCacheNode() const = 0;
 };
 
 
@@ -119,6 +133,7 @@ public:
 		const TAttribute<TOptional<FSlateRenderTransform>>& InTransform,
 		const TAttribute<FVector2D>& InTransformPivot,
 		const FName& InTag,
+		const bool InForceVolatile,
 		const TArray<TSharedRef<ISlateMetaData>>& InMetaData);
 
 	void SWidgetConstruct( const TAttribute<FText> & InToolTipText ,
@@ -129,9 +144,10 @@ public:
 		const TAttribute<TOptional<FSlateRenderTransform>>& InTransform,
 		const TAttribute<FVector2D>& InTransformPivot,
 		const FName& InTag,
+		const bool InForceVolatile,
 		const TArray<TSharedRef<ISlateMetaData>>& InMetaData)
 	{
-		Construct(InToolTipText, InToolTip, InCursor, InEnabledState, InVisibility, InTransform, InTransformPivot, InTag, InMetaData);
+		Construct(InToolTipText, InToolTip, InCursor, InEnabledState, InVisibility, InTransform, InTransformPivot, InTag, InForceVolatile, InMetaData);
 	}
 
 	//
@@ -640,7 +656,7 @@ public:
 	void SetEnabled( const TAttribute<bool>& InEnabledState )
 	{
 		EnabledState = InEnabledState;
-		InvalidateLayout();
+		Invalidate(EInvalidateWidget::LayoutAndVolatility);
 	}
 
 	/** @return Whether or not this widget is enabled */
@@ -693,7 +709,7 @@ public:
 		if ( !Visibility.IdenticalTo(InVisibility) )
 		{
 			Visibility = InVisibility;
-			InvalidateLayout();
+			Invalidate(EInvalidateWidget::LayoutAndVolatility);
 		}
 	}
 
@@ -720,30 +736,59 @@ public:
 	FORCEINLINE void ForceVolatile(bool bForce)
 	{
 		bForceVolatile = bForce;
-		InvalidateLayout();
+		Invalidate(EInvalidateWidget::LayoutAndVolatility);
 	}
 
 	/**
 	 * Invalidates the widget from the view of a layout caching widget that may own this widget.
 	 * will force the owning widget to redraw and cache children on the next paint pass.
 	 */
-	FORCEINLINE void InvalidateLayout()
+	FORCEINLINE void Invalidate(EInvalidateWidget Invalidate)
 	{
 		const bool bWasVolatile = IsVolatileIndirectly() || IsVolatile();
+		const bool bVolatilityChanged = Invalidate == EInvalidateWidget::LayoutAndVolatility ? Advanced_InvalidateVolatility() : false;
+
+		if ( bWasVolatile == false || bVolatilityChanged )
+		{
+			Advanced_ForceInvalidateLayout();
+		}
+	}
+
+	/**
+	 * Recalculates volatility of the widget and caches the result.  Should be called any time 
+	 * anything examined by your implementation of ComputeVolatility is changed.
+	 */
+	FORCEINLINE void CacheVolatility()
+	{
+		bCachedVolatile = bForceVolatile || ComputeVolatility();
+	}
+
+protected:
+
+	/**
+	 * Recalculates and caches volatility and returns 'true' if the volatility changed.
+	 */
+	FORCEINLINE bool Advanced_InvalidateVolatility()
+	{
 		const bool bWasDirectlyVolatile = IsVolatile();
 		CacheVolatility();
-		const bool bVolatilityChanged = bWasDirectlyVolatile != IsVolatile();
+		return bWasDirectlyVolatile != IsVolatile();
+	}
 
-		if ( bWasVolatile == false || bVolatilityChanged )
+	/**
+	 * Forces invalidation, doesn't check volatility.
+	 */
+	FORCEINLINE void Advanced_ForceInvalidateLayout()
+	{
+		TSharedPtr<ILayoutCache> Cache = LayoutCache.Pin();
+		if ( Cache.IsValid() )
 		{
-			TSharedPtr<ILayoutCache> Cache = LayoutCache.Pin();
-			if ( Cache.IsValid() )
-			{
-				Cache->InvalidateWidget(this);
-			}
+			Cache->InvalidateWidget(this);
 		}
 	}
 
+public:
+
 	/** @return the render transform of the widget. */
 	FORCEINLINE const TOptional<FSlateRenderTransform>& GetRenderTransform() const
 	{
@@ -754,7 +799,7 @@ public:
 	FORCEINLINE void SetRenderTransform(TAttribute<TOptional<FSlateRenderTransform>> InTransform)
 	{
 		RenderTransform = InTransform;
-		InvalidateLayout();
+		Invalidate(EInvalidateWidget::LayoutAndVolatility);
 	}
 
 	/** @return the pivot point of the render transform. */
@@ -929,22 +974,13 @@ protected:
 	/** @return a brush to draw focus, nullptr if no focus drawing is desired */
 	virtual const FSlateBrush* GetFocusBrush() const;
 
-	/**
-	 * Recalculates volatility of the widget and caches the result.  Should be called any time 
-	 * anything examined by your implementation of ComputeVolatility is changed.
-	 */
-	FORCEINLINE void CacheVolatility()
-	{
-		bCachedVolatile = bForceVolatile || ComputeVolatility();
-	}
-
 	/**
 	 * Recomputes the volatility of the widget.  If you have additional state you automatically want to make
 	 * the widget volatile, you should sample that information here.
 	 */
 	virtual bool ComputeVolatility() const
 	{
-		return Visibility.IsBound();
+		return Visibility.IsBound() || EnabledState.IsBound() || RenderTransform.IsBound();
 	}
 
 private:
diff --git a/Engine/Source/Runtime/UMG/Private/Components/InvalidationBox.cpp b/Engine/Source/Runtime/UMG/Private/Components/InvalidationBox.cpp
--- a/Engine/Source/Runtime/UMG/Private/Components/InvalidationBox.cpp
+++ b/Engine/Source/Runtime/UMG/Private/Components/InvalidationBox.cpp
@@ -57,11 +57,11 @@ void UInvalidationBox::OnSlotRemoved(UPanelSlot* Slot)
 	}
 }
 
-void UInvalidationBox::Invalidate()
+void UInvalidationBox::InvalidateCache()
 {
 	if ( MyInvalidationPanel.IsValid() )
 	{
-		return MyInvalidationPanel->Invalidate();
+		return MyInvalidationPanel->InvalidateCache();
 	}
 }
 
diff --git a/Engine/Source/Runtime/UMG/Private/Components/Widget.cpp b/Engine/Source/Runtime/UMG/Private/Components/Widget.cpp
--- a/Engine/Source/Runtime/UMG/Private/Components/Widget.cpp
+++ b/Engine/Source/Runtime/UMG/Private/Components/Widget.cpp
@@ -408,12 +408,12 @@ void UWidget::ForceLayoutPrepass()
 	}
 }
 
-void UWidget::InvalidateLayout()
+void UWidget::InvalidateLayoutAndVolatility()
 {
 	TSharedPtr<SWidget> SafeWidget = GetCachedWidget();
 	if ( SafeWidget.IsValid() )
 	{
-		SafeWidget->InvalidateLayout();
+		SafeWidget->Invalidate(EInvalidateWidget::LayoutAndVolatility);
 	}
 }
 
diff --git a/Engine/Source/Runtime/UMG/Private/UserWidget.cpp b/Engine/Source/Runtime/UMG/Private/UserWidget.cpp
--- a/Engine/Source/Runtime/UMG/Private/UserWidget.cpp
+++ b/Engine/Source/Runtime/UMG/Private/UserWidget.cpp
@@ -225,7 +225,7 @@ void UUserWidget::PlayAnimation( const UWidgetAnimation* InAnimation, float Star
 		TSharedPtr<SWidget> CachedWidget = GetCachedWidget();
 		if ( CachedWidget.IsValid() )
 		{
-			CachedWidget->InvalidateLayout();
+			CachedWidget->Invalidate(EInvalidateWidget::LayoutAndVolatility);
 		}
 
 		OnAnimationStarted( InAnimation );
@@ -755,7 +755,7 @@ void UUserWidget::TickActionsAndAnimation(const FGeometry& MyGeometry, float InD
 		TSharedPtr<SWidget> CachedWidget = GetCachedWidget();
 		if ( CachedWidget.IsValid() )
 		{
-			CachedWidget->InvalidateLayout();
+			CachedWidget->Invalidate(EInvalidateWidget::LayoutAndVolatility);
 		}
 	}
 
diff --git a/Engine/Source/Runtime/UMG/Public/Components/InvalidationBox.h b/Engine/Source/Runtime/UMG/Public/Components/InvalidationBox.h
--- a/Engine/Source/Runtime/UMG/Public/Components/InvalidationBox.h
+++ b/Engine/Source/Runtime/UMG/Public/Components/InvalidationBox.h
@@ -19,8 +19,8 @@ public:
 	/**
 	 * 
 	 */
-	UFUNCTION(BlueprintCallable, Category="Invalidate")
-	void Invalidate();
+	UFUNCTION(BlueprintCallable, Category="Invalidation Box")
+	void InvalidateCache();
 
 	virtual void ReleaseSlateResources(bool bReleaseChildren) override;
 
diff --git a/Engine/Source/Runtime/UMG/Public/Components/Widget.h b/Engine/Source/Runtime/UMG/Public/Components/Widget.h
--- a/Engine/Source/Runtime/UMG/Public/Components/Widget.h
+++ b/Engine/Source/Runtime/UMG/Public/Components/Widget.h
@@ -339,7 +339,7 @@ public:
 	 * will force the owning widget to redraw and cache children on the next paint pass.
 	 */
 	UFUNCTION(BlueprintCallable, Category="Widget")
-	void InvalidateLayout();
+	void InvalidateLayoutAndVolatility();
 
 	/**
 	 * Gets the widgets desired size.
