Fixed non-unity build error

diff --git a/Engine/Source/Runtime/SlateCore/Private/Layout/WidgetCaching.cpp b/Engine/Source/Runtime/SlateCore/Private/Layout/WidgetCaching.cpp
--- a/Engine/Source/Runtime/SlateCore/Private/Layout/WidgetCaching.cpp
+++ b/Engine/Source/Runtime/SlateCore/Private/Layout/WidgetCaching.cpp
@@ -2,6 +2,7 @@
 
 #include "SlateCorePrivatePCH.h"
 #include "WidgetCaching.h"
+#include "HittestGrid.h"
 
 // FCachedWidgetNode
 //-------------------------------------------------------------
