Slate - Custom Draw elements no longer are attempted to be early clipped using the known scissor rect before going through the element batcher, since their clip rect info isn't valid we just ignore doing this step for those elements.

diff --git a/Engine/Source/Runtime/SlateCore/Private/Rendering/ElementBatcher.cpp b/Engine/Source/Runtime/SlateCore/Private/Rendering/ElementBatcher.cpp
--- a/Engine/Source/Runtime/SlateCore/Private/Rendering/ElementBatcher.cpp
+++ b/Engine/Source/Runtime/SlateCore/Private/Rendering/ElementBatcher.cpp
@@ -112,7 +112,7 @@ void FSlateElementBatcher::AddElements( FSlateWindowElementList& WindowElementLi
 		if ( !bIsFullyClipped )
 		{
 			// do this check in here do we can short circuit above more quickly, but still name this variable so its meaning is clear.
-			const bool bIsScissored = DrawElement.GetScissorRect().IsSet() && !DrawElement.GetScissorRect().GetValue().DoesIntersect(FShortRect(InClippingRect));
+			const bool bIsScissored = DrawElement.GetScissorRect().IsSet() && DrawElement.GetElementType() != FSlateDrawElement::ET_Custom && !DrawElement.GetScissorRect().GetValue().DoesIntersect(FShortRect(InClippingRect));
 			// scissor rects are sort of a low level hack, so no one konws to clip against them. Instead we do it here to make sure the element is not actually rendered.
 			if (!bIsScissored)
 			{
