Added a safe way to delete a dynamic slate brush that may be in use.  Fixed up use cases of unsafe deletion

Fixes https://jira.ol.epicgames.net/browse/UE-18531

diff --git a/Engine/Source/Developer/StandaloneRenderer/Private/OpenGL/SlateOpenGLRenderer.cpp b/Engine/Source/Developer/StandaloneRenderer/Private/OpenGL/SlateOpenGLRenderer.cpp
--- a/Engine/Source/Developer/StandaloneRenderer/Private/OpenGL/SlateOpenGLRenderer.cpp
+++ b/Engine/Source/Developer/StandaloneRenderer/Private/OpenGL/SlateOpenGLRenderer.cpp
@@ -148,6 +148,9 @@ void FSlateOpenGLRenderer::DrawWindows( FSlateDrawBuffer& InWindowDrawBuffer )
 	}
 
 	FontCache->ConditionalFlushCache();
+
+	// Safely release the references now that we are finished rendering with the dynamic brushes
+	DynamicBrushesToRemove.Empty();
 }
 
 
@@ -209,17 +212,20 @@ void FSlateOpenGLRenderer::UpdateFullscreenState( const TSharedRef<SWindow> InWi
 
 void FSlateOpenGLRenderer::ReleaseDynamicResource( const FSlateBrush& Brush )
 {
-	ensure( IsInGameThread() );
 	TextureManager->ReleaseDynamicTextureResource( Brush );
 }
 
 
 bool FSlateOpenGLRenderer::GenerateDynamicImageResource(FName ResourceName, uint32 Width, uint32 Height, const TArray< uint8 >& Bytes)
 {
-	ensure( IsInGameThread() );
 	return TextureManager->CreateDynamicTextureResource(ResourceName, Width, Height, Bytes) != NULL;
 }
 
+void FSlateOpenGLRenderer::RemoveDynamicBrushResource( TSharedPtr<FSlateDynamicImageBrush> BrushToRemove )
+{
+	DynamicBrushesToRemove.Add( BrushToRemove );
+}
+
 
 void FSlateOpenGLRenderer::LoadStyleResources( const ISlateStyle& InStyle )
 {
diff --git a/Engine/Source/Developer/StandaloneRenderer/Private/OpenGL/SlateOpenGLRenderer.h b/Engine/Source/Developer/StandaloneRenderer/Private/OpenGL/SlateOpenGLRenderer.h
--- a/Engine/Source/Developer/StandaloneRenderer/Private/OpenGL/SlateOpenGLRenderer.h
+++ b/Engine/Source/Developer/StandaloneRenderer/Private/OpenGL/SlateOpenGLRenderer.h
@@ -146,6 +146,7 @@ public:
 	virtual void RestoreSystemResolution(const TSharedRef<SWindow> InWindow) override {}
 	virtual void ReleaseDynamicResource( const FSlateBrush& Brush ) override;
 	virtual bool GenerateDynamicImageResource(FName ResourceName, uint32 Width, uint32 Height, const TArray< uint8 >& Bytes) override;
+	virtual void RemoveDynamicBrushResource( TSharedPtr<FSlateDynamicImageBrush> BrushToRemove ) override;
 	virtual void LoadStyleResources(const ISlateStyle& Style) override;
 	virtual FSlateUpdatableTexture* CreateUpdatableTexture(uint32 Width, uint32 Height) override;
 	virtual void ReleaseUpdatableTexture(FSlateUpdatableTexture* Texture) override;
@@ -176,4 +177,6 @@ private:
 	const ISlateStyle& Style;
 	/** Shared OpenGL context */
 	FSlateOpenGLContext SharedContext;
+	/** Dynamic image brushes to remove when safe */
+	TArray<TSharedPtr<FSlateDynamicImageBrush>> DynamicBrushesToRemove;
 };
diff --git a/Engine/Source/Developer/StandaloneRenderer/Private/Windows/D3D/SlateD3DRenderer.cpp b/Engine/Source/Developer/StandaloneRenderer/Private/Windows/D3D/SlateD3DRenderer.cpp
--- a/Engine/Source/Developer/StandaloneRenderer/Private/Windows/D3D/SlateD3DRenderer.cpp
+++ b/Engine/Source/Developer/StandaloneRenderer/Private/Windows/D3D/SlateD3DRenderer.cpp
@@ -235,16 +235,19 @@ void FSlateD3DRenderer::UpdateFullscreenState( const TSharedRef<SWindow> InWindo
 
 void FSlateD3DRenderer::ReleaseDynamicResource( const FSlateBrush& Brush )
 {
-	ensure( IsInGameThread() );
 	TextureManager->ReleaseDynamicTextureResource( Brush );
 }
 
 bool FSlateD3DRenderer::GenerateDynamicImageResource(FName ResourceName, uint32 Width, uint32 Height, const TArray< uint8 >& Bytes)
 {
-	ensure( IsInGameThread() );
 	return TextureManager->CreateDynamicTextureResource(ResourceName, Width, Height, Bytes) != NULL;
 }
 
+void FSlateD3DRenderer::RemoveDynamicBrushResource( TSharedPtr<FSlateDynamicImageBrush> BrushToRemove )
+{
+	DynamicBrushesToRemove.Add( BrushToRemove );
+}
+
 void FSlateD3DRenderer::Private_ResizeViewport( const TSharedRef<SWindow> InWindow, uint32 Width, uint32 Height, bool bFullscreen )
 {
 	FSlateD3DViewport* Viewport = WindowToViewportMap.Find( &InWindow.Get() );
@@ -399,6 +402,8 @@ void FSlateD3DRenderer::DrawWindows( FSlateDrawBuffer& InWindowDrawBuffer )
 	// flush the cache if needed
 	FontCache->ConditionalFlushCache();
 
+	// Safely release the references now that we are finished rendering with the dynamic brushes
+	DynamicBrushesToRemove.Empty();
 }
 
 void FSlateD3DRenderer::OnWindowDestroyed( const TSharedRef<SWindow>& InWindow )
diff --git a/Engine/Source/Developer/StandaloneRenderer/Private/Windows/D3D/SlateD3DRenderer.h b/Engine/Source/Developer/StandaloneRenderer/Private/Windows/D3D/SlateD3DRenderer.h
--- a/Engine/Source/Developer/StandaloneRenderer/Private/Windows/D3D/SlateD3DRenderer.h
+++ b/Engine/Source/Developer/StandaloneRenderer/Private/Windows/D3D/SlateD3DRenderer.h
@@ -55,6 +55,7 @@ public:
 	virtual void RestoreSystemResolution(const TSharedRef<SWindow> InWindow) override {}
 	virtual void ReleaseDynamicResource( const FSlateBrush& Brush ) override;
 	virtual bool GenerateDynamicImageResource(FName ResourceName, uint32 Width, uint32 Height, const TArray< uint8 >& Bytes) override;
+	virtual void RemoveDynamicBrushResource( TSharedPtr<FSlateDynamicImageBrush> BrushToRemove ) override;
 	virtual void LoadStyleResources( const ISlateStyle& Style ) override;
 	virtual FSlateUpdatableTexture* CreateUpdatableTexture(uint32 Width, uint32 Height) override;
 	virtual void ReleaseUpdatableTexture(FSlateUpdatableTexture* Texture) override;
@@ -74,6 +75,7 @@ private:
 	TSharedPtr<FSlateElementBatcher> ElementBatcher;
 	TSharedPtr<FSlateD3DTextureManager> TextureManager;
 	TSharedPtr<FSlateD3D11RenderingPolicy> RenderingPolicy;
+	TArray<TSharedPtr<FSlateDynamicImageBrush>> DynamicBrushesToRemove;
 };
 
 #endif
\ No newline at end of file
diff --git a/Engine/Source/Editor/IntroTutorials/Private/STutorialsBrowser.cpp b/Engine/Source/Editor/IntroTutorials/Private/STutorialsBrowser.cpp
--- a/Engine/Source/Editor/IntroTutorials/Private/STutorialsBrowser.cpp
+++ b/Engine/Source/Editor/IntroTutorials/Private/STutorialsBrowser.cpp
@@ -79,7 +79,12 @@ public:
 	}
 
 	virtual ~FTutorialListEntry_Category()
-	{}
+	{
+		if( DynamicBrush.IsValid() )
+		{
+			DynamicBrush->ReleaseResource();
+		}
+	}
 
 	virtual TSharedRef<ITableRow> OnGenerateTutorialRow(const TSharedRef<STableViewBase>& OwnerTable) const override
 	{
@@ -296,7 +301,12 @@ public:
 	}
 
 	virtual ~FTutorialListEntry_Tutorial()
-	{}
+	{
+		if( DynamicBrush.IsValid() )
+		{
+			DynamicBrush->ReleaseResource();
+		}
+	}
 
 	virtual TSharedRef<ITableRow> OnGenerateTutorialRow(const TSharedRef<STableViewBase>& OwnerTable) const override
 	{
diff --git a/Engine/Source/Runtime/Slate/Private/Framework/Text/SlateImageRun.cpp b/Engine/Source/Runtime/Slate/Private/Framework/Text/SlateImageRun.cpp
--- a/Engine/Source/Runtime/Slate/Private/Framework/Text/SlateImageRun.cpp
+++ b/Engine/Source/Runtime/Slate/Private/Framework/Text/SlateImageRun.cpp
@@ -80,6 +80,11 @@ FSlateImageRun::FSlateImageRun( const FRunInfo& InRunInfo, const TSharedRef< con
 	Image = DynamicBrush.Get();
 }
 
+FSlateImageRun::~FSlateImageRun()
+{
+	DynamicBrush->ReleaseResource();
+}
+
 const TArray< TSharedRef<SWidget> >& FSlateImageRun::GetChildren() 
 {
 	static TArray< TSharedRef<SWidget> > NoChildren;
diff --git a/Engine/Source/Runtime/Slate/Public/Framework/Text/SlateImageRun.h b/Engine/Source/Runtime/Slate/Public/Framework/Text/SlateImageRun.h
--- a/Engine/Source/Runtime/Slate/Public/Framework/Text/SlateImageRun.h
+++ b/Engine/Source/Runtime/Slate/Public/Framework/Text/SlateImageRun.h
@@ -14,7 +14,7 @@ public:
 
 public:
 
-	virtual ~FSlateImageRun() {}
+	virtual ~FSlateImageRun();
 
 	virtual FTextRange GetTextRange() const override;
 	virtual void SetTextRange( const FTextRange& Value ) override;
diff --git a/Engine/Source/Runtime/SlateCore/Private/Brushes/SlateDynamicImageBrush.cpp b/Engine/Source/Runtime/SlateCore/Private/Brushes/SlateDynamicImageBrush.cpp
--- a/Engine/Source/Runtime/SlateCore/Private/Brushes/SlateDynamicImageBrush.cpp
+++ b/Engine/Source/Runtime/SlateCore/Private/Brushes/SlateDynamicImageBrush.cpp
@@ -29,21 +29,39 @@ TSharedPtr<FSlateDynamicImageBrush> FSlateDynamicImageBrush::CreateWithImageData
 	return Brush;
 }
 
-FSlateDynamicImageBrush::~FSlateDynamicImageBrush( )
+void FSlateDynamicImageBrush::ReleaseResource()
+{
+	if( FSlateApplicationBase::IsInitialized() )
+	{
+		// Safely remove the brush we were using. 
+		FSlateApplicationBase::Get().GetRenderer()->RemoveDynamicBrushResource( AsShared() );
+	}	
+}
+
+void FSlateDynamicImageBrush::ReleaseResourceInternal()
 {
-	if (FSlateApplicationBase::IsInitialized())
+	if( bIsInitalized )
 	{
-		// Brush resource is no longer referenced by this object
-		if (ResourceObject && bRemoveResourceFromRootSet)
+		bIsInitalized = false;
+		if (FSlateApplicationBase::IsInitialized())
 		{
-			ResourceObject->RemoveFromRoot();
-		}
+			// Brush resource is no longer referenced by this object
+			if (ResourceObject && bRemoveResourceFromRootSet)
+			{
+				ResourceObject->RemoveFromRoot();
+			}
 
-		TSharedPtr<FSlateRenderer> Renderer = FSlateApplicationBase::Get().GetRenderer();
+			TSharedPtr<FSlateRenderer> Renderer = FSlateApplicationBase::Get().GetRenderer();
 
-		if (Renderer.IsValid())
-		{
-			FSlateApplicationBase::Get().GetRenderer()->ReleaseDynamicResource(*this);
+			if (Renderer.IsValid())
+			{
+				FSlateApplicationBase::Get().GetRenderer()->ReleaseDynamicResource(*this);
+			}
 		}
 	}
 }
+
+FSlateDynamicImageBrush::~FSlateDynamicImageBrush( )
+{
+	ReleaseResourceInternal();
+}
diff --git a/Engine/Source/Runtime/SlateCore/Public/Rendering/SlateRenderer.h b/Engine/Source/Runtime/SlateCore/Public/Rendering/SlateRenderer.h
--- a/Engine/Source/Runtime/SlateCore/Public/Rendering/SlateRenderer.h
+++ b/Engine/Source/Runtime/SlateCore/Public/Rendering/SlateRenderer.h
@@ -81,27 +81,6 @@ public:
 	 */
 	virtual void DrawWindows() {}
 	
-	/**
-	 * You must call this before calling CopyWindowsToVirtualScreenBuffer(), to setup the render targets first.
-	 * 
-	 * @param	bPrimaryWorkAreaOnly	True if we should capture only the primary monitor's work area, or false to capture the entire desktop spanning all monitors
-	 * @param	ScreenScaling	How much to downscale the desktop size
-	 * @param	LiveStreamingService	Optional pointer to a live streaming service this buffer needs to work with
-	 * @return	The virtual screen rectangle.  The size of this rectangle will be the size of the render target buffer.
-	 */
-	virtual FIntRect SetupVirtualScreenBuffer( const bool bPrimaryWorkAreaOnly, const float ScreenScaling, class ILiveStreamingService* LiveStreamingService) { return FIntRect(); }
-
-
-	/**
-	 * Copies all slate windows out to a buffer at half resolution with debug information
-	 * like the mouse cursor and any keypresses.
-	 */
-	virtual void CopyWindowsToVirtualScreenBuffer(const TArray<FString>& KeypressBuffer) {}
-	
-	/** Allows and disallows access to the crash tracker buffer data on the CPU */
-	virtual void MapVirtualScreenBuffer(void** OutImageData) {}
-	virtual void UnmapVirtualScreenBuffer() {}
-	
 	/** Callback that fires after Slate has rendered each window, each frame */
 	DECLARE_MULTICAST_DELEGATE_TwoParams( FOnSlateWindowRendered, SWindow&, void* );
 	FOnSlateWindowRendered& OnSlateWindowRendered() { return SlateWindowRendered; }
@@ -130,6 +109,20 @@ public:
 	 */
 	virtual bool GenerateDynamicImageResource( FName ResourceName, uint32 Width, uint32 Height, const TArray< uint8 >& Bytes ) { return false; }
 
+	/**
+	 * Queues a dynamic brush for removal when it is safe.  The brush is not immediately released but you should consider the brush destroyed and no longer usable
+	 *
+	 * @param BrushToRemove	The brush to queue for removal which is no longer valid to use
+	 */
+	virtual void RemoveDynamicBrushResource( TSharedPtr<FSlateDynamicImageBrush> BrushToRemove ) = 0;
+
+	/** 
+	 * Releases a specific resource.  
+	 * It is unlikely that you want to call this directly.  Use RemoveDynamicBrushResource instead
+	 */
+	virtual void ReleaseDynamicResource( const FSlateBrush& Brush ) = 0;
+
+
 	/** Called when a window is destroyed to give the renderer a chance to free resources */
 	virtual void OnWindowDestroyed( const TSharedRef<SWindow>& InWindow ) = 0;
 	
@@ -176,9 +169,6 @@ public:
 	 */
 	virtual void LoadStyleResources( const ISlateStyle& Style ) {}
 
-	/** Releases a specific resource */
-	virtual void ReleaseDynamicResource( const FSlateBrush& Brush ) = 0;
-
 	/**
 	 * Returns whether or not a viewport should be in  fullscreen
 	 *
@@ -235,6 +225,28 @@ public:
 	 */
 	virtual ISlateAtlasProvider* GetFontAtlasProvider();
 
+
+	/**
+	 * You must call this before calling CopyWindowsToVirtualScreenBuffer(), to setup the render targets first.
+	 * 
+	 * @param	bPrimaryWorkAreaOnly	True if we should capture only the primary monitor's work area, or false to capture the entire desktop spanning all monitors
+	 * @param	ScreenScaling	How much to downscale the desktop size
+	 * @param	LiveStreamingService	Optional pointer to a live streaming service this buffer needs to work with
+	 * @return	The virtual screen rectangle.  The size of this rectangle will be the size of the render target buffer.
+	 */
+	virtual FIntRect SetupVirtualScreenBuffer( const bool bPrimaryWorkAreaOnly, const float ScreenScaling, class ILiveStreamingService* LiveStreamingService) { return FIntRect(); }
+
+
+	/**
+	 * Copies all slate windows out to a buffer at half resolution with debug information
+	 * like the mouse cursor and any keypresses.
+	 */
+	virtual void CopyWindowsToVirtualScreenBuffer(const TArray<FString>& KeypressBuffer) {}
+	
+	/** Allows and disallows access to the crash tracker buffer data on the CPU */
+	virtual void MapVirtualScreenBuffer(void** OutImageData) {}
+	virtual void UnmapVirtualScreenBuffer() {}
+
 private:
 
 	// Non-copyable
diff --git a/Engine/Source/Runtime/SlateRHIRenderer/Private/SlateRHIRenderer.cpp b/Engine/Source/Runtime/SlateRHIRenderer/Private/SlateRHIRenderer.cpp
--- a/Engine/Source/Runtime/SlateRHIRenderer/Private/SlateRHIRenderer.cpp
+++ b/Engine/Source/Runtime/SlateRHIRenderer/Private/SlateRHIRenderer.cpp
@@ -219,7 +219,7 @@ FSlateDrawBuffer& FSlateRHIRenderer::GetDrawBuffer()
 		else
 		{
 			FlushCommands();
-			UE_LOG(LogSlate, Log, TEXT("Slate: Had to block on waiting for a draw buffer"));
+			UE_LOG(LogSlate, Warning, TEXT("Slate: Had to block on waiting for a draw buffer"));
 			FreeBufferIndex = (FreeBufferIndex + 1) % NumDrawBuffers;
 		}
 	
@@ -227,11 +227,13 @@ FSlateDrawBuffer& FSlateRHIRenderer::GetDrawBuffer()
 		Buffer = &DrawBuffers[FreeBufferIndex];
 	}
 
+	// Safely remove brushes by emptying the array and releasing references
+	DynamicBrushesToRemove[FreeBufferIndex].Empty();
 
 	Buffer->ClearBuffer();
 	return *Buffer;
 #else
-	// With this method buffers are created on this thread and deleted on the rendering thead
+	// With this method buffers are created on this thread and deleted on the rendering thread
 	return *(new FSlateDrawBuffer);
 #endif
 }
@@ -1054,7 +1056,6 @@ bool FSlateRHIRenderer::GenerateDynamicImageResource( FName ResourceName, uint32
 {
 	check( IsInGameThread() );
 
-	
 	TSharedPtr<FSlateDynamicTextureResource> TextureResource = ResourceManager->GetDynamicTextureResourceByName( ResourceName );
 	if( !TextureResource.IsValid() )
 	{
@@ -1063,6 +1064,14 @@ bool FSlateRHIRenderer::GenerateDynamicImageResource( FName ResourceName, uint32
 	return TextureResource.IsValid();
 }
 
+void FSlateRHIRenderer::RemoveDynamicBrushResource( TSharedPtr<FSlateDynamicImageBrush> BrushToRemove )
+{
+	if( BrushToRemove.IsValid() )
+	{
+		DynamicBrushesToRemove[FreeBufferIndex].Add( BrushToRemove );
+	}
+}
+
 /**
  * Gives the renderer a chance to wait for any render commands to be completed before returning/
  */
diff --git a/Engine/Source/Runtime/SlateRHIRenderer/Private/SlateRHIRenderer.h b/Engine/Source/Runtime/SlateRHIRenderer/Private/SlateRHIRenderer.h
--- a/Engine/Source/Runtime/SlateRHIRenderer/Private/SlateRHIRenderer.h
+++ b/Engine/Source/Runtime/SlateRHIRenderer/Private/SlateRHIRenderer.h
@@ -179,6 +179,7 @@ public:
 	virtual void FlushCommands() const override;
 	virtual void Sync() const override;
 	virtual void ReleaseDynamicResource( const FSlateBrush& InBrush ) override;
+	virtual void RemoveDynamicBrushResource( TSharedPtr<FSlateDynamicImageBrush> BrushToRemove ) override;
 	virtual FIntPoint GenerateDynamicImageResource(const FName InTextureName) override;
 	virtual bool GenerateDynamicImageResource( FName ResourceName, uint32 Width, uint32 Height, const TArray< uint8 >& Bytes ) override;
 	virtual void* GetViewportResource( const SWindow& Window ) override;
@@ -289,6 +290,8 @@ private:
 	/** Drawing policy */
 	TSharedPtr<FSlateRHIRenderingPolicy> RenderingPolicy;
 
+	TArray< TSharedPtr<FSlateDynamicImageBrush> > DynamicBrushesToRemove[NumDrawBuffers];
+
 	/** The resource which allows us to capture the editor to a buffer */
 	FSlateCrashReportResource* CrashTrackerResource;
 
