Added missing files

diff --git a/Engine/Source/Runtime/SlateCore/Public/Brushes/SlateDynamicImageBrush.h b/Engine/Source/Runtime/SlateCore/Public/Brushes/SlateDynamicImageBrush.h
--- a/Engine/Source/Runtime/SlateCore/Public/Brushes/SlateDynamicImageBrush.h
+++ b/Engine/Source/Runtime/SlateCore/Public/Brushes/SlateDynamicImageBrush.h
@@ -7,7 +7,7 @@
  * Ignores the Margin. Just renders the image. Can tile the image instead of stretching.
  */
 struct SLATECORE_API FSlateDynamicImageBrush
-	: public FSlateBrush
+	: public FSlateBrush, public TSharedFromThis<FSlateDynamicImageBrush>
 {
 	/**
 	 * @param InTexture		The UTexture2D being used for this brush.
@@ -26,6 +26,7 @@ struct SLATECORE_API FSlateDynamicImageBrush
 	)
 		: FSlateBrush(ESlateBrushDrawType::Image, FName(TEXT("None")), FMargin(0.0f), InTiling, InImageType, InImageSize, InTint, (UObject*)InTexture)
 		, bRemoveResourceFromRootSet(false)
+		, bIsInitalized(true)
 	{
 		bIsDynamicallyLoaded = true;
 		InitFromTextureObject(InTextureName);
@@ -48,6 +49,7 @@ struct SLATECORE_API FSlateDynamicImageBrush
 	)
 		: FSlateBrush(ESlateBrushDrawType::Image, FName(TEXT("None")), FMargin(0.0f), InTiling, InImageType, InImageSize, InTint, (UObject*)InTexture)
 		, bRemoveResourceFromRootSet(false)
+		, bIsInitalized(true)
 	{
 		bIsDynamicallyLoaded = true;
 		InitFromTextureObject(InTextureName);
@@ -60,15 +62,15 @@ struct SLATECORE_API FSlateDynamicImageBrush
 	 * @param InTiling		How do we tile if at all?
 	 * @param InImageType		The type of image this this is
 	 */
-	FSlateDynamicImageBrush( 
+	FORCENOINLINE FSlateDynamicImageBrush( 
 		const FName InTextureName,
 		const FVector2D& InImageSize,
 		const FLinearColor& InTint = FLinearColor(1.0f, 1.0f, 1.0f, 1.0f), 
 		ESlateBrushTileType::Type InTiling = ESlateBrushTileType::NoTile, 
-		ESlateBrushImageType::Type InImageType = ESlateBrushImageType::FullColor
-	)
+		ESlateBrushImageType::Type InImageType = ESlateBrushImageType::FullColor )
 		: FSlateBrush(ESlateBrushDrawType::Image, InTextureName, FMargin(0.0f), InTiling, InImageType, InImageSize, InTint, nullptr, true)
 		, bRemoveResourceFromRootSet(false)
+		, bIsInitalized(true)
 	{
 		bIsDynamicallyLoaded = true;
 	}
@@ -89,10 +91,17 @@ struct SLATECORE_API FSlateDynamicImageBrush
 		ESlateBrushTileType::Type InTiling = ESlateBrushTileType::NoTile,
 		ESlateBrushImageType::Type InImageType = ESlateBrushImageType::FullColor);
 
-	/** Destructor. */
+	/**
+	 * Releases the resource when it is safe to do so
+	 */
+	void ReleaseResource();
+
+		/** Destructor. */
 	virtual ~FSlateDynamicImageBrush();
 
 private:
+	void ReleaseResourceInternal();
+
 	void InitFromTextureObject(FName InTextureName)
 	{
 		// if we have a texture, make a unique name
@@ -112,6 +121,9 @@ private:
 	}
 
 private:
-	// Tracks if Resource was in root set to avoid unnecessary removing it from there.
-	bool bRemoveResourceFromRootSet;
+	/** Tracks if Resource was in root set to avoid unnecessary removing it from there. */
+	bool bRemoveResourceFromRootSet : 1;
+
+	/** If the resource has been initialized */
+	bool bIsInitalized : 1;
 };
