Optimized Slate font caching code and refactored related classes.

- Check whether a character exists within a font set before attempting to load it, and replace with substitute character if not, to minimize FreeType API calls and excessive caching
- Added ability to customize font fallback level; normally if a character does not exist in a font, a localization fallback font set is checked, and then a last resort font set, which can potentially be expensive
- Optimized font texture atlas, to store empty/filled slots in linked lists, instead of in a single binary tree with expensive searches, and stopped tracking empty slots smaller than 2x2 pixels
- Optimized some unnecessary calls to FreeType API

diff --git a/Engine/Source/Runtime/Engine/Private/Font.cpp b/Engine/Source/Runtime/Engine/Private/Font.cpp
--- a/Engine/Source/Runtime/Engine/Private/Font.cpp
+++ b/Engine/Source/Runtime/Engine/Private/Font.cpp
@@ -213,8 +213,8 @@ void UFont::GetCharSize(TCHAR InCh, float& Width, float& Height) const
 			const float FontScale = 1.0f;
 			const FSlateFontInfo LegacyFontInfo = GetLegacySlateFontInfo();
 			FCharacterList& CharacterList = FontCache->GetCharacterList(LegacyFontInfo, FontScale);
+			const FCharacterEntry& Entry = CharacterList.GetCharacter(LegacyFontInfo, InCh);
 
-			const FCharacterEntry& Entry = CharacterList[InCh];
 			Width = Entry.XAdvance;
 
 			// The height of the character will always be the maximum height of any character in this font
@@ -245,7 +245,7 @@ int8 UFont::GetCharKerning(TCHAR First, TCHAR Second) const
 			const FSlateFontInfo LegacyFontInfo = GetLegacySlateFontInfo();
 			FCharacterList& CharacterList = FontCache->GetCharacterList(LegacyFontInfo, FontScale);
 
-			return CharacterList.GetKerning(First, Second);
+			return CharacterList.GetKerning(LegacyFontInfo, First, Second);
 		}
 		break;
 
@@ -265,8 +265,8 @@ int16 UFont::GetCharHorizontalOffset(TCHAR InCh) const
 		const float FontScale = 1.0f;
 		const FSlateFontInfo LegacyFontInfo = GetLegacySlateFontInfo();
 		FCharacterList& CharacterList = FontCache->GetCharacterList(LegacyFontInfo, FontScale);
+		const FCharacterEntry& Entry = CharacterList.GetCharacter(LegacyFontInfo, InCh);
 
-		const FCharacterEntry& Entry = CharacterList[InCh];
 		return Entry.HorizontalOffset;
 	}
 
diff --git a/Engine/Source/Runtime/Engine/Private/UserInterface/CanvasItem.cpp b/Engine/Source/Runtime/Engine/Private/UserInterface/CanvasItem.cpp
--- a/Engine/Source/Runtime/Engine/Private/UserInterface/CanvasItem.cpp
+++ b/Engine/Source/Runtime/Engine/Private/UserInterface/CanvasItem.cpp
@@ -1155,7 +1155,7 @@ void FCanvasTextItem::DrawStringInternal_RuntimeCache( FCanvas* InCanvas, const
 		}
 		else
 		{
-			const FCharacterEntry& Entry = CharacterList[ CurrentChar ];
+			const FCharacterEntry& Entry = CharacterList.GetCharacter(LegacyFontInfo, CurrentChar);
 
 			if( FontTexture == nullptr || Entry.TextureIndex != FontTextureIndex )
 			{
@@ -1174,7 +1174,7 @@ void FCanvasTextItem::DrawStringInternal_RuntimeCache( FCanvas* InCanvas, const
 
 			const bool bIsWhitespace = FChar::IsWhitespace(CurrentChar);
 
-			if( !bIsWhitespace && PreviousCharEntry.IsValidEntry() )
+			if( !bIsWhitespace && PreviousCharEntry.IsCached() )
 			{
 				Kerning = CharacterList.GetKerning( PreviousCharEntry, Entry ) * Scale.X;
 			}
diff --git a/Engine/Source/Runtime/SlateCore/Private/Fonts/FontCache.cpp b/Engine/Source/Runtime/SlateCore/Private/Fonts/FontCache.cpp
--- a/Engine/Source/Runtime/SlateCore/Private/Fonts/FontCache.cpp
+++ b/Engine/Source/Runtime/SlateCore/Private/Fonts/FontCache.cpp
@@ -40,8 +40,11 @@ namespace FontCacheConstants
 	const uint32 MeasureCacheSize = 500;
 }
 
+// Character used to substitute invalid font characters
+const TCHAR InvalidSubChar = '\u001A';
+
 #if WITH_FREETYPE
-const uint32 GlyphFlags = FT_LOAD_NO_BITMAP;
+const uint32 GlobalGlyphFlags = FT_LOAD_NO_BITMAP;
 
 /**
  * Memory allocation functions to be used only by freetype
@@ -63,6 +66,30 @@ static void FreetypeFree( FT_Memory Memory, void* Block )
 
 #endif // WITH_FREETYPE
 
+
+/**
+ * Helper for pop/pushing the font fallback level, within function calls
+ */
+class FScopedFontFallback
+{
+public:
+	FScopedFontFallback(EFontFallback& OutFontFallback, EFontFallback SetFontFallback)
+		: FontFallback(OutFontFallback)
+		, OldFontFallback(OutFontFallback)
+	{
+		FontFallback = SetFontFallback;
+	}
+
+	~FScopedFontFallback()
+	{
+		FontFallback = OldFontFallback;
+	}
+
+private:
+	EFontFallback& FontFallback;
+	EFontFallback OldFontFallback;
+};
+
 /**
  * Cached data for a given typeface
  */
@@ -277,6 +304,35 @@ private:
  */
 class FFreeTypeInterface
 {
+#if WITH_FREETYPE
+	/**
+	 * Internal struct for passing around information about loading a character
+	 */
+	struct FFaceCharData
+	{
+		/** The font face for the character */
+		FT_Face Face;
+
+		/** The glyph index for the character */
+		FT_UInt GlyphIndex;
+
+		/** The glyph flags that should be used for loading the characters glyph */
+		uint32 GlyphFlags;
+
+		/** The fallback font set the character was loaded from */
+		EFontFallback CharFallbackLevel;
+
+
+		FFaceCharData()
+			: Face(NULL)
+			, GlyphIndex(0)
+			, GlyphFlags(0)
+			, CharFallbackLevel(EFontFallback::FF_NoFallback)
+		{
+		}
+	};
+#endif
+
 public:
 	FFreeTypeInterface()
 	{
@@ -355,98 +411,55 @@ public:
 	int16 GetBaseline( const FSlateFontInfo& InFontInfo, const float InScale )
 	{
 #if WITH_FREETYPE
+		// Just get info for the null character 
+		TCHAR Char = 0;
 		const FFontData& FontData = GetDefaultFontData(InFontInfo);
+		FFaceCharData CharData = GetFontFaceForCharacter(FontData, Char, EFontFallback::FF_Max);
 
-		// Just render the null character 
-		TCHAR Char = 0;
+		check(CharData.Face != NULL);
+
+		FT_Error Error = LoadGlyph(FontData, CharData, InFontInfo.Size, InScale);
 
-		// Render the character 
-		FCharacterRenderData NewRenderData;
-		GetRenderData(FontData, InFontInfo.Size, Char, NewRenderData, InScale);
 
-		return NewRenderData.MeasureInfo.GlobalDescender;
+		check(Error == 0);
+
+		return (CharData.Face->size->metrics.descender / 64) * InScale;
 #else
 		return 0;
 #endif // WITH_FREETYPE
 	}
 
+
 	/** 
 	 * Creates render data for a specific character 
 	 * 
-	 * @param InFontData	Raw font data to render the character with
-	 * @param InSize		The size of the font to draw
-	 * @param Char			The character to render
-	 * @param OutCharInfo	Will contain the created render data
+	 * @param InFontData		Raw font data to render the character with
+	 * @param InSize			The size of the font to draw
+	 * @param Char				The character to render
+	 * @param OutRenderData		Will contain the created render data
+	 * @param InScale			The scale of the font
+	 * @param OutFallbackLevel	Outputs the fallback level of the font
 	 */
-	void GetRenderData( const FFontData& InFontData, const int32 InSize, TCHAR Char, FCharacterRenderData& OutRenderData, const float InScale )
+	void GetRenderData(const FFontData& InFontData, const int32 InSize, TCHAR Char, FCharacterRenderData& OutRenderData,
+						const float InScale, EFontFallback* OutFallbackLevel=NULL)
 	{
 #if WITH_FREETYPE
-		// Find or load the face if needed
-		FT_UInt GlyphIndex = 0;
-		FT_Face FontFace = GetFontFace( InFontData );
+		FFaceCharData CharData = GetFontFaceForCharacter(InFontData, Char, EFontFallback::FF_Max);
 
-		if ( FontFace != nullptr ) 
-		{
-			// Get the index to the glyph in the font face
-			GlyphIndex = FT_Get_Char_Index( FontFace, Char );
-		}
-
-		uint32 LocalGlyphFlags = GlyphFlags;
-
-		switch(InFontData.Hinting)
-		{
-		case EFontHinting::Auto:		LocalGlyphFlags |= FT_LOAD_FORCE_AUTOHINT; break;
-		case EFontHinting::AutoLight:	LocalGlyphFlags |= FT_LOAD_TARGET_LIGHT; break;
-		case EFontHinting::Monochrome:	LocalGlyphFlags |= FT_LOAD_TARGET_MONO | FT_LOAD_FORCE_AUTOHINT; break;
-		case EFontHinting::None:		LocalGlyphFlags |= FT_LOAD_NO_AUTOHINT | FT_LOAD_NO_HINTING; break;
-		case EFontHinting::Default:
-		default:						LocalGlyphFlags |= FT_LOAD_TARGET_NORMAL; break;
-		}
+		check(CharData.Face != NULL);
 
-		// If the requested glyph doesn't exist, use the localization fallback font.
-		if ( FontFace == nullptr || (Char != 0 && GlyphIndex == 0) )
+		if (OutFallbackLevel != NULL)
 		{
-			FontFace = GetFontFace( FLegacySlateFontInfoCache::Get().GetFallbackFontData() );
-			if (FontFace != nullptr)
-			{					
-				GlyphIndex = FT_Get_Char_Index( FontFace, Char );
-				LocalGlyphFlags |= FT_LOAD_FORCE_AUTOHINT;
-			}
+			*OutFallbackLevel = CharData.CharFallbackLevel;
 		}
 
-		// If the requested glyph doesn't exist, use the last resort fallback font.
-		if ( FontFace == nullptr || ( Char != 0 && GlyphIndex == 0 ) )
-		{
-			FontFace = GetFontFace( FLegacySlateFontInfoCache::Get().GetLastResortFontData() );
-			check( FontFace );
-			GlyphIndex = FT_Get_Char_Index( FontFace, Char );
-			LocalGlyphFlags |= FT_LOAD_FORCE_AUTOHINT;
-		}
+		FT_Error Error = LoadGlyph(InFontData, CharData, InSize, InScale);
 
-		// Set the character size to render at (needs to be in 1/64 of a "point")
-		FT_Error Error = FT_Set_Char_Size( FontFace, 0, InSize*64, FontCacheConstants::HorizontalDPI, FontCacheConstants::VerticalDPI );
 		check(Error==0);
 
-		if( InScale != 1.0f )
-		{
-			FT_Matrix ScaleMatrix;
-			ScaleMatrix.xy = 0;
-			ScaleMatrix.xx = (FT_Fixed)(InScale * 65536);
-			ScaleMatrix.yy = (FT_Fixed)(InScale * 65536);
-			ScaleMatrix.yx = 0;
-			FT_Set_Transform( FontFace, &ScaleMatrix, nullptr );
-		}
-		else
-		{
-			FT_Set_Transform( FontFace, nullptr, nullptr );
-		}
-
-		// Load the glyph.  Force using the freetype hinter because not all true type fonts have their own hinting
-		Error = FT_Load_Glyph( FontFace, GlyphIndex, LocalGlyphFlags );
-		check(Error==0);
 
 		// Get the slot for the glyph.  This contains measurement info
-		FT_GlyphSlot Slot = FontFace->glyph;
+		FT_GlyphSlot Slot = CharData.Face->glyph;
 		
 		FT_Render_Glyph( Slot, FT_RENDER_MODE_NORMAL );
 
@@ -505,20 +518,20 @@ public:
 		FT_Get_Glyph( Slot, &Glyph );
 		FT_Glyph_Get_CBox( Glyph, FT_GLYPH_BBOX_PIXELS, &GlyphBox );
 
-		int32 Height = (FT_MulFix( FontFace->height, FontFace->size->metrics.y_scale ) / 64) * InScale;
+		int32 Height = (FT_MulFix( CharData.Face->height, CharData.Face->size->metrics.y_scale ) / 64) * InScale;
 
 		// Set measurement info for this character
 		OutRenderData.Char = Char;
-		OutRenderData.HasKerning = FT_HAS_KERNING( FontFace ) != 0;
+		OutRenderData.HasKerning = FT_HAS_KERNING( CharData.Face ) != 0;
 		OutRenderData.MeasureInfo.SizeX = Bitmap->width;
 		OutRenderData.MeasureInfo.SizeY = Bitmap->rows;
 		OutRenderData.MaxHeight = Height;
 
 		// Need to divide by 64 to get pixels;
 		// Ascender is not scaled by freetype.  Scale it now. 
-		OutRenderData.MeasureInfo.GlobalAscender = ( FontFace->size->metrics.ascender / 64 ) * InScale;
+		OutRenderData.MeasureInfo.GlobalAscender = ( CharData.Face->size->metrics.ascender / 64 ) * InScale;
 		// Descender is not scaled by freetype.  Scale it now. 
-		OutRenderData.MeasureInfo.GlobalDescender = ( FontFace->size->metrics.descender / 64 ) * InScale;
+		OutRenderData.MeasureInfo.GlobalDescender = ( CharData.Face->size->metrics.descender / 64 ) * InScale;
 		// Note we use Slot->advance instead of Slot->metrics.horiAdvance because Slot->Advance contains transformed position (needed if we scale)
 		OutRenderData.MeasureInfo.XAdvance =  Slot->advance.x / 64;
 		OutRenderData.MeasureInfo.HorizontalOffset = Slot->bitmap_left;
@@ -706,6 +719,28 @@ public:
 #endif // WITH_FREETYPE
 	}
 
+
+	/**
+	 * Whether or not the specified character, within the specified font, can be loaded with the specified maximum font fallback level
+	 *
+	 * @param InFontData		Information about the font to load
+	 * @param Char				The character being loaded
+	 * @param MaxFallbackLevel	The maximum fallback level to try for the font
+	 * @return					Whether or not the character can be loaded
+	 */
+	bool CanLoadCharacter(const FFontData& InFontData, TCHAR Char, EFontFallback MaxFallbackLevel)
+	{
+		bool bReturnVal = false;
+
+#if WITH_FREETYPE
+		FFaceCharData CharData = GetFontFaceForCharacter(InFontData, Char, MaxFallbackLevel);
+
+		bReturnVal = CharData.Face != NULL && CharData.GlyphIndex != 0;
+#endif
+
+		return bReturnVal;
+	}
+
 private:
 
 #if WITH_FREETYPE
@@ -769,6 +804,7 @@ private:
 		return BestMatchFont;
 	}
 
+
 	/**
 	 * Gets or loads a freetype font face
 	 *
@@ -820,6 +856,121 @@ private:
 		return (FaceAndMemory) ? FaceAndMemory->Face : nullptr;
 	}
 
+
+	/**
+	 * Wrapper for above, which reverts to fallback or last resort fonts if the face could not be loaded
+	 *
+	 * @param InFontData		Information about the font to load
+	 * @param Char				The character being loaded (required for checking if a fallback font is needed)
+	 * @param MaxFallbackLevel	The maximum fallback level to try for the font
+	 * @return					Returns the character font face data
+	 */
+	FFaceCharData GetFontFaceForCharacter(const FFontData& InFontData, TCHAR Char, EFontFallback MaxFallbackLevel)
+	{
+		FFaceCharData ReturnVal;
+
+		FT_Face Face = GetFontFace(InFontData);
+		FT_UInt GlyphIndex = 0;
+		bool bOverrideFallback = Char == InvalidSubChar;
+
+		if (Face != NULL) 
+		{
+			// Get the index to the glyph in the font face
+			GlyphIndex = FT_Get_Char_Index(Face, Char);
+			ReturnVal.CharFallbackLevel = EFontFallback::FF_NoFallback;
+		}
+
+		// If the requested glyph doesn't exist, use the localization fallback font
+		if (Face == NULL || (Char != 0 && GlyphIndex == 0))
+		{
+			bool bCanFallback = bOverrideFallback || MaxFallbackLevel >= EFontFallback::FF_LocalizedFallback;
+
+			if (bCanFallback)
+			{
+				Face = GetFontFace(FLegacySlateFontInfoCache::Get().GetFallbackFontData());
+
+				if (Face != NULL)
+				{	
+					GlyphIndex = FT_Get_Char_Index(Face, Char);
+
+					ReturnVal.CharFallbackLevel = EFontFallback::FF_LocalizedFallback;
+					ReturnVal.GlyphFlags |= FT_LOAD_FORCE_AUTOHINT;
+				}
+			}
+		}
+
+		// If the requested glyph doesn't exist, use the last resort fallback font
+		if (Face == NULL || (Char != 0 && GlyphIndex == 0))
+		{
+			bool bCanFallback = bOverrideFallback || MaxFallbackLevel >= EFontFallback::FF_LastResortFallback;
+
+			if (bCanFallback)
+			{
+				Face = GetFontFace(FLegacySlateFontInfoCache::Get().GetLastResortFontData());
+
+				if (Face != NULL)
+				{
+					GlyphIndex = FT_Get_Char_Index(Face, Char);
+
+					ReturnVal.CharFallbackLevel = EFontFallback::FF_LastResortFallback;
+					ReturnVal.GlyphFlags |= FT_LOAD_FORCE_AUTOHINT;
+				}
+			}
+		}
+
+		ReturnVal.Face = Face;
+		ReturnVal.GlyphIndex = GlyphIndex;
+
+		return ReturnVal;
+	}
+
+
+	/**
+	 * Helper split from GetRenderData, for use with GetBaseline
+	 */
+	FT_Error LoadGlyph(const FFontData& InFontData, FFaceCharData& InCharData, const int32 InSize, const float InScale)
+	{
+		// Set the character size to render at (needs to be in 1/64 of a "point")
+		FT_Error Error = FT_Set_Char_Size(InCharData.Face, 0, InSize*64, FontCacheConstants::HorizontalDPI,
+											FontCacheConstants::VerticalDPI);
+
+		check(Error==0);
+
+		if( InScale != 1.0f )
+		{
+			FT_Matrix ScaleMatrix;
+			ScaleMatrix.xy = 0;
+			ScaleMatrix.xx = (FT_Fixed)(InScale * 65536);
+			ScaleMatrix.yy = (FT_Fixed)(InScale * 65536);
+			ScaleMatrix.yx = 0;
+			FT_Set_Transform( InCharData.Face, &ScaleMatrix, nullptr );
+		}
+		else
+		{
+			FT_Set_Transform( InCharData.Face, nullptr, nullptr );
+		}
+
+
+		// Setup additional glyph flags
+		InCharData.GlyphFlags |= GlobalGlyphFlags;
+
+		switch(InFontData.Hinting)
+		{
+		case EFontHinting::Auto:		InCharData.GlyphFlags |= FT_LOAD_FORCE_AUTOHINT; break;
+		case EFontHinting::AutoLight:	InCharData.GlyphFlags |= FT_LOAD_TARGET_LIGHT; break;
+		case EFontHinting::Monochrome:	InCharData.GlyphFlags |= FT_LOAD_TARGET_MONO | FT_LOAD_FORCE_AUTOHINT; break;
+		case EFontHinting::None:		InCharData.GlyphFlags |= FT_LOAD_NO_AUTOHINT | FT_LOAD_NO_HINTING; break;
+		case EFontHinting::Default:
+		default:						InCharData.GlyphFlags |= FT_LOAD_TARGET_NORMAL; break;
+		}
+
+		// Load the glyph.  Force using the freetype hinter because not all true type fonts have their own hinting
+		Error = FT_Load_Glyph( InCharData.Face, InCharData.GlyphIndex, InCharData.GlyphFlags );
+
+		return Error;
+	}
+
+
 private:
 	struct FFontFaceAndMemory
 	{
@@ -852,6 +1003,7 @@ private:
 #endif // WITH_FREETYPE
 };
 
+
 FKerningTable::FKerningTable( const FSlateFontCache& InFontCache )
 	: DirectAccessTable( nullptr )
 	, FontCache( InFontCache )
@@ -941,6 +1093,7 @@ FCharacterList::FCharacterList( const FSlateFontKey& InFontKey, const FSlateFont
 	, MaxDirectIndexedEntries( FontCacheConstants::DirectAccessSize )
 	, MaxHeight( 0 )
 	, Baseline( 0 )
+	, FontFallback(EFontFallback::FF_Max)
 {
 	const FCompositeFont* const CompositeFont = InFontKey.GetFontInfo().GetCompositeFont();
 	if( CompositeFont )
@@ -955,8 +1108,10 @@ bool FCharacterList::IsStale() const
 	return !CompositeFont || CompositeFontHistoryRevision != CompositeFont->HistoryRevision;
 }
 
-int8 FCharacterList::GetKerning( TCHAR FirstChar, TCHAR SecondChar )
+int8 FCharacterList::GetKerning(const FSlateFontInfo& InFontInfo, TCHAR FirstChar, TCHAR SecondChar)
 {
+	FScopedFontFallback FallbackScope(FontFallback, InFontInfo.FontFallback);
+
 	const FCharacterEntry First = GetCharacter( FirstChar );
 	const FCharacterEntry Second = GetCharacter( SecondChar );
 	return GetKerning( First, Second );
@@ -994,7 +1149,7 @@ uint16 FCharacterList::GetMaxHeight() const
 
 int16 FCharacterList::GetBaseline() const
 {
-	if( Baseline == 0 )
+	if (Baseline == 0)
 	{
 		Baseline = FontCache.GetBaseline( FontKey.GetFontInfo(), FontKey.GetScale() );
 	}
@@ -1002,44 +1157,102 @@ int16 FCharacterList::GetBaseline() const
 	return Baseline;
 }
 
-const FCharacterEntry& FCharacterList::GetCharacter( TCHAR Character )
+bool FCharacterList::CanCacheCharacter(TCHAR Character)
 {
-	if( Character < MaxDirectIndexedEntries )
+	bool bReturnVal = false;
+
+	if (Character == InvalidSubChar)
+	{
+		bReturnVal = true;
+	}
+	else
 	{
-		// The character can be indexed directly
+		float SubFontScalingFactor = 1.0f;
+		const FFontData& FontData = FontCache.GetFontDataForCharacter(FontKey.GetFontInfo(), Character, SubFontScalingFactor);
 
-		int32 NumToAdd = (Character - DirectIndexEntries.Num()) + 1;
-		if( NumToAdd > 0 )
+		bReturnVal = FontCache.FTInterface->CanLoadCharacter(FontData, Character, FontFallback);
+	}
+
+	return bReturnVal;
+}
+
+const FCharacterEntry& FCharacterList::GetCharacter(const FSlateFontInfo& InFontInfo, TCHAR Character)
+{
+	FScopedFontFallback FallbackScope(FontFallback, InFontInfo.FontFallback);
+
+	return GetCharacter(Character);
+}
+
+const FCharacterEntry& FCharacterList::GetCharacter(TCHAR Character)
+{
+	FCharacterEntry* ReturnVal = NULL;
+	bool bDirectIndexChar = Character < MaxDirectIndexedEntries;
+
+	// First get a reference to the character, if it is already mapped (mapped does not mean cached though)
+	if (bDirectIndexChar)
+	{
+		if (DirectIndexEntries.IsValidIndex(Character))
 		{
-			// The character doesn't exist yet in the index and there is not enough space
-			// Resize the array now
-			DirectIndexEntries.AddZeroed( NumToAdd );
+			ReturnVal = &DirectIndexEntries[Character];
 		}
-			
-		FCharacterEntry& CharacterEntry = DirectIndexEntries[ Character ];
-		if( CharacterEntry.IsValidEntry() )
+	}
+	else
+	{
+		ReturnVal = MappedEntries.Find(Character);
+	}
+
+
+	// Determine whether the character needs caching, and map it if needed
+	bool bNeedCaching = false;
+
+	if (ReturnVal != NULL)
+	{
+		bNeedCaching = !ReturnVal->IsCached();
+
+		// If the character needs caching, but can't be cached, reject the character
+		if (bNeedCaching && !CanCacheCharacter(Character))
 		{
-			return CharacterEntry;
+			bNeedCaching = false;
+			ReturnVal = NULL;
+		}
+	}
+	// Only map the character if it can be cached
+	else if (CanCacheCharacter(Character))
+	{
+		bNeedCaching = true;
+
+		if (bDirectIndexChar)
+		{
+			DirectIndexEntries.AddZeroed((Character - DirectIndexEntries.Num()) + 1);
+			ReturnVal = &DirectIndexEntries[Character];
 		}
 		else
 		{
-			// Character has not been cached yet
-			return CacheCharacter( Character );
+			ReturnVal = &(MappedEntries.Add(Character));
 		}
 	}
-	else
+
+
+	if (ReturnVal != NULL)
 	{
-		FCharacterEntry& CharacterEntry = MappedEntries.FindOrAdd( Character );
-		if( CharacterEntry.IsValidEntry() )
+		if (bNeedCaching)
 		{
-			return CharacterEntry;
+			ReturnVal = &(CacheCharacter(Character));
 		}
-		else
+		// For already-cached characters, reject characters that don't fall within maximum font fallback level requirements
+		else if (Character != InvalidSubChar && FontFallback < ReturnVal->FallbackLevel)
 		{
-			// Character has not been cached yet
-			return CacheCharacter( Character );
+			ReturnVal = NULL;
 		}
 	}
+
+	// The character is not valid, replace with the invalid character substitute
+	if (ReturnVal == NULL)
+	{
+		ReturnVal = (FCharacterEntry*)&(GetCharacter(InvalidSubChar));
+	}
+
+	return *ReturnVal;
 }
 
 FCharacterEntry& FCharacterList::CacheCharacter( TCHAR Character )
@@ -1049,7 +1262,7 @@ FCharacterEntry& FCharacterList::CacheCharacter( TCHAR Character )
 
 	check( bSuccess );
 
-	if( Character < MaxDirectIndexedEntries && NewEntry.IsValidEntry() )
+	if( Character < MaxDirectIndexedEntries && NewEntry.IsCached() )
 	{
 		DirectIndexEntries[ Character ] = NewEntry;
 		return DirectIndexEntries[ Character ];
@@ -1100,8 +1313,10 @@ bool FSlateFontCache::AddNewEntry( TCHAR Character, const FSlateFontKey& InKey,
 
 	// Render the character 
 	FCharacterRenderData RenderData;
+	EFontFallback CharFallbackLevel;
 	const float FontScale = InKey.GetScale() * SubFontScalingFactor;
-	FTInterface->GetRenderData( FontData, InKey.GetFontInfo().Size, Character, RenderData, FontScale );
+
+	FTInterface->GetRenderData( FontData, InKey.GetFontInfo().Size, Character, RenderData, FontScale, &CharFallbackLevel);
 
 	// the index of the atlas where the character is stored
 	int32 AtlasIndex = 0;
@@ -1155,6 +1370,7 @@ bool FSlateFontCache::AddNewEntry( TCHAR Character, const FSlateFontKey& InKey,
 		OutCharacterEntry.HorizontalOffset = RenderData.MeasureInfo.HorizontalOffset;
 		OutCharacterEntry.TextureIndex = AtlasIndex;
 		OutCharacterEntry.HasKerning = RenderData.HasKerning;
+		OutCharacterEntry.FallbackLevel = CharFallbackLevel;
 		OutCharacterEntry.Valid = 1;
 	}
 
diff --git a/Engine/Source/Runtime/SlateCore/Private/Fonts/FontMeasure.cpp b/Engine/Source/Runtime/SlateCore/Private/Fonts/FontMeasure.cpp
--- a/Engine/Source/Runtime/SlateCore/Private/Fonts/FontMeasure.cpp
+++ b/Engine/Source/Runtime/SlateCore/Private/Fonts/FontMeasure.cpp
@@ -200,12 +200,12 @@ FVector2D FSlateFontMeasure::MeasureStringInternal( const FString& Text, int32 S
 		}
 		else
 		{
-			const FCharacterEntry& Entry = CharacterList[CurrentChar];
+			const FCharacterEntry& Entry = CharacterList.GetCharacter(InFontInfo, CurrentChar);
 
 			int32 Kerning = 0;
 			if( PreviousChar != 0 )
 			{
-				Kerning = CharacterList.GetKerning( CharacterList[ PreviousChar ], Entry );
+				Kerning = CharacterList.GetKerning(CharacterList.GetCharacter(InFontInfo, PreviousChar), Entry);
 			}
 
 			PreviousChar = CurrentChar;
@@ -313,7 +313,7 @@ uint16 FSlateFontMeasure::GetMaxCharacterHeight( const FSlateFontInfo& InFontInf
 int8 FSlateFontMeasure::GetKerning(const FSlateFontInfo& InFontInfo, float FontScale, TCHAR PreviousCharacter, TCHAR CurrentCharacter) const
 {
 	FCharacterList& CharacterList = FontCache->GetCharacterList( InFontInfo, FontScale );
-	return CharacterList.GetKerning( PreviousCharacter, CurrentCharacter );
+	return CharacterList.GetKerning( InFontInfo, PreviousCharacter, CurrentCharacter );
 }
 
 int16 FSlateFontMeasure::GetBaseline( const FSlateFontInfo& InFontInfo, float FontScale ) const
diff --git a/Engine/Source/Runtime/SlateCore/Private/Fonts/SlateFontInfo.cpp b/Engine/Source/Runtime/SlateCore/Private/Fonts/SlateFontInfo.cpp
--- a/Engine/Source/Runtime/SlateCore/Private/Fonts/SlateFontInfo.cpp
+++ b/Engine/Source/Runtime/SlateCore/Private/Fonts/SlateFontInfo.cpp
@@ -16,6 +16,7 @@ FSlateFontInfo::FSlateFontInfo( )
 	, Size(0)
 	, FontName_DEPRECATED()
 	, Hinting_DEPRECATED(EFontHinting::Default)
+	, FontFallback(EFontFallback::FF_Max)
 {
 }
 
@@ -28,6 +29,7 @@ FSlateFontInfo::FSlateFontInfo( TSharedPtr<const FCompositeFont> InCompositeFont
 	, Size(InSize)
 	, FontName_DEPRECATED()
 	, Hinting_DEPRECATED(EFontHinting::Default)
+	, FontFallback(EFontFallback::FF_Max)
 {
 	if (!InCompositeFont.IsValid())
 	{
@@ -44,6 +46,7 @@ FSlateFontInfo::FSlateFontInfo( const UObject* InFontObject, const int32 InSize,
 	, Size(InSize)
 	, FontName_DEPRECATED()
 	, Hinting_DEPRECATED(EFontHinting::Default)
+	, FontFallback(EFontFallback::FF_Max)
 {
 	if (InFontObject)
 	{
@@ -68,6 +71,7 @@ FSlateFontInfo::FSlateFontInfo( const FString& InFontName, uint16 InSize, EFontH
 	, Size(InSize)
 	, FontName_DEPRECATED(*InFontName)
 	, Hinting_DEPRECATED(InHinting)
+	, FontFallback(EFontFallback::FF_Max)
 {
 	//Useful for debugging style breakages
 	//check( FPaths::FileExists( FontName.ToString() ) );
@@ -84,6 +88,7 @@ FSlateFontInfo::FSlateFontInfo( const FName& InFontName, uint16 InSize, EFontHin
 	, Size(InSize)
 	, FontName_DEPRECATED(InFontName)
 	, Hinting_DEPRECATED(InHinting)
+	, FontFallback(EFontFallback::FF_Max)
 {
 	//Useful for debugging style breakages
 	//check( FPaths::FileExists( FontName.ToString() ) );
@@ -100,6 +105,7 @@ FSlateFontInfo::FSlateFontInfo( const ANSICHAR* InFontName, uint16 InSize, EFont
 	, Size(InSize)
 	, FontName_DEPRECATED(InFontName)
 	, Hinting_DEPRECATED(InHinting)
+	, FontFallback(EFontFallback::FF_Max)
 {
 	//Useful for debugging style breakages
 	//check( FPaths::FileExists( FontName.ToString() ) );
@@ -116,6 +122,7 @@ FSlateFontInfo::FSlateFontInfo( const WIDECHAR* InFontName, uint16 InSize, EFont
 	, Size(InSize)
 	, FontName_DEPRECATED(InFontName)
 	, Hinting_DEPRECATED(InHinting)
+	, FontFallback(EFontFallback::FF_Max)
 {
 	//Useful for debugging style breakages
 	//check( FPaths::FileExists( FontName.ToString() ) );
diff --git a/Engine/Source/Runtime/SlateCore/Private/Rendering/ElementBatcher.cpp b/Engine/Source/Runtime/SlateCore/Private/Rendering/ElementBatcher.cpp
--- a/Engine/Source/Runtime/SlateCore/Private/Rendering/ElementBatcher.cpp
+++ b/Engine/Source/Runtime/SlateCore/Private/Rendering/ElementBatcher.cpp
@@ -617,7 +617,7 @@ void FSlateElementBatcher::AddTextElement(const FSlateDrawElement& DrawElement)
 		}
 		else
 		{
-			const FCharacterEntry& Entry = CharacterList[ CurrentChar ];
+			const FCharacterEntry& Entry = CharacterList.GetCharacter(InPayload.FontInfo, CurrentChar);
 
 			if( FontAtlasTexture == nullptr || Entry.TextureIndex != FontTextureIndex )
 			{
@@ -641,7 +641,7 @@ void FSlateElementBatcher::AddTextElement(const FSlateDrawElement& DrawElement)
 
 			const bool bIsWhitespace = FText::IsWhitespace(CurrentChar);
 
-			if( !bIsWhitespace && PreviousCharEntry.IsValidEntry() )
+			if( !bIsWhitespace && PreviousCharEntry.IsCached() )
 			{
 				Kerning = CharacterList.GetKerning( PreviousCharEntry, Entry );
 			}
diff --git a/Engine/Source/Runtime/SlateCore/Private/Textures/TextureAtlas.cpp b/Engine/Source/Runtime/SlateCore/Private/Textures/TextureAtlas.cpp
--- a/Engine/Source/Runtime/SlateCore/Private/Textures/TextureAtlas.cpp
+++ b/Engine/Source/Runtime/SlateCore/Private/Textures/TextureAtlas.cpp
@@ -22,8 +22,30 @@ FSlateTextureAtlas::~FSlateTextureAtlas()
 void FSlateTextureAtlas::Empty()
 {
 	// Remove all nodes
-	DestroyNodes( RootNode );
-	RootNode = nullptr;
+	TArray<FAtlasedTextureSlot*> DeleteSlots;
+
+	for (FAtlasedTextureSlot::TIterator SlotIt(AtlasUsedSlots); SlotIt; SlotIt.Next())
+	{
+		FAtlasedTextureSlot& CurSlot = *SlotIt;
+		DeleteSlots.Add(&CurSlot);
+	}
+
+	for (FAtlasedTextureSlot::TIterator SlotIt(AtlasEmptySlots); SlotIt; SlotIt.Next())
+	{
+		FAtlasedTextureSlot& CurSlot = *SlotIt;
+		DeleteSlots.Add(&CurSlot);
+	}
+
+	AtlasUsedSlots = NULL;
+	AtlasEmptySlots = NULL;
+
+	for (FAtlasedTextureSlot* CurSlot : DeleteSlots)
+	{
+		delete CurSlot;
+	}
+
+	DeleteSlots.Empty();
+
 
 	STAT(uint32 MemoryBefore = AtlasData.GetAllocatedSize());
 
@@ -66,33 +88,15 @@ void FSlateTextureAtlas::MarkTextureDirty()
 	bNeedsUpdate = true;
 }
 
-
-const FAtlasedTextureSlot* FSlateTextureAtlas::FindSlotForTexture( uint32 InWidth, uint32 InHeight )
-{
-	return FindSlotForTexture(*RootNode, InWidth, InHeight);
-}
-
-void FSlateTextureAtlas::DestroyNodes( FAtlasedTextureSlot* StartNode )
+void FSlateTextureAtlas::InitAtlasData()
 {
-	if (StartNode->Left)
-	{
-		DestroyNodes(StartNode->Left);
-	}
-
-	if (StartNode->Right)
-	{
-		DestroyNodes(StartNode->Right);
-	}
+	check(AtlasEmptySlots == NULL && AtlasData.Num() == 0);
 
-	delete StartNode;
-}
+	FAtlasedTextureSlot* RootSlot = new FAtlasedTextureSlot(0, 0, AtlasWidth, AtlasHeight, GetPaddingAmount());
 
+	RootSlot->LinkHead(AtlasEmptySlots);
 
-void FSlateTextureAtlas::InitAtlasData()
-{
-	check(RootNode == nullptr && AtlasData.Num() == 0);
 
-	RootNode = new FAtlasedTextureSlot(0, 0, AtlasWidth, AtlasHeight, GetPaddingAmount());
 	AtlasData.Reserve(AtlasWidth * AtlasHeight * BytesPerPixel);
 	AtlasData.AddZeroed(AtlasWidth * AtlasHeight * BytesPerPixel);
 
@@ -217,36 +221,9 @@ void FSlateTextureAtlas::CopyDataIntoSlot( const FAtlasedTextureSlot* SlotToCopy
 	}
 }
 
-
-const FAtlasedTextureSlot* FSlateTextureAtlas::FindSlotForTexture( FAtlasedTextureSlot& Start, uint32 InWidth, uint32 InHeight )
+const FAtlasedTextureSlot* FSlateTextureAtlas::FindSlotForTexture(uint32 InWidth, uint32 InHeight)
 {
-	// If there are left and right slots there are empty regions around this slot.  
-	// It also means this slot is occupied by a texture
-	if (Start.Left || Start.Right)
-	{
-		// Recursively search left for the smallest empty slot that can fit the texture
-		if (Start.Left)
-		{
-			const FAtlasedTextureSlot* NewSlot = FindSlotForTexture(*Start.Left, InWidth, InHeight);
-			if (NewSlot)
-			{
-				return NewSlot;
-			}
-		}
-
-		// Recursively search right for the smallest empty slot that can fit the texture
-		if (Start.Right)
-		{
-			const FAtlasedTextureSlot* NewSlot = FindSlotForTexture(*Start.Right, InWidth, InHeight);
-			if(NewSlot)
-			{
-				return NewSlot;
-			}
-		}
-
-		// Not enough space
-		return nullptr;
-	}
+	FAtlasedTextureSlot* ReturnVal = NULL;
 
 	// Account for padding on both sides
 	const uint32 Padding = GetPaddingAmount();
@@ -254,34 +231,87 @@ const FAtlasedTextureSlot* FSlateTextureAtlas::FindSlotForTexture( FAtlasedTextu
 	const uint32 PaddedWidth = InWidth + TotalPadding;
 	const uint32 PaddedHeight = InHeight + TotalPadding;
 
-	// This slot can't fit the character
-	if ((PaddedWidth > Start.Width) || (PaddedHeight > Start.Height))
+	// Previously, slots were stored as a binary tree - this has been replaced with a linked-list of slots on the edge of the tree
+	// (slots on the edge of the tree represent empty slots); this iterates empty slots in same order as a binary depth-first-search,
+	// except much faster.
+	for (FAtlasedTextureSlot::TIterator SlotIt(AtlasEmptySlots); SlotIt; SlotIt++)
 	{
-		// not enough space
-		return nullptr;
-	}
-	
-	// The width and height of the new child node
-	const uint32 RemainingWidth =  FMath::Max<int32>(0, Start.Width - PaddedWidth);
-	const uint32 RemainingHeight = FMath::Max<int32>(0, Start.Height - PaddedHeight);
+		FAtlasedTextureSlot& CurSlot = *SlotIt;
 
-	// Split the remaining area around this slot into two children.
-	if (RemainingHeight <= RemainingWidth)
-	{
-		// Split vertically
-		Start.Left = new FAtlasedTextureSlot(Start.X, Start.Y + PaddedHeight, PaddedWidth, RemainingHeight, Padding);
-		Start.Right = new FAtlasedTextureSlot(Start.X + PaddedWidth, Start.Y, RemainingWidth, Start.Height, Padding);
+		if (PaddedWidth <= CurSlot.Width && PaddedHeight <= CurSlot.Height)
+		{
+			ReturnVal = &CurSlot;
+			break;
+		}
 	}
-	else
+
+
+	if (ReturnVal != NULL)
 	{
-		// Split horizontally
-		Start.Left = new FAtlasedTextureSlot(Start.X + PaddedWidth, Start.Y, RemainingWidth, PaddedHeight, Padding);
-		Start.Right = new FAtlasedTextureSlot(Start.X, Start.Y + PaddedHeight, Start.Width, RemainingHeight, Padding);
-	}
+		// The width and height of the new child node
+		const uint32 RemainingWidth =  FMath::Max<int32>(0, ReturnVal->Width - PaddedWidth);
+		const uint32 RemainingHeight = FMath::Max<int32>(0, ReturnVal->Height - PaddedHeight);
 
-	// Shrink the slot to the remaining area.
-	Start.Width = PaddedWidth;
-	Start.Height = PaddedHeight;
+		// New slots must have a minimum width/height, to avoid excessive slots i.e. excessive memory usage and iteration.
+		// No glyphs seem to use slots this small, and cutting these slots out improves performance/memory-usage a fair bit
+		const uint32 MinSlotDim = 2;
+
+		// Split the remaining area around this slot into two children.
+		if (RemainingHeight >= MinSlotDim && RemainingWidth >= MinSlotDim)
+		{
+			FAtlasedTextureSlot* LeftSlot = NULL;
+			FAtlasedTextureSlot* RightSlot = NULL;
 
-	return &Start;
+			if (RemainingHeight <= RemainingWidth)
+			{
+				// Split vertically
+				// - - - - - - - - -
+				// |       |       |
+				// |  Slot |       |
+				// |       |       |
+				// | - - - | Right |
+				// |       |       |
+				// |  Left |       |
+				// |       |       |
+				// - - - - - - - - -
+				LeftSlot = new FAtlasedTextureSlot(ReturnVal->X, ReturnVal->Y + PaddedHeight, PaddedWidth, RemainingHeight, Padding);
+				RightSlot = new FAtlasedTextureSlot(ReturnVal->X + PaddedWidth, ReturnVal->Y, RemainingWidth, ReturnVal->Height, Padding);
+			}
+			else
+			{
+				// Split horizontally
+				// - - - - - - - - -
+				// |       |       |
+				// |  Slot | Left  |
+				// |       |       |
+				// | - - - - - - - |
+				// |               |
+				// |     Right     |
+				// |               |
+				// - - - - - - - - -
+				LeftSlot = new FAtlasedTextureSlot(ReturnVal->X + PaddedWidth, ReturnVal->Y, RemainingWidth, PaddedHeight, Padding);
+				RightSlot = new FAtlasedTextureSlot(ReturnVal->X, ReturnVal->Y + PaddedHeight, ReturnVal->Width, RemainingHeight, Padding);
+			}
+
+			// Replace the old slot within AtlasEmptySlots, with the new Left and Right slot, then add the old slot to AtlasUsedSlots
+			LeftSlot->LinkReplace(ReturnVal);
+			RightSlot->LinkAfter(LeftSlot);
+
+			ReturnVal->LinkHead(AtlasUsedSlots);
+		}
+		else
+		{
+			// Remove the old slot from AtlasEmptySlots, into AtlasUsedSlots
+			ReturnVal->Unlink();
+			ReturnVal->LinkHead(AtlasUsedSlots);
+		}
+
+
+		// Shrink the slot to the remaining area.
+		ReturnVal->Width = PaddedWidth;
+		ReturnVal->Height = PaddedHeight;
+	}
+
+	return ReturnVal;
 }
+
diff --git a/Engine/Source/Runtime/SlateCore/Public/Fonts/FontCache.h b/Engine/Source/Runtime/SlateCore/Public/Fonts/FontCache.h
--- a/Engine/Source/Runtime/SlateCore/Public/Fonts/FontCache.h
+++ b/Engine/Source/Runtime/SlateCore/Public/Fonts/FontCache.h
@@ -36,6 +36,8 @@ struct SLATECORE_API FCharacterEntry
 	uint8 TextureIndex;
 	/** 1 if this entry has kerning, 0 otherwise. */
 	bool HasKerning;
+	/** The fallback level this character represents */
+	EFontFallback FallbackLevel;
 	/** 1 if this entry is valid, 0 otherwise. */
 	bool Valid;
 
@@ -49,7 +51,7 @@ struct SLATECORE_API FCharacterEntry
 	 * 
 	 * @param InEntry The entry to check
 	 */
-	FORCEINLINE bool IsValidEntry() const
+	FORCEINLINE bool IsCached() const
 	{
 		return Valid;
 	}
@@ -129,17 +131,16 @@ public:
 		return DirectIndexEntries.IsValidIndex( Character ) || ( Character >= MaxDirectIndexedEntries && MappedEntries.Contains( Character ) );
 	}
 
+
 	/**
 	 * Gets data about how to render and measure a character 
 	 * Caching and atlasing it if needed
 	 *
-	 * @param Character	The character to get
-	 * @return Data about the character
+	 * @param InFontInfo	The higher level font info (used for determining the font fallback level - FontKey does not cache this)
+	 * @param Character		The character to get
+	 * @return				Data about the character
 	 */
-	const FCharacterEntry& operator[]( TCHAR Character )
-	{
-		return GetCharacter( Character );
-	}
+	const FCharacterEntry& GetCharacter(const FSlateFontInfo& InFontInfo, TCHAR Character);
 
 	/** Check to see if our cached data is potentially stale for our font */
 	bool IsStale() const;
@@ -147,11 +148,12 @@ public:
 	/**
 	 * Gets a kerning value for a pair of characters
 	 *
+	 * @param InFontInfo	The higher level font info (used for determining the font fallback level - FontKey does not cache this)
 	 * @param FirstChar		The first character in the pair
 	 * @param SecondChar	The second character in the pair
 	 * @return The kerning value
 	 */
-	int8 GetKerning( TCHAR FirstChar, TCHAR SecondChar );
+	int8 GetKerning(const FSlateFontInfo& InFontInfo, TCHAR FirstChar, TCHAR SecondChar);
 
 	/**
 	 * Gets a kerning value for a pair of character entries
@@ -176,6 +178,13 @@ public:
 
 private:
 
+	/**
+	 * Returns whether the specified character is valid for caching (i.e. whether it matches the FontFallback level)
+	 *
+	 * @param Character		The character to check
+	 */
+	bool CanCacheCharacter(TCHAR Character);
+
 	/**
 	 * Gets data about how to render and measure a character 
 	 * Caching and atlasing it if needed
@@ -213,6 +222,8 @@ private:
 	mutable uint16 MaxHeight;
 	/** The offset from the bottom of the max character height to the baseline. */
 	mutable int16 Baseline;
+	/** Temporarily adjusts the font fallback level, during key functions calls */
+	EFontFallback FontFallback;
 };
 
 /**
@@ -221,6 +232,8 @@ private:
  */
 class SLATECORE_API FSlateFontCache : public ISlateAtlasProvider
 {
+	friend FCharacterList;
+
 public:
 	/**
 	 * Constructor
diff --git a/Engine/Source/Runtime/SlateCore/Public/Fonts/SlateFontInfo.h b/Engine/Source/Runtime/SlateCore/Public/Fonts/SlateFontInfo.h
--- a/Engine/Source/Runtime/SlateCore/Public/Fonts/SlateFontInfo.h
+++ b/Engine/Source/Runtime/SlateCore/Public/Fonts/SlateFontInfo.h
@@ -6,6 +6,26 @@
 #include "CompositeFont.h"
 #include "SlateFontInfo.generated.h"
 
+
+/**
+ * Sets the maximum font fallback level, for when a character can't be found in the selected font set.
+ *
+ * UI code that renders strings from a third party (e.g. player chat in a multiplayer game), should restrict font fallback to localized,
+ * (or to no fallback, if international font isn't important), to prevent potential performance problems.
+ */
+UENUM()
+enum class EFontFallback : uint8
+{
+	/** No fallback font */
+	FF_NoFallback,
+	/** Fallback to localized font set */
+	FF_LocalizedFallback,
+	/** Fallback to last resort font set */
+	FF_LastResortFallback,
+	/** Tries all fallbacks */
+	FF_Max
+};
+
 /**
  * A representation of a font in Slate.
  */
@@ -43,6 +63,11 @@ private:
 	UPROPERTY()
 	EFontHinting Hinting_DEPRECATED;
 
+public:
+
+	/** The font fallback level. Runtime only, don't set on shared FSlateFontInfo, as it may change the font elsewhere (make a copy). */
+	EFontFallback FontFallback;
+
 public:
 
 	/** Default constructor. */
diff --git a/Engine/Source/Runtime/SlateCore/Public/Textures/TextureAtlas.h b/Engine/Source/Runtime/SlateCore/Public/Textures/TextureAtlas.h
--- a/Engine/Source/Runtime/SlateCore/Public/Textures/TextureAtlas.h
+++ b/Engine/Source/Runtime/SlateCore/Public/Textures/TextureAtlas.h
@@ -29,16 +29,17 @@ enum class ESlateTextureAtlasOwnerThread : uint8
 	Render,
 };
 
+
 /**
- * Structure holding information about where a texture is located in the atlas
- * Pointers to left and right children build a tree of texture rectangles so we can easily find optimal slots for textures
+ * Structure holding information about where a texture is located in the atlas. Inherits a linked-list interface.
+ *
+ * When a slot is occupied by texture data, the remaining space in the slot (if big enough) is split off into two new (smaller) slots,
+ * building a tree of texture rectangles which, instead of being stored as a tree, are flattened into two linked-lists:
+ *	- AtlastEmptySlots:	A linked-list of empty slots ready for texture data - iterates in same order as a depth-first-search on a tree
+ *	- AtlasUsedSlots:	An unordered linked-list of slots containing texture data
  */
-struct FAtlasedTextureSlot
+struct FAtlasedTextureSlot : public TIntrusiveLinkedList<FAtlasedTextureSlot>
 {
-	/** Left child slot. If nullptr there is no texture data here */
-	FAtlasedTextureSlot* Left;
-	/** Right child slot. If nullptr there is no texture data here  */
-	FAtlasedTextureSlot* Right;
 	/** The X position of the character in the texture */
 	uint32 X;
 	/** The Y position of the character in the texture */
@@ -49,17 +50,16 @@ struct FAtlasedTextureSlot
 	uint32 Height;
 	/** Uniform Padding. can only be zero or one. See ESlateTextureAtlasPaddingStyle. */
 	uint8 Padding;
+
+
 	FAtlasedTextureSlot( uint32 InX, uint32 InY, uint32 InWidth, uint32 InHeight, uint8 InPadding )
-		: Left(nullptr)
-		, Right(nullptr)
+		: TIntrusiveLinkedList<FAtlasedTextureSlot>()
 		, X(InX)
 		, Y(InY)
 		, Width(InWidth)
 		, Height(InHeight)
 		, Padding(InPadding)
-	
 	{
-
 	}
 };
 
@@ -71,7 +71,8 @@ class SLATECORE_API FSlateTextureAtlas
 public:
 	FSlateTextureAtlas( uint32 InWidth, uint32 InHeight, uint32 InBytesPerPixel, ESlateTextureAtlasPaddingStyle InPaddingStyle )
 		: AtlasData()
-		, RootNode( nullptr )
+		, AtlasUsedSlots(NULL)
+		, AtlasEmptySlots(NULL)
 		, AtlasWidth( InWidth )
 		, AtlasHeight( InHeight )
 		, BytesPerPixel( InBytesPerPixel )
@@ -112,21 +113,14 @@ public:
 	virtual void ConditionalUpdateTexture() = 0;
 	
 protected:
-	/** 
-	 * Finds the optimal slot for a texture in the atlas starting the search from the root
+	/**
+	 * Finds the optimal slot for a texture in the atlas
 	 * 
 	 * @param Width The width of the texture we are adding
 	 * @param Height The height of the texture we are adding
 	 */
 	const FAtlasedTextureSlot* FindSlotForTexture( uint32 InWidth, uint32 InHeight );
 
-	/** 
-	 * Destroys everything in the atlas starting at the provided node
-	 *
-	 * @param StartNode
-	 */
-	void DestroyNodes( FAtlasedTextureSlot* StartNode );
-	
 	/**
 	 * Creates enough space for a single texture the width and height of the atlas
 	 */
@@ -175,27 +169,18 @@ protected:
 	void CopyDataIntoSlot( const FAtlasedTextureSlot* SlotToCopyTo, const TArray<uint8>& Data );
 
 private:
-	/** 
-	 * Finds the optimal slot for a texture in the atlas.  
-	 * Does this by doing a DFS over all existing slots to see if the texture can fit in 
-	 * the empty area next to an existing slot
-	 * 
-	 * @param Start	The start slot to check
-	 * @param Width The width of the texture we are adding
-	 * @param Height The height of the texture we are adding
-	 */
-	const FAtlasedTextureSlot* FindSlotForTexture( FAtlasedTextureSlot& Start, uint32 InWidth, uint32 InHeight );
-
 	/** Returns the amount of padding needed for the current padding style */
-	int32 GetPaddingAmount() const
+	FORCEINLINE int32 GetPaddingAmount() const
 	{
 		return (PaddingStyle == ESlateTextureAtlasPaddingStyle::NoPadding) ? 0 : 1;
 	}
 protected:
 	/** Actual texture data contained in the atlas */
 	TArray<uint8> AtlasData;
-	/** Root node for the tree of data.  */
-	FAtlasedTextureSlot* RootNode;
+	/** The list of atlas slots pointing to used texture data in the atlas */
+	FAtlasedTextureSlot* AtlasUsedSlots;
+	/** The list of atlas slots pointing to empty texture data in the atlas */
+	FAtlasedTextureSlot* AtlasEmptySlots;
 	/** Width of the atlas */
 	uint32 AtlasWidth;
 	/** Height of the atlas */
