Removed SLATE_TEXT_ATTRIBUTE and SLATE_TEXT_ARGUMENT

They were deprecated in 4.8, and have been removed for the 4.10 release.

diff --git a/Engine/Source/Developer/FriendsAndChat/Private/Layers/UI/SFriendsAndChatCombo.cpp b/Engine/Source/Developer/FriendsAndChat/Private/Layers/UI/SFriendsAndChatCombo.cpp
--- a/Engine/Source/Developer/FriendsAndChat/Private/Layers/UI/SFriendsAndChatCombo.cpp
+++ b/Engine/Source/Developer/FriendsAndChat/Private/Layers/UI/SFriendsAndChatCombo.cpp
@@ -41,7 +41,7 @@ public:
 		SLATE_NAMED_SLOT(FArguments, Content)
 
 		/** Text to display on main button. */
-		SLATE_TEXT_ATTRIBUTE(ButtonText)
+		SLATE_ATTRIBUTE(FText, ButtonText)
 
 		/** Whether the optional icon is shown */
 		SLATE_ATTRIBUTE(bool, bShowIcon)
diff --git a/Engine/Source/Developer/FriendsAndChat/Private/Layers/UI/SFriendsAndChatCombo.h b/Engine/Source/Developer/FriendsAndChat/Private/Layers/UI/SFriendsAndChatCombo.h
--- a/Engine/Source/Developer/FriendsAndChat/Private/Layers/UI/SFriendsAndChatCombo.h
+++ b/Engine/Source/Developer/FriendsAndChat/Private/Layers/UI/SFriendsAndChatCombo.h
@@ -80,7 +80,7 @@ public:
 		SLATE_NAMED_SLOT(FArguments, Content)
 
 		/** Text to display on main button. */
-		SLATE_TEXT_ATTRIBUTE(ButtonText)
+		SLATE_ATTRIBUTE(FText, ButtonText)
 
 		/** Whether the optional icon is shown */
 		SLATE_ATTRIBUTE(bool, bShowIcon)
diff --git a/Engine/Source/Runtime/Slate/Private/Framework/Application/SlateApplication.cpp b/Engine/Source/Runtime/Slate/Private/Framework/Application/SlateApplication.cpp
--- a/Engine/Source/Runtime/Slate/Private/Framework/Application/SlateApplication.cpp
+++ b/Engine/Source/Runtime/Slate/Private/Framework/Application/SlateApplication.cpp
@@ -3748,17 +3748,6 @@ TSharedRef<SWidget> FSlateApplication::MakeWindowTitleBar( const TSharedRef<SWin
 }
 
 
-TSharedRef<IToolTip> FSlateApplication::MakeToolTip( const TAttribute<FString>& ToolTipString )
-{
-	PRAGMA_DISABLE_DEPRECATION_WARNINGS
-
-	return SNew(SToolTip)
-		.Text(ToolTipString);
-
-	PRAGMA_ENABLE_DEPRECATION_WARNINGS
-}
-
-
 TSharedRef<IToolTip> FSlateApplication::MakeToolTip(const TAttribute<FText>& ToolTipText)
 {
 	return SNew(SToolTip)
diff --git a/Engine/Source/Runtime/Slate/Public/Framework/Application/SlateApplication.h b/Engine/Source/Runtime/Slate/Public/Framework/Application/SlateApplication.h
--- a/Engine/Source/Runtime/Slate/Public/Framework/Application/SlateApplication.h
+++ b/Engine/Source/Runtime/Slate/Public/Framework/Application/SlateApplication.h
@@ -1077,8 +1077,6 @@ public:
 	virtual bool IsWindowHousingInteractiveTooltip(const TSharedRef<const SWindow>& WindowToTest) const override;
 	virtual TSharedRef<SWidget> MakeImage( const TAttribute<const FSlateBrush*>& Image, const TAttribute<FSlateColor>& Color, const TAttribute<EVisibility>& Visibility ) const override;
 	virtual TSharedRef<SWidget> MakeWindowTitleBar( const TSharedRef<SWindow>& Window, const TSharedPtr<SWidget>& CenterContent, EHorizontalAlignment CenterContentAlignment, TSharedPtr<IWindowTitleBar>& OutTitleBar ) const override;
-	DEPRECATED(4.8, "Passing text to Slate as FString is deprecated, please use FText instead (likely via a LOCTEXT).")
-	virtual TSharedRef<IToolTip> MakeToolTip( const TAttribute<FString>& ToolTipString ) override;
 	virtual TSharedRef<IToolTip> MakeToolTip( const TAttribute<FText>& ToolTipText ) override;
 	virtual TSharedRef<IToolTip> MakeToolTip( const FText& ToolTipText ) override;
 	virtual void RequestDestroyWindow( TSharedRef<SWindow> WindowToDestroy ) override;
diff --git a/Engine/Source/Runtime/Slate/Public/Widgets/Input/SButton.h b/Engine/Source/Runtime/Slate/Public/Widgets/Input/SButton.h
--- a/Engine/Source/Runtime/Slate/Public/Widgets/Input/SButton.h
+++ b/Engine/Source/Runtime/Slate/Public/Widgets/Input/SButton.h
@@ -48,7 +48,7 @@ public:
 		SLATE_ATTRIBUTE( FMargin, ContentPadding )
 	
 		/** The text to display in this button, if no custom content is specified */
-		SLATE_TEXT_ATTRIBUTE( Text )
+		SLATE_ATTRIBUTE( FText, Text )
 	
 		/** Called when the button is clicked */
 		SLATE_EVENT( FOnClicked, OnClicked )
diff --git a/Engine/Source/Runtime/Slate/Public/Widgets/Input/SEditableComboBox.h b/Engine/Source/Runtime/Slate/Public/Widgets/Input/SEditableComboBox.h
--- a/Engine/Source/Runtime/Slate/Public/Widgets/Input/SEditableComboBox.h
+++ b/Engine/Source/Runtime/Slate/Public/Widgets/Input/SEditableComboBox.h
@@ -38,7 +38,7 @@ public:
 		, _OptionsSource()
 	 { }
 
-		SLATE_TEXT_ATTRIBUTE(AddButtonToolTip)
+		SLATE_ATTRIBUTE(FText, AddButtonToolTip)
 
 		SLATE_STYLE_ARGUMENT(FButtonStyle, ButtonStyle)
 
@@ -66,9 +66,9 @@ public:
 
 		SLATE_ARGUMENT(const TArray<OptionType>*, OptionsSource)
 
-		SLATE_TEXT_ATTRIBUTE(RemoveButtonToolTip)
+		SLATE_ATTRIBUTE(FText, RemoveButtonToolTip)
 
-		SLATE_TEXT_ATTRIBUTE(RenameButtonToolTip)
+		SLATE_ATTRIBUTE(FText, RenameButtonToolTip)
 
 	SLATE_END_ARGS()
 
diff --git a/Engine/Source/Runtime/Slate/Public/Widgets/Input/SExpandableButton.h b/Engine/Source/Runtime/Slate/Public/Widgets/Input/SExpandableButton.h
--- a/Engine/Source/Runtime/Slate/Public/Widgets/Input/SExpandableButton.h
+++ b/Engine/Source/Runtime/Slate/Public/Widgets/Input/SExpandableButton.h
@@ -16,10 +16,10 @@ public:
 		{}
 
 		/** The text to display in this button in it's collapsed state (if nothing is specified for CollapsedButtonContent) */
-		SLATE_TEXT_ATTRIBUTE( CollapsedText )
+		SLATE_ATTRIBUTE( FText, CollapsedText )
 
 		/** The text to display in this button in it's expanded state (if nothing is specified for ExpandedButtonContent) */
-		SLATE_TEXT_ATTRIBUTE( ExpandedText )
+		SLATE_ATTRIBUTE( FText, ExpandedText )
 
 		/** Slot for this button's collapsed content (optional) */
 		SLATE_NAMED_SLOT( FArguments, CollapsedButtonContent )
diff --git a/Engine/Source/Runtime/Slate/Public/Widgets/Input/SHyperlink.h b/Engine/Source/Runtime/Slate/Public/Widgets/Input/SHyperlink.h
--- a/Engine/Source/Runtime/Slate/Public/Widgets/Input/SHyperlink.h
+++ b/Engine/Source/Runtime/Slate/Public/Widgets/Input/SHyperlink.h
@@ -21,7 +21,7 @@ public:
 		, _OnNavigate()
 		{}
 
-		SLATE_TEXT_ATTRIBUTE( Text )
+		SLATE_ATTRIBUTE( FText, Text )
 		SLATE_STYLE_ARGUMENT( FHyperlinkStyle, Style )
 		SLATE_STYLE_ARGUMENT( FTextBlockStyle, TextStyle )
 		SLATE_STYLE_ARGUMENT( FButtonStyle, UnderlineStyle )
diff --git a/Engine/Source/Runtime/Slate/Public/Widgets/Input/SNumericEntryBox.h b/Engine/Source/Runtime/Slate/Public/Widgets/Input/SNumericEntryBox.h
--- a/Engine/Source/Runtime/Slate/Public/Widgets/Input/SNumericEntryBox.h
+++ b/Engine/Source/Runtime/Slate/Public/Widgets/Input/SNumericEntryBox.h
@@ -57,7 +57,7 @@ public:
 		/** The value that should be displayed.  This value is optional in the case where a value cannot be determined */
 		SLATE_ATTRIBUTE( TOptional<NumericType>, Value )
 		/** The string to display if the value cannot be determined */
-		SLATE_TEXT_ARGUMENT( UndeterminedString )
+		SLATE_ARGUMENT( FText, UndeterminedString )
 		/** Font color and opacity */
 		SLATE_ATTRIBUTE( FSlateFontInfo, Font )
 		/** Whether or not the user should be able to change the value by dragging with the mouse cursor */
diff --git a/Engine/Source/Runtime/Slate/Public/Widgets/Input/SRichTextHyperlink.h b/Engine/Source/Runtime/Slate/Public/Widgets/Input/SRichTextHyperlink.h
--- a/Engine/Source/Runtime/Slate/Public/Widgets/Input/SRichTextHyperlink.h
+++ b/Engine/Source/Runtime/Slate/Public/Widgets/Input/SRichTextHyperlink.h
@@ -13,7 +13,7 @@ public:
 		, _Style(&FCoreStyle::Get().GetWidgetStyle< FHyperlinkStyle >("Hyperlink"))
 		, _OnNavigate()
 	{}
-		SLATE_TEXT_ATTRIBUTE( Text )
+		SLATE_ATTRIBUTE( FText, Text )
 		SLATE_STYLE_ARGUMENT( FHyperlinkStyle, Style )
 		SLATE_EVENT( FSimpleDelegate, OnNavigate )
 	SLATE_END_ARGS()
diff --git a/Engine/Source/Runtime/Slate/Public/Widgets/Input/STextComboPopup.h b/Engine/Source/Runtime/Slate/Public/Widgets/Input/STextComboPopup.h
--- a/Engine/Source/Runtime/Slate/Public/Widgets/Input/STextComboPopup.h
+++ b/Engine/Source/Runtime/Slate/Public/Widgets/Input/STextComboPopup.h
@@ -22,7 +22,7 @@ public:
 	SLATE_BEGIN_ARGS( STextComboPopup ) : _AutoFocus(true) {}
 
 		/** Label, placed before combo box */
-		SLATE_TEXT_ARGUMENT( Label )
+		SLATE_ARGUMENT( FText, Label )
 
 		/** Selection of strings to pick from */
 		SLATE_ARGUMENT( TArray<FString>, TextOptions )
diff --git a/Engine/Source/Runtime/Slate/Public/Widgets/Input/STextEntryPopup.h b/Engine/Source/Runtime/Slate/Public/Widgets/Input/STextEntryPopup.h
--- a/Engine/Source/Runtime/Slate/Public/Widgets/Input/STextEntryPopup.h
+++ b/Engine/Source/Runtime/Slate/Public/Widgets/Input/STextEntryPopup.h
@@ -14,7 +14,7 @@ public:
 		{}
 
 		/** Label, placed before text entry box */
-		SLATE_TEXT_ARGUMENT( Label )
+		SLATE_ARGUMENT( FText, Label )
 
 		/** Test to place into text entry box before anything is typed */
 		SLATE_ARGUMENT( FText, DefaultText )
diff --git a/Engine/Source/Runtime/Slate/Public/Widgets/Layout/SExpandableArea.h b/Engine/Source/Runtime/Slate/Public/Widgets/Layout/SExpandableArea.h
--- a/Engine/Source/Runtime/Slate/Public/Widgets/Layout/SExpandableArea.h
+++ b/Engine/Source/Runtime/Slate/Public/Widgets/Layout/SExpandableArea.h
@@ -33,7 +33,7 @@ public:
 		/** Content displayed inside the area that is expanded */
 		SLATE_NAMED_SLOT( FArguments, BodyContent )
 		/** The title to display.  Not used if header content is supplied */
-		SLATE_TEXT_ATTRIBUTE( AreaTitle )
+		SLATE_ATTRIBUTE( FText, AreaTitle )
 		/** Whether or not the area is initially collapsed */
 		SLATE_ARGUMENT( bool, InitiallyCollapsed )
 		/** The maximum height of the area */
diff --git a/Engine/Source/Runtime/Slate/Public/Widgets/Notifications/SErrorHint.h b/Engine/Source/Runtime/Slate/Public/Widgets/Notifications/SErrorHint.h
--- a/Engine/Source/Runtime/Slate/Public/Widgets/Notifications/SErrorHint.h
+++ b/Engine/Source/Runtime/Slate/Public/Widgets/Notifications/SErrorHint.h
@@ -11,7 +11,7 @@ public:
 	SLATE_BEGIN_ARGS( SErrorHint )
 		: _ErrorText()
 		{}
-		SLATE_TEXT_ARGUMENT(ErrorText)
+		SLATE_ARGUMENT(FText, ErrorText)
 	SLATE_END_ARGS()
 
 	void Construct(const FArguments& InArgs);
diff --git a/Engine/Source/Runtime/Slate/Public/Widgets/Notifications/SErrorText.h b/Engine/Source/Runtime/Slate/Public/Widgets/Notifications/SErrorText.h
--- a/Engine/Source/Runtime/Slate/Public/Widgets/Notifications/SErrorText.h
+++ b/Engine/Source/Runtime/Slate/Public/Widgets/Notifications/SErrorText.h
@@ -33,7 +33,7 @@ public:
 		, _AutoWrapText(false)
 	{ }
 
-		SLATE_TEXT_ARGUMENT(ErrorText)
+		SLATE_ARGUMENT(FText, ErrorText)
 		SLATE_ATTRIBUTE(FSlateColor, BackgroundColor)
 		SLATE_ATTRIBUTE(FSlateFontInfo, Font)
 		SLATE_ATTRIBUTE(bool, AutoWrapText)
diff --git a/Engine/Source/Runtime/Slate/Public/Widgets/SToolTip.h b/Engine/Source/Runtime/Slate/Public/Widgets/SToolTip.h
--- a/Engine/Source/Runtime/Slate/Public/Widgets/SToolTip.h
+++ b/Engine/Source/Runtime/Slate/Public/Widgets/SToolTip.h
@@ -23,7 +23,7 @@ public:
 	{ }
 
 		/** The text displayed in this tool tip */
-		SLATE_TEXT_ATTRIBUTE(Text)
+		SLATE_ATTRIBUTE(FText, Text)
 
 		/** Arbitrary content to be displayed in the tool tip; overrides any text that may be set. */
 		SLATE_DEFAULT_SLOT(FArguments, Content)
diff --git a/Engine/Source/Runtime/Slate/Public/Widgets/Text/STextBlock.h b/Engine/Source/Runtime/Slate/Public/Widgets/Text/STextBlock.h
--- a/Engine/Source/Runtime/Slate/Public/Widgets/Text/STextBlock.h
+++ b/Engine/Source/Runtime/Slate/Public/Widgets/Text/STextBlock.h
@@ -45,7 +45,7 @@ public:
 		{}
 
 		/** The text displayed in this text block */
-		SLATE_TEXT_ATTRIBUTE( Text )
+		SLATE_ATTRIBUTE( FText, Text )
 
 		/** Pointer to a style of the text block, which dictates the font, color, and shadow options. */
 		SLATE_STYLE_ARGUMENT( FTextBlockStyle, TextStyle )
diff --git a/Engine/Source/Runtime/Slate/Public/Widgets/Views/SHeaderRow.h b/Engine/Source/Runtime/Slate/Public/Widgets/Views/SHeaderRow.h
--- a/Engine/Source/Runtime/Slate/Public/Widgets/Views/SHeaderRow.h
+++ b/Engine/Source/Runtime/Slate/Public/Widgets/Views/SHeaderRow.h
@@ -102,8 +102,8 @@ public:
 			, _OnSort()
 			{}
 			SLATE_ARGUMENT( FName, ColumnId )
-			SLATE_TEXT_ATTRIBUTE( DefaultLabel )
-			SLATE_TEXT_ATTRIBUTE( DefaultTooltip )
+			SLATE_ATTRIBUTE( FText, DefaultLabel )
+			SLATE_ATTRIBUTE( FText, DefaultTooltip )
 			SLATE_ATTRIBUTE( float, FillWidth )
 			SLATE_ARGUMENT( TOptional< float >, FixedWidth )
 			SLATE_ATTRIBUTE( float, ManualWidth )
diff --git a/Engine/Source/Runtime/SlateCore/Private/Widgets/SWidget.cpp b/Engine/Source/Runtime/SlateCore/Private/Widgets/SWidget.cpp
--- a/Engine/Source/Runtime/SlateCore/Private/Widgets/SWidget.cpp
+++ b/Engine/Source/Runtime/SlateCore/Private/Widgets/SWidget.cpp
@@ -598,12 +598,6 @@ FSlateColor SWidget::GetForegroundColor() const
 }
 
 
-void SWidget::SetToolTipText( const TAttribute<FString>& ToolTipString )
-{
-	ToolTip = FSlateApplicationBase::Get().MakeToolTip(ToolTipString);
-}
-
-
 void SWidget::SetToolTipText(const TAttribute<FText>& ToolTipText)
 {
 	ToolTip = FSlateApplicationBase::Get().MakeToolTip(ToolTipText);
diff --git a/Engine/Source/Runtime/SlateCore/Public/Application/SlateApplicationBase.h b/Engine/Source/Runtime/SlateCore/Public/Application/SlateApplicationBase.h
--- a/Engine/Source/Runtime/SlateCore/Public/Application/SlateApplicationBase.h
+++ b/Engine/Source/Runtime/SlateCore/Public/Application/SlateApplicationBase.h
@@ -281,14 +281,6 @@ public:
 	 */
 	virtual TSharedRef<SWidget> MakeImage( const TAttribute<const FSlateBrush*>& Image, const TAttribute<FSlateColor>& Color, const TAttribute<EVisibility>& Visibility ) const = 0;
 
-	/**
-	 * Creates a tool tip with the specified string.
-	 *
-	 * @param ToolTipString The string attribute to assign to the tool tip.
-	 * @return The tool tip.
-	 */
-	virtual TSharedRef<IToolTip> MakeToolTip( const TAttribute<FString>& ToolTipString ) = 0;
-
 	/**
 	 * Creates a tool tip with the specified text.
 	 *
diff --git a/Engine/Source/Runtime/SlateCore/Public/Widgets/DeclarativeSyntaxSupport.h b/Engine/Source/Runtime/SlateCore/Public/Widgets/DeclarativeSyntaxSupport.h
--- a/Engine/Source/Runtime/SlateCore/Public/Widgets/DeclarativeSyntaxSupport.h
+++ b/Engine/Source/Runtime/SlateCore/Public/Widgets/DeclarativeSyntaxSupport.h
@@ -109,10 +109,15 @@ class IToolTip;
 	};
 
 
-#define INTERNAL_SLATE_DECL_ATTRIBUTE( Var, AttrType, AttrName ) \
-		WidgetArgsType& AttrName( const AttrType& InAttribute ) \
+/**
+ * Use this macro to add a attribute to the declaration of your widget.
+ * An attribute can be a value or a function.
+ */
+#define SLATE_ATTRIBUTE( AttrType, AttrName ) \
+		TAttribute< AttrType > _##AttrName; \
+		WidgetArgsType& AttrName( const TAttribute< AttrType >& InAttribute ) \
 		{ \
-			Var = InAttribute; \
+			_##AttrName = InAttribute; \
 			return this->Me(); \
 		} \
 	\
@@ -121,31 +126,31 @@ class IToolTip;
 		template< typename StaticFuncPtr > \
 		WidgetArgsType& AttrName##_Static( StaticFuncPtr InFunc )	\
 		{ \
-			Var = AttrType::Create( AttrType::FGetter::CreateStatic( InFunc ) ); \
+			_##AttrName = TAttribute< AttrType >::Create( TAttribute< AttrType >::FGetter::CreateStatic( InFunc ) ); \
 			return this->Me(); \
 		} \
 		template< typename Var1Type > \
-		WidgetArgsType& AttrName##_Static( typename AttrType::FGetter::template TStaticDelegate_OneVar< Var1Type >::FFuncPtr InFunc, Var1Type Var1 )	\
+		WidgetArgsType& AttrName##_Static( typename TAttribute< AttrType >::FGetter::template TStaticDelegate_OneVar< Var1Type >::FFuncPtr InFunc, Var1Type Var1 )	\
 		{ \
-			Var = AttrType::Create( AttrType::FGetter::CreateStatic( InFunc, Var1 ) ); \
+			_##AttrName = TAttribute< AttrType >::Create( TAttribute< AttrType >::FGetter::CreateStatic( InFunc, Var1 ) ); \
 			return this->Me(); \
 		} \
 		template< typename Var1Type, typename Var2Type > \
-		WidgetArgsType& AttrName##_Static( typename AttrType::FGetter::template TStaticDelegate_TwoVars< Var1Type, Var2Type >::FFuncPtr InFunc, Var1Type Var1, Var2Type Var2 )	\
+		WidgetArgsType& AttrName##_Static( typename TAttribute< AttrType >::FGetter::template TStaticDelegate_TwoVars< Var1Type, Var2Type >::FFuncPtr InFunc, Var1Type Var1, Var2Type Var2 )	\
 		{ \
-			Var = AttrType::Create( AttrType::FGetter::CreateStatic( InFunc, Var1, Var2 ) ); \
+			_##AttrName = TAttribute< AttrType >::Create( TAttribute< AttrType >::FGetter::CreateStatic( InFunc, Var1, Var2 ) ); \
 			return this->Me(); \
 		} \
 		template< typename Var1Type, typename Var2Type, typename Var3Type > \
-		WidgetArgsType& AttrName##_Static( typename AttrType::FGetter::template TStaticDelegate_ThreeVars< Var1Type, Var2Type, Var3Type >::FFuncPtr InFunc, Var1Type Var1, Var2Type Var2, Var3Type Var3 )	\
+		WidgetArgsType& AttrName##_Static( typename TAttribute< AttrType >::FGetter::template TStaticDelegate_ThreeVars< Var1Type, Var2Type, Var3Type >::FFuncPtr InFunc, Var1Type Var1, Var2Type Var2, Var3Type Var3 )	\
 		{ \
-			Var = AttrType::Create( AttrType::FGetter::CreateStatic( InFunc, Var1, Var2, Var3 ) ); \
+			_##AttrName = TAttribute< AttrType >::Create( TAttribute< AttrType >::FGetter::CreateStatic( InFunc, Var1, Var2, Var3 ) ); \
 			return this->Me(); \
 		} \
 		template< typename Var1Type, typename Var2Type, typename Var3Type, typename Var4Type > \
-		WidgetArgsType& AttrName##_Static( typename AttrType::FGetter::template TStaticDelegate_FourVars< Var1Type, Var2Type, Var3Type, Var4Type >::FFuncPtr InFunc, Var1Type Var1, Var2Type Var2, Var3Type Var3, Var4Type Var4 )	\
+		WidgetArgsType& AttrName##_Static( typename TAttribute< AttrType >::FGetter::template TStaticDelegate_FourVars< Var1Type, Var2Type, Var3Type, Var4Type >::FFuncPtr InFunc, Var1Type Var1, Var2Type Var2, Var3Type Var3, Var4Type Var4 )	\
 		{ \
-			Var = AttrType::Create( AttrType::FGetter::CreateStatic( InFunc, Var1, Var2, Var3, Var4 ) ); \
+			_##AttrName = TAttribute< AttrType >::Create( TAttribute< AttrType >::FGetter::CreateStatic( InFunc, Var1, Var2, Var3, Var4 ) ); \
 			return this->Me(); \
 		} \
 	\
@@ -154,381 +159,141 @@ class IToolTip;
 		template<typename FunctorType> \
 		WidgetArgsType& AttrName##_Lambda(FunctorType&& InFunctor) \
 		{ \
-			Var = AttrType::Create(AttrType::FGetter::CreateLambda(Forward<FunctorType>(InFunctor))); \
+			_##AttrName = TAttribute< AttrType >::Create(TAttribute< AttrType >::FGetter::CreateLambda(Forward<FunctorType>(InFunctor))); \
 			return this->Me(); \
 		} \
 	\
 		/* Bind attribute with delegate to a raw C++ class method */ \
 		template< class UserClass >	\
-		WidgetArgsType& AttrName##_Raw( UserClass* InUserObject, typename AttrType::FGetter::template TRawMethodDelegate_Const< UserClass >::FMethodPtr InFunc )	\
+		WidgetArgsType& AttrName##_Raw( UserClass* InUserObject, typename TAttribute< AttrType >::FGetter::template TRawMethodDelegate_Const< UserClass >::FMethodPtr InFunc )	\
 		{ \
-			Var = AttrType::Create( AttrType::FGetter::CreateRaw( InUserObject, InFunc ) ); \
+			_##AttrName = TAttribute< AttrType >::Create( TAttribute< AttrType >::FGetter::CreateRaw( InUserObject, InFunc ) ); \
 			return this->Me(); \
 		} \
 		template< class UserClass, typename Var1Type >	\
-		WidgetArgsType& AttrName##_Raw( UserClass* InUserObject, typename AttrType::FGetter::template TRawMethodDelegate_OneVar_Const< UserClass, Var1Type  >::FMethodPtr InFunc, Var1Type Var1 )	\
+		WidgetArgsType& AttrName##_Raw( UserClass* InUserObject, typename TAttribute< AttrType >::FGetter::template TRawMethodDelegate_OneVar_Const< UserClass, Var1Type  >::FMethodPtr InFunc, Var1Type Var1 )	\
 		{ \
-			Var = AttrType::Create( AttrType::FGetter::CreateRaw( InUserObject, InFunc, Var1 ) ); \
+			_##AttrName = TAttribute< AttrType >::Create( TAttribute< AttrType >::FGetter::CreateRaw( InUserObject, InFunc, Var1 ) ); \
 			return this->Me(); \
 		} \
 		template< class UserClass, typename Var1Type, typename Var2Type >	\
-		WidgetArgsType& AttrName##_Raw( UserClass* InUserObject, typename AttrType::FGetter::template TRawMethodDelegate_TwoVars_Const< UserClass, Var1Type, Var2Type >::FMethodPtr InFunc, Var1Type Var1, Var2Type Var2 )	\
+		WidgetArgsType& AttrName##_Raw( UserClass* InUserObject, typename TAttribute< AttrType >::FGetter::template TRawMethodDelegate_TwoVars_Const< UserClass, Var1Type, Var2Type >::FMethodPtr InFunc, Var1Type Var1, Var2Type Var2 )	\
 		{ \
-			Var = AttrType::Create( AttrType::FGetter::CreateRaw( InUserObject, InFunc, Var1, Var2 ) ); \
+			_##AttrName = TAttribute< AttrType >::Create( TAttribute< AttrType >::FGetter::CreateRaw( InUserObject, InFunc, Var1, Var2 ) ); \
 			return this->Me(); \
 		} \
 		template< class UserClass, typename Var1Type, typename Var2Type, typename Var3Type >	\
-		WidgetArgsType& AttrName##_Raw( UserClass* InUserObject, typename AttrType::FGetter::template TRawMethodDelegate_ThreeVars_Const< UserClass, Var1Type, Var2Type, Var3Type >::FMethodPtr InFunc, Var1Type Var1, Var2Type Var2, Var3Type Var3 )	\
+		WidgetArgsType& AttrName##_Raw( UserClass* InUserObject, typename TAttribute< AttrType >::FGetter::template TRawMethodDelegate_ThreeVars_Const< UserClass, Var1Type, Var2Type, Var3Type >::FMethodPtr InFunc, Var1Type Var1, Var2Type Var2, Var3Type Var3 )	\
 		{ \
-			Var = AttrType::Create( AttrType::FGetter::CreateRaw( InUserObject, InFunc, Var1, Var2, Var3 ) ); \
+			_##AttrName = TAttribute< AttrType >::Create( TAttribute< AttrType >::FGetter::CreateRaw( InUserObject, InFunc, Var1, Var2, Var3 ) ); \
 			return this->Me(); \
 		} \
 		template< class UserClass, typename Var1Type, typename Var2Type, typename Var3Type, typename Var4Type >	\
-		WidgetArgsType& AttrName##_Raw( UserClass* InUserObject, typename AttrType::FGetter::template TRawMethodDelegate_FourVars_Const< UserClass, Var1Type, Var2Type, Var3Type, Var4Type >::FMethodPtr InFunc, Var1Type Var1, Var2Type Var2, Var3Type Var3, Var4Type Var4 )	\
+		WidgetArgsType& AttrName##_Raw( UserClass* InUserObject, typename TAttribute< AttrType >::FGetter::template TRawMethodDelegate_FourVars_Const< UserClass, Var1Type, Var2Type, Var3Type, Var4Type >::FMethodPtr InFunc, Var1Type Var1, Var2Type Var2, Var3Type Var3, Var4Type Var4 )	\
 		{ \
-			Var = AttrType::Create( AttrType::FGetter::CreateRaw( InUserObject, InFunc, Var1, Var2, Var3, Var4 ) ); \
+			_##AttrName = TAttribute< AttrType >::Create( TAttribute< AttrType >::FGetter::CreateRaw( InUserObject, InFunc, Var1, Var2, Var3, Var4 ) ); \
 			return this->Me(); \
 		} \
 	\
 		/* Bind attribute with delegate to a shared pointer-based class method.  Slate mostly uses shared pointers so we use an overload for this type of binding. */ \
 		template< class UserClass >	\
-		WidgetArgsType& AttrName( TSharedRef< UserClass > InUserObjectRef, typename AttrType::FGetter::template TSPMethodDelegate_Const< UserClass >::FMethodPtr InFunc )	\
+		WidgetArgsType& AttrName( TSharedRef< UserClass > InUserObjectRef, typename TAttribute< AttrType >::FGetter::template TSPMethodDelegate_Const< UserClass >::FMethodPtr InFunc )	\
 		{ \
-			Var = AttrType( InUserObjectRef, InFunc ); \
+			_##AttrName = TAttribute< AttrType >( InUserObjectRef, InFunc ); \
 			return this->Me(); \
 		} \
 		template< class UserClass, typename Var1Type >	\
-		WidgetArgsType& AttrName( TSharedRef< UserClass > InUserObjectRef, typename AttrType::FGetter::template TSPMethodDelegate_OneVar_Const< UserClass, Var1Type >::FMethodPtr InFunc, Var1Type Var1 )	\
+		WidgetArgsType& AttrName( TSharedRef< UserClass > InUserObjectRef, typename TAttribute< AttrType >::FGetter::template TSPMethodDelegate_OneVar_Const< UserClass, Var1Type >::FMethodPtr InFunc, Var1Type Var1 )	\
 		{ \
-			Var = AttrType::Create( AttrType::FGetter::CreateSP( InUserObjectRef, InFunc, Var1 ) ); \
+			_##AttrName = TAttribute< AttrType >::Create( TAttribute< AttrType >::FGetter::CreateSP( InUserObjectRef, InFunc, Var1 ) ); \
 			return this->Me(); \
 		} \
 		template< class UserClass, typename Var1Type, typename Var2Type >	\
-		WidgetArgsType& AttrName( TSharedRef< UserClass > InUserObjectRef, typename AttrType::FGetter::template TSPMethodDelegate_TwoVars_Const< UserClass, Var1Type, Var2Type >::FMethodPtr InFunc, Var1Type Var1, Var2Type Var2 )	\
+		WidgetArgsType& AttrName( TSharedRef< UserClass > InUserObjectRef, typename TAttribute< AttrType >::FGetter::template TSPMethodDelegate_TwoVars_Const< UserClass, Var1Type, Var2Type >::FMethodPtr InFunc, Var1Type Var1, Var2Type Var2 )	\
 		{ \
-			Var = AttrType::Create( AttrType::FGetter::CreateSP( InUserObjectRef, InFunc, Var1, Var2 ) ); \
+			_##AttrName = TAttribute< AttrType >::Create( TAttribute< AttrType >::FGetter::CreateSP( InUserObjectRef, InFunc, Var1, Var2 ) ); \
 			return this->Me(); \
 		} \
 		template< class UserClass, typename Var1Type, typename Var2Type, typename Var3Type >	\
-		WidgetArgsType& AttrName( TSharedRef< UserClass > InUserObjectRef, typename AttrType::FGetter::template TSPMethodDelegate_ThreeVars_Const< UserClass, Var1Type, Var2Type, Var3Type >::FMethodPtr InFunc, Var1Type Var1, Var2Type Var2, Var3Type Var3 )	\
+		WidgetArgsType& AttrName( TSharedRef< UserClass > InUserObjectRef, typename TAttribute< AttrType >::FGetter::template TSPMethodDelegate_ThreeVars_Const< UserClass, Var1Type, Var2Type, Var3Type >::FMethodPtr InFunc, Var1Type Var1, Var2Type Var2, Var3Type Var3 )	\
 		{ \
-			Var = AttrType::Create( AttrType::FGetter::CreateSP( InUserObjectRef, InFunc, Var1, Var2, Var3 ) ); \
+			_##AttrName = TAttribute< AttrType >::Create( TAttribute< AttrType >::FGetter::CreateSP( InUserObjectRef, InFunc, Var1, Var2, Var3 ) ); \
 			return this->Me(); \
 		} \
 		template< class UserClass, typename Var1Type, typename Var2Type, typename Var3Type, typename Var4Type >	\
-		WidgetArgsType& AttrName( TSharedRef< UserClass > InUserObjectRef, typename AttrType::FGetter::template TSPMethodDelegate_FourVars_Const< UserClass, Var1Type, Var2Type, Var3Type, Var4Type >::FMethodPtr InFunc, Var1Type Var1, Var2Type Var2, Var3Type Var3, Var4Type Var4 )	\
+		WidgetArgsType& AttrName( TSharedRef< UserClass > InUserObjectRef, typename TAttribute< AttrType >::FGetter::template TSPMethodDelegate_FourVars_Const< UserClass, Var1Type, Var2Type, Var3Type, Var4Type >::FMethodPtr InFunc, Var1Type Var1, Var2Type Var2, Var3Type Var3, Var4Type Var4 )	\
 		{ \
-			Var = AttrType::Create( AttrType::FGetter::CreateSP( InUserObjectRef, InFunc, Var1, Var2, Var3, Var4 ) ); \
+			_##AttrName = TAttribute< AttrType >::Create( TAttribute< AttrType >::FGetter::CreateSP( InUserObjectRef, InFunc, Var1, Var2, Var3, Var4 ) ); \
 			return this->Me(); \
 		} \
 	\
 		/* Bind attribute with delegate to a shared pointer-based class method.  Slate mostly uses shared pointers so we use an overload for this type of binding. */ \
 		template< class UserClass >	\
-		WidgetArgsType& AttrName( UserClass* InUserObject, typename AttrType::FGetter::template TSPMethodDelegate_Const< UserClass >::FMethodPtr InFunc )	\
+		WidgetArgsType& AttrName( UserClass* InUserObject, typename TAttribute< AttrType >::FGetter::template TSPMethodDelegate_Const< UserClass >::FMethodPtr InFunc )	\
 		{ \
-			Var = AttrType::Create( AttrType::FGetter::CreateSP( InUserObject, InFunc ) ); \
+			_##AttrName = TAttribute< AttrType >::Create( TAttribute< AttrType >::FGetter::CreateSP( InUserObject, InFunc ) ); \
 			return this->Me(); \
 		} \
 		template< class UserClass, typename Var1Type >	\
-		WidgetArgsType& AttrName( UserClass* InUserObject, typename AttrType::FGetter::template TSPMethodDelegate_OneVar_Const< UserClass, Var1Type >::FMethodPtr InFunc, Var1Type Var1 )	\
+		WidgetArgsType& AttrName( UserClass* InUserObject, typename TAttribute< AttrType >::FGetter::template TSPMethodDelegate_OneVar_Const< UserClass, Var1Type >::FMethodPtr InFunc, Var1Type Var1 )	\
 		{ \
-			Var = AttrType::Create( AttrType::FGetter::CreateSP( InUserObject, InFunc, Var1 ) ); \
+			_##AttrName = TAttribute< AttrType >::Create( TAttribute< AttrType >::FGetter::CreateSP( InUserObject, InFunc, Var1 ) ); \
 			return this->Me(); \
 		} \
 		template< class UserClass, typename Var1Type, typename Var2Type >	\
-		WidgetArgsType& AttrName( UserClass* InUserObject, typename AttrType::FGetter::template TSPMethodDelegate_TwoVars_Const< UserClass, Var1Type, Var2Type >::FMethodPtr InFunc, Var1Type Var1, Var2Type Var2 )	\
+		WidgetArgsType& AttrName( UserClass* InUserObject, typename TAttribute< AttrType >::FGetter::template TSPMethodDelegate_TwoVars_Const< UserClass, Var1Type, Var2Type >::FMethodPtr InFunc, Var1Type Var1, Var2Type Var2 )	\
 		{ \
-			Var = AttrType::Create( AttrType::FGetter::CreateSP( InUserObject, InFunc, Var1, Var2 ) ); \
+			_##AttrName = TAttribute< AttrType >::Create( TAttribute< AttrType >::FGetter::CreateSP( InUserObject, InFunc, Var1, Var2 ) ); \
 			return this->Me(); \
 		} \
 		template< class UserClass, typename Var1Type, typename Var2Type, typename Var3Type >	\
-		WidgetArgsType& AttrName( UserClass* InUserObject, typename AttrType::FGetter::template TSPMethodDelegate_ThreeVars_Const< UserClass, Var1Type, Var2Type, Var3Type >::FMethodPtr InFunc, Var1Type Var1, Var2Type Var2, Var3Type Var3 )	\
+		WidgetArgsType& AttrName( UserClass* InUserObject, typename TAttribute< AttrType >::FGetter::template TSPMethodDelegate_ThreeVars_Const< UserClass, Var1Type, Var2Type, Var3Type >::FMethodPtr InFunc, Var1Type Var1, Var2Type Var2, Var3Type Var3 )	\
 		{ \
-			Var = AttrType::Create( AttrType::FGetter::CreateSP( InUserObject, InFunc, Var1, Var2, Var3 ) ); \
+			_##AttrName = TAttribute< AttrType >::Create( TAttribute< AttrType >::FGetter::CreateSP( InUserObject, InFunc, Var1, Var2, Var3 ) ); \
 			return this->Me(); \
 		} \
 		template< class UserClass, typename Var1Type, typename Var2Type, typename Var3Type, typename Var4Type >	\
-		WidgetArgsType& AttrName( UserClass* InUserObject, typename AttrType::FGetter::template TSPMethodDelegate_FourVars_Const< UserClass, Var1Type, Var2Type, Var3Type, Var4Type >::FMethodPtr InFunc, Var1Type Var1, Var2Type Var2, Var3Type Var3, Var4Type Var4 )	\
+		WidgetArgsType& AttrName( UserClass* InUserObject, typename TAttribute< AttrType >::FGetter::template TSPMethodDelegate_FourVars_Const< UserClass, Var1Type, Var2Type, Var3Type, Var4Type >::FMethodPtr InFunc, Var1Type Var1, Var2Type Var2, Var3Type Var3, Var4Type Var4 )	\
 		{ \
-			Var = AttrType::Create( AttrType::FGetter::CreateSP( InUserObject, InFunc, Var1, Var2, Var3, Var4 ) ); \
+			_##AttrName = TAttribute< AttrType >::Create( TAttribute< AttrType >::FGetter::CreateSP( InUserObject, InFunc, Var1, Var2, Var3, Var4 ) ); \
 			return this->Me(); \
 		} \
 	\
 		/* Bind attribute with delegate to a UObject-based class method */ \
 		template< class UserClass >	\
-		WidgetArgsType& AttrName##_UObject( UserClass* InUserObject, typename AttrType::FGetter::template TUObjectMethodDelegate_Const< UserClass >::FMethodPtr InFunc )	\
+		WidgetArgsType& AttrName##_UObject( UserClass* InUserObject, typename TAttribute< AttrType >::FGetter::template TUObjectMethodDelegate_Const< UserClass >::FMethodPtr InFunc )	\
 		{ \
-			Var = AttrType::Create( AttrType::FGetter::CreateUObject( InUserObject, InFunc ) ); \
+			_##AttrName = TAttribute< AttrType >::Create( TAttribute< AttrType >::FGetter::CreateUObject( InUserObject, InFunc ) ); \
 			return this->Me(); \
 		} \
 		template< class UserClass, typename Var1Type >	\
-		WidgetArgsType& AttrName##_UObject( UserClass* InUserObject, typename AttrType::FGetter::template TUObjectMethodDelegate_OneVar_Const< UserClass, Var1Type >::FMethodPtr InFunc, Var1Type Var1 )	\
+		WidgetArgsType& AttrName##_UObject( UserClass* InUserObject, typename TAttribute< AttrType >::FGetter::template TUObjectMethodDelegate_OneVar_Const< UserClass, Var1Type >::FMethodPtr InFunc, Var1Type Var1 )	\
 		{ \
-			Var = AttrType::Create( AttrType::FGetter::CreateUObject( InUserObject, InFunc, Var1 ) ); \
+			_##AttrName = TAttribute< AttrType >::Create( TAttribute< AttrType >::FGetter::CreateUObject( InUserObject, InFunc, Var1 ) ); \
 			return this->Me(); \
 		} \
 		template< class UserClass, typename Var1Type, typename Var2Type  >	\
-		WidgetArgsType& AttrName##_UObject( UserClass* InUserObject, typename AttrType::FGetter::template TUObjectMethodDelegate_TwoVars_Const< UserClass, Var1Type, Var2Type >::FMethodPtr InFunc, Var1Type Var1, Var2Type Var2 )	\
+		WidgetArgsType& AttrName##_UObject( UserClass* InUserObject, typename TAttribute< AttrType >::FGetter::template TUObjectMethodDelegate_TwoVars_Const< UserClass, Var1Type, Var2Type >::FMethodPtr InFunc, Var1Type Var1, Var2Type Var2 )	\
 		{ \
-			Var = AttrType::Create( AttrType::FGetter::CreateUObject( InUserObject, InFunc, Var1, Var2 ) ); \
+			_##AttrName = TAttribute< AttrType >::Create( TAttribute< AttrType >::FGetter::CreateUObject( InUserObject, InFunc, Var1, Var2 ) ); \
 			return this->Me(); \
 		} \
 		template< class UserClass, typename Var1Type, typename Var2Type, typename Var3Type >	\
-		WidgetArgsType& AttrName##_UObject( UserClass* InUserObject, typename AttrType::FGetter::template TUObjectMethodDelegate_ThreeVars_Const< UserClass, Var1Type, Var2Type, Var3Type >::FMethodPtr InFunc, Var1Type Var1, Var2Type Var2, Var3Type Var3 )	\
+		WidgetArgsType& AttrName##_UObject( UserClass* InUserObject, typename TAttribute< AttrType >::FGetter::template TUObjectMethodDelegate_ThreeVars_Const< UserClass, Var1Type, Var2Type, Var3Type >::FMethodPtr InFunc, Var1Type Var1, Var2Type Var2, Var3Type Var3 )	\
 		{ \
-			Var = AttrType::Create( AttrType::FGetter::CreateUObject( InUserObject, InFunc, Var1, Var2, Var3 ) ); \
+			_##AttrName = TAttribute< AttrType >::Create( TAttribute< AttrType >::FGetter::CreateUObject( InUserObject, InFunc, Var1, Var2, Var3 ) ); \
 			return this->Me(); \
 		} \
 		template< class UserClass, typename Var1Type, typename Var2Type, typename Var3Type, typename Var4Type >	\
-		WidgetArgsType& AttrName##_UObject( UserClass* InUserObject, typename AttrType::FGetter::template TUObjectMethodDelegate_FourVars_Const< UserClass, Var1Type, Var2Type, Var3Type, Var4Type >::FMethodPtr InFunc, Var1Type Var1, Var2Type Var2, Var3Type Var3, Var4Type Var4 )	\
+		WidgetArgsType& AttrName##_UObject( UserClass* InUserObject, typename TAttribute< AttrType >::FGetter::template TUObjectMethodDelegate_FourVars_Const< UserClass, Var1Type, Var2Type, Var3Type, Var4Type >::FMethodPtr InFunc, Var1Type Var1, Var2Type Var2, Var3Type Var3, Var4Type Var4 )	\
 		{ \
-			Var = AttrType::Create( AttrType::FGetter::CreateUObject( InUserObject, InFunc, Var1, Var2, Var3, Var4 ) ); \
+			_##AttrName = TAttribute< AttrType >::Create( TAttribute< AttrType >::FGetter::CreateUObject( InUserObject, InFunc, Var1, Var2, Var3, Var4 ) ); \
 			return this->Me(); \
 		}
 
 
 /**
- * Use this macro to add a attribute to the declaration of your widget.
- * An attribute can be a value or a function.
- */
-#define SLATE_ATTRIBUTE( AttrType, AttrName ) \
-	TAttribute< AttrType > _##AttrName; \
-	INTERNAL_SLATE_DECL_ATTRIBUTE( _##AttrName, TAttribute< AttrType >, AttrName )
-
-
-#define SLATE_TEXT_ATTRIBUTE_IMPL( Var, AttrName ) \
-	TAttribute< FText > Var; \
-	struct AttrName##_Local\
-	{\
-		static FText PassThroughDelegate( TAttribute< FString >::FGetter StringDelegate )\
-		{\
-			return (StringDelegate.IsBound()) ? FText::FromString( StringDelegate.Execute() ) : FText::GetEmpty();\
-		}\
-		\
-		static FText PassThroughAttribute( TAttribute< FString > StringAttribute )\
-		{\
-			return FText::FromString( StringAttribute.Get( FString() ) );\
-		}\
-	};\
-	WidgetArgsType& AttrName( const FText& Text ) \
-	{ \
-		Var = TAttribute< FText >( Text ); \
-		return this->Me(); \
-	} \
-	DEPRECATED(4.8, "Passing text to Slate as FString is deprecated, please use FText instead (likely via a LOCTEXT).")\
-	WidgetArgsType& AttrName( const FString& String ) \
-	{ \
-		Var = TAttribute< FText >( FText::FromString( String ) ); \
-		return this->Me(); \
-	} \
-	DEPRECATED(4.8, "Passing text to Slate as FString is deprecated, please use FText instead (likely via a LOCTEXT).")\
-	WidgetArgsType& AttrName( const TAttribute< FString >& InAttribute ) \
-	{ \
-		if ( InAttribute.IsBound() ) {\
-			Var = TAttribute< FText >::Create( TAttribute< FText >::FGetter::CreateStatic( &AttrName##_Local::PassThroughAttribute, InAttribute ) ); \
-		} else {\
-			Var = FText::FromString( InAttribute.Get() );\
-		}\
-		return this->Me(); \
-	} \
-	/* Bind attribute with delegate to a global function \
-	* NOTE: We use a template here to avoid 'typename' issues when hosting attributes inside templated classes */ \
-	DEPRECATED(4.8, "Passing text to Slate as FString is deprecated, please use FText instead (likely via a LOCTEXT).")\
-	WidgetArgsType& AttrName##_Static( TAttribute< FString >::FGetter::FStaticDelegate::FFuncPtr InFunc )	\
-	{ \
-		Var = TAttribute< FText >::Create( TAttribute< FText >::FGetter::CreateStatic( &AttrName##_Local::PassThroughDelegate, TAttribute< FString >::FGetter::CreateStatic( InFunc ) ) ); \
-		return this->Me(); \
-	} \
-	template< typename Var1Type > \
-	DEPRECATED(4.8, "Passing text to Slate as FString is deprecated, please use FText instead (likely via a LOCTEXT).")\
-	WidgetArgsType& AttrName##_Static( typename TAttribute< FString >::FGetter::template TStaticDelegate_OneVar< Var1Type >::FFuncPtr InFunc, Var1Type Var1 )	\
-	{ \
-		Var = TAttribute< FText >::Create( TAttribute< FText >::FGetter::CreateStatic( &AttrName##_Local::PassThroughDelegate, TAttribute< FString >::FGetter::CreateStatic( InFunc, Var1 ) ) ); \
-		return this->Me(); \
-	} \
-	template< typename Var1Type, typename Var2Type > \
-	DEPRECATED(4.8, "Passing text to Slate as FString is deprecated, please use FText instead (likely via a LOCTEXT).")\
-	WidgetArgsType& AttrName##_Static( typename TAttribute< FString >::FGetter::template TStaticDelegate_TwoVars< Var1Type, Var2Type >::FFuncPtr InFunc, Var1Type Var1, Var2Type Var2 )	\
-	{ \
-		Var = TAttribute< FText >::Create( TAttribute< FText >::FGetter::CreateStatic( &AttrName##_Local::PassThroughDelegate, TAttribute< FString >::FGetter::CreateStatic( InFunc, Var1, Var2 ) ) ); \
-		return this->Me(); \
-	} \
-	template< typename Var1Type, typename Var2Type, typename Var3Type > \
-	DEPRECATED(4.8, "Passing text to Slate as FString is deprecated, please use FText instead (likely via a LOCTEXT).")\
-	WidgetArgsType& AttrName##_Static( typename TAttribute< FString >::FGetter::template TStaticDelegate_ThreeVars< Var1Type, Var2Type, Var3Type >::FFuncPtr InFunc, Var1Type Var1, Var2Type Var2, Var3Type Var3 )	\
-	{ \
-		Var = TAttribute< FText >::Create( TAttribute< FText >::FGetter::CreateStatic( &AttrName##_Local::PassThroughDelegate, TAttribute< FString >::FGetter::CreateStatic( InFunc, Var1, Var2, Var3 ) ) ); \
-		return this->Me(); \
-	} \
-	template< typename Var1Type, typename Var2Type, typename Var3Type, typename Var4Type > \
-	DEPRECATED(4.8, "Passing text to Slate as FString is deprecated, please use FText instead (likely via a LOCTEXT).")\
-	WidgetArgsType& AttrName##_Static( typename TAttribute< FString >::FGetter::template TStaticDelegate_FourVars< Var1Type, Var2Type, Var3Type, Var4Type >::FFuncPtr InFunc, Var1Type Var1, Var2Type Var2, Var3Type Var3, Var4Type Var4 )	\
-	{ \
-		Var = TAttribute< FText >::Create( TAttribute< FText >::FGetter::CreateStatic( &AttrName##_Local::PassThroughDelegate, TAttribute< FString >::FGetter::CreateStatic( InFunc, Var1, Var2, Var3, Var4 ) ) ); \
-		return this->Me(); \
-	} \
-\
-	/* We don't support binding FString lambdas to FText attributes because we currently have no way to disambiguate them from FText lambdas \
-	template<typename FunctorType> \
-	WidgetArgsType& AttrName##_Lambda(FunctorType&& InFunctor) */ \
-\
-	/* Bind attribute with delegate to a raw C++ class method */ \
-	template< class UserClass >	\
-	DEPRECATED(4.8, "Passing text to Slate as FString is deprecated, please use FText instead (likely via a LOCTEXT).")\
-	WidgetArgsType& AttrName##_Raw( UserClass* InUserObject, typename TAttribute< FString >::FGetter::template TRawMethodDelegate_Const< UserClass >::FMethodPtr InFunc )	\
-	{ \
-		Var = TAttribute< FText >::Create( TAttribute< FText >::FGetter::CreateStatic( &AttrName##_Local::PassThroughDelegate, TAttribute< FString >::FGetter::CreateRaw( InUserObject, InFunc ) ) ); \
-		return this->Me(); \
-	} \
-	template< class UserClass, typename Var1Type >	\
-	DEPRECATED(4.8, "Passing text to Slate as FString is deprecated, please use FText instead (likely via a LOCTEXT).")\
-	WidgetArgsType& AttrName##_Raw( UserClass* InUserObject, typename TAttribute< FString >::FGetter::template TRawMethodDelegate_OneVar_Const< UserClass, Var1Type  >::FMethodPtr InFunc, Var1Type Var1 )	\
-	{ \
-		Var = TAttribute< FText >::Create( TAttribute< FText >::FGetter::CreateStatic( &AttrName##_Local::PassThroughDelegate, TAttribute< FString >::FGetter::CreateRaw( InUserObject, InFunc, Var1 ) ) ); \
-		return this->Me(); \
-	} \
-	template< class UserClass, typename Var1Type, typename Var2Type >	\
-	DEPRECATED(4.8, "Passing text to Slate as FString is deprecated, please use FText instead (likely via a LOCTEXT).")\
-	WidgetArgsType& AttrName##_Raw( UserClass* InUserObject, typename TAttribute< FString >::FGetter::template TRawMethodDelegate_TwoVars_Const< UserClass, Var1Type, Var2Type >::FMethodPtr InFunc, Var1Type Var1, Var2Type Var2 )	\
-	{ \
-		Var = TAttribute< FText >::Create( TAttribute< FText >::FGetter::CreateStatic( &AttrName##_Local::PassThroughDelegate, TAttribute< FString >::FGetter::CreateRaw( InUserObject, InFunc, Var1, Var2 ) ) ); \
-		return this->Me(); \
-	} \
-	template< class UserClass, typename Var1Type, typename Var2Type, typename Var3Type >	\
-	DEPRECATED(4.8, "Passing text to Slate as FString is deprecated, please use FText instead (likely via a LOCTEXT).")\
-	WidgetArgsType& AttrName##_Raw( UserClass* InUserObject, typename TAttribute< FString >::FGetter::template TRawMethodDelegate_ThreeVars_Const< UserClass, Var1Type, Var2Type, Var3Type >::FMethodPtr InFunc, Var1Type Var1, Var2Type Var2, Var3Type Var3 )	\
-	{ \
-		Var = TAttribute< FText >::Create( TAttribute< FText >::FGetter::CreateStatic( &AttrName##_Local::PassThroughDelegate, TAttribute< FString >::FGetter::CreateRaw( InUserObject, InFunc, Var1, Var2, Var3 ) ) ); \
-		return this->Me(); \
-	} \
-	template< class UserClass, typename Var1Type, typename Var2Type, typename Var3Type, typename Var4Type >	\
-	DEPRECATED(4.8, "Passing text to Slate as FString is deprecated, please use FText instead (likely via a LOCTEXT).")\
-	WidgetArgsType& AttrName##_Raw( UserClass* InUserObject, typename TAttribute< FString >::FGetter::template TRawMethodDelegate_FourVars_Const< UserClass, Var1Type, Var2Type, Var3Type, Var4Type >::FMethodPtr InFunc, Var1Type Var1, Var2Type Var2, Var3Type Var3, Var4Type Var4 )	\
-	{ \
-		Var = TAttribute< FText >::Create( TAttribute< FText >::FGetter::CreateStatic( &AttrName##_Local::PassThroughDelegate, TAttribute< FString >::FGetter::CreateRaw( InUserObject, InFunc, Var1, Var2, Var3, Var4 ) ) ); \
-		return this->Me(); \
-	} \
-\
-	/* Bind attribute with delegate to a shared pointer-based class method.  Slate mostly uses shared pointers so we use an overload for this type of binding. */ \
-	template< class UserClass >	\
-	DEPRECATED(4.8, "Passing text to Slate as FString is deprecated, please use FText instead (likely via a LOCTEXT).")\
-	WidgetArgsType& AttrName( TSharedRef< UserClass > InUserObjectRef, typename TAttribute< FString >::FGetter::template TSPMethodDelegate_Const< UserClass >::FMethodPtr InFunc )	\
-	{ \
-		Var = TAttribute< FText >::Create( TAttribute< FText >::FGetter::CreateStatic( &AttrName##_Local::PassThroughDelegate, TAttribute< FString >::FGetter::CreateSP( InUserObjectRef, InFunc ) ) ); \
-		return this->Me(); \
-	} \
-	template< class UserClass, typename Var1Type >	\
-	DEPRECATED(4.8, "Passing text to Slate as FString is deprecated, please use FText instead (likely via a LOCTEXT).")\
-	WidgetArgsType& AttrName( TSharedRef< UserClass > InUserObjectRef, typename TAttribute< FString >::FGetter::template TSPMethodDelegate_OneVar_Const< UserClass, Var1Type >::FMethodPtr InFunc, Var1Type Var1 )	\
-	{ \
-		Var = TAttribute< FText >::Create( TAttribute< FText >::FGetter::CreateStatic( &AttrName##_Local::PassThroughDelegate, TAttribute< FString >::FGetter::CreateSP( InUserObjectRef, InFunc, Var1 ) ) ); \
-		return this->Me(); \
-	} \
-	template< class UserClass, typename Var1Type, typename Var2Type >	\
-	DEPRECATED(4.8, "Passing text to Slate as FString is deprecated, please use FText instead (likely via a LOCTEXT).")\
-	WidgetArgsType& AttrName( TSharedRef< UserClass > InUserObjectRef, typename TAttribute< FString >::FGetter::template TSPMethodDelegate_TwoVars_Const< UserClass, Var1Type, Var2Type >::FMethodPtr InFunc, Var1Type Var1, Var2Type Var2 )	\
-	{ \
-		Var = TAttribute< FText >::Create( TAttribute< FText >::FGetter::CreateStatic( &AttrName##_Local::PassThroughDelegate, TAttribute< FString >::FGetter::CreateSP( InUserObjectRef, InFunc, Var1, Var2 ) ) ); \
-		return this->Me(); \
-	} \
-	template< class UserClass, typename Var1Type, typename Var2Type, typename Var3Type >	\
-	DEPRECATED(4.8, "Passing text to Slate as FString is deprecated, please use FText instead (likely via a LOCTEXT).")\
-	WidgetArgsType& AttrName( TSharedRef< UserClass > InUserObjectRef, typename TAttribute< FString >::FGetter::template TSPMethodDelegate_ThreeVars_Const< UserClass, Var1Type, Var2Type, Var3Type >::FMethodPtr InFunc, Var1Type Var1, Var2Type Var2, Var3Type Var3 )	\
-	{ \
-		Var = TAttribute< FText >::Create( TAttribute< FText >::FGetter::CreateStatic( &AttrName##_Local::PassThroughDelegate, TAttribute< FString >::FGetter::CreateSP( InUserObjectRef, InFunc, Var1, Var2, Var3 ) ) ); \
-		return this->Me(); \
-	} \
-	template< class UserClass, typename Var1Type, typename Var2Type, typename Var3Type, typename Var4Type >	\
-	DEPRECATED(4.8, "Passing text to Slate as FString is deprecated, please use FText instead (likely via a LOCTEXT).")\
-	WidgetArgsType& AttrName( TSharedRef< UserClass > InUserObjectRef, typename TAttribute< FString >::FGetter::template TSPMethodDelegate_FourVars_Const< UserClass, Var1Type, Var2Type, Var3Type, Var4Type >::FMethodPtr InFunc, Var1Type Var1, Var2Type Var2, Var3Type Var3, Var4Type Var4 )	\
-	{ \
-		Var = TAttribute< FText >::Create( TAttribute< FText >::FGetter::CreateStatic( &AttrName##_Local::PassThroughDelegate, TAttribute< FString >::FGetter::CreateSP( InUserObjectRef, InFunc, Var1, Var2, Var3, Var4 ) ) ); \
-		return this->Me(); \
-	} \
-\
-	/* Bind attribute with delegate to a shared pointer-based class method.  Slate mostly uses shared pointers so we use an overload for this type of binding. */ \
-	template< class UserClass >	\
-	DEPRECATED(4.8, "Passing text to Slate as FString is deprecated, please use FText instead (likely via a LOCTEXT).")\
-	WidgetArgsType& AttrName( UserClass* InUserObject, typename TAttribute< FString >::FGetter::template TSPMethodDelegate_Const< UserClass >::FMethodPtr InFunc )	\
-	{ \
-		Var = TAttribute< FText >::Create( TAttribute< FText >::FGetter::CreateStatic( &AttrName##_Local::PassThroughDelegate, TAttribute< FString >::FGetter::CreateSP( InUserObject, InFunc ) ) ); \
-		return this->Me(); \
-	} \
-	template< class UserClass, typename Var1Type >	\
-	DEPRECATED(4.8, "Passing text to Slate as FString is deprecated, please use FText instead (likely via a LOCTEXT).")\
-	WidgetArgsType& AttrName( UserClass* InUserObject, typename TAttribute< FString >::FGetter::template TSPMethodDelegate_OneVar_Const< UserClass, Var1Type >::FMethodPtr InFunc, Var1Type Var1 )	\
-	{ \
-		Var = TAttribute< FText >::Create( TAttribute< FText >::FGetter::CreateStatic( &AttrName##_Local::PassThroughDelegate, TAttribute< FString >::FGetter::CreateSP( InUserObject, InFunc, Var1 ) ) ); \
-		return this->Me(); \
-	} \
-	template< class UserClass, typename Var1Type, typename Var2Type >	\
-	DEPRECATED(4.8, "Passing text to Slate as FString is deprecated, please use FText instead (likely via a LOCTEXT).")\
-	WidgetArgsType& AttrName( UserClass* InUserObject, typename TAttribute< FString >::FGetter::template TSPMethodDelegate_TwoVars_Const< UserClass, Var1Type, Var2Type >::FMethodPtr InFunc, Var1Type Var1, Var2Type Var2 )	\
-	{ \
-		Var = TAttribute< FText >::Create( TAttribute< FText >::FGetter::CreateStatic( &AttrName##_Local::PassThroughDelegate, TAttribute< FString >::FGetter::CreateSP( InUserObject, InFunc, Var1, Var2 ) ) ); \
-		return this->Me(); \
-	} \
-	template< class UserClass, typename Var1Type, typename Var2Type, typename Var3Type >	\
-	DEPRECATED(4.8, "Passing text to Slate as FString is deprecated, please use FText instead (likely via a LOCTEXT).")\
-	WidgetArgsType& AttrName( UserClass* InUserObject, typename TAttribute< FString >::FGetter::template TSPMethodDelegate_ThreeVars_Const< UserClass, Var1Type, Var2Type, Var3Type >::FMethodPtr InFunc, Var1Type Var1, Var2Type Var2, Var3Type Var3 )	\
-	{ \
-		Var = TAttribute< FText >::Create( TAttribute< FText >::FGetter::CreateStatic( &AttrName##_Local::PassThroughDelegate, TAttribute< FString >::FGetter::CreateSP( InUserObject, InFunc, Var1, Var2, Var3 ) ) ); \
-		return this->Me(); \
-	} \
-	template< class UserClass, typename Var1Type, typename Var2Type, typename Var3Type, typename Var4Type >	\
-	DEPRECATED(4.8, "Passing text to Slate as FString is deprecated, please use FText instead (likely via a LOCTEXT).")\
-	WidgetArgsType& AttrName( UserClass* InUserObject, typename TAttribute< FString >::FGetter::template TSPMethodDelegate_FourVars_Const< UserClass, Var1Type, Var2Type, Var3Type, Var4Type >::FMethodPtr InFunc, Var1Type Var1, Var2Type Var2, Var3Type Var3, Var4Type Var4 )	\
-	{ \
-		Var = TAttribute< FText >::Create( TAttribute< FText >::FGetter::CreateStatic( &AttrName##_Local::PassThroughDelegate, TAttribute< FString >::FGetter::CreateSP( InUserObject, InFunc, Var1, Var2, Var3, Var4 ) ) ); \
-		return this->Me(); \
-	} \
-\
-	/* Bind attribute with delegate to a UObject-based class method */ \
-	template< class UserClass >	\
-	DEPRECATED(4.8, "Passing text to Slate as FString is deprecated, please use FText instead (likely via a LOCTEXT).")\
-	WidgetArgsType& AttrName##_UObject( UserClass* InUserObject, typename TAttribute< FString >::FGetter::template TUObjectMethodDelegate_Const< UserClass >::FMethodPtr InFunc )	\
-	{ \
-		Var = TAttribute< FText >::Create( TAttribute< FText >::FGetter::CreateStatic( &AttrName##_Local::PassThroughDelegate, TAttribute< FString >::FGetter::CreateUObject( InUserObject, InFunc ) ) ); \
-		return this->Me(); \
-	} \
-	template< class UserClass, typename Var1Type >	\
-	DEPRECATED(4.8, "Passing text to Slate as FString is deprecated, please use FText instead (likely via a LOCTEXT).")\
-	WidgetArgsType& AttrName##_UObject( UserClass* InUserObject, typename TAttribute< FString >::FGetter::template TUObjectMethodDelegate_OneVar_Const< UserClass, Var1Type >::FMethodPtr InFunc, Var1Type Var1 )	\
-	{ \
-		Var = TAttribute< FText >::Create( TAttribute< FText >::FGetter::CreateStatic( &AttrName##_Local::PassThroughDelegate, TAttribute< FString >::FGetter::CreateUObject( InUserObject, InFunc, Var1 ) ) ); \
-		return this->Me(); \
-	} \
-	template< class UserClass, typename Var1Type, typename Var2Type  >	\
-	DEPRECATED(4.8, "Passing text to Slate as FString is deprecated, please use FText instead (likely via a LOCTEXT).")\
-	WidgetArgsType& AttrName##_UObject( UserClass* InUserObject, typename TAttribute< FString >::FGetter::template TUObjectMethodDelegate_TwoVars_Const< UserClass, Var1Type, Var2Type >::FMethodPtr InFunc, Var1Type Var1, Var2Type Var2 )	\
-	{ \
-		Var = TAttribute< FText >::Create( TAttribute< FText >::FGetter::CreateStatic( &AttrName##_Local::PassThroughDelegate, TAttribute< FString >::FGetter::CreateUObject( InUserObject, InFunc, Var1, Var2 ) ) ); \
-		return this->Me(); \
-	} \
-	template< class UserClass, typename Var1Type, typename Var2Type, typename Var3Type >	\
-	DEPRECATED(4.8, "Passing text to Slate as FString is deprecated, please use FText instead (likely via a LOCTEXT).")\
-	WidgetArgsType& AttrName##_UObject( UserClass* InUserObject, typename TAttribute< FString >::FGetter::template TUObjectMethodDelegate_ThreeVars_Const< UserClass, Var1Type, Var2Type, Var3Type >::FMethodPtr InFunc, Var1Type Var1, Var2Type Var2, Var3Type Var3 )	\
-	{ \
-		Var = TAttribute< FText >::Create( TAttribute< FText >::FGetter::CreateStatic( &AttrName##_Local::PassThroughDelegate, TAttribute< FString >::FGetter::CreateUObject( InUserObject, InFunc, Var1, Var2, Var3 ) ) ); \
-		return this->Me(); \
-	} \
-	template< class UserClass, typename Var1Type, typename Var2Type, typename Var3Type, typename Var4Type >	\
-	DEPRECATED(4.8, "Passing text to Slate as FString is deprecated, please use FText instead (likely via a LOCTEXT).")\
-	WidgetArgsType& AttrName##_UObject( UserClass* InUserObject, typename TAttribute< FString >::FGetter::template TUObjectMethodDelegate_FourVars_Const< UserClass, Var1Type, Var2Type, Var3Type, Var4Type >::FMethodPtr InFunc, Var1Type Var1, Var2Type Var2, Var3Type Var3, Var4Type Var4 )	\
-	{ \
-		Var = TAttribute< FText >::Create( TAttribute< FText >::FGetter::CreateStatic( &AttrName##_Local::PassThroughDelegate, TAttribute< FString >::FGetter::CreateUObject( InUserObject, InFunc, Var1, Var2, Var3, Var4 ) ) ); \
-		return this->Me(); \
-	}\
-	INTERNAL_SLATE_DECL_ATTRIBUTE( Var, TAttribute< FText >, AttrName )
-
-/**
- * Declare an FText attribute which also allows content passthrough from FString
- */
-#define SLATE_TEXT_ATTRIBUTE( AttrName ) SLATE_TEXT_ATTRIBUTE_IMPL( _##AttrName, AttrName )
-
-
-/**
- * Use this macro to declar a slate argument.
+ * Use this macro to declare a slate argument.
  * Arguments differ from attributes in that they can only be values
  */
 #define SLATE_ARGUMENT( ArgType, ArgName ) \
@@ -539,21 +304,6 @@ class IToolTip;
 			return this->Me(); \
 		}
 
-#define SLATE_TEXT_ARGUMENT( ArgName ) \
-	FText _##ArgName; \
-	DEPRECATED(4.8, "Passing text to Slate as FString is deprecated, please use FText instead (likely via a LOCTEXT).")\
-	WidgetArgsType& ArgName( FString InArg ) \
-	{ \
-		_##ArgName = FText::FromString(InArg); \
-		return this->Me(); \
-	}\
-	\
-	WidgetArgsType& ArgName( FText InArg ) \
-	{ \
-		_##ArgName = InArg; \
-		return this->Me(); \
-	}
-
 
 /**
  * Use this macro to declare a slate argument.
@@ -1046,7 +796,7 @@ struct TSlateBaseNamedArgs
 		return Me();
 	}
 
-	SLATE_TEXT_ATTRIBUTE( ToolTipText )
+	SLATE_ATTRIBUTE( FText, ToolTipText )
 	SLATE_ARGUMENT( TSharedPtr<IToolTip>, ToolTip )
 	SLATE_ATTRIBUTE( TOptional<EMouseCursor::Type>, Cursor )
 	SLATE_ATTRIBUTE( bool, IsEnabled )
diff --git a/Engine/Source/Runtime/SlateCore/Public/Widgets/SWidget.h b/Engine/Source/Runtime/SlateCore/Public/Widgets/SWidget.h
--- a/Engine/Source/Runtime/SlateCore/Public/Widgets/SWidget.h
+++ b/Engine/Source/Runtime/SlateCore/Public/Widgets/SWidget.h
@@ -819,9 +819,6 @@ public:
 	 *
 	 * @param InToolTipText  the text that should appear in the tool tip
 	 */
-	DEPRECATED(4.8, "Passing text to Slate as FString is deprecated, please use FText instead (likely via a LOCTEXT).")
-	void SetToolTipText( const TAttribute<FString>& InToolTipText );
-
 	void SetToolTipText(const TAttribute<FText>& ToolTipText);
 
 	void SetToolTipText( const FText& InToolTipText );
