Exposed the ability to customize the keyboard focus brush used by default by all SWidgets.

Merging using UE4-To-UE4-LauncherDev

diff --git a/Engine/Source/Runtime/SlateCore/Private/Styling/CoreStyle.cpp b/Engine/Source/Runtime/SlateCore/Private/Styling/CoreStyle.cpp
--- a/Engine/Source/Runtime/SlateCore/Private/Styling/CoreStyle.cpp
+++ b/Engine/Source/Runtime/SlateCore/Private/Styling/CoreStyle.cpp
@@ -91,6 +91,16 @@ void FCoreStyle::SetPressedSelectionColor( const FLinearColor& NewColor )
 	FSlateCoreStyle::SetColor(Style->SelectionColor_Pressed_LinearRef, NewColor);
 }
 
+void FCoreStyle::SetFocusBrush(FSlateBrush* NewBrush)
+{
+	TSharedRef<FSlateCoreStyle> Style = StaticCastSharedRef<FSlateCoreStyle>(Instance.ToSharedRef());
+	FSlateStyleRegistry::UnRegisterSlateStyle(Style.Get());
+
+	Style->Set("FocusRectangle", NewBrush);
+
+	FSlateStyleRegistry::RegisterSlateStyle(Style.Get());
+}
+
 
 #define IMAGE_BRUSH( RelativePath, ... ) FSlateImageBrush(Style->RootToContentDir(RelativePath, TEXT(".png")), __VA_ARGS__)
 #define BOX_BRUSH( RelativePath, ... ) FSlateBoxBrush(Style->RootToContentDir(RelativePath, TEXT(".png")), __VA_ARGS__)
@@ -193,7 +203,7 @@ TSharedRef<ISlateStyle> FCoreStyle::Create( const FName& InStyleSetName )
 
 		Style->Set("BoxShadow", new BOX_BRUSH( "Common/BoxShadow" , FMargin( 5.0f / 64.0f)));
 
-		Style->Set("FocusRectangle", new BORDER_BRUSH( "Old/DashedBorder", FMargin(6.0f / 32.0f)));
+		Style->Set("FocusRectangle", new BORDER_BRUSH( "Old/DashedBorder", FMargin(6.0f / 32.0f), FLinearColor(1, 1, 1, 0.5)));
 	}
 
 	// Important colors
diff --git a/Engine/Source/Runtime/SlateCore/Private/Widgets/SWidget.cpp b/Engine/Source/Runtime/SlateCore/Private/Widgets/SWidget.cpp
--- a/Engine/Source/Runtime/SlateCore/Private/Widgets/SWidget.cpp
+++ b/Engine/Source/Runtime/SlateCore/Private/Widgets/SWidget.cpp
@@ -714,7 +714,7 @@ int32 SWidget::Paint(const FPaintArgs& Args, const FGeometry& AllottedGeometry,
 					BrushResource,
 					MyClippingRect,
 					ESlateDrawEffect::None,
-					FColor(255, 255, 255, 128)
+					BrushResource->GetTint(InWidgetStyle)
 					);
 			}
 		}
diff --git a/Engine/Source/Runtime/SlateCore/Public/Styling/CoreStyle.h b/Engine/Source/Runtime/SlateCore/Public/Styling/CoreStyle.h
--- a/Engine/Source/Runtime/SlateCore/Public/Styling/CoreStyle.h
+++ b/Engine/Source/Runtime/SlateCore/Public/Styling/CoreStyle.h
@@ -26,6 +26,7 @@ public:
 	static void SetSelectionColor( const FLinearColor& NewColor );
 	static void SetInactiveSelectionColor( const FLinearColor& NewColor );
 	static void SetPressedSelectionColor( const FLinearColor& NewColor );
+	static void SetFocusBrush(FSlateBrush* NewBrush);
 
 	// todo: jdale - These are only here because of UTouchInterface::Activate and the fact that GetDynamicImageBrush is non-const
 	static const TSharedPtr<FSlateDynamicImageBrush> GetDynamicImageBrush( FName BrushTemplate, FName TextureName, const ANSICHAR* Specifier = nullptr );
