UMG - Rolling back the change to restrict access to the image property on FSlateBrush.  The change is breaking too much old code with no easy way to upgrade people.

diff --git a/Engine/Source/Runtime/SlateCore/Public/Styling/SlateBrush.h b/Engine/Source/Runtime/SlateCore/Public/Styling/SlateBrush.h
--- a/Engine/Source/Runtime/SlateCore/Public/Styling/SlateBrush.h
+++ b/Engine/Source/Runtime/SlateCore/Public/Styling/SlateBrush.h
@@ -293,7 +293,7 @@ protected:
 	/**
 	 * The image to render for this brush, can be a UTexture2D or Material.
 	 */
-	UPROPERTY(EditAnywhere, Category=Brush, meta=( DisplayThumbnail="true", DisplayName="Image", AllowedClasses="Texture2D,MaterialInterface" ))
+	UPROPERTY(EditAnywhere, BlueprintReadWrite, Category=Brush, meta=( DisplayThumbnail="true", DisplayName="Image", AllowedClasses="Texture2D,MaterialInterface" ))
 	UObject* ResourceObject;
 
 	/** The name of the rendering resource to use */
