Fixes pull request GitHub 1376.
When a ReleaseMouseCaptor is called after a CaptureMouse or a UseHighPrecisionMouse, it is ineffective.

diff --git a/Engine/Source/Runtime/SlateCore/Public/Input/Reply.h b/Engine/Source/Runtime/SlateCore/Public/Input/Reply.h
--- a/Engine/Source/Runtime/SlateCore/Public/Input/Reply.h
+++ b/Engine/Source/Runtime/SlateCore/Public/Input/Reply.h
@@ -124,6 +124,7 @@ public:
 	 */
 	FReply& ReleaseMouseCapture()
 	{
+		this->MouseCaptor.Reset();
 		this->bReleaseMouseCapture = true;
 		this->bUseHighPrecisionMouse = false;
 		return Me();
