Exposed FSlateSound to BluePrint. Classes derived from SoundBase can be set. SButton::SetStyle now updates Hovered and pressed sounds.
[UE-19542] - Cannot dynamically change when pressed sound for buttons in UMG

diff --git a/Engine/Source/Runtime/Slate/Private/Widgets/Input/SButton.cpp b/Engine/Source/Runtime/Slate/Private/Widgets/Input/SButton.cpp
--- a/Engine/Source/Runtime/Slate/Private/Widgets/Input/SButton.cpp
+++ b/Engine/Source/Runtime/Slate/Private/Widgets/Input/SButton.cpp
@@ -383,4 +383,7 @@ void SButton::SetButtonStyle(const FButtonStyle* ButtonStyle)
 
 	BorderPadding = Style->NormalPadding;
 	PressedBorderPadding = Style->PressedPadding;
+
+	HoveredSound = Style->HoveredSlateSound;
+	PressedSound = Style->PressedSlateSound;
 }
diff --git a/Engine/Source/Runtime/SlateCore/Public/Sound/SlateSound.h b/Engine/Source/Runtime/SlateCore/Public/Sound/SlateSound.h
--- a/Engine/Source/Runtime/SlateCore/Public/Sound/SlateSound.h
+++ b/Engine/Source/Runtime/SlateCore/Public/Sound/SlateSound.h
@@ -64,7 +64,7 @@ protected:
 	 * Pointer to the USoundBase. Holding onto it as a UObject because USoundBase is not available in Slate core.
 	 * Edited via FSlateSoundStructCustomization to ensure you can only set USoundBase assets on it.
 	 */
-	UPROPERTY(EditAnywhere, Category=Sound)
+	UPROPERTY(EditAnywhere, BlueprintReadWrite, Category=Sound, meta=(DisplayName="Sound", AllowedClasses="SoundBase"))
 	UObject* ResourceObject;
 
 	/** The legacy resource name; only used by sounds that have been set-up in code, or otherwise upgraded from old FName properties, set to NAME_None in non-legacy instances */
