UMG / Slate - There are now several low level flags on widgets to help reduce calls to virtual functions that are not implemented, bCanEverTick, bCanEverFocus, bCanHaveChildren, possibly more in the future.

diff --git a/Engine/Source/Runtime/Slate/Private/Widgets/Layout/SBorder.cpp b/Engine/Source/Runtime/Slate/Private/Widgets/Layout/SBorder.cpp
--- a/Engine/Source/Runtime/Slate/Private/Widgets/Layout/SBorder.cpp
+++ b/Engine/Source/Runtime/Slate/Private/Widgets/Layout/SBorder.cpp
@@ -1,35 +1,31 @@
 // Copyright 1998-2015 Epic Games, Inc. All Rights Reserved.
 
 #include "SlatePrivatePCH.h"
- 
+
+static FName SBorderTypeName("SBorder");
+
 SBorder::SBorder()
 	: BorderImage( FCoreStyle::Get().GetBrush( "Border" ) )
 	, BorderBackgroundColor( FLinearColor::White )
 	, DesiredSizeScale(FVector2D(1,1))
-{ }
-
+{
+}
 
-/**
- * Construct this widget
- *
- * @param	InArgs	The declaration data for this widget
- */
 void SBorder::Construct( const SBorder::FArguments& InArgs )
 {
+	// Only do this if we're exactly an SBorder
+	if ( GetType() == SBorderTypeName )
+	{
+		bCanTick = false;
+		bCanSupportFocus = false;
+	}
+
 	ContentScale = InArgs._ContentScale;
 	ColorAndOpacity = InArgs._ColorAndOpacity;
 	DesiredSizeScale = InArgs._DesiredSizeScale;
 
 	ShowDisabledEffect = InArgs._ShowEffectWhenDisabled;
 
-	ChildSlot
-		.HAlign(InArgs._HAlign)
-		.VAlign(InArgs._VAlign)
-		.Padding(InArgs._Padding)
-	[
-		InArgs._Content.Widget
-	];
-
 	BorderImage = InArgs._BorderImage;
 	BorderBackgroundColor = InArgs._BorderBackgroundColor;
 	ForegroundColor = InArgs._ForegroundColor;
@@ -37,14 +33,16 @@ void SBorder::Construct( const SBorder::FArguments& InArgs )
 	MouseButtonUpHandler = InArgs._OnMouseButtonUp;
 	MouseMoveHandler = InArgs._OnMouseMove;
 	MouseDoubleClickHandler = InArgs._OnMouseDoubleClick;
-}
 
+	ChildSlot
+		.HAlign(InArgs._HAlign)
+		.VAlign(InArgs._VAlign)
+		.Padding(InArgs._Padding)
+	[
+		InArgs._Content.Widget
+	];
+}
 
-/**
- * Sets the content for this border
- *
- * @param	InContent	The widget to use as content for the border
- */
 void SBorder::SetContent( TSharedRef< SWidget > InContent )
 {
 	ChildSlot
@@ -58,7 +56,6 @@ const TSharedRef< SWidget >& SBorder::GetContent() const
 	return ChildSlot.GetWidget();
 }
 
-/** Clears out the content for the border */
 void SBorder::ClearContent()
 {
 	ChildSlot.DetachWidget();
@@ -93,15 +90,6 @@ int32 SBorder::OnPaint( const FPaintArgs& Args, const FGeometry& AllottedGeometr
 	return SCompoundWidget::OnPaint(Args, AllottedGeometry, MyClippingRect.IntersectionWith( AllottedGeometry.GetClippingRect() ), OutDrawElements, LayerId, CompoundedWidgetStyle, bEnabled );
 }
 
-
-/**
- * The system calls this method to notify the widget that a mouse button was pressed within it. This event is bubbled.
- *
- * @param MyGeometry The Geometry of the widget receiving the event
- * @param MouseEvent Information about the input event
- *
- * @return Whether the event was handled along with possible requests for the system to take action.
- */
 FReply SBorder::OnMouseButtonDown( const FGeometry& MyGeometry, const FPointerEvent& MouseEvent )
 {
 	if ( MouseButtonDownHandler.IsBound() )
@@ -116,14 +104,6 @@ FReply SBorder::OnMouseButtonDown( const FGeometry& MyGeometry, const FPointerEv
 	}
 }
 
-/**
- * The system calls this method to notify the widget that a mouse button was release within it. This event is bubbled.
- *
- * @param MyGeometry The Geometry of the widget receiving the event
- * @param MouseEvent Information about the input event
- *
- * @return Whether the event was handled along with possible requests for the system to take action.
- */
 FReply SBorder::OnMouseButtonUp( const FGeometry& MyGeometry, const FPointerEvent& MouseEvent )
 {
 	if ( MouseButtonUpHandler.IsBound() )
@@ -138,14 +118,6 @@ FReply SBorder::OnMouseButtonUp( const FGeometry& MyGeometry, const FPointerEven
 	}
 }
 
-/**
- * The system calls this method to notify the widget that a mouse moved within it. This event is bubbled.
- *
- * @param MyGeometry The Geometry of the widget receiving the event
- * @param MouseEvent Information about the input event
- *
- * @return Whether the event was handled along with possible requests for the system to take action.
- */
 FReply SBorder::OnMouseMove( const FGeometry& MyGeometry, const FPointerEvent& MouseEvent )
 {
 	if ( MouseMoveHandler.IsBound() )
@@ -160,14 +132,6 @@ FReply SBorder::OnMouseMove( const FGeometry& MyGeometry, const FPointerEvent& M
 	}
 }
 
-/**
- * Called when a mouse button is double clicked
- *
- * @param  InMyGeometry  Widget geometry
- * @param  InMouseEvent  Mouse button event
- *
- * @return  Returns whether the event was handled, along with other possible actions
- */
 FReply SBorder::OnMouseButtonDoubleClick( const FGeometry& MyGeometry, const FPointerEvent& MouseEvent )
 {
 	if ( MouseDoubleClickHandler.IsBound() )
diff --git a/Engine/Source/Runtime/Slate/Private/Widgets/Layout/SBox.cpp b/Engine/Source/Runtime/Slate/Private/Widgets/Layout/SBox.cpp
--- a/Engine/Source/Runtime/Slate/Private/Widgets/Layout/SBox.cpp
+++ b/Engine/Source/Runtime/Slate/Private/Widgets/Layout/SBox.cpp
@@ -6,6 +6,8 @@
 SBox::SBox()
 : ChildSlot()
 {
+	bCanTick = false;
+	bCanSupportFocus = false;
 }
 
 void SBox::Construct( const FArguments& InArgs )
diff --git a/Engine/Source/Runtime/Slate/Private/Widgets/Layout/SDPIScaler.cpp b/Engine/Source/Runtime/Slate/Private/Widgets/Layout/SDPIScaler.cpp
--- a/Engine/Source/Runtime/Slate/Private/Widgets/Layout/SDPIScaler.cpp
+++ b/Engine/Source/Runtime/Slate/Private/Widgets/Layout/SDPIScaler.cpp
@@ -6,6 +6,8 @@
 SDPIScaler::SDPIScaler()
 : ChildSlot()
 {
+	bCanTick = false;
+	bCanSupportFocus = false;
 }
 
 void SDPIScaler::Construct( const FArguments& InArgs )
@@ -61,4 +63,4 @@ void SDPIScaler::SetDPIScale(TAttribute<float> InDPIScale)
 float SDPIScaler::GetRelativeLayoutScale(const FSlotBase& Child) const
 {
 	return DPIScale.Get();
-}
\ No newline at end of file
+}
diff --git a/Engine/Source/Runtime/Slate/Private/Widgets/SCanvas.cpp b/Engine/Source/Runtime/Slate/Private/Widgets/SCanvas.cpp
--- a/Engine/Source/Runtime/Slate/Private/Widgets/SCanvas.cpp
+++ b/Engine/Source/Runtime/Slate/Private/Widgets/SCanvas.cpp
@@ -4,16 +4,12 @@
 #include "LayoutUtils.h"
 
 
-/* SCanvas structors
- *****************************************************************************/
-
 SCanvas::SCanvas()
-	: Children()
-{ }
-
-
-/* SCanvas interface
- *****************************************************************************/
+: Children()
+{
+	bCanTick = false;
+	bCanSupportFocus = false;
+}
 
 void SCanvas::Construct( const SCanvas::FArguments& InArgs )
 {
@@ -24,13 +20,11 @@ void SCanvas::Construct( const SCanvas::FArguments& InArgs )
 	}
 }
 
-
 void SCanvas::ClearChildren( )
 {
 	Children.Empty();
 }
 
-
 int32 SCanvas::RemoveSlot( const TSharedRef<SWidget>& SlotWidget )
 {
 	for (int32 SlotIdx = 0; SlotIdx < Children.Num(); ++SlotIdx)
@@ -45,10 +39,6 @@ int32 SCanvas::RemoveSlot( const TSharedRef<SWidget>& SlotWidget )
 	return -1;
 }
 
-
-/* SWidget overrides
- *****************************************************************************/
-
 void SCanvas::OnArrangeChildren( const FGeometry& AllottedGeometry, FArrangedChildren& ArrangedChildren ) const
 {
 	if (Children.Num() > 0)
@@ -101,7 +91,6 @@ void SCanvas::OnArrangeChildren( const FGeometry& AllottedGeometry, FArrangedChi
 	}
 }
 
-
 int32 SCanvas::OnPaint( const FPaintArgs& Args, const FGeometry& AllottedGeometry, const FSlateRect& MyClippingRect, FSlateWindowElementList& OutDrawElements, int32 LayerId, const FWidgetStyle& InWidgetStyle, bool bParentEnabled ) const
 {
 	FArrangedChildren ArrangedChildren(EVisibility::Visible);
@@ -127,14 +116,12 @@ int32 SCanvas::OnPaint( const FPaintArgs& Args, const FGeometry& AllottedGeometr
 	return MaxLayerId;
 }
 
-
 FVector2D SCanvas::ComputeDesiredSize( float ) const
 {
 	// Canvas widgets have no desired size -- their size is always determined by their container
 	return FVector2D::ZeroVector;
 }
 
-
 FChildren* SCanvas::GetChildren()
 {
 	return &Children;
diff --git a/Engine/Source/Runtime/Slate/Public/Widgets/Images/SImage.h b/Engine/Source/Runtime/Slate/Public/Widgets/Images/SImage.h
--- a/Engine/Source/Runtime/Slate/Public/Widgets/Images/SImage.h
+++ b/Engine/Source/Runtime/Slate/Public/Widgets/Images/SImage.h
@@ -27,6 +27,13 @@ public:
 
 	SLATE_END_ARGS()
 
+	/** Constructor */
+	SImage()
+	{
+		bCanTick = false;
+		bCanSupportFocus = false;
+	}
+
 	/**
 	 * Construct this widget
 	 *
diff --git a/Engine/Source/Runtime/Slate/Public/Widgets/LayerManager/STooltipPresenter.h b/Engine/Source/Runtime/Slate/Public/Widgets/LayerManager/STooltipPresenter.h
--- a/Engine/Source/Runtime/Slate/Public/Widgets/LayerManager/STooltipPresenter.h
+++ b/Engine/Source/Runtime/Slate/Public/Widgets/LayerManager/STooltipPresenter.h
@@ -16,6 +16,12 @@ public:
 
 	SLATE_END_ARGS()
 
+	/** Constructor */
+	STooltipPresenter()
+	{
+		bCanSupportFocus = false;
+	}
+
 	void Construct(const FArguments& InArgs);
 
 	void SetContent(const TSharedRef<SWidget>& InWidget);
diff --git a/Engine/Source/Runtime/Slate/Public/Widgets/Layout/SConstraintCanvas.cpp b/Engine/Source/Runtime/Slate/Public/Widgets/Layout/SConstraintCanvas.cpp
--- a/Engine/Source/Runtime/Slate/Public/Widgets/Layout/SConstraintCanvas.cpp
+++ b/Engine/Source/Runtime/Slate/Public/Widgets/Layout/SConstraintCanvas.cpp
@@ -12,6 +12,8 @@
 SConstraintCanvas::SConstraintCanvas()
 : Children()
 {
+	bCanTick = false;
+	bCanSupportFocus = false;
 }
 
 void SConstraintCanvas::Construct( const SConstraintCanvas::FArguments& InArgs )
@@ -23,13 +25,19 @@ void SConstraintCanvas::Construct( const SConstraintCanvas::FArguments& InArgs )
 	}
 }
 
-void SConstraintCanvas::ClearChildren( )
+void SConstraintCanvas::ClearChildren()
 {
-	Children.Empty();
+	if ( Children.Num() )
+	{
+		Invalidate(EInvalidateWidget::Layout);
+		Children.Empty();
+	}
 }
 
 int32 SConstraintCanvas::RemoveSlot( const TSharedRef<SWidget>& SlotWidget )
 {
+	Invalidate(EInvalidateWidget::Layout);
+
 	for (int32 SlotIdx = 0; SlotIdx < Children.Num(); ++SlotIdx)
 	{
 		if (SlotWidget == Children[SlotIdx].GetWidget())
diff --git a/Engine/Source/Runtime/Slate/Public/Widgets/Layout/SConstraintCanvas.h b/Engine/Source/Runtime/Slate/Public/Widgets/Layout/SConstraintCanvas.h
--- a/Engine/Source/Runtime/Slate/Public/Widgets/Layout/SConstraintCanvas.h
+++ b/Engine/Source/Runtime/Slate/Public/Widgets/Layout/SConstraintCanvas.h
@@ -110,6 +110,8 @@ public:
 	 */
 	FSlot& AddSlot()
 	{
+		Invalidate(EInvalidateWidget::Layout);
+
 		SConstraintCanvas::FSlot& NewSlot = *(new FSlot());
 		this->Children.Add( &NewSlot );
 		return NewSlot;
diff --git a/Engine/Source/Runtime/Slate/Public/Widgets/Layout/SPopup.h b/Engine/Source/Runtime/Slate/Public/Widgets/Layout/SPopup.h
--- a/Engine/Source/Runtime/Slate/Public/Widgets/Layout/SPopup.h
+++ b/Engine/Source/Runtime/Slate/Public/Widgets/Layout/SPopup.h
@@ -17,6 +17,12 @@ public:
 
 	SLATE_END_ARGS()
 
+	SPopup()
+	{
+		bCanTick = false;
+		bCanSupportFocus = false;
+	}
+
 	void Construct(const FArguments& InArgs);
 
 private:
diff --git a/Engine/Source/Runtime/Slate/Public/Widgets/Layout/SScaleBox.h b/Engine/Source/Runtime/Slate/Public/Widgets/Layout/SScaleBox.h
--- a/Engine/Source/Runtime/Slate/Public/Widgets/Layout/SScaleBox.h
+++ b/Engine/Source/Runtime/Slate/Public/Widgets/Layout/SScaleBox.h
@@ -90,6 +90,13 @@ public:
 
 	SLATE_END_ARGS()
 
+	/** Constructor */
+	SScaleBox()
+	{
+		bCanTick = false;
+		bCanSupportFocus = false;
+	}
+
 	void Construct(const FArguments& InArgs);
 	
 	// SWidget interface
diff --git a/Engine/Source/Runtime/Slate/Public/Widgets/Layout/SSpacer.h b/Engine/Source/Runtime/Slate/Public/Widgets/Layout/SSpacer.h
--- a/Engine/Source/Runtime/Slate/Public/Widgets/Layout/SSpacer.h
+++ b/Engine/Source/Runtime/Slate/Public/Widgets/Layout/SSpacer.h
@@ -16,6 +16,12 @@ public:
 		SLATE_ATTRIBUTE( FVector2D, Size )
 	SLATE_END_ARGS()
 
+	SSpacer()
+	{
+		bCanTick = false;
+		bCanSupportFocus = false;
+	}
+
 	/**
 	 * Construct this widget
 	 *
diff --git a/Engine/Source/Runtime/Slate/Public/Widgets/Text/STextBlock.h b/Engine/Source/Runtime/Slate/Public/Widgets/Text/STextBlock.h
--- a/Engine/Source/Runtime/Slate/Public/Widgets/Text/STextBlock.h
+++ b/Engine/Source/Runtime/Slate/Public/Widgets/Text/STextBlock.h
@@ -99,6 +99,13 @@ public:
 
 	SLATE_END_ARGS()
 
+	/** Constructor */
+	STextBlock()
+	{
+		bCanTick = false;
+		bCanSupportFocus = false;
+	}
+
 	/**
 	 * Construct this widget
 	 *
diff --git a/Engine/Source/Runtime/SlateCore/Private/Widgets/SNullWidget.cpp b/Engine/Source/Runtime/SlateCore/Private/Widgets/SNullWidget.cpp
--- a/Engine/Source/Runtime/SlateCore/Private/Widgets/SNullWidget.cpp
+++ b/Engine/Source/Runtime/SlateCore/Private/Widgets/SNullWidget.cpp
@@ -32,6 +32,12 @@ public:
 		);
 	}
 
+	SNullWidgetContent()
+	{
+		bCanTick = false;
+		bCanSupportFocus = false;
+	}
+
 private:
 	virtual void SetVisibility( TAttribute<EVisibility> InVisibility ) override final
 	{
diff --git a/Engine/Source/Runtime/SlateCore/Private/Widgets/SOverlay.cpp b/Engine/Source/Runtime/SlateCore/Private/Widgets/SOverlay.cpp
--- a/Engine/Source/Runtime/SlateCore/Private/Widgets/SOverlay.cpp
+++ b/Engine/Source/Runtime/SlateCore/Private/Widgets/SOverlay.cpp
@@ -7,13 +7,10 @@
 SOverlay::SOverlay()
 : Children()
 {
+	bCanTick = false;
+	bCanSupportFocus = false;
 }
 
-/**
- * Construct this widget
- *
- * @param	InArgs	The declaration data for this widget
- */
 void SOverlay::Construct( const SOverlay::FArguments& InArgs )
 {
 	const int32 NumSlots = InArgs.Slots.Num();
@@ -23,7 +20,6 @@ void SOverlay::Construct( const SOverlay::FArguments& InArgs )
 	}
 }
 
-
 void SOverlay::OnArrangeChildren( const FGeometry& AllottedGeometry, FArrangedChildren& ArrangedChildren ) const
 {
 	for ( int32 ChildIndex = 0; ChildIndex < Children.Num(); ++ChildIndex )
@@ -45,13 +41,6 @@ void SOverlay::OnArrangeChildren( const FGeometry& AllottedGeometry, FArrangedCh
 	}
 }
 
-
-/**
- * A Panel's desired size in the space required to arrange of its children on the screen while respecting all of
- * the children's desired sizes and any layout-related options specified by the user. See StackPanel for an example.
- *
- * @return The desired size.
- */
 FVector2D SOverlay::ComputeDesiredSize( float ) const
 {
 	FVector2D MaxSize(0,0);
@@ -70,7 +59,6 @@ FVector2D SOverlay::ComputeDesiredSize( float ) const
 	return MaxSize;
 }
 
-/** @return  The children of a panel in a slot-agnostic way. */
 FChildren* SOverlay::GetChildren()
 {
 	return &Children;
@@ -167,13 +155,11 @@ void SOverlay::ClearChildren()
 	Children.Empty();
 }
 
-/** Returns the number of child widgets */
 int32 SOverlay::GetNumWidgets() const
 {
 	return Children.Num();
 }
 
-
 void SOverlay::RemoveSlot( TSharedRef< SWidget > Widget )
 {
 	// Search and remove
diff --git a/Engine/Source/Runtime/SlateCore/Private/Widgets/SWidget.cpp b/Engine/Source/Runtime/SlateCore/Private/Widgets/SWidget.cpp
--- a/Engine/Source/Runtime/SlateCore/Private/Widgets/SWidget.cpp
+++ b/Engine/Source/Runtime/SlateCore/Private/Widgets/SWidget.cpp
@@ -17,6 +17,9 @@ SLATE_DECLARE_CYCLE_COUNTER(GSlateGetVisibility, "GetVisibility");
 
 TAutoConsoleVariable<int32> TickInvisibleWidgets(TEXT("Slate.TickInvisibleWidgets"), 0, TEXT("Controls whether invisible widgets are ticked."));
 
+SLATECORE_API int32 bFoldTick = 1;
+FAutoConsoleVariableRef FoldTick(TEXT("Slate.FoldTick"), bFoldTick, TEXT("When folding, call Tick as part of the paint pass instead of a separate tick pass."));
+
 SWidget::SWidget()
 	: Cursor( TOptional<EMouseCursor::Type>() )
 	, EnabledState( true )
@@ -26,7 +29,10 @@ SWidget::SWidget()
 	, DesiredSize(FVector2D::ZeroVector)
 	, ToolTip()
 	, bIsHovered(false)
-	, bToolTipForceFieldEnabled( false )
+	, bCanTick(true)
+	, bCanSupportFocus(true)
+	, bCanHaveChildren(true)
+	, bToolTipForceFieldEnabled(false)
 	, bForceVolatile(false)
 	, bCachedVolatile(false)
 	, bInheritedVolatility(false)
@@ -73,7 +79,6 @@ void SWidget::Construct(
 	MetaData = InMetaData;
 }
 
-
 FReply SWidget::OnFocusReceived(const FGeometry& MyGeometry, const FFocusEvent& InFocusEvent)
 {
 	return FReply::Unhandled();
@@ -92,13 +97,11 @@ FReply SWidget::OnKeyChar( const FGeometry& MyGeometry, const FCharacterEvent& I
 	return FReply::Unhandled();
 }
 
-
 FReply SWidget::OnPreviewKeyDown( const FGeometry& MyGeometry, const FKeyEvent& InKeyEvent )
 {
 	return FReply::Unhandled();
 }
 
-
 FReply SWidget::OnKeyDown( const FGeometry& MyGeometry, const FKeyEvent& InKeyEvent )
 {
 	if (SupportsKeyboardFocus())
@@ -139,61 +142,51 @@ FReply SWidget::OnKeyDown( const FGeometry& MyGeometry, const FKeyEvent& InKeyEv
 	return FReply::Unhandled();
 }
 
-
 FReply SWidget::OnKeyUp( const FGeometry& MyGeometry, const FKeyEvent& InKeyEvent )
 {
 	return FReply::Unhandled();
 }
 
-
 FReply SWidget::OnAnalogValueChanged( const FGeometry& MyGeometry, const FAnalogInputEvent& InAnalogInputEvent )
 {
 	return FReply::Unhandled();
 }
 
-
 FReply SWidget::OnMouseButtonDown( const FGeometry& MyGeometry, const FPointerEvent& MouseEvent )
 {
 	return FReply::Unhandled();
 }
 
-
 FReply SWidget::OnPreviewMouseButtonDown( const FGeometry& MyGeometry, const FPointerEvent& MouseEvent )
 {
 	return FReply::Unhandled();
 }
 
-
 FReply SWidget::OnMouseButtonUp( const FGeometry& MyGeometry, const FPointerEvent& MouseEvent )
 {
 	return FReply::Unhandled();
 }
 
-
 FReply SWidget::OnMouseMove( const FGeometry& MyGeometry, const FPointerEvent& MouseEvent )
 {
 	return FReply::Unhandled();
 }
 
-
 void SWidget::OnMouseEnter( const FGeometry& MyGeometry, const FPointerEvent& MouseEvent )
 {
 	bIsHovered = true;
 }
 
-
 void SWidget::OnMouseLeave( const FPointerEvent& MouseEvent )
 {
 	bIsHovered = false;
 }
 
-
 FReply SWidget::OnMouseWheel( const FGeometry& MyGeometry, const FPointerEvent& MouseEvent )
 {
 	return FReply::Unhandled();
 }
 
-
 FCursorReply SWidget::OnCursorQuery( const FGeometry& MyGeometry, const FPointerEvent& CursorEvent ) const
 {
 	TOptional<EMouseCursor::Type> TheCursor = Cursor.Get();
@@ -207,89 +200,74 @@ TOptional<TSharedRef<SWidget>> SWidget::OnMapCursor(const FCursorReply& CursorRe
 	return TOptional<TSharedRef<SWidget>>();
 }
 
-
 FReply SWidget::OnMouseButtonDoubleClick( const FGeometry& InMyGeometry, const FPointerEvent& InMouseEvent )
 {
 	return FReply::Unhandled();
 }
 
-
 bool SWidget::OnVisualizeTooltip( const TSharedPtr<SWidget>& TooltipContent )
 {
 	return false;
 }
 
-
 FReply SWidget::OnDragDetected( const FGeometry& MyGeometry, const FPointerEvent& MouseEvent )
 {
 	return FReply::Unhandled();
 }
 
-
 void SWidget::OnDragEnter( const FGeometry& MyGeometry, const FDragDropEvent& DragDropEvent )
 {
 }
 
-
 void SWidget::OnDragLeave( const FDragDropEvent& DragDropEvent )
 {
 }
 
-
 FReply SWidget::OnDragOver( const FGeometry& MyGeometry, const FDragDropEvent& DragDropEvent )
 {
 	return FReply::Unhandled();
 }
 
-
 FReply SWidget::OnDrop( const FGeometry& MyGeometry, const FDragDropEvent& DragDropEvent )
 {
 	return FReply::Unhandled();
 }
 
-
 FReply SWidget::OnControllerButtonPressed( const FGeometry& MyGeometry, const FControllerEvent& ControllerEvent )
 {
 	return FReply::Unhandled();
 }
 
-
 FReply SWidget::OnControllerButtonReleased( const FGeometry& MyGeometry, const FControllerEvent& ControllerEvent )
 {
 	return FReply::Unhandled();
 }
 
-
 FReply SWidget::OnControllerAnalogValueChanged( const FGeometry& MyGeometry, const FControllerEvent& ControllerEvent )
 {
 	return FReply::Unhandled();
 }
 
-
 FReply SWidget::OnTouchGesture( const FGeometry& MyGeometry, const FPointerEvent& GestureEvent )
 {
 	return FReply::Unhandled();
 }
 
-
 FReply SWidget::OnTouchStarted( const FGeometry& MyGeometry, const FPointerEvent& InTouchEvent )
 {
 	return FReply::Unhandled();
 }
 
-
 FReply SWidget::OnTouchMoved( const FGeometry& MyGeometry, const FPointerEvent& InTouchEvent )
 {
 	return FReply::Unhandled();
 }
 
-
 FReply SWidget::OnTouchEnded( const FGeometry& MyGeometry, const FPointerEvent& InTouchEvent )
 {
 	return FReply::Unhandled();
 }
 
-
 FReply SWidget::OnMotionDetected( const FGeometry& MyGeometry, const FMotionEvent& InMotionEvent )
 {
 	return FReply::Unhandled();
@@ -310,19 +288,16 @@ TSharedPtr<struct FVirtualPointerPosition> SWidget::TranslateMouseCoordinateFor3
 	return nullptr;
 }
 
-
 void SWidget::OnFinishedPointerInput()
 {
 
 }
 
-
 void SWidget::OnFinishedKeyInput()
 {
 
 }
 
-
 FNavigationReply SWidget::OnNavigation(const FGeometry& MyGeometry, const FNavigationEvent& InNavigationEvent)
 {
 	EUINavigation Type = InNavigationEvent.GetNavigationType();
@@ -335,7 +310,6 @@ FNavigationReply SWidget::OnNavigation(const FGeometry& MyGeometry, const FNavig
 	return FNavigationReply::Escape();
 }
 
-
 EWindowZone::Type SWidget::GetWindowZoneOverride() const
 {
 	// No special behavior.  Override this in derived widgets, if needed.
@@ -378,20 +352,24 @@ void SWidget::SlatePrepass()
 
 void SWidget::SlatePrepass(float LayoutScaleMultiplier)
 {
-	SLATE_CYCLE_COUNTER_SCOPE_CUSTOM_DETAILED(SLATE_STATS_DETAIL_LEVEL_MED, GSlatePrepass, GetType());
-	// Cache child desired sizes first. This widget's desired size is
-	// a function of its children's sizes.
-	FChildren* MyChildren = this->GetChildren();
-	int32 NumChildren = MyChildren->Num();
-	for ( int32 ChildIndex=0; ChildIndex < NumChildren; ++ChildIndex )
-	{
-		const TSharedRef<SWidget>& Child = MyChildren->GetChildAt(ChildIndex);
+	//SLATE_CYCLE_COUNTER_SCOPE_CUSTOM_DETAILED(SLATE_STATS_DETAIL_LEVEL_MED, GSlatePrepass, GetType());
 
-		if ( Child->Visibility.Get() != EVisibility::Collapsed )
+	if ( bCanHaveChildren )
+	{
+		// Cache child desired sizes first. This widget's desired size is
+		// a function of its children's sizes.
+		FChildren* MyChildren = this->GetChildren();
+		int32 NumChildren = MyChildren->Num();
+		for ( int32 ChildIndex=0; ChildIndex < NumChildren; ++ChildIndex )
 		{
-			const float ChildLayoutScaleMultiplier = GetRelativeLayoutScale(MyChildren->GetSlotAt(ChildIndex));
-			// Recur: Descend down the widget tree.
-			Child->SlatePrepass(LayoutScaleMultiplier*ChildLayoutScaleMultiplier);
+			const TSharedRef<SWidget>& Child = MyChildren->GetChildAt(ChildIndex);
+
+			if ( Child->Visibility.Get() != EVisibility::Collapsed )
+			{
+				const float ChildLayoutScaleMultiplier = GetRelativeLayoutScale(MyChildren->GetSlotAt(ChildIndex));
+				// Recur: Descend down the widget tree.
+				Child->SlatePrepass(LayoutScaleMultiplier*ChildLayoutScaleMultiplier);
+			}
 		}
 	}
 
@@ -412,18 +390,21 @@ const FVector2D& SWidget::GetDesiredSize() const
 
 void SWidget::CachePrepass(TWeakPtr<ILayoutCache> InLayoutCache)
 {
-	FChildren* MyChildren = this->GetChildren();
-	int32 NumChildren = MyChildren->Num();
-	for ( int32 ChildIndex=0; ChildIndex < NumChildren; ++ChildIndex )
+	if ( bCanHaveChildren )
 	{
-		const TSharedRef<SWidget>& Child = MyChildren->GetChildAt(ChildIndex);
-		if ( Child->GetVisibility().IsVisible() == false )
+		FChildren* MyChildren = this->GetChildren();
+		int32 NumChildren = MyChildren->Num();
+		for ( int32 ChildIndex=0; ChildIndex < NumChildren; ++ChildIndex )
 		{
-			Child->LayoutCache = InLayoutCache;
-		}
-		else
-		{
-			Child->CachePrepass(InLayoutCache);
+			const TSharedRef<SWidget>& Child = MyChildren->GetChildAt(ChildIndex);
+			if ( Child->GetVisibility().IsVisible() == false )
+			{
+				Child->LayoutCache = InLayoutCache;
+			}
+			else
+			{
+				Child->CachePrepass(InLayoutCache);
+			}
 		}
 	}
 }
@@ -468,7 +449,6 @@ bool SWidget::HasMouseCapture() const
 	return FSlateApplicationBase::Get().HasMouseCapture(SharedThis(this));
 }
 
-
 void SWidget::OnMouseCaptureLost()
 {
 	
@@ -513,7 +493,6 @@ void SWidget::FindChildGeometries_Helper( const FGeometry& MyGeometry, const TSe
 	}	
 }
 
-
 FGeometry SWidget::FindChildGeometry( const FGeometry& MyGeometry, TSharedRef<SWidget> WidgetToFind ) const
 {
 	// We just need to find the one WidgetToFind among our descendants.
@@ -528,14 +507,12 @@ FGeometry SWidget::FindChildGeometry( const FGeometry& MyGeometry, TSharedRef<SW
 	return Result.FindChecked( WidgetToFind ).Geometry;
 }
 
-
 int32 SWidget::FindChildUnderMouse( const FArrangedChildren& Children, const FPointerEvent& MouseEvent )
 {
 	const FVector2D& AbsoluteCursorLocation = MouseEvent.GetScreenSpacePosition();
 	return SWidget::FindChildUnderPosition( Children, AbsoluteCursorLocation );
 }
 
-
 int32 SWidget::FindChildUnderPosition( const FArrangedChildren& Children, const FVector2D& ArrangedSpacePosition )
 {
 	const int32 NumChildren = Children.Num();
@@ -555,67 +532,57 @@ int32 SWidget::FindChildUnderPosition( const FArrangedChildren& Children, const
 	return INDEX_NONE;
 }
 
-
 FString SWidget::ToString() const
 {
 	return FString::Printf(TEXT("%s [%s]"), *this->TypeOfWidget.ToString(), *this->GetReadableLocation() );
 }
 
-
 FString SWidget::GetTypeAsString() const
 {
 	return this->TypeOfWidget.ToString();
 }
 
-
 FName SWidget::GetType() const
 {
-	return this->TypeOfWidget;
+	return TypeOfWidget;
 }
 
-
 FString SWidget::GetReadableLocation() const
 {
 	return FString::Printf(TEXT("%s(%d)"), *FPaths::GetCleanFilename(this->CreatedInLocation.GetPlainNameString()), this->CreatedInLocation.GetNumber());
 }
 
-
 FName SWidget::GetCreatedInLocation() const
 {
-	return this->CreatedInLocation;
+	return CreatedInLocation;
 }
 
 FName SWidget::GetTag() const
 {
-	return this->Tag;
+	return Tag;
 }
 
-
 FSlateColor SWidget::GetForegroundColor() const
 {
 	static FSlateColor NoColor = FSlateColor::UseForeground();
 	return NoColor;
 }
 
-
 void SWidget::SetToolTipText(const TAttribute<FText>& ToolTipText)
 {
 	ToolTip = FSlateApplicationBase::Get().MakeToolTip(ToolTipText);
 }
 
-
 void SWidget::SetToolTipText( const FText& ToolTipText )
 {
 	ToolTip = FSlateApplicationBase::Get().MakeToolTip(ToolTipText);
 }
 
-
 void SWidget::SetToolTip( const TSharedPtr<IToolTip> & InToolTip )
 {
 	ToolTip = InToolTip;
 }
 
-
 TSharedPtr<IToolTip> SWidget::GetToolTip()
 {
 	return ToolTip;
@@ -640,18 +607,13 @@ void SWidget::SetCursor( const TAttribute< TOptional<EMouseCursor::Type> >& InCu
 	Cursor = InCursor;
 }
 
-
 void SWidget::SetDebugInfo( const ANSICHAR* InType, const ANSICHAR* InFile, int32 OnLine )
 {
-	this->TypeOfWidget = InType;
-	this->
-		CreatedInLocation = FName( InFile );
-	this->CreatedInLocation.SetNumber(OnLine);
+	TypeOfWidget = InType;
+	CreatedInLocation = FName( InFile );
+	CreatedInLocation.SetNumber(OnLine);
 }
 
-SLATECORE_API int32 bFoldTick = 1;
-FAutoConsoleVariableRef FoldTick(TEXT("Slate.FoldTick"), bFoldTick, TEXT("When folding, call Tick as part of the paint pass instead of a separate tick pass."));
-
 int32 SWidget::Paint(const FPaintArgs& Args, const FGeometry& AllottedGeometry, const FSlateRect& MyClippingRect, FSlateWindowElementList& OutDrawElements, int32 LayerId, const FWidgetStyle& InWidgetStyle, bool bParentEnabled) const
 {
 	INC_DWORD_STAT(STAT_SlateNumPaintedWidgets);
@@ -672,15 +634,14 @@ int32 SWidget::Paint(const FPaintArgs& Args, const FGeometry& AllottedGeometry,
 		// Volatile widgets don't have an associated layout cache.
 		LayoutCache.Reset();
 
+		const int32 VolatileLayerId = LayerId + 1;
 		OutDrawElements.QueueVolatilePainting(
-			FSlateWindowElementList::FVolatilePaint(SharedThis(this), Args, AllottedGeometry, MyClippingRect, LayerId, InWidgetStyle, bParentEnabled));
+			FSlateWindowElementList::FVolatilePaint(SharedThis(this), Args, AllottedGeometry, MyClippingRect, VolatileLayerId, InWidgetStyle, bParentEnabled));
 
-		// Add a layer padding so that cached elements drawn on top of this volatile widget have some airspace room and not be drawn on top
-		// of when the volatile element draws.
-		return LayerId + 1000;
+		return VolatileLayerId;
 	}
 
-	if ( bFoldTick )
+	if ( bFoldTick && bCanTick )
 	{
 		FGeometry TickGeometry = AllottedGeometry;
 		TickGeometry.AppendTransform( FSlateLayoutTransform(Args.GetWindowToDesktopTransform()) );
@@ -696,9 +657,8 @@ int32 SWidget::Paint(const FPaintArgs& Args, const FGeometry& AllottedGeometry,
 	// Paint the geometry of this widget.
 	int32 NewLayerID = OnPaint(UpdatedArgs, AllottedGeometry, MyClippingRect, OutDrawElements, LayerId, InWidgetStyle, bParentEnabled);
 
-	// HACK Doing this on every widget is expensive to call this virtual on everyone, do we really need this?
-	// it's not even all that useful on a shipped game you'd never use the built in keyboard focusing system.
-	if (SupportsKeyboardFocus())
+	// Check if we need to show the keyboard focus ring, this is only necessary if the widget could be focused.
+	if ( bCanSupportFocus && SupportsKeyboardFocus() )
 	{
 		bool bShowUserFocus = FSlateApplicationBase::Get().ShowUserFocus(SharedThis(this));
 		if (bShowUserFocus)
@@ -728,18 +688,6 @@ float SWidget::GetRelativeLayoutScale(const FSlotBase& Child) const
 	return 1.0f;
 }
 
-void SWidget::ArrangeChildren(const FGeometry& AllottedGeometry, FArrangedChildren& ArrangedChildren) const
-{
-	SLATE_CYCLE_COUNTER_SCOPE_CUSTOM_DETAILED(SLATE_STATS_DETAIL_LEVEL_MED, GSlateArrangeChildren, GetType());
-	OnArrangeChildren(AllottedGeometry, ArrangedChildren);
-}
-
-EVisibility SWidget::GetVisibility() const
-{
-	SLATE_CYCLE_COUNTER_SCOPE_CUSTOM_DETAILED(SLATE_STATS_DETAIL_LEVEL_HI,GSlateGetVisibility, GetType());
-	return Visibility.Get();
-}
-
 TSharedRef<FActiveTimerHandle> SWidget::RegisterActiveTimer(float TickPeriod, FWidgetActiveTimerDelegate TickFunction)
 {
 	TSharedRef<FActiveTimerHandle> ActiveTimerHandle = MakeShareable(new FActiveTimerHandle(TickPeriod, TickFunction, FSlateApplicationBase::Get().GetCurrentTime() + TickPeriod));
diff --git a/Engine/Source/Runtime/SlateCore/Public/Widgets/SBoxPanel.h b/Engine/Source/Runtime/SlateCore/Public/Widgets/SBoxPanel.h
--- a/Engine/Source/Runtime/SlateCore/Public/Widgets/SBoxPanel.h
+++ b/Engine/Source/Runtime/SlateCore/Public/Widgets/SBoxPanel.h
@@ -197,6 +197,9 @@ public:
 	{
 		SHorizontalBox::FSlot& NewSlot = *new SHorizontalBox::FSlot();
 		this->Children.Add( &NewSlot );
+
+		Invalidate(EInvalidateWidget::Layout);
+
 		return NewSlot;
 	}
 
@@ -208,6 +211,9 @@ public:
 		}
 		SHorizontalBox::FSlot& NewSlot = *new SHorizontalBox::FSlot();
 		this->Children.Insert(&NewSlot, Index);
+
+		Invalidate(EInvalidateWidget::Layout);
+
 		return NewSlot;
 	}
 
@@ -219,6 +225,8 @@ public:
 	FORCENOINLINE SHorizontalBox()
 	: SBoxPanel( Orient_Horizontal )
 	{
+		bCanTick = false;
+		bCanSupportFocus = false;
 	}
 
 	/**
@@ -340,6 +348,9 @@ public:
 	{
 		SVerticalBox::FSlot& NewSlot = *new SVerticalBox::FSlot();
 		this->Children.Add( &NewSlot );
+
+		Invalidate(EInvalidateWidget::Layout);
+
 		return NewSlot;
 	}
 
@@ -351,6 +362,9 @@ public:
 		}
 		SVerticalBox::FSlot& NewSlot = *new SVerticalBox::FSlot();
 		this->Children.Insert(&NewSlot, Index);
+
+		Invalidate(EInvalidateWidget::Layout);
+
 		return NewSlot;
 	}
 
@@ -362,6 +376,8 @@ public:
 	FORCENOINLINE SVerticalBox()
 	: SBoxPanel( Orient_Vertical )
 	{
+		bCanTick = false;
+		bCanSupportFocus = false;
 	}
 
 	/**
diff --git a/Engine/Source/Runtime/SlateCore/Public/Widgets/SOverlay.h b/Engine/Source/Runtime/SlateCore/Public/Widgets/SOverlay.h
--- a/Engine/Source/Runtime/SlateCore/Public/Widgets/SOverlay.h
+++ b/Engine/Source/Runtime/SlateCore/Public/Widgets/SOverlay.h
@@ -32,7 +32,7 @@ public:
 			, HAlignment(HAlign_Fill)
 			, VAlignment(VAlign_Fill)
 			, SlotPadding(0.0f)
- 			, ZOrder(0)
+			, ZOrder(0)
 		{ }
 
 		FOverlaySlot& HAlign( EHorizontalAlignment InHAlignment )
@@ -78,7 +78,7 @@ public:
 		TAttribute< FMargin > SlotPadding;
 
 		/** Slots with larger ZOrder values will draw above slots with smaller ZOrder values.  Slots
-		    with the same ZOrder will simply draw in the order they were added.  Currently this only
+			with the same ZOrder will simply draw in the order they were added.  Currently this only
 			works for overlay slots that are added dynamically with AddWidget() and RemoveWidget() */
 		int32 ZOrder;
 	};
@@ -95,10 +95,6 @@ public:
 
 	SOverlay();
 
-	virtual ~SOverlay()
-	{
-	}
-
 	/**
 	 * Construct this widget.
 	 *
diff --git a/Engine/Source/Runtime/SlateCore/Public/Widgets/SWidget.h b/Engine/Source/Runtime/SlateCore/Public/Widgets/SWidget.h
--- a/Engine/Source/Runtime/SlateCore/Public/Widgets/SWidget.h
+++ b/Engine/Source/Runtime/SlateCore/Public/Widgets/SWidget.h
@@ -566,7 +566,7 @@ public:
 	 */
 	private: virtual FVector2D ComputeDesiredSize(float LayoutScaleMultiplier) const = 0;
 
-	public:
+public:
 
 	/** What is the Child's scale relative to this widget. */
 	virtual float GetRelativeLayoutScale( const FSlotBase& Child ) const;
@@ -580,7 +580,10 @@ public:
 	 * @param AllottedGeometry    The geometry allotted for this widget by its parent.
 	 * @param ArrangedChildren    The array to which to add the WidgetGeometries that represent the arranged children.
 	 */
-	void ArrangeChildren( const FGeometry& AllottedGeometry, FArrangedChildren& ArrangedChildren ) const;
+	FORCEINLINE void ArrangeChildren(const FGeometry& AllottedGeometry, FArrangedChildren& ArrangedChildren) const
+	{
+		OnArrangeChildren(AllottedGeometry, ArrangedChildren);
+	}
 
 	/**
 	 * Every widget that has children must implement this method. This allows for iteration over the Widget's
@@ -701,7 +704,7 @@ public:
 	virtual bool IsDirectlyHovered() const;
 
 	/** @return is this widget visible, hidden or collapsed */
-	EVisibility GetVisibility() const;
+	FORCEINLINE EVisibility GetVisibility() const { return Visibility.Get(); }
 
 	/** @param InVisibility  should this widget be */
 	virtual void SetVisibility( TAttribute<EVisibility> InVisibility )
@@ -1106,7 +1109,19 @@ private:
 
 protected:
 	/** Is this widget hovered? */
-	bool bIsHovered;
+	bool bIsHovered : 1;
+
+	/** Can the widget ever be ticked. */
+	bool bCanTick : 1;
+
+	/** Can the widget ever support keyboard focus */
+	bool bCanSupportFocus : 1;
+
+	/**
+	 * Can the widget ever support children?  This will be false on SLeafWidgets, 
+	 * rather than setting this directly, you should probably inherit from SLeafWidget.
+	 */
+	bool bCanHaveChildren : 1;
 
 private:
 
@@ -1114,14 +1129,14 @@ private:
 	 * Whether this widget is a "tool tip force field".  That is, tool-tips should never spawn over the area
 	 * occupied by this widget, and will instead be repelled to an outside edge
 	 */
-	bool bToolTipForceFieldEnabled;
+	bool bToolTipForceFieldEnabled : 1;
 
 	/** Should we be forcing this widget to be volatile at all times and redrawn every frame? */
-	bool bForceVolatile;
+	bool bForceVolatile : 1;
 
 	/** The last cached volatility of this widget.  Cached so that we don't need to recompute volatility every frame. */
-	bool bCachedVolatile;
+	bool bCachedVolatile : 1;
 
 	/** If we're owned by a volatile widget, we need inherit that volatility and use as part of our volatility, but don't cache it. */
-	mutable bool bInheritedVolatility;
+	mutable bool bInheritedVolatility : 1;
 };
