UMG/Slate - The invalidation panel can now cache a GPU Buffer for a set of drawable elements, this avoids needing to go through the element batcher which can be very expensive on large UI set.  Currently only supported on the RHI Renderer.  This change required we introduce a sort of logical layering system into the rendering stack for Slate, so that volatile elements that could expand to be N long, could be sorted relative to their parent in terms of layer id, but without knowing that ahead of time, it became impossible to properly sort child elements that came after the volatile set.  The logical layering avoids that by allowing us to sort each logical group independently, and as we traverse the tree only sort along each branch.

diff --git a/Engine/Source/Developer/StandaloneRenderer/Private/OpenGL/SlateOpenGLRenderer.cpp b/Engine/Source/Developer/StandaloneRenderer/Private/OpenGL/SlateOpenGLRenderer.cpp
--- a/Engine/Source/Developer/StandaloneRenderer/Private/OpenGL/SlateOpenGLRenderer.cpp
+++ b/Engine/Source/Developer/StandaloneRenderer/Private/OpenGL/SlateOpenGLRenderer.cpp
@@ -129,7 +129,7 @@ void FSlateOpenGLRenderer::DrawWindows( FSlateDrawBuffer& InWindowDrawBuffer )
 			
 			FSlateBatchData& BatchData = ElementList.GetBatchData();
 
-			BatchData.CreateRenderBatches();
+			BatchData.CreateRenderBatches(ElementList.GetRootDrawLayer().GetElementBatchMap());
 
 			RenderingPolicy->UpdateVertexAndIndexBuffers( BatchData );
 
diff --git a/Engine/Source/Developer/StandaloneRenderer/Private/Windows/D3D/SlateD3DRenderer.cpp b/Engine/Source/Developer/StandaloneRenderer/Private/Windows/D3D/SlateD3DRenderer.cpp
--- a/Engine/Source/Developer/StandaloneRenderer/Private/Windows/D3D/SlateD3DRenderer.cpp
+++ b/Engine/Source/Developer/StandaloneRenderer/Private/Windows/D3D/SlateD3DRenderer.cpp
@@ -364,7 +364,7 @@ void FSlateD3DRenderer::DrawWindows( FSlateDrawBuffer& InWindowDrawBuffer )
 			FSlateBatchData& BatchData = ElementList.GetBatchData();
 			{
 				SLATE_CYCLE_COUNTER_SCOPE(GRendererUpdateBuffers);
-				BatchData.CreateRenderBatches();
+				BatchData.CreateRenderBatches(ElementList.GetRootDrawLayer().GetElementBatchMap());
 				RenderingPolicy->UpdateVertexAndIndexBuffers(BatchData);
 			}
 
diff --git a/Engine/Source/Runtime/Core/Private/Math/Color.cpp b/Engine/Source/Runtime/Core/Private/Math/Color.cpp
--- a/Engine/Source/Runtime/Core/Private/Math/Color.cpp
+++ b/Engine/Source/Runtime/Core/Private/Math/Color.cpp
@@ -25,6 +25,10 @@ const FColor FColor::Blue(0,0,255);
 const FColor FColor::Yellow(255,255,0);
 const FColor FColor::Cyan(0,255,255);
 const FColor FColor::Magenta(255,0,255);
+const FColor FColor::Orange(243, 156, 18);
+const FColor FColor::Turquoise(26, 188, 156);
+const FColor FColor::Silver(189, 195, 199);
+const FColor FColor::Emerald(46, 204, 113);
 
 /**
 * Helper used by FColor -> FLinearColor conversion. We don't use a lookup table as unlike pow, multiplication is fast.
diff --git a/Engine/Source/Runtime/Core/Public/Math/Color.h b/Engine/Source/Runtime/Core/Public/Math/Color.h
--- a/Engine/Source/Runtime/Core/Public/Math/Color.h
+++ b/Engine/Source/Runtime/Core/Public/Math/Color.h
@@ -572,6 +572,10 @@ public:
 	static CORE_API const FColor Yellow;
 	static CORE_API const FColor Cyan;
 	static CORE_API const FColor Magenta;
+	static CORE_API const FColor Orange;
+	static CORE_API const FColor Turquoise;
+	static CORE_API const FColor Silver;
+	static CORE_API const FColor Emerald;
 
 private:
 	/**
diff --git a/Engine/Source/Runtime/Engine/Private/Slate/SGameLayerManager.cpp b/Engine/Source/Runtime/Engine/Private/Slate/SGameLayerManager.cpp
--- a/Engine/Source/Runtime/Engine/Private/Slate/SGameLayerManager.cpp
+++ b/Engine/Source/Runtime/Engine/Private/Slate/SGameLayerManager.cpp
@@ -165,6 +165,14 @@ void SGameLayerManager::Tick(const FGeometry& AllottedGeometry, const double InC
 	UpdateLayout();
 }
 
+int32 SGameLayerManager::OnPaint(const FPaintArgs& Args, const FGeometry& AllottedGeometry, const FSlateRect& MyClippingRect, FSlateWindowElementList& OutDrawElements, int32 LayerId, const FWidgetStyle& InWidgetStyle, bool bParentEnabled) const
+{
+	//FPlatformMisc::BeginNamedEvent(FColor::Green, "Game Layer Manager");
+	int32 ResultLayer = SCompoundWidget::OnPaint(Args, AllottedGeometry, MyClippingRect, OutDrawElements, LayerId, InWidgetStyle, bParentEnabled);
+	//FPlatformMisc::EndNamedEvent();
+	return ResultLayer;
+}
+
 bool SGameLayerManager::OnVisualizeTooltip(const TSharedPtr<SWidget>& TooltipContent)
 {
 	TooltipPresenter->SetContent(TooltipContent.IsValid() ? TooltipContent.ToSharedRef() : SNullWidget::NullWidget);
diff --git a/Engine/Source/Runtime/Engine/Public/Slate/SGameLayerManager.h b/Engine/Source/Runtime/Engine/Public/Slate/SGameLayerManager.h
--- a/Engine/Source/Runtime/Engine/Public/Slate/SGameLayerManager.h
+++ b/Engine/Source/Runtime/Engine/Public/Slate/SGameLayerManager.h
@@ -89,6 +89,7 @@ public:
 
 	// Begin SWidget overrides
 	virtual void Tick(const FGeometry& AllottedGeometry, const double InCurrentTime, const float InDeltaTime) override;
+	virtual int32 OnPaint(const FPaintArgs& Args, const FGeometry& AllottedGeometry, const FSlateRect& MyClippingRect, FSlateWindowElementList& OutDrawElements, int32 LayerId, const FWidgetStyle& InWidgetStyle, bool bParentEnabled) const override;
 	virtual bool OnVisualizeTooltip(const TSharedPtr<SWidget>& TooltipContent) override;
 	// End SWidget overrides
 
diff --git a/Engine/Source/Runtime/Slate/Private/Framework/Application/SlateApplication.cpp b/Engine/Source/Runtime/Slate/Private/Framework/Application/SlateApplication.cpp
--- a/Engine/Source/Runtime/Slate/Private/Framework/Application/SlateApplication.cpp
+++ b/Engine/Source/Runtime/Slate/Private/Framework/Application/SlateApplication.cpp
@@ -924,7 +924,9 @@ void FSlateApplication::TickWindowAndChildren( TSharedRef<SWindow> WindowToTick
 void FSlateApplication::DrawWindows()
 {
 	SLATE_CYCLE_COUNTER_SCOPE(GSlateDrawWindows);
+	FPlatformMisc::BeginNamedEvent(FColor::Magenta, "Slate::DrawWindows");
 	PrivateDrawWindows();
+	FPlatformMisc::EndNamedEvent();
 }
 
 struct FDrawWindowArgs
@@ -948,13 +950,19 @@ void FSlateApplication::DrawWindowAndChildren( const TSharedRef<SWindow>& Window
 		// Switch to the appropriate world for drawing
 		FScopedSwitchWorldHack SwitchWorld( WindowToDraw );
 
+		//FPlatformMisc::BeginNamedEvent(FColor::Magenta, "Slate::DrawPrep");
 		FSlateWindowElementList& WindowElementList = DrawWindowArgs.OutDrawBuffer.AddWindowElementList( WindowToDraw );
+		//FPlatformMisc::EndNamedEvent();
 
 		// Drawing is done in window space, so null out the positions and keep the size.
 		FGeometry WindowGeometry = WindowToDraw->GetWindowGeometryInWindow();
 		int32 MaxLayerId = 0;
 		{
+			//FPlatformMisc::BeginNamedEvent(FColor::Magenta, "Slate::ClearHitTestGrid");
 			WindowToDraw->GetHittestGrid()->ClearGridForNewFrame( VirtualDesktopRect );
+			//FPlatformMisc::EndNamedEvent();
+
+			FPlatformMisc::BeginNamedEvent(FColor::Magenta, "Slate::DrawWindow");
 			MaxLayerId = WindowToDraw->PaintWindow(
 				FPaintArgs(WindowToDraw, *WindowToDraw->GetHittestGrid(), WindowToDraw->GetPositionInScreen(), GetCurrentTime(), GetDeltaTime()),
 				WindowGeometry, WindowToDraw->GetClippingRectangleInWindow(),
@@ -962,6 +970,7 @@ void FSlateApplication::DrawWindowAndChildren( const TSharedRef<SWindow>& Window
 				0,
 				FWidgetStyle(),
 				WindowToDraw->IsEnabled() );
+			FPlatformMisc::EndNamedEvent();
 
 			// Draw drag drop operation if it's windowless.
 			if ( IsDragDropping() && DragDropContent->IsWindowlessOperation() )
@@ -1137,10 +1146,14 @@ void FSlateApplication::PrivateDrawWindows( TSharedPtr<SWindow> DrawOnlyThisWind
 
 	if ( !SkipSecondPrepass.GetValueOnGameThread() )
 	{
+		//FPlatformMisc::BeginNamedEvent(FColor::Magenta, "Slate::Prepass");
 		DrawPrepass( DrawOnlyThisWindow );
+		//FPlatformMisc::EndNamedEvent();
 	}
 
+	//FPlatformMisc::BeginNamedEvent(FColor::Magenta, "Slate::GetDrawBuffer");
 	FDrawWindowArgs DrawWindowArgs( Renderer->GetDrawBuffer(), WidgetsUnderCursor );
+	//FPlatformMisc::EndNamedEvent();
 
 	{
 		SCOPE_CYCLE_COUNTER( STAT_SlateDrawWindowTime );
@@ -1281,12 +1294,16 @@ void FSlateApplication::Tick()
 	{
 	SCOPE_CYCLE_COUNTER( STAT_SlateTickTime );
 	SLATE_CYCLE_COUNTER_SCOPE(GSlateTotalTickTime);
+
+	FPlatformMisc::BeginNamedEvent(FColor::Magenta, "Slate::Tick");
 	
 	{
 		const float DeltaTime = GetDeltaTime();
 
 		SCOPE_CYCLE_COUNTER( STAT_SlateMessageTick );
 
+		//FPlatformMisc::BeginNamedEvent(FColor::Magenta, "Slate::Platform");
+
 		// We need to pump messages here so that slate can receive input.  
 		if( (ActiveModalWindows.Num() > 0) || GIntraFrameDebuggingGameThread )
 		{
@@ -1302,13 +1319,18 @@ void FSlateApplication::Tick()
 		PlatformApplication->Tick( DeltaTime );
 
 		PlatformApplication->ProcessDeferredEvents( DeltaTime );
+
+		//FPlatformMisc::EndNamedEvent();
 	}
 
+	//FPlatformMisc::BeginNamedEvent(FColor::Magenta, "Slate::UpdateCursorLockRegion");
 	// The widget locking the cursor to its bounds may have been reshaped.
 	// Check if the widget was reshaped and update the cursor lock
 	// bounds if needed.
 	UpdateCursorLockRegion();
+	//FPlatformMisc::EndNamedEvent();
 
+	//FPlatformMisc::BeginNamedEvent(FColor::Magenta, "Slate::CaptureAndToolTipUpdate");
 	// When Slate captures the mouse, it is up to us to set the cursor 
 	// because the OS assumes that we own the mouse.
 	if (MouseCaptor.HasCapture())
@@ -1324,6 +1346,7 @@ void FSlateApplication::Tick()
 		const bool AllowSpawningOfToolTips = false;
 		UpdateToolTip( AllowSpawningOfToolTips );
 	}
+	//FPlatformMisc::EndNamedEvent();
 
 
 	// Advance time
@@ -1442,6 +1465,8 @@ void FSlateApplication::Tick()
 	}
 	}
 
+	FPlatformMisc::EndNamedEvent();
+
 	// Update Slate Stats
 	SLATE_STATS_END_FRAME(GetCurrentTime());
 }
diff --git a/Engine/Source/Runtime/Slate/Private/Framework/Text/TextLayout.cpp b/Engine/Source/Runtime/Slate/Private/Framework/Text/TextLayout.cpp
--- a/Engine/Source/Runtime/Slate/Private/Framework/Text/TextLayout.cpp
+++ b/Engine/Source/Runtime/Slate/Private/Framework/Text/TextLayout.cpp
@@ -1761,11 +1761,6 @@ void FTextLayout::SetMargin( const FMargin& InMargin )
 	DirtyFlags |= EDirtyState::Layout;
 }
 
-FMargin FTextLayout::GetMargin() const
-{
-	return Margin;
-}
-
 void FTextLayout::SetScale( float Value )
 {
 	if ( Scale == Value )
@@ -1781,11 +1776,6 @@ void FTextLayout::SetScale( float Value )
 	ClearWrappingCache();
 }
 
-float FTextLayout::GetScale() const
-{
-	return Scale;
-}
-
 void FTextLayout::SetJustification( ETextJustify::Type Value )
 {
 	if ( Justification == Value )
@@ -1797,11 +1787,6 @@ void FTextLayout::SetJustification( ETextJustify::Type Value )
 	DirtyFlags |= EDirtyState::Layout;
 }
 
-ETextJustify::Type FTextLayout::GetJustification() const
-{
-	return Justification;
-}
-
 void FTextLayout::SetLineHeightPercentage( float Value )
 {
 	if ( LineHeightPercentage != Value )
@@ -1811,11 +1796,6 @@ void FTextLayout::SetLineHeightPercentage( float Value )
 	}
 }
 
-float FTextLayout::GetLineHeightPercentage() const
-{
-	return LineHeightPercentage;
-}
-
 void FTextLayout::SetWrappingWidth( float Value )
 {
 	if ( WrappingWidth != Value )
@@ -1825,11 +1805,6 @@ void FTextLayout::SetWrappingWidth( float Value )
 	}
 }
 
-float FTextLayout::GetWrappingWidth() const
-{
-	return WrappingWidth;
-}
-
 FVector2D FTextLayout::GetDrawSize() const
 {
 	return TextLayoutSize.GetDrawSize();
@@ -1845,16 +1820,6 @@ FVector2D FTextLayout::GetSize() const
 	return TextLayoutSize.GetDrawSize() * ( 1 / Scale );
 }
 
-const TArray< FTextLayout::FLineModel >& FTextLayout::GetLineModels() const
-{
-	return LineModels;
-}
-
-const TArray< FTextLayout::FLineView >& FTextLayout::GetLineViews() const
-{
-	return LineViews;
-}
-
 FTextLayout::~FTextLayout()
 {
 
diff --git a/Engine/Source/Runtime/Slate/Private/Widgets/SInvalidationPanel.cpp b/Engine/Source/Runtime/Slate/Private/Widgets/SInvalidationPanel.cpp
--- a/Engine/Source/Runtime/Slate/Private/Widgets/SInvalidationPanel.cpp
+++ b/Engine/Source/Runtime/Slate/Private/Widgets/SInvalidationPanel.cpp
@@ -5,7 +5,13 @@
 #include "SInvalidationPanel.h"
 #include "WidgetCaching.h"
 
-#if !UE_BUILD_SHIPPING
+DECLARE_CYCLE_STAT(TEXT("Invalidation Time"), STAT_InvalidationTime, STATGROUP_Slate);
+DECLARE_DWORD_COUNTER_STAT(TEXT("Num Cached Elements"), STAT_SlateNumCachedElements, STATGROUP_Slate);
+DECLARE_DWORD_COUNTER_STAT(TEXT("Num Invalidated Elements"), STAT_SlateNumInvalidatedElements, STATGROUP_Slate);
+DECLARE_DWORD_COUNTER_STAT(TEXT("Num Volatile Widgets"), STAT_SlateNumVolatileWidgets, STATGROUP_Slate);
+
+
+#if !(UE_BUILD_SHIPPING || UE_BUILD_TEST)
 
 /** True if we should allow widgets to be cached in the UI at all. */
 TAutoConsoleVariable<int32> InvalidationDebugging(
@@ -43,7 +49,6 @@ void SInvalidationPanel::Construct( const FArguments& InArgs )
 	bCanCache = true;
 	RootCacheNode = nullptr;
 	LastUsedCachedNodeIndex = 0;
-	LastLayerId = 0;
 	LastHitTestIndex = 0;
 
 	bCacheRelativeTransforms = InArgs._CacheRelativeTransforms;
@@ -59,7 +64,7 @@ SInvalidationPanel::~SInvalidationPanel()
 
 bool SInvalidationPanel::GetCanCache() const
 {
-#if !UE_BUILD_SHIPPING
+#if !(UE_BUILD_SHIPPING || UE_BUILD_TEST)
 	return bCanCache && EnableWidgetCaching.GetValueOnGameThread() == 1;
 #else
 	return bCanCache;
@@ -75,8 +80,12 @@ void SInvalidationPanel::SetCanCache(bool InCanCache)
 
 void SInvalidationPanel::Tick( const FGeometry& AllottedGeometry, const double InCurrentTime, const float InDeltaTime )
 {
+	//FPlatformMisc::BeginNamedEvent(FColor::Magenta, "Slate::InvalidationPanel::Tick");
+
 	if ( GetCanCache() )
 	{
+		SCOPE_CYCLE_COUNTER(STAT_InvalidationTime);
+
 		const bool bWasCachingNeeded = bNeedsCaching;
 
 		if ( bNeedsCaching == false )
@@ -116,6 +125,8 @@ void SInvalidationPanel::Tick( const FGeometry& AllottedGeometry, const double I
 			CachePrepass(SharedThis(this));
 		}
 	}
+
+	//FPlatformMisc::EndNamedEvent();
 }
 
 FChildren* SInvalidationPanel::GetChildren()
@@ -134,10 +145,10 @@ void SInvalidationPanel::InvalidateWidget(SWidget* InvalidateWidget)
 {
 	bNeedsCaching = true;
 
-#if !UE_BUILD_SHIPPING
+#if !(UE_BUILD_SHIPPING || UE_BUILD_TEST)
 	if ( InvalidateWidget != nullptr && IsInvalidationDebuggingEnabled() )
 	{
-		InvalidatorWidgets.Add(InvalidateWidget->AsShared());
+		InvalidatorWidgets.Add(InvalidateWidget->AsShared(), 1);
 	}
 #endif
 }
@@ -160,30 +171,74 @@ FCachedWidgetNode* SInvalidationPanel::CreateCacheNode() const
 	return NewNode;
 }
 
+TSharedPtr< FSlateWindowElementList > SInvalidationPanel::GetNextCachedElementList(const TSharedPtr<SWindow>& CurrentWindow) const
+{
+	TSharedPtr< FSlateWindowElementList > NextElementList;
+
+	// Move any inactive element lists in the active pool to the inactive pool.
+	for ( int32 i = 0; i < ActiveCachedElementListPool.Num(); i++ )
+	{
+		if ( ActiveCachedElementListPool[i]->IsCachedRenderDataInUse() == false )
+		{
+			InactiveCachedElementListPool.Add(ActiveCachedElementListPool[i]);
+			ActiveCachedElementListPool.RemoveAtSwap(i, 1, false);
+		}
+	}
+
+	// Remove inactive lists that don't belong to this window.
+	for ( int32 i = 0; i < InactiveCachedElementListPool.Num(); i++ )
+	{
+		if ( InactiveCachedElementListPool[i]->GetWindow() != CurrentWindow )
+		{
+			InactiveCachedElementListPool.RemoveAtSwap(i, 1, false);
+		}
+	}
+
+	// Create a new element list if none are available, or use an existing one.
+	if ( InactiveCachedElementListPool.Num() == 0 )
+	{
+		NextElementList = MakeShareable(new FSlateWindowElementList(CurrentWindow));
+	}
+	else
+	{
+		NextElementList = InactiveCachedElementListPool[0];
+		NextElementList->ResetBuffers();
+
+		InactiveCachedElementListPool.RemoveAtSwap(0, 1, false);
+	}
+
+	ActiveCachedElementListPool.Add(NextElementList);
+
+	return NextElementList;
+}
+
 int32 SInvalidationPanel::OnPaint( const FPaintArgs& Args, const FGeometry& AllottedGeometry, const FSlateRect& MyClippingRect, FSlateWindowElementList& OutDrawElements, int32 LayerId, const FWidgetStyle& InWidgetStyle, bool bParentEnabled ) const
 {
 	if ( GetCanCache() )
 	{
+		SCOPE_CYCLE_COUNTER(STAT_InvalidationTime);
+
+		//FPlatformMisc::BeginNamedEvent(FColor::Magenta, "Slate::InvalidationPanel::Paint");
+
 		const bool bWasCachingNeeded = bNeedsCaching;
 
 		if ( bNeedsCaching )
 		{
+			//FPlatformMisc::BeginNamedEvent(FColor::Red, "Slate::Invalidation");
+
 			SInvalidationPanel* MutableThis = const_cast<SInvalidationPanel*>( this );
 
 			// Always set the caching flag to false first, during the paint / tick pass we may change something
 			// to volatile and need to re-cache.
 			bNeedsCaching = false;
 
+			// Mark that we're in the process of invalidating.
 			bIsInvalidating = true;
 
-			if ( !CachedWindowElements.IsValid() || CachedWindowElements->GetWindow() != OutDrawElements.GetWindow() )
-			{
-				CachedWindowElements = MakeShareable(new FSlateWindowElementList(OutDrawElements.GetWindow()));
-			}
-			else
-			{
-				CachedWindowElements->Reset();
-			}
+			CachedWindowElements = GetNextCachedElementList(OutDrawElements.GetWindow());
+
+			// Reset the render data handle in case it was in use, and we're not overriding it this frame.
+			CachedRenderData.Reset();
 
 			// Reset the cached node pool index so that we effectively reset the pool.
 			LastUsedCachedNodeIndex = 0;
@@ -204,25 +259,52 @@ int32 SInvalidationPanel::OnPaint( const FPaintArgs& Args, const FGeometry& Allo
 			if ( bCacheRelativeTransforms )
 			{
 				CachedAbsolutePosition = AllottedGeometry.Position;
+				AbsoluteDeltaPosition = FVector2D(0, 0);
 			}
 
-			LastLayerId = LayerId;
+#if WITH_ENGINE
+			CachedRenderData = CachedWindowElements->CacheRenderData();
+#endif
+
 			LastHitTestIndex = Args.GetLastHitTestIndex();
 
 			bIsInvalidating = false;
+
+			//FPlatformMisc::EndNamedEvent();
 		}
 
+		//FPlatformMisc::BeginNamedEvent(FColor::Yellow, "Slate::RecordHitTestGeometry");
+
 		// The hit test grid is actually populated during the initial cache phase, so don't bother
 		// recording the hit test geometry on the same frame that we regenerate the cache.
 		if ( bWasCachingNeeded == false )
 		{
+			INC_DWORD_STAT_BY(STAT_SlateNumCachedElements, CachedWindowElements->GetDrawElements().Num());
+
 			RootCacheNode->RecordHittestGeometry(Args.GetGrid(), Args.GetLastHitTestIndex());
 		}
+		else
+		{
+			INC_DWORD_STAT_BY(STAT_SlateNumInvalidatedElements, CachedWindowElements->GetDrawElements().Num());
+		}
+		//FPlatformMisc::EndNamedEvent();
+
+		int32 OutMaxChildLayer = CachedMaxChildLayer;
 
 		if ( bCacheRelativeTransforms )
 		{
-			FVector2D DeltaPosition = AllottedGeometry.Position - CachedAbsolutePosition;
+			FVector2D NewAbsoluteDeltaPosition = AllottedGeometry.Position - CachedAbsolutePosition;
+			FVector2D RelativeDeltaPosition = NewAbsoluteDeltaPosition - AbsoluteDeltaPosition;
+			AbsoluteDeltaPosition = NewAbsoluteDeltaPosition;
+
+#if WITH_ENGINE
+			if ( bWasCachingNeeded == false && RelativeDeltaPosition.IsZero() == false )
+			{
+				CachedWindowElements->UpdateCacheRenderData(RelativeDeltaPosition);
+			}
 
+			FSlateDrawElement::MakeCachedBuffer(OutDrawElements, LayerId, CachedRenderData);
+#else
 			const TArray<FSlateDrawElement>& CachedElements = CachedWindowElements->GetDrawElements();
 			const int32 CachedElementCount = CachedElements.Num();
 			for ( int32 Index = 0; Index < CachedElementCount; Index++ )
@@ -230,26 +312,35 @@ int32 SInvalidationPanel::OnPaint( const FPaintArgs& Args, const FGeometry& Allo
 				const FSlateDrawElement& LocalElement = CachedElements[Index];
 				FSlateDrawElement AbsElement = LocalElement;
 
-				AbsElement.SetPosition(LocalElement.GetPosition() + DeltaPosition);
-				AbsElement.SetClippingRect(LocalElement.GetClippingRect().OffsetBy(DeltaPosition));
+				AbsElement.SetPosition(LocalElement.GetPosition() + AbsoluteDeltaPosition);
+				AbsElement.SetClippingRect(LocalElement.GetClippingRect().OffsetBy(AbsoluteDeltaPosition));
 
 				OutDrawElements.AddItem(AbsElement);
 			}
+#endif
 		}
 		else
 		{
+#if WITH_ENGINE
+			FSlateDrawElement::MakeCachedBuffer(OutDrawElements, LayerId, CachedRenderData);
+#else
 			OutDrawElements.AppendDrawElements(CachedWindowElements->GetDrawElements());
+#endif
 		}
 
-		int32 OutMaxChildLayer = CachedMaxChildLayer;
-
+		//FPlatformMisc::BeginNamedEvent(FColor::Orange, "Slate::VolatilePainting");
 		// Paint the volatile elements
 		if ( CachedWindowElements.IsValid() )
 		{
-			OutMaxChildLayer = FMath::Max(CachedMaxChildLayer, CachedWindowElements->PaintVolatile(OutDrawElements));
+			const TArray<TSharedPtr<FSlateWindowElementList::FVolatilePaint>>& VolatileElements = CachedWindowElements->GetVolatileElements();
+			INC_DWORD_STAT_BY(STAT_SlateNumVolatileWidgets, VolatileElements.Num());
+
+			// TODO Offset? AbsoluteDeltaPosition
+			OutMaxChildLayer = FMath::Max(OutMaxChildLayer, CachedWindowElements->PaintVolatile(OutDrawElements));
 		}
+		//FPlatformMisc::EndNamedEvent();
 
-#if !UE_BUILD_SHIPPING
+#if !(UE_BUILD_SHIPPING || UE_BUILD_TEST)
 
 		if ( IsInvalidationDebuggingEnabled() )
 		{
@@ -257,7 +348,7 @@ int32 SInvalidationPanel::OnPaint( const FPaintArgs& Args, const FGeometry& Allo
 			{
 				check(Args.IsCaching() == false);
 				//const bool bShowOutlineAsCached = Args.IsCaching() || bWasCachingNeeded == false;
-				const FLinearColor DebugTint = bWasCachingNeeded ? FLinearColor::Red : FLinearColor::Green;
+				const FLinearColor DebugTint = bWasCachingNeeded ? FLinearColor::Red : ( bCacheRelativeTransforms ? FLinearColor::Blue : FLinearColor::Green );
 
 				FGeometry ScaledOutline = AllottedGeometry.MakeChild(FVector2D(0, 0), AllottedGeometry.GetLocalSize() * AllottedGeometry.Scale, Inverse(AllottedGeometry.Scale));
 
@@ -272,10 +363,22 @@ int32 SInvalidationPanel::OnPaint( const FPaintArgs& Args, const FGeometry& Allo
 				);
 			}
 
+			static const FName InvalidationPanelName(TEXT("SInvalidationPanel"));
+
 			// Draw a yellow outline around any volatile elements.
-			const TArray< TSharedRef<FSlateWindowElementList::FVolatilePaint> >& VolatileElements = CachedWindowElements->GetVolatileElements();
-			for ( const TSharedRef<FSlateWindowElementList::FVolatilePaint>& VolatileElement : VolatileElements )
+			const TArray<TSharedPtr<FSlateWindowElementList::FVolatilePaint>>& VolatileElements = CachedWindowElements->GetVolatileElements();
+			for ( const TSharedPtr<FSlateWindowElementList::FVolatilePaint>& VolatileElement : VolatileElements )
 			{
+				// Ignore drawing the volatility rect for child invalidation panels, that's not really important, since
+				// they're always volatile and it will make it hard to see when they're invalidated.
+				if ( const SWidget* Widget = VolatileElement->GetWidget() )
+				{
+					if ( Widget->GetType() == InvalidationPanelName )
+					{
+						continue;
+					}
+				}
+
 				FSlateDrawElement::MakeBox(
 					OutDrawElements,
 					++OutMaxChildLayer,
@@ -287,10 +390,11 @@ int32 SInvalidationPanel::OnPaint( const FPaintArgs& Args, const FGeometry& Allo
 				);
 			}
 
-			// Draw a white flash for any widget that invalidated us this frame.
-			for ( TWeakPtr<SWidget> Invalidator : InvalidatorWidgets )
+			// Draw a red flash for any widget that invalidated us recently, we slowly 
+			// fade out the flashes over time, unless the widget invalidates us again.
+			for ( TMap<TWeakPtr<SWidget>, double>::TIterator It(InvalidatorWidgets); It; ++It )
 			{
-				TSharedPtr<SWidget> SafeInvalidator = Invalidator.Pin();
+				TSharedPtr<SWidget> SafeInvalidator = It.Key().Pin();
 				if ( SafeInvalidator.IsValid() )
 				{
 					FWidgetPath WidgetPath;
@@ -306,17 +410,28 @@ int32 SInvalidationPanel::OnPaint( const FPaintArgs& Args, const FGeometry& Allo
 							FCoreStyle::Get().GetBrush(TEXT("WhiteBrush")),
 							MyClippingRect,
 							ESlateDrawEffect::None,
-							FLinearColor::White.CopyWithNewOpacity(0.6f)
+							FLinearColor::Red.CopyWithNewOpacity(0.75f * It.Value())
 						);
 					}
+
+					It.Value() -= FApp::GetDeltaTime();
+
+					if ( It.Value() <= 0 )
+					{
+						It.RemoveCurrent();
+					}
+				}
+				else
+				{
+					It.RemoveCurrent();
 				}
 			}
-
-			InvalidatorWidgets.Reset();
 		}
 
 #endif
 
+		//FPlatformMisc::EndNamedEvent();
+
 		return OutMaxChildLayer;
 	}
 	else
diff --git a/Engine/Source/Runtime/Slate/Private/Widgets/SViewport.cpp b/Engine/Source/Runtime/Slate/Private/Widgets/SViewport.cpp
--- a/Engine/Source/Runtime/Slate/Private/Widgets/SViewport.cpp
+++ b/Engine/Source/Runtime/Slate/Private/Widgets/SViewport.cpp
@@ -3,6 +3,8 @@
 #include "SlatePrivatePCH.h"
 #include "HittestGrid.h"
 
+DECLARE_CYCLE_STAT(TEXT("Viewport Update Time"), STAT_ViewportUpdateTime, STATGROUP_Slate);
+
 /* SViewport constructors
  *****************************************************************************/
 
@@ -52,6 +54,8 @@ EActiveTimerReturnType SViewport::EnsureTick(double InCurrentTime, float InDelta
 
 int32 SViewport::OnPaint( const FPaintArgs& Args, const FGeometry& AllottedGeometry, const FSlateRect& MyClippingRect, FSlateWindowElementList& OutDrawElements, int32 LayerId, const FWidgetStyle& InWidgetStyle, bool bParentEnabled ) const
 {
+	SCOPE_CYCLE_COUNTER(STAT_ViewportUpdateTime);
+
 	bool bEnabled = ShouldBeEnabled( bParentEnabled );
 	bool bShowDisabledEffect = ShowDisabledEffect.Get();
 	ESlateDrawEffect::Type DrawEffects = bShowDisabledEffect && !bEnabled ? ESlateDrawEffect::DisabledEffect : ESlateDrawEffect::None;
@@ -123,19 +127,20 @@ int32 SViewport::OnPaint( const FPaintArgs& Args, const FGeometry& AllottedGeome
 		);
 	}
 
-
-	if(CustomHitTestPath.IsValid())
+	// If there are any custom hit testable widgets in the 3D world we need to register their custom hit test path here.
+	if ( CustomHitTestPath.IsValid() )
 	{
 		Args.InsertCustomHitTestPath(CustomHitTestPath.ToSharedRef(), LastHitTestIndex);
 	}
 
-
 	return Layer;
 }
 
 void SViewport::Tick( const FGeometry& AllottedGeometry, const double InCurrentTime, const float InDeltaTime )
 {
-	if(ViewportInterface.IsValid())
+	SCOPE_CYCLE_COUNTER(STAT_ViewportUpdateTime);
+
+	if ( ViewportInterface.IsValid() )
 	{
 		ViewportInterface.Pin()->Tick(AllottedGeometry, InCurrentTime, InDeltaTime);
 	}
@@ -160,13 +165,11 @@ FReply SViewport::OnMouseButtonDown( const FGeometry& MyGeometry, const FPointer
 	return ViewportInterface.IsValid() ? ViewportInterface.Pin()->OnMouseButtonDown(MyGeometry, MouseEvent) : FReply::Unhandled();
 }
 
-
 FReply SViewport::OnMouseButtonUp( const FGeometry& MyGeometry, const FPointerEvent& MouseEvent )
 {
 	return ViewportInterface.IsValid() ? ViewportInterface.Pin()->OnMouseButtonUp(MyGeometry, MouseEvent) : FReply::Unhandled();
 }
 
-
 void SViewport::OnMouseEnter( const FGeometry& MyGeometry, const FPointerEvent& MouseEvent )
 {
 	SCompoundWidget::OnMouseEnter(MyGeometry, MouseEvent);
@@ -177,7 +180,6 @@ void SViewport::OnMouseEnter( const FGeometry& MyGeometry, const FPointerEvent&
 	}
 }
 
-
 void SViewport::OnMouseLeave( const FPointerEvent& MouseEvent )
 {
 	SCompoundWidget::OnMouseLeave(MouseEvent);
@@ -188,31 +190,26 @@ void SViewport::OnMouseLeave( const FPointerEvent& MouseEvent )
 	}
 }
 
-
 FReply SViewport::OnMouseMove( const FGeometry& MyGeometry, const FPointerEvent& MouseEvent )
 {
 	return ViewportInterface.IsValid() ? ViewportInterface.Pin()->OnMouseMove(MyGeometry, MouseEvent) : FReply::Unhandled();
 }
 
-
 FReply SViewport::OnMouseWheel( const FGeometry& MyGeometry, const FPointerEvent& MouseEvent )
 {
 	return ViewportInterface.IsValid() ? ViewportInterface.Pin()->OnMouseWheel(MyGeometry, MouseEvent) : FReply::Unhandled();
 }
 
-
 FReply SViewport::OnMouseButtonDoubleClick( const FGeometry& MyGeometry, const FPointerEvent& MouseEvent )
 {
 	return ViewportInterface.IsValid() ? ViewportInterface.Pin()->OnMouseButtonDoubleClick(MyGeometry, MouseEvent) : FReply::Unhandled();
 }
 
-
 FReply SViewport::OnKeyDown( const FGeometry& MyGeometry, const FKeyEvent& KeyEvent )
 {
 	return ViewportInterface.IsValid() ? ViewportInterface.Pin()->OnKeyDown(MyGeometry, KeyEvent) : FReply::Unhandled();
 }
 
- 
 FReply SViewport::OnKeyUp( const FGeometry& MyGeometry, const FKeyEvent& KeyEvent )
 {
 	return ViewportInterface.IsValid() ? ViewportInterface.Pin()->OnKeyUp(MyGeometry, KeyEvent) : FReply::Unhandled();
@@ -238,7 +235,6 @@ FReply SViewport::OnFocusReceived( const FGeometry& MyGeometry, const FFocusEven
 	return ViewportInterface.IsValid() ? ViewportInterface.Pin()->OnFocusReceived(InFocusEvent) : FReply::Unhandled();
 }
 
-
 void SViewport::OnFocusLost( const FFocusEvent& InFocusEvent )
 {
 	if (ViewportInterface.IsValid())
@@ -273,31 +269,26 @@ void SViewport::OnWindowClosed( const TSharedRef<SWindow>& WindowBeingClosed )
 	}
 }
 
-
 FReply SViewport::OnTouchStarted( const FGeometry& MyGeometry, const FPointerEvent& InTouchEvent )
 {
 	return ViewportInterface.IsValid() ? ViewportInterface.Pin()->OnTouchStarted(MyGeometry, InTouchEvent) : FReply::Unhandled();
 }
 
-
 FReply SViewport::OnTouchMoved( const FGeometry& MyGeometry, const FPointerEvent& InTouchEvent )
 {
 	return ViewportInterface.IsValid() ? ViewportInterface.Pin()->OnTouchMoved(MyGeometry, InTouchEvent) : FReply::Unhandled();
 }
 
-
 FReply SViewport::OnTouchEnded( const FGeometry& MyGeometry, const FPointerEvent& InTouchEvent )
 {
 	return ViewportInterface.IsValid() ? ViewportInterface.Pin()->OnTouchEnded(MyGeometry, InTouchEvent) : FReply::Unhandled();
 }
 
-
 FReply SViewport::OnTouchGesture( const FGeometry& MyGeometry, const FPointerEvent& GestureEvent )
 {
 	return ViewportInterface.IsValid() ? ViewportInterface.Pin()->OnTouchGesture(MyGeometry, GestureEvent) : FReply::Unhandled();
 }
 
-
 FReply SViewport::OnMotionDetected( const FGeometry& MyGeometry, const FMotionEvent& InMotionEvent )
 {
 	return ViewportInterface.IsValid() ? ViewportInterface.Pin()->OnMotionDetected(MyGeometry, InMotionEvent) : FReply::Unhandled();
diff --git a/Engine/Source/Runtime/Slate/Public/Framework/Text/TextLayout.h b/Engine/Source/Runtime/Slate/Public/Framework/Text/TextLayout.h
--- a/Engine/Source/Runtime/Slate/Public/Framework/Text/TextLayout.h
+++ b/Engine/Source/Runtime/Slate/Public/Framework/Text/TextLayout.h
@@ -277,26 +277,26 @@ public:
 
 	virtual ~FTextLayout();
 
-	const TArray< FTextLayout::FLineView >& GetLineViews() const;
-	const TArray< FTextLayout::FLineModel >& GetLineModels() const;
+	FORCEINLINE const TArray< FTextLayout::FLineView >& GetLineViews() const { return LineViews; }
+	FORCEINLINE const TArray< FTextLayout::FLineModel >& GetLineModels() const { return LineModels; }
 
 	FVector2D GetSize() const;
 	FVector2D GetDrawSize() const;
 	FVector2D GetWrappedSize() const;
 
-	float GetWrappingWidth() const;
+	FORCEINLINE float GetWrappingWidth() const { return WrappingWidth; }
 	void SetWrappingWidth( float Value );
 
-	float GetLineHeightPercentage() const;
+	FORCEINLINE float GetLineHeightPercentage() const { return LineHeightPercentage; }
 	void SetLineHeightPercentage( float Value );
 
-	ETextJustify::Type GetJustification() const;
+	FORCEINLINE ETextJustify::Type GetJustification() const { return Justification; }
 	void SetJustification( ETextJustify::Type Value );
 
-	float GetScale() const;
+	FORCEINLINE float GetScale() const { return Scale; }
 	void SetScale( float Value );
 
-	FMargin GetMargin() const;
+	FORCEINLINE FMargin GetMargin() const { return Margin; }
 	void SetMargin( const FMargin& InMargin );
 
 	void SetVisibleRegion( const FVector2D& InViewSize, const FVector2D& InScrollOffset );
diff --git a/Engine/Source/Runtime/Slate/Public/Widgets/SInvalidationPanel.h b/Engine/Source/Runtime/Slate/Public/Widgets/SInvalidationPanel.h
--- a/Engine/Source/Runtime/Slate/Public/Widgets/SInvalidationPanel.h
+++ b/Engine/Source/Runtime/Slate/Public/Widgets/SInvalidationPanel.h
@@ -40,27 +40,35 @@ public:
 
 	void SetContent(const TSharedRef< SWidget >& InContent);
 
-#if !UE_BUILD_SHIPPING
+#if !(UE_BUILD_SHIPPING || UE_BUILD_TEST)
 	static bool IsInvalidationDebuggingEnabled();
 	static void EnableInvalidationDebugging(bool bEnable);
 #endif
 
+private:
+	TSharedPtr< FSlateWindowElementList > GetNextCachedElementList(const TSharedPtr<SWindow>& CurrentWindow) const;
+
 private:
 	FGeometry LastAllottedGeometry;
 
 	FSimpleSlot EmptyChildSlot;
 	FVector2D CachedDesiredSize;
 
-#if !UE_BUILD_SHIPPING
-	mutable TSet<TWeakPtr<SWidget>> InvalidatorWidgets;
+#if !(UE_BUILD_SHIPPING || UE_BUILD_TEST)
+	mutable TMap<TWeakPtr<SWidget>, double> InvalidatorWidgets;
 #endif
 
 	mutable FCachedWidgetNode* RootCacheNode;
+	mutable TArray< TSharedPtr< FSlateWindowElementList > > ActiveCachedElementListPool;
+	mutable TArray< TSharedPtr< FSlateWindowElementList > > InactiveCachedElementListPool;
 	mutable TSharedPtr< FSlateWindowElementList > CachedWindowElements;
+	mutable TSharedPtr<FSlateRenderDataHandle, ESPMode::ThreadSafe> CachedRenderData;
+	
 	mutable FVector2D CachedAbsolutePosition;
+	mutable FVector2D AbsoluteDeltaPosition;
+
 	mutable TArray< FCachedWidgetNode* > NodePool;
 	mutable int32 LastUsedCachedNodeIndex;
-	mutable int32 LastLayerId;
 	mutable int32 LastHitTestIndex;
 
 	mutable int32 CachedMaxChildLayer;
diff --git a/Engine/Source/Runtime/SlateCore/Private/Rendering/DrawElements.cpp b/Engine/Source/Runtime/SlateCore/Private/Rendering/DrawElements.cpp
--- a/Engine/Source/Runtime/SlateCore/Private/Rendering/DrawElements.cpp
+++ b/Engine/Source/Runtime/SlateCore/Private/Rendering/DrawElements.cpp
@@ -1,6 +1,17 @@
 // Copyright 1998-2015 Epic Games, Inc. All Rights Reserved.
 
 #include "SlateCorePrivatePCH.h"
+#include "ReflectionMetadata.h"
+
+#if !(UE_BUILD_SHIPPING || UE_BUILD_TEST)
+
+/**  */
+TAutoConsoleVariable<int32> DrawVolatileWidgets(
+	TEXT("Slate.DrawVolatileWidgets"),
+	true,
+	TEXT("Whether to draw volatile widgets"));
+
+#endif
 
 DECLARE_CYCLE_STAT(TEXT("FSlateDrawElement::Make Time"), STAT_SlateDrawElementMakeTime, STATGROUP_SlateVerbose);
 
@@ -159,6 +170,26 @@ void FSlateDrawElement::MakeCustom( FSlateWindowElementList& ElementList, uint32
 	DrawElt.DataPayload.SetCustomDrawerPayloadProperties( CustomDrawer );
 }
 
+void FSlateDrawElement::MakeCachedBuffer(FSlateWindowElementList& ElementList, uint32 InLayer, TSharedPtr<FSlateRenderDataHandle, ESPMode::ThreadSafe>& CachedRenderDataHandle)
+{
+	SCOPE_CYCLE_COUNTER(STAT_SlateDrawElementMakeTime)
+	FSlateDrawElement& DrawElt = ElementList.AddUninitialized();
+	DrawElt.Init(InLayer, FPaintGeometry(), FSlateRect(1, 1, 1, 1), ESlateDrawEffect::None);
+	DrawElt.RenderTransform = FSlateRenderTransform();
+	DrawElt.ElementType = ET_CachedBuffer;
+	DrawElt.DataPayload.SetCachedBufferPayloadProperties(CachedRenderDataHandle.Get());
+}
+
+void FSlateDrawElement::MakeLayer(FSlateWindowElementList& ElementList, uint32 InLayer, TSharedPtr<FSlateDrawLayerHandle, ESPMode::ThreadSafe>& DrawLayerHandle)
+{
+	SCOPE_CYCLE_COUNTER(STAT_SlateDrawElementMakeTime)
+	FSlateDrawElement& DrawElt = ElementList.AddUninitialized();
+	DrawElt.Init(InLayer, FPaintGeometry(), FSlateRect(1, 1, 1, 1), ESlateDrawEffect::None);
+	DrawElt.RenderTransform = FSlateRenderTransform();
+	DrawElt.ElementType = ET_Layer;
+	DrawElt.DataPayload.SetLayerPayloadProperties(DrawLayerHandle.Get());
+}
+
 FVector2D FSlateDrawElement::GetRotationPoint(const FPaintGeometry& PaintGeometry, const TOptional<FVector2D>& UserRotationPoint, ERotationSpace RotationSpace)
 {
 	FVector2D RotationPoint(0, 0);
@@ -197,6 +228,8 @@ void FSlateBatchData::Reset()
 	NumBatchedVertices = 0;
 	NumBatchedIndices = 0;
 	NumLayers = 0;
+
+	RenderDataHandle.Reset();
 }
 
 void FSlateBatchData::AssignVertexArrayToBatch( FSlateElementBatch& Batch )
@@ -242,6 +275,12 @@ void FSlateBatchData::FillVertexAndIndexBuffer( uint8* VertexBuffer, uint8* Inde
 	int32 VertexOffset = 0;
 	for( const FSlateRenderBatch& Batch : RenderBatches )
 	{
+		// Ignore foreign batches that are inserted into our render set.
+		if ( RenderDataHandle != Batch.CachedRenderHandle )
+		{
+			continue;
+		}
+
 		if( Batch.VertexArrayIndex != INDEX_NONE && Batch.IndexArrayIndex != INDEX_NONE )
 		{
 			TArray<FSlateVertex>& Vertices = BatchVertexArrays[Batch.VertexArrayIndex];
@@ -268,58 +307,142 @@ void FSlateBatchData::FillVertexAndIndexBuffer( uint8* VertexBuffer, uint8* Inde
 	}
 }
 
-void FSlateBatchData::CreateRenderBatches()
+void FSlateBatchData::UpdateRenderBatches(FVector2D PositionOffset)
 {
-	checkSlow( IsInRenderingThread() );
+	checkSlow(IsInRenderingThread());
+
+	for ( const FSlateRenderBatch& Batch : RenderBatches )
+	{
+		if ( Batch.VertexArrayIndex != INDEX_NONE && Batch.IndexArrayIndex != INDEX_NONE )
+		{
+			TArray<FSlateVertex>& Vertices = BatchVertexArrays[Batch.VertexArrayIndex];
+
+			int32 Count = Vertices.Num();
+			for ( int32 VertexIndex = 0; VertexIndex < Count; VertexIndex++ )
+			{
+				FSlateVertex& Vertex = Vertices[VertexIndex];
+				Vertex.Position[0] += PositionOffset.X;
+				Vertex.Position[1] += PositionOffset.Y;
+				Vertex.ClipRect.TopLeft += PositionOffset;
+			}
+		}
+	}
+}
 
-	LayerToElementBatches.KeySort( TLess<uint32>() );
+void FSlateBatchData::CreateRenderBatches(FElementBatchMap& LayerToElementBatches)
+{
+	checkSlow( IsInRenderingThread() );
 
 	uint32 VertexOffset = 0;
 	uint32 IndexOffset = 0;
+
+	Merge(LayerToElementBatches, VertexOffset, IndexOffset);
+
+	// 
+	if ( RenderDataHandle.IsValid() )
+	{
+		RenderDataHandle->SetRenderBatches(&RenderBatches);
+	}
+}
+
+void FSlateBatchData::Merge(FElementBatchMap& InLayerToElementBatches, uint32& VertexOffset, uint32& IndexOffset)
+{
+	InLayerToElementBatches.KeySort(TLess<uint32>());
+
+	const bool bExpandLayersAndCachedHandles = RenderDataHandle.IsValid() == false;
+
 	// For each element batch add its vertices and indices to the bulk lists.
-	for (FElementBatchMap::TIterator It(LayerToElementBatches); It; ++It)
+	for ( FElementBatchMap::TIterator It(InLayerToElementBatches); It; ++It )
 	{
+		uint32 Layer = It.Key();
 		FElementBatchArray& ElementBatches = It.Value();
 
-		if( ElementBatches.Num() > 0 )
+		if ( ElementBatches.Num() > 0 )
 		{
 			++NumLayers;
-			for (FElementBatchArray::TIterator BatchIt(ElementBatches); BatchIt; ++BatchIt)
+			for ( FElementBatchArray::TIterator BatchIt(ElementBatches); BatchIt; ++BatchIt )
 			{
 				FSlateElementBatch& ElementBatch = *BatchIt;
 
-				bool bIsMaterial = ElementBatch.GetShaderResource() && ElementBatch.GetShaderResource()->GetType() == ESlateShaderResource::Material;
-
-				if (!ElementBatch.GetCustomDrawer().IsValid())
+				if ( ElementBatch.GetCustomDrawer().IsValid() )
 				{
-					TArray<FSlateVertex>& BatchVertices = GetBatchVertexList(ElementBatch);
-					TArray<SlateIndex>& BatchIndices = GetBatchIndexList(ElementBatch);
-	
-					// We should have at least some vertices and indices in the batch or none at all
-					check(BatchVertices.Num() > 0 && BatchIndices.Num() > 0 || BatchVertices.Num() == 0 && BatchIndices.Num() == 0);
-
-					if (BatchVertices.Num() > 0 && BatchIndices.Num() > 0)
-					{
-						int32 NumVertices = BatchVertices.Num();
-						int32 NumIndices = BatchIndices.Num();
-
-						AddRenderBatch( ElementBatch, NumVertices, NumIndices, VertexOffset, IndexOffset );
-				
-						VertexOffset += BatchVertices.Num();
-						IndexOffset += BatchIndices.Num();
-					}
+					AddRenderBatch(Layer, ElementBatch, 0, 0, 0, 0);
 				}
 				else
 				{
-					AddRenderBatch( ElementBatch, 0, 0, 0, 0 );
+					if ( bExpandLayersAndCachedHandles )
+					{
+						if ( FSlateRenderDataHandle* RenderHandle = ElementBatch.GetCachedRenderHandle().Get() )
+						{
+							if ( TArray<FSlateRenderBatch>* ForeignBatches = RenderHandle->GetRenderBatches() )
+							{
+								TArray<FSlateRenderBatch>& ForeignBatchesRef = *ForeignBatches;
+								for ( int32 i = 0; i < ForeignBatches->Num(); i++ )
+								{
+									TSharedPtr<FSlateDrawLayerHandle, ESPMode::ThreadSafe> LayerHandle = ForeignBatchesRef[i].LayerHandle.Pin();
+									if ( LayerHandle.IsValid() )
+									{
+										// If a record was added for a layer, but nothing was ever drawn for it, the batch map will be null.
+										if ( LayerHandle->BatchMap )
+										{
+											Merge(*LayerHandle->BatchMap, VertexOffset, IndexOffset);
+											LayerHandle->BatchMap = nullptr;
+										}
+									}
+									else
+									{
+										RenderBatches.Add(ForeignBatchesRef[i]);
+									}
+									//RenderBatches.Append(*RenderHandle->GetRenderBatches());
+								}
+							}
+
+							continue;
+						}
+					}
+					else
+					{
+						// Insert if we're not expanding
+						if ( FSlateDrawLayerHandle* LayerHandle = ElementBatch.GetLayerHandle().Get() )
+						{
+							AddRenderBatch(Layer, ElementBatch, 0, 0, 0, 0);
+							continue;
+						}
+					}
+					
+					if ( ElementBatch.VertexArrayIndex != INDEX_NONE && ElementBatch.IndexArrayIndex != INDEX_NONE )
+					{
+						TArray<FSlateVertex>& BatchVertices = GetBatchVertexList(ElementBatch);
+						TArray<SlateIndex>& BatchIndices = GetBatchIndexList(ElementBatch);
+
+						// We should have at least some vertices and indices in the batch or none at all
+						check(BatchVertices.Num() > 0 && BatchIndices.Num() > 0 || BatchVertices.Num() == 0 && BatchIndices.Num() == 0);
+
+						if ( BatchVertices.Num() > 0 && BatchIndices.Num() > 0 )
+						{
+							const int32 NumVertices = BatchVertices.Num();
+							const int32 NumIndices = BatchIndices.Num();
+
+							AddRenderBatch(Layer, ElementBatch, NumVertices, NumIndices, VertexOffset, IndexOffset);
+
+							VertexOffset += BatchVertices.Num();
+							IndexOffset += BatchIndices.Num();
+						}
+					}
 				}
 			}
 		}
+
 		ElementBatches.Reset();
 	}
+}
 
+void FSlateBatchData::SortRenderBatches()
+{
+	//RenderBatches.StableSort();
 }
 
+
 FSlateWindowElementList::FDeferredPaint::FDeferredPaint( const TSharedRef<const SWidget>& InWidgetToPaint, const FPaintArgs& InArgs, const FGeometry InAllottedGeometry, const FSlateRect InMyClippingRect, const FWidgetStyle& InWidgetStyle, bool InParentEnabled )
 : WidgetToPaintPtr( InWidgetToPaint )
 , Args( InArgs )
@@ -345,16 +468,14 @@ int32 FSlateWindowElementList::FDeferredPaint::ExecutePaint( int32 LayerId, FSla
 
 void FSlateWindowElementList::QueueDeferredPainting( const FDeferredPaint& InDeferredPaint )
 {
-	DeferredPaintList.Add( MakeShareable( new FDeferredPaint( InDeferredPaint ) ) );
+	DeferredPaintList.Add(MakeShareable(new FDeferredPaint(InDeferredPaint)));
 }
 
 int32 FSlateWindowElementList::PaintDeferred(int32 LayerId)
 {
 	for ( int32 i = 0; i < DeferredPaintList.Num(); ++i )
 	{
-		const TSharedRef< FDeferredPaint >& Args = DeferredPaintList[i];
-
-		LayerId = Args->ExecutePaint(LayerId, *this);
+		LayerId = DeferredPaintList[i]->ExecutePaint(LayerId, *this);
 	}
 
 	return LayerId;
@@ -377,40 +498,202 @@ int32 FSlateWindowElementList::FVolatilePaint::ExecutePaint(FSlateWindowElementL
 	TSharedPtr<const SWidget> WidgetToPaint = WidgetToPaintPtr.Pin();
 	if ( WidgetToPaint.IsValid() )
 	{
+		//FPlatformMisc::BeginNamedEvent(FColor::Orange, "Overhead");
+		//FPlatformMisc::BeginNamedEvent(FColor::Orange, *FReflectionMetaData::GetWidgetDebugInfo(WidgetToPaint));
+
 		// Have to run a slate pre-pass for all volatile elements, some widgets cache information like 
 		// the STextBlock.  This may be all kinds of terrible an idea to do during paint.
 		SWidget* MutableWidget = const_cast<SWidget*>( WidgetToPaint.Get() );
 		MutableWidget->SlatePrepass(AllottedGeometry.Scale);
 
-		return WidgetToPaint->Paint(Args, AllottedGeometry, MyClippingRect, OutDrawElements, LayerId, WidgetStyle, bParentEnabled);
+		const int32 NewLayer = WidgetToPaint->Paint(Args, AllottedGeometry, MyClippingRect, OutDrawElements, LayerId, WidgetStyle, bParentEnabled);
+		
+		//FPlatformMisc::EndNamedEvent();
+		//FPlatformMisc::EndNamedEvent();
+
+		return NewLayer;
 	}
 
 	return LayerId;
 }
 
-void FSlateWindowElementList::QueueVolatilePainting(const FVolatilePaint& InDeferredPaint)
+void FSlateWindowElementList::QueueVolatilePainting(const FVolatilePaint& InVolatilePaint)
 {
-	VolatilePaintList.Add(MakeShareable(new FVolatilePaint(InDeferredPaint)));
+	TSharedPtr< FSlateDrawLayerHandle, ESPMode::ThreadSafe > LayerHandle = MakeShareable(new FSlateDrawLayerHandle());
+
+	FSlateDrawElement::MakeLayer(*this, InVolatilePaint.GetLayerId(), LayerHandle);
+
+	const int32 NewEntryIndex = VolatilePaintList.Add(MakeShareable(new FVolatilePaint(InVolatilePaint)));
+	VolatilePaintList[NewEntryIndex]->LayerHandle = LayerHandle;
 }
 
 int32 FSlateWindowElementList::PaintVolatile(FSlateWindowElementList& OutElementList)
 {
+#if !(UE_BUILD_SHIPPING || UE_BUILD_TEST)
+	const bool bDrawVolatile = DrawVolatileWidgets.GetValueOnGameThread() == 1;
+#else
+	const bool bDrawVolatile = true;
+#endif
+
+	const static FName InvalidationPanelName(TEXT("SInvalidationPanelName"));
+
 	int32 MaxLayerId = 0;
 
 	for ( int32 VolatileIndex = 0; VolatileIndex < VolatilePaintList.Num(); ++VolatileIndex )
 	{
-		const TSharedRef< FVolatilePaint >& Args = VolatilePaintList[VolatileIndex];
+		const TSharedPtr<FVolatilePaint>& Args = VolatilePaintList[VolatileIndex];
 
+		// If we're not drawing volatile elements, just draw invalidation panels.  They're
+		// the exception, as we're likely trying to determine the volatile widget overhead.
+		if ( bDrawVolatile == false )
+		{
+			if ( const SWidget* Widget = Args->GetWidget() )
+			{
+				if ( Widget->GetType() != InvalidationPanelName )
+				{
+					continue;
+				}
+			}
+		}
+
+		OutElementList.BeginLogicalLayer(Args->LayerHandle);
 		MaxLayerId = FMath::Max(MaxLayerId, Args->ExecutePaint(OutElementList));
+		OutElementList.EndLogicalLayer();
 	}
 
 	return MaxLayerId;
 }
 
-void FSlateWindowElementList::Reset()
+void FSlateWindowElementList::BeginLogicalLayer(const TSharedPtr<FSlateDrawLayerHandle, ESPMode::ThreadSafe>& LayerHandle)
 {
-	BatchData.Reset();
-	DrawElements.Reset();
+	// Don't attempt to begin logical layers inside a cached view of the data.
+	checkSlow(!IsCachedRenderDataInUse());
+
+	//FPlatformMisc::BeginNamedEvent(FColor::Orange, "FindLayer");
+	TSharedPtr<FSlateDrawLayer> Layer = DrawLayers.FindRef(LayerHandle);
+	//FPlatformMisc::EndNamedEvent();
+
+	if ( !Layer.IsValid() )
+	{
+		if ( DrawLayerPool.Num() > 0 )
+		{
+			Layer = DrawLayerPool.Pop(false);
+		}
+		else
+		{
+			Layer = MakeShareable(new FSlateDrawLayer());
+		}
+
+		//FPlatformMisc::BeginNamedEvent(FColor::Orange, "AddLayer");
+		DrawLayers.Add(LayerHandle, Layer);
+		//FPlatformMisc::EndNamedEvent();
+	}
+
+	//FPlatformMisc::BeginNamedEvent(FColor::Orange, "PushLayer");
+	DrawStack.Push(Layer.Get());
+	//FPlatformMisc::EndNamedEvent();
+}
+
+void FSlateWindowElementList::EndLogicalLayer()
+{
+	DrawStack.Pop();
+}
+
+FSlateRenderDataHandle::FSlateRenderDataHandle(FSlateRenderer* InRenderer)
+	: Renderer(InRenderer)
+	, RenderBatches(nullptr)
+{
+}
+
+FSlateRenderDataHandle::~FSlateRenderDataHandle()
+{
+	if ( Renderer )
+	{
+		Renderer->ReleaseCachedRenderData(this);
+	}
+}
+
+void FSlateRenderDataHandle::Disconnect()
+{
+	Renderer = nullptr;
+	RenderBatches = nullptr;
+}
+
+TSharedRef<FSlateRenderDataHandle, ESPMode::ThreadSafe> FSlateWindowElementList::CacheRenderData()
+{
+	// Don't attempt to use this slate window element list if the cache is still being used.
+	checkSlow(!IsCachedRenderDataInUse());
+
+	TSharedPtr<FSlateRenderer> Renderer = FSlateApplicationBase::Get().GetRenderer();
+
+	TSharedRef<FSlateRenderDataHandle, ESPMode::ThreadSafe> CachedRenderDataHandleRef = Renderer->CacheElementRenderData(*this);
+	CachedRenderDataHandle = CachedRenderDataHandleRef;
+
+	return CachedRenderDataHandleRef;
+}
+
+void FSlateWindowElementList::UpdateCacheRenderData(FVector2D PositionOffset)
+{
+	// This should only be done if we're using cached render data
+	checkSlow(IsCachedRenderDataInUse());
+
+	TSharedPtr<FSlateRenderer> Renderer = FSlateApplicationBase::Get().GetRenderer();
+	Renderer->UpdateElementRenderData(*this, PositionOffset);
+}
+
+void FSlateWindowElementList::PreDraw_RenderThread()
+{
+	check(IsInRenderingThread());
+
+	for ( auto& Entry : DrawLayers )
+	{
+		checkSlow(Entry.Key->BatchMap == nullptr);
+		Entry.Key->BatchMap = &Entry.Value->GetElementBatchMap();
+	}
+}
+
+void FSlateWindowElementList::PostDraw_RenderThread()
+{
+	check(IsInRenderingThread());
+
+	for ( auto& Entry : DrawLayers )
+	{
+		Entry.Key->BatchMap = nullptr;
+	}
+}
+
+void FSlateWindowElementList::ResetBuffers()
+{
+	// Don't attempt to use this slate window element list if the cache is still being used.
+	checkSlow(!IsCachedRenderDataInUse());
+
 	DeferredPaintList.Reset();
 	VolatilePaintList.Reset();
+	BatchData.Reset();
+
+#if SLATE_POOL_DRAW_ELEMENTS
+	DrawElementFreePool.Append(RootDrawLayer.DrawElements);
+#endif
+
+	// Reset the draw elements on the root draw layer
+	RootDrawLayer.DrawElements.Reset();
+
+	// Return child draw layers to the pool, and reset their draw elements.
+	for ( auto& Entry : DrawLayers )
+	{
+#if SLATE_POOL_DRAW_ELEMENTS
+		TArray<FSlateDrawElement*>& DrawElements = Entry.Value->DrawElements;
+		DrawElementFreePool.Append(DrawElements);
+#else
+		TArray<FSlateDrawElement>& DrawElements = Entry.Value->DrawElements;
+#endif
+
+		DrawElements.Reset();
+		DrawLayerPool.Add(Entry.Value);
+	}
+
+	DrawLayers.Reset();
+
+	DrawStack.Reset();
+	DrawStack.Push(&RootDrawLayer);
 }
diff --git a/Engine/Source/Runtime/SlateCore/Private/Rendering/ElementBatcher.cpp b/Engine/Source/Runtime/SlateCore/Private/Rendering/ElementBatcher.cpp
--- a/Engine/Source/Runtime/SlateCore/Private/Rendering/ElementBatcher.cpp
+++ b/Engine/Source/Runtime/SlateCore/Private/Rendering/ElementBatcher.cpp
@@ -7,7 +7,7 @@
 #include "SlateStats.h"
 
 DECLARE_CYCLE_STAT(TEXT("Find Batch For Element Time"), STAT_SlateFindBatchForElement, STATGROUP_SlateVerbose);
-DECLARE_DWORD_COUNTER_STAT(TEXT("Num Elements (Prebatch)"), STAT_SlateNumPrebatchElements, STATGROUP_SlateVerbose);
+DECLARE_DWORD_COUNTER_STAT(TEXT("Num Elements (Prebatch)"), STAT_SlateNumPrebatchElements, STATGROUP_Slate);
 
 DECLARE_CYCLE_STAT(TEXT("Add Elements Time"), STAT_SlateAddElements, STATGROUP_Slate);
 
@@ -50,25 +50,48 @@ FSlateRotatedClipRectType ToSnappedRotatedRect(const FSlateRect& ClipRectInLayou
 
 FSlateElementBatcher::FSlateElementBatcher( TSharedRef<FSlateRenderingPolicy> InRenderingPolicy )
 	: BatchData( nullptr )
+	, DrawLayer( nullptr )
 	, ResourceManager( *InRenderingPolicy->GetResourceManager() )
 	, FontCache( *InRenderingPolicy->GetFontCache() )
+	, NumDrawnBatchesStat(0)
 	, PixelCenterOffset( InRenderingPolicy->GetPixelCenterOffset() )
 	, bSRGBVertexColor( !InRenderingPolicy->IsVertexColorInLinearSpace() )
 {
-
 }
 
-
 FSlateElementBatcher::~FSlateElementBatcher()
 {
 }
 
-
-void FSlateElementBatcher::AddElements( FSlateWindowElementList& WindowElementList )
+void FSlateElementBatcher::AddElements(FSlateWindowElementList& WindowElementList)
 {
 	SLATE_CYCLE_COUNTER_SCOPE(GSlateAddElements);
+	FPlatformMisc::BeginNamedEvent(FColor::Magenta, "Slate::AddElements");
+
+	SCOPE_CYCLE_COUNTER(STAT_SlateAddElements);
+
+	NumDrawnBatchesStat = 0;
+
+	BatchData = &WindowElementList.GetBatchData();
+
+	AddElements(WindowElementList.GetRootDrawLayer());
+
+	TMap < TSharedPtr<FSlateDrawLayerHandle, ESPMode::ThreadSafe>, TSharedPtr<FSlateDrawLayer> >& DrawLayers = WindowElementList.GetChildDrawLayers();
+	for ( auto& Entry : DrawLayers )
+	{
+		AddElements(*Entry.Value.Get());
+	}
+
+	// Done with the element list
+	BatchData = nullptr;
+
+	SET_DWORD_STAT(STAT_SlateNumPrebatchElements, NumDrawnBatchesStat);
+
+	FPlatformMisc::EndNamedEvent();
+}
 
-	SCOPE_CYCLE_COUNTER( STAT_SlateAddElements );
+void FSlateElementBatcher::AddElements(FSlateDrawLayer& InDrawLayer)
+{
 	// This stuff is just for the counters. Could be scoped by an #ifdef if necessary.
 	static_assert(
 		FSlateDrawElement::EElementType::ET_Box == 0 &&
@@ -80,10 +103,12 @@ void FSlateElementBatcher::AddElements( FSlateWindowElementList& WindowElementLi
 		FSlateDrawElement::EElementType::ET_Viewport == 6 &&
 		FSlateDrawElement::EElementType::ET_Border == 7 &&
 		FSlateDrawElement::EElementType::ET_Custom == 8 &&
-		FSlateDrawElement::EElementType::ET_Count == 9, 
-		"If FSlateDrawElement::EElementType is modified, this array must be made to match.");
+		FSlateDrawElement::EElementType::ET_CachedBuffer == 9 &&
+		FSlateDrawElement::EElementType::ET_Layer == 10 &&
+		FSlateDrawElement::EElementType::ET_Count == 11,
+		"If FSlateDrawElement::EElementType is modified, this array must be made to match." );
 
-	static FName ElementFNames[] = 
+	static FName ElementFNames[] =
 	{
 		FName(TEXT("Box")),
 		FName(TEXT("DebugQuad")),
@@ -94,17 +119,33 @@ void FSlateElementBatcher::AddElements( FSlateWindowElementList& WindowElementLi
 		FName(TEXT("Viewport")),
 		FName(TEXT("Border")),
 		FName(TEXT("Custom")),
+		FName(TEXT("CachedBuffer")),
+		FName(TEXT("Layer")),
 	};
 
-	BatchData = &WindowElementList.GetBatchData();
-
-	const TArray<FSlateDrawElement>& DrawElements = WindowElementList.GetDrawElements();
+	DrawLayer = &InDrawLayer;
 
-	int32 NumDrawnBatches = 0;
+#if SLATE_POOL_DRAW_ELEMENTS
+	const TArray<FSlateDrawElement*>& DrawElements = InDrawLayer.DrawElements;
+#else
+	const TArray<FSlateDrawElement>& DrawElements = InDrawLayer.DrawElements;
+#endif
 
 	for( int32 DrawElementIndex = 0; DrawElementIndex < DrawElements.Num(); ++DrawElementIndex )
 	{
+#if SLATE_POOL_DRAW_ELEMENTS
+		if ( DrawElementIndex < ( DrawElements.Num() - 1 ) )
+		{
+			FPlatformMisc::Prefetch(DrawElements[DrawElementIndex + 1]);
+		}
+#endif
+
+#if SLATE_POOL_DRAW_ELEMENTS
+		const FSlateDrawElement& DrawElement = *DrawElements[DrawElementIndex];
+#else
 		const FSlateDrawElement& DrawElement = DrawElements[DrawElementIndex];
+#endif
+
 		const FSlateRect& InClippingRect = DrawElement.GetClippingRect();
 	
 		// A zero or negatively sized clipping rect means the geometry will not be displayed
@@ -112,12 +153,24 @@ void FSlateElementBatcher::AddElements( FSlateWindowElementList& WindowElementLi
 		
 		if ( !bIsFullyClipped )
 		{
+			bool bIsScissored = false;
+
 			// do this check in here do we can short circuit above more quickly, but still name this variable so its meaning is clear.
-			const bool bIsScissored = DrawElement.GetScissorRect().IsSet() && DrawElement.GetElementType() != FSlateDrawElement::ET_Custom && !DrawElement.GetScissorRect().GetValue().DoesIntersect(FShortRect(InClippingRect));
+			switch ( DrawElement.GetElementType() )
+			{
+			case FSlateDrawElement::ET_Custom:
+			case FSlateDrawElement::ET_CachedBuffer:
+			case FSlateDrawElement::ET_Layer:
+				break;
+			default:
+				bIsScissored = DrawElement.GetScissorRect().IsSet() && !DrawElement.GetScissorRect().GetValue().DoesIntersect(FShortRect(InClippingRect));
+				break;
+			}
+
 			// scissor rects are sort of a low level hack, so no one konws to clip against them. Instead we do it here to make sure the element is not actually rendered.
 			if (!bIsScissored)
 			{
-				++NumDrawnBatches;
+				++NumDrawnBatchesStat;
 
 				// time just the adding of the element. The clipping stuff will be counted in exclusive time for the non-typed timer.
 				SLATE_CYCLE_COUNTER_SCOPE_CUSTOM_DETAILED(SLATE_STATS_DETAIL_LEVEL_MED, GSlateAddElements, ElementFNames[DrawElement.GetElementType()]);
@@ -151,6 +204,12 @@ void FSlateElementBatcher::AddElements( FSlateWindowElementList& WindowElementLi
 				case FSlateDrawElement::ET_Custom:
 					AddCustomElement( DrawElement );
 					break;
+				case FSlateDrawElement::ET_CachedBuffer:
+					AddCachedBuffer( DrawElement );
+					break;
+				case FSlateDrawElement::ET_Layer:
+					AddLayer(DrawElement);
+					break;
 				default:
 					checkf(0, TEXT("Invalid element type"));
 					break;
@@ -158,11 +217,6 @@ void FSlateElementBatcher::AddElements( FSlateWindowElementList& WindowElementLi
 			}
 		}
 	}
-
-	SET_DWORD_STAT( STAT_SlateNumPrebatchElements, NumDrawnBatches );
-
-	// Done with the element list
-	BatchData = nullptr;
 }
 
 void FSlateElementBatcher::AddQuadElement( const FSlateDrawElement& DrawElement, FColor Color )
@@ -1463,7 +1517,7 @@ void FSlateElementBatcher::AddBorderElement( const FSlateDrawElement& DrawElemen
 
 void FSlateElementBatcher::AddCustomElement( const FSlateDrawElement& DrawElement )
 {
-	FElementBatchMap& LayerToElementBatches = BatchData->GetElementBatchMap();
+	FElementBatchMap& LayerToElementBatches = DrawLayer->GetElementBatchMap();
 
 	const FSlateDataPayload& InPayload = DrawElement.GetDataPayload();
 	uint32 Layer = DrawElement.GetLayer();
@@ -1479,12 +1533,59 @@ void FSlateElementBatcher::AddCustomElement( const FSlateDrawElement& DrawElemen
 		}
 		check( ElementBatches );
 
-
 		// Custom elements are not batched together 
 		ElementBatches->Add( FSlateElementBatch( InPayload.CustomDrawer, DrawElement.GetScissorRect() ) );
 	}
 }
 
+void FSlateElementBatcher::AddCachedBuffer(const FSlateDrawElement& DrawElement)
+{
+	FElementBatchMap& LayerToElementBatches = DrawLayer->GetElementBatchMap();
+
+	const FSlateDataPayload& InPayload = DrawElement.GetDataPayload();
+	uint32 Layer = DrawElement.GetLayer();
+
+	if ( InPayload.CachedRenderData )
+	{
+		// See if the layer already exists.
+		FElementBatchArray* ElementBatches = LayerToElementBatches.Find(Layer);
+		if ( !ElementBatches )
+		{
+			// The layer doesn't exist so make it now
+			ElementBatches = &LayerToElementBatches.Add(Layer, FElementBatchArray());
+		}
+		check(ElementBatches);
+
+		// Custom elements are not batched together
+		TSharedPtr< FSlateRenderDataHandle, ESPMode::ThreadSafe > RenderData = InPayload.CachedRenderData->AsShared();
+		ElementBatches->Add(FSlateElementBatch(RenderData, DrawElement.GetScissorRect()));
+	}
+}
+
+void FSlateElementBatcher::AddLayer(const FSlateDrawElement& DrawElement)
+{
+	FElementBatchMap& LayerToElementBatches = DrawLayer->GetElementBatchMap();
+
+	const FSlateDataPayload& InPayload = DrawElement.GetDataPayload();
+	uint32 Layer = DrawElement.GetLayer();
+
+	if ( InPayload.LayerHandle )
+	{
+		// See if the layer already exists.
+		FElementBatchArray* ElementBatches = LayerToElementBatches.Find(Layer);
+		if ( !ElementBatches )
+		{
+			// The layer doesn't exist so make it now
+			ElementBatches = &LayerToElementBatches.Add(Layer, FElementBatchArray());
+		}
+		check(ElementBatches);
+
+		// Custom elements are not batched together
+		TSharedPtr< FSlateDrawLayerHandle, ESPMode::ThreadSafe > LayerHandle = InPayload.LayerHandle->AsShared();
+		ElementBatches->Add(FSlateElementBatch(LayerHandle, DrawElement.GetScissorRect()));
+	}
+}
+
 
 FSlateElementBatch& FSlateElementBatcher::FindBatchForElement( 
 	uint32 Layer, 
@@ -1499,7 +1600,7 @@ FSlateElementBatch& FSlateElementBatcher::FindBatchForElement(
 	SLATE_CYCLE_COUNTER_SCOPE_DETAILED(SLATE_STATS_DETAIL_LEVEL_HI, GSlateFindBatchTime);
 
 //	SCOPE_CYCLE_COUNTER( STAT_SlateFindBatchForElement );
-	FElementBatchMap& LayerToElementBatches = BatchData->GetElementBatchMap();
+	FElementBatchMap& LayerToElementBatches = DrawLayer->GetElementBatchMap();
 
 	// See if the layer already exists.
 	FElementBatchArray* ElementBatches = LayerToElementBatches.Find( Layer );
@@ -1521,8 +1622,8 @@ FSlateElementBatch& FSlateElementBatcher::FindBatchForElement(
 		int32 Index = ElementBatches->Add( TempBatch );
 		ElementBatch = &(*ElementBatches)[Index];
 
-		BatchData->AssignVertexArrayToBatch( *ElementBatch );
-		BatchData->AssignIndexArrayToBatch( *ElementBatch );
+		BatchData->AssignVertexArrayToBatch(*ElementBatch);
+		BatchData->AssignIndexArrayToBatch(*ElementBatch);
 	}
 	check( ElementBatch );
 
diff --git a/Engine/Source/Runtime/SlateCore/Private/Rendering/SlateDrawBuffer.cpp b/Engine/Source/Runtime/SlateCore/Private/Rendering/SlateDrawBuffer.cpp
--- a/Engine/Source/Runtime/SlateCore/Private/Rendering/SlateDrawBuffer.cpp
+++ b/Engine/Source/Runtime/SlateCore/Private/Rendering/SlateDrawBuffer.cpp
@@ -19,7 +19,7 @@ FSlateWindowElementList& FSlateDrawBuffer::AddWindowElementList(TSharedRef<SWind
 			WindowElementLists.Add(WindowElements);
 			WindowElementListsPool.RemoveAtSwap(WindowIndex);
 
-			WindowElements->Reset();
+			WindowElements->ResetBuffers();
 
 			return *WindowElements;
 		}
diff --git a/Engine/Source/Runtime/SlateCore/Private/Rendering/SlateRenderer.cpp b/Engine/Source/Runtime/SlateCore/Private/Rendering/SlateRenderer.cpp
--- a/Engine/Source/Runtime/SlateCore/Private/Rendering/SlateRenderer.cpp
+++ b/Engine/Source/Runtime/SlateCore/Private/Rendering/SlateRenderer.cpp
@@ -55,6 +55,15 @@ ISlateAtlasProvider* FSlateRenderer::GetFontAtlasProvider()
 	return nullptr;
 }
 
+TSharedRef<FSlateRenderDataHandle, ESPMode::ThreadSafe> FSlateRenderer::CacheElementRenderData(FSlateWindowElementList& ElementList)
+{
+	return MakeShareable(new FSlateRenderDataHandle(this));
+}
+
+void FSlateRenderer::UpdateElementRenderData(FSlateWindowElementList& ElementList, FVector2D PositionOffset)
+{
+	// Do Nothing
+}
 
 /* Global functions
  *****************************************************************************/
diff --git a/Engine/Source/Runtime/SlateCore/Private/Widgets/SWindow.cpp b/Engine/Source/Runtime/SlateCore/Private/Widgets/SWindow.cpp
--- a/Engine/Source/Runtime/SlateCore/Private/Widgets/SWindow.cpp
+++ b/Engine/Source/Runtime/SlateCore/Private/Widgets/SWindow.cpp
@@ -384,7 +384,7 @@ void SWindow::ConstructWindowInternals()
 	// Setup widget that represents the main area of the window.  That is, everything inside the window's border.
 	TSharedRef< SVerticalBox > MainWindowArea = 
 		SNew( SVerticalBox )
-			.Visibility( EVisibility::SelfHitTestInvisible );
+		.Visibility( EVisibility::SelfHitTestInvisible );
 
 	if (bCreateTitleBar)
 	{
@@ -408,10 +408,10 @@ void SWindow::ConstructWindowInternals()
 		}
 
 		MainWindowArea->AddSlot()
-			.AutoHeight()
-			[
-				FSlateApplicationBase::Get().MakeWindowTitleBar(SharedThis(this), nullptr, TitleContentAlignment, TitleBar)
-			];
+		.AutoHeight()
+		[
+			FSlateApplicationBase::Get().MakeWindowTitleBar(SharedThis(this), nullptr, TitleContentAlignment, TitleBar)
+		];
 	}
 	else
 	{
@@ -437,57 +437,56 @@ void SWindow::ConstructWindowInternals()
 		this->ChildSlot
 		[
 			SAssignNew(WindowOverlay, SOverlay)
-				.Visibility(EVisibility::SelfHitTestInvisible)
+			.Visibility(EVisibility::SelfHitTestInvisible)
 
 			// window background
 			+ SOverlay::Slot()
-				[
-					FSlateApplicationBase::Get().MakeImage(
-						WindowBackgroundAttr,
-						FLinearColor::White,
-						WindowContentVisibility
-					)
-				]
+			[
+				FSlateApplicationBase::Get().MakeImage(
+					WindowBackgroundAttr,
+					FLinearColor::White,
+					WindowContentVisibility
+				)
+			]
 
 			// window border
 			+ SOverlay::Slot()
-				[
-					FSlateApplicationBase::Get().MakeImage(
-						&Style->BorderBrush,
-						FLinearColor::White,
-						WindowContentVisibility
-					)
-				]
+			[
+				FSlateApplicationBase::Get().MakeImage(
+					&Style->BorderBrush,
+					FLinearColor::White,
+					WindowContentVisibility
+				)
+			]
 
 			// main area
 			+ SOverlay::Slot()
+			[
+				SNew(SVerticalBox)
+				.Visibility(WindowContentVisibility)
+
+				+ SVerticalBox::Slot()					
+				.Padding(TAttribute<FMargin>::Create(TAttribute<FMargin>::FGetter::CreateSP(this, &SWindow::GetWindowBorderSize, false)))
 				[
-					SNew(SVerticalBox)
-						.Visibility(WindowContentVisibility)
-
-					+ SVerticalBox::Slot()
-					
-						.Padding(TAttribute<FMargin>::Create(TAttribute<FMargin>::FGetter::CreateSP(this, &SWindow::GetWindowBorderSize, false)))
-						[
-							MainWindowArea
-						]
+					MainWindowArea
 				]
+			]
 
 			// pop-up layer
 			+ SOverlay::Slot()
-				[
-					SAssignNew(PopupLayer, SPopupLayer, SharedThis(this))
-				]
+			[
+				SAssignNew(PopupLayer, SPopupLayer, SharedThis(this))
+			]
 
 			// window outline
 			+ SOverlay::Slot()
-				[
-					FSlateApplicationBase::Get().MakeImage(
-						WindowOutlineAttr,
-						WindowOutlineColorAttr,
-						WindowContentVisibility
-					)
-				]
+			[
+				FSlateApplicationBase::Get().MakeImage(
+					WindowOutlineAttr,
+					WindowOutlineColorAttr,
+					WindowContentVisibility
+				)
+			]
 		];
 	}
 	else if( bHasOSWindowBorder )
@@ -497,7 +496,7 @@ void SWindow::ConstructWindowInternals()
 			SAssignNew(WindowOverlay, SOverlay)
 			+ SOverlay::Slot()
 			[
-					MainWindowArea
+				MainWindowArea
 			]
 			+ SOverlay::Slot()
 			[
diff --git a/Engine/Source/Runtime/SlateCore/Public/Rendering/DrawElements.h b/Engine/Source/Runtime/SlateCore/Public/Rendering/DrawElements.h
--- a/Engine/Source/Runtime/SlateCore/Public/Rendering/DrawElements.h
+++ b/Engine/Source/Runtime/SlateCore/Public/Rendering/DrawElements.h
@@ -23,6 +23,9 @@ struct FSlateGradientStop
 template <> struct TIsPODType<FSlateGradientStop> { enum { Value = true }; };
 
 
+class FSlateDrawLayerHandle;
+
+
 class FSlateDataPayload
 {
 public:
@@ -66,6 +69,12 @@ public:
 	// Custom drawer data
 	TWeakPtr<ICustomSlateElement, ESPMode::ThreadSafe> CustomDrawer;
 
+	// Cached render data
+	class FSlateRenderDataHandle* CachedRenderData;
+
+	// Layer handle
+	FSlateDrawLayerHandle* LayerHandle;
+
 	// Line data
 	ESlateLineJoinType::Type SegmentJoinType;
 	bool bAntialias;
@@ -75,6 +84,8 @@ public:
 		, BrushResource(nullptr)
 		, RotationPoint(FVector2D::ZeroVector)
 		, Viewport(nullptr)
+		, CachedRenderData(nullptr)
+		, LayerHandle(nullptr)
 	{ }
 
 	void SetBoxPayloadProperties( const FSlateBrush* InBrush, const FLinearColor& InTint )
@@ -139,6 +150,18 @@ public:
 	{
 		CustomDrawer = InCustomDrawer;
 	}
+
+	void SetCachedBufferPayloadProperties(FSlateRenderDataHandle* InRenderDataHandle)
+	{
+		CachedRenderData = InRenderDataHandle;
+		checkSlow(CachedRenderData);
+	}
+
+	void SetLayerPayloadProperties(FSlateDrawLayerHandle* InLayerHandle)
+	{
+		LayerHandle = InLayerHandle;
+		checkSlow(LayerHandle);
+	}
 };
 
 /**
@@ -159,6 +182,8 @@ public:
 		ET_Viewport,
 		ET_Border,
 		ET_Custom,
+		ET_CachedBuffer,
+		ET_Layer,
 		ET_Count,
 	};
 
@@ -313,19 +338,23 @@ public:
 	 */
 	SLATECORE_API static void MakeCustom( FSlateWindowElementList& ElementList, uint32 InLayer, TSharedPtr<ICustomSlateElement, ESPMode::ThreadSafe> CustomDrawer );
 
+	SLATECORE_API static void MakeCachedBuffer(FSlateWindowElementList& ElementList, uint32 InLayer, TSharedPtr<FSlateRenderDataHandle, ESPMode::ThreadSafe>& CachedRenderDataHandle);
+
+	SLATECORE_API static void MakeLayer(FSlateWindowElementList& ElementList, uint32 InLayer, TSharedPtr<FSlateDrawLayerHandle, ESPMode::ThreadSafe>& DrawLayerHandle);
+
 
-	EElementType GetElementType() const { return ElementType; }
-	uint32 GetLayer() const { return Layer; }
-	const FSlateRenderTransform& GetRenderTransform() const { return RenderTransform; }
-	SLATECORE_API const FVector2D& GetPosition() const { return Position; }
-	SLATECORE_API void SetPosition(const FVector2D& InPosition) { Position = Position; }
-	const FVector2D& GetLocalSize() const { return LocalSize; }
-	float GetScale() const { return Scale; }
-	SLATECORE_API const FSlateRect& GetClippingRect() const { return ClippingRect; }
-	SLATECORE_API void SetClippingRect(const FSlateRect& InClippingRect) { ClippingRect = InClippingRect; }
-	const FSlateDataPayload& GetDataPayload() const { return DataPayload; }
-	uint32 GetDrawEffects() const { return DrawEffects; }
-	const TOptional<FShortRect>& GetScissorRect() const { return ScissorRect; }
+	FORCEINLINE EElementType GetElementType() const { return ElementType; }
+	FORCEINLINE uint32 GetLayer() const { return Layer; }
+	FORCEINLINE const FSlateRenderTransform& GetRenderTransform() const { return RenderTransform; }
+	FORCEINLINE const FVector2D& GetPosition() const { return Position; }
+	FORCEINLINE void SetPosition(const FVector2D& InPosition) { Position = Position; }
+	FORCEINLINE const FVector2D& GetLocalSize() const { return LocalSize; }
+	FORCEINLINE float GetScale() const { return Scale; }
+	FORCEINLINE const FSlateRect& GetClippingRect() const { return ClippingRect; }
+	FORCEINLINE void SetClippingRect(const FSlateRect& InClippingRect) { ClippingRect = InClippingRect; }
+	FORCEINLINE const FSlateDataPayload& GetDataPayload() const { return DataPayload; }
+	FORCEINLINE uint32 GetDrawEffects() const { return DrawEffects; }
+	FORCEINLINE const TOptional<FShortRect>& GetScissorRect() const { return ScissorRect; }
 
 private:
 	void Init(uint32 InLayer, const FPaintGeometry& PaintGeometry, const FSlateRect& InClippingRect, ESlateDrawEffect::Type InDrawEffects);
@@ -346,6 +375,7 @@ private:
 	TOptional<FShortRect> ScissorRect;
 };
 
+
 /**
  * Shader parameters for slate
  */
@@ -373,6 +403,26 @@ struct FShaderParams
 	}
 };
 
+class FSlateRenderer;
+class FSlateRenderBatch;
+
+class SLATECORE_API FSlateRenderDataHandle : public TSharedFromThis < FSlateRenderDataHandle, ESPMode::ThreadSafe >
+{
+public:
+	FSlateRenderDataHandle(FSlateRenderer* InRenderer);
+
+	virtual ~FSlateRenderDataHandle();
+
+	void Disconnect();
+
+	void SetRenderBatches(TArray<FSlateRenderBatch>* InRenderBatches) { RenderBatches = InRenderBatches; }
+	TArray<FSlateRenderBatch>* GetRenderBatches() { return RenderBatches; }
+
+private:
+	FSlateRenderer* Renderer;
+	TArray<FSlateRenderBatch>* RenderBatches;
+};
+
 /** 
  * Represents an element batch for rendering. 
  */
@@ -395,6 +445,21 @@ public:
 		, IndexArrayIndex(INDEX_NONE)
 	{}
 
+	FSlateElementBatch( TSharedPtr<FSlateRenderDataHandle, ESPMode::ThreadSafe> InCachedRenderHandle, const TOptional<FShortRect>& ScissorRect)
+		: BatchKey( InCachedRenderHandle, ScissorRect )
+		, ShaderResource( nullptr )
+		, NumElementsInBatch(0)
+		, VertexArrayIndex(INDEX_NONE)
+		, IndexArrayIndex(INDEX_NONE)
+	{}
+
+	FSlateElementBatch( TSharedPtr<FSlateDrawLayerHandle, ESPMode::ThreadSafe> InLayerHandle, const TOptional<FShortRect>& ScissorRect)
+		: BatchKey( InLayerHandle, ScissorRect )
+		, ShaderResource( nullptr )
+		, NumElementsInBatch(0)
+		, VertexArrayIndex(INDEX_NONE)
+		, IndexArrayIndex(INDEX_NONE)
+	{}
 
 	bool operator==( const FSlateElementBatch& Other ) const
 	{	
@@ -414,10 +479,14 @@ public:
 	ESlateDrawEffect::Type GetDrawEffects() const { return BatchKey.DrawEffects; }
 	const TOptional<FShortRect>& GetScissorRect() const { return BatchKey.ScissorRect; }
 	const TWeakPtr<ICustomSlateElement, ESPMode::ThreadSafe> GetCustomDrawer() const { return BatchKey.CustomDrawer; }
+	const TSharedPtr<FSlateRenderDataHandle, ESPMode::ThreadSafe> GetCachedRenderHandle() const { return BatchKey.CachedRenderHandle; }
+	const TSharedPtr<FSlateDrawLayerHandle, ESPMode::ThreadSafe> GetLayerHandle() const { return BatchKey.LayerHandle; }
 private:
 	struct FBatchKey
 	{
 		const TWeakPtr<ICustomSlateElement, ESPMode::ThreadSafe> CustomDrawer;
+		const TSharedPtr<FSlateRenderDataHandle, ESPMode::ThreadSafe> CachedRenderHandle;
+		const TSharedPtr<FSlateDrawLayerHandle, ESPMode::ThreadSafe> LayerHandle;
 		const FShaderParams ShaderParams;
 		const ESlateBatchDrawFlag::Type DrawFlags;	
 		const ESlateShader::Type ShaderType;
@@ -445,6 +514,26 @@ private:
 			, ScissorRect( InScissorRect )
 		{}
 
+		FBatchKey( TSharedPtr<FSlateRenderDataHandle, ESPMode::ThreadSafe> InCachedRenderHandle, const TOptional<FShortRect>& InScissorRect)
+			: CachedRenderHandle(InCachedRenderHandle)
+			, ShaderParams()
+			, DrawFlags(ESlateBatchDrawFlag::None)
+			, ShaderType(ESlateShader::Default)
+			, DrawPrimitiveType(ESlateDrawPrimitive::TriangleList)
+			, DrawEffects(ESlateDrawEffect::None)
+			, ScissorRect(InScissorRect)
+		{}
+
+		FBatchKey(TSharedPtr<FSlateDrawLayerHandle, ESPMode::ThreadSafe> InLayerHandle, const TOptional<FShortRect>& InScissorRect)
+			: LayerHandle(InLayerHandle)
+			, ShaderParams()
+			, DrawFlags(ESlateBatchDrawFlag::None)
+			, ShaderType(ESlateShader::Default)
+			, DrawPrimitiveType(ESlateDrawPrimitive::TriangleList)
+			, DrawEffects(ESlateDrawEffect::None)
+			, ScissorRect(InScissorRect)
+		{}
+
 		bool operator==( const FBatchKey& Other ) const
 		{
 			return DrawFlags == Other.DrawFlags
@@ -454,6 +543,8 @@ private:
 				&& ShaderParams == Other.ShaderParams
 				&& ScissorRect == Other.ScissorRect
 				&& CustomDrawer == Other.CustomDrawer
+				&& CachedRenderHandle == Other.CachedRenderHandle
+				&& LayerHandle == Other.LayerHandle
 				;
 		}
 
@@ -463,6 +554,7 @@ private:
 			// NOTE: Assumes these enum types are 8 bits.
 			uint32 RunningHash = (uint32)InBatchKey.DrawFlags << 24 | (uint32)InBatchKey.ShaderType << 16 | (uint32)InBatchKey.DrawPrimitiveType << 8 | (uint32)InBatchKey.DrawEffects << 0;
 			RunningHash = InBatchKey.CustomDrawer.IsValid() ? PointerHash(InBatchKey.CustomDrawer.Pin().Get(), RunningHash) : RunningHash;
+			RunningHash = InBatchKey.CachedRenderHandle.IsValid() ? PointerHash(InBatchKey.CachedRenderHandle.Get(), RunningHash) : RunningHash;
 			RunningHash = HashCombine(GetTypeHash(InBatchKey.ShaderParams.PixelParams), RunningHash);
 			// NOTE: Assumes this type is 64 bits, no padding.
 			RunningHash = InBatchKey.ScissorRect.IsSet() ? HashCombine(GetTypeHash(*reinterpret_cast<const uint64*>(&InBatchKey.ScissorRect.GetValue())), RunningHash) : RunningHash;
@@ -488,10 +580,13 @@ public:
 class FSlateRenderBatch
 {
 public:
-	FSlateRenderBatch( const FSlateElementBatch& InBatch, int32 InNumVertices, int32 InNumIndices, int32 InVertexOffset, int32 InIndexOffset )
-		: ShaderParams( InBatch.GetShaderParams() )
+	FSlateRenderBatch(uint32 InLayer, const FSlateElementBatch& InBatch, TSharedPtr<FSlateRenderDataHandle, ESPMode::ThreadSafe> InRenderHandle, int32 InNumVertices, int32 InNumIndices, int32 InVertexOffset, int32 InIndexOffset)
+		: Layer( InLayer )
+		, ShaderParams( InBatch.GetShaderParams() )
 		, Texture( InBatch.GetShaderResource() )
 		, CustomDrawer( InBatch.GetCustomDrawer() )
+		, LayerHandle( InBatch.GetLayerHandle() )
+		, CachedRenderHandle( InRenderHandle )
 		, DrawFlags( InBatch.GetDrawFlags() )
 		, ShaderType( InBatch.GetShaderType() )
 		, DrawPrimitiveType( InBatch.GetPrimitiveType() )
@@ -505,7 +600,16 @@ public:
 		, NumIndices( InNumIndices )
 	{}
 
+	// Render Batch Sort Operator
+	bool operator < ( const FSlateRenderBatch& RenderBatchIn ) const
+	{
+		return Layer < RenderBatchIn.Layer;
+	}
+
 public:
+	/** The layer we need to sort by when  */
+	uint32 Layer;
+
 	const FShaderParams ShaderParams;
 
 	/** Texture to use with this batch.  */
@@ -513,6 +617,10 @@ public:
 
 	const TWeakPtr<ICustomSlateElement, ESPMode::ThreadSafe> CustomDrawer;
 
+	const TWeakPtr<FSlateDrawLayerHandle, ESPMode::ThreadSafe> LayerHandle;
+
+	const TSharedPtr<FSlateRenderDataHandle, ESPMode::ThreadSafe> CachedRenderHandle;
+
 	const ESlateBatchDrawFlag::Type DrawFlags;	
 
 	const ESlateShader::Type ShaderType;
@@ -553,8 +661,6 @@ public:
 
 	void Reset();
 
-	FElementBatchMap& GetElementBatchMap() { return LayerToElementBatches; }
-
 	/**
 	 * Returns a list of element batches for this window
 	 */
@@ -585,6 +691,9 @@ public:
 	/** @return Total number of batched layers */
 	int32 GetNumLayers() const { return NumLayers; }
 
+	/** Set the associated vertex/index buffer handle. */
+	void SetRenderDataHandle(TSharedPtr<FSlateRenderDataHandle, ESPMode::ThreadSafe> InRenderDataHandle) { RenderDataHandle = InRenderDataHandle; }
+
 	/** 
 	 * Fills batch data into the actual vertex and index buffer
 	 *
@@ -596,20 +705,32 @@ public:
 	/** 
 	 * Creates rendering data from batched elements
 	 */
-	SLATECORE_API void CreateRenderBatches();
+	SLATECORE_API void CreateRenderBatches(FElementBatchMap& LayerToElementBatches);
+
+	/**
+	 * Patches the render data to be offset by the new layer delta.
+	 */
+	SLATECORE_API void UpdateRenderBatches(FVector2D PositionOffset);
+
+	/** 
+	 * Sort the rendered element batches.
+	 */
+	SLATECORE_API void SortRenderBatches();
 
 private:
-	void AddRenderBatch( const FSlateElementBatch& InElementBatch, int32 InNumVertices, int32 InNumIndices, int32 InVertexOffset, int32 InIndexOffset )
+	void Merge(FElementBatchMap& InLayerToElementBatches, uint32& VertexOffset, uint32& IndexOffset);
+
+	void AddRenderBatch( uint32 InLayer, const FSlateElementBatch& InElementBatch, int32 InNumVertices, int32 InNumIndices, int32 InVertexOffset, int32 InIndexOffset )
 	{
 		NumBatchedVertices += InNumVertices;
 		NumBatchedIndices += InNumIndices;
 
-		RenderBatches.Add( FSlateRenderBatch(InElementBatch, InNumVertices, InNumIndices, InVertexOffset, InIndexOffset) );
+		RenderBatches.Add( FSlateRenderBatch(InLayer, InElementBatch, RenderDataHandle, InNumVertices, InNumIndices, InVertexOffset, InIndexOffset) );
 	}
 private:
 
-	// Element batch maps sorted by layer.
-	FElementBatchMap LayerToElementBatches;
+	// The associated render data handle if these render batches are not in the default vertex/index buffer
+	TSharedPtr<FSlateRenderDataHandle, ESPMode::ThreadSafe> RenderDataHandle;
 
 	// Array of vertex lists that are currently free (have no elements in them).
 	TArray<uint32> VertexArrayFreeList;
@@ -634,6 +755,41 @@ private:
 
 };
 
+class FSlateRenderer;
+
+class FSlateDrawLayer
+{
+public:
+	FElementBatchMap& GetElementBatchMap() { return LayerToElementBatches; }
+
+public:
+	// Element batch maps sorted by layer.
+	FElementBatchMap LayerToElementBatches;
+
+#if SLATE_POOL_DRAW_ELEMENTS
+	/** The elements drawn on this layer */
+	TArray<FSlateDrawElement*> DrawElements;
+#else
+	/** The elements drawn on this layer */
+	TArray<FSlateDrawElement> DrawElements;
+#endif
+};
+
+/**
+ * 
+ */
+class FSlateDrawLayerHandle : public TSharedFromThis < FSlateDrawLayerHandle, ESPMode::ThreadSafe >
+{
+public:
+	FSlateDrawLayerHandle()
+		: BatchMap(nullptr)
+	{
+	}
+
+	FElementBatchMap* BatchMap;
+};
+
+
 /**
  * Represents a top level window and its draw elements.
  */
@@ -649,7 +805,8 @@ public:
 	 */
 	explicit FSlateWindowElementList( TSharedPtr<SWindow> InWindow = TSharedPtr<SWindow>() )
 		: TopLevelWindow( InWindow )
-	{		
+	{
+		DrawStack.Push(&RootDrawLayer);
 	}
 	
 	/** @return Get the window that we will be painting */
@@ -664,12 +821,22 @@ public:
 		return TopLevelWindow.Pin();
 	}
 	
+#if SLATE_POOL_DRAW_ELEMENTS
+	
+	/** @return Get the draw elements that we want to render into this window */
+	FORCEINLINE const TArray<FSlateDrawElement*>& GetDrawElements() const
+	{
+		return RootDrawLayer.DrawElements;
+	}
+
+#else
+
 	/** @return Get the draw elements that we want to render into this window */
 	FORCEINLINE const TArray<FSlateDrawElement>& GetDrawElements() const
 	{
-		return DrawElements;
+		return RootDrawLayer.DrawElements;
 	}
-	
+
 	/**
 	 * Add a draw element to the list
 	 *
@@ -677,25 +844,47 @@ public:
 	 */
 	FORCEINLINE void AddItem(const FSlateDrawElement& InDrawElement)
 	{
-		DrawElements.Add( InDrawElement );
+		TArray<FSlateDrawElement>& ActiveDrawElements = DrawStack.Last()->DrawElements;
+		ActiveDrawElements.Add(InDrawElement);
 	}
+#endif
 
 	/**
 	 * Creates an uninitialized draw element
 	 */
 	FORCEINLINE FSlateDrawElement& AddUninitialized()
 	{
-		const int32 InsertIdx = DrawElements.AddDefaulted();
-		return DrawElements[InsertIdx];
+#if SLATE_POOL_DRAW_ELEMENTS
+		FSlateDrawElement* DrawElement = ( DrawElementFreePool.Num() > 0 ) ? DrawElementFreePool.Pop() : new FSlateDrawElement();
+		DrawStack.Last()->DrawElements.Push(DrawElement);
+
+		return *DrawElement;
+#else
+		TArray<FSlateDrawElement>& ActiveDrawElements = DrawStack.Last()->DrawElements;
+		const int32 InsertIdx = ActiveDrawElements.AddDefaulted();
+		return ActiveDrawElements[InsertIdx];
+#endif
 	}
 
+#if SLATE_POOL_DRAW_ELEMENTS
 	/**
 	 * Append draw elements to the list of draw elements
 	 */
+	FORCEINLINE void AppendDrawElements(const TArray<FSlateDrawElement*>& InDrawElements)
+	{
+		TArray<FSlateDrawElement*>& ActiveDrawElements = DrawStack.Last()->DrawElements;
+		ActiveDrawElements.Append(InDrawElements);
+	}
+#else
+	/**
+	* Append draw elements to the list of draw elements
+	*/
 	FORCEINLINE void AppendDrawElements(const TArray<FSlateDrawElement>& InDrawElements)
 	{
-		DrawElements.Append(InDrawElements);
+		TArray<FSlateDrawElement>& ActiveDrawElements = DrawStack.Last()->DrawElements;
+		ActiveDrawElements.Append(InDrawElements);
 	}
+#endif
 
 	/**
 	 * Some widgets may want to paint their children after after another, loosely-related widget finished painting.
@@ -726,9 +915,14 @@ public:
 	public:
 		SLATECORE_API FVolatilePaint(const TSharedRef<const SWidget>& InWidgetToPaint, const FPaintArgs& InArgs, const FGeometry InAllottedGeometry, const FSlateRect InMyClippingRect, int32 LayerId, const FWidgetStyle& InWidgetStyle, bool InParentEnabled);
 
-		int32 ExecutePaint( FSlateWindowElementList& OutDrawElements ) const;
+		int32 ExecutePaint(FSlateWindowElementList& OutDrawElements) const;
 
+		FORCEINLINE const SWidget* GetWidget() const { return WidgetToPaintPtr.Pin().Get(); }
 		FORCEINLINE FGeometry GetGeometry() const { return AllottedGeometry; }
+		FORCEINLINE int32 GetLayerId() const { return LayerId; }
+
+	public:
+		TSharedPtr< FSlateDrawLayerHandle, ESPMode::ThreadSafe > LayerHandle;
 	 
 	private:
 		const TWeakPtr<const SWidget> WidgetToPaintPtr;
@@ -744,31 +938,80 @@ public:
 
 	SLATECORE_API int32 PaintVolatile(FSlateWindowElementList& OutElementList);
 
-	SLATECORE_API const TArray< TSharedRef<FVolatilePaint> >& GetVolatileElements() const { return VolatilePaintList; }
+	SLATECORE_API void BeginLogicalLayer(const TSharedPtr<FSlateDrawLayerHandle, ESPMode::ThreadSafe>& LayerHandle);
+	SLATECORE_API void EndLogicalLayer();
+
+	SLATECORE_API const TArray< TSharedPtr<FVolatilePaint> >& GetVolatileElements() const { return VolatilePaintList; }
 	
 	/**
 	 * Remove all the elements from this draw list.
 	 */
-	SLATECORE_API void Reset();
+	SLATECORE_API void ResetBuffers();
 
 	FSlateBatchData& GetBatchData() { return BatchData; }
 
+	FSlateDrawLayer& GetRootDrawLayer() { return RootDrawLayer; }
+
+	TMap < TSharedPtr<FSlateDrawLayerHandle, ESPMode::ThreadSafe>, TSharedPtr<FSlateDrawLayer> >& GetChildDrawLayers() { return DrawLayers; }
+
+	/**
+	 * Caches this element list on the renderer, generating all needed index and vertex buffers.
+	 */
+	SLATECORE_API TSharedRef<FSlateRenderDataHandle, ESPMode::ThreadSafe> CacheRenderData();
+
+	SLATECORE_API void UpdateCacheRenderData(FVector2D PositionOffset);
+
+	SLATECORE_API TSharedPtr<FSlateRenderDataHandle, ESPMode::ThreadSafe> GetCachedRenderDataHandle() const
+	{
+		return CachedRenderDataHandle.Pin();
+	}
+
+	SLATECORE_API bool IsCachedRenderDataInUse() const
+	{
+		return CachedRenderDataHandle.IsValid();
+	}
+
+	SLATECORE_API void PreDraw_RenderThread();
+
+	SLATECORE_API void PostDraw_RenderThread();
+
 private:
-	/** Batched data used for rendering */
-	FSlateBatchData BatchData;
 
 	/** The top level window which these elements are being drawn on */
 	TWeakPtr<SWindow> TopLevelWindow;
-	
+
+	/** Batched data used for rendering */
+	FSlateBatchData BatchData;
+
+	/** The base draw layer/context. */
+	FSlateDrawLayer RootDrawLayer;
+
+	/** */
+	TMap < TSharedPtr<FSlateDrawLayerHandle, ESPMode::ThreadSafe>, TSharedPtr<FSlateDrawLayer> > DrawLayers;
+
+	/** */
+	TArray< TSharedPtr<FSlateDrawLayer> > DrawLayerPool;
+
+	/** */
+	TArray< FSlateDrawLayer* > DrawStack;
+
+#if SLATE_POOL_DRAW_ELEMENTS
 	/** List of draw elements for the window */
-	TArray<FSlateDrawElement> DrawElements;
+	TArray<FSlateDrawElement*> DrawElementFreePool;
+#endif
 
 	/**
 	 * Some widgets want their logical children to appear at a different "layer" in the physical hierarchy.
 	 * We accomplish this by deferring their painting.
 	 */
-	TArray< TSharedRef<FDeferredPaint> > DeferredPaintList;
+	TArray< TSharedPtr<FDeferredPaint> > DeferredPaintList;
 
 	/** The widgets be cached for a later paint pass when the invalidation host paints. */
-	TArray< TSharedRef<FVolatilePaint> > VolatilePaintList;
+	TArray< TSharedPtr<FVolatilePaint> > VolatilePaintList;
+
+	/**
+	 * Handle to the cached render data associated with this element list.  Will only exist if 
+	 * this element list is being used for invalidation / caching.
+	 */
+	mutable TWeakPtr<FSlateRenderDataHandle, ESPMode::ThreadSafe> CachedRenderDataHandle;
 };
diff --git a/Engine/Source/Runtime/SlateCore/Public/Rendering/ElementBatcher.h b/Engine/Source/Runtime/SlateCore/Public/Rendering/ElementBatcher.h
--- a/Engine/Source/Runtime/SlateCore/Public/Rendering/ElementBatcher.h
+++ b/Engine/Source/Runtime/SlateCore/Public/Rendering/ElementBatcher.h
@@ -31,7 +31,9 @@ public:
 	 * Resets all stored data accumulated during the batching process
 	 */
 	void ResetBatches();
+
 private:
+	void AddElements(FSlateDrawLayer& InDrawLayer);
 
 	/** 
 	 * Creates vertices necessary to draw a Quad element 
@@ -85,6 +87,10 @@ private:
 	 */
 	void AddCustomElement( const FSlateDrawElement& DrawElement );
 
+	void AddCachedBuffer( const FSlateDrawElement& DrawElement );
+
+	void AddLayer(const FSlateDrawElement& DrawElement);
+
 	/** 
 	 * Finds an batch for an element based on the passed in parameters
 	 * Elements with common parameters and layers will batched together.
@@ -108,12 +114,18 @@ private:
 	/** Batch data currently being filled in */
 	FSlateBatchData* BatchData;
 
+	/** The draw layer currently being accumulated */
+	FSlateDrawLayer* DrawLayer;
+
 	/** Resource manager for accessing shader resources */
 	FSlateShaderResourceManager& ResourceManager;
 
 	/** Font cache used to layout text */
 	FSlateFontCache& FontCache;
 
+	/** Track the number of drawn batches from the previous frame to report to stats. */
+	int32 NumDrawnBatchesStat;
+
 	/** Offset to use when supporting 1:1 texture to pixel snapping */
 	const float PixelCenterOffset;
 
diff --git a/Engine/Source/Runtime/SlateCore/Public/Rendering/SlateRenderer.h b/Engine/Source/Runtime/SlateCore/Public/Rendering/SlateRenderer.h
--- a/Engine/Source/Runtime/SlateCore/Public/Rendering/SlateRenderer.h
+++ b/Engine/Source/Runtime/SlateCore/Public/Rendering/SlateRenderer.h
@@ -14,6 +14,7 @@ struct FSlateBrush;
 
 typedef FRHITexture2D* FTexture2DRHIParamRef;
 
+
 /**
  * Abstract base class for Slate renderers.
  */
@@ -62,10 +63,10 @@ public:
 	virtual void UpdateFullscreenState( const TSharedRef<SWindow> InWindow, uint32 OverrideResX = 0, uint32 OverrideResY = 0 ) = 0;
 
 	/**
-     * Restore the given window to the resolution settings currently cached by the engine
+	 * Restore the given window to the resolution settings currently cached by the engine
 	 * 
 	 * @param InWindow    -> The window to restore to the cached settings
-     */
+	 */
 	virtual void RestoreSystemResolution(const TSharedRef<SWindow> InWindow) = 0;
 
 	/** 
@@ -232,6 +233,16 @@ public:
 	 */
 	virtual ISlateAtlasProvider* GetFontAtlasProvider();
 
+	/**
+	 * Converts and caches the elements list data as the final rendered vertex and index buffer data,
+	 * returning a handle to it to allow issuing future draw commands using it.
+	 */
+	virtual TSharedRef<FSlateRenderDataHandle, ESPMode::ThreadSafe> CacheElementRenderData(FSlateWindowElementList& ElementList);
+
+	/**
+	 * Patch existing render data to have a new layer delta applied to all elements.
+	 */
+	virtual void UpdateElementRenderData(FSlateWindowElementList& ElementList, FVector2D PositionOffset);
 
 	/**
 	 * You must call this before calling CopyWindowsToVirtualScreenBuffer(), to setup the render targets first.
@@ -243,7 +254,6 @@ public:
 	 */
 	virtual FIntRect SetupVirtualScreenBuffer( const bool bPrimaryWorkAreaOnly, const float ScreenScaling, class ILiveStreamingService* LiveStreamingService) { return FIntRect(); }
 
-
 	/**
 	 * Copies all slate windows out to a buffer at half resolution with debug information
 	 * like the mouse cursor and any keypresses.
@@ -254,6 +264,9 @@ public:
 	virtual void MapVirtualScreenBuffer(void** OutImageData) {}
 	virtual void UnmapVirtualScreenBuffer() {}
 
+protected:
+	virtual void ReleaseCachedRenderData(FSlateRenderDataHandle* RenderHandle) { }
+
 private:
 
 	// Non-copyable
@@ -267,6 +280,8 @@ protected:
 
 	/** Callback that fires after Slate has rendered each window, each frame */
 	FOnSlateWindowRendered SlateWindowRendered;
+
+	friend class FSlateRenderDataHandle;
 };
 
 
diff --git a/Engine/Source/Runtime/SlateCore/Public/Types/ReflectionMetadata.h b/Engine/Source/Runtime/SlateCore/Public/Types/ReflectionMetadata.h
--- a/Engine/Source/Runtime/SlateCore/Public/Types/ReflectionMetadata.h
+++ b/Engine/Source/Runtime/SlateCore/Public/Types/ReflectionMetadata.h
@@ -27,6 +27,25 @@ public:
 	/** The class the constructed the slate widget. */
 	TWeakObjectPtr<UClass> Class;
 
-	/** The asset that owns the widget and is responsible for its specific existance. */
+	/** The asset that owns the widget and is responsible for its specific existence. */
 	TWeakObjectPtr<UObject> Asset;
+
+public:
+
+	static FString GetWidgetDebugInfo(TSharedPtr<const SWidget>& InWidget)
+	{
+		// UMG widgets have meta-data to help track them
+		TSharedPtr<FReflectionMetaData> MetaData = InWidget->GetMetaData<FReflectionMetaData>();
+		if ( MetaData.IsValid() && MetaData->Asset.Get() != nullptr )
+		{
+			const FName AssetName = MetaData->Asset->GetFName();
+			const FName WidgetName = MetaData->Name;
+
+			return FString::Printf(TEXT("%s [%s]"), *AssetName.ToString(), *WidgetName.ToString());
+		}
+		else
+		{
+			return InWidget->GetReadableLocation();
+		}
+	}
 };
diff --git a/Engine/Source/Runtime/SlateCore/Public/Widgets/SLeafWidget.h b/Engine/Source/Runtime/SlateCore/Public/Widgets/SLeafWidget.h
--- a/Engine/Source/Runtime/SlateCore/Public/Widgets/SLeafWidget.h
+++ b/Engine/Source/Runtime/SlateCore/Public/Widgets/SLeafWidget.h
@@ -13,6 +13,11 @@ class SLATECORE_API SLeafWidget
 	: public SWidget
 {
 public:
+	SLeafWidget()
+	{
+		bCanHaveChildren = false;
+	}
+
 	virtual void SetVisibility( TAttribute<EVisibility> InVisibility ) override final;
 
 private:
diff --git a/Engine/Source/Runtime/SlateRHIRenderer/Private/Slate3DRenderer.cpp b/Engine/Source/Runtime/SlateRHIRenderer/Private/Slate3DRenderer.cpp
--- a/Engine/Source/Runtime/SlateRHIRenderer/Private/Slate3DRenderer.cpp
+++ b/Engine/Source/Runtime/SlateRHIRenderer/Private/Slate3DRenderer.cpp
@@ -85,8 +85,9 @@ void FSlate3DRenderer::DrawWindow_GameThread(FSlateDrawBuffer& DrawBuffer)
 void FSlate3DRenderer::DrawWindowToTarget_RenderThread( FRHICommandListImmediate& RHICmdList, UTextureRenderTarget2D* RenderTarget, FSlateDrawBuffer& InDrawBuffer )
 {
 	FSlateBatchData& BatchData = InDrawBuffer.GetWindowElementLists()[0]->GetBatchData();
+	FElementBatchMap& RootBatchMap = InDrawBuffer.GetWindowElementLists()[0]->GetRootDrawLayer().GetElementBatchMap();
 
-	BatchData.CreateRenderBatches();
+	BatchData.CreateRenderBatches(RootBatchMap);
 	RenderTargetPolicy->BeginDrawingWindows();
 
 	RenderTargetPolicy->UpdateVertexAndIndexBuffers( RHICmdList, BatchData );
diff --git a/Engine/Source/Runtime/SlateRHIRenderer/Private/SlateRHIRenderer.cpp b/Engine/Source/Runtime/SlateRHIRenderer/Private/SlateRHIRenderer.cpp
--- a/Engine/Source/Runtime/SlateRHIRenderer/Private/SlateRHIRenderer.cpp
+++ b/Engine/Source/Runtime/SlateRHIRenderer/Private/SlateRHIRenderer.cpp
@@ -403,9 +403,12 @@ void FSlateRHIRenderer::DrawWindow_RenderThread(FRHICommandListImmediate& RHICmd
 		SCOPE_CYCLE_COUNTER( STAT_SlateRenderingRTTime );
 
 		FSlateBatchData& BatchData = WindowElementList.GetBatchData();
+		FElementBatchMap& RootBatchMap = WindowElementList.GetRootDrawLayer().GetElementBatchMap();
+
+		WindowElementList.PreDraw_RenderThread();
 
 		// Update the vertex and index buffer	
-		BatchData.CreateRenderBatches();	
+		BatchData.CreateRenderBatches(RootBatchMap);
 		RenderingPolicy->UpdateVertexAndIndexBuffers(RHICmdList, BatchData);
 		// should have been created by the game thread
 		check( IsValidRef(ViewportInfo.ViewportRHI) );
@@ -418,7 +421,7 @@ void FSlateRHIRenderer::DrawWindow_RenderThread(FRHICommandListImmediate& RHICmd
 		const uint32 ViewportHeight = (ViewportRT) ? ViewportRT->GetSizeY() : ViewportInfo.Height;
 
 		RHICmdList.BeginDrawingViewport( ViewportInfo.ViewportRHI, FTextureRHIRef() );
-		RHICmdList.SetViewport( 0,0,0,ViewportWidth, ViewportHeight, 0.0f ); 
+		RHICmdList.SetViewport(0, 0, 0, ViewportWidth, ViewportHeight, 0.0f);
 
 		if( ViewportInfo.bRequiresStencilTest )
 		{
@@ -427,7 +430,7 @@ void FSlateRHIRenderer::DrawWindow_RenderThread(FRHICommandListImmediate& RHICmd
 			// Reset the backbuffer as our color render target and also set a depth stencil buffer
 			SetRenderTarget(RHICmdList, BackBuffer, ViewportInfo.DepthStencil);
 			// Clear the stencil buffer
-			RHICmdList.Clear( false, FLinearColor::White, false, 0.0f, true, 0x00, FIntRect());
+			RHICmdList.Clear(false, FLinearColor::White, false, 0.0f, true, 0x00, FIntRect());
 		}
 		else
 		{
@@ -436,7 +439,7 @@ void FSlateRHIRenderer::DrawWindow_RenderThread(FRHICommandListImmediate& RHICmd
 
 		if (bClear)
 		{
-			RHICmdList.Clear( true, FLinearColor(ForceInitToZero), false, 0.0f, true, 0x00, FIntRect());
+			RHICmdList.Clear(true, FLinearColor(ForceInitToZero), false, 0.0f, true, 0x00, FIntRect());
 		}
 
 #if DEBUG_OVERDRAW
@@ -454,6 +457,8 @@ void FSlateRHIRenderer::DrawWindow_RenderThread(FRHICommandListImmediate& RHICmd
 				BatchData.GetRenderBatches()
 			);
 		}
+
+		WindowElementList.PostDraw_RenderThread();
 	}
 
 	bool bNeedCallFinishFrameForStereo = false;
@@ -949,8 +954,9 @@ void FSlateRHIRenderer::CopyWindowsToVirtualScreenBuffer(const TArray<FString>&
 		RHICmdList.SetBlendState(TStaticBlendState<CW_RGBA, BO_Add, BF_SourceAlpha, BF_InverseSourceAlpha, BO_Add, BF_Zero, BF_One>::GetRHI());
 		
 		FSlateBatchData& BatchData = Context.SlateElementList->GetBatchData();
+		FElementBatchMap& RootBatchMap = Context.SlateElementList->GetRootDrawLayer().GetElementBatchMap();
 
-		BatchData.CreateRenderBatches();
+		BatchData.CreateRenderBatches(RootBatchMap);
 
 		Context.RenderPolicy->UpdateVertexAndIndexBuffers(RHICmdList, BatchData);
 		if( BatchData.GetRenderBatches().Num() > 0 )
@@ -958,7 +964,7 @@ void FSlateRHIRenderer::CopyWindowsToVirtualScreenBuffer(const TArray<FString>&
 			FTexture2DRHIRef UnusedTargetTexture;
 			FSlateBackBuffer UnusedTarget( UnusedTargetTexture, FIntPoint::ZeroValue );
 
-			Context.RenderPolicy->DrawElements(RHICmdList, UnusedTarget, CreateProjectionMatrix(Context.ViewportSize.X, Context.ViewportSize.Y),BatchData.GetRenderBatches());
+			Context.RenderPolicy->DrawElements(RHICmdList, UnusedTarget, CreateProjectionMatrix(Context.ViewportSize.X, Context.ViewportSize.Y), BatchData.GetRenderBatches());
 		}
 	});
 
@@ -1238,3 +1244,93 @@ void FSlateRHIRenderer::SetWindowRenderTarget(const SWindow& Window, IViewportRe
 		ViewInfo->RTProvider = Provider;
 	}
 }
+
+TSharedRef<FSlateRenderDataHandle, ESPMode::ThreadSafe> FSlateRHIRenderer::CacheElementRenderData(FSlateWindowElementList& ElementList)
+{
+	TSharedRef<FSlateRenderDataHandle, ESPMode::ThreadSafe> RenderDataHandle = MakeShareable(new FSlateRenderDataHandle(this));
+
+	checkSlow(ElementList.GetChildDrawLayers().Num() == 0);
+
+	// Add all elements for this window to the element batcher
+	ElementBatcher->AddElements(ElementList);
+
+	// All elements for this window have been batched and rendering data updated
+	ElementBatcher->ResetBatches();
+
+	struct FCacheElementBatchesContext
+	{
+		FSlateRHIRenderingPolicy* RenderPolicy;
+		FSlateWindowElementList* SlateElementList;
+		TSharedRef<FSlateRenderDataHandle, ESPMode::ThreadSafe> RenderDataHandle;
+	};
+	FCacheElementBatchesContext CacheElementBatchesContext =
+	{
+		RenderingPolicy.Get(),
+		&ElementList,
+		RenderDataHandle,
+	};
+	ENQUEUE_UNIQUE_RENDER_COMMAND_ONEPARAMETER(
+		CacheElementBatches,
+		FCacheElementBatchesContext, Context, CacheElementBatchesContext,
+	{
+		FSlateBatchData& BatchData = Context.SlateElementList->GetBatchData();
+		FElementBatchMap& RootBatchMap = Context.SlateElementList->GetRootDrawLayer().GetElementBatchMap();
+
+		BatchData.SetRenderDataHandle(Context.RenderDataHandle);
+		BatchData.CreateRenderBatches(RootBatchMap);
+		Context.RenderPolicy->UpdateVertexAndIndexBuffers(RHICmdList, BatchData, Context.RenderDataHandle);
+	});
+
+	return RenderDataHandle;
+}
+
+void FSlateRHIRenderer::UpdateElementRenderData(FSlateWindowElementList& ElementList, FVector2D PositionOffset)
+{
+	struct FUpdateCacheElementBatchesContext
+	{
+		FSlateRHIRenderingPolicy* RenderPolicy;
+		FSlateWindowElementList* SlateElementList;
+		TSharedPtr<FSlateRenderDataHandle, ESPMode::ThreadSafe> RenderDataHandle;
+		FVector2D PositionOffset;
+	};
+	FUpdateCacheElementBatchesContext UpdateCacheElementBatchesContext =
+	{
+		RenderingPolicy.Get(),
+		&ElementList,
+		ElementList.GetCachedRenderDataHandle(),
+		PositionOffset
+	};
+	ENQUEUE_UNIQUE_RENDER_COMMAND_ONEPARAMETER(
+		UpdateCacheElementBatches,
+		FUpdateCacheElementBatchesContext, Context, UpdateCacheElementBatchesContext,
+	{
+		if ( Context.RenderDataHandle.IsValid() )
+		{
+			auto RenderDataHandleRef = Context.RenderDataHandle.ToSharedRef();
+
+			FSlateBatchData& BatchData = Context.SlateElementList->GetBatchData();
+			BatchData.UpdateRenderBatches(Context.PositionOffset);
+			Context.RenderPolicy->UpdateVertexAndIndexBuffers(RHICmdList, BatchData, RenderDataHandleRef);
+		}
+	});
+}
+
+void FSlateRHIRenderer::ReleaseCachedRenderData(FSlateRenderDataHandle* InRenderHandle)
+{
+	struct FReleaseCachedRenderDataContext
+	{
+		FSlateRHIRenderingPolicy* RenderPolicy;
+		FSlateRenderDataHandle* RenderDataHandle;
+	};
+	FReleaseCachedRenderDataContext ReleaseCachedRenderDataContext =
+	{
+		RenderingPolicy.Get(),
+		InRenderHandle,
+	};
+	ENQUEUE_UNIQUE_RENDER_COMMAND_ONEPARAMETER(
+		ReleaseCachedRenderData,
+		FReleaseCachedRenderDataContext, Context, ReleaseCachedRenderDataContext,
+	{
+		Context.RenderPolicy->ReleaseCachedRenderData(Context.RenderDataHandle);
+	});
+}
diff --git a/Engine/Source/Runtime/SlateRHIRenderer/Private/SlateRHIRenderer.h b/Engine/Source/Runtime/SlateRHIRenderer/Private/SlateRHIRenderer.h
--- a/Engine/Source/Runtime/SlateRHIRenderer/Private/SlateRHIRenderer.h
+++ b/Engine/Source/Runtime/SlateRHIRenderer/Private/SlateRHIRenderer.h
@@ -188,6 +188,9 @@ public:
 	virtual void ReleaseUpdatableTexture(FSlateUpdatableTexture* Texture) override;
 	virtual ISlateAtlasProvider* GetTextureAtlasProvider() override;
 	virtual void ReleaseAccessedResources() override;
+	virtual TSharedRef<FSlateRenderDataHandle, ESPMode::ThreadSafe> CacheElementRenderData(FSlateWindowElementList& ElementList) override;
+	virtual void UpdateElementRenderData(FSlateWindowElementList& ElementList, FVector2D PositionOffset) override;
+	virtual void ReleaseCachedRenderData(FSlateRenderDataHandle* RenderHandle) override;
 
 	/** Draws windows from a FSlateDrawBuffer on the render thread */
 	void DrawWindow_RenderThread(FRHICommandListImmediate& RHICmdList, const FSlateRHIRenderer::FViewportInfo& ViewportInfo, FSlateWindowElementList& WindowElementList, bool bLockToVsync, bool bClear);
@@ -219,16 +222,8 @@ public:
 	 */
 	virtual void ReloadTextureResources() override;
 
-
 	virtual void LoadStyleResources( const ISlateStyle& Style ) override;
 
-	/**
-	 * Returns the viewport RHI reference for the provided window
-	 *
-	 * @param Window	The window to get the RHI viewport from 
-	 */
-
-
 	/** Returns whether shaders that Slate depends on have been compiled. */
 	virtual bool AreShadersInitialized() const override;
 
diff --git a/Engine/Source/Runtime/SlateRHIRenderer/Private/SlateRHIRenderingPolicy.cpp b/Engine/Source/Runtime/SlateRHIRenderer/Private/SlateRHIRenderingPolicy.cpp
--- a/Engine/Source/Runtime/SlateRHIRenderer/Private/SlateRHIRenderingPolicy.cpp
+++ b/Engine/Source/Runtime/SlateRHIRenderer/Private/SlateRHIRenderingPolicy.cpp
@@ -10,6 +10,7 @@
 #include "EngineModule.h"
 #include "SlateUTextureResource.h"
 #include "SlateMaterialResource.h"
+#include "Rendering/DrawElements.h"
 
 DECLARE_CYCLE_STAT(TEXT("Update Buffers RT"), STAT_SlateUpdateBufferRTTime, STATGROUP_Slate);
 DECLARE_CYCLE_STAT(TEXT("PreFill Buffers RT"), STAT_SlatePreFullBufferRTTime, STATGROUP_Slate);
@@ -40,12 +41,26 @@ void FSlateElementIndexBuffer::Init( int32 MinNumIndices )
 {
 	MinBufferSize = sizeof(SlateIndex) * FMath::Max( MinNumIndices, 200 );
 
-	BeginInitResource(this);
+	if ( IsInRenderingThread() )
+	{
+		InitResource();
+	}
+	else
+	{
+		BeginInitResource(this);
+	}
 }
 
 void FSlateElementIndexBuffer::Destroy()
 {
-	BeginReleaseResource(this);
+	if ( IsInRenderingThread() )
+	{
+		ReleaseResource();
+	}
+	else
+	{
+		BeginReleaseResource(this);
+	}
 }
 
 
@@ -177,6 +192,23 @@ void FSlateRHIRenderingPolicy::ReleaseResources()
 		VertexBuffers[BufferIndex].Destroy();
 		IndexBuffers[BufferIndex].Destroy();
 	}
+
+	for ( TCachedBufferMap::TIterator BufferIt(CachedBuffers); BufferIt; ++BufferIt )
+	{
+		FSlateRenderDataHandle* Handle = BufferIt.Key();
+		FCachedRenderBuffers* Buffer = BufferIt.Value();
+
+		Handle->Disconnect();
+
+		Buffer->VertexBuffer.Destroy();
+		Buffer->IndexBuffer.Destroy();
+	}
+
+	for ( FCachedRenderBuffers* PooledBuffer : CachedBufferPool )
+	{
+		PooledBuffer->VertexBuffer.Destroy();
+		PooledBuffer->IndexBuffer.Destroy();
+	}
 }
 
 void FSlateRHIRenderingPolicy::BeginDrawingWindows()
@@ -234,10 +266,55 @@ struct FSlateUpdateVertexAndIndexBuffers : public FRHICommand<FSlateUpdateVertex
 	
 		VertexBuffer.UnlockBuffer_RHIThread();
 		IndexBuffer.UnlockBuffer_RHIThread();
+
+		BatchData.SortRenderBatches();
 	}
 };
 
 void FSlateRHIRenderingPolicy::UpdateVertexAndIndexBuffers(FRHICommandListImmediate& RHICmdList, FSlateBatchData& InBatchData)
+{
+	TSlateElementVertexBuffer<FSlateVertex>& VertexBuffer = VertexBuffers[CurrentBufferIndex];
+	FSlateElementIndexBuffer& IndexBuffer = IndexBuffers[CurrentBufferIndex];
+
+	UpdateVertexAndIndexBuffers(RHICmdList, InBatchData, VertexBuffer, IndexBuffer);
+}
+
+void FSlateRHIRenderingPolicy::UpdateVertexAndIndexBuffers(FRHICommandListImmediate& RHICmdList, FSlateBatchData& InBatchData, TSharedRef<FSlateRenderDataHandle, ESPMode::ThreadSafe> RenderHandle)
+{
+	FCachedRenderBuffers* Buffers = CachedBuffers.FindRef(&RenderHandle.Get());
+	if ( Buffers == nullptr )
+	{
+		// If the cached buffer pool is empty, time to create a new one!
+		if ( CachedBufferPool.Num() == 0 )
+		{
+			Buffers = new FCachedRenderBuffers();
+			Buffers->VertexBuffer.Init(200);
+			Buffers->IndexBuffer.Init(200);
+		}
+		else
+		{
+			// If we found one in the pool, lets use it!
+			Buffers = CachedBufferPool[0];
+			CachedBufferPool.RemoveAtSwap(0, 1, false);
+		}
+
+		CachedBuffers.Add(&RenderHandle.Get(), Buffers);
+	}
+
+	UpdateVertexAndIndexBuffers(RHICmdList, InBatchData, Buffers->VertexBuffer, Buffers->IndexBuffer);
+}
+
+void FSlateRHIRenderingPolicy::ReleaseCachedRenderData(FSlateRenderDataHandle* InRenderHandle)
+{
+	FCachedRenderBuffers* Buffers = CachedBuffers.FindRef(InRenderHandle);
+	if ( ensure(Buffers != nullptr) )
+	{
+		CachedBufferPool.Add(Buffers);
+		CachedBuffers.Remove(InRenderHandle);
+	}
+}
+
+void FSlateRHIRenderingPolicy::UpdateVertexAndIndexBuffers(FRHICommandListImmediate& RHICmdList, FSlateBatchData& InBatchData, TSlateElementVertexBuffer<FSlateVertex>& VertexBuffer, FSlateElementIndexBuffer& IndexBuffer)
 {
 	SCOPE_CYCLE_COUNTER( STAT_SlateUpdateBufferRTTime );
 
@@ -249,15 +326,12 @@ void FSlateRHIRenderingPolicy::UpdateVertexAndIndexBuffers(FRHICommandListImmedi
 
 	if( InBatchData.GetRenderBatches().Num() > 0  && NumVertices > 0 && NumIndices > 0)
 	{
-		TSlateElementVertexBuffer<FSlateVertex>& VertexBuffer = VertexBuffers[CurrentBufferIndex];
-		FSlateElementIndexBuffer& IndexBuffer = IndexBuffers[CurrentBufferIndex];
-
 		bool bShouldShrinkResources = false;
 
 		VertexBuffer.PreFillBuffer(NumVertices, bShouldShrinkResources);
 		IndexBuffer.PreFillBuffer(NumIndices, bShouldShrinkResources);
 
-		if(!GRHIThread || RHICmdList.Bypass())
+		if(1 || !GRHIThread || RHICmdList.Bypass())
 		{
 			uint8* VertexBufferData = (uint8*)VertexBuffer.LockBuffer_RenderThread(NumVertices);
 			uint8* IndexBufferData =  (uint8*)IndexBuffer.LockBuffer_RenderThread(NumIndices);
@@ -266,6 +340,8 @@ void FSlateRHIRenderingPolicy::UpdateVertexAndIndexBuffers(FRHICommandListImmedi
 	
 			VertexBuffer.UnlockBuffer_RenderThread();
 			IndexBuffer.UnlockBuffer_RenderThread();
+
+			InBatchData.SortRenderBatches();
 		}
 		else
 		{
@@ -273,6 +349,9 @@ void FSlateRHIRenderingPolicy::UpdateVertexAndIndexBuffers(FRHICommandListImmedi
 		}
 	}
 
+	checkSlow(VertexBuffer.GetBufferUsageSize() <= VertexBuffer.GetBufferSize());
+	checkSlow(IndexBuffer.GetBufferUsageSize() <= IndexBuffer.GetBufferSize());
+
 	SET_DWORD_STAT( STAT_SlateNumLayers, InBatchData.GetNumLayers() );
 	SET_DWORD_STAT( STAT_SlateNumBatches, InBatchData.GetRenderBatches().Num() );
 	SET_DWORD_STAT( STAT_SlateVertexCount, InBatchData.GetNumBatchedVertices() );
@@ -358,9 +437,6 @@ void FSlateRHIRenderingPolicy::DrawElements(FRHICommandListImmediate& RHICmdList
 	// Should only be called by the rendering thread
 	check(IsInRenderingThread());
 
-	TSlateElementVertexBuffer<FSlateVertex>& VertexBuffer = VertexBuffers[CurrentBufferIndex];
-	FSlateElementIndexBuffer& IndexBuffer = IndexBuffers[CurrentBufferIndex];
-
 	float TimeSeconds = FPlatformTime::Seconds() - GStartTime;
 	float RealTimeSeconds = FPlatformTime::Seconds() - GStartTime;
 	float DeltaTimeSeconds = FApp::GetDeltaTime();
@@ -390,15 +466,50 @@ void FSlateRHIRenderingPolicy::DrawElements(FRHICommandListImmediate& RHICmdList
 
 	FSamplerStateRHIRef BilinearClamp = TStaticSamplerState<SF_Bilinear,AM_Clamp,AM_Clamp,AM_Clamp>::GetRHI();
 
-	if (GRHISupportsBaseVertexIndex)
+	TSlateElementVertexBuffer<FSlateVertex>* VertexBuffer = &VertexBuffers[CurrentBufferIndex];
+	FSlateElementIndexBuffer* IndexBuffer = &IndexBuffers[CurrentBufferIndex];
+
+	if ( GRHISupportsBaseVertexIndex )
 	{
-		RHICmdList.SetStreamSource(0, VertexBuffer.VertexBufferRHI, sizeof(FSlateVertex), 0);
+		RHICmdList.SetStreamSource(0, VertexBuffer->VertexBufferRHI, sizeof(FSlateVertex), 0);
 	}
 
+	const FSlateRenderDataHandle* LastHandle = nullptr;
+
 	// Draw each element
 	for( int32 BatchIndex = 0; BatchIndex < RenderBatches.Num(); ++BatchIndex )
 	{
 		const FSlateRenderBatch& RenderBatch = RenderBatches[BatchIndex];
+		const FSlateRenderDataHandle* RenderHandle = RenderBatch.CachedRenderHandle.Get();
+
+		if ( RenderHandle != LastHandle )
+		{
+			LastHandle = RenderHandle;
+
+			if ( RenderHandle )
+			{
+				FCachedRenderBuffers* Buffers = CachedBuffers.FindRef(RenderHandle);
+				if ( Buffers != nullptr )
+				{
+					VertexBuffer = &Buffers->VertexBuffer;
+					IndexBuffer = &Buffers->IndexBuffer;
+				}
+			}
+			else
+			{
+				VertexBuffer = &VertexBuffers[CurrentBufferIndex];
+				IndexBuffer = &IndexBuffers[CurrentBufferIndex];
+			}
+
+			checkSlow(VertexBuffer);
+			checkSlow(IndexBuffer);
+
+			if ( GRHISupportsBaseVertexIndex )
+			{
+				RHICmdList.SetStreamSource(0, VertexBuffer->VertexBufferRHI, sizeof(FSlateVertex), 0);
+			}
+		}
+
 		const FSlateShaderResource* ShaderResource = RenderBatch.Texture;
 
 		const ESlateBatchDrawFlag::Type DrawFlags = RenderBatch.DrawFlags;
@@ -499,12 +610,12 @@ void FSlateRHIRenderingPolicy::DrawElements(FRHICommandListImmediate& RHICmdList
 				// for RHIs that can't handle VertexOffset, we need to offset the stream source each time
 				if (!GRHISupportsBaseVertexIndex)
 				{
-					RHICmdList.SetStreamSource(0, VertexBuffer.VertexBufferRHI, sizeof(FSlateVertex), RenderBatch.VertexOffset * sizeof(FSlateVertex));
-					RHICmdList.DrawIndexedPrimitive(IndexBuffer.IndexBufferRHI, GetRHIPrimitiveType(RenderBatch.DrawPrimitiveType), 0, 0, RenderBatch.NumVertices, RenderBatch.IndexOffset, PrimitiveCount, 1);
+					RHICmdList.SetStreamSource(0, VertexBuffer->VertexBufferRHI, sizeof(FSlateVertex), RenderBatch.VertexOffset * sizeof(FSlateVertex));
+					RHICmdList.DrawIndexedPrimitive(IndexBuffer->IndexBufferRHI, GetRHIPrimitiveType(RenderBatch.DrawPrimitiveType), 0, 0, RenderBatch.NumVertices, RenderBatch.IndexOffset, PrimitiveCount, 1);
 				}
 				else
 				{
-					RHICmdList.DrawIndexedPrimitive(IndexBuffer.IndexBufferRHI, GetRHIPrimitiveType(RenderBatch.DrawPrimitiveType), RenderBatch.VertexOffset, 0, RenderBatch.NumVertices, RenderBatch.IndexOffset, PrimitiveCount, 1);
+					RHICmdList.DrawIndexedPrimitive(IndexBuffer->IndexBufferRHI, GetRHIPrimitiveType(RenderBatch.DrawPrimitiveType), RenderBatch.VertexOffset, 0, RenderBatch.NumVertices, RenderBatch.IndexOffset, PrimitiveCount, 1);
 				}
 
 			}
@@ -556,12 +667,12 @@ void FSlateRHIRenderingPolicy::DrawElements(FRHICommandListImmediate& RHICmdList
 					// for RHIs that can't handle VertexOffset, we need to offset the stream source each time
 					if (!GRHISupportsBaseVertexIndex)
 					{
-						RHICmdList.SetStreamSource(0, VertexBuffer.VertexBufferRHI, sizeof(FSlateVertex), RenderBatch.VertexOffset * sizeof(FSlateVertex));
-						RHICmdList.DrawIndexedPrimitive(IndexBuffer.IndexBufferRHI, GetRHIPrimitiveType(RenderBatch.DrawPrimitiveType), 0, 0, RenderBatch.NumVertices, RenderBatch.IndexOffset, PrimitiveCount, 1);
+						RHICmdList.SetStreamSource(0, VertexBuffer->VertexBufferRHI, sizeof(FSlateVertex), RenderBatch.VertexOffset * sizeof(FSlateVertex));
+						RHICmdList.DrawIndexedPrimitive(IndexBuffer->IndexBufferRHI, GetRHIPrimitiveType(RenderBatch.DrawPrimitiveType), 0, 0, RenderBatch.NumVertices, RenderBatch.IndexOffset, PrimitiveCount, 1);
 					}
 					else
 					{
-						RHICmdList.DrawIndexedPrimitive(IndexBuffer.IndexBufferRHI, GetRHIPrimitiveType(RenderBatch.DrawPrimitiveType), RenderBatch.VertexOffset, 0, RenderBatch.NumVertices, RenderBatch.IndexOffset, PrimitiveCount, 1);
+						RHICmdList.DrawIndexedPrimitive(IndexBuffer->IndexBufferRHI, GetRHIPrimitiveType(RenderBatch.DrawPrimitiveType), RenderBatch.VertexOffset, 0, RenderBatch.NumVertices, RenderBatch.IndexOffset, PrimitiveCount, 1);
 					}
 				}
 			}
@@ -577,7 +688,7 @@ void FSlateRHIRenderingPolicy::DrawElements(FRHICommandListImmediate& RHICmdList
 
 				// Something may have messed with the viewport size so set it back to the full target.
 				RHICmdList.SetViewport( 0,0,0,BackBuffer.GetSizeXY().X, BackBuffer.GetSizeXY().Y, 0.0f ); 
-				RHICmdList.SetStreamSource(0, VertexBuffer.VertexBufferRHI, sizeof(FSlateVertex), 0);
+				RHICmdList.SetStreamSource(0, VertexBuffer->VertexBufferRHI, sizeof(FSlateVertex), 0);
 
 			}
 		}
diff --git a/Engine/Source/Runtime/SlateRHIRenderer/Private/SlateRHIRenderingPolicy.h b/Engine/Source/Runtime/SlateRHIRenderer/Private/SlateRHIRenderingPolicy.h
--- a/Engine/Source/Runtime/SlateRHIRenderer/Private/SlateRHIRenderingPolicy.h
+++ b/Engine/Source/Runtime/SlateRHIRenderer/Private/SlateRHIRenderingPolicy.h
@@ -32,12 +32,26 @@ public:
 	{
 		MinBufferSize = sizeof(VertexType) * FMath::Max( MinNumVertices, 200 );
 
-		BeginInitResource(this);
+		if ( IsInRenderingThread() )
+		{
+			InitResource();
+		}
+		else
+		{
+			BeginInitResource(this);
+		}
 	}
 
 	void Destroy()
 	{
-		BeginReleaseResource(this);
+		if ( IsInRenderingThread() )
+		{
+			ReleaseResource();
+		}
+		else
+		{
+			BeginReleaseResource(this);
+		}
 	}
 
 	/** Initializes the vertex buffers RHI resource. */
@@ -232,6 +246,9 @@ public:
 	~FSlateRHIRenderingPolicy();
 
 	void UpdateVertexAndIndexBuffers(FRHICommandListImmediate& RHICmdList, FSlateBatchData& BatchData);
+	void UpdateVertexAndIndexBuffers(FRHICommandListImmediate& RHICmdList, FSlateBatchData& BatchData, TSharedRef<FSlateRenderDataHandle, ESPMode::ThreadSafe> RenderHandle);
+
+	void ReleaseCachedRenderData(FSlateRenderDataHandle* InRenderHandle);
 
 	virtual void DrawElements(FRHICommandListImmediate& RHICmdList, class FSlateBackBuffer& BackBuffer, const FMatrix& ViewProjectionMatrix, const TArray<FSlateRenderBatch>& RenderBatches, bool bAllowSwtichVerticalAxis=true);
 
@@ -246,6 +263,9 @@ public:
 	void EndDrawingWindows();
 
 	void SetUseGammaCorrection( bool bInUseGammaCorrection ) { bGammaCorrect = bInUseGammaCorrection; }
+protected:
+	void UpdateVertexAndIndexBuffers(FRHICommandListImmediate& RHICmdList, FSlateBatchData& BatchData, TSlateElementVertexBuffer<FSlateVertex>& VertexBuffer, FSlateElementIndexBuffer& IndexBuffer);
+
 private:
 	/**
 	 * Returns the pixel shader that should be used for the specified ShaderType and DrawEffects
@@ -264,6 +284,16 @@ private:
 	TSlateElementVertexBuffer<FSlateVertex> VertexBuffers[SlateRHIConstants::NumBuffers];
 	FSlateElementIndexBuffer IndexBuffers[SlateRHIConstants::NumBuffers];
 
+	struct FCachedRenderBuffers
+	{
+		TSlateElementVertexBuffer<FSlateVertex> VertexBuffer;
+		FSlateElementIndexBuffer IndexBuffer;
+	};
+
+	typedef TMap< FSlateRenderDataHandle*, FCachedRenderBuffers* > TCachedBufferMap;
+	TCachedBufferMap CachedBuffers;
+	TArray< FCachedRenderBuffers* > CachedBufferPool;
+
 	TSharedRef<FSlateRHIResourceManager> ResourceManager;
 	TSharedPtr<FSlateFontCache> FontCache;
 
diff --git a/Engine/Source/Runtime/UMG/Private/UserWidget.cpp b/Engine/Source/Runtime/UMG/Private/UserWidget.cpp
--- a/Engine/Source/Runtime/UMG/Private/UserWidget.cpp
+++ b/Engine/Source/Runtime/UMG/Private/UserWidget.cpp
@@ -166,6 +166,11 @@ void UUserWidget::PostInitProperties()
 
 UWorld* UUserWidget::GetWorld() const
 {
+	if ( UWorld* LastWorld = CachedWorld.Get() )
+	{
+		return LastWorld;
+	}
+
 	if ( HasAllFlags(RF_ClassDefaultObject) )
 	{
 		// If we are a CDO, we must return nullptr instead of calling Outer->GetWorld() to fool UObject::ImplementsGetWorld.
@@ -178,6 +183,7 @@ UWorld* UUserWidget::GetWorld() const
 	{
 		if ( UWorld* World = PlayerContext.GetWorld() )
 		{
+			CachedWorld = World;
 			return World;
 		}
 	}
@@ -191,6 +197,7 @@ UWorld* UUserWidget::GetWorld() const
 		UWorld* World = Outer->GetWorld();
 		if ( World )
 		{
+			CachedWorld = World;
 			return World;
 		}
 
diff --git a/Engine/Source/Runtime/UMG/Public/Blueprint/UserWidget.h b/Engine/Source/Runtime/UMG/Public/Blueprint/UserWidget.h
--- a/Engine/Source/Runtime/UMG/Public/Blueprint/UserWidget.h
+++ b/Engine/Source/Runtime/UMG/Public/Blueprint/UserWidget.h
@@ -920,6 +920,9 @@ private:
 
 	FLocalPlayerContext PlayerContext;
 
+	/** Get World calls can be expensive for Widgets, we speed them up by caching the last found world until it goes away. */
+	mutable TWeakObjectPtr<UWorld> CachedWorld;
+
 	/** Has this widget been initialized by its class yet? */
 	bool bInitialized;
 };
