UMG / Slate - Now reporting cached materials and textures to prevent garbage collection of cached resources.  Additionally moving the resource proxy creation to inside of the Make* functions in order to have a better stack trace when a junk GC'ed resource is accessed.

diff --git a/Engine/Source/Developer/StandaloneRenderer/Private/OpenGL/SlateOpenGLRenderer.cpp b/Engine/Source/Developer/StandaloneRenderer/Private/OpenGL/SlateOpenGLRenderer.cpp
--- a/Engine/Source/Developer/StandaloneRenderer/Private/OpenGL/SlateOpenGLRenderer.cpp
+++ b/Engine/Source/Developer/StandaloneRenderer/Private/OpenGL/SlateOpenGLRenderer.cpp
@@ -65,6 +65,7 @@ void FSlateOpenGLRenderer::Initialize()
 	const uint32 TextureSize = 1024;
 
 	TextureManager = MakeShareable( new FSlateOpenGLTextureManager );
+	FSlateDataPayload::ResourceManager = TextureManager.Get();
 
 	FontCache = MakeShareable( new FSlateFontCache( MakeShareable( new FSlateOpenGLFontAtlasFactory ) ) );
 	FontMeasure = FSlateFontMeasure::Create( FontCache.ToSharedRef() );
diff --git a/Engine/Source/Developer/StandaloneRenderer/Private/Windows/D3D/SlateD3DRenderer.cpp b/Engine/Source/Developer/StandaloneRenderer/Private/Windows/D3D/SlateD3DRenderer.cpp
--- a/Engine/Source/Developer/StandaloneRenderer/Private/Windows/D3D/SlateD3DRenderer.cpp
+++ b/Engine/Source/Developer/StandaloneRenderer/Private/Windows/D3D/SlateD3DRenderer.cpp
@@ -77,8 +77,9 @@ void FSlateD3DRenderer::Initialize()
 {
 	CreateDevice();
 
-
 	TextureManager = MakeShareable( new FSlateD3DTextureManager );
+	FSlateDataPayload::ResourceManager = TextureManager.Get();
+
 	TextureManager->LoadUsedTextures();
 
 	FontCache = MakeShareable( new FSlateFontCache( MakeShareable( new FSlateD3DFontAtlasFactory ) ) );
diff --git a/Engine/Source/Runtime/Slate/Private/Widgets/SInvalidationPanel.cpp b/Engine/Source/Runtime/Slate/Private/Widgets/SInvalidationPanel.cpp
--- a/Engine/Source/Runtime/Slate/Private/Widgets/SInvalidationPanel.cpp
+++ b/Engine/Source/Runtime/Slate/Private/Widgets/SInvalidationPanel.cpp
@@ -64,9 +64,6 @@ SInvalidationPanel::~SInvalidationPanel()
 
 bool SInvalidationPanel::GetCanCache() const
 {
-	//HACK: Disabling invalidation panel until material resource reporting is done.
-	return false;
-
 #if !(UE_BUILD_SHIPPING || UE_BUILD_TEST)
 	return bCanCache && EnableWidgetCaching.GetValueOnGameThread() == 1;
 #else
@@ -144,6 +141,11 @@ FChildren* SInvalidationPanel::GetChildren()
 	}
 }
 
+void SInvalidationPanel::AddReferencedObjects(FReferenceCollector& Collector)
+{
+	Collector.AddReferencedObjects(CachedResources);
+}
+
 void SInvalidationPanel::InvalidateWidget(SWidget* InvalidateWidget)
 {
 	bNeedsCaching = true;
@@ -259,6 +261,22 @@ int32 SInvalidationPanel::OnPaint( const FPaintArgs& Args, const FGeometry& Allo
 				InWidgetStyle,
 				bParentEnabled);
 
+			{
+				const TArray<FSlateDrawElement>& CachedElements = CachedWindowElements->GetDrawElements();
+				const int32 CachedElementCount = CachedElements.Num();
+				for ( int32 Index = 0; Index < CachedElementCount; Index++ )
+				{
+					const FSlateDrawElement& LocalElement = CachedElements[Index];
+					if ( const FSlateBrush* BrushResource = LocalElement.GetDataPayload().BrushResource )
+					{
+						if ( UObject* ResourceObject = BrushResource->GetResourceObject() )
+						{
+							CachedResources.Add(ResourceObject);
+						}
+					}
+				}
+			}
+
 			if ( bCacheRelativeTransforms )
 			{
 				CachedAbsolutePosition = AllottedGeometry.Position;
diff --git a/Engine/Source/Runtime/Slate/Public/Widgets/SInvalidationPanel.h b/Engine/Source/Runtime/Slate/Public/Widgets/SInvalidationPanel.h
--- a/Engine/Source/Runtime/Slate/Public/Widgets/SInvalidationPanel.h
+++ b/Engine/Source/Runtime/Slate/Public/Widgets/SInvalidationPanel.h
@@ -4,7 +4,7 @@
 
 #include "Input/HittestGrid.h"
 
-class SLATE_API SInvalidationPanel : public SCompoundWidget, public ILayoutCache
+class SLATE_API SInvalidationPanel : public SCompoundWidget, public FGCObject, public ILayoutCache
 {
 public:
 	SLATE_BEGIN_ARGS( SInvalidationPanel )
@@ -24,6 +24,8 @@ public:
 
 	FORCEINLINE void InvalidateCache() { bNeedsCaching = true; }
 
+	virtual void AddReferencedObjects(FReferenceCollector& Collector) override;
+
 	// ILayoutCache overrides
 	virtual void InvalidateWidget(SWidget* InvalidateWidget) override;
 	virtual FCachedWidgetNode* CreateCacheNode() const override;
@@ -63,6 +65,8 @@ private:
 	mutable TArray< TSharedPtr< FSlateWindowElementList > > InactiveCachedElementListPool;
 	mutable TSharedPtr< FSlateWindowElementList > CachedWindowElements;
 	mutable TSharedPtr<FSlateRenderDataHandle, ESPMode::ThreadSafe> CachedRenderData;
+
+	mutable TSet<UObject*> CachedResources;
 	
 	mutable FVector2D CachedAbsolutePosition;
 	mutable FVector2D AbsoluteDeltaPosition;
diff --git a/Engine/Source/Runtime/SlateCore/Private/Rendering/DrawElements.cpp b/Engine/Source/Runtime/SlateCore/Private/Rendering/DrawElements.cpp
--- a/Engine/Source/Runtime/SlateCore/Private/Rendering/DrawElements.cpp
+++ b/Engine/Source/Runtime/SlateCore/Private/Rendering/DrawElements.cpp
@@ -3,18 +3,10 @@
 #include "SlateCorePrivatePCH.h"
 #include "ReflectionMetadata.h"
 
-#if !(UE_BUILD_SHIPPING || UE_BUILD_TEST)
-
-/**  */
-TAutoConsoleVariable<int32> DrawVolatileWidgets(
-	TEXT("Slate.DrawVolatileWidgets"),
-	true,
-	TEXT("Whether to draw volatile widgets"));
-
-#endif
-
 DECLARE_CYCLE_STAT(TEXT("FSlateDrawElement::Make Time"), STAT_SlateDrawElementMakeTime, STATGROUP_SlateVerbose);
 
+FSlateShaderResourceManager* FSlateDataPayload::ResourceManager;
+
 void FSlateDrawElement::Init(uint32 InLayer, const FPaintGeometry& PaintGeometry, const FSlateRect& InClippingRect, ESlateDrawEffect::Type InDrawEffects)
 {
 	RenderTransform = PaintGeometry.GetAccumulatedRenderTransform();
@@ -529,33 +521,12 @@ void FSlateWindowElementList::QueueVolatilePainting(const FVolatilePaint& InVola
 
 int32 FSlateWindowElementList::PaintVolatile(FSlateWindowElementList& OutElementList)
 {
-#if !(UE_BUILD_SHIPPING || UE_BUILD_TEST)
-	const bool bDrawVolatile = DrawVolatileWidgets.GetValueOnGameThread() == 1;
-#else
-	const bool bDrawVolatile = true;
-#endif
-
-	const static FName InvalidationPanelName(TEXT("SInvalidationPanelName"));
-
 	int32 MaxLayerId = 0;
 
 	for ( int32 VolatileIndex = 0; VolatileIndex < VolatilePaintList.Num(); ++VolatileIndex )
 	{
 		const TSharedPtr<FVolatilePaint>& Args = VolatilePaintList[VolatileIndex];
 
-		// If we're not drawing volatile elements, just draw invalidation panels.  They're
-		// the exception, as we're likely trying to determine the volatile widget overhead.
-		if ( bDrawVolatile == false )
-		{
-			if ( const SWidget* Widget = Args->GetWidget() )
-			{
-				if ( Widget->GetType() != InvalidationPanelName )
-				{
-					continue;
-				}
-			}
-		}
-
 		OutElementList.BeginLogicalLayer(Args->LayerHandle);
 		MaxLayerId = FMath::Max(MaxLayerId, Args->ExecutePaint(OutElementList));
 		OutElementList.EndLogicalLayer();
diff --git a/Engine/Source/Runtime/SlateCore/Private/Rendering/ElementBatcher.cpp b/Engine/Source/Runtime/SlateCore/Private/Rendering/ElementBatcher.cpp
--- a/Engine/Source/Runtime/SlateCore/Private/Rendering/ElementBatcher.cpp
+++ b/Engine/Source/Runtime/SlateCore/Private/Rendering/ElementBatcher.cpp
@@ -15,6 +15,9 @@ SLATE_DECLARE_CYCLE_COUNTER(GSlateAddElements, "Add Elements");
 SLATE_DECLARE_CYCLE_COUNTER(GSlateFindBatchTime, "FindElementForBatch");
 SLATE_DECLARE_CYCLE_COUNTER(GSlateFillBatchBuffers, "FillBatchBuffers");
 
+DECLARE_DWORD_COUNTER_STAT(TEXT("Elements (Box)"), STAT_SlateNumBoxElements, STATGROUP_Slate);
+DECLARE_DWORD_COUNTER_STAT(TEXT("Elements (Text)"), STAT_SlateNumTextElements, STATGROUP_Slate);
+
 // Super-hacky way of storing the scissor rect so we don't have to change all the FSlateDrawElement APIs for this hacky support.
 SLATECORE_API TOptional<FShortRect> GSlateScissorRect;
 
@@ -70,7 +73,7 @@ void FSlateElementBatcher::AddElements(FSlateWindowElementList& WindowElementLis
 
 	SCOPE_CYCLE_COUNTER(STAT_SlateAddElements);
 
-	NumDrawnBatchesStat = 0;
+	NumDrawnBatchesStat = NumDrawnBoxesStat = NumDrawnTextsStat 0;
 
 	BatchData = &WindowElementList.GetBatchData();
 
@@ -86,6 +89,8 @@ void FSlateElementBatcher::AddElements(FSlateWindowElementList& WindowElementLis
 	BatchData = nullptr;
 
 	SET_DWORD_STAT(STAT_SlateNumPrebatchElements, NumDrawnBatchesStat);
+	SET_DWORD_STAT(STAT_SlateNumBoxElements, NumDrawnBoxesStat);
+	SET_DWORD_STAT(STAT_SlateNumTextElements, NumDrawnTextsStat);
 
 	FPlatformMisc::EndNamedEvent();
 }
@@ -199,7 +204,7 @@ void FSlateElementBatcher::AddElements(FSlateDrawLayer& InDrawLayer)
 					AddViewportElement( DrawElement );
 					break;
 				case FSlateDrawElement::ET_Border:
-					AddBorderElement( DrawElement );
+					//AddBorderElement( DrawElement );
 					break;
 				case FSlateDrawElement::ET_Custom:
 					AddCustomElement( DrawElement );
@@ -283,13 +288,22 @@ FSlateRenderTransform GetBoxRenderTransform(const FSlateDrawElement& DrawElement
 
 void FSlateElementBatcher::AddBoxElement( const FSlateDrawElement& DrawElement )
 {
+	const FSlateDataPayload& InPayload = DrawElement.GetDataPayload();
+	FColor Tint = InPayload.Tint.ToFColor(bSRGBVertexColor);
+
+	if ( Tint.A == 0 )
+	{
+		return;
+	}
+
+	NumDrawnBoxesStat++;
+
 	const FSlateRenderTransform& ElementRenderTransform = DrawElement.GetRenderTransform();
 	const FSlateRenderTransform RenderTransform = GetBoxRenderTransform(DrawElement);
 	//const FVector2D& InPosition = DrawElement.GetPosition();
 	//const FVector2D& Size = DrawElement.GetSize();
 	const FVector2D& LocalSize = DrawElement.GetLocalSize();
 	//float Scale = DrawElement.GetScale();
-	const FSlateDataPayload& InPayload = DrawElement.GetDataPayload();
 	const FSlateRect& InClippingRect = DrawElement.GetClippingRect();
 	ESlateDrawEffect::Type InDrawEffects = DrawElement.GetDrawEffects();
 	uint32 Layer = DrawElement.GetLayer();
@@ -318,7 +332,7 @@ void FSlateElementBatcher::AddBoxElement( const FSlateDrawElement& DrawElement )
 
 		FVector2D HalfTexel;
 
-		FSlateShaderResourceProxy* ResourceProxy = ResourceManager.GetShaderResource(*BrushResource);
+		const FSlateShaderResourceProxy* ResourceProxy = InPayload.ResourceProxy;
 		FSlateShaderResource* Resource = nullptr;
 		if( ResourceProxy )
 		{
@@ -354,8 +368,6 @@ void FSlateElementBatcher::AddBoxElement( const FSlateDrawElement& DrawElement )
 		}
 
 
-		FColor Tint = InPayload.Tint.ToFColor(bSRGBVertexColor);
-
 		const ESlateBrushTileType::Type TilingRule = BrushResource->Tiling;
 		const bool bTileHorizontal = (TilingRule == ESlateBrushTileType::Both || TilingRule == ESlateBrushTileType::Horizontal);
 		const bool bTileVertical = (TilingRule == ESlateBrushTileType::Both || TilingRule == ESlateBrushTileType::Vertical);
@@ -584,11 +596,20 @@ void FSlateElementBatcher::AddBoxElement( const FSlateDrawElement& DrawElement )
 
 void FSlateElementBatcher::AddTextElement(const FSlateDrawElement& DrawElement)
 {
+	const FSlateDataPayload& InPayload = DrawElement.GetDataPayload();
+	FColor Tint = InPayload.Tint.ToFColor(bSRGBVertexColor);
+
+	if ( Tint.A == 0 )
+	{
+		return;
+	}
+
+	NumDrawnTextsStat++;
+
 	//const FVector2D& InPosition = DrawElement.GetPosition();
 	//const FVector2D& Size = DrawElement.GetSize();
 	//const FVector2D& LocalSize = DrawElement.GetLocalSize();
 	//float Scale = DrawElement.GetScale();
-	const FSlateDataPayload& InPayload = DrawElement.GetDataPayload();
 	const FSlateRect& InClippingRect = DrawElement.GetClippingRect();
 	ESlateDrawEffect::Type InDrawEffects = DrawElement.GetDrawEffects();
 	uint32 Layer = DrawElement.GetLayer();
@@ -648,8 +669,6 @@ void FSlateElementBatcher::AddTextElement(const FSlateDrawElement& DrawElement)
 	
 	const bool bIsFontMaterial = InPayload.FontInfo.FontMaterial != nullptr;
 
-	FColor FinalColor = InPayload.Tint.ToFColor(bSRGBVertexColor);
-
 	uint32 NumChars = Text.Len();
 
 	uint32 NumLines = 1;
@@ -755,10 +774,10 @@ void FSlateElementBatcher::AddTextElement(const FSlateDrawElement& DrawElement)
 					}
 
 					// Add four vertices to the list of verts to be added to the vertex buffer
-					BatchVerticesRef[ VertexOffset++ ] = FSlateVertex( RenderTransform, UpperLeft,								FVector4(U,V,				Ut,Vt),			FVector2D(0.0f,0.0f), FinalColor, RenderClipRect );
-					BatchVerticesRef[ VertexOffset++ ] = FSlateVertex( RenderTransform, FVector2D(LowerRight.X,UpperLeft.Y),	FVector4(U+SizeU, V,		UtMax,Vt),		FVector2D(1.0f,0.0f), FinalColor, RenderClipRect );
-					BatchVerticesRef[ VertexOffset++ ] = FSlateVertex( RenderTransform, FVector2D(UpperLeft.X,LowerRight.Y),	FVector4(U, V+SizeV,		Ut,VtMax),		FVector2D(0.0f,1.0f), FinalColor, RenderClipRect );
-					BatchVerticesRef[ VertexOffset++ ] = FSlateVertex( RenderTransform, LowerRight,								FVector4(U+SizeU, V+SizeV,	UtMax,VtMax),	FVector2D(1.0f,1.0f), FinalColor, RenderClipRect );
+					BatchVerticesRef[ VertexOffset++ ] = FSlateVertex( RenderTransform, UpperLeft,								FVector4(U,V,				Ut,Vt),			FVector2D(0.0f,0.0f), Tint, RenderClipRect );
+					BatchVerticesRef[ VertexOffset++ ] = FSlateVertex( RenderTransform, FVector2D(LowerRight.X,UpperLeft.Y),	FVector4(U+SizeU, V,		UtMax,Vt),		FVector2D(1.0f,0.0f), Tint, RenderClipRect );
+					BatchVerticesRef[ VertexOffset++ ] = FSlateVertex( RenderTransform, FVector2D(UpperLeft.X,LowerRight.Y),	FVector4(U, V+SizeV,		Ut,VtMax),		FVector2D(0.0f,1.0f), Tint, RenderClipRect );
+					BatchVerticesRef[ VertexOffset++ ] = FSlateVertex( RenderTransform, LowerRight,								FVector4(U+SizeU, V+SizeV,	UtMax,VtMax),	FVector2D(1.0f,1.0f), Tint, RenderClipRect );
 
 					BatchIndicesRef[IndexOffset++] = IndexStart + 0;
 					BatchIndicesRef[IndexOffset++] = IndexStart + 1;
@@ -1241,30 +1260,24 @@ void FSlateElementBatcher::AddViewportElement( const FSlateDrawElement& DrawElem
 		DrawFlags |= ESlateBatchDrawFlag::NoGamma;
 	}
 
-	TSharedPtr<const ISlateViewport> ViewportPin = InPayload.Viewport.Pin();
-
-	FSlateShaderResource* ViewportResource =  nullptr;
+	FSlateShaderResource* ViewportResource = InPayload.ViewportRenderTargetTexture;
 	ESlateShader::Type ShaderType = ESlateShader::Default;
 
-	if( ViewportPin.IsValid() )
+	if( InPayload.bViewportTextureAlphaOnly )
 	{
-		ViewportResource = ViewportPin->GetViewportRenderTargetTexture();
-
-		if( ViewportPin->IsViewportTextureAlphaOnly() )
-		{
-			// This is a slight hack, but the font shader is the same as the general shader except it reads alpha only textures
-			ShaderType = ESlateShader::Font;
-		}
+		// This is a slight hack, but the font shader is the same as the general shader except it reads alpha only textures
+		ShaderType = ESlateShader::Font;
 	}
 
+
 	FSlateElementBatch& ElementBatch = FindBatchForElement( Layer, FShaderParams(), ViewportResource, ESlateDrawPrimitive::TriangleList, ShaderType, InDrawEffects, DrawFlags, DrawElement.GetScissorRect() );
 	TArray<FSlateVertex>& BatchVertices = BatchData->GetBatchVertexList(ElementBatch);
 	TArray<SlateIndex>& BatchIndices = BatchData->GetBatchIndexList(ElementBatch);
 
 	// Tag this batch as requiring vsync if the viewport requires it.
-	if( ViewportPin.IsValid() )
+	if( ViewportResource != nullptr && !InPayload.bAllowViewportScaling )
 	{
-		bRequiresVsync |= ViewportPin->RequiresVsync();
+		bRequiresVsync |= InPayload.bRequiresVSync;
 	}
 
 	// Do pixel snapping
@@ -1272,7 +1285,7 @@ void FSlateElementBatcher::AddViewportElement( const FSlateDrawElement& DrawElem
 	FVector2D BotRight(LocalSize);
 
 	// If the viewport disallows scaling, force size to current texture size.
-	if (ViewportResource != nullptr && !ViewportPin->AllowScaling())
+	if (ViewportResource != nullptr && !InPayload.bAllowViewportScaling)
 	{
 		BotRight = FVector2D(ViewportResource->GetWidth(), ViewportResource->GetHeight());
 	}
@@ -1325,7 +1338,7 @@ void FSlateElementBatcher::AddBorderElement( const FSlateDrawElement& DrawElemen
 	uint32 TextureHeight = 1;
 
 	// Currently borders are not atlased because they are tiled.  So we just assume the texture proxy holds the actual texture
-	FSlateShaderResourceProxy* ResourceProxy = ResourceManager.GetShaderResource( *InPayload.BrushResource );
+	const FSlateShaderResourceProxy* ResourceProxy = InPayload.ResourceProxy;
 	FSlateShaderResource* Resource = ResourceProxy ? ResourceProxy->Resource : nullptr;
 	if( Resource )
 	{
diff --git a/Engine/Source/Runtime/SlateCore/Public/Rendering/DrawElements.h b/Engine/Source/Runtime/SlateCore/Public/Rendering/DrawElements.h
--- a/Engine/Source/Runtime/SlateCore/Public/Rendering/DrawElements.h
+++ b/Engine/Source/Runtime/SlateCore/Public/Rendering/DrawElements.h
@@ -40,6 +40,7 @@ public:
 
 	// Brush data
 	const FSlateBrush* BrushResource;
+	const FSlateShaderResourceProxy* ResourceProxy;
 
 	// Box Data
 	FVector2D RotationPoint;
@@ -51,7 +52,7 @@ public:
 	// Font data
 	FSlateFontInfo FontInfo;
 	FString Text;
-	
+
 	// Gradient data
 	TArray<FSlateGradientStop> GradientStops;
 	EOrientation GradientType;
@@ -59,8 +60,11 @@ public:
 	// Line data
 	TArray<FVector2D> Points;
 
-	// Viewport data (intentionally weak to allow the source element to be destructed after this element has been added for drawing)
-	TWeakPtr<const ISlateViewport> Viewport;
+	// Viewport data
+	FSlateShaderResource* ViewportRenderTargetTexture;
+	bool bAllowViewportScaling;
+	bool bViewportTextureAlphaOnly;
+	bool bRequiresVSync;
 
 	// Misc data
 	bool bGammaCorrect;
@@ -79,19 +83,26 @@ public:
 	ESlateLineJoinType::Type SegmentJoinType;
 	bool bAntialias;
 
+	SLATECORE_API static FSlateShaderResourceManager* ResourceManager;
+
 	FSlateDataPayload()
 		: Tint(FLinearColor::White)
 		, BrushResource(nullptr)
+		, ResourceProxy(nullptr)
 		, RotationPoint(FVector2D::ZeroVector)
-		, Viewport(nullptr)
-		, CachedRenderData(nullptr)
-		, LayerHandle(nullptr)
+		, ViewportRenderTargetTexture(nullptr)
+		, bViewportTextureAlphaOnly(false)
+		, bRequiresVSync(false)
+		, bGammaCorrect(false)
+		, bAllowBlending(false)
+		, CustomDrawer()
 	{ }
 
 	void SetBoxPayloadProperties( const FSlateBrush* InBrush, const FLinearColor& InTint )
 	{
 		Tint = InTint;
 		BrushResource = InBrush;
+		ResourceProxy = ResourceManager->GetShaderResource(*InBrush);
 		Angle = 0.0f;
 	}
 
@@ -99,6 +110,7 @@ public:
 	{
 		Tint = InTint;
 		BrushResource = InBrush;
+		ResourceProxy = ResourceManager->GetShaderResource(*InBrush);
 		RotationPoint = LocalRotationPoint;
 		RotationPoint.DiagnosticCheckNaN();
 		Angle = InAngle;
@@ -141,7 +153,10 @@ public:
 	void SetViewportPayloadProperties( const TSharedPtr<const ISlateViewport>& InViewport, const FLinearColor& InTint, bool bInGammaCorrect, bool bInAllowBlending )
 	{
 		Tint = InTint;
-		Viewport = InViewport;
+		ViewportRenderTargetTexture = InViewport->GetViewportRenderTargetTexture();
+		bAllowViewportScaling = InViewport->AllowScaling();
+		bViewportTextureAlphaOnly = InViewport->IsViewportTextureAlphaOnly();
+		bRequiresVSync = InViewport->RequiresVsync();
 		bGammaCorrect = bInGammaCorrect;
 		bAllowBlending = bInAllowBlending;
 	}
diff --git a/Engine/Source/Runtime/SlateCore/Public/Rendering/ElementBatcher.h b/Engine/Source/Runtime/SlateCore/Public/Rendering/ElementBatcher.h
--- a/Engine/Source/Runtime/SlateCore/Public/Rendering/ElementBatcher.h
+++ b/Engine/Source/Runtime/SlateCore/Public/Rendering/ElementBatcher.h
@@ -126,6 +126,12 @@ private:
 	/** Track the number of drawn batches from the previous frame to report to stats. */
 	int32 NumDrawnBatchesStat;
 
+	/** Track the number of drawn boxes from the previous frame to report to stats. */
+	int32 NumDrawnBoxesStat;
+
+	/** Track the number of drawn texts from the previous frame to report to stats. */
+	int32 NumDrawnTextsStat;
+
 	/** Offset to use when supporting 1:1 texture to pixel snapping */
 	const float PixelCenterOffset;
 
diff --git a/Engine/Source/Runtime/SlateRHIRenderer/Private/SlateRHIRendererModule.cpp b/Engine/Source/Runtime/SlateRHIRenderer/Private/SlateRHIRendererModule.cpp
--- a/Engine/Source/Runtime/SlateRHIRenderer/Private/SlateRHIRendererModule.cpp
+++ b/Engine/Source/Runtime/SlateRHIRenderer/Private/SlateRHIRendererModule.cpp
@@ -81,6 +81,7 @@ private:
 		if( !ResourceManager.IsValid() )
 		{
 			ResourceManager = MakeShareable( new FSlateRHIResourceManager );
+			FSlateDataPayload::ResourceManager = ResourceManager.Get();
 		}
 
 		if( !FontCache.IsValid() )
