Build error fix.

diff --git a/Engine/Source/Runtime/SlateCore/Private/Rendering/ElementBatcher.cpp b/Engine/Source/Runtime/SlateCore/Private/Rendering/ElementBatcher.cpp
--- a/Engine/Source/Runtime/SlateCore/Private/Rendering/ElementBatcher.cpp
+++ b/Engine/Source/Runtime/SlateCore/Private/Rendering/ElementBatcher.cpp
@@ -73,7 +73,13 @@ void FSlateElementBatcher::AddElements(FSlateWindowElementList& WindowElementLis
 
 	SCOPE_CYCLE_COUNTER(STAT_SlateAddElements);
 
+>>>> ORIGINAL //depot/UE4-Orion/Engine/Source/Runtime/SlateCore/Private/Rendering/ElementBatcher.cpp#12
+	NumDrawnBatchesStat = NumDrawnBoxesStat = NumDrawnTextsStat;
+==== THEIRS //depot/UE4-Orion/Engine/Source/Runtime/SlateCore/Private/Rendering/ElementBatcher.cpp#13
+	NumDrawnBatchesStat = NumDrawnBoxesStat = NumDrawnTextsStat = 0;
+==== YOURS //Nick.Darnell_Dev/UE4/Engine/Source/Runtime/SlateCore/Private/Rendering/ElementBatcher.cpp
 	NumDrawnBatchesStat = NumDrawnBoxesStat = NumDrawnTextsStat 0;
+<<<<
 
 	BatchData = &WindowElementList.GetBatchData();
 
