Added a way to globally invalidate all invalidation panels.
This is expensive and should only be used when there is some major change to a resource that could be held onto by the invalidation panels

Merged from Orion CL 2659385

diff --git a/Engine/Source/Runtime/Slate/Private/Widgets/SInvalidationPanel.cpp b/Engine/Source/Runtime/Slate/Private/Widgets/SInvalidationPanel.cpp
--- a/Engine/Source/Runtime/Slate/Private/Widgets/SInvalidationPanel.cpp
+++ b/Engine/Source/Runtime/Slate/Private/Widgets/SInvalidationPanel.cpp
@@ -39,6 +39,8 @@ TAutoConsoleVariable<int32> EnableWidgetCaching(
 
 void SInvalidationPanel::Construct( const FArguments& InArgs )
 {
+	FSlateApplicationBase::Get().OnGlobalInvalidate().AddSP( this, &SInvalidationPanel::OnGlobalInvalidate );
+
 	ChildSlot
 	[
 		InArgs._Content.Widget
@@ -217,6 +219,11 @@ TSharedPtr< FSlateWindowElementList > SInvalidationPanel::GetNextCachedElementLi
 	return NextElementList;
 }
 
+void SInvalidationPanel::OnGlobalInvalidate()
+{
+	InvalidateCache();
+}
+
 int32 SInvalidationPanel::OnPaint( const FPaintArgs& Args, const FGeometry& AllottedGeometry, const FSlateRect& MyClippingRect, FSlateWindowElementList& OutDrawElements, int32 LayerId, const FWidgetStyle& InWidgetStyle, bool bParentEnabled ) const
 {
 	if ( GetCanCache() )
diff --git a/Engine/Source/Runtime/Slate/Public/Widgets/SInvalidationPanel.h b/Engine/Source/Runtime/Slate/Public/Widgets/SInvalidationPanel.h
--- a/Engine/Source/Runtime/Slate/Public/Widgets/SInvalidationPanel.h
+++ b/Engine/Source/Runtime/Slate/Public/Widgets/SInvalidationPanel.h
@@ -49,7 +49,7 @@ public:
 
 private:
 	TSharedPtr< FSlateWindowElementList > GetNextCachedElementList(const TSharedPtr<SWindow>& CurrentWindow) const;
-
+	void OnGlobalInvalidate();
 private:
 	FGeometry LastAllottedGeometry;
 
diff --git a/Engine/Source/Runtime/SlateCore/Private/Application/SlateApplicationBase.cpp b/Engine/Source/Runtime/SlateCore/Private/Application/SlateApplicationBase.cpp
--- a/Engine/Source/Runtime/SlateCore/Private/Application/SlateApplicationBase.cpp
+++ b/Engine/Source/Runtime/SlateCore/Private/Application/SlateApplicationBase.cpp
@@ -82,3 +82,8 @@ bool FSlateApplicationBase::IsSlateAsleep()
 {
 	return bIsSlateAsleep;
 }
+
+void FSlateApplicationBase::InvalidateAllWidgets() const
+{
+	OnGlobalInvalidateEvent.Broadcast();
+}
\ No newline at end of file
diff --git a/Engine/Source/Runtime/SlateCore/Private/Fonts/FontCache.cpp b/Engine/Source/Runtime/SlateCore/Private/Fonts/FontCache.cpp
--- a/Engine/Source/Runtime/SlateCore/Private/Fonts/FontCache.cpp
+++ b/Engine/Source/Runtime/SlateCore/Private/Fonts/FontCache.cpp
@@ -1490,6 +1490,9 @@ void FSlateFontCache::ReleaseResources()
 
 void FSlateFontCache::FlushCache() const
 {
+	// Ensure all invalidation panels are cleared of cached widgets
+	FSlateApplicationBase::Get().InvalidateAllWidgets();
+
 	FontToCharacterListCache.Empty();
 	FTInterface->Flush();
 
diff --git a/Engine/Source/Runtime/SlateCore/Public/Application/SlateApplicationBase.h b/Engine/Source/Runtime/SlateCore/Public/Application/SlateApplicationBase.h
--- a/Engine/Source/Runtime/SlateCore/Public/Application/SlateApplicationBase.h
+++ b/Engine/Source/Runtime/SlateCore/Public/Application/SlateApplicationBase.h
@@ -345,6 +345,17 @@ public:
 	 */
 	virtual void SetAllUserFocus(const FWidgetPath& InFocusPath, const EFocusCause InCause) = 0;
 
+	/**
+	 * Gets a delegate that is invoked when a global invalidate of all widgets should occur
+	 */
+	DECLARE_EVENT(FSlateApplicationBase, FOnGlobalInvalidate);
+	FOnGlobalInvalidate& OnGlobalInvalidate()  { return OnGlobalInvalidateEvent; }
+
+	/**
+	 * Notifies all invalidation panels that they should invalidate their contents
+	 * Note: this is a very expensive call and should only be done in non-performance critical situations
+	 */
+	void InvalidateAllWidgets() const;
 private:
 	/**
 	 * Implementation for active timer registration. See SWidget::RegisterActiveTimer.
@@ -464,6 +475,8 @@ public:
 	}
 
 protected:
+	/** multicast delegate to broadcast when a global invalidate is requested */
+	FOnGlobalInvalidate OnGlobalInvalidateEvent;
 
 	// Gets set when Slate goes to sleep and cleared when active.
 	bool bIsSlateAsleep;
