Engine - Build fix defining WITH_SLATE_STATS if undefined.

diff --git a/Engine/Source/Runtime/SlateCore/Public/Stats/SlateStats.h b/Engine/Source/Runtime/SlateCore/Public/Stats/SlateStats.h
--- a/Engine/Source/Runtime/SlateCore/Public/Stats/SlateStats.h
+++ b/Engine/Source/Runtime/SlateCore/Public/Stats/SlateStats.h
@@ -4,6 +4,11 @@
 
 /** Set to zero to disable the slate stats system completely, meaning all stats-related structures and code are compiled out. */
 //#define SLATE_STATS (!(UE_BUILD_SHIPPING || UE_BUILD_TEST) || WITH_EDITOR)
+
+#ifndef WITH_SLATE_STATS
+	#define WITH_SLATE_STATS 0
+#endif
+
 #define SLATE_STATS (WITH_SLATE_STATS || WITH_EDITOR)
 
 /** 
