Slate: Removed an unnecessary pass-by-value in FindPathToWidget
[INTEGRATE] Change 2667195 by [REDACTED_PERSON][REDACTED_HANDLE] on 2015/08/25 02:10:52

diff --git a/Engine/Source/Runtime/SlateCore/Private/Application/SlateWindowHelper.cpp b/Engine/Source/Runtime/SlateCore/Private/Application/SlateWindowHelper.cpp
--- a/Engine/Source/Runtime/SlateCore/Private/Application/SlateWindowHelper.cpp
+++ b/Engine/Source/Runtime/SlateCore/Private/Application/SlateWindowHelper.cpp
@@ -88,7 +88,7 @@ bool FSlateWindowHelper::ContainsWindow( const TArray<TSharedRef<SWindow>>& Wind
 }
 
 
-bool FSlateWindowHelper::FindPathToWidget( const TArray<TSharedRef<SWindow>> WindowsToSearch,  TSharedRef<const SWidget> InWidget, FWidgetPath& OutWidgetPath, EVisibility VisibilityFilter )
+bool FSlateWindowHelper::FindPathToWidget( const TArray<TSharedRef<SWindow>>& WindowsToSearch,  TSharedRef<const SWidget> InWidget, FWidgetPath& OutWidgetPath, EVisibility VisibilityFilter )
 {
 	SCOPE_CYCLE_COUNTER(STAT_FindPathToWidget);
 
diff --git a/Engine/Source/Runtime/SlateCore/Public/Application/SlateWindowHelper.h b/Engine/Source/Runtime/SlateCore/Public/Application/SlateWindowHelper.h
--- a/Engine/Source/Runtime/SlateCore/Public/Application/SlateWindowHelper.h
+++ b/Engine/Source/Runtime/SlateCore/Public/Application/SlateWindowHelper.h
@@ -56,7 +56,7 @@ public:
 	 * @param VisibilityFilter Widgets must have this type of visibility to be included the path.
 	 * @return True if the widget path was found.
 	 */
-	static bool FindPathToWidget( const TArray<TSharedRef<SWindow>> WindowsToSearch, TSharedRef<const SWidget> InWidget, FWidgetPath& OutWidgetPath, EVisibility VisibilityFilter = EVisibility::Visible );
+	static bool FindPathToWidget( const TArray<TSharedRef<SWindow>>& WindowsToSearch, TSharedRef<const SWidget> InWidget, FWidgetPath& OutWidgetPath, EVisibility VisibilityFilter = EVisibility::Visible );
 
 	/**
 	 * Searches the given collection of windows to find the Slate window that corresponds to the specified platform window.
