UE4 - merged fix for slate creating stats before main starts

diff --git a/Engine/Source/Runtime/SlateCore/Private/Widgets/SWidget.cpp b/Engine/Source/Runtime/SlateCore/Private/Widgets/SWidget.cpp
--- a/Engine/Source/Runtime/SlateCore/Private/Widgets/SWidget.cpp
+++ b/Engine/Source/Runtime/SlateCore/Private/Widgets/SWidget.cpp
@@ -38,7 +38,10 @@ SWidget::SWidget()
 	, bCachedVolatile(false)
 	, bInheritedVolatility(false)
 {
-	INC_DWORD_STAT(STAT_SlateTotalWidgets);
+	if (GIsRunning)
+	{
+		INC_DWORD_STAT(STAT_SlateTotalWidgets);
+	}
 }
 
 SWidget::~SWidget()
