Slate: Changed a local reserve in HittestGrid::GetBubblePath to use an inline allocation
[INTEGRATE] Change 2668351 by [REDACTED_PERSON][REDACTED_HANDLE] on 2015/08/25 16:46:43

diff --git a/Engine/Source/Runtime/SlateCore/Private/Input/HittestGrid.cpp b/Engine/Source/Runtime/SlateCore/Private/Input/HittestGrid.cpp
--- a/Engine/Source/Runtime/SlateCore/Private/Input/HittestGrid.cpp
+++ b/Engine/Source/Runtime/SlateCore/Private/Input/HittestGrid.cpp
@@ -159,8 +159,7 @@ TArray<FWidgetAndPointer> FHittestGrid::GetBubblePath(FVector2D DesktopSpaceCoor
 		const FIntPoint LRIndex = GetCellCoordinate(CursorPositionInGrid + RadiusVector);
 
 		//first, find all the overlapping cells
-		TArray<FIntPoint> CellIndexes;
-		CellIndexes.Reserve(16);
+		TArray<FIntPoint, TInlineAllocator<16>> CellIndexes;
 
 		for (int32 YIndex = ULIndex.Y; YIndex <= LRIndex.Y; ++YIndex)
 		{
