UE4 - Fixed slate vertex and index recycling to not grow each batch to be the size of the largest one.

diff --git a/Engine/Source/Runtime/SlateCore/Private/Rendering/DrawElements.cpp b/Engine/Source/Runtime/SlateCore/Private/Rendering/DrawElements.cpp
--- a/Engine/Source/Runtime/SlateCore/Private/Rendering/DrawElements.cpp
+++ b/Engine/Source/Runtime/SlateCore/Private/Rendering/DrawElements.cpp
@@ -227,37 +227,39 @@ void FSlateBatchData::Reset()
 	RenderDataHandle.Reset();
 }
 
+#define MAX_VERT_ARRAY_RECYCLE (200)
+
 void FSlateBatchData::AssignVertexArrayToBatch( FSlateElementBatch& Batch )
 {
 	// Get a free vertex array
 	if (VertexArrayFreeList.Num() > 0)
 	{
 		Batch.VertexArrayIndex = VertexArrayFreeList.Pop(/*bAllowShrinking=*/ false);
-		BatchVertexArrays[Batch.VertexArrayIndex].Reserve(200);
 	}
 	else
 	{
 		// There are no free vertex arrays so we must add one		
 		uint32 NewIndex = BatchVertexArrays.Add(TArray<FSlateVertex>());
-		BatchVertexArrays[NewIndex].Reserve(200);
+		BatchVertexArrays[NewIndex].Reserve(MAX_VERT_ARRAY_RECYCLE);
 
 		Batch.VertexArrayIndex = NewIndex;
 	}
 }
 
+#define MAX_INDEX_ARRAY_RECYCLE (500)
+
 void FSlateBatchData::AssignIndexArrayToBatch( FSlateElementBatch& Batch )
 {
 	// Get a free index array
 	if (IndexArrayFreeList.Num() > 0)
 	{
 		Batch.IndexArrayIndex = IndexArrayFreeList.Pop(/*bAllowShrinking=*/ false);
-		BatchIndexArrays[Batch.IndexArrayIndex].Reserve(200);
 	}
 	else
 	{
 		// There are no free index arrays so we must add one
 		uint32 NewIndex = BatchIndexArrays.Add(TArray<SlateIndex>());
-		BatchIndexArrays[NewIndex].Reserve(500);
+		BatchIndexArrays[NewIndex].Reserve(MAX_INDEX_ARRAY_RECYCLE);
 
 		Batch.IndexArrayIndex = NewIndex;
 	}
@@ -297,6 +299,17 @@ void FSlateBatchData::FillVertexAndIndexBuffer( uint8* VertexBuffer, uint8* Inde
 
 				Vertices.Reset();
 				Indices.Reset();
+
+				if (Vertices.GetSlack() > MAX_VERT_ARRAY_RECYCLE)
+				{
+					Vertices.Empty();
+					Vertices.Reserve(MAX_VERT_ARRAY_RECYCLE);
+				}
+				if (Indices.GetSlack() > MAX_INDEX_ARRAY_RECYCLE)
+				{
+					Indices.Empty();
+					Indices.Reserve(MAX_INDEX_ARRAY_RECYCLE);
+				}
 			}
 		}
 	}
