PR #1525: Update SlateDynamicImageBrush.cpp (Contributed by [REDACTED_PERSON])

diff --git a/Engine/Source/Runtime/SlateCore/Private/Brushes/SlateDynamicImageBrush.cpp b/Engine/Source/Runtime/SlateCore/Private/Brushes/SlateDynamicImageBrush.cpp
--- a/Engine/Source/Runtime/SlateCore/Private/Brushes/SlateDynamicImageBrush.cpp
+++ b/Engine/Source/Runtime/SlateCore/Private/Brushes/SlateDynamicImageBrush.cpp
@@ -31,11 +31,7 @@ TSharedPtr<FSlateDynamicImageBrush> FSlateDynamicImageBrush::CreateWithImageData
 
 void FSlateDynamicImageBrush::ReleaseResource()
 {
-	if( FSlateApplicationBase::IsInitialized() )
-	{
-		// Safely remove the brush we were using. 
-		FSlateApplicationBase::Get().GetRenderer()->RemoveDynamicBrushResource( AsShared() );
-	}	
+	ReleaseResourceInternal();	
 }
 
 void FSlateDynamicImageBrush::ReleaseResourceInternal()
