UE4 - Fixed malloc churn problem with FHittestGrid::FCachedWidget

diff --git a/Engine/Source/Runtime/SlateCore/Private/Input/HittestGrid.cpp b/Engine/Source/Runtime/SlateCore/Private/Input/HittestGrid.cpp
--- a/Engine/Source/Runtime/SlateCore/Private/Input/HittestGrid.cpp
+++ b/Engine/Source/Runtime/SlateCore/Private/Input/HittestGrid.cpp
@@ -117,7 +117,7 @@ struct FHittestGrid::FCachedWidget
 	FGeometry CachedGeometry;
 	// @todo umg : ideally this clipping rect is optional and we only have them on a small number of widgets.
 	FSlateRect ClippingRect;
-	TArray<int32> Children;
+	TArray<int32, TInlineAllocator<16> > Children;
 	int32 ParentIndex;
 };
 
