Merging CL 2684004 using UE4_to_UE4-Orion

Implemented a basic slate resource handle for fast access without doing map lookups from Brush -> Rendering resource.  Not hooked up to anything yet

diff --git a/Engine/Source/Developer/StandaloneRenderer/Private/OpenGL/SlateOpenGLRenderer.cpp b/Engine/Source/Developer/StandaloneRenderer/Private/OpenGL/SlateOpenGLRenderer.cpp
--- a/Engine/Source/Developer/StandaloneRenderer/Private/OpenGL/SlateOpenGLRenderer.cpp
+++ b/Engine/Source/Developer/StandaloneRenderer/Private/OpenGL/SlateOpenGLRenderer.cpp
@@ -222,6 +222,11 @@ bool FSlateOpenGLRenderer::GenerateDynamicImageResource(FName ResourceName, uint
 	return TextureManager->CreateDynamicTextureResource(ResourceName, Width, Height, Bytes) != NULL;
 }
 
+FSlateResourceHandle FSlateOpenGLRenderer::GetResourceHandle( const FSlateBrush& Brush )
+{
+	return TextureManager->GetResourceHandle( Brush );
+}
+
 void FSlateOpenGLRenderer::RemoveDynamicBrushResource( TSharedPtr<FSlateDynamicImageBrush> BrushToRemove )
 {
 	DynamicBrushesToRemove.Add( BrushToRemove );
diff --git a/Engine/Source/Developer/StandaloneRenderer/Private/OpenGL/SlateOpenGLRenderer.h b/Engine/Source/Developer/StandaloneRenderer/Private/OpenGL/SlateOpenGLRenderer.h
--- a/Engine/Source/Developer/StandaloneRenderer/Private/OpenGL/SlateOpenGLRenderer.h
+++ b/Engine/Source/Developer/StandaloneRenderer/Private/OpenGL/SlateOpenGLRenderer.h
@@ -146,6 +146,7 @@ public:
 	virtual void RestoreSystemResolution(const TSharedRef<SWindow> InWindow) override {}
 	virtual void ReleaseDynamicResource( const FSlateBrush& Brush ) override;
 	virtual bool GenerateDynamicImageResource(FName ResourceName, uint32 Width, uint32 Height, const TArray< uint8 >& Bytes) override;
+	virtual FSlateResourceHandle GetResourceHandle( const FSlateBrush& Brush ) override;
 	virtual void RemoveDynamicBrushResource( TSharedPtr<FSlateDynamicImageBrush> BrushToRemove ) override;
 	virtual void LoadStyleResources(const ISlateStyle& Style) override;
 	virtual FSlateUpdatableTexture* CreateUpdatableTexture(uint32 Width, uint32 Height) override;
diff --git a/Engine/Source/Developer/StandaloneRenderer/Private/Windows/D3D/SlateD3DRenderer.cpp b/Engine/Source/Developer/StandaloneRenderer/Private/Windows/D3D/SlateD3DRenderer.cpp
--- a/Engine/Source/Developer/StandaloneRenderer/Private/Windows/D3D/SlateD3DRenderer.cpp
+++ b/Engine/Source/Developer/StandaloneRenderer/Private/Windows/D3D/SlateD3DRenderer.cpp
@@ -244,6 +244,11 @@ bool FSlateD3DRenderer::GenerateDynamicImageResource(FName ResourceName, uint32
 	return TextureManager->CreateDynamicTextureResource(ResourceName, Width, Height, Bytes) != NULL;
 }
 
+FSlateResourceHandle FSlateD3DRenderer::GetResourceHandle( const FSlateBrush& Brush )
+{
+	return TextureManager->GetResourceHandle( Brush );
+}
+
 void FSlateD3DRenderer::RemoveDynamicBrushResource( TSharedPtr<FSlateDynamicImageBrush> BrushToRemove )
 {
 	DynamicBrushesToRemove.Add( BrushToRemove );
diff --git a/Engine/Source/Developer/StandaloneRenderer/Private/Windows/D3D/SlateD3DRenderer.h b/Engine/Source/Developer/StandaloneRenderer/Private/Windows/D3D/SlateD3DRenderer.h
--- a/Engine/Source/Developer/StandaloneRenderer/Private/Windows/D3D/SlateD3DRenderer.h
+++ b/Engine/Source/Developer/StandaloneRenderer/Private/Windows/D3D/SlateD3DRenderer.h
@@ -55,6 +55,7 @@ public:
 	virtual void RestoreSystemResolution(const TSharedRef<SWindow> InWindow) override {}
 	virtual void ReleaseDynamicResource( const FSlateBrush& Brush ) override;
 	virtual bool GenerateDynamicImageResource(FName ResourceName, uint32 Width, uint32 Height, const TArray< uint8 >& Bytes) override;
+	virtual FSlateResourceHandle GetResourceHandle( const FSlateBrush& Brush ) override;
 	virtual void RemoveDynamicBrushResource( TSharedPtr<FSlateDynamicImageBrush> BrushToRemove ) override;
 	virtual void LoadStyleResources( const ISlateStyle& Style ) override;
 	virtual FSlateUpdatableTexture* CreateUpdatableTexture(uint32 Width, uint32 Height) override;
diff --git a/Engine/Source/Runtime/SlateCore/Private/Rendering/DrawElements.cpp b/Engine/Source/Runtime/SlateCore/Private/Rendering/DrawElements.cpp
--- a/Engine/Source/Runtime/SlateCore/Private/Rendering/DrawElements.cpp
+++ b/Engine/Source/Runtime/SlateCore/Private/Rendering/DrawElements.cpp
@@ -47,6 +47,27 @@ void FSlateDrawElement::MakeBox(
 	DrawElt.DataPayload.SetBoxPayloadProperties( InBrush, InTint );
 }
 
+void FSlateDrawElement::MakeBox( 
+	FSlateWindowElementList& ElementList,
+	uint32 InLayer, 
+	const FPaintGeometry& PaintGeometry, 
+	const FSlateBrush* InBrush, 
+	const FSlateResourceHandle& InRenderingHandle,
+	const FSlateRect& InClippingRect, 
+	ESlateDrawEffect::Type InDrawEffects, 
+	const FLinearColor& InTint )
+{
+	SCOPE_CYCLE_COUNTER( STAT_SlateDrawElementMakeTime )
+
+	FSlateShaderResourceProxy* RenderingProxy = InRenderingHandle.Data->Proxy;
+
+	PaintGeometry.CommitTransformsIfUsingLegacyConstructor();
+	FSlateDrawElement& DrawElt = ElementList.AddUninitialized();
+	DrawElt.Init(InLayer, PaintGeometry, InClippingRect, InDrawEffects);
+	DrawElt.ElementType = (InBrush->DrawAs == ESlateBrushDrawType::Border) ? ET_Border : ET_Box;
+	DrawElt.DataPayload.SetBoxPayloadProperties( InBrush, InTint, RenderingProxy );
+}
+
 void FSlateDrawElement::MakeRotatedBox( 
 	FSlateWindowElementList& ElementList,
 	uint32 InLayer, 
diff --git a/Engine/Source/Runtime/SlateCore/Private/Rendering/ShaderResourceManager.cpp b/Engine/Source/Runtime/SlateCore/Private/Rendering/ShaderResourceManager.cpp
new file mode 100644
--- /dev/null
+++ b/Engine/Source/Runtime/SlateCore/Private/Rendering/ShaderResourceManager.cpp
@@ -0,0 +1,22 @@
+// Copyright 1998-2015 Epic Games, Inc. All Rights Reserved.
+
+#include "SlateCorePrivatePCH.h"
+#include "ShaderResourceManager.h"
+
+FSlateResourceHandle FSlateShaderResourceManager::GetResourceHandle( const FSlateBrush& InBrush )
+{
+	FSlateShaderResourceProxy* Proxy = GetShaderResource( InBrush );
+
+	FSlateResourceHandle NewHandle;
+	if( Proxy )
+	{
+		if( !Proxy->HandleData.IsValid() )
+		{
+			Proxy->HandleData = MakeShareable( new FSlateSharedHandleData( Proxy ) );
+		}
+
+		NewHandle.Data = Proxy->HandleData;
+	}
+
+	return NewHandle;
+}
\ No newline at end of file
diff --git a/Engine/Source/Runtime/SlateCore/Public/Rendering/DrawElements.h b/Engine/Source/Runtime/SlateCore/Public/Rendering/DrawElements.h
--- a/Engine/Source/Runtime/SlateCore/Public/Rendering/DrawElements.h
+++ b/Engine/Source/Runtime/SlateCore/Public/Rendering/DrawElements.h
@@ -99,19 +99,25 @@ public:
 		, CustomDrawer()
 	{ }
 
-	void SetBoxPayloadProperties( const FSlateBrush* InBrush, const FLinearColor& InTint )
+	void SetBoxPayloadProperties( const FSlateBrush* InBrush, const FLinearColor& InTint, FSlateShaderResourceProxy* InResourceProxy = nullptr )
 	{
 		Tint = InTint;
+
 		BrushResource = InBrush;
-		ResourceProxy = ResourceManager->GetShaderResource(*InBrush);
+		if( InResourceProxy )
+		{
+			ResourceProxy = InResourceProxy;
+		}
+		else
+		{
+			ResourceProxy = ResourceManager->GetShaderResource(*InBrush);
+		}
 		Angle = 0.0f;
 	}
 
-	void SetRotatedBoxPayloadProperties( const FSlateBrush* InBrush, float InAngle, const FVector2D& LocalRotationPoint, const FLinearColor& InTint )
+	void SetRotatedBoxPayloadProperties( const FSlateBrush* InBrush, float InAngle, const FVector2D& LocalRotationPoint, const FLinearColor& InTint, FSlateShaderResourceProxy* InResourceProxy = nullptr )
 	{
-		Tint = InTint;
-		BrushResource = InBrush;
-		ResourceProxy = ResourceManager->GetShaderResource(*InBrush);
+		SetBoxPayloadProperties( InBrush, InTint, InResourceProxy );
 		RotationPoint = LocalRotationPoint;
 		RotationPoint.DiagnosticCheckNaN();
 		Angle = InAngle;
@@ -253,6 +259,16 @@ public:
 		ESlateDrawEffect::Type InDrawEffects = ESlateDrawEffect::None, 
 		const FLinearColor& InTint = FLinearColor::White );
 
+	SLATECORE_API static void MakeBox(
+		FSlateWindowElementList& ElementList,
+		uint32 InLayer, 
+		const FPaintGeometry& PaintGeometry, 
+		const FSlateBrush* InBrush, 
+		const FSlateResourceHandle& InRenderingHandle, 
+		const FSlateRect& InClippingRect, 
+		ESlateDrawEffect::Type InDrawEffects = ESlateDrawEffect::None, 
+		const FLinearColor& InTint = FLinearColor::White );
+	
 	// !!! DEPRECATED !!! Use a render transform om your widget instead.
 	SLATECORE_API static void MakeRotatedBox(
 		FSlateWindowElementList& ElementList,
@@ -265,7 +281,6 @@ public:
 		TOptional<FVector2D> InRotationPoint = TOptional<FVector2D>(),
 		ERotationSpace RotationSpace = RelativeToElement,
 		const FLinearColor& InTint = FLinearColor::White );
-
 	/**
 	 * Creates a text element which displays a string of a rendered in a certain font on the screen
 	 *
diff --git a/Engine/Source/Runtime/SlateCore/Public/Rendering/ShaderResourceManager.h b/Engine/Source/Runtime/SlateCore/Public/Rendering/ShaderResourceManager.h
new file mode 100644
--- /dev/null
+++ b/Engine/Source/Runtime/SlateCore/Public/Rendering/ShaderResourceManager.h
@@ -0,0 +1,98 @@
+// Copyright 1998-2015 Epic Games, Inc. All Rights Reserved.
+
+#pragma once
+
+class ISlateAtlasProvider;
+
+/**
+ * Represents a new texture that has been loaded but no resource created for it
+ */
+struct FNewTextureInfo
+{
+	/** Raw data */
+	FSlateTextureDataPtr TextureData;
+	/** Whether or not the texture should be atlased */
+	bool bShouldAtlas;
+	/** Whether or not the texture is in srgb space */
+	bool bSrgb;
+	FNewTextureInfo()
+		: bShouldAtlas(true)
+		, bSrgb(true)
+	{
+
+	}
+};
+
+struct FCompareFNewTextureInfoByTextureSize
+{
+	FORCEINLINE bool operator()( const FNewTextureInfo& A, const FNewTextureInfo& B ) const
+	{
+		return (B.TextureData->GetWidth()+B.TextureData->GetHeight()) < (A.TextureData->GetWidth()+A.TextureData->GetHeight());
+	}
+};
+
+
+/** 
+ * Base texture manager class used by a Slate renderer to manage texture resources
+ */
+class SLATECORE_API FSlateShaderResourceManager
+{
+public:
+	FSlateShaderResourceManager() {};
+	virtual ~FSlateShaderResourceManager()
+	{
+		ClearTextureMap();
+	}
+
+
+	/** 
+	 * Returns a texture associated with the passed in name.  Should return nullptr if not found 
+	 */
+	virtual FSlateShaderResourceProxy* GetShaderResource( const FSlateBrush& InBrush ) = 0;
+
+	/**
+	 * Creates a handle to a Slate resource
+	 * A handle is used as fast path for looking up a rendering resource for a given brush when adding Slate draw elements
+	 * This can be cached and stored safely in code.  It will become invalid when a resource is destroyed
+	 * It is expensive to create a resource so do not do it in time sensitive areas
+	 *
+	 * @param	Brush		The brush to get a rendering resource handle 
+	 * @return	The created resource handle.  
+	 */
+	virtual FSlateResourceHandle GetResourceHandle( const FSlateBrush& InBrush );
+
+
+	virtual FSlateShaderResource* GetFontShaderResource( uint32 FontAtlasIndex, FSlateShaderResource* FontTextureAtlas, const class UObject* FontMaterial ) { return FontTextureAtlas; }
+
+	/**
+	 * Returns the way to access the texture atlas information from this resource manager
+	 */
+	virtual ISlateAtlasProvider* GetTextureAtlasProvider() = 0;
+
+protected:
+
+	void ClearTextureMap()
+	{
+		// delete all allocated textures
+		for( TMap<FName,FSlateShaderResourceProxy*>::TIterator It(ResourceMap); It; ++It )
+		{
+			delete It.Value();
+		}
+		ResourceMap.Empty();
+	}
+
+	FString GetResourcePath( const FSlateBrush& InBrush ) const
+	{
+		// assume the brush name contains the whole path
+		return InBrush.GetResourceName().ToString();
+	}
+
+	/** Mapping of names to texture pointers */
+	TMap<FName,FSlateShaderResourceProxy*> ResourceMap;
+
+private:
+	// Non-copyable
+	FSlateShaderResourceManager(const FSlateShaderResourceManager&);
+	FSlateShaderResourceManager& operator=(const FSlateShaderResourceManager&);
+
+};
diff --git a/Engine/Source/Runtime/SlateCore/Public/Rendering/SlateRenderer.h b/Engine/Source/Runtime/SlateCore/Public/Rendering/SlateRenderer.h
--- a/Engine/Source/Runtime/SlateCore/Public/Rendering/SlateRenderer.h
+++ b/Engine/Source/Runtime/SlateCore/Public/Rendering/SlateRenderer.h
@@ -110,6 +110,17 @@ public:
 	 */
 	virtual bool GenerateDynamicImageResource( FName ResourceName, uint32 Width, uint32 Height, const TArray< uint8 >& Bytes ) { return false; }
 
+	/**
+	 * Creates a handle to a Slate resource
+	 * A handle is used as fast path for looking up a rendering resource for a given brush when adding Slate draw elements
+	 * This can be cached and stored safely in code.  It will become invalid when a resource is destroyed
+	 * It is expensive to create a resource so do not do it in time sensitive areas
+	 *
+	 * @param	Brush		The brush to get a rendering resource handle 
+	 * @return	The created resource handle.  
+	 */
+	virtual FSlateResourceHandle GetResourceHandle( const FSlateBrush& Brush ) = 0;
+
 	/**
 	 * Queues a dynamic brush for removal when it is safe.  The brush is not immediately released but you should consider the brush destroyed and no longer usable
 	 *
diff --git a/Engine/Source/Runtime/SlateCore/Public/SlateCore.h b/Engine/Source/Runtime/SlateCore/Public/SlateCore.h
--- a/Engine/Source/Runtime/SlateCore/Public/SlateCore.h
+++ b/Engine/Source/Runtime/SlateCore/Public/SlateCore.h
@@ -80,7 +80,7 @@ enum class EActiveTimerReturnType : uint8;
 #include "SlateTextureData.h"
 #include "SlateUpdatableTexture.h"
 #include "TextureAtlas.h"
-#include "TextureManager.h"
+#include "ShaderResourceManager.h"
 
 // Fonts
 #include "SlateFontInfo.h"
diff --git a/Engine/Source/Runtime/SlateCore/Public/Textures/SlateShaderResource.h b/Engine/Source/Runtime/SlateCore/Public/Textures/SlateShaderResource.h
--- a/Engine/Source/Runtime/SlateCore/Public/Textures/SlateShaderResource.h
+++ b/Engine/Source/Runtime/SlateCore/Public/Textures/SlateShaderResource.h
@@ -2,6 +2,8 @@
 
 #pragma once
 
+class FSlateShaderResourceProxy;
+
 DECLARE_MEMORY_STAT_EXTERN(TEXT("Texture Data GPU Memory"), STAT_SlateTextureGPUMemory, STATGROUP_SlateMemory, SLATECORE_API);
 
 namespace ESlateShaderResource
@@ -60,6 +62,22 @@ public:
 	virtual ~FSlateShaderResource() { }
 };
 
+/**
+ * Data used to lookup a resource from a handle and to determine validity of the handle.  
+ * This is shared between the handle and the resource
+ *
+ */
+class FSlateSharedHandleData
+{
+public:
+	FSlateSharedHandleData( FSlateShaderResourceProxy* InProxy = nullptr )
+		: Proxy( InProxy )
+	{}
+
+public:
+	/** Rendering proxy used to directly access the rendering resource quickly */
+	FSlateShaderResourceProxy* Proxy;
+};
 
 /** 
  * A proxy resource.  
@@ -83,6 +101,9 @@ public:
 	/** The size of the texture.  Regardless of atlasing this is the size of the actual texture */
 	FIntPoint ActualSize;
 
+	/** Shared data between resources and handles to this resource.  Is created the first time a handle is needed */
+	TSharedPtr<FSlateSharedHandleData> HandleData;
+
 	/** Default constructor. */
 	FSlateShaderResourceProxy( )
 		: StartUV(0.0f, 0.0f)
@@ -90,6 +111,16 @@ public:
 		, Resource(nullptr)
 		, ActualSize(0.0f, 0.0f)
 	{ }
+
+	~FSlateShaderResourceProxy()
+	{
+		if( HandleData.IsValid() )
+		{
+			// Handles exist for this proxy.
+			// Null out the proxy to invalidate all handles
+			HandleData->Proxy = nullptr; 
+		}
+	}
 };
 
 
@@ -148,3 +179,20 @@ public:
 	virtual FSlateShaderResource* GetViewportRenderTargetTexture() = 0;
 };
 
+/**
+ * A SlateResourceHandle is used as fast path for looking up a rendering resource for a given brush when adding Slate draw elements
+ * This can be cached and stored safely in code.  It will become invalid when a resource is destroyed
+*/
+class FSlateResourceHandle
+{
+	friend class FSlateShaderResourceManager;
+	friend class FSlateDrawElement;
+public:
+	/**
+	 * @return true if the handle still points to a valid rendering resource
+	 */
+	bool IsValid() const { return Data.IsValid() && Data->Proxy; }
+private:
+	/** Internal data to pair the handle to the resource */
+	TSharedPtr<FSlateSharedHandleData> Data;
+};
\ No newline at end of file
diff --git a/Engine/Source/Runtime/SlateCore/Public/Textures/TextureManager.h b/Engine/Source/Runtime/SlateCore/Public/Textures/TextureManager.h
deleted file mode 100644
diff --git a/Engine/Source/Runtime/SlateRHIRenderer/Private/SlateRHIRenderer.cpp b/Engine/Source/Runtime/SlateRHIRenderer/Private/SlateRHIRenderer.cpp
--- a/Engine/Source/Runtime/SlateRHIRenderer/Private/SlateRHIRenderer.cpp
+++ b/Engine/Source/Runtime/SlateRHIRenderer/Private/SlateRHIRenderer.cpp
@@ -1083,6 +1083,11 @@ bool FSlateRHIRenderer::GenerateDynamicImageResource( FName ResourceName, uint32
 	return TextureResource.IsValid();
 }
 
+FSlateResourceHandle FSlateRHIRenderer::GetResourceHandle( const FSlateBrush& Brush )
+{
+	return ResourceManager->GetResourceHandle( Brush );
+}
+
 void FSlateRHIRenderer::RemoveDynamicBrushResource( TSharedPtr<FSlateDynamicImageBrush> BrushToRemove )
 {
 	if( BrushToRemove.IsValid() )
diff --git a/Engine/Source/Runtime/SlateRHIRenderer/Private/SlateRHIRenderer.h b/Engine/Source/Runtime/SlateRHIRenderer/Private/SlateRHIRenderer.h
--- a/Engine/Source/Runtime/SlateRHIRenderer/Private/SlateRHIRenderer.h
+++ b/Engine/Source/Runtime/SlateRHIRenderer/Private/SlateRHIRenderer.h
@@ -178,6 +178,7 @@ public:
 	virtual void RemoveDynamicBrushResource( TSharedPtr<FSlateDynamicImageBrush> BrushToRemove ) override;
 	virtual FIntPoint GenerateDynamicImageResource(const FName InTextureName) override;
 	virtual bool GenerateDynamicImageResource( FName ResourceName, uint32 Width, uint32 Height, const TArray< uint8 >& Bytes ) override;
+	virtual FSlateResourceHandle GetResourceHandle( const FSlateBrush& Brush ) override;
 	virtual void* GetViewportResource( const SWindow& Window ) override;
 	virtual void SetColorVisionDeficiencyType( uint32 Type ) override;
 	virtual FSlateUpdatableTexture* CreateUpdatableTexture(uint32 Width, uint32 Height) override;
diff --git a/Engine/Source/Runtime/SlateRHIRenderer/Private/SlateRHIResourceManager.cpp b/Engine/Source/Runtime/SlateRHIRenderer/Private/SlateRHIResourceManager.cpp
--- a/Engine/Source/Runtime/SlateRHIRenderer/Private/SlateRHIResourceManager.cpp
+++ b/Engine/Source/Runtime/SlateRHIRenderer/Private/SlateRHIResourceManager.cpp
@@ -10,6 +10,7 @@
 DECLARE_DWORD_ACCUMULATOR_STAT(TEXT("Num Texture Atlases"), STAT_SlateNumTextureAtlases, STATGROUP_SlateMemory);
 DECLARE_DWORD_ACCUMULATOR_STAT(TEXT("Num Non-Atlased Textures"), STAT_SlateNumNonAtlasedTextures, STATGROUP_SlateMemory);
 DECLARE_DWORD_ACCUMULATOR_STAT(TEXT("Num Dynamic Textures"), STAT_SlateNumDynamicTextures, STATGROUP_SlateMemory);
+DECLARE_CYCLE_STAT(TEXT("GetResource Time"), STAT_SlateGetResourceTime, STATGROUP_SlateVerbose);
 
 FDynamicResourceMap::FDynamicResourceMap()
 	: TextureMemorySincePurge(0)
@@ -430,6 +431,8 @@ static void LoadUObjectForBrush( const FSlateBrush& InBrush )
 
 FSlateShaderResourceProxy* FSlateRHIResourceManager::GetShaderResource( const FSlateBrush& InBrush )
 {
+	SCOPE_CYCLE_COUNTER( STAT_SlateGetResourceTime );
+
 	check( IsThreadSafeForSlateRendering() );
 
 	FSlateShaderResourceProxy* Texture = NULL;
