Fix slate styles not working with SlateBrushAssets

Merging 2686963

//depot/UE4-UT/...

to //depot/UE4/...

diff --git a/Engine/Source/Runtime/SlateCore/Private/Styling/SlateStyleSet.cpp b/Engine/Source/Runtime/SlateCore/Private/Styling/SlateStyleSet.cpp
--- a/Engine/Source/Runtime/SlateCore/Private/Styling/SlateStyleSet.cpp
+++ b/Engine/Source/Runtime/SlateCore/Private/Styling/SlateStyleSet.cpp
@@ -16,7 +16,10 @@ FSlateStyleSet::~FSlateStyleSet()
 	// Delete all allocated brush resources.
 	for ( TMap< FName, FSlateBrush* >::TIterator It(BrushResources); It; ++It )
 	{
-		delete It.Value();
+		if (!It.Value()->HasUObject())
+		{
+			delete It.Value();
+		}
 	}
 }
 
