Optimised FElementBatchMap manipulation by making the value a TUniqueObj.

diff --git a/Engine/Source/Runtime/SlateCore/Private/Rendering/DrawElements.cpp b/Engine/Source/Runtime/SlateCore/Private/Rendering/DrawElements.cpp
--- a/Engine/Source/Runtime/SlateCore/Private/Rendering/DrawElements.cpp
+++ b/Engine/Source/Runtime/SlateCore/Private/Rendering/DrawElements.cpp
@@ -371,7 +371,7 @@ void FSlateBatchData::Merge(FElementBatchMap& InLayerToElementBatches, uint32& V
 	for ( FElementBatchMap::TIterator It(InLayerToElementBatches); It; ++It )
 	{
 		uint32 Layer = It.Key();
-		FElementBatchArray& ElementBatches = It.Value();
+		FElementBatchArray& ElementBatches = *It.Value();
 
 		if ( ElementBatches.Num() > 0 )
 		{
diff --git a/Engine/Source/Runtime/SlateCore/Private/Rendering/ElementBatcher.cpp b/Engine/Source/Runtime/SlateCore/Private/Rendering/ElementBatcher.cpp
--- a/Engine/Source/Runtime/SlateCore/Private/Rendering/ElementBatcher.cpp
+++ b/Engine/Source/Runtime/SlateCore/Private/Rendering/ElementBatcher.cpp
@@ -1545,16 +1545,16 @@ void FSlateElementBatcher::AddCustomElement( const FSlateDrawElement& DrawElemen
 	if( InPayload.CustomDrawer.IsValid() )
 	{
 		// See if the layer already exists.
-		FElementBatchArray* ElementBatches = LayerToElementBatches.Find( Layer );
+		TUniqueObj<FElementBatchArray>* ElementBatches = LayerToElementBatches.Find( Layer );
 		if( !ElementBatches )
 		{
 			// The layer doesn't exist so make it now
-			ElementBatches = &LayerToElementBatches.Add( Layer, FElementBatchArray() );
+			ElementBatches = &LayerToElementBatches.Add( Layer );
 		}
 		check( ElementBatches );
 
 		// Custom elements are not batched together 
-		ElementBatches->Add( FSlateElementBatch( InPayload.CustomDrawer, DrawElement.GetScissorRect() ) );
+		(*ElementBatches)->Add( FSlateElementBatch( InPayload.CustomDrawer, DrawElement.GetScissorRect() ) );
 	}
 }
 
@@ -1568,17 +1568,17 @@ void FSlateElementBatcher::AddCachedBuffer(const FSlateDrawElement& DrawElement)
 	if ( InPayload.CachedRenderData )
 	{
 		// See if the layer already exists.
-		FElementBatchArray* ElementBatches = LayerToElementBatches.Find(Layer);
+		TUniqueObj<FElementBatchArray>* ElementBatches = LayerToElementBatches.Find(Layer);
 		if ( !ElementBatches )
 		{
 			// The layer doesn't exist so make it now
-			ElementBatches = &LayerToElementBatches.Add(Layer, FElementBatchArray());
+			ElementBatches = &LayerToElementBatches.Add(Layer);
 		}
 		check(ElementBatches);
 
 		// Custom elements are not batched together
 		TSharedPtr< FSlateRenderDataHandle, ESPMode::ThreadSafe > RenderData = InPayload.CachedRenderData->AsShared();
-		ElementBatches->Add(FSlateElementBatch(RenderData, InPayload.CachedRenderDataOffset, DrawElement.GetScissorRect()));
+		(*ElementBatches)->Add(FSlateElementBatch(RenderData, InPayload.CachedRenderDataOffset, DrawElement.GetScissorRect()));
 	}
 }
 
@@ -1592,17 +1592,17 @@ void FSlateElementBatcher::AddLayer(const FSlateDrawElement& DrawElement)
 	if ( InPayload.LayerHandle )
 	{
 		// See if the layer already exists.
-		FElementBatchArray* ElementBatches = LayerToElementBatches.Find(Layer);
+		TUniqueObj<FElementBatchArray>* ElementBatches = LayerToElementBatches.Find(Layer);
 		if ( !ElementBatches )
 		{
 			// The layer doesn't exist so make it now
-			ElementBatches = &LayerToElementBatches.Add(Layer, FElementBatchArray());
+			ElementBatches = &LayerToElementBatches.Add(Layer);
 		}
 		check(ElementBatches);
 
 		// Custom elements are not batched together
 		TSharedPtr< FSlateDrawLayerHandle, ESPMode::ThreadSafe > LayerHandle = InPayload.LayerHandle->AsShared();
-		ElementBatches->Add(FSlateElementBatch(LayerHandle, DrawElement.GetScissorRect()));
+		(*ElementBatches)->Add(FSlateElementBatch(LayerHandle, DrawElement.GetScissorRect()));
 	}
 }
 
@@ -1623,11 +1623,11 @@ FSlateElementBatch& FSlateElementBatcher::FindBatchForElement(
 	FElementBatchMap& LayerToElementBatches = DrawLayer->GetElementBatchMap();
 
 	// See if the layer already exists.
-	FElementBatchArray* ElementBatches = LayerToElementBatches.Find( Layer );
+	TUniqueObj<FElementBatchArray>* ElementBatches = LayerToElementBatches.Find( Layer );
 	if( !ElementBatches )
 	{
 		// The layer doesn't exist so make it now
-		ElementBatches = &LayerToElementBatches.Add( Layer, FElementBatchArray() );
+		ElementBatches = &LayerToElementBatches.Add( Layer );
 	}
 
 	checkSlow( ElementBatches );
@@ -1635,12 +1635,12 @@ FSlateElementBatch& FSlateElementBatcher::FindBatchForElement(
 	// Create a temp batch so we can use it as our key to find if the same batch already exists
 	FSlateElementBatch TempBatch( InTexture, ShaderParams, ShaderType, PrimitiveType, DrawEffects, DrawFlags, ScissorRect );
 
-	FSlateElementBatch* ElementBatch = ElementBatches->FindByKey( TempBatch );
+	FSlateElementBatch* ElementBatch = (*ElementBatches)->FindByKey( TempBatch );
 	if( !ElementBatch )
 	{
 		// No batch with the specified parameter exists.  Create it from the temp batch.
-		int32 Index = ElementBatches->Add( TempBatch );
-		ElementBatch = &(*ElementBatches)[Index];
+		int32 Index = (*ElementBatches)->Add( TempBatch );
+		ElementBatch = &(**ElementBatches)[Index];
 
 		BatchData->AssignVertexArrayToBatch(*ElementBatch);
 		BatchData->AssignIndexArrayToBatch(*ElementBatch);
diff --git a/Engine/Source/Runtime/SlateCore/Public/Rendering/DrawElements.h b/Engine/Source/Runtime/SlateCore/Public/Rendering/DrawElements.h
--- a/Engine/Source/Runtime/SlateCore/Public/Rendering/DrawElements.h
+++ b/Engine/Source/Runtime/SlateCore/Public/Rendering/DrawElements.h
@@ -2,6 +2,7 @@
 
 #pragma once
 
+#include "UniqueObj.h"
 
 class SWindow;
 class FSlateViewportInterface;
@@ -691,7 +692,7 @@ public:
 
 typedef TArray<FSlateElementBatch, TInlineAllocator<1>> FElementBatchArray;
 
-typedef TMap<uint32,FElementBatchArray> FElementBatchMap ;
+typedef TMap<uint32,TUniqueObj<FElementBatchArray>> FElementBatchMap ;
 
 class FSlateBatchData
 {
