Fixed Slate atlas textures no longer making full use of their page space

UE-21089 - Slate Atlas Textures are not making full use of page space

diff --git a/Engine/Source/Runtime/SlateCore/Private/Textures/TextureAtlas.cpp b/Engine/Source/Runtime/SlateCore/Private/Textures/TextureAtlas.cpp
--- a/Engine/Source/Runtime/SlateCore/Private/Textures/TextureAtlas.cpp
+++ b/Engine/Source/Runtime/SlateCore/Private/Textures/TextureAtlas.cpp
@@ -257,7 +257,7 @@ const FAtlasedTextureSlot* FSlateTextureAtlas::FindSlotForTexture(uint32 InWidth
 		const uint32 MinSlotDim = 2;
 
 		// Split the remaining area around this slot into two children.
-		if (RemainingHeight >= MinSlotDim && RemainingWidth >= MinSlotDim)
+		if (RemainingHeight >= MinSlotDim || RemainingWidth >= MinSlotDim)
 		{
 			FAtlasedTextureSlot* LeftSlot = NULL;
 			FAtlasedTextureSlot* RightSlot = NULL;
