UE4 - remove malloc churn on strings in slate.

diff --git a/Engine/Source/Runtime/SlateCore/Private/Rendering/DrawElements.cpp b/Engine/Source/Runtime/SlateCore/Private/Rendering/DrawElements.cpp
--- a/Engine/Source/Runtime/SlateCore/Private/Rendering/DrawElements.cpp
+++ b/Engine/Source/Runtime/SlateCore/Private/Rendering/DrawElements.cpp
@@ -7,6 +7,24 @@ DECLARE_CYCLE_STAT(TEXT("FSlateDrawElement::Make Time"), STAT_SlateDrawElementMa
 
 FSlateShaderResourceManager* FSlateDataPayload::ResourceManager;
 
+void FSlateDataPayload::SetTextPayloadProperties( FSlateWindowElementList& ElementList, const FString& InText, const FSlateFontInfo& InFontInfo, const FLinearColor& InTint )
+{
+	Tint = InTint;
+	FontInfo = InFontInfo;
+	SIZE_T Count = InText.Len() + 1;
+	ImmutableText = (TCHAR*)ElementList.Alloc(sizeof(TCHAR) * Count, ALIGNOF(TCHAR));
+	if (Count > 1)
+	{
+		FCString::Strcpy(ImmutableText, Count, InText.GetCharArray().GetData());
+		check(!ImmutableText[Count - 1]);
+	}
+	else
+	{
+		ImmutableText[0] = 0;
+	}
+}
+
+
 void FSlateDrawElement::Init(uint32 InLayer, const FPaintGeometry& PaintGeometry, const FSlateRect& InClippingRect, ESlateDrawEffect::Type InDrawEffects)
 {
 	RenderTransform = PaintGeometry.GetAccumulatedRenderTransform();
@@ -98,7 +116,8 @@ void FSlateDrawElement::MakeText( FSlateWindowElementList& ElementList, uint32 I
 	FSlateDrawElement& DrawElt = ElementList.AddUninitialized();
 	DrawElt.Init(InLayer, PaintGeometry, InClippingRect, InDrawEffects);
 	DrawElt.ElementType = ET_Text;
-	DrawElt.DataPayload.SetTextPayloadProperties( FString( EndIndex - StartIndex, *InText + StartIndex ), InFontInfo, InTint );
+	//fixme, no reason to create a temp string here, add an overload for SetTextPayloadProperties instead
+	DrawElt.DataPayload.SetTextPayloadProperties(  ElementList, FString( EndIndex - StartIndex, *InText + StartIndex ), InFontInfo, InTint );
 }
 
 
@@ -109,7 +128,7 @@ void FSlateDrawElement::MakeText( FSlateWindowElementList& ElementList, uint32 I
 	FSlateDrawElement& DrawElt = ElementList.AddUninitialized();
 	DrawElt.Init(InLayer, PaintGeometry, InClippingRect, InDrawEffects);
 	DrawElt.ElementType = ET_Text;
-	DrawElt.DataPayload.SetTextPayloadProperties( InText, InFontInfo, InTint );
+	DrawElt.DataPayload.SetTextPayloadProperties(  ElementList, InText, InFontInfo, InTint );
 }
 
 
@@ -120,7 +139,8 @@ void FSlateDrawElement::MakeText( FSlateWindowElementList& ElementList, uint32 I
 	FSlateDrawElement& DrawElt = ElementList.AddUninitialized();
 	DrawElt.Init(InLayer, PaintGeometry, InClippingRect, InDrawEffects);
 	DrawElt.ElementType = ET_Text;
-	DrawElt.DataPayload.SetTextPayloadProperties( InText.ToString(), InFontInfo, InTint );
+	//fixme, alloc here 
+	DrawElt.DataPayload.SetTextPayloadProperties(  ElementList, InText.ToString(), InFontInfo, InTint );
 }
 
 void FSlateDrawElement::MakeGradient( FSlateWindowElementList& ElementList, uint32 InLayer, const FPaintGeometry& PaintGeometry, TArray<FSlateGradientStop> InGradientStops, EOrientation InGradientType, const FSlateRect& InClippingRect, ESlateDrawEffect::Type InDrawEffects, bool bGammaCorrect )
@@ -655,6 +675,9 @@ void FSlateWindowElementList::PostDraw_ParallelThread()
 	CachedRenderHandlesInUse.Reset();
 }
 
+DECLARE_MEMORY_STAT(TEXT("FSlateWindowElementList MemManager"), STAT_FSlateWindowElementListMemManager, STATGROUP_SlateVerbose);
+DECLARE_DWORD_COUNTER_STAT(TEXT("FSlateWindowElementList MemManager Count"), STAT_FSlateWindowElementListMemManagerCount, STATGROUP_SlateVerbose);
+
 void FSlateWindowElementList::ResetBuffers()
 {
 	// Don't attempt to use this slate window element list if the cache is still being used.
@@ -689,4 +712,9 @@ void FSlateWindowElementList::ResetBuffers()
 
 	DrawStack.Reset();
 	DrawStack.Push(&RootDrawLayer);
+
+	INC_DWORD_STAT(STAT_FSlateWindowElementListMemManagerCount);
+	INC_MEMORY_STAT_BY(STAT_FSlateWindowElementListMemManager, MemManager.GetByteCount());
+
+	MemManager.Flush();
 }
diff --git a/Engine/Source/Runtime/SlateCore/Private/Rendering/ElementBatcher.cpp b/Engine/Source/Runtime/SlateCore/Private/Rendering/ElementBatcher.cpp
--- a/Engine/Source/Runtime/SlateCore/Private/Rendering/ElementBatcher.cpp
+++ b/Engine/Source/Runtime/SlateCore/Private/Rendering/ElementBatcher.cpp
@@ -635,9 +635,9 @@ void FSlateElementBatcher::AddTextElement(const FSlateDrawElement& DrawElement)
 	// Used to clip individual characters as we generate them.
 	FSlateRect LocalClipRect = TransformRect(InverseLayoutTransform, InClippingRect);
 
-	const FString& Text = InPayload.Text;
+	int32 Len = InPayload.ImmutableText ? FCString::Strlen(InPayload.ImmutableText) : 0;
 	// Nothing to do if no text
-	if( Text.Len() == 0 )
+	if( Len == 0 )
 	{
 		return;
 	}
@@ -676,12 +676,12 @@ void FSlateElementBatcher::AddTextElement(const FSlateDrawElement& DrawElement)
 	
 	const bool bIsFontMaterial = InPayload.FontInfo.FontMaterial != nullptr;
 
-	uint32 NumChars = Text.Len();
+	uint32 NumChars = Len;
 
 	uint32 NumLines = 1;
 	for( uint32 CharIndex = 0; CharIndex < NumChars; ++CharIndex )
 	{
-		const TCHAR CurrentChar = Text[ CharIndex ];
+		const TCHAR CurrentChar = InPayload.ImmutableText[ CharIndex ];
 
 		const bool IsNewline = (CurrentChar == '\n');
 
diff --git a/Engine/Source/Runtime/SlateCore/Public/Rendering/DrawElements.h b/Engine/Source/Runtime/SlateCore/Public/Rendering/DrawElements.h
--- a/Engine/Source/Runtime/SlateCore/Public/Rendering/DrawElements.h
+++ b/Engine/Source/Runtime/SlateCore/Public/Rendering/DrawElements.h
@@ -52,13 +52,13 @@ public:
 
 	// Font data
 	FSlateFontInfo FontInfo;
-	FString Text;
+	TCHAR* ImmutableText;
 
-	// Gradient data
+	// Gradient data (fixme, this should be allocated with FSlateWindowElementList::Alloc)
 	TArray<FSlateGradientStop> GradientStops;
 	EOrientation GradientType;
 
-	// Line data
+	// Line data (fixme, this should be allocated with FSlateWindowElementList::Alloc)
 	TArray<FVector2D> Points;
 
 	// Viewport data
@@ -92,6 +92,7 @@ public:
 		, BrushResource(nullptr)
 		, ResourceProxy(nullptr)
 		, RotationPoint(FVector2D::ZeroVector)
+		, ImmutableText(nullptr)
 		, ViewportRenderTargetTexture(nullptr)
 		, bViewportTextureAlphaOnly(false)
 		, bRequiresVSync(false)
@@ -124,12 +125,7 @@ public:
 		Angle = InAngle;
 	}
 
-	void SetTextPayloadProperties( const FString& InText, const FSlateFontInfo& InFontInfo, const FLinearColor& InTint )
-	{
-		Tint = InTint;
-		FontInfo = InFontInfo;
-		Text = InText;
-	}
+	void SetTextPayloadProperties( FSlateWindowElementList& DrawBuffer, const FString& InText, const FSlateFontInfo& InFontInfo, const FLinearColor& InTint );
 
 	void SetGradientPayloadProperties( const TArray<FSlateGradientStop>& InGradientStops, EOrientation InGradientType, bool bInGammaCorrect )
 	{
@@ -837,6 +833,7 @@ public:
 	 */
 	explicit FSlateWindowElementList( TSharedPtr<SWindow> InWindow = TSharedPtr<SWindow>() )
 		: TopLevelWindow( InWindow )
+		, MemManager(0)
 	{
 		DrawStack.Push(&RootDrawLayer);
 	}
@@ -980,6 +977,23 @@ public:
 	 */
 	SLATECORE_API void ResetBuffers();
 
+	/**
+	 * Allocate memory that remains valid until ResetBuffers is called.
+	 */
+	FORCEINLINE_DEBUGGABLE void* Alloc(int32 AllocSize, int32 Alignment = MIN_ALIGNMENT)
+	{
+		return MemManager.Alloc(AllocSize, Alignment);
+	}
+
+	/**
+	 * Allocate memory for a type that remains valid until ResetBuffers is called.
+	 */
+	template <typename T>
+	FORCEINLINE_DEBUGGABLE void* Alloc()
+	{
+		return MemManager.Alloc(sizeof(T), ALIGNOF(T));
+	}
+
 	FSlateBatchData& GetBatchData() { return BatchData; }
 
 	FSlateDrawLayer& GetRootDrawLayer() { return RootDrawLayer; }
@@ -1053,4 +1067,7 @@ private:
 	 * this element list is being used for invalidation / caching.
 	 */
 	mutable TWeakPtr<FSlateRenderDataHandle, ESPMode::ThreadSafe> CachedRenderDataHandle;
+
+	// Mem stack for temp allocations
+	FMemStackBase MemManager; 
 };
