UE-20550 UMG/Slate - Correctly setting the effected buttons in touch pointer events, restoring input for checkbox/slider on iOS that were only checking the EffectingButton.

diff --git a/Engine/Source/Runtime/SlateCore/Public/Input/Events.h b/Engine/Source/Runtime/SlateCore/Public/Input/Events.h
--- a/Engine/Source/Runtime/SlateCore/Public/Input/Events.h
+++ b/Engine/Source/Runtime/SlateCore/Public/Input/Events.h
@@ -704,7 +704,7 @@ public:
 		, LastScreenSpacePosition(InLastScreenSpacePosition)
 		, CursorDelta(InScreenSpacePosition - InLastScreenSpacePosition)
 		, PressedButtons(bPressLeftMouseButton ? FTouchKeySet::StandardSet : FTouchKeySet::EmptySet)
-		, EffectingButton()
+		, EffectingButton(EKeys::LeftMouseButton)
 		, PointerIndex(InPointerIndex)
 		, TouchpadIndex(InTouchpadIndex)
 		, bIsTouchEvent(true)
