Fixed mem leak in FSlateBatchData that can occur when a batch has no vertex or index data.

ElementBatches with no vertex data were not being added to the RenderBatches and therefore the ArrayIndices were never returned to their respective freelists.

diff --git a/Engine/Source/Runtime/SlateCore/Private/Rendering/DrawElements.cpp b/Engine/Source/Runtime/SlateCore/Private/Rendering/DrawElements.cpp
--- a/Engine/Source/Runtime/SlateCore/Private/Rendering/DrawElements.cpp
+++ b/Engine/Source/Runtime/SlateCore/Private/Rendering/DrawElements.cpp
@@ -468,6 +468,11 @@ void FSlateBatchData::Merge(FElementBatchMap& InLayerToElementBatches, uint32& V
 							VertexOffset += BatchVertices.Num();
 							IndexOffset += BatchIndices.Num();
 						}
+						else
+						{
+							VertexArrayFreeList.Add(ElementBatch.VertexArrayIndex);
+							IndexArrayFreeList.Add(ElementBatch.IndexArrayIndex);
+						}
 					}
 				}
 			}
