Fix for constructing SWindows with an auto center rule in place; now their size will always be clamped to not be bigger than the preferred work area.

diff --git a/Engine/Source/Editor/MainFrame/Private/MainFrameModule.h b/Engine/Source/Editor/MainFrame/Private/MainFrameModule.h
--- a/Engine/Source/Editor/MainFrame/Private/MainFrameModule.h
+++ b/Engine/Source/Editor/MainFrame/Private/MainFrameModule.h
@@ -125,7 +125,7 @@ protected:
 public:
 
 	/** Get the size of the project browser window */
-	static FVector2D GetProjectBrowserWindowSize() { return FVector2D(1100, 700); }
+	static FVector2D GetProjectBrowserWindowSize() { return FVector2D(1100, 740); }
 
 private:
 
diff --git a/Engine/Source/Runtime/SlateCore/Private/Widgets/SWindow.cpp b/Engine/Source/Runtime/SlateCore/Private/Widgets/SWindow.cpp
--- a/Engine/Source/Runtime/SlateCore/Private/Widgets/SWindow.cpp
+++ b/Engine/Source/Runtime/SlateCore/Private/Widgets/SWindow.cpp
@@ -287,6 +287,10 @@ void SWindow::Construct(const FArguments& InArgs)
 			break;
 		}
 
+		// Clamp window size to be no greater than the work area size
+		WindowSize.X = FMath::Min(WindowSize.X, AutoCenterRect.GetSize().X);
+		WindowSize.Y = FMath::Min(WindowSize.Y, AutoCenterRect.GetSize().Y);
+
 		// Setup a position and size for the main frame window that's centered in the desktop work area
 		const FVector2D DisplayTopLeft( AutoCenterRect.Left, AutoCenterRect.Top );
 		const FVector2D DisplaySize( AutoCenterRect.Right - AutoCenterRect.Left, AutoCenterRect.Bottom - AutoCenterRect.Top );
