Merged UE4-Orion->UE4.  Slate hardware instancing support

2697248 Added SMeshWidget and underlying support to draw a custom mesh from Slate.
2696183 Fixed world position offset not being usable without texture coordinate in Slate Materials.
2695823 Removed some unused code
2695811 Added hardware instancing support to Slate draw elements
2695462 Added custom Slate shader type for completely custom materials that make no assumption on how they are used

diff --git a/Engine/Shaders/SlateElementPixelShader.usf b/Engine/Shaders/SlateElementPixelShader.usf
--- a/Engine/Shaders/SlateElementPixelShader.usf
+++ b/Engine/Shaders/SlateElementPixelShader.usf
@@ -5,15 +5,12 @@
 #include "GammaCorrectionCommon.usf"
 #include "SlateShaderCommon.usf"
 
-/** Shader types (mirrored from ESlateShader::Type in RenderingCommon.h */
-#define ST_Default		0
-#define ST_Border		1
-#define ST_Font			2
-#define ST_LineSegment	3
 
 /** Enable to debug clipping issues */
 #define DEBUG_CLIPPING 0
 
+#define DO_CLIPPING SHADER_TYPE != ST_Custom
+
 /** Display gamma x:gamma curve adjustment, y:inverse gamma (1/GEngine->DisplayGamma) */
 float2 GammaValues;
 /** Per element params */
@@ -51,7 +48,7 @@ void FillTexCoords( inout FMaterialPixelParameters Parameters, VertexToPixelInte
 		Parameters.TexCoords[1] = InVertex.MaterialTexCoords;
 	#endif
 #else
-	#if NUM_CUSTOMIZED_UVS == 0
+	#if NUM_CUSTOMIZED_UVS == 0 || SHADER_TYPE == ST_Custom
 		#if NUM_MATERIAL_TEXCOORDS == 1
 			Parameters.TexCoords[0] = InDefaultUV;
 		#elif NUM_MATERIAL_TEXCOORDS == 2
@@ -61,6 +58,15 @@ void FillTexCoords( inout FMaterialPixelParameters Parameters, VertexToPixelInte
 			Parameters.TexCoords[0] = InDefaultUV;
 			Parameters.TexCoords[1] = GetUV( InVertex, 0 );
 			Parameters.TexCoords[2] = GetUV( InVertex, 1 );
+		#elif NUM_MATERIAL_TEXCOORDS >= 4
+			Parameters.TexCoords[0] = InDefaultUV;
+			Parameters.TexCoords[1] = GetUV( InVertex, 0 );
+			Parameters.TexCoords[2] = GetUV( InVertex, 1 );
+			UNROLL
+			for( int CoordinateIndex = 3; CoordinateIndex < NUM_MATERIAL_TEXCOORDS; CoordinateIndex++ )
+			{
+				Parameters.TexCoords[CoordinateIndex] = GetUV( InVertex, CoordinateIndex );
+			}
 		#endif
 	#else
 		// Note:  We assume users with customized UVs are overriding behavior and know what they are doing
@@ -118,13 +124,11 @@ float4 GetColor( VertexToPixelInterpolants InVertex, float2 UV )
 	float4 FinalColor;
 	
 #if USE_MATERIALS
-	float4 BaseColor = GetMaterialColor( InVertex, UV );
+	FinalColor = GetMaterialColor( InVertex, UV );
 #else
-	float4 BaseColor = GetTextureColor( InVertex, UV );
+	FinalColor = GetTextureColor( InVertex, UV )*InVertex.Color;
 #endif
 
-	FinalColor = BaseColor*InVertex.Color;
-
 	return FinalColor;
 }
 
@@ -145,7 +149,7 @@ float4 GetFontElementColor( VertexToPixelInterpolants InVertex )
 }
 #endif
 
-#if SHADER_TYPE == ST_Default
+#if SHADER_TYPE == ST_Default || SHADER_TYPE == ST_Custom
 float4 GetDefaultElementColor( VertexToPixelInterpolants InVertex )
 {
 	return GetColor( InVertex, GetUV( InVertex, 0) * GetUV( InVertex, 1 ) );
@@ -251,18 +255,20 @@ float4 GetLineSegmentElementColor( VertexToPixelInterpolants InVertex )
 
 float4 Main( VertexToPixelInterpolants VIn ) : SV_Target0
 {
-	// Clip pixels which are outside of the clipping rect
-	float2 ClipOrigin = VIn.ClipOriginAndPos.xy;
-	float2 WindowPos = VIn.ClipOriginAndPos.zw;
-	float4 ClipTest = PointInParallelogram(WindowPos, ClipOrigin, VIn.ClipExtents);
+#if DO_CLIPPING
+		// Clip pixels which are outside of the clipping rect
+		float2 ClipOrigin = VIn.ClipOriginAndPos.xy;
+		float2 WindowPos = VIn.ClipOriginAndPos.zw;
+		float4 ClipTest = PointInParallelogram(WindowPos, ClipOrigin, VIn.ClipExtents);
 	
-#if !DEBUG_CLIPPING
-	// Clip pixels which are outside of the clipping rect
-	clip(ClipTest);
+	#if !DEBUG_CLIPPING
+		// Clip pixels which are outside of the clipping rect
+		clip(ClipTest);
+	#endif
 #endif
 
 	float4 OutColor;
-#if SHADER_TYPE == ST_Default 
+#if SHADER_TYPE == ST_Default || SHADER_TYPE == ST_Custom
 	OutColor = GetDefaultElementColor( VIn );
 #elif SHADER_TYPE == ST_Border
 	OutColor = GetBorderElementColor( VIn );
@@ -286,6 +292,7 @@ float4 Main( VertexToPixelInterpolants VIn ) : SV_Target0
 #if !USE_MATERIALS
 	OutColor = ApplyColorVisionDeficiencyFilter(OutColor);
 #endif
+
 	// gamma correct
 	OutColor.rgb = GammaCorrect(OutColor.rgb);
 
diff --git a/Engine/Shaders/SlateShaderCommon.usf b/Engine/Shaders/SlateShaderCommon.usf
--- a/Engine/Shaders/SlateShaderCommon.usf
+++ b/Engine/Shaders/SlateShaderCommon.usf
@@ -17,6 +17,14 @@
 	#define NUM_SLATE_TEXCOORDS 1
 #endif
 
+
+/** Shader types (mirrored from ESlateShader::Type in RenderingCommon.h */
+#define ST_Default		0
+#define ST_Border		1
+#define ST_Font			2
+#define ST_LineSegment		3
+#define ST_Custom		4
+
 /** Color Vision Deficiency Types (mirrored in EditorUserSettings.h) */
 #define CVD_NormalVision	0
 #define CVD_Deuteranomly	1
diff --git a/Engine/Shaders/SlateVertexShader.usf b/Engine/Shaders/SlateVertexShader.usf
--- a/Engine/Shaders/SlateVertexShader.usf
+++ b/Engine/Shaders/SlateVertexShader.usf
@@ -15,7 +15,8 @@ VertexToPixelInterpolants Main(
 	in float2 InPosition : ATTRIBUTE2,
 	in float2 InClipOrigin : ATTRIBUTE3,
 	in float4 InClipExtents : ATTRIBUTE4,
-	in float4 InColor : ATTRIBUTE5
+	in float4 InColor : ATTRIBUTE5,
+	in float4 InstanceParam : ATTRIBUTE6
 	)
 {
 	VertexToPixelInterpolants VOut = (VertexToPixelInterpolants)0;
@@ -34,51 +35,65 @@ VertexToPixelInterpolants Main(
 #if USE_MATERIALS
 	FMaterialVertexParameters VertexParameters = (FMaterialVertexParameters)0;
 	VertexParameters.WorldPosition = WorldPosition.xyz;
-	
-	// Todo do not use per-instance parameters
-	/*#if USE_INSTANCING
-		VertexParameters.PerInstanceParams = float4(0,0,0,0);
-	#endif // USE_INSTANCING*/
 
 	VertexParameters.VertexColor = FinalVertexColor;
-	WorldPosition.xyz += GetMaterialWorldPositionOffset(VertexParameters);
-		
-	#if HAS_MATERIAL_TEXCOORDS
-		#if NUM_MATERIAL_TEXCOORDS_VERTEX == 1
-			VertexParameters.TexCoords[0] = InTextureCoordinates.xy;
-		#elif NUM_MATERIAL_TEXCOORDS_VERTEX >= 2
-			VertexParameters.TexCoords[0] = InTextureCoordinates.xy;
-			VertexParameters.TexCoords[1] = InTextureCoordinates.zw;
-		#endif
-
-		float2 CustomizedUVs[NUM_SLATE_TEXCOORDS];
-		GetMaterialCustomizedUVs(VertexParameters, CustomizedUVs);
-
-		#if ES2_PROFILE
-			#if NUM_SLATE_TEXCOORDS == 1
-				VOut.TextureCoordinates[0].xy = CustomizedUVs[0];
-			#elif NUM_SLATE_TEXCOORDS >= 2
-				VOut.TextureCoordinates[0].xy = CustomizedUVs[0];
-				VOut.TextureCoordinates[0].zw = CustomizedUVs[1];
-			#endif
-		#else
-			UNROLL
-			for( int CoordinateIndex = 0; CoordinateIndex < NUM_SLATE_TEXCOORDS; CoordinateIndex++ )
-			{
-				FLATTEN
-				if (CoordinateIndex % 2)
-				{
-					VOut.TextureCoordinates[CoordinateIndex / 2].zw = CustomizedUVs[CoordinateIndex];
-				}
-				else
+
+	#if NUM_MATERIAL_TEXCOORDS_VERTEX
+
+		float2 UVArrays[8] = 
+		{
+			InTextureCoordinates.xy,
+			InTextureCoordinates.zw,
+			InstanceParam.xy,
+			InstanceParam.zw,
+			InMaterialTextureCoordinates,
+			InClipOrigin,
+			InClipExtents.xy,
+			InClipExtents.zw
+		};
+
+		UNROLL
+		for( int CoordinateIndex = 0; CoordinateIndex < NUM_MATERIAL_TEXCOORDS_VERTEX; CoordinateIndex++ )
+		{
+			VertexParameters.TexCoords[CoordinateIndex] = UVArrays[CoordinateIndex];
+		}
+
+		#if HAS_MATERIAL_TEXCOORDS
+
+			float2 CustomizedUVs[NUM_SLATE_TEXCOORDS];
+			GetMaterialCustomizedUVs(VertexParameters, CustomizedUVs);
+
+			#if ES2_PROFILE
+				#if NUM_SLATE_TEXCOORDS == 1
+					VOut.TextureCoordinates[0].xy = CustomizedUVs[0];
+				#elif NUM_SLATE_TEXCOORDS >= 2
+					VOut.TextureCoordinates[0].xy = CustomizedUVs[0];
+					VOut.TextureCoordinates[0].zw = CustomizedUVs[1];
+				#endif
+			#else
+				UNROLL
+				for( int CoordinateIndex = 0; CoordinateIndex < NUM_SLATE_TEXCOORDS; CoordinateIndex++ )
 				{
-					VOut.TextureCoordinates[CoordinateIndex / 2].xy = CustomizedUVs[CoordinateIndex];
+					FLATTEN
+					if (CoordinateIndex % 2)
+					{
+						VOut.TextureCoordinates[CoordinateIndex / 2].zw = CustomizedUVs[CoordinateIndex];
+					}
+					else
+					{
+						VOut.TextureCoordinates[CoordinateIndex / 2].xy = CustomizedUVs[CoordinateIndex];
+					}
 				}
-			}
-		#endif
-	#endif // NUM_MATERIAL_TEXCOORDS
+			#endif
+		#endif // HAS_MATERIAL_TEXCOORDS
+	#endif // NUM_MATERIAL_TEXCOORDS_VERTEX
+
+	WorldPosition.xyz += GetMaterialWorldPositionOffset(VertexParameters);
+
 #endif // USE_MATERIALS
 
+	
+
 	VOut.Position = mul(WorldPosition, ViewProjection);
 
 #if ES2_PROFILE && COMPILER_GLSL_ES2
diff --git a/Engine/Source/Runtime/AppFramework/Private/Widgets/Testing/STestSuite.cpp b/Engine/Source/Runtime/AppFramework/Private/Widgets/Testing/STestSuite.cpp
--- a/Engine/Source/Runtime/AppFramework/Private/Widgets/Testing/STestSuite.cpp
+++ b/Engine/Source/Runtime/AppFramework/Private/Widgets/Testing/STestSuite.cpp
@@ -466,6 +466,12 @@ public:
 				]
 				+ SVerticalBox::Slot()
 				.FillHeight(3)
+				[
+					SNew(SCustomPaintWidget)
+					.OnPaintHandler(this, &SElementTesting::TestCustomVerts)
+				]
+				+ SVerticalBox::Slot()
+				.FillHeight(3)
 				[
 					SNew( SDynamicBrushTest )
 				]
@@ -698,6 +704,70 @@ private:
 
 		return InParams.Layer;
 	}
+
+	int32 TestCustomVerts(const FOnPaintHandlerParams& InParams)
+	{
+		const float Radius = FMath::Min(InParams.Geometry.Size.X, InParams.Geometry.Size.Y)*0.5f;
+		const FVector2D Center = InParams.Geometry.AbsolutePosition + InParams.Geometry.Size*0.5f;
+
+		const FSlateBrush* MyBrush = FCoreStyle::Get().GetBrush("ColorWheel.HueValueCircle");
+		FSlateShaderResourceProxy *ResourceProxy = FSlateDataPayload::ResourceManager->GetShaderResource(*MyBrush);
+		FVector2D UVCenter = FVector2D::ZeroVector;
+		FVector2D UVRadius = FVector2D(1,1);
+		if (ResourceProxy != nullptr)
+		{
+			UVRadius = 0.5f*ResourceProxy->SizeUV;
+			UVCenter = ResourceProxy->StartUV + UVRadius;
+		}
+
+		// Make a triangle fan in the area allotted
+		const int NumTris = 12;
+		TArray<FSlateVertex> Verts;
+		Verts.Reserve(NumTris*3);
+
+		// Center Vertex
+		Verts.AddZeroed();
+		{
+			FSlateVertex& NewVert = Verts.Last();
+			NewVert.Position[0] = Center.X;
+			NewVert.Position[1] = Center.Y;
+			NewVert.TexCoords[0] = UVCenter.X;
+			NewVert.TexCoords[1] = UVCenter.Y;
+			NewVert.TexCoords[2] = NewVert.TexCoords[3] = 1.0f;
+			NewVert.Color = FColor::White;
+			NewVert.ClipRect = FSlateRotatedRect(InParams.ClippingRect);
+		}
+
+		for (int i = 0; i < NumTris; ++i)
+		{
+			Verts.AddZeroed();
+			{
+				const float Angle = (2*PI*i) / NumTris;
+				const FVector2D EdgeDirection(FMath::Cos(Angle), FMath::Sin(Angle));
+				const FVector2D Edge(Radius*EdgeDirection);
+				FSlateVertex& NewVert = Verts.Last();
+				NewVert.Position[0] = Center.X + Edge.X;
+				NewVert.Position[1] = Center.Y + Edge.Y;
+				NewVert.TexCoords[0] = UVCenter.X + UVRadius.X*EdgeDirection.X;
+				NewVert.TexCoords[1] = UVCenter.Y + UVRadius.Y*EdgeDirection.Y;
+				NewVert.TexCoords[2] = NewVert.TexCoords[3] = 1.0f;
+				NewVert.Color = FColor::White;
+				NewVert.ClipRect = FSlateRotatedRect(InParams.ClippingRect);
+			}
+		}
+
+		TArray<SlateIndex> Indexes;
+		for (int i = 1; i <= NumTris; ++i)
+		{
+			Indexes.Add(0);
+			Indexes.Add(i);
+			Indexes.Add( (i+1 > 12) ? (1) : (i+1) );
+		}
+
+		FSlateDrawElement::MakeCustomVerts(InParams.OutDrawElements, InParams.Layer, MyBrush, Verts, Indexes, nullptr);
+
+		return InParams.Layer;
+	}
 };
 
 
diff --git a/Engine/Source/Runtime/SlateCore/Private/Rendering/DrawElements.cpp b/Engine/Source/Runtime/SlateCore/Private/Rendering/DrawElements.cpp
--- a/Engine/Source/Runtime/SlateCore/Private/Rendering/DrawElements.cpp
+++ b/Engine/Source/Runtime/SlateCore/Private/Rendering/DrawElements.cpp
@@ -203,6 +203,16 @@ void FSlateDrawElement::MakeCustom( FSlateWindowElementList& ElementList, uint32
 	DrawElt.DataPayload.SetCustomDrawerPayloadProperties( CustomDrawer );
 }
 
+void FSlateDrawElement::MakeCustomVerts(FSlateWindowElementList& ElementList, uint32 InLayer, const FSlateBrush* InBrush, const TArray<FSlateVertex>& InVerts, const TArray<SlateIndex>& InIndexes, ISlateUpdatableInstanceBuffer* InInstanceData)
+{
+	SCOPE_CYCLE_COUNTER(STAT_SlateDrawElementMakeTime)
+	FSlateDrawElement& DrawElt = ElementList.AddUninitialized();
+	DrawElt.Init(InLayer, FPaintGeometry(),/*Clip rect is ignored*/ FSlateRect(1,1,1,1), ESlateDrawEffect::None);
+	DrawElt.RenderTransform = FSlateRenderTransform();
+	DrawElt.ElementType = ET_CustomVerts;
+	DrawElt.DataPayload.SetCustomVertsPayloadProperties(InBrush, InVerts, InIndexes, InInstanceData);
+}
+
 void FSlateDrawElement::MakeCachedBuffer(FSlateWindowElementList& ElementList, uint32 InLayer, TSharedPtr<FSlateRenderDataHandle, ESPMode::ThreadSafe>& CachedRenderDataHandle, const FVector2D& Offset)
 {
 	SCOPE_CYCLE_COUNTER(STAT_SlateDrawElementMakeTime)
diff --git a/Engine/Source/Runtime/SlateCore/Private/Rendering/ElementBatcher.cpp b/Engine/Source/Runtime/SlateCore/Private/Rendering/ElementBatcher.cpp
--- a/Engine/Source/Runtime/SlateCore/Private/Rendering/ElementBatcher.cpp
+++ b/Engine/Source/Runtime/SlateCore/Private/Rendering/ElementBatcher.cpp
@@ -108,9 +108,10 @@ void FSlateElementBatcher::AddElements(FSlateDrawLayer& InDrawLayer)
 		FSlateDrawElement::EElementType::ET_Viewport == 6 &&
 		FSlateDrawElement::EElementType::ET_Border == 7 &&
 		FSlateDrawElement::EElementType::ET_Custom == 8 &&
-		FSlateDrawElement::EElementType::ET_CachedBuffer == 9 &&
-		FSlateDrawElement::EElementType::ET_Layer == 10 &&
-		FSlateDrawElement::EElementType::ET_Count == 11,
+		FSlateDrawElement::EElementType::ET_CustomVerts == 9 &&
+		FSlateDrawElement::EElementType::ET_CachedBuffer == 10 &&
+		FSlateDrawElement::EElementType::ET_Layer == 11 &&
+		FSlateDrawElement::EElementType::ET_Count == 12,
 		"If FSlateDrawElement::EElementType is modified, this array must be made to match." );
 
 	static FName ElementFNames[] =
@@ -124,6 +125,7 @@ void FSlateElementBatcher::AddElements(FSlateDrawLayer& InDrawLayer)
 		FName(TEXT("Viewport")),
 		FName(TEXT("Border")),
 		FName(TEXT("Custom")),
+		FName(TEXT("CustomVerts")),
 		FName(TEXT("CachedBuffer")),
 		FName(TEXT("Layer")),
 	};
@@ -164,6 +166,7 @@ void FSlateElementBatcher::AddElements(FSlateDrawLayer& InDrawLayer)
 			switch ( DrawElement.GetElementType() )
 			{
 			case FSlateDrawElement::ET_Custom:
+			case FSlateDrawElement::ET_CustomVerts:
 			case FSlateDrawElement::ET_CachedBuffer:
 			case FSlateDrawElement::ET_Layer:
 				break;
@@ -209,6 +212,9 @@ void FSlateElementBatcher::AddElements(FSlateDrawLayer& InDrawLayer)
 				case FSlateDrawElement::ET_Custom:
 					AddCustomElement( DrawElement );
 					break;
+				case FSlateDrawElement::ET_CustomVerts:
+					AddCustomVerts(DrawElement);
+					break;
 				case FSlateDrawElement::ET_CachedBuffer:
 					AddCachedBuffer( DrawElement );
 					break;
@@ -1558,6 +1564,53 @@ void FSlateElementBatcher::AddCustomElement( const FSlateDrawElement& DrawElemen
 	}
 }
 
+void FSlateElementBatcher::AddCustomVerts(const FSlateDrawElement& DrawElement)
+{
+	FElementBatchMap& LayerToElementBatches = DrawLayer->GetElementBatchMap();
+
+	const FSlateDataPayload& InPayload = DrawElement.GetDataPayload();
+	uint32 Layer = DrawElement.GetLayer();
+
+	if (InPayload.CustomVertsData.Num() >0)
+	{
+		// See if the layer already exists.
+		TUniqueObj<FElementBatchArray>* ElementBatches = LayerToElementBatches.Find(Layer);
+		if (!ElementBatches)
+		{
+			// The layer doesn't exist so make it now
+			ElementBatches = &LayerToElementBatches.Add( Layer );
+		}
+		check(ElementBatches);
+
+		FSlateElementBatch NewBatch(
+			InPayload.ResourceProxy != nullptr ? InPayload.ResourceProxy->Resource : nullptr,
+			FShaderParams(),
+			ESlateShader::Custom,
+			ESlateDrawPrimitive::TriangleList,
+			ESlateDrawEffect::None,
+			ESlateBatchDrawFlag::None,
+			DrawElement.GetScissorRect(),
+			InPayload.InstanceData == nullptr ? 1 : InPayload.InstanceData->GetNumInstances(),
+			InPayload.InstanceData
+		);
+
+		int32 Index = (*ElementBatches)->Add(NewBatch);
+		FSlateElementBatch* ElementBatch = &(**ElementBatches)[Index];
+
+		BatchData->AssignVertexArrayToBatch(*ElementBatch);
+		BatchData->AssignIndexArrayToBatch(*ElementBatch);
+
+		TArray<FSlateVertex>& BatchVertices = BatchData->GetBatchVertexList(*ElementBatch);
+		TArray<SlateIndex>& BatchIndices = BatchData->GetBatchIndexList(*ElementBatch);
+
+		// Vertex Buffer since  it is already in slate format it is a straight copy
+		BatchVertices = InPayload.CustomVertsData;
+		BatchIndices = InPayload.CustomVertsIndexData;
+	}
+
+
+}
+
 void FSlateElementBatcher::AddCachedBuffer(const FSlateDrawElement& DrawElement)
 {
 	FElementBatchMap& LayerToElementBatches = DrawLayer->GetElementBatchMap();
diff --git a/Engine/Source/Runtime/SlateCore/Public/Rendering/DrawElements.h b/Engine/Source/Runtime/SlateCore/Public/Rendering/DrawElements.h
--- a/Engine/Source/Runtime/SlateCore/Public/Rendering/DrawElements.h
+++ b/Engine/Source/Runtime/SlateCore/Public/Rendering/DrawElements.h
@@ -46,6 +46,7 @@ public:
 	// Box Data
 	FVector2D RotationPoint;
 	float Angle;
+	ISlateUpdatableInstanceBuffer* InstanceData;
 
 	// Spline/Line Data
 	float Thickness;
@@ -74,6 +75,10 @@ public:
 	// Custom drawer data
 	TWeakPtr<ICustomSlateElement, ESPMode::ThreadSafe> CustomDrawer;
 
+	// Custom verts data
+	TArray<FSlateVertex> CustomVertsData;
+	TArray<SlateIndex> CustomVertsIndexData;
+
 	// Cached render data
 	class FSlateRenderDataHandle* CachedRenderData;
 	FVector2D CachedRenderDataOffset;
@@ -92,6 +97,7 @@ public:
 		, BrushResource(nullptr)
 		, ResourceProxy(nullptr)
 		, RotationPoint(FVector2D::ZeroVector)
+		, InstanceData(nullptr)
 		, ImmutableText(nullptr)
 		, ViewportRenderTargetTexture(nullptr)
 		, bViewportTextureAlphaOnly(false)
@@ -101,10 +107,12 @@ public:
 		, CustomDrawer()
 	{ }
 
-	void SetBoxPayloadProperties( const FSlateBrush* InBrush, const FLinearColor& InTint, FSlateShaderResourceProxy* InResourceProxy = nullptr )
+	void SetBoxPayloadProperties( const FSlateBrush* InBrush, const FLinearColor& InTint, FSlateShaderResourceProxy* InResourceProxy = nullptr, ISlateUpdatableInstanceBuffer* InInstanceData = nullptr )
 	{
 		Tint = InTint;
 
+		InstanceData = InInstanceData;
+
 		BrushResource = InBrush;
 		if( InResourceProxy )
 		{
@@ -170,6 +178,14 @@ public:
 		CustomDrawer = InCustomDrawer;
 	}
 
+	void SetCustomVertsPayloadProperties(const FSlateBrush* InBrush, const TArray<FSlateVertex>& InVerts, const TArray<SlateIndex>& InIndexes, ISlateUpdatableInstanceBuffer* InInstanceData)
+	{
+		ResourceProxy = ResourceManager->GetShaderResource(*InBrush);
+		CustomVertsData = InVerts;
+		CustomVertsIndexData = InIndexes;
+		InstanceData = InInstanceData;
+	}
+
 	void SetCachedBufferPayloadProperties(FSlateRenderDataHandle* InRenderDataHandle, const FVector2D& Offset)
 	{
 		CachedRenderData = InRenderDataHandle;
@@ -202,6 +218,7 @@ public:
 		ET_Viewport,
 		ET_Border,
 		ET_Custom,
+		ET_CustomVerts,
 		ET_CachedBuffer,
 		ET_Layer,
 		ET_Count,
@@ -366,6 +383,8 @@ public:
 	 * @param CustomDrawer		   Interface to a drawer which will be called when Slate renders this element
 	 */
 	SLATECORE_API static void MakeCustom( FSlateWindowElementList& ElementList, uint32 InLayer, TSharedPtr<ICustomSlateElement, ESPMode::ThreadSafe> CustomDrawer );
+	
+	SLATECORE_API static void MakeCustomVerts(FSlateWindowElementList& ElementList, uint32 InLayer, const FSlateBrush* InBrush, const TArray<FSlateVertex>& InVerts, const TArray<SlateIndex>& InIndexes, ISlateUpdatableInstanceBuffer* InInstanceData);
 
 	SLATECORE_API static void MakeCachedBuffer(FSlateWindowElementList& ElementList, uint32 InLayer, TSharedPtr<FSlateRenderDataHandle, ESPMode::ThreadSafe>& CachedRenderDataHandle, const FVector2D& Offset);
 
@@ -469,8 +488,8 @@ private:
 class FSlateElementBatch
 {
 public: 
-	FSlateElementBatch( const FSlateShaderResource* InShaderResource, const FShaderParams& InShaderParams, ESlateShader::Type ShaderType, ESlateDrawPrimitive::Type PrimitiveType, ESlateDrawEffect::Type DrawEffects, ESlateBatchDrawFlag::Type DrawFlags, const TOptional<FShortRect>& ScissorRect )
-		: BatchKey( InShaderParams, ShaderType, PrimitiveType, DrawEffects, DrawFlags, ScissorRect )
+	FSlateElementBatch( const FSlateShaderResource* InShaderResource, const FShaderParams& InShaderParams, ESlateShader::Type ShaderType, ESlateDrawPrimitive::Type PrimitiveType, ESlateDrawEffect::Type DrawEffects, ESlateBatchDrawFlag::Type DrawFlags, const TOptional<FShortRect>& ScissorRect, int32 InstanceCount=1, ISlateUpdatableInstanceBuffer* InstanceData=nullptr )
+		: BatchKey( InShaderParams, ShaderType, PrimitiveType, DrawEffects, DrawFlags, ScissorRect, InstanceCount, InstanceData )
 		, ShaderResource(InShaderResource)
 		, NumElementsInBatch(0)
 		, VertexArrayIndex(INDEX_NONE)
@@ -522,6 +541,9 @@ public:
 	const TSharedPtr<FSlateRenderDataHandle, ESPMode::ThreadSafe> GetCachedRenderHandle() const { return BatchKey.CachedRenderHandle; }
 	FVector2D GetCachedRenderDataOffset() const { return BatchKey.CachedRenderDataOffset; }
 	const TSharedPtr<FSlateDrawLayerHandle, ESPMode::ThreadSafe> GetLayerHandle() const { return BatchKey.LayerHandle; }
+	int32 GetInstanceCount() const{ return BatchKey.InstanceCount; }
+	const ISlateUpdatableInstanceBuffer* GetInstanceData() const { return BatchKey.InstanceData; }
+
 private:
 	struct FBatchKey
 	{
@@ -535,14 +557,18 @@ private:
 		const ESlateDrawPrimitive::Type DrawPrimitiveType;
 		const ESlateDrawEffect::Type DrawEffects;
 		const TOptional<FShortRect> ScissorRect;
+		const int32 InstanceCount;
+		const ISlateUpdatableInstanceBuffer* InstanceData;
 
-		FBatchKey( const FShaderParams& InShaderParams, ESlateShader::Type InShaderType, ESlateDrawPrimitive::Type InDrawPrimitiveType, ESlateDrawEffect::Type InDrawEffects, ESlateBatchDrawFlag::Type InDrawFlags, const TOptional<FShortRect>& InScissorRect )
+		FBatchKey( const FShaderParams& InShaderParams, ESlateShader::Type InShaderType, ESlateDrawPrimitive::Type InDrawPrimitiveType, ESlateDrawEffect::Type InDrawEffects, ESlateBatchDrawFlag::Type InDrawFlags, const TOptional<FShortRect>& InScissorRect, int32 InInstanceCount, ISlateUpdatableInstanceBuffer* InInstanceBuffer)
 			: ShaderParams( InShaderParams )
 			, DrawFlags( InDrawFlags )
 			, ShaderType( InShaderType )
 			, DrawPrimitiveType( InDrawPrimitiveType )
 			, DrawEffects( InDrawEffects )
 			, ScissorRect( InScissorRect )
+			, InstanceCount( InInstanceCount )
+			, InstanceData( InInstanceBuffer )
 		{
 		}
 
@@ -554,6 +580,7 @@ private:
 			, DrawPrimitiveType( ESlateDrawPrimitive::TriangleList )
 			, DrawEffects( ESlateDrawEffect::None )
 			, ScissorRect( InScissorRect )
+			, InstanceCount(1)
 		{}
 
 		FBatchKey( TSharedPtr<FSlateRenderDataHandle, ESPMode::ThreadSafe> InCachedRenderHandle, FVector2D InCachedRenderDataOffset, const TOptional<FShortRect>& InScissorRect)
@@ -565,6 +592,7 @@ private:
 			, DrawPrimitiveType(ESlateDrawPrimitive::TriangleList)
 			, DrawEffects(ESlateDrawEffect::None)
 			, ScissorRect(InScissorRect)
+			, InstanceCount(1)
 		{}
 
 		FBatchKey(TSharedPtr<FSlateDrawLayerHandle, ESPMode::ThreadSafe> InLayerHandle, const TOptional<FShortRect>& InScissorRect)
@@ -575,6 +603,7 @@ private:
 			, DrawPrimitiveType(ESlateDrawPrimitive::TriangleList)
 			, DrawEffects(ESlateDrawEffect::None)
 			, ScissorRect(InScissorRect)
+			, InstanceCount(1)
 		{}
 
 		bool operator==( const FBatchKey& Other ) const
@@ -588,7 +617,8 @@ private:
 				&& CustomDrawer == Other.CustomDrawer
 				&& CachedRenderHandle == Other.CachedRenderHandle
 				&& LayerHandle == Other.LayerHandle
-				;
+				&& InstanceCount == Other.InstanceCount
+				&& InstanceData == Other.InstanceData;
 		}
 
 		/** Compute an efficient hash for this type for use in hash containers. */
@@ -601,6 +631,9 @@ private:
 			RunningHash = HashCombine(GetTypeHash(InBatchKey.ShaderParams.PixelParams), RunningHash);
 			// NOTE: Assumes this type is 64 bits, no padding.
 			RunningHash = InBatchKey.ScissorRect.IsSet() ? HashCombine(GetTypeHash(*reinterpret_cast<const uint64*>(&InBatchKey.ScissorRect.GetValue())), RunningHash) : RunningHash;
+			RunningHash = InBatchKey.InstanceCount > 1 ? HashCombine( InBatchKey.InstanceCount, RunningHash ) : RunningHash;
+			RunningHash = InBatchKey.InstanceData ? HashCombine( PointerHash(InBatchKey.InstanceData), RunningHash ) : RunningHash;
+
 			return RunningHash;
 			//return FCrc::MemCrc32(&InBatchKey.ShaderParams, sizeof(FShaderParams)) ^ ((InBatchKey.ShaderType << 16) | (InBatchKey.DrawFlags+InBatchKey.ShaderType+InBatchKey.DrawPrimitiveType+InBatchKey.DrawEffects));
 		}
@@ -627,6 +660,8 @@ public:
 		: Layer( InLayer )
 		, ShaderParams( InBatch.GetShaderParams() )
 		, Texture( InBatch.GetShaderResource() )
+		, InstanceData( InBatch.GetInstanceData() )
+		, InstanceCount( InBatch.GetInstanceCount() )
 		, CustomDrawer( InBatch.GetCustomDrawer() )
 		, LayerHandle( InBatch.GetLayerHandle() )
 		, CachedRenderHandle( InRenderHandle )
@@ -655,6 +690,10 @@ public:
 	/** Texture to use with this batch.  */
 	const FSlateShaderResource* Texture;
 
+	const ISlateUpdatableInstanceBuffer* InstanceData;
+
+	const int32 InstanceCount;
+
 	const TWeakPtr<ICustomSlateElement, ESPMode::ThreadSafe> CustomDrawer;
 
 	const TWeakPtr<FSlateDrawLayerHandle, ESPMode::ThreadSafe> LayerHandle;
diff --git a/Engine/Source/Runtime/SlateCore/Public/Rendering/ElementBatcher.h b/Engine/Source/Runtime/SlateCore/Public/Rendering/ElementBatcher.h
--- a/Engine/Source/Runtime/SlateCore/Public/Rendering/ElementBatcher.h
+++ b/Engine/Source/Runtime/SlateCore/Public/Rendering/ElementBatcher.h
@@ -89,6 +89,8 @@ private:
 	 */
 	void AddCustomElement( const FSlateDrawElement& DrawElement );
 
+	void AddCustomVerts( const FSlateDrawElement& DrawElement );
+
 	void AddCachedBuffer( const FSlateDrawElement& DrawElement );
 
 	void AddLayer(const FSlateDrawElement& DrawElement);
diff --git a/Engine/Source/Runtime/SlateCore/Public/Rendering/RenderingCommon.h b/Engine/Source/Runtime/SlateCore/Public/Rendering/RenderingCommon.h
--- a/Engine/Source/Runtime/SlateCore/Public/Rendering/RenderingCommon.h
+++ b/Engine/Source/Runtime/SlateCore/Public/Rendering/RenderingCommon.h
@@ -41,8 +41,10 @@ namespace ESlateShader
 	const Type Border = 1;
 	/** Font shader, same as default except uses an alpha only texture */
 	const Type Font = 2;
-	/** Line segment shader. For drawing anti-aliased lines **/
+	/** Line segment shader. For drawing anti-aliased lines */
 	const Type LineSegment = 3;
+	/** For completely customized materials.  Makes no assumptions on use*/
+	const Type Custom = 4;
 };
 
 /**
@@ -96,7 +98,7 @@ namespace ESlateLineJoinType
  * We provide a ctor that does the work common to slate drawing, but you could technically 
  * create this any way you want.
  */
-struct FSlateRotatedRect
+struct SLATECORE_API FSlateRotatedRect
 {
 	/** Default ctor. */
 	FSlateRotatedRect();
@@ -535,3 +537,83 @@ public:
 	 */
 	virtual void DrawRenderThread(class FRHICommandListImmediate& RHICmdList, const void* RenderTarget) = 0;
 };
+
+/**
+ * Represents a per instance data buffer for a custom Slate mesh element.
+ * Use FSlateInstanceBufferUpdate to update the per-instance data.
+ *  e.g.
+ *    TSharedRef<FSlateInstanceBufferUpdate> NewUpdate = InstanceBuffer.BeginUpdate();
+ *     NewUpdate.GetData().Add( FVector4(1,1,1,1) )
+ *     FSlateInstanceBufferUpdate::CommitUpdate(NewUpdate);
+ */
+class ISlateUpdatableInstanceBuffer
+{
+public:
+	virtual ~ISlateUpdatableInstanceBuffer(){};
+	friend class FSlateInstanceBufferUpdate;
+
+   /**
+	* Use this method to begin a new update to this instance of the buffer:
+	*/
+	virtual TSharedPtr<class FSlateInstanceBufferUpdate> BeginUpdate() = 0;
+
+	/** How many instances should we draw? */
+	virtual int32 GetNumInstances() const = 0;
+
+private:
+	friend class FSlateInstanceBufferUpdate;
+
+	/** Updates rendering data for the GPU */
+	virtual void UpdateRenderingData(int32 NumInstancesToUse) = 0;
+
+	/** @return an array of instance data that is safe to populate (e.g not in use by the renderer) */
+	virtual TArray<FVector4>& GetBufferData() = 0;
+};
+
+/** Represents an update to the per-instance buffer. */
+class FSlateInstanceBufferUpdate
+{
+public:
+	/** Access the per-instance data for modiciation */
+	FORCEINLINE TArray<FVector4>& GetData(){ return Data; }
+	
+	/** Set number of instances; if it is not set, Data.Num() will be used instead.  */
+	void SetInstanceCount(int32 InCount){ InstanceCount = InCount; }
+	
+	/** Send an update to the render thread */
+	static void CommitUpdate(TSharedPtr<FSlateInstanceBufferUpdate>& UpdateToCommit)
+	{
+		ensure(UpdateToCommit.GetSharedReferenceCount() == 1);
+		UpdateToCommit->CommitUpdate_Internal();
+		UpdateToCommit.Reset();
+	}
+
+	~FSlateInstanceBufferUpdate()
+	{
+		if (!bWasCommitted)
+		{
+			CommitUpdate_Internal();
+		}
+	}
+
+private:
+	friend class FSlateUpdatableInstanceBuffer;
+	FSlateInstanceBufferUpdate(ISlateUpdatableInstanceBuffer& InBuffer)
+		: Buffer(InBuffer)
+		, Data(InBuffer.GetBufferData())
+		, InstanceCount(INDEX_NONE)
+		, bWasCommitted(false)
+	{
+	}
+
+	void CommitUpdate_Internal()
+	{
+		Buffer.UpdateRenderingData(InstanceCount == INDEX_NONE ? Data.Num() : InstanceCount);
+		bWasCommitted = true;
+	}
+
+	ISlateUpdatableInstanceBuffer& Buffer;
+	TArray<FVector4>& Data;
+	int32 InstanceCount;
+	bool bWasCommitted;
+};
diff --git a/Engine/Source/Runtime/SlateCore/Public/Styling/SlateBrush.h b/Engine/Source/Runtime/SlateCore/Public/Styling/SlateBrush.h
--- a/Engine/Source/Runtime/SlateCore/Public/Styling/SlateBrush.h
+++ b/Engine/Source/Runtime/SlateCore/Public/Styling/SlateBrush.h
@@ -163,6 +163,8 @@ public:
 		, UVRegion(ForceInit)
 	{ }
 
+	virtual ~FSlateBrush(){}
+
 public:
 
 	/**
diff --git a/Engine/Source/Runtime/SlateRHIRenderer/Private/SlateMaterialShader.cpp b/Engine/Source/Runtime/SlateRHIRenderer/Private/SlateMaterialShader.cpp
--- a/Engine/Source/Runtime/SlateRHIRenderer/Private/SlateMaterialShader.cpp
+++ b/Engine/Source/Runtime/SlateRHIRenderer/Private/SlateMaterialShader.cpp
@@ -153,6 +153,8 @@ IMPLEMENT_SLATE_VERTEXMATERIALSHADER_TYPE(false);
 	typedef TSlateMaterialShaderPS<ESlateShader::ShaderType, bDrawDisabledEffect> TSlateMaterialShaderPS##ShaderType##bDrawDisabledEffect; \
 	IMPLEMENT_MATERIAL_SHADER_TYPE(template<>, TSlateMaterialShaderPS##ShaderType##bDrawDisabledEffect, TEXT("SlateElementPixelShader"), TEXT("Main"), SF_Pixel);
 
+IMPLEMENT_SLATE_MATERIALSHADER_TYPE(Custom, false)
+
 IMPLEMENT_SLATE_MATERIALSHADER_TYPE(Default,true);
 IMPLEMENT_SLATE_MATERIALSHADER_TYPE(Default,false);
 IMPLEMENT_SLATE_MATERIALSHADER_TYPE(Border,true);
diff --git a/Engine/Source/Runtime/SlateRHIRenderer/Private/SlateRHIRendererModule.cpp b/Engine/Source/Runtime/SlateRHIRenderer/Private/SlateRHIRendererModule.cpp
--- a/Engine/Source/Runtime/SlateRHIRenderer/Private/SlateRHIRendererModule.cpp
+++ b/Engine/Source/Runtime/SlateRHIRenderer/Private/SlateRHIRendererModule.cpp
@@ -2,6 +2,7 @@
 
 #include "SlateRHIRendererPrivatePCH.h"
 #include "Slate3DRenderer.h"
+#include "SlateUpdatableBuffer.h"
 
 class FSlateRHIFontAtlasFactory : public ISlateFontAtlasFactory
 {
@@ -71,6 +72,11 @@ public:
 		return MakeShareable(new FSlateRHIFontAtlasFactory);
 	}
 
+	virtual TSharedRef<ISlateUpdatableInstanceBuffer> CreateInstanceBuffer( int32 InitialInstanceCount ) override
+	{
+		return MakeShareable( new FSlateUpdatableInstanceBuffer(InitialInstanceCount) );
+	}
+
 	virtual void StartupModule( ) override { }
 	virtual void ShutdownModule( ) override { }
 
diff --git a/Engine/Source/Runtime/SlateRHIRenderer/Private/SlateRHIRenderingPolicy.cpp b/Engine/Source/Runtime/SlateRHIRenderer/Private/SlateRHIRenderingPolicy.cpp
--- a/Engine/Source/Runtime/SlateRHIRenderer/Private/SlateRHIRenderingPolicy.cpp
+++ b/Engine/Source/Runtime/SlateRHIRenderer/Private/SlateRHIRenderingPolicy.cpp
@@ -11,6 +11,7 @@
 #include "SlateUTextureResource.h"
 #include "SlateMaterialResource.h"
 #include "Rendering/DrawElements.h"
+#include "SlateUpdatableBuffer.h"
 
 DECLARE_CYCLE_STAT(TEXT("Update Buffers RT"), STAT_SlateUpdateBufferRTTime, STATGROUP_Slate);
 DECLARE_CYCLE_STAT(TEXT("PreFill Buffers RT"), STAT_SlatePreFullBufferRTTime, STATGROUP_Slate);
@@ -701,11 +702,21 @@ void FSlateRHIRenderingPolicy::DrawElements(FRHICommandListImmediate& RHICmdList
 				if (!GRHISupportsBaseVertexIndex)
 				{
 					RHICmdList.SetStreamSource(0, VertexBuffer->VertexBufferRHI, sizeof(FSlateVertex), RenderBatch.VertexOffset * sizeof(FSlateVertex));
-					RHICmdList.DrawIndexedPrimitive(IndexBuffer->IndexBufferRHI, GetRHIPrimitiveType(RenderBatch.DrawPrimitiveType), 0, 0, RenderBatch.NumVertices, RenderBatch.IndexOffset, PrimitiveCount, 1);
+					RHICmdList.DrawIndexedPrimitive(IndexBuffer->IndexBufferRHI, GetRHIPrimitiveType(RenderBatch.DrawPrimitiveType), 0, 0, RenderBatch.NumVertices, RenderBatch.IndexOffset, PrimitiveCount, RenderBatch.InstanceCount);
 				}
 				else
 				{
-					RHICmdList.DrawIndexedPrimitive(IndexBuffer->IndexBufferRHI, GetRHIPrimitiveType(RenderBatch.DrawPrimitiveType), RenderBatch.VertexOffset, 0, RenderBatch.NumVertices, RenderBatch.IndexOffset, PrimitiveCount, 1);
+					if( RenderBatch.InstanceCount > 1 && RenderBatch.InstanceData != nullptr )
+					{
+						FSlateUpdatableInstanceBuffer* InstanceBuffer = (FSlateUpdatableInstanceBuffer*)RenderBatch.InstanceData;
+						InstanceBuffer->BindStreamSource( RHICmdList, 1 );
+					}
+					else
+					{
+						RHICmdList.SetStreamSource(1, nullptr, 0, 0);
+					}
+
+					RHICmdList.DrawIndexedPrimitive(IndexBuffer->IndexBufferRHI, GetRHIPrimitiveType(RenderBatch.DrawPrimitiveType), RenderBatch.VertexOffset, 0, RenderBatch.NumVertices, RenderBatch.IndexOffset, PrimitiveCount, RenderBatch.InstanceCount);
 				}
 
 			}
@@ -724,7 +735,7 @@ void FSlateRHIRenderingPolicy::DrawElements(FRHICommandListImmediate& RHICmdList
 
 				FSlateMaterialShaderPS* PixelShader = GetMaterialPixelShader( Material, ShaderType, DrawEffects );
 
-				const bool bUseInstancing = false;
+				const bool bUseInstancing = RenderBatch.InstanceCount > 1;
 				FSlateMaterialShaderVS* VertexShader = GetMaterialVertexShader( Material, bUseInstancing );
 
 				if( VertexShader && PixelShader )
@@ -761,11 +772,21 @@ void FSlateRHIRenderingPolicy::DrawElements(FRHICommandListImmediate& RHICmdList
 					if (!GRHISupportsBaseVertexIndex)
 					{
 						RHICmdList.SetStreamSource(0, VertexBuffer->VertexBufferRHI, sizeof(FSlateVertex), RenderBatch.VertexOffset * sizeof(FSlateVertex));
-						RHICmdList.DrawIndexedPrimitive(IndexBuffer->IndexBufferRHI, GetRHIPrimitiveType(RenderBatch.DrawPrimitiveType), 0, 0, RenderBatch.NumVertices, RenderBatch.IndexOffset, PrimitiveCount, 1);
+						RHICmdList.DrawIndexedPrimitive(IndexBuffer->IndexBufferRHI, GetRHIPrimitiveType(RenderBatch.DrawPrimitiveType), 0, 0, RenderBatch.NumVertices, RenderBatch.IndexOffset, PrimitiveCount, RenderBatch.InstanceCount);
 					}
 					else
 					{
-						RHICmdList.DrawIndexedPrimitive(IndexBuffer->IndexBufferRHI, GetRHIPrimitiveType(RenderBatch.DrawPrimitiveType), RenderBatch.VertexOffset, 0, RenderBatch.NumVertices, RenderBatch.IndexOffset, PrimitiveCount, 1);
+						if( bUseInstancing && RenderBatch.InstanceData != nullptr )
+						{
+							FSlateUpdatableInstanceBuffer* InstanceBuffer = (FSlateUpdatableInstanceBuffer*)RenderBatch.InstanceData;
+							InstanceBuffer->BindStreamSource( RHICmdList, 1 );
+						}
+						else
+						{
+							RHICmdList.SetStreamSource(1, nullptr, 0, 0);
+						}
+
+						RHICmdList.DrawIndexedPrimitive(IndexBuffer->IndexBufferRHI, GetRHIPrimitiveType(RenderBatch.DrawPrimitiveType), RenderBatch.VertexOffset, 0, RenderBatch.NumVertices, RenderBatch.IndexOffset, PrimitiveCount, RenderBatch.InstanceCount);
 					}
 				}
 			}
@@ -782,7 +803,6 @@ void FSlateRHIRenderingPolicy::DrawElements(FRHICommandListImmediate& RHICmdList
 				// Something may have messed with the viewport size so set it back to the full target.
 				RHICmdList.SetViewport( 0,0,0,BackBuffer.GetSizeXY().X, BackBuffer.GetSizeXY().Y, 0.0f ); 
 				RHICmdList.SetStreamSource(0, VertexBuffer->VertexBufferRHI, sizeof(FSlateVertex), 0);
-
 			}
 		}
 
@@ -914,6 +934,11 @@ FSlateMaterialShaderPS* FSlateRHIRenderingPolicy::GetMaterialPixelShader( const
 			FoundShader = MaterialShaderMap->GetShader(&TSlateMaterialShaderPS<ESlateShader::Font, false>::StaticType);
 		}
 		break;
+	case ESlateShader::Custom:
+		{
+			FoundShader = MaterialShaderMap->GetShader(&TSlateMaterialShaderPS<ESlateShader::Custom, false>::StaticType);
+		}
+		break;
 	default:
 		checkf(false, TEXT("Unsupported Slate shader type for use with materials"));
 		break;
diff --git a/Engine/Source/Runtime/SlateRHIRenderer/Private/SlateRHIRenderingPolicy.h b/Engine/Source/Runtime/SlateRHIRenderer/Private/SlateRHIRenderingPolicy.h
--- a/Engine/Source/Runtime/SlateRHIRenderer/Private/SlateRHIRenderingPolicy.h
+++ b/Engine/Source/Runtime/SlateRHIRenderer/Private/SlateRHIRenderingPolicy.h
@@ -7,11 +7,6 @@ class FSlateRHIResourceManager;
 
 #include "RenderingPolicy.h"
 
-namespace SlateRHIConstants
-{
-	// Number of vertex and index buffers we swap between when drawing windows
-	const int32 NumBuffers = 3;
-}
 
 /** 
  * Vertex buffer containing all Slate vertices
diff --git a/Engine/Source/Runtime/SlateRHIRenderer/Private/SlateShaders.cpp b/Engine/Source/Runtime/SlateRHIRenderer/Private/SlateShaders.cpp
--- a/Engine/Source/Runtime/SlateRHIRenderer/Private/SlateShaders.cpp
+++ b/Engine/Source/Runtime/SlateRHIRenderer/Private/SlateShaders.cpp
@@ -49,6 +49,9 @@ void FSlateVertexDeclaration::InitRHI()
 	Elements.Add(FVertexElement(0, STRUCT_OFFSET(FSlateVertex, ClipRect) + STRUCT_OFFSET(FSlateRotatedClipRectType, TopLeft), VET_Float2, 3, Stride));
 	Elements.Add(FVertexElement(0, STRUCT_OFFSET(FSlateVertex, ClipRect) + STRUCT_OFFSET(FSlateRotatedClipRectType, ExtentX), VET_Float4, 4, Stride));
 	Elements.Add(FVertexElement(0, STRUCT_OFFSET(FSlateVertex, Color), VET_Color, 5, Stride));
+	
+	static const bool bUseInstanceIndex = true;
+	Elements.Add(FVertexElement(1, 0, VET_Float4, 6, sizeof(FVector4), bUseInstanceIndex));
 
 	VertexDeclarationRHI = RHICreateVertexDeclaration(Elements);
 }
diff --git a/Engine/Source/Runtime/SlateRHIRenderer/Private/SlateUpdatableBuffer.cpp b/Engine/Source/Runtime/SlateRHIRenderer/Private/SlateUpdatableBuffer.cpp
new file mode 100644
--- /dev/null
+++ b/Engine/Source/Runtime/SlateRHIRenderer/Private/SlateUpdatableBuffer.cpp
@@ -0,0 +1,101 @@
+// Copyright 1998-2015 Epic Games, Inc. All Rights Reserved.
+
+#include "SlateRHIRendererPrivatePCH.h"
+#include "SlateUpdatableBuffer.h"
+
+DECLARE_CYCLE_STAT(TEXT("UpdateInstanceBuffer Time"), STAT_SlateUpdateInstanceBuffer, STATGROUP_Slate);
+
+struct FSlateUpdateInstanceBufferCommand : public FRHICommand<FSlateUpdateInstanceBufferCommand>
+{
+	TSlateElementVertexBuffer<FVector4>& InstanceBuffer;
+	const TArray<FVector4>& InstanceData;
+
+	FSlateUpdateInstanceBufferCommand(TSlateElementVertexBuffer<FVector4>& InInstanceBuffer, const TArray<FVector4>& InInstanceData )
+		: InstanceBuffer(InInstanceBuffer)
+		, InstanceData(InInstanceData)
+	{}
+
+	void Execute(FRHICommandListBase& CmdList)
+	{
+		SCOPE_CYCLE_COUNTER( STAT_SlateUpdateInstanceBuffer );
+
+		uint8* InstanceBufferData = (uint8*)InstanceBuffer.LockBuffer_RenderThread(InstanceData.Num());
+
+		FMemory::Memcpy(InstanceBufferData, InstanceData.GetData(), InstanceData.Num()*sizeof(FVector4) );
+	
+		InstanceBuffer.UnlockBuffer_RenderThread();
+	}
+};
+
+FSlateUpdatableInstanceBuffer::FSlateUpdatableInstanceBuffer( int32 InitialInstanceCount )
+	: FreeBufferIndex(0)
+{
+	InstanceBufferResource.Init(InitialInstanceCount);
+	for( int32 BufferIndex = 0; BufferIndex < SlateRHIConstants::NumBuffers; ++BufferIndex )
+	{
+		BufferData[BufferIndex].Reserve( InitialInstanceCount );
+	}
+
+}
+
+FSlateUpdatableInstanceBuffer::~FSlateUpdatableInstanceBuffer()
+{
+	InstanceBufferResource.Destroy();
+	FlushRenderingCommands();
+}
+
+void FSlateUpdatableInstanceBuffer::BindStreamSource(FRHICommandListImmediate& RHICmdList, int32 StreamIndex)
+{
+	RHICmdList.SetStreamSource(StreamIndex, InstanceBufferResource.VertexBufferRHI, sizeof(FVector4), 0);
+}
+
+TSharedPtr<class FSlateInstanceBufferUpdate> FSlateUpdatableInstanceBuffer::BeginUpdate()
+{
+	return MakeShareable(new FSlateInstanceBufferUpdate(*this));
+}
+
+int32 FSlateUpdatableInstanceBuffer::GetNumInstances() const
+{
+	return NumInstances;
+}
+
+void FSlateUpdatableInstanceBuffer::UpdateRenderingData(int32 NumInstancesToUse)
+{
+	NumInstances = NumInstancesToUse;
+	// Enqueue a command to unlock the draw buffer after all windows have been drawn
+	ENQUEUE_UNIQUE_RENDER_COMMAND_TWOPARAMETER( SlateBeginDrawingWindowsCommand, 
+		FSlateUpdatableInstanceBuffer&, Self, *this,
+		int32, BufferIndex, FreeBufferIndex,
+	{
+		Self.UpdateRenderingData_RenderThread(RHICmdList, BufferIndex);
+	});
+	
+	FreeBufferIndex = (FreeBufferIndex + 1) % SlateRHIConstants::NumBuffers;
+}
+
+TArray<FVector4>& FSlateUpdatableInstanceBuffer::GetBufferData()
+{
+	return BufferData[FreeBufferIndex];
+}
+
+void FSlateUpdatableInstanceBuffer::UpdateRenderingData_RenderThread(FRHICommandListImmediate& RHICmdList, int32 BufferIndex)
+{
+	SCOPE_CYCLE_COUNTER( STAT_SlateUpdateInstanceBuffer );
+
+	const TArray<FVector4>& RenderThreadBufferData = BufferData[BufferIndex];
+	InstanceBufferResource.PreFillBuffer( RenderThreadBufferData.Num(), false );
+
+	if(!GRHIThread || RHICmdList.Bypass())
+	{
+		uint8* InstanceBufferData = (uint8*)InstanceBufferResource.LockBuffer_RenderThread(RenderThreadBufferData.Num());
+
+		FMemory::Memcpy(InstanceBufferData, RenderThreadBufferData.GetData(), RenderThreadBufferData.Num()*sizeof(FVector4) );
+	
+		InstanceBufferResource.UnlockBuffer_RenderThread();
+	}
+	else
+	{
+		new (RHICmdList.AllocCommand<FSlateUpdateInstanceBufferCommand>()) FSlateUpdateInstanceBufferCommand(InstanceBufferResource, RenderThreadBufferData);
+	}
+}
+
diff --git a/Engine/Source/Runtime/SlateRHIRenderer/Private/SlateUpdatableBuffer.h b/Engine/Source/Runtime/SlateRHIRenderer/Private/SlateUpdatableBuffer.h
new file mode 100644
--- /dev/null
+++ b/Engine/Source/Runtime/SlateRHIRenderer/Private/SlateUpdatableBuffer.h
@@ -0,0 +1,33 @@
+// Copyright 1998-2015 Epic Games, Inc. All Rights Reserved.
+
+#pragma once
+
+/**
+ * Represents a per instance data buffer for a custom Slate mesh element
+ */
+class FSlateUpdatableInstanceBuffer : public ISlateUpdatableInstanceBuffer
+{
+public:
+
+	FSlateUpdatableInstanceBuffer( int32 InstanceCount );
+	~FSlateUpdatableInstanceBuffer();
+	void BindStreamSource( FRHICommandListImmediate& RHICmdList, int32 StreamIndex);
+
+private:
+	// BEGIN ISlateUpdatableInstanceBuffer
+	virtual TSharedPtr<class FSlateInstanceBufferUpdate> BeginUpdate() override;
+	virtual int32 GetNumInstances() const override;
+	virtual void UpdateRenderingData(int32 NumInstancesToUse) override;
+	virtual TArray<FVector4>& GetBufferData() override;
+	// END ISlateUpdatableInstanceBuffer
+
+	void UpdateRenderingData_RenderThread(FRHICommandListImmediate& RHICmdList, int32 BufferIndex);
+private:
+	TArray<FVector4> BufferData[SlateRHIConstants::NumBuffers];
+	TSlateElementVertexBuffer<FVector4> InstanceBufferResource;
+	int32 NumInstances;
+	int32 FreeBufferIndex;
+};
+
+
+
diff --git a/Engine/Source/Runtime/SlateRHIRenderer/Public/Interfaces/ISlateRHIRendererModule.h b/Engine/Source/Runtime/SlateRHIRenderer/Public/Interfaces/ISlateRHIRendererModule.h
--- a/Engine/Source/Runtime/SlateRHIRenderer/Public/Interfaces/ISlateRHIRendererModule.h
+++ b/Engine/Source/Runtime/SlateRHIRenderer/Public/Interfaces/ISlateRHIRendererModule.h
@@ -8,6 +8,12 @@
 class ISlate3DRenderer;
 class ISlateFontAtlasFactory;
 
+namespace SlateRHIConstants
+{
+	// Number of vertex and index buffers we swap between when drawing windows
+	const int32 NumBuffers = 3;
+}
+
 /**
  * Interface for the Slate RHI Renderer module.
  */
@@ -26,4 +32,6 @@ public:
 	virtual TSharedRef<ISlate3DRenderer> CreateSlate3DRenderer() = 0;
 
 	virtual TSharedRef<ISlateFontAtlasFactory> CreateSlateFontAtlasFactory() = 0;
+
+	virtual TSharedRef<ISlateUpdatableInstanceBuffer> CreateInstanceBuffer( int32 InitialInstanceCount ) = 0;
 };
diff --git a/Engine/Source/Runtime/UMG/Private/Slate/SMeshWidget.cpp b/Engine/Source/Runtime/UMG/Private/Slate/SMeshWidget.cpp
new file mode 100644
--- /dev/null
+++ b/Engine/Source/Runtime/UMG/Private/Slate/SMeshWidget.cpp
@@ -0,0 +1,135 @@
+// Copyright 1998-2015 Epic Games, Inc. All Rights Reserved.
+
+#include "UMGPrivatePCH.h"
+
+#include "SMeshWidget.h"
+#include "Engine/StaticMesh.h"
+#include "StaticMeshResources.h"
+#include "SlateMaterialBrush.h"
+#include "Runtime/SlateRHIRenderer/Public/Interfaces/ISlateRHIRendererModule.h"
+
+
+void SMeshWidget::Construct(const FArguments& Args)
+{
+	if (Args._StaticMeshSource != nullptr)
+	{
+		SetMesh(*Args._StaticMeshSource);
+	}
+}
+
+static const FVector2D DontCare(FVector2D(64, 64));
+void SMeshWidget::SetMesh(const UStaticMesh& InStaticMesh)
+{	
+	Brush = MakeShareable( new FSlateMaterialBrush(*InStaticMesh.GetMaterial(0), DontCare) );
+	SMeshWidget::StaticMeshToSlateRenderData(InStaticMesh, VertexData, IndexData);
+}
+
+static const FName SlateRHIModuleName("SlateRHIRenderer");
+TSharedPtr<FSlateInstanceBufferUpdate> SMeshWidget::BeginPerInstanceBufferUpdate(int32 InitialSize)
+{
+	if (!PerInstanceBuffer.IsValid())
+	{
+		PerInstanceBuffer = FModuleManager::Get().GetModuleChecked<ISlateRHIRendererModule>(SlateRHIModuleName).CreateInstanceBuffer(InitialSize);
+	}
+
+	return PerInstanceBuffer->BeginUpdate();
+}
+
+
+int32 SMeshWidget::OnPaint(const FPaintArgs& Args, const FGeometry& AllottedGeometry, const FSlateRect& MyClippingRect, FSlateWindowElementList& OutDrawElements, int32 LayerId, const FWidgetStyle& InWidgetStyle, bool bParentEnabled) const
+{
+	if (Brush.IsValid() && VertexData.Num() > 0 && IndexData.Num() > 0)
+	{
+		FSlateDrawElement::MakeCustomVerts(OutDrawElements, LayerId, Brush.Get(), VertexData, IndexData, PerInstanceBuffer.Get());
+	}
+
+	return LayerId;
+}
+
+
+FVector2D SMeshWidget::ComputeDesiredSize(float) const
+{
+	return FVector2D(256,256);
+}
+
+
+void SMeshWidget::StaticMeshToSlateRenderData(const UStaticMesh& DataSource, TArray<FSlateVertex>& OutSlateVerts, TArray<SlateIndex>& OutIndexes)
+{
+	const FStaticMeshLODResources& LOD = DataSource.RenderData->LODResources[0];
+	const int32 NumSections = LOD.Sections.Num();
+	if (NumSections > 1)
+	{
+		UE_LOG(LogUMG, Warning, TEXT("StaticMesh %s has %d sections. SMeshWidget expects a static mesh with 1 section."), *DataSource.GetName(), NumSections);
+	}
+	else
+	{
+		// Populate Index data
+		{
+			FIndexArrayView SourceIndexes = LOD.IndexBuffer.GetArrayView();
+			const int32 NumIndexes = SourceIndexes.Num();
+			OutIndexes.Empty();
+			OutIndexes.Reserve(NumIndexes);
+			for (int32 i = 0; i < NumIndexes; ++i)
+			{
+				OutIndexes.Add(SourceIndexes[i]);
+			}
+		}
+
+		// Populate Vertex Data
+		{
+			const uint32 NumVerts = LOD.PositionVertexBuffer.GetNumVertices();
+			OutSlateVerts.Empty();
+			OutSlateVerts.Reserve(NumVerts);
+
+			static const int32 MAX_SUPPORTED_UV_SETS = 6;
+			const int32 TexCoordsPerVertex = LOD.GetNumTexCoords();
+			if (TexCoordsPerVertex > MAX_SUPPORTED_UV_SETS)
+			{
+				UE_LOG(LogStaticMesh, Warning, TEXT("[%s] has %d UV sets; slate vertex data supports at most %d"), *DataSource.GetName(), TexCoordsPerVertex, MAX_SUPPORTED_UV_SETS);
+			}
+
+			for (uint32 i = 0; i < NumVerts; ++i)
+			{
+				FSlateVertex& NewVert = OutSlateVerts[OutSlateVerts.AddUninitialized()];
+
+				// Copy Position
+				{
+					const FVector& Position = LOD.PositionVertexBuffer.VertexPosition(i);
+					NewVert.Position[0] = Position.X;
+					NewVert.Position[1] = Position.Y;
+				}
+
+				// Copy Color
+				if (LOD.ColorVertexBuffer.GetNumVertices() > 0)
+				{
+					NewVert.Color = LOD.ColorVertexBuffer.VertexColor(i);
+				}
+
+
+				// Copy all the UVs that we have, and as many as we can fit.
+				{
+					const FVector2D& UV0 = (TexCoordsPerVertex > 0) ? LOD.VertexBuffer.GetVertexUV(i, 0) : FVector2D(1, 1);
+					NewVert.TexCoords[0] = UV0.X;
+					NewVert.TexCoords[1] = UV0.Y;
+
+					const FVector2D& UV1 = (TexCoordsPerVertex > 1) ? LOD.VertexBuffer.GetVertexUV(i, 1) : FVector2D(1, 1);
+					NewVert.TexCoords[2] = UV1.X;
+					NewVert.TexCoords[3] = UV1.Y;
+
+					const FVector2D& UV2 = (TexCoordsPerVertex > 2) ? LOD.VertexBuffer.GetVertexUV(i, 2) : FVector2D(1, 1);
+					NewVert.MaterialTexCoords[0] = UV2.X;
+					NewVert.MaterialTexCoords[1] = UV2.Y;
+
+					const FVector2D& UV3 = (TexCoordsPerVertex > 3) ? LOD.VertexBuffer.GetVertexUV(i, 3) : FVector2D(1, 1);
+					NewVert.ClipRect.TopLeft = UV3;
+
+					const FVector2D& UV4 = (TexCoordsPerVertex > 4) ? LOD.VertexBuffer.GetVertexUV(i, 4) : FVector2D(1, 1);
+					NewVert.ClipRect.ExtentX = UV4;
+
+					const FVector2D& UV5 = (TexCoordsPerVertex > 5) ? LOD.VertexBuffer.GetVertexUV(i, 5) : FVector2D(1, 1);
+					NewVert.ClipRect.ExtentY = UV5;
+				}
+			}
+		}
+	}
+}
\ No newline at end of file
diff --git a/Engine/Source/Runtime/UMG/Public/Slate/SMeshWidget.h b/Engine/Source/Runtime/UMG/Public/Slate/SMeshWidget.h
new file mode 100644
--- /dev/null
+++ b/Engine/Source/Runtime/UMG/Public/Slate/SMeshWidget.h
@@ -0,0 +1,50 @@
+// Copyright 1998-2015 Epic Games, Inc. All Rights Reserved.
+
+#pragma once
+
+class UStaticMesh;
+class FSlateInstanceBufferUpdate;
+
+/**
+ * A widget that draws vertexes provided by a 2.5D StaticMesh.
+ * The Mesh's material is used.
+ * Hardware instancing is supported.
+ */
+class UMG_API SMeshWidget : public SLeafWidget
+{
+
+public:
+	SLATE_BEGIN_ARGS(SMeshWidget)
+		: _StaticMeshSource(nullptr)
+	{}
+		/** The StaticMesh asset that should be drawn. */
+		SLATE_ARGUMENT(UStaticMesh*, StaticMeshSource)
+	SLATE_END_ARGS()
+
+	void Construct(const FArguments& Args);
+
+	/** Draw the InStaticMesh when this widget paints */
+	void SetMesh(const UStaticMesh& InStaticMesh);
+
+	/** Begin an update to the per instance buffer. Enables hardware instancing. */
+	TSharedPtr<FSlateInstanceBufferUpdate> BeginPerInstanceBufferUpdate(int32 InitialSize);
+
+protected:
+	// BEGIN SLeafWidget interface
+	virtual int32 OnPaint(const FPaintArgs& Args, const FGeometry& AllottedGeometry, const FSlateRect& MyClippingRect, FSlateWindowElementList& OutDrawElements, int32 LayerId, const FWidgetStyle& InWidgetStyle, bool bParentEnabled) const override;
+	virtual FVector2D ComputeDesiredSize(float) const override;
+	// END SLeafWidget interface
+
+private:
+	/** Populate OutSlateVerts and OutIndexes with data from this static mesh such that Slate can render it. */
+	static void StaticMeshToSlateRenderData(const UStaticMesh& DataSource, TArray<FSlateVertex>& OutSlateVerts, TArray<SlateIndex>& OutIndexes);
+
+	/** Holds a copy of the Static Mesh's data converted to a format that Slate understands. */
+	TArray<FSlateVertex> VertexData;
+	/** Connectivity data: Order in which the vertexes occur to make up a series of triangles. */
+	TArray<SlateIndex> IndexData;
+	/** Holds on to the material that is found on the StaticMesh. */
+	TSharedPtr<FSlateBrush> Brush;
+	/** Per instance data that can be passed to */
+	TSharedPtr<ISlateUpdatableInstanceBuffer> PerInstanceBuffer;
+};
\ No newline at end of file
