Slate -  Huge performance improve for invalidation panels on the render thread.  The layer to element map is now specialized to deal with our primiarly linearly ordered layers with short ranges in the symantic layering used by the invalidation panel to merge disperate element trees.  Test went from 5ms to 0.2ms.

diff --git a/Engine/Source/Runtime/SlateCore/Private/Rendering/DrawElements.cpp b/Engine/Source/Runtime/SlateCore/Private/Rendering/DrawElements.cpp
--- a/Engine/Source/Runtime/SlateCore/Private/Rendering/DrawElements.cpp
+++ b/Engine/Source/Runtime/SlateCore/Private/Rendering/DrawElements.cpp
@@ -75,7 +75,13 @@ void FSlateDrawElement::MakeBox(
 	ESlateDrawEffect::Type InDrawEffects, 
 	const FLinearColor& InTint )
 {
-	SCOPE_CYCLE_COUNTER( STAT_SlateDrawElementMakeTime )
+	SCOPE_CYCLE_COUNTER(STAT_SlateDrawElementMakeTime)
+
+	// Ignore invalid rendering handles.
+	if ( !InRenderingHandle.Data.IsValid() )
+	{
+		return;
+	}
 
 	FSlateShaderResourceProxy* RenderingProxy = InRenderingHandle.Data->Proxy;
 
@@ -317,11 +323,11 @@ void FSlateBatchData::AssignIndexArrayToBatch( FSlateElementBatch& Batch )
 
 }
 
-void FSlateBatchData::FillVertexAndIndexBuffer( uint8* VertexBuffer, uint8* IndexBuffer )
+void FSlateBatchData::FillVertexAndIndexBuffer(uint8* VertexBuffer, uint8* IndexBuffer)
 {
 	int32 IndexOffset = 0;
 	int32 VertexOffset = 0;
-	for( const FSlateRenderBatch& Batch : RenderBatches )
+	for ( const FSlateRenderBatch& Batch : RenderBatches )
 	{
 		// Ignore foreign batches that are inserted into our render set.
 		if ( RenderDataHandle != Batch.CachedRenderHandle )
@@ -329,34 +335,34 @@ void FSlateBatchData::FillVertexAndIndexBuffer( uint8* VertexBuffer, uint8* Inde
 			continue;
 		}
 
-		if( Batch.VertexArrayIndex != INDEX_NONE && Batch.IndexArrayIndex != INDEX_NONE )
+		if ( Batch.VertexArrayIndex != INDEX_NONE && Batch.IndexArrayIndex != INDEX_NONE )
 		{
 			TArray<FSlateVertex>& Vertices = BatchVertexArrays[Batch.VertexArrayIndex];
 			TArray<SlateIndex>& Indices = BatchIndexArrays[Batch.IndexArrayIndex];
 
-			if(Vertices.Num() && Indices.Num())
+			if ( Vertices.Num() && Indices.Num() )
 			{
 				uint32 RequiredVertexSize = Vertices.Num() * Vertices.GetTypeSize();
 				uint32 RequiredIndexSize = Indices.Num() * Indices.GetTypeSize();
 
-				FMemory::Memcpy(VertexBuffer+VertexOffset, Vertices.GetData(), RequiredVertexSize);
-				FMemory::Memcpy(IndexBuffer+IndexOffset, Indices.GetData(), RequiredIndexSize);
+				FMemory::Memcpy(VertexBuffer + VertexOffset, Vertices.GetData(), RequiredVertexSize);
+				FMemory::Memcpy(IndexBuffer + IndexOffset, Indices.GetData(), RequiredIndexSize);
 
 				VertexArrayFreeList.Add(Batch.VertexArrayIndex);
 				IndexArrayFreeList.Add(Batch.IndexArrayIndex);
 
-				IndexOffset += (Indices.Num()*sizeof(SlateIndex));
-				VertexOffset += (Vertices.Num()*sizeof(FSlateVertex));
+				IndexOffset += ( Indices.Num()*sizeof(SlateIndex) );
+				VertexOffset += ( Vertices.Num()*sizeof(FSlateVertex) );
 
 				Vertices.Reset();
 				Indices.Reset();
 
-				if (Vertices.GetSlack() > MAX_VERT_ARRAY_RECYCLE)
+				if ( Vertices.GetSlack() > MAX_VERT_ARRAY_RECYCLE )
 				{
 					Vertices.Empty();
 					Vertices.Reserve(MAX_VERT_ARRAY_RECYCLE);
 				}
-				if (Indices.GetSlack() > MAX_INDEX_ARRAY_RECYCLE)
+				if ( Indices.GetSlack() > MAX_INDEX_ARRAY_RECYCLE )
 				{
 					Indices.Empty();
 					Indices.Reserve(MAX_INDEX_ARRAY_RECYCLE);
@@ -368,13 +374,17 @@ void FSlateBatchData::FillVertexAndIndexBuffer( uint8* VertexBuffer, uint8* Inde
 
 void FSlateBatchData::CreateRenderBatches(FElementBatchMap& LayerToElementBatches)
 {
-	checkSlow( IsInRenderingThread() );
+	checkSlow(IsInRenderingThread());
 
 	uint32 VertexOffset = 0;
 	uint32 IndexOffset = 0;
 
+	FPlatformMisc::BeginNamedEvent(FColor::Magenta, "SlateRT::CreateRenderBatches");
+
 	Merge(LayerToElementBatches, VertexOffset, IndexOffset);
 
+	FPlatformMisc::EndNamedEvent();
+
 	// 
 	if ( RenderDataHandle.IsValid() )
 	{
@@ -393,103 +403,98 @@ void FSlateBatchData::AddRenderBatch(uint32 InLayer, const FSlateElementBatch& I
 
 void FSlateBatchData::Merge(FElementBatchMap& InLayerToElementBatches, uint32& VertexOffset, uint32& IndexOffset)
 {
-	InLayerToElementBatches.KeySort(TLess<uint32>());
+	InLayerToElementBatches.Sort();
 
 	const bool bExpandLayersAndCachedHandles = RenderDataHandle.IsValid() == false;
 
-	// For each element batch add its vertices and indices to the bulk lists.
-	for ( FElementBatchMap::TIterator It(InLayerToElementBatches); It; ++It )
+	InLayerToElementBatches.ForEachLayer([&] (uint32 Layer, FElementBatchArray& ElementBatches)
 	{
-		uint32 Layer = It.Key();
-		FElementBatchArray& ElementBatches = *It.Value();
-
-		if ( ElementBatches.Num() > 0 )
+		++NumLayers;
+		for ( FElementBatchArray::TIterator BatchIt(ElementBatches); BatchIt; ++BatchIt )
 		{
-			++NumLayers;
-			for ( FElementBatchArray::TIterator BatchIt(ElementBatches); BatchIt; ++BatchIt )
-			{
-				FSlateElementBatch& ElementBatch = *BatchIt;
+			FSlateElementBatch& ElementBatch = *BatchIt;
 
-				if ( ElementBatch.GetCustomDrawer().IsValid() )
-				{
-					AddRenderBatch(Layer, ElementBatch, 0, 0, 0, 0);
-				}
-				else
+			if ( ElementBatch.GetCustomDrawer().IsValid() )
+			{
+				AddRenderBatch(Layer, ElementBatch, 0, 0, 0, 0);
+			}
+			else
+			{
+				if ( bExpandLayersAndCachedHandles )
 				{
-					if ( bExpandLayersAndCachedHandles )
+					if ( FSlateRenderDataHandle* RenderHandle = ElementBatch.GetCachedRenderHandle().Get() )
 					{
-						if ( FSlateRenderDataHandle* RenderHandle = ElementBatch.GetCachedRenderHandle().Get() )
-						{
-							DynamicOffset += ElementBatch.GetCachedRenderDataOffset();
+						DynamicOffset += ElementBatch.GetCachedRenderDataOffset();
 
-							if ( TArray<FSlateRenderBatch>* ForeignBatches = RenderHandle->GetRenderBatches() )
+						if ( TArray<FSlateRenderBatch>* ForeignBatches = RenderHandle->GetRenderBatches() )
+						{
+							TArray<FSlateRenderBatch>& ForeignBatchesRef = *ForeignBatches;
+							for ( int32 i = 0; i < ForeignBatches->Num(); i++ )
 							{
-								TArray<FSlateRenderBatch>& ForeignBatchesRef = *ForeignBatches;
-								for ( int32 i = 0; i < ForeignBatches->Num(); i++ )
+								TSharedPtr<FSlateDrawLayerHandle, ESPMode::ThreadSafe> LayerHandle = ForeignBatchesRef[i].LayerHandle.Pin();
+								if ( LayerHandle.IsValid() )
 								{
-									TSharedPtr<FSlateDrawLayerHandle, ESPMode::ThreadSafe> LayerHandle = ForeignBatchesRef[i].LayerHandle.Pin();
-									if ( LayerHandle.IsValid() )
+									// If a record was added for a layer, but nothing was ever drawn for it, the batch map will be null.
+									if ( LayerHandle->BatchMap )
 									{
-										// If a record was added for a layer, but nothing was ever drawn for it, the batch map will be null.
-										if ( LayerHandle->BatchMap )
-										{
-											Merge(*LayerHandle->BatchMap, VertexOffset, IndexOffset);
-											LayerHandle->BatchMap = nullptr;
-										}
-									}
-									else
-									{
-										const int32 Index = RenderBatches.Add(ForeignBatchesRef[i]);
-										RenderBatches[Index].DynamicOffset = DynamicOffset;
+										Merge(*LayerHandle->BatchMap, VertexOffset, IndexOffset);
+										LayerHandle->BatchMap = nullptr;
 									}
 								}
+								else
+								{
+									const int32 Index = RenderBatches.Add(ForeignBatchesRef[i]);
+									RenderBatches[Index].DynamicOffset = DynamicOffset;
+								}
 							}
+						}
 
-							DynamicOffset -= ElementBatch.GetCachedRenderDataOffset();
+						DynamicOffset -= ElementBatch.GetCachedRenderDataOffset();
 
-							continue;
-						}
+						continue;
 					}
-					else
+				}
+				else
+				{
+					// Insert if we're not expanding
+					if ( FSlateDrawLayerHandle* LayerHandle = ElementBatch.GetLayerHandle().Get() )
 					{
-						// Insert if we're not expanding
-						if ( FSlateDrawLayerHandle* LayerHandle = ElementBatch.GetLayerHandle().Get() )
-						{
-							AddRenderBatch(Layer, ElementBatch, 0, 0, 0, 0);
-							continue;
-						}
+						AddRenderBatch(Layer, ElementBatch, 0, 0, 0, 0);
+						continue;
 					}
+				}
 					
-					if ( ElementBatch.VertexArrayIndex != INDEX_NONE && ElementBatch.IndexArrayIndex != INDEX_NONE )
-					{
-						TArray<FSlateVertex>& BatchVertices = GetBatchVertexList(ElementBatch);
-						TArray<SlateIndex>& BatchIndices = GetBatchIndexList(ElementBatch);
+				if ( ElementBatch.VertexArrayIndex != INDEX_NONE && ElementBatch.IndexArrayIndex != INDEX_NONE )
+				{
+					TArray<FSlateVertex>& BatchVertices = GetBatchVertexList(ElementBatch);
+					TArray<SlateIndex>& BatchIndices = GetBatchIndexList(ElementBatch);
 
-						// We should have at least some vertices and indices in the batch or none at all
-						check(BatchVertices.Num() > 0 && BatchIndices.Num() > 0 || BatchVertices.Num() == 0 && BatchIndices.Num() == 0);
+					// We should have at least some vertices and indices in the batch or none at all
+					check(BatchVertices.Num() > 0 && BatchIndices.Num() > 0 || BatchVertices.Num() == 0 && BatchIndices.Num() == 0);
 
-						if ( BatchVertices.Num() > 0 && BatchIndices.Num() > 0 )
-						{
-							const int32 NumVertices = BatchVertices.Num();
-							const int32 NumIndices = BatchIndices.Num();
+					if ( BatchVertices.Num() > 0 && BatchIndices.Num() > 0 )
+					{
+						const int32 NumVertices = BatchVertices.Num();
+						const int32 NumIndices = BatchIndices.Num();
 
-							AddRenderBatch(Layer, ElementBatch, NumVertices, NumIndices, VertexOffset, IndexOffset);
+						AddRenderBatch(Layer, ElementBatch, NumVertices, NumIndices, VertexOffset, IndexOffset);
 
-							VertexOffset += BatchVertices.Num();
-							IndexOffset += BatchIndices.Num();
-						}
-						else
-						{
-							VertexArrayFreeList.Add(ElementBatch.VertexArrayIndex);
-							IndexArrayFreeList.Add(ElementBatch.IndexArrayIndex);
-						}
+						VertexOffset += BatchVertices.Num();
+						IndexOffset += BatchIndices.Num();
+					}
+					else
+					{
+						VertexArrayFreeList.Add(ElementBatch.VertexArrayIndex);
+						IndexArrayFreeList.Add(ElementBatch.IndexArrayIndex);
 					}
 				}
 			}
 		}
 
 		ElementBatches.Reset();
-	}
+	});
+
+	InLayerToElementBatches.Reset();
 }
 
 
@@ -545,6 +550,8 @@ FSlateWindowElementList::FVolatilePaint::FVolatilePaint(const TSharedRef<const S
 
 int32 FSlateWindowElementList::FVolatilePaint::ExecutePaint(FSlateWindowElementList& OutDrawElements) const
 {
+	static const FName InvalidationPanelName(TEXT("SInvalidationPanel"));
+
 	TSharedPtr<const SWidget> WidgetToPaint = WidgetToPaintPtr.Pin();
 	if ( WidgetToPaint.IsValid() )
 	{
@@ -554,7 +561,10 @@ int32 FSlateWindowElementList::FVolatilePaint::ExecutePaint(FSlateWindowElementL
 		// the STextBlock.  This may be all kinds of terrible an idea to do during paint.
 		SWidget* MutableWidget = const_cast<SWidget*>( WidgetToPaint.Get() );
 
-		MutableWidget->SlatePrepass(AllottedGeometry.Scale);
+		if ( MutableWidget->GetType() != InvalidationPanelName )
+		{
+			MutableWidget->SlatePrepass(AllottedGeometry.Scale);
+		}
 
 		const int32 NewLayer = WidgetToPaint->Paint(Args, AllottedGeometry, MyClippingRect, OutDrawElements, LayerId, WidgetStyle, bParentEnabled);
 		
diff --git a/Engine/Source/Runtime/SlateCore/Public/Rendering/DrawElements.h b/Engine/Source/Runtime/SlateCore/Public/Rendering/DrawElements.h
--- a/Engine/Source/Runtime/SlateCore/Public/Rendering/DrawElements.h
+++ b/Engine/Source/Runtime/SlateCore/Public/Rendering/DrawElements.h
@@ -727,7 +727,108 @@ public:
 
 typedef TArray<FSlateElementBatch, TInlineAllocator<1>> FElementBatchArray;
 
-typedef TMap<uint32,TUniqueObj<FElementBatchArray>> FElementBatchMap ;
+class FElementBatchMap
+{
+public:
+	FElementBatchMap()
+	{
+		Reset();
+	}
+
+	FORCEINLINE_DEBUGGABLE int32 Num() const { return Layers.Num() + OverflowLayers.Num(); }
+	
+	FORCEINLINE_DEBUGGABLE TUniqueObj<FElementBatchArray>* Find(uint32 Layer)
+	{
+		if ( Layer < (uint32)Layers.Num() )
+		{
+			if ( ActiveLayers[Layer] )
+			{
+				return &Layers[Layer];
+			}
+
+			return nullptr;
+		}
+		else
+		{
+			return OverflowLayers.Find(Layer);
+		}
+	}
+
+	FORCEINLINE_DEBUGGABLE TUniqueObj<FElementBatchArray>& Add(uint32 Layer)
+	{
+		if ( Layer < (uint32)Layers.Num() )
+		{
+			MinLayer = FMath::Min(Layer, MinLayer);
+			MaxLayer = FMath::Max(Layer, MaxLayer);
+			ActiveLayers[Layer] = true;
+			return Layers[Layer];
+		}
+		else
+		{
+			return OverflowLayers.Add(Layer);
+		}
+	}
+
+	FORCEINLINE_DEBUGGABLE void Sort()
+	{
+		OverflowLayers.KeySort(TLess<uint32>());
+	}
+
+	template <typename TFunc>
+	FORCEINLINE_DEBUGGABLE void ForEachLayer(const TFunc& Process)
+	{
+		if ( MinLayer < (uint32)Layers.Num() )
+		{
+			for ( TBitArray<>::FConstIterator BitIt(ActiveLayers, MinLayer); BitIt; ++BitIt )
+			{
+				if ( BitIt.GetValue() == 0 )
+				{
+					continue;
+				}
+
+				const int32 BitIndex = BitIt.GetIndex();
+				FElementBatchArray& ElementBatches = *Layers[BitIndex];
+
+				if ( ElementBatches.Num() > 0 )
+				{
+					Process(BitIndex, ElementBatches);
+				}
+
+				if ( ((uint32)BitIndex) >= MaxLayer )
+				{
+					break;
+				}
+			}
+		}
+
+		for ( TMap<uint32, TUniqueObj<FElementBatchArray>>::TIterator It(OverflowLayers); It; ++It )
+		{
+			uint32 Layer = It.Key();
+			FElementBatchArray& ElementBatches = *It.Value();
+
+			if ( ElementBatches.Num() > 0 )
+			{
+				Process(Layer, ElementBatches);
+			}
+		}
+	}
+
+	FORCEINLINE_DEBUGGABLE void Reset()
+	{
+		MinLayer = UINT_MAX;
+		MaxLayer = 0;
+		ActiveLayers.Init(false, Layers.Num());
+		OverflowLayers.Reset();
+	}
+
+private:
+	TBitArray<> ActiveLayers;
+	TStaticArray<TUniqueObj<FElementBatchArray>, 256> Layers;
+	TMap<uint32, TUniqueObj<FElementBatchArray>> OverflowLayers;
+	uint32 MinLayer;
+	uint32 MaxLayer;
+};
+
 
 class FSlateBatchData
 {
