Slate - Inlining several small functions in Slate to save a few hundred microseconds.

diff --git a/Engine/Source/Runtime/SlateCore/Private/Input/HittestGrid.cpp b/Engine/Source/Runtime/SlateCore/Private/Input/HittestGrid.cpp
--- a/Engine/Source/Runtime/SlateCore/Private/Input/HittestGrid.cpp
+++ b/Engine/Source/Runtime/SlateCore/Private/Input/HittestGrid.cpp
@@ -45,19 +45,16 @@ FVector2D ClosestPointOnSlateRotatedRect(const FVector2D &Point, const FSlateRot
 	return RetPoint;
 }
 
-
-float DistanceSqToSlateRotatedRect(const FVector2D &Point, const FSlateRotatedRect& RotatedRect)
+FORCEINLINE float DistanceSqToSlateRotatedRect(const FVector2D &Point, const FSlateRotatedRect& RotatedRect)
 {
 	return FVector2D::DistSquared(ClosestPointOnSlateRotatedRect(Point, RotatedRect), Point);
 }
 
-
-bool IsOverlappingSlateRotatedRect(const FVector2D& Point, const float Radius, const FSlateRotatedRect& RotatedRect)
+FORCEINLINE bool IsOverlappingSlateRotatedRect(const FVector2D& Point, const float Radius, const FSlateRotatedRect& RotatedRect)
 {
 	return DistanceSqToSlateRotatedRect( Point, RotatedRect ) <= (Radius * Radius);
 }
 
-
 bool ContainsInteractableWidget(const TArray<FWidgetAndPointer>& PathToTest)
 {
 	for (int32 i = PathToTest.Num() - 1; i >= 0; --i)
@@ -71,8 +68,6 @@ bool ContainsInteractableWidget(const TArray<FWidgetAndPointer>& PathToTest)
 	return false;
 }
 
-
-
 //
 // FHittestGrid
 //
@@ -540,8 +535,6 @@ TSharedPtr<SWidget> FHittestGrid::FindNextFocusableWidget(const FArrangedWidget&
 	return Widget;
 }
 
-
-
 FIntPoint FHittestGrid::GetCellCoordinate(FVector2D Position)
 {
 	return FIntPoint(
@@ -549,18 +542,6 @@ FIntPoint FHittestGrid::GetCellCoordinate(FVector2D Position)
 		FMath::Min(FMath::Max(FMath::FloorToInt(Position.Y / CellSize.Y), 0), NumCells.Y - 1));
 }
 
-FHittestGrid::FCell& FHittestGrid::CellAt( const int32 X, const int32 Y )
-{
-	check( (Y*NumCells.X + X) < Cells.Num() );
-	return Cells[ Y*NumCells.X + X ];
-}
-
-const FHittestGrid::FCell& FHittestGrid::CellAt( const int32 X, const int32 Y ) const
-{
-	check( (Y*NumCells.X + X) < Cells.Num() );
-	return Cells[ Y*NumCells.X + X ];
-}
-
 bool FHittestGrid::IsValidCellCoord(const FIntPoint& CellCoord) const
 {
 	return IsValidCellCoord(CellCoord.X, CellCoord.Y);
diff --git a/Engine/Source/Runtime/SlateCore/Private/Layout/ArrangedChildren.cpp b/Engine/Source/Runtime/SlateCore/Private/Layout/ArrangedChildren.cpp
--- a/Engine/Source/Runtime/SlateCore/Private/Layout/ArrangedChildren.cpp
+++ b/Engine/Source/Runtime/SlateCore/Private/Layout/ArrangedChildren.cpp
@@ -6,16 +6,7 @@
 /* FArrangedChildren interface
  *****************************************************************************/
 
-void FArrangedChildren::AddWidget( const FArrangedWidget& InWidgetGeometry )
+void FArrangedChildren::AddWidget(const FArrangedWidget& InWidgetGeometry)
 {
 	AddWidget(InWidgetGeometry.Widget->GetVisibility(), InWidgetGeometry);
 }
-
-
-void FArrangedChildren::AddWidget( EVisibility VisibilityOverride, const FArrangedWidget& InWidgetGeometry )
-{
-	if (this->Accepts(VisibilityOverride))
-	{
-		Array.Add(InWidgetGeometry);
-	}
-}
diff --git a/Engine/Source/Runtime/SlateCore/Private/Layout/Geometry.cpp b/Engine/Source/Runtime/SlateCore/Private/Layout/Geometry.cpp
--- a/Engine/Source/Runtime/SlateCore/Private/Layout/Geometry.cpp
+++ b/Engine/Source/Runtime/SlateCore/Private/Layout/Geometry.cpp
@@ -20,193 +20,7 @@ FGeometry& FGeometry::operator=(const FGeometry& RHS)
 	return *this;
 }
 
-FGeometry::FGeometry(const FVector2D& OffsetFromParent, const FVector2D& ParentAbsolutePosition, const FVector2D& InLocalSize, float InScale) 
-	: Size(InLocalSize)
-	, Scale(1.0f)
-	, AbsolutePosition(0.0f, 0.0f)
-{
-	// Since OffsetFromParent is given as a LocalSpaceOffset, we MUST convert this offset into the space of the parent to construct a valid layout transform.
-	// The extra TransformPoint below does this by converting the local offset to an offset in parent space.
-	FVector2D LayoutOffset = TransformPoint(InScale, OffsetFromParent);
-
-	FSlateLayoutTransform ParentAccumulatedLayoutTransform(InScale, ParentAbsolutePosition);
-	FSlateLayoutTransform LocalLayoutTransform(LayoutOffset);
-	FSlateLayoutTransform AccumulatedLayoutTransform = Concatenate(LocalLayoutTransform, ParentAccumulatedLayoutTransform);
-	AccumulatedRenderTransform = TransformCast<FSlateRenderTransform>(AccumulatedLayoutTransform);
-	// HACK to allow us to make FGeometry public members immutable to catch misuse.
-	const_cast<FVector2D&>(AbsolutePosition) = AccumulatedLayoutTransform.GetTranslation();
-	const_cast<float&>(Scale) = AccumulatedLayoutTransform.GetScale();
-	const_cast<FVector2D&>(Position) = LocalLayoutTransform.GetTranslation();
-}
-
-FGeometry::FGeometry(const FVector2D& InLocalSize, const FSlateLayoutTransform& InLocalLayoutTransform, const FSlateRenderTransform& InLocalRenderTransform, const FVector2D& InLocalRenderTransformPivot, const FSlateLayoutTransform& ParentAccumulatedLayoutTransform, const FSlateRenderTransform& ParentAccumulatedRenderTransform) 
-	: Size(InLocalSize)
-	, Scale(1.0f)
-	, AbsolutePosition(0.0f, 0.0f)
-	, AccumulatedRenderTransform(
-		Concatenate(
-			// convert the pivot to local space and make it the origin
-			Inverse(TransformPoint(FScale2D(InLocalSize), InLocalRenderTransformPivot)), 
-			// apply the render transform in local space centered around the pivot
-			InLocalRenderTransform, 
-			// translate the pivot point back.
-			TransformPoint(FScale2D(InLocalSize), InLocalRenderTransformPivot), 
-			// apply the layout transform next.
-			InLocalLayoutTransform, 
-			// finally apply the parent accumulated transform, which takes us to the root.
-			ParentAccumulatedRenderTransform))
-{
-	FSlateLayoutTransform AccumulatedLayoutTransform = Concatenate(InLocalLayoutTransform, ParentAccumulatedLayoutTransform);
-	// HACK to allow us to make FGeometry public members immutable to catch misuse.
-	const_cast<FVector2D&>(AbsolutePosition) = AccumulatedLayoutTransform.GetTranslation();
-	const_cast<float&>(Scale) = AccumulatedLayoutTransform.GetScale();
-	const_cast<FVector2D&>(Position) = InLocalLayoutTransform.GetTranslation();
-}
-
-FGeometry::FGeometry(const FVector2D& InLocalSize, const FSlateLayoutTransform& InLocalLayoutTransform, const FSlateLayoutTransform& ParentAccumulatedLayoutTransform, const FSlateRenderTransform& ParentAccumulatedRenderTransform) 
-	: Size(InLocalSize)
-	, Scale(1.0f)
-	, AbsolutePosition(0.0f, 0.0f)
-	, AccumulatedRenderTransform(Concatenate(InLocalLayoutTransform, ParentAccumulatedRenderTransform))
-{
-	FSlateLayoutTransform AccumulatedLayoutTransform = Concatenate(InLocalLayoutTransform, ParentAccumulatedLayoutTransform);
-	// HACK to allow us to make FGeometry public members immutable to catch misuse.
-	const_cast<FVector2D&>(AbsolutePosition) = AccumulatedLayoutTransform.GetTranslation();
-	const_cast<float&>(Scale) = AccumulatedLayoutTransform.GetScale();
-	const_cast<FVector2D&>(Position) = InLocalLayoutTransform.GetTranslation();
-}
-
-void FGeometry::AppendTransform(const FSlateLayoutTransform& LayoutTransform)
-{
-	FSlateLayoutTransform AccumulatedLayoutTransform = ::Concatenate(GetAccumulatedLayoutTransform(), LayoutTransform);
-	AccumulatedRenderTransform = ::Concatenate(AccumulatedRenderTransform, LayoutTransform);
-	const_cast<FVector2D&>(AbsolutePosition) = AccumulatedLayoutTransform.GetTranslation();
-	const_cast<float&>(Scale) = AccumulatedLayoutTransform.GetScale();
-}
-
-FGeometry FGeometry::MakeRoot(const FVector2D& LocalSize, const FSlateLayoutTransform& LayoutTransform)
-{
-	return FGeometry(LocalSize, LayoutTransform, FSlateLayoutTransform(), FSlateRenderTransform());
-}
-
-FGeometry FGeometry::MakeChild(const FVector2D& LocalSize, const FSlateLayoutTransform& LayoutTransform, const FSlateRenderTransform& RenderTransform, const FVector2D& RenderTransformPivot) const
-{
-	return FGeometry(LocalSize, LayoutTransform, RenderTransform, RenderTransformPivot, GetAccumulatedLayoutTransform(), GetAccumulatedRenderTransform());
-}
-
-FGeometry FGeometry::MakeChild( const FVector2D& LocalSize, const FSlateLayoutTransform& LayoutTransform ) const
-{
-	return FGeometry(LocalSize, LayoutTransform, GetAccumulatedLayoutTransform(), GetAccumulatedRenderTransform());
-}
-
-FArrangedWidget FGeometry::MakeChild(const TSharedRef<SWidget>& ChildWidget, const FLayoutGeometry& LayoutGeometry) const
-{
-	return MakeChild(ChildWidget, LayoutGeometry.GetSizeInLocalSpace(), LayoutGeometry.GetLocalToParentTransform());
-}
-
-FArrangedWidget FGeometry::MakeChild(const TSharedRef<SWidget>& ChildWidget, const FVector2D& LocalSize, const FSlateLayoutTransform& LayoutTransform) const
-{
-	// If there is no render transform set, use the simpler MakeChild call that doesn't bother concatenating the render transforms.
-	// This saves a significant amount of overhead since every widget does this, and most children don't have a render transform.
-	if (ChildWidget->GetRenderTransform().IsSet())
-	{
-		return FArrangedWidget(ChildWidget, MakeChild(LocalSize, LayoutTransform, ChildWidget->GetRenderTransform().GetValue(), ChildWidget->GetRenderTransformPivot()));
-	}
-	else
-	{
-		return FArrangedWidget(ChildWidget, MakeChild(LocalSize, LayoutTransform));
-	}
-}
-
-FGeometry FGeometry::MakeChild( const FVector2D& ChildOffset, const FVector2D& LocalSize, float LocalScale ) const
-{
-	// Since ChildOffset is given as a LocalSpaceOffset, we MUST convert this offset into the space of the parent to construct a valid layout transform.
-	// The extra TransformPoint below does this by converting the local offset to an offset in parent space.
-	return FGeometry(LocalSize, FSlateLayoutTransform(LocalScale, TransformPoint(LocalScale, ChildOffset)), GetAccumulatedLayoutTransform(), GetAccumulatedRenderTransform());
-}
-
-FArrangedWidget FGeometry::MakeChild( const TSharedRef<SWidget>& ChildWidget, const FVector2D& ChildOffset, const FVector2D& LocalSize, float ChildScale) const
-{
-	// Since ChildOffset is given as a LocalSpaceOffset, we MUST convert this offset into the space of the parent to construct a valid layout transform.
-	// The extra TransformPoint below does this by converting the local offset to an offset in parent space.
-	return MakeChild(ChildWidget, LocalSize, FSlateLayoutTransform(ChildScale, TransformPoint(ChildScale, ChildOffset)));
-}
-
-FPaintGeometry FGeometry::ToPaintGeometry() const
-{
-	return FPaintGeometry(GetAccumulatedLayoutTransform(), GetAccumulatedRenderTransform(), Size);
-}
-
-FPaintGeometry FGeometry::ToPaintGeometry(const FSlateLayoutTransform& LayoutTransform) const
-{
-	return ToPaintGeometry(Size, LayoutTransform);
-}
-
-
-FPaintGeometry FGeometry::ToPaintGeometry(const FVector2D& LocalSize, const FSlateLayoutTransform& LayoutTransform) const
-{
-	FSlateLayoutTransform NewAccumulatedLayoutTransform = Concatenate(LayoutTransform, GetAccumulatedLayoutTransform());
-	return FPaintGeometry(NewAccumulatedLayoutTransform, Concatenate(LayoutTransform, GetAccumulatedRenderTransform()), LocalSize);
-}
-
-FPaintGeometry FGeometry::ToPaintGeometry(const FVector2D& LocalOffset, const FVector2D& LocalSize, float LocalScale) const
-{
-	// Since ChildOffset is given as a LocalSpaceOffset, we MUST convert this offset into the space of the parent to construct a valid layout transform.
-	// The extra TransformPoint below does this by converting the local offset to an offset in parent space.
-	return ToPaintGeometry(LocalSize, FSlateLayoutTransform(LocalScale, TransformPoint(LocalScale, LocalOffset)));
-}
-
-FPaintGeometry FGeometry::ToOffsetPaintGeometry(const FVector2D& LocalOffset) const
-{
-	return ToPaintGeometry(FSlateLayoutTransform(LocalOffset));
-}
-
-FPaintGeometry FGeometry::ToInflatedPaintGeometry( const FVector2D& InflateAmount ) const
-{
-	// This essentially adds (or subtracts) a border around the widget. We scale the size then offset by the border amount.
-	// Note this is not scaling child widgets, so the scale is not changing.
-	FVector2D NewSize = Size + InflateAmount * 2;
-	return ToPaintGeometry(NewSize, FSlateLayoutTransform(-InflateAmount));
-}
-
-FSlateRect FGeometry::GetClippingRect() const
-{
-	return TransformRect(GetAccumulatedLayoutTransform(), FSlateRect(FVector2D(0.0f,0.0f), Size));
-}
-
-FVector2D FGeometry::AbsoluteToLocal(FVector2D AbsoluteCoordinate) const
-{
-	// this render transform invert is a little expensive. We might consider caching it.
-	return TransformPoint(Inverse(GetAccumulatedRenderTransform()), AbsoluteCoordinate);
-}
-
-
-FVector2D FGeometry::LocalToAbsolute(FVector2D LocalCoordinate) const
-{
-	return TransformPoint(GetAccumulatedRenderTransform(), LocalCoordinate);
-}
-
-FVector2D FGeometry::LocalToRoundedLocal(FVector2D LocalCoordinate) const
-{
-	const FVector2D AbsoluteCoordinate = LocalToAbsolute(LocalCoordinate);
-	const FVector2D AbsoluteCoordinateRounded = FVector2D(FMath::RoundToInt(AbsoluteCoordinate.X), FMath::RoundToInt(AbsoluteCoordinate.Y));
-
-	return AbsoluteToLocal(AbsoluteCoordinateRounded);
-}
-
-bool FGeometry::IsUnderLocation(const FVector2D& AbsoluteCoordinate) const
-{
-	// this render transform invert is a little expensive. We might consider caching it.
-	return FSlateRect(FVector2D(0.0f, 0.0f), Size).ContainsPoint(TransformPoint(Inverse(GetAccumulatedRenderTransform()), AbsoluteCoordinate));
-}
-
-
 FString FGeometry::ToString() const
 {
 	return FString::Printf(TEXT("[Abs=%s, Scale=%.2f, Size=%s]"), *AbsolutePosition.ToString(), Scale, *Size.ToString());
 }
-
-FVector2D FGeometry::GetDrawSize() const
-{
-	return TransformVector(GetAccumulatedLayoutTransform(), Size);
-}
diff --git a/Engine/Source/Runtime/SlateCore/Private/Rendering/SlateRenderer.cpp b/Engine/Source/Runtime/SlateCore/Private/Rendering/SlateRenderer.cpp
--- a/Engine/Source/Runtime/SlateCore/Private/Rendering/SlateRenderer.cpp
+++ b/Engine/Source/Runtime/SlateCore/Private/Rendering/SlateRenderer.cpp
@@ -15,7 +15,7 @@ void FSlateRenderer::FlushFontCache( )
 
 bool FSlateRenderer::IsViewportFullscreen( const SWindow& Window ) const
 {
-	check( IsThreadSafeForSlateRendering() );
+	checkSlow( IsThreadSafeForSlateRendering() );
 
 	bool bFullscreen = false;
 
diff --git a/Engine/Source/Runtime/SlateCore/Private/SlotBase.cpp b/Engine/Source/Runtime/SlateCore/Private/SlotBase.cpp
--- a/Engine/Source/Runtime/SlateCore/Private/SlotBase.cpp
+++ b/Engine/Source/Runtime/SlateCore/Private/SlotBase.cpp
@@ -14,17 +14,6 @@ FSlotBase::FSlotBase( const TSharedRef<SWidget>& InWidget )
 	
 }
 
-void FSlotBase::AttachWidget( const TSharedRef<SWidget>& InWidget )
-{
-	Widget = InWidget;
-}
-
-
-const TSharedRef<SWidget>& FSlotBase::GetWidget() const
-{
-	return Widget;
-}
-
 const TSharedPtr<SWidget> FSlotBase::DetachWidget()
 {
 	if (Widget != SNullWidget::NullWidget)
diff --git a/Engine/Source/Runtime/SlateCore/Private/Types/PaintArgs.cpp b/Engine/Source/Runtime/SlateCore/Private/Types/PaintArgs.cpp
--- a/Engine/Source/Runtime/SlateCore/Private/Types/PaintArgs.cpp
+++ b/Engine/Source/Runtime/SlateCore/Private/Types/PaintArgs.cpp
@@ -19,19 +19,6 @@ FPaintArgs::FPaintArgs(const SWidget& Parent, FHittestGrid& InHittestGrid, FVect
 {
 }
 
-FPaintArgs FPaintArgs::WithNewParent( const SWidget& Parent ) const
-{
-	FPaintArgs Args = FPaintArgs(Parent, this->Grid, this->WindowOffset, this->CurrentTime, this->DeltaTime);
-	Args.LastHittestIndex = this->LastHittestIndex;
-	Args.LastRecordedVisibility = this->LastRecordedVisibility;
-	Args.LayoutCache = this->LayoutCache;
-	Args.ParentCacheNode = this->ParentCacheNode;
-	Args.bIsCaching = this->bIsCaching;
-	Args.bIsVolatilityPass = this->bIsVolatilityPass;
-
-	return Args;
-}
-
 FPaintArgs FPaintArgs::EnableCaching(ILayoutCache* InLayoutCache, FCachedWidgetNode* InParentCacheNode, bool bEnableCaching, bool bEnableVolatile) const
 {
 	FPaintArgs UpdatedArgs(*this);
diff --git a/Engine/Source/Runtime/SlateCore/Private/Widgets/SWidget.cpp b/Engine/Source/Runtime/SlateCore/Private/Widgets/SWidget.cpp
--- a/Engine/Source/Runtime/SlateCore/Private/Widgets/SWidget.cpp
+++ b/Engine/Source/Runtime/SlateCore/Private/Widgets/SWidget.cpp
@@ -409,11 +409,6 @@ void SWidget::CacheDesiredSize(float LayoutScaleMultiplier)
 	this->Advanced_SetDesiredSize(this->ComputeDesiredSize(LayoutScaleMultiplier));
 }
 
-const FVector2D& SWidget::GetDesiredSize() const
-{
-	return DesiredSize;
-}
-
 void SWidget::CachePrepass(ILayoutCache* InLayoutCache)
 {
 	if ( bCanHaveChildren )
diff --git a/Engine/Source/Runtime/SlateCore/Public/Application/SlateApplicationBase.h b/Engine/Source/Runtime/SlateCore/Public/Application/SlateApplicationBase.h
--- a/Engine/Source/Runtime/SlateCore/Public/Application/SlateApplicationBase.h
+++ b/Engine/Source/Runtime/SlateCore/Public/Application/SlateApplicationBase.h
@@ -389,7 +389,7 @@ public:
 	 */
 	static FSlateApplicationBase& Get( )
 	{
-		check(IsThreadSafeForSlateRendering());
+		checkSlow(IsThreadSafeForSlateRendering());
 		return *CurrentBaseApplication;
 	}
 
diff --git a/Engine/Source/Runtime/SlateCore/Public/Input/HittestGrid.h b/Engine/Source/Runtime/SlateCore/Public/Input/HittestGrid.h
--- a/Engine/Source/Runtime/SlateCore/Public/Input/HittestGrid.h
+++ b/Engine/Source/Runtime/SlateCore/Public/Input/HittestGrid.h
@@ -137,8 +137,17 @@ private:
 	FIntPoint GetCellCoordinate(FVector2D Position);
 
 	/** Access a cell at coordinates X, Y. Coordinates are row and column indexes. */
-	FCell& CellAt( const int32 X, const int32 Y );
-	const FCell& CellAt( const int32 X, const int32 Y ) const;
+	FORCEINLINE_DEBUGGABLE FCell& CellAt(const int32 X, const int32 Y)
+	{
+		check(( Y*NumCells.X + X ) < Cells.Num());
+		return Cells[Y*NumCells.X + X];
+	}
+
+	FORCEINLINE_DEBUGGABLE const FCell& CellAt( const int32 X, const int32 Y ) const
+	{
+		check(( Y*NumCells.X + X ) < Cells.Num());
+		return Cells[Y*NumCells.X + X];
+	}
 
 	/** All the widgets and their arranged geometries encountered this frame. */
 	TSharedRef< TArray<FCachedWidget> > WidgetsCachedThisFrame;
diff --git a/Engine/Source/Runtime/SlateCore/Public/Layout/ArrangedChildren.h b/Engine/Source/Runtime/SlateCore/Public/Layout/ArrangedChildren.h
--- a/Engine/Source/Runtime/SlateCore/Public/Layout/ArrangedChildren.h
+++ b/Engine/Source/Runtime/SlateCore/Public/Layout/ArrangedChildren.h
@@ -60,25 +60,34 @@ class SLATECORE_API FArrangedChildren
 	 *                             of layout or performance (i.e. prevent redundant call to Widget->GetVisibility())
 	 * @param InWidgetGeometry     The arranged widget (i.e. widget and its geometry)
 	 */
-	void AddWidget( EVisibility VisibilityOverride, const FArrangedWidget& InWidgetGeometry );
+	FORCEINLINE_DEBUGGABLE void AddWidget(EVisibility VisibilityOverride, const FArrangedWidget& InWidgetGeometry)
+	{
+		if ( Accepts(VisibilityOverride) )
+		{
+			Array.Add(InWidgetGeometry);
+		}
+	}
 	
 	/**
 	 * Add an arranged widget (i.e. widget and its resulting geometry) to the list of Arranged children
 	 * based on the the visibility filter and the arranged widget's visibility
 	 */
-	void AddWidget( const FArrangedWidget& InWidgetGeometry );
+	void AddWidget(const FArrangedWidget& InWidgetGeometry);
 
-	EVisibility GetFilter() const
+	FORCEINLINE_DEBUGGABLE EVisibility GetFilter() const
 	{
 		return VisibilityFilter;
 	}
 
-	bool Accepts( EVisibility InVisibility ) const
+	FORCEINLINE_DEBUGGABLE bool Accepts(EVisibility InVisibility) const
 	{
 		return 0 != (InVisibility.Value & VisibilityFilter.Value);
 	}
 
-	bool Allows3DWidgets() const { return bAllow3DWidgets; }
+	FORCEINLINE_DEBUGGABLE bool Allows3DWidgets() const
+	{
+		return bAllow3DWidgets;
+	}
 
 	// We duplicate parts of TArray's interface here!
 	// Inheriting from "public TArray<FArrangedWidget>"
@@ -92,75 +101,63 @@ private:
 	bool bAllow3DWidgets;
 public:
 
-	FArrangedWidgetArray& GetInternalArray()
+	FORCEINLINE FArrangedWidgetArray& GetInternalArray()
 	{
 		return Array;
 	}
 
-	const FArrangedWidgetArray& GetInternalArray() const
+	FORCEINLINE const FArrangedWidgetArray& GetInternalArray() const
 	{
 		return Array;
 	}
 
-	int32 Num() const
+	FORCEINLINE int32 Num() const
 	{
 		return Array.Num();
 	}
 
-	const FArrangedWidget& operator[]( int32 Index ) const
+	FORCEINLINE const FArrangedWidget& operator[]( int32 Index ) const
 	{
 		return Array[Index];
 	}
 
-	FArrangedWidget& operator[]( int32 Index )
+	FORCEINLINE FArrangedWidget& operator[](int32 Index)
 	{
 		return Array[Index];
 	}
 
-	//// @todo umg na deprecate this
-	//const FArrangedWidget& operator()(int32 Index) const
-	//{
-	//	return Array[Index];
-	//}
-
-	//// @todo umg na deprecate this
-	//FArrangedWidget& operator()(int32 Index)
-	//{
-	//	return Array[Index];
-	//}
-
-	const FArrangedWidget& Last() const
+	FORCEINLINE const FArrangedWidget& Last() const
 	{
 		return Array.Last();
 	}
 
-	FArrangedWidget& Last()
+	FORCEINLINE FArrangedWidget& Last()
 	{
 		return Array.Last();
 	}
 
-	int32 FindItemIndex(const FArrangedWidget& ItemToFind ) const
+	FORCEINLINE int32 FindItemIndex(const FArrangedWidget& ItemToFind) const
 	{
 		return Array.Find(ItemToFind);
 	}
 
 	template<typename PredicateType>
-	int32 IndexOfByPredicate( const PredicateType& Pred ) const
+	FORCEINLINE int32 IndexOfByPredicate(const PredicateType& Pred) const
 	{
 		return Array.IndexOfByPredicate( Pred );
 	}
 
-	void Remove( int32 Index, int32 Count=1 )
+	FORCEINLINE void Remove(int32 Index, int32 Count=1)
 	{
 		Array.RemoveAt(Index, Count);
 	}
 
-	void Append( const FArrangedChildren& Source )
+	FORCEINLINE void Append(const FArrangedChildren& Source)
 	{
 		Array.Append( Source.Array );
 	}
 
-	void Empty()
+	FORCEINLINE void Empty()
 	{
 		Array.Empty();
 	}
diff --git a/Engine/Source/Runtime/SlateCore/Public/Layout/Geometry.h b/Engine/Source/Runtime/SlateCore/Public/Layout/Geometry.h
--- a/Engine/Source/Runtime/SlateCore/Public/Layout/Geometry.h
+++ b/Engine/Source/Runtime/SlateCore/Public/Layout/Geometry.h
@@ -47,7 +47,24 @@ public:
 	 * @param InSize                   The size of this geometry.
 	 * @param InScale                  The scale of this geometry with respect to Normal Slate Coordinates.
 	 */
-	FGeometry( const FVector2D& OffsetFromParent, const FVector2D& ParentAbsolutePosition, const FVector2D& InLocalSize, float InScale );
+	FGeometry( const FVector2D& OffsetFromParent, const FVector2D& ParentAbsolutePosition, const FVector2D& InLocalSize, float InScale )
+		: Size(InLocalSize)
+		, Scale(1.0f)
+		, AbsolutePosition(0.0f, 0.0f)
+	{
+		// Since OffsetFromParent is given as a LocalSpaceOffset, we MUST convert this offset into the space of the parent to construct a valid layout transform.
+		// The extra TransformPoint below does this by converting the local offset to an offset in parent space.
+		FVector2D LayoutOffset = TransformPoint(InScale, OffsetFromParent);
+
+		FSlateLayoutTransform ParentAccumulatedLayoutTransform(InScale, ParentAbsolutePosition);
+		FSlateLayoutTransform LocalLayoutTransform(LayoutOffset);
+		FSlateLayoutTransform AccumulatedLayoutTransform = Concatenate(LocalLayoutTransform, ParentAccumulatedLayoutTransform);
+		AccumulatedRenderTransform = TransformCast<FSlateRenderTransform>(AccumulatedLayoutTransform);
+		// HACK to allow us to make FGeometry public members immutable to catch misuse.
+		const_cast<FVector2D&>( AbsolutePosition ) = AccumulatedLayoutTransform.GetTranslation();
+		const_cast<float&>( Scale ) = AccumulatedLayoutTransform.GetScale();
+		const_cast<FVector2D&>( Position ) = LocalLayoutTransform.GetTranslation();
+	}
 
 private:
 	/**
@@ -66,7 +83,29 @@ private:
 		const FSlateRenderTransform& InLocalRenderTransform, 
 		const FVector2D& InLocalRenderTransformPivot, 
 		const FSlateLayoutTransform& ParentAccumulatedLayoutTransform, 
-		const FSlateRenderTransform& ParentAccumulatedRenderTransform);
+		const FSlateRenderTransform& ParentAccumulatedRenderTransform)
+		: Size(InLocalSize)
+		, Scale(1.0f)
+		, AbsolutePosition(0.0f, 0.0f)
+		, AccumulatedRenderTransform(
+		Concatenate(
+		// convert the pivot to local space and make it the origin
+		Inverse(TransformPoint(FScale2D(InLocalSize), InLocalRenderTransformPivot)),
+		// apply the render transform in local space centered around the pivot
+		InLocalRenderTransform,
+		// translate the pivot point back.
+		TransformPoint(FScale2D(InLocalSize), InLocalRenderTransformPivot),
+		// apply the layout transform next.
+		InLocalLayoutTransform,
+		// finally apply the parent accumulated transform, which takes us to the root.
+		ParentAccumulatedRenderTransform))
+	{
+		FSlateLayoutTransform AccumulatedLayoutTransform = Concatenate(InLocalLayoutTransform, ParentAccumulatedLayoutTransform);
+		// HACK to allow us to make FGeometry public members immutable to catch misuse.
+		const_cast<FVector2D&>( AbsolutePosition ) = AccumulatedLayoutTransform.GetTranslation();
+		const_cast<float&>( Scale ) = AccumulatedLayoutTransform.GetScale();
+		const_cast<FVector2D&>( Position ) = InLocalLayoutTransform.GetTranslation();
+	}
 
 	/**
 	 * Construct a new geometry with a given size in LocalSpace (and identity render transform) that is attached to a parent geometry with the given layout and render transform. 
@@ -76,11 +115,22 @@ private:
 	 * @param ParentAccumulatedLayoutTransform	The accumulated layout transform of the parent widget. AccumulatedLayoutTransform = Concat(LocalLayoutTransform, ParentAccumulatedLayoutTransform).
 	 * @param ParentAccumulatedRenderTransform	The accumulated render transform of the parent widget. AccumulatedRenderTransform = Concat(LocalRenderTransform, LocalLayoutTransform, ParentAccumulatedRenderTransform).
 	 */
-	FGeometry( 
-		const FVector2D& InLocalSize, 
-		const FSlateLayoutTransform& InLocalLayoutTransform, 
-		const FSlateLayoutTransform& ParentAccumulatedLayoutTransform, 
-		const FSlateRenderTransform& ParentAccumulatedRenderTransform);
+	FGeometry(
+		const FVector2D& InLocalSize,
+		const FSlateLayoutTransform& InLocalLayoutTransform,
+		const FSlateLayoutTransform& ParentAccumulatedLayoutTransform,
+		const FSlateRenderTransform& ParentAccumulatedRenderTransform)
+		: Size(InLocalSize)
+		, Scale(1.0f)
+		, AbsolutePosition(0.0f, 0.0f)
+		, AccumulatedRenderTransform(Concatenate(InLocalLayoutTransform, ParentAccumulatedRenderTransform))
+	{
+		FSlateLayoutTransform AccumulatedLayoutTransform = Concatenate(InLocalLayoutTransform, ParentAccumulatedLayoutTransform);
+		// HACK to allow us to make FGeometry public members immutable to catch misuse.
+		const_cast<FVector2D&>( AbsolutePosition ) = AccumulatedLayoutTransform.GetTranslation();
+		const_cast<float&>( Scale ) = AccumulatedLayoutTransform.GetScale();
+		const_cast<FVector2D&>( Position ) = InLocalLayoutTransform.GetTranslation();
+	}
 
 public:
 
@@ -119,7 +169,10 @@ public:
 	 * @param LayoutTransform	Layout transform of the geometry.
 	 * @return					The new root geometry
 	 */
-	static FGeometry MakeRoot( const FVector2D& LocalSize, const FSlateLayoutTransform& LayoutTransform );
+	FORCEINLINE_DEBUGGABLE static FGeometry MakeRoot(const FVector2D& LocalSize, const FSlateLayoutTransform& LayoutTransform)
+	{
+		return FGeometry(LocalSize, LayoutTransform, FSlateLayoutTransform(), FSlateRenderTransform());
+	}
 
 	/**
 	 * Create a child geometry relative to this one with a given local space size, layout transform, and render transform.
@@ -133,7 +186,10 @@ public:
 	 *
 	 * @return					The new child geometry.
 	 */
-	FGeometry MakeChild( const FVector2D& LocalSize, const FSlateLayoutTransform& LayoutTransform, const FSlateRenderTransform& RenderTransform, const FVector2D& RenderTransformPivot ) const;
+	FORCEINLINE_DEBUGGABLE FGeometry MakeChild(const FVector2D& LocalSize, const FSlateLayoutTransform& LayoutTransform, const FSlateRenderTransform& RenderTransform, const FVector2D& RenderTransformPivot) const
+	{
+		return FGeometry(LocalSize, LayoutTransform, RenderTransform, RenderTransformPivot, GetAccumulatedLayoutTransform(), GetAccumulatedRenderTransform());
+	}
 
 	/**
 	 * Create a child geometry relative to this one with a given local space size, layout transform, and identity render transform.
@@ -145,7 +201,10 @@ public:
 	 *
 	 * @return					The new child geometry.
 	 */
-	FGeometry MakeChild( const FVector2D& LocalSize, const FSlateLayoutTransform& LayoutTransform ) const;
+	FORCEINLINE_DEBUGGABLE FGeometry MakeChild(const FVector2D& LocalSize, const FSlateLayoutTransform& LayoutTransform) const
+	{
+		return FGeometry(LocalSize, LayoutTransform, GetAccumulatedLayoutTransform(), GetAccumulatedRenderTransform());
+	}
 
 	/**
 	 * Create a child geometry+widget relative to this one using the given LayoutGeometry.
@@ -155,7 +214,7 @@ public:
 	 *
 	 * @return					The new child geometry.
 	 */
-	FArrangedWidget MakeChild( const TSharedRef<SWidget>& ChildWidget, const FLayoutGeometry& LayoutGeometry ) const;
+	FArrangedWidget MakeChild(const TSharedRef<SWidget>& ChildWidget, const FLayoutGeometry& LayoutGeometry) const;
 
 	/**
 	 * Create a child geometry+widget relative to this one with a given local space size and layout transform.
@@ -169,7 +228,7 @@ public:
 	 *
 	 * @return					The new child geometry+widget.
 	 */
-	FArrangedWidget MakeChild( const TSharedRef<SWidget>& ChildWidget, const FVector2D& LocalSize, const FSlateLayoutTransform& LayoutTransform ) const;
+	FArrangedWidget MakeChild(const TSharedRef<SWidget>& ChildWidget, const FVector2D& LocalSize, const FSlateLayoutTransform& LayoutTransform) const;
 
 	/**
 	 * !!! DEPRECATED FUNCTION !!! Use MakeChild taking a layout transform instead!
@@ -181,7 +240,12 @@ public:
 	 *
 	 * @return				The new child geometry.
 	 */
-	FGeometry MakeChild( const FVector2D& ChildOffset, const FVector2D& LocalSize, float ChildScale = 1.0f ) const;
+	FORCEINLINE_DEBUGGABLE FGeometry MakeChild(const FVector2D& ChildOffset, const FVector2D& LocalSize, float ChildScale = 1.0f) const
+	{
+		// Since ChildOffset is given as a LocalSpaceOffset, we MUST convert this offset into the space of the parent to construct a valid layout transform.
+		// The extra TransformPoint below does this by converting the local offset to an offset in parent space.
+		return FGeometry(LocalSize, FSlateLayoutTransform(ChildScale, TransformPoint(ChildScale, ChildOffset)), GetAccumulatedLayoutTransform(), GetAccumulatedRenderTransform());
+	}
 
 	/**
 	 * !!! DEPRECATED FUNCTION !!! Use MakeChild taking a layout transform instead!
@@ -195,14 +259,17 @@ public:
 	 *
 	 * @return				The new child geometry+widget.
 	 */
-	FArrangedWidget MakeChild( const TSharedRef<SWidget>& ChildWidget, const FVector2D& ChildOffset, const FVector2D& LocalSize, float ChildScale = 1.0f) const;
+	FArrangedWidget MakeChild(const TSharedRef<SWidget>& ChildWidget, const FVector2D& ChildOffset, const FVector2D& LocalSize, float ChildScale = 1.0f) const;
 
 	/**
 	 * Create a paint geometry that represents this geometry.
 	 * 
 	 * @return	The new paint geometry.
 	 */
-	FPaintGeometry ToPaintGeometry() const;
+	FORCEINLINE_DEBUGGABLE FPaintGeometry ToPaintGeometry() const
+	{
+		return FPaintGeometry(GetAccumulatedLayoutTransform(), GetAccumulatedRenderTransform(), Size);
+	}
 
 	/**
 	 * Create a paint geometry relative to this one with a given local space size and layout transform.
@@ -213,7 +280,11 @@ public:
 	 *
 	 * @return					The new paint geometry derived from this one.
 	 */
-	FPaintGeometry ToPaintGeometry( const FVector2D& LocalSize, const FSlateLayoutTransform& LayoutTransform) const;
+	FORCEINLINE_DEBUGGABLE FPaintGeometry ToPaintGeometry(const FVector2D& LocalSize, const FSlateLayoutTransform& LayoutTransform) const
+	{
+		FSlateLayoutTransform NewAccumulatedLayoutTransform = Concatenate(LayoutTransform, GetAccumulatedLayoutTransform());
+		return FPaintGeometry(NewAccumulatedLayoutTransform, Concatenate(LayoutTransform, GetAccumulatedRenderTransform()), LocalSize);
+	}
 
 	/**
 	 * Create a paint geometry with the same size as this geometry with a given layout transform.
@@ -223,7 +294,10 @@ public:
 	 * 
 	 * @return					The new paint geometry derived from this one.
 	 */
-	FPaintGeometry ToPaintGeometry( const FSlateLayoutTransform& LayoutTransform) const;
+	FORCEINLINE_DEBUGGABLE FPaintGeometry ToPaintGeometry(const FSlateLayoutTransform& LayoutTransform) const
+	{
+		return ToPaintGeometry(Size, LayoutTransform);
+	}
 
 	/**
 	 * !!! DEPRECATED FUNCTION !!! Use ToPaintGeometry taking a layout transform instead!
@@ -236,7 +310,12 @@ public:
 	 * 
 	 * @return				The new paint geometry derived from this one.
 	 */
-	FPaintGeometry ToPaintGeometry( const FVector2D& LocalOffset, const FVector2D& LocalSize, float LocalScale = 1.0f ) const;
+	FORCEINLINE_DEBUGGABLE FPaintGeometry ToPaintGeometry(const FVector2D& LocalOffset, const FVector2D& LocalSize, float LocalScale = 1.0f) const
+	{
+		// Since ChildOffset is given as a LocalSpaceOffset, we MUST convert this offset into the space of the parent to construct a valid layout transform.
+		// The extra TransformPoint below does this by converting the local offset to an offset in parent space.
+		return ToPaintGeometry(LocalSize, FSlateLayoutTransform(LocalScale, TransformPoint(LocalScale, LocalOffset)));
+	}
 
 	/**
 	 * !!! DEPRECATED FUNCTION !!! Use ToPaintGeometry taking a layout transform instead!
@@ -247,7 +326,10 @@ public:
 	 * 
 	 * @return				The new paint geometry derived from this one.
 	 */
-	FPaintGeometry ToOffsetPaintGeometry( const FVector2D& LocalOffset ) const;
+	FORCEINLINE_DEBUGGABLE FPaintGeometry ToOffsetPaintGeometry(const FVector2D& LocalOffset) const
+	{
+		return ToPaintGeometry(FSlateLayoutTransform(LocalOffset));
+	}
 
 	/**
 	 * Create a paint geometry relative to this one that whose local space is "inflated" by the specified amount in each direction.
@@ -257,21 +339,35 @@ public:
 	 * 
 	 * @return				The new paint geometry derived from this one.
 	 */
-	FPaintGeometry ToInflatedPaintGeometry( const FVector2D& InflateAmount ) const;
+	FORCEINLINE_DEBUGGABLE FPaintGeometry ToInflatedPaintGeometry(const FVector2D& InflateAmount) const
+	{
+		// This essentially adds (or subtracts) a border around the widget. We scale the size then offset by the border amount.
+		// Note this is not scaling child widgets, so the scale is not changing.
+		FVector2D NewSize = Size + InflateAmount * 2;
+		return ToPaintGeometry(NewSize, FSlateLayoutTransform(-InflateAmount));
+	}
 
 	/** 
 	 * Absolute coordinates could be either desktop or window space depending on what space the root of the widget hierarchy is in.
 	 * 
 	 * @return true if the provided location in absolute coordinates is within the bounds of this geometry. 
 	 */
-	bool IsUnderLocation( const FVector2D& AbsoluteCoordinate ) const;
+	FORCEINLINE_DEBUGGABLE bool IsUnderLocation(const FVector2D& AbsoluteCoordinate) const
+	{
+		// this render transform invert is a little expensive. We might consider caching it.
+		return FSlateRect(FVector2D(0.0f, 0.0f), Size).ContainsPoint(TransformPoint(Inverse(GetAccumulatedRenderTransform()), AbsoluteCoordinate));
+	}
 
 	/** 
 	 * Absolute coordinates could be either desktop or window space depending on what space the root of the widget hierarchy is in.
 	 * 
 	 * @return Transforms AbsoluteCoordinate into the local space of this Geometry. 
 	 */
-	FVector2D AbsoluteToLocal( FVector2D AbsoluteCoordinate ) const;
+	FORCEINLINE_DEBUGGABLE FVector2D AbsoluteToLocal(FVector2D AbsoluteCoordinate) const
+	{
+		// this render transform invert is a little expensive. We might consider caching it.
+		return TransformPoint(Inverse(GetAccumulatedRenderTransform()), AbsoluteCoordinate);
+	}
 
 	/**
 	 * Translates local coordinates into absolute coordinates
@@ -280,7 +376,10 @@ public:
 	 * 
 	 * @return  Absolute coordinates
 	 */
-	FVector2D LocalToAbsolute( FVector2D LocalCoordinate ) const;
+	FORCEINLINE_DEBUGGABLE FVector2D LocalToAbsolute(FVector2D LocalCoordinate) const
+	{
+		return TransformPoint(GetAccumulatedRenderTransform(), LocalCoordinate);
+	}
 
 	/**
 	 * Translates the local coordinates into local coordinates that after being transformed into absolute space will be rounded
@@ -288,7 +387,13 @@ public:
 	 * and not have the window start on a half pixel, which can cause the contents to jitter in relation to eachother as the tooltip 
 	 * or popup moves around.
 	 */
-	FVector2D LocalToRoundedLocal( FVector2D LocalCoordinate ) const;
+	FORCEINLINE_DEBUGGABLE FVector2D LocalToRoundedLocal(FVector2D LocalCoordinate) const
+	{
+		const FVector2D AbsoluteCoordinate = LocalToAbsolute(LocalCoordinate);
+		const FVector2D AbsoluteCoordinateRounded = FVector2D(FMath::RoundToInt(AbsoluteCoordinate.X), FMath::RoundToInt(AbsoluteCoordinate.Y));
+
+		return AbsoluteToLocal(AbsoluteCoordinateRounded);
+	}
 	
 	/**
 	 * !!! DEPRECATED !!! This legacy function does not account for render transforms.
@@ -301,10 +406,13 @@ public:
 	 *
 	 * @return  Allotted geometry rectangle in absolute coordinates.
 	 */
-	FSlateRect GetClippingRect( ) const;
+	FORCEINLINE_DEBUGGABLE FSlateRect GetClippingRect() const
+	{
+		return TransformRect(GetAccumulatedLayoutTransform(), FSlateRect(FVector2D(0.0f, 0.0f), Size));
+	}
 	
 	/** @return A String representation of this Geometry */
-	FString ToString( ) const;
+	FString ToString() const;
 
 	/** 
 	 * !!! DEPRECATED !!! This legacy function does not account for render transforms.
@@ -312,16 +420,22 @@ public:
 	 * Absolute coordinates could be either desktop or window space depending on what space the root of the widget hierarchy is in.
 	 *
 	 * @return the size of the geometry in absolute space */
-	FVector2D GetDrawSize() const;
+	FORCEINLINE_DEBUGGABLE FVector2D GetDrawSize() const
+	{
+		return TransformVector(GetAccumulatedLayoutTransform(), Size);
+	}
 
 	/** @return the size of the geometry in local space. */
-	const FVector2D& GetLocalSize() const { return Size; }
+	FORCEINLINE const FVector2D& GetLocalSize() const { return Size; }
 
 	/** @return the accumulated render transform. Shouldn't be needed in general. */
-	const FSlateRenderTransform& GetAccumulatedRenderTransform() const { return AccumulatedRenderTransform; }
+	FORCEINLINE const FSlateRenderTransform& GetAccumulatedRenderTransform() const { return AccumulatedRenderTransform; }
 
 	/** @return the accumulated layout transform. Shouldn't be needed in general. */
-	FSlateLayoutTransform GetAccumulatedLayoutTransform() const { return FSlateLayoutTransform(Scale, AbsolutePosition); }
+	FORCEINLINE FSlateLayoutTransform GetAccumulatedLayoutTransform() const
+	{
+		return FSlateLayoutTransform(Scale, AbsolutePosition);
+	}
 
 	/**
 	 * Special case method to append a layout transform to a geometry.
@@ -331,7 +445,14 @@ public:
 	 * 
 	 * @param LayoutTransform	An additional layout transform to append to this geoemtry.
 	 */
-	void AppendTransform(const FSlateLayoutTransform& LayoutTransform);
+	FORCEINLINE_DEBUGGABLE void AppendTransform(const FSlateLayoutTransform& LayoutTransform)
+	{
+		FSlateLayoutTransform AccumulatedLayoutTransform = ::Concatenate(GetAccumulatedLayoutTransform(), LayoutTransform);
+		AccumulatedRenderTransform = ::Concatenate(AccumulatedRenderTransform, LayoutTransform);
+		const_cast<FVector2D&>( AbsolutePosition ) = AccumulatedLayoutTransform.GetTranslation();
+		const_cast<float&>( Scale ) = AccumulatedLayoutTransform.GetScale();
+	}
+
 public:
 	/** 
 	 * 
diff --git a/Engine/Source/Runtime/SlateCore/Public/SlotBase.h b/Engine/Source/Runtime/SlateCore/Public/SlotBase.h
--- a/Engine/Source/Runtime/SlateCore/Public/SlotBase.h
+++ b/Engine/Source/Runtime/SlateCore/Public/SlotBase.h
@@ -12,16 +12,22 @@ public:
 
 	FSlotBase( const TSharedRef<SWidget>& InWidget );
 
-	void AttachWidget( const TSharedRef<SWidget>& InWidget );
-
 	virtual ~FSlotBase();
 
+	FORCEINLINE_DEBUGGABLE void AttachWidget( const TSharedRef<SWidget>& InWidget )
+	{
+		Widget = InWidget;
+	}
+
 	/**
 	 * Access the widget in the current slot.
 	 * There will always be a widget in the slot; sometimes it is
 	 * the SNullWidget instance.
 	 */
-	const TSharedRef<SWidget>& GetWidget() const;
+	FORCEINLINE_DEBUGGABLE const TSharedRef<SWidget>& GetWidget() const
+	{
+		return Widget;
+	}
 
 	/**
 	 * Remove the widget from its current slot.
diff --git a/Engine/Source/Runtime/SlateCore/Public/Types/PaintArgs.h b/Engine/Source/Runtime/SlateCore/Public/Types/PaintArgs.h
--- a/Engine/Source/Runtime/SlateCore/Public/Types/PaintArgs.h
+++ b/Engine/Source/Runtime/SlateCore/Public/Types/PaintArgs.h
@@ -22,12 +22,26 @@ class SLATECORE_API FPaintArgs
 {
 public:
 	FPaintArgs( const SWidget& Parent, FHittestGrid& InHittestGrid, FVector2D InWindowOffset, double InCurrentTime, float InDeltaTime );
+	
 	FORCEINLINE FPaintArgs WithNewParent(const SWidget* Parent) const
 	{
 		checkSlow(Parent);
 		return WithNewParent(*Parent);
 	}
-	FPaintArgs WithNewParent( const SWidget& Parent ) const;
+
+	FORCEINLINE_DEBUGGABLE FPaintArgs WithNewParent(const SWidget& Parent) const
+	{
+		FPaintArgs Args = FPaintArgs(Parent, this->Grid, this->WindowOffset, this->CurrentTime, this->DeltaTime);
+		Args.LastHittestIndex = this->LastHittestIndex;
+		Args.LastRecordedVisibility = this->LastRecordedVisibility;
+		Args.LayoutCache = this->LayoutCache;
+		Args.ParentCacheNode = this->ParentCacheNode;
+		Args.bIsCaching = this->bIsCaching;
+		Args.bIsVolatilityPass = this->bIsVolatilityPass;
+
+		return Args;
+	}
+
 	FPaintArgs EnableCaching(ILayoutCache* InLayoutCache, FCachedWidgetNode* InParentCacheNode, bool bEnableCaching, bool bEnableVolatility) const;
 	FPaintArgs RecordHittestGeometry(const SWidget* Widget, const FGeometry& WidgetGeometry, const FSlateRect& InClippingRect) const;
 	FPaintArgs InsertCustomHitTestPath( TSharedRef<ICustomHitTestPath> CustomHitTestPath, int32 HitTestIndex ) const;
diff --git a/Engine/Source/Runtime/SlateCore/Public/Widgets/SWidget.h b/Engine/Source/Runtime/SlateCore/Public/Widgets/SWidget.h
--- a/Engine/Source/Runtime/SlateCore/Public/Widgets/SWidget.h
+++ b/Engine/Source/Runtime/SlateCore/Public/Widgets/SWidget.h
@@ -8,6 +8,7 @@
 #include "SlateColor.h"
 #include "WidgetActiveTimerDelegate.h"
 #include "DeclarativeSyntaxSupport.h"
+#include "Layout/LayoutGeometry.h"
 
 class ISlateMetaData;
 class FActiveTimerHandle;
@@ -546,8 +547,9 @@ public:
 	 */
 	void SlatePrepass(float LayoutScaleMultiplier);
 
+	public:
 	/** @return the DesiredSize that was computed the last time CacheDesiredSize() was called. */
-	public:  const FVector2D& GetDesiredSize() const;
+	FORCEINLINE const FVector2D& GetDesiredSize() const { return DesiredSize; }
 
 	/**
 	 * The system calls this method. It performs a breadth-first traversal of every visible widget and asks
@@ -1149,3 +1151,33 @@ private:
 	/** If we're owned by a volatile widget, we need inherit that volatility and use as part of our volatility, but don't cache it. */
 	mutable bool bInheritedVolatility : 1;
 };
+
+//=================================================================
+// FGeometry Arranged Widget Inlined Functions
+//=================================================================
+
+FORCEINLINE_DEBUGGABLE FArrangedWidget FGeometry::MakeChild(const TSharedRef<SWidget>& ChildWidget, const FVector2D& LocalSize, const FSlateLayoutTransform& LayoutTransform) const
+{
+	// If there is no render transform set, use the simpler MakeChild call that doesn't bother concatenating the render transforms.
+	// This saves a significant amount of overhead since every widget does this, and most children don't have a render transform.
+	if ( ChildWidget->GetRenderTransform().IsSet() )
+	{
+		return FArrangedWidget(ChildWidget, MakeChild(LocalSize, LayoutTransform, ChildWidget->GetRenderTransform().GetValue(), ChildWidget->GetRenderTransformPivot()));
+	}
+	else
+	{
+		return FArrangedWidget(ChildWidget, MakeChild(LocalSize, LayoutTransform));
+	}
+}
+
+FORCEINLINE_DEBUGGABLE FArrangedWidget FGeometry::MakeChild(const TSharedRef<SWidget>& ChildWidget, const FLayoutGeometry& LayoutGeometry) const
+{
+	return MakeChild(ChildWidget, LayoutGeometry.GetSizeInLocalSpace(), LayoutGeometry.GetLocalToParentTransform());
+}
+
+FORCEINLINE_DEBUGGABLE FArrangedWidget FGeometry::MakeChild(const TSharedRef<SWidget>& ChildWidget, const FVector2D& ChildOffset, const FVector2D& LocalSize, float ChildScale) const
+{
+	// Since ChildOffset is given as a LocalSpaceOffset, we MUST convert this offset into the space of the parent to construct a valid layout transform.
+	// The extra TransformPoint below does this by converting the local offset to an offset in parent space.
+	return MakeChild(ChildWidget, LocalSize, FSlateLayoutTransform(ChildScale, TransformPoint(ChildScale, ChildOffset)));
+}
diff --git a/Engine/Source/Runtime/SlateRHIRenderer/Private/SlateRHIFontTexture.cpp b/Engine/Source/Runtime/SlateRHIRenderer/Private/SlateRHIFontTexture.cpp
--- a/Engine/Source/Runtime/SlateRHIRenderer/Private/SlateRHIFontTexture.cpp
+++ b/Engine/Source/Runtime/SlateRHIRenderer/Private/SlateRHIFontTexture.cpp
@@ -80,7 +80,7 @@ FSlateFontAtlasRHI::~FSlateFontAtlasRHI()
 
 void FSlateFontAtlasRHI::ReleaseResources()
 {
-	check( IsThreadSafeForSlateRendering() );
+	checkSlow( IsThreadSafeForSlateRendering() );
 
 	BeginReleaseResource( FontTexture.Get() );
 }
@@ -101,7 +101,7 @@ void FSlateFontAtlasRHI::ConditionalUpdateTexture()
 		}
 		else
 		{
-			check( IsThreadSafeForSlateRendering() );
+			checkSlow( IsThreadSafeForSlateRendering() );
 
 			BeginInitResource( FontTexture.Get() );
 
diff --git a/Engine/Source/Runtime/SlateRHIRenderer/Private/SlateRHIRenderer.cpp b/Engine/Source/Runtime/SlateRHIRenderer/Private/SlateRHIRenderer.cpp
--- a/Engine/Source/Runtime/SlateRHIRenderer/Private/SlateRHIRenderer.cpp
+++ b/Engine/Source/Runtime/SlateRHIRenderer/Private/SlateRHIRenderer.cpp
@@ -279,7 +279,7 @@ void FSlateRHIRenderer::CreateViewport( const TSharedRef<SWindow> Window )
 
 void FSlateRHIRenderer::ConditionalResizeViewport( FViewportInfo* ViewInfo, uint32 Width, uint32 Height, bool bFullscreen )
 {
-	check( IsThreadSafeForSlateRendering() );
+	checkSlow( IsThreadSafeForSlateRendering() );
 
 	if( IsInGameThread() && !IsInSlateThread() && ViewInfo && ( ViewInfo->Height != Height || ViewInfo->Width != Width ||  ViewInfo->bFullscreen != bFullscreen || !IsValidRef(ViewInfo->ViewportRHI) ) )
 	{
@@ -371,7 +371,7 @@ void FSlateRHIRenderer::RestoreSystemResolution(const TSharedRef<SWindow> InWind
 /** Called when a window is destroyed to give the renderer a chance to free resources */
 void FSlateRHIRenderer::OnWindowDestroyed( const TSharedRef<SWindow>& InWindow )
 {
-	check(IsThreadSafeForSlateRendering());
+	checkSlow(IsThreadSafeForSlateRendering());
 
 	FViewportInfo** ViewportInfoPtr = WindowToViewportInfo.Find( &InWindow.Get() );
 	if( ViewportInfoPtr )
@@ -576,7 +576,7 @@ void FSlateRHIRenderer::PrepareToTakeScreenshot(const FIntRect& Rect, TArray<FCo
  */
 void FSlateRHIRenderer::DrawWindows_Private( FSlateDrawBuffer& WindowDrawBuffer )
 {
-	check( IsThreadSafeForSlateRendering() );
+	checkSlow( IsThreadSafeForSlateRendering() );
 
 	// Enqueue a command to unlock the draw buffer after all windows have been drawn
 	ENQUEUE_UNIQUE_RENDER_COMMAND_ONEPARAMETER( SlateBeginDrawingWindowsCommand, 
@@ -917,7 +917,7 @@ void FSlateRHIRenderer::CopyWindowsToVirtualScreenBuffer(const TArray<FString>&
 	const FIntPoint ScaledCursorLocation = FIntPoint(MouseCursorLocation.X * XScaling, MouseCursorLocation.Y * YScaling) - VirtualScreenPos;
 
 	FSlateWindowElementList* WindowElementList = CrashTrackerResource->GetNextElementList();
-	*WindowElementList = FSlateWindowElementList(TSharedPtr<SWindow>());
+	WindowElementList->ResetBuffers();
 
 	// Don't draw cursor when it is hidden (mouse looking, scrolling, etc.)
 	// @todo livestream: The cursor is probably still hidden when dragging with the mouse captured (grabby hand)
@@ -1148,7 +1148,7 @@ void FSlateRHIRenderer::ReleaseDynamicResource( const FSlateBrush& InBrush )
 
 void* FSlateRHIRenderer::GetViewportResource( const SWindow& Window )
 {
-	check(IsThreadSafeForSlateRendering());
+	checkSlow(IsThreadSafeForSlateRendering());
 
 	FViewportInfo** InfoPtr = WindowToViewportInfo.Find( &Window );
 
@@ -1244,7 +1244,7 @@ void FSlateRHIRenderer::ReleaseAccessedResources()
 
 void FSlateRHIRenderer::RequestResize( const TSharedPtr<SWindow>& Window, uint32 NewWidth, uint32 NewHeight )
 {
-	check( IsThreadSafeForSlateRendering() );
+	checkSlow( IsThreadSafeForSlateRendering() );
 
 	FViewportInfo* ViewInfo = WindowToViewportInfo.FindRef( Window.Get() );
 
diff --git a/Engine/Source/Runtime/SlateRHIRenderer/Private/SlateRHIResourceManager.cpp b/Engine/Source/Runtime/SlateRHIRenderer/Private/SlateRHIResourceManager.cpp
--- a/Engine/Source/Runtime/SlateRHIRenderer/Private/SlateRHIResourceManager.cpp
+++ b/Engine/Source/Runtime/SlateRHIRenderer/Private/SlateRHIResourceManager.cpp
@@ -297,7 +297,7 @@ bool FSlateRHIResourceManager::LoadTexture( const FSlateBrush& InBrush, uint32&
  */
 bool FSlateRHIResourceManager::LoadTexture( const FName& TextureName, const FString& ResourcePath, uint32& Width, uint32& Height, TArray<uint8>& DecodedImage )
 {
-	check( IsThreadSafeForSlateRendering() );
+	checkSlow( IsThreadSafeForSlateRendering() );
 
 	bool bSucceeded = true;
 	uint32 BytesPerPixel = 4;
@@ -433,7 +433,7 @@ FSlateShaderResourceProxy* FSlateRHIResourceManager::GetShaderResource( const FS
 {
 	SCOPE_CYCLE_COUNTER( STAT_SlateGetResourceTime );
 
-	check( IsThreadSafeForSlateRendering() );
+	checkSlow( IsThreadSafeForSlateRendering() );
 
 	FSlateShaderResourceProxy* Texture = NULL;
 	if( !InBrush.IsDynamicallyLoaded() && !InBrush.HasUObject() )
@@ -581,7 +581,7 @@ TSharedPtr<FSlateUTextureResource> FSlateRHIResourceManager::MakeDynamicUTexture
 
 FSlateShaderResourceProxy* FSlateRHIResourceManager::FindOrCreateDynamicTextureResource(const FSlateBrush& InBrush)
 {
-	check( IsThreadSafeForSlateRendering() );
+	checkSlow( IsThreadSafeForSlateRendering() );
 
 	const FName ResourceName = InBrush.GetResourceName();
 	if ( ResourceName.IsValid() && ResourceName != NAME_None )
@@ -654,7 +654,7 @@ FSlateShaderResourceProxy* FSlateRHIResourceManager::FindOrCreateDynamicTextureR
 
 FSlateMaterialResource* FSlateRHIResourceManager::GetMaterialResource(const UObject* InMaterial, FVector2D ImageSize, FSlateShaderResource* TextureMask )
 {
-	check(IsThreadSafeForSlateRendering());
+	checkSlow(IsThreadSafeForSlateRendering());
 
 	const UMaterialInterface* Material = CastChecked<UMaterialInterface>(InMaterial);
 
@@ -707,7 +707,7 @@ bool FSlateRHIResourceManager::ContainsTexture( const FName& ResourceName ) cons
 
 void FSlateRHIResourceManager::ReleaseDynamicResource( const FSlateBrush& InBrush )
 {
-	check( IsThreadSafeForSlateRendering() )
+	checkSlow( IsThreadSafeForSlateRendering() )
 
 	// Note: Only dynamically loaded or utexture brushes can be dynamically released
 	if( InBrush.HasUObject() || InBrush.IsDynamicallyLoaded() )
@@ -798,7 +798,7 @@ void FSlateRHIResourceManager::UpdateTextureAtlases()
 
 void FSlateRHIResourceManager::ReleaseResources()
 {
-	check( IsThreadSafeForSlateRendering() );
+	checkSlow( IsThreadSafeForSlateRendering() );
 
 	for( int32 AtlasIndex = 0; AtlasIndex < TextureAtlases.Num(); ++AtlasIndex )
 	{
@@ -845,7 +845,7 @@ void FSlateRHIResourceManager::DeleteResources()
 
 void FSlateRHIResourceManager::ReloadTextures()
 {
-	check( IsThreadSafeForSlateRendering() );
+	checkSlow( IsThreadSafeForSlateRendering() );
 
 	// Release rendering resources
 	ReleaseResources();
diff --git a/Engine/Source/Runtime/SlateRHIRenderer/Private/SlateRHITextureAtlas.cpp b/Engine/Source/Runtime/SlateRHIRenderer/Private/SlateRHITextureAtlas.cpp
--- a/Engine/Source/Runtime/SlateRHIRenderer/Private/SlateRHITextureAtlas.cpp
+++ b/Engine/Source/Runtime/SlateRHIRenderer/Private/SlateRHITextureAtlas.cpp
@@ -60,7 +60,7 @@ void FSlateTextureAtlasRHI::UpdateTexture_RenderThread( FSlateTextureData* Rende
 
 void FSlateTextureAtlasRHI::ConditionalUpdateTexture( )
 {
-	check(IsThreadSafeForSlateRendering());
+	checkSlow(IsThreadSafeForSlateRendering());
 
 	if (bNeedsUpdate)
 	{
