Adding a protection to slate to prevent the font cache from being flushed when it's not the thread that currently owns slate rendering.  The cause was an automation smoke test that ran early on in the engine init that tests changing cultures, which flushes the font cache, but does so from the game thread.  However, currently the rendering thread is the owner of the slate thread, since it's playing a loading screen, thus the crash.

diff --git a/Engine/Source/Runtime/SlateCore/Private/Fonts/FontCache.cpp b/Engine/Source/Runtime/SlateCore/Private/Fonts/FontCache.cpp
--- a/Engine/Source/Runtime/SlateCore/Private/Fonts/FontCache.cpp
+++ b/Engine/Source/Runtime/SlateCore/Private/Fonts/FontCache.cpp
@@ -825,9 +825,9 @@ bool FSlateFontCache::ConditionalFlushCache()
 	bool bFlushed = false;
 	if( bFlushRequested )
 	{
-		FlushCache();
-		bFlushed = true;
 		bFlushRequested = false;
+		FlushCache();
+		bFlushed = !bFlushRequested;
 	}
 
 	return bFlushed;
@@ -851,24 +851,31 @@ void FSlateFontCache::ReleaseResources()
 
 void FSlateFontCache::FlushCache() const
 {
-	FlushData();
+	if ( IsThreadSafeForSlateRendering() )
+	{
+		FlushData();
 
-	FontToCharacterListCache.Empty();
-	ShapedGlyphToAtlasData.Empty();
+		FontToCharacterListCache.Empty();
+		ShapedGlyphToAtlasData.Empty();
 
-	for( int32 AtlasIndex = 0; AtlasIndex < FontAtlases.Num(); ++AtlasIndex ) 
-	{
-		FontAtlases[AtlasIndex]->ReleaseResources();
-	}
+		for ( int32 AtlasIndex = 0; AtlasIndex < FontAtlases.Num(); ++AtlasIndex )
+		{
+			FontAtlases[AtlasIndex]->ReleaseResources();
+		}
 
-	// hack
-	FSlateApplicationBase::Get().GetRenderer()->FlushCommands();
+		// hack
+		FSlateApplicationBase::Get().GetRenderer()->FlushCommands();
 
-	SET_DWORD_STAT( STAT_SlateNumFontAtlases, 0 );
+		SET_DWORD_STAT(STAT_SlateNumFontAtlases, 0);
 
-	FontAtlases.Empty();
+		FontAtlases.Empty();
 
-	UE_LOG( LogSlate, Verbose, TEXT("Slate font cache was flushed") );
+		UE_LOG(LogSlate, Verbose, TEXT("Slate font cache was flushed"));
+	}
+	else
+	{
+		bFlushRequested = true;
+	}
 }
 
 void FSlateFontCache::FlushData() const
diff --git a/Engine/Source/Runtime/SlateCore/Private/Fonts/SlateFontRenderer.cpp b/Engine/Source/Runtime/SlateCore/Private/Fonts/SlateFontRenderer.cpp
--- a/Engine/Source/Runtime/SlateCore/Private/Fonts/SlateFontRenderer.cpp
+++ b/Engine/Source/Runtime/SlateCore/Private/Fonts/SlateFontRenderer.cpp
@@ -237,7 +237,7 @@ void FSlateFontRenderer::GetRenderData(const FShapedGlyphEntry& InShapedGlyph, F
 {
 #if WITH_FREETYPE
 	FFreeTypeFaceGlyphData FaceGlyphData;
-	FaceGlyphData.FaceAndMemory = InShapedGlyph.FontFaceData->FontFace.IsValid() ? InShapedGlyph.FontFaceData->FontFace.Pin() : nullptr;
+	FaceGlyphData.FaceAndMemory = InShapedGlyph.FontFaceData->FontFace.Pin();
 	FaceGlyphData.GlyphIndex = InShapedGlyph.GlyphIndex;
 	FaceGlyphData.GlyphFlags = InShapedGlyph.FontFaceData->GlyphFlags;
 	check(FaceGlyphData.FaceAndMemory.IsValid());
diff --git a/Engine/Source/Runtime/SlateCore/Private/Fonts/SlateTextShaper.cpp b/Engine/Source/Runtime/SlateCore/Private/Fonts/SlateTextShaper.cpp
--- a/Engine/Source/Runtime/SlateCore/Private/Fonts/SlateTextShaper.cpp
+++ b/Engine/Source/Runtime/SlateCore/Private/Fonts/SlateTextShaper.cpp
@@ -230,8 +230,7 @@ void FSlateTextShaper::PerformKerningOnlyTextShaping(const TCHAR* InText, const
 	// Step 2) Now we use the font cache to get the size for each character, and kerning for each character pair
 	{
 		FCharacterList& CharacterList = FontCache->GetCharacterList(InFontInfo, InFontScale);
-		OutGlyphsToRender.Reserve(InTextLen);
-
+		OutGlyphsToRender.Reserve(OutGlyphsToRender.Num() + InTextLen);
 		for (const FKerningOnlyTextSequenceEntry& KerningOnlyTextSequenceEntry : KerningOnlyTextSequence)
 		{
 			const float FinalFontScale = InFontScale * KerningOnlyTextSequenceEntry.SubFontScalingFactor;
diff --git a/Engine/Source/Runtime/SlateCore/Private/Rendering/SlateRenderer.cpp b/Engine/Source/Runtime/SlateCore/Private/Rendering/SlateRenderer.cpp
--- a/Engine/Source/Runtime/SlateCore/Private/Rendering/SlateRenderer.cpp
+++ b/Engine/Source/Runtime/SlateCore/Private/Rendering/SlateRenderer.cpp
@@ -147,7 +147,16 @@ void FSlateRenderer::ReleaseCachingResourcesFor(const ILayoutCache* Cacher)
 /* Global functions
  *****************************************************************************/
 
-bool IsThreadSafeForSlateRendering( )
+bool IsThreadSafeForSlateRendering()
 {
-	return ((GSlateLoadingThreadId != 0) || IsInGameThread());
-}
+	if ( IsInGameThread() )
+	{
+		return GSlateLoadingThreadId == 0;
+	}
+	else
+	{
+		return FPlatformTLS::GetCurrentThreadId() == GSlateLoadingThreadId;
+	}
+
+	return false;
+}
\ No newline at end of file
diff --git a/Engine/Source/Runtime/SlateCore/Public/Fonts/FontCache.h b/Engine/Source/Runtime/SlateCore/Public/Fonts/FontCache.h
--- a/Engine/Source/Runtime/SlateCore/Public/Fonts/FontCache.h
+++ b/Engine/Source/Runtime/SlateCore/Public/Fonts/FontCache.h
@@ -727,5 +727,5 @@ private:
 	TSharedRef<ISlateFontAtlasFactory> FontAtlasFactory;
 
 	/** Whether or not we have a pending request to flush the cache when it is safe to do so */
-	mutable bool bFlushRequested;
+	volatile mutable bool bFlushRequested;
 };
