Slate - Additional fix to the fix for UE-23292.  Introducing DoesThreadOwnSlateRendering() that does what IsThreadSafeForSlateRendering() was defined to do.  IsThreadSafeForSlateRendering() has been restored to the previous implementation to mean "can this thread be a slate rendering thread" as opposied to being the current owner of slate rendering.

diff --git a/Engine/Source/Runtime/SlateCore/Private/Fonts/FontCache.cpp b/Engine/Source/Runtime/SlateCore/Private/Fonts/FontCache.cpp
--- a/Engine/Source/Runtime/SlateCore/Private/Fonts/FontCache.cpp
+++ b/Engine/Source/Runtime/SlateCore/Private/Fonts/FontCache.cpp
@@ -851,7 +851,7 @@ void FSlateFontCache::ReleaseResources()
 
 void FSlateFontCache::FlushCache() const
 {
-	if ( IsThreadSafeForSlateRendering() )
+	if ( DoesThreadOwnSlateRendering() )
 	{
 		FlushData();
 
diff --git a/Engine/Source/Runtime/SlateCore/Private/Rendering/SlateRenderer.cpp b/Engine/Source/Runtime/SlateCore/Private/Rendering/SlateRenderer.cpp
--- a/Engine/Source/Runtime/SlateCore/Private/Rendering/SlateRenderer.cpp
+++ b/Engine/Source/Runtime/SlateCore/Private/Rendering/SlateRenderer.cpp
@@ -148,6 +148,11 @@ void FSlateRenderer::ReleaseCachingResourcesFor(const ILayoutCache* Cacher)
  *****************************************************************************/
 
 bool IsThreadSafeForSlateRendering()
+{
+	return ( ( GSlateLoadingThreadId != 0 ) || IsInGameThread() );
+}
+
+bool DoesThreadOwnSlateRendering()
 {
 	if ( IsInGameThread() )
 	{
diff --git a/Engine/Source/Runtime/SlateCore/Public/Rendering/SlateRenderer.h b/Engine/Source/Runtime/SlateCore/Public/Rendering/SlateRenderer.h
--- a/Engine/Source/Runtime/SlateCore/Public/Rendering/SlateRenderer.h
+++ b/Engine/Source/Runtime/SlateCore/Public/Rendering/SlateRenderer.h
@@ -408,4 +408,10 @@ protected:
  * If the slate loading thread exists, then yes, it is always safe
  * Otherwise, we have to be on the game thread
  */
-bool SLATECORE_API IsThreadSafeForSlateRendering( );
+bool SLATECORE_API IsThreadSafeForSlateRendering();
+
+/**
+ * If it's the game thread, and there's no loading thread, then it owns slate rendering.
+ * However if there's a loading thread, it is the exlusive owner of slate rendering.
+ */
+bool SLATECORE_API DoesThreadOwnSlateRendering();
