Slate - Changed a massive amount of things attempting to improve the Slate and VR intermingling.  Popup layers should now function on widget components, lots of changes to make focus work better.  Added hacks to the tab manager to allow us to steal tabs as they open, alternative plan in the works, probably going to undo some portion of these changes later.  The 3DWidget renderer now handles child windows and properly draws those nested windows onto the render target.  In the same way tooltips can have a visualize tooltip, all widgets now have the option to Visualize Popup.  Assorted hacks in various locations to make slate handle virtual windows better and not explode just because there's no native window associated with something.

diff --git a/Engine/Plugins/Runtime/WindowsMoviePlayer/Source/WindowsMoviePlayer/WindowsMoviePlayer.Build.cs b/Engine/Plugins/Runtime/WindowsMoviePlayer/Source/WindowsMoviePlayer/WindowsMoviePlayer.Build.cs
--- a/Engine/Plugins/Runtime/WindowsMoviePlayer/Source/WindowsMoviePlayer/WindowsMoviePlayer.Build.cs
+++ b/Engine/Plugins/Runtime/WindowsMoviePlayer/Source/WindowsMoviePlayer/WindowsMoviePlayer.Build.cs
@@ -27,6 +27,7 @@ namespace UnrealBuildTool.Rules
                     "MoviePlayer",
                     "RenderCore",
                     "RHI",
+                    "SlateCore",
                     "Slate"
 				}
 				);
diff --git a/Engine/Source/Editor/VREditor/VREditorFloatingUI.cpp b/Engine/Source/Editor/VREditor/VREditorFloatingUI.cpp
--- a/Engine/Source/Editor/VREditor/VREditorFloatingUI.cpp
+++ b/Engine/Source/Editor/VREditor/VREditorFloatingUI.cpp
@@ -6,6 +6,8 @@
 #include "VREditorBaseUserWidget.h"
 #include "VREditorMode.h"
 #include "WidgetComponent.h"
+#include "LevelEditor.h"
+#include "ILevelViewport.h"
 
 
 namespace VREd
@@ -36,7 +38,8 @@ AVREditorFloatingUI::AVREditorFloatingUI()
 	check( SceneComponent != nullptr );
 	this->RootComponent = SceneComponent;
 
-	this->WidgetComponent = CreateDefaultSubobject<UWidgetComponent>( TEXT( "WidgetComponent" ), bTransient );
+	WidgetComponent = CreateDefaultSubobject<UWidgetComponent>( TEXT( "WidgetComponent" ), bTransient );
+	WidgetComponent->SetEditTimeUsable(true);
 	WidgetComponent->AttachTo( SceneComponent );
 }
 
@@ -210,7 +213,7 @@ void AVREditorFloatingUI::UpdateFadingState( const float DeltaTime )
  		// Set material color
 		// NOTE: We intentionally make the UI quite dark here, so it doesn't bloom out in HDR
 		const float UIBrightness = FadeAlpha * VREd::UIBrightness->GetFloat();
- 		WidgetComponent->SetTintColorAndOpacity( FLinearColor( UIBrightness, UIBrightness, UIBrightness ).CopyWithNewOpacity( FadeAlpha ) );
+		WidgetComponent->SetTintColorAndOpacity( FLinearColor( UIBrightness, UIBrightness, UIBrightness ).CopyWithNewOpacity( FadeAlpha ) );
 	}
 }
 
@@ -443,6 +446,30 @@ void AVREditorFloatingUI::ShowUI( const bool bShow, const bool bAllowFading, con
 			FadeAlpha = bShow ? 1.0f : 0.0f;
 		}
 
+		// @todo vreditor Could fail a million different ways.
+		if ( bShow )
+		{
+			FLevelEditorModule& LevelEditorModule = FModuleManager::Get().GetModuleChecked<FLevelEditorModule>(TEXT("LevelEditor"));
+			TSharedPtr<ILevelViewport> ActiveLevelViewport = LevelEditorModule.GetFirstActiveViewport();
+			
+			if ( ActiveLevelViewport.IsValid() )
+			{
+				TSharedPtr<SViewport> ViewportWidget = ActiveLevelViewport->GetViewportWidget().Pin();
+				WidgetComponent->RegisterHitTesterWithViewport(ViewportWidget);
+			}
+		}
+		else
+		{
+			FLevelEditorModule& LevelEditorModule = FModuleManager::Get().GetModuleChecked<FLevelEditorModule>(TEXT("LevelEditor"));
+			TSharedPtr<ILevelViewport> ActiveLevelViewport = LevelEditorModule.GetFirstActiveViewport();
+
+			if ( ActiveLevelViewport.IsValid() )
+			{
+				TSharedPtr<SViewport> ViewportWidget = ActiveLevelViewport->GetViewportWidget().Pin();
+				WidgetComponent->UnregisterHitTesterWithViewport(ViewportWidget);
+			}
+		}
+
 		FadeDelay = InitFadeDelay;
 	}
 }
diff --git a/Engine/Source/Editor/VREditor/VREditorMode.cpp b/Engine/Source/Editor/VREditor/VREditorMode.cpp
--- a/Engine/Source/Editor/VREditor/VREditorMode.cpp
+++ b/Engine/Source/Editor/VREditor/VREditorMode.cpp
@@ -391,6 +391,17 @@ void FVREditorMode::Enter()
 
 	WorldInteraction.Reset( new FVREditorWorldInteraction( *this ) );
 
+	// If we're not actually in VR mode, and instead we've forced it, we should spawn some UI at the Origin to help test
+	if ( !bActuallyUsingVR )
+	{
+		const int32 HandIndex = 1;
+		const bool bShouldShow = true;
+
+		FVirtualHand& Hand = VirtualHands[HandIndex];
+		Hand.Transform = FTransform(FVector(0, 0, 0));
+
+		UISystem->ShowEditorUIPanel(FVREditorUISystem::EEditorUIPanel::ContentBrowser, HandIndex, bShouldShow);
+	}
 
 	bIsFullyInitialized = true;
 }
diff --git a/Engine/Source/Editor/VREditor/VREditorUISystem.cpp b/Engine/Source/Editor/VREditor/VREditorUISystem.cpp
--- a/Engine/Source/Editor/VREditor/VREditorUISystem.cpp
+++ b/Engine/Source/Editor/VREditor/VREditorUISystem.cpp
@@ -601,7 +601,9 @@ void FVREditorUISystem::CreateUIs()
 			Config.bCanShowRealTimeThumbnails = false;
 			Config.InitialAssetViewType = EAssetViewType::Tile;
 			Config.bCanShowDevelopersFolder = false;
-			Config.bCanShowFolders = false;
+			Config.bCanShowFolders = true;
+			Config.bUsePathPicker = true;
+			Config.bCanShowFilters = true;
 			Config.bCanShowAssetSearch = false;
 			Config.bUseSourcesView = true;
 			Config.bExpandSourcesView = true;
@@ -730,6 +732,30 @@ void FVREditorUISystem::CreateUIs()
 
 			EditorUIPanels[ (int32)EEditorUIPanel::Tutorial ] = TutorialUI;		
 		}
+
+		{
+			const FIntPoint Resolution(VREd::EditorUIResolutionX->GetInt(), VREd::EditorUIResolutionY->GetInt());
+
+			const bool bWithSceneComponent = false;
+			AVREditorFloatingUI* TabManagerUI = GetOwner().SpawnTransientSceneActor< AVREditorFloatingUI >(TEXT("TabManager"), bWithSceneComponent);
+			TabManagerUI->SetSlateWidget(*this, SNullWidget::NullWidget, Resolution, VREd::EditorUISize->GetFloat(), AVREditorFloatingUI::EDockedTo::Nothing);
+			TabManagerUI->ShowUI(true);
+			TabManagerUI->SetRelativeOffset(FVector(0, 0, 1000));
+			FloatingUIs.Add(TabManagerUI);
+
+			TSharedPtr<SWindow> TabManagerWindow = TabManagerUI->GetWidgetComponent()->GetSlateWindow();
+			TSharedRef<SWindow> TabManagerWindowRef = TabManagerWindow.ToSharedRef();
+			ProxyTabManager = MakeShareable(new FProxyTabmanager(TabManagerWindowRef));
+
+			TSharedRef<FTabManager::FLayout> LoadedLayout = FTabManager::NewLayout(TEXT("VRTabManager"));
+			LoadedLayout->AddArea(FTabManager::NewPrimaryArea());
+
+			TSharedPtr<SWidget> DockLayout = FGlobalTabmanager::Get()->RestoreFrom(LoadedLayout, TabManagerWindowRef, false);
+			TabManagerWindowRef->SetContent(DockLayout.ToSharedRef());
+
+			// We're going to start stealing tabs from the global tab manager inserting them into the world instead.
+			FGlobalTabmanager::Get()->SetProxyTabManager(ProxyTabManager);
+		}
 	}
 }
 
@@ -748,6 +774,11 @@ void FVREditorUISystem::CleanUpActorsBeforeMapChangeOrSimulate()
 	EditorUIPanels.Reset();
 	QuickRadialMenu = nullptr;
 	QuickMenuUI = nullptr;
+
+	ProxyTabManager.Reset();
+
+	// Remove the proxy tab manager, we don't want to steal tabs any more.
+	FGlobalTabmanager::Get()->SetProxyTabManager(TSharedPtr<FProxyTabmanager>());
 }
 
 
diff --git a/Engine/Source/Editor/VREditor/VREditorUISystem.h b/Engine/Source/Editor/VREditor/VREditorUISystem.h
--- a/Engine/Source/Editor/VREditor/VREditorUISystem.h
+++ b/Engine/Source/Editor/VREditor/VREditorUISystem.h
@@ -153,5 +153,13 @@ protected:
 
 	/** If the panels should become visible or invisible next toggle*/
 	bool bPanelVisibilityToggle;
+
+
+	//
+	// Tab Manager UI
+	//
+
+	/** Allows us to steal the global tabs and show them in the world. */
+	TSharedPtr<FProxyTabmanager> ProxyTabManager;
 };
 
diff --git a/Engine/Source/Editor/VREditor/VREditorWidgetComponent.cpp b/Engine/Source/Editor/VREditor/VREditorWidgetComponent.cpp
new file mode 100644
--- /dev/null
+++ b/Engine/Source/Editor/VREditor/VREditorWidgetComponent.cpp
@@ -0,0 +1,10 @@
+// Copyright 1998-2015 Epic Games, Inc. All Rights Reserved.
+
+#include "VREditorModule.h"
+#include "VREditorWidgetComponent.h"
+
+UVREditorWidgetComponent::UVREditorWidgetComponent(const FObjectInitializer& ObjectInitializer)
+	: Super( ObjectInitializer )
+{
+
+}
diff --git a/Engine/Source/Editor/VREditor/VREditorWidgetComponent.h b/Engine/Source/Editor/VREditor/VREditorWidgetComponent.h
new file mode 100644
--- /dev/null
+++ b/Engine/Source/Editor/VREditor/VREditorWidgetComponent.h
@@ -0,0 +1,19 @@
+// Copyright 1998-2015 Epic Games, Inc. All Rights Reserved.
+
+#pragma once
+
+#include "Components/WidgetComponent.h"
+#include "VREditorWidgetComponent.generated.h"
+
+/**
+ * A specialized WidgetComponent for the VREditor
+ */
+UCLASS(hidedropdown)
+class UVREditorWidgetComponent : public UWidgetComponent
+{
+	GENERATED_BODY()
+
+public:
+	/** Default constructor that sets up CDO properties */
+	UVREditorWidgetComponent(const FObjectInitializer& ObjectInitializer);
+};
\ No newline at end of file
diff --git a/Engine/Source/Editor/VREditor/VREditorWorldInteraction.cpp b/Engine/Source/Editor/VREditor/VREditorWorldInteraction.cpp
--- a/Engine/Source/Editor/VREditor/VREditorWorldInteraction.cpp
+++ b/Engine/Source/Editor/VREditor/VREditorWorldInteraction.cpp
@@ -2683,6 +2683,13 @@ void FVREditorWorldInteraction::OnAssetDragStartedFromContentBrowser( const TArr
 		}
 	}
 
+	// We're always expecting a hand to be hovering at the time we receive this event, but in forced VR mode
+	// there may not be a hand if you're interacting with the legit CB.
+	if ( !FVREditorMode::IsActuallyUsingVR() && PlacingWithHandIndex == INDEX_NONE )
+	{
+		return;
+	}
+
 	// We're always expecting a hand to be hovering at the time we receive this event
 	if( PlacingWithHandIndex != INDEX_NONE )	// @todo vreditor: This ensure fires when in "fake VR" mode when you drag an asset (Expected)
 	{
diff --git a/Engine/Source/Runtime/Engine/Private/Slate/SceneViewport.cpp b/Engine/Source/Runtime/Engine/Private/Slate/SceneViewport.cpp
--- a/Engine/Source/Runtime/Engine/Private/Slate/SceneViewport.cpp
+++ b/Engine/Source/Runtime/Engine/Private/Slate/SceneViewport.cpp
@@ -288,13 +288,16 @@ void FSceneViewport::OnDrawViewport( const FGeometry& AllottedGeometry, const FS
 	if( GetSizeXY() != DrawSize )
 	{
 		TSharedPtr<SWindow> Window = FSlateApplication::Get().FindWidgetWindow( ViewportWidget.Pin().ToSharedRef() );
-
-		check(Window.IsValid());
-		if (Window->IsViewportSizeDrivenByWindow())
+		if ( Window.IsValid() )
 		{
-			ResizeViewport(FMath::Max(0, DrawSize.X), FMath::Max(0, DrawSize.Y), Window->GetWindowMode(), 0, 0);
+			//@HACK VREDITOR
+			//check(Window.IsValid());
+			if ( Window->IsViewportSizeDrivenByWindow() )
+			{
+				ResizeViewport(FMath::Max(0, DrawSize.X), FMath::Max(0, DrawSize.Y), Window->GetWindowMode(), 0, 0);
+			}
 		}
-	}	
+	}
 	
 	// Cannot pass negative canvas positions
 	float CanvasMinX = FMath::Max(0.0f, AllottedGeometry.AbsolutePosition.X);
@@ -398,6 +401,20 @@ FReply FSceneViewport::OnMouseButtonDown( const FGeometry& InGeometry, const FPo
 				CurrentReplyState = FReply::Unhandled();
 			}
 		}
+		else
+		{
+			TSharedRef<SViewport> ViewportWidgetRef = ViewportWidget.Pin().ToSharedRef();
+			if ( ViewportWidgetRef->HasUserFocusedDescendants(InMouseEvent.GetUserIndex()) )
+			{
+				// If we're still focused on a descendant, force it to stay focused.  Need to do this to fool SlateApplication
+				// otherwise it will reassign focus thinking that nobody requested it.
+				TSharedPtr<SWidget> FocusedWidgetPtr = FSlateApplication::Get().GetUserFocusedWidget(InMouseEvent.GetUserIndex());
+				if ( FocusedWidgetPtr.IsValid() )
+				{
+					CurrentReplyState.SetUserFocus(FocusedWidgetPtr.ToSharedRef(), EFocusCause::SetDirectly, false);
+				}
+			}
+		}
 
 		// a new menu was opened if there was previously not a menu visible but now there is
 		const bool bNewMenuWasOpened = !bAnyMenuWasVisible && FSlateApplication::Get().AnyMenusVisible();
diff --git a/Engine/Source/Runtime/Slate/Private/Framework/Application/MenuStack.cpp b/Engine/Source/Runtime/Slate/Private/Framework/Application/MenuStack.cpp
--- a/Engine/Source/Runtime/Slate/Private/Framework/Application/MenuStack.cpp
+++ b/Engine/Source/Runtime/Slate/Private/Framework/Application/MenuStack.cpp
@@ -210,7 +210,7 @@ TSharedRef<IMenu> FMenuStack::Push(const FWidgetPath& InOwnerPath, const TShared
 		ActiveMethod = InMethod.IsSet() ? FPopupMethodReply::UseMethod(InMethod.GetValue()) : QueryPopupMethod(InOwnerPath);
 
 		// The host window is determined when a new root menu is pushed
-		SetHostWindow(InOwnerPath.GetWindow());
+		SetHostPath(InOwnerPath);
 	}
 
 	return PushInternal(ParentMenu, InContent, Anchor, TransitionEffect, bFocusImmediately, ActiveMethod.GetShouldThrottle(), bIsCollapsedByParent, bEnablePerPixelTransparency);
@@ -245,7 +245,7 @@ TSharedRef<IMenu> FMenuStack::PushHosted(const FWidgetPath& InOwnerPath, const T
 		ActiveMethod = FPopupMethodReply::UseMethod(EPopupMethod::UseCurrentWindow);
 
 		// The host window is determined when a new root menu is pushed
-		SetHostWindow(InOwnerPath.GetWindow());
+		SetHostPath(InOwnerPath);
 	}
 
 	return PushHosted(ParentMenu, InMenuHost, InContent, OutWrappedContent, TransitionEffect, ShouldThrottle);
@@ -580,27 +580,32 @@ void FMenuStack::DismissInternal(int32 FirstStackIndexToRemove)
 	}
 }
 
-void FMenuStack::SetHostWindow(TSharedPtr<SWindow> InWindow)
+void FMenuStack::SetHostPath(const FWidgetPath& InOwnerPath)
 {
-	if (HostWindow != InWindow)
+	if ( HostPopupLayer.IsValid() )
 	{
-		// If the host window is changing, remove the popup panel from the previous host
-		if (HostWindow.IsValid() && HostWindowPopupPanel.IsValid())
+		if ( !InOwnerPath.ContainsWidget(HostPopupLayer->GetHost()) )
 		{
-			HostWindow->RemoveOverlaySlot(HostWindowPopupPanel.ToSharedRef());
+			HostPopupLayer->Remove();
+			HostPopupLayer.Reset();
 			HostWindowPopupPanel.Reset();
 		}
+	}
 
-		HostWindow = InWindow;
+	HostWindow = InOwnerPath.IsValid() ? InOwnerPath.GetWindow() : TSharedPtr<SWindow>();
 
-		// If the host window is changing, add a popup panel to the new host
-		if (HostWindow.IsValid())
+	if ( HostWindow.IsValid() && !HostWindowPopupPanel.IsValid() )
+	{
+		TSharedRef<SMenuPanel> NewHostWindowPopupPanel = SNew(SMenuPanel);
+		for ( int i = InOwnerPath.Widgets.Num() - 1; i >= 0; i-- )
 		{
-			// setup a popup layer on the new window
-			HostWindow->AddOverlaySlot()
-			[
-				SAssignNew(HostWindowPopupPanel, SMenuPanel)
-			];
+			const TSharedRef<SWidget>& HostWidget = InOwnerPath.Widgets[i].Widget;
+			HostPopupLayer = HostWidget->OnVisualizePopup(NewHostWindowPopupPanel);
+			if ( HostPopupLayer.IsValid() )
+			{
+				HostWindowPopupPanel = NewHostWindowPopupPanel;
+				break;
+			}
 		}
 	}
 }
@@ -631,7 +636,7 @@ void FMenuStack::OnMenuDestroyed(TSharedRef<IMenu> InMenu)
 				FSlateThrottleManager::Get().LeaveResponsiveMode(ThrottleHandle);
 			}
 
-			SetHostWindow(TSharedPtr<SWindow>());
+			SetHostPath(FWidgetPath());
 		}
 	}
 }
@@ -713,7 +718,7 @@ void FMenuStack::OnWindowDestroyed(TSharedRef<SWindow> InWindow)
 		Stack.Empty();
 		CachedContentMap.Empty();
 
-		SetHostWindow(TSharedPtr<SWindow>());
+		SetHostPath(FWidgetPath());
 	}
 	else
 	{
diff --git a/Engine/Source/Runtime/Slate/Private/Framework/Application/SlateApplication.cpp b/Engine/Source/Runtime/Slate/Private/Framework/Application/SlateApplication.cpp
--- a/Engine/Source/Runtime/Slate/Private/Framework/Application/SlateApplication.cpp
+++ b/Engine/Source/Runtime/Slate/Private/Framework/Application/SlateApplication.cpp
@@ -1890,8 +1890,9 @@ void FSlateApplication::SetModalWindowStackEndedDelegate(FModalWindowStackEnded
 
 TSharedRef<SWindow> FSlateApplication::AddWindowAsNativeChild( TSharedRef<SWindow> InSlateWindow, TSharedRef<SWindow> InParentWindow, const bool bShowImmediately )
 {
+	// @VREDITOR HACK
 	// Parent window must already have been added
-	checkSlow(FSlateWindowHelper::ContainsWindow(SlateWindows, InParentWindow));
+	//checkSlow(FSlateWindowHelper::ContainsWindow(SlateWindows, InParentWindow));
 
 	// Add the Slate window to the Slate application's top-level window array.  Note that neither the Slate window
 	// or the native window are ready to be used yet, however we need to make sure they're in the Slate window
@@ -1899,16 +1900,21 @@ TSharedRef<SWindow> FSlateApplication::AddWindowAsNativeChild( TSharedRef<SWindo
 	// activation message may be sent by the OS as soon as the window is shown (in the Init function), and if we
 	// don't add the Slate window to our window list, we wouldn't be able to route that message to the window.
 	InParentWindow->AddChildWindow( InSlateWindow );
-	TSharedRef<FGenericWindow> NewWindow = MakeWindow( InSlateWindow, bShowImmediately );
 
-	if( bShowImmediately )
+	// Only make native generic windows if the parent has one.
+	if ( InParentWindow->GetNativeWindow()->GetOSWindowHandle() )
 	{
-		InSlateWindow->ShowWindow();
+		TSharedRef<FGenericWindow> NewWindow = MakeWindow(InSlateWindow, bShowImmediately);
 
-		//@todo Slate: Potentially dangerous and annoying if all slate windows that are created steal focus.
-		if( InSlateWindow->SupportsKeyboardFocus() && InSlateWindow->IsFocusedInitially() )
+		if ( bShowImmediately )
 		{
-			InSlateWindow->GetNativeWindow()->SetWindowFocus();
+			InSlateWindow->ShowWindow();
+
+			//@todo Slate: Potentially dangerous and annoying if all slate windows that are created steal focus.
+			if ( InSlateWindow->SupportsKeyboardFocus() && InSlateWindow->IsFocusedInitially() )
+			{
+				InSlateWindow->GetNativeWindow()->SetWindowFocus();
+			}
 		}
 	}
 
@@ -2397,6 +2403,8 @@ bool FSlateApplication::SetUserFocus(const uint32 InUserIndex, const FWidgetPath
 
 	if (InFocusPath.IsValid())
 	{
+		//UE_LOG(LogSlate, Warning, TEXT("Focus for user %i seeking focus path:\n%s"), InUserIndex, *InFocusPath.ToString());
+
 		for (int32 WidgetIndex = InFocusPath.Widgets.Num() - 1; WidgetIndex >= 0; --WidgetIndex)
 		{
 			const FArrangedWidget& WidgetToFocus = InFocusPath.Widgets[WidgetIndex];
@@ -2407,6 +2415,7 @@ bool FSlateApplication::SetUserFocus(const uint32 InUserIndex, const FWidgetPath
 				// Is we aren't changing focus then simply return
 				if (WidgetToFocus.Widget == OldFocusedWidget)
 				{
+					//UE_LOG(LogSlate, Warning, TEXT("--Focus Has Not Changed--"));
 					return false;
 				}
 				NewFocusedWidget = WidgetToFocus.Widget;
@@ -4430,15 +4439,21 @@ FReply FSlateApplication::RoutePointerDownEvent(FWidgetPath& WidgetsUnderPointer
 
 	const TSharedPtr<SWidget> PreviouslyFocusedWidget = GetKeyboardFocusedWidget();
 
-	FReply Reply = FEventRouter::Route<FReply>(this, FEventRouter::FTunnelPolicy(WidgetsUnderPointer), PointerEvent, [] (const FArrangedWidget TargetWidget, const FPointerEvent& Event)
+	int32 WidgetHandledIndex = 0;
+
+	FReply Reply = FEventRouter::Route<FReply>(this, FEventRouter::FTunnelPolicy(WidgetsUnderPointer), PointerEvent, [&WidgetHandledIndex] (const FArrangedWidget TargetWidget, const FPointerEvent& Event)
 	{
+		WidgetHandledIndex++;
 		return TargetWidget.Widget->OnPreviewMouseButtonDown(TargetWidget.Geometry, Event);
 	});
 
 	if ( !Reply.IsEventHandled() )
 	{
-		Reply = FEventRouter::Route<FReply>(this, FEventRouter::FBubblePolicy(WidgetsUnderPointer), PointerEvent, [this] (const FArrangedWidget TargetWidget, const FPointerEvent& Event)
+		WidgetHandledIndex = WidgetsUnderPointer.Widgets.Num();
+
+		Reply = FEventRouter::Route<FReply>(this, FEventRouter::FBubblePolicy(WidgetsUnderPointer), PointerEvent, [this, &WidgetHandledIndex] (const FArrangedWidget TargetWidget, const FPointerEvent& Event)
 		{
+			WidgetHandledIndex--;
 			FReply ThisReply = FReply::Unhandled();
 			if ( !ThisReply.IsEventHandled() )
 			{
@@ -4454,20 +4469,27 @@ FReply FSlateApplication::RoutePointerDownEvent(FWidgetPath& WidgetsUnderPointer
 			return ThisReply;
 		});
 	}
+	else
+	{
+		// Since we always advance 1, need to back that out if it was handled by the preview mouse down code.
+		WidgetHandledIndex--;
+	}
+
 	LOG_EVENT(EEventLog::MouseButtonDown, Reply);
 
 	// If none of the widgets requested keyboard focus to be set (or set the keyboard focus explicitly), set it to the leaf-most widget under the mouse.
 	// On Mac we prevent the OS from activating the window on mouse down, so we have full control and can activate only if there's nothing draggable under the mouse cursor.
 	const bool bFocusChangedByEventHandler = PreviouslyFocusedWidget != GetKeyboardFocusedWidget();
-	if ((!Reply.GetUserFocusRecepient().IsValid() || (PLATFORM_MAC && PointerEvent.GetEffectingButton() == EKeys::LeftMouseButton && !DragDetector.DetectDragForWidget.IsValid()))
-		&& (!bFocusChangedByEventHandler || bNeedToActivateWindow))
+	if ( ( !Reply.GetUserFocusRecepient().IsValid() || ( PLATFORM_MAC && PointerEvent.GetEffectingButton() == EKeys::LeftMouseButton && !DragDetector.DetectDragForWidget.IsValid() ) )
+		&& ( !bFocusChangedByEventHandler || bNeedToActivateWindow ) )
 	{
-		// The event handler for OnMouseButton down may have altered the widget hierarchy.
-		// Refresh the previously-cached widget path.
-		WidgetsUnderPointer = LocateWindowUnderMouse(PointerEvent.GetScreenSpacePosition(), GetInteractiveTopLevelWindows());
+		// If the event wasn't handled, we need to resample the screen to find
+		// @HACK VREDITOR No longer resampling from the screen down, that won't
+		//                work well with 3D widgets, so we're just going to use the existing widgets under pointer
+		//                that we already passed in.
 
 		bool bFocusCandidateFound = false;
-		for ( int32 WidgetIndex = WidgetsUnderPointer.Widgets.Num() - 1; !bFocusCandidateFound && WidgetIndex >= 0; --WidgetIndex )
+		for ( int32 WidgetIndex = WidgetHandledIndex; !bFocusCandidateFound && WidgetIndex >= 0; --WidgetIndex )
 		{
 			FArrangedWidget& CurWidget = WidgetsUnderPointer.Widgets[WidgetIndex];
 			if ( CurWidget.Widget->SupportsKeyboardFocus() )
@@ -4480,7 +4502,7 @@ FReply FSlateApplication::RoutePointerDownEvent(FWidgetPath& WidgetsUnderPointer
 
 #if PLATFORM_MAC
 		TSharedPtr<SWindow> TopLevelWindow = WidgetsUnderPointer.TopLevelWindow;
-		if ( bNeedToActivateWindow || (TopLevelWindow.IsValid() && TopLevelWindow->GetNativeWindow()->GetOSWindowHandle() != ActiveWindow) )
+		if ( bNeedToActivateWindow || ( TopLevelWindow.IsValid() && TopLevelWindow->GetNativeWindow()->GetOSWindowHandle() != ActiveWindow ) )
 		{
 			// Clicking on a context menu should not activate anything
 			// @todo: This needs to be updated when we have window type in SWindow and we no longer have to guess if WidgetsUnderCursor.TopLevelWindow is a menu
@@ -4489,7 +4511,7 @@ FReply FSlateApplication::RoutePointerDownEvent(FWidgetPath& WidgetsUnderPointer
 			{
 				MouseCaptorHelper Captor = MouseCaptor;
 				FPlatformMisc::ActivateApplication();
-				if (TopLevelWindow.IsValid())
+				if ( TopLevelWindow.IsValid() )
 				{
 					TopLevelWindow->BringToFront(true);
 				}
diff --git a/Engine/Source/Runtime/Slate/Private/Framework/Docking/TabManager.cpp b/Engine/Source/Runtime/Slate/Private/Framework/Docking/TabManager.cpp
--- a/Engine/Source/Runtime/Slate/Private/Framework/Docking/TabManager.cpp
+++ b/Engine/Source/Runtime/Slate/Private/Framework/Docking/TabManager.cpp
@@ -966,7 +966,6 @@ TSharedRef<SDockTab> FTabManager::InvokeTab( const FTabId& TabId )
 
 TSharedRef<SDockTab> FTabManager::InvokeTab_Internal( const FTabId& TabId )
 {
-
 	// Tab Spawning Rules:
 	// 
 	//     * Find live instance --yes--> use it.
@@ -1078,34 +1077,38 @@ void FTabManager::InsertDocumentTab(FName PlaceholderId, const FSearchPreference
 	
 	if (bTabNotManaged)
 	{
-		TSharedPtr<SDockTab> LiveTab = SearchPreference.Search(*this, PlaceholderId, UnmanagedTab);
+		OpenUnmanagedTab(PlaceholderId, SearchPreference, UnmanagedTab);
+	}
+
+	DrawAttention(UnmanagedTab);
+	if (bPlaySpawnAnim)
+	{
+		UnmanagedTab->PlaySpawnAnim();
+	}
+}
+
+void FTabManager::OpenUnmanagedTab(FName PlaceholderId, const FSearchPreference& SearchPreference, const TSharedRef<SDockTab>& UnmanagedTab)
+{
+	TSharedPtr<SDockTab> LiveTab = SearchPreference.Search(*this, PlaceholderId, UnmanagedTab);
 		
-		if (LiveTab.IsValid())
+	if (LiveTab.IsValid())
+	{
+		LiveTab->GetParent()->GetParentDockTabStack()->OpenTab( UnmanagedTab );
+	}
+	else
+	{
+		TSharedPtr<SDockingTabStack> StackToSpawnIn = FindPotentiallyClosedTab( PlaceholderId );
+		if( StackToSpawnIn.IsValid() == false )
 		{
+			UE_LOG(LogTabManager, Warning, TEXT("Unable to insert tab '%s'."), *(PlaceholderId.ToString()));
+			LiveTab = InvokeTab_Internal( FTabId( PlaceholderId ) );
 			LiveTab->GetParent()->GetParentDockTabStack()->OpenTab( UnmanagedTab );
 		}
 		else
 		{
-			TSharedPtr<SDockingTabStack> StackToSpawnIn = FindPotentiallyClosedTab( PlaceholderId );
-			if( StackToSpawnIn.IsValid() == false )
-			{
-				UE_LOG(LogTabManager, Warning, TEXT("Unable to insert tab '%s'."), *(PlaceholderId.ToString()));
-				LiveTab = InvokeTab_Internal( FTabId( PlaceholderId ) );
-				LiveTab->GetParent()->GetParentDockTabStack()->OpenTab( UnmanagedTab );
-			}
-			else
-			{
-				StackToSpawnIn->OpenTab( UnmanagedTab );
-			}
+			StackToSpawnIn->OpenTab( UnmanagedTab );
 		}
 	}
-
-	DrawAttention(UnmanagedTab);
-	if (bPlaySpawnAnim)
-	{
-		UnmanagedTab->PlaySpawnAnim();
-	}
-	
 }
 
 FTabManager::FTabManager( const TSharedPtr<SDockTab>& InOwnerTab, const TSharedRef<FTabSpawner> & InNomadTabSpawner )
@@ -1989,4 +1992,40 @@ void FGlobalTabmanager::UpdateStats()
 	AllAreasWindowMaxCount = FMath::Max(AllAreasWindowMaxCount, ParentWindows.Num());
 }
 
+void FGlobalTabmanager::OpenUnmanagedTab(FName PlaceholderId, const FSearchPreference& SearchPreference, const TSharedRef<SDockTab>& UnmanagedTab)
+{
+	if ( ProxyTabManager.IsValid() )
+	{
+		ProxyTabManager->OpenUnmanagedTab(PlaceholderId, SearchPreference, UnmanagedTab);
+	}
+	else
+	{
+		FTabManager::OpenUnmanagedTab(PlaceholderId, SearchPreference, UnmanagedTab);
+	}
+}
+
+void FGlobalTabmanager::SetProxyTabManager(TSharedPtr<FProxyTabmanager> InProxyTabManager)
+{
+	ProxyTabManager = InProxyTabManager;
+}
+
+void FProxyTabmanager::OpenUnmanagedTab(FName PlaceholderId, const FSearchPreference& SearchPreference, const TSharedRef<SDockTab>& UnmanagedTab)
+{
+	// No layout info about this tab found; start 
+	TSharedRef<FArea> NewAreaForTab = FTabManager::NewArea(FTabManager::FallbackWindowSize)
+		->Split
+		(
+		FTabManager::NewStack()
+		->AddTab(UnmanagedTab->GetLayoutIdentifier(), ETabState::OpenedTab)
+		);
+
+	TSharedRef<SDockingArea> DockingArea = RestoreArea(NewAreaForTab, ParentWindow.Pin());
+
+	const TSharedPtr<SDockTab> NewlyOpenedTab = DockingArea->GetAllChildTabs()[0];
+	check(NewlyOpenedTab.IsValid());
+		
+	NewlyOpenedTab->GetParent()->GetParentDockTabStack()->OpenTab(UnmanagedTab);
+}
+
+
 #undef LOCTEXT_NAMESPACE
diff --git a/Engine/Source/Runtime/Slate/Private/Framework/MultiBox/SMenuEntryBlock.cpp b/Engine/Source/Runtime/Slate/Private/Framework/MultiBox/SMenuEntryBlock.cpp
--- a/Engine/Source/Runtime/Slate/Private/Framework/MultiBox/SMenuEntryBlock.cpp
+++ b/Engine/Source/Runtime/Slate/Private/Framework/MultiBox/SMenuEntryBlock.cpp
@@ -1042,18 +1042,20 @@ void SMenuEntryBlock::OnClicked( bool bCheckBoxClicked )
 		if( PinnedMenuAnchor.IsValid() && PinnedMenuAnchor->ShouldOpenDueToClick() )
 		{
 			FWidgetPath WidgetPath;
-			FSlateApplication::Get().GeneratePathToWidgetChecked( PinnedMenuAnchor->AsShared(), WidgetPath );
-
-			// Don't process clicks that attempt to open sub-menus when the parent is queued for destruction.
-			if (!FSlateApplication::Get().IsWindowInDestroyQueue( WidgetPath.GetWindow() ))
+			FSlateApplication::Get().GeneratePathToWidgetUnchecked(PinnedMenuAnchor->AsShared(), WidgetPath);
+			if ( WidgetPath.IsValid() )
 			{
-				// Close other open pull-down menus from this menu bar
-				OwnerMultiBoxWidget.Pin()->CloseSummonedMenus();
+				// Don't process clicks that attempt to open sub-menus when the parent is queued for destruction.
+				if ( !FSlateApplication::Get().IsWindowInDestroyQueue(WidgetPath.GetWindow()) )
+				{
+					// Close other open pull-down menus from this menu bar
+					OwnerMultiBoxWidget.Pin()->CloseSummonedMenus();
 
-				PinnedMenuAnchor->SetIsOpen( true );
+					PinnedMenuAnchor->SetIsOpen(true);
 
-				// Also tell the multibox about this open pull-down menu, so it can be closed later if we need to
-				OwnerMultiBoxWidget.Pin()->SetSummonedMenu( PinnedMenuAnchor.ToSharedRef() );
+					// Also tell the multibox about this open pull-down menu, so it can be closed later if we need to
+					OwnerMultiBoxWidget.Pin()->SetSummonedMenu(PinnedMenuAnchor.ToSharedRef());
+				}
 			}
 		}
 	}
diff --git a/Engine/Source/Runtime/Slate/Public/Framework/Application/MenuStack.h b/Engine/Source/Runtime/Slate/Public/Framework/Application/MenuStack.h
--- a/Engine/Source/Runtime/Slate/Public/Framework/Application/MenuStack.h
+++ b/Engine/Source/Runtime/Slate/Public/Framework/Application/MenuStack.h
@@ -272,11 +272,9 @@ private:
 	void DismissInternal(int32 FirstStackIndexToRemove);
 
 	/**
-	 * Handles changes to the host window. Makes sure overlay slots are added/removed as appropriate.
-	 *
-	 * @param	InWindow	The new host window.
+	 * Handles changes to the owner path, potentially locating and establishing the new host. Makes sure overlay slots are added/removed as appropriate.
 	 */
-	void SetHostWindow(TSharedPtr<SWindow> InWindow);
+	void SetHostPath(const FWidgetPath& InOwnerPath);
 
 	/**
 	 * Callback method called by all menus when they are dismissed.  Handles changes to the stack state when a menu is destroyed.
@@ -393,6 +391,9 @@ private:
 	/** The parent window of the root menu in the stack. NOT the actual menu window if it's a CreateNewWindow */
 	TSharedPtr<SMenuPanel> HostWindowPopupPanel;
 
+	/** The popup layer that contains our HostWindowPopupPanel. */
+	TSharedPtr<FPopupLayer> HostPopupLayer;
+
 	/** The array of menus in the stack */
 	TArray<TSharedPtr<class FMenuBase>> Stack;
 
diff --git a/Engine/Source/Runtime/Slate/Public/Framework/Docking/TabManager.h b/Engine/Source/Runtime/Slate/Public/Framework/Docking/TabManager.h
--- a/Engine/Source/Runtime/Slate/Public/Framework/Docking/TabManager.h
+++ b/Engine/Source/Runtime/Slate/Public/Framework/Docking/TabManager.h
@@ -622,13 +622,13 @@ class SLATE_API FTabManager : public TSharedFromThis<FTabManager>
 		void InsertNewDocumentTab( FName PlaceholderId, const FSearchPreference& SearchPreference, const TSharedRef<SDockTab>& UnmanagedTab );
 		
 		/** Insert a new UnmanagedTab document tab next to an existing tab (closed or open) that has the PlaceholdId. */
-		void InsertNewDocumentTab( FName PlaceholderId, ESearchPreference::Type SearchPreference, const TSharedRef<SDockTab>& UnmanagedTab );
+		void InsertNewDocumentTab(FName PlaceholderId, ESearchPreference::Type SearchPreference, const TSharedRef<SDockTab>& UnmanagedTab);
 
 		/**
 		 * Much like InsertNewDocumentTab, but the UnmanagedTab is not seen by the user as newly-created.
 		 * e.g. Opening an restores multiple previously opened documents; these are not seen as new tabs.
 		 */
-		void RestoreDocumentTab( FName PlaceholderId, ESearchPreference::Type SearchPreference, const TSharedRef<SDockTab>& UnmanagedTab );
+		void RestoreDocumentTab(FName PlaceholderId, ESearchPreference::Type SearchPreference, const TSharedRef<SDockTab>& UnmanagedTab);
 
 		/**
 		 * Opens tab if it is closed at the last known location.  If it already exists, it will draw attention to the tab.
@@ -636,7 +636,7 @@ class SLATE_API FTabManager : public TSharedFromThis<FTabManager>
 		 * @param TabId The tab identifier.
 		 * @return The existing or newly spawned tab instance.
 		 */
-		TSharedRef<SDockTab> InvokeTab( const FTabId& TabId );
+		virtual TSharedRef<SDockTab> InvokeTab( const FTabId& TabId );
 
 		virtual ~FTabManager()
 		{
@@ -689,7 +689,9 @@ class SLATE_API FTabManager : public TSharedFromThis<FTabManager>
 
 	protected:
 		
-		void InsertDocumentTab(FName PlaceholderId, const FSearchPreference& SearchPreference, const TSharedRef<SDockTab>& UnmanagedTab, bool bPlaySpawnAnim);
+		void InsertDocumentTab( FName PlaceholderId, const FSearchPreference& SearchPreference, const TSharedRef<SDockTab>& UnmanagedTab, bool bPlaySpawnAnim );
+
+		virtual void OpenUnmanagedTab(FName PlaceholderId, const FSearchPreference& SearchPreference, const TSharedRef<SDockTab>& UnmanagedTab);
 			
 		void PopulateTabSpawnerMenu_Helper( FMenuBuilder& PopulateMe, struct FPopulateTabSpawnerMenu_Args Args );
 
@@ -811,7 +813,7 @@ class SLATE_API FTabManager : public TSharedFromThis<FTabManager>
 
 
 
-
+class FProxyTabmanager;
 
 
 class SLATE_API FGlobalTabmanager : public FTabManager
@@ -895,12 +897,16 @@ public:
 	/** Returns the highest number of parent windows that were open simultaneously during this session */
 	int32 GetMaximumWindowCount() const { return AllAreasWindowMaxCount; }
 
+	void SetProxyTabManager(TSharedPtr<FProxyTabmanager> InProxyTabManager);
+
 protected:
 	virtual void OnTabForegrounded( const TSharedPtr<SDockTab>& NewForegroundTab, const TSharedPtr<SDockTab>& BackgroundedTab ) override;
 	virtual void OnTabRelocated( const TSharedRef<SDockTab>& RelocatedTab, const TSharedPtr<SWindow>& NewOwnerWindow ) override;
 	virtual void OnTabClosing( const TSharedRef<SDockTab>& TabBeingClosed ) override;
 	virtual void UpdateStats() override;
 
+	virtual void OpenUnmanagedTab(FName PlaceholderId, const FSearchPreference& SearchPreference, const TSharedRef<SDockTab>& UnmanagedTab) override;
+
 public:
 	virtual void OnTabManagerClosing() override;
 
@@ -980,4 +986,23 @@ private:
 	/** Keeps track of the running-maximum number of unique parent windows in all dock areas and sub-managers during this session */
 	int32 AllAreasWindowMaxCount;
 
+	/**  */
+	TSharedPtr<FProxyTabmanager> ProxyTabManager;
+};
+
+
+class SLATE_API FProxyTabmanager : public FTabManager
+{
+public:
+
+	FProxyTabmanager(TSharedRef<SWindow> InParentWindow)
+		: FTabManager(TSharedPtr<SDockTab>(), MakeShareable(new FTabSpawner()))
+		, ParentWindow(InParentWindow)
+	{
+	}
+
+	virtual void OpenUnmanagedTab(FName PlaceholderId, const FSearchPreference& SearchPreference, const TSharedRef<SDockTab>& UnmanagedTab) override;
+
+private:
+	TWeakPtr<SWindow> ParentWindow;
 };
diff --git a/Engine/Source/Runtime/SlateCore/Private/Input/Reply.cpp b/Engine/Source/Runtime/SlateCore/Private/Input/Reply.cpp
new file mode 100644
--- /dev/null
+++ b/Engine/Source/Runtime/SlateCore/Private/Input/Reply.cpp
@@ -0,0 +1,23 @@
+// Copyright 1998-2016 Epic Games, Inc. All Rights Reserved.
+
+#include "SlateCorePrivatePCH.h"
+
+FReply& FReply::SetUserFocus(TSharedRef<SWidget> GiveMeFocus, EFocusCause ReasonFocusIsChanging, bool bInAllUsers)
+{
+	this->bSetUserFocus = true;
+	this->FocusRecipient = GiveMeFocus;
+	this->FocusChangeReason = ReasonFocusIsChanging;
+	this->bReleaseUserFocus = false;
+	this->bAllUsers = bInAllUsers;
+	return Me();
+}
+
+FReply& FReply::ClearUserFocus(EFocusCause ReasonFocusIsChanging, bool bInAllUsers)
+{
+	this->FocusRecipient = nullptr;
+	this->FocusChangeReason = ReasonFocusIsChanging;
+	this->bReleaseUserFocus = true;
+	this->bSetUserFocus = false;
+	this->bAllUsers = bInAllUsers;
+	return Me();
+}
\ No newline at end of file
diff --git a/Engine/Source/Runtime/SlateCore/Private/Rendering/DrawElements.cpp b/Engine/Source/Runtime/SlateCore/Private/Rendering/DrawElements.cpp
--- a/Engine/Source/Runtime/SlateCore/Private/Rendering/DrawElements.cpp
+++ b/Engine/Source/Runtime/SlateCore/Private/Rendering/DrawElements.cpp
@@ -526,12 +526,12 @@ void FSlateBatchData::Merge(FElementBatchMap& InLayerToElementBatches, uint32& V
 
 
 FSlateWindowElementList::FDeferredPaint::FDeferredPaint( const TSharedRef<const SWidget>& InWidgetToPaint, const FPaintArgs& InArgs, const FGeometry InAllottedGeometry, const FSlateRect InMyClippingRect, const FWidgetStyle& InWidgetStyle, bool InParentEnabled )
-: WidgetToPaintPtr( InWidgetToPaint )
-, Args( InArgs )
-, AllottedGeometry( InAllottedGeometry )
-, MyClippingRect( InMyClippingRect )
-, WidgetStyle( InWidgetStyle )
-, bParentEnabled( InParentEnabled )
+	: WidgetToPaintPtr( InWidgetToPaint )
+	, Args( InArgs )
+	, AllottedGeometry( InAllottedGeometry )
+	, MyClippingRect( InMyClippingRect )
+	, WidgetStyle( InWidgetStyle )
+	, bParentEnabled( InParentEnabled )
 {
 }
 
@@ -555,14 +555,34 @@ void FSlateWindowElementList::QueueDeferredPainting( const FDeferredPaint& InDef
 
 int32 FSlateWindowElementList::PaintDeferred(int32 LayerId)
 {
-	for ( int32 i = 0; i < DeferredPaintList.Num(); ++i )
+	bNeedsDeferredResolve = false;
+
+	int32 ResolveIndex = ResolveToDeferredIndex.Pop(false);
+
+	for ( int32 i = ResolveIndex; i < DeferredPaintList.Num(); ++i )
 	{
 		LayerId = DeferredPaintList[i]->ExecutePaint(LayerId, *this);
 	}
 
+	for ( int32 i = DeferredPaintList.Num() - 1; i >= ResolveIndex; --i )
+	{
+		DeferredPaintList.RemoveAt(i, 1, false);
+	}
+
 	return LayerId;
 }
 
+void FSlateWindowElementList::BeginDeferredGroup()
+{
+	ResolveToDeferredIndex.Add(DeferredPaintList.Num());
+}
+
+void FSlateWindowElementList::EndDeferredGroup()
+{
+	bNeedsDeferredResolve = true;
+}
+
+
 
 FSlateWindowElementList::FVolatilePaint::FVolatilePaint(const TSharedRef<const SWidget>& InWidgetToPaint, const FPaintArgs& InArgs, const FGeometry InAllottedGeometry, const FSlateRect InMyClippingRect, int32 InLayerId, const FWidgetStyle& InWidgetStyle, bool InParentEnabled)
 	: WidgetToPaintPtr(InWidgetToPaint)
diff --git a/Engine/Source/Runtime/SlateCore/Private/Widgets/SOverlay.cpp b/Engine/Source/Runtime/SlateCore/Private/Widgets/SOverlay.cpp
--- a/Engine/Source/Runtime/SlateCore/Private/Widgets/SOverlay.cpp
+++ b/Engine/Source/Runtime/SlateCore/Private/Widgets/SOverlay.cpp
@@ -176,7 +176,7 @@ int32 SOverlay::GetNumWidgets() const
 	return Children.Num();
 }
 
-void SOverlay::RemoveSlot( TSharedRef< SWidget > Widget )
+bool SOverlay::RemoveSlot( TSharedRef< SWidget > Widget )
 {
 	// Search and remove
 	for( int32 CurSlotIndex = 0; CurSlotIndex < Children.Num(); ++CurSlotIndex )
@@ -185,7 +185,9 @@ void SOverlay::RemoveSlot( TSharedRef< SWidget > Widget )
 		if( CurSlot.GetWidget() == Widget )
 		{
 			Children.RemoveAt( CurSlotIndex );
-			break;
+			return true;
 		}
 	}
+
+	return false;
 }
diff --git a/Engine/Source/Runtime/SlateCore/Private/Widgets/SWidget.cpp b/Engine/Source/Runtime/SlateCore/Private/Widgets/SWidget.cpp
--- a/Engine/Source/Runtime/SlateCore/Private/Widgets/SWidget.cpp
+++ b/Engine/Source/Runtime/SlateCore/Private/Widgets/SWidget.cpp
@@ -217,6 +217,11 @@ bool SWidget::OnVisualizeTooltip( const TSharedPtr<SWidget>& TooltipContent )
 	return false;
 }
 
+TSharedPtr<FPopupLayer> SWidget::OnVisualizePopup(const TSharedRef<SWidget>& PopupContent)
+{
+	return TSharedPtr<FPopupLayer>();
+}
+
 FReply SWidget::OnDragDetected( const FGeometry& MyGeometry, const FPointerEvent& MouseEvent )
 {
 	return FReply::Unhandled();
@@ -691,6 +696,11 @@ int32 SWidget::Paint(const FPaintArgs& Args, const FGeometry& AllottedGeometry,
 		}
 	}
 
+	if ( OutDrawElements.ShouldResolveDeferred() )
+	{
+		NewLayerID = OutDrawElements.PaintDeferred(NewLayerID);
+	}
+
 	return NewLayerID;
 }
 
diff --git a/Engine/Source/Runtime/SlateCore/Private/Widgets/SWindow.cpp b/Engine/Source/Runtime/SlateCore/Private/Widgets/SWindow.cpp
--- a/Engine/Source/Runtime/SlateCore/Private/Widgets/SWindow.cpp
+++ b/Engine/Source/Runtime/SlateCore/Private/Widgets/SWindow.cpp
@@ -12,169 +12,20 @@ namespace SWindowDefs
 	static const int32 CornerRadius = 6;
 }
 
-
-/**
- * An internal overlay used by Slate to support in-window pop ups and tooltips.
- * The overlay ignores DPI scaling when it does its own arrangement, but otherwise
- * passes all DPI scale values through.
- */
-class SPopupLayer : public SPanel
+FOverlayPopupLayer::FOverlayPopupLayer(const TSharedRef<SWidget>& HostWidget, const TSharedRef<SWidget>& PopupContent, TSharedPtr<SOverlay> Overlay)
+	: FPopupLayer(HostWidget, PopupContent)
+	, Overlay(Overlay)
 {
-public:
-	SLATE_BEGIN_ARGS( SPopupLayer )
-		{
-			_Visibility = EVisibility::SelfHitTestInvisible;
-		}
-
-		SLATE_SUPPORTS_SLOT( FPopupLayerSlot )
-
-	SLATE_END_ARGS()
-
-	SPopupLayer()
-	: Children()
-	{}
-
-	void Construct( const FArguments& InArgs, const TSharedRef<SWindow>& InWindow )
-	{
-		OwnerWindow = InWindow;
-
-		const int32 NumSlots = InArgs.Slots.Num();
-		for ( int32 SlotIndex = 0; SlotIndex < NumSlots; ++SlotIndex )
-		{
-			Children.Add( InArgs.Slots[SlotIndex] );
-		}
-	}
-
-	/** Make a new ListPanel::Slot  */
-	FPopupLayerSlot& Slot()
-	{
-		return *(new FPopupLayerSlot());
-	}
-	
-	/** Add a slot to the ListPanel */
-	FPopupLayerSlot& AddSlot(int32 InsertAtIndex = INDEX_NONE)
-	{
-		FPopupLayerSlot& NewSlot = *new FPopupLayerSlot();
-		if (InsertAtIndex == INDEX_NONE)
-		{
-			this->Children.Add( &NewSlot );
-		}
-		else
-		{
-			this->Children.Insert( &NewSlot, InsertAtIndex );
-		}
-	
-		return NewSlot;
-	}
-
-	void RemoveSlot(const TSharedRef<SWidget>& WidgetToRemove)
-	{
-		for( int32 CurSlotIndex = 0; CurSlotIndex < Children.Num(); ++CurSlotIndex )
-		{
-			const FPopupLayerSlot& CurSlot = Children[ CurSlotIndex ];
-			if( CurSlot.GetWidget() == WidgetToRemove )
-			{
-				Children.RemoveAt( CurSlotIndex );
-				return;
-			}
-		}
-	}
-
-private:
-
-	/**
-	 * Each child slot essentially tries to place their contents at a specified position on the screen
-	 * and scale as the widget initiating the popup, both of which are stored in the slot attributes.
-	 * The tricky part is that the scale we are given is the fully accumulated layout scale of the widget, which already incorporates 
-	 * the DPI Scale of the window. The DPI Scale is also applied to the overlay since it is part of the window,
-	 * so this scale needs to be factored out when determining the scale of the child geometry that will be created to hold the popup.
-	 * We also optionally adjust the window position to keep it within the client bounds of the top-level window. This must be done in screenspace.
-	 * This means some hairy transformation calculus goes on to ensure the computations are done in the proper space so scale is respected.
-	 * 
-	 * There are 3 transformational spaces involved, each clearly specified in the variable names:
-	 * Screen      - Basically desktop space. Contains desktop offset and DPI scale.
-	 * WindowLocal - local space of the SWindow containing this popup. Screenspace == Concat(WindowLocal, DPI Scale, Desktop Offset)
-	 * ChildLocal  - space of the child widget we want to display in the popup. The widget's LayoutTransform takes us from ChildLocal to WindowLocal space.
-	 */
-	virtual void OnArrangeChildren( const FGeometry& AllottedGeometry, FArrangedChildren& ArrangedChildren ) const override
-	{
-		// skip all this work if there are no children to arrange.
-		if (Children.Num() == 0) return;
-
-		// create a transform from screen to local space.
-		// This assumes that the PopupLayer is part of an Overlay that takes up the entire window space.
-		// We should technically be using the AllottedGeometry to transform from AbsoluteToLocal space just in case it has an additional scale on it.
-		// But we can't because the absolute space of the geometry is sometimes given in desktop space (picking, ticking) 
-		// and sometimes in window space (painting), and we can't necessarily tell by inspection so we have to just make an assumption here.
-		FSlateLayoutTransform ScreenToWindowLocal = (ensure(OwnerWindow.IsValid())) ? Inverse(OwnerWindow.Pin()->GetLocalToScreenTransform()) : FSlateLayoutTransform();
-		
-		for ( int32 ChildIndex = 0; ChildIndex < Children.Num(); ++ChildIndex )
-		{
-			const FPopupLayerSlot& CurChild = Children[ChildIndex];
-			const EVisibility ChildVisibility = CurChild.GetWidget()->GetVisibility();
-			if ( ArrangedChildren.Accepts(ChildVisibility) )
-			{
-				// This scale+translate forms the ChildLocal to Screenspace transform.
-				// The translation may be adjusted based on clamping, but the scale is accurate, 
-				// so we can transform vectors into screenspace using the scale alone.
-				const float ChildLocalToScreenScale = CurChild.Scale_Attribute.Get();
-				FVector2D ChildLocalToScreenOffset = CurChild.DesktopPosition_Attribute.Get();
-				// The size of the child is either the desired size of the widget (computed in the child's local space) or the size override (specified in screen space)
-				const FVector2D ChildSizeChildLocal = CurChild.GetWidget()->GetDesiredSize();
-				// Convert the desired size to screen space. Here is were we convert a vector to screenspace
-				// before we have the final position in screenspace (which would be needed to transform a point).
-				FVector2D ChildSizeScreenspace = TransformVector(ChildLocalToScreenScale, ChildSizeChildLocal);
-				// But then allow each size dimension to be overridden by the slot, which specifies the overrides in screen space.
-				ChildSizeScreenspace = FVector2D(
-						CurChild.WidthOverride_Attribute.IsSet() ? CurChild.WidthOverride_Attribute.Get() : ChildSizeScreenspace.X,
-						CurChild.HeightOverride_Attribute.IsSet() ? CurChild.HeightOverride_Attribute.Get() : ChildSizeScreenspace.Y);
-				
-				// If clamping, move the screen space position to ensure the screen space size stays within the client rect of the top level window.
-				if(CurChild.Clamp_Attribute.Get())
-				{
-					const FSlateRect WindowClientRectScreenspace = (ensure(OwnerWindow.IsValid())) ? OwnerWindow.Pin()->GetClientRectInScreen() : FSlateRect();
-					const FVector2D ClampBufferScreenspace = CurChild.ClampBuffer_Attribute.Get();
-					const FSlateRect ClampedWindowClientRectScreenspace = WindowClientRectScreenspace.InsetBy(FMargin(ClampBufferScreenspace.X, ClampBufferScreenspace.Y));
-					// Find how much our child wants to extend beyond our client space and subtract that amount, but don't push it past the client edge.
-					ChildLocalToScreenOffset.X = FMath::Max(WindowClientRectScreenspace.Left, ChildLocalToScreenOffset.X - FMath::Max(0.0f, (ChildLocalToScreenOffset.X  + ChildSizeScreenspace.X ) - ClampedWindowClientRectScreenspace.Right));
-					ChildLocalToScreenOffset.Y = FMath::Max(WindowClientRectScreenspace.Top, ChildLocalToScreenOffset.Y - FMath::Max(0.0f, (ChildLocalToScreenOffset.Y  + ChildSizeScreenspace.Y ) - ClampedWindowClientRectScreenspace.Bottom));
-				}
-
-				// We now have the final position, so construct the transform from ChildLocal to Screenspace
-				const FSlateLayoutTransform ChildLocalToScreen(ChildLocalToScreenScale, ChildLocalToScreenOffset);
-				// Using this we can compute the transform from ChildLocal to WindowLocal, which is effectively the LayoutTransform of the child widget.
-				const FSlateLayoutTransform ChildLocalToWindowLocal = Concatenate(ChildLocalToScreen, ScreenToWindowLocal);
-				// The ChildSize needs to be given in ChildLocal space when constructing a geometry.
-				const FVector2D ChildSizeLocalspace = TransformVector(Inverse(ChildLocalToScreen), ChildSizeScreenspace);
-
-				// The position is explicitly in desktop pixels.
-				// The size and DPI scale come from the widget that is using
-				// this overlay to "punch" through the UI.
-				ArrangedChildren.AddWidget(AllottedGeometry.MakeChild(CurChild.GetWidget(), ChildSizeLocalspace, ChildLocalToWindowLocal));
-			}
-		}
-	}
-
-	virtual FVector2D ComputeDesiredSize(float) const override
-	{
-		return FVector2D(100,100);
-	}
-
-	/**
-	 * All widgets must provide a way to access their children in a layout-agnostic way.
-	 * Panels store their children in Slots, which creates a dilemma. Most panels
-	 * can store their children in a TPanelChildren<Slot>, where the Slot class
-	 * provides layout information about the child it stores. In that case
-	 * GetChildren should simply return the TPanelChildren<Slot>. See StackPanel for an example.
-	 */
-	virtual FChildren* GetChildren() override
-	{
-		return &Children;
-	}
+	Overlay->AddSlot()
+	[
+		PopupContent
+	];
+}
 
-	TPanelChildren<FPopupLayerSlot> Children;
-	TWeakPtr<SWindow> OwnerWindow;
-};
+void FOverlayPopupLayer::Remove()
+{
+	Overlay->RemoveSlot(GetContent());
+}
 
 FVector2D SWindow::GetWindowSizeFromClientSize(FVector2D InClientSize)
 {
@@ -500,12 +351,6 @@ void SWindow::ConstructWindowInternals()
 				]
 			]
 
-			// pop-up layer
-			+ SOverlay::Slot()
-			[
-				SAssignNew(PopupLayer, SPopupLayer, SharedThis(this))
-			]
-
 			// window outline
 			+ SOverlay::Slot()
 			[
@@ -517,7 +362,7 @@ void SWindow::ConstructWindowInternals()
 			]
 		];
 	}
-	else if( bHasOSWindowBorder )
+	else if ( bHasOSWindowBorder || bVirtualWindow )
 	{
 		this->ChildSlot
 		[
@@ -526,10 +371,6 @@ void SWindow::ConstructWindowInternals()
 			[
 				MainWindowArea
 			]
-			+ SOverlay::Slot()
-			[
-				SAssignNew(PopupLayer, SPopupLayer, SharedThis(this))
-			]
 		];
 	}
 }
@@ -1044,17 +885,14 @@ void SWindow::RemoveOverlaySlot( const TSharedRef<SWidget>& InContent )
 	}
 }
 
-/** Return a new slot in the popup layer. Assumes that the window has a popup layer. */
-FPopupLayerSlot& SWindow::AddPopupLayerSlot()
+TSharedPtr<FPopupLayer> SWindow::OnVisualizePopup(const TSharedRef<SWidget>& PopupContent)
 {
-	ensure( PopupLayer.IsValid() );
-	return PopupLayer->AddSlot();
-}
+	if ( WindowOverlay.IsValid() )
+	{
+		return MakeShareable(new FOverlayPopupLayer(SharedThis(this), PopupContent, WindowOverlay));
+	}
 
-/** Counterpart to AddPopupLayerSlot */
-void SWindow::RemovePopupLayerSlot( const TSharedRef<SWidget>& WidgetToRemove )
-{
-	PopupLayer->RemoveSlot( WidgetToRemove );
+	return TSharedPtr<FPopupLayer>();
 }
 
 /** @return should this window show up in the taskbar */
@@ -1193,12 +1031,22 @@ bool SWindow::IsVisible() const
 
 bool SWindow::IsWindowMaximized() const
 {
-	return NativeWindow->IsMaximized();
+	if ( NativeWindow.IsValid() )
+	{
+		return NativeWindow->IsMaximized();
+	}
+
+	return false;
 }
 
 bool SWindow::IsWindowMinimized() const
 {
-	return NativeWindow->IsMinimized();
+	if ( NativeWindow.IsValid() )
+	{
+		return NativeWindow->IsMinimized();
+	}
+
+	return false;
 }
 
 
@@ -1690,10 +1538,18 @@ SWindow::SWindow()
 int32 SWindow::PaintWindow( const FPaintArgs& Args, const FGeometry& AllottedGeometry, const FSlateRect& MyClippingRect, FSlateWindowElementList& OutDrawElements, int32 LayerId, const FWidgetStyle& InWidgetStyle, bool bParentEnabled ) const
 {
 	LayerId = Paint(Args, AllottedGeometry, MyClippingRect, OutDrawElements, LayerId, InWidgetStyle, bParentEnabled);
-	LayerId = OutDrawElements.PaintDeferred( LayerId );
+	//LayerId = OutDrawElements.PaintDeferred( LayerId );
 	return LayerId;
 }
 
+int32 SWindow::OnPaint(const FPaintArgs& Args, const FGeometry& AllottedGeometry, const FSlateRect& MyClippingRect, FSlateWindowElementList& OutDrawElements, int32 LayerId, const FWidgetStyle& InWidgetStyle, bool bParentEnabled) const
+{
+	OutDrawElements.BeginDeferredGroup();
+	int32 MaxLayer = SCompoundWidget::OnPaint(Args, AllottedGeometry, MyClippingRect, OutDrawElements, LayerId, InWidgetStyle, bParentEnabled);
+	OutDrawElements.EndDeferredGroup();
+
+	return MaxLayer;
+}
 
 FOptionalSize SWindow::GetTitleBarSize() const
 {
diff --git a/Engine/Source/Runtime/SlateCore/Public/Input/Reply.h b/Engine/Source/Runtime/SlateCore/Public/Input/Reply.h
--- a/Engine/Source/Runtime/SlateCore/Public/Input/Reply.h
+++ b/Engine/Source/Runtime/SlateCore/Public/Input/Reply.h
@@ -13,7 +13,7 @@ class SWidget;
  * For example, a widget may handle an OnMouseDown event by asking the system to give mouse capture to a specific Widget.
  * To do this, return FReply::CaptureMouse( NewMouseCapture ).
  */
-class FReply : public TReplyBase<FReply>
+class SLATECORE_API FReply : public TReplyBase<FReply>
 {
 public:
 		
@@ -45,15 +45,7 @@ public:
 	}
 
 	/** An event should return FReply::Handled().SetUserFocus( SomeWidget ) as a means of asking the system to set users focus to the provided widget*/
-	FReply& SetUserFocus(TSharedRef<SWidget> GiveMeFocus, EFocusCause ReasonFocusIsChanging = EFocusCause::SetDirectly, bool bInAllUsers = false)
-	{
-		this->bSetUserFocus = true;
-		this->FocusRecipient = GiveMeFocus;
-		this->FocusChangeReason = ReasonFocusIsChanging;
-		this->bReleaseUserFocus = false;
-		this->bAllUsers = bInAllUsers;
-		return Me();
-	}
+	FReply& SetUserFocus(TSharedRef<SWidget> GiveMeFocus, EFocusCause ReasonFocusIsChanging = EFocusCause::SetDirectly, bool bInAllUsers = false);
 	
 	DEPRECATED(4.6, "FReply::CaptureJoystick() is deprecated, use FReply::SetUserFocus() instead.")
 	FReply& CaptureJoystick(TSharedRef<SWidget> InJoystickCaptor, bool bInAllJoysticks = false)
@@ -74,15 +66,7 @@ public:
 	}
 
 	/** An event should return a FReply::Handled().ClearUserFocus() to ask the system to clear user focus*/
-	FReply& ClearUserFocus(EFocusCause ReasonFocusIsChanging, bool bInAllUsers = false)
-	{
-		this->FocusRecipient = nullptr;
-		this->FocusChangeReason = ReasonFocusIsChanging;
-		this->bReleaseUserFocus = true;
-		this->bSetUserFocus = false;
-		this->bAllUsers = bInAllUsers;
-		return Me();
-	}
+	FReply& ClearUserFocus(EFocusCause ReasonFocusIsChanging, bool bInAllUsers = false);
 
 	DEPRECATED(4.6, "FReply::ReleaseJoystickCapture() is deprecated, use FReply::ClearUserFocus() instead.")
 	FReply& ReleaseJoystickCapture(bool bInAllJoysticks = false)
diff --git a/Engine/Source/Runtime/SlateCore/Public/Layout/WidgetPath.h b/Engine/Source/Runtime/SlateCore/Public/Layout/WidgetPath.h
--- a/Engine/Source/Runtime/SlateCore/Public/Layout/WidgetPath.h
+++ b/Engine/Source/Runtime/SlateCore/Public/Layout/WidgetPath.h
@@ -71,6 +71,9 @@ public:
 	/** A valid path has at least one widget in it */
 	bool IsValid() const;
 	
+	/**
+	 * Builds a string representation of the widget path.
+	 */
 	FString ToString() const;
 
 	/**
@@ -138,11 +141,16 @@ public:
 	 */
 	bool MoveFocus(int32 PathLevel, EUINavigation NavigationType);
 
+public:
 
+	/** The widgets that make up the widget path, the first item is the root widget, the end is the widget this path was built for. */
 	FArrangedChildren Widgets;
+
+	/** The top level window of this widget path. */
 	TSharedPtr< SWindow > TopLevelWindow;
-	TArray< TSharedPtr<FVirtualPointerPosition> > VirtualPointerPositions;
 
+	/** The virtual representation of the mouse at each level in the widget path.  Due to 3D widgets, the space you transition to can be completely arbitrary as you traverse the tree. */
+	TArray< TSharedPtr<FVirtualPointerPosition> > VirtualPointerPositions;
 
 private:
 
diff --git a/Engine/Source/Runtime/SlateCore/Public/Rendering/DrawElements.h b/Engine/Source/Runtime/SlateCore/Public/Rendering/DrawElements.h
--- a/Engine/Source/Runtime/SlateCore/Public/Rendering/DrawElements.h
+++ b/Engine/Source/Runtime/SlateCore/Public/Rendering/DrawElements.h
@@ -1043,6 +1043,8 @@ public:
 	 */
 	explicit FSlateWindowElementList( TSharedPtr<SWindow> InWindow = TSharedPtr<SWindow>() )
 		: TopLevelWindow( InWindow )
+		, bNeedsDeferredResolve( false )
+		, ResolveToDeferredIndex()
 		, MemManager(0)
 	{
 		DrawStack.Push(&RootDrawLayer);
@@ -1149,6 +1151,11 @@ public:
 
 	int32 PaintDeferred(int32 LayerId);
 
+	bool ShouldResolveDeferred() const { return bNeedsDeferredResolve; }
+
+	void BeginDeferredGroup();
+	void EndDeferredGroup();
+
 	struct FVolatilePaint
 	{
 	public:
@@ -1269,6 +1276,9 @@ private:
 	 */
 	TArray< TSharedPtr<FDeferredPaint> > DeferredPaintList;
 
+	bool bNeedsDeferredResolve;
+	TArray<int32> ResolveToDeferredIndex;
+
 	/** The widgets be cached for a later paint pass when the invalidation host paints. */
 	TArray< TSharedPtr<FVolatilePaint> > VolatilePaintList;
 
diff --git a/Engine/Source/Runtime/SlateCore/Public/Widgets/SOverlay.h b/Engine/Source/Runtime/SlateCore/Public/Widgets/SOverlay.h
--- a/Engine/Source/Runtime/SlateCore/Public/Widgets/SOverlay.h
+++ b/Engine/Source/Runtime/SlateCore/Public/Widgets/SOverlay.h
@@ -110,7 +110,7 @@ public:
 	 *
 	 * @param	Widget	The widget content to remove
 	 */
-	void RemoveSlot( TSharedRef< SWidget > Widget );
+	bool RemoveSlot( TSharedRef< SWidget > Widget );
 
 	/** Adds a slot at the specified location (ignores Z-order) */
 	FOverlaySlot& AddSlot(int32 ZOrder=INDEX_NONE);
diff --git a/Engine/Source/Runtime/SlateCore/Public/Widgets/SWidget.h b/Engine/Source/Runtime/SlateCore/Public/Widgets/SWidget.h
--- a/Engine/Source/Runtime/SlateCore/Public/Widgets/SWidget.h
+++ b/Engine/Source/Runtime/SlateCore/Public/Widgets/SWidget.h
@@ -99,6 +99,32 @@ public:
 };
 
 
+/**
+ * An FPopupLayer hosts the pop-up content which could be anything you want to appear on top of a widget.
+ * The widget must understand how to host pop-ups to make use of this.
+ */
+class FPopupLayer : public TSharedFromThis<FPopupLayer>
+{
+public:
+	FPopupLayer(const TSharedRef<SWidget>& HostWidget, const TSharedRef<SWidget>& PopupContent)
+		: HostWidget(HostWidget)
+		, PopupContent(PopupContent)
+	{
+	}
+
+	virtual ~FPopupLayer() { }
+	
+	virtual TSharedRef<SWidget> GetHost() { return HostWidget; }
+	virtual TSharedRef<SWidget> GetContent() { return PopupContent; }
+
+	virtual void Remove() = 0;
+
+private:
+	TSharedRef<SWidget> HostWidget;
+	TSharedRef<SWidget> PopupContent;
+};
+
+
 class IToolTip;
 
 /**
@@ -378,6 +404,15 @@ public:
 	 */
 	virtual bool OnVisualizeTooltip(const TSharedPtr<SWidget>& TooltipContent);
 
+	/**
+	 * Visualize a new pop-up if possible.  If it's not possible for this widget to host the pop-up
+	 * content you'll get back an invalid pointer to the layer.
+	 * 
+	 * @param PopupContent The widget to try and host overlaid on top of the widget.
+	 * @return a valid FPopupLayer if this widget supported hosting it.
+	 */
+	virtual TSharedPtr<FPopupLayer> OnVisualizePopup(const TSharedRef<SWidget>& PopupContent);
+
 	/**
 	 * Called when Slate detects that a widget started to be dragged.
 	 * Usage:
diff --git a/Engine/Source/Runtime/SlateCore/Public/Widgets/SWindow.h b/Engine/Source/Runtime/SlateCore/Public/Widgets/SWindow.h
--- a/Engine/Source/Runtime/SlateCore/Public/Widgets/SWindow.h
+++ b/Engine/Source/Runtime/SlateCore/Public/Widgets/SWindow.h
@@ -70,6 +70,20 @@ struct FWindowTransparency
 	EWindowTransparency Value;
 };
 
+/**
+ * Simple overlay layer to allow content to be laid out on a Window or similar widget.
+ */
+class FOverlayPopupLayer : public FPopupLayer
+{
+public:
+	FOverlayPopupLayer(const TSharedRef<SWidget>& HostWidget, const TSharedRef<SWidget>& PopupContent, TSharedPtr<SOverlay> Overlay);
+
+	virtual void Remove() override;
+
+private:
+	TSharedPtr<SOverlay> Overlay;
+};
+
 /**
  * SWindow is a platform-agnostic representation of a top-level window.
  */
@@ -414,11 +428,7 @@ public:
 	 */
 	void RemoveOverlaySlot( const TSharedRef<SWidget>& InContent );
 
-	/** Return a new slot in the popup layer. Assumes that the window has a popup layer. */
-	struct FPopupLayerSlot& AddPopupLayerSlot();
-
-	/** Counterpart to AddPopupLayerSlot */
-	void RemovePopupLayerSlot( const TSharedRef<SWidget>& WidgetToRemove );
+	virtual TSharedPtr<FPopupLayer> OnVisualizePopup(const TSharedRef<SWidget>& PopupContent) override;
 
 	/**
 	 * Sets a widget to use as a full window overlay, or clears an existing widget if set.  When set, this widget will be drawn on top of all other window content.
@@ -643,6 +653,8 @@ private:
 	virtual FReply OnMouseButtonUp( const FGeometry& MyGeometry, const FPointerEvent& MouseEvent ) override;
 	virtual FReply OnMouseMove( const FGeometry& MyGeometry, const FPointerEvent& MouseEvent ) override;
 
+	virtual int32 OnPaint(const FPaintArgs& Args, const FGeometry& AllottedGeometry, const FSlateRect& MyClippingRect, FSlateWindowElementList& OutDrawElements, int32 LayerId, const FWidgetStyle& InWidgetStyle, bool bParentEnabled) const override;
+
 	/** The window's desired size takes into account the ratio between the slate units and the pixel size */
 	virtual FVector2D ComputeDesiredSize(float) const override;
 
@@ -816,6 +828,9 @@ protected:
 	/** True if this window displays the os window border instead of drawing one in slate */
 	bool bHasOSWindowBorder : 1;
 
+	/** True if this window is virtual and not directly rendered by slate application or the OS. */
+	bool bVirtualWindow : 1;
+
 	/** True if this window displays an enabled close button on the toolbar area */
 	bool bHasCloseButton : 1;
 
@@ -898,7 +913,7 @@ protected:
 	const FWindowStyle* Style;
 	const FSlateBrush* WindowBackground;
 
-private:
+protected:
 
 	/** Min and Max values for Width and Height; all optional. */
 	FWindowSizeLimits SizeLimits;
@@ -968,12 +983,12 @@ private:
 	// The margin around the edges of the window that will be detected as places the user can grab to resize the window. 
 	FMargin UserResizeBorder;
 
-private:
+protected:
 	
-	virtual
-
 	void ConstructWindowInternals();
 
+private:
+
 	/**
 	 * @return EVisibility::Visible if we are showing this viewports content.  EVisibility::Hidden otherwise (we hide the content during full screen overlays)
 	 */
@@ -992,80 +1007,6 @@ private:
 };
 
 
-/**
- * Popups, tooltips, drag and drop decorators all can be executed without creating a new window.
- * This slot along with the SWindow::AddPopupLayerSlot() API enabled it.
- */
-struct FPopupLayerSlot : public TSlotBase<FPopupLayerSlot>
-{
-public:
-	FPopupLayerSlot()
-	: TSlotBase<FPopupLayerSlot>()
-	, DesktopPosition_Attribute(FVector2D::ZeroVector)
-	, WidthOverride_Attribute()
-	, HeightOverride_Attribute()
-	, Scale_Attribute(1.0f)
-	, Clamp_Attribute(false)
-	, ClampBuffer_Attribute(FVector2D::ZeroVector)
-	{}
-
-	/** Pixel position in desktop space */
-	FPopupLayerSlot& DesktopPosition( const TAttribute<FVector2D>& InDesktopPosition )
-	{
-		DesktopPosition_Attribute = InDesktopPosition;
-		return *this;
-	}
-
-	/** Width override in pixels */
-	FPopupLayerSlot& WidthOverride( const TAttribute<float>& InWidthOverride )
-	{
-		WidthOverride_Attribute = InWidthOverride;
-		return *this;
-	}
-
-	/** Width override in pixels */
-	FPopupLayerSlot& HeightOverride( const TAttribute<float>& InHeightOverride )
-	{
-		HeightOverride_Attribute = InHeightOverride;
-		return *this;
-	}
-
-	/** DPI scaling to be applied to the contents of this slot */
-	FPopupLayerSlot& Scale( const TAttribute<float>& InScale )
-	{
-		Scale_Attribute = InScale;
-		return *this;
-	}
-
-	/** Should this slot be kept within the parent window */
-	FPopupLayerSlot& ClampToWindow( const TAttribute<bool>& InClamp_Attribute)
-	{
-		Clamp_Attribute = InClamp_Attribute;
-		return *this;
-	}
-
-	/** If this slot is kept within the parent window, how far from the edges should we clamp it */
-	FPopupLayerSlot& ClampBuffer( const TAttribute<FVector2D>& InClampBuffer_Attribute )
-	{
-		ClampBuffer_Attribute = InClampBuffer_Attribute;
-		return *this;
-	}
-
-private:
-	/** SPopupLayer arranges FPopupLayerSlots, so it needs to know all about */
-	friend class SPopupLayer;
-	/** TPanelChildren need access to the Widget member */
-	friend class TPanelChildren<FPopupLayerSlot>;
-
-	TAttribute<FVector2D> DesktopPosition_Attribute;
-	TAttribute<float> WidthOverride_Attribute;
-	TAttribute<float> HeightOverride_Attribute;
-	TAttribute<float> Scale_Attribute;
-	TAttribute<bool> Clamp_Attribute;
-	TAttribute<FVector2D> ClampBuffer_Attribute;
-};
-
-
 #if WITH_EDITOR
 
 /**
diff --git a/Engine/Source/Runtime/SlateRHIRenderer/Private/Slate3DRenderer.cpp b/Engine/Source/Runtime/SlateRHIRenderer/Private/Slate3DRenderer.cpp
--- a/Engine/Source/Runtime/SlateRHIRenderer/Private/Slate3DRenderer.cpp
+++ b/Engine/Source/Runtime/SlateRHIRenderer/Private/Slate3DRenderer.cpp
@@ -56,26 +56,26 @@ void FSlate3DRenderer::DrawWindow_GameThread(FSlateDrawBuffer& DrawBuffer)
 
 	TArray<TSharedPtr<FSlateWindowElementList>>& WindowElementLists = DrawBuffer.GetWindowElementLists();
 
-	// Only one window element list for now.
-	check( WindowElementLists.Num() == 1);
-
-	FSlateWindowElementList& ElementList = *WindowElementLists[0];
-
-	TSharedPtr<SWindow> Window = ElementList.GetWindow();
-
-	if (Window.IsValid())
+	for ( int32 WindowIndex = 0; WindowIndex < WindowElementLists.Num(); WindowIndex++ )
 	{
-		const FVector2D WindowSize = Window->GetSizeInScreen();
-		if (WindowSize.X > 0 && WindowSize.Y > 0)
-		{
-			// Add all elements for this window to the element batcher
-			ElementBatcher->AddElements(ElementList);
+		FSlateWindowElementList& ElementList = *WindowElementLists[WindowIndex];
 
-			// Update the font cache with new text after elements are batched
-			FontCache->UpdateCache();
+		TSharedPtr<SWindow> Window = ElementList.GetWindow();
 
-			// All elements for this window have been batched and rendering data updated
-			ElementBatcher->ResetBatches();
+		if ( Window.IsValid() )
+		{
+			const FVector2D WindowSize = Window->GetSizeInScreen();
+			if ( WindowSize.X > 0 && WindowSize.Y > 0 )
+			{
+				// Add all elements for this window to the element batcher
+				ElementBatcher->AddElements(ElementList);
+
+				// Update the font cache with new text after elements are batched
+				FontCache->UpdateCache();
+
+				// All elements for this window have been batched and rendering data updated
+				ElementBatcher->ResetBatches();
+			}
 		}
 	}
 }
@@ -94,17 +94,9 @@ void FSlate3DRenderer::DrawWindowToTarget_RenderThread( FRHICommandListImmediate
 {
 	SCOPED_DRAW_EVENT( InRHICmdList, SlateRenderToTarget );
 
-	checkSlow( WindowDrawBuffer.GetWindowElementLists().Num() == 1 );
 	checkSlow( RenderTargetResource );
 
-	FSlateWindowElementList& WindowElementList = *WindowDrawBuffer.GetWindowElementLists()[ 0 ].Get();
-
-	FSlateBatchData& BatchData = WindowElementList.GetBatchData();
-	FElementBatchMap& RootBatchMap = WindowElementList.GetRootDrawLayer().GetElementBatchMap();
-
-	WindowElementList.PreDraw_ParallelThread();
-
-	BatchData.CreateRenderBatches(RootBatchMap);
+	TArray< TSharedPtr<FSlateWindowElementList> >& WindowsToDraw = WindowDrawBuffer.GetWindowElementLists();
 
 	// Enqueue a command to unlock the draw buffer after all windows have been drawn
 	ENQUEUE_UNIQUE_RENDER_COMMAND_ONEPARAMETER(SlateBeginDrawingWindowsCommand,
@@ -113,8 +105,6 @@ void FSlate3DRenderer::DrawWindowToTarget_RenderThread( FRHICommandListImmediate
 			Policy.BeginDrawingWindows();
 		});
 
-	RenderTargetPolicy->UpdateVertexAndIndexBuffers( InRHICmdList, BatchData );
-	
 	// Set render target and clear.
 	FTexture2DRHIRef RTResource = RenderTargetResource->GetTextureRHI();
 	FRHIRenderTargetView ColorRTV(RTResource);
@@ -125,22 +115,35 @@ void FSlate3DRenderer::DrawWindowToTarget_RenderThread( FRHICommandListImmediate
 	InRHICmdList.TransitionResource(EResourceTransitionAccess::EWritable, RTResource);
 	InRHICmdList.SetRenderTargetsAndClear(Info);
 
-	FMatrix ProjectionMatrix = FSlateRHIRenderer::CreateProjectionMatrix(RTResource->GetSizeX(), RTResource->GetSizeY());
-	if (BatchData.GetRenderBatches().Num() > 0)
+	for ( int32 WindowIndex = 0; WindowIndex < WindowsToDraw.Num(); WindowIndex++ )
 	{
-		FSlateBackBuffer BackBufferTarget(RenderTargetResource->GetTextureRHI(), FIntPoint(RTResource->GetSizeX(), RTResource->GetSizeY()));
-
-		// The scene renderer will handle it in this case
-		const bool bAllowSwitchVerticalAxis = false;
-
-		RenderTargetPolicy->DrawElements
-		(
-			InRHICmdList,
-			BackBufferTarget,
-			ProjectionMatrix,
-			BatchData.GetRenderBatches(),
-			bAllowSwitchVerticalAxis
-		);
+		FSlateWindowElementList& WindowElementList = *WindowsToDraw[WindowIndex].Get();
+
+		FSlateBatchData& BatchData = WindowElementList.GetBatchData();
+		FElementBatchMap& RootBatchMap = WindowElementList.GetRootDrawLayer().GetElementBatchMap();
+
+		WindowElementList.PreDraw_ParallelThread();
+
+		BatchData.CreateRenderBatches(RootBatchMap);
+
+		RenderTargetPolicy->UpdateVertexAndIndexBuffers(InRHICmdList, BatchData);
+
+		FMatrix ProjectionMatrix = FSlateRHIRenderer::CreateProjectionMatrix(RTResource->GetSizeX(), RTResource->GetSizeY());
+		if ( BatchData.GetRenderBatches().Num() > 0 )
+		{
+			FSlateBackBuffer BackBufferTarget(RenderTargetResource->GetTextureRHI(), FIntPoint(RTResource->GetSizeX(), RTResource->GetSizeY()));
+
+			// The scene renderer will handle it in this case
+			const bool bAllowSwitchVerticalAxis = false;
+
+			RenderTargetPolicy->DrawElements(
+				InRHICmdList,
+				BackBufferTarget,
+				ProjectionMatrix,
+				BatchData.GetRenderBatches(),
+				bAllowSwitchVerticalAxis
+			);
+		}
 	}
 
 	ENQUEUE_UNIQUE_RENDER_COMMAND_TWOPARAMETER(SlateEndDrawingWindowsCommand,
diff --git a/Engine/Source/Runtime/UMG/Private/Components/WidgetComponent.cpp b/Engine/Source/Runtime/UMG/Private/Components/WidgetComponent.cpp
--- a/Engine/Source/Runtime/UMG/Private/Components/WidgetComponent.cpp
+++ b/Engine/Source/Runtime/UMG/Private/Components/WidgetComponent.cpp
@@ -14,6 +14,8 @@
 #include "SGameLayerManager.h"
 #include "Slate/WidgetRenderer.h"
 #include "Slate/SWorldWidgetScreenLayer.h"
+#include "Widgets/LayerManager/STooltipPresenter.h"
+#include "Widgets/Layout/SPopup.h"
 
 DECLARE_CYCLE_STAT(TEXT("3DHitTesting"), STAT_Slate3DHitTesting, STATGROUP_Slate);
 
@@ -299,7 +301,7 @@ public:
 	{
 		SCOPE_CYCLE_COUNTER(STAT_Slate3DHitTesting);
 
-		if( World.IsValid() && ensure( World->IsGameWorld() ) )
+		if( World.IsValid() /*&& ensure( World->IsGameWorld() )*/ )
 		{
 			UWorld* SafeWorld = World.Get();
 			if ( SafeWorld )
@@ -583,22 +585,7 @@ void UWidgetComponent::OnRegister()
 			if ( GetWorld()->IsGameWorld() )
 			{
 				TSharedPtr<SViewport> GameViewportWidget = GEngine->GetGameViewportWidget();
-
-				if ( GameViewportWidget.IsValid() )
-				{
-					TSharedPtr<ICustomHitTestPath> CustomHitTestPath = GameViewportWidget->GetCustomHitTestPath();
-					if ( !CustomHitTestPath.IsValid() )
-					{
-						CustomHitTestPath = MakeShareable(new FWidget3DHitTester(GetWorld()));
-						GameViewportWidget->SetCustomHitTestPath(CustomHitTestPath);
-					}
-
-					TSharedPtr<FWidget3DHitTester> Widget3DHitTester = StaticCastSharedPtr<FWidget3DHitTester>(CustomHitTestPath);
-					if ( Widget3DHitTester->GetWorld() == GetWorld() )
-					{
-						Widget3DHitTester->RegisterWidgetComponent(this);
-					}
-				}
+				RegisterHitTesterWithViewport(GameViewportWidget);
 			}
 
 			if( !MaterialInstance )
@@ -629,18 +616,7 @@ void UWidgetComponent::OnUnregister()
 		TSharedPtr<SViewport> GameViewportWidget = GEngine->GetGameViewportWidget();
 		if( GameViewportWidget.IsValid() )
 		{
-			TSharedPtr<ICustomHitTestPath> CustomHitTestPath = GameViewportWidget->GetCustomHitTestPath();
-			if( CustomHitTestPath.IsValid() )
-			{
-				TSharedPtr<FWidget3DHitTester> WidgetHitTestPath = StaticCastSharedPtr<FWidget3DHitTester>(CustomHitTestPath);
-
-				WidgetHitTestPath->UnregisterWidgetComponent( this );
-
-				if ( WidgetHitTestPath->GetNumRegisteredComponents() == 0 )
-				{
-					GameViewportWidget->SetCustomHitTestPath( nullptr );
-				}
-			}
+			UnregisterHitTesterWithViewport(GameViewportWidget);
 		}
 	}
 
@@ -654,6 +630,46 @@ void UWidgetComponent::OnUnregister()
 	Super::OnUnregister();
 }
 
+void UWidgetComponent::RegisterHitTesterWithViewport(TSharedPtr<SViewport> ViewportWidget)
+{
+#if !UE_SERVER
+	if ( ViewportWidget.IsValid() )
+	{
+		TSharedPtr<ICustomHitTestPath> CustomHitTestPath = ViewportWidget->GetCustomHitTestPath();
+		if ( !CustomHitTestPath.IsValid() )
+		{
+			CustomHitTestPath = MakeShareable(new FWidget3DHitTester(GetWorld()));
+			ViewportWidget->SetCustomHitTestPath(CustomHitTestPath);
+		}
+
+		TSharedPtr<FWidget3DHitTester> Widget3DHitTester = StaticCastSharedPtr<FWidget3DHitTester>(CustomHitTestPath);
+		if ( Widget3DHitTester->GetWorld() == GetWorld() )
+		{
+			Widget3DHitTester->RegisterWidgetComponent(this);
+		}
+	}
+#endif
+
+}
+
+void UWidgetComponent::UnregisterHitTesterWithViewport(TSharedPtr<SViewport> ViewportWidget)
+{
+#if !UE_SERVER
+	TSharedPtr<ICustomHitTestPath> CustomHitTestPath = ViewportWidget->GetCustomHitTestPath();
+	if ( CustomHitTestPath.IsValid() )
+	{
+		TSharedPtr<FWidget3DHitTester> WidgetHitTestPath = StaticCastSharedPtr<FWidget3DHitTester>(CustomHitTestPath);
+
+		WidgetHitTestPath->UnregisterWidgetComponent(this);
+
+		if ( WidgetHitTestPath->GetNumRegisteredComponents() == 0 )
+		{
+			ViewportWidget->SetCustomHitTestPath(nullptr);
+		}
+	}
+#endif
+}
+
 void UWidgetComponent::DestroyComponent(bool bPromoteChildren/*= false*/)
 {
 	Super::DestroyComponent(bPromoteChildren);
@@ -692,14 +708,9 @@ void UWidgetComponent::TickComponent(float DeltaTime, enum ELevelTick TickType,
 
 	if ( Space != EWidgetSpace::Screen )
 	{
-		const float RenderTimeThreshold = .5f;
-		if ( IsVisible() )
+		if ( ShouldDrawWidget() )
 		{
-			// If we don't tick when off-screen, don't bother ticking if it hasn't been rendered recently
-			if ( TickWhenOffscreen || GetWorld()->TimeSince(LastRenderTime) <= RenderTimeThreshold )
-			{
-				DrawWidgetToRenderTarget(DeltaTime);
-			}
+			DrawWidgetToRenderTarget(DeltaTime);
 		}
 	}
 	else
@@ -757,6 +768,21 @@ void UWidgetComponent::TickComponent(float DeltaTime, enum ELevelTick TickType,
 #endif // !UE_SERVER
 }
 
+bool UWidgetComponent::ShouldDrawWidget() const
+{
+	const float RenderTimeThreshold = .5f;
+	if ( IsVisible() )
+	{
+		// If we don't tick when off-screen, don't bother ticking if it hasn't been rendered recently
+		if ( TickWhenOffscreen || GetWorld()->TimeSince(LastRenderTime) <= RenderTimeThreshold )
+		{
+			return true;
+		}
+	}
+
+	return false;
+}
+
 void UWidgetComponent::DrawWidgetToRenderTarget(float DeltaTime)
 {
 	if ( GUsingNullRHI )
@@ -925,9 +951,8 @@ void UWidgetComponent::InitWidget()
 		}
 		
 #if WITH_EDITOR
-		if ( Widget && !GetWorld()->IsGameWorld() )
+		if ( Widget && !GetWorld()->IsGameWorld() && !bEditTimeUsable )
 		{
-			//@todo vreditor: This prevents the UI from executing blueprint events in the editor. For VR editor we need this but still need to disable it for user placed in world components in 2D and VR editors
 			if( !GEnableVREditorHacks )
 			{
 				// Prevent native ticking of editor component previews
@@ -996,10 +1021,10 @@ void UWidgetComponent::UpdateWidget()
 	{
 		if ( Space != EWidgetSpace::Screen )
 		{
-			TSharedPtr<SWidget> NewWidget;
+			TSharedPtr<SWidget> NewSlateWidget;
 			if (Widget)
 			{
-				NewWidget = Widget->TakeWidget();
+				NewSlateWidget = Widget->TakeWidget();
 			}
 
 			if ( !SlateWindow.IsValid() )
@@ -1014,17 +1039,26 @@ void UWidgetComponent::UpdateWidget()
 
 			SlateWindow->Resize(DrawSize);
 
-			if (NewWidget.IsValid())
+			if ( NewSlateWidget.IsValid() )
 			{
-				SlateWindow->SetContent(NewWidget.ToSharedRef());
+				if ( NewSlateWidget != CurrentSlateWidget )
+				{
+					CurrentSlateWidget = NewSlateWidget;
+					SlateWindow->SetContent(NewSlateWidget.ToSharedRef());
+				}
 			}
 			else if( SlateWidget.IsValid() )
 			{
-				SlateWindow->SetContent( SlateWidget.ToSharedRef() );
+				if ( SlateWidget != CurrentSlateWidget )
+				{
+					CurrentSlateWidget = SlateWidget;
+					SlateWindow->SetContent(SlateWidget.ToSharedRef());
+				}
 			}
 			else
 			{
-				SlateWindow->SetContent(SNullWidget::NullWidget);
+				CurrentSlateWidget = SNullWidget::NullWidget;
+				SlateWindow->SetContent( SNullWidget::NullWidget );
 			}
 		}
 		else
diff --git a/Engine/Source/Runtime/UMG/Private/Slate/WidgetRenderer.cpp b/Engine/Source/Runtime/UMG/Private/Slate/WidgetRenderer.cpp
--- a/Engine/Source/Runtime/UMG/Private/Slate/WidgetRenderer.cpp
+++ b/Engine/Source/Runtime/UMG/Private/Slate/WidgetRenderer.cpp
@@ -9,16 +9,44 @@
 #include "ISlate3DRenderer.h"
 #endif // !UE_SERVER
 
+#include "Widgets/LayerManager/STooltipPresenter.h"
+#include "Widgets/Layout/SPopup.h"
+
 extern SLATECORE_API int32 bFoldTick;
 
 void SVirtualWindow::Construct(const FArguments& InArgs)
 {
 	bIsPopupWindow = true;
+	bVirtualWindow = true;
 	SetCachedSize(InArgs._Size);
 	SetNativeWindow(MakeShareable(new FGenericWindow()));
+
+	ConstructWindowInternals();
+
+	WindowOverlay->AddSlot()
+	[
+		SNew(SPopup)
+		[
+			SAssignNew(TooltipPresenter, STooltipPresenter)
+		]
+	];
+
 	SetContent(SNullWidget::NullWidget);
 }
 
+FPopupMethodReply SVirtualWindow::OnQueryPopupMethod() const
+{
+	return FPopupMethodReply::UseMethod(EPopupMethod::UseCurrentWindow)
+		.SetShouldThrottle(EShouldThrottle::No);
+}
+
+bool SVirtualWindow::OnVisualizeTooltip(const TSharedPtr<SWidget>& TooltipContent)
+{
+	TooltipPresenter->SetContent(TooltipContent.IsValid() ? TooltipContent.ToSharedRef() : SNullWidget::NullWidget);
+
+	return true;
+}
+
 FWidgetRenderer::FWidgetRenderer(bool bUseGammaCorrection)
 	: bPrepassNeeded(true)
 	, bUseGammaSpace(bUseGammaCorrection)
@@ -115,38 +143,15 @@ void FWidgetRenderer::DrawWindow(
 	{
 		return;
 	}
-	
-	if ( !bFoldTick )
-	{
-		Window->TickWidgetsRecursively(WindowGeometry, FApp::GetCurrentTime(), DeltaTime);
-	}
 
 	if ( bPrepassNeeded )
 	{
-		// Ticking can cause geometry changes.  Recompute
-		Window->SlatePrepass(WindowGeometry.Scale);
+		PrepassWindowAndChildren(Window, WindowGeometry.Scale);
 	}
-
-	// Prepare the test grid 
-	HitTestGrid->ClearGridForNewFrame(WindowClipRect);
-
-	// Get the free buffer & add our virtual window
+	
 	FSlateDrawBuffer& DrawBuffer = Renderer->GetDrawBuffer();
-	FSlateWindowElementList& WindowElementList = DrawBuffer.AddWindowElementList(Window);
 
-	int32 MaxLayerId = 0;
-	{
-		FPaintArgs PaintArgs(Window.Get(), HitTestGrid.Get(), FVector2D::ZeroVector, FApp::GetCurrentTime(), DeltaTime);
-
-		// Paint the window
-		MaxLayerId = Window->Paint(
-			PaintArgs,
-			WindowGeometry, WindowClipRect,
-			WindowElementList,
-			0,
-			FWidgetStyle(),
-			Window->IsEnabled());
-	}
+	DrawWindowAndChildren(RenderTarget, DrawBuffer, HitTestGrid, Window, WindowGeometry, WindowClipRect, DeltaTime);
 
 	Renderer->DrawWindow_GameThread(DrawBuffer);
 
@@ -171,3 +176,60 @@ void FWidgetRenderer::DrawWindow(
 		});
 #endif // !UE_SERVER
 }
+
+void FWidgetRenderer::PrepassWindowAndChildren(TSharedRef<SWindow> Window, float Scale)
+{
+	Window->SlatePrepass(Scale);
+
+	const TArray< TSharedRef<SWindow> >& WindowChildren = Window->GetChildWindows();
+	for ( int32 ChildIndex=0; ChildIndex < WindowChildren.Num(); ++ChildIndex )
+	{
+		const TSharedRef<SWindow>& ChildWindow = WindowChildren[ChildIndex];
+		PrepassWindowAndChildren(ChildWindow, Scale);
+	}
+}
+
+void FWidgetRenderer::DrawWindowAndChildren(
+	UTextureRenderTarget2D* RenderTarget,
+	FSlateDrawBuffer& DrawBuffer,
+	TSharedRef<FHittestGrid> HitTestGrid,
+	TSharedRef<SWindow> Window,
+	FGeometry WindowGeometry,
+	FSlateRect WindowClipRect,
+	float DeltaTime)
+{
+	if ( !bFoldTick )
+	{
+		Window->TickWidgetsRecursively(WindowGeometry, FApp::GetCurrentTime(), DeltaTime);
+	}
+
+	// Prepare the test grid 
+	HitTestGrid->ClearGridForNewFrame(WindowClipRect);
+
+	// Get the free buffer & add our virtual window
+	FSlateWindowElementList& WindowElementList = DrawBuffer.AddWindowElementList(Window);
+
+	int32 MaxLayerId = 0;
+	{
+		FPaintArgs PaintArgs(Window.Get(), HitTestGrid.Get(), FVector2D::ZeroVector, FApp::GetCurrentTime(), DeltaTime);
+
+		// Paint the window
+		MaxLayerId = Window->Paint(
+			PaintArgs,
+			WindowGeometry, WindowClipRect,
+			WindowElementList,
+			0,
+			FWidgetStyle(),
+			Window->IsEnabled());
+	}
+
+	// Draw the child windows
+	const TArray< TSharedRef<SWindow> >& WindowChildren = Window->GetChildWindows();
+	for ( int32 ChildIndex=0; ChildIndex < WindowChildren.Num(); ++ChildIndex )
+	{
+		const TSharedRef<SWindow>& ChildWindow = WindowChildren[ChildIndex];
+		FGeometry ChildWindowGeometry = ChildWindow->GetWindowGeometryInWindow();
+		FSlateRect ChildWindowClipRect = ChildWindow->GetClippingRectangleInWindow();
+		DrawWindowAndChildren(RenderTarget, DrawBuffer, HitTestGrid, ChildWindow, ChildWindowGeometry, ChildWindowClipRect, DeltaTime);
+	}
+}
diff --git a/Engine/Source/Runtime/UMG/Public/Components/WidgetComponent.h b/Engine/Source/Runtime/UMG/Public/Components/WidgetComponent.h
--- a/Engine/Source/Runtime/UMG/Public/Components/WidgetComponent.h
+++ b/Engine/Source/Runtime/UMG/Public/Components/WidgetComponent.h
@@ -23,6 +23,8 @@ enum class EWidgetBlendMode : uint8
 	Transparent
 };
 
+
+
 /**
  * Beware! This feature is experimental and may be substantially changed or removed in future releases.
  * A 3D instance of a Widget Blueprint that can be interacted with in the world.
@@ -50,6 +52,9 @@ public:
 
 	void ApplyComponentInstanceData(class FWidgetComponentInstanceData* ComponentInstanceData);
 
+	void RegisterHitTesterWithViewport(TSharedPtr<SViewport> ViewportWidget);
+	void UnregisterHitTesterWithViewport(TSharedPtr<SViewport> ViewportWidget);
+
 	// Begin UObject
 	virtual void PostLoad() override;
 	// End UObject
@@ -195,6 +200,13 @@ public:
 
 	void SetWidgetSpace( EWidgetSpace NewSpace ) { Space = NewSpace; }
 
+	bool GetEditTimeUsable() const { return bEditTimeUsable; }
+
+	void SetEditTimeUsable(bool Value) { bEditTimeUsable = Value; }
+
+protected:
+	virtual bool ShouldDrawWidget() const;
+
 protected:
 	void RemoveWidgetFromScreen();
 
@@ -271,6 +283,9 @@ protected:
 	/** The Slate widget to be displayed by this component.  Only one of either Widget or SlateWidget can be used */
 	TSharedPtr<SWidget> SlateWidget;
 
+	/** The slate widget currently being drawn. */
+	TWeakPtr<SWidget> CurrentSlateWidget;
+
 	/** The body setup of the displayed quad */
 	UPROPERTY(Transient, DuplicateTransient)
 	class UBodySetup* BodySetup;
@@ -312,6 +327,13 @@ protected:
 
 	UPROPERTY(Transient, DuplicateTransient)
 	bool bAddedToScreen;
+
+	/**
+	 * Allows the widget component to be used at editor time.  For use in the VR-Editor.
+	 */
+	UPROPERTY()
+	bool bEditTimeUsable;
+
 protected:
 
 	/** The grid used to find actual hit actual widgets once input has been translated to the components local space */
diff --git a/Engine/Source/Runtime/UMG/Public/Slate/WidgetRenderer.h b/Engine/Source/Runtime/UMG/Public/Slate/WidgetRenderer.h
--- a/Engine/Source/Runtime/UMG/Public/Slate/WidgetRenderer.h
+++ b/Engine/Source/Runtime/UMG/Public/Slate/WidgetRenderer.h
@@ -18,6 +18,12 @@ class SVirtualWindow : public SWindow
 
 public:
 	void Construct(const FArguments& InArgs);
+
+	virtual FPopupMethodReply OnQueryPopupMethod() const override;
+	virtual bool OnVisualizeTooltip(const TSharedPtr<SWidget>& TooltipContent) override;
+
+private:
+	TSharedPtr<class STooltipPresenter> TooltipPresenter;
 };
 
 /**
@@ -60,6 +66,18 @@ public:
 		FSlateRect WindowClipRect,
 		float DeltaTime);
 
+private:
+	void PrepassWindowAndChildren(TSharedRef<SWindow> Window, float Scale);
+
+	void DrawWindowAndChildren(
+		UTextureRenderTarget2D* RenderTarget,
+		FSlateDrawBuffer& DrawBuffer,
+		TSharedRef<FHittestGrid> HitTestGrid,
+		TSharedRef<SWindow> Window,
+		FGeometry WindowGeometry,
+		FSlateRect WindowClipRect,
+		float DeltaTime);
+
 private:
 	/** The slate 3D renderer used to render the user slate widget */
 	TSharedPtr<ISlate3DRenderer, ESPMode::ThreadSafe> Renderer;
