Slate: Re-added support for legacy pop-up layers on Slate windows
- This feature was still in use by some game teams, so we're putting it back in
- Now, we support for FPopupLayerSlot and FOverlayPopupLayer, depending on your needs

diff --git a/Engine/Source/Runtime/SlateCore/Private/Widgets/SWindow.cpp b/Engine/Source/Runtime/SlateCore/Private/Widgets/SWindow.cpp
--- a/Engine/Source/Runtime/SlateCore/Private/Widgets/SWindow.cpp
+++ b/Engine/Source/Runtime/SlateCore/Private/Widgets/SWindow.cpp
@@ -27,6 +27,172 @@ void FOverlayPopupLayer::Remove()
 	Overlay->RemoveSlot(GetContent());
 }
 
+
+
+/**
+ * An internal overlay used by Slate to support in-window pop ups and tooltips.
+ * The overlay ignores DPI scaling when it does its own arrangement, but otherwise
+ * passes all DPI scale values through.
+ */
+class SPopupLayer : public SPanel
+{
+public:
+	SLATE_BEGIN_ARGS( SPopupLayer )
+		{
+			_Visibility = EVisibility::SelfHitTestInvisible;
+		}
+
+		SLATE_SUPPORTS_SLOT( FPopupLayerSlot )
+
+	SLATE_END_ARGS()
+
+	SPopupLayer()
+	: Children()
+	{}
+
+	void Construct( const FArguments& InArgs, const TSharedRef<SWindow>& InWindow )
+	{
+		OwnerWindow = InWindow;
+
+		const int32 NumSlots = InArgs.Slots.Num();
+		for ( int32 SlotIndex = 0; SlotIndex < NumSlots; ++SlotIndex )
+		{
+			Children.Add( InArgs.Slots[SlotIndex] );
+		}
+	}
+
+	/** Make a new ListPanel::Slot  */
+	FPopupLayerSlot& Slot()
+	{
+		return *(new FPopupLayerSlot());
+	}
+	
+	/** Add a slot to the ListPanel */
+	FPopupLayerSlot& AddSlot(int32 InsertAtIndex = INDEX_NONE)
+	{
+		FPopupLayerSlot& NewSlot = *new FPopupLayerSlot();
+		if (InsertAtIndex == INDEX_NONE)
+		{
+			this->Children.Add( &NewSlot );
+		}
+		else
+		{
+			this->Children.Insert( &NewSlot, InsertAtIndex );
+		}
+	
+		return NewSlot;
+	}
+
+	void RemoveSlot(const TSharedRef<SWidget>& WidgetToRemove)
+	{
+		for( int32 CurSlotIndex = 0; CurSlotIndex < Children.Num(); ++CurSlotIndex )
+		{
+			const FPopupLayerSlot& CurSlot = Children[ CurSlotIndex ];
+			if( CurSlot.GetWidget() == WidgetToRemove )
+			{
+				Children.RemoveAt( CurSlotIndex );
+				return;
+			}
+		}
+	}
+
+private:
+
+	/**
+	 * Each child slot essentially tries to place their contents at a specified position on the screen
+	 * and scale as the widget initiating the popup, both of which are stored in the slot attributes.
+	 * The tricky part is that the scale we are given is the fully accumulated layout scale of the widget, which already incorporates 
+	 * the DPI Scale of the window. The DPI Scale is also applied to the overlay since it is part of the window,
+	 * so this scale needs to be factored out when determining the scale of the child geometry that will be created to hold the popup.
+	 * We also optionally adjust the window position to keep it within the client bounds of the top-level window. This must be done in screenspace.
+	 * This means some hairy transformation calculus goes on to ensure the computations are done in the proper space so scale is respected.
+	 * 
+	 * There are 3 transformational spaces involved, each clearly specified in the variable names:
+	 * Screen      - Basically desktop space. Contains desktop offset and DPI scale.
+	 * WindowLocal - local space of the SWindow containing this popup. Screenspace == Concat(WindowLocal, DPI Scale, Desktop Offset)
+	 * ChildLocal  - space of the child widget we want to display in the popup. The widget's LayoutTransform takes us from ChildLocal to WindowLocal space.
+	 */
+	virtual void OnArrangeChildren( const FGeometry& AllottedGeometry, FArrangedChildren& ArrangedChildren ) const override
+	{
+		// skip all this work if there are no children to arrange.
+		if (Children.Num() == 0) return;
+
+		// create a transform from screen to local space.
+		// This assumes that the PopupLayer is part of an Overlay that takes up the entire window space.
+		// We should technically be using the AllottedGeometry to transform from AbsoluteToLocal space just in case it has an additional scale on it.
+		// But we can't because the absolute space of the geometry is sometimes given in desktop space (picking, ticking) 
+		// and sometimes in window space (painting), and we can't necessarily tell by inspection so we have to just make an assumption here.
+		FSlateLayoutTransform ScreenToWindowLocal = (ensure(OwnerWindow.IsValid())) ? Inverse(OwnerWindow.Pin()->GetLocalToScreenTransform()) : FSlateLayoutTransform();
+		
+		for ( int32 ChildIndex = 0; ChildIndex < Children.Num(); ++ChildIndex )
+		{
+			const FPopupLayerSlot& CurChild = Children[ChildIndex];
+			const EVisibility ChildVisibility = CurChild.GetWidget()->GetVisibility();
+			if ( ArrangedChildren.Accepts(ChildVisibility) )
+			{
+				// This scale+translate forms the ChildLocal to Screenspace transform.
+				// The translation may be adjusted based on clamping, but the scale is accurate, 
+				// so we can transform vectors into screenspace using the scale alone.
+				const float ChildLocalToScreenScale = CurChild.Scale_Attribute.Get();
+				FVector2D ChildLocalToScreenOffset = CurChild.DesktopPosition_Attribute.Get();
+				// The size of the child is either the desired size of the widget (computed in the child's local space) or the size override (specified in screen space)
+				const FVector2D ChildSizeChildLocal = CurChild.GetWidget()->GetDesiredSize();
+				// Convert the desired size to screen space. Here is were we convert a vector to screenspace
+				// before we have the final position in screenspace (which would be needed to transform a point).
+				FVector2D ChildSizeScreenspace = TransformVector(ChildLocalToScreenScale, ChildSizeChildLocal);
+				// But then allow each size dimension to be overridden by the slot, which specifies the overrides in screen space.
+				ChildSizeScreenspace = FVector2D(
+						CurChild.WidthOverride_Attribute.IsSet() ? CurChild.WidthOverride_Attribute.Get() : ChildSizeScreenspace.X,
+						CurChild.HeightOverride_Attribute.IsSet() ? CurChild.HeightOverride_Attribute.Get() : ChildSizeScreenspace.Y);
+				
+				// If clamping, move the screen space position to ensure the screen space size stays within the client rect of the top level window.
+				if(CurChild.Clamp_Attribute.Get())
+				{
+					const FSlateRect WindowClientRectScreenspace = (ensure(OwnerWindow.IsValid())) ? OwnerWindow.Pin()->GetClientRectInScreen() : FSlateRect();
+					const FVector2D ClampBufferScreenspace = CurChild.ClampBuffer_Attribute.Get();
+					const FSlateRect ClampedWindowClientRectScreenspace = WindowClientRectScreenspace.InsetBy(FMargin(ClampBufferScreenspace.X, ClampBufferScreenspace.Y));
+					// Find how much our child wants to extend beyond our client space and subtract that amount, but don't push it past the client edge.
+					ChildLocalToScreenOffset.X = FMath::Max(WindowClientRectScreenspace.Left, ChildLocalToScreenOffset.X - FMath::Max(0.0f, (ChildLocalToScreenOffset.X  + ChildSizeScreenspace.X ) - ClampedWindowClientRectScreenspace.Right));
+					ChildLocalToScreenOffset.Y = FMath::Max(WindowClientRectScreenspace.Top, ChildLocalToScreenOffset.Y - FMath::Max(0.0f, (ChildLocalToScreenOffset.Y  + ChildSizeScreenspace.Y ) - ClampedWindowClientRectScreenspace.Bottom));
+				}
+
+				// We now have the final position, so construct the transform from ChildLocal to Screenspace
+				const FSlateLayoutTransform ChildLocalToScreen(ChildLocalToScreenScale, ChildLocalToScreenOffset);
+				// Using this we can compute the transform from ChildLocal to WindowLocal, which is effectively the LayoutTransform of the child widget.
+				const FSlateLayoutTransform ChildLocalToWindowLocal = Concatenate(ChildLocalToScreen, ScreenToWindowLocal);
+				// The ChildSize needs to be given in ChildLocal space when constructing a geometry.
+				const FVector2D ChildSizeLocalspace = TransformVector(Inverse(ChildLocalToScreen), ChildSizeScreenspace);
+
+				// The position is explicitly in desktop pixels.
+				// The size and DPI scale come from the widget that is using
+				// this overlay to "punch" through the UI.
+				ArrangedChildren.AddWidget(AllottedGeometry.MakeChild(CurChild.GetWidget(), ChildSizeLocalspace, ChildLocalToWindowLocal));
+			}
+		}
+	}
+
+	virtual FVector2D ComputeDesiredSize(float) const override
+	{
+		return FVector2D(100,100);
+	}
+
+	/**
+	 * All widgets must provide a way to access their children in a layout-agnostic way.
+	 * Panels store their children in Slots, which creates a dilemma. Most panels
+	 * can store their children in a TPanelChildren<Slot>, where the Slot class
+	 * provides layout information about the child it stores. In that case
+	 * GetChildren should simply return the TPanelChildren<Slot>. See StackPanel for an example.
+	 */
+	virtual FChildren* GetChildren() override
+	{
+		return &Children;
+	}
+
+	TPanelChildren<FPopupLayerSlot> Children;
+	TWeakPtr<SWindow> OwnerWindow;
+};
+
+
 FVector2D SWindow::GetWindowSizeFromClientSize(FVector2D InClientSize)
 {
 	// If this is a regular non-OS window, we need to compensate for the border and title bar area that we will add
@@ -353,6 +519,12 @@ void SWindow::ConstructWindowInternals()
 				]
 			]
 
+			// pop-up layer
+			+ SOverlay::Slot()
+			[
+				SAssignNew(PopupLayer, SPopupLayer, SharedThis(this))
+			]
+
 			// window outline
 			+ SOverlay::Slot()
 			[
@@ -373,6 +545,10 @@ void SWindow::ConstructWindowInternals()
 			[
 				MainWindowArea
 			]
+			+ SOverlay::Slot()
+			[
+				SAssignNew(PopupLayer, SPopupLayer, SharedThis(this))
+			]
 		];
 	}
 }
@@ -897,6 +1073,19 @@ TSharedPtr<FPopupLayer> SWindow::OnVisualizePopup(const TSharedRef<SWidget>& Pop
 	return TSharedPtr<FPopupLayer>();
 }
 
+/** Return a new slot in the popup layer. Assumes that the window has a popup layer. */
+FPopupLayerSlot& SWindow::AddPopupLayerSlot()
+{
+	ensure( PopupLayer.IsValid() );
+	return PopupLayer->AddSlot();
+}
+
+/** Counterpart to AddPopupLayerSlot */
+void SWindow::RemovePopupLayerSlot( const TSharedRef<SWidget>& WidgetToRemove )
+{
+	PopupLayer->RemoveSlot( WidgetToRemove );
+}
+
 /** @return should this window show up in the taskbar */
 bool SWindow::AppearsInTaskbar() const
 {
diff --git a/Engine/Source/Runtime/SlateCore/Public/Widgets/SWidget.h b/Engine/Source/Runtime/SlateCore/Public/Widgets/SWidget.h
--- a/Engine/Source/Runtime/SlateCore/Public/Widgets/SWidget.h
+++ b/Engine/Source/Runtime/SlateCore/Public/Widgets/SWidget.h
@@ -406,10 +406,12 @@ public:
 
 	/**
 	 * Visualize a new pop-up if possible.  If it's not possible for this widget to host the pop-up
-	 * content you'll get back an invalid pointer to the layer.
+	 * content you'll get back an invalid pointer to the layer.  The returned FPopupLayer allows you 
+	 * to remove the pop-up when you're done with it
 	 * 
 	 * @param PopupContent The widget to try and host overlaid on top of the widget.
-	 * @return a valid FPopupLayer if this widget supported hosting it.
+	 *
+	 * @return a valid FPopupLayer if this widget supported hosting it.  You can call Remove() on this to destroy the pop-up.
 	 */
 	virtual TSharedPtr<FPopupLayer> OnVisualizePopup(const TSharedRef<SWidget>& PopupContent);
 
diff --git a/Engine/Source/Runtime/SlateCore/Public/Widgets/SWindow.h b/Engine/Source/Runtime/SlateCore/Public/Widgets/SWindow.h
--- a/Engine/Source/Runtime/SlateCore/Public/Widgets/SWindow.h
+++ b/Engine/Source/Runtime/SlateCore/Public/Widgets/SWindow.h
@@ -436,8 +436,23 @@ public:
 	 */
 	void RemoveOverlaySlot( const TSharedRef<SWidget>& InContent );
 
+	/**
+	 * Visualize a new pop-up if possible.  If it's not possible for this widget to host the pop-up
+	 * content you'll get back an invalid pointer to the layer.  The returned FPopupLayer allows you 
+	 * to remove the pop-up when you're done with it
+	 * 
+	 * @param PopupContent The widget to try and host overlaid on top of the widget.
+	 *
+	 * @return a valid FPopupLayer if this widget supported hosting it.  You can call Remove() on this to destroy the pop-up.
+	 */
 	virtual TSharedPtr<FPopupLayer> OnVisualizePopup(const TSharedRef<SWidget>& PopupContent) override;
 
+	/** Return a new slot in the popup layer. Assumes that the window has a popup layer. */
+	struct FPopupLayerSlot& AddPopupLayerSlot();
+
+	/** Counterpart to AddPopupLayerSlot */
+	void RemovePopupLayerSlot( const TSharedRef<SWidget>& WidgetToRemove );
+
 	/**
 	 * Sets a widget to use as a full window overlay, or clears an existing widget if set.  When set, this widget will be drawn on top of all other window content.
 	 *
@@ -1034,6 +1049,80 @@ private:
 };
 
 
+/**
+ * Popups, tooltips, drag and drop decorators all can be executed without creating a new window.
+ * This slot along with the SWindow::AddPopupLayerSlot() API enabled it.
+ */
+struct FPopupLayerSlot : public TSlotBase<FPopupLayerSlot>
+{
+public:
+	FPopupLayerSlot()
+	: TSlotBase<FPopupLayerSlot>()
+	, DesktopPosition_Attribute(FVector2D::ZeroVector)
+	, WidthOverride_Attribute()
+	, HeightOverride_Attribute()
+	, Scale_Attribute(1.0f)
+	, Clamp_Attribute(false)
+	, ClampBuffer_Attribute(FVector2D::ZeroVector)
+	{}
+
+	/** Pixel position in desktop space */
+	FPopupLayerSlot& DesktopPosition( const TAttribute<FVector2D>& InDesktopPosition )
+	{
+		DesktopPosition_Attribute = InDesktopPosition;
+		return *this;
+	}
+
+	/** Width override in pixels */
+	FPopupLayerSlot& WidthOverride( const TAttribute<float>& InWidthOverride )
+	{
+		WidthOverride_Attribute = InWidthOverride;
+		return *this;
+	}
+
+	/** Width override in pixels */
+	FPopupLayerSlot& HeightOverride( const TAttribute<float>& InHeightOverride )
+	{
+		HeightOverride_Attribute = InHeightOverride;
+		return *this;
+	}
+
+	/** DPI scaling to be applied to the contents of this slot */
+	FPopupLayerSlot& Scale( const TAttribute<float>& InScale )
+	{
+		Scale_Attribute = InScale;
+		return *this;
+	}
+
+	/** Should this slot be kept within the parent window */
+	FPopupLayerSlot& ClampToWindow( const TAttribute<bool>& InClamp_Attribute)
+	{
+		Clamp_Attribute = InClamp_Attribute;
+		return *this;
+	}
+
+	/** If this slot is kept within the parent window, how far from the edges should we clamp it */
+	FPopupLayerSlot& ClampBuffer( const TAttribute<FVector2D>& InClampBuffer_Attribute )
+	{
+		ClampBuffer_Attribute = InClampBuffer_Attribute;
+		return *this;
+	}
+
+private:
+	/** SPopupLayer arranges FPopupLayerSlots, so it needs to know all about */
+	friend class SPopupLayer;
+	/** TPanelChildren need access to the Widget member */
+	friend class TPanelChildren<FPopupLayerSlot>;
+
+	TAttribute<FVector2D> DesktopPosition_Attribute;
+	TAttribute<float> WidthOverride_Attribute;
+	TAttribute<float> HeightOverride_Attribute;
+	TAttribute<float> Scale_Attribute;
+	TAttribute<bool> Clamp_Attribute;
+	TAttribute<FVector2D> ClampBuffer_Attribute;
+};
+
+
 #if WITH_EDITOR
 
 /**
