Slate: Fixed variable shadowing compiler warning on Clang

diff --git a/Engine/Source/Runtime/SlateCore/Public/Widgets/SWidget.h b/Engine/Source/Runtime/SlateCore/Public/Widgets/SWidget.h
--- a/Engine/Source/Runtime/SlateCore/Public/Widgets/SWidget.h
+++ b/Engine/Source/Runtime/SlateCore/Public/Widgets/SWidget.h
@@ -106,9 +106,9 @@ public:
 class FPopupLayer : public TSharedFromThis<FPopupLayer>
 {
 public:
-	FPopupLayer(const TSharedRef<SWidget>& HostWidget, const TSharedRef<SWidget>& PopupContent)
-		: HostWidget(HostWidget)
-		, PopupContent(PopupContent)
+	FPopupLayer(const TSharedRef<SWidget>& InitHostWidget, const TSharedRef<SWidget>& InitPopupContent)
+		: HostWidget(InitHostWidget)
+		, PopupContent(InitPopupContent)
 	{
 	}
 
