Fixed another compiler warning on Mac

diff --git a/Engine/Source/Runtime/SlateCore/Private/Widgets/SWindow.cpp b/Engine/Source/Runtime/SlateCore/Private/Widgets/SWindow.cpp
--- a/Engine/Source/Runtime/SlateCore/Private/Widgets/SWindow.cpp
+++ b/Engine/Source/Runtime/SlateCore/Private/Widgets/SWindow.cpp
@@ -12,9 +12,9 @@ namespace SWindowDefs
 	static const int32 CornerRadius = 6;
 }
 
-FOverlayPopupLayer::FOverlayPopupLayer(const TSharedRef<SWidget>& HostWidget, const TSharedRef<SWidget>& PopupContent, TSharedPtr<SOverlay> Overlay)
-	: FPopupLayer(HostWidget, PopupContent)
-	, Overlay(Overlay)
+FOverlayPopupLayer::FOverlayPopupLayer(const TSharedRef<SWidget>& InitHostWidget, const TSharedRef<SWidget>& InitPopupContent, TSharedPtr<SOverlay> InitOverlay)
+	: FPopupLayer(InitHostWidget, InitPopupContent)
+	, Overlay(InitOverlay)
 {
 	Overlay->AddSlot()
 	[
diff --git a/Engine/Source/Runtime/SlateCore/Public/Widgets/SWindow.h b/Engine/Source/Runtime/SlateCore/Public/Widgets/SWindow.h
--- a/Engine/Source/Runtime/SlateCore/Public/Widgets/SWindow.h
+++ b/Engine/Source/Runtime/SlateCore/Public/Widgets/SWindow.h
@@ -76,7 +76,7 @@ struct FWindowTransparency
 class FOverlayPopupLayer : public FPopupLayer
 {
 public:
-	FOverlayPopupLayer(const TSharedRef<SWidget>& HostWidget, const TSharedRef<SWidget>& PopupContent, TSharedPtr<SOverlay> Overlay);
+	FOverlayPopupLayer(const TSharedRef<SWidget>& InitHostWidget, const TSharedRef<SWidget>& InitPopupContent, TSharedPtr<SOverlay> InitOverlay);
 
 	virtual void Remove() override;
 
