Fixed compile error with recent change

diff --git a/Engine/Source/Runtime/SlateCore/Private/Widgets/SWindow.cpp b/Engine/Source/Runtime/SlateCore/Private/Widgets/SWindow.cpp
--- a/Engine/Source/Runtime/SlateCore/Private/Widgets/SWindow.cpp
+++ b/Engine/Source/Runtime/SlateCore/Private/Widgets/SWindow.cpp
@@ -18,7 +18,7 @@ FOverlayPopupLayer::FOverlayPopupLayer(const TSharedRef<SWidget>& InitHostWidget
 {
 	Overlay->AddSlot()
 	[
-		PopupContent
+		InitPopupContent
 	];
 }
 
