Copying //UE4/Dev-Editor to //UE4/Dev-Main (Source: //UE4/Dev-Editor @ 3082391)

diff --git a/Engine/Source/Runtime/Internationalization/Public/TextLocalizationResourceGenerator.h b/Engine/Source/Runtime/Internationalization/Public/TextLocalizationResourceGenerator.h
new file mode 100644
--- /dev/null
+++ b/Engine/Source/Runtime/Internationalization/Public/TextLocalizationResourceGenerator.h
@@ -0,0 +1,73 @@
+// Copyright 1998-2016 Epic Games, Inc. All Rights Reserved.
+
+#pragma once 
+
+#include "LocTextHelper.h"
+
+class FTextLocalizationResourceGenerator
+{
+public:
+	/**
+	 * Given a loc text helper, generate a compiled LocRes file for the given culture into the destination archive.
+	 */
+	INTERNATIONALIZATION_API static bool Generate(const FLocTextHelper& InLocTextHelper, const FString& InCultureToGenerate, const bool bSkipSourceCheck, FArchive& InDestinationArchive);
+
+	/**
+	 * Given a config file, generate a compiled LocRes file for the active culture and use it to update the live-entries in the localization manager.
+	 */
+	INTERNATIONALIZATION_API static bool GenerateAndUpdateLiveEntriesFromConfig(const FString& InConfigFilePath, const bool bSkipSourceCheck);
+
+private:
+
+	struct FLocalizationEntryTracker
+	{
+		struct FEntry
+		{
+			uint32 SourceStringHash;
+			FString LocalizedString;
+		};
+
+		typedef TArray<FEntry> FEntryArray;
+
+
+		struct FKeyTableKeyFuncs : BaseKeyFuncs<FEntryArray, FString, false>
+		{
+			static FORCEINLINE const FString& GetSetKey(const TPair<FString,FEntryArray>& Element)
+			{
+				return Element.Key;
+			}
+			static FORCEINLINE bool Matches(const FString& A, const FString& B)
+			{
+				return A.Equals( B, ESearchCase::CaseSensitive );
+			}
+			static FORCEINLINE uint32 GetKeyHash(const FString& Key)
+			{
+				return FCrc::StrCrc32<TCHAR>(*Key);
+			}
+		};
+		typedef TMap<FString, FEntryArray, FDefaultSetAllocator, FKeyTableKeyFuncs> FKeyTable;
+
+
+		struct FNamespaceTableKeyFuncs : BaseKeyFuncs<FKeyTable, FString, false>
+		{
+			static FORCEINLINE const FString& GetSetKey(const TPair<FString,FKeyTable>& Element)
+			{
+				return Element.Key;
+			}
+			static FORCEINLINE bool Matches(const FString& A, const FString& B)
+			{
+				return A.Equals( B, ESearchCase::CaseSensitive );
+			}
+			static FORCEINLINE uint32 GetKeyHash(const FString& Key)
+			{
+				return FCrc::StrCrc32<TCHAR>(*Key);
+			}
+		};
+		typedef TMap<FString, FKeyTable, FDefaultSetAllocator, FNamespaceTableKeyFuncs> FNamespaceTable;
+
+		FNamespaceTable Namespaces;
+
+		void ReportCollisions() const;
+		bool WriteToArchive(FArchive* const Archive);
+	};
+};
\ No newline at end of file
diff --git a/Engine/Source/Runtime/Landscape/Private/LandscapeBlueprintSupport.cpp b/Engine/Source/Runtime/Landscape/Private/LandscapeBlueprintSupport.cpp
--- a/Engine/Source/Runtime/Landscape/Private/LandscapeBlueprintSupport.cpp
+++ b/Engine/Source/Runtime/Landscape/Private/LandscapeBlueprintSupport.cpp
@@ -17,7 +17,7 @@ void ALandscapeProxy::EditorApplySpline(USplineComponent* InSplineComponent, flo
 	if (InSplineComponent && !GetWorld()->IsGameWorld())
 	{
 		TArray<FLandscapeSplineInterpPoint> Points;
-		LandscapeSplineRaster::Pointify(InSplineComponent->SplineInfo, Points, NumSubdivisions, 0.0f, 0.0f, StartWidth, EndWidth, StartSideFalloff, EndSideFalloff, StartRoll, EndRoll);
+		LandscapeSplineRaster::Pointify(InSplineComponent->SplineCurves.Position, Points, NumSubdivisions, 0.0f, 0.0f, StartWidth, EndWidth, StartSideFalloff, EndSideFalloff, StartRoll, EndRoll);
 
 		FTransform SplineToWorld = InSplineComponent->GetComponentTransform();
 		LandscapeSplineRaster::RasterizeSegmentPoints(GetLandscapeInfo(), MoveTemp(Points), SplineToWorld, bRaiseHeights, bLowerHeights, PaintLayer);
diff --git a/Engine/Source/Runtime/Launch/Private/LaunchEngineLoop.cpp b/Engine/Source/Runtime/Launch/Private/LaunchEngineLoop.cpp
--- a/Engine/Source/Runtime/Launch/Private/LaunchEngineLoop.cpp
+++ b/Engine/Source/Runtime/Launch/Private/LaunchEngineLoop.cpp
@@ -3252,6 +3252,12 @@ bool FEngineLoop::AppInit( )
 		return false;
 	}
 
+	// Register the callback that allows the text localization manager to load data for plugins
+	FTextLocalizationManager::Get().GatherAdditionalLocResPathsCallback.AddLambda([](TArray<FString>& OutLocResPaths)
+	{
+		IPluginManager::Get().GetLocalizationPathsForEnabledPlugins(OutLocResPaths);
+	});
+
 	PreInitHMDDevice();
 
 	// Put the command line and config info into the suppression system
diff --git a/Engine/Source/Runtime/MovieScene/Private/MovieScene.cpp b/Engine/Source/Runtime/MovieScene/Private/MovieScene.cpp
--- a/Engine/Source/Runtime/MovieScene/Private/MovieScene.cpp
+++ b/Engine/Source/Runtime/MovieScene/Private/MovieScene.cpp
@@ -55,7 +55,7 @@ void UMovieScene::OptimizeObjectArray(TArray<T>& ObjectArray)
 		// @todo: ObjectBindings mapped by ID to avoid linear search
 		for (int32 BindingIndex = 0; BindingIndex < ObjectBindings.Num(); ++BindingIndex)
 		{
-			FMovieSceneBinding& Binding =ObjectBindings[BindingIndex];
+			FMovieSceneBinding& Binding = ObjectBindings[BindingIndex];
 			if (Binding.GetObjectGuid() != ObjectGuid)
 			{
 				continue;
@@ -436,6 +436,30 @@ UMovieSceneTrack* UMovieScene::AddTrack( TSubclassOf<UMovieSceneTrack> TrackClas
 	return CreatedType;
 }
 
+bool UMovieScene::AddGivenTrack(UMovieSceneTrack* InTrack, const FGuid& ObjectGuid)
+{
+	check(ObjectGuid.IsValid());
+
+	Modify();
+	for (auto& Binding : ObjectBindings)
+	{
+		if (Binding.GetObjectGuid() == ObjectGuid)
+		{
+			for (auto& Track : Binding.GetTracks())
+			{
+				if (Track->GetTrackName() == InTrack->GetTrackName())
+				{
+					return false;
+				}
+			}
+			InTrack->Rename(nullptr, this);
+			check(InTrack);
+			Binding.AddTrack(*InTrack);
+			return true;
+		}
+	}
+	return false;
+}
 
 bool UMovieScene::RemoveTrack(UMovieSceneTrack& Track)
 {
diff --git a/Engine/Source/Runtime/MovieScene/Public/MovieScene.h b/Engine/Source/Runtime/MovieScene/Public/MovieScene.h
--- a/Engine/Source/Runtime/MovieScene/Public/MovieScene.h
+++ b/Engine/Source/Runtime/MovieScene/Public/MovieScene.h
@@ -212,6 +212,18 @@ public:
 	 */
 	UMovieSceneTrack* AddTrack(TSubclassOf<UMovieSceneTrack> TrackClass, const FGuid& ObjectGuid);
 
+	/**
+	* Adds a given track.
+	*
+	* Note: Function will only add if the track is not already exist.
+	*
+	* @param InTrack The track to add.
+	* @param ObjectGuid The runtime object guid that the type should bind to.
+	* @see  FindTrack, RemoveTrack
+	* @return true if the track is successfully added, false otherwise.
+	*/
+	bool AddGivenTrack(UMovieSceneTrack* InTrack, const FGuid& ObjectGuid);
+
 	/**
 	 * Adds a track.
 	 *
diff --git a/Engine/Source/Runtime/MovieScene/Public/MovieSceneBinding.h b/Engine/Source/Runtime/MovieScene/Public/MovieSceneBinding.h
--- a/Engine/Source/Runtime/MovieScene/Public/MovieSceneBinding.h
+++ b/Engine/Source/Runtime/MovieScene/Public/MovieSceneBinding.h
@@ -87,7 +87,7 @@ struct FMovieSceneBinding
 	 *
 	 * @param NewTrack	The track to add
 	 */
-	void AddTrack(UMovieSceneTrack& NewTrack);
+	void MOVIESCENE_API AddTrack(UMovieSceneTrack& NewTrack);
 
 	/**
 	 * Removes a track from this binding
diff --git a/Engine/Source/Runtime/PhysXFormats/Private/PhysXFormats.cpp b/Engine/Source/Runtime/PhysXFormats/Private/PhysXFormats.cpp
--- a/Engine/Source/Runtime/PhysXFormats/Private/PhysXFormats.cpp
+++ b/Engine/Source/Runtime/PhysXFormats/Private/PhysXFormats.cpp
@@ -97,14 +97,7 @@ public:
 		PConvexMeshDesc.points.data = SrcBuffer.GetData();
 		PConvexMeshDesc.points.count = SrcBuffer.Num();
 		PConvexMeshDesc.points.stride = sizeof(FVector);
-		if (bDeformableMesh)
-		{
-			PConvexMeshDesc.flags = PxConvexFlag::eCOMPUTE_CONVEX | PxConvexFlag::eINFLATE_CONVEX;
-		}
-		else
-		{
-			PConvexMeshDesc.flags = PxConvexFlag::eCOMPUTE_CONVEX;
-		}
+		PConvexMeshDesc.flags = PxConvexFlag::eCOMPUTE_CONVEX;
 
 		// Set up cooking
 		const PxCookingParams Params = PhysXCooking->getParams();
diff --git a/Engine/Source/Runtime/Projects/Private/LocalizationDescriptor.cpp b/Engine/Source/Runtime/Projects/Private/LocalizationDescriptor.cpp
new file mode 100644
--- /dev/null
+++ b/Engine/Source/Runtime/Projects/Private/LocalizationDescriptor.cpp
@@ -0,0 +1,175 @@
+// Copyright 1998-2016 Epic Games, Inc. All Rights Reserved.
+
+#include "ProjectsPrivatePCH.h"
+#include "LocalizationDescriptor.h"
+
+#define LOCTEXT_NAMESPACE "LocalizationDescriptor"
+
+ELocalizationTargetDescriptorLoadingPolicy::Type ELocalizationTargetDescriptorLoadingPolicy::FromString(const TCHAR *String)
+{
+	ELocalizationTargetDescriptorLoadingPolicy::Type TestType = (ELocalizationTargetDescriptorLoadingPolicy::Type)0;
+	for(; TestType < ELocalizationTargetDescriptorLoadingPolicy::Max; TestType = (ELocalizationTargetDescriptorLoadingPolicy::Type)(TestType + 1))
+	{
+		const TCHAR *TestString = ToString(TestType);
+		if (FCString::Stricmp(String, TestString) == 0)
+		{
+			break;
+		}
+	}
+	return TestType;
+}
+
+const TCHAR* ELocalizationTargetDescriptorLoadingPolicy::ToString(const ELocalizationTargetDescriptorLoadingPolicy::Type Value)
+{
+	switch (Value)
+	{
+	case Never:
+		return TEXT("Never");
+
+	case Always:
+		return TEXT("Always");
+
+	case Editor:
+		return TEXT("Editor");
+
+	case Game:
+		return TEXT("Game");
+
+	case PropertyNames:
+		return TEXT("PropertyNames");
+
+	case ToolTips:
+		return TEXT("ToolTips");
+
+	default:
+		ensureMsgf(false, TEXT("ELocalizationTargetDescriptorLoadingPolicy::ToString - Unrecognized ELocalizationTargetDescriptorLoadingPolicy value: %i"), Value);
+		return nullptr;
+	}
+}
+
+FLocalizationTargetDescriptor::FLocalizationTargetDescriptor(FString InName, ELocalizationTargetDescriptorLoadingPolicy::Type InLoadingPolicy)
+	: Name(MoveTemp(InName))
+	, LoadingPolicy(InLoadingPolicy)
+{
+}
+
+bool FLocalizationTargetDescriptor::Read(const FJsonObject& InObject, FText& OutFailReason)
+{
+	// Read the target name
+	TSharedPtr<FJsonValue> NameValue = InObject.TryGetField(TEXT("Name"));
+	if (!NameValue.IsValid() || NameValue->Type != EJson::String)
+	{
+		OutFailReason = LOCTEXT("TargetWithoutAName", "Found a 'Localization Target' entry with a missing 'Name' field");
+		return false;
+	}
+	Name = NameValue->AsString();
+
+	// Read the target loading policy
+	TSharedPtr<FJsonValue> LoadingPolicyValue = InObject.TryGetField(TEXT("LoadingPolicy"));
+	if (LoadingPolicyValue.IsValid() && LoadingPolicyValue->Type == EJson::String)
+	{
+		LoadingPolicy = ELocalizationTargetDescriptorLoadingPolicy::FromString(*LoadingPolicyValue->AsString());
+		if (LoadingPolicy == ELocalizationTargetDescriptorLoadingPolicy::Max)
+		{
+			OutFailReason = FText::Format(LOCTEXT("TargetWithInvalidLoadingPolicy", "Localization Target entry '{0}' specified an unrecognized target LoadingPolicy '{1}'"), FText::FromString(Name), FText::FromString(LoadingPolicyValue->AsString()));
+			return false;
+		}
+	}
+
+	return true;
+}
+
+bool FLocalizationTargetDescriptor::ReadArray(const FJsonObject& InObject, const TCHAR* InName, TArray<FLocalizationTargetDescriptor>& OutTargets, FText& OutFailReason)
+{
+	bool bResult = true;
+
+	TSharedPtr<FJsonValue> TargetsArrayValue = InObject.TryGetField(InName);
+	if (TargetsArrayValue.IsValid() && TargetsArrayValue->Type == EJson::Array)
+	{
+		const TArray<TSharedPtr<FJsonValue>>& TargetsArray = TargetsArrayValue->AsArray();
+		for (const TSharedPtr<FJsonValue>& TargetValue : TargetsArray)
+		{
+			if (TargetValue.IsValid() && TargetValue->Type == EJson::Object)
+			{
+				FLocalizationTargetDescriptor Descriptor;
+				if (Descriptor.Read(*TargetValue->AsObject().Get(), OutFailReason))
+				{
+					OutTargets.Add(Descriptor);
+				}
+				else
+				{
+					bResult = false;
+				}
+			}
+			else
+			{
+				OutFailReason = LOCTEXT("TargetWithInvalidTargetsArray", "The 'Localization Targets' array has invalid contents and was not able to be loaded.");
+				bResult = false;
+			}
+		}
+	}
+
+	return bResult;
+}
+
+void FLocalizationTargetDescriptor::Write(TJsonWriter<>& InWriter) const
+{
+	InWriter.WriteObjectStart();
+	InWriter.WriteValue(TEXT("Name"), Name);
+	InWriter.WriteValue(TEXT("LoadingPolicy"), FString(ELocalizationTargetDescriptorLoadingPolicy::ToString(LoadingPolicy)));
+	InWriter.WriteObjectEnd();
+}
+
+void FLocalizationTargetDescriptor::WriteArray(TJsonWriter<>& InWriter, const TCHAR* InName, const TArray<FLocalizationTargetDescriptor>& InTargets)
+{
+	if (InTargets.Num() > 0)
+	{
+		InWriter.WriteArrayStart(InName);
+		for (const FLocalizationTargetDescriptor& Target : InTargets)
+		{
+			Target.Write(InWriter);
+		}
+		InWriter.WriteArrayEnd();
+	}
+}
+
+bool FLocalizationTargetDescriptor::ShouldLoadLocalizationTarget() const
+{
+	switch (LoadingPolicy)
+	{
+	case ELocalizationTargetDescriptorLoadingPolicy::Never:
+		return false;
+
+	case ELocalizationTargetDescriptorLoadingPolicy::Always:
+		return true;
+
+	case ELocalizationTargetDescriptorLoadingPolicy::Editor:
+		return WITH_EDITOR;
+
+	case ELocalizationTargetDescriptorLoadingPolicy::Game:
+		return FApp::IsGame();
+
+	case ELocalizationTargetDescriptorLoadingPolicy::PropertyNames:
+#if WITH_EDITOR
+		{
+			bool bShouldLoadLocalizedPropertyNames = true;
+			if (!GConfig->GetBool(TEXT("Internationalization"), TEXT("ShouldLoadLocalizedPropertyNames"), bShouldLoadLocalizedPropertyNames, GEditorSettingsIni))
+			{
+				GConfig->GetBool(TEXT("Internationalization"), TEXT("ShouldLoadLocalizedPropertyNames"), bShouldLoadLocalizedPropertyNames, GEngineIni);
+			}
+			return bShouldLoadLocalizedPropertyNames;
+		}
+#else	// WITH_EDITOR
+		return false;
+#endif	// WITH_EDITOR
+
+	case ELocalizationTargetDescriptorLoadingPolicy::ToolTips:
+		return WITH_EDITOR;
+
+	default:
+		ensureMsgf(false, TEXT("FLocalizationTargetDescriptor::ShouldLoadLocalizationTarget - Unrecognized ELocalizationTargetDescriptorLoadingPolicy value: %i"), LoadingPolicy);
+		return false;
+	}
+}
+
+#undef LOCTEXT_NAMESPACE
diff --git a/Engine/Source/Runtime/Projects/Private/PluginDescriptor.cpp b/Engine/Source/Runtime/Projects/Private/PluginDescriptor.cpp
--- a/Engine/Source/Runtime/Projects/Private/PluginDescriptor.cpp
+++ b/Engine/Source/Runtime/Projects/Private/PluginDescriptor.cpp
@@ -98,6 +98,11 @@ bool FPluginDescriptor::Read(const FString& Text, FText& OutFailReason)
 		return false;
 	}
 
+	if (!FLocalizationTargetDescriptor::ReadArray(Object, TEXT("LocalizationTargets"), LocalizationTargets, OutFailReason))
+	{
+		return false;
+	}
+
 	Object.TryGetBoolField(TEXT("EnabledByDefault"), bEnabledByDefault);
 	Object.TryGetBoolField(TEXT("CanContainContent"), bCanContainContent);
 	Object.TryGetBoolField(TEXT("IsBetaVersion"), bIsBetaVersion);
@@ -154,6 +159,9 @@ FString FPluginDescriptor::ToString() const
 	Writer.WriteValue(TEXT("Installed"), bInstalled);
 
 	FModuleDescriptor::WriteArray(Writer, TEXT("Modules"), Modules);
+
+	FLocalizationTargetDescriptor::WriteArray(Writer, TEXT("LocalizationTargets"), LocalizationTargets);
+
 	if(!bRequiresBuildPlatform)
 	{
 		Writer.WriteValue(TEXT("RequiresBuildPlatform"), bRequiresBuildPlatform);
diff --git a/Engine/Source/Runtime/Projects/Private/PluginManager.cpp b/Engine/Source/Runtime/Projects/Private/PluginManager.cpp
--- a/Engine/Source/Runtime/Projects/Private/PluginManager.cpp
+++ b/Engine/Source/Runtime/Projects/Private/PluginManager.cpp
@@ -598,6 +598,33 @@ bool FPluginManager::LoadModulesForEnabledPlugins( const ELoadingPhase::Type Loa
 	return true;
 }
 
+void FPluginManager::GetLocalizationPathsForEnabledPlugins( TArray<FString>& OutLocResPaths )
+{
+	// Figure out which plugins are enabled
+	if (!ConfigureEnabledPlugins())
+	{
+		return;
+	}
+
+	// Gather the paths from all plugins that have localization targets that are loaded based on the current runtime environment
+	for (const TSharedRef<FPlugin>& Plugin : AllPlugins)
+	{
+		if (!Plugin->bEnabled || Plugin->GetDescriptor().LocalizationTargets.Num() == 0)
+		{
+			continue;
+		}
+		
+		const FString PluginLocDir = Plugin->GetContentDir() / TEXT("Localization");
+		for (const FLocalizationTargetDescriptor& LocTargetDesc : Plugin->GetDescriptor().LocalizationTargets)
+		{
+			if (LocTargetDesc.ShouldLoadLocalizationTarget())
+			{
+				OutLocResPaths.Add(PluginLocDir / LocTargetDesc.Name);
+			}
+		}
+	}
+}
+
 void FPluginManager::SetRegisterMountPointDelegate( const FRegisterMountPointDelegate& Delegate )
 {
 	RegisterMountPointDelegate = Delegate;
diff --git a/Engine/Source/Runtime/Projects/Private/PluginManager.h b/Engine/Source/Runtime/Projects/Private/PluginManager.h
--- a/Engine/Source/Runtime/Projects/Private/PluginManager.h
+++ b/Engine/Source/Runtime/Projects/Private/PluginManager.h
@@ -63,6 +63,7 @@ public:
 	/** IPluginManager interface */
 	virtual void RefreshPluginsList() override;
 	virtual bool LoadModulesForEnabledPlugins( const ELoadingPhase::Type LoadingPhase ) override;
+	virtual void GetLocalizationPathsForEnabledPlugins( TArray<FString>& OutLocResPaths ) override;
 	virtual void SetRegisterMountPointDelegate( const FRegisterMountPointDelegate& Delegate ) override;
 	virtual bool AreRequiredPluginsAvailable() override;
 	virtual bool CheckModuleCompatibility( TArray<FString>& OutIncompatibleModules ) override;
diff --git a/Engine/Source/Runtime/Projects/Public/Interfaces/IPluginManager.h b/Engine/Source/Runtime/Projects/Public/Interfaces/IPluginManager.h
--- a/Engine/Source/Runtime/Projects/Public/Interfaces/IPluginManager.h
+++ b/Engine/Source/Runtime/Projects/Public/Interfaces/IPluginManager.h
@@ -142,6 +142,13 @@ public:
 	 */
 	virtual bool LoadModulesForEnabledPlugins( const ELoadingPhase::Type LoadingPhase ) = 0;
 
+	/**
+	 * Get the localization paths for all enabled plugins.
+	 *
+	 * @param	OutLocResPaths	Array to populate with the localization paths for all enabled plugins.
+	 */
+	virtual void GetLocalizationPathsForEnabledPlugins( TArray<FString>& OutLocResPaths ) = 0;
+
 	/** Delegate type for mounting content paths.  Used internally by FPackageName code. */
 	DECLARE_DELEGATE_TwoParams( FRegisterMountPointDelegate, const FString& /* Root content path */, const FString& /* Directory name */ );
 
diff --git a/Engine/Source/Runtime/Projects/Public/PluginDescriptor.h b/Engine/Source/Runtime/Projects/Public/PluginDescriptor.h
--- a/Engine/Source/Runtime/Projects/Public/PluginDescriptor.h
+++ b/Engine/Source/Runtime/Projects/Public/PluginDescriptor.h
@@ -3,6 +3,7 @@
 #pragma once
 
 #include "ModuleDescriptor.h"
+#include "LocalizationDescriptor.h"
 #include "CustomBuildSteps.h"
 
 /**
@@ -68,6 +69,9 @@ struct PROJECTS_API FPluginDescriptor
 	/** List of all modules associated with this plugin */
 	TArray<FModuleDescriptor> Modules;
 
+	/** List of all localization targets associated with this plugin */
+	TArray<FLocalizationTargetDescriptor> LocalizationTargets;
+
 	/** Whether this plugin should be enabled by default for all projects */
 	bool bEnabledByDefault;
 
diff --git a/Engine/Source/Runtime/Renderer/Private/PostProcess/PostProcessSelectionOutline.cpp b/Engine/Source/Runtime/Renderer/Private/PostProcess/PostProcessSelectionOutline.cpp
--- a/Engine/Source/Runtime/Renderer/Private/PostProcess/PostProcessSelectionOutline.cpp
+++ b/Engine/Source/Runtime/Renderer/Private/PostProcess/PostProcessSelectionOutline.cpp
@@ -20,6 +20,7 @@
 
 void FRCPassPostProcessSelectionOutlineColor::Process(FRenderingCompositePassContext& Context)
 {
+#if WITH_EDITOR
 	SCOPED_DRAW_EVENT(Context.RHICmdList, PostProcessSelectionOutlineBuffer);
 
 	const FPooledRenderTargetDesc* SceneColorInputDesc = GetInputDesc(ePId_Input0);
@@ -39,6 +40,7 @@ void FRCPassPostProcessSelectionOutlineColor::Process(FRenderingCompositePassCon
 	
 	// Set the render target/viewport.
 	SetRenderTarget(Context.RHICmdList, FTextureRHIParamRef(), DestRenderTarget.TargetableTexture);
+
 	// This is a reversed Z depth surface, so 0.0f is the far plane.
 	Context.RHICmdList.Clear(false, FLinearColor(0, 0, 0, 0), true, (float)ERHIZBuffer::FarPlane, true, 0, FIntRect());
 	Context.SetViewportAndCallRHI(ViewRect);
@@ -47,53 +49,44 @@ void FRCPassPostProcessSelectionOutlineColor::Process(FRenderingCompositePassCon
 	{
 		FHitProxyDrawingPolicyFactory::ContextType FactoryContext;
 
-		//@todo - use memstack
-		TMap<FName, int32> ActorNameToStencilIndex;
-		TMap<const FPrimitiveSceneProxy*, int32> IndividuallySelectedProxies;
-		ActorNameToStencilIndex.Add(NAME_BSP, 1);
-
 		Context.RHICmdList.SetRasterizerState(TStaticRasterizerState<>::GetRHI());
 		Context.RHICmdList.SetBlendState(TStaticBlendStateWriteMask<CW_NONE, CW_NONE, CW_NONE, CW_NONE>::GetRHI());
 
+		// Note that the stencil value will overflow with enough selected objects
+		FEditorSelectionDrawingPolicy::ResetStencilValues();
+
+		// Run selection pass on static elements
+		FScene* Scene = View.Family->Scene->GetRenderScene();
+		if(Scene)
+		{	
+			Scene->EditorSelectionDrawList.DrawVisible(Context.RHICmdList, View, View.StaticMeshEditorSelectionMap, View.StaticMeshBatchVisibility);
+		}
+	
 		for (int32 MeshBatchIndex = 0; MeshBatchIndex < View.DynamicMeshElements.Num(); MeshBatchIndex++)
 		{
 			const FMeshBatchAndRelevance& MeshBatchAndRelevance = View.DynamicMeshElements[MeshBatchIndex];
 			const FPrimitiveSceneProxy* PrimitiveSceneProxy = MeshBatchAndRelevance.PrimitiveSceneProxy;
 
-#if WITH_EDITOR
 			// Selected actors should be subdued if any component is individually selected
 			bool bActorSelectionColorIsSubdued = View.bHasSelectedComponents;
-#else
-			bool bActorSelectionColorIsSubdued = false;
-#endif
+
 			if (PrimitiveSceneProxy->IsSelected() && MeshBatchAndRelevance.Mesh->bUseSelectionOutline && PrimitiveSceneProxy->WantsSelectionOutline())
 			{
-				const int32* AssignedStencilIndexPtr = PrimitiveSceneProxy->IsIndividuallySelected() ? IndividuallySelectedProxies.Find( PrimitiveSceneProxy ) : ActorNameToStencilIndex.Find(PrimitiveSceneProxy->GetOwnerName());
-
-				if (!AssignedStencilIndexPtr)
+				int32 StencilValue = 1;
+				if(PrimitiveSceneProxy->GetOwnerName() != NAME_BSP)
 				{
-					if( PrimitiveSceneProxy->IsIndividuallySelected() )
-					{
-						// Any component that is individually selected should have a stencil value of < 128 so that it can have a unique color.  We offset the value by 2 because 0 means no selection and 1 is for bsp
-						int32 StencilValue = IndividuallySelectedProxies.Num() % 126 + 2;
-						AssignedStencilIndexPtr = &IndividuallySelectedProxies.Add(PrimitiveSceneProxy, StencilValue);
-					}
-					else
-					{
-						// If we are subduing actor color highlight then use the top level bits to indicate that to the shader.  
-						int32 StencilValue = bActorSelectionColorIsSubdued ? ActorNameToStencilIndex.Num() % 128 + 128 : ActorNameToStencilIndex.Num() % 126 + 2;
-						AssignedStencilIndexPtr = &ActorNameToStencilIndex.Add(PrimitiveSceneProxy->GetOwnerName(), StencilValue);
-					}
+					StencilValue = FEditorSelectionDrawingPolicy::GetStencilValue(View, PrimitiveSceneProxy);
+
 				}
 
 				// Note that the stencil value will overflow with enough selected objects
-				Context.RHICmdList.SetDepthStencilState(TStaticDepthStencilState<true, CF_DepthNearOrEqual, true, CF_Always, SO_Keep, SO_Keep, SO_Replace>::GetRHI(), *AssignedStencilIndexPtr);
+				Context.RHICmdList.SetDepthStencilState(TStaticDepthStencilState<true, CF_DepthNearOrEqual, true, CF_Always, SO_Keep, SO_Keep, SO_Replace>::GetRHI(), StencilValue);
 
 				const FMeshBatch& MeshBatch = *MeshBatchAndRelevance.Mesh;
 				FHitProxyDrawingPolicyFactory::DrawDynamicMesh(Context.RHICmdList, View, FactoryContext, MeshBatch, false, true, MeshBatchAndRelevance.PrimitiveSceneProxy, MeshBatch.BatchHitProxyId);
 			}
 		}
-		
+
 		// to get an outline around the objects if it's partly outside of the screen
 		{
 			FIntRect InnerRect = ViewRect;
@@ -124,6 +117,7 @@ void FRCPassPostProcessSelectionOutlineColor::Process(FRenderingCompositePassCon
 	
 	// Resolve to the output
 	Context.RHICmdList.CopyToResolveTarget(DestRenderTarget.TargetableTexture, DestRenderTarget.ShaderResourceTexture, false, FResolveParams());
+#endif	
 }
 
 FPooledRenderTargetDesc FRCPassPostProcessSelectionOutlineColor::ComputeOutputDesc(EPassOutputId InPassOutputId) const
diff --git a/Engine/Source/Runtime/Renderer/Private/RendererScene.cpp b/Engine/Source/Runtime/Renderer/Private/RendererScene.cpp
--- a/Engine/Source/Runtime/Renderer/Private/RendererScene.cpp
+++ b/Engine/Source/Runtime/Renderer/Private/RendererScene.cpp
@@ -2159,6 +2159,9 @@ void FScene::UpdateStaticDrawListsForMaterials_RenderThread(FRHICommandListImmed
 	MaskedDepthDrawList.GetUsedPrimitivesBasedOnMaterials(SceneFeatureLevel, Materials, PrimitivesToUpdate);
 	HitProxyDrawList.GetUsedPrimitivesBasedOnMaterials(SceneFeatureLevel, Materials, PrimitivesToUpdate);
 	HitProxyDrawList_OpaqueOnly.GetUsedPrimitivesBasedOnMaterials(SceneFeatureLevel, Materials, PrimitivesToUpdate);
+#if WITH_EDITOR
+	EditorSelectionDrawList.GetUsedPrimitivesBasedOnMaterials(SceneFeatureLevel, Materials, PrimitivesToUpdate);
+#endif
 	VelocityDrawList.GetUsedPrimitivesBasedOnMaterials(SceneFeatureLevel, Materials, PrimitivesToUpdate);
 	WholeSceneShadowDepthDrawList.GetUsedPrimitivesBasedOnMaterials(SceneFeatureLevel, Materials, PrimitivesToUpdate);
 	WholeSceneReflectiveShadowMapDrawList.GetUsedPrimitivesBasedOnMaterials(SceneFeatureLevel, Materials, PrimitivesToUpdate);
@@ -2389,6 +2392,9 @@ void FScene::DumpStaticMeshDrawListStats() const
 	DUMP_DRAW_LIST(MobileBasePassUniformLightMapPolicyDrawListWithCSM[EBasePass_Masked]);
 	DUMP_DRAW_LIST(HitProxyDrawList);
 	DUMP_DRAW_LIST(HitProxyDrawList_OpaqueOnly);
+#if WITH_EDITOR
+	DUMP_DRAW_LIST(EditorSelectionDrawList);
+#endif
 	DUMP_DRAW_LIST(VelocityDrawList);
 	DUMP_DRAW_LIST(WholeSceneShadowDepthDrawList);
 #undef DUMP_DRAW_LIST
diff --git a/Engine/Source/Runtime/Renderer/Private/SceneHitProxyRendering.cpp b/Engine/Source/Runtime/Renderer/Private/SceneHitProxyRendering.cpp
--- a/Engine/Source/Runtime/Renderer/Private/SceneHitProxyRendering.cpp
+++ b/Engine/Source/Runtime/Renderer/Private/SceneHitProxyRendering.cpp
@@ -152,8 +152,7 @@ IMPLEMENT_MATERIAL_SHADER_TYPE(,FHitProxyPS,TEXT("HitProxyPixelShader"),TEXT("Ma
 FHitProxyDrawingPolicy::FHitProxyDrawingPolicy(
 	const FVertexFactory* InVertexFactory,
 	const FMaterialRenderProxy* InMaterialRenderProxy,
-	ERHIFeatureLevel::Type InFeatureLevel
-	):
+	ERHIFeatureLevel::Type InFeatureLevel ):
 	FMeshDrawingPolicy(InVertexFactory, InMaterialRenderProxy, *InMaterialRenderProxy->GetMaterial(InFeatureLevel))
 {
 	HullShader = NULL;
@@ -243,6 +242,54 @@ void FHitProxyDrawingPolicy::SetMeshRenderState(
 		));
 }
 
+#if WITH_EDITOR
+
+int32 FEditorSelectionDrawingPolicy::PrimitiveSelectionIndex;
+int32 FEditorSelectionDrawingPolicy::IndividuallySelectedProxyIndex;
+
+void FEditorSelectionDrawingPolicy::SetMeshRenderState( FRHICommandList& RHICmdList, const FSceneView& View, const FPrimitiveSceneProxy* PrimitiveSceneProxy, const FMeshBatch& Mesh, int32 BatchElementIndex, bool bBackFace, const FMeshDrawingRenderState& DrawRenderState, const FHitProxyId HitProxyId, const ContextDataType PolicyContext )
+{
+	FHitProxyDrawingPolicy::SetMeshRenderState(RHICmdList, View, PrimitiveSceneProxy, Mesh, BatchElementIndex, bBackFace, DrawRenderState, HitProxyId, PolicyContext);
+
+	int32 StencilValue = GetStencilValue(View, PrimitiveSceneProxy);
+		
+	RHICmdList.SetDepthStencilState(TStaticDepthStencilState<true, CF_DepthNearOrEqual, true, CF_Always, SO_Keep, SO_Keep, SO_Replace>::GetRHI(), StencilValue);
+
+	//RHICmdList.SetRasterizerState(TStaticRasterizerState<>::GetRHI());
+	RHICmdList.SetBlendState(TStaticBlendStateWriteMask<CW_NONE, CW_NONE, CW_NONE, CW_NONE>::GetRHI());
+
+}
+
+int32 FEditorSelectionDrawingPolicy::GetStencilValue(const FSceneView& View, const FPrimitiveSceneProxy* PrimitiveSceneProxy)
+{
+	const bool bActorSelectionColorIsSubdued = View.bHasSelectedComponents;
+
+	int32 StencilValue = 0;
+	if(PrimitiveSceneProxy->IsIndividuallySelected())
+	{
+		// Any component that is individually selected should have a stencil value of < 128 so that it can have a unique color.  We offset the value by 2 because 0 means no selection and 1 is for bsp
+		StencilValue = IndividuallySelectedProxyIndex % 126 + 2;
+		++IndividuallySelectedProxyIndex;
+	}
+	else
+	{
+			
+		// If we are subduing actor color highlight then use the top level bits to indicate that to the shader.  
+		StencilValue = bActorSelectionColorIsSubdued ? PrimitiveSelectionIndex % 128 + 128 : PrimitiveSelectionIndex % 126 + 2;
+		++PrimitiveSelectionIndex;
+	}
+
+	return StencilValue;
+}
+
+void FEditorSelectionDrawingPolicy::ResetStencilValues()
+{
+	PrimitiveSelectionIndex = 0;
+	IndividuallySelectedProxyIndex = 0;
+}
+
+#endif
+
 void FHitProxyDrawingPolicyFactory::AddStaticMesh(FScene* Scene,FStaticMesh* StaticMesh,ContextType)
 {
 	checkSlow( Scene->RequiresHitProxies() );
@@ -264,6 +311,13 @@ void FHitProxyDrawingPolicyFactory::AddStaticMesh(FScene* Scene,FStaticMesh* Sta
 		);
 
 #if WITH_EDITOR
+
+	Scene->EditorSelectionDrawList.AddMesh(
+		StaticMesh,
+		StaticMesh->BatchHitProxyId,
+		FEditorSelectionDrawingPolicy(StaticMesh->VertexFactory, MaterialRenderProxy, Scene->GetFeatureLevel()),
+		Scene->GetFeatureLevel());
+
 	// If the mesh isn't translucent then we'll also add it to the "opaque-only" draw list.  Depending
 	// on user preferences in the editor, we may use this draw list to disallow selection of
 	// translucent objects in perspective viewports
diff --git a/Engine/Source/Runtime/Renderer/Private/SceneHitProxyRendering.h b/Engine/Source/Runtime/Renderer/Private/SceneHitProxyRendering.h
--- a/Engine/Source/Runtime/Renderer/Private/SceneHitProxyRendering.h
+++ b/Engine/Source/Runtime/Renderer/Private/SceneHitProxyRendering.h
@@ -60,6 +60,40 @@ private:
 	class FHitProxyDS* DomainShader;
 };
 
+#if WITH_EDITOR
+class FEditorSelectionDrawingPolicy : public FHitProxyDrawingPolicy
+{
+public:
+	FEditorSelectionDrawingPolicy(const FVertexFactory* InVertexFactory, const FMaterialRenderProxy* InMaterialRenderProxy, ERHIFeatureLevel::Type InFeatureLevel)
+		: FHitProxyDrawingPolicy(InVertexFactory, InMaterialRenderProxy, InFeatureLevel)
+	{}
+
+	// FMeshDrawingPolicy interface.
+	void SetMeshRenderState( FRHICommandList& RHICmdList, const FSceneView& View, const FPrimitiveSceneProxy* PrimitiveSceneProxy, const FMeshBatch& Mesh, int32 BatchElementIndex, bool bBackFace, const FMeshDrawingRenderState& DrawRenderState, const FHitProxyId HitProxyId, const ContextDataType PolicyContext );
+
+	/**
+	 * Gets the value that should be written into the editor selection stencil buffer for a given primitive
+	 * There will be a unique stencil value for each primitive until the max stencil buffer value is written and then the values will repeat
+	 *
+	 * @param View					The view being rendered
+	 * @param PrimitiveSceneProxy	The scene proxy for the primitive being rendered
+	 * @return the stencil value that should be written
+	 */
+	static int32 GetStencilValue(const FSceneView& View, const FPrimitiveSceneProxy* PrimitiveSceneProxy);
+
+	/**
+	 * Resets all unique stencil values
+	 */
+	static void ResetStencilValues();
+
+private:
+	/** Current selection index for a primitive */
+	static int32 PrimitiveSelectionIndex;
+	/** Current selection index used when components are selected directly */
+	static int32 IndividuallySelectedProxyIndex;
+};
+#endif
+
 /**
  * A drawing policy factory for the hit proxy drawing policy.
  */
diff --git a/Engine/Source/Runtime/Renderer/Private/ScenePrivate.h b/Engine/Source/Runtime/Renderer/Private/ScenePrivate.h
--- a/Engine/Source/Runtime/Renderer/Private/ScenePrivate.h
+++ b/Engine/Source/Runtime/Renderer/Private/ScenePrivate.h
@@ -1810,6 +1810,10 @@ public:
 	template<typename LightMapPolicyType>
 	TStaticMeshDrawList<TMobileBasePassDrawingPolicy<LightMapPolicyType, 0> >& GetMobileBasePassCSMDrawList(EBasePassDrawListType DrawType);
 
+#if WITH_EDITOR
+	/** Draw list to use for selected static meshes in the editor only */
+	TStaticMeshDrawList<FEditorSelectionDrawingPolicy> EditorSelectionDrawList;
+#endif
 	/**
 	 * The following arrays are densely packed primitive data needed by various
 	 * rendering passes. PrimitiveSceneInfo->PackedIndex maintains the index
diff --git a/Engine/Source/Runtime/Renderer/Private/SceneRendering.h b/Engine/Source/Runtime/Renderer/Private/SceneRendering.h
--- a/Engine/Source/Runtime/Renderer/Private/SceneRendering.h
+++ b/Engine/Source/Runtime/Renderer/Private/SceneRendering.h
@@ -653,6 +653,11 @@ public:
 	/** A map from static mesh ID to a boolean dithered LOD fade in value. */
 	FSceneBitArray StaticMeshFadeInDitheredLODMap;
 
+#if WITH_EDITOR
+	/** A map from static mesh ID to editor selection visibility (whether or not it is selected AND should be drawn).  */
+	FSceneBitArray StaticMeshEditorSelectionMap;
+#endif
+
 	/** An array of batch element visibility masks, valid only for meshes
 	 set visible in either StaticMeshVisibilityMap or StaticMeshShadowDepthMap. */
 	TArray<uint64,SceneRenderingAllocator> StaticMeshBatchVisibility;
diff --git a/Engine/Source/Runtime/Renderer/Private/SceneVisibility.cpp b/Engine/Source/Runtime/Renderer/Private/SceneVisibility.cpp
--- a/Engine/Source/Runtime/Renderer/Private/SceneVisibility.cpp
+++ b/Engine/Source/Runtime/Renderer/Private/SceneVisibility.cpp
@@ -1366,6 +1366,7 @@ namespace EMarkMaskBits
 		StaticMeshOccluderMapMask = 0x8,
 		StaticMeshFadeOutDitheredLODMapMask = 0x10,
 		StaticMeshFadeInDitheredLODMapMask = 0x20,
+		StaticMeshEditorSelectedMask = 0x40,
 	};
 }
 
@@ -1455,6 +1456,7 @@ struct FRelevancePacket
 			const bool bDynamicRelevance = ViewRelevance.bDynamicRelevance;
 			const bool bShadowRelevance = ViewRelevance.bShadowRelevance;
 			const bool bEditorRelevance = ViewRelevance.bEditorPrimitiveRelevance;
+			const bool bEditorSelectionRelevance = ViewRelevance.bEditorStaticSelectionRelevance;
 			const bool bTranslucentRelevance = ViewRelevance.HasTranslucency();
 
 			if (bStaticRelevance && (bDrawRelevance || bShadowRelevance))
@@ -1665,6 +1667,13 @@ struct FRelevancePacket
 						}
 						bNeedsBatchVisibility = true;
 					}
+
+#if WITH_EDITOR
+					if(ViewRelevance.bDrawRelevance && ViewRelevance.bEditorStaticSelectionRelevance)
+					{
+						MarkMask |= EMarkMaskBits::StaticMeshEditorSelectedMask;
+					}
+#endif
 					if (MarkMask)
 					{
 						MarkMasks[StaticMesh.Id] = MarkMask;
@@ -1807,6 +1816,9 @@ static void ComputeAndMarkRelevanceForViewParallel(
 	uint32* RESTRICT StaticMeshOccluderMap_Words = View.StaticMeshOccluderMap.GetData();
 	uint32* RESTRICT StaticMeshFadeOutDitheredLODMap_Words = View.StaticMeshFadeOutDitheredLODMap.GetData();
 	uint32* RESTRICT StaticMeshFadeInDitheredLODMap_Words = View.StaticMeshFadeInDitheredLODMap.GetData();
+#if WITH_EDITOR
+	uint32* RESTRICT StaticMeshEditorSelectionMap_Words = View.StaticMeshEditorSelectionMap.GetData();
+#endif
 	const uint64* RESTRICT MarkMasks64 = (const uint64* RESTRICT)MarkMasks;
 	const uint8* RESTRICT MarkMasks8 = MarkMasks;
 	for (int32 BaseIndex = 0; BaseIndex < NumMesh; BaseIndex += 32)
@@ -1817,6 +1829,7 @@ static void ComputeAndMarkRelevanceForViewParallel(
 		uint32 StaticMeshOccluderMap_Word = 0;
 		uint32 StaticMeshFadeOutDitheredLODMap_Word = 0;
 		uint32 StaticMeshFadeInDitheredLODMap_Word = 0;
+		uint32 StaticMeshEditorSelectionMap_Word = 0;
 		uint32 Mask = 1;
 		bool bAny = false;
 		for (int32 QWordIndex = 0; QWordIndex < 4; QWordIndex++)
@@ -1832,6 +1845,9 @@ static void ComputeAndMarkRelevanceForViewParallel(
 					StaticMeshOccluderMap_Word |= (MaskMask & EMarkMaskBits::StaticMeshOccluderMapMask) ? Mask : 0;
 					StaticMeshFadeOutDitheredLODMap_Word |= (MaskMask & EMarkMaskBits::StaticMeshFadeOutDitheredLODMapMask) ? Mask : 0;
 					StaticMeshFadeInDitheredLODMap_Word |= (MaskMask & EMarkMaskBits::StaticMeshFadeInDitheredLODMapMask) ? Mask : 0;
+#if WITH_EDITOR
+					StaticMeshEditorSelectionMap_Word |= (MaskMask & EMarkMaskBits::StaticMeshEditorSelectedMask) ? Mask : 0;
+#endif
 				}
 				bAny = true;
 			}
@@ -1850,6 +1866,9 @@ static void ComputeAndMarkRelevanceForViewParallel(
 			*StaticMeshOccluderMap_Words = StaticMeshOccluderMap_Word;
 			*StaticMeshFadeOutDitheredLODMap_Words = StaticMeshFadeOutDitheredLODMap_Word;
 			*StaticMeshFadeInDitheredLODMap_Words = StaticMeshFadeInDitheredLODMap_Word;
+#if WITH_EDITOR
+			*StaticMeshEditorSelectionMap_Words = StaticMeshEditorSelectionMap_Word;
+#endif
 		}
 		StaticMeshVisibilityMap_Words++;
 		StaticMeshVelocityMap_Words++;
@@ -1857,6 +1876,9 @@ static void ComputeAndMarkRelevanceForViewParallel(
 		StaticMeshOccluderMap_Words++;
 		StaticMeshFadeOutDitheredLODMap_Words++;
 		StaticMeshFadeInDitheredLODMap_Words++;
+#if WITH_EDITOR
+		StaticMeshEditorSelectionMap_Words++;
+#endif
 	}
 }
 
@@ -2347,6 +2369,10 @@ void FSceneRenderer::ComputeViewVisibility(FRHICommandListImmediate& RHICmdList)
 
 		View.VisibleLightInfos.Empty(Scene->Lights.GetMaxIndex());
 
+#if WITH_EDITOR
+		View.StaticMeshEditorSelectionMap.Init(false, Scene->StaticMeshes.GetMaxIndex());
+#endif
+
 		// The dirty list allocation must take into account the max possible size because when GILCUpdatePrimTaskEnabled is true,
 		// the indirect lighting cache will be update on by threaded job, which can not do reallocs on the buffer (since it uses the SceneRenderingAllocator).
 		View.DirtyPrecomputedLightingBufferPrimitives.Reserve(Scene->Primitives.Num());
diff --git a/Engine/Source/Runtime/Slate/Private/Framework/Docking/SDockingTabWell.cpp b/Engine/Source/Runtime/Slate/Private/Framework/Docking/SDockingTabWell.cpp
--- a/Engine/Source/Runtime/Slate/Private/Framework/Docking/SDockingTabWell.cpp
+++ b/Engine/Source/Runtime/Slate/Private/Framework/Docking/SDockingTabWell.cpp
@@ -372,6 +372,14 @@ FReply SDockingTabWell::OnDrop( const FGeometry& MyGeometry, const FDragDropEven
 	{
 		if (DragDropOperation->CanDockInNode(ParentTabStackPtr.Pin().ToSharedRef(), FDockingDragOperation::DockingViaTabWell))
 		{
+			// It's rare, but sometimes a drop operation can happen after we drag a tab out of a docking tab well but before the engine has a
+			// chance to notify the next docking tab well that a drag operation has entered it. In this case, just use the tab referenced by the
+			// drag/drop operation
+			if (!TabBeingDraggedPtr.IsValid())
+			{
+				TabBeingDraggedPtr = DragDropOperation->GetTabBeingDragged();
+			}
+			
 			if ( ensure( TabBeingDraggedPtr.IsValid() ) )
 			{
 				// We dropped a Tab into this TabWell.
diff --git a/Engine/Source/Runtime/Slate/Private/Framework/Text/PlainTextLayoutMarshaller.cpp b/Engine/Source/Runtime/Slate/Private/Framework/Text/PlainTextLayoutMarshaller.cpp
--- a/Engine/Source/Runtime/Slate/Private/Framework/Text/PlainTextLayoutMarshaller.cpp
+++ b/Engine/Source/Runtime/Slate/Private/Framework/Text/PlainTextLayoutMarshaller.cpp
@@ -4,6 +4,7 @@
 #include "PlainTextLayoutMarshaller.h"
 #include "SlateTextLayout.h"
 #include "SlatePasswordRun.h"
+#include "SlateTextUnderlineLineHighlighter.h"
 
 TSharedRef< FPlainTextLayoutMarshaller > FPlainTextLayoutMarshaller::Create()
 {
@@ -29,9 +30,17 @@ void FPlainTextLayoutMarshaller::SetText(const FString& SourceString, FTextLayou
 	TArray<FTextLayout::FNewLineData> LinesToAdd;
 	LinesToAdd.Reserve(LineRanges.Num());
 
+	TArray<FTextLineHighlight> LineHighlightsToAdd;
+	TSharedPtr<FSlateTextUnderlineLineHighlighter> UnderlineLineHighlighter;
+	if (!DefaultTextStyle.UnderlineBrush.GetResourceName().IsNone())
+	{
+		UnderlineLineHighlighter = FSlateTextUnderlineLineHighlighter::Create(DefaultTextStyle.UnderlineBrush, DefaultTextStyle.Font, DefaultTextStyle.ColorAndOpacity, DefaultTextStyle.ShadowOffset, DefaultTextStyle.ShadowColorAndOpacity);
+	}
+
 	const bool bUsePasswordRun = bIsPassword.Get(false);
-	for (const FTextRange& LineRange : LineRanges)
+	for (int32 LineIndex = 0; LineIndex < LineRanges.Num(); ++LineIndex)
 	{
+		const FTextRange& LineRange = LineRanges[LineIndex];
 		TSharedRef<FString> LineText = MakeShareable(new FString(SourceString.Mid(LineRange.BeginIndex, LineRange.Len())));
 
 		TArray<TSharedRef<IRun>> Runs;
@@ -44,10 +53,16 @@ void FPlainTextLayoutMarshaller::SetText(const FString& SourceString, FTextLayou
 			Runs.Add(FSlateTextRun::Create(FRunInfo(), LineText, DefaultTextStyle));
 		}
 
+		if (UnderlineLineHighlighter.IsValid())
+		{
+			LineHighlightsToAdd.Add(FTextLineHighlight(LineIndex, FTextRange(0, LineRange.Len()), FSlateTextUnderlineLineHighlighter::DefaultZIndex, UnderlineLineHighlighter.ToSharedRef()));
+		}
+
 		LinesToAdd.Emplace(MoveTemp(LineText), MoveTemp(Runs));
 	}
 
 	TargetTextLayout.AddLines(LinesToAdd);
+	TargetTextLayout.SetLineHighlights(LineHighlightsToAdd);
 }
 
 void FPlainTextLayoutMarshaller::GetText(FString& TargetString, const FTextLayout& SourceTextLayout)
diff --git a/Engine/Source/Runtime/Slate/Private/Framework/Text/RichTextLayoutMarshaller.cpp b/Engine/Source/Runtime/Slate/Private/Framework/Text/RichTextLayoutMarshaller.cpp
--- a/Engine/Source/Runtime/Slate/Private/Framework/Text/RichTextLayoutMarshaller.cpp
+++ b/Engine/Source/Runtime/Slate/Private/Framework/Text/RichTextLayoutMarshaller.cpp
@@ -4,6 +4,7 @@
 #include "RichTextLayoutMarshaller.h"
 #include "RichTextMarkupProcessing.h"
 #include "SlateTextLayout.h"
+#include "SlateTextUnderlineLineHighlighter.h"
 
 #if WITH_FANCY_TEXT
 
@@ -32,21 +33,27 @@ void FRichTextLayoutMarshaller::SetText(const FString& SourceString, FTextLayout
 	TArray<FTextLayout::FNewLineData> LinesToAdd;
 	LinesToAdd.Reserve(LineParseResultsArray.Num());
 
+	TArray<FTextLineHighlight> LineHighlightsToAdd;
+	TMap<const FTextBlockStyle*, TSharedPtr<FSlateTextUnderlineLineHighlighter>> CachedUnderlineHighlighters;
+
 	// Iterate through parsed line results and create processed lines with runs.
-	for (const FTextLineParseResults& LineParseResults : LineParseResultsArray)
+	for (int32 LineIndex = 0; LineIndex < LineParseResultsArray.Num(); ++LineIndex)
 	{
+		const FTextLineParseResults& LineParseResults = LineParseResultsArray[LineIndex];
+
 		TSharedRef<FString> ModelString = MakeShareable(new FString());
 		TArray< TSharedRef< IRun > > Runs;
 
 		for (const FTextRunParseResults& RunParseResult : LineParseResults.Runs)
 		{
-			AppendRunsForText(RunParseResult, ProcessedString, DefaultTextStyle, ModelString, TargetTextLayout, Runs);
+			AppendRunsForText(LineIndex, RunParseResult, ProcessedString, DefaultTextStyle, ModelString, TargetTextLayout, Runs, LineHighlightsToAdd, CachedUnderlineHighlighters);
 		}
 
 		LinesToAdd.Emplace(MoveTemp(ModelString), MoveTemp(Runs));
 	}
 
 	TargetTextLayout.AddLines(LinesToAdd);
+	TargetTextLayout.SetLineHighlights(LineHighlightsToAdd);
 }
 
 void FRichTextLayoutMarshaller::GetText(FString& TargetString, const FTextLayout& SourceTextLayout)
@@ -113,12 +120,17 @@ TSharedPtr< ITextDecorator > FRichTextLayoutMarshaller::TryGetDecorator(const FS
 	return nullptr;
 }
 
-void FRichTextLayoutMarshaller::AppendRunsForText(const FTextRunParseResults& TextRun,
+void FRichTextLayoutMarshaller::AppendRunsForText(
+	const int32 LineIndex,
+	const FTextRunParseResults& TextRun,
 	const FString& ProcessedString,
 	const FTextBlockStyle& DefaultTextStyle,
 	const TSharedRef<FString>& InOutModelText,
 	FTextLayout& TargetTextLayout,
-	TArray<TSharedRef<IRun>>& Runs)
+	TArray<TSharedRef<IRun>>& Runs,
+	TArray<FTextLineHighlight>& LineHighlights,
+	TMap<const FTextBlockStyle*, TSharedPtr<FSlateTextUnderlineLineHighlighter>>& CachedUnderlineHighlighters
+	)
 {
 	TSharedPtr< ISlateRun > Run;
 	TSharedPtr< ITextDecorator > Decorator = TryGetDecorator(ProcessedString, TextRun);
@@ -153,6 +165,18 @@ void FRichTextLayoutMarshaller::AppendRunsForText(const FTextRunParseResults& Te
 
 		// Create run.
 		Run = FSlateTextRun::Create(RunInfo, InOutModelText, *TextBlockStyle, ModelRange);
+
+		if (!TextBlockStyle->UnderlineBrush.GetResourceName().IsNone())
+		{
+			TSharedPtr<FSlateTextUnderlineLineHighlighter> UnderlineLineHighlighter = CachedUnderlineHighlighters.FindRef(TextBlockStyle);
+			if (!UnderlineLineHighlighter.IsValid())
+			{
+				UnderlineLineHighlighter = FSlateTextUnderlineLineHighlighter::Create(TextBlockStyle->UnderlineBrush, TextBlockStyle->Font, TextBlockStyle->ColorAndOpacity, TextBlockStyle->ShadowOffset, TextBlockStyle->ShadowColorAndOpacity);
+				CachedUnderlineHighlighters.Add(TextBlockStyle, UnderlineLineHighlighter);
+			}
+
+			LineHighlights.Add(FTextLineHighlight(LineIndex, ModelRange, FSlateTextUnderlineLineHighlighter::DefaultZIndex, UnderlineLineHighlighter.ToSharedRef()));
+		}
 	}
 
 	Runs.Add(Run.ToSharedRef());
diff --git a/Engine/Source/Runtime/Slate/Private/Framework/Text/SlateTextUnderlineLineHighlighter.cpp b/Engine/Source/Runtime/Slate/Private/Framework/Text/SlateTextUnderlineLineHighlighter.cpp
new file mode 100644
--- /dev/null
+++ b/Engine/Source/Runtime/Slate/Private/Framework/Text/SlateTextUnderlineLineHighlighter.cpp
@@ -0,0 +1,79 @@
+// Copyright 1998-2016 Epic Games, Inc. All Rights Reserved.
+
+#include "SlatePrivatePCH.h"
+#include "SlateTextUnderlineLineHighlighter.h"
+
+FSlateTextUnderlineLineHighlighter::FSlateTextUnderlineLineHighlighter(const FSlateBrush& InUnderlineBrush, const FSlateFontInfo& InFontInfo, const FSlateColor InColorAndOpacity, const FVector2D InShadowOffset, const FLinearColor InShadowColorAndOpacity)
+	: UnderlineBrush(InUnderlineBrush)
+	, FontInfo(InFontInfo)
+	, ColorAndOpacity(InColorAndOpacity)
+	, ShadowOffset(InShadowOffset)
+	, ShadowColorAndOpacity(InShadowColorAndOpacity)
+{
+}
+
+int32 FSlateTextUnderlineLineHighlighter::OnPaint(const FPaintArgs& Args, const FTextLayout::FLineView& Line, const float OffsetX, const float Width, const FTextBlockStyle& DefaultStyle, const FGeometry& AllottedGeometry, const FSlateRect& MyClippingRect, FSlateWindowElementList& OutDrawElements, int32 LayerId, const FWidgetStyle& InWidgetStyle, bool bParentEnabled) const
+{
+	TSharedRef<FSlateFontCache> FontCache = FSlateApplication::Get().GetRenderer()->GetFontCache();
+
+	const uint16 MaxHeight = FontCache->GetMaxCharacterHeight(FontInfo, AllottedGeometry.Scale);
+	const int16 Baseline = FontCache->GetBaseline(FontInfo, AllottedGeometry.Scale);
+
+	int16 UnderlinePos, UnderlineThickness;
+	FontCache->GetUnderlineMetrics(FontInfo, AllottedGeometry.Scale, UnderlinePos, UnderlineThickness);
+
+	const FVector2D Location(Line.Offset.X + OffsetX, Line.Offset.Y + MaxHeight + Baseline - (UnderlinePos * 0.5f));
+	const FVector2D Size(Width, FMath::Max<int16>(1, UnderlineThickness));
+
+	// The block size and offset values are pre-scaled, so we need to account for that when converting the block offsets into paint geometry
+	const float InverseScale = Inverse(AllottedGeometry.Scale);
+
+	if (Size.X)
+	{
+		const FLinearColor LineColorAndOpacity = ColorAndOpacity.GetColor(InWidgetStyle);
+
+		const bool ShouldDropShadow = ShadowColorAndOpacity.A > 0.f && ShadowOffset.SizeSquared() > 0.f;
+
+		// A negative shadow offset should be applied as a positive offset to the underline to avoid clipping issues
+		const FVector2D DrawShadowOffset(
+			(ShadowOffset.X > 0.0f) ? ShadowOffset.X * AllottedGeometry.Scale : 0.0f,
+			(ShadowOffset.Y > 0.0f) ? ShadowOffset.Y * AllottedGeometry.Scale : 0.0f
+			);
+		const FVector2D DrawUnderlineOffset(
+			(ShadowOffset.X < 0.0f) ? -ShadowOffset.X * AllottedGeometry.Scale : 0.0f,
+			(ShadowOffset.Y < 0.0f) ? -ShadowOffset.Y * AllottedGeometry.Scale : 0.0f
+			);
+
+		// Draw the optional shadow
+		if (ShouldDropShadow)
+		{
+			FSlateDrawElement::MakeBox(
+				OutDrawElements,
+				++LayerId,
+				AllottedGeometry.ToPaintGeometry(TransformVector(InverseScale, Size), FSlateLayoutTransform(TransformPoint(InverseScale, Location + DrawShadowOffset))),
+				&UnderlineBrush,
+				MyClippingRect,
+				bParentEnabled ? ESlateDrawEffect::None : ESlateDrawEffect::DisabledEffect,
+				ShadowColorAndOpacity * InWidgetStyle.GetColorAndOpacityTint()
+				);
+		}
+
+		// Draw underline
+		FSlateDrawElement::MakeBox(
+			OutDrawElements,
+			++LayerId,
+			AllottedGeometry.ToPaintGeometry(TransformVector(InverseScale, Size), FSlateLayoutTransform(TransformPoint(InverseScale, Location + DrawUnderlineOffset))),
+			&UnderlineBrush,
+			MyClippingRect,
+			bParentEnabled ? ESlateDrawEffect::None : ESlateDrawEffect::DisabledEffect,
+			LineColorAndOpacity * InWidgetStyle.GetColorAndOpacityTint()
+			);
+	}
+
+	return LayerId;
+}
+
+TSharedRef<FSlateTextUnderlineLineHighlighter> FSlateTextUnderlineLineHighlighter::Create(const FSlateBrush& InUnderlineBrush, const FSlateFontInfo& InFontInfo, const FSlateColor InColorAndOpacity, const FVector2D InShadowOffset, const FLinearColor InShadowColorAndOpacity)
+{
+	return MakeShareable(new FSlateTextUnderlineLineHighlighter(InUnderlineBrush, InFontInfo, InColorAndOpacity, InShadowOffset, InShadowColorAndOpacity));
+}
diff --git a/Engine/Source/Runtime/Slate/Private/Framework/Text/SyntaxHighlighterTextLayoutMarshaller.cpp b/Engine/Source/Runtime/Slate/Private/Framework/Text/SyntaxHighlighterTextLayoutMarshaller.cpp
--- a/Engine/Source/Runtime/Slate/Private/Framework/Text/SyntaxHighlighterTextLayoutMarshaller.cpp
+++ b/Engine/Source/Runtime/Slate/Private/Framework/Text/SyntaxHighlighterTextLayoutMarshaller.cpp
@@ -2,6 +2,7 @@
 
 #include "SlatePrivatePCH.h"
 #include "SyntaxHighlighterTextLayoutMarshaller.h"
+#include "SlateTextUnderlineLineHighlighter.h"
 
 #if WITH_FANCY_TEXT
 
@@ -76,10 +77,15 @@ void FRichTextSyntaxHighlighterTextLayoutMarshaller::ParseTokens(const FString&
 	TArray<FTextLayout::FNewLineData> LinesToAdd;
 	LinesToAdd.Reserve(TokenizedLines.Num());
 
+	TArray<FTextLineHighlight> LineHighlightsToAdd;
+	TMap<const FTextBlockStyle*, TSharedPtr<FSlateTextUnderlineLineHighlighter>> CachedUnderlineHighlighters;
+
 	// Parse the tokens, generating the styled runs for each line
 	EParseState ParseState = EParseState::LookingForNode;
-	for(const FSyntaxTokenizer::FTokenizedLine& TokenizedLine : TokenizedLines)
+	for (int32 LineIndex = 0; LineIndex < TokenizedLines.Num(); ++LineIndex)
 	{
+		const FSyntaxTokenizer::FTokenizedLine& TokenizedLine = TokenizedLines[LineIndex];
+
 		TSharedRef<FString> ModelString = MakeShareable(new FString());
 		TArray< TSharedRef< IRun > > Runs;
 
@@ -91,7 +97,7 @@ void FRichTextSyntaxHighlighterTextLayoutMarshaller::ParseTokens(const FString&
 			ModelString->Append(TokenText);
 
 			FRunInfo RunInfo(TEXT("SyntaxHighlight.Normal"));
-			FTextBlockStyle TextBlockStyle = SyntaxTextStyle.NormalTextStyle;
+			const FTextBlockStyle* TextBlockStyle = &SyntaxTextStyle.NormalTextStyle;
 
 			const bool bIsWhitespace = FString(TokenText).TrimTrailing().IsEmpty();
 			if(!bIsWhitespace)
@@ -102,41 +108,41 @@ void FRichTextSyntaxHighlighterTextLayoutMarshaller::ParseTokens(const FString&
 					if(ParseState == EParseState::LookingForNode && TokenText == TEXT("<"))
 					{
 						RunInfo.Name = TEXT("SyntaxHighlight.Node");
-						TextBlockStyle = SyntaxTextStyle.NodeTextStyle;
+						TextBlockStyle = &SyntaxTextStyle.NodeTextStyle;
 						ParseState = EParseState::LookingForNodeName;
 						bHasMatchedSyntax = true;
 					}
 					else if(ParseState == EParseState::LookingForNodeAttributeKey && TokenText == TEXT(">"))
 					{
 						RunInfo.Name = TEXT("SyntaxHighlight.Node");
-						TextBlockStyle = SyntaxTextStyle.NodeTextStyle;
+						TextBlockStyle = &SyntaxTextStyle.NodeTextStyle;
 						ParseState = EParseState::LookingForNode;
 					}
 					else if(ParseState == EParseState::LookingForNode && TokenText == TEXT("</>"))
 					{
 						RunInfo.Name = TEXT("SyntaxHighlight.Node");
-						TextBlockStyle = SyntaxTextStyle.NodeTextStyle;
+						TextBlockStyle = &SyntaxTextStyle.NodeTextStyle;
 						// No state change
 						bHasMatchedSyntax = true;
 					}
 					else if(ParseState == EParseState::LookingForNodeAttributeKey && TokenText == TEXT("="))
 					{
 						RunInfo.Name = TEXT("SyntaxHighlight.NodeAttribueAssignment");
-						TextBlockStyle = SyntaxTextStyle.NodeAttribueAssignmentTextStyle;
+						TextBlockStyle = &SyntaxTextStyle.NodeAttribueAssignmentTextStyle;
 						ParseState = EParseState::LookingForNodeAttribueValueBegin;
 						bHasMatchedSyntax = true;
 					}
 					else if(ParseState == EParseState::LookingForNodeAttribueValueBegin && TokenText == TEXT("\""))
 					{
 						RunInfo.Name = TEXT("SyntaxHighlight.NodeAttributeValue");
-						TextBlockStyle = SyntaxTextStyle.NodeAttributeValueTextStyle;
+						TextBlockStyle = &SyntaxTextStyle.NodeAttributeValueTextStyle;
 						ParseState = EParseState::LookingForNodeAttribueValueBody;
 						bHasMatchedSyntax = true;
 					}
 					else if(ParseState == EParseState::LookingForNodeAttribueValueBody && TokenText == TEXT("\""))
 					{
 						RunInfo.Name = TEXT("SyntaxHighlight.NodeAttributeValue");
-						TextBlockStyle = SyntaxTextStyle.NodeAttributeValueTextStyle;
+						TextBlockStyle = &SyntaxTextStyle.NodeAttributeValueTextStyle;
 						ParseState = EParseState::LookingForNodeAttributeKey;
 						bHasMatchedSyntax = true;
 					}
@@ -149,32 +155,45 @@ void FRichTextSyntaxHighlighterTextLayoutMarshaller::ParseTokens(const FString&
 					if(ParseState == EParseState::LookingForNodeName)
 					{
 						RunInfo.Name = TEXT("SyntaxHighlight.Node");
-						TextBlockStyle = SyntaxTextStyle.NodeTextStyle;
+						TextBlockStyle = &SyntaxTextStyle.NodeTextStyle;
 						ParseState = EParseState::LookingForNodeAttributeKey;
 					}
 					else if(ParseState == EParseState::LookingForNodeAttributeKey)
 					{
 						RunInfo.Name = TEXT("SyntaxHighlight.NodeAttributeKey");
-						TextBlockStyle = SyntaxTextStyle.NodeAttributeKeyTextStyle;
+						TextBlockStyle = &SyntaxTextStyle.NodeAttributeKeyTextStyle;
 						// No state change, a key can be multiple tokens - consume until we find an equals
 					}
 					else if(ParseState == EParseState::LookingForNodeAttribueValueBody)
 					{
 						RunInfo.Name = TEXT("SyntaxHighlight.NodeAttributeValue");
-						TextBlockStyle = SyntaxTextStyle.NodeAttributeValueTextStyle;
+						TextBlockStyle = &SyntaxTextStyle.NodeAttributeValueTextStyle;
 						// No state change, a value can be multiple tokens - consume until we find a closing quote
 					}
 				}
 			}
 
-			TSharedRef< ISlateRun > Run = FSlateTextRun::Create(RunInfo, ModelString, TextBlockStyle, ModelRange);
+			TSharedRef< ISlateRun > Run = FSlateTextRun::Create(RunInfo, ModelString, *TextBlockStyle, ModelRange);
 			Runs.Add(Run);
+
+			if (!TextBlockStyle->UnderlineBrush.GetResourceName().IsNone())
+			{
+				TSharedPtr<FSlateTextUnderlineLineHighlighter> UnderlineLineHighlighter = CachedUnderlineHighlighters.FindRef(TextBlockStyle);
+				if (!UnderlineLineHighlighter.IsValid())
+				{
+					UnderlineLineHighlighter = FSlateTextUnderlineLineHighlighter::Create(TextBlockStyle->UnderlineBrush, TextBlockStyle->Font, TextBlockStyle->ColorAndOpacity, TextBlockStyle->ShadowOffset, TextBlockStyle->ShadowColorAndOpacity);
+					CachedUnderlineHighlighters.Add(TextBlockStyle, UnderlineLineHighlighter);
+				}
+
+				LineHighlightsToAdd.Add(FTextLineHighlight(LineIndex, ModelRange, FSlateTextUnderlineLineHighlighter::DefaultZIndex, UnderlineLineHighlighter.ToSharedRef()));
+			}
 		}
 
 		LinesToAdd.Emplace(MoveTemp(ModelString), MoveTemp(Runs));
 	}
 
 	TargetTextLayout.AddLines(LinesToAdd);
+	TargetTextLayout.SetLineHighlights(LineHighlightsToAdd);
 }
 
 FRichTextSyntaxHighlighterTextLayoutMarshaller::FRichTextSyntaxHighlighterTextLayoutMarshaller(TSharedPtr< FSyntaxTokenizer > InTokenizer, const FSyntaxTextStyle& InSyntaxTextStyle)
diff --git a/Engine/Source/Runtime/Slate/Private/Framework/Text/TextLayout.cpp b/Engine/Source/Runtime/Slate/Private/Framework/Text/TextLayout.cpp
--- a/Engine/Source/Runtime/Slate/Private/Framework/Text/TextLayout.cpp
+++ b/Engine/Source/Runtime/Slate/Private/Framework/Text/TextLayout.cpp
@@ -1199,13 +1199,13 @@ void FTextLayout::AddRunRenderer( const FTextRunRenderer& Renderer )
 		if ( LineModel.RunRenderers[ Index ].Range.BeginIndex > Renderer.Range.BeginIndex )
 		{
 			checkf(Index == 0 || LineModel.RunRenderers[Index - 1].Range.EndIndex <= Renderer.Range.BeginIndex, TEXT("Renderers cannot overlap!\n\tDebug Source: %s"), *DebugSourceInfo.Get(FString()));
-			LineModel.RunRenderers.Insert( Renderer, Index - 1 );
+			LineModel.RunRenderers.Insert( Renderer, FMath::Max(0, Index - 1) );
 			bWasInserted = true;
 		}
 		else if ( LineModel.RunRenderers[ Index ].Range.EndIndex > Renderer.Range.EndIndex )
 		{
 			checkf(LineModel.RunRenderers[Index].Range.BeginIndex >= Renderer.Range.EndIndex, TEXT("Renderers cannot overlap!\n\tDebug Source: %s"), *DebugSourceInfo.Get(FString()));
-			LineModel.RunRenderers.Insert( Renderer, Index - 1 );
+			LineModel.RunRenderers.Insert( Renderer, FMath::Max(0, Index - 1) );
 			bWasInserted = true;
 		}
 	}
@@ -1218,6 +1218,29 @@ void FTextLayout::AddRunRenderer( const FTextRunRenderer& Renderer )
 	DirtyFlags |= ETextLayoutDirtyState::Layout;
 }
 
+void FTextLayout::RemoveRunRenderer( const FTextRunRenderer& Renderer )
+{
+	checkf(LineModels.IsValidIndex(Renderer.LineIndex), TEXT("Renderers must be for a valid Line Index!\n\tDebug Source: %s"), *DebugSourceInfo.Get(FString()));
+
+	FLineModel& LineModel = LineModels[Renderer.LineIndex];
+
+	bool bWasRemoved = false;
+	for (int32 Index = 0; Index < LineModel.RunRenderers.Num(); ++Index)
+	{
+		if (LineModel.RunRenderers[Index] == Renderer)
+		{
+			LineModel.RunRenderers.RemoveAt(Index, 1, /*bAllowShrinking*/false);
+			bWasRemoved = true;
+			break;
+		}
+	}
+
+	if (bWasRemoved)
+	{
+		DirtyFlags |= ETextLayoutDirtyState::Layout;
+	}
+}
+
 void FTextLayout::ClearLineHighlights()
 {
 	for (int32 Index = 0; Index < LineModels.Num(); Index++)
@@ -1253,7 +1276,7 @@ void FTextLayout::AddLineHighlight( const FTextLineHighlight& Highlight )
 	{
 		if ( LineModel.LineHighlights[ Index ].ZOrder > Highlight.ZOrder )
 		{
-			LineModel.LineHighlights.Insert( Highlight, Index - 1 );
+			LineModel.LineHighlights.Insert( Highlight, FMath::Max(0, Index - 1) );
 			bWasInserted = true;
 		}
 	}
@@ -1266,6 +1289,30 @@ void FTextLayout::AddLineHighlight( const FTextLineHighlight& Highlight )
 	DirtyFlags |= ETextLayoutDirtyState::Highlights;
 }
 
+void FTextLayout::RemoveLineHighlight( const FTextLineHighlight& Highlight )
+{
+	checkf(LineModels.IsValidIndex(Highlight.LineIndex), TEXT("Highlights must be for a valid Line Index!\n\tDebug Source: %s"), *DebugSourceInfo.Get(FString()));
+	checkf(Highlight.ZOrder, TEXT("The highlight Z-order must be <0 to create an underlay, or >0 to create an overlay!\n\tDebug Source: %s"), *DebugSourceInfo.Get(FString()));
+
+	FLineModel& LineModel = LineModels[Highlight.LineIndex];
+
+	bool bWasRemoved = false;
+	for (int32 Index = 0; Index < LineModel.LineHighlights.Num(); ++Index)
+	{
+		if (LineModel.LineHighlights[Index] == Highlight)
+		{
+			LineModel.LineHighlights.RemoveAt(Index, 1, /*bAllowShrinking*/false);
+			bWasRemoved = true;
+			break;
+		}
+	}
+
+	if (bWasRemoved)
+	{
+		DirtyFlags |= ETextLayoutDirtyState::Highlights;
+	}
+}
+
 FTextLocation FTextLayout::GetTextLocationAt( const FLineView& LineView, const FVector2D& Relative, ETextHitPoint* const OutHitPoint ) const
 {
 	for (int32 BlockIndex = 0; BlockIndex < LineView.Blocks.Num(); BlockIndex++)
diff --git a/Engine/Source/Runtime/Slate/Private/Widgets/Text/SlateEditableTextLayout.cpp b/Engine/Source/Runtime/Slate/Private/Widgets/Text/SlateEditableTextLayout.cpp
--- a/Engine/Source/Runtime/Slate/Private/Widgets/Text/SlateEditableTextLayout.cpp
+++ b/Engine/Source/Runtime/Slate/Private/Widgets/Text/SlateEditableTextLayout.cpp
@@ -328,6 +328,9 @@ bool FSlateEditableTextLayout::SetEditableText(const FText& TextToSet, const boo
 		ClearSelection();
 		TextLayout->ClearLines();
 
+		TextLayout->ClearLineHighlights();
+		TextLayout->ClearRunRenderers();
+
 		Marshaller->SetText(TextToSetString, *TextLayout);
 
 		const TArray< FTextLayout::FLineModel >& Lines = TextLayout->GetLineModels();
@@ -2248,9 +2251,9 @@ void FSlateEditableTextLayout::UpdateCursorHighlight()
 
 	RemoveCursorHighlight();
 
-	static const int32 SelectionHighlightZOrder = -1; // draw below the text
-	static const int32 CompositionRangeZOrder = 1; // draw above the text
-	static const int32 CursorZOrder = 2; // draw above the text and the composition
+	static const int32 SelectionHighlightZOrder = -10; // draw below the text
+	static const int32 CompositionRangeZOrder = 10; // draw above the text
+	static const int32 CursorZOrder = 11; // draw above the text and the composition
 
 	const FTextLocation CursorInteractionPosition = CursorInfo.GetCursorInteractionLocation();
 	const FTextLocation SelectionLocation = SelectionStart.Get(CursorInteractionPosition);
@@ -2277,7 +2280,7 @@ void FSlateEditableTextLayout::UpdateCursorHighlight()
 			const bool bCursorInRange = (CompositionBeginLocation.GetLineIndex() == CursorInteractionPosition.GetLineIndex() && Range.InclusiveContains(CursorInteractionPosition.GetOffset()));
 			if (!Range.IsEmpty() && bCursorInRange)
 			{
-				TextLayout->AddLineHighlight(FTextLineHighlight(CompositionBeginLocation.GetLineIndex(), Range, CompositionRangeZOrder, TextCompositionHighlighter.ToSharedRef()));
+				ActiveLineHighlights.Add(FTextLineHighlight(CompositionBeginLocation.GetLineIndex(), Range, CompositionRangeZOrder, TextCompositionHighlighter.ToSharedRef()));
 			}
 		}
 	}
@@ -2296,7 +2299,7 @@ void FSlateEditableTextLayout::UpdateCursorHighlight()
 		if (SelectionBeginningLineIndex == SelectionEndLineIndex)
 		{
 			const FTextRange Range(SelectionBeginningLineOffset, SelectionEndLineOffset);
-			TextLayout->AddLineHighlight(FTextLineHighlight(SelectionBeginningLineIndex, Range, SelectionHighlightZOrder, TextSelectionHighlighter.ToSharedRef()));
+			ActiveLineHighlights.Add(FTextLineHighlight(SelectionBeginningLineIndex, Range, SelectionHighlightZOrder, TextSelectionHighlighter.ToSharedRef()));
 		}
 		else
 		{
@@ -2307,17 +2310,17 @@ void FSlateEditableTextLayout::UpdateCursorHighlight()
 				if (LineIndex == SelectionBeginningLineIndex)
 				{
 					const FTextRange Range(SelectionBeginningLineOffset, Lines[LineIndex].Text->Len());
-					TextLayout->AddLineHighlight(FTextLineHighlight(LineIndex, Range, SelectionHighlightZOrder, TextSelectionHighlighter.ToSharedRef()));
+					ActiveLineHighlights.Add(FTextLineHighlight(LineIndex, Range, SelectionHighlightZOrder, TextSelectionHighlighter.ToSharedRef()));
 				}
 				else if (LineIndex == SelectionEndLineIndex)
 				{
 					const FTextRange Range(0, SelectionEndLineOffset);
-					TextLayout->AddLineHighlight(FTextLineHighlight(LineIndex, Range, SelectionHighlightZOrder, TextSelectionHighlighter.ToSharedRef()));
+					ActiveLineHighlights.Add(FTextLineHighlight(LineIndex, Range, SelectionHighlightZOrder, TextSelectionHighlighter.ToSharedRef()));
 				}
 				else
 				{
 					const FTextRange Range(0, Lines[LineIndex].Text->Len());
-					TextLayout->AddLineHighlight(FTextLineHighlight(LineIndex, Range, SelectionHighlightZOrder, TextSelectionHighlighter.ToSharedRef()));
+					ActiveLineHighlights.Add(FTextLineHighlight(LineIndex, Range, SelectionHighlightZOrder, TextSelectionHighlighter.ToSharedRef()));
 				}
 			}
 		}
@@ -2333,23 +2336,32 @@ void FSlateEditableTextLayout::UpdateCursorHighlight()
 
 		if (LineTextLength == 0)
 		{
-			TextLayout->AddLineHighlight(FTextLineHighlight(CursorPosition.GetLineIndex(), FTextRange(0, 0), CursorZOrder, CursorLineHighlighter.ToSharedRef()));
+			ActiveLineHighlights.Add(FTextLineHighlight(CursorPosition.GetLineIndex(), FTextRange(0, 0), CursorZOrder, CursorLineHighlighter.ToSharedRef()));
 		}
 		else if (CursorPosition.GetOffset() == LineTextLength)
 		{
-			TextLayout->AddLineHighlight(FTextLineHighlight(CursorPosition.GetLineIndex(), FTextRange(LineTextLength - 1, LineTextLength), CursorZOrder, CursorLineHighlighter.ToSharedRef()));
+			ActiveLineHighlights.Add(FTextLineHighlight(CursorPosition.GetLineIndex(), FTextRange(LineTextLength - 1, LineTextLength), CursorZOrder, CursorLineHighlighter.ToSharedRef()));
 		}
 		else
 		{
-			TextLayout->AddLineHighlight(FTextLineHighlight(CursorPosition.GetLineIndex(), FTextRange(CursorPosition.GetOffset(), CursorPosition.GetOffset() + 1), CursorZOrder, CursorLineHighlighter.ToSharedRef()));
+			ActiveLineHighlights.Add(FTextLineHighlight(CursorPosition.GetLineIndex(), FTextRange(CursorPosition.GetOffset(), CursorPosition.GetOffset() + 1), CursorZOrder, CursorLineHighlighter.ToSharedRef()));
 		}
 	}
+
+	// We don't use SetLineHighlights here as we don't want to remove any line highlights that other code might have added (eg, underlines)
+	for (const FTextLineHighlight& LineHighlight : ActiveLineHighlights)
+	{
+		TextLayout->AddLineHighlight(LineHighlight);
+	}
 }
 
 void FSlateEditableTextLayout::RemoveCursorHighlight()
 {
-	TextLayout->ClearRunRenderers();
-	TextLayout->ClearLineHighlights();
+	for (const FTextLineHighlight& LineHighlight : ActiveLineHighlights)
+	{
+		TextLayout->RemoveLineHighlight(LineHighlight);
+	}
+	ActiveLineHighlights.Empty();
 }
 
 void FSlateEditableTextLayout::UpdatePreferredCursorScreenOffsetInLine()
diff --git a/Engine/Source/Runtime/Slate/Public/Framework/MultiBox/MultiBox.h b/Engine/Source/Runtime/Slate/Public/Framework/MultiBox/MultiBox.h
--- a/Engine/Source/Runtime/Slate/Public/Framework/MultiBox/MultiBox.h
+++ b/Engine/Source/Runtime/Slate/Public/Framework/MultiBox/MultiBox.h
@@ -583,7 +583,7 @@ public:
 	TSharedPtr<STextBlock> GetSearchTextWidget(); 
 
 	/**
-	* Get the block widget holding the search text
+	* Set the block widget holding the search text
 	*
 	* @param	BlockWidget		The widget to set
 	*/
diff --git a/Engine/Source/Runtime/Slate/Public/Framework/MultiBox/MultiBoxBuilder.h b/Engine/Source/Runtime/Slate/Public/Framework/MultiBox/MultiBoxBuilder.h
--- a/Engine/Source/Runtime/Slate/Public/Framework/MultiBox/MultiBoxBuilder.h
+++ b/Engine/Source/Runtime/Slate/Public/Framework/MultiBox/MultiBoxBuilder.h
@@ -155,6 +155,7 @@ public:
 	 * Adds a menu entry
 	 *
 	 * @param	InCommand			The command associated with this menu entry
+	 * @param	InExtensionHook		The section hook. Can be NAME_None
 	 * @param	InLabelOverride		Optional label override.  If omitted, then the action's label will be used instead.
 	 * @param	InToolTipOverride	Optional tool tip override.	 If omitted, then the action's label will be used instead.
 	 * @param	InIconOverride		Optional name of the slate brush to use for the tool bar image.  If omitted, then the action's icon will be used instead.
@@ -169,6 +170,7 @@ public:
 	 * @param	InToolTip	Tool tip used when hovering over the menu entry
 	 * @param	InIcon		The icon to use		
 	 * @param	UIAction	Actions to execute on this menu item.
+	 * @param	InExtensionHook			The section hook. Can be NAME_None
 	 * @param	UserInterfaceActionType	Type of interface action
 	 * @param	InTutorialHighlightName	Optional name to identify this widget and highlight during tutorials
 	 */
@@ -236,7 +238,11 @@ public:
 	 * @param	InMenuLabel			The text that should be shown for the menu
 	 * @param	InToolTip			The tooltip that should be shown when the menu is hovered over
 	 * @param	InSubMenu			Sub-Menu object which creates menu entries for the sub-menu
+	 * @param	InUIAction			Actions to execute on this menu item.
+	 * @param	InExtensionHook		The section hook. Can be NAME_None
+	 * @param	InUserInterfaceActionType	Type of interface action
 	 * @param bInOpenSubMenuOnClick Sub-menu will open only if the sub-menu entry is clicked
+	 * @param	InIcon				The icon to use
 	 */
 	void AddSubMenu( const TAttribute<FText>& InMenuLabel, const TAttribute<FText>& InToolTip, const FNewMenuDelegate& InSubMenu, const FUIAction& InUIAction, FName InExtensionHook, const EUserInterfaceActionType::Type InUserInterfaceActionType, const bool bInOpenSubMenuOnClick = false, const FSlateIcon& InIcon = FSlateIcon() );
 
diff --git a/Engine/Source/Runtime/Slate/Public/Framework/Text/RichTextLayoutMarshaller.h b/Engine/Source/Runtime/Slate/Public/Framework/Text/RichTextLayoutMarshaller.h
--- a/Engine/Source/Runtime/Slate/Public/Framework/Text/RichTextLayoutMarshaller.h
+++ b/Engine/Source/Runtime/Slate/Public/Framework/Text/RichTextLayoutMarshaller.h
@@ -8,6 +8,8 @@
 #include "IRichTextMarkupParser.h"
 #include "IRichTextMarkupWriter.h"
 
+class FSlateTextUnderlineLineHighlighter;
+
 /**
  * Get/set the raw text to/from a text layout as rich text
  */
@@ -40,12 +42,17 @@ protected:
 
 	TSharedPtr< ITextDecorator > TryGetDecorator(const FString& Line, const FTextRunParseResults& TextRun) const;
 
-	virtual void AppendRunsForText(const FTextRunParseResults& TextRun, 
+	virtual void AppendRunsForText(
+		const int32 LineIndex,
+		const FTextRunParseResults& TextRun, 
 		const FString& ProcessedString, 
 		const FTextBlockStyle& DefaultTextStyle,
 		const TSharedRef<FString>& InOutModelText, 
 		FTextLayout& TargetTextLayout,
-		TArray<TSharedRef<IRun>>& Runs);
+		TArray<TSharedRef<IRun>>& Runs,
+		TArray<FTextLineHighlight>& LineHighlights,
+		TMap<const FTextBlockStyle*, TSharedPtr<FSlateTextUnderlineLineHighlighter>>& CachedUnderlineHighlighters
+		);
 
 	/** The parser used to resolve any markup used in the provided string. */
 	TSharedPtr< IRichTextMarkupParser > Parser;
diff --git a/Engine/Source/Runtime/Slate/Public/Framework/Text/SlateTextUnderlineLineHighlighter.h b/Engine/Source/Runtime/Slate/Public/Framework/Text/SlateTextUnderlineLineHighlighter.h
new file mode 100644
--- /dev/null
+++ b/Engine/Source/Runtime/Slate/Public/Framework/Text/SlateTextUnderlineLineHighlighter.h
@@ -0,0 +1,31 @@
+// Copyright 1998-2016 Epic Games, Inc. All Rights Reserved.
+#pragma once
+
+/** Run highlighter used to draw underlines */
+class SLATE_API FSlateTextUnderlineLineHighlighter : public ISlateLineHighlighter
+{
+public:
+	static TSharedRef<FSlateTextUnderlineLineHighlighter> Create(const FSlateBrush& InUnderlineBrush, const FSlateFontInfo& InFontInfo, const FSlateColor InColorAndOpacity, const FVector2D InShadowOffset, const FLinearColor InShadowColorAndOpacity);
+
+	virtual int32 OnPaint(const FPaintArgs& Args, const FTextLayout::FLineView& Line, const float OffsetX, const float Width, const FTextBlockStyle& DefaultStyle, const FGeometry& AllottedGeometry, const FSlateRect& MyClippingRect, FSlateWindowElementList& OutDrawElements, int32 LayerId, const FWidgetStyle& InWidgetStyle, bool bParentEnabled) const override;
+
+	static const int32 DefaultZIndex = 1;
+
+protected:
+	FSlateTextUnderlineLineHighlighter(const FSlateBrush& InUnderlineBrush, const FSlateFontInfo& InFontInfo, const FSlateColor InColorAndOpacity, const FVector2D InShadowOffset, const FLinearColor InShadowColorAndOpacity);
+
+	/** Brush used to draw the underline */
+	FSlateBrush UnderlineBrush;
+
+	/** Font the underline is associated with */
+	FSlateFontInfo FontInfo;
+
+	/** The color to draw the underline (typically matches the text its associated with) */
+	FSlateColor ColorAndOpacity;
+
+	/** Offset at which to draw the shadow (if any) */
+	FVector2D ShadowOffset;
+
+	/** The color to draw the shadow */
+	FLinearColor ShadowColorAndOpacity;
+};
diff --git a/Engine/Source/Runtime/Slate/Public/Framework/Text/TextLayout.h b/Engine/Source/Runtime/Slate/Public/Framework/Text/TextLayout.h
--- a/Engine/Source/Runtime/Slate/Public/Framework/Text/TextLayout.h
+++ b/Engine/Source/Runtime/Slate/Public/Framework/Text/TextLayout.h
@@ -418,6 +418,11 @@ public:
 	*/
 	void AddRunRenderer( const FTextRunRenderer& Renderer );
 
+	/**
+	* Removes a single run renderer to the existing set of renderers.
+	*/
+	void RemoveRunRenderer( const FTextRunRenderer& Renderer );
+
 	/**
 	* Clears all line highlights
 	*/
@@ -433,6 +438,11 @@ public:
 	*/
 	void AddLineHighlight( const FTextLineHighlight& Highlight );
 
+	/**
+	* Removes a single line highlight to the existing set of highlights.
+	*/
+	void RemoveLineHighlight( const FTextLineHighlight& Highlight );
+
 	/**
 	* Updates the TextLayout's if any changes have occurred since the last update.
 	*/
diff --git a/Engine/Source/Runtime/Slate/Public/Framework/Text/TextLineHighlight.h b/Engine/Source/Runtime/Slate/Public/Framework/Text/TextLineHighlight.h
--- a/Engine/Source/Runtime/Slate/Public/Framework/Text/TextLineHighlight.h
+++ b/Engine/Source/Runtime/Slate/Public/Framework/Text/TextLineHighlight.h
@@ -14,6 +14,22 @@ struct SLATE_API FTextLineHighlight
 
 	}
 
+	bool operator==(const FTextLineHighlight& Other) const
+	{
+		return LineIndex == Other.LineIndex
+			&& Range == Other.Range
+			&& ZOrder == Other.ZOrder
+			&& Highlighter == Other.Highlighter;
+	}
+
+	bool operator!=(const FTextLineHighlight& Other) const
+	{
+		return LineIndex != Other.LineIndex
+			|| Range != Other.Range
+			|| ZOrder != Other.ZOrder
+			|| Highlighter != Other.Highlighter;
+	}
+
 	int32 LineIndex;
 	FTextRange Range;
 	int32 ZOrder;
diff --git a/Engine/Source/Runtime/Slate/Public/Framework/Text/TextRunRenderer.h b/Engine/Source/Runtime/Slate/Public/Framework/Text/TextRunRenderer.h
--- a/Engine/Source/Runtime/Slate/Public/Framework/Text/TextRunRenderer.h
+++ b/Engine/Source/Runtime/Slate/Public/Framework/Text/TextRunRenderer.h
@@ -13,6 +13,20 @@ struct SLATE_API FTextRunRenderer
 
 	}
 
+	bool operator==(const FTextRunRenderer& Other) const
+	{
+		return LineIndex == Other.LineIndex
+			&& Range == Other.Range
+			&& Renderer == Other.Renderer;
+	}
+
+	bool operator!=(const FTextRunRenderer& Other) const
+	{
+		return LineIndex != Other.LineIndex
+			|| Range != Other.Range
+			|| Renderer != Other.Renderer;
+	}
+
 	int32 LineIndex;
 	FTextRange Range;
 	TSharedRef< IRunRenderer > Renderer;
diff --git a/Engine/Source/Runtime/Slate/Public/Widgets/Text/SlateEditableTextLayout.h b/Engine/Source/Runtime/Slate/Public/Widgets/Text/SlateEditableTextLayout.h
--- a/Engine/Source/Runtime/Slate/Public/Widgets/Text/SlateEditableTextLayout.h
+++ b/Engine/Source/Runtime/Slate/Public/Widgets/Text/SlateEditableTextLayout.h
@@ -508,6 +508,9 @@ private:
 	/** Run renderer used to draw the active text selection */
 	TSharedPtr<SlateEditableTextTypes::FTextSelectionHighlighter> TextSelectionHighlighter;
 
+	/** Line highlights that have been added from this editable text layout (used for cleanup without removing) */
+	TArray<FTextLineHighlight> ActiveLineHighlights;
+
 	/** The scroll offset (in unscaled Slate units) for this text */
 	FVector2D ScrollOffset;
 
diff --git a/Engine/Source/Runtime/SlateCore/Private/Fonts/FontBulkData.cpp b/Engine/Source/Runtime/SlateCore/Private/Fonts/FontBulkData.cpp
--- a/Engine/Source/Runtime/SlateCore/Private/Fonts/FontBulkData.cpp
+++ b/Engine/Source/Runtime/SlateCore/Private/Fonts/FontBulkData.cpp
@@ -102,7 +102,7 @@ void UFontBulkData::Serialize(FArchive& Ar)
 	Super::Serialize(Ar);
 	BulkData.Serialize(Ar, this);
 
-	if( !GIsEditor )
+	if( !GIsEditor && Ar.IsLoading() )
 	{
 		BulkData.SetBulkDataFlags( BULKDATA_SingleUse );
 	}
diff --git a/Engine/Source/Runtime/SlateCore/Private/Fonts/FontCache.cpp b/Engine/Source/Runtime/SlateCore/Private/Fonts/FontCache.cpp
--- a/Engine/Source/Runtime/SlateCore/Private/Fonts/FontCache.cpp
+++ b/Engine/Source/Runtime/SlateCore/Private/Fonts/FontCache.cpp
@@ -1015,6 +1015,11 @@ int16 FSlateFontCache::GetBaseline( const FSlateFontInfo& InFontInfo, float Font
 	return FontRenderer->GetBaseline(InFontInfo, FontScale);
 }
 
+void FSlateFontCache::GetUnderlineMetrics( const FSlateFontInfo& InFontInfo, const float FontScale, int16& OutUnderlinePos, int16& OutUnderlineThickness ) const
+{
+	FontRenderer->GetUnderlineMetrics(InFontInfo, FontScale, OutUnderlinePos, OutUnderlineThickness);
+}
+
 int8 FSlateFontCache::GetKerning( const FFontData& InFontData, const int32 InSize, TCHAR First, TCHAR Second, float Scale ) const
 {
 	return FontRenderer->GetKerning(InFontData, InSize, First, Second, Scale);
diff --git a/Engine/Source/Runtime/SlateCore/Private/Fonts/SlateFontRenderer.cpp b/Engine/Source/Runtime/SlateCore/Private/Fonts/SlateFontRenderer.cpp
--- a/Engine/Source/Runtime/SlateCore/Private/Fonts/SlateFontRenderer.cpp
+++ b/Engine/Source/Runtime/SlateCore/Private/Fonts/SlateFontRenderer.cpp
@@ -91,6 +91,27 @@ int16 FSlateFontRenderer::GetBaseline(const FSlateFontInfo& InFontInfo, const fl
 #endif // WITH_FREETYPE
 }
 
+void FSlateFontRenderer::GetUnderlineMetrics(const FSlateFontInfo& InFontInfo, const float InScale, int16& OutUnderlinePos, int16& OutUnderlineThickness) const
+{
+#if WITH_FREETYPE
+	const FFontData& FontData = CompositeFontCache->GetDefaultFontData(InFontInfo);
+	FT_Face FontFace = GetFontFace(FontData);
+
+	if (FontFace && FT_IS_SCALABLE(FontFace))
+	{
+		FreeTypeUtils::ApplySizeAndScale(FontFace, InFontInfo.Size, InScale);
+
+		OutUnderlinePos = static_cast<int16>(FreeTypeUtils::Convert26Dot6ToRoundedPixel<int32>(FT_MulFix(FontFace->underline_position, FontFace->size->metrics.y_scale)) * InScale);
+		OutUnderlineThickness = static_cast<int16>(FreeTypeUtils::Convert26Dot6ToRoundedPixel<int32>(FT_MulFix(FontFace->underline_thickness, FontFace->size->metrics.y_scale)) * InScale);
+	}
+	else
+#endif // WITH_FREETYPE
+	{
+		OutUnderlinePos = 0;
+		OutUnderlineThickness = 0;
+	}
+}
+
 bool FSlateFontRenderer::HasKerning(const FFontData& InFontData) const
 {
 #if WITH_FREETYPE
diff --git a/Engine/Source/Runtime/SlateCore/Private/Fonts/SlateFontRenderer.h b/Engine/Source/Runtime/SlateCore/Private/Fonts/SlateFontRenderer.h
--- a/Engine/Source/Runtime/SlateCore/Private/Fonts/SlateFontRenderer.h
+++ b/Engine/Source/Runtime/SlateCore/Private/Fonts/SlateFontRenderer.h
@@ -72,6 +72,14 @@ public:
 	 */
 	int16 GetBaseline(const FSlateFontInfo& InFontInfo, const float InScale) const;
 
+	/**
+	 * Get the underline metrics used by any character in the default font
+	 *
+	 * @param OutUnderlinePos		The offset from the baseline to the center of the underline bar
+	 * @param OutUnderlineThickness	The thickness of the underline bar
+	 */
+	void GetUnderlineMetrics(const FSlateFontInfo& InFontInfo, const float InScale, int16& OutUnderlinePos, int16& OutUnderlineThickness) const;
+
 	/**
 	 * @param Whether or not the font has kerning
 	 */
diff --git a/Engine/Source/Runtime/SlateCore/Private/Rendering/DrawElements.cpp b/Engine/Source/Runtime/SlateCore/Private/Rendering/DrawElements.cpp
--- a/Engine/Source/Runtime/SlateCore/Private/Rendering/DrawElements.cpp
+++ b/Engine/Source/Runtime/SlateCore/Private/Rendering/DrawElements.cpp
@@ -5,6 +5,8 @@
 
 DECLARE_CYCLE_STAT(TEXT("FSlateDrawElement::Make Time"), STAT_SlateDrawElementMakeTime, STATGROUP_SlateVerbose);
 
+DEFINE_STAT(STAT_SlateBufferPoolMemory);
+
 FSlateShaderResourceManager* FSlateDataPayload::ResourceManager;
 
 void FSlateDataPayload::SetTextPayloadProperties( FSlateWindowElementList& ElementList, const FString& InText, const FSlateFontInfo& InFontInfo, const FLinearColor& InTint, const int32 InStartIndex, const int32 InEndIndex )
@@ -326,6 +328,7 @@ void FSlateBatchData::Reset()
 }
 
 #define MAX_VERT_ARRAY_RECYCLE (200)
+#define MAX_INDEX_ARRAY_RECYCLE (500)
 
 void FSlateBatchData::AssignVertexArrayToBatch( FSlateElementBatch& Batch )
 {
@@ -337,15 +340,13 @@ void FSlateBatchData::AssignVertexArrayToBatch( FSlateElementBatch& Batch )
 	else
 	{
 		// There are no free vertex arrays so we must add one		
-		uint32 NewIndex = BatchVertexArrays.Add(TArray<FSlateVertex>());
-		BatchVertexArrays[NewIndex].Reserve(MAX_VERT_ARRAY_RECYCLE);
+		uint32 NewIndex = BatchVertexArrays.Add(FSlateVertexArray());
+		ResetVertexArray(BatchVertexArrays[NewIndex]);
 
 		Batch.VertexArrayIndex = NewIndex;
 	}
 }
 
-#define MAX_INDEX_ARRAY_RECYCLE (500)
-
 void FSlateBatchData::AssignIndexArrayToBatch( FSlateElementBatch& Batch )
 {
 	// Get a free index array
@@ -356,8 +357,8 @@ void FSlateBatchData::AssignIndexArrayToBatch( FSlateElementBatch& Batch )
 	else
 	{
 		// There are no free index arrays so we must add one
-		uint32 NewIndex = BatchIndexArrays.Add(TArray<SlateIndex>());
-		BatchIndexArrays[NewIndex].Reserve(MAX_INDEX_ARRAY_RECYCLE);
+		uint32 NewIndex = BatchIndexArrays.Add(FSlateIndexArray());
+		ResetIndexArray(BatchIndexArrays[NewIndex]);
 
 		Batch.IndexArrayIndex = NewIndex;
 	}
@@ -378,8 +379,8 @@ void FSlateBatchData::FillVertexAndIndexBuffer(uint8* VertexBuffer, uint8* Index
 
 		if ( Batch.VertexArrayIndex != INDEX_NONE && Batch.IndexArrayIndex != INDEX_NONE )
 		{
-			TArray<FSlateVertex>& Vertices = BatchVertexArrays[Batch.VertexArrayIndex];
-			TArray<SlateIndex>& Indices = BatchIndexArrays[Batch.IndexArrayIndex];
+			FSlateVertexArray& Vertices = BatchVertexArrays[Batch.VertexArrayIndex];
+			FSlateIndexArray& Indices = BatchIndexArrays[Batch.IndexArrayIndex];
 
 			if ( Vertices.Num() && Indices.Num() )
 			{
@@ -397,13 +398,12 @@ void FSlateBatchData::FillVertexAndIndexBuffer(uint8* VertexBuffer, uint8* Index
 
 				if ( Vertices.GetSlack() > MAX_VERT_ARRAY_RECYCLE )
 				{
-					Vertices.Empty();
-					Vertices.Reserve(MAX_VERT_ARRAY_RECYCLE);
+					ResetVertexArray(Vertices);
 				}
+
 				if ( Indices.GetSlack() > MAX_INDEX_ARRAY_RECYCLE )
 				{
-					Indices.Empty();
-					Indices.Reserve(MAX_INDEX_ARRAY_RECYCLE);
+					ResetIndexArray(Indices);
 				}
 			}
 
@@ -442,6 +442,17 @@ void FSlateBatchData::AddRenderBatch(uint32 InLayer, const FSlateElementBatch& I
 	RenderBatches[Index].DynamicOffset = DynamicOffset;
 }
 
+void FSlateBatchData::ResetVertexArray(FSlateVertexArray& InOutVertexArray)
+{
+	InOutVertexArray.Empty(MAX_VERT_ARRAY_RECYCLE);
+}
+
+void FSlateBatchData::ResetIndexArray(FSlateIndexArray& InOutIndexArray)
+{
+	InOutIndexArray.Empty(MAX_INDEX_ARRAY_RECYCLE);
+}
+
+
 void FSlateBatchData::Merge(FElementBatchMap& InLayerToElementBatches, uint32& VertexOffset, uint32& IndexOffset)
 {
 	InLayerToElementBatches.Sort();
@@ -507,8 +518,8 @@ void FSlateBatchData::Merge(FElementBatchMap& InLayerToElementBatches, uint32& V
 					
 				if ( ElementBatch.VertexArrayIndex != INDEX_NONE && ElementBatch.IndexArrayIndex != INDEX_NONE )
 				{
-					TArray<FSlateVertex>& BatchVertices = GetBatchVertexList(ElementBatch);
-					TArray<SlateIndex>& BatchIndices = GetBatchIndexList(ElementBatch);
+					FSlateVertexArray& BatchVertices = GetBatchVertexList(ElementBatch);
+					FSlateIndexArray& BatchIndices = GetBatchIndexList(ElementBatch);
 
 					// We should have at least some vertices and indices in the batch or none at all
 					check((BatchVertices.Num() > 0 && BatchIndices.Num() > 0) || (BatchVertices.Num() == 0 && BatchIndices.Num() == 0));
diff --git a/Engine/Source/Runtime/SlateCore/Private/Rendering/ElementBatcher.cpp b/Engine/Source/Runtime/SlateCore/Private/Rendering/ElementBatcher.cpp
--- a/Engine/Source/Runtime/SlateCore/Private/Rendering/ElementBatcher.cpp
+++ b/Engine/Source/Runtime/SlateCore/Private/Rendering/ElementBatcher.cpp
@@ -258,8 +258,8 @@ void FSlateElementBatcher::AddQuadElement( const FSlateDrawElement& DrawElement,
 	FSlateRotatedClipRectType RenderClipRect = ToSnappedRotatedRect(InClippingRect, InverseLayoutTransform, RenderTransform);
 
 	FSlateElementBatch& ElementBatch = FindBatchForElement( Layer, FShaderParams(), nullptr, ESlateDrawPrimitive::TriangleList, ESlateShader::Default, ESlateDrawEffect::None, ESlateBatchDrawFlag::Wireframe|ESlateBatchDrawFlag::NoBlending, DrawElement.GetScissorRect() );
-	TArray<FSlateVertex>& BatchVertices = BatchData->GetBatchVertexList(ElementBatch);
-	TArray<SlateIndex>& BatchIndices = BatchData->GetBatchIndexList(ElementBatch);
+	FSlateVertexArray& BatchVertices = BatchData->GetBatchVertexList(ElementBatch);
+	FSlateIndexArray& BatchIndices = BatchData->GetBatchIndexList(ElementBatch);
 
 	// Determine the four corners of the quad
 	FVector2D TopLeft = FVector2D::ZeroVector;
@@ -398,8 +398,8 @@ void FSlateElementBatcher::AddBoxElement( const FSlateDrawElement& DrawElement )
 		DrawFlags |= ( ( bTileHorizontal ? ESlateBatchDrawFlag::TileU : 0 ) | ( bTileVertical ? ESlateBatchDrawFlag::TileV : 0 ) );
 
 		FSlateElementBatch& ElementBatch = FindBatchForElement( Layer, FShaderParams(), Resource, ESlateDrawPrimitive::TriangleList, ESlateShader::Default, InDrawEffects, DrawFlags, DrawElement.GetScissorRect() );
-		TArray<FSlateVertex>& BatchVertices = BatchData->GetBatchVertexList(ElementBatch);
-		TArray<SlateIndex>& BatchIndices = BatchData->GetBatchIndexList(ElementBatch);
+		FSlateVertexArray& BatchVertices = BatchData->GetBatchVertexList(ElementBatch);
+		FSlateIndexArray& BatchIndices = BatchData->GetBatchIndexList(ElementBatch);
 
 		float HorizontalTiling = bTileHorizontal ? LocalSize.X/TextureWidth : 1.0f;
 		float VerticalTiling = bTileVertical ? LocalSize.Y/TextureHeight : 1.0f;
@@ -665,8 +665,8 @@ void FSlateElementBatcher::AddTextElement(const FSlateDrawElement& DrawElement)
 	FSlateShaderResource* FontShaderResource = nullptr;
 
 	FSlateElementBatch* ElementBatch = nullptr;
-	TArray<FSlateVertex>* BatchVertices = nullptr;
-	TArray<SlateIndex>* BatchIndices = nullptr;
+	FSlateVertexArray* BatchVertices = nullptr;
+	FSlateIndexArray* BatchIndices = nullptr;
 
 	uint32 VertexOffset = 0;
 	uint32 IndexOffset = 0;
@@ -766,8 +766,8 @@ void FSlateElementBatcher::AddTextElement(const FSlateDrawElement& DrawElement)
 
 				if( FSlateRect::DoRectanglesIntersect( LocalClipRect, CharRect ) )
 				{
-					TArray<FSlateVertex>& BatchVerticesRef = *BatchVertices;
-					TArray<SlateIndex>& BatchIndicesRef = *BatchIndices;
+					FSlateVertexArray& BatchVerticesRef = *BatchVertices;
+					FSlateIndexArray& BatchIndicesRef = *BatchIndices;
 
 					FVector2D UpperLeft( X, Y );
 					FVector2D UpperRight( X+SizeX, Y );
@@ -868,8 +868,8 @@ void FSlateElementBatcher::AddShapedTextElement( const FSlateDrawElement& DrawEl
 	FSlateShaderResource* FontShaderResource = nullptr;
 
 	FSlateElementBatch* ElementBatch = nullptr;
-	TArray<FSlateVertex>* BatchVertices = nullptr;
-	TArray<SlateIndex>* BatchIndices = nullptr;
+	FSlateVertexArray* BatchVertices = nullptr;
+	FSlateIndexArray* BatchIndices = nullptr;
 
 	uint32 VertexOffset = 0;
 	uint32 IndexOffset = 0;
@@ -936,8 +936,8 @@ void FSlateElementBatcher::AddShapedTextElement( const FSlateDrawElement& DrawEl
 				const FSlateRect CharRect(X, Y, X + SizeX, Y + SizeY);
 				if (FSlateRect::DoRectanglesIntersect(LocalClipRect, CharRect))
 				{
-					TArray<FSlateVertex>& BatchVerticesRef = *BatchVertices;
-					TArray<SlateIndex>& BatchIndicesRef = *BatchIndices;
+					FSlateVertexArray& BatchVerticesRef = *BatchVertices;
+					FSlateIndexArray& BatchIndicesRef = *BatchIndices;
 
 					const FVector2D UpperLeft = CharRect.GetTopLeft();
 					const FVector2D UpperRight = CharRect.GetTopRight();
@@ -1017,8 +1017,8 @@ void FSlateElementBatcher::AddGradientElement( const FSlateDrawElement& DrawElem
 			InPayload.BatchFlags,
 			DrawElement.GetScissorRect() );
 
-	TArray<FSlateVertex>& BatchVertices = BatchData->GetBatchVertexList(ElementBatch);
-	TArray<SlateIndex>& BatchIndices = BatchData->GetBatchIndexList(ElementBatch);
+	FSlateVertexArray& BatchVertices = BatchData->GetBatchVertexList(ElementBatch);
+	FSlateIndexArray& BatchIndices = BatchData->GetBatchIndexList(ElementBatch);
 
 	// Determine the four corners of the quad containing the gradient
 	FVector2D TopLeft = FVector2D::ZeroVector;
@@ -1170,8 +1170,8 @@ void FSlateElementBatcher::AddSplineElement( const FSlateDrawElement& DrawElemen
 
 	// Find a batch for the element
 	FSlateElementBatch& ElementBatch = FindBatchForElement( Layer, FShaderParams::MakePixelShaderParams( FVector4( InPayload.Thickness,Radius,0,0) ), nullptr, ESlateDrawPrimitive::TriangleList, ESlateShader::LineSegment, InDrawEffects, ESlateBatchDrawFlag::None, DrawElement.GetScissorRect() );
-	TArray<FSlateVertex>& BatchVertices = BatchData->GetBatchVertexList(ElementBatch);
-	TArray<SlateIndex>& BatchIndices = BatchData->GetBatchIndexList(ElementBatch);
+	FSlateVertexArray& BatchVertices = BatchData->GetBatchVertexList(ElementBatch);
+	FSlateIndexArray& BatchIndices = BatchData->GetBatchIndexList(ElementBatch);
 
 	const FVector2D StartPt = InPayload.StartPt;
 	const FVector2D StartDir = InPayload.StartDir;
@@ -1317,8 +1317,8 @@ void FSlateElementBatcher::AddLineElement( const FSlateDrawElement& DrawElement
 
 		// Find a batch for the element
 		FSlateElementBatch& ElementBatch = FindBatchForElement( Layer, FShaderParams::MakePixelShaderParams( FVector4(RequestedThickness,Radius,0,0) ), nullptr, ESlateDrawPrimitive::TriangleList, ESlateShader::LineSegment, DrawEffects, ESlateBatchDrawFlag::None, DrawElement.GetScissorRect() );
-		TArray<FSlateVertex>& BatchVertices = BatchData->GetBatchVertexList(ElementBatch);
-		TArray<SlateIndex>& BatchIndices = BatchData->GetBatchIndexList(ElementBatch);
+		FSlateVertexArray& BatchVertices = BatchData->GetBatchVertexList(ElementBatch);
+		FSlateIndexArray& BatchIndices = BatchData->GetBatchIndexList(ElementBatch);
 
 		const TArray<FVector2D>& Points = InPayload.Points;
 
@@ -1427,8 +1427,8 @@ void FSlateElementBatcher::AddLineElement( const FSlateDrawElement& DrawElement
 	{
 		// Find a batch for the element
 		FSlateElementBatch& ElementBatch = FindBatchForElement( Layer, FShaderParams(), nullptr, ESlateDrawPrimitive::LineList, ESlateShader::Default, DrawEffects, ESlateBatchDrawFlag::None, DrawElement.GetScissorRect() );
-		TArray<FSlateVertex>& BatchVertices = BatchData->GetBatchVertexList(ElementBatch);
-		TArray<SlateIndex>& BatchIndices = BatchData->GetBatchIndexList(ElementBatch);
+		FSlateVertexArray& BatchVertices = BatchData->GetBatchVertexList(ElementBatch);
+		FSlateIndexArray& BatchIndices = BatchData->GetBatchIndexList(ElementBatch);
 
 		// Generate the rest of the segments
 		for( int32 Point = 0; Point < InPayload.Points.Num()-1; ++Point )
@@ -1479,8 +1479,8 @@ void FSlateElementBatcher::AddViewportElement( const FSlateDrawElement& DrawElem
 
 
 	FSlateElementBatch& ElementBatch = FindBatchForElement( Layer, FShaderParams(), ViewportResource, ESlateDrawPrimitive::TriangleList, ShaderType, InDrawEffects, DrawFlags, DrawElement.GetScissorRect() );
-	TArray<FSlateVertex>& BatchVertices = BatchData->GetBatchVertexList(ElementBatch);
-	TArray<SlateIndex>& BatchIndices = BatchData->GetBatchIndexList(ElementBatch);
+	FSlateVertexArray& BatchVertices = BatchData->GetBatchVertexList(ElementBatch);
+	FSlateIndexArray& BatchIndices = BatchData->GetBatchIndexList(ElementBatch);
 
 	// Tag this batch as requiring vsync if the viewport requires it.
 	if( ViewportResource != nullptr && !InPayload.bAllowViewportScaling )
@@ -1613,8 +1613,8 @@ void FSlateElementBatcher::AddBorderElement( const FSlateDrawElement& DrawElemen
 	ESlateBatchDrawFlag::Type DrawFlags = (ESlateBatchDrawFlag::TileU|ESlateBatchDrawFlag::TileV);
 
 	FSlateElementBatch& ElementBatch = FindBatchForElement( Layer, ShaderParams, Resource, ESlateDrawPrimitive::TriangleList, ESlateShader::Border, InDrawEffects, DrawFlags, DrawElement.GetScissorRect() );
-	TArray<FSlateVertex>& BatchVertices = BatchData->GetBatchVertexList(ElementBatch);
-	TArray<SlateIndex>& BatchIndices = BatchData->GetBatchIndexList(ElementBatch);
+	FSlateVertexArray& BatchVertices = BatchData->GetBatchVertexList(ElementBatch);
+	FSlateIndexArray& BatchIndices = BatchData->GetBatchIndexList(ElementBatch);
 
 	// Ensure tiling of at least 1.  
 	TopTiling = TopTiling >= 1.0f ? TopTiling : 1.0f;
@@ -1796,8 +1796,8 @@ void FSlateElementBatcher::AddCustomVerts(const FSlateDrawElement& DrawElement)
 		BatchData->AssignVertexArrayToBatch(*ElementBatch);
 		BatchData->AssignIndexArrayToBatch(*ElementBatch);
 
-		TArray<FSlateVertex>& BatchVertices = BatchData->GetBatchVertexList(*ElementBatch);
-		TArray<SlateIndex>& BatchIndices = BatchData->GetBatchIndexList(*ElementBatch);
+		FSlateVertexArray& BatchVertices = BatchData->GetBatchVertexList(*ElementBatch);
+		FSlateIndexArray& BatchIndices = BatchData->GetBatchIndexList(*ElementBatch);
 
 		// Vertex Buffer since  it is already in slate format it is a straight copy
 		BatchVertices = InPayload.CustomVertsData;
diff --git a/Engine/Source/Runtime/SlateCore/Private/Styling/CoreStyle.cpp b/Engine/Source/Runtime/SlateCore/Private/Styling/CoreStyle.cpp
--- a/Engine/Source/Runtime/SlateCore/Private/Styling/CoreStyle.cpp
+++ b/Engine/Source/Runtime/SlateCore/Private/Styling/CoreStyle.cpp
@@ -158,8 +158,11 @@ TSharedRef<ISlateStyle> FCoreStyle::Create( const FName& InStyleSetName )
 
 	Style->Set("SmallFont", TTF_FONT("Fonts/Roboto-Regular", 8));
 
+	FSlateBrush* DefaultTextUnderlineBrush = new IMAGE_BRUSH("Old/White", Icon8x8, FLinearColor::White, ESlateBrushTileType::Both);
+	Style->Set("DefaultTextUnderline", DefaultTextUnderlineBrush);
+
 	// Normal Text
-	FTextBlockStyle NormalText = FTextBlockStyle()
+	const FTextBlockStyle NormalText = FTextBlockStyle()
 		.SetFont(TTF_FONT("Fonts/Roboto-Regular", 9))
 		.SetColorAndOpacity(FSlateColor::UseForeground())
 		.SetShadowOffset(FVector2D::ZeroVector)
@@ -167,29 +170,39 @@ TSharedRef<ISlateStyle> FCoreStyle::Create( const FName& InStyleSetName )
 		.SetHighlightColor(FLinearColor(0.02f, 0.3f, 0.0f))
 		.SetHighlightShape(BOX_BRUSH("Common/TextBlockHighlightShape", FMargin(3.f /8.f)));
 
+	const FTextBlockStyle NormalUnderlinedText = FTextBlockStyle(NormalText)
+		.SetUnderlineBrush(*DefaultTextUnderlineBrush);
+
 	// Monospaced Text
 	const FTextBlockStyle MonospacedText = FTextBlockStyle()
-			.SetFont(TTF_FONT("Fonts/DroidSansMono", 10))
-			.SetColorAndOpacity(FSlateColor::UseForeground())
-			.SetShadowOffset(FVector2D::ZeroVector)
-			.SetShadowColorAndOpacity(FLinearColor::Black)
-			.SetHighlightColor(FLinearColor(0.02f, 0.3f, 0.0f))
-			.SetHighlightShape(BOX_BRUSH("Common/TextBlockHighlightShape", FMargin(3.f/8.f))
-			);
+		.SetFont(TTF_FONT("Fonts/DroidSansMono", 10))
+		.SetColorAndOpacity(FSlateColor::UseForeground())
+		.SetShadowOffset(FVector2D::ZeroVector)
+		.SetShadowColorAndOpacity(FLinearColor::Black)
+		.SetHighlightColor(FLinearColor(0.02f, 0.3f, 0.0f))
+		.SetHighlightShape(BOX_BRUSH("Common/TextBlockHighlightShape", FMargin(3.f/8.f))
+		);
+
+	const FTextBlockStyle MonospacedUnderlinedText = FTextBlockStyle(MonospacedText)
+		.SetUnderlineBrush(*DefaultTextUnderlineBrush);
 
 	Style->Set("MonospacedText", MonospacedText);
+	Style->Set("MonospacedUnderlinedText", MonospacedUnderlinedText);
 
 	// Small Text
-	FTextBlockStyle SmallText = FTextBlockStyle(NormalText)
+	const FTextBlockStyle SmallText = FTextBlockStyle(NormalText)
 		.SetFont(TTF_FONT("Fonts/Roboto-Regular", 8));
 
+	const FTextBlockStyle SmallUnderlinedText = FTextBlockStyle(SmallText)
+		.SetUnderlineBrush(*DefaultTextUnderlineBrush);
+
 	// Embossed Text
 	Style->Set("EmbossedText", FTextBlockStyle(NormalText)
 		.SetFont(TTF_FONT("Fonts/Roboto-Regular", 24))
 		.SetColorAndOpacity(FLinearColor::Black )
 		.SetShadowOffset( FVector2D(0.0f, 1.0f))
 		.SetShadowColorAndOpacity(FLinearColor(0.8f, 0.8f, 0.8f, 0.5))
-	);
+		);
 
 	// Common brushes
 	FSlateBrush* GenericWhiteBox = new IMAGE_BRUSH("Old/White", Icon16x16);
@@ -460,7 +473,10 @@ TSharedRef<ISlateStyle> FCoreStyle::Create( const FName& InStyleSetName )
 	// STextBlock defaults...
 	{
 		Style->Set("NormalText", NormalText);
+		Style->Set("NormalUnderlinedText", NormalUnderlinedText);
+
 		Style->Set("SmallText", SmallText);
+		Style->Set("SmallUnderlinedText", SmallUnderlinedText);
 	}
 
 	// SInlineEditableTextBlock
diff --git a/Engine/Source/Runtime/SlateCore/Private/Styling/SlateTypes.cpp b/Engine/Source/Runtime/SlateCore/Private/Styling/SlateTypes.cpp
--- a/Engine/Source/Runtime/SlateCore/Private/Styling/SlateTypes.cpp
+++ b/Engine/Source/Runtime/SlateCore/Private/Styling/SlateTypes.cpp
@@ -71,6 +71,12 @@ FTextBlockStyle::FTextBlockStyle()
 
 const FName FTextBlockStyle::TypeName( TEXT("FTextBlockStyle") );
 
+void FTextBlockStyle::GetResources(TArray< const FSlateBrush* >& OutBrushes) const
+{
+	OutBrushes.Add(&HighlightShape);
+	OutBrushes.Add(&UnderlineBrush);
+}
+
 const FTextBlockStyle& FTextBlockStyle::GetDefault()
 {
 	static TSharedPtr< FTextBlockStyle > Default;
diff --git a/Engine/Source/Runtime/SlateCore/Private/Textures/TextureAtlas.cpp b/Engine/Source/Runtime/SlateCore/Private/Textures/TextureAtlas.cpp
--- a/Engine/Source/Runtime/SlateCore/Private/Textures/TextureAtlas.cpp
+++ b/Engine/Source/Runtime/SlateCore/Private/Textures/TextureAtlas.cpp
@@ -5,7 +5,7 @@
 
 DEFINE_STAT(STAT_SlateTextureGPUMemory);
 DEFINE_STAT(STAT_SlateTextureDataMemory);
-DECLARE_MEMORY_STAT(TEXT("Texture Atlas CPU Memory"), STAT_SlateTextureAtlasMemory, STATGROUP_SlateMemory);
+DECLARE_MEMORY_STAT(TEXT("Texture Atlas Memory (CPU)"), STAT_SlateTextureAtlasMemory, STATGROUP_SlateMemory);
 
 ESlateTextureAtlasThreadId GetCurrentSlateTextureAtlasThreadId()
 {
diff --git a/Engine/Source/Runtime/SlateCore/Public/Fonts/FontCache.h b/Engine/Source/Runtime/SlateCore/Public/Fonts/FontCache.h
--- a/Engine/Source/Runtime/SlateCore/Public/Fonts/FontCache.h
+++ b/Engine/Source/Runtime/SlateCore/Public/Fonts/FontCache.h
@@ -778,6 +778,16 @@ public:
 	 */
 	int16 GetBaseline( const FSlateFontInfo& InFontInfo, float FontScale ) const;
 
+	/**
+	 * Get the underline metrics for the specified font.
+	 *
+	 * @param InFontInfo			A descriptor of the font to get character size for
+	 * @param FontScale				The scale to apply to the font
+	 * @param OutUnderlinePos		The offset from the baseline to the center of the underline bar
+	 * @param OutUnderlineThickness	The thickness of the underline bar
+	 */
+	void GetUnderlineMetrics( const FSlateFontInfo& InFontInfo, const float FontScale, int16& OutUnderlinePos, int16& OutUnderlineThickness ) const;
+
 	/**
 	 * Calculates the kerning amount for a pair of characters
 	 *
