Copying //UE4/Dev-Editor to //UE4/Dev-Main (Source: //UE4/Dev-Editor @ 3082391)

diff --git a/Engine/Source/Runtime/SlateCore/Public/Rendering/DrawElements.h b/Engine/Source/Runtime/SlateCore/Public/Rendering/DrawElements.h
--- a/Engine/Source/Runtime/SlateCore/Public/Rendering/DrawElements.h
+++ b/Engine/Source/Runtime/SlateCore/Public/Rendering/DrawElements.h
@@ -8,6 +8,8 @@
 class SWindow;
 class FSlateViewportInterface;
 
+DECLARE_MEMORY_STAT_EXTERN(TEXT("Vertex/Index Buffer Pool Memory (CPU)"), STAT_SlateBufferPoolMemory, STATGROUP_SlateMemory, SLATECORE_API );
+
 
 struct FSlateGradientStop
 {
@@ -894,6 +896,60 @@ private:
 	uint32 MaxLayer;
 };
 
+#if STATS
+
+class FSlateStatTrackingMemoryAllocator : public FDefaultAllocator
+{
+public:
+	typedef FDefaultAllocator Super;
+
+	class ForAnyElementType : public FDefaultAllocator::ForAnyElementType
+	{
+	public:
+		typedef FDefaultAllocator::ForAnyElementType Super;
+
+		ForAnyElementType()
+			: AllocatedSize(0)
+		{
+
+		}
+
+		/** Destructor. */
+		~ForAnyElementType()
+		{
+			if(AllocatedSize)
+			{
+				DEC_DWORD_STAT_BY(STAT_SlateBufferPoolMemory, AllocatedSize);
+			}
+		}
+
+		void ResizeAllocation(int32 PreviousNumElements, int32 NumElements, int32 NumBytesPerElement)
+		{
+			const int32 NewSize = NumElements * NumBytesPerElement;
+			INC_DWORD_STAT_BY(STAT_SlateBufferPoolMemory, NewSize - AllocatedSize);
+			AllocatedSize = NewSize;
+
+			Super::ResizeAllocation(PreviousNumElements, NumElements, NumBytesPerElement);
+		}
+
+	private:
+		ForAnyElementType(const ForAnyElementType&);
+		ForAnyElementType& operator=(const ForAnyElementType&);
+	private:
+		int32 AllocatedSize;
+	};
+};
+
+typedef TArray<FSlateVertex, FSlateStatTrackingMemoryAllocator> FSlateVertexArray;
+typedef TArray<SlateIndex, FSlateStatTrackingMemoryAllocator> FSlateIndexArray;
+
+#else
+
+typedef TArray<FSlateVertex> FSlateVertexArray;
+typedef TArray<SlateIndex> FSlateIndexArray;
+
+#endif
+
 
 class FSlateBatchData
 {
@@ -923,10 +979,10 @@ public:
 	void AssignIndexArrayToBatch( FSlateElementBatch& Batch );
 
 	/** @return the list of vertices for a batch */
-	TArray<FSlateVertex>& GetBatchVertexList( FSlateElementBatch& Batch ) { return BatchVertexArrays[Batch.VertexArrayIndex]; }
+	FSlateVertexArray& GetBatchVertexList( FSlateElementBatch& Batch ) { return BatchVertexArrays[Batch.VertexArrayIndex]; }
 
 	/** @return the list of indices for a batch */
-	TArray<SlateIndex>& GetBatchIndexList( FSlateElementBatch& Batch ) { return BatchIndexArrays[Batch.IndexArrayIndex]; }
+	FSlateIndexArray& GetBatchIndexList( FSlateElementBatch& Batch ) { return BatchIndexArrays[Batch.IndexArrayIndex]; }
 
 	/** @return The total number of batched vertices */
 	int32 GetNumBatchedVertices() const { return NumBatchedVertices; }
@@ -958,6 +1014,18 @@ private:
 
 	void AddRenderBatch(uint32 InLayer, const FSlateElementBatch& InElementBatch, int32 InNumVertices, int32 InNumIndices, int32 InVertexOffset, int32 InIndexOffset);
 
+	/**
+	 * Resets an array from the pool of vertex arrays
+	 * This will empty the array and give it a reasonable starting memory amount for when it is reused
+	 */
+	void ResetVertexArray(FSlateVertexArray& InOutVertexArray);
+
+	/**
+	* Resets an array from the pool of index arrays
+	* This will empty the array and give it a reasonable starting memory amount for when it is reused
+	*/
+	void ResetIndexArray(FSlateIndexArray& InOutIndexArray);
+
 private:
 
 	// The associated render data handle if these render batches are not in the default vertex/index buffer
@@ -970,10 +1038,10 @@ private:
 	TArray<uint32> IndexArrayFreeList;
 
 	// Array of vertex lists for batching vertices. We use this method for quickly resetting the arrays without deleting memory.
-	TArray<TArray<FSlateVertex>> BatchVertexArrays;
+	TArray<FSlateVertexArray> BatchVertexArrays;
 
 	// Array of vertex lists for batching indices. We use this method for quickly resetting the arrays without deleting memory.
-	TArray<TArray<SlateIndex>> BatchIndexArrays;
+	TArray<FSlateIndexArray> BatchIndexArrays;
 
 	/** List of element batches sorted by later for use in rendering (for threaded renderers, can only be accessed from the render thread)*/
 	TArray<FSlateRenderBatch> RenderBatches;
diff --git a/Engine/Source/Runtime/SlateCore/Public/Styling/SlateTypes.h b/Engine/Source/Runtime/SlateCore/Public/Styling/SlateTypes.h
--- a/Engine/Source/Runtime/SlateCore/Public/Styling/SlateTypes.h
+++ b/Engine/Source/Runtime/SlateCore/Public/Styling/SlateTypes.h
@@ -189,6 +189,8 @@ struct SLATECORE_API FTextBlockStyle : public FSlateWidgetStyle
 
 	virtual ~FTextBlockStyle() {}
 
+	virtual void GetResources( TArray< const FSlateBrush* >& OutBrushes ) const override;
+
 	static const FName TypeName;
 	virtual const FName GetTypeName() const override { return TypeName; };
 
@@ -240,6 +242,11 @@ struct SLATECORE_API FTextBlockStyle : public FSlateWidgetStyle
 	UPROPERTY(EditAnywhere, BlueprintReadWrite, Category=Appearance, AdvancedDisplay)
 	FSlateBrush HighlightShape;
 	FTextBlockStyle& SetHighlightShape( const FSlateBrush& InHighlightShape ){ HighlightShape = InHighlightShape; return *this; }
+
+	/** The brush used to draw an underline under the text (if any) */
+	UPROPERTY(EditAnywhere, BlueprintReadWrite, Category=Appearance, AdvancedDisplay)
+	FSlateBrush UnderlineBrush;
+	FTextBlockStyle& SetUnderlineBrush( const FSlateBrush& InUnderlineBrush ){ UnderlineBrush = InUnderlineBrush; return *this; }
 };
 
 /**
diff --git a/Engine/Source/Runtime/SlateCore/Public/Textures/SlateShaderResource.h b/Engine/Source/Runtime/SlateCore/Public/Textures/SlateShaderResource.h
--- a/Engine/Source/Runtime/SlateCore/Public/Textures/SlateShaderResource.h
+++ b/Engine/Source/Runtime/SlateCore/Public/Textures/SlateShaderResource.h
@@ -4,7 +4,7 @@
 
 class FSlateShaderResourceProxy;
 
-DECLARE_MEMORY_STAT_EXTERN(TEXT("Texture Data GPU Memory"), STAT_SlateTextureGPUMemory, STATGROUP_SlateMemory, SLATECORE_API);
+DECLARE_MEMORY_STAT_EXTERN(TEXT("Texture Data Memory (GPU)"), STAT_SlateTextureGPUMemory, STATGROUP_SlateMemory, SLATECORE_API);
 
 namespace ESlateShaderResource
 {
diff --git a/Engine/Source/Runtime/SlateCore/Public/Textures/SlateTextureData.h b/Engine/Source/Runtime/SlateCore/Public/Textures/SlateTextureData.h
--- a/Engine/Source/Runtime/SlateCore/Public/Textures/SlateTextureData.h
+++ b/Engine/Source/Runtime/SlateCore/Public/Textures/SlateTextureData.h
@@ -2,7 +2,7 @@
 
 #pragma once
 
-DECLARE_MEMORY_STAT_EXTERN(TEXT("Texture Data CPU Memory"), STAT_SlateTextureDataMemory, STATGROUP_SlateMemory, SLATECORE_API);
+DECLARE_MEMORY_STAT_EXTERN(TEXT("Texture Data Memory (CPU)"), STAT_SlateTextureDataMemory, STATGROUP_SlateMemory, SLATECORE_API);
 
 /**
  * Holds texture data for upload to a rendering resource
diff --git a/Engine/Source/Runtime/SlateCore/Public/Widgets/SWindow.h b/Engine/Source/Runtime/SlateCore/Public/Widgets/SWindow.h
--- a/Engine/Source/Runtime/SlateCore/Public/Widgets/SWindow.h
+++ b/Engine/Source/Runtime/SlateCore/Public/Widgets/SWindow.h
@@ -1158,7 +1158,7 @@ private:
 #if WITH_EDITOR
 
 /**
- * Hack to switch wolds in a scope and switch back when we fall out of scope                   
+ * Hack to switch worlds in a scope and switch back when we fall out of scope                   
  */
 struct FScopedSwitchWorldHack
 {
diff --git a/Engine/Source/Runtime/SlateRHIRenderer/Private/SlateElementIndexBuffer.cpp b/Engine/Source/Runtime/SlateRHIRenderer/Private/SlateElementIndexBuffer.cpp
--- a/Engine/Source/Runtime/SlateRHIRenderer/Private/SlateElementIndexBuffer.cpp
+++ b/Engine/Source/Runtime/SlateRHIRenderer/Private/SlateElementIndexBuffer.cpp
@@ -3,7 +3,7 @@
 #include "SlateRHIRendererPrivatePCH.h"
 #include "SlateElementIndexBuffer.h"
 
-DECLARE_MEMORY_STAT(TEXT("Index Buffer Memory"), STAT_SlateIndexBufferMemory, STATGROUP_SlateMemory);
+DECLARE_MEMORY_STAT(TEXT("Index Buffer Memory (GPU)"), STAT_SlateIndexBufferMemory, STATGROUP_SlateMemory);
 
 FSlateElementIndexBuffer::FSlateElementIndexBuffer()
 	: BufferSize(0)	 
diff --git a/Engine/Source/Runtime/SlateRHIRenderer/Private/SlateElementVertexBuffer.h b/Engine/Source/Runtime/SlateRHIRenderer/Private/SlateElementVertexBuffer.h
--- a/Engine/Source/Runtime/SlateRHIRenderer/Private/SlateElementVertexBuffer.h
+++ b/Engine/Source/Runtime/SlateRHIRenderer/Private/SlateElementVertexBuffer.h
@@ -2,7 +2,7 @@
 
 #pragma once
 
-DECLARE_MEMORY_STAT_EXTERN(TEXT("Vertex Buffer Memory"), STAT_SlateVertexBufferMemory, STATGROUP_SlateMemory, );
+DECLARE_MEMORY_STAT_EXTERN(TEXT("Vertex Buffer Memory (GPU)"), STAT_SlateVertexBufferMemory, STATGROUP_SlateMemory, );
 
 /**
  * Vertex buffer containing all Slate vertices
diff --git a/Engine/Source/Runtime/SlateRHIRenderer/Private/SlateRHIRenderer.cpp b/Engine/Source/Runtime/SlateRHIRenderer/Private/SlateRHIRenderer.cpp
--- a/Engine/Source/Runtime/SlateRHIRenderer/Private/SlateRHIRenderer.cpp
+++ b/Engine/Source/Runtime/SlateRHIRenderer/Private/SlateRHIRenderer.cpp
@@ -590,12 +590,20 @@ void FSlateRHIRenderer::DrawWindows_Private( FSlateDrawBuffer& WindowDrawBuffer
 
 				bLockToVsync = ElementBatcher->RequiresVsync();
 
-				if( !GIsEditor )
+				bool bForceVsyncFromCVar = false;
+				if(GIsEditor)
+				{
+					static IConsoleVariable* CVar = IConsoleManager::Get().FindConsoleVariable(TEXT("r.VSyncEditor"));
+					bForceVsyncFromCVar = (CVar->GetInt() != 0);
+				}
+				else
 				{
 					static IConsoleVariable* CVar = IConsoleManager::Get().FindConsoleVariable(TEXT("r.VSync"));
-					bLockToVsync = bLockToVsync || (CVar->GetInt() != 0);
+					bForceVsyncFromCVar = (CVar->GetInt() != 0);
 				}
 
+				bLockToVsync |= bForceVsyncFromCVar;
+
 				// All elements for this window have been batched and rendering data updated
 				ElementBatcher->ResetBatches();
 
diff --git a/Engine/Source/Runtime/UMG/Public/Animation/WidgetAnimation.h b/Engine/Source/Runtime/UMG/Public/Animation/WidgetAnimation.h
--- a/Engine/Source/Runtime/UMG/Public/Animation/WidgetAnimation.h
+++ b/Engine/Source/Runtime/UMG/Public/Animation/WidgetAnimation.h
@@ -85,7 +85,10 @@ public:
 	 */
 	UMG_API void ReplacePossessableObject(const FGuid& OldId, const FGuid& NewId, UObject& OldObject, UObject& NewObject);
 
+	/** Get Animation bindings of the animation */
 	const TArray<FWidgetAnimationBinding>& GetBindings() const { return AnimationBindings; }
+	/** Get the preview widget of the animation, if any */
+	const TWeakObjectPtr<UUserWidget> GetPreviewWidget() const { return PreviewWidget; }
 public:
 
 	/** Pointer to the movie scene that controls this animation. */
diff --git a/Engine/Source/Runtime/Windows/D3D11RHI/Private/D3D11Viewport.cpp b/Engine/Source/Runtime/Windows/D3D11RHI/Private/D3D11Viewport.cpp
--- a/Engine/Source/Runtime/Windows/D3D11RHI/Private/D3D11Viewport.cpp
+++ b/Engine/Source/Runtime/Windows/D3D11RHI/Private/D3D11Viewport.cpp
@@ -41,11 +41,11 @@ namespace RHIConsoleVariables
 		ECVF_RenderThreadSafe
 		);
 
-	int32 bForceThirtyHz = 1;
-	static FAutoConsoleVariableRef CVarForceThirtyHz(
-		TEXT("RHI.ForceThirtyHz"),
-		bForceThirtyHz,
-		TEXT("If true, the display will never update more often than 30Hz."),
+	int32 TargetRefreshRate = 0;
+	static FAutoConsoleVariableRef CVarTargetRefreshRate(
+		TEXT("RHI.TargetRefreshRate"),
+		TargetRefreshRate,
+		TEXT("If non-zero, the display will never update more often than the target refresh rate (in Hz)."),
 		ECVF_RenderThreadSafe
 		);
 
@@ -288,7 +288,8 @@ void FD3D11Viewport::PresentWithVsyncDWM()
 	QueryPerformanceCounter(&Cycles);
 	FMemory::Memzero(TimingInfo);
 	TimingInfo.cbSize = sizeof(DWM_TIMING_INFO);
-	DwmGetCompositionTimingInfo(WindowHandle, &TimingInfo);
+	// Starting at windows 8.1 null must be passed into this method for it to work.  null also works on previous versions
+	DwmGetCompositionTimingInfo(nullptr, &TimingInfo);
 
 	uint64 QpcAtFlip = Cycles.QuadPart;
 	uint64 CyclesSinceLastFlip = Cycles.QuadPart - LastFlipTime;
@@ -298,14 +299,15 @@ void FD3D11Viewport::PresentWithVsyncDWM()
 
 	// Find the smallest multiple of the refresh rate that is >= 33ms, our target frame rate.
 	float RefreshPeriod = DisplayRefreshPeriod;
-	if (RHIConsoleVariables::bForceThirtyHz && RefreshPeriod > 1.0f)
+	if(RHIConsoleVariables::TargetRefreshRate > 0 && RefreshPeriod > 1.0f)
 	{
-		while (RefreshPeriod - (1000.0f / 30.0f) < -1.0f)
+		while(RefreshPeriod - (1000.0f / RHIConsoleVariables::TargetRefreshRate) < -1.0f)
 		{
 			RefreshPeriod *= 2.0f;
 		}
 	}
 
+
 	// If the last frame hasn't completed yet, we don't know how long the GPU took.
 	bool bValidGPUTime = (TimingInfo.cFrameComplete > LastFrameComplete);
 	if (bValidGPUTime)
