Copying //UE4/WEX-Staging to Dev-Main (//UE4/Dev-Main)

===================
MAJOR FEATURES + CHANGES
===================

Change 3123735 on 2016/09/13 by [REDACTED_PERSON]

	#UE4 - added FNames for some common online features to be refactored later

Change 3123608 on 2016/09/13 by [REDACTED_PERSON]

	#OSS - refactor FUniqueNetIdRepl to derive from FUniqueNetIdWrapper and clean up redundant code

Change 3120074 on 2016/09/09 by [REDACTED_PERSON]

	#Analytics: Move log message to more correct location

Change 3120073 on 2016/09/09 by [REDACTED_PERSON]

	#Analytics: Adjustments to ET.DroppedSubmission per Wes's feedback
	- Don't try to restore events at all if caching is disabled (but do log the error)
	- Stop accumulating events (even the ET.DroppedSubmission event) at 1024 cached to defend against run away accumulation

Change 3119959 on 2016/09/09 by [REDACTED_PERSON]

	update to engine to provide the generic chunk install when one is not implemented or requested

Change 3119378 on 2016/09/09 by [REDACTED_PERSON]

	#Analytics #ET: Added a config option to not drop events in the event of a failure to flush.
	- Events are re-added to the queue (may result in some reordering but timestamps will still be there) in the event of a failure
	- Also appending an ET.DroppedSubmission event in this case so we can track how often clients fail. Has HTTP_STATUS and URL attributes.
	- Events accumulate within a given application run but no attempt is made to persist across crashes etc.

Change 3118773 on 2016/09/08 by [REDACTED_PERSON]

	bring over fix for missing debug information in IOS executables

Change 3118574 on 2016/09/08 by [REDACTED_PERSON]

	pulled over architecture fix for IOS from Main

Change 3117672 on 2016/09/08 by [REDACTED_PERSON]

	Adding 3544_iPadMini3_ChAIRQA

Change 3116529 on 2016/09/07 by [REDACTED_PERSON]

	#UE4 - reverted started IOS purchase/store work from default engine OSS plugins

Change 3116010 on 2016/09/07 by [REDACTED_PERSON]

	#WEX - Copying //WEX/Dev-Mobile to Dev-Main (//WEX/Dev-Main) to get OnlineSubsystemiOS changes

Change 3114411 on 2016/09/06 by [REDACTED_PERSON]

	fix for deploying movies to the correct directory

Change 3113944 on 2016/09/06 by [REDACTED_PERSON]

	addition of resave packages command to UAT

Change 3112948 on 2016/09/02 by [REDACTED_PERSON]

	#WEX
	- Fixing file length on Android so that we can write out stats files (change is being submitted into Dev-Mobile by [REDACTED_PERSON])

Change 3112567 on 2016/09/02 by [REDACTED_PERSON]

	#UE4 - fixed possible exception when checking directories that don't exist  (copy of CL#3099217)

Change 3112055 on 2016/09/02 by [REDACTED_PERSON]

	#WEX
	- Merging mobile branch changes into main

Change 3108827 on 2016/08/31 by [REDACTED_PERSON]

	fix for warning

Change 3105012 on 2016/08/29 by [REDACTED_PERSON]

	enable ICMP for Android

Change 3103322 on 2016/08/26 by [REDACTED_PERSON]

	#WEX: Fix for clang warnings in Font Outline code

Change 3102935 on 2016/08/26 by [REDACTED_PERSON]

	Merging CL 3102878 from
	//UE4/Dev-Editor/Engine/...
	to //WEX/Main/Engine/...
	------
	Added support for outline fonts
	- An outline size (in slate units), optional material and optional fill color can be specified with each font info.
	- Outlines do not contribute to measurement directly so the text measuring and shaping methods have been modified to account for outlines
	- Fixed a bug where font materials do not work properly if part of the font's rendered glyphs were in a different atlas

Change 3102541 on 2016/08/26 by [REDACTED_PERSON]

	fix for build warning on Mac and IOS

Change 3101820 on 2016/08/25 by [REDACTED_PERSON]

	Moved the online changes to the plugins directory

Change 3101808 on 2016/08/25 by [REDACTED_PERSON]

	Merging using WEX_Main_to_UE4_WEX_Staging bringing in UE4 engine 4.14

Change 3097879 on 2016/08/23 by [REDACTED_PERSON]

	#Analytics #OSS: Adjusted cohort selection algorithm and test cases

Change 3096606 on 2016/08/22 by [REDACTED_PERSON]

	#WEX
	- Adding bAllowWindowResize, bAllowClose, bAllowMaximize, and bAllowMinimize boolean to GeneralProjectSettings
	- Preventing us from resizing the window for the time being (in the future we may only allow along aspect ratio resizing)

Change 3094946 on 2016/08/19 by [REDACTED_PERSON]

	#OSS - Added FUserOnlineAccountMcp::SelectCohort

Change 3094942 on 2016/08/19 by [REDACTED_PERSON]

	#UE4 - Made FMD5 const-correct

Change 3092494 on 2016/08/17 by [REDACTED_PERSON]

	#WEX
	- Making sure we never build zip64 instead of zip

Change 3090760 on 2016/08/16 by [REDACTED_PERSON]

	Copy of CL 3089133

	Fix for task bar displayed over the fullscreen window on Windows 10 with Anniversary Update

Change 3090759 on 2016/08/16 by [REDACTED_PERSON]

	Copy of CL 3078927

	Updated WindowTitleBarArea widget to not override window zone in fullscreen mode, to prevent window from being moved. That required adding separate handling for double click in fullscreen, as it's no longer handled by window action.

Change 3087872 on 2016/08/12 by [REDACTED_PERSON]

	#UE4 - cleaned up IOS store/purchase interface (first pass, minus IAP restore)

Change 3084182 on 2016/08/10 by [REDACTED_PERSON]

	revert out the OpenGL shader compression code

Change 3082565 on 2016/08/09 by [REDACTED_PERSON]

	Fix building with VS2015 update 3.

Change 3082557 on 2016/08/09 by [REDACTED_PERSON]

	Fix UBT makefile being invalidated to update adaptive unity build settings, even if that module happens to not include that file in a unity file. Keep a list of all files included by unity files as well as files in the working set.

Change 3082456 on 2016/08/09 by [REDACTED_PERSON]

	#UE4 - fixed typo

Change 3082439 on 2016/08/09 by [REDACTED_PERSON]

	#UE4 - added CanMakePurchase call to IOS

Change 3081905 on 2016/08/09 by [REDACTED_PERSON]

	Editor: Made the text colors and font size in the output log configurable in the editor appearance settings (no changes to default values ... yet)

Change 3080932 on 2016/08/08 by [REDACTED_PERSON]

	#UE4 - New IOS purchasing/store interface v2
	- added interfaces to main IOS subsystem
	- added proper destruction of interfaces to Shutdown of IOS
	- moved FStoreKitHelper to its own file
	-- extended it for new v2 (improvements forthcoming)
	- MCPCatalogHelper returns bogus user id for IOS app store
	#tests very basic IAP stuff so far

Change 3080217 on 2016/08/07 by [REDACTED_PERSON]

	Engine: Prevented a startup warning when SpectatorClass is nullptr, as not all games require a spectator pawn
	- Also reduced the number of GetWorld() calls in APlayerController::SpawnSpectatorPawn()

Change 3080046 on 2016/08/06 by [REDACTED_PERSON]

	Engine: Moved where scissor rect reset happens for custom slate drawables to avoid a conflict with an existing fix in another branch

Change 3080032 on 2016/08/06 by [REDACTED_PERSON]

	UMG: Fixed a bug where screen-mode UWidgetComponent widgets were drawn incorrectly offset in splitscreen or with aspect-ratio constrained cameras

Change 3080031 on 2016/08/06 by [REDACTED_PERSON]

	Engine: Add the option to return player viewport-relative positions to ProjectWorldLocationToScreenWithDistance, ProjectWorldLocationToScreen, and ProjectWorldToScreen, which is useful if the position is going to be used for widgets in splitscreen or with aspect-ratio constrained cameras

Change 3080029 on 2016/08/06 by [REDACTED_PERSON]

	Engine: Fixed a bug where the debug console and other debug rendering would be an incorrect size (based on the last player viewport) and also be partially clipped (depending on what in Slate rendered previously)
	- This fixes issues with the console being offset and clipped when using aspect ratio constrained cameras or split screen

Change 3079656 on 2016/08/05 by [REDACTED_PERSON]

	#WEX - basic IOS changes to project
	- added OnlineSubsystemIOS
	- added some default settings
	- removed GoogleVR from project

Change 3078971 on 2016/08/05 by [REDACTED_PERSON]

	Updating to match check-in for UE4 Main @ CL 3078968

Change 3078025 on 2016/08/04 by [REDACTED_PERSON]

	Copy of CLs 3073978 and 3075931

	- More reliable way of checking if the cursor should be changed to resize cursor in bordeless window mode
	- On Windows, lock the cursor to the center of the rect if the cursor is hidden to avoid problems with borderless window's round corners not treated as part of the window.

Change 3075415 on 2016/08/03 by [REDACTED_PERSON]

	reduce the metal command buffers

Change 3071457 on 2016/07/31 by [REDACTED_PERSON]

	#WEX
	Blueprint indexing

	[REDACTED_HANDLE]: This change and anything in Engine/Content can be stomped by any engine integration. This is just to help with not having to resave all of these to udpate a few of our own content blueprints for search indexing.

Change 3068661 on 2016/07/28 by [REDACTED_PERSON]

	#WEX - changed the max number of possible UObjects allowed by a factor of 10 to reduce the memory footprint

Change 3068500 on 2016/07/28 by [REDACTED_PERSON]

	#OSSMCP: Use correct HttpRequest creation method to respect game service config

Change 3066945 on 2016/07/27 by [REDACTED_PERSON]

	Reproduced CL 3063869 from Michael.Noland
	>> Engine: Added a cvar (t.FPSChart.OpenFolderOnDump) to control whether or not FPS charts automatically open the profiling folder when stopfpschart is executed, which can be useful to avoid a bunch of open
	>> windows while doing automated testing

Change 3063495 on 2016/07/25 by [REDACTED_PERSON]

	Copy of CL 3063426

	Borderless window support improvements:

	- the cursor changes to resize when hovering over the window edge
	- added a way for widgets to register a delegate that's called when window actions occur (maximize, restore, etc.)
	- used window action notification for WindowTitleBarArea to improve how toggling fullscreen on double click is handled

Change 3063431 on 2016/07/25 by [REDACTED_PERSON]

	Copy of CL 3063057 - Use round corners for windows with no system title bar and border only in windowed mode.

Change 3062654 on 2016/07/23 by [REDACTED_PERSON]

	Copy of CL 3046975 and 3056204 - Support for making the game window borderless (no system border or title bar). Disabled by default. Enabling requires adding bUseBorderlessWindow=True to [/Script/EngineSettings.GeneralProjectSettings] in DefaultGame.ini. The game using this is responsible for adding WindowTitleBarArea widget to its UI, as well as window minimize/maximize/close buttons.

Change 3062647 on 2016/07/23 by [REDACTED_PERSON]

	Copy of CL 3029211 - Added a setting (on by default) to make the game window preserve its content's aspect ratio while being resized by user

Change 3062646 on 2016/07/23 by [REDACTED_PERSON]

	Copy of CLs 3039855, 3042644 and 3042911 - Added an option to toggle fullscreen with F11 key in addition to Alt+Enter

Change 3062638 on 2016/07/23 by [REDACTED_PERSON]

	Copy of CL 3038201 and CL 3046803 -Added WindowTitleBarArea widget

Change 3062056 on 2016/07/22 by [REDACTED_PERSON]

	addition of optimization for ios compile times

Change 3054586 on 2016/07/18 by [REDACTED_PERSON]

	#WEX
	- Adding tags around my engine level change

Change 3054581 on 2016/07/18 by [REDACTED_PERSON]

	#WEX
	- Removing previous change, making all buttons ignore space bar and enter as we don't really care about that functionality

Change 3048243 on 2016/07/13 by [REDACTED_PERSON]

	This one actually has the changes from rev4

Change 3046649 on 2016/07/12 by [REDACTED_PERSON]

	Updating to match provision in UE4 Main @ CL 3046262

Change 3046127 on 2016/07/12 by [REDACTED_PERSON]

	#UE4, #OnlineSubSystem - Hotfix in the ExpirationDate field early so we can update the OGF plugin

Change 3034707 on 2016/06/30 by [REDACTED_PERSON]

	update provision to go with latest certificate

Change 3031429 on 2016/06/28 by [REDACTED_PERSON]

	#WEX: porting an engine change we depend on in latest OGF (pre integrate)

Change 3030084 on 2016/06/27 by [REDACTED_PERSON]

	#GameCatalog: Add code to export attributes in itemGrants

Change 3030073 on 2016/06/27 by [REDACTED_PERSON]

	#Json: Make JsonObjectWrapper play nice with serialization and Import/Export Text

Change 3030029 on 2016/06/27 by [REDACTED_PERSON]

	#WEX: Fix for FJsonObjectWrapper::ImportTextItem (use FParse::QuotedString to read from Buffer)

Change 3029740 on 2016/06/27 by [REDACTED_PERSON]

	#OGF #JsonObjectWrapper
	- Add attributes to catalog grants
	[REDACTED_HANDLE] - I didn't get a chance to test this today. Things are crazy for PS+. All the code should already be there on the backend though. Here's a shelf in case you need it asap

Change 3028704 on 2016/06/27 by [REDACTED_PERSON]

	Duplicating 3027482 from //Orion/Main

	Read TaggedPropertyRedirects from all config files to allow plugins to register property redirectors

	You'll need this before you grab the latest OGF or your catalog prices will go away, so here it is now

	#ue4
	#tests none

Change 3021448 on 2016/06/21 by [REDACTED_PERSON]

	potential fix for android apk size issue

Change 3020999 on 2016/06/21 by [REDACTED_PERSON]

	#WEX: Likely fix for WEX-1610

Change 3008450 on 2016/06/09 by [REDACTED_PERSON]

	#PF
	- WEX-1737, WEX-1744
	- Adding the ability to set new layers in widget components
	- WidgetComponents are now blueprintable
	- New blueprint for the level marker menu widget components
	- Map marker menus are now in a layer above other widget components

Change 3007804 on 2016/06/09 by [REDACTED_PERSON]

	fix for build set up failure

Change 3007292 on 2016/06/09 by [REDACTED_PERSON]

	add the WEX e-mail stuff back in, seems to have gotten lost in the transition

Change 3004478 on 2016/06/07 by [REDACTED_PERSON]

	for now have the cooker respect the bCookAll flag in the project settings.

Change 3000256 on 2016/06/03 by [REDACTED_PERSON]

	fix for iOS compile warning

Change 2998304 on 2016/06/02 by [REDACTED_PERSON]

	#PF
	- Fixing Windowed Mode

Change 2994269 on 2016/05/31 by [REDACTED_PERSON]

	Merging

	//depot/UE4-WEX/...

	to //WEX/Main/...

Change 2987181 on 2016/05/23 by [REDACTED_PERSON]

	Merging

	//UE4/WEX-Staging/.p4ignore.txt
	//UE4/WEX-Staging/Engine/...
	//UE4/WEX-Staging/GenerateProjectFiles.bat
	//UE4/WEX-Staging/GenerateProjectFiles.command
	//UE4/WEX-Staging/GenerateProjectFiles.sh
	//UE4/WEX-Staging/UE4Games.uprojectdirs
	//UE4/WEX-Staging/WEX/...

	to //WEX/Main/...

Change 2984959 on 2016/05/20 by [REDACTED_PERSON]

	re-applying HSL engine change

	#PF
	PF-292
	- Make sure to regenerate the list when you open the recipe view
	- Fixes cases where you buy an item in the store then return to evolve

Change 2984957 on 2016/05/20 by [REDACTED_PERSON]

	re-apply engine change from HSL

	#PF
	PF-33
	- Check if we are in BeginDestroyed on Animation updates, possibly fixes a crash on level transition
	- Make interactive items play their mouseover animations
	- Heroes make the screen shake if they are on cooldown, we should figure out how to remove that on PC

Change 2984956 on 2016/05/20 by [REDACTED_PERSON]

	re-apply HSL change

	#PF
	PF-11
	- Remove simulated touch with the mouse
	- Add Right-click support to the game
	- Right click now does the special attack

Change 2984345 on 2016/05/19 by [REDACTED_PERSON]

	Copying

	//UE4/WEX-Staging/...

	to //WEX/Main/...

	This should reset the merge history

Change 2981872 on 2016/05/18 by [REDACTED_PERSON]

	fixes for IOS build of WEX

Change 2980734 on 2016/05/17 by [REDACTED_PERSON]

	Copying

	//depot/UE4-WEX/...

	to //WEX/Main/...

	Populating WEX stream from old depot at CL2979954

diff --git a/Engine/Config/BaseDeviceProfiles.ini b/Engine/Config/BaseDeviceProfiles.ini
--- a/Engine/Config/BaseDeviceProfiles.ini
+++ b/Engine/Config/BaseDeviceProfiles.ini
@@ -118,9 +118,9 @@ BaseProfileName=
 +CVars=r.PostProcessAAQuality=0
 // Not sure on this one - may need the whole D Buffer stuff going
 +CVars=r.Decal.StencilSizeThreshold=-1
-+CVars=r.MetalVertexParameterSize=100
-+CVars=r.MetalPixelParameterSize=100
-+CVars=r.MetalComputeParameterSize=100
++CVars=r.MetalVertexParameterSize=50
++CVars=r.MetalPixelParameterSize=50
++CVars=r.MetalComputeParameterSize=0
 
 [iPad2 DeviceProfile]
 DeviceType=IOS
diff --git a/Engine/Config/BaseEditorPerProjectUserSettings.ini b/Engine/Config/BaseEditorPerProjectUserSettings.ini
--- a/Engine/Config/BaseEditorPerProjectUserSettings.ini
+++ b/Engine/Config/BaseEditorPerProjectUserSettings.ini
@@ -73,6 +73,14 @@ PressedSelectionColor=(R=0.701f,G=0.225f,B=0.003f,A=1.000000)
 InactiveSelectionColor=(R=0.25f,G=0.25f,B=0.25f,A=1.000000)
 SelectorColor=(R=0.701f,G=0.225f,B=0.003f,A=1.000000)
 
+LogBackgroundColor=(R=0.047059,G=0.047059,B=0.047059,A=1.000000)
+LogSelectionBackgroundColor=(R=0.132868,G=0.132868,B=0.132868,A=1.000000)
+LogNormalColor=(R=0.401978,G=0.401978,B=0.401978,A=1.000000)
+LogCommandColor=(R=0.033105,G=0.723055,B=0.033105,A=1.000000)
+LogWarningColor=(R=0.496933,G=0.496933,B=0.057805,A=1.000000)
+LogErrorColor=(R=0.723055,G=0.000000,B=0.000000,A=1.000000)
+LogFontSize=9
+
 ; If true, all toolbars will use small icons without labels
 bUseSmallToolBarIcons=False
 ; Enables animated transitions for certain menus and pop-up windows.  Note that animations may be automatically disabled at low frame rates in order to improve responsiveness.
diff --git a/Engine/Config/BaseEngine.ini b/Engine/Config/BaseEngine.ini
--- a/Engine/Config/BaseEngine.ini
+++ b/Engine/Config/BaseEngine.ini
@@ -2493,6 +2493,7 @@ bAllowRemoteRotation=True
 bUseRemoteAsVirtualJoystick=True
 bUseAbsoluteDpadValues=False
 bGeneratedSYMFile=False
+bDisableHTTPS=false
 bUseRSync=True
 BundleDisplayName=[PROJECT_NAME]
 BundleName=[PROJECT_NAME]
diff --git a/Engine/Config/IOS/IOSEngine.ini b/Engine/Config/IOS/IOSEngine.ini
--- a/Engine/Config/IOS/IOSEngine.ini
+++ b/Engine/Config/IOS/IOSEngine.ini
@@ -1,10 +1,13 @@
+[/Script/Engine.GarbageCollectionSettings]
+gc.MaxObjectsInGame=131072
+
 [Audio]
 AudioDeviceModuleName=IOSAudio
 ; Defines a platform-specific volume headroom (in dB) for audio to provide better platform consistency with respect to volume levels.
 PlatformHeadroomDB=0
 
 [OnlineSubsystem]
-DefaultPlatformService=IOS
+DefaultPlatformService=MCP
 
 [Advertising]
 DefaultProviderName=IOSAdvertising
@@ -12,8 +15,13 @@ DefaultProviderName=IOSAdvertising
 [DeviceProfileManager]
 DeviceProfileSelectionModule="IOSDeviceProfileSelector"
 
+[/Script/IOSRuntimeSettings.IOSRuntimeSettings]
+bEnableGameCenterSupport=true
+bEnableCloudKitSupport=true
+
 [OnlineSubsystemIOS.Store]
-bSupportInAppPurchasing=False
+bSupportsInAppPurchasing=true
+bUseStoreV2=true
 
 [SlateRenderer]
 NumPreallocatedVertices=200
diff --git a/Engine/Source/Developer/iOS/IOSPlatformEditor/Private/IOSTargetSettingsCustomization.cpp b/Engine/Source/Developer/iOS/IOSPlatformEditor/Private/IOSTargetSettingsCustomization.cpp
--- a/Engine/Source/Developer/iOS/IOSPlatformEditor/Private/IOSTargetSettingsCustomization.cpp
+++ b/Engine/Source/Developer/iOS/IOSPlatformEditor/Private/IOSTargetSettingsCustomization.cpp
@@ -130,7 +130,7 @@ void FIOSTargetSettingsCustomization::UpdateStatus()
 
 		// Now split up the log into multiple lines
 		TArray<FString> LogLines;
-		OutputMessage.ParseIntoArray(LogLines, LINE_TERMINATOR, true);
+		OutputMessage.ParseIntoArray(LogLines, TEXT("\n"), true);
 		
 		// format of the line being read here!!
 		bool bCerts = false;
diff --git a/Engine/Source/Editor/DetailCustomizations/Private/SlateFontInfoCustomization.cpp b/Engine/Source/Editor/DetailCustomizations/Private/SlateFontInfoCustomization.cpp
--- a/Engine/Source/Editor/DetailCustomizations/Private/SlateFontInfoCustomization.cpp
+++ b/Engine/Source/Editor/DetailCustomizations/Private/SlateFontInfoCustomization.cpp
@@ -84,7 +84,9 @@ void FSlateFontInfoStructCustomization::CustomizeHeader(TSharedRef<IPropertyHand
 
 void FSlateFontInfoStructCustomization::CustomizeChildren(TSharedRef<IPropertyHandle> InStructPropertyHandle, IDetailChildrenBuilder& InStructBuilder, IPropertyTypeCustomizationUtils& InStructCustomizationUtils)
 {
-	InStructBuilder.AddChildProperty( InStructPropertyHandle->GetChildHandle( TEXT("FontMaterial") ).ToSharedRef() );
+	InStructBuilder.AddChildProperty(InStructPropertyHandle->GetChildHandle( TEXT("FontMaterial")).ToSharedRef());
+
+	InStructBuilder.AddChildProperty(InStructPropertyHandle->GetChildHandle(TEXT("OutlineSettings")).ToSharedRef());
 }
 
 bool FSlateFontInfoStructCustomization::OnFilterFontAsset(const FAssetData& InAssetData)
diff --git a/Engine/Source/Editor/EditorStyle/Private/SlateEditorStyle.cpp b/Engine/Source/Editor/EditorStyle/Private/SlateEditorStyle.cpp
--- a/Engine/Source/Editor/EditorStyle/Private/SlateEditorStyle.cpp
+++ b/Engine/Source/Editor/EditorStyle/Private/SlateEditorStyle.cpp
@@ -60,6 +60,13 @@ FSlateEditorStyle::FStyle::FStyle( const TWeakObjectPtr< UEditorStyleSettings >&
 	, SelectionColor_Inactive_LinearRef( MakeShareable( new FLinearColor( 0.25f, 0.25f, 0.25f ) ) )
 	, SelectionColor_Pressed_LinearRef( MakeShareable( new FLinearColor( 0.701f, 0.225f, 0.003f ) ) )
 
+	, LogColor_Background_LinearRef(MakeShareable(new FLinearColor(FColor(0xFF3E3E3E))))
+	, LogColor_SelectionBackground_LinearRef(MakeShareable(new FLinearColor(FColor(0xff666666))))
+	, LogColor_Normal_LinearRef(MakeShareable(new FLinearColor(FColor(0xffaaaaaa))))
+	, LogColor_Command_LinearRef(MakeShareable(new FLinearColor(FColor(0xff33dd33))))
+	, LogColor_Warning_LinearRef(MakeShareable(new FLinearColor(FColor(0xffbbbb44))))
+	, LogColor_Error_LinearRef(MakeShareable(new FLinearColor(FColor(0xffdd0000))))
+
 	// These are the Slate colors which reference those above; these are the colors to put into the style
 	, DefaultForeground( DefaultForeground_LinearRef )
 	, InvertedForeground( InvertedForeground_LinearRef )
@@ -69,6 +76,13 @@ FSlateEditorStyle::FStyle::FStyle( const TWeakObjectPtr< UEditorStyleSettings >&
 	, SelectionColor_Inactive( SelectionColor_Inactive_LinearRef )
 	, SelectionColor_Pressed( SelectionColor_Pressed_LinearRef )
 
+	, LogColor_Background( LogColor_Background_LinearRef )
+	, LogColor_SelectionBackground( LogColor_SelectionBackground_LinearRef )
+	, LogColor_Normal( LogColor_Normal_LinearRef )
+	, LogColor_Command( LogColor_Command_LinearRef )
+	, LogColor_Warning( LogColor_Warning_LinearRef )
+	, LogColor_Error( LogColor_Error_LinearRef )
+
 	, InheritedFromBlueprintTextColor(FLinearColor(0.25f, 0.5f, 1.0f))
 
 	, Settings( InSettings )
@@ -102,6 +116,13 @@ void FSlateEditorStyle::FStyle::SyncSettings()
 		SetColor( SelectionColor_Inactive_LinearRef, Settings->InactiveSelectionColor );
 		SetColor( SelectionColor_Pressed_LinearRef, Settings->PressedSelectionColor );
 
+		SetColor( LogColor_Background_LinearRef, Settings->LogBackgroundColor );
+		SetColor( LogColor_SelectionBackground_LinearRef, Settings->LogSelectionBackgroundColor );
+		SetColor( LogColor_Normal_LinearRef, Settings->LogNormalColor );
+		SetColor( LogColor_Command_LinearRef, Settings->LogCommandColor );
+		SetColor( LogColor_Warning_LinearRef, Settings->LogWarningColor );
+		SetColor( LogColor_Error_LinearRef, Settings->LogErrorColor );
+
 		// The subdued selection color is derived from the selection color
 		auto SubduedSelectionColor = Settings->GetSubduedSelectionColor();
 		SetColor( SelectionColor_Subdued_LinearRef, SubduedSelectionColor );
@@ -1074,32 +1095,35 @@ void FSlateEditorStyle::FStyle::SetupGeneralStyles()
 	// Output Log Window
 #if WITH_EDITOR || IS_PROGRAM
 	{
+		const int32 LogFontSize = Settings.IsValid() ? Settings->LogFontSize : 9;
+
 		const FTextBlockStyle NormalLogText = FTextBlockStyle(NormalText)
-			.SetFont( TTF_FONT( "Fonts/DroidSansMono", 9 ) )
-			.SetColorAndOpacity( FLinearColor(FColor(0xffaaaaaa)) )
-			.SetSelectedBackgroundColor( FLinearColor(FColor(0xff666666)) );
+			.SetFont( TTF_FONT( "Fonts/DroidSansMono", LogFontSize ) )
+			.SetColorAndOpacity( LogColor_Normal )
+			.SetSelectedBackgroundColor( LogColor_SelectionBackground );
 
 		Set("Log.Normal", NormalLogText );
 
 		Set("Log.Command", FTextBlockStyle(NormalLogText)
-			.SetColorAndOpacity( FLinearColor(FColor(0xff33dd33)) )
+			.SetColorAndOpacity( LogColor_Command )
 			);
 
 		Set("Log.Warning", FTextBlockStyle(NormalLogText)
-			.SetColorAndOpacity( FLinearColor(FColor(0xffbbbb44)) )
+			.SetColorAndOpacity( LogColor_Warning )
 			);
 
 		Set("Log.Error", FTextBlockStyle(NormalLogText)
-			.SetColorAndOpacity( FLinearColor(FColor(0xffdd0000)) )
+			.SetColorAndOpacity( LogColor_Error )
 			);
 
 		Set("Log.TabIcon", new IMAGE_BRUSH( "Icons/icon_tab_OutputLog_16x", Icon16x16 ) );
 
 		Set("Log.TextBox", FEditableTextBoxStyle(NormalEditableTextBoxStyle)
-			.SetBackgroundImageNormal( BOX_BRUSH( "Common/GroupBorder", FMargin(4.0f/16.0f) ) )
-			.SetBackgroundImageHovered( BOX_BRUSH( "Common/GroupBorder", FMargin(4.0f/16.0f) ) )
-			.SetBackgroundImageFocused( BOX_BRUSH( "Common/GroupBorder", FMargin(4.0f/16.0f) ) )
-			.SetBackgroundImageReadOnly( BOX_BRUSH( "Common/GroupBorder", FMargin(4.0f/16.0f) ) )
+			.SetBackgroundImageNormal( BOX_BRUSH( "Common/WhiteGroupBorder", FMargin(4.0f/16.0f) ) )
+			.SetBackgroundImageHovered( BOX_BRUSH( "Common/WhiteGroupBorder", FMargin(4.0f/16.0f) ) )
+			.SetBackgroundImageFocused( BOX_BRUSH( "Common/WhiteGroupBorder", FMargin(4.0f/16.0f) ) )
+			.SetBackgroundImageReadOnly( BOX_BRUSH( "Common/WhiteGroupBorder", FMargin(4.0f/16.0f) ) )
+			.SetBackgroundColor( LogColor_Background )
 			);
 
 		Set("DebugConsole.Background", new BOX_BRUSH("Old/Menu_Background", FMargin(8.0f / 64.0f)));
diff --git a/Engine/Source/Editor/EditorStyle/Private/SlateEditorStyle.h b/Engine/Source/Editor/EditorStyle/Private/SlateEditorStyle.h
--- a/Engine/Source/Editor/EditorStyle/Private/SlateEditorStyle.h
+++ b/Engine/Source/Editor/EditorStyle/Private/SlateEditorStyle.h
@@ -116,6 +116,13 @@ public:
 		const TSharedRef< FLinearColor > SelectionColor_Inactive_LinearRef;
 		const TSharedRef< FLinearColor > SelectionColor_Pressed_LinearRef;
 
+		const TSharedRef< FLinearColor > LogColor_Background_LinearRef;
+		const TSharedRef< FLinearColor > LogColor_SelectionBackground_LinearRef;
+		const TSharedRef< FLinearColor > LogColor_Normal_LinearRef;
+		const TSharedRef< FLinearColor > LogColor_Command_LinearRef;
+		const TSharedRef< FLinearColor > LogColor_Warning_LinearRef;
+		const TSharedRef< FLinearColor > LogColor_Error_LinearRef;
+
 		// These are the Slate colors which reference those above; these are the colors to put into the style
 		const FSlateColor DefaultForeground;
 		const FSlateColor InvertedForeground;
@@ -125,6 +132,13 @@ public:
 		const FSlateColor SelectionColor_Inactive;
 		const FSlateColor SelectionColor_Pressed;
 
+		const FSlateColor LogColor_Background;
+		const FSlateColor LogColor_SelectionBackground;
+		const FSlateColor LogColor_Normal;
+		const FSlateColor LogColor_Command;
+		const FSlateColor LogColor_Warning;
+		const FSlateColor LogColor_Error;
+
 		// These are common colors used thruout the editor in mutliple style elements
 		const FSlateColor InheritedFromBlueprintTextColor;
 
diff --git a/Engine/Source/Editor/EditorStyle/Public/Classes/EditorStyleSettings.h b/Engine/Source/Editor/EditorStyle/Public/Classes/EditorStyleSettings.h
--- a/Engine/Source/Editor/EditorStyle/Public/Classes/EditorStyleSettings.h
+++ b/Engine/Source/Editor/EditorStyle/Public/Classes/EditorStyleSettings.h
@@ -116,12 +116,40 @@ public:
 	UPROPERTY(config)
 	uint32 bShowLaunchMenus : 1;
 
+	/** The color used for the background in the output log */
+	UPROPERTY(EditAnywhere, config, Category="Output Log", meta=(DisplayName="Background Color"))
+	FLinearColor LogBackgroundColor;
+
+	/** The color used for the background of selected text in the output log */
+	UPROPERTY(EditAnywhere, config, Category="Output Log", meta=(DisplayName="Selection Background Color"))
+	FLinearColor LogSelectionBackgroundColor;
+
+	/** The color used for normal text in the output log */
+	UPROPERTY(EditAnywhere, config, Category="Output Log", meta=(DisplayName="Normal Text Color"))
+	FLinearColor LogNormalColor;
+
+	/** The color used for normal text in the output log */
+	UPROPERTY(EditAnywhere, config, Category="Output Log", meta=(DisplayName="Command Text Color"))
+	FLinearColor LogCommandColor;
+
+	/** The color used for warning log lines */
+	UPROPERTY(EditAnywhere, config, Category="Output Log", meta=(DisplayName="Warning Text Color"))
+	FLinearColor LogWarningColor;
+
+	/** The color used for error log lines */
+	UPROPERTY(EditAnywhere, config, Category="Output Log", meta=(DisplayName="Error Text Color"))
+	FLinearColor LogErrorColor;
+
+	/** The font size used in the output log */
+	UPROPERTY(EditAnywhere, config, Category="Output Log", meta=(DisplayName="Log Font Size", ConfigRestartRequired=true))
+	int32 LogFontSize;
+
 	/** The display mode for timestamps in the output log */
-	UPROPERTY(EditAnywhere, config, Category=Logging)
+	UPROPERTY(EditAnywhere, config, Category="Output Log")
 	TEnumAsByte<ELogTimes::Type> LogTimestampMode;
 
 	/** Should warnings and errors in the Output Log during "Play in Editor" be promoted to the message log? */
-	UPROPERTY(EditAnywhere, config, Category=Logging)
+	UPROPERTY(EditAnywhere, config, Category="Output Log")
 	bool bPromoteOutputLogWarningsDuringPIE;
 
 	/** New asset editor tabs will open at the specified location. */
diff --git a/Engine/Source/Editor/FontEditor/Private/SFontEditorViewport.cpp b/Engine/Source/Editor/FontEditor/Private/SFontEditorViewport.cpp
--- a/Engine/Source/Editor/FontEditor/Private/SFontEditorViewport.cpp
+++ b/Engine/Source/Editor/FontEditor/Private/SFontEditorViewport.cpp
@@ -230,7 +230,7 @@ void FFontEditorViewportClient::Draw(FViewport* Viewport, FCanvas* Canvas)
 						{
 							if (GlyphToRender.bIsVisible)
 							{
-								const FShapedGlyphFontAtlasData GlyphAtlasData = FontCache->GetShapedGlyphFontAtlasData(GlyphToRender);
+								const FShapedGlyphFontAtlasData GlyphAtlasData = FontCache->GetShapedGlyphFontAtlasData(GlyphToRender, FFontOutlineSettings::NoOutline);
 
 								const float X = CurPos.X + LineX + GlyphAtlasData.HorizontalOffset + GlyphToRender.XOffset;
 								const float Y = CurPos.Y - GlyphAtlasData.VerticalOffset + GlyphToRender.YOffset + ShapedPreviewText->GetTextBaseline() + ShapedPreviewText->GetMaxTextHeight();
@@ -260,7 +260,7 @@ void FFontEditorViewportClient::Draw(FViewport* Viewport, FCanvas* Canvas)
 
 								if (GlyphToRender.bIsVisible)
 								{
-									const FShapedGlyphFontAtlasData GlyphAtlasData = FontCache->GetShapedGlyphFontAtlasData(GlyphToRender);
+									const FShapedGlyphFontAtlasData GlyphAtlasData = FontCache->GetShapedGlyphFontAtlasData(GlyphToRender,FFontOutlineSettings::NoOutline);
 
 									const float X = CurPos.X + LineX + GlyphAtlasData.HorizontalOffset + GlyphToRender.XOffset;
 									const float Y = CurPos.Y - GlyphAtlasData.VerticalOffset + GlyphToRender.YOffset + ShapedPreviewText->GetTextBaseline() + ShapedPreviewText->GetMaxTextHeight();
diff --git a/Engine/Source/Editor/UnrealEd/Private/Commandlets/ChunkDependencyInfo.h b/Engine/Source/Editor/UnrealEd/Private/Commandlets/ChunkDependencyInfo.h
--- a/Engine/Source/Editor/UnrealEd/Private/Commandlets/ChunkDependencyInfo.h
+++ b/Engine/Source/Editor/UnrealEd/Private/Commandlets/ChunkDependencyInfo.h
@@ -33,7 +33,7 @@ struct FChunkDependencyTreeNode
 };
 
 UCLASS(config=Engine, defaultconfig)
-class UChunkDependencyInfo : public UObject
+class UNREALED_API UChunkDependencyInfo : public UObject
 {
 	GENERATED_UCLASS_BODY()
 
diff --git a/Engine/Source/Programs/AutomationTool/AutomationUtils/CommandletUtils.cs b/Engine/Source/Programs/AutomationTool/AutomationUtils/CommandletUtils.cs
--- a/Engine/Source/Programs/AutomationTool/AutomationUtils/CommandletUtils.cs
+++ b/Engine/Source/Programs/AutomationTool/AutomationUtils/CommandletUtils.cs
@@ -119,6 +119,25 @@ namespace AutomationTool
 			RunCommandlet(ProjectName, UE4Exe, "ResavePackages", String.Format("-buildlighting -MapsOnly -ProjectOnly -AllowCommandletRendering -SkipSkinVerify {0} {1}", MapsToRebuildLighting, Parameters));
 		}
 
+        /// <summary>
+        /// Runs RebuildLightMaps commandlet.
+        /// </summary>
+        /// <param name="ProjectName">Project name.</param>
+        /// <param name="UE4Exe">The name of the UE4 Editor executable to use.</param>
+        /// <param name="Maps">List of maps to rebuild light maps for. Can be null in which case -MapIniSection=AllMaps is used.</param>
+        /// <param name="Parameters">List of additional parameters.</param>
+        public static void ResavePackagesCommandlet(FileReference ProjectName, string UE4Exe = "UE4Editor-Cmd.exe", string[] Maps = null, string Parameters = "")
+        {
+            string MapsToRebuildLighting = "";
+            if (!IsNullOrEmpty(Maps))
+            {
+                MapsToRebuildLighting = "-Map=" + CombineCommandletParams(Maps);
+                MapsToRebuildLighting.Trim();
+            }
+
+            RunCommandlet(ProjectName, UE4Exe, "ResavePackages", String.Format((!String.IsNullOrEmpty(MapsToRebuildLighting) ? "-MapsOnly" : "") + "-ProjectOnly {0} {1}", MapsToRebuildLighting, Parameters));
+        }
+
         /// <summary>
         /// Runs GenerateDistillFileSets commandlet.
         /// </summary>
diff --git a/Engine/Source/Programs/AutomationTool/Scripts/AutomationScripts.Automation.csproj b/Engine/Source/Programs/AutomationTool/Scripts/AutomationScripts.Automation.csproj
--- a/Engine/Source/Programs/AutomationTool/Scripts/AutomationScripts.Automation.csproj
+++ b/Engine/Source/Programs/AutomationTool/Scripts/AutomationScripts.Automation.csproj
@@ -94,6 +94,7 @@
     <Compile Include="InstalledBuild.cs" />
     <Compile Include="ListThirdPartySoftware.Automation.cs" />
     <Compile Include="RebuildLightMapsCommand.Automation.cs" />
+    <Compile Include="ResavePackagesCommand.Automation.cs" />
     <Compile Include="RocketBuild.Automation.cs" />
     <Compile Include="IPhonePackager.Automation.cs" />
     <Compile Include="LauncherLocalization.Automation.cs" />
diff --git a/Engine/Source/Programs/AutomationTool/Scripts/CopyBuildToStagingDirectory.Automation.cs b/Engine/Source/Programs/AutomationTool/Scripts/CopyBuildToStagingDirectory.Automation.cs
--- a/Engine/Source/Programs/AutomationTool/Scripts/CopyBuildToStagingDirectory.Automation.cs
+++ b/Engine/Source/Programs/AutomationTool/Scripts/CopyBuildToStagingDirectory.Automation.cs
@@ -798,7 +798,13 @@ public partial class Project : CommandUtils
 
 		var UnrealPakResponseFile = CreatePakResponseFileFromStagingManifest(SC);
 
-		CreatePak(Params, SC, UnrealPakResponseFile, SC.ShortProjectName, Params.SignPak);
+        if (Params.CreateChunkInstall)
+        {
+            string ChunkInstallBasePath = CombinePaths(Params.ChunkInstallDirectory, SC.FinalCookPlatform);
+            string CloudDir = MakePathSafeToUseWithCommandLine(CombinePaths(ChunkInstallBasePath, "CloudDir"));
+            InternalUtils.SafeDeleteDirectory(CloudDir, true);
+        }
+        CreatePak(Params, SC, UnrealPakResponseFile, SC.ShortProjectName, Params.SignPak);
 	}
 
 	/// <summary>
@@ -1165,7 +1171,14 @@ public partial class Project : CommandUtils
 
 		IEnumerable<Tuple<Dictionary<string,string>, string>> PakPairs = PakResponseFiles.Zip(ChunkList, (a, b) => Tuple.Create(a, b));
 
-		System.Threading.Tasks.Parallel.ForEach(PakPairs, (PakPair) =>
+        if (Params.CreateChunkInstall)
+        {
+            string ChunkInstallBasePath = CombinePaths(Params.ChunkInstallDirectory, SC.FinalCookPlatform);
+            string CloudDir = MakePathSafeToUseWithCommandLine(CombinePaths(ChunkInstallBasePath, "CloudDir"));
+            InternalUtils.SafeDeleteDirectory(CloudDir, true);
+        }
+
+        System.Threading.Tasks.Parallel.ForEach(PakPairs, (PakPair) =>
 		{
 			var ChunkName = Path.GetFileNameWithoutExtension(PakPair.Item2);
 			CreatePak(Params, SC, PakPair.Item1, ChunkName, Params.SignPak);
diff --git a/Engine/Source/Programs/AutomationTool/Scripts/ResavePackagesCommand.Automation.cs b/Engine/Source/Programs/AutomationTool/Scripts/ResavePackagesCommand.Automation.cs
new file mode 100644
--- /dev/null
+++ b/Engine/Source/Programs/AutomationTool/Scripts/ResavePackagesCommand.Automation.cs
@@ -0,0 +1,375 @@
+// Copyright 1998-2016 Epic Games, Inc. All Rights Reserved.
+using System;
+using System.Collections.Generic;
+using System.IO;
+using System.Threading;
+using System.Reflection;
+using System.Linq;
+using System.Net.Mail;
+using AutomationTool;
+using UnrealBuildTool;
+
+/// <summary>
+/// Helper command used for resaving a projects packages.
+/// </summary>
+/// <remarks>
+/// Command line parameters used by this command:
+/// -project					- Absolute path to a .uproject file
+/// -CommandletTargetName		- The Target used in running the commandlet
+/// -StakeholdersEmailAddresses	- Users to notify of completion
+/// 
+/// </remarks>
+namespace AutomationScripts.Automation
+{
+	[RequireP4]
+	public class ResavePackages : BuildCommand
+	{
+        // The rebuild lighting process
+        #region ResavePackages Command
+
+        public override void ExecuteBuild()
+		{
+			Log("********** RESAVE PACKAGES COMMAND STARTED **********");
+
+			try
+			{
+				var Params = new ProjectParams
+				(
+					Command: this,
+					// Shared
+					RawProjectPath: ProjectPath
+				);
+
+				// Sync and build our targets required for the commandlet to run correctly.
+				// P4.Sync(String.Format("-f {0}/...#head", P4Env.BuildRootP4));
+
+                bool NoBuild = ParseParam("nobuild");
+
+                if (!NoBuild)
+                {
+                    BuildNecessaryTargets();
+                }
+				CreateChangelist(Params);
+				RunResavePackagesCommandlet(Params);
+				SubmitResavedPackages();
+			}
+			catch (Exception ProcessEx)
+			{
+				Log("********** RESAVE PACKAGES COMMAND FAILED **********");
+                Log("Error message: {0}", ProcessEx.Message);
+				HandleFailure(ProcessEx.Message);
+				throw ProcessEx;
+			}
+
+			// The processes steps have completed successfully.
+			HandleSuccess();
+
+			Log("********** RESAVE PACKAGES COMMAND COMPLETED **********");
+		}
+
+		#endregion
+
+		// Broken down steps used to run the process.
+		#region RebuildLightMaps Process Steps
+
+		private void BuildNecessaryTargets()
+		{
+			Log("Running Step:- RebuildLightMaps::BuildNecessaryTargets");
+			UE4Build.BuildAgenda Agenda = new UE4Build.BuildAgenda();
+            Agenda.AddTarget("UnrealHeaderTool", UnrealBuildTool.UnrealTargetPlatform.Win64, UnrealBuildTool.UnrealTargetConfiguration.Development);
+			Agenda.AddTarget("ShaderCompileWorker", UnrealBuildTool.UnrealTargetPlatform.Win64, UnrealBuildTool.UnrealTargetConfiguration.Development);
+			Agenda.AddTarget("UnrealLightmass", UnrealBuildTool.UnrealTargetPlatform.Win64, UnrealBuildTool.UnrealTargetConfiguration.Development);
+			Agenda.AddTarget(CommandletTargetName, UnrealBuildTool.UnrealTargetPlatform.Win64, UnrealBuildTool.UnrealTargetConfiguration.Development);
+
+			try
+			{
+				UE4Build Builder = new UE4Build(this);
+				Builder.Build(Agenda, InDeleteBuildProducts: true, InUpdateVersionFiles: true, InForceNoXGE: false, InChangelistNumberOverride: GetLatestCodeChange());
+				UE4Build.CheckBuildProducts(Builder.BuildProductFiles);
+			}
+			catch (AutomationException Ex)
+			{
+				LogError("Rebuild Light Maps has failed.");
+				throw Ex;
+			}
+		}
+
+		private int GetLatestCodeChange()
+		{
+			List<P4Connection.ChangeRecord> ChangeRecords;
+			if(!P4.Changes(out ChangeRecords, String.Format("-m 1 //{0}/....cpp@<{1} //{0}/....h@<{1} //{0}/....cs@<{1} //{0}/....usf@<{1}", P4Env.Client, P4Env.Changelist), WithClient: true))
+			{
+				throw new AutomationException("Couldn't enumerate latest change from branch");
+			}
+			return ChangeRecords.Max(x => x.CL);
+		}
+
+		private void CreateChangelist(ProjectParams Params)
+		{
+			Log("Running Step:- RebuildLightMaps::CheckOutMaps");
+			// Setup a P4 Cl we will use to submit the new lightmaps
+			WorkingCL = P4.CreateChange(P4Env.Client, String.Format("{0} rebuilding lightmaps from changelist {1}\n#rb None\n#tests None", Params.ShortProjectName, P4Env.Changelist));
+			Log("Working in {0}", WorkingCL);
+
+		}
+
+		private void RunResavePackagesCommandlet(ProjectParams Params)
+		{
+			Log("Running Step:- ResavePackages::RunResavePackagesCommandlet");
+
+			// Find the commandlet binary
+			string UE4EditorExe = HostPlatform.Current.GetUE4ExePath(Params.UE4Exe);
+			if (!FileExists(UE4EditorExe))
+			{
+				LogError("Missing " + UE4EditorExe + " executable. Needs to be built first.");
+				throw new AutomationException("Missing " + UE4EditorExe + " executable. Needs to be built first.");
+			}
+
+			// Now let's rebuild lightmaps for the project
+			try
+			{
+				var CommandletParams = IsBuildMachine ? "-unattended -buildmachine -fileopenlog" : "-fileopenlog";
+                CommandletParams += " -AutoCheckOutPackages";
+                if (P4Enabled)
+                {
+                    CommandletParams += String.Format(" -SCCProvider={0} -P4Port={1} -P4User={2} -P4Client={3} -P4Changelist={4} -P4Passwd={5}", "Perforce", P4Env.P4Port, P4Env.User, P4Env.Client, WorkingCL.ToString(), P4.GetAuthenticationToken());
+                }
+				ResavePackagesCommandlet(Params.RawProjectPath, Params.UE4Exe, Params.MapsToRebuildLightMaps.ToArray(), CommandletParams);
+			}
+			catch (Exception Ex)
+			{
+                string FinalLogLines = "No log file found";
+                AutomationException AEx = Ex as AutomationException;
+                if ( AEx != null )
+                {
+                    string LogFile = AEx.LogFileName;
+                    UnrealBuildTool.Log.TraceWarning("Attempting to load file {0}", LogFile);
+                    if ( LogFile != "")
+                    {
+                        
+                        UnrealBuildTool.Log.TraceWarning("Attempting to read file {0}", LogFile);
+                        try
+                        {
+                            string[] AllLogFile = ReadAllLines(LogFile);
+
+                            FinalLogLines = "Important log entries\n";
+                            foreach (string LogLine in AllLogFile)
+                            {
+                                if (LogLine.Contains("[REPORT]"))
+                                {
+                                    FinalLogLines += LogLine + "\n";
+                                }
+                            }
+                        }
+                        catch (Exception)
+                        {
+                            // we don't care about this because if this is hit then there is no log file the exception probably has more info
+                            LogError("Could not find log file " + LogFile);
+                        }
+                    }
+                }
+
+				// Something went wrong with the commandlet. Abandon this run, don't check in any updated files, etc.
+				LogError("Resave Packages has failed. because "+ Ex.ToString());
+				throw new AutomationException(ExitCode.Error_Unknown, Ex, "ResavePackages failed. {0}", FinalLogLines);
+			}
+		}
+
+		private void SubmitResavedPackages()
+		{
+			Log("Running Step:- ResavePackages::SubmitResavedPackages");
+
+			// Check everything in!
+			if (WorkingCL != -1)
+			{
+                Log("Running Step:- Submitting CL " + WorkingCL);
+				int SubmittedCL;
+				P4.Submit(WorkingCL, out SubmittedCL, true, true);
+				Log("INFO: Packages successfully submitted in cl "+ SubmittedCL.ToString());
+			}
+		}
+
+        #endregion
+
+        // Helper functions and procedure steps necessary for running the commandlet successfully
+        #region ResavePackages Helper Functions
+
+        /**
+		 * Parse the P4 output for any errors that we really care about.
+		 * e.g. umaps and assets are exclusive checkout files, if we cant check out a map for this reason
+		 *		then we need to stop.
+		 */
+        private bool FoundCheckOutErrorInP4Output(string Output)
+		{
+			bool bHadAnError = false;
+
+			var Lines = Output.Split(new string[] { Environment.NewLine }, StringSplitOptions.RemoveEmptyEntries);
+			foreach (string Line in Lines)
+			{
+				// Check for log spew that matches exclusive checkout failure
+				// http://answers.perforce.com/articles/KB/3114
+				if (Line.Contains("can't edit exclusive file already opened"))
+				{
+					bHadAnError = true;
+					break;
+				}
+			}
+
+			return bHadAnError;
+		}
+
+		/**
+		 * Cleanup anything this build may leave behind and inform the user
+		 */
+		private void HandleFailure(String FailureMessage)
+		{
+			try
+			{
+				if (WorkingCL != -1)
+				{
+					P4.RevertAll(WorkingCL);
+					P4.DeleteChange(WorkingCL);
+				}
+				SendCompletionMessage(false, FailureMessage);
+			}
+			catch (P4Exception P4Ex)
+			{
+				LogError("Failed to clean up P4 changelist: " + P4Ex.Message);
+			}
+			catch (Exception SendMailEx)
+			{
+				LogError("Failed to notify that build succeeded: " + SendMailEx.Message);
+			}
+		}
+
+		/**
+		 * Perform any post completion steps needed. I.e. Notify stakeholders etc.
+		 */
+		private void HandleSuccess()
+		{
+			try
+			{
+				SendCompletionMessage(true, "Successfully rebuilt lightmaps.");
+			}
+			catch (Exception SendMailEx)
+			{
+				LogError("Failed to notify that build succeeded: " + SendMailEx.Message);
+			}
+		}
+
+		/**
+		 * Notify stakeholders of the commandlet results
+		 */
+		void SendCompletionMessage(bool bWasSuccessful, String MessageBody)
+		{
+			MailMessage Message = new System.Net.Mail.MailMessage();
+			Message.Priority = MailPriority.High;
+			Message.From = new MailAddress("[REDACTED_EMAIL]");
+
+            string Branch = "Unknown";
+            if ( P4Enabled )
+            {
+                Branch = P4Env.BuildRootP4;
+            }
+
+			foreach (String NextStakeHolder in StakeholdersEmailAddresses)
+			{
+				Message.To.Add(new MailAddress(NextStakeHolder));
+			}
+
+			Message.CC.Add(new MailAddress("[REDACTED_EMAIL]"));
+            Message.CC.Add(new MailAddress("[REDACTED_EMAIL]"));
+            Message.CC.Add(new MailAddress("[REDACTED_EMAIL]"));
+			Message.Subject = String.Format("Nightly lightmap rebuild {0} for {1}", bWasSuccessful ? "[SUCCESS]" : "[FAILED]", Branch);
+			Message.Body = MessageBody;
+            /*Attachment Attach = new Attachment();
+            Message.Attachments.Add()*/
+			try
+			{
+				SmtpClient MailClient = new SmtpClient("smtp.epicgames.net");
+				MailClient.Send(Message);
+			}
+			catch (Exception Ex)
+			{
+				LogError("Failed to send notify email to {0} ({1})", String.Join(", ", StakeholdersEmailAddresses.ToArray()), Ex.Message);
+			}
+		}
+
+		#endregion
+
+		// Member vars used in multiple steps.
+		#region RebuildLightMaps Property Set-up
+
+		// Users to notify if the process fails or succeeds.
+		List<String> StakeholdersEmailAddresses
+		{
+			get
+			{
+				String UnprocessedEmailList = ParseParamValue("StakeholdersEmailAddresses");
+				if (String.IsNullOrEmpty(UnprocessedEmailList) == false)
+				{
+					return UnprocessedEmailList.Split('+').ToList();
+				}
+				else
+				{
+					return null;
+				}
+			}
+		}
+
+		// The Changelist used when doing the work.
+		private int WorkingCL = -1;
+
+		// The target name of the commandlet binary we wish to build and run.
+		private String CommandletTargetName
+		{
+			get
+			{
+				return ParseParamValue("CommandletTargetName", "");
+			}
+		}
+
+		// Process command-line and find a project file. This is necessary for the commandlet to run successfully
+		private FileReference ProjectFullPath;
+		public virtual FileReference ProjectPath
+		{
+			get
+			{
+				if (ProjectFullPath == null)
+				{
+					var OriginalProjectName = ParseParamValue("project", "");
+					var ProjectName = OriginalProjectName;
+					ProjectName = ProjectName.Trim(new char[] { '\"' });
+					if (ProjectName.IndexOfAny(new char[] { '\\', '/' }) < 0)
+					{
+						ProjectName = CombinePaths(CmdEnv.LocalRoot, ProjectName, ProjectName + ".uproject");
+					}
+					else if (!FileExists_NoExceptions(ProjectName))
+					{
+						ProjectName = CombinePaths(CmdEnv.LocalRoot, ProjectName);
+					}
+					if (FileExists_NoExceptions(ProjectName))
+					{
+						ProjectFullPath = new FileReference(ProjectName);
+					}
+					else
+					{
+						var Branch = new BranchInfo(new List<UnrealTargetPlatform> { UnrealBuildTool.BuildHostPlatform.Current.Platform });
+						var GameProj = Branch.FindGame(OriginalProjectName);
+						if (GameProj != null)
+						{
+							ProjectFullPath = GameProj.FilePath;
+						}
+						if (!FileExists_NoExceptions(ProjectFullPath.FullName))
+						{
+							throw new AutomationException("Could not find a project file {0}.", ProjectName);
+						}
+					}
+				}
+				return ProjectFullPath;
+			}
+		}
+
+		#endregion
+	}
+}
\ No newline at end of file
diff --git a/Engine/Source/Programs/UnrealBuildTool/IOS/UEDeployIOS.cs b/Engine/Source/Programs/UnrealBuildTool/IOS/UEDeployIOS.cs
--- a/Engine/Source/Programs/UnrealBuildTool/IOS/UEDeployIOS.cs
+++ b/Engine/Source/Programs/UnrealBuildTool/IOS/UEDeployIOS.cs
@@ -15,15 +15,15 @@ namespace UnrealBuildTool
 	public class UEDeployIOS : UEBuildDeploy
 	{
 		FileReference ProjectFile;
-		IOSPlatformContext	IOSPlatformContext;
+        IOSPlatformContext IOSPlatformContext;
 
-		public UEDeployIOS(FileReference InProjectFile, IOSPlatformContext inIOSPlatformContext)
-		{
-			ProjectFile = InProjectFile;
-			IOSPlatformContext = inIOSPlatformContext;
-		}
+        public UEDeployIOS(FileReference InProjectFile, IOSPlatformContext inIOSPlatformContext)
+        {
+            ProjectFile = InProjectFile;
+            IOSPlatformContext = inIOSPlatformContext;
+        }
 
-		protected UnrealPluginLanguage UPL = null;
+        protected UnrealPluginLanguage UPL = null;
 
 		public void SetIOSPluginData(List<string> Architectures, List<string> inPluginExtraData)
 		{
@@ -236,19 +236,24 @@ namespace UnrealBuildTool
 			string BundleName;
 			Ini.GetString("/Script/IOSRuntimeSettings.IOSRuntimeSettings", "BundleName", out BundleName);
 
+			// disable https requirement
+            bool bDisableHTTPS;
+            Ini.GetBool("/Script/IOSRuntimeSettings.IOSRuntimeSettings", "bDisableHTTPS", out bDisableHTTPS);
+
 			// short version string
 			string BundleShortVersion;
 			Ini.GetString("/Script/IOSRuntimeSettings.IOSRuntimeSettings", "VersionInfo", out BundleShortVersion);
-			
-			string RequiredCaps = "";
-			
-			if(InThis != null)
-			{
-				// required capabilities
-	            RequiredCaps += InThis.IOSPlatformContext.GetRequiredCapabilities();
-			}
 
-			Ini.GetBool("/Script/IOSRuntimeSettings.IOSRuntimeSettings", "bSupportsOpenGLES2", out bSupported);
+            // required capabilities
+            string RequiredCaps = "";
+
+            if (InThis != null)
+            {
+                // required capabilities
+                RequiredCaps += InThis.IOSPlatformContext.GetRequiredCapabilities();
+            }
+
+            Ini.GetBool("/Script/IOSRuntimeSettings.IOSRuntimeSettings", "bSupportsOpenGLES2", out bSupported);
 			RequiredCaps += bSupported ? "\t\t<string>opengles-2</string>\n" : "";
 			if (!bSupported)
 			{
@@ -484,6 +489,16 @@ namespace UnrealBuildTool
 			// disable exempt encryption
 			Text.AppendLine("\t<key>ITSAppUsesNonExemptEncryption</key>");
 			Text.AppendLine("\t<false/>");
+			
+			// disable HTTPS requirement
+            if (bDisableHTTPS)
+            {
+                Text.AppendLine("\t<key>NSAppTransportSecurity</key>");
+                Text.AppendLine("\t\t<dict>");
+                Text.AppendLine("\t\t\t<key>NSAllowsArbitraryLoads</key><true/>");
+                Text.AppendLine("\t\t</dict>");
+            }
+            
 			if (bEnableFacebookSupport)
 			{
 				Text.AppendLine("\t<key>FacebookAppID</key>");
@@ -788,9 +803,9 @@ namespace UnrealBuildTool
 
 			// Run through iOS APL file
 			IOSPlatformContext PlatformContext = new IOSPlatformContext(InTarget.ProjectFile);
-			PlatformContext.SetUpProjectEnvironment(InTarget.Configuration);
+            PlatformContext.SetUpProjectEnvironment(InTarget.Configuration);
 
-			string BaseSoName = InTarget.OutputPaths[0].FullName;
+            string BaseSoName = InTarget.OutputPaths[0].FullName;
 
 			// get the receipt
 			UnrealTargetPlatform Platform = InTarget.Platform;
diff --git a/Engine/Source/Programs/UnrealBuildTool/System/Unity.cs b/Engine/Source/Programs/UnrealBuildTool/System/Unity.cs
--- a/Engine/Source/Programs/UnrealBuildTool/System/Unity.cs
+++ b/Engine/Source/Programs/UnrealBuildTool/System/Unity.cs
@@ -17,6 +17,10 @@ namespace UnrealBuildTool
 		/// </summary>
 		public static HashSet<FileItem> SourceFileWorkingSet = new HashSet<FileItem>();
 
+		/// <summary>
+		/// Set of files which are candidates for being part of the working set.
+		/// </summary>
+		public static HashSet<FileItem> CandidateSourceFilesForWorkingSet = new HashSet<FileItem>();
 
 		/// <summary>
 		/// A class which represents a list of files and the sum of their lengths.
@@ -259,6 +263,12 @@ namespace UnrealBuildTool
 
 					else
 					{
+						// If adaptive unity build is enabled for this module, add this source file to the set that will invalidate the makefile
+						if(bUseAdaptiveUnityBuild)
+						{
+							CandidateSourceFilesForWorkingSet.Add(CPPFile);
+						}
+
 						// Compile this file as part of the unity blob
 						CPPUnityFileBuilder.AddFile(CPPFile);
 
diff --git a/Engine/Source/Programs/UnrealBuildTool/System/UnrealBuildTool.cs b/Engine/Source/Programs/UnrealBuildTool/System/UnrealBuildTool.cs
--- a/Engine/Source/Programs/UnrealBuildTool/System/UnrealBuildTool.cs
+++ b/Engine/Source/Programs/UnrealBuildTool/System/UnrealBuildTool.cs
@@ -1906,6 +1906,7 @@ namespace UnrealBuildTool
 						UBTMakefile.TargetNameToUObjectModules = TargetNameToUObjectModules;
 						UBTMakefile.Targets = Targets;
 						UBTMakefile.SourceFileWorkingSet = Unity.SourceFileWorkingSet;
+						UBTMakefile.CandidateSourceFilesForWorkingSet = Unity.CandidateSourceFilesForWorkingSet;
 
 						if (BuildConfiguration.bUseUBTMakefiles)
 						{
@@ -2319,12 +2320,16 @@ namespace UnrealBuildTool
 			/// </summary>
 			public List<UEBuildTarget> Targets;
 
-
 			/// <summary>
 			/// Current working set of source files, for when bUseAdaptiveUnityBuild is enabled
 			/// </summary>
 			public HashSet<FileItem> SourceFileWorkingSet = new HashSet<FileItem>();
 
+			/// <summary>
+			/// Set of source files which are included in unity files, but which should invalidate the makefile if modified (for when bUseAdaptiveUnityBuild is enabled)
+			/// </summary>
+			public HashSet<FileItem> CandidateSourceFilesForWorkingSet = new HashSet<FileItem>();
+
 			public UBTMakefile()
 			{
 				Version = CurrentVersion;
@@ -2344,6 +2349,7 @@ namespace UnrealBuildTool
 				TargetNameToUObjectModules = (Dictionary<string, List<UHTModuleInfo>>)Info.GetValue("nu", typeof(Dictionary<string, List<UHTModuleInfo>>));
 				Targets = (List<UEBuildTarget>)Info.GetValue("ta", typeof(List<UEBuildTarget>));
 				SourceFileWorkingSet = (HashSet<FileItem>)Info.GetValue("ws", typeof(HashSet<FileItem>));
+				CandidateSourceFilesForWorkingSet = (HashSet<FileItem>)Info.GetValue("wc", typeof(HashSet<FileItem>));
 			}
 
 			public void GetObjectData(SerializationInfo Info, StreamingContext Context)
@@ -2356,6 +2362,7 @@ namespace UnrealBuildTool
 				Info.AddValue("nu", TargetNameToUObjectModules);
 				Info.AddValue("ta", Targets);
 				Info.AddValue("ws", SourceFileWorkingSet);
+				Info.AddValue("wc", CandidateSourceFilesForWorkingSet);
 			}
 
 
@@ -2368,7 +2375,8 @@ namespace UnrealBuildTool
 					EnvironmentVariables != null &&
 					TargetNameToUObjectModules != null && TargetNameToUObjectModules.Count > 0 &&
 					Targets != null && Targets.Count > 0 &&
-					SourceFileWorkingSet != null;
+					SourceFileWorkingSet != null &&
+					CandidateSourceFilesForWorkingSet != null;
 			}
 		}
 
@@ -2615,37 +2623,25 @@ namespace UnrealBuildTool
 			// iteration times.)
 			if (BuildConfiguration.bUseAdaptiveUnityBuild)
 			{
-				// Get all .cpp files in processed modules newer than the makefile itself
-				foreach (UEBuildTarget Target in LoadedUBTMakefile.Targets)
+				// Check if any source files in the working set no longer belong in it
+				foreach(FileItem SourceFile in LoadedUBTMakefile.SourceFileWorkingSet)
 				{
-					foreach (FlatModuleCsDataType FlatModule in Target.FlatModuleCsData.Values)
+					if(!UnrealBuildTool.ShouldSourceFileBePartOfWorkingSet(SourceFile.AbsolutePath) && File.GetLastWriteTimeUtc(SourceFile.AbsolutePath) > UBTMakefileInfo.LastWriteTimeUtc)
 					{
-						if (FlatModule != null && FlatModule.ModuleSourceFolder != null)
-						{
-							// Has the directory been touched?  Don't bother checking folders that haven't been written to since the
-							// Makefile was created. 
-							//if( Directory.GetLastWriteTimeUtc( FlatModule.ModuleSourceFolder ) > UBTMakefileInfo.LastWriteTimeUtc )
-							{
-								foreach (FileReference SourceFilePath in FlatModule.ModuleSourceFolder.EnumerateFileReferences("*.cpp", SearchOption.AllDirectories))
-								{
-									// Was the file written to since the Makefile was created?  We're not able to do a fully exhaustive
-									// check for working sets when in assembler mode, so we might get it wrong occasionally.  This is
-									// why we only want to check newly-modified or created files.  Otherwise, we could potentially
-									// invalidate the Makefile every single run.
-									if (File.GetLastWriteTimeUtc(SourceFilePath.FullName) > UBTMakefileInfo.LastWriteTimeUtc)
-									{
-										bool bShouldBePartOfWorkingSet = UnrealBuildTool.ShouldSourceFileBePartOfWorkingSet(SourceFilePath.FullName);
-										bool bIsAlreadyPartOfWorkingSet = LoadedUBTMakefile.SourceFileWorkingSet.Contains(FileItem.GetItemByFileReference(SourceFilePath));
-										if (bShouldBePartOfWorkingSet != bIsAlreadyPartOfWorkingSet)
-										{
-											Log.TraceVerbose("{0} {1} part of source working set and now {2}; invalidating makefile ({3})", SourceFilePath, bIsAlreadyPartOfWorkingSet ? "was" : "was not", bShouldBePartOfWorkingSet ? "is" : "is not", UBTMakefileInfo.FullName);
-											ReasonNotLoaded = string.Format("working set of source files changed");
-											return null;
-										}
-									}
-								}
-							}
-						}
+						Log.TraceVerbose("{0} was part of source working set and now is not; invalidating makefile ({1})", SourceFile.AbsolutePath, UBTMakefileInfo.FullName);
+						ReasonNotLoaded = string.Format("working set of source files changed");
+						return null;
+					}
+				}
+
+				// Check if any source files that are eligable for being in the working set have been modified
+				foreach (FileItem SourceFile in LoadedUBTMakefile.CandidateSourceFilesForWorkingSet)
+				{
+					if(UnrealBuildTool.ShouldSourceFileBePartOfWorkingSet(SourceFile.AbsolutePath) && File.GetLastWriteTimeUtc(SourceFile.AbsolutePath) > UBTMakefileInfo.LastWriteTimeUtc)
+					{
+						Log.TraceVerbose("{0} was part of source working set and now is not; invalidating makefile ({1})", SourceFile.AbsolutePath, UBTMakefileInfo.FullName);
+						ReasonNotLoaded = string.Format("working set of source files changed");
+						return null;
 					}
 				}
 			}
diff --git a/Engine/Source/Runtime/Analytics/AnalyticsET/Private/IAnalyticsProviderET.cpp b/Engine/Source/Runtime/Analytics/AnalyticsET/Private/IAnalyticsProviderET.cpp
--- a/Engine/Source/Runtime/Analytics/AnalyticsET/Private/IAnalyticsProviderET.cpp
+++ b/Engine/Source/Runtime/Analytics/AnalyticsET/Private/IAnalyticsProviderET.cpp
@@ -222,6 +222,8 @@ private:
 	bool bInDestructor;
 	/** True to use the legacy backend server protocol that uses URL params. */
 	bool UseLegacyProtocol;
+	/** False to disable dropping events when a flush fails due to network reasons */
+	bool bDropEventsOnFlushFailure;
 	/** AppEnvironment to use. */
 	FString AppEnvironment;
 	/** UploadType to use. */
@@ -265,7 +267,7 @@ private:
 	/**
 	* Delegate called when an event Http request completes
 	*/
-	void EventRequestComplete(FHttpRequestPtr HttpRequest, FHttpResponsePtr HttpResponse, bool bSucceeded);
+	void EventRequestComplete(FHttpRequestPtr HttpRequest, FHttpResponsePtr HttpResponse, bool bSucceeded, TSharedPtr< TArray<FAnalyticsEventEntry> > FlushedEvents);
 };
 
 TSharedPtr<IAnalyticsProviderET> FAnalyticsET::CreateAnalyticsProvider(const Config& ConfigValues) const
@@ -292,6 +294,7 @@ FAnalyticsProviderET::FAnalyticsProviderET(const FAnalyticsET::Config& ConfigVal
 	, FlushEventsCountdown(MaxCachedElapsedTime)
 	, bInDestructor(false)
 	, UseLegacyProtocol(ConfigValues.UseLegacyProtocol)
+	, bDropEventsOnFlushFailure(ConfigValues.bDropEventsOnFlushFailure)
 {
 	if (APIKey.IsEmpty() || APIServer.IsEmpty())
 	{
@@ -454,6 +457,13 @@ void FAnalyticsProviderET::FlushEvents()
 
 		if (!UseLegacyProtocol)
 		{
+			TSharedPtr< TArray<FAnalyticsEventEntry> > FlushedEvents;
+			if (!bDropEventsOnFlushFailure)
+			{
+				FlushedEvents = TSharedPtr< TArray<FAnalyticsEventEntry> >(new TArray<FAnalyticsEventEntry>());
+				FlushedEvents->Reserve(CachedEvents.Num());
+			}
+
 			TSharedRef< TJsonWriter<TCHAR, TCondensedJsonPrintPolicy<TCHAR> > > JsonWriter = TJsonWriterFactory<TCHAR, TCondensedJsonPrintPolicy<TCHAR> >::Create(&Payload);
 			JsonWriter->WriteObjectStart();
 			JsonWriter->WriteArrayStart(TEXT("Events"));
@@ -462,11 +472,20 @@ void FAnalyticsProviderET::FlushEvents()
 				if (Entry.bIsDefaultAttributes)
 				{
 					// This is the default attributes, so update the array.
-					CurrentDefaultAttributes = MoveTemp(Entry.Attributes);
+					if (FlushedEvents.IsValid())
+					{
+						CurrentDefaultAttributes = Entry.Attributes; // need to copy
+						FlushedEvents->Emplace(MoveTemp(Entry));
+					}
+					else
+					{
+						CurrentDefaultAttributes = MoveTemp(Entry.Attributes);
+					}
 				}
 				else
 				{
 					++EventCount;
+
 					// event entry
 					JsonWriter->WriteObjectStart();
 					JsonWriter->WriteValue(TEXT("EventName"), Entry.EventName);
@@ -493,6 +512,12 @@ void FAnalyticsProviderET::FlushEvents()
 						}
 					}
 					JsonWriter->WriteObjectEnd();
+
+					// move the entry into the flushed
+					if (FlushedEvents.IsValid())
+					{
+						FlushedEvents->Emplace(MoveTemp(Entry));
+					}
 				}
 			}
 
@@ -535,7 +560,7 @@ void FAnalyticsProviderET::FlushEvents()
 			// Don't set a response callback if we are in our destructor, as the instance will no longer be there to call.
 			if (!bInDestructor)
 			{
-				HttpRequest->OnProcessRequestComplete().BindSP(this, &FAnalyticsProviderET::EventRequestComplete);
+				HttpRequest->OnProcessRequestComplete().BindSP(this, &FAnalyticsProviderET::EventRequestComplete, FlushedEvents);
 			}
 			HttpRequest->ProcessRequest();
 		}
@@ -597,7 +622,10 @@ void FAnalyticsProviderET::FlushEvents()
 						*EventParams));
 					PayloadSize = HttpRequest->GetURL().Len();
 					HttpRequest->SetVerb(TEXT("GET"));
-					HttpRequest->OnProcessRequestComplete().BindRaw(this, &FAnalyticsProviderET::EventRequestComplete);
+					if (!bInDestructor)
+					{
+						HttpRequest->OnProcessRequestComplete().BindSP(this, &FAnalyticsProviderET::EventRequestComplete, TSharedPtr< TArray<FAnalyticsEventEntry> >());
+					}
 					HttpRequest->ProcessRequest();
 				}
 			}
@@ -692,14 +720,48 @@ void FAnalyticsProviderET::SetDefaultEventAttributes(TArray<FAnalyticsEventAttri
 	}
 }
 
-void FAnalyticsProviderET::EventRequestComplete(FHttpRequestPtr HttpRequest, FHttpResponsePtr HttpResponse, bool bSucceeded)
+void FAnalyticsProviderET::EventRequestComplete(FHttpRequestPtr HttpRequest, FHttpResponsePtr HttpResponse, bool, TSharedPtr< TArray<FAnalyticsEventEntry> > FlushedEvents)
 {
-	if (bSucceeded && HttpResponse.IsValid())
+	// process responses
+	bool bEventsDelivered = false;
+	if (HttpResponse.IsValid())
 	{
 		UE_LOG(LogAnalytics, VeryVerbose, TEXT("[%s] ET response for [%s]. Code: %d. Payload: %s"), *APIKey, *HttpRequest->GetURL(), HttpResponse->GetResponseCode(), *HttpResponse->GetContentAsString());
+		if (EHttpResponseCodes::IsOk(HttpResponse->GetResponseCode()))
+		{
+			bEventsDelivered = true;
+		}
 	}
 	else
 	{
 		UE_LOG(LogAnalytics, VeryVerbose, TEXT("[%s] ET response for [%s]. No response"), *APIKey, *HttpRequest->GetURL());
 	}
+
+	// if the events were not delivered and FlushedEvents is passed, requeue them at the beginning
+	if (FlushedEvents.IsValid() && !bEventsDelivered)
+	{
+		// add a dropped submission event so we can see how often this is happening
+		if (bShouldCacheEvents && CachedEvents.Num() < 1024)
+		{
+			TArray<FAnalyticsEventAttribute> Attributes;
+			Attributes.Emplace(FAnalyticsEventAttribute(FString(TEXT("HTTP_STATUS")), FString::Printf(TEXT("%d"), HttpResponse.IsValid() ? HttpResponse->GetResponseCode() : 0)));
+			Attributes.Emplace(FAnalyticsEventAttribute(FString(TEXT("URL")), HttpRequest->GetURL()));
+			Attributes.Emplace(FAnalyticsEventAttribute(FString(TEXT("EVENTS_IN_BATCH")), FString::Printf(TEXT("%d"), FlushedEvents->Num())));
+			Attributes.Emplace(FAnalyticsEventAttribute(FString(TEXT("EVENTS_QUEUED")), FString::Printf(TEXT("%d"), CachedEvents.Num())));
+			CachedEvents.Emplace(FAnalyticsEventEntry(FString(TEXT("ET.DroppedSubmission")), MoveTemp(Attributes), false, false));
+		}
+
+		// if we're being super spammy or have been offline forever, just leave it at the ET.DroppedSubmission event
+		if (bShouldCacheEvents && CachedEvents.Num() < 256)
+		{
+			UE_LOG(LogAnalytics, Log, TEXT("[%s] ET Requeuing %d analytics events due to failure to send"), *APIKey, FlushedEvents->Num());
+
+			// put them at the beginning since it should include a default attributes entry and we don't want to change the current default attributes
+			CachedEvents.Insert(*FlushedEvents, 0);
+		}
+		else
+		{
+			UE_LOG(LogAnalytics, Error, TEXT("[%s] ET dropping %d analytics events due to too many in queue (%d)"), *APIKey, FlushedEvents->Num(), CachedEvents.Num());
+		}
+	}
 }
diff --git a/Engine/Source/Runtime/Analytics/AnalyticsET/Public/AnalyticsET.h b/Engine/Source/Runtime/Analytics/AnalyticsET/Public/AnalyticsET.h
--- a/Engine/Source/Runtime/Analytics/AnalyticsET/Public/AnalyticsET.h
+++ b/Engine/Source/Runtime/Analytics/AnalyticsET/Public/AnalyticsET.h
@@ -49,6 +49,8 @@ public:
 		FString AppVersionET;
 		/** When true, sends events using the legacy ET protocol that passes all attributes as URL parameters. Defaults to false. */
 		bool UseLegacyProtocol = false;
+		/** When true (default), events are dropped if flush fails */
+		bool bDropEventsOnFlushFailure = true;
 		/** The AppEnvironment that the data router should use. Defaults to GetDefaultAppEnvironment. */
 		FString AppEnvironment;
 		/** The UploadType that the data router should use. Defaults to GetDefaultUploadType. */
diff --git a/Engine/Source/Runtime/AppFramework/Private/Widgets/Testing/STestSuite.cpp b/Engine/Source/Runtime/AppFramework/Private/Widgets/Testing/STestSuite.cpp
--- a/Engine/Source/Runtime/AppFramework/Private/Widgets/Testing/STestSuite.cpp
+++ b/Engine/Source/Runtime/AppFramework/Private/Widgets/Testing/STestSuite.cpp
@@ -543,13 +543,16 @@ private:
 		const FText Text = LOCTEXT("TestText", "The quick brown fox jumps over the lazy dog 0123456789");
 		const FString FontName( FPaths::EngineContentDir() / TEXT("Slate/Fonts/Roboto-Regular.ttf" ) );
 		uint32 FontSize = 14;
+		FSlateFontInfo FontInfo(FontName, FontSize);
+		FontInfo.OutlineSettings.OutlineColor = FLinearColor::Blue;
+		FontInfo.OutlineSettings.OutlineSize = 2;
 
 		FSlateDrawElement::MakeText(
 			InParams.OutDrawElements,
 			InParams.Layer,
 			InParams.Geometry.ToPaintGeometry(FVector2D(0,0), InParams.Geometry.Size, FontScale),
 			Text.ToString(),
-			FSlateFontInfo( FontName,FontSize ),
+			FontInfo,
 			InParams.ClippingRect,
 			InParams.bEnabled ? ESlateDrawEffect::None : ESlateDrawEffect::DisabledEffect,
 			FColor( 255, 255, 255 )
diff --git a/Engine/Source/Runtime/Apple/MetalRHI/Private/MetalShaders.cpp b/Engine/Source/Runtime/Apple/MetalRHI/Private/MetalShaders.cpp
--- a/Engine/Source/Runtime/Apple/MetalRHI/Private/MetalShaders.cpp
+++ b/Engine/Source/Runtime/Apple/MetalRHI/Private/MetalShaders.cpp
@@ -128,8 +128,6 @@ TMetalBaseShader<BaseResourceType, ShaderType>::TMetalBaseShader(const TArray<ui
 	FMetalCompiledShaderKey Key(CodeLength, CodeCRC);
 
 	// Find the existing compiled shader in the cache.
-	Function = GetMetalCompiledShaderCache().FindRef(Key);
-	if (!Function)
 	{
 		id<MTLLibrary> Library;
 		
diff --git a/Engine/Source/Runtime/Core/Private/Android/AndroidFile.cpp b/Engine/Source/Runtime/Core/Private/Android/AndroidFile.cpp
--- a/Engine/Source/Runtime/Core/Private/Android/AndroidFile.cpp
+++ b/Engine/Source/Runtime/Core/Private/Android/AndroidFile.cpp
@@ -296,7 +296,8 @@ public:
 			// Can't write to assets.
 			return false;
 		}
-		
+
+		bool bSuccess = true;
 		while (BytesToWrite)
 		{
 			check(BytesToWrite >= 0);
@@ -304,7 +305,8 @@ public:
 			check(Source);
 			if (pwrite(File->Handle, Source, ThisSize, CurrentOffset) != ThisSize)
 			{
-				return false;
+				bSuccess = false;
+				break;
 			}
 			CurrentOffset += ThisSize;
 			Source += ThisSize;
@@ -314,7 +316,7 @@ public:
 		// Update the cached file length
 		Length = FMath::Max(Length, CurrentOffset);
 
-		return true;
+		return bSuccess;
 	}
 
 	virtual int64 Size() override
diff --git a/Engine/Source/Runtime/Core/Private/Android/AndroidMisc.cpp b/Engine/Source/Runtime/Core/Private/Android/AndroidMisc.cpp
--- a/Engine/Source/Runtime/Core/Private/Android/AndroidMisc.cpp
+++ b/Engine/Source/Runtime/Core/Private/Android/AndroidMisc.cpp
@@ -724,7 +724,7 @@ class IPlatformChunkInstall* FAndroidMisc::GetPlatformChunkInstall()
 		else
 		{
 			// Placeholder instance
-			ChunkInstall = new FGenericPlatformChunkInstall();
+			ChunkInstall = FGenericPlatformMisc::GetPlatformChunkInstall();
 		}
 	}
 
diff --git a/Engine/Source/Runtime/Core/Private/IOS/IOSPlatformMisc.cpp b/Engine/Source/Runtime/Core/Private/IOS/IOSPlatformMisc.cpp
--- a/Engine/Source/Runtime/Core/Private/IOS/IOSPlatformMisc.cpp
+++ b/Engine/Source/Runtime/Core/Private/IOS/IOSPlatformMisc.cpp
@@ -668,8 +668,33 @@ FString FIOSPlatformMisc::GetUniqueDeviceId()
 
 class IPlatformChunkInstall* FIOSPlatformMisc::GetPlatformChunkInstall()
 {
-    static FIOSChunkInstall Singleton;
-    return &Singleton;
+	static IPlatformChunkInstall* ChunkInstall = nullptr;
+	if (!ChunkInstall)
+	{
+		FString ProviderName;
+		GConfig->GetString(TEXT("StreamingInstall"), TEXT("DefaultProviderName"), ProviderName, GEngineIni);
+		if (ProviderName == TEXT("HTTPChunkInstaller"))
+		{
+			IPlatformChunkInstallModule* PlatformChunkInstallModule = FModuleManager::LoadModulePtr<IPlatformChunkInstallModule>("HTTPChunkInstaller");
+			if (PlatformChunkInstallModule != NULL)
+			{
+				// Attempt to grab the platform installer
+				ChunkInstall = PlatformChunkInstallModule->GetPlatformChunkInstall();
+			}
+		}
+		else if (ProviderName == TEXT("IOSChunkInstaller"))
+		{
+			static FIOSChunkInstall Singleton;
+			ChunkInstall = &Singleton;
+		}
+		if (!ChunkInstall)
+		{
+			// Placeholder instance
+			ChunkInstall = FGenericPlatformMisc::GetPlatformChunkInstall();
+		}
+	}
+
+	return ChunkInstall;
 }
 
 
diff --git a/Engine/Source/Runtime/Core/Private/Windows/WindowsPlatformMisc.cpp b/Engine/Source/Runtime/Core/Private/Windows/WindowsPlatformMisc.cpp
--- a/Engine/Source/Runtime/Core/Private/Windows/WindowsPlatformMisc.cpp
+++ b/Engine/Source/Runtime/Core/Private/Windows/WindowsPlatformMisc.cpp
@@ -2863,7 +2863,7 @@ IPlatformChunkInstall* FWindowsPlatformMisc::GetPlatformChunkInstall()
 #endif
 		{
 			// Placeholder instance
-			ChunkInstall = new FGenericPlatformChunkInstall();
+			ChunkInstall = FGenericPlatformMisc::GetPlatformChunkInstall();
 		}
 	}
 
diff --git a/Engine/Source/Runtime/Core/Public/UObject/UnrealNames.inl b/Engine/Source/Runtime/Core/Public/UObject/UnrealNames.inl
--- a/Engine/Source/Runtime/Core/Public/UObject/UnrealNames.inl
+++ b/Engine/Source/Runtime/Core/Public/UObject/UnrealNames.inl
@@ -133,6 +133,10 @@ REGISTER_NAME(283,PendingNetDriver)
 REGISTER_NAME(284,BeaconNetDriver)
 REGISTER_NAME(285,FlushNetDormancy)
 REGISTER_NAME(286,DemoNetDriver)
+REGISTER_NAME(287,GameSession)
+REGISTER_NAME(288,PartySession)
+REGISTER_NAME(289,GamePort)
+REGISTER_NAME(290,BeaconPort)
 
 // Texture settings.
 REGISTER_NAME(300,Linear)
diff --git a/Engine/Source/Runtime/CoreUObject/Public/UObject/CoreOnline.h b/Engine/Source/Runtime/CoreUObject/Public/UObject/CoreOnline.h
--- a/Engine/Source/Runtime/CoreUObject/Public/UObject/CoreOnline.h
+++ b/Engine/Source/Runtime/CoreUObject/Public/UObject/CoreOnline.h
@@ -13,7 +13,7 @@ const FName GamePort(TEXT("GamePort"));
 const FName BeaconPort(TEXT("BeaconPort"));
 
 #if !CPP
-// Circular dependency on Core vs UHT means we have to noexport this struct so tools can build
+// Circular dependency on Core vs UHT means we have to noexport these structs so tools can build
 USTRUCT(noexport)
 struct FJoinabilitySettings
 {
@@ -41,6 +41,13 @@ struct FJoinabilitySettings
 	UPROPERTY()
 	int32 MaxPartySize;
 };
+
+USTRUCT(noexport)
+struct FUniqueNetIdWrapper
+{
+	//GENERATED_BODY()
+};
+
 #endif
 
 struct FJoinabilitySettings
@@ -105,8 +112,8 @@ protected:
 public:
 
 	/**
-	*	Comparison operator
-	*/
+	 *	Comparison operator
+	 */
 	bool operator==(const FUniqueNetId& Other) const
 	{
 		return Other.Compare(*this);
@@ -169,17 +176,27 @@ public:
 	virtual ~FUniqueNetId() {}
 };
 
-// This is not a USTRUCT due to cross module dependencies
 struct FUniqueNetIdWrapper
 {
 public:
 
 	FUniqueNetIdWrapper() 
-	: UniqueNetId()
+		: UniqueNetId()
+	{
+	}
+
+	FUniqueNetIdWrapper(const TSharedRef<const FUniqueNetId>& InUniqueNetId)
+		: UniqueNetId(InUniqueNetId)
 	{
 	}
+
 	FUniqueNetIdWrapper(const TSharedPtr<const FUniqueNetId>& InUniqueNetId)
-	: UniqueNetId(InUniqueNetId)
+		: UniqueNetId(InUniqueNetId)
+	{
+	}
+
+	FUniqueNetIdWrapper(const FUniqueNetIdWrapper& InUniqueNetId)
+		: UniqueNetId(InUniqueNetId.UniqueNetId)
 	{
 	}
 
@@ -214,6 +231,12 @@ public:
 		return IsValid() ? UniqueNetId->ToString() : TEXT("INVALID");
 	}
 
+	/** Convert this value to a string with additional information */
+	virtual FString ToDebugString() const
+	{
+		return IsValid() ? UniqueNetId->ToDebugString() : TEXT("INVALID");
+	}
+
 	/** Is the FUniqueNetId wrapped in this object valid */
 	bool IsValid() const
 	{
@@ -254,5 +277,6 @@ public:
 
 protected:
 
+	// Actual unique id
 	TSharedPtr<const FUniqueNetId> UniqueNetId;
 };
diff --git a/Engine/Source/Runtime/Engine/Classes/GameFramework/OnlineReplStructs.h b/Engine/Source/Runtime/Engine/Classes/GameFramework/OnlineReplStructs.h
--- a/Engine/Source/Runtime/Engine/Classes/GameFramework/OnlineReplStructs.h
+++ b/Engine/Source/Runtime/Engine/Classes/GameFramework/OnlineReplStructs.h
@@ -17,7 +17,7 @@ class FJsonValue;
  * over network RPC and actor replication
  */
 USTRUCT(BlueprintType)
-struct FUniqueNetIdRepl
+struct FUniqueNetIdRepl : public FUniqueNetIdWrapper
 {
 	GENERATED_USTRUCT_BODY()
 
@@ -25,48 +25,27 @@ struct FUniqueNetIdRepl
 	{
 	}
 
- 	FUniqueNetIdRepl(const FUniqueNetIdWrapper& InWrapper)
- 	: UniqueNetId(InWrapper.GetUniqueNetId())
- 	{
- 	}
-
-	FUniqueNetIdRepl(const TSharedRef<const FUniqueNetId>& InUniqueNetId) :
-		UniqueNetId(InUniqueNetId)
+	FUniqueNetIdRepl(const FUniqueNetIdRepl& InWrapper)
+		: FUniqueNetIdWrapper(InWrapper.UniqueNetId)
 	{
 	}
 
-	FUniqueNetIdRepl(const TSharedPtr<const FUniqueNetId>& InUniqueNetId) :
-		UniqueNetId(InUniqueNetId)
-	{
-	}
+  	FUniqueNetIdRepl(const FUniqueNetIdWrapper& InWrapper)
+  		: FUniqueNetIdWrapper(InWrapper)
+  	{
+  	}
 
-	~FUniqueNetIdRepl() {}
-
-	/** Assignment operator */
-	FUniqueNetIdRepl& operator=(const FUniqueNetIdRepl& Other)
+ 	FUniqueNetIdRepl(const TSharedRef<const FUniqueNetId>& InUniqueNetId)
+		: FUniqueNetIdWrapper(InUniqueNetId)
 	{
-		if(&Other != this)
-		{
-			UniqueNetId = Other.UniqueNetId;
-		}
-		
-		return *this;
-	}
-
-	/** Comparison operator */
-	bool operator==(FUniqueNetIdRepl const& Other) const
-	{
-		// Both invalid structs or both valid and deep comparison equality
-		bool bBothValid = IsValid() && Other.IsValid();
-		bool bBothInvalid = !IsValid() && !Other.IsValid();
-		return (bBothInvalid || (bBothValid && (*UniqueNetId == *Other.UniqueNetId)));
 	}
+ 
+ 	FUniqueNetIdRepl(const TSharedPtr<const FUniqueNetId>& InUniqueNetId)
+		: FUniqueNetIdWrapper(InUniqueNetId)
+ 	{
+ 	}
 
-	/** Comparison operator */
-	bool operator!=(FUniqueNetIdRepl const& Other) const
-	{
-		return !(FUniqueNetIdRepl::operator==(Other));
-	}
+	~FUniqueNetIdRepl() {}
 
     /** Export contents of this struct as a string */
 	bool ExportTextItem(FString& ValueStr, FUniqueNetIdRepl const& DefaultValue, UObject* Parent, int32 PortFlags, UObject* ExportRootScope) const;
@@ -83,69 +62,11 @@ struct FUniqueNetIdRepl
 	/** Serialization to any FArchive */
 	bool Serialize(FArchive& Ar);
 	
-	/** Convert this value to a string */
-	FString ToString() const
-	{
-		return IsValid() ? UniqueNetId->ToString() : TEXT("INVALID");
-	}
-
-	/** Convert this value to a string with additional information */
-	virtual FString ToDebugString() const 
-	{
-		return IsValid() ? UniqueNetId->ToDebugString() : TEXT("INVALID");
-	}
-
 	/** Convert this unique id to a json value */
 	TSharedRef<FJsonValue> ToJson() const;
 	/** Create a unique id from a json string */
 	void FromJson(const FString& InValue);
 
-	/** Is the FUniqueNetId wrapped in this object valid */
-	bool IsValid() const
-	{
-		return UniqueNetId.IsValid() && UniqueNetId->IsValid();
-	}
-	
-	/** 
-	 * Assign a unique id to this wrapper object
-	 *
-	 * @param InUniqueNetId id to associate
-	 */
-	void SetUniqueNetId(const TSharedPtr<const FUniqueNetId>& InUniqueNetId)
-	{
-		UniqueNetId = InUniqueNetId;
-	}
-
-	/** @return unique id associated with this wrapper object */
-	const TSharedPtr<const FUniqueNetId>& GetUniqueNetId() const
-	{
-		return UniqueNetId;
-	}
-
-	/**
-	 * Dereference operator returns a reference to the FUniqueNetId
-	 */
-	const FUniqueNetId& operator*() const
-	{
-		return *UniqueNetId;
-	}
-
-	/**
-	 * Arrow operator returns a pointer to this FUniqueNetId
-	 */
-	const FUniqueNetId* operator->() const
-	{
-		return UniqueNetId.Get();
-	}
-
-	/**
-	 * implicit cast operator to FUniqueNetIdWrapper
-	 */
- 	FORCEINLINE operator FUniqueNetIdWrapper() const
- 	{
- 		return FUniqueNetIdWrapper(UniqueNetId);
- 	}
-
 	friend inline uint32 GetTypeHash(FUniqueNetIdRepl const& Value)
 	{
 		if (Value.UniqueNetId.IsValid())
@@ -157,7 +78,6 @@ struct FUniqueNetIdRepl
 	}
 
 protected:
-	TSharedPtr<const FUniqueNetId> UniqueNetId;
 
 	/** Helper to create an FUniqueNetId from a string */
 	void UniqueIdFromString(const FString& Contents);
diff --git a/Engine/Source/Runtime/Engine/Classes/GameFramework/PlayerController.h b/Engine/Source/Runtime/Engine/Classes/GameFramework/PlayerController.h
--- a/Engine/Source/Runtime/Engine/Classes/GameFramework/PlayerController.h
+++ b/Engine/Source/Runtime/Engine/Classes/GameFramework/PlayerController.h
@@ -498,13 +498,13 @@ public:
 	 * @return true if the world coordinate was successfully projected to the screen.
 	 */
 	UFUNCTION(BlueprintCallable, Category = "Game|Player", meta = (DisplayName = "ConvertWorldLocationToScreenLocation", Keywords = "project"))
-	bool ProjectWorldLocationToScreen(FVector WorldLocation, FVector2D& ScreenLocation) const;
+	bool ProjectWorldLocationToScreen(FVector WorldLocation, FVector2D& ScreenLocation, bool bPlayerViewportRelative = false) const;
 
 	/**
 	 * Convert a World Space 3D position into a 3D Screen Space position.
 	 * @return true if the world coordinate was successfully projected to the screen.
 	 */
-	bool ProjectWorldLocationToScreenWithDistance(FVector WorldLocation, FVector& ScreenLocation) const;
+	bool ProjectWorldLocationToScreenWithDistance(FVector WorldLocation, FVector& ScreenLocation, bool bPlayerViewportRelative = false) const;
 	
 	/** Positions the mouse cursor in screen space, in pixels. */
 	UFUNCTION( BlueprintCallable, Category="Game|Player", meta = (DisplayName = "SetMousePosition", Keywords = "mouse" ))
diff --git a/Engine/Source/Runtime/Engine/Classes/Kismet/GameplayStatics.h b/Engine/Source/Runtime/Engine/Classes/Kismet/GameplayStatics.h
--- a/Engine/Source/Runtime/Engine/Classes/Kismet/GameplayStatics.h
+++ b/Engine/Source/Runtime/Engine/Classes/Kismet/GameplayStatics.h
@@ -885,9 +885,10 @@ class ENGINE_API UGameplayStatics : public UBlueprintFunctionLibrary
 	 * @param Player			Project using this player's view.
 	 * @param WorldPosition		World position to project.
 	 * @param ScreenPosition	(out) Corresponding 2D position in screen space
+	 * @param bPlayerViewportRelative	Should this be relative to the player viewport subregion (useful when using player attached widgets in split screen)
 	 */
 	UFUNCTION(BlueprintPure, Category = "Utilities")
-	static bool ProjectWorldToScreen(APlayerController const* Player, const FVector& WorldPosition, FVector2D& ScreenPosition);
+	static bool ProjectWorldToScreen(APlayerController const* Player, const FVector& WorldPosition, FVector2D& ScreenPosition, bool bPlayerViewportRelative = false);
 
 	//~ Utility functions for interacting with Options strings
 
diff --git a/Engine/Source/Runtime/Engine/Private/Animation/AnimSingleNodeInstanceProxy.cpp b/Engine/Source/Runtime/Engine/Private/Animation/AnimSingleNodeInstanceProxy.cpp
--- a/Engine/Source/Runtime/Engine/Private/Animation/AnimSingleNodeInstanceProxy.cpp
+++ b/Engine/Source/Runtime/Engine/Private/Animation/AnimSingleNodeInstanceProxy.cpp
@@ -31,10 +31,7 @@ bool FAnimSingleNodeInstanceProxy::Evaluate(FPoseContext& Output)
 		= false;
 #endif
 
-	if (CurrentAsset != NULL &&
-//@HSL_BEGIN - CCL - Seeing crashes due to processing animations that have begun to be destroyed
-		!CurrentAsset->HasAnyFlags(RF_BeginDestroyed))
-//@HSL_END
+	if (CurrentAsset != NULL && !CurrentAsset->HasAnyFlags(RF_BeginDestroyed))
 	{
 		//@TODO: animrefactor: Seems like more code duplication than we need
 		if (UBlendSpaceBase* BlendSpace = Cast<UBlendSpaceBase>(CurrentAsset))
diff --git a/Engine/Source/Runtime/Engine/Private/GameEngine.cpp b/Engine/Source/Runtime/Engine/Private/GameEngine.cpp
--- a/Engine/Source/Runtime/Engine/Private/GameEngine.cpp
+++ b/Engine/Source/Runtime/Engine/Private/GameEngine.cpp
@@ -331,6 +331,10 @@ TSharedRef<SWindow> UGameEngine::CreateGameWindow()
 	const FText WindowTitle = FText::Format(WindowTitleVar, Args);
 	const bool bShouldPreserveAspectRatio = GetDefault<UGeneralProjectSettings>()->bShouldWindowPreserveAspectRatio;
 	const bool bUseBorderlessWindow = GetDefault<UGeneralProjectSettings>()->bUseBorderlessWindow;
+	const bool bAllowWindowResize = GetDefault<UGeneralProjectSettings>()->bAllowWindowResize;
+	const bool bAllowClose = GetDefault<UGeneralProjectSettings>()->bAllowClose;
+	const bool bAllowMaximize = GetDefault<UGeneralProjectSettings>()->bAllowMaximize;
+	const bool bAllowMinimize = GetDefault<UGeneralProjectSettings>()->bAllowMinimize;
 
 	// Allow optional winX/winY parameters to set initial window position
 	EAutoCenter::Type AutoCenterType = EAutoCenter::PrimaryWorkArea;
@@ -387,7 +391,11 @@ TSharedRef<SWindow> UGameEngine::CreateGameWindow()
 	.UseOSWindowBorder(!bUseBorderlessWindow)
 	.CreateTitleBar(!bUseBorderlessWindow)
 	.ShouldPreserveAspectRatio(bShouldPreserveAspectRatio)
-	.LayoutBorder(bUseBorderlessWindow ? FMargin(0) : FMargin(5, 5, 5, 5));
+	.LayoutBorder(bUseBorderlessWindow ? FMargin(0) : FMargin(5, 5, 5, 5))
+	.SizingRule(bAllowWindowResize ? ESizingRule::UserSized : ESizingRule::FixedSize)
+	.HasCloseButton(bAllowClose)
+	.SupportsMinimize(bAllowMinimize)
+	.SupportsMaximize(bAllowMaximize);
 
 	const bool bShowImmediately = false;
 
diff --git a/Engine/Source/Runtime/Engine/Private/GameViewportClient.cpp b/Engine/Source/Runtime/Engine/Private/GameViewportClient.cpp
--- a/Engine/Source/Runtime/Engine/Private/GameViewportClient.cpp
+++ b/Engine/Source/Runtime/Engine/Private/GameViewportClient.cpp
@@ -1256,6 +1256,10 @@ void UGameViewportClient::Draw(FViewport* InViewport, FCanvas* SceneCanvas)
 		// Render the console.
 		if (ViewportConsole)
 		{
+			// Reset the debug canvas to be full-screen before drawing the console
+			// (the debug draw service above has messed with the viewport size to fit it to a single player's subregion)
+			DebugCanvasObject->Init(InViewport->GetSizeXY().X, InViewport->GetSizeXY().Y, NULL);
+
 			ViewportConsole->PostRender_Console(DebugCanvasObject);
 		}
 	}
diff --git a/Engine/Source/Runtime/Engine/Private/GameplayStatics.cpp b/Engine/Source/Runtime/Engine/Private/GameplayStatics.cpp
--- a/Engine/Source/Runtime/Engine/Private/GameplayStatics.cpp
+++ b/Engine/Source/Runtime/Engine/Private/GameplayStatics.cpp
@@ -2014,7 +2014,7 @@ bool UGameplayStatics::DeprojectScreenToWorld(APlayerController const* Player, c
 	return false;
 }
 
-bool UGameplayStatics::ProjectWorldToScreen(APlayerController const* Player, const FVector& WorldPosition, FVector2D& ScreenPosition)
+bool UGameplayStatics::ProjectWorldToScreen(APlayerController const* Player, const FVector& WorldPosition, FVector2D& ScreenPosition, bool bPlayerViewportRelative)
 {
 	ULocalPlayer* const LP = Player ? Player->GetLocalPlayer() : nullptr;
 	if (LP && LP->ViewportClient)
@@ -2024,7 +2024,14 @@ bool UGameplayStatics::ProjectWorldToScreen(APlayerController const* Player, con
 		if (LP->GetProjectionData(LP->ViewportClient->Viewport, eSSP_FULL, /*out*/ ProjectionData))
 		{
 			FMatrix const ViewProjectionMatrix = ProjectionData.ComputeViewProjectionMatrix();
-			return FSceneView::ProjectWorldToScreen(WorldPosition, ProjectionData.GetConstrainedViewRect(), ViewProjectionMatrix, ScreenPosition);
+			const bool bResult = FSceneView::ProjectWorldToScreen(WorldPosition, ProjectionData.GetConstrainedViewRect(), ViewProjectionMatrix, ScreenPosition);
+
+			if (bPlayerViewportRelative)
+			{
+				ScreenPosition -= FVector2D(ProjectionData.GetConstrainedViewRect().Min);
+			}
+
+			return bResult;
 		}
 	}
 
diff --git a/Engine/Source/Runtime/Engine/Private/PlayerController.cpp b/Engine/Source/Runtime/Engine/Private/PlayerController.cpp
--- a/Engine/Source/Runtime/Engine/Private/PlayerController.cpp
+++ b/Engine/Source/Runtime/Engine/Private/PlayerController.cpp
@@ -1905,15 +1905,15 @@ bool APlayerController::DeprojectScreenPositionToWorld(float ScreenX, float Scre
 }
 
 
-bool APlayerController::ProjectWorldLocationToScreen(FVector WorldLocation, FVector2D& ScreenLocation) const
+bool APlayerController::ProjectWorldLocationToScreen(FVector WorldLocation, FVector2D& ScreenLocation, bool bPlayerViewportRelative) const
 {
-	return UGameplayStatics::ProjectWorldToScreen(this, WorldLocation, ScreenLocation);
+	return UGameplayStatics::ProjectWorldToScreen(this, WorldLocation, ScreenLocation, bPlayerViewportRelative);
 }
 
-bool APlayerController::ProjectWorldLocationToScreenWithDistance(FVector WorldLocation, FVector& ScreenLocation) const
+bool APlayerController::ProjectWorldLocationToScreenWithDistance(FVector WorldLocation, FVector& ScreenLocation, bool bPlayerViewportRelative) const
 {
 	FVector2D ScreenLoc2D;
-	if (UGameplayStatics::ProjectWorldToScreen(this, WorldLocation, ScreenLoc2D))
+	if (UGameplayStatics::ProjectWorldToScreen(this, WorldLocation, ScreenLoc2D, bPlayerViewportRelative))
 	{
 		// find distance
 		ULocalPlayer const* const LP = GetLocalPlayer();
@@ -4177,35 +4177,38 @@ void APlayerController::SetSpectatorPawn(class ASpectatorPawn* NewSpectatorPawn)
 
 ASpectatorPawn* APlayerController::SpawnSpectatorPawn()
 {
-	ASpectatorPawn* SpawnedSpectator = NULL;
+	ASpectatorPawn* SpawnedSpectator = nullptr;
 
 	// Only spawned for the local player
-	if (GetSpectatorPawn() == NULL && IsLocalController())
+	if ((GetSpectatorPawn() == nullptr) && IsLocalController())
 	{
-		AGameState const* const GameState = GetWorld()->GameState;
-		if (GameState)
+		UWorld* World = GetWorld();
+		if (AGameState const* const GameState = World->GameState)
 		{
-			FActorSpawnParameters SpawnParams;
-			SpawnParams.Owner = this;
-			SpawnParams.SpawnCollisionHandlingOverride = ESpawnActorCollisionHandlingMethod::AlwaysSpawn;
-
-			SpawnParams.ObjectFlags |= RF_Transient;	// We never want to save spectator pawns into a map
-			SpawnedSpectator = GetWorld()->SpawnActor<ASpectatorPawn>(GameState->SpectatorClass, GetSpawnLocation(), GetControlRotation(), SpawnParams);
-			if (SpawnedSpectator)
+			if (UClass* SpectatorClass = GameState->SpectatorClass)
 			{
-				SpawnedSpectator->SetReplicates(false); // Client-side only
-				SpawnedSpectator->PossessedBy(this);
-				SpawnedSpectator->PawnClientRestart();
-				if (SpawnedSpectator->PrimaryActorTick.bStartWithTickEnabled)
+				FActorSpawnParameters SpawnParams;
+				SpawnParams.Owner = this;
+				SpawnParams.SpawnCollisionHandlingOverride = ESpawnActorCollisionHandlingMethod::AlwaysSpawn;
+
+				SpawnParams.ObjectFlags |= RF_Transient;	// We never want to save spectator pawns into a map
+				SpawnedSpectator = World->SpawnActor<ASpectatorPawn>(SpectatorClass, GetSpawnLocation(), GetControlRotation(), SpawnParams);
+				if (SpawnedSpectator)
 				{
-					SpawnedSpectator->SetActorTickEnabled(true);
-				}
+					SpawnedSpectator->SetReplicates(false); // Client-side only
+					SpawnedSpectator->PossessedBy(this);
+					SpawnedSpectator->PawnClientRestart();
+					if (SpawnedSpectator->PrimaryActorTick.bStartWithTickEnabled)
+					{
+						SpawnedSpectator->SetActorTickEnabled(true);
+					}
 
-				UE_LOG(LogPlayerController, Verbose, TEXT("Spawned spectator %s [server:%d]"), *GetNameSafe(SpawnedSpectator), GetNetMode() < NM_Client);
-			}
-			else
-			{
-				UE_LOG(LogPlayerController, Warning, TEXT("Failed to spawn spectator with class %s"), GameState->SpectatorClass ? *GameState->SpectatorClass->GetName() : TEXT("NULL"));
+					UE_LOG(LogPlayerController, Verbose, TEXT("Spawned spectator %s [server:%d]"), *GetNameSafe(SpawnedSpectator), GetNetMode() < NM_Client);
+				}
+				else
+				{
+					UE_LOG(LogPlayerController, Warning, TEXT("Failed to spawn spectator with class %s"), *GetNameSafe(SpectatorClass));
+				}
 			}
 		}
 		else
diff --git a/Engine/Source/Runtime/Engine/Private/Slate/SceneViewport.cpp b/Engine/Source/Runtime/Engine/Private/Slate/SceneViewport.cpp
--- a/Engine/Source/Runtime/Engine/Private/Slate/SceneViewport.cpp
+++ b/Engine/Source/Runtime/Engine/Private/Slate/SceneViewport.cpp
@@ -1158,9 +1158,7 @@ void FSceneViewport::ResizeFrame(uint32 NewWindowSizeX, uint32 NewWindowSizeY, E
 			}
 
 			// Resize viewport
-// @HSL_CHANGE_BEGIN - [REDACTED_EMAIL] - 5/31/2016 - Fixing windowed mode
 			FVector2D ViewportSize = WindowToResize->GetWindowSizeFromClientSize(FVector2D(SizeX, SizeY));
-// @HSL_CHANGE_END - [REDACTED_EMAIL] - 5/31/2016
 			FVector2D NewViewportSize = WindowToResize->GetViewportSize();
 
 			if (NewViewportSize != ViewportSize || NewWindowMode != WindowMode)
diff --git a/Engine/Source/Runtime/Engine/Private/UserInterface/CanvasItem.cpp b/Engine/Source/Runtime/Engine/Private/UserInterface/CanvasItem.cpp
--- a/Engine/Source/Runtime/Engine/Private/UserInterface/CanvasItem.cpp
+++ b/Engine/Source/Runtime/Engine/Private/UserInterface/CanvasItem.cpp
@@ -1330,7 +1330,7 @@ void FCanvasShapedTextItem::DrawStringInternal( FCanvas* InCanvas, const FVector
 
 		if( GlyphToRender.bIsVisible )
 		{
-			const FShapedGlyphFontAtlasData GlyphAtlasData = FontCache->GetShapedGlyphFontAtlasData(GlyphToRender);
+			const FShapedGlyphFontAtlasData GlyphAtlasData = FontCache->GetShapedGlyphFontAtlasData(GlyphToRender,FFontOutlineSettings::NoOutline);
 
 			if (GlyphAtlasData.Valid)
 			{
diff --git a/Engine/Source/Runtime/EngineSettings/Classes/GeneralProjectSettings.h b/Engine/Source/Runtime/EngineSettings/Classes/GeneralProjectSettings.h
--- a/Engine/Source/Runtime/EngineSettings/Classes/GeneralProjectSettings.h
+++ b/Engine/Source/Runtime/EngineSettings/Classes/GeneralProjectSettings.h
@@ -74,4 +74,16 @@ class ENGINESETTINGS_API UGeneralProjectSettings
 	/** Should the game attempt to start in VR, regardless of whether -vr was set on the commandline */
 	UPROPERTY(config, EditAnywhere, Category = Settings)
 	bool bStartInVR;
+	
+	UPROPERTY(config, EditAnywhere, Category = Settings)
+	bool bAllowWindowResize;
+
+	UPROPERTY(config, EditAnywhere, Category = Settings)
+	bool bAllowClose;
+
+	UPROPERTY(config, EditAnywhere, Category = Settings)
+	bool bAllowMaximize;
+
+	UPROPERTY(config, EditAnywhere, Category = Settings)
+	bool bAllowMinimize;
 };
diff --git a/Engine/Source/Runtime/EngineSettings/Private/EngineSettingsModule.cpp b/Engine/Source/Runtime/EngineSettings/Private/EngineSettingsModule.cpp
--- a/Engine/Source/Runtime/EngineSettings/Private/EngineSettingsModule.cpp
+++ b/Engine/Source/Runtime/EngineSettings/Private/EngineSettingsModule.cpp
@@ -58,7 +58,12 @@ UGeneralEngineSettings::UGeneralEngineSettings( const FObjectInitializer& Object
 UGeneralProjectSettings::UGeneralProjectSettings( const FObjectInitializer& ObjectInitializer )
 	: Super(ObjectInitializer)
 	, bShouldWindowPreserveAspectRatio(true)
+	, bUseBorderlessWindow(false)
 	, bStartInVR(false)
+	, bAllowWindowResize(true)
+	, bAllowClose(true)
+	, bAllowMaximize(true)
+	, bAllowMinimize(true)
 { }
 
 
diff --git a/Engine/Source/Runtime/Online/ICMP/Private/IcmpPrivatePCH.h b/Engine/Source/Runtime/Online/ICMP/Private/IcmpPrivatePCH.h
--- a/Engine/Source/Runtime/Online/ICMP/Private/IcmpPrivatePCH.h
+++ b/Engine/Source/Runtime/Online/ICMP/Private/IcmpPrivatePCH.h
@@ -29,6 +29,10 @@
 #define PLATFORM_SUPPORTS_ICMP 0
 #define PLATFORM_USES_POSIX_IMCP 0
 #define PING_ALLOWS_CUSTOM_THREAD_SIZE 0
+#elif PLATFORM_ANDROID
+#define PLATFORM_SUPPORTS_ICMP 1
+#define PLATFORM_USES_POSIX_IMCP 1
+#define PING_ALLOWS_CUSTOM_THREAD_SIZE 0
 #else
 #define PLATFORM_SUPPORTS_ICMP 0
 #define PLATFORM_USES_POSIX_IMCP 0
diff --git a/Engine/Source/Runtime/Slate/Private/Framework/Text/SlateTextRun.cpp b/Engine/Source/Runtime/Slate/Private/Framework/Text/SlateTextRun.cpp
--- a/Engine/Source/Runtime/Slate/Private/Framework/Text/SlateTextRun.cpp
+++ b/Engine/Source/Runtime/Slate/Private/Framework/Text/SlateTextRun.cpp
@@ -41,13 +41,17 @@ FVector2D FSlateTextRun::Measure( int32 BeginIndex, int32 EndIndex, float Scale,
 {
 	const FVector2D ShadowOffsetToApply((EndIndex == Range.EndIndex) ? FMath::Abs(Style.ShadowOffset.X * Scale) : 0.0f, FMath::Abs(Style.ShadowOffset.Y * Scale));
 
+	// Offset the measured shaped text by the outline since the outline was not factored into the size of the text
+	const float OutlineSize = Style.Font.OutlineSettings.OutlineSize * Scale;
+	const FVector2D OutlineSizeToApply(OutlineSize, OutlineSize);
+
 	if ( EndIndex - BeginIndex == 0 )
 	{
 		return FVector2D( ShadowOffsetToApply.X * Scale, GetMaxHeight( Scale ) );
 	}
 
 	// Use the full text range (rather than the run range) so that text that spans runs will still be shaped correctly
-	return ShapedTextCacheUtil::MeasureShapedText(TextContext.ShapedTextCache, FCachedShapedTextKey(FTextRange(0, Text->Len()), Scale, TextContext, Style.Font), FTextRange(BeginIndex, EndIndex), **Text) + ShadowOffsetToApply;
+	return ShapedTextCacheUtil::MeasureShapedText(TextContext.ShapedTextCache, FCachedShapedTextKey(FTextRange(0, Text->Len()), Scale, TextContext, Style.Font), FTextRange(BeginIndex, EndIndex), **Text) + ShadowOffsetToApply + OutlineSizeToApply;
 }
 
 int8 FSlateTextRun::GetKerning(int32 CurrentIndex, float Scale, const FRunTextContext& TextContext) const
diff --git a/Engine/Source/Runtime/Slate/Private/Widgets/Input/SButton.cpp b/Engine/Source/Runtime/Slate/Private/Widgets/Input/SButton.cpp
--- a/Engine/Source/Runtime/Slate/Private/Widgets/Input/SButton.cpp
+++ b/Engine/Source/Runtime/Slate/Private/Widgets/Input/SButton.cpp
@@ -177,7 +177,8 @@ FReply SButton::OnKeyUp(const FGeometry& MyGeometry, const FKeyEvent& InKeyEvent
 {
 	FReply Reply = FReply::Unhandled();
 
-	if (IsEnabled() && (InKeyEvent.GetKey() == EKeys::Enter || InKeyEvent.GetKey() == EKeys::SpaceBar || InKeyEvent.GetKey() == EKeys::Gamepad_FaceButton_Bottom))
+
+	if (IsEnabled() && (/* InKeyEvent.GetKey() == EKeys::Enter || InKeyEvent.GetKey() == EKeys::SpaceBar || */ InKeyEvent.GetKey() == EKeys::Gamepad_FaceButton_Bottom))
 	{
 		const bool bWasPressed = bIsPressed;
 
diff --git a/Engine/Source/Runtime/SlateCore/Private/Fonts/FontCache.cpp b/Engine/Source/Runtime/SlateCore/Private/Fonts/FontCache.cpp
--- a/Engine/Source/Runtime/SlateCore/Private/Fonts/FontCache.cpp
+++ b/Engine/Source/Runtime/SlateCore/Private/Fonts/FontCache.cpp
@@ -7,6 +7,7 @@
 #include "SlateFontRenderer.h"
 #include "SlateTextShaper.h"
 #include "LegacySlateFontInfoCache.h"
+#include "SlateFontInfo.h"
 
 DECLARE_DWORD_ACCUMULATOR_STAT(TEXT("Num Font Atlases"), STAT_SlateNumFontAtlases, STATGROUP_SlateMemory);
 DECLARE_DWORD_ACCUMULATOR_STAT(TEXT("Num Font Non-Atlased Textures"), STAT_SlateNumFontNonAtlasedTextures, STATGROUP_SlateMemory);
@@ -40,25 +41,27 @@ ETextShapingMethod GetDefaultTextShapingMethod()
 }
 
 
-FShapedGlyphEntryKey::FShapedGlyphEntryKey(const TSharedPtr<FShapedGlyphFaceData>& InFontFaceData, uint32 InGlyphIndex)
+FShapedGlyphEntryKey::FShapedGlyphEntryKey(const TSharedPtr<FShapedGlyphFaceData>& InFontFaceData, uint32 InGlyphIndex, const FFontOutlineSettings& InOutlineSettings)
 	: FontFace(InFontFaceData->FontFace)
 	, FontSize(InFontFaceData->FontSize)
+	, OutlineSize(InOutlineSettings.OutlineSize)
 	, FontScale(InFontFaceData->FontScale)
 	, GlyphIndex(InGlyphIndex)
 	, KeyHash(0)
 {
 	KeyHash = HashCombine(KeyHash, GetTypeHash(FontFace));
 	KeyHash = HashCombine(KeyHash, GetTypeHash(FontSize));
+	KeyHash = HashCombine(KeyHash, GetTypeHash(InOutlineSettings));
 	KeyHash = HashCombine(KeyHash, GetTypeHash(FontScale));
 	KeyHash = HashCombine(KeyHash, GetTypeHash(GlyphIndex));
 }
 
-
-FShapedGlyphSequence::FShapedGlyphSequence(TArray<FShapedGlyphEntry> InGlyphsToRender, const int16 InTextBaseline, const uint16 InMaxTextHeight, const UObject* InFontMaterial, const FSourceTextRange& InSourceTextRange)
+FShapedGlyphSequence::FShapedGlyphSequence(TArray<FShapedGlyphEntry> InGlyphsToRender, const int16 InTextBaseline, const uint16 InMaxTextHeight, const UObject* InFontMaterial, const FFontOutlineSettings& InOutlineSettings, const FSourceTextRange& InSourceTextRange)
 	: GlyphsToRender(MoveTemp(InGlyphsToRender))
 	, TextBaseline(InTextBaseline)
 	, MaxTextHeight(InMaxTextHeight)
 	, FontMaterial(InFontMaterial)
+	, OutlineSettings(InOutlineSettings)
 	, SequenceWidth(0)
 	, GlyphFontFaces()
 	, SourceIndicesToGlyphData(InSourceTextRange)
@@ -251,7 +254,7 @@ bool FShapedGlyphSequence::HasFoundGlyphAtOffset(FSlateFontCache& InFontCache, c
 	for (int32 SubGlyphIndex = InCurrentGlyphIndex;; ++SubGlyphIndex)
 	{
 		const FShapedGlyphEntry& SubGlyph = GlyphsToRender[SubGlyphIndex];
-		const FShapedGlyphFontAtlasData SubGlyphAtlasData = InFontCache.GetShapedGlyphFontAtlasData(SubGlyph);
+		const FShapedGlyphFontAtlasData SubGlyphAtlasData = InFontCache.GetShapedGlyphFontAtlasData(SubGlyph,FFontOutlineSettings::NoOutline);
 		TotalGlyphSpacing += SubGlyphAtlasData.HorizontalOffset + SubGlyph.XAdvance;
 		TotalGlyphAdvance += SubGlyph.XAdvance;
 
@@ -320,7 +323,7 @@ FShapedGlyphSequencePtr FShapedGlyphSequence::GetSubSequence(const int32 InStart
 
 	if (EnumerateVisualGlyphsInSourceRange(InStartIndex, InEndIndex, GlyphCallback) == EEnumerateGlyphsResult::EnumerationComplete)
 	{
-		return MakeShareable(new FShapedGlyphSequence(MoveTemp(SubGlyphsToRender), TextBaseline, MaxTextHeight, FontMaterial, FSourceTextRange(InStartIndex, InEndIndex - InStartIndex)));
+		return MakeShareable(new FShapedGlyphSequence(MoveTemp(SubGlyphsToRender), TextBaseline, MaxTextHeight, FontMaterial, OutlineSettings, FSourceTextRange(InStartIndex, InEndIndex - InStartIndex)));
 	}
 
 	return nullptr;
@@ -769,7 +772,7 @@ bool FSlateFontCache::AddNewEntry( TCHAR Character, const FSlateFontKey& InKey,
 	EFontFallback CharFallbackLevel;
 	const float FontScale = InKey.GetScale() * SubFontScalingFactor;
 
-	const bool bDidRender = FontRenderer->GetRenderData( FontData, InKey.GetFontInfo().Size, Character, RenderData, FontScale, &CharFallbackLevel);
+	const bool bDidRender = FontRenderer->GetRenderData( FontData, InKey.GetFontInfo().Size, InKey.GetFontOutlineSettings(), Character, RenderData, FontScale, &CharFallbackLevel);
 
 	OutCharacterEntry.Valid = bDidRender && AddNewEntry(RenderData, OutCharacterEntry.TextureIndex, OutCharacterEntry.StartU, OutCharacterEntry.StartV, OutCharacterEntry.USize, OutCharacterEntry.VSize);
 	if (OutCharacterEntry.Valid)
@@ -789,11 +792,11 @@ bool FSlateFontCache::AddNewEntry( TCHAR Character, const FSlateFontKey& InKey,
 	return OutCharacterEntry.Valid;
 }
 
-bool FSlateFontCache::AddNewEntry( const FShapedGlyphEntry& InShapedGlyph, FShapedGlyphFontAtlasData& OutAtlasData )
+bool FSlateFontCache::AddNewEntry(const FShapedGlyphEntry& InShapedGlyph, const FFontOutlineSettings& InOutlineSettings, FShapedGlyphFontAtlasData& OutAtlasData)
 {
 	// Render the glyph
 	FCharacterRenderData RenderData;
-	const bool bDidRender = FontRenderer->GetRenderData(InShapedGlyph, RenderData);
+	const bool bDidRender = FontRenderer->GetRenderData(InShapedGlyph, InOutlineSettings, RenderData);
 
 	OutAtlasData.Valid = bDidRender && AddNewEntry(RenderData, OutAtlasData.TextureIndex, OutAtlasData.StartU, OutAtlasData.StartV, OutAtlasData.USize, OutAtlasData.VSize);
 	if (OutAtlasData.Valid)
@@ -925,10 +928,10 @@ FShapedGlyphSequenceRef FSlateFontCache::ShapeUnidirectionalText( const TCHAR* I
 	return TextShaper->ShapeUnidirectionalText(InText, InTextStart, InTextLen, InFontInfo, InFontScale, InTextDirection, InTextShapingMethod);
 }
 
-FCharacterList& FSlateFontCache::GetCharacterList( const FSlateFontInfo &InFontInfo, float FontScale )
+FCharacterList& FSlateFontCache::GetCharacterList( const FSlateFontInfo &InFontInfo, float FontScale, const FFontOutlineSettings& InOutlineSettings )
 {
 	// Create a key for looking up each character
-	const FSlateFontKey FontKey( InFontInfo, FontScale );
+	const FSlateFontKey FontKey( InFontInfo, InOutlineSettings, FontScale );
 
 	//@HSL_BEGIN - Chance.Lyon - A critical section to help make this class thread-safe */
 	FScopeLock ScopeLock(&CacheCriticalSection);
@@ -952,18 +955,19 @@ FCharacterList& FSlateFontCache::GetCharacterList( const FSlateFontInfo &InFontI
 	return FontToCharacterListCache.Add( FontKey, MakeShareable( new FCharacterList( FontKey, *this ) ) ).Get();
 }
 
-FShapedGlyphFontAtlasData FSlateFontCache::GetShapedGlyphFontAtlasData( const FShapedGlyphEntry& InShapedGlyph )
+FShapedGlyphFontAtlasData FSlateFontCache::GetShapedGlyphFontAtlasData( const FShapedGlyphEntry& InShapedGlyph, const FFontOutlineSettings& InOutlineSettings )
 {
 	check(IsInGameThread() || IsInRenderingThread());
 
+	uint8 CachedTypeIndex = (uint8)(InOutlineSettings.OutlineSize <= 0 ? EFontCacheAtlasDataType::Regular : EFontCacheAtlasDataType::Outline);
 	const ESlateTextureAtlasThreadId AtlasThreadId = GetCurrentSlateTextureAtlasThreadId();
 	check(AtlasThreadId != ESlateTextureAtlasThreadId::Unknown);
 
-	const int32 CachedAtlasDataIndex = (AtlasThreadId == ESlateTextureAtlasThreadId::Game) ? 0 : 1;
+	const int32 CachedAtlasDataThreadIndex = (AtlasThreadId == ESlateTextureAtlasThreadId::Game) ? 0 : 1;
 
 	// Has the atlas data already been cached on the glyph?
 	{
-		TSharedPtr<FShapedGlyphFontAtlasData> CachedAtlasDataPin = InShapedGlyph.CachedAtlasData[CachedAtlasDataIndex].Pin();
+		TSharedPtr<FShapedGlyphFontAtlasData> CachedAtlasDataPin = InShapedGlyph.CachedAtlasData[CachedTypeIndex][CachedAtlasDataThreadIndex].Pin();
 		if (CachedAtlasDataPin.IsValid())
 		{
 			return *CachedAtlasDataPin;
@@ -972,23 +976,23 @@ FShapedGlyphFontAtlasData FSlateFontCache::GetShapedGlyphFontAtlasData( const FS
 
 	// Not cached on the glyph, so create a key for to look up this glyph, as it may
 	// have already been cached by another shaped text sequence
-	const FShapedGlyphEntryKey GlyphKey(InShapedGlyph.FontFaceData, InShapedGlyph.GlyphIndex);
+	const FShapedGlyphEntryKey GlyphKey(InShapedGlyph.FontFaceData, InShapedGlyph.GlyphIndex, InOutlineSettings);
 
 	// Has the atlas data already been cached by another shaped text sequence?
 	const TSharedRef<FShapedGlyphFontAtlasData>* FoundAtlasData = ShapedGlyphToAtlasData.Find(GlyphKey);
 	if (FoundAtlasData)
 	{
-		InShapedGlyph.CachedAtlasData[CachedAtlasDataIndex] = *FoundAtlasData;
+		InShapedGlyph.CachedAtlasData[CachedTypeIndex][CachedAtlasDataThreadIndex] = *FoundAtlasData;
 		return **FoundAtlasData;
 	}
 
 	// Not cached at all... create a new entry
 	TSharedRef<FShapedGlyphFontAtlasData> NewAtlasData = MakeShareable(new FShapedGlyphFontAtlasData());
-	AddNewEntry(InShapedGlyph, *NewAtlasData);
+	AddNewEntry(InShapedGlyph, InOutlineSettings, *NewAtlasData);
 
 	if (NewAtlasData->Valid)
 	{
-		InShapedGlyph.CachedAtlasData[CachedAtlasDataIndex] = NewAtlasData;
+		InShapedGlyph.CachedAtlasData[CachedTypeIndex][CachedAtlasDataThreadIndex] = NewAtlasData;
 		ShapedGlyphToAtlasData.Add(GlyphKey, NewAtlasData);
 	}
 
diff --git a/Engine/Source/Runtime/SlateCore/Private/Fonts/FontCacheFreeType.h b/Engine/Source/Runtime/SlateCore/Private/Fonts/FontCacheFreeType.h
--- a/Engine/Source/Runtime/SlateCore/Private/Fonts/FontCacheFreeType.h
+++ b/Engine/Source/Runtime/SlateCore/Private/Fonts/FontCacheFreeType.h
@@ -21,6 +21,7 @@
 	#include FT_MODULE_H
 	#include FT_BITMAP_H
 	#include FT_ADVANCES_H
+	#include FT_STROKER_H
 #endif // WITH_FREETYPE
 
 
diff --git a/Engine/Source/Runtime/SlateCore/Private/Fonts/FontMeasure.cpp b/Engine/Source/Runtime/SlateCore/Private/Fonts/FontMeasure.cpp
--- a/Engine/Source/Runtime/SlateCore/Private/Fonts/FontMeasure.cpp
+++ b/Engine/Source/Runtime/SlateCore/Private/Fonts/FontMeasure.cpp
@@ -276,7 +276,7 @@ FVector2D FSlateFontMeasure::MeasureStringInternal( const FString& Text, int32 S
 FMeasureCache* FSlateFontMeasure::FindOrAddMeasureCache( const FSlateFontInfo& InFontInfo, const float InFontScale ) const
 {
 #if USE_MEASURE_CACHING
-	FSlateFontKey FontKey(InFontInfo, InFontScale);
+	FSlateFontKey FontKey(InFontInfo, FFontOutlineSettings::NoOutline, InFontScale);
 	const FCompositeFont* const CompositeFont = InFontInfo.GetCompositeFont();
 
 	TSharedPtr< FSlateFontMeasureCache > FoundMeasureCache = FontToMeasureCache.FindRef( FontKey );
diff --git a/Engine/Source/Runtime/SlateCore/Private/Fonts/SlateFontInfo.cpp b/Engine/Source/Runtime/SlateCore/Private/Fonts/SlateFontInfo.cpp
--- a/Engine/Source/Runtime/SlateCore/Private/Fonts/SlateFontInfo.cpp
+++ b/Engine/Source/Runtime/SlateCore/Private/Fonts/SlateFontInfo.cpp
@@ -8,6 +8,8 @@
 /* FSlateFontInfo structors
  *****************************************************************************/
 
+FFontOutlineSettings FFontOutlineSettings::NoOutline;
+
 FSlateFontInfo::FSlateFontInfo( )
 	: FontObject(nullptr)
 	, FontMaterial(nullptr)
diff --git a/Engine/Source/Runtime/SlateCore/Private/Fonts/SlateFontRenderer.cpp b/Engine/Source/Runtime/SlateCore/Private/Fonts/SlateFontRenderer.cpp
--- a/Engine/Source/Runtime/SlateCore/Private/Fonts/SlateFontRenderer.cpp
+++ b/Engine/Source/Runtime/SlateCore/Private/Fonts/SlateFontRenderer.cpp
@@ -6,6 +6,17 @@
 #include "FontCacheCompositeFont.h"
 #include "LegacySlateFontInfoCache.h"
 
+/**
+ * Method for rendering fonts with the possibility of an outline.
+ * 0 = freetype does everything and generates a bitmap for the base glyph. 
+ * 1 = We override the freetype rasterizer.  Can help with some rendering anomalies on complicated fonts when freetype generates a wildly different stroke from the base glyph
+ * Note: The font cache must be flushed if changing this in the middle of a running instance
+ */ 
+static int32 OutlineFontRenderMethod = 1;
+FAutoConsoleVariableRef CVarOutlineFontRenderMethod(
+	TEXT("Slate.OutlineFontRenderMethod"),
+	OutlineFontRenderMethod,
+	TEXT("Changes the render method for outline fonts.  0 = freetype does everything and generates a bitmap for the base glyph (default).  1 = We override the freetype rasterizer.  Can help with some rendering anomalies on complicated fonts."));
 
 namespace SlateFontRendererUtils
 {
@@ -58,7 +69,9 @@ uint16 FSlateFontRenderer::GetMaxHeight(const FSlateFontInfo& InFontInfo, const
 		FFreeTypeGlyphCache::FCachedGlyphData CachedGlyphData;
 		if (FTGlyphCache->FindOrCache(FaceGlyphData.FaceAndMemory->GetFace(), FaceGlyphData.GlyphIndex, FaceGlyphData.GlyphFlags, InFontInfo.Size, InScale, CachedGlyphData))
 		{
-			return static_cast<uint16>(FreeTypeUtils::Convert26Dot6ToRoundedPixel<int32>(FT_MulFix(CachedGlyphData.Height, CachedGlyphData.SizeMetrics.y_scale)) * InScale);
+			// Adjust the height by the size of the outline that was applied.  The outline is counted twice since it surrounds the font on top and bottom
+			const float HeightAdjustment = InFontInfo.OutlineSettings.OutlineSize * 2;
+			return static_cast<uint16>((FreeTypeUtils::Convert26Dot6ToRoundedPixel<int32>(FT_MulFix(CachedGlyphData.Height, CachedGlyphData.SizeMetrics.y_scale)) + HeightAdjustment) * InScale);
 		}
 	}
 
@@ -81,7 +94,9 @@ int16 FSlateFontRenderer::GetBaseline(const FSlateFontInfo& InFontInfo, const fl
 		FFreeTypeGlyphCache::FCachedGlyphData CachedGlyphData;
 		if (FTGlyphCache->FindOrCache(FaceGlyphData.FaceAndMemory->GetFace(), FaceGlyphData.GlyphIndex, FaceGlyphData.GlyphFlags, InFontInfo.Size, InScale, CachedGlyphData))
 		{
-			return static_cast<int16>(FreeTypeUtils::Convert26Dot6ToRoundedPixel<int32>(CachedGlyphData.SizeMetrics.descender) * InScale);
+			// Adjust the height by the size of the outline that was applied.  The outline is counted twice since it surrounds the font on top and bottom
+			const float HeightAdjustment = InFontInfo.OutlineSettings.OutlineSize * 2;
+			return static_cast<int16>((FreeTypeUtils::Convert26Dot6ToRoundedPixel<int32>(CachedGlyphData.SizeMetrics.descender) + HeightAdjustment) * InScale);
 		}
 	}
 
@@ -232,7 +247,7 @@ FFreeTypeFaceGlyphData FSlateFontRenderer::GetFontFaceForCharacter(const FFontDa
 
 #endif // WITH_FREETYPE
 
-bool FSlateFontRenderer::GetRenderData(const FFontData& InFontData, const int32 InSize, TCHAR Char, FCharacterRenderData& OutRenderData, const float InScale, EFontFallback* OutFallbackLevel) const
+bool FSlateFontRenderer::GetRenderData(const FFontData& InFontData, const int32 InSize, const FFontOutlineSettings& InOutlineSettings, TCHAR Char, FCharacterRenderData& OutRenderData, const float InScale, EFontFallback* OutFallbackLevel) const
 {
 #if WITH_FREETYPE
 	FFreeTypeFaceGlyphData FaceGlyphData = GetFontFaceForCharacter(InFontData, Char, EFontFallback::FF_Max);
@@ -251,19 +266,20 @@ bool FSlateFontRenderer::GetRenderData(const FFontData& InFontData, const int32
 		check(Error == 0);
 
 		OutRenderData.Char = Char;
-		return GetRenderData(FaceGlyphData, InScale, OutRenderData);
+		return GetRenderData(FaceGlyphData, InScale, InOutlineSettings, OutRenderData);
 	}
 #endif // WITH_FREETYPE
 	return false;
 }
 
-bool FSlateFontRenderer::GetRenderData(const FShapedGlyphEntry& InShapedGlyph, FCharacterRenderData& OutRenderData) const
+bool FSlateFontRenderer::GetRenderData(const FShapedGlyphEntry& InShapedGlyph, const FFontOutlineSettings& InOutlineSettings, FCharacterRenderData& OutRenderData) const
 {
 #if WITH_FREETYPE
 	FFreeTypeFaceGlyphData FaceGlyphData;
 	FaceGlyphData.FaceAndMemory = InShapedGlyph.FontFaceData->FontFace.Pin();
 	FaceGlyphData.GlyphIndex = InShapedGlyph.GlyphIndex;
 	FaceGlyphData.GlyphFlags = InShapedGlyph.FontFaceData->GlyphFlags;
+
 	if (FaceGlyphData.FaceAndMemory.IsValid())
 	{
 		check(FaceGlyphData.FaceAndMemory->IsValid());
@@ -272,7 +288,7 @@ bool FSlateFontRenderer::GetRenderData(const FShapedGlyphEntry& InShapedGlyph, F
 		check(Error == 0);
 
 		OutRenderData.Char = 0;
-		return GetRenderData(FaceGlyphData, InShapedGlyph.FontFaceData->FontScale, OutRenderData);
+		return GetRenderData(FaceGlyphData, InShapedGlyph.FontFaceData->FontScale, InOutlineSettings, OutRenderData);
 	}
 #endif // WITH_FREETYPE
 	return false;
@@ -280,62 +296,260 @@ bool FSlateFontRenderer::GetRenderData(const FShapedGlyphEntry& InShapedGlyph, F
 
 #if WITH_FREETYPE
 
-bool FSlateFontRenderer::GetRenderData(const FFreeTypeFaceGlyphData& InFaceGlyphData, const float InScale, FCharacterRenderData& OutRenderData) const
+/**
+ * Represents one or more pixels of a rasterized glyph that have the same coverage (filled amount)
+ */
+struct FRasterizerSpan
+{
+	/** Start x location of the span */
+	int32 X;
+	/** Start y location of the span  */
+	int32 Y;
+	
+	/** Length of the span */
+	int32 Width;
+	
+	/** How "filled" the span is where 0 is completely transparent and 255 is completely opaque */
+	uint8 Coverage;
+
+	FRasterizerSpan(int32 InX, int32 InY, int32 InWidth, uint8 InCoverage)
+		: X(InX)
+		, Y(InY)
+		, Width(InWidth)
+		, Coverage(InCoverage)
+	{}
+};
+
+/**
+ * Represents a single rasterized glyph
+ */
+struct FRasterizerSpanList
+{
+	/** List of spans in the glyph */
+	TArray<FRasterizerSpan> Spans;
+	/** Bounds around the glyph */
+	FBox2D BoundingBox;
+
+	FRasterizerSpanList()
+		: BoundingBox(0)
+	{}
+};
+
+
+/**
+ * Rasterisizes a font glyph
+ */
+void RenderOutlineRows(FT_Library Library, FT_Outline* Outline, FRasterizerSpanList& OutSpansListOuter)
+{
+	auto RasterizerCallback = [](const int32 Y, const int32 Count, const FT_Span* const Spans, void* const UserData)
+	{
+		FRasterizerSpanList& OutSpanList = *static_cast<FRasterizerSpanList*>(UserData);
+
+		TArray<FRasterizerSpan>& OutSpans = OutSpanList.Spans;
+		FBox2D& BoundingBox = OutSpanList.BoundingBox;
+
+		OutSpans.Reserve(OutSpans.Num() + Count);
+		for(int32 SpanIndex = 0; SpanIndex < Count; ++SpanIndex)
+		{
+			const FT_Span& Span = Spans[SpanIndex];
+
+			BoundingBox += FVector2D(Span.x, Y);
+			BoundingBox += FVector2D(Span.x + Span.len - 1, Y);
+
+			OutSpans.Add(FRasterizerSpan(Span.x, Y, Span.len, Span.coverage));
+		}
+	};
+
+	FT_Raster_Params RasterParams;
+	FMemory::Memzero<FT_Raster_Params>(RasterParams);
+	RasterParams.flags = FT_RASTER_FLAG_AA | FT_RASTER_FLAG_DIRECT;
+	RasterParams.gray_spans = RasterizerCallback;
+	RasterParams.user = &OutSpansListOuter;
+
+	FT_Outline_Render(Library, Outline, &RasterParams);
+
+}
+
+bool FSlateFontRenderer::GetRenderData(const FFreeTypeFaceGlyphData& InFaceGlyphData, float InScale, const FFontOutlineSettings& InOutlineSettings, FCharacterRenderData& OutRenderData) const
 {
 	FT_Face Face = InFaceGlyphData.FaceAndMemory->GetFace();
 
-	// Get the slot for the glyph.  This contains measurement info
+	FT_Bitmap* Bitmap = nullptr;
+
+	// Get the lot for the glyph.  This contains measurement info
 	FT_GlyphSlot Slot = Face->glyph;
+
+	float ScaledOutlineSize = FMath::RoundToFloat(InOutlineSettings.OutlineSize * InScale);
+
+	if( (ScaledOutlineSize > 0 || OutlineFontRenderMethod == 1) && Slot->format == FT_GLYPH_FORMAT_OUTLINE)
+	{
+		// Render the filled area first
+		FRasterizerSpanList FillSpans;
+		RenderOutlineRows(FTLibrary->GetLibrary(), &Slot->outline, FillSpans);
+
+		FRasterizerSpanList OutlineSpans;
 		
-	FT_Render_Glyph(Slot, FT_RENDER_MODE_NORMAL);
+		FT_Stroker Stroker = nullptr;
+		FT_Glyph Glyph = nullptr;
 
-	// one byte per pixel 
-	const uint32 GlyphPixelSize = 1;
+		// If there is an outline, render it second after applying a border stroke to the font to produce an outline
+		if(ScaledOutlineSize > 0)
+		{
+			FT_Stroker_New(FTLibrary->GetLibrary(), &Stroker);
+			FT_Stroker_Set(Stroker, FMath::TruncToInt(FreeTypeUtils::ConvertPixelTo26Dot6<float>(ScaledOutlineSize)), FT_STROKER_LINECAP_ROUND, FT_STROKER_LINEJOIN_ROUND, 0);
 
-	FT_Bitmap* Bitmap = nullptr;
+			FT_Get_Glyph(Slot, &Glyph);
 
-	FT_Bitmap NewBitmap;
-	if (Slot->bitmap.pixel_mode == FT_PIXEL_MODE_MONO)
-	{
-		FT_Bitmap_New(&NewBitmap);
-		// Convert the mono font to 8bbp from 1bpp
-		FT_Bitmap_Convert(FTLibrary->GetLibrary(), &Slot->bitmap, &NewBitmap, 4);
+			bool bInner = false;
+			FT_Glyph_StrokeBorder(&Glyph, Stroker, bInner ? 1 : 0, 0);
+
+			FT_Outline* Outline = &reinterpret_cast<FT_OutlineGlyph>(Glyph)->outline;
+
+			RenderOutlineRows(FTLibrary->GetLibrary(), Outline, OutlineSpans);
+		}
+
+		const FBox2D BoundingBox = FillSpans.BoundingBox + OutlineSpans.BoundingBox;
+
+		FVector2D Size = BoundingBox.GetSize();
+
+		//Note: we add 1 to width and height because we need full pixels
+		int32 Width = FMath::TruncToInt(Size.X)+1;
+		int32 Height = FMath::TruncToInt(Size.Y)+1;
+
+		OutRenderData.RawPixels.Reset();
+
+		OutRenderData.MeasureInfo.SizeX = Width;
+		OutRenderData.MeasureInfo.SizeY = Height;
+
+		OutRenderData.RawPixels.AddZeroed(OutRenderData.MeasureInfo.SizeX * OutRenderData.MeasureInfo.SizeY);
+
+		int32 XMin = BoundingBox.Min.X;
+		int32 YMin = BoundingBox.Min.Y;
+
+		// Compute and copy the pixels for the total filled area of the glyph. 
+
+		// Copy the outline area first
+		const TArray<FRasterizerSpan>& FirstSpanList = OutlineSpans.Spans;
+		{
+			for(const FRasterizerSpan& Span : FirstSpanList)
+			{
+				for(int32 W = 0; W < Span.Width; ++W)
+				{
+					OutRenderData.RawPixels[((int32)((Height - 1 - (Span.Y - YMin)) * Width + Span.X - XMin + W))] = Span.Coverage;
+
+				}
+			}
+		}
+
+		// If there is an outline, it was rasterized by freetype with the filled area included
+		// This code will eliminate the filled area if the user requests an outline with separate translucency for the fill area
+		const TArray<FRasterizerSpan>& SecondSpanList = FillSpans.Spans;
+		{
+			if(ScaledOutlineSize > 0)
+			{
+				for (const FRasterizerSpan& Span : SecondSpanList)
+				{
+					for (int32 W = 0; W < Span.Width; ++W)
+					{
+						uint8& Src = OutRenderData.RawPixels[((int32)((Height - 1 - (Span.Y - YMin)) * Width + Span.X - XMin + W))];
+
+						if(InOutlineSettings.bSeparateFillAlpha)
+						{
+							// this method is better for transparent fill areas
+							Src = Span.Coverage ? FMath::Abs(Src - Span.Coverage) : 0;
+						}
+						else
+						{
+							// This method is better for opaque fill areas 
+							Src = Span.Coverage == 255 ? Span.Coverage : Src;
+						}
+					}
+				}
+			}
+			else
+			{
+				for (const FRasterizerSpan& Span : SecondSpanList)
+				{
+					for (int32 w = 0; w < Span.Width; ++w)
+					{
+						OutRenderData.RawPixels[((int32)((Height - 1 - (Span.Y - YMin)) * Width + Span.X - XMin + w))] = Span.Coverage;
+					}
+				}
+			}
+		}
+
+		FT_Stroker_Done(Stroker);
+		FT_Done_Glyph(Glyph);
+
+		// Note: in order to render the stroke properly AND to get proper measurements this must be done after rendering the stroke
+		FT_Render_Glyph(Slot, FT_RENDER_MODE_NORMAL);
 
-		Bitmap = &NewBitmap;
 	}
 	else
 	{
-		Bitmap = &Slot->bitmap;
-	}
-		
-	OutRenderData.RawPixels.Reset();
-	OutRenderData.RawPixels.AddUninitialized(Bitmap->rows * Bitmap->width);
+		// This path renders a standard font with no outline.  This may occur if the outline failed to generate
 
-	// Nothing to do for zero width or height glyphs
-	if (OutRenderData.RawPixels.Num())
-	{
-		// Copy the rendered bitmap to our raw pixels array
-		if (Slot->bitmap.pixel_mode != FT_PIXEL_MODE_MONO)
+		FT_Render_Glyph(Slot, FT_RENDER_MODE_NORMAL);
+
+		// one byte per pixel 
+		const uint32 GlyphPixelSize = 1;
+
+		FT_Bitmap NewBitmap;
+		if(Slot->bitmap.pixel_mode == FT_PIXEL_MODE_MONO)
 		{
-			for (uint32 Row = 0; Row < (uint32)Bitmap->rows; ++Row)
-			{
-				// Copy a single row. Note Bitmap.pitch contains the offset (in bytes) between rows.  Not always equal to Bitmap.width!
-				FMemory::Memcpy(&OutRenderData.RawPixels[Row*Bitmap->width], &Bitmap->buffer[Row*Bitmap->pitch], Bitmap->width*GlyphPixelSize);
-			}
+			FT_Bitmap_New(&NewBitmap);
+			// Convert the mono font to 8bbp from 1bpp
+			FT_Bitmap_Convert(FTLibrary->GetLibrary(), &Slot->bitmap, &NewBitmap, 4);
 
+			Bitmap = &NewBitmap;
 		}
 		else
 		{
-			// In Mono a value of 1 means the pixel is drawn and a value of zero means it is not. 
-			// So we must check each pixel and convert it to a color.
-			for (uint32 Height = 0; Height < (uint32)Bitmap->rows; ++Height )
+			Bitmap = &Slot->bitmap;
+		}
+
+		OutRenderData.RawPixels.Reset();
+		OutRenderData.RawPixels.AddUninitialized(Bitmap->rows * Bitmap->width);
+
+
+		// Nothing to do for zero width or height glyphs
+		if(OutRenderData.RawPixels.Num())
+		{
+			// Copy the rendered bitmap to our raw pixels array
+			if(Bitmap->pixel_mode != FT_PIXEL_MODE_MONO)
 			{
-				for (uint32 Width = 0; Width < (uint32)Bitmap->width; ++Width)
+				for(uint32 Row = 0; Row < (uint32)Bitmap->rows; ++Row)
 				{
-					OutRenderData.RawPixels[Height*Bitmap->width+Width] = Bitmap->buffer[Height*Bitmap->pitch+Width] == 1 ? 255 : 0;
+					// Copy a single row. Note Bitmap.pitch contains the offset (in bytes) between rows.  Not always equal to Bitmap.width!
+					FMemory::Memcpy(&OutRenderData.RawPixels[Row*Bitmap->width], &Bitmap->buffer[Row*Bitmap->pitch], Bitmap->width*GlyphPixelSize);
 				}
+
 			}
+			else
+			{
+				// In Mono a value of 1 means the pixel is drawn and a value of zero means it is not. 
+				// So we must check each pixel and convert it to a color.
+				for(uint32 Height = 0; Height < (uint32)Bitmap->rows; ++Height)
+				{
+					for(uint32 Width = 0; Width < (uint32)Bitmap->width; ++Width)
+					{
+						OutRenderData.RawPixels[Height*Bitmap->width+Width] = Bitmap->buffer[Height*Bitmap->pitch+Width] == 1 ? 255 : 0;
+					}
+				}
+			}
+		}
+
+		if(Bitmap->pixel_mode == FT_PIXEL_MODE_MONO)
+		{
+			FT_Bitmap_Done(FTLibrary->GetLibrary(), Bitmap);
+
 		}
+		OutRenderData.MeasureInfo.SizeX = Bitmap->width;
+		OutRenderData.MeasureInfo.SizeY = Bitmap->rows;
+
+		// Reset the outline to zero.  If we are in this path, either the outline failed to generate because the font supported or there is no outline
+		// We do not want to take it into account if it failed to generate
+		ScaledOutlineSize = 0;
 	}
 
 	const int32 Height = static_cast<int32>(FreeTypeUtils::Convert26Dot6ToRoundedPixel<int32>(FT_MulFix(Face->height, Face->size->metrics.y_scale)) * InScale);
@@ -343,8 +557,7 @@ bool FSlateFontRenderer::GetRenderData(const FFreeTypeFaceGlyphData& InFaceGlyph
 	// Set measurement info for this character
 	OutRenderData.GlyphIndex = InFaceGlyphData.GlyphIndex;
 	OutRenderData.HasKerning = FT_HAS_KERNING(Face) != 0;
-	OutRenderData.MeasureInfo.SizeX = Bitmap->width;
-	OutRenderData.MeasureInfo.SizeY = Bitmap->rows;
+
 	OutRenderData.MaxHeight = Height;
 
 	// Ascender is not scaled by freetype.  Scale it now. 
@@ -354,12 +567,8 @@ bool FSlateFontRenderer::GetRenderData(const FFreeTypeFaceGlyphData& InFaceGlyph
 	// Note we use Slot->advance instead of Slot->metrics.horiAdvance because Slot->Advance contains transformed position (needed if we scale)
 	OutRenderData.MeasureInfo.XAdvance = FreeTypeUtils::Convert26Dot6ToRoundedPixel<int16>(Slot->advance.x);
 	OutRenderData.MeasureInfo.HorizontalOffset = Slot->bitmap_left;
-	OutRenderData.MeasureInfo.VerticalOffset = Slot->bitmap_top;
+	OutRenderData.MeasureInfo.VerticalOffset = Slot->bitmap_top + ScaledOutlineSize;
 
-	if (Slot->bitmap.pixel_mode == FT_PIXEL_MODE_MONO)
-	{
-		FT_Bitmap_Done(FTLibrary->GetLibrary(), Bitmap);
-	}
 
 	return true;
 }
diff --git a/Engine/Source/Runtime/SlateCore/Private/Fonts/SlateFontRenderer.h b/Engine/Source/Runtime/SlateCore/Private/Fonts/SlateFontRenderer.h
--- a/Engine/Source/Runtime/SlateCore/Private/Fonts/SlateFontRenderer.h
+++ b/Engine/Source/Runtime/SlateCore/Private/Fonts/SlateFontRenderer.h
@@ -122,21 +122,22 @@ public:
 	 * Creates render data for a specific character 
 	 * 
 	 * @param InFontData		Raw font data to render the character with
-	 * @param InSize			The size of the font to draw
-	 * @param Char				The character to render
+	 * @param InSize		The size of the font to draw
+	 * @param InOutlineSettings	Outline settings to apply when rendering the characer
+	 * @param Char			The character to render
 	 * @param OutRenderData		Will contain the created render data
-	 * @param InScale			The scale of the font
+	 * @param InScale		The scale of the font
 	 * @param OutFallbackLevel	Outputs the fallback level of the font
 	 */
-	bool GetRenderData(const FFontData& InFontData, const int32 InSize, TCHAR Char, FCharacterRenderData& OutRenderData, const float InScale, EFontFallback* OutFallbackLevel = nullptr) const;
-	bool GetRenderData(const FShapedGlyphEntry& InShapedGlyph, FCharacterRenderData& OutRenderData) const;
+	bool GetRenderData(const FFontData& InFontData, const int32 InSize, const FFontOutlineSettings& InOutlineSettings, TCHAR Char, FCharacterRenderData& OutRenderData, const float InScale, EFontFallback* OutFallbackLevel = nullptr) const;
+	bool GetRenderData(const FShapedGlyphEntry& InShapedGlyph, const FFontOutlineSettings& InOutlineSettings, FCharacterRenderData& OutRenderData) const;
 
 private:
 #if WITH_FREETYPE
 	/**
 	 * Internal, shared implementation of GetRenderData
 	 */
-	bool GetRenderData(const FFreeTypeFaceGlyphData& InFaceGlyphData, const float InScale, FCharacterRenderData& OutRenderData) const;
+	bool GetRenderData(const FFreeTypeFaceGlyphData& InFaceGlyphData, float InScale, const FFontOutlineSettings& InOutlineSettings, FCharacterRenderData& OutRenderData) const;
 
 	/** 
 	 * Gets or loads a FreeType font face 
diff --git a/Engine/Source/Runtime/SlateCore/Private/Fonts/SlateTextShaper.cpp b/Engine/Source/Runtime/SlateCore/Private/Fonts/SlateTextShaper.cpp
--- a/Engine/Source/Runtime/SlateCore/Private/Fonts/SlateTextShaper.cpp
+++ b/Engine/Source/Runtime/SlateCore/Private/Fonts/SlateTextShaper.cpp
@@ -157,7 +157,7 @@ FShapedGlyphSequenceRef FSlateTextShaper::FinalizeTextShaping(TArray<FShapedGlyp
 	}
 #endif // WITH_FREETYPE
 
-	return MakeShareable(new FShapedGlyphSequence(MoveTemp(InGlyphsToRender), TextBaseline, MaxHeight, InFontInfo.FontMaterial, InSourceTextRange));
+	return MakeShareable(new FShapedGlyphSequence(MoveTemp(InGlyphsToRender), TextBaseline, MaxHeight, InFontInfo.FontMaterial, InFontInfo.OutlineSettings, InSourceTextRange));
 }
 
 #if WITH_FREETYPE
diff --git a/Engine/Source/Runtime/SlateCore/Private/Rendering/ElementBatcher.cpp b/Engine/Source/Runtime/SlateCore/Private/Rendering/ElementBatcher.cpp
--- a/Engine/Source/Runtime/SlateCore/Private/Rendering/ElementBatcher.cpp
+++ b/Engine/Source/Runtime/SlateCore/Private/Rendering/ElementBatcher.cpp
@@ -615,19 +615,26 @@ void FSlateElementBatcher::AddBoxElement( const FSlateDrawElement& DrawElement )
 void FSlateElementBatcher::AddTextElement(const FSlateDrawElement& DrawElement)
 {
 	const FSlateDataPayload& InPayload = DrawElement.GetDataPayload();
-	FColor Tint = PackVertexColor(InPayload.Tint);
+	FColor BaseTint = PackVertexColor(InPayload.Tint);
 
-	if ( Tint.A == 0 )
+	const FFontOutlineSettings& OutlineSettings = InPayload.FontInfo.OutlineSettings;
+
+	// Don't do anything if there the font would be completely transparent 
+	if ((BaseTint.A == 0 && OutlineSettings.OutlineSize == 0) || (BaseTint.A == 0 && OutlineSettings.OutlineColor.A == 0))
 	{
 		return;
 	}
 
+	int32 Len = InPayload.ImmutableText ? FCString::Strlen(InPayload.ImmutableText) : 0;
+	// Nothing to do if no text
+	if (Len == 0)
+	{
+		return;
+	}
+
+
 	NumDrawnTextsStat++;
 
-	//const FVector2D& InPosition = DrawElement.GetPosition();
-	//const FVector2D& Size = DrawElement.GetSize();
-	//const FVector2D& LocalSize = DrawElement.GetLocalSize();
-	//float Scale = DrawElement.GetScale();
 	const FSlateRect& InClippingRect = DrawElement.GetClippingRect();
 	ESlateDrawEffect::Type InDrawEffects = DrawElement.GetDrawEffects();
 	uint32 Layer = DrawElement.GetLayer();
@@ -646,173 +653,195 @@ void FSlateElementBatcher::AddTextElement(const FSlateDrawElement& DrawElement)
 	// Used to clip individual characters as we generate them.
 	FSlateRect LocalClipRect = TransformRect(InverseLayoutTransform, InClippingRect);
 
-	int32 Len = InPayload.ImmutableText ? FCString::Strlen(InPayload.ImmutableText) : 0;
-	// Nothing to do if no text
-	if( Len == 0 )
-	{
-		return;
-	}
 
 	FSlateFontCache& FontCache = *RenderingPolicy->GetFontCache();
 	FSlateShaderResourceManager& ResourceManager = *RenderingPolicy->GetResourceManager();
 
-	FCharacterList& CharacterList = FontCache.GetCharacterList( InPayload.FontInfo, FontScale );
+	const UObject* BaseFontMaterial = InPayload.FontInfo.FontMaterial;
+	const UObject* OutlineFontMaterial = OutlineSettings.OutlineMaterial;
 
-	float MaxHeight = CharacterList.GetMaxHeight();
+	bool bOutlineFont = OutlineSettings.OutlineSize > 0.0f;
 
-	uint32 FontTextureIndex = 0;
-	FSlateShaderResource* FontAtlasTexture = nullptr;
-	FSlateShaderResource* FontShaderResource = nullptr;
+	const float OutlineSize = OutlineSettings.OutlineSize;
 
-	FSlateElementBatch* ElementBatch = nullptr;
-	FSlateVertexArray* BatchVertices = nullptr;
-	FSlateIndexArray* BatchIndices = nullptr;
 
-	uint32 VertexOffset = 0;
-	uint32 IndexOffset = 0;
+	auto BuildFontGeometry = [&](const FFontOutlineSettings& InOutlineSettings, const FColor& InTint, const UObject* FontMaterial, int32 InLayer, int32 InOutlineHorizontalOffset)
+	{
+		FCharacterList& CharacterList = FontCache.GetCharacterList(InPayload.FontInfo, FontScale, InOutlineSettings);
 
-	float InvTextureSizeX = 0;
-	float InvTextureSizeY = 0;
+		float MaxHeight = CharacterList.GetMaxHeight();
 
-	float LineX = 0;
+		uint32 FontTextureIndex = 0;
+		FSlateShaderResource* FontAtlasTexture = nullptr;
+		FSlateShaderResource* FontShaderResource = nullptr;
 
-	FCharacterEntry PreviousCharEntry;
+		FSlateElementBatch* ElementBatch = nullptr;
+		FSlateVertexArray* BatchVertices = nullptr;
+		FSlateIndexArray* BatchIndices = nullptr;
 
-	int32 Kerning = 0;
+		uint32 VertexOffset = 0;
+		uint32 IndexOffset = 0;
 
-	FVector2D TopLeft(0,0);
+		float InvTextureSizeX = 0;
+		float InvTextureSizeY = 0;
 
-	const float PosX = TopLeft.X;
-	float PosY = TopLeft.Y;
+		float LineX = 0;
 
-	LineX = PosX;
-	
-	const bool bIsFontMaterial = InPayload.FontInfo.FontMaterial != nullptr;
+		FCharacterEntry PreviousCharEntry;
 
-	uint32 NumChars = Len;
+		int32 Kerning = 0;
 
-	uint32 NumLines = 1;
-	for( uint32 CharIndex = 0; CharIndex < NumChars; ++CharIndex )
-	{
-		const TCHAR CurrentChar = InPayload.ImmutableText[ CharIndex ];
+		FVector2D TopLeft(0, 0);
 
-		const bool IsNewline = (CurrentChar == '\n');
+		const float PosX = TopLeft.X;
+		float PosY = TopLeft.Y;
 
-		if (IsNewline)
-		{
-			// Move down: we are drawing the next line.
-			PosY += MaxHeight;
-			// Carriage return 
-			LineX = PosX;
+		LineX = PosX;
 
-			++NumLines;
+		const bool bIsFontMaterial = FontMaterial != nullptr;
 
-		}
-		else
+		uint32 NumChars = Len;
+
+		uint32 NumLines = 1;
+		for( uint32 CharIndex = 0; CharIndex < NumChars; ++CharIndex )
 		{
-			const FCharacterEntry& Entry = CharacterList.GetCharacter(CurrentChar, InPayload.FontInfo.FontFallback);
+			const TCHAR CurrentChar = InPayload.ImmutableText[ CharIndex ];
+
+			const bool IsNewline = (CurrentChar == '\n');
 
-			if( Entry.Valid && (FontAtlasTexture == nullptr || Entry.TextureIndex != FontTextureIndex) )
+			if (IsNewline)
 			{
-				// Font has a new texture for this glyph. Refresh the batch we use and the index we are currently using
-				FontTextureIndex = Entry.TextureIndex;
+				// Move down: we are drawing the next line.
+				PosY += MaxHeight;
+				// Carriage return 
+				LineX = PosX;
 
-				FontAtlasTexture = FontCache.GetSlateTextureResource( FontTextureIndex );
-				check(FontAtlasTexture);
+				++NumLines;
 
-				FontShaderResource = ResourceManager.GetFontShaderResource( FontTextureIndex, FontAtlasTexture, InPayload.FontInfo.FontMaterial );
-				check(FontShaderResource);
+			}
+			else
+			{
+				const FCharacterEntry& Entry = CharacterList.GetCharacter(CurrentChar, InPayload.FontInfo.FontFallback);
 
-				ElementBatch = &FindBatchForElement( Layer, FShaderParams(), FontShaderResource, ESlateDrawPrimitive::TriangleList, ESlateShader::Font, InDrawEffects, ESlateBatchDrawFlag::None, DrawElement.GetScissorRect() );
+				if( Entry.Valid && (FontAtlasTexture == nullptr || Entry.TextureIndex != FontTextureIndex) )
+				{
+					// Font has a new texture for this glyph. Refresh the batch we use and the index we are currently using
+					FontTextureIndex = Entry.TextureIndex;
 
-				BatchVertices = &BatchData->GetBatchVertexList(*ElementBatch);
-				BatchIndices = &BatchData->GetBatchIndexList(*ElementBatch);
+					FontAtlasTexture = FontCache.GetSlateTextureResource( FontTextureIndex );
+					check(FontAtlasTexture);
 
-				VertexOffset = BatchVertices->Num();
-				IndexOffset = BatchIndices->Num();
-				
-				InvTextureSizeX = 1.0f/FontAtlasTexture->GetWidth();
-				InvTextureSizeY = 1.0f/FontAtlasTexture->GetHeight();
-			}
+					FontShaderResource = ResourceManager.GetFontShaderResource( FontTextureIndex, FontAtlasTexture, InPayload.FontInfo.FontMaterial );
+					check(FontShaderResource);
 
-			const bool bIsWhitespace = !Entry.Valid || FText::IsWhitespace(CurrentChar);
+					ElementBatch = &FindBatchForElement(InLayer, FShaderParams(), FontShaderResource, ESlateDrawPrimitive::TriangleList, ESlateShader::Font, InDrawEffects, ESlateBatchDrawFlag::None, DrawElement.GetScissorRect() );
 
-			if( !bIsWhitespace && PreviousCharEntry.Valid )
-			{
-				Kerning = CharacterList.GetKerning( PreviousCharEntry, Entry );
-			}
-			else
-			{
-				Kerning = 0;
-			}
+					BatchVertices = &BatchData->GetBatchVertexList(*ElementBatch);
+					BatchIndices = &BatchData->GetBatchIndexList(*ElementBatch);
 
-			LineX += Kerning;
-			PreviousCharEntry = Entry;
+					VertexOffset = BatchVertices->Num();
+					IndexOffset = BatchIndices->Num();
+				
+					InvTextureSizeX = 1.0f/FontAtlasTexture->GetWidth();
+					InvTextureSizeY = 1.0f/FontAtlasTexture->GetHeight();
+				}
 
-			if( !bIsWhitespace )
-			{
-				const float X = LineX + Entry.HorizontalOffset;
-				// Note PosX,PosY is the upper left corner of the bounding box representing the string.  This computes the Y position of the baseline where text will sit
+				const bool bIsWhitespace = !Entry.Valid || FText::IsWhitespace(CurrentChar);
 
-				const float Y = PosY - Entry.VerticalOffset+MaxHeight+Entry.GlobalDescender;
-				const float U = Entry.StartU * InvTextureSizeX;
-				const float V = Entry.StartV * InvTextureSizeY;
-				const float SizeX = Entry.USize;
-				const float SizeY = Entry.VSize;
-				const float SizeU = Entry.USize * InvTextureSizeX;
-				const float SizeV = Entry.VSize * InvTextureSizeY;
+				if( !bIsWhitespace && PreviousCharEntry.Valid )
+				{
+					Kerning = CharacterList.GetKerning( PreviousCharEntry, Entry );
+				}
+				else
+				{
+					Kerning = 0;
+				}
 
-				FSlateRect CharRect( X, Y, X+SizeX, Y+SizeY );
+				LineX += Kerning;
+				PreviousCharEntry = Entry;
 
-				if( FSlateRect::DoRectanglesIntersect( LocalClipRect, CharRect ) )
+				if( !bIsWhitespace )
 				{
-					FSlateVertexArray& BatchVerticesRef = *BatchVertices;
-					FSlateIndexArray& BatchIndicesRef = *BatchIndices;
+					const float X = LineX + Entry.HorizontalOffset+InOutlineHorizontalOffset;
+					// Note PosX,PosY is the upper left corner of the bounding box representing the string.  This computes the Y position of the baseline where text will sit
 
-					FVector2D UpperLeft( X, Y );
-					FVector2D UpperRight( X+SizeX, Y );
-					FVector2D LowerLeft( X, Y+SizeY );
-					FVector2D LowerRight( X+SizeX, Y+SizeY );
+					const float Y = PosY - Entry.VerticalOffset+MaxHeight+Entry.GlobalDescender;
+					const float U = Entry.StartU * InvTextureSizeX;
+					const float V = Entry.StartV * InvTextureSizeY;
+					const float SizeX = Entry.USize;
+					const float SizeY = Entry.VSize;
+					const float SizeU = Entry.USize * InvTextureSizeX;
+					const float SizeV = Entry.VSize * InvTextureSizeY;
 
-					// Add four vertices for this quad
-					BatchVerticesRef.AddUninitialized( 4 );
-					// Add six indices for this quad
-					BatchIndicesRef.AddUninitialized( 6 );
+					FSlateRect CharRect( X, Y, X+SizeX, Y+SizeY );
 
-					// The start index of these vertices in the index buffer
-					uint32 IndexStart = VertexOffset;
-
-					float Ut = 0.0f, Vt = 0.0f, UtMax = 0.0f, VtMax = 0.0f;
-					if( bIsFontMaterial )
+					if( FSlateRect::DoRectanglesIntersect( LocalClipRect, CharRect ) )
 					{
-						float DistAlpha = (float)CharIndex/NumChars;
-						float DistAlphaNext = (float)(CharIndex+1)/NumChars;
-
-						// This creates a set of UV's that goes from 0-1, from left to right of the string in U and 0-1 baseline to baseline top to bottom in V
-						Ut = FMath::Lerp(0.0f, 1.0f, DistAlpha);
-						Vt = FMath::Lerp(0.0f, 1.0f, UpperLeft.Y/(MaxHeight*NumLines));
-
-						UtMax = FMath::Lerp(0.0f, 1.0f, DistAlphaNext);
-						VtMax = FMath::Lerp(0.0f, 1.0f, LowerLeft.Y/(MaxHeight*NumLines));
+						FSlateVertexArray& BatchVerticesRef = *BatchVertices;
+						FSlateIndexArray& BatchIndicesRef = *BatchIndices;
+
+						FVector2D UpperLeft( X, Y );
+						FVector2D UpperRight( X+SizeX, Y );
+						FVector2D LowerLeft( X, Y+SizeY );
+						FVector2D LowerRight( X+SizeX, Y+SizeY );
+
+						// Add four vertices for this quad
+						BatchVerticesRef.AddUninitialized( 4 );
+						// Add six indices for this quad
+						BatchIndicesRef.AddUninitialized( 6 );
+
+						// The start index of these vertices in the index buffer
+						uint32 IndexStart = VertexOffset;
+
+						float Ut = 0.0f, Vt = 0.0f, UtMax = 0.0f, VtMax = 0.0f;
+						if( bIsFontMaterial )
+						{
+							float DistAlpha = (float)CharIndex/NumChars;
+							float DistAlphaNext = (float)(CharIndex+1)/NumChars;
+
+							// This creates a set of UV's that goes from 0-1, from left to right of the string in U and 0-1 baseline to baseline top to bottom in V
+							Ut = FMath::Lerp(0.0f, 1.0f, DistAlpha);
+							Vt = FMath::Lerp(0.0f, 1.0f, UpperLeft.Y/(MaxHeight*NumLines));
+
+							UtMax = FMath::Lerp(0.0f, 1.0f, DistAlphaNext);
+							VtMax = FMath::Lerp(0.0f, 1.0f, LowerLeft.Y/(MaxHeight*NumLines));
+						}
+
+						// Add four vertices to the list of verts to be added to the vertex buffer
+						BatchVerticesRef[ VertexOffset++ ] = FSlateVertex( RenderTransform, UpperLeft,								FVector4(U,V,				Ut,Vt),			FVector2D(0.0f,0.0f), InTint, RenderClipRect );
+						BatchVerticesRef[ VertexOffset++ ] = FSlateVertex( RenderTransform, FVector2D(LowerRight.X,UpperLeft.Y),	FVector4(U+SizeU, V,		UtMax,Vt),		FVector2D(1.0f,0.0f), InTint, RenderClipRect );
+						BatchVerticesRef[ VertexOffset++ ] = FSlateVertex( RenderTransform, FVector2D(UpperLeft.X,LowerRight.Y),	FVector4(U, V+SizeV,		Ut,VtMax),		FVector2D(0.0f,1.0f), InTint, RenderClipRect );
+						BatchVerticesRef[ VertexOffset++ ] = FSlateVertex( RenderTransform, LowerRight,								FVector4(U+SizeU, V+SizeV,	UtMax,VtMax),	FVector2D(1.0f,1.0f), InTint, RenderClipRect );
+
+						BatchIndicesRef[IndexOffset++] = IndexStart + 0;
+						BatchIndicesRef[IndexOffset++] = IndexStart + 1;
+						BatchIndicesRef[IndexOffset++] = IndexStart + 2;
+						BatchIndicesRef[IndexOffset++] = IndexStart + 1;
+						BatchIndicesRef[IndexOffset++] = IndexStart + 3;
+						BatchIndicesRef[IndexOffset++] = IndexStart + 2;
 					}
-
-					// Add four vertices to the list of verts to be added to the vertex buffer
-					BatchVerticesRef[ VertexOffset++ ] = FSlateVertex( RenderTransform, UpperLeft,								FVector4(U,V,				Ut,Vt),			FVector2D(0.0f,0.0f), Tint, RenderClipRect );
-					BatchVerticesRef[ VertexOffset++ ] = FSlateVertex( RenderTransform, FVector2D(LowerRight.X,UpperLeft.Y),	FVector4(U+SizeU, V,		UtMax,Vt),		FVector2D(1.0f,0.0f), Tint, RenderClipRect );
-					BatchVerticesRef[ VertexOffset++ ] = FSlateVertex( RenderTransform, FVector2D(UpperLeft.X,LowerRight.Y),	FVector4(U, V+SizeV,		Ut,VtMax),		FVector2D(0.0f,1.0f), Tint, RenderClipRect );
-					BatchVerticesRef[ VertexOffset++ ] = FSlateVertex( RenderTransform, LowerRight,								FVector4(U+SizeU, V+SizeV,	UtMax,VtMax),	FVector2D(1.0f,1.0f), Tint, RenderClipRect );
-
-					BatchIndicesRef[IndexOffset++] = IndexStart + 0;
-					BatchIndicesRef[IndexOffset++] = IndexStart + 1;
-					BatchIndicesRef[IndexOffset++] = IndexStart + 2;
-					BatchIndicesRef[IndexOffset++] = IndexStart + 1;
-					BatchIndicesRef[IndexOffset++] = IndexStart + 3;
-					BatchIndicesRef[IndexOffset++] = IndexStart + 2;
 				}
-			}
 
-			LineX += Entry.XAdvance;
+				LineX += Entry.XAdvance;
+			}
 		}
+	};
+
+	if (bOutlineFont)
+	{
+		// Build geometry for the outline
+		BuildFontGeometry(OutlineSettings, PackVertexColor(OutlineSettings.OutlineColor), OutlineFontMaterial, Layer, 0);
+
+		//The fill area was measured without an outline so it must be shifted by the scaled outline size
+		float HorizontalOffset = FMath::RoundToFloat(OutlineSize * FontScale);
+
+		// Build geometry for the base font which is always rendered on top of the outline
+		BuildFontGeometry(FFontOutlineSettings::NoOutline, BaseTint, BaseFontMaterial, Layer + 1, HorizontalOffset);
+	}
+	else
+	{
+		// No outline, draw normally
+		BuildFontGeometry(FFontOutlineSettings::NoOutline, BaseTint, BaseFontMaterial, Layer, 0);
 	}
 }
 
@@ -821,14 +850,18 @@ void FSlateElementBatcher::AddShapedTextElement( const FSlateDrawElement& DrawEl
 	const FSlateDataPayload& InPayload = DrawElement.GetDataPayload();
 	check(InPayload.ShapedGlyphSequence.IsValid());
 
-	FColor Tint = PackVertexColor(InPayload.Tint);
-	if (Tint.A == 0)
+	const FFontOutlineSettings& OutlineSettings = InPayload.ShapedGlyphSequence->GetFontOutlineSettings();
+
+	const TArray<FShapedGlyphEntry>& GlyphsToRender = InPayload.ShapedGlyphSequence->GetGlyphsToRender();
+	if (GlyphsToRender.Num() == 0)
 	{
 		return;
 	}
 
-	const TArray<FShapedGlyphEntry>& GlyphsToRender = InPayload.ShapedGlyphSequence->GetGlyphsToRender();
-	if (GlyphsToRender.Num() == 0)
+	FColor BaseTint = PackVertexColor(InPayload.Tint);
+
+	// Don't do anything if there the font would be completely transparent 
+	if((BaseTint.A == 0 && OutlineSettings.OutlineSize == 0) || (BaseTint.A == 0 && OutlineSettings.OutlineColor.A == 0))
 	{
 		return;
 	}
@@ -841,10 +874,6 @@ void FSlateElementBatcher::AddShapedTextElement( const FSlateDrawElement& DrawEl
 
 	NumDrawnTextsStat++;
 
-	//const FVector2D& InPosition = DrawElement.GetPosition();
-	//const FVector2D& Size = DrawElement.GetSize();
-	//const FVector2D& LocalSize = DrawElement.GetLocalSize();
-	//float Scale = DrawElement.GetScale();
 	const FSlateRect& InClippingRect = DrawElement.GetClippingRect();
 	ESlateDrawEffect::Type InDrawEffects = DrawElement.GetDrawEffects();
 	uint32 Layer = DrawElement.GetLayer();
@@ -863,128 +892,155 @@ void FSlateElementBatcher::AddShapedTextElement( const FSlateDrawElement& DrawEl
 	// Used to clip individual characters as we generate them.
 	FSlateRect LocalClipRect = TransformRect(InverseLayoutTransform, InClippingRect);
 
-	uint32 FontTextureIndex = 0;
-	FSlateShaderResource* FontAtlasTexture = nullptr;
-	FSlateShaderResource* FontShaderResource = nullptr;
+	const UObject* BaseFontMaterial = InPayload.ShapedGlyphSequence->GetFontMaterial();
+	const UObject* OutlineFontMaterial = OutlineSettings.OutlineMaterial;
 
-	FSlateElementBatch* ElementBatch = nullptr;
-	FSlateVertexArray* BatchVertices = nullptr;
-	FSlateIndexArray* BatchIndices = nullptr;
+	bool bOutlineFont = OutlineSettings.OutlineSize > 0.0f;
 
-	uint32 VertexOffset = 0;
-	uint32 IndexOffset = 0;
+	const float OutlineSize = OutlineSettings.OutlineSize;
 
-	float InvTextureSizeX = 0;
-	float InvTextureSizeY = 0;
+	auto BuildFontGeometry = [&](const FFontOutlineSettings& InOutlineSettings, const FColor& InTint, const UObject* FontMaterial, int32 InLayer, int32 InHorizontalOffset)
+	{
+		FVector2D TopLeft(0, 0);
 
-	FVector2D TopLeft(0,0);
+		const float PosX = TopLeft.X+InHorizontalOffset;
+		float PosY = TopLeft.Y;
 
-	const float PosX = TopLeft.X;
-	float PosY = TopLeft.Y;
+		float LineX = PosX;
+		float LineY = PosY;
 
-	float LineX = PosX;
-	float LineY = PosY;
-	
-	const UObject* FontMaterial = InPayload.ShapedGlyphSequence->GetFontMaterial();
-	const bool bIsFontMaterial = FontMaterial != nullptr;
+		int32 FontTextureIndex = -1;
+		FSlateShaderResource* FontAtlasTexture = nullptr;
+		FSlateShaderResource* FontShaderResource = nullptr;
 
-	const int32 NumGlyphs = GlyphsToRender.Num();
-	for (int32 GlyphIndex = 0; GlyphIndex < NumGlyphs; ++GlyphIndex)
-	{
-		const FShapedGlyphEntry& GlyphToRender = GlyphsToRender[GlyphIndex];
+		FSlateElementBatch* ElementBatch = nullptr;
+		FSlateVertexArray* BatchVertices = nullptr;
+		FSlateIndexArray* BatchIndices = nullptr;
 
-		if (GlyphToRender.bIsVisible)
-		{
-			const FShapedGlyphFontAtlasData GlyphAtlasData = FontCache.GetShapedGlyphFontAtlasData(GlyphToRender);
+		uint32 VertexOffset = 0;
+		uint32 IndexOffset = 0;
 
-			if (GlyphAtlasData.Valid)
-			{
-				if (FontAtlasTexture == nullptr || GlyphAtlasData.TextureIndex != FontTextureIndex)
-				{
-					// Font has a new texture for this glyph. Refresh the batch we use and the index we are currently using
-					FontTextureIndex = GlyphAtlasData.TextureIndex;
+		float InvTextureSizeX = 0;
+		float InvTextureSizeY = 0;
 
-					FontAtlasTexture = FontCache.GetSlateTextureResource(FontTextureIndex);
-					check(FontAtlasTexture);
+		const bool bIsFontMaterial = FontMaterial != nullptr;
 
-					FontShaderResource = ResourceManager.GetFontShaderResource(FontTextureIndex, FontAtlasTexture, FontMaterial);
-					check(FontShaderResource);
+		const int32 NumGlyphs = GlyphsToRender.Num();
+		for (int32 GlyphIndex = 0; GlyphIndex < NumGlyphs; ++GlyphIndex)
+		{
+			const FShapedGlyphEntry& GlyphToRender = GlyphsToRender[GlyphIndex];
 
-					ElementBatch = &FindBatchForElement(Layer, FShaderParams(), FontShaderResource, ESlateDrawPrimitive::TriangleList, ESlateShader::Font, InDrawEffects, ESlateBatchDrawFlag::None, DrawElement.GetScissorRect());
+			if (GlyphToRender.bIsVisible)
+			{
+				const FShapedGlyphFontAtlasData GlyphAtlasData = FontCache.GetShapedGlyphFontAtlasData(GlyphToRender, InOutlineSettings);
+				 
+				if (GlyphAtlasData.Valid)
+				{
+					if (FontAtlasTexture == nullptr || GlyphAtlasData.TextureIndex != FontTextureIndex)
+					{
+						// Font has a new texture for this glyph. Refresh the batch we use and the index we are currently using
+						FontTextureIndex = GlyphAtlasData.TextureIndex;
 
-					BatchVertices = &BatchData->GetBatchVertexList(*ElementBatch);
-					BatchIndices = &BatchData->GetBatchIndexList(*ElementBatch);
+						FontAtlasTexture = FontCache.GetSlateTextureResource(FontTextureIndex);
+						check(FontAtlasTexture);
 
-					VertexOffset = BatchVertices->Num();
-					IndexOffset = BatchIndices->Num();
-				
-					InvTextureSizeX = 1.0f/FontAtlasTexture->GetWidth();
-					InvTextureSizeY = 1.0f/FontAtlasTexture->GetHeight();
-				}
+						FontShaderResource = ResourceManager.GetFontShaderResource(FontTextureIndex, FontAtlasTexture, FontMaterial);
+						check(FontShaderResource);
 
-				const float X = LineX + GlyphAtlasData.HorizontalOffset + GlyphToRender.XOffset;
-				// Note PosX,PosY is the upper left corner of the bounding box representing the string.  This computes the Y position of the baseline where text will sit
+						ElementBatch = &FindBatchForElement(InLayer, FShaderParams(), FontShaderResource, ESlateDrawPrimitive::TriangleList, ESlateShader::Font, InDrawEffects, ESlateBatchDrawFlag::None, DrawElement.GetScissorRect());
 
-				const float Y = LineY - GlyphAtlasData.VerticalOffset + GlyphToRender.YOffset + MaxHeight + TextBaseline;
-				const float U = GlyphAtlasData.StartU * InvTextureSizeX;
-				const float V = GlyphAtlasData.StartV * InvTextureSizeY;
-				const float SizeX = GlyphAtlasData.USize;
-				const float SizeY = GlyphAtlasData.VSize;
-				const float SizeU = GlyphAtlasData.USize * InvTextureSizeX;
-				const float SizeV = GlyphAtlasData.VSize * InvTextureSizeY;
+						BatchVertices = &BatchData->GetBatchVertexList(*ElementBatch);
+						BatchIndices = &BatchData->GetBatchIndexList(*ElementBatch);
 
-				const FSlateRect CharRect(X, Y, X + SizeX, Y + SizeY);
-				if (FSlateRect::DoRectanglesIntersect(LocalClipRect, CharRect))
-				{
-					FSlateVertexArray& BatchVerticesRef = *BatchVertices;
-					FSlateIndexArray& BatchIndicesRef = *BatchIndices;
+						VertexOffset = BatchVertices->Num();
+						IndexOffset = BatchIndices->Num();
 
-					const FVector2D UpperLeft = CharRect.GetTopLeft();
-					const FVector2D UpperRight = CharRect.GetTopRight();
-					const FVector2D LowerLeft = CharRect.GetBottomLeft();
-					const FVector2D LowerRight = CharRect.GetBottomRight();
+						InvTextureSizeX = 1.0f / FontAtlasTexture->GetWidth();
+						InvTextureSizeY = 1.0f / FontAtlasTexture->GetHeight();
+					}
 
-					// Add four vertices for this quad
-					BatchVerticesRef.AddUninitialized(4);
-					// Add six indices for this quad
-					BatchIndicesRef.AddUninitialized(6);
+					const float X = LineX + GlyphAtlasData.HorizontalOffset + GlyphToRender.XOffset;
+					// Note PosX,PosY is the upper left corner of the bounding box representing the string.  This computes the Y position of the baseline where text will sit
 
-					// The start index of these vertices in the index buffer
-					uint32 IndexStart = VertexOffset;
+					const float Y = LineY - GlyphAtlasData.VerticalOffset + GlyphToRender.YOffset + MaxHeight + TextBaseline;
+					const float U = GlyphAtlasData.StartU * InvTextureSizeX;
+					const float V = GlyphAtlasData.StartV * InvTextureSizeY;
+					const float SizeX = GlyphAtlasData.USize;
+					const float SizeY = GlyphAtlasData.VSize;
+					const float SizeU = GlyphAtlasData.USize * InvTextureSizeX;
+					const float SizeV = GlyphAtlasData.VSize * InvTextureSizeY;
 
-					float Ut = 0.0f, Vt = 0.0f, UtMax = 0.0f, VtMax = 0.0f;
-					if (bIsFontMaterial)
+					const FSlateRect CharRect(X, Y, X + SizeX, Y + SizeY);
+					if (FSlateRect::DoRectanglesIntersect(LocalClipRect, CharRect))
 					{
-						float DistAlpha = (float)GlyphIndex/NumGlyphs;
-						float DistAlphaNext = (float)(GlyphIndex+1)/NumGlyphs;
-
-						// This creates a set of UV's that goes from 0-1, from left to right of the string in U and 0-1 baseline to baseline top to bottom in V
-						Ut = FMath::Lerp(0.0f, 1.0f, DistAlpha);
-						Vt = FMath::Lerp(0.0f, 1.0f, UpperLeft.Y/MaxHeight);
-
-						UtMax = FMath::Lerp(0.0f, 1.0f, DistAlphaNext);
-						VtMax = FMath::Lerp(0.0f, 1.0f, LowerLeft.Y/MaxHeight);
+						FSlateVertexArray& BatchVerticesRef = *BatchVertices;
+						FSlateIndexArray& BatchIndicesRef = *BatchIndices;
+
+						const FVector2D UpperLeft = CharRect.GetTopLeft();
+						const FVector2D UpperRight = CharRect.GetTopRight();
+						const FVector2D LowerLeft = CharRect.GetBottomLeft();
+						const FVector2D LowerRight = CharRect.GetBottomRight();
+
+						// Add four vertices for this quad
+						BatchVerticesRef.AddUninitialized(4);
+						// Add six indices for this quad
+						BatchIndicesRef.AddUninitialized(6);
+
+						// The start index of these vertices in the index buffer
+						uint32 IndexStart = VertexOffset;
+
+						float Ut = 0.0f, Vt = 0.0f, UtMax = 0.0f, VtMax = 0.0f;
+						if (bIsFontMaterial)
+						{
+							float DistAlpha = (float)GlyphIndex / NumGlyphs;
+							float DistAlphaNext = (float)(GlyphIndex + 1) / NumGlyphs;
+
+							// This creates a set of UV's that goes from 0-1, from left to right of the string in U and 0-1 baseline to baseline top to bottom in V
+							Ut = FMath::Lerp(0.0f, 1.0f, DistAlpha);
+							Vt = FMath::Lerp(0.0f, 1.0f, UpperLeft.Y / MaxHeight);
+
+							UtMax = FMath::Lerp(0.0f, 1.0f, DistAlphaNext);
+							VtMax = FMath::Lerp(0.0f, 1.0f, LowerLeft.Y / MaxHeight);
+						}
+
+						// Add four vertices to the list of verts to be added to the vertex buffer
+						BatchVerticesRef[VertexOffset++] = FSlateVertex(RenderTransform, UpperLeft,								FVector4(U, V, Ut, Vt),							FVector2D(0.0f, 0.0f), InTint, RenderClipRect);
+						BatchVerticesRef[VertexOffset++] = FSlateVertex(RenderTransform, FVector2D(LowerRight.X, UpperLeft.Y),	FVector4(U + SizeU, V, UtMax, Vt),				FVector2D(1.0f, 0.0f), InTint, RenderClipRect);
+						BatchVerticesRef[VertexOffset++] = FSlateVertex(RenderTransform, FVector2D(UpperLeft.X, LowerRight.Y),	FVector4(U, V + SizeV, Ut, VtMax),				FVector2D(0.0f, 1.0f), InTint, RenderClipRect);
+						BatchVerticesRef[VertexOffset++] = FSlateVertex(RenderTransform, LowerRight,							FVector4(U + SizeU, V + SizeV, UtMax, VtMax),	FVector2D(1.0f, 1.0f), InTint, RenderClipRect);
+
+						BatchIndicesRef[IndexOffset++] = IndexStart + 0;
+						BatchIndicesRef[IndexOffset++] = IndexStart + 1;
+						BatchIndicesRef[IndexOffset++] = IndexStart + 2;
+						BatchIndicesRef[IndexOffset++] = IndexStart + 1;
+						BatchIndicesRef[IndexOffset++] = IndexStart + 3;
+						BatchIndicesRef[IndexOffset++] = IndexStart + 2;
 					}
-
-					// Add four vertices to the list of verts to be added to the vertex buffer
-					BatchVerticesRef[VertexOffset++] = FSlateVertex(RenderTransform, UpperLeft,								FVector4(U,V,				Ut,Vt),			FVector2D(0.0f,0.0f), Tint, RenderClipRect);
-					BatchVerticesRef[VertexOffset++] = FSlateVertex(RenderTransform, FVector2D(LowerRight.X,UpperLeft.Y),	FVector4(U+SizeU, V,		UtMax,Vt),		FVector2D(1.0f,0.0f), Tint, RenderClipRect);
-					BatchVerticesRef[VertexOffset++] = FSlateVertex(RenderTransform, FVector2D(UpperLeft.X,LowerRight.Y),	FVector4(U, V+SizeV,		Ut,VtMax),		FVector2D(0.0f,1.0f), Tint, RenderClipRect);
-					BatchVerticesRef[VertexOffset++] = FSlateVertex(RenderTransform, LowerRight,							FVector4(U+SizeU, V+SizeV,	UtMax,VtMax),	FVector2D(1.0f,1.0f), Tint, RenderClipRect);
-
-					BatchIndicesRef[IndexOffset++] = IndexStart + 0;
-					BatchIndicesRef[IndexOffset++] = IndexStart + 1;
-					BatchIndicesRef[IndexOffset++] = IndexStart + 2;
-					BatchIndicesRef[IndexOffset++] = IndexStart + 1;
-					BatchIndicesRef[IndexOffset++] = IndexStart + 3;
-					BatchIndicesRef[IndexOffset++] = IndexStart + 2;
 				}
 			}
+
+			LineX += GlyphToRender.XAdvance;
+			LineY += GlyphToRender.YAdvance;
 		}
+	};
 
-		LineX += GlyphToRender.XAdvance;
-		LineY += GlyphToRender.YAdvance;
+	if (bOutlineFont)
+	{
+		// Build geometry for the outline
+		BuildFontGeometry(OutlineSettings, PackVertexColor(OutlineSettings.OutlineColor), OutlineFontMaterial, Layer, 0);
+		
+		//The fill area was measured without an outline so it must be shifted by the scaled outline size
+		float HorizontalOffset = FMath::RoundToFloat(OutlineSize * FontScale);
+
+		// Build geometry for the base font which is always rendered on top of the outline 
+		BuildFontGeometry(FFontOutlineSettings::NoOutline, BaseTint, BaseFontMaterial, Layer+1, HorizontalOffset);
 	}
+	else
+	{
+		// No outline
+		BuildFontGeometry(FFontOutlineSettings::NoOutline, BaseTint, BaseFontMaterial, Layer, 0);
+	}
+
 }
 
 void FSlateElementBatcher::AddGradientElement( const FSlateDrawElement& DrawElement )
diff --git a/Engine/Source/Runtime/SlateCore/Public/Fonts/FontCache.h b/Engine/Source/Runtime/SlateCore/Public/Fonts/FontCache.h
--- a/Engine/Source/Runtime/SlateCore/Public/Fonts/FontCache.h
+++ b/Engine/Source/Runtime/SlateCore/Public/Fonts/FontCache.h
@@ -21,6 +21,19 @@ class FShapedGlyphFaceData;
 struct FCharacterRenderData;
 struct FSlateFontKey;
 
+enum class EFontCacheAtlasDataType : uint8
+{
+	/** Data was cached for a regular non-outline font */
+	Regular = 0,
+
+	/** Data was cached for a outline (stroked) font */
+	Outline,
+
+	/** Must be last */
+	Num,
+};
+
+
 /** 
  * Methods that can be used to shape text.
  * @note If you change this enum, make sure and update CVarDefaultTextShapingMethod and GetDefaultTextShapingMethod.
@@ -149,21 +162,23 @@ private:
 	/** 
 	 * Pointer to the cached atlas data for this glyph entry.
 	 * This is cached on the glyph by FSlateFontCache::GetShapedGlyphFontAtlasData to avoid repeated map look ups.
-	 * Index 0 is the cached value for the game thread font cache. Index 1 is the cached value for the render thread font cache.
+	 * First index is to determine if this is a cached outline glyph or a regular glyph
+	 * Second index is the index of the thread dependent font cache. Index 0 is the cached value for the game thread font cache. Index 1 is the cached value for the render thread font cache.
 	 */
-	mutable TWeakPtr<FShapedGlyphFontAtlasData> CachedAtlasData[2];
+	mutable TWeakPtr<FShapedGlyphFontAtlasData> CachedAtlasData[(uint8)EFontCacheAtlasDataType::Num][2];
 };
 
 /** Minimal FShapedGlyphEntry key information used for map lookups */
 struct FShapedGlyphEntryKey
 {
 public:
-	FShapedGlyphEntryKey(const TSharedPtr<FShapedGlyphFaceData>& InFontFaceData, uint32 InGlyphIndex);
+	FShapedGlyphEntryKey(const TSharedPtr<FShapedGlyphFaceData>& InFontFaceData, uint32 InGlyphIndex, const FFontOutlineSettings& InOutlineSettings);
 
 	FORCEINLINE bool operator==(const FShapedGlyphEntryKey& Other) const
 	{
 		return FontFace == Other.FontFace 
 			&& FontSize == Other.FontSize
+			&& OutlineSize == Other.OutlineSize
 			&& FontScale == Other.FontScale
 			&& GlyphIndex == Other.GlyphIndex;
 	}
@@ -183,6 +198,8 @@ private:
 	TWeakPtr<FFreeTypeFace> FontFace;
 	/** Provides the point size used to render the font */
 	int32 FontSize;
+	/** The size in pixels of the outline to render for the font */
+	float OutlineSize;
 	/** Provides the final scale used to render to the font */
 	float FontScale;
 	/** The index of this glyph in the FreeType face */
@@ -207,7 +224,7 @@ public:
 		int32 TextLen;
 	};
 
-	FShapedGlyphSequence(TArray<FShapedGlyphEntry> InGlyphsToRender, const int16 InTextBaseline, const uint16 InMaxTextHeight, const UObject* InFontMaterial, const FSourceTextRange& InSourceTextRange);
+	FShapedGlyphSequence(TArray<FShapedGlyphEntry> InGlyphsToRender, const int16 InTextBaseline, const uint16 InMaxTextHeight, const UObject* InFontMaterial, const struct FFontOutlineSettings& InOutlineSettings, const FSourceTextRange& InSourceTextRange);
 	~FShapedGlyphSequence();
 
 	/** Get the amount of memory allocated to this sequence */
@@ -237,6 +254,12 @@ public:
 		return FontMaterial;
 	}
 
+	/** Get the font outline settings to use when rendering these glyphs */
+	const FFontOutlineSettings& GetFontOutlineSettings() const
+	{
+		return OutlineSettings;
+	}
+
 	/** Check to see whether this glyph sequence is dirty (ie, contains glyphs with invalid font pointers) */
 	bool IsDirty() const;
 
@@ -415,6 +438,8 @@ private:
 	uint16 MaxTextHeight;
 	/** The material to use when rendering these glyphs */
 	const UObject* FontMaterial;
+	/** Outline settings to use when rendering these glyphs */
+	FFontOutlineSettings OutlineSettings;
 	/** The cached width of the entire sequence */
 	int32 SequenceWidth;
 	/** The set of fonts being used by the glyphs within this sequence */
@@ -679,17 +704,16 @@ public:
 	/** 
 	 * Gets information for how to draw all characters in the specified string. Caches characters as they are found
 	 * 
-	 * @param Text				The string to get character information for
 	 * @param InFontInfo		Information about the font that the string is drawn with
 	 * @param FontScale			The scale to apply to the font
 	 * @param OutCharacterEntries	Populated array of character entries. Indices of characters in Text match indices in this array
 	 */
-	class FCharacterList& GetCharacterList( const FSlateFontInfo &InFontInfo, float FontScale );
+	class FCharacterList& GetCharacterList( const FSlateFontInfo &InFontInfo, float FontScale, const FFontOutlineSettings& InOutlineSettings = FFontOutlineSettings::NoOutline);
 
 	/**
 	 * Get the atlas information for the given shaped glyph. This information will be cached if required 
 	 */
-	FShapedGlyphFontAtlasData GetShapedGlyphFontAtlasData( const FShapedGlyphEntry& InShapedGlyph );
+	FShapedGlyphFontAtlasData GetShapedGlyphFontAtlasData( const FShapedGlyphEntry& InShapedGlyph, const FFontOutlineSettings& InOutlineSettings);
 
 	/** 
 	 * Add a new entries into a cache atlas
@@ -701,7 +725,7 @@ public:
 	 */
 	bool AddNewEntry( TCHAR Character, const FSlateFontKey& InKey, FCharacterEntry& OutCharacterEntry );
 
-	bool AddNewEntry( const FShapedGlyphEntry& InShapedGlyph, FShapedGlyphFontAtlasData& OutAtlasData );
+	bool AddNewEntry( const FShapedGlyphEntry& InShapedGlyph, const FFontOutlineSettings& InOutlineSettings, FShapedGlyphFontAtlasData& OutAtlasData );
 
 	bool AddNewEntry( const FCharacterRenderData InRenderData, uint8& OutTextureIndex, uint16& OutGlyphX, uint16& OutGlyphY, uint16& OutGlyphWidth, uint16& OutGlyphHeight );
 
diff --git a/Engine/Source/Runtime/SlateCore/Public/Fonts/FontTypes.h b/Engine/Source/Runtime/SlateCore/Public/Fonts/FontTypes.h
--- a/Engine/Source/Runtime/SlateCore/Public/Fonts/FontTypes.h
+++ b/Engine/Source/Runtime/SlateCore/Public/Fonts/FontTypes.h
@@ -4,12 +4,14 @@
 struct SLATECORE_API FSlateFontKey
 {
 public:
-	FSlateFontKey( const FSlateFontInfo& InInfo, const float InScale )
+	FSlateFontKey( const FSlateFontInfo& InInfo, const FFontOutlineSettings& InFontOutlineSettings, const float InScale )
 		: FontInfo( InInfo )
+		, OutlineSettings(InFontOutlineSettings)
 		, Scale( InScale )
 		, KeyHash( 0 )
 	{
 		KeyHash = HashCombine(KeyHash, GetTypeHash(FontInfo));
+		KeyHash = HashCombine(KeyHash, GetTypeHash(OutlineSettings));
 		KeyHash = HashCombine(KeyHash, GetTypeHash(Scale));
 	}
 
@@ -23,9 +25,14 @@ public:
 		return Scale;
 	}
 
+	FORCEINLINE const FFontOutlineSettings& GetFontOutlineSettings() const
+	{
+		return OutlineSettings;
+	}
+
 	FORCEINLINE bool operator==(const FSlateFontKey& Other ) const
 	{
-		return FontInfo == Other.FontInfo && Scale == Other.Scale;
+		return FontInfo == Other.FontInfo && OutlineSettings == Other.OutlineSettings && Scale == Other.Scale;
 	}
 
 	FORCEINLINE bool operator!=(const FSlateFontKey& Other ) const
@@ -40,6 +47,7 @@ public:
 
 private:
 	FSlateFontInfo FontInfo;
+	FFontOutlineSettings OutlineSettings;
 	float Scale;
 	uint32 KeyHash;
 };
diff --git a/Engine/Source/Runtime/SlateCore/Public/Fonts/SlateFontInfo.h b/Engine/Source/Runtime/SlateCore/Public/Fonts/SlateFontInfo.h
--- a/Engine/Source/Runtime/SlateCore/Public/Fonts/SlateFontInfo.h
+++ b/Engine/Source/Runtime/SlateCore/Public/Fonts/SlateFontInfo.h
@@ -26,6 +26,61 @@ enum class EFontFallback : uint8
 	FF_Max
 };
 
+/**
+ * Settings for applying an outline to a font
+ */
+USTRUCT()
+struct SLATECORE_API FFontOutlineSettings
+{
+	GENERATED_USTRUCT_BODY()
+
+	/** Size of the outline in slate units (at 1.0 font scale this unit is a pixel)*/
+	UPROPERTY(EditAnywhere, BlueprintReadWrite, Category=OutlineSettings, meta=(ClampMin="0"))
+	int32 OutlineSize;
+
+	/** Optional material to apply to the outline */
+	UPROPERTY(EditAnywhere, BlueprintReadWrite, Category=SlateStyleRules, meta=(AllowedClasses="MaterialInterface"))
+	const UObject* OutlineMaterial;
+
+	/** The color of the outline for any character in this font */
+	UPROPERTY(EditAnywhere, BlueprintReadWrite, Category=OutlineSettings)
+	FLinearColor OutlineColor;
+
+	/** 
+	 * If checked, the outline will be completely translucent where the filled area will be.  This allows for a separate fill alpha value
+	 * The trade off when enabling this is slightly worse quality for completely opaque fills where the inner outline border meets the fill area
+	 */
+	UPROPERTY(EditAnywhere, BlueprintReadWrite, Category=OutlineSettings)
+	bool bSeparateFillAlpha;
+
+	FFontOutlineSettings()
+		: OutlineSize(0)
+		, OutlineMaterial(nullptr)
+		, OutlineColor(FColor::Black)
+		, bSeparateFillAlpha(false)
+	{}
+
+	bool operator==(const FFontOutlineSettings& Other) const
+	{
+		return OutlineSize == Other.OutlineSize
+			&& OutlineMaterial == Other.OutlineMaterial
+			&& OutlineColor == Other.OutlineColor
+			&& bSeparateFillAlpha == Other.bSeparateFillAlpha;
+	}
+
+	friend inline uint32 GetTypeHash(const FFontOutlineSettings& OutlineSettings)
+	{
+		uint32 Hash = 0;
+		Hash = HashCombine(Hash, GetTypeHash(OutlineSettings.OutlineMaterial));
+		Hash = HashCombine(Hash, GetTypeHash(OutlineSettings.OutlineSize));
+		Hash = HashCombine(Hash, GetTypeHash(OutlineSettings.OutlineColor));
+		Hash = HashCombine(Hash, GetTypeHash(OutlineSettings.bSeparateFillAlpha));
+		return Hash;
+	}
+
+	static FFontOutlineSettings NoOutline;
+};
+
 /**
  * A representation of a font in Slate.
  */
@@ -42,6 +97,10 @@ struct SLATECORE_API FSlateFontInfo
 	UPROPERTY(EditAnywhere, BlueprintReadWrite, Category=SlateStyleRules, meta=(AllowedClasses="MaterialInterface"))
 	const UObject* FontMaterial;
 
+	/** Settings for applying an outline to a font */
+	UPROPERTY(EditAnywhere, BlueprintReadWrite, Category=SlateStyleRules)
+	FFontOutlineSettings OutlineSettings;
+
 	/** The composite font data to use (valid when used with a Slate style set in C++) */
 	TSharedPtr<const FCompositeFont> CompositeFont;
 
@@ -139,7 +198,8 @@ public:
 	{
 		return FontObject == Other.FontObject
 			&& FontMaterial == Other.FontMaterial
-			&& CompositeFont == Other.CompositeFont 
+			&& OutlineSettings == Other.OutlineSettings
+			&& CompositeFont == Other.CompositeFont
 			&& TypefaceFontName == Other.TypefaceFontName
 			&& Size == Other.Size;
 	}
@@ -180,6 +240,7 @@ public:
 		uint32 Hash = 0;
 		Hash = HashCombine(Hash, GetTypeHash(FontInfo.FontObject));
 		Hash = HashCombine(Hash, GetTypeHash(FontInfo.FontMaterial));
+		Hash = HashCombine(Hash, GetTypeHash(FontInfo.OutlineSettings));
 		Hash = HashCombine(Hash, GetTypeHash(FontInfo.CompositeFont));
 		Hash = HashCombine(Hash, GetTypeHash(FontInfo.TypefaceFontName));
 		Hash = HashCombine(Hash, GetTypeHash(FontInfo.Size));
@@ -207,4 +268,4 @@ struct TStructOpsTypeTraits<FSlateFontInfo>
 	{
 		WithPostSerialize = true,
 	};
-};
+};
\ No newline at end of file
diff --git a/Engine/Source/Runtime/SlateCore/Public/Rendering/ShaderResourceManager.h b/Engine/Source/Runtime/SlateCore/Public/Rendering/ShaderResourceManager.h
--- a/Engine/Source/Runtime/SlateCore/Public/Rendering/ShaderResourceManager.h
+++ b/Engine/Source/Runtime/SlateCore/Public/Rendering/ShaderResourceManager.h
@@ -62,7 +62,7 @@ public:
 	virtual FSlateResourceHandle GetResourceHandle( const FSlateBrush& InBrush );
 
 
-	virtual FSlateShaderResource* GetFontShaderResource( uint32 FontAtlasIndex, FSlateShaderResource* FontTextureAtlas, const class UObject* FontMaterial ) { return FontTextureAtlas; }
+	virtual FSlateShaderResource* GetFontShaderResource( int32 InTextureAtlasIndex, FSlateShaderResource* FontTextureAtlas, const class UObject* FontMaterial ) { return FontTextureAtlas; }
 
 	/**
 	 * Returns the way to access the texture atlas information from this resource manager
diff --git a/Engine/Source/Runtime/SlateCore/Public/Widgets/SWindow.h b/Engine/Source/Runtime/SlateCore/Public/Widgets/SWindow.h
--- a/Engine/Source/Runtime/SlateCore/Public/Widgets/SWindow.h
+++ b/Engine/Source/Runtime/SlateCore/Public/Widgets/SWindow.h
@@ -709,10 +709,8 @@ private:
 	virtual FVector2D ComputeDesiredSize(float) const override;
 
 public:
-// @HSL_CHANGE_BEGIN - [REDACTED_EMAIL] - 5/31/2016 - need access to this to fix windowed mode
 	// For a given client size, calculate the window size required to accomodate any potential non-OS borders and tilebars
 	FVector2D GetWindowSizeFromClientSize(FVector2D InClientSize);
-// @HSL_CHANGE_END - [REDACTED_EMAIL] - 5/31/2016 - need access to this to fix windowed mode
 
 	/** @return true if this window will be focused when it is first shown */
 	inline bool IsFocusedInitially() const
diff --git a/Engine/Source/Runtime/SlateRHIRenderer/Private/SlateRHIResourceManager.cpp b/Engine/Source/Runtime/SlateRHIRenderer/Private/SlateRHIResourceManager.cpp
--- a/Engine/Source/Runtime/SlateRHIRenderer/Private/SlateRHIResourceManager.cpp
+++ b/Engine/Source/Runtime/SlateRHIRenderer/Private/SlateRHIResourceManager.cpp
@@ -46,9 +46,9 @@ TSharedPtr<FSlateAtlasedTextureResource> FDynamicResourceMap::GetAtlasedTextureR
 	return nullptr;
 }
 
-TSharedPtr<FSlateMaterialResource> FDynamicResourceMap::GetMaterialResource( const UMaterialInterface* Material ) const
+TSharedPtr<FSlateMaterialResource> FDynamicResourceMap::GetMaterialResource( const FMaterialKey& InKey ) const
 {
-	return MaterialMap.FindRef( Material );
+	return MaterialMap.FindRef( InKey );
 }
 
 
@@ -68,10 +68,10 @@ void FDynamicResourceMap::AddUTextureResource( UTexture* TextureObject, TSharedR
 	}
 }
 
-void FDynamicResourceMap::AddMaterialResource( const UMaterialInterface* Material, TSharedRef<FSlateMaterialResource> InMaterialResource )
+void FDynamicResourceMap::AddMaterialResource( const FMaterialKey& InKey, TSharedRef<FSlateMaterialResource> InMaterialResource )
 {
-	check( Material == InMaterialResource->GetMaterialObject() );
-	MaterialMap.Add(Material, InMaterialResource);
+	check(InKey.Material == InMaterialResource->GetMaterialObject() );
+	MaterialMap.Add(InKey, InMaterialResource);
 }
 
 void FDynamicResourceMap::RemoveDynamicTextureResource(FName ResourceName)
@@ -88,9 +88,9 @@ void FDynamicResourceMap::RemoveUTextureResource( UTexture* TextureObject )
 	}
 }
 
-void FDynamicResourceMap::RemoveMaterialResource( const UMaterialInterface* Material )
+void FDynamicResourceMap::RemoveMaterialResource( const FMaterialKey& InKey )
 {
-	MaterialMap.Remove(Material);
+	MaterialMap.Remove(InKey);
 }
 
 void FDynamicResourceMap::AddAtlasedTextureResource(UTexture* TextureObject, TSharedRef<FSlateAtlasedTextureResource> InResource)
@@ -172,8 +172,8 @@ void FDynamicResourceMap::RemoveExpiredMaterialResources(TArray< TSharedPtr<FSla
 	{
 		for ( FMaterialResourceMap::TIterator It(MaterialMap); It; ++It )
 		{
-			TWeakObjectPtr<UMaterialInterface>& Key = It.Key();
-			if ( !Key.IsValid() )
+			FMaterialKey& Key = It.Key();
+			if ( !Key.Material.IsValid() )
 			{
 				RemovedMaterials.Push(It.Value());
 
@@ -511,7 +511,7 @@ FSlateShaderResourceProxy* FSlateRHIResourceManager::GetShaderResource( const FS
 	}
 	else if (InBrush.GetResourceObject() && InBrush.GetResourceObject()->IsA<UMaterialInterface>())
 	{
-		FSlateMaterialResource* Resource = GetMaterialResource(InBrush.GetResourceObject(), InBrush.ImageSize, nullptr);
+		FSlateMaterialResource* Resource = GetMaterialResource(InBrush.GetResourceObject(), InBrush.ImageSize, nullptr, 0);
 		Texture = Resource->SlateProxy;
 	}
 	else if( InBrush.IsDynamicallyLoaded() || ( InBrush.HasUObject() ) )
@@ -528,7 +528,7 @@ FSlateShaderResourceProxy* FSlateRHIResourceManager::GetShaderResource( const FS
 	return Texture;
 }
 
-FSlateShaderResource* FSlateRHIResourceManager::GetFontShaderResource( uint32 FontAtlasIndex, FSlateShaderResource* FontTextureAtlas, const class UObject* FontMaterial )
+FSlateShaderResource* FSlateRHIResourceManager::GetFontShaderResource( int32 InTextureAtlasIndex, FSlateShaderResource* FontTextureAtlas, const class UObject* FontMaterial )
 {
 	if( FontMaterial == nullptr )
 	{
@@ -536,7 +536,7 @@ FSlateShaderResource* FSlateRHIResourceManager::GetFontShaderResource( uint32 Fo
 	}
 	else
 	{
-		return GetMaterialResource( FontMaterial, FVector2D::ZeroVector, FontTextureAtlas );
+		return GetMaterialResource( FontMaterial, FVector2D::ZeroVector, FontTextureAtlas, InTextureAtlasIndex );
 	}
 }
 
@@ -741,13 +741,15 @@ FSlateShaderResourceProxy* FSlateRHIResourceManager::FindOrCreateDynamicTextureR
 	return nullptr;
 }
 
-FSlateMaterialResource* FSlateRHIResourceManager::GetMaterialResource(const UObject* InMaterial, FVector2D ImageSize, FSlateShaderResource* TextureMask )
+FSlateMaterialResource* FSlateRHIResourceManager::GetMaterialResource(const UObject* InMaterial, FVector2D ImageSize, FSlateShaderResource* TextureMask, int32 InMaskKey )
 {
 	checkSlow(IsThreadSafeForSlateRendering());
 
 	const UMaterialInterface* Material = CastChecked<UMaterialInterface>(InMaterial);
 
-	TSharedPtr<FSlateMaterialResource> MaterialResource = DynamicResourceMap.GetMaterialResource(Material);
+	FMaterialKey Key(Material, InMaskKey);
+
+	TSharedPtr<FSlateMaterialResource> MaterialResource = DynamicResourceMap.GetMaterialResource(Key);
 	if (!MaterialResource.IsValid())
 	{
 		// Get a resource from the free list if possible
@@ -762,7 +764,7 @@ FSlateMaterialResource* FSlateRHIResourceManager::GetMaterialResource(const UObj
 			MaterialResource = MakeShareable(new FSlateMaterialResource(*Material, ImageSize, TextureMask));
 		}
 		
-		DynamicResourceMap.AddMaterialResource(Material, MaterialResource.ToSharedRef());
+		DynamicResourceMap.AddMaterialResource(Key, MaterialResource.ToSharedRef());
 	}
 	else if( MaterialResource->GetTextureMaskResource() != TextureMask )
 	{
@@ -823,9 +825,11 @@ void FSlateRHIResourceManager::ReleaseDynamicResource( const FSlateBrush& InBrus
 			{
 				UMaterialInterface* Material = Cast<UMaterialInterface>(ResourceObject);
 
-				TSharedPtr<FSlateMaterialResource> MaterialResource = DynamicResourceMap.GetMaterialResource(Material);
+				FMaterialKey Key(Material, 0);
+
+				TSharedPtr<FSlateMaterialResource> MaterialResource = DynamicResourceMap.GetMaterialResource(Key);
 				
-				DynamicResourceMap.RemoveMaterialResource(Material);
+				DynamicResourceMap.RemoveMaterialResource(Key);
 
 				if (MaterialResource.IsValid())
 				{
diff --git a/Engine/Source/Runtime/SlateRHIRenderer/Private/SlateRHIResourceManager.h b/Engine/Source/Runtime/SlateRHIRenderer/Private/SlateRHIResourceManager.h
--- a/Engine/Source/Runtime/SlateRHIRenderer/Private/SlateRHIResourceManager.h
+++ b/Engine/Source/Runtime/SlateRHIRenderer/Private/SlateRHIResourceManager.h
@@ -11,6 +11,31 @@ class FSlateUTextureResource;
 class FSlateMaterialResource;
 class FSlateAtlasedTextureResource;
 
+/** 
+ * Lookup key for materials.  Sometimes the same material is used with different masks so there must be
+ * unique resource per material/mask combination
+ */
+struct FMaterialKey
+{
+	TWeakObjectPtr<UMaterialInterface>	Material;
+	int32 MaskKey;
+
+	FMaterialKey(const UMaterialInterface* InMaterial, int32 InMaskKey)
+		: Material(InMaterial)
+		, MaskKey(InMaskKey)
+	{}
+
+	bool operator==(const FMaterialKey& Other) const
+	{
+		return Material == Other.Material && MaskKey == Other.MaskKey;
+	}
+
+	friend uint32 GetTypeHash(const FMaterialKey& Key)
+	{
+		return HashCombine(GetTypeHash(Key.Material), Key.MaskKey);
+	}
+};
+
 struct FDynamicResourceMap
 {
 public:
@@ -22,7 +47,7 @@ public:
 
 	TSharedPtr<FSlateAtlasedTextureResource> GetAtlasedTextureResource(UTexture* InObject) const;
 
-	TSharedPtr<FSlateMaterialResource> GetMaterialResource( const UMaterialInterface* Material ) const;
+	TSharedPtr<FSlateMaterialResource> GetMaterialResource( const FMaterialKey& InKey ) const;
 
 	void AddUTextureResource( UTexture* TextureObject, TSharedRef<FSlateUTextureResource> InResource );
 	void RemoveUTextureResource( UTexture* TextureObject );
@@ -30,8 +55,8 @@ public:
 	void AddDynamicTextureResource( FName ResourceName, TSharedRef<FSlateDynamicTextureResource> InResource);
 	void RemoveDynamicTextureResource( FName ResourceName );
 
-	void AddMaterialResource( const UMaterialInterface* Material, TSharedRef<FSlateMaterialResource> InResource );
-	void RemoveMaterialResource( const UMaterialInterface* Material );
+	void AddMaterialResource( const FMaterialKey& InKey, TSharedRef<FSlateMaterialResource> InResource );
+	void RemoveMaterialResource( const FMaterialKey& InKey );
 
 	void AddAtlasedTextureResource(UTexture* TextureObject, TSharedRef<FSlateAtlasedTextureResource> InResource);
 	void RemoveAtlasedTextureResource(UTexture* TextureObject);
@@ -55,23 +80,20 @@ public:
 private:
 	TMap<FName, TSharedPtr<FSlateDynamicTextureResource> > NativeTextureMap;
 	
-	typedef TMap<TWeakObjectPtr<UTexture>, TSharedPtr<FSlateUTextureResource> > FTextureResourceMap;
-
 	/** Map of all texture resources */
+	typedef TMap<TWeakObjectPtr<UTexture>, TSharedPtr<FSlateUTextureResource> > FTextureResourceMap;
 	FTextureResourceMap TextureMap;
 
-	uint64 TextureMemorySincePurge;
-
-	typedef TMap<TWeakObjectPtr<UMaterialInterface>, TSharedPtr<FSlateMaterialResource> > FMaterialResourceMap;
-
 	/** Map of all material resources */
+	typedef TMap<FMaterialKey, TSharedPtr<FSlateMaterialResource> > FMaterialResourceMap;
 	FMaterialResourceMap MaterialMap;
 
-	typedef TMap<TWeakObjectPtr<UObject>, TSharedPtr<FSlateAtlasedTextureResource> > FObjectResourceMap;
-
 	/** Map of all object resources */
+	typedef TMap<TWeakObjectPtr<UObject>, TSharedPtr<FSlateAtlasedTextureResource> > FObjectResourceMap;
 	FObjectResourceMap ObjectMap;
 
+	uint64 TextureMemorySincePurge;
+
 	int32 LastExpiredMaterialNumMarker;
 };
 
@@ -135,7 +157,7 @@ public:
 
 	/** FSlateShaderResourceManager interface */
 	virtual FSlateShaderResourceProxy* GetShaderResource( const FSlateBrush& InBrush ) override;
-	virtual FSlateShaderResource* GetFontShaderResource( uint32 FontAtlasIndex, FSlateShaderResource* FontTextureAtlas, const class UObject* FontMaterial ) override;
+	virtual FSlateShaderResource* GetFontShaderResource( int32 InTextureAtlasIndex, FSlateShaderResource* FontTextureAtlas, const class UObject* FontMaterial ) override;
 	virtual ISlateAtlasProvider* GetTextureAtlasProvider() override;
 
 	/**
@@ -255,7 +277,7 @@ private:
 	 *
 	 * @param InMaterial	The material object
 	 */
-	FSlateMaterialResource* GetMaterialResource( const UObject* InMaterial, FVector2D ImageSize, FSlateShaderResource* TextureMask );
+	FSlateMaterialResource* GetMaterialResource( const UObject* InMaterial, FVector2D ImageSize, FSlateShaderResource* TextureMask, int32 InMaskKey );
 
 	/**
 	 * Called when the application exists before the UObject system shuts down so we can free object resources
diff --git a/Engine/Source/Runtime/UMG/Private/WidgetLayoutLibrary.cpp b/Engine/Source/Runtime/UMG/Private/WidgetLayoutLibrary.cpp
--- a/Engine/Source/Runtime/UMG/Private/WidgetLayoutLibrary.cpp
+++ b/Engine/Source/Runtime/UMG/Private/WidgetLayoutLibrary.cpp
@@ -29,7 +29,7 @@ bool UWidgetLayoutLibrary::ProjectWorldLocationToWidgetPositionWithDistance(APla
 	if ( PlayerController )
 	{
 		FVector PixelLocation;
-		const bool bProjected = PlayerController->ProjectWorldLocationToScreenWithDistance(WorldLocation, PixelLocation);
+		const bool bProjected = PlayerController->ProjectWorldLocationToScreenWithDistance(WorldLocation, PixelLocation, /*bPlayerViewportRelative=*/ true);
 
 		if ( bProjected )
 		{
@@ -53,7 +53,7 @@ bool UWidgetLayoutLibrary::ProjectWorldLocationToWidgetPositionWithDistance(APla
 		}
 	}
 
-	ScreenPosition = FVector(0, 0, 0);
+	ScreenPosition = FVector::ZeroVector;
 
 	return false;
 }
