Fixed window de-activation clobbering "window to focus on activate"

This broke the logic that restored window focus when a context menu was closed.

diff --git a/Engine/Source/Runtime/SlateCore/Private/Widgets/SWindow.cpp b/Engine/Source/Runtime/SlateCore/Private/Widgets/SWindow.cpp
--- a/Engine/Source/Runtime/SlateCore/Private/Widgets/SWindow.cpp
+++ b/Engine/Source/Runtime/SlateCore/Private/Widgets/SWindow.cpp
@@ -1457,6 +1457,8 @@ bool SWindow::OnIsActiveChanged( const FWindowActivateEvent& ActivateEvent )
 		OnWindowDeactivated.ExecuteIfBound();	// deprecated
 		WindowDeactivatedEvent.Broadcast();
 
+		WidgetFocusedOnDeactivate.Reset();
+
 		const EWindowMode::Type WindowMode = GetWindowMode();
 		// If the window is not fullscreen, we do not want to automatically recapture the mouse unless an external UI such as Steam is open. Fullscreen windows we do.
 		if (WindowMode != EWindowMode::Fullscreen && WidgetToFocusOnActivate.IsValid() && WidgetToFocusOnActivate.Pin()->HasMouseCapture() && !FSlateApplicationBase::Get().IsExternalUIOpened())
@@ -1466,10 +1468,10 @@ bool SWindow::OnIsActiveChanged( const FWindowActivateEvent& ActivateEvent )
 		else if (SupportsKeyboardFocus())
 		{
 			// If we have no specific widget to focus then cache the currently focused widget so we can restore its focus when we regain focus
-			WidgetToFocusOnActivate = FSlateApplicationBase::Get().GetKeyboardFocusedWidget();
-			if (!WidgetToFocusOnActivate.IsValid())
+			WidgetFocusedOnDeactivate = FSlateApplicationBase::Get().GetKeyboardFocusedWidget();
+			if (!WidgetFocusedOnDeactivate.IsValid())
 			{
-				WidgetToFocusOnActivate = FSlateApplicationBase::Get().GetUserFocusedWidget(0);
+				WidgetFocusedOnDeactivate = FSlateApplicationBase::Get().GetUserFocusedWidget(0);
 			}
 		}
 	}
@@ -1497,7 +1499,8 @@ bool SWindow::OnIsActiveChanged( const FWindowActivateEvent& ActivateEvent )
 			else if (SupportsKeyboardFocus())
 			{
 				FWidgetPath WindowWidgetPath;
-				if (FSlateWindowHelper::FindPathToWidget(JustThisWindow, AsShared(), WindowWidgetPath))
+				TSharedRef<SWidget> WindowWidgetToFocus = WidgetFocusedOnDeactivate.IsValid() ? WidgetFocusedOnDeactivate.Pin().ToSharedRef() : AsShared();
+				if (FSlateWindowHelper::FindPathToWidget(JustThisWindow, WindowWidgetToFocus, WindowWidgetPath))
 				{
 					FSlateApplicationBase::Get().SetAllUserFocusAllowingDescendantFocus(WindowWidgetPath, EFocusCause::SetDirectly);
 				}
diff --git a/Engine/Source/Runtime/SlateCore/Public/Widgets/SWindow.h b/Engine/Source/Runtime/SlateCore/Public/Widgets/SWindow.h
--- a/Engine/Source/Runtime/SlateCore/Public/Widgets/SWindow.h
+++ b/Engine/Source/Runtime/SlateCore/Public/Widgets/SWindow.h
@@ -1024,6 +1024,10 @@ protected:
 		restore focus to a widget after a popup has been dismissed. */
 	TWeakPtr< SWidget > WidgetToFocusOnActivate;
 
+	/** Widget that had keyboard focus when this window was last de-activated, if any.  This is used to
+		restore focus to a widget after the window regains focus. */
+	TWeakPtr< SWidget > WidgetFocusedOnDeactivate;
+
 	/** Style used to draw this window */
 	const FWindowStyle* Style;
 	const FSlateBrush* WindowBackground;
