Copying //UE4/Dev-Editor to //UE4/Dev-Main (Source: //UE4/Dev-Editor @ 3279756)

diff --git a/Engine/Source/Runtime/Slate/Private/Framework/Application/SlateApplication.cpp b/Engine/Source/Runtime/Slate/Private/Framework/Application/SlateApplication.cpp
--- a/Engine/Source/Runtime/Slate/Private/Framework/Application/SlateApplication.cpp
+++ b/Engine/Source/Runtime/Slate/Private/Framework/Application/SlateApplication.cpp
@@ -24,6 +24,7 @@
 #include "Framework/Application/SWindowTitleBar.h"
 #include "Input/HittestGrid.h"
 #include "Stats/SlateStats.h"
+#include "HardwareCursor.h"
 
 #include "Framework/Application/IWidgetReflector.h"
 #include "Framework/Commands/GenericCommands.h"
@@ -979,15 +980,15 @@ void FSlateApplication::PlaySound( const FSlateSound& SoundToPlay, int32 UserInd
 }
 
 float FSlateApplication::GetSoundDuration(const FSlateSound& Sound) const
-	{
+{
 	return SlateSoundDevice->GetSoundDuration(Sound);
 }
 
 FVector2D FSlateApplication::GetCursorPos() const
 {
-	if ( PlatformApplication->Cursor.IsValid() )
+	if ( ICursor* Cursor = PlatformApplication->Cursor.Get() )
 	{
-		return PlatformApplication->Cursor->GetPosition();
+		return Cursor->GetPosition();
 	}
 
 	return FVector2D( 0, 0 );
@@ -1000,9 +1001,9 @@ FVector2D FSlateApplication::GetLastCursorPos() const
 
 void FSlateApplication::SetCursorPos( const FVector2D& MouseCoordinate )
 {
-	if ( PlatformApplication->Cursor.IsValid() )
+	if (ICursor* Cursor = PlatformApplication->Cursor.Get())
 	{
-		return PlatformApplication->Cursor->SetPosition( MouseCoordinate.X, MouseCoordinate.Y );
+		return Cursor->SetPosition( MouseCoordinate.X, MouseCoordinate.Y );
 	}
 }
 
@@ -2101,40 +2102,6 @@ TSharedRef<SWindow> FSlateApplication::AddWindowAsNativeChild( TSharedRef<SWindo
 	return InSlateWindow;
 }
 
-
-TSharedRef<SWindow> FSlateApplication::PushMenu( const TSharedRef<SWidget>& InParentContent, const TSharedRef<SWidget>& InContent, const FVector2D& SummonLocation, const FPopupTransitionEffect& TransitionEffect, const bool bFocusImmediately, const bool bShouldAutoSize, const FVector2D& WindowSize, const FVector2D& SummonLocationSize )
-{
-#if !(UE_BUILD_SHIPPING && WITH_EDITOR)
-	// The would-be parent of the new menu being pushed is about to be destroyed.  Any children added to an about to be destroyed window will also be destroyed
-	FWidgetPath WidgetPath;
-	GeneratePathToWidgetChecked(InParentContent, WidgetPath);
-	if (IsWindowInDestroyQueue( WidgetPath.GetWindow() ))
-	{
-		UE_LOG(LogSlate, Warning, TEXT("FSlateApplication::PushMenu() called when parent window queued for destroy. New menu will be destroyed."));
-	}
-#endif
-
-	// This is a DEPRECATED version of PushMenu(). It creates a new menu but forces the popup method to EPopupMethod::CreateNewWindow
-
-	// Wrap content in a fixed size box if autosizing isn't wanted
-	TSharedRef<SWidget> WappedContent = InContent;
-	if (!bShouldAutoSize)
-	{
-		WappedContent = SNew(SBox)
-			.WidthOverride(WindowSize.X)
-			.HeightOverride(WindowSize.Y)
-			[
-				InContent
-			];
-	}
-
-	TSharedPtr<IMenu> Menu = PushMenu(InParentContent, FWidgetPath(), WappedContent, SummonLocation, TransitionEffect, bFocusImmediately, SummonLocationSize, EPopupMethod::CreateNewWindow);
-	TSharedPtr<SWindow> Window = Menu->GetOwnedWindow();
-	check(Window.IsValid());
-
-	return Window.ToSharedRef();
-}
-
 TSharedPtr<IMenu> FSlateApplication::PushMenu(const TSharedRef<SWidget>& InParentWidget, const FWidgetPath& InOwnerPath, const TSharedRef<SWidget>& InContent, const FVector2D& SummonLocation, const FPopupTransitionEffect& TransitionEffect, const bool bFocusImmediately, const FVector2D& SummonLocationSize, TOptional<EPopupMethod> Method, const bool bIsCollapsedByParent)
 {
 	// Caller supplied a valid path? Pass it to the menu stack.
@@ -2182,17 +2149,6 @@ TSharedPtr<IMenu> FSlateApplication::PushHostedMenu(const TSharedPtr<IMenu>& InP
 	return MenuStack.PushHosted(InParentMenu, InMenuHost, InContent, OutWrappedContent, TransitionEffect, ShouldThrottle, bIsCollapsedByParent);
 }
 
-bool FSlateApplication::HasOpenSubMenus( TSharedRef<SWindow> Window ) const
-{
-	// deprecated
-	TSharedPtr<IMenu> Menu = MenuStack.FindMenuFromWindow(Window);
-	if (Menu.IsValid())
-	{
-		return MenuStack.HasOpenSubMenus(Menu);
-	}
-	return false;
-}
-
 bool FSlateApplication::HasOpenSubMenus(TSharedPtr<IMenu> InMenu) const
 {
 	return MenuStack.HasOpenSubMenus(InMenu);
@@ -2218,16 +2174,6 @@ void FSlateApplication::DismissAllMenus()
 	MenuStack.DismissAll();
 }
 
-void FSlateApplication::DismissMenu( TSharedRef<SWindow> MenuWindowToDismiss )
-{
-	// deprecated
-	TSharedPtr<IMenu> Menu = MenuStack.FindMenuFromWindow(MenuWindowToDismiss);
-	if (Menu.IsValid())
-	{
-		DismissMenu(Menu);
-	}
-}
-
 void FSlateApplication::DismissMenu(const TSharedPtr<IMenu>& InFromMenu)
 {
 	MenuStack.DismissFrom(InFromMenu);
@@ -2246,13 +2192,6 @@ void FSlateApplication::DismissMenuByWidget(const TSharedRef<SWidget>& InWidgetI
 	}
 }
 
-int32 FSlateApplication::GetLocationInMenuStack( TSharedRef<SWindow> WindowToFind ) const
-{
-	// DEPRECATED
-	return MenuStack.FindLocationInStack(WindowToFind);
-}
-
-
 void FSlateApplication::RequestDestroyWindow( TSharedRef<SWindow> InWindowToDestroy )
 {
 	// Logging to track down window shutdown issues with movie loading threads. Too spammy in editor builds with all the windows
@@ -2439,11 +2378,6 @@ void FSlateApplication::SetUserFocusToGameViewport(uint32 UserIndex, EFocusCause
 	}
 }
 
-void FSlateApplication::SetFocusToGameViewport()
-{
-	SetUserFocusToGameViewport(GetUserIndexForKeyboard(), EFocusCause::SetDirectly);
-}
-
 void FSlateApplication::SetAllUserFocusToGameViewport(EFocusCause ReasonFocusIsChanging /* = EFocusCause::SetDirectly*/)
 {
 	TSharedPtr< SViewport > CurrentGameViewportWidget = GameViewportWidget.Pin();
@@ -2459,11 +2393,6 @@ void FSlateApplication::SetAllUserFocusToGameViewport(EFocusCause ReasonFocusIsC
 	}
 }
 
-void FSlateApplication::SetJoystickCaptorToGameViewport()
-{
-	SetAllUserFocusToGameViewport();
-}
-
 void FSlateApplication::ActivateGameViewport()
 {
 	// Only focus the window if the application is active, if not the application activation sequence will take care of it.
@@ -2552,11 +2481,6 @@ TSharedPtr<SWidget> FSlateApplication::GetUserFocusedWidget(uint32 UserIndex) co
 	return TSharedPtr<SWidget>();
 }
 
-TSharedPtr<SWidget> FSlateApplication::GetJoystickCaptor(uint32 UserIndex) const
-{
-	return GetUserFocusedWidget(UserIndex);
-}
-
 void FSlateApplication::ClearUserFocus(uint32 UserIndex, EFocusCause ReasonFocusIsChanging /* = EFocusCause::SetDirectly*/)
 {
 	SetUserFocus(UserIndex, FWidgetPath(), ReasonFocusIsChanging);
@@ -2567,11 +2491,6 @@ void FSlateApplication::ClearAllUserFocus(EFocusCause ReasonFocusIsChanging /*=
 	SetAllUserFocus(FWidgetPath(), ReasonFocusIsChanging);
 }
 
-void FSlateApplication::ReleaseJoystickCapture(uint32 UserIndex)
-{
-	ClearUserFocus(UserIndex);
-}
-
 bool FSlateApplication::SetKeyboardFocus(const TSharedPtr< SWidget >& OptionalWidgetToFocus, EFocusCause ReasonFocusIsChanging /* = EFocusCause::SetDirectly*/)
 {
 	return SetUserFocus(GetUserIndexForKeyboard(), OptionalWidgetToFocus, ReasonFocusIsChanging);
@@ -4868,6 +4787,14 @@ bool FSlateApplication::ProcessMouseButtonDownEvent( const TSharedPtr< FGenericW
 {
 	SCOPE_CYCLE_COUNTER(STAT_ProcessMouseButtonDown);
 
+#if WITH_EDITOR
+	//Send the key input to all pre input key down listener function
+	if (OnApplicationMousePreInputButtonDownListenerEvent.IsBound())
+	{
+		OnApplicationMousePreInputButtonDownListenerEvent.Broadcast(MouseEvent);
+	}
+#endif //WITH_EDITOR
+
 	QueueSynthesizedMouseMove();
 	SetLastUserInteractionTime(this->GetCurrentTime());
 	LastUserInteractionTimeForThrottling = LastUserInteractionTime;
@@ -6609,4 +6536,21 @@ void FSlateApplication::NavigateFromWidgetUnderCursor(EUINavigation InNavigation
 	}
 }
 
+void FSlateApplication::RegisterCursor(EMouseCursor::Type CursorType, TSharedPtr<FHardwareCursor> HardwareCursor)
+{
+	HardwareCursors.Add(CursorType, HardwareCursor);
+
+	if (ICursor* Cursor = PlatformApplication->Cursor.Get())
+	{
+		if (HardwareCursor.IsValid())
+		{
+			Cursor->SetTypeShape(CursorType, HardwareCursor->GetHandle());
+		}
+		else
+		{
+			Cursor->SetTypeShape(CursorType, nullptr);
+		}
+	}
+}
+
 #undef SLATE_HAS_WIDGET_REFLECTOR
diff --git a/Engine/Source/Runtime/Slate/Private/Framework/MultiBox/SClippingHorizontalBox.cpp b/Engine/Source/Runtime/Slate/Private/Framework/MultiBox/SClippingHorizontalBox.cpp
--- a/Engine/Source/Runtime/Slate/Private/Framework/MultiBox/SClippingHorizontalBox.cpp
+++ b/Engine/Source/Runtime/Slate/Private/Framework/MultiBox/SClippingHorizontalBox.cpp
@@ -96,7 +96,7 @@ int32 SClippingHorizontalBox::OnPaint( const FPaintArgs& Args, const FGeometry&
 		}
 	
 		bool bEnabled = ShouldBeEnabled( bParentEnabled );
-		ESlateDrawEffect::Type DrawEffects = !bEnabled ? ESlateDrawEffect::DisabledEffect : ESlateDrawEffect::None;
+		ESlateDrawEffect DrawEffects = !bEnabled ? ESlateDrawEffect::DisabledEffect : ESlateDrawEffect::None;
 		FSlateColor BorderBackgroundColor = FLinearColor::White;
 
 		// Draw the custom border
diff --git a/Engine/Source/Runtime/Slate/Private/Framework/Text/SlateImageRun.cpp b/Engine/Source/Runtime/Slate/Private/Framework/Text/SlateImageRun.cpp
--- a/Engine/Source/Runtime/Slate/Private/Framework/Text/SlateImageRun.cpp
+++ b/Engine/Source/Runtime/Slate/Private/Framework/Text/SlateImageRun.cpp
@@ -151,7 +151,7 @@ int32 FSlateImageRun::OnPaint( const FPaintArgs& Args, const FTextLayout::FLineV
 	if ( Image->DrawAs != ESlateBrushDrawType::NoDrawType )
 	{
 		const FLinearColor FinalColorAndOpacity( InWidgetStyle.GetColorAndOpacityTint() * Image->GetTint( InWidgetStyle ) );
-		const uint32 DrawEffects = bParentEnabled ? ESlateDrawEffect::None : ESlateDrawEffect::DisabledEffect;
+		const ESlateDrawEffect DrawEffects = bParentEnabled ? ESlateDrawEffect::None : ESlateDrawEffect::DisabledEffect;
 		FSlateDrawElement::MakeBox(
 			OutDrawElements, 
 			++LayerId, 
diff --git a/Engine/Source/Runtime/Slate/Private/Framework/Text/SlatePasswordRun.cpp b/Engine/Source/Runtime/Slate/Private/Framework/Text/SlatePasswordRun.cpp
--- a/Engine/Source/Runtime/Slate/Private/Framework/Text/SlatePasswordRun.cpp
+++ b/Engine/Source/Runtime/Slate/Private/Framework/Text/SlatePasswordRun.cpp
@@ -44,7 +44,7 @@ int8 FSlatePasswordRun::GetKerning(int32 CurrentIndex, float Scale, const FRunTe
 int32 FSlatePasswordRun::OnPaint(const FPaintArgs& Args, const FTextLayout::FLineView& Line, const TSharedRef< ILayoutBlock >& Block, const FTextBlockStyle& DefaultStyle, const FGeometry& AllottedGeometry, const FSlateRect& MyClippingRect, FSlateWindowElementList& OutDrawElements, int32 LayerId, const FWidgetStyle& InWidgetStyle, bool bParentEnabled) const
 {
 	const FSlateRect ClippingRect = AllottedGeometry.GetClippingRect().IntersectionWith(MyClippingRect);
-	const ESlateDrawEffect::Type DrawEffects = bParentEnabled ? ESlateDrawEffect::None : ESlateDrawEffect::DisabledEffect;
+	const ESlateDrawEffect DrawEffects = bParentEnabled ? ESlateDrawEffect::None : ESlateDrawEffect::DisabledEffect;
 
 	const bool ShouldDropShadow = Style.ShadowColorAndOpacity.A > 0.f && Style.ShadowOffset.SizeSquared() > 0.f;
 	const FVector2D BlockLocationOffset = Block->GetLocationOffset();
diff --git a/Engine/Source/Runtime/Slate/Private/Framework/Text/SlateTextLayout.cpp b/Engine/Source/Runtime/Slate/Private/Framework/Text/SlateTextLayout.cpp
--- a/Engine/Source/Runtime/Slate/Private/Framework/Text/SlateTextLayout.cpp
+++ b/Engine/Source/Runtime/Slate/Private/Framework/Text/SlateTextLayout.cpp
@@ -51,7 +51,7 @@ void FSlateTextLayout::ArrangeChildren( const FGeometry& AllottedGeometry, FArra
 int32 FSlateTextLayout::OnPaint( const FPaintArgs& Args, const FGeometry& AllottedGeometry, const FSlateRect& MyClippingRect, FSlateWindowElementList& OutDrawElements, int32 LayerId, const FWidgetStyle& InWidgetStyle, bool bParentEnabled ) const
 {
 	const FSlateRect ClippingRect = AllottedGeometry.GetClippingRect().IntersectionWith(MyClippingRect);
-	const ESlateDrawEffect::Type DrawEffects = bParentEnabled ? ESlateDrawEffect::None : ESlateDrawEffect::DisabledEffect;
+	const ESlateDrawEffect DrawEffects = bParentEnabled ? ESlateDrawEffect::None : ESlateDrawEffect::DisabledEffect;
 
 	static bool ShowDebug = false;
 	FLinearColor BlockHue( 0, 1.0f, 1.0f, 0.5 );
diff --git a/Engine/Source/Runtime/Slate/Private/Widgets/Colors/SColorBlock.cpp b/Engine/Source/Runtime/Slate/Private/Widgets/Colors/SColorBlock.cpp
--- a/Engine/Source/Runtime/Slate/Private/Widgets/Colors/SColorBlock.cpp
+++ b/Engine/Source/Runtime/Slate/Private/Widgets/Colors/SColorBlock.cpp
@@ -25,7 +25,7 @@ int32 SColorBlock::OnPaint( const FPaintArgs& Args, const FGeometry& AllottedGeo
 {
 	const FSlateBrush* GenericBrush = FCoreStyle::Get().GetBrush( "GenericWhiteBox" );
 
-	const ESlateDrawEffect::Type DrawEffects = ESlateDrawEffect::None;
+	const ESlateDrawEffect DrawEffects = ESlateDrawEffect::None;
 	
 	FLinearColor InColor = Color.Get();
 	if (ColorIsHSV.Get())
diff --git a/Engine/Source/Runtime/Slate/Private/Widgets/Input/SButton.cpp b/Engine/Source/Runtime/Slate/Private/Widgets/Input/SButton.cpp
--- a/Engine/Source/Runtime/Slate/Private/Widgets/Input/SButton.cpp
+++ b/Engine/Source/Runtime/Slate/Private/Widgets/Input/SButton.cpp
@@ -78,7 +78,7 @@ int32 SButton::OnPaint(const FPaintArgs& Args, const FGeometry& AllottedGeometry
 
 	const FSlateBrush* BrushResource = !bShowDisabledEffect && !bEnabled ? DisabledImage : GetBorder();
 	
-	ESlateDrawEffect::Type DrawEffects = bShowDisabledEffect && !bEnabled ? ESlateDrawEffect::DisabledEffect : ESlateDrawEffect::None;
+	ESlateDrawEffect DrawEffects = bShowDisabledEffect && !bEnabled ? ESlateDrawEffect::DisabledEffect : ESlateDrawEffect::None;
 
 	if (BrushResource && BrushResource->DrawAs != ESlateBrushDrawType::NoDrawType)
 	{
diff --git a/Engine/Source/Runtime/Slate/Private/Widgets/Input/SEditableText.cpp b/Engine/Source/Runtime/Slate/Private/Widgets/Input/SEditableText.cpp
--- a/Engine/Source/Runtime/Slate/Private/Widgets/Input/SEditableText.cpp
+++ b/Engine/Source/Runtime/Slate/Private/Widgets/Input/SEditableText.cpp
@@ -464,7 +464,16 @@ EKeyboardType SEditableText::GetVirtualKeyboardType() const
 
 TSharedRef<SWidget> SEditableText::GetSlateWidget()
 {
-	return SharedThis(this);
+	return AsShared();
+}
+
+TSharedPtr<SWidget> SEditableText::GetSlateWidgetPtr()
+{
+	if (DoesSharedInstanceExist())
+	{
+		return AsShared();
+	}
+	return nullptr;
 }
 
 TSharedPtr<SWidget> SEditableText::BuildContextMenuContent() const
diff --git a/Engine/Source/Runtime/Slate/Private/Widgets/Layout/SBorder.cpp b/Engine/Source/Runtime/Slate/Private/Widgets/Layout/SBorder.cpp
--- a/Engine/Source/Runtime/Slate/Private/Widgets/Layout/SBorder.cpp
+++ b/Engine/Source/Runtime/Slate/Private/Widgets/Layout/SBorder.cpp
@@ -68,7 +68,7 @@ int32 SBorder::OnPaint( const FPaintArgs& Args, const FGeometry& AllottedGeometr
 		
 	const bool bEnabled = ShouldBeEnabled(bParentEnabled);
 	const bool bShowDisabledEffect = ShowDisabledEffect.Get();
-	ESlateDrawEffect::Type DrawEffects = bShowDisabledEffect && !bEnabled ? ESlateDrawEffect::DisabledEffect : ESlateDrawEffect::None;
+	ESlateDrawEffect DrawEffects = bShowDisabledEffect && !bEnabled ? ESlateDrawEffect::DisabledEffect : ESlateDrawEffect::None;
 
 	if ( BrushResource && BrushResource->DrawAs != ESlateBrushDrawType::NoDrawType )
 	{
diff --git a/Engine/Source/Runtime/Slate/Private/Widgets/Layout/SConstraintCanvas.cpp b/Engine/Source/Runtime/Slate/Private/Widgets/Layout/SConstraintCanvas.cpp
--- a/Engine/Source/Runtime/Slate/Private/Widgets/Layout/SConstraintCanvas.cpp
+++ b/Engine/Source/Runtime/Slate/Private/Widgets/Layout/SConstraintCanvas.cpp
@@ -84,9 +84,9 @@ void SConstraintCanvas::ArrangeLayeredChildren(const FGeometry& AllottedGeometry
 		// or allow the explicit ZOrder value to place children in the same layer. This will allow users to have non-touching
 		// children render into the same layer and have a chance of being batched by the Slate renderer for better performance.
 #if WITH_EDITOR
-		bool bExplicitChildZOrder = GetDefault<USlateSettings>()->bExplicitCanvasChildZOrder;
+		const bool bExplicitChildZOrder = GetDefault<USlateSettings>()->bExplicitCanvasChildZOrder;
 #else
-		static bool bExplicitChildZOrder = GetDefault<USlateSettings>()->bExplicitCanvasChildZOrder;
+		const static bool bExplicitChildZOrder = GetDefault<USlateSettings>()->bExplicitCanvasChildZOrder;
 #endif
 		// Sort the children based on zorder.
 		TArray< FChildZOrder, TInlineAllocator<64> > SlotOrder;
diff --git a/Engine/Source/Runtime/Slate/Private/Widgets/Layout/SSplitter.cpp b/Engine/Source/Runtime/Slate/Private/Widgets/Layout/SSplitter.cpp
--- a/Engine/Source/Runtime/Slate/Private/Widgets/Layout/SSplitter.cpp
+++ b/Engine/Source/Runtime/Slate/Private/Widgets/Layout/SSplitter.cpp
@@ -57,6 +57,7 @@ void SSplitter::Construct( const SSplitter::FArguments& InArgs )
 	HoveredHandleIndex = INDEX_NONE;
 	bIsResizing = false;
 	Style = InArgs._Style;
+	OnGetMaxSlotSize = InArgs._OnGetMaxSlotSize;
 
 	for (int32 SlotIndex = 0; SlotIndex < InArgs.Slots.Num(); ++SlotIndex )
 	{
@@ -197,7 +198,7 @@ int32 SSplitter::OnPaint( const FPaintArgs& Args, const FGeometry& AllottedGeome
 				GeometryAfterSplitter.ToPaintGeometry( HandlePosition, HandleSize, 1.0f ),
 				NormalHandleBrush,
 				MyClippingRect,
-				ShouldBeEnabled( bParentEnabled ),
+				ShouldBeEnabled(bParentEnabled) ? ESlateDrawEffect::None : ESlateDrawEffect::DisabledEffect,
 				InWidgetStyle.GetColorAndOpacityTint() * NormalHandleBrush->TintColor.GetSpecifiedColor()
 			);
 		}
@@ -209,7 +210,7 @@ int32 SSplitter::OnPaint( const FPaintArgs& Args, const FGeometry& AllottedGeome
 				GeometryAfterSplitter.ToPaintGeometry( HandlePosition, HandleSize, 1.0f ),
 				&Style->HandleHighlightBrush,
 				MyClippingRect,
-				ShouldBeEnabled( bParentEnabled ),
+				ShouldBeEnabled(bParentEnabled) ? ESlateDrawEffect::None : ESlateDrawEffect::DisabledEffect,
 				InWidgetStyle.GetColorAndOpacityTint() * Style->HandleHighlightBrush.TintColor.GetSpecifiedColor()
 			);	
 		}
@@ -343,14 +344,7 @@ FReply SSplitter::OnMouseMove( const FGeometry& MyGeometry, const FPointerEvent&
 	{
 		if ( !MouseEvent.GetCursorDelta().IsZero() )
 		{
-			if (Orientation == Orient_Horizontal)
-			{
-				HandleResizing<Orient_Horizontal>( PhysicalSplitterHandleSize, ResizeMode, HoveredHandleIndex, LocalMousePosition, Children, LayoutChildren );
-			}
-			else
-			{
-				HandleResizing<Orient_Vertical>( PhysicalSplitterHandleSize, ResizeMode, HoveredHandleIndex, LocalMousePosition, Children, LayoutChildren );
-			}
+			HandleResizingByMousePosition(Orientation, PhysicalSplitterHandleSize, ResizeMode, HoveredHandleIndex, LocalMousePosition, Children, LayoutChildren );
 		}
 
 		return FReply::Handled();
@@ -375,6 +369,25 @@ FReply SSplitter::OnMouseMove( const FGeometry& MyGeometry, const FPointerEvent&
 
 }
 
+FReply SSplitter::OnMouseButtonDoubleClick(const FGeometry& InMyGeometry, const FPointerEvent& InMouseEvent)
+{
+	if (OnGetMaxSlotSize.IsBound())
+	{
+		FVector2D MaxSlotSize = OnGetMaxSlotSize.Execute(HoveredHandleIndex);
+
+		if (!MaxSlotSize.IsZero())
+		{
+			TArray<FLayoutGeometry> LayoutChildren = ArrangeChildrenForLayout(InMyGeometry);
+
+			HandleResizingBySize(Orientation, PhysicalSplitterHandleSize, ResizeMode, HoveredHandleIndex, MaxSlotSize, Children, LayoutChildren);
+
+			return FReply::Handled();
+		}
+	}
+
+	return SWidget::OnMouseButtonDoubleClick(InMyGeometry, InMouseEvent);
+}
+
 void SSplitter::OnMouseLeave( const FPointerEvent& MouseEvent )
 {
 	if ( !bIsResizing )
@@ -487,25 +500,19 @@ void SSplitter::FindAllResizeableSlotsAfterHandle( int32 DraggedHandle, const TP
 	}
 };
 
-
-template<EOrientation SplitterOrientation>
-void SSplitter::HandleResizing( const float PhysicalSplitterHandleSize, const ESplitterResizeMode::Type ResizeMode, int32 DraggedHandle, const FVector2D& LocalMousePos, TPanelChildren<FSlot>& Children, const TArray<FLayoutGeometry>& ChildGeometries )
+void SSplitter::HandleResizingDelta(EOrientation SplitterOrientation, const float PhysicalSplitterHandleSize, const ESplitterResizeMode::Type ResizeMode, int32 DraggedHandle, float Delta, TPanelChildren<FSlot>& Children, const TArray<FLayoutGeometry>& ChildGeometries)
 {
 	const int32 NumChildren = Children.Num();
-
 	const int32 AxisIndex = (SplitterOrientation == Orient_Horizontal) ? 0 : 1;
 
 	// Note:
 	//  - Prev vs. Next refers to the widgets in the order they are laid out (left->right, top->bottom).
 	//  - New vs. Old refers to the Old values for width/height vs. the post-resize values.
 
-	const float HandlePos = ChildGeometries[DraggedHandle+1].GetLocalToParentTransform().GetTranslation().Component(AxisIndex) - PhysicalSplitterHandleSize / 2;
-	float Delta = LocalMousePos.Component(AxisIndex) - HandlePos;
-
-	const int32 SlotBeforeDragHandle = FindResizeableSlotBeforeHandle( DraggedHandle, Children );
+	const int32 SlotBeforeDragHandle = FindResizeableSlotBeforeHandle(DraggedHandle, Children);
 
 	TArray< int32 > SlotsAfterDragHandleIndicies;
-	if ( ResizeMode == ESplitterResizeMode::Fixed )
+	if (ResizeMode == ESplitterResizeMode::FixedPosition)
 	{
 		const int32 SlotAfterDragHandle = FindResizeableSlotAfterHandle( DraggedHandle, Children );
 
@@ -514,7 +521,7 @@ void SSplitter::HandleResizing( const float PhysicalSplitterHandleSize, const ES
 			SlotsAfterDragHandleIndicies.Add( SlotAfterDragHandle );
 		}
 	}
-	else if ( ResizeMode == ESplitterResizeMode::Fill )
+	else if (ResizeMode == ESplitterResizeMode::Fill || ResizeMode == ESplitterResizeMode::FixedSize)
 	{
 		FindAllResizeableSlotsAfterHandle( DraggedHandle, Children, /*OUT*/ SlotsAfterDragHandleIndicies );
 	}
@@ -551,20 +558,25 @@ void SSplitter::HandleResizing( const float PhysicalSplitterHandleSize, const ES
 
 		// Distribute the Delta across the affected slots after the drag handle
 		float UnusedDelta = Delta;
+
 		for (int DistributionCount = 0; DistributionCount < SlotsAfterDragHandle.Num() && UnusedDelta != 0; DistributionCount++)
 		{
-			float DividedDelta = UnusedDelta / SlotsAfterDragHandle.Num();
+			float DividedDelta = ResizeMode != ESplitterResizeMode::FixedSize ? UnusedDelta / SlotsAfterDragHandle.Num() : UnusedDelta;
 			UnusedDelta = 0;
+
 			for (int SlotIndex = 0; SlotIndex < SlotsAfterDragHandle.Num(); SlotIndex++)
 			{
 				FSlotInfo& SlotInfo = SlotsAfterDragHandle[ SlotIndex ];
 
-				float CurrentSize = ClampChild( SlotInfo.Geometry->GetSizeInParentSpace().Component(AxisIndex) );
-				SlotInfo.NewSize = ClampChild( CurrentSize - DividedDelta );
+				if (ResizeMode != ESplitterResizeMode::FixedSize || (ResizeMode == ESplitterResizeMode::FixedSize && SlotIndex == SlotsAfterDragHandle.Num() - 1)) // resize only the last handle in the case of fixed size
+				{
+					float CurrentSize = ClampChild(SlotInfo.Geometry->GetSizeInParentSpace().Component(AxisIndex));
+					SlotInfo.NewSize = ClampChild(CurrentSize - DividedDelta);
 
-				// If one of the slots couldn't be fully adjusted by the delta due to min/max constraints then
-				// the leftover delta needs to be evenly distributed to all of the other slots
-				UnusedDelta += SlotInfo.NewSize - ( CurrentSize - DividedDelta );
+					// If one of the slots couldn't be fully adjusted by the delta due to min/max constraints then
+					// the leftover delta needs to be evenly distributed to all of the other slots
+					UnusedDelta += SlotInfo.NewSize - (CurrentSize - DividedDelta);
+				}
 			}
 		}
 
@@ -617,6 +629,26 @@ void SSplitter::HandleResizing( const float PhysicalSplitterHandleSize, const ES
 	}
 }
 
+void SSplitter::HandleResizingBySize(EOrientation SplitterOrientation, const float PhysicalSplitterHandleSize, const ESplitterResizeMode::Type ResizeMode, int32 DraggedHandle, const FVector2D& DesiredSize, TPanelChildren<FSlot>& Children, const TArray<FLayoutGeometry>& ChildGeometries)
+{
+	const int32 AxisIndex = (SplitterOrientation == Orient_Horizontal) ? 0 : 1;
+
+	float CurrentSlotSize = ChildGeometries[DraggedHandle].GetSizeInParentSpace().Component(AxisIndex);
+	float Delta = DesiredSize.Component(AxisIndex) - CurrentSlotSize;
+
+	HandleResizingDelta(SplitterOrientation, PhysicalSplitterHandleSize, ResizeMode, DraggedHandle, Delta, Children, ChildGeometries);
+}
+
+void SSplitter::HandleResizingByMousePosition(EOrientation SplitterOrientation, const float PhysicalSplitterHandleSize, const ESplitterResizeMode::Type ResizeMode, int32 DraggedHandle, const FVector2D& LocalMousePos, TPanelChildren<FSlot>& Children, const TArray<FLayoutGeometry>& ChildGeometries )
+{
+	const int32 AxisIndex = (SplitterOrientation == Orient_Horizontal) ? 0 : 1;
+
+	const float HandlePos = ChildGeometries[DraggedHandle+1].GetLocalToParentTransform().GetTranslation().Component(AxisIndex) - PhysicalSplitterHandleSize / 2;
+	float Delta = LocalMousePos.Component(AxisIndex) - HandlePos;
+
+	HandleResizingDelta(SplitterOrientation, PhysicalSplitterHandleSize, ResizeMode, DraggedHandle, Delta, Children, ChildGeometries);
+}
+
 /**
 * @param ProposedSize  A size that a child would like to be
 *
diff --git a/Engine/Source/Runtime/Slate/Private/Widgets/Text/SMultiLineEditableText.cpp b/Engine/Source/Runtime/Slate/Private/Widgets/Text/SMultiLineEditableText.cpp
--- a/Engine/Source/Runtime/Slate/Private/Widgets/Text/SMultiLineEditableText.cpp
+++ b/Engine/Source/Runtime/Slate/Private/Widgets/Text/SMultiLineEditableText.cpp
@@ -266,7 +266,16 @@ EKeyboardType SMultiLineEditableText::GetVirtualKeyboardType() const
 
 TSharedRef<SWidget> SMultiLineEditableText::GetSlateWidget()
 {
-	return SharedThis(this);
+	return AsShared();
+}
+
+TSharedPtr<SWidget> SMultiLineEditableText::GetSlateWidgetPtr()
+{
+	if (DoesSharedInstanceExist())
+	{
+		return AsShared();
+	}
+	return nullptr;
 }
 
 TSharedPtr<SWidget> SMultiLineEditableText::BuildContextMenuContent() const
diff --git a/Engine/Source/Runtime/Slate/Private/Widgets/Text/SlateEditableTextLayout.cpp b/Engine/Source/Runtime/Slate/Private/Widgets/Text/SlateEditableTextLayout.cpp
--- a/Engine/Source/Runtime/Slate/Private/Widgets/Text/SlateEditableTextLayout.cpp
+++ b/Engine/Source/Runtime/Slate/Private/Widgets/Text/SlateEditableTextLayout.cpp
@@ -1818,7 +1818,11 @@ void FSlateEditableTextLayout::OnContextMenuClosed(TSharedRef<IMenu> Menu)
 	// to know that the window is still available for OnFocusReceived and OnFocusLost even though it's about to be destroyed
 
 	// Give our owner widget focus when the context menu has been dismissed
-	FSlateApplication::Get().SetKeyboardFocus(OwnerWidget->GetSlateWidget(), EFocusCause::OtherWidgetLostFocus);
+	TSharedPtr<SWidget> OwnerSlateWidget = OwnerWidget->GetSlateWidgetPtr();
+	if (OwnerSlateWidget.IsValid())
+	{
+		FSlateApplication::Get().SetKeyboardFocus(OwnerSlateWidget, EFocusCause::OtherWidgetLostFocus);
+	}
 }
 
 void FSlateEditableTextLayout::InsertRunAtCursor(TSharedRef<IRun> InRun)
diff --git a/Engine/Source/Runtime/Slate/Private/Widgets/Views/SHeaderRow.cpp b/Engine/Source/Runtime/Slate/Private/Widgets/Views/SHeaderRow.cpp
--- a/Engine/Source/Runtime/Slate/Private/Widgets/Views/SHeaderRow.cpp
+++ b/Engine/Source/Runtime/Slate/Private/Widgets/Views/SHeaderRow.cpp
@@ -229,20 +229,6 @@ public:
 		return FReply::Unhandled();
 	}
 
-
-private:
-
-
-	const FSlateBrush* GetHeaderBackgroundBrush() const
-	{
-		if ( IsHovered() && SortMode.IsBound() )
-		{
-			return &Style->HoveredBrush;
-		}
-
-		return &Style->NormalBrush;
-	}
-
 	EVisibility GetMenuOverlayVisibility() const
 	{
 		if (ComboVisibility == EHeaderComboVisibility::OnHover)
@@ -256,6 +242,24 @@ private:
 		return EVisibility::Visible;
 	}
 
+	const FVector2D& GetMenuOverlaySize() const 
+	{ 
+		return MenuOverlay.IsValid() ? MenuOverlay->GetDesiredSize() : FVector2D::ZeroVector; 
+	}
+
+private:
+
+
+	const FSlateBrush* GetHeaderBackgroundBrush() const
+	{
+		if ( IsHovered() && SortMode.IsBound() )
+		{
+			return &Style->HoveredBrush;
+		}
+
+		return &Style->NormalBrush;
+	}
+	
 	const FSlateBrush* GetComboButtonBorderBrush() const
 	{
 		if ( ComboButton.IsValid() && ( ComboButton->IsHovered() || ComboButton->IsOpen() ) )
@@ -415,6 +419,8 @@ void SHeaderRow::Construct( const FArguments& InArgs )
 	ScrollBarThickness = FVector2D::ZeroVector;
 	ScrollBarVisibility = EVisibility::Collapsed;
 	Style = InArgs._Style;
+	OnGetMaxRowSizeForColumn = InArgs._OnGetMaxRowSizeForColumn;
+	ResizeMode = InArgs._ResizeMode;
 
 	if ( InArgs._OnColumnsChanged.IsBound() )
 	{
@@ -539,9 +545,37 @@ void SHeaderRow::SetColumnWidth( const FName& InColumnId, float InWidth )
 	}
 }
 
+FVector2D SHeaderRow::GetRowSizeForSlotIndex(int32 SlotIndex) const
+{
+	if (Columns.IsValidIndex(SlotIndex))
+	{
+		const TSharedPtr<STableColumnHeader>& HeaderWidget = HeaderWidgets[SlotIndex];
+		const FColumn& Column = Columns[SlotIndex];
+
+		FVector2D HeaderSize = HeaderWidget->GetDesiredSize();
+
+		if (Column.HeaderMenuContent.Widget != SNullWidget::NullWidget && HeaderWidget->GetMenuOverlayVisibility() != EVisibility::Visible)
+		{
+			HeaderSize += HeaderWidget->GetMenuOverlaySize();
+		}
+
+		if (OnGetMaxRowSizeForColumn.IsBound())
+		{
+			// It's assume that a header is at the top, so the sizing is for the width
+			FVector2D MaxChildColumnSize = OnGetMaxRowSizeForColumn.Execute(Column.ColumnId, EOrientation::Orient_Horizontal);
+
+			return MaxChildColumnSize.Component(EOrientation::Orient_Horizontal) < HeaderSize.Component(EOrientation::Orient_Horizontal) ? HeaderSize : MaxChildColumnSize;
+		}
+	}
+
+	return FVector2D::ZeroVector;
+}
+
 void SHeaderRow::RegenerateWidgets()
 {
 	const float SplitterHandleDetectionSize = 5.0f;
+	HeaderWidgets.Empty();
+
 	TSharedPtr<SSplitter> Splitter;
 
 	TSharedRef< SHorizontalBox > Box = 
@@ -551,9 +585,10 @@ void SHeaderRow::RegenerateWidgets()
 		[
 			SAssignNew(Splitter, SSplitter)
 			.Style( &Style->ColumnSplitterStyle )
-			.ResizeMode( ESplitterResizeMode::Fill )
+			.ResizeMode(ResizeMode)
 			.PhysicalSplitterHandleSize( 0.0f )
 			.HitDetectionSplitterHandleSize( SplitterHandleDetectionSize )
+			.OnGetMaxSlotSize(this, &SHeaderRow::GetRowSizeForSlotIndex)
 		]
 
 		+SHorizontalBox::Slot()
@@ -587,6 +622,8 @@ void SHeaderRow::RegenerateWidgets()
 					SAssignNew(NewlyMadeHeader, STableColumnHeader, SomeColumn, DefaultPadding)
 					.Style((SlotIndex + 1 == Columns.Num()) ? &Style->LastColumnStyle : &Style->ColumnStyle);
 
+				HeaderWidgets.Add(NewlyMadeHeader);
+
 				switch (SomeColumn.SizeRule)
 				{
 				case EColumnSizeMode::Fill:
@@ -688,18 +725,18 @@ void SHeaderRow::RegenerateWidgets()
 						[
 							SNew(SBox)
 							.WidthOverride(WidthBinding)
-						[
-							SNew(SOverlay)
-							+ SOverlay::Slot()
-						[
-							NewHeader
-						]
-					+ SOverlay::Slot()
-						.HAlign(HAlign_Right)
-						[
-							SizingGrip
-						]
-						]
+							[
+								SNew(SOverlay)
+								+ SOverlay::Slot()
+								[
+									NewHeader
+								]
+								+ SOverlay::Slot()
+								.HAlign(HAlign_Right)
+								[
+									SizingGrip
+								]
+							]
 						];
 				}
 				break;
diff --git a/Engine/Source/Runtime/Slate/Public/Framework/Application/HardwareCursor.h b/Engine/Source/Runtime/Slate/Public/Framework/Application/HardwareCursor.h
new file mode 100644
--- /dev/null
+++ b/Engine/Source/Runtime/Slate/Public/Framework/Application/HardwareCursor.h
@@ -0,0 +1,73 @@
+// Copyright 1998-2017 Epic Games, Inc. All Rights Reserved.
+
+#pragma once
+
+#include "CoreMinimal.h"
+#include "ICursor.h"
+
+#if PLATFORM_WINDOWS
+#include "WindowsHWrapper.h"
+#endif
+
+struct SDL_Cursor;
+
+/**
+ * Provides a cross desktop platform solution for providing hardware cursors.  These cursors
+ * generally require platform specific formats.  To try and combat this, this class standardizes
+ * on .png files.  However, on different platforms that support it, it searches for platform specific
+ * formats first if you want to take advantage of those capabilities.
+ *
+ * Windows:
+ *   .ani -> .cur -> .png
+ *
+ * Mac:
+ *   .tiff -> .png
+ *
+ * Linux:
+ *   .png
+ *
+ * Additionally, you have the option to include different resolutions.  On Windows the typical size
+ * cursors are 32x32 or 64x64.  On Macs, a multi-resolution tiff can be provided.
+ *
+ * Mac: https://developer.apple.com/library/content/documentation/GraphicsAnimation/Conceptual/HighResolutionOSX/Optimizing/Optimizing.html
+ *
+ * In order to deal with this, your cursor path should include the base cursor with no size information
+ * however, you should name them like this, Pointer_32x32.png and Pointer_64x64.png.  This only applies for
+ * .png, since .cur|.ani|.tiff already support multiple resolutions.
+ *
+ * If no exact size is found, the closest cursor in size is used.
+ */
+class SLATE_API FHardwareCursor
+{
+public:
+	/**
+	 * HotSpot needs to be in normalized UV coordinates since it may apply to different resolution images. 
+	 */
+	FHardwareCursor(const FString& InPathToCursorWithoutExtension, FVector2D HotSpot);
+	FHardwareCursor(const TArray<FColor>& Pixels, FIntPoint Size, FVector2D HotSpot);
+
+	~FHardwareCursor();
+
+	/** Gets the platform specific handle to the cursor that was allocated.  If loading the cursor failed, this value will be null. */
+	void* GetHandle();
+
+	/**
+	 * Replaces the shape for the given cursor for this platform.  So if you wanted to replace the default
+	 * arrow cursor on the platform with this custom hardware cursor, you would pass in EMouseCursor::Default.
+	 */
+	void ReplaceTheShapeFor(EMouseCursor::Type InCursorType);
+
+private:
+	void CreateCursorFromRGBABuffer(const FColor* Pixels, int32 Width, int32 Height, FVector2D InHotSpot);
+
+	bool FindAndLoadCorrectResolutionPng(TArray<uint8>& Result, const FString& InPathToCursorWithoutExtension);
+
+private:
+#if PLATFORM_WINDOWS
+	HCURSOR CursorHandle;
+#elif PLATFORM_MAC
+	NSCursor* CursorHandle;
+#elif PLATFORM_LINUX
+	SDL_Cursor* CursorHandle;
+#endif
+};
diff --git a/Engine/Source/Runtime/Slate/Public/Framework/Application/SlateApplication.h b/Engine/Source/Runtime/Slate/Public/Framework/Application/SlateApplication.h
--- a/Engine/Source/Runtime/Slate/Public/Framework/Application/SlateApplication.h
+++ b/Engine/Source/Runtime/Slate/Public/Framework/Application/SlateApplication.h
@@ -338,21 +338,6 @@ public:
 	 */
 	TSharedRef<SWindow> AddWindowAsNativeChild( TSharedRef<SWindow> InSlateWindow, TSharedRef<SWindow> InParentWindow, const bool bShowImmediately = true );
 
-	/**
-	 * Creates a new Menu window and adds it to the menu stack.
-	 *
-	 * @param InWindow				The parent of the menu.  If there is already an open menu this parent must exist in the menu stack or the menu stack is dismissed and a new one started
-	 * @param Content				The content to be placed inside the new window
-	 * @param SummonLocation		The location where this menu should be summoned
-	 * @param TransitionEffect		Animation to use when the popup appears
-	 * @param bFocusImmediately		Should the popup steal focus when shown?
-	 * @param bShouldAutoSize		True if the newCalculatePopupWindowPosition window should automatically size itself to its content
-	 * @param WindowSize			When bShouldAutoSize=false, this must be set to the size of the window to be created
-	 * @param SummonLocationSize	An optional rect which describes an area in which the menu may not appear
-	 */
-	DEPRECATED(4.9, "PushMenu() returning a window is deprecated. Use the new version of PushMenu() that returns an IMenu.")
-	TSharedRef<SWindow> PushMenu( const TSharedRef<SWidget>& InParentContent, const TSharedRef<SWidget>& InContent, const FVector2D& SummonLocation, const FPopupTransitionEffect& TransitionEffect, const bool bFocusImmediately = true, const bool bShouldAutoSize = true, const FVector2D& WindowSize = FVector2D::ZeroVector, const FVector2D& SummonLocationSize = FVector2D::ZeroVector );
-
 	/**
 	 * Creates a new Menu and adds it to the menu stack.
 	 * Menus are always auto-sized. Use fixed-size content if a fixed size is required.
@@ -412,10 +397,6 @@ public:
 	 */	
 	TSharedPtr<IMenu> PushHostedMenu(const TSharedPtr<IMenu>& InParentMenu, const TSharedRef<IMenuHost>& InMenuHost, const TSharedRef<SWidget>& InContent, TSharedPtr<SWidget>& OutWrappedContent, const FPopupTransitionEffect& TransitionEffect, EShouldThrottle ShouldThrottle, const bool bIsCollapsedByParent = true);
 
-	/** @return Returns whether the window has child menus. */
-	DEPRECATED(4.9, "HasOpenSubMenus() taking a window is deprecated. Use HasOpenSubMenus() taking an IMenu as a parameter.")
-	bool HasOpenSubMenus( TSharedRef<SWindow> Window ) const;
-
 	/** @return Returns whether the menu has child menus. */
 	bool HasOpenSubMenus(TSharedPtr<IMenu> InMenu) const;
 
@@ -436,14 +417,6 @@ public:
 	/** Dismisses all open menus */
 	void DismissAllMenus();
 
-	/**
-	 * Dismisses a menu window and all its children
-	 *
-	 * @param MenuWindowToDismiss	The window to dismiss, any children, grandchildren etc will also be dismissed
-	 */
-	DEPRECATED(4.9, "DismissMenu() taking a window is deprecated. Use DismissMenu() taking an IMenu, DismissMenuByWidget() or DismissAll().")
-	void DismissMenu( TSharedRef<SWindow> MenuWindowToDismiss );
-
 	/**
 	 * Dismisses a menu and all its children
 	 *
@@ -458,15 +431,6 @@ public:
 	 */
 	void DismissMenuByWidget(const TSharedRef<SWidget>& InWidgetInMenu);
 
-	/**
-	 * Finds the window in the menu stack
-	 *
-	 * @param WindowToFind	The window to look for
-	 * @return The level in the stack  that the window is in or INDEX_NONE if it is not found
-	 */
-	DEPRECATED(4.9, "GetLocationInMenuStack() is deprecated. Shouldn't be needed.")
-	int32 GetLocationInMenuStack( TSharedRef<SWindow> WindowToFind ) const;
-
 	/**
 	 * HACK: Don't use this unless shutting down a game viewport
 	 * Game viewport windows need to be destroyed instantly or else the viewport could tick and access deleted data
@@ -553,17 +517,11 @@ public:
 	 */
 	void SetUserFocusToGameViewport(uint32 UserIndex, EFocusCause ReasonFocusIsChanging = EFocusCause::SetDirectly);
 
-	DEPRECATED(4.6, "FSlateApplication::SetFocusToGameViewport() is deprecated, use FSlateApplication::SetUserFocusToGameViewport() instead.")
-	void SetFocusToGameViewport();
-
 	/**
 	 * Sets all users focus to the SWidget representing the currently active game viewport
 	 */
 	void SetAllUserFocusToGameViewport(EFocusCause ReasonFocusIsChanging = EFocusCause::SetDirectly);
 
-	DEPRECATED(4.6, "FSlateApplication::SetJoystickCaptorToGameViewport() is deprecated, use FSlateApplication::SetAllUserFocusToGameViewport() instead.")
-	void SetJoystickCaptorToGameViewport();
-
 	/**
 	 * Activates the Game Viewport if it is properly childed under a window
 	 */
@@ -586,18 +544,12 @@ public:
 	 */
 	void SetAllUserFocus(const TSharedPtr<SWidget>& WidgetToFocus, EFocusCause ReasonFocusIsChanging = EFocusCause::SetDirectly);
 
-	DEPRECATED(4.6, "FSlateApplication::GetJoystickCaptor() is deprecated, use FSlateApplication::GetUserFocusedWidget() instead.")
-	TSharedPtr< SWidget > GetJoystickCaptor(uint32 UserIndex) const;
-
 	/** Releases the users focus from whatever it currently is on. */
 	void ClearUserFocus(uint32 UserIndex, EFocusCause ReasonFocusIsChanging = EFocusCause::SetDirectly);
 
 	/** Releases the focus for all users from whatever it currently is on. */
 	void ClearAllUserFocus(EFocusCause ReasonFocusIsChanging = EFocusCause::SetDirectly);
 
-	DEPRECATED(4.6, "FSlateApplication::ReleaseJoystickCapture() is deprecated, use FSlateApplication::ClearUserFocus() instead.")
-	void ReleaseJoystickCapture(uint32 UserIndex);
-
 	/**
 	 * Sets the Keyboard focus to the specified SWidget
 	 */
@@ -617,6 +569,14 @@ public:
 	*/
 	DECLARE_EVENT_OneParam(FSlateApplication, FOnApplicationPreInputKeyDownListener, const FKeyEvent&);
 	FOnApplicationPreInputKeyDownListener& OnApplicationPreInputKeyDownListener() { return OnApplicationPreInputKeyDownListenerEvent; }
+
+	/**
+	* Gets a delegate that is invoked before the mouse input button down get process by slate widgets bubble system.
+	* Its read only and you cannot mark the input as handled.
+	*/
+	DECLARE_EVENT_OneParam(FSlateApplication, FOnApplicationMousePreInputButtonDownListener, const FPointerEvent&);
+	FOnApplicationMousePreInputButtonDownListener& OnApplicationMousePreInputButtonDownListener() { return OnApplicationMousePreInputButtonDownListenerEvent; }
+
 #endif //WITH_EDITOR
 
 	/**
@@ -760,12 +720,6 @@ public:
 	/** Get the current drag-dropping content */
 	TSharedPtr<class FDragDropOperation> GetDragDroppingContent() const;
 
-	DEPRECATED(4.6, "Use CancelDragDrop")
-	void EndDragDrop()
-	{
-		CancelDragDrop();
-	}
-
 	/** Cancels any in flight drag and drops */
 	void CancelDragDrop();
 
@@ -1261,6 +1215,9 @@ public:
 	/** Getter for the cursor radius */
 	float GetCursorRadius() const;
 
+	/**  */
+	void RegisterCursor(EMouseCursor::Type CursorType, TSharedPtr<class FHardwareCursor> Cursor);
+
 public:
 
 	//~ Begin FSlateApplicationBase Interface
@@ -1739,6 +1696,9 @@ private:
 	/** The hit-test radius of the cursor. Default value is 0. */
 	float CursorRadius;
 
+	/**  */
+	TMap<EMouseCursor::Type, TSharedPtr<class FHardwareCursor>> HardwareCursors;
+
 	/**
 	 * All users currently registered with Slate.  Normally this is 1, but in a 
 	 * situation where multiple users are providing input you need to track ui state
@@ -2023,5 +1983,11 @@ private:
 	* User Function cannot mark the input as handled.
 	*/
 	FOnApplicationPreInputKeyDownListener OnApplicationPreInputKeyDownListenerEvent;
+
+	/**
+	* Delegate that is invoked before the mouse input button get process by slate widgets bubble system.
+	* User Function cannot mark the input as handled.
+	*/
+	FOnApplicationMousePreInputButtonDownListener OnApplicationMousePreInputButtonDownListenerEvent;
 #endif // WITH_EDITOR
 };
diff --git a/Engine/Source/Runtime/Slate/Public/Widgets/Input/SCheckBox.h b/Engine/Source/Runtime/Slate/Public/Widgets/Input/SCheckBox.h
--- a/Engine/Source/Runtime/Slate/Public/Widgets/Input/SCheckBox.h
+++ b/Engine/Source/Runtime/Slate/Public/Widgets/Input/SCheckBox.h
@@ -18,18 +18,6 @@
 
 class SBorder;
 
-/** DEPRECATED 4.6 - Do not use */
-//@Todo slate: Remove this as soon as the 4.6 deprecated API is Removed.
-namespace ESlateCheckBoxState
-{
-	DEPRECATED(4.6, "ESlateCheckBoxState::Type is deprecated and was renamed to ECheckBoxState. Please use that type instead.")
-	typedef ::ECheckBoxState Type;
-
-	const ECheckBoxState Unchecked = ECheckBoxState::Unchecked;
-	const ECheckBoxState Checked = ECheckBoxState::Checked;
-	const ECheckBoxState Undetermined = ECheckBoxState::Undetermined;
-}
-
 
 /** Delegate that is executed when the check box state changes */
 DECLARE_DELEGATE_OneParam( FOnCheckStateChanged, ECheckBoxState );
@@ -80,22 +68,6 @@ public:
 		/** Whether the check box is currently in a checked state */
 		SLATE_ATTRIBUTE( ECheckBoxState, IsChecked )
 
-		// TODO Remove this function when IsChecked(bool InIsChecked) is removed.  It has to be here to prevent ambiguous conversions.
-		/** Whether the check box is currently in a checked state */
-		FArguments& IsChecked(ECheckBoxState InIsChecked)
-		{
-			_IsChecked = InIsChecked;
-			return Me();
-		}
-
-		/** Whether the check box is currently in a checked state */
-		DEPRECATED(4.3, "Please use IsChecked(TAttribute<ECheckBoxState>)")
-		FArguments& IsChecked(bool InIsChecked)
-		{
-			_IsChecked = InIsChecked ? ECheckBoxState::Checked : ECheckBoxState::Unchecked;
-			return Me();
-		}
-
 		/** How the content of the toggle button should align within the given space*/
 		SLATE_ARGUMENT( EHorizontalAlignment, HAlign )
 
diff --git a/Engine/Source/Runtime/Slate/Public/Widgets/Input/SEditableText.h b/Engine/Source/Runtime/Slate/Public/Widgets/Input/SEditableText.h
--- a/Engine/Source/Runtime/Slate/Public/Widgets/Input/SEditableText.h
+++ b/Engine/Source/Runtime/Slate/Public/Widgets/Input/SEditableText.h
@@ -316,6 +316,7 @@ protected:
 	virtual void EnsureActiveTick() override;
 	virtual EKeyboardType GetVirtualKeyboardType() const override;
 	virtual TSharedRef<SWidget> GetSlateWidget() override;
+	virtual TSharedPtr<SWidget> GetSlateWidgetPtr() override;
 	virtual TSharedPtr<SWidget> BuildContextMenuContent() const override;
 	virtual void OnTextChanged(const FText& InText) override;
 	virtual void OnTextCommitted(const FText& InText, const ETextCommit::Type InTextAction) override;
diff --git a/Engine/Source/Runtime/Slate/Public/Widgets/Input/SNumericEntryBox.h b/Engine/Source/Runtime/Slate/Public/Widgets/Input/SNumericEntryBox.h
--- a/Engine/Source/Runtime/Slate/Public/Widgets/Input/SNumericEntryBox.h
+++ b/Engine/Source/Runtime/Slate/Public/Widgets/Input/SNumericEntryBox.h
@@ -24,6 +24,7 @@
 #include "Widgets/Input/SEditableText.h"
 #include "Widgets/Input/SSpinBox.h"
 #include "Widgets/Colors/SColorBlock.h"
+#include "UnrealString.h"
 
 /**
  * Implementation for a box that only accepts a numeric value or that can display an undetermined value via a string
@@ -479,7 +480,7 @@ private:
 		else
 		{
 			NumericType NumericValue;
-			if (Lex::TryParseString(NumericValue, *NewValue.ToString()))
+			if (LexicalConversion::TryParseString(NumericValue, *NewValue.ToString()))
 			{
 				OnValueChanged.ExecuteIfBound( NumericValue );
 			}
diff --git a/Engine/Source/Runtime/Slate/Public/Widgets/Input/SSpinBox.h b/Engine/Source/Runtime/Slate/Public/Widgets/Input/SSpinBox.h
--- a/Engine/Source/Runtime/Slate/Public/Widgets/Input/SSpinBox.h
+++ b/Engine/Source/Runtime/Slate/Public/Widgets/Input/SSpinBox.h
@@ -242,7 +242,7 @@ public:
 		const int32 BackgroundLayer = LayerId;
 
 		const bool bEnabled = ShouldBeEnabled( bParentEnabled );
-		const ESlateDrawEffect::Type DrawEffects = bEnabled ? ESlateDrawEffect::None : ESlateDrawEffect::DisabledEffect;
+		const ESlateDrawEffect DrawEffects = bEnabled ? ESlateDrawEffect::None : ESlateDrawEffect::DisabledEffect;
 
 		FSlateDrawElement::MakeBox(
 			OutDrawElements,
@@ -337,6 +337,12 @@ public:
 		{
 			if( bDragging )
 			{
+				NumericType CurrentDelta = Delta.Get();
+				if (CurrentDelta != 0)
+				{
+					InternalValue = Snap(InternalValue, CurrentDelta);
+				}
+
 				NotifyValueCommitted();
 			}
 
@@ -613,14 +619,7 @@ protected:
 	/** @return the value being observed by the spinbox as a string */
 	FString GetValueAsString() const
 	{
-		auto CurrentValue = ValueAttribute.Get();
-		NumericType CurrentDelta = Delta.Get();
-		if( CurrentDelta != 0 )
-		{
-			CurrentValue = (NumericType)Snap(CurrentValue, CurrentDelta);
-		}
-		
-		return Interface->ToString(CurrentValue);
+		return Interface->ToString(ValueAttribute.Get());
 	}
 
 	/** @return the value being observed by the spinbox as FText - todo: spinbox FText support (reimplement me) */
diff --git a/Engine/Source/Runtime/Slate/Public/Widgets/Layout/SSplitter.h b/Engine/Source/Runtime/Slate/Public/Widgets/Layout/SSplitter.h
--- a/Engine/Source/Runtime/Slate/Public/Widgets/Layout/SSplitter.h
+++ b/Engine/Source/Runtime/Slate/Public/Widgets/Layout/SSplitter.h
@@ -25,7 +25,8 @@ namespace ESplitterResizeMode
 {
 	enum Type
 	{
-		Fixed,		
+		FixedPosition,
+		FixedSize,
 		Fill,
 	};
 }
@@ -53,6 +54,8 @@ public:
 		/** The new size coefficient of the slot */
 		float );
 
+	DECLARE_DELEGATE_RetVal_OneParam(FVector2D, FOnGetMaxSlotSize, int32);
+
 public:
 	class FSlot : public TSlotBase<FSlot>
 	{
@@ -106,7 +109,7 @@ public:
 	SLATE_BEGIN_ARGS(SSplitter)
 		: _Style( &FCoreStyle::Get().GetWidgetStyle<FSplitterStyle>("Splitter") )
 		, _Orientation( Orient_Horizontal )
-		, _ResizeMode( ESplitterResizeMode::Fixed )
+		, _ResizeMode( ESplitterResizeMode::FixedPosition )
 		, _PhysicalSplitterHandleSize( 5.0f )
 		, _HitDetectionSplitterHandleSize( 5.0f )
 		, _OnSplitterFinishedResizing()
@@ -126,6 +129,8 @@ public:
 		SLATE_ARGUMENT( float, HitDetectionSplitterHandleSize )
 
 		SLATE_EVENT( FSimpleDelegate, OnSplitterFinishedResizing )
+		
+		SLATE_EVENT(FOnGetMaxSlotSize, OnGetMaxSlotSize)
 
 	SLATE_END_ARGS()
 
@@ -199,6 +204,8 @@ public:
 	 * @return Whether the event was handled along with possible requests for the system to take action.
 	 */
 	virtual FReply OnMouseButtonUp( const FGeometry& MyGeometry, const FPointerEvent& MouseEvent ) override;
+
+	virtual FReply OnMouseButtonDoubleClick(const FGeometry& InMyGeometry, const FPointerEvent& InMouseEvent) override;
 	
 	virtual FReply OnMouseMove( const FGeometry& MyGeometry, const FPointerEvent& MouseEvent ) override;
 
@@ -253,8 +260,9 @@ protected:
 	 * @param Children         A reference to this splitter's children array; we will modify the children's layout values.
 	 * @param ChildGeometries  The arranged children; we need their sizes and positions so that we can perform a resizing.
 	 */
-	template<EOrientation SplitterOrientation>
-	static void HandleResizing( const float PhysicalSplitterHandleSize, const ESplitterResizeMode::Type ResizeMode, int32 DraggedHandle, const FVector2D& LocalMousePos, TPanelChildren<FSlot>& Children, const TArray<FLayoutGeometry>& ChildGeometries );
+	static void HandleResizingByMousePosition(EOrientation Orientation, const float PhysicalSplitterHandleSize, const ESplitterResizeMode::Type ResizeMode, int32 DraggedHandle, const FVector2D& LocalMousePos, TPanelChildren<FSlot>& Children, const TArray<FLayoutGeometry>& ChildGeometries );
+	static void HandleResizingDelta(EOrientation Orientation, const float PhysicalSplitterHandleSize, const ESplitterResizeMode::Type ResizeMode, int32 DraggedHandle, float Delta, TPanelChildren<FSlot>& Children, const TArray<FLayoutGeometry>& ChildGeometries);
+	static void HandleResizingBySize(EOrientation Orientation, const float PhysicalSplitterHandleSize, const ESplitterResizeMode::Type ResizeMode, int32 DraggedHandle, const FVector2D& DesiredSize, TPanelChildren<FSlot>& Children, const TArray<FLayoutGeometry>& ChildGeometries);
 
 	/**
 	 * @param ProposedSize  A size that a child would like to be
@@ -283,6 +291,7 @@ protected:
 	ESplitterResizeMode::Type ResizeMode;
 
 	FSimpleDelegate OnSplitterFinishedResizing;
+	FOnGetMaxSlotSize OnGetMaxSlotSize;
 
 	/** The thickness of the grip area that the user uses to resize a splitter */
 	float PhysicalSplitterHandleSize;
diff --git a/Engine/Source/Runtime/Slate/Public/Widgets/Text/ISlateEditableTextWidget.h b/Engine/Source/Runtime/Slate/Public/Widgets/Text/ISlateEditableTextWidget.h
--- a/Engine/Source/Runtime/Slate/Public/Widgets/Text/ISlateEditableTextWidget.h
+++ b/Engine/Source/Runtime/Slate/Public/Widgets/Text/ISlateEditableTextWidget.h
@@ -221,9 +221,12 @@ public:
 	/** Get the type of virtual keyboard to use for this widget */
 	virtual EKeyboardType GetVirtualKeyboardType() const = 0;
 
-	/** Get our the Slate widget this interface is representing */
+	/** Get the Slate widget this interface is representing (may not be called during destruction) */
 	virtual TSharedRef<SWidget> GetSlateWidget() = 0;
 
+	/** Get the Slate widget this interface is representing (may be null during destruction) */
+	virtual TSharedPtr<SWidget> GetSlateWidgetPtr() = 0;
+
 	/** Build the context menu content to use for this widget (if any) */
 	virtual TSharedPtr<SWidget> BuildContextMenuContent() const = 0;
 
diff --git a/Engine/Source/Runtime/Slate/Public/Widgets/Text/SMultiLineEditableText.h b/Engine/Source/Runtime/Slate/Public/Widgets/Text/SMultiLineEditableText.h
--- a/Engine/Source/Runtime/Slate/Public/Widgets/Text/SMultiLineEditableText.h
+++ b/Engine/Source/Runtime/Slate/Public/Widgets/Text/SMultiLineEditableText.h
@@ -338,6 +338,7 @@ protected:
 	virtual void EnsureActiveTick() override;
 	virtual EKeyboardType GetVirtualKeyboardType() const override;
 	virtual TSharedRef<SWidget> GetSlateWidget() override;
+	virtual TSharedPtr<SWidget> GetSlateWidgetPtr() override;
 	virtual TSharedPtr<SWidget> BuildContextMenuContent() const override;
 	virtual void OnTextChanged(const FText& InText) override;
 	virtual void OnTextCommitted(const FText& InText, const ETextCommit::Type InTextAction) override;
diff --git a/Engine/Source/Runtime/Slate/Public/Widgets/Views/SHeaderRow.h b/Engine/Source/Runtime/Slate/Public/Widgets/Views/SHeaderRow.h
--- a/Engine/Source/Runtime/Slate/Public/Widgets/Views/SHeaderRow.h
+++ b/Engine/Source/Runtime/Slate/Public/Widgets/Views/SHeaderRow.h
@@ -12,6 +12,7 @@
 #include "Styling/CoreStyle.h"
 #include "Widgets/Layout/SBorder.h"
 #include "Styling/SlateWidgetStyleAsset.h"
+#include "SSplitter.h"
 
 class SScrollBar;
 
@@ -77,6 +78,8 @@ DECLARE_DELEGATE_ThreeParams( FOnSortModeChanged, EColumnSortPriority::Type, con
 /**	Callback when the width of the column changes */
 DECLARE_DELEGATE_OneParam( FOnWidthChanged, float );
 
+/**	Callback to fetch the max row width for a specified column id */
+DECLARE_DELEGATE_RetVal_TwoParams(FVector2D, FOnGetMaxRowSizeForColumn, const FName&, EOrientation);
 
 /**
  * The header that appears above lists and trees when they are showing multiple columns.
@@ -242,10 +245,14 @@ public:
 
 	SLATE_BEGIN_ARGS(SHeaderRow)
 		: _Style( &FCoreStyle::Get().GetWidgetStyle<FHeaderRowStyle>("TableView.Header") )
+		, _ResizeMode(ESplitterResizeMode::Fill)
 		{}
+
 		SLATE_STYLE_ARGUMENT( FHeaderRowStyle, Style )
 		SLATE_SUPPORTS_SLOT_WITH_ARGS( FColumn )	
 		SLATE_EVENT( FColumnsChanged::FDelegate, OnColumnsChanged )
+		SLATE_EVENT(FOnGetMaxRowSizeForColumn, OnGetMaxRowSizeForColumn)
+		SLATE_ARGUMENT(ESplitterResizeMode::Type, ResizeMode)
 	SLATE_END_ARGS()
 
 	void Construct( const FArguments& InArgs );
@@ -278,6 +285,12 @@ public:
 	/** Sets the column, with the specified name, to the desired width */
 	void SetColumnWidth( const FName& InColumnId, float InWidth );
 
+	/** Will return the size for this row at the specified slot index */
+	FVector2D GetRowSizeForSlotIndex(int32 SlotIndex) const;
+
+	/** Simple function to set the delegate to fetch the max row size for column id */
+	void SetOnGetMaxRowSizeForColumn(const FOnGetMaxRowSizeForColumn& Delegate) { OnGetMaxRowSizeForColumn = Delegate; }
+
 private:
 
 	/** Regenerates all widgets in the header */
@@ -285,9 +298,12 @@ private:
 
 	/** Information about the various columns */
 	TIndirectArray<FColumn> Columns;
+	TArray<TSharedPtr<class STableColumnHeader>> HeaderWidgets;
 
 	FVector2D ScrollBarThickness;
 	TAttribute< EVisibility > ScrollBarVisibility;
 	const FHeaderRowStyle* Style;
 	FColumnsChanged ColumnsChanged;
+	ESplitterResizeMode::Type ResizeMode;
+	FOnGetMaxRowSizeForColumn OnGetMaxRowSizeForColumn;
 };
diff --git a/Engine/Source/Runtime/Slate/Public/Widgets/Views/SListView.h b/Engine/Source/Runtime/Slate/Public/Widgets/Views/SListView.h
--- a/Engine/Source/Runtime/Slate/Public/Widgets/Views/SListView.h
+++ b/Engine/Source/Runtime/Slate/Public/Widgets/Views/SListView.h
@@ -1226,6 +1226,33 @@ public:
 		TListTypeTraits<ItemType>::AddReferencedObjects( Collector, WidgetGenerator.ItemsWithGeneratedWidgets, SelectedItems );
 	}
 
+	/**
+	* Will determine the max row size for the specified column id
+	*
+	* @param ColumnId  Column Id
+	* @param Orientation  Orientation that is main axis you want to query
+	*
+	* @return The max size for a column Id.
+	*/
+	FVector2D GetMaxRowSizeForColumn(const FName& ColumnId, EOrientation Orientation)
+	{
+		FVector2D MaxSize = FVector2D::ZeroVector;
+
+		for (auto It = WidgetGenerator.WidgetMapToItem.CreateConstIterator(); It; ++It)
+		{
+			const ITableRow* TableRow = It.Key();
+			FVector2D NewMaxSize = TableRow->GetRowSizeForColumn(ColumnId);
+
+			// We'll return the full size, but we only take into consideration the asked axis for the calculation of the size
+			if (NewMaxSize.Component(Orientation) > MaxSize.Component(Orientation))
+			{
+				MaxSize = NewMaxSize;
+			}
+		}
+
+		return MaxSize;
+	}
+
 protected:
 
 	FOnItemToString_Debug GetDefaultDebugDelegate()
diff --git a/Engine/Source/Runtime/Slate/Public/Widgets/Views/STableRow.h b/Engine/Source/Runtime/Slate/Public/Widgets/Views/STableRow.h
--- a/Engine/Source/Runtime/Slate/Public/Widgets/Views/STableRow.h
+++ b/Engine/Source/Runtime/Slate/Public/Widgets/Views/STableRow.h
@@ -61,6 +61,9 @@ class SLATE_API ITableRow
 		/** Called when the expander arrow for this row is shift+clicked */
 		virtual void Private_OnExpanderArrowShiftClicked() = 0;
 
+		/** @return the size for the specified column name */
+		virtual FVector2D GetRowSizeForColumn(const FName& InColumnName) const = 0;
+
 	protected:
 		/** Called to query the selection mode for the row */
 		virtual ESelectionMode::Type GetSelectionMode() const = 0;
@@ -877,6 +880,12 @@ public:
 		return OwnerWidget->Private_IsItemSelected(*MyItem);
 	}
 
+	/** By default, this function does nothing, it should be implemented by derived class */
+	virtual FVector2D GetRowSizeForColumn(const FName& InColumnName) const override
+	{
+		return FVector2D::ZeroVector;
+	}
+
 	/** Protected constructor; SWidgets should only be instantiated via declarative syntax. */
 	STableRow()
 		: IndexInList(0)
@@ -1169,6 +1178,11 @@ protected:
 		ColumnIdToSlotContents.Empty();
 	}
 
+	const TSharedRef<SWidget>* GetWidgetFromColumnId(const FName& ColumnId) const
+	{
+		return ColumnIdToSlotContents.Find(ColumnId);
+	}
+
 private:
 	
 	TSharedPtr<SHorizontalBox> Box;
diff --git a/Engine/Source/Runtime/Slate/Slate.Build.cs b/Engine/Source/Runtime/Slate/Slate.Build.cs
--- a/Engine/Source/Runtime/Slate/Slate.Build.cs
+++ b/Engine/Source/Runtime/Slate/Slate.Build.cs
@@ -15,7 +15,8 @@ public class Slate : ModuleRules
 				"InputCore",
 				"Json",
 				"SlateCore",
-			}
+                "ImageWrapper"
+            }
 		);
 
 		PrivateIncludePaths.AddRange(
@@ -55,5 +56,9 @@ public class Slate : ModuleRules
 		{
 			AddEngineThirdPartyPrivateStaticDependencies(Target, "XInput");
 		}
-	}
+        else if (Target.Platform == UnrealTargetPlatform.Linux)
+        {
+            AddEngineThirdPartyPrivateStaticDependencies(Target, "SDL2");
+        }
+    }
 }
diff --git a/Engine/Source/Runtime/SlateCore/Private/Fonts/FontCacheHarfBuzz.cpp b/Engine/Source/Runtime/SlateCore/Private/Fonts/FontCacheHarfBuzz.cpp
--- a/Engine/Source/Runtime/SlateCore/Private/Fonts/FontCacheHarfBuzz.cpp
+++ b/Engine/Source/Runtime/SlateCore/Private/Fonts/FontCacheHarfBuzz.cpp
@@ -3,6 +3,7 @@
 #include "Fonts/FontCacheHarfBuzz.h"
 #include "Fonts/FontCache.h"
 #include "Fonts/FontCacheFreeType.h"
+#include "Fonts/SlateFontRenderer.h"
 
 #if WITH_HARFBUZZ
 	#if PLATFORM_WINDOWS
@@ -166,6 +167,21 @@ FORCEINLINE int32 get_ft_flags(hb_font_t* InFont)
 	return hb_ft_font_get_load_flags(InFont->parent);
 }
 
+hb_bool_t get_nominal_glyph(hb_font_t* InFont, void* InFontData, hb_codepoint_t InUnicodeChar, hb_codepoint_t *OutGlyphIndex, void *InUserData)
+{
+	FT_Face FreeTypeFace = get_ft_face(InFont);
+
+	*OutGlyphIndex = FT_Get_Char_Index(FreeTypeFace, InUnicodeChar);
+
+	// If the given font can't render that character (as the fallback font may be missing), try again with the fallback character
+	if (InUnicodeChar != 0 && *OutGlyphIndex == 0)
+	{
+		*OutGlyphIndex = FT_Get_Char_Index(FreeTypeFace, SlateFontRendererUtils::InvalidSubChar);
+	}
+
+	return InUnicodeChar == 0 || *OutGlyphIndex != 0;
+}
+
 hb_position_t get_glyph_h_advance(hb_font_t* InFont, void* InFontData, hb_codepoint_t InGlyphIndex, void* InUserData)
 {
 	FT_Face FreeTypeFace = get_ft_face(InFont);
@@ -297,6 +313,7 @@ FHarfBuzzFontFactory::FHarfBuzzFontFactory(FFreeTypeGlyphCache* InFTGlyphCache,
 #if WITH_HARFBUZZ
 	CustomHarfBuzzFuncs = hb_font_funcs_create();
 
+	hb_font_funcs_set_nominal_glyph_func(CustomHarfBuzzFuncs, &HarfBuzzFontFunctions::Internal::get_nominal_glyph, nullptr, nullptr);
 	hb_font_funcs_set_glyph_h_advance_func(CustomHarfBuzzFuncs, &HarfBuzzFontFunctions::Internal::get_glyph_h_advance, nullptr, nullptr);
 	hb_font_funcs_set_glyph_v_advance_func(CustomHarfBuzzFuncs, &HarfBuzzFontFunctions::Internal::get_glyph_v_advance, nullptr, nullptr);
 	hb_font_funcs_set_glyph_v_origin_func(CustomHarfBuzzFuncs, &HarfBuzzFontFunctions::Internal::get_glyph_v_origin, nullptr, nullptr);
diff --git a/Engine/Source/Runtime/SlateCore/Private/Fonts/LegacySlateFontInfoCache.cpp b/Engine/Source/Runtime/SlateCore/Private/Fonts/LegacySlateFontInfoCache.cpp
--- a/Engine/Source/Runtime/SlateCore/Private/Fonts/LegacySlateFontInfoCache.cpp
+++ b/Engine/Source/Runtime/SlateCore/Private/Fonts/LegacySlateFontInfoCache.cpp
@@ -24,6 +24,8 @@ FLegacySlateFontInfoCache::FLegacySlateFontInfoCache()
 	, LocalizedFallbackFontDataHistoryVersion(0)
 	, LocalizedFallbackFontFrameCounter(0)
 {
+	LastResortFontPath = FPaths::EngineContentDir() / TEXT("SlateDebug/Fonts/LastResort.ttf");
+	bIsLastResortFontAvailable = FPaths::FileExists(LastResortFontPath);
 }
 
 TSharedPtr<const FCompositeFont> FLegacySlateFontInfoCache::GetCompositeFont(const FName& InLegacyFontName, const EFontHinting InLegacyFontHinting)
@@ -61,9 +63,6 @@ TSharedPtr<const FCompositeFont> FLegacySlateFontInfoCache::GetCompositeFont(con
 		}
 	}
 
-	// Don't allow GC while we perform this allocation
-	FGCScopeGuard GCGuard;
-
 	TSharedRef<const FCompositeFont> NewCompositeFont = MakeShareable(new FStandaloneCompositeFont(NAME_None, LegacyFontPath, InLegacyFontHinting, EFontLoadingPolicy::LazyLoad));
 	LegacyFontNameToCompositeFont.Add(LegacyFontKey, NewCompositeFont);
 	return NewCompositeFont;
@@ -73,9 +72,6 @@ TSharedPtr<const FCompositeFont> FLegacySlateFontInfoCache::GetSystemFont()
 {
 	if (!SystemFont.IsValid())
 	{
-		// Don't allow GC while we perform this allocation
-		FGCScopeGuard GCGuard;
-
 		const TArray<uint8> FontBytes = FPlatformMisc::GetSystemFontBytes();
 		if (FontBytes.Num() > 0)
 		{
@@ -103,9 +99,6 @@ const FFontData& FLegacySlateFontInfoCache::GetLocalizedFallbackFontData()
 	// By only allowing it to update once per-frame, we ensure that the font cache has been flushed (which happens at the end of the frame) before we return a new font
 	if (!LocalizedFallbackFontData.IsValid() || (LocalizedFallbackFontDataHistoryVersion != CurrentHistoryVersion && LocalizedFallbackFontFrameCounter != CurrentFrameCounter))
 	{
-		// Don't allow GC while we perform this allocation
-		FGCScopeGuard GCGuard;
-
 		LocalizedFallbackFontDataHistoryVersion = CurrentHistoryVersion;
 		LocalizedFallbackFontFrameCounter = CurrentFrameCounter;
 
@@ -135,12 +128,17 @@ uint16 FLegacySlateFontInfoCache::GetLocalizedFallbackFontRevision() const
 	return LocalizedFallbackFontRevision;
 }
 
+bool FLegacySlateFontInfoCache::IsLastResortFontAvailable() const
+{
+	return bIsLastResortFontAvailable;
+}
+
 TSharedPtr<const FCompositeFont> FLegacySlateFontInfoCache::GetLastResortFont()
 {
 	// GetLastResortFont may be called from multiple threads at once
 	FScopeLock Lock(&LastResortFontCS);
 
-	if (!LastResortFont.IsValid())
+	if (!LastResortFont.IsValid() && bIsLastResortFontAvailable)
 	{
 		const FFontData& FontData = GetLastResortFontData();
 		LastResortFont = MakeShareable(new FStandaloneCompositeFont(NAME_None, FontData.GetFontFilename(), FontData.GetHinting(), FontData.GetLoadingPolicy()));
@@ -156,11 +154,7 @@ const FFontData& FLegacySlateFontInfoCache::GetLastResortFontData()
 
 	if (!LastResortFontData.IsValid())
 	{
-		// Don't allow GC while we perform this allocation
-		FGCScopeGuard GCGuard;
-
-		const FString LastResortFontPath = FPaths::EngineContentDir() / TEXT("SlateDebug/Fonts/LastResort.ttf");
-		LastResortFontData = MakeShareable(new FFontData(LastResortFontPath, EFontHinting::Default, EFontLoadingPolicy::LazyLoad));
+		LastResortFontData = MakeShareable(new FFontData(bIsLastResortFontAvailable ? LastResortFontPath : FString(), EFontHinting::Default, EFontLoadingPolicy::LazyLoad));
 	}
 
 	return *LastResortFontData;
diff --git a/Engine/Source/Runtime/SlateCore/Private/Fonts/LegacySlateFontInfoCache.h b/Engine/Source/Runtime/SlateCore/Private/Fonts/LegacySlateFontInfoCache.h
--- a/Engine/Source/Runtime/SlateCore/Private/Fonts/LegacySlateFontInfoCache.h
+++ b/Engine/Source/Runtime/SlateCore/Private/Fonts/LegacySlateFontInfoCache.h
@@ -37,6 +37,11 @@ public:
 	 */
 	uint16 GetLocalizedFallbackFontRevision() const;
 
+	/**
+	 * Is that last resort fallback font available (not all builds have it).
+	 */
+	bool IsLastResortFontAvailable() const;
+
 	/**
 	 * Get (or create) the last resort fallback font
 	 */
@@ -110,5 +115,8 @@ private:
 	FCriticalSection LastResortFontCS;
 	FCriticalSection LastResortFontDataCS;
 
+	FString LastResortFontPath;
+	bool bIsLastResortFontAvailable;
+
 	static TSharedPtr<FLegacySlateFontInfoCache> Singleton;
 };
diff --git a/Engine/Source/Runtime/SlateCore/Private/Fonts/SlateFontRenderer.cpp b/Engine/Source/Runtime/SlateCore/Private/Fonts/SlateFontRenderer.cpp
--- a/Engine/Source/Runtime/SlateCore/Private/Fonts/SlateFontRenderer.cpp
+++ b/Engine/Source/Runtime/SlateCore/Private/Fonts/SlateFontRenderer.cpp
@@ -189,58 +189,68 @@ bool FSlateFontRenderer::CanLoadCharacter(const FFontData& InFontData, TCHAR Cha
 FFreeTypeFaceGlyphData FSlateFontRenderer::GetFontFaceForCharacter(const FFontData& InFontData, TCHAR Char, EFontFallback MaxFallbackLevel) const 
 {
 	FFreeTypeFaceGlyphData ReturnVal;
-
-	TSharedPtr<FFreeTypeFace> FaceAndMemory = CompositeFontCache->GetFontFace(InFontData);
-	uint32 GlyphIndex = 0;
 	const bool bOverrideFallback = Char == SlateFontRendererUtils::InvalidSubChar;
 
-	if (FaceAndMemory.IsValid())
+	// Try the requested font first
 	{
-		// Get the index to the glyph in the font face
-		GlyphIndex = FT_Get_Char_Index(FaceAndMemory->GetFace(), Char);
-		ReturnVal.CharFallbackLevel = EFontFallback::FF_NoFallback;
+		ReturnVal.FaceAndMemory = CompositeFontCache->GetFontFace(InFontData);
+
+		if (ReturnVal.FaceAndMemory.IsValid())
+		{
+			ReturnVal.GlyphIndex = FT_Get_Char_Index(ReturnVal.FaceAndMemory->GetFace(), Char);
+			ReturnVal.CharFallbackLevel = EFontFallback::FF_NoFallback;
+		}
 	}
 
 	// If the requested glyph doesn't exist, use the localization fallback font
-	if (!FaceAndMemory.IsValid() || (Char != 0 && GlyphIndex == 0))
+	if (Char != 0 && ReturnVal.GlyphIndex == 0)
 	{
 		const bool bCanFallback = bOverrideFallback || MaxFallbackLevel >= EFontFallback::FF_LocalizedFallback;
 
 		if (bCanFallback)
 		{
-			FaceAndMemory = CompositeFontCache->GetFontFace(FLegacySlateFontInfoCache::Get().GetLocalizedFallbackFontData());
+			ReturnVal.FaceAndMemory = CompositeFontCache->GetFontFace(FLegacySlateFontInfoCache::Get().GetLocalizedFallbackFontData());
 
-			if (FaceAndMemory.IsValid())
+			if (ReturnVal.FaceAndMemory.IsValid())
 			{	
-				GlyphIndex = FT_Get_Char_Index(FaceAndMemory->GetFace(), Char);
+				ReturnVal.GlyphIndex = FT_Get_Char_Index(ReturnVal.FaceAndMemory->GetFace(), Char);
 
-				ReturnVal.CharFallbackLevel = EFontFallback::FF_LocalizedFallback;
-				ReturnVal.GlyphFlags |= FT_LOAD_FORCE_AUTOHINT;
+				if (ReturnVal.GlyphIndex != 0)
+				{
+					ReturnVal.CharFallbackLevel = EFontFallback::FF_LocalizedFallback;
+					ReturnVal.GlyphFlags |= FT_LOAD_FORCE_AUTOHINT;
+				}
 			}
 		}
 	}
 
 	// If the requested glyph doesn't exist, use the last resort fallback font
-	if (!FaceAndMemory.IsValid() || (Char != 0 && GlyphIndex == 0))
+	if (Char != 0 && ReturnVal.GlyphIndex == 0)
 	{
-		bool bCanFallback = bOverrideFallback || MaxFallbackLevel >= EFontFallback::FF_LastResortFallback;
+		const bool bCanFallback = bOverrideFallback || MaxFallbackLevel >= EFontFallback::FF_LastResortFallback;
 
-		if (bCanFallback)
+		if (bCanFallback && FLegacySlateFontInfoCache::Get().IsLastResortFontAvailable())
 		{
-			FaceAndMemory = CompositeFontCache->GetFontFace(FLegacySlateFontInfoCache::Get().GetLastResortFontData());
+			ReturnVal.FaceAndMemory = CompositeFontCache->GetFontFace(FLegacySlateFontInfoCache::Get().GetLastResortFontData());
 
-			if (FaceAndMemory.IsValid())
+			if (ReturnVal.FaceAndMemory.IsValid())
 			{
-				GlyphIndex = FT_Get_Char_Index(FaceAndMemory->GetFace(), Char);
+				ReturnVal.GlyphIndex = FT_Get_Char_Index(ReturnVal.FaceAndMemory->GetFace(), Char);
 
-				ReturnVal.CharFallbackLevel = EFontFallback::FF_LastResortFallback;
-				ReturnVal.GlyphFlags |= FT_LOAD_FORCE_AUTOHINT;
+				if (ReturnVal.GlyphIndex != 0)
+				{
+					ReturnVal.CharFallbackLevel = EFontFallback::FF_LastResortFallback;
+					ReturnVal.GlyphFlags |= FT_LOAD_FORCE_AUTOHINT;
+				}
 			}
 		}
 	}
 
-	ReturnVal.FaceAndMemory = FaceAndMemory;
-	ReturnVal.GlyphIndex = GlyphIndex;
+	// Found an invalid glyph?
+	if (Char != 0 && ReturnVal.GlyphIndex == 0)
+	{
+		ReturnVal.FaceAndMemory.Reset();
+	}
 
 	return ReturnVal;
 }
diff --git a/Engine/Source/Runtime/SlateCore/Private/Fonts/SlateFontRenderer.h b/Engine/Source/Runtime/SlateCore/Private/Fonts/SlateFontRenderer.h
--- a/Engine/Source/Runtime/SlateCore/Private/Fonts/SlateFontRenderer.h
+++ b/Engine/Source/Runtime/SlateCore/Private/Fonts/SlateFontRenderer.h
@@ -38,7 +38,7 @@ namespace SlateFontRendererUtils
 {
 
 /** Character used to substitute invalid font characters */
-const TCHAR InvalidSubChar = '\u001A';
+const TCHAR InvalidSubChar = TEXT('\uFFFD');
 
 #if WITH_FREETYPE
 
diff --git a/Engine/Source/Runtime/SlateCore/Private/Fonts/SlateTextShaper.cpp b/Engine/Source/Runtime/SlateCore/Private/Fonts/SlateTextShaper.cpp
--- a/Engine/Source/Runtime/SlateCore/Private/Fonts/SlateTextShaper.cpp
+++ b/Engine/Source/Runtime/SlateCore/Private/Fonts/SlateTextShaper.cpp
@@ -196,16 +196,24 @@ void FSlateTextShaper::PerformKerningOnlyTextShaping(const TCHAR* InText, const
 		{
 			const TCHAR CurrentChar = InText[RunningTextIndex];
 
+			// First try with the actual character
 			float SubFontScalingFactor = 1.0f;
-			const FFontData& FontData = CompositeFontCache->GetFontDataForCharacter(InFontInfo, CurrentChar, SubFontScalingFactor);
-			const FFreeTypeFaceGlyphData FaceGlyphData = FontRenderer->GetFontFaceForCharacter(FontData, CurrentChar, InFontInfo.FontFallback);
+			const FFontData* FontDataPtr = &CompositeFontCache->GetFontDataForCharacter(InFontInfo, CurrentChar, SubFontScalingFactor);
+			FFreeTypeFaceGlyphData FaceGlyphData = FontRenderer->GetFontFaceForCharacter(*FontDataPtr, CurrentChar, InFontInfo.FontFallback);
 
-			if (!RunningFontDataPtr || RunningFontDataPtr != &FontData || RunningFaceAndMemory != FaceGlyphData.FaceAndMemory || RunningSubFontScalingFactor != SubFontScalingFactor)
+			// If none of our fonts can render that character (as the fallback font may be missing), try again with the fallback character
+			if (!FaceGlyphData.FaceAndMemory.IsValid())
+			{
+				FontDataPtr = &CompositeFontCache->GetFontDataForCharacter(InFontInfo, SlateFontRendererUtils::InvalidSubChar, SubFontScalingFactor);
+				FaceGlyphData = FontRenderer->GetFontFaceForCharacter(*FontDataPtr, SlateFontRendererUtils::InvalidSubChar, InFontInfo.FontFallback);
+			}
+
+			if (!RunningFontDataPtr || RunningFontDataPtr != FontDataPtr || RunningFaceAndMemory != FaceGlyphData.FaceAndMemory || RunningSubFontScalingFactor != SubFontScalingFactor)
 			{
 				AppendPendingFontDataToSequence();
 
 				SplitStartIndex = RunningTextIndex;
-				RunningFontDataPtr = &FontData;
+				RunningFontDataPtr = FontDataPtr;
 				RunningFaceAndMemory = FaceGlyphData.FaceAndMemory;
 				RunningSubFontScalingFactor = SubFontScalingFactor;
 			}
@@ -317,16 +325,24 @@ void FSlateTextShaper::PerformHarfBuzzTextShaping(const TCHAR* InText, const int
 		{
 			const TCHAR CurrentChar = InText[RunningTextIndex];
 
+			// First try with the actual character
 			float SubFontScalingFactor = 1.0f;
-			const FFontData& FontData = CompositeFontCache->GetFontDataForCharacter(InFontInfo, CurrentChar, SubFontScalingFactor);
-			const FFreeTypeFaceGlyphData FaceGlyphData = FontRenderer->GetFontFaceForCharacter(FontData, CurrentChar, InFontInfo.FontFallback);
+			const FFontData* FontDataPtr = &CompositeFontCache->GetFontDataForCharacter(InFontInfo, CurrentChar, SubFontScalingFactor);
+			FFreeTypeFaceGlyphData FaceGlyphData = FontRenderer->GetFontFaceForCharacter(*FontDataPtr, CurrentChar, InFontInfo.FontFallback);
+
+			// If none of our fonts can render that character (as the fallback font may be missing), try again with the fallback character
+			if (!FaceGlyphData.FaceAndMemory.IsValid())
+			{
+				FontDataPtr = &CompositeFontCache->GetFontDataForCharacter(InFontInfo, SlateFontRendererUtils::InvalidSubChar, SubFontScalingFactor);
+				FaceGlyphData = FontRenderer->GetFontFaceForCharacter(*FontDataPtr, SlateFontRendererUtils::InvalidSubChar, InFontInfo.FontFallback);
+			}
 
-			if (!RunningFontDataPtr || RunningFontDataPtr != &FontData || RunningFaceAndMemory != FaceGlyphData.FaceAndMemory || RunningSubFontScalingFactor != SubFontScalingFactor)
+			if (!RunningFontDataPtr || RunningFontDataPtr != FontDataPtr || RunningFaceAndMemory != FaceGlyphData.FaceAndMemory || RunningSubFontScalingFactor != SubFontScalingFactor)
 			{
 				AppendPendingFontDataToSequence();
 
 				SplitStartIndex = RunningTextIndex;
-				RunningFontDataPtr = &FontData;
+				RunningFontDataPtr = FontDataPtr;
 				RunningFaceAndMemory = FaceGlyphData.FaceAndMemory;
 				RunningSubFontScalingFactor = SubFontScalingFactor;
 			}
diff --git a/Engine/Source/Runtime/SlateCore/Private/Rendering/DrawElements.cpp b/Engine/Source/Runtime/SlateCore/Private/Rendering/DrawElements.cpp
--- a/Engine/Source/Runtime/SlateCore/Private/Rendering/DrawElements.cpp
+++ b/Engine/Source/Runtime/SlateCore/Private/Rendering/DrawElements.cpp
@@ -4,6 +4,7 @@
 #include "Application/SlateApplicationBase.h"
 #include "Widgets/SWidget.h"
 #include "Widgets/SWindow.h"
+#include "HAL/IConsoleManager.h"
 
 DECLARE_CYCLE_STAT(TEXT("FSlateDrawElement::Make Time"), STAT_SlateDrawElementMakeTime, STATGROUP_SlateVerbose);
 DECLARE_CYCLE_STAT(TEXT("FSlateDrawElement::MakeCustomVerts Time"), STAT_SlateDrawElementMakeCustomVertsTime, STATGROUP_Slate);
@@ -32,38 +33,61 @@ void FSlateDataPayload::SetTextPayloadProperties( FSlateWindowElementList& Eleme
 	}
 }
 
+int32 PixelSnapRenderRotationScaleShear = 0;
+static FAutoConsoleVariableRef CVarPixelSnapRenderRotationScaleShear(TEXT("Slate.PixelSnapRenderRotationScaleShear"), PixelSnapRenderRotationScaleShear, TEXT("If the element being drawn has a render transform with Rotation, Scale or Shear, we automatically disable pixel snapping."), ECVF_Default);
 
-void FSlateDrawElement::Init(uint32 InLayer, const FPaintGeometry& PaintGeometry, const FSlateRect& InClippingRect, ESlateDrawEffect::Type InDrawEffects)
+int32 PixelSnapRenderTranslation = 1;
+static FAutoConsoleVariableRef CVarPixelSnapRenderTranslation(TEXT("Slate.PixelSnapRenderTranslation"), PixelSnapRenderTranslation, TEXT("If the element being drawn has a render transform with Translation, we automatically disable pixel snapping."), ECVF_Default);
+
+void FSlateDrawElement::Init(uint32 InLayer, const FPaintGeometry& PaintGeometry, const FSlateRect& InClippingRect, ESlateDrawEffect InDrawEffects)
 {
 	RenderTransform = PaintGeometry.GetAccumulatedRenderTransform();
 	Position = PaintGeometry.DrawPosition;
+	Scale = PaintGeometry.DrawScale;
 	LocalSize = PaintGeometry.GetLocalSize();
 	ClippingRect = InClippingRect;
 	Layer = InLayer;
-	Scale = PaintGeometry.DrawScale;
 	DrawEffects = InDrawEffects;
 	extern SLATECORE_API TOptional<FShortRect> GSlateScissorRect;
 	ScissorRect = GSlateScissorRect;
 
+	// Calculate the layout to render transform as this is needed by several calculations downstream.
+	const FSlateLayoutTransform InverseLayoutTransform(Inverse(FSlateLayoutTransform(Scale, Position)));
+	LayoutToRenderTransform = Concatenate(InverseLayoutTransform, RenderTransform);
+
+	// TODO NDarnell We need to review FSlateApplication::Get().GetRenderer(), not thread safe to use here, or any other 
+	//      place in painting, since it could be hit during the loading thread.
+
 	// Like GSlateScissorRect above, this is a workaround because we want to keep track of the various Scenes 
 	// in use throughout the UI. We keep a synchronized set with the render thread on the SlateRenderer and 
 	// use indices to synchronize between them.
-	TSharedPtr<FSlateRenderer> Renderer = FSlateApplicationBase::Get().GetRenderer();
-	check(Renderer.IsValid());
+	FSlateRenderer* Renderer = FSlateApplicationBase::Get().GetRenderer().Get();
+	checkSlow(Renderer);
 	SceneIndex = Renderer->GetCurrentSceneIndex();
 
 	DataPayload.BatchFlags = ESlateBatchDrawFlag::None;
-	if ( InDrawEffects & ESlateDrawEffect::NoBlending )
-	{
-		DataPayload.BatchFlags |= ESlateBatchDrawFlag::NoBlending;
-	}
-	if ( InDrawEffects & ESlateDrawEffect::PreMultipliedAlpha )
+	DataPayload.BatchFlags |= static_cast<ESlateBatchDrawFlag>( static_cast<uint32>( InDrawEffects ) & static_cast<uint32>( ESlateDrawEffect::NoBlending | ESlateDrawEffect::PreMultipliedAlpha | ESlateDrawEffect::NoGamma | ESlateDrawEffect::InvertAlpha ) );
+
+	static_assert( ( ( __underlying_type(ESlateDrawEffect) )ESlateDrawEffect::NoBlending ) == ( ( __underlying_type(ESlateBatchDrawFlag) )ESlateBatchDrawFlag::NoBlending ), "Must keep ESlateBatchDrawFlag and ESlateDrawEffect partial matches" );
+	static_assert( ( ( __underlying_type(ESlateDrawEffect) )ESlateDrawEffect::PreMultipliedAlpha ) == ( ( __underlying_type(ESlateBatchDrawFlag) )ESlateBatchDrawFlag::PreMultipliedAlpha ), "Must keep ESlateBatchDrawFlag and ESlateDrawEffect partial matches" );
+	static_assert( ( ( __underlying_type(ESlateDrawEffect) )ESlateDrawEffect::NoGamma ) == ( ( __underlying_type(ESlateBatchDrawFlag) )ESlateBatchDrawFlag::NoGamma ), "Must keep ESlateBatchDrawFlag and ESlateDrawEffect partial matches" );
+	static_assert( ( ( __underlying_type(ESlateDrawEffect) )ESlateDrawEffect::InvertAlpha ) == ( ( __underlying_type(ESlateBatchDrawFlag) )ESlateBatchDrawFlag::InvertAlpha ), "Must keep ESlateBatchDrawFlag and ESlateDrawEffect partial matches" );
+
+	if ( PixelSnapRenderRotationScaleShear == 0 )
 	{
-		DataPayload.BatchFlags |= ESlateBatchDrawFlag::PreMultipliedAlpha;
+		// Check if theres any rotation, scale or sheer.  If so, don't pixel snap.
+		if ( !LayoutToRenderTransform.GetMatrix().IsIdentity() )
+		{
+			DrawEffects |= ESlateDrawEffect::NoPixelSnapping;
+		}
 	}
-	if ( InDrawEffects & ESlateDrawEffect::NoGamma )
+
+	if ( PixelSnapRenderTranslation == 0 )
 	{
-		DataPayload.BatchFlags |= ESlateBatchDrawFlag::NoGamma;
+		if ( !LayoutToRenderTransform.GetTranslation().IsZero() )
+		{
+			DrawEffects |= ESlateDrawEffect::NoPixelSnapping;
+		}
 	}
 }
 
@@ -83,9 +107,14 @@ void FSlateDrawElement::MakeBox(
 	const FPaintGeometry& PaintGeometry, 
 	const FSlateBrush* InBrush, 
 	const FSlateRect& InClippingRect, 
-	ESlateDrawEffect::Type InDrawEffects, 
+	ESlateDrawEffect InDrawEffects, 
 	const FLinearColor& InTint )
 {
+	if ( InTint.A == 0 || InBrush->DrawAs == ESlateBrushDrawType::NoDrawType )
+	{
+		return;
+	}
+
 	SCOPE_CYCLE_COUNTER( STAT_SlateDrawElementMakeTime )
 	PaintGeometry.CommitTransformsIfUsingLegacyConstructor();
 	FSlateDrawElement& DrawElt = ElementList.AddUninitialized();
@@ -101,17 +130,22 @@ void FSlateDrawElement::MakeBox(
 	const FSlateBrush* InBrush, 
 	const FSlateResourceHandle& InRenderingHandle,
 	const FSlateRect& InClippingRect, 
-	ESlateDrawEffect::Type InDrawEffects, 
+	ESlateDrawEffect InDrawEffects, 
 	const FLinearColor& InTint )
 {
-	SCOPE_CYCLE_COUNTER(STAT_SlateDrawElementMakeTime)
-
 	// Ignore invalid rendering handles.
 	if ( !InRenderingHandle.IsValid() )
 	{
 		return;
 	}
 
+	if ( InTint.A == 0 || InBrush->DrawAs == ESlateBrushDrawType::NoDrawType )
+	{
+		return;
+	}
+
+	SCOPE_CYCLE_COUNTER(STAT_SlateDrawElementMakeTime)
+
 	FSlateShaderResourceProxy* RenderingProxy = InRenderingHandle.Data->Proxy;
 
 	PaintGeometry.CommitTransformsIfUsingLegacyConstructor();
@@ -127,12 +161,17 @@ void FSlateDrawElement::MakeRotatedBox(
 	const FPaintGeometry& PaintGeometry, 
 	const FSlateBrush* InBrush, 
 	const FSlateRect& InClippingRect, 
-	ESlateDrawEffect::Type InDrawEffects, 
+	ESlateDrawEffect InDrawEffects, 
 	float Angle,
 	TOptional<FVector2D> InRotationPoint,
 	ERotationSpace RotationSpace,
 	const FLinearColor& InTint )
 {
+	if ( InTint.A == 0 || InBrush->DrawAs == ESlateBrushDrawType::NoDrawType )
+	{
+		return;
+	}
+
 	SCOPE_CYCLE_COUNTER( STAT_SlateDrawElementMakeTime )
 	PaintGeometry.CommitTransformsIfUsingLegacyConstructor();
 	FSlateDrawElement& DrawElt = ElementList.AddUninitialized();
@@ -144,7 +183,7 @@ void FSlateDrawElement::MakeRotatedBox(
 }
 
 
-void FSlateDrawElement::MakeText( FSlateWindowElementList& ElementList, uint32 InLayer, const FPaintGeometry& PaintGeometry, const FString& InText, const int32 StartIndex, const int32 EndIndex, const FSlateFontInfo& InFontInfo, const FSlateRect& InClippingRect,ESlateDrawEffect::Type InDrawEffects, const FLinearColor& InTint )
+void FSlateDrawElement::MakeText( FSlateWindowElementList& ElementList, uint32 InLayer, const FPaintGeometry& PaintGeometry, const FString& InText, const int32 StartIndex, const int32 EndIndex, const FSlateFontInfo& InFontInfo, const FSlateRect& InClippingRect,ESlateDrawEffect InDrawEffects, const FLinearColor& InTint )
 {
 	SCOPE_CYCLE_COUNTER( STAT_SlateDrawElementMakeTime )
 	PaintGeometry.CommitTransformsIfUsingLegacyConstructor();
@@ -155,7 +194,7 @@ void FSlateDrawElement::MakeText( FSlateWindowElementList& ElementList, uint32 I
 }
 
 
-void FSlateDrawElement::MakeText( FSlateWindowElementList& ElementList, uint32 InLayer, const FPaintGeometry& PaintGeometry, const FString& InText, const FSlateFontInfo& InFontInfo, const FSlateRect& InClippingRect,ESlateDrawEffect::Type InDrawEffects, const FLinearColor& InTint )
+void FSlateDrawElement::MakeText( FSlateWindowElementList& ElementList, uint32 InLayer, const FPaintGeometry& PaintGeometry, const FString& InText, const FSlateFontInfo& InFontInfo, const FSlateRect& InClippingRect,ESlateDrawEffect InDrawEffects, const FLinearColor& InTint )
 {
 	SCOPE_CYCLE_COUNTER( STAT_SlateDrawElementMakeTime )
 	PaintGeometry.CommitTransformsIfUsingLegacyConstructor();
@@ -166,7 +205,7 @@ void FSlateDrawElement::MakeText( FSlateWindowElementList& ElementList, uint32 I
 }
 
 
-void FSlateDrawElement::MakeText( FSlateWindowElementList& ElementList, uint32 InLayer, const FPaintGeometry& PaintGeometry, const FText& InText, const FSlateFontInfo& InFontInfo, const FSlateRect& InClippingRect,ESlateDrawEffect::Type InDrawEffects, const FLinearColor& InTint )
+void FSlateDrawElement::MakeText( FSlateWindowElementList& ElementList, uint32 InLayer, const FPaintGeometry& PaintGeometry, const FText& InText, const FSlateFontInfo& InFontInfo, const FSlateRect& InClippingRect,ESlateDrawEffect InDrawEffects, const FLinearColor& InTint )
 {
 	SCOPE_CYCLE_COUNTER( STAT_SlateDrawElementMakeTime )
 	PaintGeometry.CommitTransformsIfUsingLegacyConstructor();
@@ -177,7 +216,7 @@ void FSlateDrawElement::MakeText( FSlateWindowElementList& ElementList, uint32 I
 	DrawElt.DataPayload.SetTextPayloadProperties(  ElementList, InText.ToString(), InFontInfo, InTint );
 }
 
-void FSlateDrawElement::MakeShapedText( FSlateWindowElementList& ElementList, uint32 InLayer, const FPaintGeometry& PaintGeometry, const FShapedGlyphSequenceRef& InShapedGlyphSequence, const FSlateRect& InClippingRect, ESlateDrawEffect::Type InDrawEffects, const FLinearColor& InTint )
+void FSlateDrawElement::MakeShapedText( FSlateWindowElementList& ElementList, uint32 InLayer, const FPaintGeometry& PaintGeometry, const FShapedGlyphSequenceRef& InShapedGlyphSequence, const FSlateRect& InClippingRect, ESlateDrawEffect InDrawEffects, const FLinearColor& InTint )
 {
 	SCOPE_CYCLE_COUNTER( STAT_SlateDrawElementMakeTime )
 	PaintGeometry.CommitTransformsIfUsingLegacyConstructor();
@@ -187,7 +226,7 @@ void FSlateDrawElement::MakeShapedText( FSlateWindowElementList& ElementList, ui
 	DrawElt.DataPayload.SetShapedTextPayloadProperties(InShapedGlyphSequence, InTint);
 }
 
-void FSlateDrawElement::MakeGradient( FSlateWindowElementList& ElementList, uint32 InLayer, const FPaintGeometry& PaintGeometry, TArray<FSlateGradientStop> InGradientStops, EOrientation InGradientType, const FSlateRect& InClippingRect, ESlateDrawEffect::Type InDrawEffects )
+void FSlateDrawElement::MakeGradient( FSlateWindowElementList& ElementList, uint32 InLayer, const FPaintGeometry& PaintGeometry, TArray<FSlateGradientStop> InGradientStops, EOrientation InGradientType, const FSlateRect& InClippingRect, ESlateDrawEffect InDrawEffects )
 {
 	SCOPE_CYCLE_COUNTER( STAT_SlateDrawElementMakeTime )
 	PaintGeometry.CommitTransformsIfUsingLegacyConstructor();
@@ -198,7 +237,7 @@ void FSlateDrawElement::MakeGradient( FSlateWindowElementList& ElementList, uint
 }
 
 
-void FSlateDrawElement::MakeSpline( FSlateWindowElementList& ElementList, uint32 InLayer, const FPaintGeometry& PaintGeometry, const FVector2D& InStart, const FVector2D& InStartDir, const FVector2D& InEnd, const FVector2D& InEndDir, const FSlateRect InClippingRect, float InThickness, ESlateDrawEffect::Type InDrawEffects, const FLinearColor& InTint )
+void FSlateDrawElement::MakeSpline( FSlateWindowElementList& ElementList, uint32 InLayer, const FPaintGeometry& PaintGeometry, const FVector2D& InStart, const FVector2D& InStartDir, const FVector2D& InEnd, const FVector2D& InEndDir, const FSlateRect InClippingRect, float InThickness, ESlateDrawEffect InDrawEffects, const FLinearColor& InTint )
 {
 	SCOPE_CYCLE_COUNTER( STAT_SlateDrawElementMakeTime )
 	PaintGeometry.CommitTransformsIfUsingLegacyConstructor();
@@ -209,13 +248,13 @@ void FSlateDrawElement::MakeSpline( FSlateWindowElementList& ElementList, uint32
 }
 
 
-void FSlateDrawElement::MakeDrawSpaceSpline( FSlateWindowElementList& ElementList, uint32 InLayer, const FVector2D& InStart, const FVector2D& InStartDir, const FVector2D& InEnd, const FVector2D& InEndDir, const FSlateRect InClippingRect, float InThickness, ESlateDrawEffect::Type InDrawEffects, const FLinearColor& InTint )
+void FSlateDrawElement::MakeDrawSpaceSpline( FSlateWindowElementList& ElementList, uint32 InLayer, const FVector2D& InStart, const FVector2D& InStartDir, const FVector2D& InEnd, const FVector2D& InEndDir, const FSlateRect InClippingRect, float InThickness, ESlateDrawEffect InDrawEffects, const FLinearColor& InTint )
 {
 	MakeSpline( ElementList, InLayer, FPaintGeometry(), InStart, InStartDir, InEnd, InEndDir, InClippingRect, InThickness, InDrawEffects, InTint );
 }
 
 
-void FSlateDrawElement::MakeDrawSpaceGradientSpline(FSlateWindowElementList& ElementList, uint32 InLayer, const FVector2D& InStart, const FVector2D& InStartDir, const FVector2D& InEnd, const FVector2D& InEndDir, const FSlateRect InClippingRect, const TArray<FSlateGradientStop>& InGradientStops, float InThickness, ESlateDrawEffect::Type InDrawEffects)
+void FSlateDrawElement::MakeDrawSpaceGradientSpline(FSlateWindowElementList& ElementList, uint32 InLayer, const FVector2D& InStart, const FVector2D& InStartDir, const FVector2D& InEnd, const FVector2D& InEndDir, const FSlateRect InClippingRect, const TArray<FSlateGradientStop>& InGradientStops, float InThickness, ESlateDrawEffect InDrawEffects)
 {
 	SCOPE_CYCLE_COUNTER(STAT_SlateDrawElementMakeTime)
 	const FPaintGeometry PaintGeometry;
@@ -227,7 +266,7 @@ void FSlateDrawElement::MakeDrawSpaceGradientSpline(FSlateWindowElementList& Ele
 }
 
 
-void FSlateDrawElement::MakeLines(FSlateWindowElementList& ElementList, uint32 InLayer, const FPaintGeometry& PaintGeometry, const TArray<FVector2D>& Points, const FSlateRect InClippingRect, ESlateDrawEffect::Type InDrawEffects, const FLinearColor& InTint, bool bAntialias, float Thickness)
+void FSlateDrawElement::MakeLines(FSlateWindowElementList& ElementList, uint32 InLayer, const FPaintGeometry& PaintGeometry, const TArray<FVector2D>& Points, const FSlateRect InClippingRect, ESlateDrawEffect InDrawEffects, const FLinearColor& InTint, bool bAntialias, float Thickness)
 {
 	SCOPE_CYCLE_COUNTER( STAT_SlateDrawElementMakeTime )
 	PaintGeometry.CommitTransformsIfUsingLegacyConstructor();
@@ -238,7 +277,7 @@ void FSlateDrawElement::MakeLines(FSlateWindowElementList& ElementList, uint32 I
 }
 
 
-void FSlateDrawElement::MakeViewport( FSlateWindowElementList& ElementList, uint32 InLayer, const FPaintGeometry& PaintGeometry, TSharedPtr<const ISlateViewport> Viewport, const FSlateRect& InClippingRect, ESlateDrawEffect::Type InDrawEffects, const FLinearColor& InTint )
+void FSlateDrawElement::MakeViewport( FSlateWindowElementList& ElementList, uint32 InLayer, const FPaintGeometry& PaintGeometry, TSharedPtr<const ISlateViewport> Viewport, const FSlateRect& InClippingRect, ESlateDrawEffect InDrawEffects, const FLinearColor& InTint )
 {
 	SCOPE_CYCLE_COUNTER( STAT_SlateDrawElementMakeTime )
 	PaintGeometry.CommitTransformsIfUsingLegacyConstructor();
