Copying //UE4/Dev-Editor to //UE4/Dev-Main (Source: //UE4/Dev-Editor @ 3279756)

diff --git a/Engine/Source/Runtime/SlateCore/Private/Rendering/ElementBatcher.cpp b/Engine/Source/Runtime/SlateCore/Private/Rendering/ElementBatcher.cpp
--- a/Engine/Source/Runtime/SlateCore/Private/Rendering/ElementBatcher.cpp
+++ b/Engine/Source/Runtime/SlateCore/Private/Rendering/ElementBatcher.cpp
@@ -150,37 +150,39 @@ void FSlateElementBatcher::AddElements(const FSlateWindowElementList& ElementLis
 			{
 				++NumDrawnBatchesStat;
 
+				const bool EnablePixelSnapping = !EnumHasAllFlags(DrawElement.GetDrawEffects(), ESlateDrawEffect::NoPixelSnapping);
+
 				// time just the adding of the element. The clipping stuff will be counted in exclusive time for the non-typed timer.
 				SLATE_CYCLE_COUNTER_SCOPE_CUSTOM_DETAILED(SLATE_STATS_DETAIL_LEVEL_MED, GSlateAddElements, ElementFNames[DrawElement.GetElementType()]);
 				// Determine what type of element to add
 				switch( DrawElement.GetElementType() )
 				{
 				case FSlateDrawElement::ET_Box:
-					AddBoxElement( DrawElement );
+					EnablePixelSnapping ? AddBoxElement<ESlateVertexRounding::Enabled>( DrawElement ) : AddBoxElement<ESlateVertexRounding::Disabled>(DrawElement);
 					break;
 				case FSlateDrawElement::ET_DebugQuad:
-					AddQuadElement( DrawElement );
+					EnablePixelSnapping ? AddQuadElement<ESlateVertexRounding::Enabled>(DrawElement) : AddQuadElement<ESlateVertexRounding::Disabled>(DrawElement);
 					break;
 				case FSlateDrawElement::ET_Text:
-					AddTextElement( DrawElement );
+					EnablePixelSnapping ? AddTextElement<ESlateVertexRounding::Enabled>(DrawElement) : AddTextElement<ESlateVertexRounding::Disabled>(DrawElement);
 					break;
 				case FSlateDrawElement::ET_ShapedText:
-					AddShapedTextElement( DrawElement );
+					EnablePixelSnapping ? AddShapedTextElement<ESlateVertexRounding::Enabled>(DrawElement) : AddShapedTextElement<ESlateVertexRounding::Disabled>(DrawElement);
 					break;
 				case FSlateDrawElement::ET_Spline:
-					AddSplineElement( DrawElement );
+					EnablePixelSnapping ? AddSplineElement<ESlateVertexRounding::Enabled>(DrawElement) : AddSplineElement<ESlateVertexRounding::Disabled>(DrawElement);
 					break;
 				case FSlateDrawElement::ET_Line:
-					AddLineElement( DrawElement );
+					EnablePixelSnapping ? AddLineElement<ESlateVertexRounding::Enabled>(DrawElement) : AddLineElement<ESlateVertexRounding::Disabled>(DrawElement);
 					break;
 				case FSlateDrawElement::ET_Gradient:
-					AddGradientElement( DrawElement );
+					EnablePixelSnapping ? AddGradientElement<ESlateVertexRounding::Enabled>(DrawElement) : AddGradientElement<ESlateVertexRounding::Disabled>(DrawElement);
 					break;
 				case FSlateDrawElement::ET_Viewport:
-					AddViewportElement( DrawElement );
+					EnablePixelSnapping ? AddViewportElement<ESlateVertexRounding::Enabled>(DrawElement) : AddViewportElement<ESlateVertexRounding::Disabled>(DrawElement);
 					break;
 				case FSlateDrawElement::ET_Border:
-					AddBorderElement( DrawElement );
+					EnablePixelSnapping ? AddBorderElement<ESlateVertexRounding::Enabled>(DrawElement) : AddBorderElement<ESlateVertexRounding::Disabled>(DrawElement);
 					break;
 				case FSlateDrawElement::ET_Custom:
 					AddCustomElement( DrawElement );
@@ -213,6 +215,7 @@ FColor FSlateElementBatcher::PackVertexColor(const FLinearColor& InLinearColor)
 	return InLinearColor.ToFColor(bSRGBVertexColor);
 }
 
+template<ESlateVertexRounding Rounding>
 void FSlateElementBatcher::AddQuadElement( const FSlateDrawElement& DrawElement, FColor Color )
 {
 	const FSlateRenderTransform& RenderTransform = DrawElement.GetRenderTransform();
@@ -222,7 +225,7 @@ void FSlateElementBatcher::AddQuadElement( const FSlateDrawElement& DrawElement,
 	//float Scale = DrawElement.GetScale();
 	//const FSlateDataPayload& InPayload = DrawElement.GetDataPayload();
 	const FSlateRect& InClippingRect = DrawElement.GetClippingRect();
-	//ESlateDrawEffect::Type InDrawEffects = DrawElement.GetDrawEffects();
+	ESlateDrawEffect InDrawEffects = DrawElement.GetDrawEffects();
 	uint32 Layer = DrawElement.GetLayer();
 
 	// extract the layout transform from the draw element
@@ -243,10 +246,10 @@ void FSlateElementBatcher::AddQuadElement( const FSlateDrawElement& DrawElement,
 	uint32 IndexStart = BatchVertices.Num();
 
 	// Add four vertices to the list of verts to be added to the vertex buffer
-	BatchVertices.Add( FSlateVertex( RenderTransform, TopLeft, FVector2D(0.0f,0.0f), Color, RenderClipRect ) );
-	BatchVertices.Add( FSlateVertex( RenderTransform, TopRight, FVector2D(1.0f,0.0f), Color, RenderClipRect ) );
-	BatchVertices.Add( FSlateVertex( RenderTransform, BotLeft, FVector2D(0.0f,1.0f), Color, RenderClipRect ) );
-	BatchVertices.Add( FSlateVertex( RenderTransform, BotRight, FVector2D(1.0f,1.0f),Color, RenderClipRect ) );
+	BatchVertices.Add( FSlateVertex::Make<Rounding>( RenderTransform, TopLeft, FVector2D(0.0f,0.0f), Color, RenderClipRect ) );
+	BatchVertices.Add( FSlateVertex::Make<Rounding>( RenderTransform, TopRight, FVector2D(1.0f,0.0f), Color, RenderClipRect ) );
+	BatchVertices.Add( FSlateVertex::Make<Rounding>( RenderTransform, BotLeft, FVector2D(0.0f,1.0f), Color, RenderClipRect ) );
+	BatchVertices.Add( FSlateVertex::Make<Rounding>( RenderTransform, BotRight, FVector2D(1.0f,1.0f),Color, RenderClipRect ) );
 
 	// The offset into the index buffer where this quads indices start
 	uint32 IndexOffsetStart = BatchIndices.Num();
@@ -274,316 +277,308 @@ FSlateRenderTransform GetBoxRenderTransform(const FSlateDrawElement& DrawElement
 	return Concatenate(RotationTransform, ElementRenderTransform);
 }
 
-
-void FSlateElementBatcher::AddBoxElement( const FSlateDrawElement& DrawElement )
+template<ESlateVertexRounding Rounding>
+void FSlateElementBatcher::AddBoxElement(const FSlateDrawElement& DrawElement)
 {
+	NumDrawnBoxesStat++;
+
 	const FSlateDataPayload& InPayload = DrawElement.GetDataPayload();
-	FColor Tint = PackVertexColor(InPayload.Tint);
 
-	if ( Tint.A == 0 )
-	{
-		return;
-	}
+	check(InPayload.BrushResource);
+	const FSlateBrush* BrushResource = InPayload.BrushResource;
 
-	NumDrawnBoxesStat++;
+	ensureMsgf(BrushResource->DrawAs != ESlateBrushDrawType::NoDrawType, TEXT("This should have been filtered out earlier in the Make... call."));
 
+	const FColor Tint = PackVertexColor(InPayload.Tint);
 	const FSlateRenderTransform& ElementRenderTransform = DrawElement.GetRenderTransform();
 	const FSlateRenderTransform RenderTransform = GetBoxRenderTransform(DrawElement);
-	//const FVector2D& InPosition = DrawElement.GetPosition();
-	//const FVector2D& Size = DrawElement.GetSize();
 	const FVector2D& LocalSize = DrawElement.GetLocalSize();
-	//float Scale = DrawElement.GetScale();
+
 	const FSlateRect& InClippingRect = DrawElement.GetClippingRect();
-	ESlateDrawEffect::Type InDrawEffects = DrawElement.GetDrawEffects();
-	uint32 Layer = DrawElement.GetLayer();
+	const ESlateDrawEffect InDrawEffects = DrawElement.GetDrawEffects();
+	const uint32 Layer = DrawElement.GetLayer();
 
-	// extract the layout transform from the draw element
-	FSlateLayoutTransform InverseLayoutTransform(Inverse(FSlateLayoutTransform(DrawElement.GetScale(), DrawElement.GetPosition())));
-	// The clip rect is NOT subject to the rotations specified by MakeRotatedBox.
-	FSlateRotatedClipRectType RenderClipRect = FSlateRotatedClipRectType::MakeSnappedRotatedRect(InClippingRect, InverseLayoutTransform, ElementRenderTransform);
+	const float DrawScale = DrawElement.GetScale();
 
-	check(InPayload.BrushResource);
-	const FSlateBrush* BrushResource = InPayload.BrushResource;
+	const FSlateRotatedClipRectType RenderClipRect = DrawElement.CalculateRenderClippingRect();
 
-	if ( BrushResource->DrawAs != ESlateBrushDrawType::NoDrawType )
-	{
-		// Do pixel snapping
-		FVector2D TopLeft(0,0);
-		FVector2D BotRight(LocalSize);
+	// Do pixel snapping
+	FVector2D TopLeft(0,0);
+	FVector2D BotRight(LocalSize);
 
-		uint32 TextureWidth = 1;
-		uint32 TextureHeight = 1;
+	uint32 TextureWidth = 1;
+	uint32 TextureHeight = 1;
 
-		// Get the default start and end UV.  If the texture is atlased this value will be a subset of this
-		FVector2D StartUV = FVector2D(0.0f,0.0f);
-		FVector2D EndUV = FVector2D(1.0f,1.0f);
-		FVector2D SizeUV;
+	// Get the default start and end UV.  If the texture is atlased this value will be a subset of this
+	FVector2D StartUV = FVector2D(0.0f,0.0f);
+	FVector2D EndUV = FVector2D(1.0f,1.0f);
+	FVector2D SizeUV;
 
-		FVector2D HalfTexel;
+	FVector2D HalfTexel;
 
-		const FSlateShaderResourceProxy* ResourceProxy = InPayload.ResourceProxy;
-		FSlateShaderResource* Resource = nullptr;
-		if( ResourceProxy )
+	const FSlateShaderResourceProxy* ResourceProxy = InPayload.ResourceProxy;
+	FSlateShaderResource* Resource = nullptr;
+	if( ResourceProxy )
+	{
+		// The actual texture for rendering.  If the texture is atlased this is the atlas
+		Resource = ResourceProxy->Resource;
+		// The width and height of the texture (non-atlased size)
+		TextureWidth = ResourceProxy->ActualSize.X != 0 ? ResourceProxy->ActualSize.X : 1;
+		TextureHeight = ResourceProxy->ActualSize.Y != 0 ? ResourceProxy->ActualSize.Y : 1;
+
+		// Texel offset
+		HalfTexel = FVector2D( PixelCenterOffset/TextureWidth, PixelCenterOffset/TextureHeight );
+
+		FBox2D BrushUV = BrushResource->GetUVRegion();
+		//In case brush has valid UV region - use it instead of proxy UV
+		if (BrushUV.bIsValid)
 		{
-			// The actual texture for rendering.  If the texture is atlased this is the atlas
-			Resource = ResourceProxy->Resource;
-			// The width and height of the texture (non-atlased size)
-			TextureWidth = ResourceProxy->ActualSize.X != 0 ? ResourceProxy->ActualSize.X : 1;
-			TextureHeight = ResourceProxy->ActualSize.Y != 0 ? ResourceProxy->ActualSize.Y : 1;
-
-			// Texel offset
-			HalfTexel = FVector2D( PixelCenterOffset/TextureWidth, PixelCenterOffset/TextureHeight );
-
-			FBox2D BrushUV = BrushResource->GetUVRegion();
-			//In case brush has valid UV region - use it instead of proxy UV
-			if (BrushUV.bIsValid)
-			{
-				SizeUV = BrushUV.GetSize();
-				StartUV = BrushUV.Min + HalfTexel;
-				EndUV = StartUV + SizeUV;
-			}
-			else
-			{
-				SizeUV = ResourceProxy->SizeUV;
-				StartUV = ResourceProxy->StartUV + HalfTexel;
-				EndUV = StartUV + ResourceProxy->SizeUV;
-			}
+			SizeUV = BrushUV.GetSize();
+			StartUV = BrushUV.Min + HalfTexel;
+			EndUV = StartUV + SizeUV;
 		}
 		else
 		{
-			// no texture
-			SizeUV = FVector2D(1.0f,1.0f);
-			HalfTexel = FVector2D( PixelCenterOffset, PixelCenterOffset );
+			SizeUV = ResourceProxy->SizeUV;
+			StartUV = ResourceProxy->StartUV + HalfTexel;
+			EndUV = StartUV + ResourceProxy->SizeUV;
 		}
+	}
+	else
+	{
+		// no texture
+		SizeUV = FVector2D(1.0f,1.0f);
+		HalfTexel = FVector2D( PixelCenterOffset, PixelCenterOffset );
+	}
 
 
-		const ESlateBrushTileType::Type TilingRule = BrushResource->Tiling;
-		const bool bTileHorizontal = (TilingRule == ESlateBrushTileType::Both || TilingRule == ESlateBrushTileType::Horizontal);
-		const bool bTileVertical = (TilingRule == ESlateBrushTileType::Both || TilingRule == ESlateBrushTileType::Vertical);
+	const ESlateBrushTileType::Type TilingRule = BrushResource->Tiling;
+	const bool bTileHorizontal = (TilingRule == ESlateBrushTileType::Both || TilingRule == ESlateBrushTileType::Horizontal);
+	const bool bTileVertical = (TilingRule == ESlateBrushTileType::Both || TilingRule == ESlateBrushTileType::Vertical);
 
-		const ESlateBrushMirrorType::Type MirroringRule = BrushResource->Mirroring;
-		const bool bMirrorHorizontal = (MirroringRule == ESlateBrushMirrorType::Both || MirroringRule == ESlateBrushMirrorType::Horizontal);
-		const bool bMirrorVertical = (MirroringRule == ESlateBrushMirrorType::Both || MirroringRule == ESlateBrushMirrorType::Vertical);
+	const ESlateBrushMirrorType::Type MirroringRule = BrushResource->Mirroring;
+	const bool bMirrorHorizontal = (MirroringRule == ESlateBrushMirrorType::Both || MirroringRule == ESlateBrushMirrorType::Horizontal);
+	const bool bMirrorVertical = (MirroringRule == ESlateBrushMirrorType::Both || MirroringRule == ESlateBrushMirrorType::Vertical);
 
-		// Pass the tiling information as a flag so we can pick the correct texture addressing mode
-		ESlateBatchDrawFlag::Type DrawFlags = InPayload.BatchFlags;
-		DrawFlags |= ( ( bTileHorizontal ? ESlateBatchDrawFlag::TileU : 0 ) | ( bTileVertical ? ESlateBatchDrawFlag::TileV : 0 ) );
+	// Pass the tiling information as a flag so we can pick the correct texture addressing mode
+	ESlateBatchDrawFlag DrawFlags = InPayload.BatchFlags;
+	DrawFlags |= ( ( bTileHorizontal ? ESlateBatchDrawFlag::TileU : ESlateBatchDrawFlag::None ) | ( bTileVertical ? ESlateBatchDrawFlag::TileV : ESlateBatchDrawFlag::None ) );
 
-		FSlateElementBatch& ElementBatch = FindBatchForElement( Layer, FShaderParams(), Resource, ESlateDrawPrimitive::TriangleList, ESlateShader::Default, InDrawEffects, DrawFlags, DrawElement.GetScissorRect(), DrawElement.GetSceneIndex());
-		FSlateVertexArray& BatchVertices = BatchData->GetBatchVertexList(ElementBatch);
-		FSlateIndexArray& BatchIndices = BatchData->GetBatchIndexList(ElementBatch);
+	FSlateElementBatch& ElementBatch = FindBatchForElement( Layer, FShaderParams(), Resource, ESlateDrawPrimitive::TriangleList, ESlateShader::Default, InDrawEffects, DrawFlags, DrawElement.GetScissorRect(), DrawElement.GetSceneIndex());
+	FSlateVertexArray& BatchVertices = BatchData->GetBatchVertexList(ElementBatch);
+	FSlateIndexArray& BatchIndices = BatchData->GetBatchIndexList(ElementBatch);
 
-		float HorizontalTiling = bTileHorizontal ? LocalSize.X/TextureWidth : 1.0f;
-		float VerticalTiling = bTileVertical ? LocalSize.Y/TextureHeight : 1.0f;
+	float HorizontalTiling = bTileHorizontal ? LocalSize.X / TextureWidth : 1.0f;
+	float VerticalTiling = bTileVertical ? LocalSize.Y / TextureHeight : 1.0f;
 
-		const FVector2D Tiling( HorizontalTiling,VerticalTiling );
+	const FVector2D Tiling( HorizontalTiling, VerticalTiling );
 
-		// The start index of these vertices in the index buffer
-		uint32 IndexStart = BatchVertices.Num();
-		// The offset into the index buffer where this elements indices start
-		uint32 IndexOffsetStart = BatchIndices.Num();
+	// The start index of these vertices in the index buffer
+	uint32 IndexStart = BatchVertices.Num();
+	// The offset into the index buffer where this elements indices start
+	uint32 IndexOffsetStart = BatchIndices.Num();
 
-		const FMargin& Margin = BrushResource->Margin;
+	const FMargin& Margin = BrushResource->Margin;
 
-		if ( BrushResource->DrawAs != ESlateBrushDrawType::Image &&
-			( Margin.Left != 0.0f || Margin.Top != 0.0f || Margin.Right != 0.0f || Margin.Bottom != 0.0f ) )
+	if ( BrushResource->DrawAs != ESlateBrushDrawType::Image &&
+		( Margin.Left != 0.0f || Margin.Top != 0.0f || Margin.Right != 0.0f || Margin.Bottom != 0.0f ) )
+	{
+		// Create 9 quads for the box element based on the following diagram
+		//     ___LeftMargin    ___RightMargin
+		//    /                /
+		//  +--+-------------+--+
+		//  |  |c1           |c2| ___TopMargin
+		//  +--o-------------o--+
+		//  |  |             |  |
+		//  |  |c3           |c4|
+		//  +--o-------------o--+
+		//  |  |             |  | ___BottomMargin
+		//  +--+-------------+--+
+
+
+		// Determine the texture coordinates for each quad
+		// These are not scaled.
+		float LeftMarginU = (Margin.Left > 0.0f)
+			? StartUV.X + Margin.Left * SizeUV.X + HalfTexel.X
+			: StartUV.X;
+		float TopMarginV = (Margin.Top > 0.0f)
+			? StartUV.Y + Margin.Top * SizeUV.Y + HalfTexel.Y
+			: StartUV.Y;
+		float RightMarginU = (Margin.Right > 0.0f)
+			? EndUV.X - Margin.Right * SizeUV.X + HalfTexel.X
+			: EndUV.X;
+		float BottomMarginV = (Margin.Bottom > 0.0f)
+			? EndUV.Y - Margin.Bottom * SizeUV.Y + HalfTexel.Y
+			: EndUV.Y;
+
+		if( bMirrorHorizontal || bMirrorVertical )
 		{
-			// Create 9 quads for the box element based on the following diagram
-			//     ___LeftMargin    ___RightMargin
-			//    /                /
-			//  +--+-------------+--+
-			//  |  |c1           |c2| ___TopMargin
-			//  +--o-------------o--+
-			//  |  |             |  |
-			//  |  |c3           |c4|
-			//  +--o-------------o--+
-			//  |  |             |  | ___BottomMargin
-			//  +--+-------------+--+
-
-
-			// Determine the texture coordinates for each quad
-			// These are not scaled.
-			float LeftMarginU = (Margin.Left > 0.0f)
-				? StartUV.X + Margin.Left * SizeUV.X + HalfTexel.X
-				: StartUV.X;
-			float TopMarginV = (Margin.Top > 0.0f)
-				? StartUV.Y + Margin.Top * SizeUV.Y + HalfTexel.Y
-				: StartUV.Y;
-			float RightMarginU = (Margin.Right > 0.0f)
-				? EndUV.X - Margin.Right * SizeUV.X + HalfTexel.X
-				: EndUV.X;
-			float BottomMarginV = (Margin.Bottom > 0.0f)
-				? EndUV.Y - Margin.Bottom * SizeUV.Y + HalfTexel.Y
-				: EndUV.Y;
-
-			if( bMirrorHorizontal || bMirrorVertical )
-			{
-				const FVector2D UVMin = StartUV;
-				const FVector2D UVMax = EndUV;
+			const FVector2D UVMin = StartUV;
+			const FVector2D UVMax = EndUV;
 
-				if( bMirrorHorizontal )
-				{
-					StartUV.X = UVMax.X - ( StartUV.X - UVMin.X );
-					EndUV.X = UVMax.X - ( EndUV.X - UVMin.X );
-					LeftMarginU = UVMax.X - ( LeftMarginU - UVMin.X );
-					RightMarginU = UVMax.X - ( RightMarginU - UVMin.X );
-				}
-				if( bMirrorVertical )
-				{
-					StartUV.Y = UVMax.Y - ( StartUV.Y - UVMin.Y );
-					EndUV.Y = UVMax.Y - ( EndUV.Y - UVMin.Y );
-					TopMarginV = UVMax.Y - ( TopMarginV - UVMin.Y );
-					BottomMarginV = UVMax.Y - ( BottomMarginV - UVMin.Y );
-				}
-			}
-
-			// Determine the margins for each quad
-
-			float LeftMarginX = TextureWidth * Margin.Left;
-			float TopMarginY = TextureHeight * Margin.Top;
-			float RightMarginX = LocalSize.X - TextureWidth * Margin.Right;
-			float BottomMarginY = LocalSize.Y - TextureHeight * Margin.Bottom;
-
-			// If the margins are overlapping the margins are too big or the button is too small
-			// so clamp margins to half of the box size
-			if( RightMarginX < LeftMarginX )
+			if( bMirrorHorizontal )
 			{
-				LeftMarginX = LocalSize.X / 2;
-				RightMarginX = LeftMarginX;
+				StartUV.X = UVMax.X - ( StartUV.X - UVMin.X );
+				EndUV.X = UVMax.X - ( EndUV.X - UVMin.X );
+				LeftMarginU = UVMax.X - ( LeftMarginU - UVMin.X );
+				RightMarginU = UVMax.X - ( RightMarginU - UVMin.X );
 			}
-
-			if( BottomMarginY < TopMarginY )
+			if( bMirrorVertical )
 			{
-				TopMarginY = LocalSize.Y / 2;
-				BottomMarginY = TopMarginY;
+				StartUV.Y = UVMax.Y - ( StartUV.Y - UVMin.Y );
+				EndUV.Y = UVMax.Y - ( EndUV.Y - UVMin.Y );
+				TopMarginV = UVMax.Y - ( TopMarginV - UVMin.Y );
+				BottomMarginV = UVMax.Y - ( BottomMarginV - UVMin.Y );
 			}
+		}
 
-			FVector2D Position = TopLeft;
-			FVector2D EndPos = BotRight;
-
-			BatchVertices.Add( FSlateVertex( RenderTransform, FVector2D( Position.X, Position.Y ),		StartUV,									Tiling,	Tint, RenderClipRect ) ); //0
-			BatchVertices.Add( FSlateVertex( RenderTransform, FVector2D( Position.X, TopMarginY ),		FVector2D( StartUV.X, TopMarginV ),			Tiling,	Tint, RenderClipRect ) ); //1
-			BatchVertices.Add( FSlateVertex( RenderTransform, FVector2D( LeftMarginX, Position.Y ),		FVector2D( LeftMarginU, StartUV.Y ),		Tiling,	Tint, RenderClipRect ) ); //2
-			BatchVertices.Add( FSlateVertex( RenderTransform, FVector2D( LeftMarginX, TopMarginY ),		FVector2D( LeftMarginU, TopMarginV ),		Tiling,	Tint, RenderClipRect ) ); //3
-			BatchVertices.Add( FSlateVertex( RenderTransform, FVector2D( RightMarginX, Position.Y ),	FVector2D( RightMarginU, StartUV.Y ),		Tiling,	Tint, RenderClipRect ) ); //4
-			BatchVertices.Add( FSlateVertex( RenderTransform, FVector2D( RightMarginX, TopMarginY ),	FVector2D( RightMarginU,TopMarginV),		Tiling,	Tint, RenderClipRect ) ); //5
-			BatchVertices.Add( FSlateVertex( RenderTransform, FVector2D( EndPos.X, Position.Y ),		FVector2D( EndUV.X, StartUV.Y ),			Tiling,	Tint, RenderClipRect ) ); //6
-			BatchVertices.Add( FSlateVertex( RenderTransform, FVector2D( EndPos.X, TopMarginY ),		FVector2D( EndUV.X, TopMarginV),			Tiling,	Tint, RenderClipRect ) ); //7
-
-			BatchVertices.Add( FSlateVertex( RenderTransform, FVector2D( Position.X, BottomMarginY ),	FVector2D( StartUV.X, BottomMarginV ),		Tiling,	Tint, RenderClipRect ) ); //8
-			BatchVertices.Add( FSlateVertex( RenderTransform, FVector2D( LeftMarginX, BottomMarginY ),	FVector2D( LeftMarginU, BottomMarginV ),	Tiling,	Tint, RenderClipRect ) ); //9
-			BatchVertices.Add( FSlateVertex( RenderTransform, FVector2D( RightMarginX, BottomMarginY ),	FVector2D( RightMarginU, BottomMarginV ),	Tiling,	Tint, RenderClipRect ) ); //10
-			BatchVertices.Add( FSlateVertex( RenderTransform, FVector2D( EndPos.X, BottomMarginY ),		FVector2D( EndUV.X, BottomMarginV ),		Tiling,	Tint, RenderClipRect ) ); //11
-			BatchVertices.Add( FSlateVertex( RenderTransform, FVector2D( Position.X, EndPos.Y ),		FVector2D( StartUV.X, EndUV.Y ),			Tiling,	Tint, RenderClipRect ) ); //12
-			BatchVertices.Add( FSlateVertex( RenderTransform, FVector2D( LeftMarginX, EndPos.Y ),		FVector2D( LeftMarginU, EndUV.Y ),			Tiling,	Tint, RenderClipRect ) ); //13
-			BatchVertices.Add( FSlateVertex( RenderTransform, FVector2D( RightMarginX, EndPos.Y ),		FVector2D( RightMarginU, EndUV.Y ),			Tiling,	Tint, RenderClipRect ) ); //14
-			BatchVertices.Add( FSlateVertex( RenderTransform, FVector2D( EndPos.X, EndPos.Y ),			EndUV,										Tiling,	Tint, RenderClipRect ) ); //15
-
-			// Top
-			BatchIndices.Add( IndexStart + 0 );
-			BatchIndices.Add( IndexStart + 1 );
-			BatchIndices.Add( IndexStart + 2 );
-			BatchIndices.Add( IndexStart + 2 );
-			BatchIndices.Add( IndexStart + 1 );
-			BatchIndices.Add( IndexStart + 3 );
-
-			BatchIndices.Add( IndexStart + 2 );
-			BatchIndices.Add( IndexStart + 3 );
-			BatchIndices.Add( IndexStart + 4 );
-			BatchIndices.Add( IndexStart + 4 );
-			BatchIndices.Add( IndexStart + 3 );
-			BatchIndices.Add( IndexStart + 5 );
-
-			BatchIndices.Add( IndexStart + 4 );
-			BatchIndices.Add( IndexStart + 5 );
-			BatchIndices.Add( IndexStart + 6 );
-			BatchIndices.Add( IndexStart + 6 );
-			BatchIndices.Add( IndexStart + 5 );
-			BatchIndices.Add( IndexStart + 7 );
-
-			// Middle
-			BatchIndices.Add( IndexStart + 1 );
-			BatchIndices.Add( IndexStart + 8 );
-			BatchIndices.Add( IndexStart + 3 );
-			BatchIndices.Add( IndexStart + 3 );
-			BatchIndices.Add( IndexStart + 8 );
-			BatchIndices.Add( IndexStart + 9 );
-
-			BatchIndices.Add( IndexStart + 3 );
-			BatchIndices.Add( IndexStart + 9 );
-			BatchIndices.Add( IndexStart + 5 );
-			BatchIndices.Add( IndexStart + 5 );
-			BatchIndices.Add( IndexStart + 9 );
-			BatchIndices.Add( IndexStart + 10 );
-
-			BatchIndices.Add( IndexStart + 5 );
-			BatchIndices.Add( IndexStart + 10 );
-			BatchIndices.Add( IndexStart + 7 );
-			BatchIndices.Add( IndexStart + 7 );
-			BatchIndices.Add( IndexStart + 10 );
-			BatchIndices.Add( IndexStart + 11 );
-
-			// Bottom
-			BatchIndices.Add( IndexStart + 8 );
-			BatchIndices.Add( IndexStart + 12 );
-			BatchIndices.Add( IndexStart + 9 );
-			BatchIndices.Add( IndexStart + 9 );
-			BatchIndices.Add( IndexStart + 12 );
-			BatchIndices.Add( IndexStart + 13 );
-
-			BatchIndices.Add( IndexStart + 9 );
-			BatchIndices.Add( IndexStart + 13 );
-			BatchIndices.Add( IndexStart + 10 );
-			BatchIndices.Add( IndexStart + 10 );
-			BatchIndices.Add( IndexStart + 13 );
-			BatchIndices.Add( IndexStart + 14 );
-
-			BatchIndices.Add( IndexStart + 10 );
-			BatchIndices.Add( IndexStart + 14 );
-			BatchIndices.Add( IndexStart + 11 );
-			BatchIndices.Add( IndexStart + 11 );
-			BatchIndices.Add( IndexStart + 14 );
-			BatchIndices.Add( IndexStart + 15 );
+		// Determine the margins for each quad
+
+		float LeftMarginX = TextureWidth * Margin.Left;
+		float TopMarginY = TextureHeight * Margin.Top;
+		float RightMarginX = LocalSize.X - TextureWidth * Margin.Right;
+		float BottomMarginY = LocalSize.Y - TextureHeight * Margin.Bottom;
+
+		// If the margins are overlapping the margins are too big or the button is too small
+		// so clamp margins to half of the box size
+		if( RightMarginX < LeftMarginX )
+		{
+			LeftMarginX = LocalSize.X / 2;
+			RightMarginX = LeftMarginX;
 		}
-		else
+
+		if( BottomMarginY < TopMarginY )
 		{
-			FVector2D TopRight = FVector2D( BotRight.X, TopLeft.Y);
-			FVector2D BotLeft =	 FVector2D( TopLeft.X, BotRight.Y);
+			TopMarginY = LocalSize.Y / 2;
+			BottomMarginY = TopMarginY;
+		}
 
-			if( bMirrorHorizontal || bMirrorVertical )
-			{
-				const FVector2D UVMin = StartUV;
-				const FVector2D UVMax = EndUV;
+		FVector2D Position = TopLeft;
+		FVector2D EndPos = BotRight;
+
+		BatchVertices.Add( FSlateVertex::Make<Rounding>( RenderTransform, FVector2D( Position.X, Position.Y ),		LocalSize, DrawScale, FVector4(StartUV,										Tiling),	Tint, RenderClipRect ) ); //0
+		BatchVertices.Add( FSlateVertex::Make<Rounding>( RenderTransform, FVector2D( Position.X, TopMarginY ),		LocalSize, DrawScale, FVector4(FVector2D( StartUV.X, TopMarginV ),			Tiling),	Tint, RenderClipRect ) ); //1
+		BatchVertices.Add( FSlateVertex::Make<Rounding>( RenderTransform, FVector2D( LeftMarginX, Position.Y ),		LocalSize, DrawScale, FVector4(FVector2D( LeftMarginU, StartUV.Y ),			Tiling),	Tint, RenderClipRect ) ); //2
+		BatchVertices.Add( FSlateVertex::Make<Rounding>( RenderTransform, FVector2D( LeftMarginX, TopMarginY ),		LocalSize, DrawScale, FVector4(FVector2D( LeftMarginU, TopMarginV ),		Tiling),	Tint, RenderClipRect ) ); //3
+		BatchVertices.Add( FSlateVertex::Make<Rounding>( RenderTransform, FVector2D( RightMarginX, Position.Y ),	LocalSize, DrawScale, FVector4(FVector2D( RightMarginU, StartUV.Y ),		Tiling),	Tint, RenderClipRect ) ); //4
+		BatchVertices.Add( FSlateVertex::Make<Rounding>( RenderTransform, FVector2D( RightMarginX, TopMarginY ),	LocalSize, DrawScale, FVector4(FVector2D( RightMarginU,TopMarginV),			Tiling),	Tint, RenderClipRect ) ); //5
+		BatchVertices.Add( FSlateVertex::Make<Rounding>( RenderTransform, FVector2D( EndPos.X, Position.Y ),		LocalSize, DrawScale, FVector4(FVector2D( EndUV.X, StartUV.Y ),				Tiling),	Tint, RenderClipRect ) ); //6
+		BatchVertices.Add( FSlateVertex::Make<Rounding>( RenderTransform, FVector2D( EndPos.X, TopMarginY ),		LocalSize, DrawScale, FVector4(FVector2D( EndUV.X, TopMarginV),				Tiling),	Tint, RenderClipRect ) ); //7
+
+		BatchVertices.Add( FSlateVertex::Make<Rounding>( RenderTransform, FVector2D( Position.X, BottomMarginY ),	LocalSize, DrawScale, FVector4(FVector2D( StartUV.X, BottomMarginV ),		Tiling),	Tint, RenderClipRect ) ); //8
+		BatchVertices.Add( FSlateVertex::Make<Rounding>( RenderTransform, FVector2D( LeftMarginX, BottomMarginY ),	LocalSize, DrawScale, FVector4(FVector2D( LeftMarginU, BottomMarginV ),		Tiling),	Tint, RenderClipRect ) ); //9
+		BatchVertices.Add( FSlateVertex::Make<Rounding>( RenderTransform, FVector2D( RightMarginX, BottomMarginY ),	LocalSize, DrawScale, FVector4(FVector2D( RightMarginU, BottomMarginV ),	Tiling),	Tint, RenderClipRect ) ); //10
+		BatchVertices.Add( FSlateVertex::Make<Rounding>( RenderTransform, FVector2D( EndPos.X, BottomMarginY ),		LocalSize, DrawScale, FVector4(FVector2D( EndUV.X, BottomMarginV ),			Tiling),	Tint, RenderClipRect ) ); //11
+		BatchVertices.Add( FSlateVertex::Make<Rounding>( RenderTransform, FVector2D( Position.X, EndPos.Y ),		LocalSize, DrawScale, FVector4(FVector2D( StartUV.X, EndUV.Y ),				Tiling),	Tint, RenderClipRect ) ); //12
+		BatchVertices.Add( FSlateVertex::Make<Rounding>( RenderTransform, FVector2D( LeftMarginX, EndPos.Y ),		LocalSize, DrawScale, FVector4(FVector2D( LeftMarginU, EndUV.Y ),			Tiling),	Tint, RenderClipRect ) ); //13
+		BatchVertices.Add( FSlateVertex::Make<Rounding>( RenderTransform, FVector2D( RightMarginX, EndPos.Y ),		LocalSize, DrawScale, FVector4(FVector2D( RightMarginU, EndUV.Y ),			Tiling),	Tint, RenderClipRect ) ); //14
+		BatchVertices.Add( FSlateVertex::Make<Rounding>( RenderTransform, FVector2D( EndPos.X, EndPos.Y ),			LocalSize, DrawScale, FVector4(EndUV,										Tiling),	Tint, RenderClipRect ) ); //15
+
+		// Top
+		BatchIndices.Add( IndexStart + 0 );
+		BatchIndices.Add( IndexStart + 1 );
+		BatchIndices.Add( IndexStart + 2 );
+		BatchIndices.Add( IndexStart + 2 );
+		BatchIndices.Add( IndexStart + 1 );
+		BatchIndices.Add( IndexStart + 3 );
+
+		BatchIndices.Add( IndexStart + 2 );
+		BatchIndices.Add( IndexStart + 3 );
+		BatchIndices.Add( IndexStart + 4 );
+		BatchIndices.Add( IndexStart + 4 );
+		BatchIndices.Add( IndexStart + 3 );
+		BatchIndices.Add( IndexStart + 5 );
+
+		BatchIndices.Add( IndexStart + 4 );
+		BatchIndices.Add( IndexStart + 5 );
+		BatchIndices.Add( IndexStart + 6 );
+		BatchIndices.Add( IndexStart + 6 );
+		BatchIndices.Add( IndexStart + 5 );
+		BatchIndices.Add( IndexStart + 7 );
+
+		// Middle
+		BatchIndices.Add( IndexStart + 1 );
+		BatchIndices.Add( IndexStart + 8 );
+		BatchIndices.Add( IndexStart + 3 );
+		BatchIndices.Add( IndexStart + 3 );
+		BatchIndices.Add( IndexStart + 8 );
+		BatchIndices.Add( IndexStart + 9 );
+
+		BatchIndices.Add( IndexStart + 3 );
+		BatchIndices.Add( IndexStart + 9 );
+		BatchIndices.Add( IndexStart + 5 );
+		BatchIndices.Add( IndexStart + 5 );
+		BatchIndices.Add( IndexStart + 9 );
+		BatchIndices.Add( IndexStart + 10 );
+
+		BatchIndices.Add( IndexStart + 5 );
+		BatchIndices.Add( IndexStart + 10 );
+		BatchIndices.Add( IndexStart + 7 );
+		BatchIndices.Add( IndexStart + 7 );
+		BatchIndices.Add( IndexStart + 10 );
+		BatchIndices.Add( IndexStart + 11 );
+
+		// Bottom
+		BatchIndices.Add( IndexStart + 8 );
+		BatchIndices.Add( IndexStart + 12 );
+		BatchIndices.Add( IndexStart + 9 );
+		BatchIndices.Add( IndexStart + 9 );
+		BatchIndices.Add( IndexStart + 12 );
+		BatchIndices.Add( IndexStart + 13 );
+
+		BatchIndices.Add( IndexStart + 9 );
+		BatchIndices.Add( IndexStart + 13 );
+		BatchIndices.Add( IndexStart + 10 );
+		BatchIndices.Add( IndexStart + 10 );
+		BatchIndices.Add( IndexStart + 13 );
+		BatchIndices.Add( IndexStart + 14 );
+
+		BatchIndices.Add( IndexStart + 10 );
+		BatchIndices.Add( IndexStart + 14 );
+		BatchIndices.Add( IndexStart + 11 );
+		BatchIndices.Add( IndexStart + 11 );
+		BatchIndices.Add( IndexStart + 14 );
+		BatchIndices.Add( IndexStart + 15 );
+	}
+	else
+	{
+		FVector2D TopRight = FVector2D( BotRight.X, TopLeft.Y);
+		FVector2D BotLeft =	 FVector2D( TopLeft.X, BotRight.Y);
 
-				if( bMirrorHorizontal )
-				{
-					StartUV.X = UVMax.X - ( StartUV.X - UVMin.X );
-					EndUV.X = UVMax.X - ( EndUV.X - UVMin.X );
-				}
-				if( bMirrorVertical )
-				{
-					StartUV.Y = UVMax.Y - ( StartUV.Y - UVMin.Y );
-					EndUV.Y = UVMax.Y - ( EndUV.Y - UVMin.Y );
-				}
+		if( bMirrorHorizontal || bMirrorVertical )
+		{
+			const FVector2D UVMin = StartUV;
+			const FVector2D UVMax = EndUV;
+
+			if( bMirrorHorizontal )
+			{
+				StartUV.X = UVMax.X - ( StartUV.X - UVMin.X );
+				EndUV.X = UVMax.X - ( EndUV.X - UVMin.X );
 			}
+			if( bMirrorVertical )
+			{
+				StartUV.Y = UVMax.Y - ( StartUV.Y - UVMin.Y );
+				EndUV.Y = UVMax.Y - ( EndUV.Y - UVMin.Y );
+			}
+		}
 
-			// Add four vertices to the list of verts to be added to the vertex buffer
-			BatchVertices.Add( FSlateVertex( RenderTransform, TopLeft,	StartUV,						Tiling,	Tint, RenderClipRect ) );
-			BatchVertices.Add( FSlateVertex( RenderTransform, TopRight,	FVector2D(EndUV.X,StartUV.Y),	Tiling,	Tint, RenderClipRect ) );
-			BatchVertices.Add( FSlateVertex( RenderTransform, BotLeft,	FVector2D(StartUV.X,EndUV.Y),	Tiling,	Tint, RenderClipRect ) );
-			BatchVertices.Add( FSlateVertex( RenderTransform, BotRight,	EndUV,							Tiling,	Tint, RenderClipRect ) );
+		// Add four vertices to the list of verts to be added to the vertex buffer
+		BatchVertices.Add(FSlateVertex::Make<Rounding>(RenderTransform, TopLeft, LocalSize, DrawScale, FVector4(StartUV, Tiling), Tint, RenderClipRect));
+		BatchVertices.Add(FSlateVertex::Make<Rounding>(RenderTransform, TopRight, LocalSize, DrawScale, FVector4(FVector2D(EndUV.X, StartUV.Y), Tiling), Tint, RenderClipRect));
+		BatchVertices.Add(FSlateVertex::Make<Rounding>(RenderTransform, BotLeft, LocalSize, DrawScale, FVector4(FVector2D(StartUV.X, EndUV.Y), Tiling), Tint, RenderClipRect));
+		BatchVertices.Add(FSlateVertex::Make<Rounding>(RenderTransform, BotRight, LocalSize, DrawScale, FVector4(EndUV, Tiling), Tint, RenderClipRect));
 
-			BatchIndices.Add( IndexStart + 0 );
-			BatchIndices.Add( IndexStart + 1 );
-			BatchIndices.Add( IndexStart + 2 );
+		BatchIndices.Add( IndexStart + 0 );
+		BatchIndices.Add( IndexStart + 1 );
+		BatchIndices.Add( IndexStart + 2 );
 
-			BatchIndices.Add( IndexStart + 2 );
-			BatchIndices.Add( IndexStart + 1 );
-			BatchIndices.Add( IndexStart + 3 );
-		}
+		BatchIndices.Add( IndexStart + 2 );
+		BatchIndices.Add( IndexStart + 1 );
+		BatchIndices.Add( IndexStart + 3 );
 	}
 }
 
+template<ESlateVertexRounding Rounding>
 void FSlateElementBatcher::AddTextElement(const FSlateDrawElement& DrawElement)
 {
 	const FSlateDataPayload& InPayload = DrawElement.GetDataPayload();
@@ -608,7 +603,7 @@ void FSlateElementBatcher::AddTextElement(const FSlateDrawElement& DrawElement)
 	NumDrawnTextsStat++;
 
 	const FSlateRect& InClippingRect = DrawElement.GetClippingRect();
-	ESlateDrawEffect::Type InDrawEffects = DrawElement.GetDrawEffects();
+	ESlateDrawEffect InDrawEffects = DrawElement.GetDrawEffects();
 	uint32 Layer = DrawElement.GetLayer();
 
 	// extract the layout transform from the draw element
@@ -780,10 +775,10 @@ void FSlateElementBatcher::AddTextElement(const FSlateDrawElement& DrawElement)
 						}
 
 						// Add four vertices to the list of verts to be added to the vertex buffer
-						BatchVerticesRef[ VertexOffset++ ] = FSlateVertex( RenderTransform, UpperLeft,								FVector4(U,V,				Ut,Vt),			FVector2D(0.0f,0.0f), InTint, RenderClipRect );
-						BatchVerticesRef[ VertexOffset++ ] = FSlateVertex( RenderTransform, FVector2D(LowerRight.X,UpperLeft.Y),	FVector4(U+SizeU, V,		UtMax,Vt),		FVector2D(1.0f,0.0f), InTint, RenderClipRect );
-						BatchVerticesRef[ VertexOffset++ ] = FSlateVertex( RenderTransform, FVector2D(UpperLeft.X,LowerRight.Y),	FVector4(U, V+SizeV,		Ut,VtMax),		FVector2D(0.0f,1.0f), InTint, RenderClipRect );
-						BatchVerticesRef[ VertexOffset++ ] = FSlateVertex( RenderTransform, LowerRight,								FVector4(U+SizeU, V+SizeV,	UtMax,VtMax),	FVector2D(1.0f,1.0f), InTint, RenderClipRect );
+						BatchVerticesRef[ VertexOffset++ ] = FSlateVertex::Make<Rounding>( RenderTransform, UpperLeft,								FVector4(U,V,				Ut,Vt),			FVector2D(0.0f,0.0f), InTint, RenderClipRect );
+						BatchVerticesRef[ VertexOffset++ ] = FSlateVertex::Make<Rounding>( RenderTransform, FVector2D(LowerRight.X,UpperLeft.Y),	FVector4(U+SizeU, V,		UtMax,Vt),		FVector2D(1.0f,0.0f), InTint, RenderClipRect );
+						BatchVerticesRef[ VertexOffset++ ] = FSlateVertex::Make<Rounding>( RenderTransform, FVector2D(UpperLeft.X,LowerRight.Y),	FVector4(U, V+SizeV,		Ut,VtMax),		FVector2D(0.0f,1.0f), InTint, RenderClipRect );
+						BatchVerticesRef[ VertexOffset++ ] = FSlateVertex::Make<Rounding>( RenderTransform, LowerRight,								FVector4(U+SizeU, V+SizeV,	UtMax,VtMax),	FVector2D(1.0f,1.0f), InTint, RenderClipRect );
 
 						BatchIndicesRef[IndexOffset++] = IndexStart + 0;
 						BatchIndicesRef[IndexOffset++] = IndexStart + 1;
@@ -817,6 +812,7 @@ void FSlateElementBatcher::AddTextElement(const FSlateDrawElement& DrawElement)
 	}
 }
 
+template<ESlateVertexRounding Rounding>
 void FSlateElementBatcher::AddShapedTextElement( const FSlateDrawElement& DrawElement )
 {
 	const FSlateDataPayload& InPayload = DrawElement.GetDataPayload();
@@ -847,7 +843,7 @@ void FSlateElementBatcher::AddShapedTextElement( const FSlateDrawElement& DrawEl
 	NumDrawnTextsStat++;
 
 	const FSlateRect& InClippingRect = DrawElement.GetClippingRect();
-	ESlateDrawEffect::Type InDrawEffects = DrawElement.GetDrawEffects();
+	ESlateDrawEffect InDrawEffects = DrawElement.GetDrawEffects();
 	uint32 Layer = DrawElement.GetLayer();
 
 	// extract the layout transform from the draw element
@@ -976,10 +972,10 @@ void FSlateElementBatcher::AddShapedTextElement( const FSlateDrawElement& DrawEl
 						}
 
 						// Add four vertices to the list of verts to be added to the vertex buffer
-						BatchVerticesRef[VertexOffset++] = FSlateVertex(RenderTransform, UpperLeft,								FVector4(U, V, Ut, Vt),							FVector2D(0.0f, 0.0f), InTint, RenderClipRect);
-						BatchVerticesRef[VertexOffset++] = FSlateVertex(RenderTransform, FVector2D(LowerRight.X, UpperLeft.Y),	FVector4(U + SizeU, V, UtMax, Vt),				FVector2D(1.0f, 0.0f), InTint, RenderClipRect);
-						BatchVerticesRef[VertexOffset++] = FSlateVertex(RenderTransform, FVector2D(UpperLeft.X, LowerRight.Y),	FVector4(U, V + SizeV, Ut, VtMax),				FVector2D(0.0f, 1.0f), InTint, RenderClipRect);
-						BatchVerticesRef[VertexOffset++] = FSlateVertex(RenderTransform, LowerRight,							FVector4(U + SizeU, V + SizeV, UtMax, VtMax),	FVector2D(1.0f, 1.0f), InTint, RenderClipRect);
+						BatchVerticesRef[VertexOffset++] = FSlateVertex::Make<Rounding>(RenderTransform, UpperLeft,								FVector4(U, V, Ut, Vt),							FVector2D(0.0f, 0.0f), InTint, RenderClipRect);
+						BatchVerticesRef[VertexOffset++] = FSlateVertex::Make<Rounding>(RenderTransform, FVector2D(LowerRight.X, UpperLeft.Y),	FVector4(U + SizeU, V, UtMax, Vt),				FVector2D(1.0f, 0.0f), InTint, RenderClipRect);
+						BatchVerticesRef[VertexOffset++] = FSlateVertex::Make<Rounding>(RenderTransform, FVector2D(UpperLeft.X, LowerRight.Y),	FVector4(U, V + SizeV, Ut, VtMax),				FVector2D(0.0f, 1.0f), InTint, RenderClipRect);
+						BatchVerticesRef[VertexOffset++] = FSlateVertex::Make<Rounding>(RenderTransform, LowerRight,							FVector4(U + SizeU, V + SizeV, UtMax, VtMax),	FVector2D(1.0f, 1.0f), InTint, RenderClipRect);
 
 						BatchIndicesRef[IndexOffset++] = IndexStart + 0;
 						BatchIndicesRef[IndexOffset++] = IndexStart + 1;
@@ -1015,6 +1011,7 @@ void FSlateElementBatcher::AddShapedTextElement( const FSlateDrawElement& DrawEl
 
 }
 
+template<ESlateVertexRounding Rounding>
 void FSlateElementBatcher::AddGradientElement( const FSlateDrawElement& DrawElement )
 {
 	const FSlateRenderTransform& RenderTransform = DrawElement.GetRenderTransform();
@@ -1024,12 +1021,10 @@ void FSlateElementBatcher::AddGradientElement( const FSlateDrawElement& DrawElem
 	//float Scale = DrawElement.GetScale();
 	const FSlateDataPayload& InPayload = DrawElement.GetDataPayload();
 	const FSlateRect& InClippingRect = DrawElement.GetClippingRect();
-	ESlateDrawEffect::Type InDrawEffects = DrawElement.GetDrawEffects();
+	ESlateDrawEffect InDrawEffects = DrawElement.GetDrawEffects();
 	uint32 Layer = DrawElement.GetLayer();
 
-	// extract the layout transform from the draw element
-	FSlateLayoutTransform InverseLayoutTransform(Inverse(FSlateLayoutTransform(DrawElement.GetScale(), DrawElement.GetPosition())));
-	FSlateRotatedClipRectType RenderClipRect = FSlateRotatedClipRectType::MakeSnappedRotatedRect(InClippingRect, InverseLayoutTransform, RenderTransform);
+	FSlateRotatedClipRectType RenderClipRect = DrawElement.CalculateRenderClippingRect();
 
 	// There must be at least one gradient stop
 	check( InPayload.GradientStops.Num() > 0 );
@@ -1130,14 +1125,14 @@ void FSlateElementBatcher::AddGradientElement( const FSlateDrawElement& DrawElem
 		if( StopIndex == 0 )
 		{
 			// First stop does not have a full quad yet so do not create indices
-			BatchVertices.Add( FSlateVertex(RenderTransform, StartPt, FVector2D::ZeroVector, FVector2D::ZeroVector, CurStop.Color.ToFColor(false), RenderClipRect ) );
-			BatchVertices.Add( FSlateVertex(RenderTransform, EndPt, FVector2D::ZeroVector, FVector2D::ZeroVector, CurStop.Color.ToFColor(false), RenderClipRect ) );
+			BatchVertices.Add(FSlateVertex::Make<Rounding>(RenderTransform, StartPt, FVector2D::ZeroVector, FVector2D::ZeroVector, CurStop.Color.ToFColor(false), RenderClipRect ) );
+			BatchVertices.Add(FSlateVertex::Make<Rounding>(RenderTransform, EndPt, FVector2D::ZeroVector, FVector2D::ZeroVector, CurStop.Color.ToFColor(false), RenderClipRect ) );
 		}
 		else
 		{
 			// All stops after the first have indices and generate quads
-			BatchVertices.Add( FSlateVertex(RenderTransform, StartPt, FVector2D::ZeroVector, FVector2D::ZeroVector, CurStop.Color.ToFColor(false), RenderClipRect ) );
-			BatchVertices.Add( FSlateVertex(RenderTransform, EndPt, FVector2D::ZeroVector, FVector2D::ZeroVector, CurStop.Color.ToFColor(false), RenderClipRect ) );
+			BatchVertices.Add(FSlateVertex::Make<Rounding>(RenderTransform, StartPt, FVector2D::ZeroVector, FVector2D::ZeroVector, CurStop.Color.ToFColor(false), RenderClipRect ) );
+			BatchVertices.Add(FSlateVertex::Make<Rounding>(RenderTransform, EndPt, FVector2D::ZeroVector, FVector2D::ZeroVector, CurStop.Color.ToFColor(false), RenderClipRect ) );
 
 			// Connect the indices to the previous vertices
 			BatchIndices.Add( IndexStart - 2 );
@@ -1151,6 +1146,7 @@ void FSlateElementBatcher::AddGradientElement( const FSlateDrawElement& DrawElem
 	}
 }
 
+template<ESlateVertexRounding Rounding>
 void FSlateElementBatcher::AddSplineElement( const FSlateDrawElement& DrawElement )
 {
 	const FSlateRenderTransform& RenderTransform = DrawElement.GetRenderTransform();
@@ -1160,12 +1156,10 @@ void FSlateElementBatcher::AddSplineElement( const FSlateDrawElement& DrawElemen
 	//float Scale = DrawElement.GetScale();
 	const FSlateDataPayload& InPayload = DrawElement.GetDataPayload();
 	const FSlateRect& InClippingRect = DrawElement.GetClippingRect();
-	ESlateDrawEffect::Type InDrawEffects = DrawElement.GetDrawEffects();
+	ESlateDrawEffect InDrawEffects = DrawElement.GetDrawEffects();
 	uint32 Layer = DrawElement.GetLayer();
 
-	// extract the layout transform from the draw element
-	FSlateLayoutTransform InverseLayoutTransform(Inverse(FSlateLayoutTransform(DrawElement.GetScale(), DrawElement.GetPosition())));
-	FSlateRotatedClipRectType RenderClipRect = FSlateRotatedClipRectType::MakeSnappedRotatedRect(InClippingRect, InverseLayoutTransform, RenderTransform);
+	const FSlateRotatedClipRectType RenderClipRect = DrawElement.CalculateRenderClippingRect();
 
 	//@todo SLATE: Merge with AddLineElement?
 
@@ -1185,7 +1179,7 @@ void FSlateElementBatcher::AddSplineElement( const FSlateDrawElement& DrawElemen
 
 	// 1 is the minimum thickness we support
 	// Thickness is given in screenspace, so convert it to local space before proceeding.
-	float InThickness = FMath::Max( 1.0f, InverseLayoutTransform.GetScale() * InPayload.Thickness );
+	float InThickness = FMath::Max( 1.0f, DrawElement.GetInverseLayoutTransform().GetScale() * InPayload.Thickness );
 
 	// The radius to use when checking the distance of pixels to the actual line.  Arbitrary value based on what looks the best
 	const float Radius = 1.5f;
@@ -1219,8 +1213,8 @@ void FSlateElementBatcher::AddSplineElement( const FSlateDrawElement& DrawElemen
 
 	FColor VertexCol = bColorGradient ? PackVertexColor(InPayload.GradientStops[0].Color) : PackVertexColor(InPayload.Tint);
 
-	BatchVertices.Add( FSlateVertex( RenderTransform, StartPos + Up, TransformPoint(RenderTransform, StartPos), TransformPoint(RenderTransform, EndPos), VertexCol, RenderClipRect ) );
-	BatchVertices.Add( FSlateVertex( RenderTransform, StartPos - Up, TransformPoint(RenderTransform, StartPos), TransformPoint(RenderTransform, EndPos), VertexCol, RenderClipRect ) );
+	BatchVertices.Add(FSlateVertex::Make<Rounding>( RenderTransform, StartPos + Up, TransformPoint(RenderTransform, StartPos), TransformPoint(RenderTransform, EndPos), VertexCol, RenderClipRect ) );
+	BatchVertices.Add(FSlateVertex::Make<Rounding>( RenderTransform, StartPos - Up, TransformPoint(RenderTransform, StartPos), TransformPoint(RenderTransform, EndPos), VertexCol, RenderClipRect ) );
 
 	// Generate the rest of the segments
 	for( int32 Step = 0; Step < NumSteps; ++Step )
@@ -1247,8 +1241,8 @@ void FSlateElementBatcher::AddSplineElement( const FSlateDrawElement& DrawElemen
 		// Create the new vertices for the thick line segment
 		Up = SegmentNormal * HalfThickness;
 
-		BatchVertices.Add( FSlateVertex( RenderTransform, EndPos + Up, TransformPoint(RenderTransform, StartPos), TransformPoint(RenderTransform, EndPos), VertexCol, RenderClipRect ) );
-		BatchVertices.Add( FSlateVertex( RenderTransform, EndPos - Up, TransformPoint(RenderTransform, StartPos), TransformPoint(RenderTransform, EndPos), VertexCol, RenderClipRect ) );
+		BatchVertices.Add(FSlateVertex::Make<Rounding>( RenderTransform, EndPos + Up, TransformPoint(RenderTransform, StartPos), TransformPoint(RenderTransform, EndPos), VertexCol, RenderClipRect ) );
+		BatchVertices.Add(FSlateVertex::Make<Rounding>( RenderTransform, EndPos - Up, TransformPoint(RenderTransform, StartPos), TransformPoint(RenderTransform, EndPos), VertexCol, RenderClipRect ) );
 
 		BatchIndices.Add( IndexStart - 2 );
 		BatchIndices.Add( IndexStart - 1 );
@@ -1305,7 +1299,7 @@ static bool LineIntersect( const FVector2D& P1, const FVector2D& P2, const FVect
 	return false;
 }
 	
-
+template<ESlateVertexRounding Rounding>
 void FSlateElementBatcher::AddLineElement( const FSlateDrawElement& DrawElement )
 {
 	const FSlateRenderTransform& RenderTransform = DrawElement.GetRenderTransform();
@@ -1315,12 +1309,10 @@ void FSlateElementBatcher::AddLineElement( const FSlateDrawElement& DrawElement
 	//float Scale = DrawElement.GetScale();
 	const FSlateDataPayload& InPayload = DrawElement.GetDataPayload();
 	const FSlateRect& InClippingRect = DrawElement.GetClippingRect();
-	ESlateDrawEffect::Type DrawEffects = DrawElement.GetDrawEffects();
+	ESlateDrawEffect DrawEffects = DrawElement.GetDrawEffects();
 	uint32 Layer = DrawElement.GetLayer();
 
-	// extract the layout transform from the draw element
-	FSlateLayoutTransform InverseLayoutTransform(Inverse(FSlateLayoutTransform(DrawElement.GetScale(), DrawElement.GetPosition())));
-	FSlateRotatedClipRectType RenderClipRect = FSlateRotatedClipRectType::MakeSnappedRotatedRect(InClippingRect, InverseLayoutTransform, RenderTransform);
+	const FSlateRotatedClipRectType RenderClipRect = DrawElement.CalculateRenderClippingRect();
 
 	if( InPayload.Points.Num() < 2 )
 	{
@@ -1358,8 +1350,8 @@ void FSlateElementBatcher::AddLineElement( const FSlateDrawElement& DrawElement
 
 		FVector2D Up = Normal * HalfThickness;
 
-		BatchVertices.Add( FSlateVertex( RenderTransform, StartPos + Up, TransformPoint(RenderTransform, StartPos), TransformPoint(RenderTransform, EndPos), FinalColor, RenderClipRect ) );
-		BatchVertices.Add( FSlateVertex( RenderTransform, StartPos - Up, TransformPoint(RenderTransform, StartPos), TransformPoint(RenderTransform, EndPos), FinalColor, RenderClipRect ) );
+		BatchVertices.Add(FSlateVertex::Make<Rounding>( RenderTransform, StartPos + Up, TransformPoint(RenderTransform, StartPos), TransformPoint(RenderTransform, EndPos), FinalColor, RenderClipRect ) );
+		BatchVertices.Add(FSlateVertex::Make<Rounding>( RenderTransform, StartPos - Up, TransformPoint(RenderTransform, StartPos), TransformPoint(RenderTransform, EndPos), FinalColor, RenderClipRect ) );
 
 		// Generate the rest of the segments
 		for( int32 Point = 1; Point < Points.Num(); ++Point )
@@ -1438,8 +1430,8 @@ void FSlateElementBatcher::AddLineElement( const FSlateDrawElement& DrawElement
 				BatchVertices.Add( StartV1 );
 			}
 
-			BatchVertices.Add( FSlateVertex( RenderTransform, IntersectUpper, StartPosRenderSpace, IntersectCenterRenderSpace, FinalColor, RenderClipRect ) );
-			BatchVertices.Add( FSlateVertex( RenderTransform, IntersectLower, StartPosRenderSpace, IntersectCenterRenderSpace, FinalColor, RenderClipRect ) );
+			BatchVertices.Add(FSlateVertex::Make<Rounding>( RenderTransform, IntersectUpper, StartPosRenderSpace, IntersectCenterRenderSpace, FinalColor, RenderClipRect ) );
+			BatchVertices.Add(FSlateVertex::Make<Rounding>( RenderTransform, IntersectLower, StartPosRenderSpace, IntersectCenterRenderSpace, FinalColor, RenderClipRect ) );
 
 			BatchIndices.Add( IndexStart - 1 );
 			BatchIndices.Add( IndexStart - 2 );
@@ -1467,8 +1459,8 @@ void FSlateElementBatcher::AddLineElement( const FSlateDrawElement& DrawElement
 			FVector2D EndPos = InPayload.Points[Point+1];
 
 
-			BatchVertices.Add( FSlateVertex( RenderTransform, StartPos, FVector2D::ZeroVector, FinalColor, RenderClipRect ) );
-			BatchVertices.Add( FSlateVertex( RenderTransform, EndPos, FVector2D::ZeroVector, FinalColor, RenderClipRect ) );
+			BatchVertices.Add( FSlateVertex::Make<Rounding>( RenderTransform, StartPos, FVector2D::ZeroVector, FinalColor, RenderClipRect ) );
+			BatchVertices.Add( FSlateVertex::Make<Rounding>( RenderTransform, EndPos, FVector2D::ZeroVector, FinalColor, RenderClipRect ) );
 
 			BatchIndices.Add( IndexStart );
 			BatchIndices.Add( IndexStart + 1 );
@@ -1476,7 +1468,7 @@ void FSlateElementBatcher::AddLineElement( const FSlateDrawElement& DrawElement
 	}
 }
 
-
+template<ESlateVertexRounding Rounding>
 void FSlateElementBatcher::AddViewportElement( const FSlateDrawElement& DrawElement )
 {
 	const FSlateRenderTransform& RenderTransform = DrawElement.GetRenderTransform();
@@ -1485,16 +1477,14 @@ void FSlateElementBatcher::AddViewportElement( const FSlateDrawElement& DrawElem
 	const FVector2D& LocalSize = DrawElement.GetLocalSize();
 	const FSlateDataPayload& InPayload = DrawElement.GetDataPayload();
 	const FSlateRect& InClippingRect = DrawElement.GetClippingRect();
-	ESlateDrawEffect::Type InDrawEffects = DrawElement.GetDrawEffects();
+	ESlateDrawEffect InDrawEffects = DrawElement.GetDrawEffects();
 	uint32 Layer = DrawElement.GetLayer();
 
-	// extract the layout transform from the draw element
-	FSlateLayoutTransform InverseLayoutTransform(Inverse(FSlateLayoutTransform(DrawElement.GetScale(), DrawElement.GetPosition())));
-	FSlateRotatedClipRectType RenderClipRect = FSlateRotatedClipRectType::MakeSnappedRotatedRect(InClippingRect, InverseLayoutTransform, RenderTransform);
+	const FSlateRotatedClipRectType RenderClipRect = DrawElement.CalculateRenderClippingRect();
 
 	const FColor FinalColor = PackVertexColor(InPayload.Tint);
 
-	ESlateBatchDrawFlag::Type DrawFlags = InPayload.BatchFlags;
+	ESlateBatchDrawFlag DrawFlags = InPayload.BatchFlags;
 
 	FSlateShaderResource* ViewportResource = InPayload.RenderTargetResource;
 	ESlateShader::Type ShaderType = ESlateShader::Default;
@@ -1534,10 +1524,10 @@ void FSlateElementBatcher::AddViewportElement( const FSlateDrawElement& DrawElem
 	uint32 IndexStart = BatchVertices.Num();
 
 	// Add four vertices to the list of verts to be added to the vertex buffer
-	BatchVertices.Add( FSlateVertex( RenderTransform, TopLeft,	FVector2D(0.0f,0.0f),	FinalColor, RenderClipRect ) );
-	BatchVertices.Add( FSlateVertex( RenderTransform, TopRight,	FVector2D(1.0f,0.0f),	FinalColor, RenderClipRect ) );
-	BatchVertices.Add( FSlateVertex( RenderTransform, BotLeft,	FVector2D(0.0f,1.0f),	FinalColor, RenderClipRect ) );
-	BatchVertices.Add( FSlateVertex( RenderTransform, BotRight,	FVector2D(1.0f,1.0f),	FinalColor, RenderClipRect ) );
+	BatchVertices.Add( FSlateVertex::Make<Rounding>( RenderTransform, TopLeft,	FVector2D(0.0f,0.0f),	FinalColor, RenderClipRect ) );
+	BatchVertices.Add( FSlateVertex::Make<Rounding>( RenderTransform, TopRight,	FVector2D(1.0f,0.0f),	FinalColor, RenderClipRect ) );
+	BatchVertices.Add( FSlateVertex::Make<Rounding>( RenderTransform, BotLeft,	FVector2D(0.0f,1.0f),	FinalColor, RenderClipRect ) );
+	BatchVertices.Add( FSlateVertex::Make<Rounding>( RenderTransform, BotRight,	FVector2D(1.0f,1.0f),	FinalColor, RenderClipRect ) );
 
 	// The offset into the index buffer where this quads indices start
 	uint32 IndexOffsetStart = BatchIndices.Num();
@@ -1552,7 +1542,7 @@ void FSlateElementBatcher::AddViewportElement( const FSlateDrawElement& DrawElem
 	BatchIndices.Add( IndexStart + 3 );
 }
 
-
+template<ESlateVertexRounding Rounding>
 void FSlateElementBatcher::AddBorderElement( const FSlateDrawElement& DrawElement )
 {
 	const FSlateRenderTransform& RenderTransform = DrawElement.GetRenderTransform();
@@ -1562,12 +1552,12 @@ void FSlateElementBatcher::AddBorderElement( const FSlateDrawElement& DrawElemen
 	//float Scale = DrawElement.GetScale();
 	const FSlateDataPayload& InPayload = DrawElement.GetDataPayload();
 	const FSlateRect& InClippingRect = DrawElement.GetClippingRect();
-	ESlateDrawEffect::Type InDrawEffects = DrawElement.GetDrawEffects();
+	ESlateDrawEffect InDrawEffects = DrawElement.GetDrawEffects();
 	uint32 Layer = DrawElement.GetLayer();
 
-	// extract the layout transform from the draw element
-	FSlateLayoutTransform InverseLayoutTransform(Inverse(FSlateLayoutTransform(DrawElement.GetScale(), DrawElement.GetPosition())));
-	FSlateRotatedClipRectType RenderClipRect = FSlateRotatedClipRectType::MakeSnappedRotatedRect(InClippingRect, InverseLayoutTransform, RenderTransform);
+	const float DrawScale = DrawElement.GetScale();
+
+	const FSlateRotatedClipRectType RenderClipRect = DrawElement.CalculateRenderClippingRect();
 
 	check( InPayload.BrushResource );
 
@@ -1582,7 +1572,7 @@ void FSlateElementBatcher::AddBorderElement( const FSlateDrawElement& DrawElemen
 		TextureWidth = Resource->GetWidth();
 		TextureHeight = Resource->GetHeight();
 	}
-	FVector2D TextureSizeLocalSpace = TransformVector(InverseLayoutTransform, FVector2D(TextureWidth, TextureHeight));
+	FVector2D TextureSizeLocalSpace = TransformVector(DrawElement.GetInverseLayoutTransform(), FVector2D(TextureWidth, TextureHeight));
  
 	// Texel offset
 	const FVector2D HalfTexel( PixelCenterOffset/TextureWidth, PixelCenterOffset/TextureHeight );
@@ -1639,7 +1629,7 @@ void FSlateElementBatcher::AddBorderElement( const FSlateDrawElement& DrawElemen
 	const FColor Tint = PackVertexColor(InPayload.Tint);
 
 	// Pass the tiling information as a flag so we can pick the correct texture addressing mode
-	ESlateBatchDrawFlag::Type DrawFlags = (ESlateBatchDrawFlag::TileU|ESlateBatchDrawFlag::TileV);
+	ESlateBatchDrawFlag DrawFlags = (ESlateBatchDrawFlag::TileU|ESlateBatchDrawFlag::TileV);
 
 	FSlateElementBatch& ElementBatch = FindBatchForElement( Layer, ShaderParams, Resource, ESlateDrawPrimitive::TriangleList, ESlateShader::Border, InDrawEffects, DrawFlags, DrawElement.GetScissorRect(), DrawElement.GetSceneIndex());
 	FSlateVertexArray& BatchVertices = BatchData->GetBatchVertexList(ElementBatch);
@@ -1661,45 +1651,45 @@ void FSlateElementBatcher::AddBorderElement( const FSlateDrawElement& DrawElemen
 	FVector2D Zero(0,0);
 
 	// Add all the vertices needed for this element.  Vertices are duplicated so that we can have some sections with no tiling and some with tiling.
-	BatchVertices.Add( FSlateVertex( RenderTransform, Position,									FVector2D( StartUV.X, StartUV.Y ),			Zero,							Tint, RenderClipRect ) ); //0
-	BatchVertices.Add( FSlateVertex( RenderTransform, FVector2D( Position.X, TopMarginY ),		FVector2D( StartUV.X, TopMarginV ),			Zero,							Tint, RenderClipRect ) ); //1
-	BatchVertices.Add( FSlateVertex( RenderTransform, FVector2D( LeftMarginX, Position.Y ),		FVector2D( LeftMarginU, StartUV.Y ),		Zero,							Tint, RenderClipRect ) ); //2
-	BatchVertices.Add( FSlateVertex( RenderTransform, FVector2D( LeftMarginX, TopMarginY ),		FVector2D( LeftMarginU, TopMarginV ),		Zero,							Tint, RenderClipRect ) ); //3
-	
-	BatchVertices.Add( FSlateVertex( RenderTransform, FVector2D( LeftMarginX, Position.Y ),		FVector2D( StartUV.X, StartUV.Y ),			FVector2D(TopTiling, 0.0f),		Tint, RenderClipRect ) ); //4
-	BatchVertices.Add( FSlateVertex( RenderTransform, FVector2D( LeftMarginX, TopMarginY ),		FVector2D( StartUV.X, TopMarginV ),			FVector2D(TopTiling, 0.0f),		Tint, RenderClipRect ) ); //5
-	BatchVertices.Add( FSlateVertex( RenderTransform, FVector2D( RightMarginX, Position.Y ),	FVector2D( EndUV.X, StartUV.Y ),			FVector2D(TopTiling, 0.0f),		Tint, RenderClipRect ) ); //6
-	BatchVertices.Add( FSlateVertex( RenderTransform, FVector2D( RightMarginX, TopMarginY ),	FVector2D( EndUV.X, TopMarginV),			FVector2D(TopTiling, 0.0f),		Tint, RenderClipRect ) ); //7
-	
-	BatchVertices.Add( FSlateVertex( RenderTransform, FVector2D( RightMarginX, Position.Y ),	FVector2D( RightMarginU, StartUV.Y ),		Zero,							Tint, RenderClipRect ) ); //8
-	BatchVertices.Add( FSlateVertex( RenderTransform, FVector2D( RightMarginX, TopMarginY ),	FVector2D( RightMarginU, TopMarginV),		Zero,							Tint, RenderClipRect ) ); //9
-	BatchVertices.Add( FSlateVertex( RenderTransform, FVector2D( EndPos.X, Position.Y ),		FVector2D( EndUV.X, StartUV.Y ),			Zero,							Tint, RenderClipRect ) ); //10
-	BatchVertices.Add( FSlateVertex( RenderTransform, FVector2D( EndPos.X, TopMarginY ),		FVector2D( EndUV.X, TopMarginV),			Zero,							Tint, RenderClipRect ) ); //11
-	
-	BatchVertices.Add( FSlateVertex( RenderTransform, FVector2D( Position.X, TopMarginY ),		FVector2D( StartUV.X, StartUV.Y ),			FVector2D(0.0f, LeftTiling),	Tint, RenderClipRect ) ); //12
-	BatchVertices.Add( FSlateVertex( RenderTransform, FVector2D( Position.X, BottomMarginY ),	FVector2D( StartUV.X, EndUV.Y ),			FVector2D(0.0f, LeftTiling),	Tint, RenderClipRect ) ); //13
-	BatchVertices.Add( FSlateVertex( RenderTransform, FVector2D( LeftMarginX, TopMarginY ),		FVector2D( LeftMarginU, StartUV.Y ),		FVector2D(0.0f, LeftTiling),	Tint, RenderClipRect ) ); //14
-	BatchVertices.Add( FSlateVertex( RenderTransform, FVector2D( LeftMarginX, BottomMarginY ),	FVector2D( LeftMarginU, EndUV.Y ),			FVector2D(0.0f, LeftTiling),	Tint, RenderClipRect ) ); //15
-	
-	BatchVertices.Add( FSlateVertex( RenderTransform, FVector2D( RightMarginX, TopMarginY ),	FVector2D( RightMarginU, StartUV.Y ),		FVector2D(0.0f, RightTiling),	Tint, RenderClipRect ) ); //16
-	BatchVertices.Add( FSlateVertex( RenderTransform, FVector2D( RightMarginX, BottomMarginY ),	FVector2D( RightMarginU, EndUV.Y ),			FVector2D(0.0f, RightTiling),	Tint, RenderClipRect ) ); //17
-	BatchVertices.Add( FSlateVertex( RenderTransform, FVector2D( EndPos.X, TopMarginY ),		FVector2D( EndUV.X, StartUV.Y ),			FVector2D(0.0f, RightTiling),	Tint, RenderClipRect ) ); //18
-	BatchVertices.Add( FSlateVertex( RenderTransform, FVector2D( EndPos.X, BottomMarginY ),		FVector2D( EndUV.X, EndUV.Y ),				FVector2D(0.0f, RightTiling),	Tint, RenderClipRect ) ); //19
-	
-	BatchVertices.Add( FSlateVertex( RenderTransform, FVector2D( Position.X, BottomMarginY ),	FVector2D( StartUV.X, BottomMarginV ),		Zero,							Tint, RenderClipRect ) ); //20
-	BatchVertices.Add( FSlateVertex( RenderTransform, FVector2D( Position.X, EndPos.Y ),		FVector2D( StartUV.X, EndUV.Y ),			Zero,							Tint, RenderClipRect ) ); //21
-	BatchVertices.Add( FSlateVertex( RenderTransform, FVector2D( LeftMarginX, BottomMarginY ),	FVector2D( LeftMarginU, BottomMarginV ),	Zero,							Tint, RenderClipRect ) ); //22
-	BatchVertices.Add( FSlateVertex( RenderTransform, FVector2D( LeftMarginX, EndPos.Y ),		FVector2D( LeftMarginU, EndUV.Y ),			Zero,							Tint, RenderClipRect ) ); //23
-	
-	BatchVertices.Add( FSlateVertex( RenderTransform, FVector2D( LeftMarginX, BottomMarginY ),	FVector2D( StartUV.X, BottomMarginV ),		FVector2D(BottomTiling, 0.0f),	Tint, RenderClipRect ) ); //24
-	BatchVertices.Add( FSlateVertex( RenderTransform, FVector2D( LeftMarginX, EndPos.Y ),		FVector2D( StartUV.X, EndUV.Y ),			FVector2D(BottomTiling, 0.0f),	Tint, RenderClipRect ) ); //25
-	BatchVertices.Add( FSlateVertex( RenderTransform, FVector2D( RightMarginX,BottomMarginY ),	FVector2D( EndUV.X, BottomMarginV ),		FVector2D(BottomTiling, 0.0f),	Tint, RenderClipRect ) ); //26
-	BatchVertices.Add( FSlateVertex( RenderTransform, FVector2D( RightMarginX, EndPos.Y ),		FVector2D( EndUV.X, EndUV.Y ),				FVector2D(BottomTiling, 0.0f),	Tint, RenderClipRect ) ); //27
-	
-	BatchVertices.Add( FSlateVertex( RenderTransform, FVector2D( RightMarginX, BottomMarginY ),	FVector2D( RightMarginU, BottomMarginV ),	Zero,							Tint, RenderClipRect ) ); //29
-	BatchVertices.Add( FSlateVertex( RenderTransform, FVector2D( RightMarginX, EndPos.Y ),		FVector2D( RightMarginU, EndUV.Y ),			Zero,							Tint, RenderClipRect ) ); //30
-	BatchVertices.Add( FSlateVertex( RenderTransform, FVector2D( EndPos.X, BottomMarginY ),		FVector2D( EndUV.X, BottomMarginV ),		Zero,							Tint, RenderClipRect ) ); //31
-	BatchVertices.Add( FSlateVertex( RenderTransform, FVector2D( EndPos.X, EndPos.Y ),			FVector2D( EndUV.X, EndUV.Y ),				Zero,							Tint, RenderClipRect ) ); //32
+	BatchVertices.Add( FSlateVertex::Make<Rounding>( RenderTransform, Position,									LocalSize, DrawScale, FVector4( StartUV.X, StartUV.Y, 0.0f, 0.0f),				Tint, RenderClipRect ) ); //0
+	BatchVertices.Add( FSlateVertex::Make<Rounding>( RenderTransform, FVector2D( Position.X, TopMarginY ),		LocalSize, DrawScale, FVector4( StartUV.X, TopMarginV, 0.0f, 0.0f),				Tint, RenderClipRect ) ); //1
+	BatchVertices.Add( FSlateVertex::Make<Rounding>( RenderTransform, FVector2D( LeftMarginX, Position.Y ),		LocalSize, DrawScale, FVector4( LeftMarginU, StartUV.Y, 0.0f, 0.0f),			Tint, RenderClipRect ) ); //2
+	BatchVertices.Add( FSlateVertex::Make<Rounding>( RenderTransform, FVector2D( LeftMarginX, TopMarginY ),		LocalSize, DrawScale, FVector4( LeftMarginU, TopMarginV, 0.0f, 0.0f),			Tint, RenderClipRect ) ); //3
+
+	BatchVertices.Add( FSlateVertex::Make<Rounding>( RenderTransform, FVector2D( LeftMarginX, Position.Y ),		LocalSize, DrawScale, FVector4( StartUV.X, StartUV.Y, TopTiling, 0.0f),			Tint, RenderClipRect ) ); //4
+	BatchVertices.Add( FSlateVertex::Make<Rounding>( RenderTransform, FVector2D( LeftMarginX, TopMarginY ),		LocalSize, DrawScale, FVector4( StartUV.X, TopMarginV, TopTiling, 0.0f),		Tint, RenderClipRect ) ); //5
+	BatchVertices.Add( FSlateVertex::Make<Rounding>( RenderTransform, FVector2D( RightMarginX, Position.Y ),	LocalSize, DrawScale, FVector4( EndUV.X, StartUV.Y, TopTiling, 0.0f),			Tint, RenderClipRect ) ); //6
+	BatchVertices.Add( FSlateVertex::Make<Rounding>( RenderTransform, FVector2D( RightMarginX, TopMarginY ),	LocalSize, DrawScale, FVector4( EndUV.X, TopMarginV, TopTiling, 0.0f),			Tint, RenderClipRect ) ); //7
+
+	BatchVertices.Add( FSlateVertex::Make<Rounding>( RenderTransform, FVector2D( RightMarginX, Position.Y ),	LocalSize, DrawScale, FVector4( RightMarginU, StartUV.Y, 0.0f, 0.0f),			Tint, RenderClipRect ) ); //8
+	BatchVertices.Add( FSlateVertex::Make<Rounding>( RenderTransform, FVector2D( RightMarginX, TopMarginY ),	LocalSize, DrawScale, FVector4( RightMarginU, TopMarginV, 0.0f, 0.0f),			Tint, RenderClipRect ) ); //9
+	BatchVertices.Add( FSlateVertex::Make<Rounding>( RenderTransform, FVector2D( EndPos.X, Position.Y ),		LocalSize, DrawScale, FVector4( EndUV.X, StartUV.Y, 0.0f, 0.0f),				Tint, RenderClipRect ) ); //10
+	BatchVertices.Add( FSlateVertex::Make<Rounding>( RenderTransform, FVector2D( EndPos.X, TopMarginY ),		LocalSize, DrawScale, FVector4( EndUV.X, TopMarginV, 0.0f, 0.0f),				Tint, RenderClipRect ) ); //11
+
+	BatchVertices.Add( FSlateVertex::Make<Rounding>( RenderTransform, FVector2D( Position.X, TopMarginY ),		LocalSize, DrawScale, FVector4( StartUV.X, StartUV.Y, 0.0f, LeftTiling),		Tint, RenderClipRect ) ); //12
+	BatchVertices.Add( FSlateVertex::Make<Rounding>( RenderTransform, FVector2D( Position.X, BottomMarginY ),	LocalSize, DrawScale, FVector4( StartUV.X, EndUV.Y, 0.0f, LeftTiling),			Tint, RenderClipRect ) ); //13
+	BatchVertices.Add( FSlateVertex::Make<Rounding>( RenderTransform, FVector2D( LeftMarginX, TopMarginY ),		LocalSize, DrawScale, FVector4( LeftMarginU, StartUV.Y, 0.0f, LeftTiling),		Tint, RenderClipRect ) ); //14
+	BatchVertices.Add( FSlateVertex::Make<Rounding>( RenderTransform, FVector2D( LeftMarginX, BottomMarginY ),	LocalSize, DrawScale, FVector4( LeftMarginU, EndUV.Y, 0.0f, LeftTiling),		Tint, RenderClipRect ) ); //15
+
+	BatchVertices.Add( FSlateVertex::Make<Rounding>( RenderTransform, FVector2D( RightMarginX, TopMarginY ),	LocalSize, DrawScale, FVector4( RightMarginU, StartUV.Y, 0.0f, RightTiling),	Tint, RenderClipRect ) ); //16
+	BatchVertices.Add( FSlateVertex::Make<Rounding>( RenderTransform, FVector2D( RightMarginX, BottomMarginY ), LocalSize, DrawScale, FVector4( RightMarginU, EndUV.Y, 0.0f, RightTiling),		Tint, RenderClipRect ) ); //17
+	BatchVertices.Add( FSlateVertex::Make<Rounding>( RenderTransform, FVector2D( EndPos.X, TopMarginY ),		LocalSize, DrawScale, FVector4( EndUV.X, StartUV.Y, 0.0f, RightTiling),			Tint, RenderClipRect ) ); //18
+	BatchVertices.Add( FSlateVertex::Make<Rounding>( RenderTransform, FVector2D( EndPos.X, BottomMarginY ),		LocalSize, DrawScale, FVector4( EndUV.X, EndUV.Y, 0.0f, RightTiling),			Tint, RenderClipRect ) ); //19
+
+	BatchVertices.Add( FSlateVertex::Make<Rounding>( RenderTransform, FVector2D( Position.X, BottomMarginY ),	LocalSize, DrawScale, FVector4( StartUV.X, BottomMarginV, 0.0f, 0.0f),			Tint, RenderClipRect ) ); //20
+	BatchVertices.Add( FSlateVertex::Make<Rounding>( RenderTransform, FVector2D( Position.X, EndPos.Y ),		LocalSize, DrawScale, FVector4( StartUV.X, EndUV.Y, 0.0f, 0.0f),				Tint, RenderClipRect ) ); //21
+	BatchVertices.Add( FSlateVertex::Make<Rounding>( RenderTransform, FVector2D( LeftMarginX, BottomMarginY ),	LocalSize, DrawScale, FVector4( LeftMarginU, BottomMarginV, 0.0f, 0.0f),		Tint, RenderClipRect ) ); //22
+	BatchVertices.Add( FSlateVertex::Make<Rounding>( RenderTransform, FVector2D( LeftMarginX, EndPos.Y ),		LocalSize, DrawScale, FVector4( LeftMarginU, EndUV.Y, 0.0f, 0.0f),				Tint, RenderClipRect ) ); //23
+
+	BatchVertices.Add( FSlateVertex::Make<Rounding>( RenderTransform, FVector2D( LeftMarginX, BottomMarginY ),	LocalSize, DrawScale, FVector4( StartUV.X, BottomMarginV, BottomTiling, 0.0f),	Tint, RenderClipRect ) ); //24
+	BatchVertices.Add( FSlateVertex::Make<Rounding>( RenderTransform, FVector2D( LeftMarginX, EndPos.Y ),		LocalSize, DrawScale, FVector4( StartUV.X, EndUV.Y, BottomTiling, 0.0f),		Tint, RenderClipRect ) ); //25
+	BatchVertices.Add( FSlateVertex::Make<Rounding>( RenderTransform, FVector2D( RightMarginX,BottomMarginY ),	LocalSize, DrawScale, FVector4( EndUV.X, BottomMarginV, BottomTiling, 0.0f),	Tint, RenderClipRect ) ); //26
+	BatchVertices.Add( FSlateVertex::Make<Rounding>( RenderTransform, FVector2D( RightMarginX, EndPos.Y ),		LocalSize, DrawScale, FVector4( EndUV.X, EndUV.Y, BottomTiling, 0.0f),			Tint, RenderClipRect ) ); //27
+
+	BatchVertices.Add( FSlateVertex::Make<Rounding>( RenderTransform, FVector2D( RightMarginX, BottomMarginY ), LocalSize, DrawScale, FVector4( RightMarginU, BottomMarginV, 0.0f, 0.0f),		Tint, RenderClipRect ) ); //29
+	BatchVertices.Add( FSlateVertex::Make<Rounding>( RenderTransform, FVector2D( RightMarginX, EndPos.Y ),		LocalSize, DrawScale, FVector4( RightMarginU, EndUV.Y, 0.0f, 0.0f),				Tint, RenderClipRect ) ); //30
+	BatchVertices.Add( FSlateVertex::Make<Rounding>( RenderTransform, FVector2D( EndPos.X, BottomMarginY ),		LocalSize, DrawScale, FVector4( EndUV.X, BottomMarginV, 0.0f, 0.0f),			Tint, RenderClipRect ) ); //31
+	BatchVertices.Add( FSlateVertex::Make<Rounding>( RenderTransform, FVector2D( EndPos.X, EndPos.Y ),			LocalSize, DrawScale, FVector4( EndUV.X, EndUV.Y, 0.0f, 0.0f),					Tint, RenderClipRect ) ); //32
 
 
 	// The offset into the index buffer where this elements indices start
@@ -1902,9 +1892,7 @@ void FSlateElementBatcher::AddPostProcessPass(const FSlateDrawElement& DrawEleme
 
 	uint32 Layer = DrawElement.GetLayer();
 
-	// extract the layout transform from the draw element
-	FSlateLayoutTransform InverseLayoutTransform(Inverse(FSlateLayoutTransform(DrawElement.GetScale(), Position)));
-	FSlateRotatedClipRectType RenderClipRect = FSlateRotatedClipRectType::MakeSnappedRotatedRect(InClippingRect, InverseLayoutTransform, RenderTransform);
+	const FSlateRotatedClipRectType RenderClipRect = DrawElement.CalculateRenderClippingRect();
 
 	// Determine the four corners of the quad
 	FVector2D TopLeft = FVector2D::ZeroVector;
@@ -1949,8 +1937,8 @@ FSlateElementBatch& FSlateElementBatcher::FindBatchForElement(
 	const FSlateShaderResource* InTexture, 
 	ESlateDrawPrimitive::Type PrimitiveType,
 	ESlateShader::Type ShaderType, 
-	ESlateDrawEffect::Type DrawEffects, 
-	ESlateBatchDrawFlag::Type DrawFlags,
+	ESlateDrawEffect DrawEffects, 
+	ESlateBatchDrawFlag DrawFlags,
 	const TOptional<FShortRect>& ScissorRect,
 	int32 SceneIndex)
 {
diff --git a/Engine/Source/Runtime/SlateCore/Private/Rendering/RenderingCommon.cpp b/Engine/Source/Runtime/SlateCore/Private/Rendering/RenderingCommon.cpp
--- a/Engine/Source/Runtime/SlateCore/Private/Rendering/RenderingCommon.cpp
+++ b/Engine/Source/Runtime/SlateCore/Private/Rendering/RenderingCommon.cpp
@@ -56,18 +56,29 @@ static FVector2D RoundToInt(const FVector2D& Vec)
 	return FVector2D(FMath::RoundToInt(Vec.X), FMath::RoundToInt(Vec.Y));
 }
 
-FSlateRotatedRect FSlateRotatedRect::MakeSnappedRotatedRect(const FSlateRect& ClipRectInLayoutWindowSpace, const FSlateLayoutTransform& InverseLayoutTransform, const FSlateRenderTransform& RenderTransform)
+FSlateRotatedRect FSlateRotatedRect::MakeRotatedRect(const FSlateRect& ClipRectInLayoutWindowSpace, const FTransform2D& LayoutToRenderTransform)
 {
-	FSlateRotatedRect RotatedRect = TransformRect(
-		Concatenate(InverseLayoutTransform, RenderTransform),
-		FSlateRotatedRect(ClipRectInLayoutWindowSpace));
+	const FSlateRotatedRect RotatedRect = TransformRect(LayoutToRenderTransform, FSlateRotatedRect(ClipRectInLayoutWindowSpace));
+
+	const FVector2D TopRight = RotatedRect.TopLeft + RotatedRect.ExtentX;
+	const FVector2D BottomLeft = RotatedRect.TopLeft + RotatedRect.ExtentY;
+
+	return FSlateRotatedClipRectType(
+		RotatedRect.TopLeft,
+		TopRight - RotatedRect.TopLeft,
+		BottomLeft - RotatedRect.TopLeft);
+}
+
+FSlateRotatedRect FSlateRotatedRect::MakeSnappedRotatedRect(const FSlateRect& ClipRectInLayoutWindowSpace, const FTransform2D& LayoutToRenderTransform)
+{
+	const FSlateRotatedRect RotatedRect = TransformRect(LayoutToRenderTransform, FSlateRotatedRect(ClipRectInLayoutWindowSpace));
 
 	// Pixel snapping is done here by rounding the resulting floats to ints, we do this before
 	// calculating the final extents of the clip box otherwise we'll get a smaller clip rect than a visual
 	// rect where each point is individually snapped.
-	FVector2D SnappedTopLeft = RoundToInt(RotatedRect.TopLeft);
-	FVector2D SnappedTopRight = RoundToInt(RotatedRect.TopLeft + RotatedRect.ExtentX);
-	FVector2D SnappedBottomLeft = RoundToInt(RotatedRect.TopLeft + RotatedRect.ExtentY);
+	const FVector2D SnappedTopLeft = RoundToInt(RotatedRect.TopLeft);
+	const FVector2D SnappedTopRight = RoundToInt(RotatedRect.TopLeft + RotatedRect.ExtentX);
+	const FVector2D SnappedBottomLeft = RoundToInt(RotatedRect.TopLeft + RotatedRect.ExtentY);
 
 	//NOTE: We explicitly do not re-snap the extent x/y, it wouldn't be correct to snap again in distance space
 	// even if two points are snapped, their distance wont necessarily be a whole number if those points are not
@@ -78,59 +89,3 @@ FSlateRotatedRect FSlateRotatedRect::MakeSnappedRotatedRect(const FSlateRect& Cl
 		SnappedBottomLeft - SnappedTopLeft);
 }
 
-FSlateVertex::FSlateVertex() 
-{
-}
-
-FSlateVertex::FSlateVertex(const FSlateRenderTransform& RenderTransform, const FVector2D& InLocalPosition, const FVector2D& InTexCoord, const FVector2D& InTexCoord2, const FColor& InColor, const FSlateRotatedClipRectType& InClipRect )
-	: ClipRect( InClipRect )
-	, Color( InColor )
-{
-	TexCoords[0] = InTexCoord.X;
-	TexCoords[1] = InTexCoord.Y;
-	TexCoords[2] = InTexCoord2.X;
-	TexCoords[3] = InTexCoord2.Y;
-
-	const FVector2D WindowPosition = TransformPoint(RenderTransform, InLocalPosition);
-
-	// Pixel snapping here.
-	Position[0] = FMath::RoundToInt(WindowPosition.X);
-	Position[1] = FMath::RoundToInt(WindowPosition.Y);
-}
-
-FSlateVertex::FSlateVertex( const FSlateRenderTransform& RenderTransform, const FVector2D& InLocalPosition, const FVector2D& InTexCoord, const FColor& InColor, const FSlateRotatedClipRectType& InClipRect )
-	: ClipRect( InClipRect )
-	, Color( InColor )
-{
-	TexCoords[0] = InTexCoord.X;
-	TexCoords[1] = InTexCoord.Y;
-	TexCoords[2] = 1.0f;
-	TexCoords[3] = 1.0f;
-
-	const FVector2D WindowPosition = TransformPoint(RenderTransform, InLocalPosition);
-
-	// Pixel snapping here.
-	Position[0] = FMath::RoundToInt(WindowPosition.X);
-	Position[1] = FMath::RoundToInt(WindowPosition.Y);
-}
-
-
-FSlateVertex::FSlateVertex( const FSlateRenderTransform& RenderTransform, const FVector2D& InLocalPosition, const FVector4& InTexCoord, const FVector2D& InMaterialTexCoords, const FColor& InColor, const FSlateRotatedClipRectType& InClipRect )
-	: ClipRect( InClipRect )
-	, Color( InColor )
-{
-
-	TexCoords[0] = InTexCoord.X;
-	TexCoords[1] = InTexCoord.Y;
-	TexCoords[2] = InTexCoord.Z;
-	TexCoords[3] = InTexCoord.W;
-
-	MaterialTexCoords[0] = InMaterialTexCoords.X;
-	MaterialTexCoords[1] = InMaterialTexCoords.Y;
-
-	const FVector2D WindowPosition = TransformPoint(RenderTransform, InLocalPosition);
-
-	// Pixel snapping here.
-	Position[0] = FMath::RoundToInt(WindowPosition.X);
-	Position[1] = FMath::RoundToInt(WindowPosition.Y);
-}
diff --git a/Engine/Source/Runtime/SlateCore/Private/Styling/CoreStyle.cpp b/Engine/Source/Runtime/SlateCore/Private/Styling/CoreStyle.cpp
--- a/Engine/Source/Runtime/SlateCore/Private/Styling/CoreStyle.cpp
+++ b/Engine/Source/Runtime/SlateCore/Private/Styling/CoreStyle.cpp
@@ -316,6 +316,9 @@ TSharedRef<ISlateStyle> FCoreStyle::Create( const FName& InStyleSetName )
 		Style->Set( "GenericCommands.Duplicate", new IMAGE_BRUSH("Icons/Edit/icon_Edit_Duplicate_16x", Icon16x16) );
 	}
 
+	// SVerticalBox Drag& Drop icon
+	Style->Set("VerticalBoxDragIndicator", new IMAGE_BRUSH("Common/VerticalBoxDragIndicator", FVector2D(6, 45)));
+
 	// SScrollBar defaults...
 	const FScrollBarStyle ScrollBar = FScrollBarStyle()
 		.SetVerticalTopSlotImage(IMAGE_BRUSH("Common/Scrollbar_Background_Vertical", FVector2D(8, 8)))
diff --git a/Engine/Source/Runtime/SlateCore/Private/Widgets/SBoxPanel.cpp b/Engine/Source/Runtime/SlateCore/Private/Widgets/SBoxPanel.cpp
--- a/Engine/Source/Runtime/SlateCore/Private/Widgets/SBoxPanel.cpp
+++ b/Engine/Source/Runtime/SlateCore/Private/Widgets/SBoxPanel.cpp
@@ -277,3 +277,219 @@ SBoxPanel::SBoxPanel( EOrientation InOrientation )
 {
 
 }
+
+
+void SDragAndDropVerticalBox::Construct(const FArguments& InArgs)
+{
+	SVerticalBox::Construct(SVerticalBox::FArguments());
+
+	OnCanAcceptDrop = InArgs._OnCanAcceptDrop;
+	OnAcceptDrop = InArgs._OnAcceptDrop;
+	OnDragDetected_Handler = InArgs._OnDragDetected;
+	OnDragEnter_Handler = InArgs._OnDragEnter;
+	OnDragLeave_Handler = InArgs._OnDragLeave;
+	OnDrop_Handler = InArgs._OnDrop;
+
+	CurrentDragOperationScreenSpaceLocation = FVector2D::ZeroVector;
+	CurrentDragOverSlotIndex = INDEX_NONE;
+}
+
+FReply SDragAndDropVerticalBox::OnMouseButtonDown(const FGeometry& MyGeometry, const FPointerEvent& MouseEvent)
+{
+	if (MouseEvent.GetEffectingButton() == EKeys::LeftMouseButton)
+	{
+		return FReply::Handled().DetectDrag(SharedThis(this), EKeys::LeftMouseButton);
+	}
+
+	return FReply::Unhandled();
+}
+
+FReply SDragAndDropVerticalBox::OnDragDetected(const FGeometry& MyGeometry, const FPointerEvent& MouseEvent)
+{
+	TPanelChildren<FSlot>* PanelChildren = (TPanelChildren<FSlot>*)GetChildren();
+	FArrangedChildren ArrangedChildren(EVisibility::Visible);
+	ArrangeChildren(MyGeometry, ArrangedChildren);
+
+	int32 NodeUnderPositionIndex = SWidget::FindChildUnderMouse(ArrangedChildren, MouseEvent);
+
+	if (PanelChildren->IsValidIndex(NodeUnderPositionIndex))
+	{
+		SVerticalBox::FSlot* Slot = &(*PanelChildren)[NodeUnderPositionIndex];
+
+		if (OnDragDetected_Handler.IsBound())
+		{
+			return OnDragDetected_Handler.Execute(MyGeometry, MouseEvent, NodeUnderPositionIndex, Slot);
+		}
+	}
+
+	return FReply::Unhandled();
+}
+
+void SDragAndDropVerticalBox::OnDragEnter(FGeometry const& MyGeometry, FDragDropEvent const& DragDropEvent)
+{
+	if (OnDragEnter_Handler.IsBound())
+	{
+		OnDragEnter_Handler.Execute(DragDropEvent);
+	}
+}
+
+void SDragAndDropVerticalBox::OnDragLeave(FDragDropEvent const& DragDropEvent)
+{
+	ItemDropZone = TOptional<EItemDropZone>();
+	CurrentDragOperationScreenSpaceLocation = FVector2D::ZeroVector;
+	CurrentDragOverSlotIndex = INDEX_NONE;
+
+	if (OnDragLeave_Handler.IsBound())
+	{
+		OnDragLeave_Handler.Execute(DragDropEvent);
+	}
+}
+
+SDragAndDropVerticalBox::EItemDropZone SDragAndDropVerticalBox::ZoneFromPointerPosition(FVector2D LocalPointerPos, const FGeometry& CurrentGeometry, const FGeometry& StartGeometry) const
+{
+	FSlateLayoutTransform StartGeometryLayoutTransform = StartGeometry.GetAccumulatedLayoutTransform();
+	FSlateLayoutTransform CurrentGeometryLayoutTransform = CurrentGeometry.GetAccumulatedLayoutTransform();
+
+	if (StartGeometryLayoutTransform.GetTranslation().Y > CurrentGeometryLayoutTransform.GetTranslation().Y) // going up
+	{
+		return EItemDropZone::AboveItem;
+	}
+	else if (StartGeometryLayoutTransform.GetTranslation().Y < CurrentGeometryLayoutTransform.GetTranslation().Y) // going down
+	{
+		return EItemDropZone::BelowItem;
+	}
+	else
+	{
+		if (LocalPointerPos.Y <= CurrentGeometry.GetLocalSize().Y / 2.0f)
+		{
+			return EItemDropZone::AboveItem;
+		}
+		else
+		{
+			return EItemDropZone::BelowItem;
+		}
+	}
+}
+
+FReply SDragAndDropVerticalBox::OnDragOver(const FGeometry& MyGeometry, const FDragDropEvent& DragDropEvent)
+{
+	if (OnCanAcceptDrop.IsBound())
+	{
+		FArrangedChildren ArrangedChildren(EVisibility::Visible);
+		ArrangeChildren(MyGeometry, ArrangedChildren);
+
+		TSharedPtr<FDragAndDropVerticalBoxOp> DragOp = DragDropEvent.GetOperationAs<FDragAndDropVerticalBoxOp>();
+
+		if (DragOp.IsValid())
+		{
+			int32 DragOverSlotIndex = SWidget::FindChildUnderPosition(ArrangedChildren, DragDropEvent.GetScreenSpacePosition());
+
+			if (ArrangedChildren.IsValidIndex(DragOverSlotIndex))
+			{
+				FVector2D LocalPointerPos = ArrangedChildren[DragOverSlotIndex].Geometry.AbsoluteToLocal(DragDropEvent.GetScreenSpacePosition());
+				EItemDropZone ItemHoverZone = ZoneFromPointerPosition(LocalPointerPos, ArrangedChildren[DragOverSlotIndex].Geometry, ArrangedChildren[DragOp->SlotIndexBeingDragged].Geometry);
+
+				TPanelChildren<FSlot>* PanelChildren = (TPanelChildren<FSlot>*)GetChildren();
+
+				if (PanelChildren->IsValidIndex(DragOverSlotIndex))
+				{
+					SVerticalBox::FSlot* Slot = &(*PanelChildren)[DragOverSlotIndex];
+
+					ItemDropZone = OnCanAcceptDrop.Execute(DragDropEvent, ItemHoverZone, Slot);
+					CurrentDragOperationScreenSpaceLocation = DragDropEvent.GetScreenSpacePosition();
+					CurrentDragOverSlotIndex = DragOverSlotIndex;
+
+					return FReply::Handled();
+				}
+			}
+		}
+	}
+
+	return FReply::Unhandled();
+}
+
+FReply SDragAndDropVerticalBox::OnDrop(const FGeometry& MyGeometry, const FDragDropEvent& DragDropEvent)
+{
+	FReply DropReply = FReply::Unhandled();
+
+	if (OnAcceptDrop.IsBound())
+	{
+		FArrangedChildren ArrangedChildren(EVisibility::Visible);
+		ArrangeChildren(MyGeometry, ArrangedChildren);
+
+		if (DragDropEvent.GetOperationAs<FDragAndDropVerticalBoxOp>().IsValid())
+		{
+			int32 NodeUnderPositionIndex = SWidget::FindChildUnderPosition(ArrangedChildren, DragDropEvent.GetScreenSpacePosition());
+			TPanelChildren<FSlot>* PanelChildren = (TPanelChildren<FSlot>*)GetChildren();
+
+			if (PanelChildren->IsValidIndex(NodeUnderPositionIndex))
+			{
+				SVerticalBox::FSlot* Slot = &(*PanelChildren)[NodeUnderPositionIndex];
+				TOptional<EItemDropZone> ReportedZone = ItemDropZone;
+
+				if (OnCanAcceptDrop.IsBound() && ItemDropZone.IsSet())
+				{
+					ReportedZone = OnCanAcceptDrop.Execute(DragDropEvent, ItemDropZone.GetValue(), Slot);
+				}
+
+				if (ReportedZone.IsSet())
+				{
+					DropReply = OnAcceptDrop.Execute(DragDropEvent, ReportedZone.GetValue(), NodeUnderPositionIndex, Slot);
+
+					if (DropReply.IsEventHandled())
+					{
+						TSharedPtr<FDragAndDropVerticalBoxOp> DragOp = DragDropEvent.GetOperationAs<FDragAndDropVerticalBoxOp>();
+
+						// Perform the slot changes
+						PanelChildren->Move(DragOp->SlotIndexBeingDragged, NodeUnderPositionIndex);
+					}
+				}
+			}
+
+			ItemDropZone = TOptional<EItemDropZone>();
+			CurrentDragOperationScreenSpaceLocation = FVector2D::ZeroVector;
+			CurrentDragOverSlotIndex = INDEX_NONE;
+		}		
+	}
+
+	return DropReply;
+}
+
+int32 SDragAndDropVerticalBox::OnPaint(const FPaintArgs& Args, const FGeometry& AllottedGeometry, const FSlateRect& MyClippingRect, FSlateWindowElementList& OutDrawElements, int32 LayerId, const FWidgetStyle& InWidgetStyle, bool bParentEnabled) const
+{
+	FArrangedChildren ArrangedChildren(EVisibility::Visible);
+	ArrangeChildren(AllottedGeometry, ArrangedChildren);
+
+	LayerId = SPanel::OnPaint(Args, AllottedGeometry, MyClippingRect, OutDrawElements, LayerId, InWidgetStyle, bParentEnabled);
+
+	if (ItemDropZone.IsSet())
+	{
+		// Draw feedback for user dropping an item above, below.
+		const FSlateBrush* DropIndicatorBrush = nullptr;
+
+		switch (ItemDropZone.GetValue())
+		{
+			default:
+			case EItemDropZone::AboveItem: DropIndicatorBrush = &DropIndicator_Above; break;
+			case EItemDropZone::BelowItem: DropIndicatorBrush = &DropIndicator_Below; break;
+		};
+
+		if (ArrangedChildren.IsValidIndex(CurrentDragOverSlotIndex))
+		{
+			const FArrangedWidget& CurWidget = ArrangedChildren[CurrentDragOverSlotIndex];
+
+			FSlateDrawElement::MakeBox
+			(
+				OutDrawElements,
+				LayerId++,
+				CurWidget.Geometry.ToPaintGeometry(),
+				DropIndicatorBrush,
+				CurWidget.Geometry.GetClippingRect(),
+				ESlateDrawEffect::None,
+				DropIndicatorBrush->GetTint(InWidgetStyle) * InWidgetStyle.GetColorAndOpacityTint()
+			);
+		}
+	}
+
+	return LayerId;
+}
\ No newline at end of file
diff --git a/Engine/Source/Runtime/SlateCore/Private/Widgets/SWindow.cpp b/Engine/Source/Runtime/SlateCore/Private/Widgets/SWindow.cpp
--- a/Engine/Source/Runtime/SlateCore/Private/Widgets/SWindow.cpp
+++ b/Engine/Source/Runtime/SlateCore/Private/Widgets/SWindow.cpp
@@ -1371,7 +1371,7 @@ bool SWindow::IsAutosized() const
 	return SizingRule == ESizingRule::Autosized;
 }
 
-void SWindow::SetSizingRule( ESizingRule::Type InSizingRule )
+void SWindow::SetSizingRule( ESizingRule InSizingRule )
 {
 	SizingRule = InSizingRule;
 }
diff --git a/Engine/Source/Runtime/SlateCore/Public/Layout/ArrangedChildren.h b/Engine/Source/Runtime/SlateCore/Public/Layout/ArrangedChildren.h
--- a/Engine/Source/Runtime/SlateCore/Public/Layout/ArrangedChildren.h
+++ b/Engine/Source/Runtime/SlateCore/Public/Layout/ArrangedChildren.h
@@ -164,4 +164,9 @@ public:
 	{
 		Array.Empty();
 	}
+
+	FORCEINLINE bool IsValidIndex(int32 Index)
+	{
+		return Array.IsValidIndex(Index);
+	}
 };
diff --git a/Engine/Source/Runtime/SlateCore/Public/Layout/Children.h b/Engine/Source/Runtime/SlateCore/Public/Layout/Children.h
--- a/Engine/Source/Runtime/SlateCore/Public/Layout/Children.h
+++ b/Engine/Source/Runtime/SlateCore/Public/Layout/Children.h
@@ -269,6 +269,23 @@ public:
 		}
 	}
 
+	void Move(int32 IndexToMove, int32 IndexToDestination)
+	{
+		// Since we're dealing with an Indirect Array, we can't move an item already in the array.
+		if (IndexToMove > IndexToDestination) // going up
+		{
+			TIndirectArray< SlotType >::Insert(new SlotType(), IndexToDestination);
+			TIndirectArray< SlotType >::Swap(IndexToDestination, IndexToMove + 1);
+			TIndirectArray< SlotType >::RemoveAt(IndexToMove + 1);
+		}
+		else // going down
+		{
+			TIndirectArray< SlotType >::Insert(new SlotType(), IndexToDestination + 1);
+			TIndirectArray< SlotType >::Swap(IndexToMove, IndexToDestination + 1);
+			TIndirectArray< SlotType >::RemoveAt(IndexToMove);
+		}
+	}
+
 	void Reserve( int32 NumToReserve )
 	{
 		TIndirectArray< SlotType >::Reserve(NumToReserve);
