Copying //UE4/Dev-Editor to //UE4/Dev-Main (Source: //UE4/Dev-Editor @ 3279756)

diff --git a/Engine/Source/Runtime/SlateCore/Public/Rendering/DrawElements.h b/Engine/Source/Runtime/SlateCore/Public/Rendering/DrawElements.h
--- a/Engine/Source/Runtime/SlateCore/Public/Rendering/DrawElements.h
+++ b/Engine/Source/Runtime/SlateCore/Public/Rendering/DrawElements.h
@@ -92,7 +92,7 @@ public:
 	bool bAntialias:1;
 
 	// Misc data
-	ESlateBatchDrawFlag::Type BatchFlags;
+	ESlateBatchDrawFlag BatchFlags;
 
 	// Custom drawer data
 	TWeakPtr<ICustomSlateElement, ESPMode::ThreadSafe> CustomDrawer;
@@ -118,7 +118,7 @@ public:
 	int32 DownsampleAmount;
 
 	// Line data
-	ESlateLineJoinType::Type SegmentJoinType;
+	ESlateLineJoinType SegmentJoinType;
 
 
 	SLATECORE_API static FSlateShaderResourceManager* ResourceManager;
@@ -201,7 +201,7 @@ public:
 		GradientStops = InGradientStops;
 	}
 
-	void SetLinesPayloadProperties( const TArray<FVector2D>& InPoints, const FLinearColor& InTint, bool bInAntialias, ESlateLineJoinType::Type InJoinType, float InThickness )
+	void SetLinesPayloadProperties( const TArray<FVector2D>& InPoints, const FLinearColor& InTint, bool bInAntialias, ESlateLineJoinType InJoinType, float InThickness )
 	{
 		Tint = InTint;
 		Thickness = InThickness;
@@ -268,7 +268,17 @@ public:
 		ET_Border,
 		ET_Custom,
 		ET_CustomVerts,
+		/**
+		 * Used for Invalidation, these buffers represent a complete cached buffer of what we normally send to the GPU to be
+		 * drawn for a series of widgets.  They're used to reduce draw overhead in situations where the UI is largely static.
+		 */
 		ET_CachedBuffer,
+		/**
+		 * These layers are different from "layerId", they're symbolic layers, used when building up cached geometry.  They allow
+		 * Slate to semantically differentiate between Layer A and Layer B, which may have completely different layerIds, which perhaps
+		 * overlap, but because they are in logically separate layers they won't intersect, the contents of Layer B would always
+		 * come after the contents of Layer A.
+		 */
 		ET_Layer,
 		ET_PostProcessPass,
 		ET_Count,
@@ -320,7 +330,7 @@ public:
 		const FPaintGeometry& PaintGeometry, 
 		const FSlateBrush* InBrush, 
 		const FSlateRect& InClippingRect, 
-		ESlateDrawEffect::Type InDrawEffects = ESlateDrawEffect::None, 
+		ESlateDrawEffect InDrawEffects = ESlateDrawEffect::None, 
 		const FLinearColor& InTint = FLinearColor::White );
 
 	SLATECORE_API static void MakeBox(
@@ -330,7 +340,7 @@ public:
 		const FSlateBrush* InBrush, 
 		const FSlateResourceHandle& InRenderingHandle, 
 		const FSlateRect& InClippingRect, 
-		ESlateDrawEffect::Type InDrawEffects = ESlateDrawEffect::None, 
+		ESlateDrawEffect InDrawEffects = ESlateDrawEffect::None, 
 		const FLinearColor& InTint = FLinearColor::White );
 	
 	// !!! DEPRECATED !!! Use a render transform om your widget instead.
@@ -340,7 +350,7 @@ public:
 		const FPaintGeometry& PaintGeometry, 
 		const FSlateBrush* InBrush, 
 		const FSlateRect& InClippingRect, 
-		ESlateDrawEffect::Type InDrawEffects = ESlateDrawEffect::None, 
+		ESlateDrawEffect InDrawEffects = ESlateDrawEffect::None, 
 		float Angle = 0.0f,
 		TOptional<FVector2D> InRotationPoint = TOptional<FVector2D>(),
 		ERotationSpace RotationSpace = RelativeToElement,
@@ -359,11 +369,11 @@ public:
 	 * @param InDrawEffects         Optional draw effects to apply
 	 * @param InTint                Color to tint the element
 	 */
-	SLATECORE_API static void MakeText( FSlateWindowElementList& ElementList, uint32 InLayer, const FPaintGeometry& PaintGeometry, const FString& InText, const int32 StartIndex, const int32 EndIndex, const FSlateFontInfo& InFontInfo, const FSlateRect& InClippingRect,ESlateDrawEffect::Type InDrawEffects = ESlateDrawEffect::None, const FLinearColor& InTint =FLinearColor::White );
+	SLATECORE_API static void MakeText( FSlateWindowElementList& ElementList, uint32 InLayer, const FPaintGeometry& PaintGeometry, const FString& InText, const int32 StartIndex, const int32 EndIndex, const FSlateFontInfo& InFontInfo, const FSlateRect& InClippingRect,ESlateDrawEffect InDrawEffects = ESlateDrawEffect::None, const FLinearColor& InTint =FLinearColor::White );
 	
-	SLATECORE_API static void MakeText( FSlateWindowElementList& ElementList, uint32 InLayer, const FPaintGeometry& PaintGeometry, const FString& InText, const FSlateFontInfo& InFontInfo, const FSlateRect& InClippingRect,ESlateDrawEffect::Type InDrawEffects = ESlateDrawEffect::None, const FLinearColor& InTint =FLinearColor::White );
+	SLATECORE_API static void MakeText( FSlateWindowElementList& ElementList, uint32 InLayer, const FPaintGeometry& PaintGeometry, const FString& InText, const FSlateFontInfo& InFontInfo, const FSlateRect& InClippingRect,ESlateDrawEffect InDrawEffects = ESlateDrawEffect::None, const FLinearColor& InTint =FLinearColor::White );
 
-	SLATECORE_API static void MakeText( FSlateWindowElementList& ElementList, uint32 InLayer, const FPaintGeometry& PaintGeometry, const FText& InText, const FSlateFontInfo& InFontInfo, const FSlateRect& InClippingRect,ESlateDrawEffect::Type InDrawEffects = ESlateDrawEffect::None, const FLinearColor& InTint =FLinearColor::White );
+	SLATECORE_API static void MakeText( FSlateWindowElementList& ElementList, uint32 InLayer, const FPaintGeometry& PaintGeometry, const FText& InText, const FSlateFontInfo& InFontInfo, const FSlateRect& InClippingRect,ESlateDrawEffect InDrawEffects = ESlateDrawEffect::None, const FLinearColor& InTint =FLinearColor::White );
 
 	/**
 	 * Creates a text element which displays a series of shaped glyphs on the screen
@@ -376,7 +386,7 @@ public:
 	 * @param InDrawEffects         Optional draw effects to apply
 	 * @param InTint                Color to tint the element
 	 */
-	SLATECORE_API static void MakeShapedText( FSlateWindowElementList& ElementList, uint32 InLayer, const FPaintGeometry& PaintGeometry, const FShapedGlyphSequenceRef& InShapedGlyphSequence, const FSlateRect& InClippingRect, ESlateDrawEffect::Type InDrawEffects = ESlateDrawEffect::None, const FLinearColor& InTint = FLinearColor::White );
+	SLATECORE_API static void MakeShapedText( FSlateWindowElementList& ElementList, uint32 InLayer, const FPaintGeometry& PaintGeometry, const FShapedGlyphSequenceRef& InShapedGlyphSequence, const FSlateRect& InClippingRect, ESlateDrawEffect InDrawEffects = ESlateDrawEffect::None, const FLinearColor& InTint = FLinearColor::White );
 
 	/**
 	 * Creates a gradient element
@@ -389,7 +399,7 @@ public:
 	 * @param InClippingRect           Parts of the element are clipped if it falls outside of this rectangle
 	 * @param InDrawEffects            Optional draw effects to apply
 	 */
-	SLATECORE_API static void MakeGradient( FSlateWindowElementList& ElementList, uint32 InLayer, const FPaintGeometry& PaintGeometry, TArray<FSlateGradientStop> InGradientStops, EOrientation InGradientType, const FSlateRect& InClippingRect, ESlateDrawEffect::Type InDrawEffects = ESlateDrawEffect::None );
+	SLATECORE_API static void MakeGradient( FSlateWindowElementList& ElementList, uint32 InLayer, const FPaintGeometry& PaintGeometry, TArray<FSlateGradientStop> InGradientStops, EOrientation InGradientType, const FSlateRect& InClippingRect, ESlateDrawEffect InDrawEffects = ESlateDrawEffect::None );
 
 	/**
 	 * Creates a spline element
@@ -405,13 +415,13 @@ public:
 	 * @param InDrawEffects         Optional draw effects to apply
 	 * @param InTint                Color to tint the element
 	 */
-	SLATECORE_API static void MakeSpline( FSlateWindowElementList& ElementList, uint32 InLayer, const FPaintGeometry& PaintGeometry, const FVector2D& InStart, const FVector2D& InStartDir, const FVector2D& InEnd, const FVector2D& InEndDir, const FSlateRect InClippingRect, float InThickness = 0.0f, ESlateDrawEffect::Type InDrawEffects = ESlateDrawEffect::None, const FLinearColor& InTint=FLinearColor::White );
+	SLATECORE_API static void MakeSpline( FSlateWindowElementList& ElementList, uint32 InLayer, const FPaintGeometry& PaintGeometry, const FVector2D& InStart, const FVector2D& InStartDir, const FVector2D& InEnd, const FVector2D& InEndDir, const FSlateRect InClippingRect, float InThickness = 0.0f, ESlateDrawEffect InDrawEffects = ESlateDrawEffect::None, const FLinearColor& InTint=FLinearColor::White );
 
 	/** Just like MakeSpline but in draw-space coordinates. This is useful for connecting already-transformed widgets together. */
-	SLATECORE_API static void MakeDrawSpaceSpline(FSlateWindowElementList& ElementList, uint32 InLayer, const FVector2D& InStart, const FVector2D& InStartDir, const FVector2D& InEnd, const FVector2D& InEndDir, const FSlateRect InClippingRect, float InThickness = 0.0f, ESlateDrawEffect::Type InDrawEffects = ESlateDrawEffect::None, const FLinearColor& InTint=FLinearColor::White);
+	SLATECORE_API static void MakeDrawSpaceSpline(FSlateWindowElementList& ElementList, uint32 InLayer, const FVector2D& InStart, const FVector2D& InStartDir, const FVector2D& InEnd, const FVector2D& InEndDir, const FSlateRect InClippingRect, float InThickness = 0.0f, ESlateDrawEffect InDrawEffects = ESlateDrawEffect::None, const FLinearColor& InTint=FLinearColor::White);
 
 	/** Just like MakeSpline but in draw-space coordinates. This is useful for connecting already-transformed widgets together. */
-	SLATECORE_API static void MakeDrawSpaceGradientSpline( FSlateWindowElementList& ElementList, uint32 InLayer, const FVector2D& InStart, const FVector2D& InStartDir, const FVector2D& InEnd, const FVector2D& InEndDir, const FSlateRect InClippingRect, const TArray<FSlateGradientStop>& InGradientStops, float InThickness = 0.0f, ESlateDrawEffect::Type InDrawEffects = ESlateDrawEffect::None );
+	SLATECORE_API static void MakeDrawSpaceGradientSpline( FSlateWindowElementList& ElementList, uint32 InLayer, const FVector2D& InStart, const FVector2D& InStartDir, const FVector2D& InEnd, const FVector2D& InEndDir, const FSlateRect InClippingRect, const TArray<FSlateGradientStop>& InGradientStops, float InThickness = 0.0f, ESlateDrawEffect InDrawEffects = ESlateDrawEffect::None );
 
 	/**
 	 * Creates a line defined by the provided points
@@ -426,7 +436,7 @@ public:
 	 * @param bAntialias               Should antialiasing be applied to the line?
 	 * @param Thickness                The thickness of the line
 	 */
-	SLATECORE_API static void MakeLines( FSlateWindowElementList& ElementList, uint32 InLayer, const FPaintGeometry& PaintGeometry, const TArray<FVector2D>& Points, const FSlateRect InClippingRect, ESlateDrawEffect::Type InDrawEffects = ESlateDrawEffect::None, const FLinearColor& InTint=FLinearColor::White, bool bAntialias = true, float Thickness = 1.0f );
+	SLATECORE_API static void MakeLines( FSlateWindowElementList& ElementList, uint32 InLayer, const FPaintGeometry& PaintGeometry, const TArray<FVector2D>& Points, const FSlateRect InClippingRect, ESlateDrawEffect InDrawEffects = ESlateDrawEffect::None, const FLinearColor& InTint=FLinearColor::White, bool bAntialias = true, float Thickness = 1.0f );
 
 	/**
 	 * Creates a viewport element which is useful for rendering custom data in a texture into Slate
@@ -440,7 +450,7 @@ public:
 	 * @param InDrawEffects            Optional draw effects to apply
 	 * @param InTint                   Color to tint the element
 	 */
-	SLATECORE_API static void MakeViewport( FSlateWindowElementList& ElementList, uint32 InLayer, const FPaintGeometry& PaintGeometry, TSharedPtr<const ISlateViewport> Viewport, const FSlateRect& InClippingRect, ESlateDrawEffect::Type InDrawEffects = ESlateDrawEffect::None, const FLinearColor& InTint=FLinearColor::White );
+	SLATECORE_API static void MakeViewport( FSlateWindowElementList& ElementList, uint32 InLayer, const FPaintGeometry& PaintGeometry, TSharedPtr<const ISlateViewport> Viewport, const FSlateRect& InClippingRect, ESlateDrawEffect InDrawEffects = ESlateDrawEffect::None, const FLinearColor& InTint=FLinearColor::White );
 
 	/**
 	 * Creates a custom element which can be used to manually draw into the Slate render target with graphics API calls rather than Slate elements
@@ -463,6 +473,7 @@ public:
 	FORCEINLINE EElementType GetElementType() const { return ElementType; }
 	FORCEINLINE uint32 GetLayer() const { return Layer; }
 	FORCEINLINE const FSlateRenderTransform& GetRenderTransform() const { return RenderTransform; }
+	FORCEINLINE const FTransform2D& GetLayoutToRenderTransform() const { return LayoutToRenderTransform; }
 	FORCEINLINE const FVector2D& GetPosition() const { return Position; }
 	FORCEINLINE void SetPosition(const FVector2D& InPosition) { Position = InPosition; }
 	FORCEINLINE const FVector2D& GetLocalSize() const { return LocalSize; }
@@ -470,25 +481,39 @@ public:
 	FORCEINLINE const FSlateRect& GetClippingRect() const { return ClippingRect; }
 	FORCEINLINE void SetClippingRect(const FSlateRect& InClippingRect) { ClippingRect = InClippingRect; }
 	FORCEINLINE const FSlateDataPayload& GetDataPayload() const { return DataPayload; }
-	FORCEINLINE uint32 GetDrawEffects() const { return DrawEffects; }
+	FORCEINLINE ESlateDrawEffect GetDrawEffects() const { return DrawEffects; }
 	FORCEINLINE const TOptional<FShortRect>& GetScissorRect() const { return ScissorRect; }
 	FORCEINLINE const int32 GetSceneIndex() const { return SceneIndex; }
 
-private:
-	void Init(uint32 InLayer, const FPaintGeometry& PaintGeometry, const FSlateRect& InClippingRect, ESlateDrawEffect::Type InDrawEffects);
+	FORCEINLINE FSlateRotatedClipRectType CalculateRenderClippingRect() const
+	{
+		// The clip rect is NOT subject to the rotations specified by MakeRotatedBox.
+		return
+			EnumHasAllFlags(DrawEffects, ESlateDrawEffect::NoPixelSnapping)
+			? FSlateRotatedClipRectType::MakeRotatedRect(ClippingRect, LayoutToRenderTransform)
+			: FSlateRotatedClipRectType::MakeSnappedRotatedRect(ClippingRect, LayoutToRenderTransform);
+	}
 
+	FORCEINLINE FSlateLayoutTransform GetInverseLayoutTransform() const
+	{
+		return Inverse(FSlateLayoutTransform(Scale, Position));
+	}
+
+private:
+	void Init(uint32 InLayer, const FPaintGeometry& PaintGeometry, const FSlateRect& InClippingRect, ESlateDrawEffect InDrawEffects);
 
 	static FVector2D GetRotationPoint( const FPaintGeometry& PaintGeometry, const TOptional<FVector2D>& UserRotationPoint, ERotationSpace RotationSpace );
 
 private:
 	FSlateDataPayload DataPayload;
 	FSlateRenderTransform RenderTransform;
+	FTransform2D LayoutToRenderTransform;
 	FSlateRect ClippingRect;
 	FVector2D Position;
 	FVector2D LocalSize;
 	float Scale;
 	uint32 Layer;
-	uint32 DrawEffects;
+	ESlateDrawEffect DrawEffects;
 	EElementType ElementType;
 	TOptional<FShortRect> ScissorRect;
 	int32 SceneIndex;
@@ -568,7 +593,7 @@ private:
 class FSlateElementBatch
 {
 public:
-	FSlateElementBatch(const FSlateShaderResource* InShaderResource, const FShaderParams& InShaderParams, ESlateShader::Type ShaderType, ESlateDrawPrimitive::Type PrimitiveType, ESlateDrawEffect::Type DrawEffects, ESlateBatchDrawFlag::Type DrawFlags, const TOptional<FShortRect>& ScissorRect, int32 InstanceCount = 0, uint32 InstanceOffset = 0, ISlateUpdatableInstanceBuffer* InstanceData = nullptr, int32 SceneIndex = -1)
+	FSlateElementBatch(const FSlateShaderResource* InShaderResource, const FShaderParams& InShaderParams, ESlateShader::Type ShaderType, ESlateDrawPrimitive::Type PrimitiveType, ESlateDrawEffect DrawEffects, ESlateBatchDrawFlag DrawFlags, const TOptional<FShortRect>& ScissorRect, int32 InstanceCount = 0, uint32 InstanceOffset = 0, ISlateUpdatableInstanceBuffer* InstanceData = nullptr, int32 SceneIndex = -1)
 		: BatchKey(InShaderParams, ShaderType, PrimitiveType, DrawEffects, DrawFlags, ScissorRect, InstanceCount, InstanceOffset, InstanceData, SceneIndex)
 		, ShaderResource(InShaderResource)
 		, NumElementsInBatch(0)
@@ -612,10 +637,10 @@ public:
 
 	const FSlateShaderResource* GetShaderResource() const { return ShaderResource; }
 	const FShaderParams& GetShaderParams() const { return BatchKey.ShaderParams; }
-	ESlateBatchDrawFlag::Type GetDrawFlags() const { return BatchKey.DrawFlags; }
+	ESlateBatchDrawFlag GetDrawFlags() const { return BatchKey.DrawFlags; }
 	ESlateDrawPrimitive::Type GetPrimitiveType() const { return BatchKey.DrawPrimitiveType; }
 	ESlateShader::Type GetShaderType() const { return BatchKey.ShaderType; } 
-	ESlateDrawEffect::Type GetDrawEffects() const { return BatchKey.DrawEffects; }
+	ESlateDrawEffect GetDrawEffects() const { return BatchKey.DrawEffects; }
 	const TOptional<FShortRect>& GetScissorRect() const { return BatchKey.ScissorRect; }
 	const TWeakPtr<ICustomSlateElement, ESPMode::ThreadSafe> GetCustomDrawer() const { return BatchKey.CustomDrawer; }
 	const TSharedPtr<FSlateRenderDataHandle, ESPMode::ThreadSafe> GetCachedRenderHandle() const { return BatchKey.CachedRenderHandle; }
@@ -633,17 +658,17 @@ private:
 		const FVector2D CachedRenderDataOffset;
 		const TSharedPtr<FSlateDrawLayerHandle, ESPMode::ThreadSafe> LayerHandle;
 		const FShaderParams ShaderParams;
-		const ESlateBatchDrawFlag::Type DrawFlags;	
+		const ESlateBatchDrawFlag DrawFlags;	
 		const ESlateShader::Type ShaderType;
 		const ESlateDrawPrimitive::Type DrawPrimitiveType;
-		const ESlateDrawEffect::Type DrawEffects;
+		const ESlateDrawEffect DrawEffects;
 		const TOptional<FShortRect> ScissorRect;
 		const int32 InstanceCount;
 		const uint32 InstanceOffset;
 		const ISlateUpdatableInstanceBuffer* InstanceData;
 		const int32 SceneIndex;
 
-		FBatchKey(const FShaderParams& InShaderParams, ESlateShader::Type InShaderType, ESlateDrawPrimitive::Type InDrawPrimitiveType, ESlateDrawEffect::Type InDrawEffects, ESlateBatchDrawFlag::Type InDrawFlags, const TOptional<FShortRect>& InScissorRect, int32 InInstanceCount, uint32 InInstanceOffset, ISlateUpdatableInstanceBuffer* InInstanceBuffer, int32 InSceneIndex)
+		FBatchKey(const FShaderParams& InShaderParams, ESlateShader::Type InShaderType, ESlateDrawPrimitive::Type InDrawPrimitiveType, ESlateDrawEffect InDrawEffects, ESlateBatchDrawFlag InDrawFlags, const TOptional<FShortRect>& InScissorRect, int32 InInstanceCount, uint32 InInstanceOffset, ISlateUpdatableInstanceBuffer* InInstanceBuffer, int32 InSceneIndex)
 			: ShaderParams(InShaderParams)
 			, DrawFlags(InDrawFlags)
 			, ShaderType(InShaderType)
@@ -657,7 +682,7 @@ private:
 		{
 		}
 
-		FBatchKey( const FShaderParams& InShaderParams, ESlateShader::Type InShaderType, ESlateDrawPrimitive::Type InDrawPrimitiveType, ESlateDrawEffect::Type InDrawEffects, ESlateBatchDrawFlag::Type InDrawFlags, const TOptional<FShortRect>& InScissorRect, int32 InInstanceCount, uint32 InInstanceOffset, ISlateUpdatableInstanceBuffer* InInstanceBuffer)
+		FBatchKey( const FShaderParams& InShaderParams, ESlateShader::Type InShaderType, ESlateDrawPrimitive::Type InDrawPrimitiveType, ESlateDrawEffect InDrawEffects, ESlateBatchDrawFlag InDrawFlags, const TOptional<FShortRect>& InScissorRect, int32 InInstanceCount, uint32 InInstanceOffset, ISlateUpdatableInstanceBuffer* InInstanceBuffer)
 			: ShaderParams( InShaderParams )
 			, DrawFlags( InDrawFlags )
 			, ShaderType( InShaderType )
@@ -814,13 +839,13 @@ public:
 
 	const TSharedPtr<FSlateRenderDataHandle, ESPMode::ThreadSafe> CachedRenderHandle;
 
-	const ESlateBatchDrawFlag::Type DrawFlags;	
+	const ESlateBatchDrawFlag DrawFlags;	
 
 	const ESlateShader::Type ShaderType;
 
 	const ESlateDrawPrimitive::Type DrawPrimitiveType;
 
-	const ESlateDrawEffect::Type DrawEffects;
+	const ESlateDrawEffect DrawEffects;
 
 	const TOptional<FShortRect> ScissorRect;
 	
diff --git a/Engine/Source/Runtime/SlateCore/Public/Rendering/RenderingCommon.h b/Engine/Source/Runtime/SlateCore/Public/Rendering/RenderingCommon.h
--- a/Engine/Source/Runtime/SlateCore/Public/Rendering/RenderingCommon.h
+++ b/Engine/Source/Runtime/SlateCore/Public/Rendering/RenderingCommon.h
@@ -59,53 +59,63 @@ namespace ESlateShader
  * Note: New effects added should be in bit mask form
  * If you add a type here you must also implement the proper shader type (TSlateElementPS).  See SlateShaders.h
  */
-namespace ESlateDrawEffect
+enum class ESlateDrawEffect : uint8
 {
-	typedef uint8 Type;
 	/** No effect applied */
-	const Type None					= 0;
-	/** Draw the element with a disabled effect */
-	const Type DisabledEffect		= 1 << 0;
-	/** Advanced: Don't read from texture alpha channel */
-	const Type IgnoreTextureAlpha	= 1 << 1;
+	None					= 0,
 	/** Advanced: Draw the element with no blending */
-	const Type NoBlending			= 1 << 2;
+	NoBlending			= 1 << 0,
 	/** Advanced: Blend using pre-multiplied alpha. Ignored if NoBlending is set. */
-	const Type PreMultipliedAlpha	= 1 << 3;
+	PreMultipliedAlpha	= 1 << 1,
 	/** Advanced: No gamma correction should be done */
-	const Type NoGamma				= 1 << 4;
+	NoGamma				= 1 << 2,
+	/** Advanced: Change the alpha value to 1 - Alpha. */
+	InvertAlpha			= 1 << 3,
+
+	// ^^ These Match ESlateBatchDrawFlag ^^
+
+	/** Disables pixel snapping */
+	NoPixelSnapping		= 1 << 4,
+	/** Draw the element with a disabled effect */
+	DisabledEffect		= 1 << 5,
+	/** Advanced: Don't read from texture alpha channel */
+	IgnoreTextureAlpha	= 1 << 6,
 };
 
+ENUM_CLASS_FLAGS(ESlateDrawEffect);
 
 /** Flags for drawing a batch */
-namespace ESlateBatchDrawFlag
+enum class ESlateBatchDrawFlag : uint8
 {
-	typedef uint8 Type;
 	/** No draw flags */
-	const Type None					= 0;
+	None					= 0,
 	/** Draw the element with no blending */
-	const Type NoBlending			= 1 << 0;
+	NoBlending			= 1 << 0,
 	/** Blend using pre-multiplied alpha. Ignored if NoBlending is set. */
-	const Type PreMultipliedAlpha	= 1 << 1;
+	PreMultipliedAlpha	= 1 << 1,
 	/** No gamma correction should be done */
-	const Type NoGamma				= 1 << 2;
+	NoGamma				= 1 << 2,
+	/** Change the alpha value to 1 - Alpha */
+	InvertAlpha			= 1 << 3,
+
+	// ^^ These Match ESlateDrawEffect ^^
+
 	/** Draw the element as wireframe */
-	const Type Wireframe			= 1 << 3;
+	Wireframe			= 1 << 4,
 	/** The element should be tiled horizontally */
-	const Type TileU				= 1 << 4;
+	TileU				= 1 << 5,
 	/** The element should be tiled vertically */
-	const Type TileV				= 1 << 5;
+	TileV				= 1 << 6
 };
 
-namespace ESlateLineJoinType
+ENUM_CLASS_FLAGS(ESlateBatchDrawFlag);
+
+enum class ESlateLineJoinType : uint8
 {
-	enum Type
-	{
-		// Joins line segments with a sharp edge (miter)
-		Sharp =	0,
-		// Simply stitches together line segments
-		Simple = 1,
-	};
+	// Joins line segments with a sharp edge (miter)
+	Sharp =	0,
+	// Simply stitches together line segments
+	Simple = 1,
 };
 
 /**
@@ -115,12 +125,16 @@ namespace ESlateLineJoinType
  */
 struct SLATECORE_API FSlateRotatedRect
 {
+public:
 	/** Default ctor. */
 	FSlateRotatedRect();
 	/** Construct a rotated rect from a given aligned rect. */
 	explicit FSlateRotatedRect(const FSlateRect& AlignedRect);
 	/** Per-element constructor. */
 	FSlateRotatedRect(const FVector2D& InTopLeft, const FVector2D& InExtentX, const FVector2D& InExtentY);
+
+public:
+
 	/** transformed Top-left corner. */
 	FVector2D TopLeft;
 	/** transformed X extent (right-left). */
@@ -128,15 +142,38 @@ struct SLATECORE_API FSlateRotatedRect
 	/** transformed Y extent (bottom-top). */
 	FVector2D ExtentY;
 
+public:
+	bool operator == (const FSlateRotatedRect& Other) const
+	{
+		return
+			TopLeft == Other.TopLeft &&
+			ExtentX == Other.ExtentX &&
+			ExtentY == Other.ExtentY;
+	}
+
+public:
+
 	/** Convert to a bounding, aligned rect. */
 	FSlateRect ToBoundingRect() const;
 	/** Point-in-rect test. */
 	bool IsUnderLocation(const FVector2D& Location) const;
 
+	static FSlateRotatedRect MakeRotatedRect(const FSlateRect& ClipRectInLayoutWindowSpace, const FSlateLayoutTransform& InverseLayoutTransform, const FSlateRenderTransform& RenderTransform)
+	{
+		return MakeRotatedRect(ClipRectInLayoutWindowSpace, Concatenate(InverseLayoutTransform, RenderTransform));
+	}
+
+	static FSlateRotatedRect MakeRotatedRect(const FSlateRect& ClipRectInLayoutWindowSpace, const FTransform2D& LayoutToRenderTransform);
+
+	static FSlateRotatedRect MakeSnappedRotatedRect(const FSlateRect& ClipRectInLayoutWindowSpace, const FSlateLayoutTransform& InverseLayoutTransform, const FSlateRenderTransform& RenderTransform)
+	{
+		return MakeSnappedRotatedRect(ClipRectInLayoutWindowSpace, Concatenate(InverseLayoutTransform, RenderTransform));
+	}
+
 	/**
 	 * Used to construct a rotated rect from an aligned clip rect and a set of layout and render transforms from the geometry, snapped to pixel boundaries. Returns a float or float16 version of the rect based on the typedef.
 	 */
-	static FSlateRotatedRect MakeSnappedRotatedRect(const FSlateRect& ClipRectInLayoutWindowSpace, const FSlateLayoutTransform& InverseLayoutTransform, const FSlateRenderTransform& RenderTransform);
+	static FSlateRotatedRect MakeSnappedRotatedRect(const FSlateRect& ClipRectInLayoutWindowSpace, const FTransform2D& LayoutToRenderTransform);
 };
 
 /**
@@ -155,6 +192,15 @@ FSlateRotatedRect TransformRect(const TransformType& Transform, const FSlateRota
 
 typedef FSlateRotatedRect FSlateRotatedClipRectType;
 
+
+enum class ESlateVertexRounding : uint8
+{
+	Disabled,
+	Enabled
+};
+
+
+
 /** 
  * A struct which defines a basic vertex seen by the Slate vertex buffers and shaders
  */
@@ -164,10 +210,10 @@ struct SLATECORE_API FSlateVertex
 	float TexCoords[4]; 
 
 	/** Texture coordinates used as pass through to materials for custom texturing. */
-	float MaterialTexCoords[2];
+	FVector2D MaterialTexCoords;
 
 	/** Position of the vertex in window space */
-	float Position[2];
+	FVector2D Position;
 
 	/** clip center/extents in render window space (window space with render transforms applied) */
 	FSlateRotatedClipRectType ClipRect;
@@ -175,10 +221,133 @@ struct SLATECORE_API FSlateVertex
 	/** Vertex color */
 	FColor Color;
 	
-	FSlateVertex();
-	FSlateVertex( const FSlateRenderTransform& RenderTransform, const FVector2D& InLocalPosition, const FVector2D& InTexCoord, const FVector2D& InTexCoord2, const FColor& InColor, const FSlateRotatedClipRectType& InClipRect );
-	FSlateVertex( const FSlateRenderTransform& RenderTransform, const FVector2D& InLocalPosition, const FVector2D& InTexCoord, const FColor& InColor, const FSlateRotatedClipRectType& InClipRect );
-	FSlateVertex( const FSlateRenderTransform& RenderTransform, const FVector2D& InLocalPosition, const FVector4& InTexCoords, const FVector2D& InMaterialTexCoords, const FColor& InColor, const FSlateRotatedClipRectType& InClipRect );
+	/** Local size of the element */
+	uint16 PixelSize[2];
+
+	FSlateVertex() {}
+
+public:
+	
+	// These constructors have more or less been deprecated, you should use FSlateVertex::Make<ESlateVertexRounding::Enabled>(...) in order to have fine control over pixel snapping or not.
+
+	FSlateVertex(const FSlateRenderTransform& RenderTransform, const FVector2D& InLocalPosition, const FVector2D& InTexCoord, const FVector2D& InTexCoord2, const FColor& InColor, const FSlateRotatedClipRectType& InClipRect)
+		: TexCoords { InTexCoord.X, InTexCoord.Y, InTexCoord2.X, InTexCoord2.Y }
+		, Position(TransformPoint(RenderTransform, InLocalPosition))
+		, ClipRect(InClipRect)
+		, Color(InColor)
+	{
+		Position.X = FMath::RoundToInt(Position.X);
+		Position.Y = FMath::RoundToInt(Position.Y);
+	}
+
+	FSlateVertex(const FSlateRenderTransform& RenderTransform, const FVector2D& InLocalPosition, const FVector2D& InTexCoord, const FColor& InColor, const FSlateRotatedClipRectType& InClipRect)
+		: TexCoords { InTexCoord.X, InTexCoord.Y, 1.0f, 1.0f }
+		, Position(TransformPoint(RenderTransform, InLocalPosition))
+		, ClipRect(InClipRect)
+		, Color(InColor)
+	{
+		Position.X = FMath::RoundToInt(Position.X);
+		Position.Y = FMath::RoundToInt(Position.Y);
+	}
+
+	FSlateVertex(const FSlateRenderTransform& RenderTransform, const FVector2D& InLocalPosition, const FVector4& InTexCoords, const FVector2D& InMaterialTexCoords, const FColor& InColor, const FSlateRotatedClipRectType& InClipRect)
+		: TexCoords { InTexCoords.X, InTexCoords.Y, InTexCoords.Z, InTexCoords.W }
+		, MaterialTexCoords(InMaterialTexCoords)
+		, Position(TransformPoint(RenderTransform, InLocalPosition))
+		, ClipRect(InClipRect)
+		, Color(InColor)
+	{
+		Position.X = FMath::RoundToInt(Position.X);
+		Position.Y = FMath::RoundToInt(Position.Y);
+	}
+
+	FSlateVertex(const FSlateRenderTransform& RenderTransform, const FVector2D& InLocalPosition, const FVector2D& InLocalSize, float Scale, const FVector4& InTexCoords, const FColor& InColor, const FSlateRotatedClipRectType& InClipRect)
+		: TexCoords { InTexCoords.X, InTexCoords.Y, InTexCoords.Z, InTexCoords.W }
+		, MaterialTexCoords(InLocalPosition.X / InLocalSize.X, InLocalPosition.Y / InLocalSize.Y)
+		, Position(TransformPoint(RenderTransform, InLocalPosition))
+		, ClipRect(InClipRect)
+		, Color(InColor)
+	{
+		Position.X = FMath::RoundToInt(Position.X);
+		Position.Y = FMath::RoundToInt(Position.Y);
+
+		PixelSize[0] = FMath::RoundToInt(InLocalSize.X * Scale);
+		PixelSize[1] = FMath::RoundToInt(InLocalSize.Y * Scale);
+	}
+
+public:
+
+	template<ESlateVertexRounding Rounding>
+	static FSlateVertex Make(const FSlateRenderTransform& RenderTransform, const FVector2D& InLocalPosition, const FVector2D& InTexCoord, const FVector2D& InTexCoord2, const FColor& InColor, const FSlateRotatedClipRectType& InClipRect)
+	{
+		FSlateVertex Vertex;
+		Vertex.TexCoords[0] = InTexCoord.X;
+		Vertex.TexCoords[1] = InTexCoord.Y;
+		Vertex.TexCoords[2] = InTexCoord2.X;
+		Vertex.TexCoords[3] = InTexCoord2.Y;
+		Vertex.InitCommon<Rounding>(RenderTransform, InLocalPosition, InColor, InClipRect);
+
+		return Vertex;
+	}
+
+	template<ESlateVertexRounding Rounding>
+	static FSlateVertex Make(const FSlateRenderTransform& RenderTransform, const FVector2D& InLocalPosition, const FVector2D& InTexCoord, const FColor& InColor, const FSlateRotatedClipRectType& InClipRect)
+	{
+		FSlateVertex Vertex;
+		Vertex.TexCoords[0] = InTexCoord.X;
+		Vertex.TexCoords[1] = InTexCoord.Y;
+		Vertex.TexCoords[2] = 1.0f;
+		Vertex.TexCoords[3] = 1.0f;
+		Vertex.InitCommon<Rounding>(RenderTransform, InLocalPosition, InColor, InClipRect);
+
+		return Vertex;
+	}
+
+	template<ESlateVertexRounding Rounding>
+	static FSlateVertex Make(const FSlateRenderTransform& RenderTransform, const FVector2D& InLocalPosition, const FVector4& InTexCoords, const FVector2D& InMaterialTexCoords, const FColor& InColor, const FSlateRotatedClipRectType& InClipRect)
+	{
+		FSlateVertex Vertex;
+		Vertex.TexCoords[0] = InTexCoords.X;
+		Vertex.TexCoords[1] = InTexCoords.Y;
+		Vertex.TexCoords[2] = InTexCoords.Z;
+		Vertex.TexCoords[3] = InTexCoords.W;
+		Vertex.MaterialTexCoords = InMaterialTexCoords;
+		Vertex.InitCommon<Rounding>(RenderTransform, InLocalPosition, InColor, InClipRect);
+
+		return Vertex;
+	}
+
+	template<ESlateVertexRounding Rounding>
+	static FSlateVertex Make(const FSlateRenderTransform& RenderTransform, const FVector2D& InLocalPosition, const FVector2D& InLocalSize, float Scale, const FVector4& InTexCoords, const FColor& InColor, const FSlateRotatedClipRectType& InClipRect)
+	{
+		FSlateVertex Vertex;
+		Vertex.TexCoords[0] = InTexCoords.X;
+		Vertex.TexCoords[1] = InTexCoords.Y;
+		Vertex.TexCoords[2] = InTexCoords.Z;
+		Vertex.TexCoords[3] = InTexCoords.W;
+		Vertex.InitCommon<Rounding>(RenderTransform, InLocalPosition, InColor, InClipRect);
+		Vertex.PixelSize[0] = FMath::RoundToInt(InLocalSize.X * Scale);
+		Vertex.PixelSize[1] = FMath::RoundToInt(InLocalSize.Y * Scale);
+
+		return Vertex;
+	}
+
+private:
+
+	template<ESlateVertexRounding Rounding>
+	FORCEINLINE void InitCommon(const FSlateRenderTransform& RenderTransform, const FVector2D& InLocalPosition, const FColor& InColor, const FSlateRotatedClipRectType& InClipRect)
+	{
+		Position = TransformPoint(RenderTransform, InLocalPosition);
+
+		if ( Rounding == ESlateVertexRounding::Enabled )
+		{
+			Position.X = FMath::RoundToInt(Position.X);
+			Position.Y = FMath::RoundToInt(Position.Y);
+		}
+
+		ClipRect = InClipRect;
+		Color = InColor;
+	}
 };
 
 template<> struct TIsPODType<FSlateVertex> { enum { Value = true }; };
diff --git a/Engine/Source/Runtime/SlateCore/Public/Rendering/SlateRenderer.h b/Engine/Source/Runtime/SlateCore/Public/Rendering/SlateRenderer.h
--- a/Engine/Source/Runtime/SlateCore/Public/Rendering/SlateRenderer.h
+++ b/Engine/Source/Runtime/SlateCore/Public/Rendering/SlateRenderer.h
@@ -252,6 +252,8 @@ public:
 	 */
 	virtual bool GenerateDynamicImageResource( FName ResourceName, uint32 Width, uint32 Height, const TArray< uint8 >& Bytes ) { return false; }
 
+	virtual bool GenerateDynamicImageResource(FName ResourceName, FSlateTextureDataRef TextureData) { return false; }
+
 	/**
 	 * Creates a handle to a Slate resource
 	 * A handle is used as fast path for looking up a rendering resource for a given brush when adding Slate draw elements
diff --git a/Engine/Source/Runtime/SlateCore/Public/Widgets/DeclarativeSyntaxSupport.h b/Engine/Source/Runtime/SlateCore/Public/Widgets/DeclarativeSyntaxSupport.h
--- a/Engine/Source/Runtime/SlateCore/Public/Widgets/DeclarativeSyntaxSupport.h
+++ b/Engine/Source/Runtime/SlateCore/Public/Widgets/DeclarativeSyntaxSupport.h
@@ -791,7 +791,7 @@ struct TSlateBaseNamedArgs
 	template<typename MetaDataType, typename Arg0Type>
 	WidgetArgsType& AddMetaData(Arg0Type InArg0)
 	{
-		MetaData.Add(MakeShareable(new MetaDataType(InArg0)));
+		MetaData.Add(MakeShared<MetaDataType>(InArg0));
 		return Me();
 	}
 
@@ -799,7 +799,7 @@ struct TSlateBaseNamedArgs
 	template<typename MetaDataType, typename Arg0Type, typename Arg1Type>
 	WidgetArgsType& AddMetaData(Arg0Type InArg0, Arg1Type InArg1)
 	{
-		MetaData.Add(MakeShareable(new MetaDataType(InArg0, InArg1)));
+		MetaData.Add(MakeShared<MetaDataType>(InArg0, InArg1));
 		return Me();
 	}
 
@@ -990,7 +990,7 @@ struct TWidgetAllocator
 {
 	static TSharedRef<WidgetType> PrivateAllocateWidget()
 	{
-		return MakeShareable( new WidgetType() );
+		return MakeShared<WidgetType>();
 	}
 };
 
diff --git a/Engine/Source/Runtime/SlateCore/Public/Widgets/SBoxPanel.h b/Engine/Source/Runtime/SlateCore/Public/Widgets/SBoxPanel.h
--- a/Engine/Source/Runtime/SlateCore/Public/Widgets/SBoxPanel.h
+++ b/Engine/Source/Runtime/SlateCore/Public/Widgets/SBoxPanel.h
@@ -12,6 +12,9 @@
 #include "Widgets/SWidget.h"
 #include "Layout/Children.h"
 #include "Widgets/SPanel.h"
+#include "ArrangedChildren.h"
+#include "DragAndDrop.h"
+#include "DrawElements.h"
 
 class FArrangedChildren;
 
@@ -399,3 +402,119 @@ public:
 	 */
 	void Construct( const FArguments& InArgs );
 };
+
+class FDragAndDropVerticalBoxOp : public FDragDropOperation
+{
+public:
+	DRAG_DROP_OPERATOR_TYPE(FDragAndDropVerticalBoxOp, FDragDropOperation)
+
+	/** Data this item represent */
+	int32 SlotIndexBeingDragged;
+	SVerticalBox::FSlot* SlotBeingDragged;
+
+	virtual ~FDragAndDropVerticalBoxOp()
+	{}
+};
+
+/** A Vertical Box Panel. See SBoxPanel for more info. */
+class SLATECORE_API SDragAndDropVerticalBox : public SVerticalBox
+{
+public:
+	/**
+	* Where we are going to drop relative to the target item.
+	*/
+	enum class EItemDropZone
+	{
+		AboveItem,
+		BelowItem
+	};
+
+	DECLARE_DELEGATE_RetVal_FourParams(FReply, FOnDragAndDropVerticalBoxDragDetected, const FGeometry&, const FPointerEvent&, int32, SVerticalBox::FSlot*)
+	DECLARE_DELEGATE_OneParam(FOnDragAndDropVerticalBoxDragEnter, FDragDropEvent const&);
+	DECLARE_DELEGATE_OneParam(FOnDragAndDropVerticalBoxDragLeave, FDragDropEvent const&);
+	DECLARE_DELEGATE_RetVal_OneParam(FReply, FOnDragAndDropVerticalBoxDrop, FDragDropEvent const&);
+
+	/** Delegate signature for querying whether this FDragDropEvent will be handled by the drop target of type ItemType. */
+	DECLARE_DELEGATE_RetVal_ThreeParams(TOptional<EItemDropZone>, FOnCanAcceptDrop, const FDragDropEvent&, SDragAndDropVerticalBox::EItemDropZone, SVerticalBox::FSlot*);
+	/** Delegate signature for handling the drop of FDragDropEvent onto target of type ItemType */
+	DECLARE_DELEGATE_RetVal_FourParams(FReply, FOnAcceptDrop, const FDragDropEvent&, SDragAndDropVerticalBox::EItemDropZone, int32, SVerticalBox::FSlot*);
+
+	SLATE_BEGIN_ARGS(SDragAndDropVerticalBox)
+		{}
+
+		// High Level DragAndDrop
+
+		/**
+		* Handle this event to determine whether a drag and drop operation can be executed on top of the target row widget.
+		* Most commonly, this is used for previewing re-ordering and re-organization operations in lists or trees.
+		* e.g. A user is dragging one item into a different spot in the list or tree.
+		*      This delegate will be called to figure out if we should give visual feedback on whether an item will
+		*      successfully drop into the list.
+		*/
+		SLATE_EVENT(FOnCanAcceptDrop, OnCanAcceptDrop)
+
+		/**
+		* Perform a drop operation onto the target row widget
+		* Most commonly used for executing a re-ordering and re-organization operations in lists or trees.
+		* e.g. A user was dragging one item into a different spot in the list; they just dropped it.
+		*      This is our chance to handle the drop by reordering items and calling for a list refresh.
+		*/
+		SLATE_EVENT(FOnAcceptDrop, OnAcceptDrop)
+
+		// Low level DragAndDrop
+		SLATE_EVENT(FOnDragAndDropVerticalBoxDragDetected, OnDragDetected)
+		SLATE_EVENT(FOnDragAndDropVerticalBoxDragEnter, OnDragEnter)
+		SLATE_EVENT(FOnDragAndDropVerticalBoxDragLeave, OnDragLeave)
+		SLATE_EVENT(FOnDragAndDropVerticalBoxDrop, OnDrop)
+	SLATE_END_ARGS()
+
+	void Construct(const FArguments& InArgs);
+
+	/** Set the Drop indicators */
+	SDragAndDropVerticalBox& SetDropIndicator_Above(const FSlateBrush& InValue) { DropIndicator_Above = InValue; return *this; }
+	SDragAndDropVerticalBox& SetDropIndicator_Below(const FSlateBrush& InValue) { DropIndicator_Below = InValue; return *this; }
+
+	/** Drag detection and handling */
+	virtual FReply OnMouseButtonDown(const FGeometry& MyGeometry, const FPointerEvent& MouseEvent) override;
+		
+	virtual FReply OnDragDetected(const FGeometry& MyGeometry, const FPointerEvent& MouseEvent) override;
+	virtual void OnDragEnter(FGeometry const& MyGeometry, FDragDropEvent const& DragDropEvent) override;
+	virtual void OnDragLeave(FDragDropEvent const& DragDropEvent) override;
+	virtual FReply OnDragOver(const FGeometry& MyGeometry, const FDragDropEvent& DragDropEvent) override;
+	virtual FReply OnDrop(const FGeometry& MyGeometry, const FDragDropEvent& DragDropEvent) override;
+
+	virtual int32 OnPaint(const FPaintArgs& Args, const FGeometry& AllottedGeometry, const FSlateRect& MyClippingRect, FSlateWindowElementList& OutDrawElements, int32 LayerId, const FWidgetStyle& InWidgetStyle, bool bParentEnabled) const override;
+
+private:
+	/** @see SDragAndDropVerticalBox's OnCanAcceptDrop event */
+	FOnCanAcceptDrop OnCanAcceptDrop;
+
+	/** @see SDragAndDropVerticalBox's OnAcceptDrop event */
+	FOnAcceptDrop OnAcceptDrop;
+
+	/** Are we currently dragging/dropping over this item? */
+	TOptional<EItemDropZone> ItemDropZone;
+
+	/** Delegate triggered when a user starts to drag a slot item */
+	FOnDragAndDropVerticalBoxDragDetected OnDragDetected_Handler;
+
+	/** Delegate triggered when a user's drag enters the bounds of this slot item */
+	FOnDragAndDropVerticalBoxDragEnter OnDragEnter_Handler;
+
+	/** Delegate triggered when a user's drag leaves the bounds of this slot item */
+	FOnDragAndDropVerticalBoxDragLeave OnDragLeave_Handler;
+
+	/** Delegate triggered when a user's drag is dropped in the bounds of this slot item */
+	FOnDragAndDropVerticalBoxDrop OnDrop_Handler;
+
+	/** Brush used to provide feedback that a user can drop below the hovered row. */
+	FSlateBrush DropIndicator_Above;
+	FSlateBrush DropIndicator_Below;
+
+	/** This is required for the paint to access which item we're hovering */
+	FVector2D CurrentDragOperationScreenSpaceLocation;
+	int32 CurrentDragOverSlotIndex;
+
+	/** @return the zone (above, below) based on where the user is hovering over */
+	EItemDropZone ZoneFromPointerPosition(FVector2D LocalPointerPos, const FGeometry& CurrentGeometry, const FGeometry& StartGeometry) const;
+};
\ No newline at end of file
diff --git a/Engine/Source/Runtime/SlateCore/Public/Widgets/SWidget.h b/Engine/Source/Runtime/SlateCore/Public/Widgets/SWidget.h
--- a/Engine/Source/Runtime/SlateCore/Public/Widgets/SWidget.h
+++ b/Engine/Source/Runtime/SlateCore/Public/Widgets/SWidget.h
@@ -54,6 +54,12 @@ DECLARE_DELEGATE_OneParam(
 
 enum class EPopupMethod : uint8;
 
+namespace SharedPointerInternals
+{
+	template <typename ObjectType>
+	class TIntrusiveReferenceController;
+}
+
 class SLATECORE_API FSlateControlledConstruction
 {
 public:
@@ -62,10 +68,10 @@ public:
 	
 private:
 	/** UI objects cannot be copy-constructed */
-	FSlateControlledConstruction(const FSlateControlledConstruction& Other){}
+	FSlateControlledConstruction(const FSlateControlledConstruction& Other) = delete;
 	
 	/** UI objects cannot be copied. */
-	void operator= (const FSlateControlledConstruction& Other){}
+	void operator= (const FSlateControlledConstruction& Other) = delete;
 
 	/** Widgets should only ever be constructed via SNew or SAssignNew */
 	void* operator new ( const size_t InSize )
@@ -73,9 +79,18 @@ private:
 		return FMemory::Malloc(InSize);
 	}
 
+	/** Widgets should only ever be constructed via SNew or SAssignNew */
+	void* operator new ( const size_t InSize, void* Addr )
+	{
+		return Addr;
+	}
+
 	template<class WidgetType, bool bIsUserWidget>
 	friend struct TWidgetAllocator;
 
+	template <typename ObjectType>
+	friend class SharedPointerInternals::TIntrusiveReferenceController;
+
 public:
 	void operator delete(void* mem)
 	{
diff --git a/Engine/Source/Runtime/SlateCore/Public/Widgets/SWindow.h b/Engine/Source/Runtime/SlateCore/Public/Widgets/SWindow.h
--- a/Engine/Source/Runtime/SlateCore/Public/Widgets/SWindow.h
+++ b/Engine/Source/Runtime/SlateCore/Public/Widgets/SWindow.h
@@ -59,37 +59,31 @@ DECLARE_DELEGATE_RetVal_OneParam( int32, FOnSwitchWorldHack, int32 );
 
 
 /** Enum to describe how to auto-center an SWindow */
-namespace EAutoCenter
+enum class EAutoCenter : uint8
 {
-	enum Type
-	{
-		/** Don't auto-center the window */
-		None,
+	/** Don't auto-center the window */
+	None,
 
-		/** Auto-center the window on the primary work area */
-		PrimaryWorkArea,
+	/** Auto-center the window on the primary work area */
+	PrimaryWorkArea,
 
-		/** Auto-center the window on the preferred work area, determined using GetPreferredWorkArea() */
-		PreferredWorkArea,
-	};
-}
+	/** Auto-center the window on the preferred work area, determined using GetPreferredWorkArea() */
+	PreferredWorkArea,
+};
 
 
 /** Enum to describe how windows are sized */
-namespace ESizingRule
+enum class ESizingRule : uint8
 {
-	enum Type
-	{
-		/* The windows size fixed and cannot be resized **/
-		FixedSize,
-		
-		/** The window size is computed from its content and cannot be resized by users */
-		Autosized,
+	/* The windows size fixed and cannot be resized **/
+	FixedSize,
+
+	/** The window size is computed from its content and cannot be resized by users */
+	Autosized,
 
-		/** The window can be resized by users */
-		UserSized,
-	};
-}
+	/** The window can be resized by users */
+	UserSized,
+};
 
 /** Proxy structure to handle deprecated construction from bool */
 struct FWindowTransparency
@@ -168,7 +162,7 @@ public:
 
 		/** The windows auto-centering mode. If set to anything other than None, then the
 			ScreenPosition value will be ignored */
-		SLATE_ARGUMENT( EAutoCenter::Type, AutoCenter )
+		SLATE_ARGUMENT( EAutoCenter, AutoCenter )
 
 		/** Screen-space position where the window should be initially located. */
 		SLATE_ARGUMENT( FVector2D, ScreenPosition )
@@ -189,7 +183,7 @@ public:
 		SLATE_ARGUMENT(bool, IsInitiallyMinimized)
 
 		/** How the window should be sized */
-		SLATE_ARGUMENT( ESizingRule::Type, SizingRule )
+		SLATE_ARGUMENT( ESizingRule, SizingRule )
 
 		/** True if this should be a 'pop-up' window */
 		SLATE_ARGUMENT( bool, IsPopupWindow )
@@ -628,7 +622,7 @@ public:
 	bool IsAutosized() const;
 
 	/** Should this window automatically derive its size based on its content or be user-drive? */
-	void SetSizingRule( ESizingRule::Type InSizingRule );
+	void SetSizingRule( ESizingRule InSizingRule );
 
 	/** @return true if this is a vanilla window, or one being used for some special purpose: e.g. tooltip or menu */
 	bool IsRegularWindow() const;
@@ -894,10 +888,10 @@ protected:
 	float Opacity;
 
 	/** How to size the window */
-	ESizingRule::Type SizingRule;
+	ESizingRule SizingRule;
 
 	/** How to auto center the window */
-	EAutoCenter::Type AutoCenterRule;
+	EAutoCenter AutoCenterRule;
 
 	/** Transparency setting for this window */
 	EWindowTransparency TransparencySupport;
diff --git a/Engine/Source/Runtime/SlateRHIRenderer/Private/SlateAtlasedTextureResource.cpp b/Engine/Source/Runtime/SlateRHIRenderer/Private/SlateAtlasedTextureResource.cpp
--- a/Engine/Source/Runtime/SlateRHIRenderer/Private/SlateAtlasedTextureResource.cpp
+++ b/Engine/Source/Runtime/SlateRHIRenderer/Private/SlateAtlasedTextureResource.cpp
@@ -24,7 +24,8 @@ FSlateShaderResourceProxy* FSlateAtlasedTextureResource::FindOrCreateAtlasedProx
 	FSlateShaderResourceProxy* Proxy = ProxyMap.FindRef(InAtlasedObject);
 	if ( Proxy == nullptr )
 	{
-		FVector2D ActualSize(TextureObject->GetSurfaceWidth(), TextureObject->GetSurfaceHeight());
+		// when we use image-DrawAsBox with PaperSprite, we need to change its actual size as its actual dimension.
+		FVector2D ActualSize(TextureObject->GetSurfaceWidth() * AtlasData.SizeUV.X, TextureObject->GetSurfaceHeight() * AtlasData.SizeUV.Y);
 
 		Proxy = new FSlateShaderResourceProxy();
 		Proxy->Resource = this;
diff --git a/Engine/Source/Runtime/SlateRHIRenderer/Private/SlateRHIRenderer.cpp b/Engine/Source/Runtime/SlateRHIRenderer/Private/SlateRHIRenderer.cpp
--- a/Engine/Source/Runtime/SlateRHIRenderer/Private/SlateRHIRenderer.cpp
+++ b/Engine/Source/Runtime/SlateRHIRenderer/Private/SlateRHIRenderer.cpp
@@ -39,8 +39,6 @@ DECLARE_FLOAT_COUNTER_STAT(TEXT("Slate UI"), Stat_GPU_SlateUI, STATGROUP_GPU);
 // Defines the maximum size that a slate viewport will create
 #define MAX_VIEWPORT_SIZE 16384
 
-#define USE_MAX_DRAWBUFFERS 0
-
 static TAutoConsoleVariable<float> CVarUILevel(
 	TEXT("r.HDR.UI.Level"),
 	1.0f,
@@ -147,10 +145,8 @@ void FSlateRHIRenderer::FViewportInfo::RecreateDepthBuffer_RenderThread()
 
 FSlateRHIRenderer::FSlateRHIRenderer( TSharedRef<FSlateFontServices> InSlateFontServices, TSharedRef<FSlateRHIResourceManager> InResourceManager )
 	: FSlateRenderer(InSlateFontServices)
-#if USE_MAX_DRAWBUFFERS
 	, EnqueuedWindowDrawBuffer(NULL)
-	, FreeBufferIndex(1)
-#endif
+	, FreeBufferIndex(0)
 	, CurrentSceneIndex(-1)
 {
 	ResourceManager = InResourceManager;
@@ -1443,6 +1439,18 @@ bool FSlateRHIRenderer::GenerateDynamicImageResource( FName ResourceName, uint32
 	return TextureResource.IsValid();
 }
 
+bool FSlateRHIRenderer::GenerateDynamicImageResource(FName ResourceName, FSlateTextureDataRef TextureData)
+{
+	check(IsInGameThread());
+
+	TSharedPtr<FSlateDynamicTextureResource> TextureResource = ResourceManager->GetDynamicTextureResourceByName(ResourceName);
+	if ( !TextureResource.IsValid() )
+	{
+		TextureResource = ResourceManager->MakeDynamicTextureResource(ResourceName, TextureData);
+	}
+	return TextureResource.IsValid();
+}
+
 FSlateResourceHandle FSlateRHIRenderer::GetResourceHandle( const FSlateBrush& Brush )
 {
 	return ResourceManager->GetResourceHandle( Brush );
diff --git a/Engine/Source/Runtime/SlateRHIRenderer/Private/SlateRHIRenderer.h b/Engine/Source/Runtime/SlateRHIRenderer/Private/SlateRHIRenderer.h
--- a/Engine/Source/Runtime/SlateRHIRenderer/Private/SlateRHIRenderer.h
+++ b/Engine/Source/Runtime/SlateRHIRenderer/Private/SlateRHIRenderer.h
@@ -229,6 +229,7 @@ public:
 	virtual void RemoveDynamicBrushResource( TSharedPtr<FSlateDynamicImageBrush> BrushToRemove ) override;
 	virtual FIntPoint GenerateDynamicImageResource(const FName InTextureName) override;
 	virtual bool GenerateDynamicImageResource( FName ResourceName, uint32 Width, uint32 Height, const TArray< uint8 >& Bytes ) override;
+	virtual bool GenerateDynamicImageResource( FName ResourceName, FSlateTextureDataRef TextureData ) override;
 	virtual FSlateResourceHandle GetResourceHandle( const FSlateBrush& Brush ) override;
 	virtual void* GetViewportResource( const SWindow& Window ) override;
 	virtual void SetColorVisionDeficiencyType( uint32 Type ) override;
diff --git a/Engine/Source/Runtime/SlateRHIRenderer/Private/SlateRHIRenderingPolicy.cpp b/Engine/Source/Runtime/SlateRHIRenderer/Private/SlateRHIRenderingPolicy.cpp
--- a/Engine/Source/Runtime/SlateRHIRenderer/Private/SlateRHIRenderingPolicy.cpp
+++ b/Engine/Source/Runtime/SlateRHIRenderer/Private/SlateRHIRenderingPolicy.cpp
@@ -321,8 +321,8 @@ void FSlateRHIRenderingPolicy::DrawElements(FRHICommandListImmediate& RHICmdList
 		}
 
 		const FSlateShaderResource* ShaderResource = RenderBatch.Texture;
-		const ESlateBatchDrawFlag::Type DrawFlags = RenderBatch.DrawFlags;
-		const ESlateDrawEffect::Type DrawEffects = RenderBatch.DrawEffects;
+		const ESlateBatchDrawFlag DrawFlags = RenderBatch.DrawFlags;
+		const ESlateDrawEffect DrawEffects = RenderBatch.DrawEffects;
 		const ESlateShader::Type ShaderType = RenderBatch.ShaderType;
 		const FShaderParams& ShaderParams = RenderBatch.ShaderParams;
 
@@ -370,9 +370,9 @@ void FSlateRHIRenderingPolicy::DrawElements(FRHICommandListImmediate& RHICmdList
 
 #if !DEBUG_OVERDRAW
 				RHICmdList.SetBlendState(
-					( RenderBatch.DrawFlags & ESlateBatchDrawFlag::NoBlending )
+					EnumHasAllFlags( DrawFlags, ESlateBatchDrawFlag::NoBlending )
 					? TStaticBlendState<>::GetRHI()
-					: ( ( RenderBatch.DrawFlags & ESlateBatchDrawFlag::PreMultipliedAlpha )
+					: ( EnumHasAllFlags(DrawFlags, ESlateBatchDrawFlag::PreMultipliedAlpha )
 						? TStaticBlendState<CW_RGBA, BO_Add, BF_One, BF_InverseSourceAlpha, BO_Add, BF_One, BF_InverseSourceAlpha>::GetRHI()
 						: TStaticBlendState<CW_RGBA, BO_Add, BF_SourceAlpha, BF_InverseSourceAlpha, BO_Add, BF_One, BF_InverseSourceAlpha>::GetRHI() )
 					);
@@ -380,7 +380,7 @@ void FSlateRHIRenderingPolicy::DrawElements(FRHICommandListImmediate& RHICmdList
 				RHICmdList.SetBlendState(TStaticBlendState<CW_RGB, BO_Add, BF_One, BF_One, BO_Add, BF_Zero, BF_InverseSourceAlpha>::GetRHI());
 #endif
 
-				if (DrawFlags & ESlateBatchDrawFlag::Wireframe)
+				if ( EnumHasAllFlags(DrawFlags, ESlateBatchDrawFlag::Wireframe) )
 				{
 					RHICmdList.SetRasterizerState(TStaticRasterizerState<FM_Wireframe, CM_None, false>::GetRHI());
 				}
@@ -409,7 +409,7 @@ void FSlateRHIRenderingPolicy::DrawElements(FRHICommandListImmediate& RHICmdList
 					else
 					{	
 						FTextureRHIParamRef NativeTextureRHI = ((TSlateTexture<FTexture2DRHIRef>*)ShaderResource)->GetTypedResource();
-						// Atlas textures that have no content are never initialised but null textures are invalid on many platforms.
+						// Atlas textures that have no content are never initialized but null textures are invalid on many platforms.
 						TextureRHI = NativeTextureRHI ? NativeTextureRHI : (FTextureRHIParamRef)GWhiteTexture->TextureRHI;
 					}
 
@@ -427,7 +427,7 @@ void FSlateRHIRenderingPolicy::DrawElements(FRHICommandListImmediate& RHICmdList
 							SamplerState = TStaticSamplerState<SF_Bilinear, AM_Wrap, AM_Wrap, AM_Wrap>::GetRHI();
 						}
 					}
-					else if ( DrawFlags & ESlateBatchDrawFlag::TileU )
+					else if ( EnumHasAllFlags(DrawFlags, ESlateBatchDrawFlag::TileU) )
 					{
 						switch ( Filter )
 						{
@@ -441,7 +441,7 @@ void FSlateRHIRenderingPolicy::DrawElements(FRHICommandListImmediate& RHICmdList
 							SamplerState = TStaticSamplerState<SF_Bilinear, AM_Wrap, AM_Clamp, AM_Wrap>::GetRHI();
 						}
 					}
-					else if ( DrawFlags & ESlateBatchDrawFlag::TileV )
+					else if ( EnumHasAllFlags(DrawFlags, ESlateBatchDrawFlag::TileV) )
 					{
 						switch ( Filter )
 						{
@@ -474,7 +474,8 @@ void FSlateRHIRenderingPolicy::DrawElements(FRHICommandListImmediate& RHICmdList
                 
 				PixelShader->SetTexture(RHICmdList, TextureRHI, SamplerState);
 				PixelShader->SetShaderParams(RHICmdList, ShaderParams.PixelParams);
-				PixelShader->SetDisplayGamma(RHICmdList, (DrawFlags & ESlateBatchDrawFlag::NoGamma) ? 1.0f : DisplayGamma);
+				PixelShader->SetDisplayGamma(RHICmdList, EnumHasAllFlags(DrawFlags, ESlateBatchDrawFlag::NoGamma ) ? 1.0f : DisplayGamma);
+				PixelShader->SetInvertAlpha(RHICmdList, EnumHasAllFlags(DrawEffects, ESlateDrawEffect::InvertAlpha ) ? true : false );
 
 	
 				// for RHIs that can't handle VertexOffset, we need to offset the stream source each time
@@ -556,9 +557,9 @@ void FSlateRHIRenderingPolicy::DrawElements(FRHICommandListImmediate& RHICmdList
 						VertexShader->SetMaterialShaderParameters(RHICmdList, ActiveSceneView, MaterialRenderProxy, Material);
 
 						PixelShader->SetParameters(RHICmdList, ActiveSceneView, MaterialRenderProxy, Material, DisplayGamma, ShaderParams.PixelParams);
-						PixelShader->SetDisplayGamma(RHICmdList, (DrawFlags & ESlateBatchDrawFlag::NoGamma) ? 1.0f : DisplayGamma);
+						PixelShader->SetDisplayGamma(RHICmdList, EnumHasAllFlags(DrawFlags, ESlateBatchDrawFlag::NoGamma) ? 1.0f : DisplayGamma);
 
-						if (DrawFlags & ESlateBatchDrawFlag::Wireframe)
+						if ( EnumHasAllFlags(DrawFlags, ESlateBatchDrawFlag::Wireframe) )
 						{
 							RHICmdList.SetRasterizerState(TStaticRasterizerState<FM_Wireframe, CM_None, false>::GetRHI());
 						}
@@ -687,7 +688,7 @@ void FSlateRHIRenderingPolicy::DrawElements(FRHICommandListImmediate& RHICmdList
 }
 
 
-FSlateElementPS* FSlateRHIRenderingPolicy::GetTexturePixelShader( ESlateShader::Type ShaderType, ESlateDrawEffect::Type DrawEffects )
+FSlateElementPS* FSlateRHIRenderingPolicy::GetTexturePixelShader( ESlateShader::Type ShaderType, ESlateDrawEffect DrawEffects )
 {
 	FSlateElementPS* PixelShader = nullptr;
 	const auto FeatureLevel = GMaxRHIFeatureLevel;
@@ -695,8 +696,8 @@ FSlateElementPS* FSlateRHIRenderingPolicy::GetTexturePixelShader( ESlateShader::
 
 #if !DEBUG_OVERDRAW
 	
-	const bool bDrawDisabled = ( DrawEffects & ESlateDrawEffect::DisabledEffect ) != 0;
-	const bool bUseTextureAlpha = ( DrawEffects & ESlateDrawEffect::IgnoreTextureAlpha ) == 0;
+	const bool bDrawDisabled = EnumHasAllFlags( DrawEffects, ESlateDrawEffect::DisabledEffect );
+	const bool bUseTextureAlpha = !EnumHasAllFlags( DrawEffects, ESlateDrawEffect::IgnoreTextureAlpha );
 
 	if ( bDrawDisabled )
 	{
@@ -773,12 +774,12 @@ FSlateElementPS* FSlateRHIRenderingPolicy::GetTexturePixelShader( ESlateShader::
 	return PixelShader;
 }
 
-FSlateMaterialShaderPS* FSlateRHIRenderingPolicy::GetMaterialPixelShader( const FMaterial* Material, ESlateShader::Type ShaderType, ESlateDrawEffect::Type DrawEffects )
+FSlateMaterialShaderPS* FSlateRHIRenderingPolicy::GetMaterialPixelShader( const FMaterial* Material, ESlateShader::Type ShaderType, ESlateDrawEffect DrawEffects )
 {
 	const FMaterialShaderMap* MaterialShaderMap = Material->GetRenderingThreadShaderMap();
 
-	const bool bDrawDisabled = (DrawEffects & ESlateDrawEffect::DisabledEffect) != 0;
-	const bool bUseTextureAlpha = (DrawEffects & ESlateDrawEffect::IgnoreTextureAlpha) == 0;
+	const bool bDrawDisabled = EnumHasAllFlags(DrawEffects, ESlateDrawEffect::DisabledEffect);
+	const bool bUseTextureAlpha = !EnumHasAllFlags(DrawEffects, ESlateDrawEffect::IgnoreTextureAlpha);
 
 	FShader* FoundShader = nullptr;
 	switch (ShaderType)
diff --git a/Engine/Source/Runtime/SlateRHIRenderer/Private/SlateRHIRenderingPolicy.h b/Engine/Source/Runtime/SlateRHIRenderer/Private/SlateRHIRenderingPolicy.h
--- a/Engine/Source/Runtime/SlateRHIRenderer/Private/SlateRHIRenderingPolicy.h
+++ b/Engine/Source/Runtime/SlateRHIRenderer/Private/SlateRHIRenderingPolicy.h
@@ -56,8 +56,8 @@ private:
 	 * @param DrawEffects	Draw effects being used
 	 * @return The pixel shader for use with the shader type and draw effects
 	 */
-	class FSlateElementPS* GetTexturePixelShader( ESlateShader::Type ShaderType, ESlateDrawEffect::Type DrawEffects );
-	class FSlateMaterialShaderPS* GetMaterialPixelShader( const class FMaterial* Material, ESlateShader::Type ShaderType, ESlateDrawEffect::Type DrawEffects );
+	class FSlateElementPS* GetTexturePixelShader( ESlateShader::Type ShaderType, ESlateDrawEffect DrawEffects );
+	class FSlateMaterialShaderPS* GetMaterialPixelShader( const class FMaterial* Material, ESlateShader::Type ShaderType, ESlateDrawEffect DrawEffects );
 	class FSlateMaterialShaderVS* GetMaterialVertexShader( const class FMaterial* Material, bool bUseInstancing );
 
 	/** @return The RHI primitive type from the Slate primitive type */
diff --git a/Engine/Source/Runtime/SlateRHIRenderer/Private/SlateRHIResourceManager.cpp b/Engine/Source/Runtime/SlateRHIRenderer/Private/SlateRHIResourceManager.cpp
--- a/Engine/Source/Runtime/SlateRHIRenderer/Private/SlateRHIResourceManager.cpp
+++ b/Engine/Source/Runtime/SlateRHIRenderer/Private/SlateRHIResourceManager.cpp
@@ -568,11 +568,15 @@ ISlateAtlasProvider* FSlateRHIResourceManager::GetTextureAtlasProvider()
 	return this;
 }
 
-TSharedPtr<FSlateDynamicTextureResource> FSlateRHIResourceManager::MakeDynamicTextureResource( FName ResourceName, uint32 Width, uint32 Height, const TArray< uint8 >& Bytes )
+TSharedPtr<FSlateDynamicTextureResource> FSlateRHIResourceManager::MakeDynamicTextureResource(FName ResourceName, uint32 Width, uint32 Height, const TArray< uint8 >& Bytes)
 {
 	// Make storage for the image
-	FSlateTextureDataRef TextureStorage = MakeShareable( new FSlateTextureData( Width, Height, GPixelFormats[PF_B8G8R8A8].BlockBytes, Bytes ) );
+	FSlateTextureDataRef TextureStorage = MakeShareable(new FSlateTextureData(Width, Height, GPixelFormats[PF_B8G8R8A8].BlockBytes, Bytes));
+	return MakeDynamicTextureResource(ResourceName, TextureStorage);
+}
 
+TSharedPtr<FSlateDynamicTextureResource> FSlateRHIResourceManager::MakeDynamicTextureResource(FName ResourceName, FSlateTextureDataRef TextureData)
+{
 	TSharedPtr<FSlateDynamicTextureResource> TextureResource;
 	// Get a resource from the free list if possible
 	if(DynamicTextureFreeList.Num() > 0)
@@ -585,19 +589,18 @@ TSharedPtr<FSlateDynamicTextureResource> FSlateRHIResourceManager::MakeDynamicTe
 		TextureResource = MakeShareable(new FSlateDynamicTextureResource(nullptr));
 	}
 
-	TextureResource->Proxy->ActualSize = FIntPoint(TextureStorage->GetWidth(), TextureStorage->GetHeight());
+	TextureResource->Proxy->ActualSize = FIntPoint(TextureData->GetWidth(), TextureData->GetHeight());
 
 
 	// Init render thread data
 	ENQUEUE_UNIQUE_RENDER_COMMAND_TWOPARAMETER(InitNewSlateDynamicTextureResource,
 		FSlateDynamicTextureResource*, TextureResource, TextureResource.Get(),
-		FSlateTextureDataPtr, InNewTextureData, TextureStorage,
+		FSlateTextureDataPtr, InNewTextureData, TextureData,
 	{
 		if(InNewTextureData.IsValid())
 		{
 			// Set the texture to use as the texture we just loaded
 			TextureResource->RHIRefTexture->SetTextureData(InNewTextureData, PF_B8G8R8A8, TexCreate_SRGB);
-
 		}
 
 		// Initialize and link the rendering resource
diff --git a/Engine/Source/Runtime/SlateRHIRenderer/Private/SlateRHIResourceManager.h b/Engine/Source/Runtime/SlateRHIRenderer/Private/SlateRHIResourceManager.h
--- a/Engine/Source/Runtime/SlateRHIRenderer/Private/SlateRHIResourceManager.h
+++ b/Engine/Source/Runtime/SlateRHIRenderer/Private/SlateRHIResourceManager.h
@@ -190,6 +190,8 @@ public:
 	 */
 	TSharedPtr<FSlateDynamicTextureResource> MakeDynamicTextureResource( FName ResourceName, uint32 Width, uint32 Height, const TArray< uint8 >& Bytes );
 
+	TSharedPtr<FSlateDynamicTextureResource> MakeDynamicTextureResource(FName ResourceName, FSlateTextureDataRef TextureData);
+
 	/**
 	 * Find a dynamic texture resource 
 	 *
diff --git a/Engine/Source/Runtime/SlateRHIRenderer/Private/SlateShaders.cpp b/Engine/Source/Runtime/SlateRHIRenderer/Private/SlateShaders.cpp
--- a/Engine/Source/Runtime/SlateRHIRenderer/Private/SlateShaders.cpp
+++ b/Engine/Source/Runtime/SlateRHIRenderer/Private/SlateShaders.cpp
@@ -1,13 +1,12 @@
 // Copyright 1998-2017 Epic Games, Inc. All Rights Reserved.
 
 #include "SlateShaders.h"
-
+#include "RenderingCommon.h"
 
 /** Flag to determine if we are running with a color vision deficiency shader on */
 uint32 GSlateShaderColorVisionDeficiencyType = 0;
 
-
-IMPLEMENT_SHADER_TYPE(, FSlateElementVS, TEXT("SlateVertexShader"),TEXT("Main"),SF_Vertex);
+IMPLEMENT_SHADER_TYPE(, FSlateElementVS, TEXT("SlateVertexShader"), TEXT("Main"), SF_Vertex);
 
 IMPLEMENT_SHADER_TYPE(, FSlateDebugOverdrawPS, TEXT("SlateElementPixelShader"), TEXT("DebugOverdrawMain"), SF_Pixel );
 
@@ -55,6 +54,7 @@ void FSlateVertexDeclaration::InitRHI()
 	Elements.Add(FVertexElement(0, STRUCT_OFFSET(FSlateVertex, ClipRect) + STRUCT_OFFSET(FSlateRotatedClipRectType, TopLeft), VET_Float2, 3, Stride));
 	Elements.Add(FVertexElement(0, STRUCT_OFFSET(FSlateVertex, ClipRect) + STRUCT_OFFSET(FSlateRotatedClipRectType, ExtentX), VET_Float4, 4, Stride));
 	Elements.Add(FVertexElement(0, STRUCT_OFFSET(FSlateVertex, Color), VET_Color, 5, Stride));
+	Elements.Add(FVertexElement(0, STRUCT_OFFSET(FSlateVertex, PixelSize), VET_UShort2, 6, Stride));
 
 	VertexDeclarationRHI = RHICreateVertexDeclaration(Elements);
 }
@@ -78,7 +78,8 @@ void FSlateInstancedVertexDeclaration::InitRHI()
 	Elements.Add(FVertexElement(0, STRUCT_OFFSET(FSlateVertex, ClipRect) + STRUCT_OFFSET(FSlateRotatedClipRectType, TopLeft), VET_Float2, 3, Stride));
 	Elements.Add(FVertexElement(0, STRUCT_OFFSET(FSlateVertex, ClipRect) + STRUCT_OFFSET(FSlateRotatedClipRectType, ExtentX), VET_Float4, 4, Stride));
 	Elements.Add(FVertexElement(0, STRUCT_OFFSET(FSlateVertex, Color), VET_Color, 5, Stride));
-	Elements.Add(FVertexElement(1, 0, VET_Float4, 6, sizeof(FVector4), true));
+	Elements.Add(FVertexElement(0, STRUCT_OFFSET(FSlateVertex, PixelSize), VET_UShort2, 6, Stride));
+	Elements.Add(FVertexElement(1, 0, VET_Float4, 7, sizeof(FVector4), true));
 	
 	VertexDeclarationRHI = RHICreateVertexDeclaration(Elements);
 }
diff --git a/Engine/Source/Runtime/SlateRHIRenderer/Private/SlateShaders.h b/Engine/Source/Runtime/SlateRHIRenderer/Private/SlateShaders.h
--- a/Engine/Source/Runtime/SlateRHIRenderer/Private/SlateShaders.h
+++ b/Engine/Source/Runtime/SlateRHIRenderer/Private/SlateShaders.h
@@ -45,7 +45,7 @@ public:
 /** The slate Vertex shader representation */
 class FSlateElementVS : public FGlobalShader
 {
-	DECLARE_SHADER_TYPE( FSlateElementVS, Global );
+	DECLARE_SHADER_TYPE(FSlateElementVS, Global);
 public:
 	/** Indicates that this shader should be cached */
 	static bool ShouldCache( EShaderPlatform Platform ) { return true; }
@@ -112,6 +112,7 @@ public:
 		TextureParameterSampler.Bind( Initializer.ParameterMap, TEXT("ElementTextureSampler"));
 		ShaderParams.Bind( Initializer.ParameterMap, TEXT("ShaderParams"));
 		GammaValues.Bind( Initializer.ParameterMap,TEXT("GammaValues"));
+		InvertAlpha.Bind(Initializer.ParameterMap, TEXT("InvertAlpha"));
 	}
 
 	static void ModifyCompilationEnvironment(EShaderPlatform Platform, FShaderCompilerEnvironment& OutEnvironment);
@@ -149,6 +150,14 @@ public:
 		SetShaderValue(RHICmdList, GetPixelShader(),GammaValues,InGammaValues);
 	}
 
+	/**
+	 * Sets if we should invert alpha of the incoming.
+	 */
+	void SetInvertAlpha(FRHICommandList& RHICmdList, bool bInvertAlpha )
+	{
+		SetShaderValue(RHICmdList, GetPixelShader(), InvertAlpha, bInvertAlpha ? 1.0f : 0.0f );
+	}
+
 	virtual bool Serialize( FArchive& Ar )
 	{
 		bool bShaderHasOutdatedParameters = FGlobalShader::Serialize( Ar );
@@ -157,6 +166,7 @@ public:
 		Ar << TextureParameterSampler;
 		Ar << ShaderParams;
 		Ar << GammaValues;
+		Ar << InvertAlpha;
 
 		return bShaderHasOutdatedParameters;
 	}
@@ -166,6 +176,7 @@ private:
 	FShaderResourceParameter TextureParameterSampler;
 	FShaderParameter ShaderParams;
 	FShaderParameter GammaValues;
+	FShaderParameter InvertAlpha;
 };
 
 /** 
diff --git a/Engine/Source/Runtime/UMG/Private/Components/ComboBoxString.cpp b/Engine/Source/Runtime/UMG/Private/Components/ComboBoxString.cpp
--- a/Engine/Source/Runtime/UMG/Private/Components/ComboBoxString.cpp
+++ b/Engine/Source/Runtime/UMG/Private/Components/ComboBoxString.cpp
@@ -51,7 +51,7 @@ void UComboBoxString::ReleaseSlateResources(bool bReleaseChildren)
 	Super::ReleaseSlateResources(bReleaseChildren);
 
 	MyComboBox.Reset();
-	ComoboBoxContent.Reset();
+	ComboBoxContent.Reset();
 }
 
 void UComboBoxString::PostLoad()
@@ -96,13 +96,13 @@ TSharedRef<SWidget> UComboBoxString::RebuildWidget()
 		.OnComboBoxOpening(BIND_UOBJECT_DELEGATE(FOnComboBoxOpening, HandleOpening))
 		.IsFocusable(bIsFocusable)
 		[
-			SAssignNew(ComoboBoxContent, SBox)
+			SAssignNew(ComboBoxContent, SBox)
 		];
 
 	if ( InitialIndex != -1 )
 	{
 		// Generate the widget for the initially selected widget if needed
-		ComoboBoxContent->SetContent(HandleGenerateWidget(CurrentOptionPtr));
+		ComboBoxContent->SetContent(HandleGenerateWidget(CurrentOptionPtr));
 	}
 
 	return MyComboBox.ToSharedRef();
@@ -181,9 +181,9 @@ void UComboBoxString::ClearSelection()
 		MyComboBox->ClearSelection();
 	}
 
-	if ( ComoboBoxContent.IsValid() )
+	if ( ComboBoxContent.IsValid() )
 	{
-		ComoboBoxContent->SetContent(SNullWidget::NullWidget);
+		ComboBoxContent->SetContent(SNullWidget::NullWidget);
 	}
 }
 
@@ -203,10 +203,10 @@ void UComboBoxString::SetSelectedOption(FString Option)
 		CurrentOptionPtr = Options[InitialIndex];
 		SelectedOption = Option;
 
-		if ( ComoboBoxContent.IsValid() )
+		if ( ComboBoxContent.IsValid() )
 		{
 			MyComboBox->SetSelectedItem(CurrentOptionPtr);
-			ComoboBoxContent->SetContent(HandleGenerateWidget(CurrentOptionPtr));
+			ComboBoxContent->SetContent(HandleGenerateWidget(CurrentOptionPtr));
 		}
 	}
 }
@@ -250,13 +250,13 @@ void UComboBoxString::HandleSelectionChanged(TSharedPtr<FString> Item, ESelectIn
 	CurrentOptionPtr = Item;
 	SelectedOption = CurrentOptionPtr.IsValid() ? CurrentOptionPtr.ToSharedRef().Get() : FString();
 
+	// When the selection changes we always generate another widget to represent the content area of the combobox.
+	ComboBoxContent->SetContent(HandleGenerateWidget(CurrentOptionPtr));
+
 	if ( !IsDesignTime() )
 	{
 		OnSelectionChanged.Broadcast(Item.IsValid() ? *Item : FString(), SelectionType);
 	}
-
-	// When the selection changes we always generate another widget to represent the content area of the comobox.
-	ComoboBoxContent->SetContent( HandleGenerateWidget(CurrentOptionPtr) );
 }
 
 void UComboBoxString::HandleOpening()
diff --git a/Engine/Source/Runtime/UMG/Private/Components/WidgetInteractionComponent.cpp b/Engine/Source/Runtime/UMG/Private/Components/WidgetInteractionComponent.cpp
--- a/Engine/Source/Runtime/UMG/Private/Components/WidgetInteractionComponent.cpp
+++ b/Engine/Source/Runtime/UMG/Private/Components/WidgetInteractionComponent.cpp
@@ -6,6 +6,7 @@
 #include "Kismet/KismetSystemLibrary.h"
 #include "Components/ArrowComponent.h"
 #include "Framework/Application/SlateApplication.h"
+#include "Kismet/GameplayStatics.h"
 
 #include "Components/WidgetComponent.h"
 
@@ -97,6 +98,8 @@ UWidgetInteractionComponent::FWidgetTraceResult UWidgetInteractionComponent::Per
 {
 	FWidgetTraceResult TraceResult;
 
+	TArray<FHitResult> MultiHits;
+
 	switch( InteractionSource )
 	{
 		case EWidgetInteractionSource::World:
@@ -111,7 +114,7 @@ UWidgetInteractionComponent::FWidgetTraceResult UWidgetInteractionComponent::Per
 			FCollisionQueryParams Params = FCollisionQueryParams::DefaultQueryParam;
 			Params.AddIgnoredComponents(PrimitiveChildren);
 
-			TraceResult.bWasHit = GetWorld()->LineTraceSingleByChannel(TraceResult.HitResult, WorldLocation, WorldLocation + ( Direction * InteractionDistance ), TraceChannel, Params);
+			GetWorld()->LineTraceMultiByChannel(MultiHits, WorldLocation, WorldLocation + ( Direction * InteractionDistance ), TraceChannel, Params);
 			break;
 		}
 		case EWidgetInteractionSource::Mouse:
@@ -124,9 +127,8 @@ UWidgetInteractionComponent::FWidgetTraceResult UWidgetInteractionComponent::Per
 			Params.AddIgnoredComponents(PrimitiveChildren);
 
 			APlayerController* PlayerController = GetWorld()->GetFirstPlayerController();
-
 			ULocalPlayer* LocalPlayer = PlayerController->GetLocalPlayer();
-			TraceResult.bWasHit = false;
+			
 			if ( LocalPlayer && LocalPlayer->ViewportClient )
 			{
 				if ( InteractionSource == EWidgetInteractionSource::Mouse )
@@ -134,7 +136,12 @@ UWidgetInteractionComponent::FWidgetTraceResult UWidgetInteractionComponent::Per
 					FVector2D MousePosition;
 					if ( LocalPlayer->ViewportClient->GetMousePosition(MousePosition) )
 					{
-						TraceResult.bWasHit = PlayerController->GetHitResultAtScreenPosition(MousePosition, TraceChannel, Params, TraceResult.HitResult);
+						FVector WorldOrigin;
+						FVector WorldDirection;
+						if ( UGameplayStatics::DeprojectScreenToWorld(PlayerController, MousePosition, WorldOrigin, WorldDirection) == true )
+						{
+							GetWorld()->LineTraceMultiByChannel(MultiHits, WorldOrigin, WorldOrigin + WorldDirection * InteractionDistance, TraceChannel, Params);
+						}
 					}
 				}
 				else if ( InteractionSource == EWidgetInteractionSource::CenterScreen )
@@ -142,15 +149,11 @@ UWidgetInteractionComponent::FWidgetTraceResult UWidgetInteractionComponent::Per
 					FVector2D ViewportSize;
 					LocalPlayer->ViewportClient->GetViewportSize(ViewportSize);
 
-					TraceResult.bWasHit = PlayerController->GetHitResultAtScreenPosition(ViewportSize * 0.5f, TraceChannel, Params, TraceResult.HitResult);
-				}
-
-				// Don't allow infinite distance hit testing.
-				if ( TraceResult.bWasHit )
-				{
-					if ( TraceResult.HitResult.Distance > InteractionDistance )
+					FVector WorldOrigin;
+					FVector WorldDirection;
+					if ( UGameplayStatics::DeprojectScreenToWorld(PlayerController, ViewportSize * 0.5f, WorldOrigin, WorldDirection) == true )
 					{
-						TraceResult = FWidgetTraceResult();
+						GetWorld()->LineTraceMultiByChannel(MultiHits, WorldOrigin, WorldOrigin + WorldDirection * InteractionDistance, TraceChannel, Params);
 					}
 				}
 			}
@@ -164,6 +167,27 @@ UWidgetInteractionComponent::FWidgetTraceResult UWidgetInteractionComponent::Per
 		}
 	}
 
+	// If it's not a custom interaction, we do some custom filtering to ignore invisible widgets.
+	if ( InteractionSource != EWidgetInteractionSource::Custom )
+	{
+		for ( const FHitResult& HitResult : MultiHits )
+		{
+			if ( UWidgetComponent* HitWidgetComponent = Cast<UWidgetComponent>(HitResult.GetComponent()) )
+			{
+				if ( HitWidgetComponent->IsVisible() )
+				{
+					TraceResult.bWasHit = true;
+					TraceResult.HitResult = HitResult;
+				}
+			}
+			else
+			{
+				// If we hit something that wasn't a widget component, we're done.
+				break;
+			}
+		}
+	}
+
 	// Resolve trace to location on widget.
 	if (TraceResult.bWasHit)
 	{
@@ -234,32 +258,27 @@ bool UWidgetInteractionComponent::CanInteractWithComponent(UWidgetComponent* Com
 	return bCanInteract;
 }
 
-void UWidgetInteractionComponent::SimulatePointerMovement()
+FWidgetPath UWidgetInteractionComponent::DetermineWidgetUnderPointer()
 {
+	FWidgetPath WidgetPathUnderPointer;
+
 	bIsHoveredWidgetInteractable = false;
 	bIsHoveredWidgetFocusable = false;
 	bIsHoveredWidgetHitTestVisible = false;
 
-	if ( !bEnableHitTesting )
-	{
-		return;
-	}
+	LocalHitLocation = LastLocalHitLocation;
 
-	if ( !CanSendInput() )
-	{
-		return;
-	}
-	
-	FWidgetPath WidgetPathUnderFinger;
 	UWidgetComponent* OldHoveredWidget = HoveredWidgetComponent;
 
+	HoveredWidgetComponent = nullptr;
+
 	FWidgetTraceResult TraceResult = PerformTrace();
 	LastHitResult = TraceResult.HitResult;
 	HoveredWidgetComponent = TraceResult.HitWidgetComponent;
 	LocalHitLocation = TraceResult.bWasHit
 		? TraceResult.LocalHitLocation
 		: LastLocalHitLocation;
-	WidgetPathUnderFinger = TraceResult.HitWidgetPath;
+	WidgetPathUnderPointer = TraceResult.HitWidgetPath;
 
 	if ( bShowDebug )
 	{
@@ -280,31 +299,6 @@ void UWidgetInteractionComponent::SimulatePointerMovement()
 			}
 		}
 	}
-	
-	FPointerEvent PointerEvent(
-		VirtualUser->GetUserIndex(),
-		PointerIndex,
-		LocalHitLocation,
-		LastLocalHitLocation,
-		PressedKeys,
-		FKey(),
-		0.0f,
-		ModifierKeys);
-	
-	if (WidgetPathUnderFinger.IsValid())
-	{
-		check(HoveredWidgetComponent);
-		LastWigetPath = WidgetPathUnderFinger;
-		
-		FSlateApplication::Get().RoutePointerMoveEvent(WidgetPathUnderFinger, PointerEvent, false);
-	}
-	else
-	{
-		FWidgetPath EmptyWidgetPath;
-		FSlateApplication::Get().RoutePointerMoveEvent(EmptyWidgetPath, PointerEvent, false);
-
-		LastWigetPath = FWeakWidgetPath();
-	}
 
 	if ( HoveredWidgetComponent )
 	{
@@ -313,20 +307,23 @@ void UWidgetInteractionComponent::SimulatePointerMovement()
 
 	LastLocalHitLocation = LocalHitLocation;
 
-	if ( WidgetPathUnderFinger.IsValid() )
+	if ( WidgetPathUnderPointer.IsValid() )
 	{
-		const FArrangedChildren::FArrangedWidgetArray& AllArrangedWidgets = WidgetPathUnderFinger.Widgets.GetInternalArray();
+		const FArrangedChildren::FArrangedWidgetArray& AllArrangedWidgets = WidgetPathUnderPointer.Widgets.GetInternalArray();
 		for ( const FArrangedWidget& ArrangedWidget : AllArrangedWidgets )
 		{
 			const TSharedRef<SWidget>& Widget = ArrangedWidget.Widget;
-			if ( Widget->IsInteractable() )
+			if ( Widget->IsEnabled() )
 			{
-				bIsHoveredWidgetInteractable = true;
-			}
+				if ( Widget->IsInteractable() )
+				{
+					bIsHoveredWidgetInteractable = true;
+				}
 
-			if ( Widget->SupportsKeyboardFocus() )
-			{
-				bIsHoveredWidgetFocusable = true;
+				if ( Widget->SupportsKeyboardFocus() )
+				{
+					bIsHoveredWidgetFocusable = true;
+				}
 			}
 
 			if ( Widget->GetVisibility().IsHitTestVisible() )
@@ -343,7 +340,49 @@ void UWidgetInteractionComponent::SimulatePointerMovement()
 			OldHoveredWidget->RequestRedraw();
 		}
 
-		OnHoveredWidgetChanged.Broadcast(HoveredWidgetComponent, OldHoveredWidget);
+		OnHoveredWidgetChanged.Broadcast( HoveredWidgetComponent, OldHoveredWidget );
+	}
+
+	return WidgetPathUnderPointer;
+}
+
+void UWidgetInteractionComponent::SimulatePointerMovement()
+{
+	if ( !bEnableHitTesting )
+	{
+		return;
+	}
+
+	if ( !CanSendInput() )
+	{
+		return;
+	}
+	
+	FWidgetPath WidgetPathUnderFinger = DetermineWidgetUnderPointer();
+	
+	FPointerEvent PointerEvent(
+		VirtualUser->GetUserIndex(),
+		PointerIndex,
+		LocalHitLocation,
+		LastLocalHitLocation,
+		PressedKeys,
+		FKey(),
+		0.0f,
+		ModifierKeys);
+	
+	if (WidgetPathUnderFinger.IsValid())
+	{
+		check(HoveredWidgetComponent);
+		LastWidgetPath = WidgetPathUnderFinger;
+		
+		FSlateApplication::Get().RoutePointerMoveEvent(WidgetPathUnderFinger, PointerEvent, false);
+	}
+	else
+	{
+		FWidgetPath EmptyWidgetPath;
+		FSlateApplication::Get().RoutePointerMoveEvent(EmptyWidgetPath, PointerEvent, false);
+
+		LastWidgetPath = FWeakWidgetPath();
 	}
 }
 
@@ -361,7 +400,13 @@ void UWidgetInteractionComponent::PressPointerKey(FKey Key)
 	
 	PressedKeys.Add(Key);
 	
-	FWidgetPath WidgetPathUnderFinger = LastWigetPath.ToWidgetPath();
+	if ( !LastWidgetPath.IsValid() )
+	{
+		// If the cached widget path isn't valid, attempt to find a valid widget since we might have received a touch input
+		LastWidgetPath = DetermineWidgetUnderPointer();
+	}
+
+	FWidgetPath WidgetPathUnderFinger = LastWidgetPath.ToWidgetPath();
 		
 	FPointerEvent PointerEvent(
 		VirtualUser->GetUserIndex(),
@@ -394,8 +439,8 @@ void UWidgetInteractionComponent::ReleasePointerKey(FKey Key)
 	
 	PressedKeys.Remove(Key);
 	
-	FWidgetPath WidgetPathUnderFinger = LastWigetPath.ToWidgetPath();
-		
+	FWidgetPath WidgetPathUnderFinger = LastWidgetPath.ToWidgetPath();
+
 	FPointerEvent PointerEvent(
 		VirtualUser->GetUserIndex(),
 		PointerIndex,
@@ -488,7 +533,7 @@ void UWidgetInteractionComponent::ScrollWheel(float ScrollDelta)
 		return;
 	}
 
-	FWidgetPath WidgetPathUnderFinger = LastWigetPath.ToWidgetPath();
+	FWidgetPath WidgetPathUnderFinger = LastWidgetPath.ToWidgetPath();
 	
 	FPointerEvent MouseWheelEvent(
 		VirtualUser->GetUserIndex(),
@@ -525,7 +570,7 @@ bool UWidgetInteractionComponent::IsOverHitTestVisibleWidget() const
 
 const FWeakWidgetPath& UWidgetInteractionComponent::GetHoveredWidgetPath() const
 {
-	return LastWigetPath;
+	return LastWidgetPath;
 }
 
 const FHitResult& UWidgetInteractionComponent::GetLastHitResult() const
diff --git a/Engine/Source/Runtime/UMG/Private/UserWidget.cpp b/Engine/Source/Runtime/UMG/Private/UserWidget.cpp
--- a/Engine/Source/Runtime/UMG/Private/UserWidget.cpp
+++ b/Engine/Source/Runtime/UMG/Private/UserWidget.cpp
@@ -12,6 +12,7 @@
 #include "Blueprint/WidgetTree.h"
 #include "Blueprint/WidgetBlueprintGeneratedClass.h"
 #include "Animation/UMGSequencePlayer.h"
+#include "UObject/UnrealType.h"
 
 
 #include "Blueprint/WidgetBlueprintLibrary.h"
@@ -494,9 +495,18 @@ void UUserWidget::OnWidgetRebuilt()
 
 	if (!IsDesignTime())
 	{
+		// Notify the widget to run per-construct.
+		NativePreConstruct();
+
 		// Notify the widget that it has been constructed.
 		NativeConstruct();
 	}
+#if WITH_EDITOR
+	else if ( HasAnyDesignerFlags(EWidgetDesignFlags::ExecutePreConstruct) )
+	{
+		NativePreConstruct();
+	}
+#endif
 }
 
 TSharedPtr<SWidget> UUserWidget::GetSlateWidgetFromName(const FName& Name) const
@@ -646,8 +656,8 @@ void UUserWidget::AddToScreen(ULocalPlayer* Player, int32 ZOrder)
 
 		FullScreenCanvas->AddSlot()
 			.Offset(BIND_UOBJECT_ATTRIBUTE(FMargin, GetFullScreenOffset))
-			.Anchors(BIND_UOBJECT_ATTRIBUTE(FAnchors, GetViewportAnchors))
-			.Alignment(BIND_UOBJECT_ATTRIBUTE(FVector2D, GetFullScreenAlignment))
+			.Anchors(BIND_UOBJECT_ATTRIBUTE(FAnchors, GetAnchorsInViewport))
+			.Alignment(BIND_UOBJECT_ATTRIBUTE(FVector2D, GetAlignmentInViewport))
 			[
 				UserSlateWidget
 			];
@@ -843,12 +853,12 @@ FMargin UUserWidget::GetFullScreenOffset() const
 	return FMargin(ViewportOffsets.Left, ViewportOffsets.Top, FinalSize.X, FinalSize.Y);
 }
 
-FAnchors UUserWidget::GetViewportAnchors() const
+FAnchors UUserWidget::GetAnchorsInViewport() const
 {
 	return ViewportAnchors;
 }
 
-FVector2D UUserWidget::GetFullScreenAlignment() const
+FVector2D UUserWidget::GetAlignmentInViewport() const
 {
 	return ViewportAlignment;
 }
@@ -894,6 +904,22 @@ void UUserWidget::SetDesignerFlags(EWidgetDesignFlags::Type NewFlags)
 	});
 }
 
+void UUserWidget::PostEditChangeProperty(FPropertyChangedEvent& PropertyChangedEvent)
+{
+	Super::PostEditChangeProperty(PropertyChangedEvent);
+
+	if ( PropertyChangedEvent.ChangeType != EPropertyChangeType::Interactive )
+	{
+		TSharedPtr<SWidget> SafeWidget = GetCachedWidget();
+		if ( SafeWidget.IsValid() )
+		{
+			// Re-Run execute PreConstruct when we get a post edit property change, to do something
+			// akin to running Sync Properties, so users don't have to recompile to see updates.
+			NativePreConstruct();
+		}
+	}
+}
+
 #endif
 
 void UUserWidget::OnAnimationStarted_Implementation(const UWidgetAnimation* Animation)
@@ -908,6 +934,11 @@ void UUserWidget::OnAnimationFinished_Implementation(const UWidgetAnimation* Ani
 
 // Native handling for SObjectWidget
 
+void UUserWidget::NativePreConstruct()
+{
+	PreConstruct(IsDesignTime());
+}
+
 void UUserWidget::NativeConstruct()
 {
 	Construct();
diff --git a/Engine/Source/Runtime/UMG/Public/Blueprint/UserWidget.h b/Engine/Source/Runtime/UMG/Public/Blueprint/UserWidget.h
--- a/Engine/Source/Runtime/UMG/Public/Blueprint/UserWidget.h
+++ b/Engine/Source/Runtime/UMG/Public/Blueprint/UserWidget.h
@@ -228,6 +228,14 @@ public:
 	UFUNCTION(BlueprintCallable, BlueprintCosmetic, Category="User Interface|Viewport")
 	void SetAlignmentInViewport(FVector2D Alignment);
 
+	/*  */
+	UFUNCTION(BlueprintCallable, BlueprintCosmetic, Category = "User Interface|Viewport")
+		FAnchors GetAnchorsInViewport() const;
+
+	/*  */
+	UFUNCTION(BlueprintCallable, BlueprintCosmetic, Category = "User Interface|Viewport")
+		FVector2D GetAlignmentInViewport() const;
+
 	/*  */
 	UFUNCTION(BlueprintPure, BlueprintCosmetic, Category="Appearance", meta=( DeprecatedFunction, DeprecationMessage="Use IsInViewport instead" ))
 	bool GetIsVisible() const;
@@ -277,6 +285,21 @@ public:
 	UFUNCTION(BlueprintCallable, BlueprintCosmetic, Category="Player")
 	class APawn* GetOwningPlayerPawn() const;
 
+	/**
+	 * Called by both the game and the editor.  Allows users to run initial setup for their widgets to better preview
+	 * the setup in the designer and since generally that same setup code is required at runtime, it's called there
+	 * as well.
+	 *
+	 * **WARNING**
+	 * This is intended purely for cosmetic updates using locally owned data, you can not safely access any game related
+	 * state, if you call something that doesn't expect to be run at editor time, you may crash the editor.
+	 *
+	 * In the event you save the asset with blueprint code that causes a crash on evaluation.  You can turn off
+	 * PreConstruct evaluation in the Widget Designer settings in the Editor Preferences.
+	 */
+	UFUNCTION(BlueprintImplementableEvent, BlueprintCosmetic, Category="User Interface")
+	void PreConstruct(bool IsDesignTime);
+
 	/**
 	 * Called after the underlying slate widget is constructed.  Depending on how the slate object is used
 	 * this event may be called multiple times due to adding and removing from the hierarchy.
@@ -847,6 +870,8 @@ public:
 	//~ End UWidget Interface
 
 	void SetDesignerFlags(EWidgetDesignFlags::Type NewFlags);
+
+	virtual void PostEditChangeProperty(FPropertyChangedEvent& PropertyChangedEvent) override;
 #endif
 
 public:
@@ -956,12 +981,11 @@ protected:
 	virtual void OnWidgetRebuilt() override;
 
 	FMargin GetFullScreenOffset() const;
-	FAnchors GetViewportAnchors() const;
-	FVector2D GetFullScreenAlignment() const;
 
 	//native SObjectWidget methods
 	friend class SObjectWidget;
 
+	virtual void NativePreConstruct();
 	virtual void NativeConstruct();
 	virtual void NativeDestruct();
 
diff --git a/Engine/Source/Runtime/UMG/Public/Components/ComboBoxString.h b/Engine/Source/Runtime/UMG/Public/Components/ComboBoxString.h
--- a/Engine/Source/Runtime/UMG/Public/Components/ComboBoxString.h
+++ b/Engine/Source/Runtime/UMG/Public/Components/ComboBoxString.h
@@ -147,10 +147,10 @@ protected:
 	/** Called by slate when it needs to generate a new item for the combobox */
 	virtual TSharedRef<SWidget> HandleGenerateWidget(TSharedPtr<FString> Item) const;
 
-	/** Called by slate when the underlying comobobox selection changes */
+	/** Called by slate when the underlying combobox selection changes */
 	virtual void HandleSelectionChanged(TSharedPtr<FString> Item, ESelectInfo::Type SelectionType);
 
-	/** Called by slate when the underlying comobobox is opening */
+	/** Called by slate when the underlying combobox is opening */
 	virtual void HandleOpening();
 
 	//~ Begin UWidget Interface
@@ -164,8 +164,8 @@ protected:
 	/** A shared pointer to the underlying slate combobox */
 	TSharedPtr< SComboBox< TSharedPtr<FString> > > MyComboBox;
 
-	/** A shared pointer to a container that holds the comobobox content that is selected */
-	TSharedPtr< SBox > ComoboBoxContent;
+	/** A shared pointer to a container that holds the combobox content that is selected */
+	TSharedPtr< SBox > ComboBoxContent;
 
 	/** A shared pointer to the current selected string */
 	TSharedPtr<FString> CurrentOptionPtr;
diff --git a/Engine/Source/Runtime/UMG/Public/Components/Widget.h b/Engine/Source/Runtime/UMG/Public/Components/Widget.h
--- a/Engine/Source/Runtime/UMG/Public/Components/Widget.h
+++ b/Engine/Source/Runtime/UMG/Public/Components/Widget.h
@@ -118,9 +118,10 @@ namespace EWidgetDesignFlags
 {
 	enum Type
 	{
-		None		= 0,
-		Designing	= 1,
-		ShowOutline	= 2,
+		None				= 0,
+		Designing			= 1,
+		ShowOutline			= 2,
+		ExecutePreConstruct	= 4
 	};
 }
 
diff --git a/Engine/Source/Runtime/UMG/Public/Components/WidgetInteractionComponent.h b/Engine/Source/Runtime/UMG/Public/Components/WidgetInteractionComponent.h
--- a/Engine/Source/Runtime/UMG/Public/Components/WidgetInteractionComponent.h
+++ b/Engine/Source/Runtime/UMG/Public/Components/WidgetInteractionComponent.h
@@ -273,7 +273,7 @@ protected:
 	};
 	/** Performs the trace and gets the hit result under the specified InteractionSource */
 	virtual FWidgetTraceResult PerformTrace() const;
-	
+
 	/**
 	 * Gets the list of components to ignore during hit testing.  Which is everything that is a parent/sibling of this 
 	 * component that's not a Widget Component.  This is so traces don't get blocked by capsules and such around the player.
@@ -286,7 +286,7 @@ protected:
 protected:
 
 	/** The last widget path under the hit result. */
-	FWeakWidgetPath LastWigetPath;
+	FWeakWidgetPath LastWidgetPath;
 
 	/** The modifier keys to simulate during key presses. */
 	FModifierKeysState ModifierKeys;
@@ -326,6 +326,11 @@ protected:
 	UPROPERTY(Transient)
 	bool bIsHoveredWidgetHitTestVisible;
 
+private:
+
+	/** Returns the path to the widget that is currently beneath the pointer */
+	FWidgetPath DetermineWidgetUnderPointer();
+
 private_subobject:
 #if WITH_EDITORONLY_DATA
 
diff --git a/Engine/Source/Runtime/WebBrowser/Private/CEF/CEFWebBrowserWindow.cpp b/Engine/Source/Runtime/WebBrowser/Private/CEF/CEFWebBrowserWindow.cpp
--- a/Engine/Source/Runtime/WebBrowser/Private/CEF/CEFWebBrowserWindow.cpp
+++ b/Engine/Source/Runtime/WebBrowser/Private/CEF/CEFWebBrowserWindow.cpp
@@ -1518,7 +1518,7 @@ void FCEFWebBrowserWindow::OnCursorChange(CefCursorHandle CefCursor, CefRenderHa
 
 				if (PlatformCursor.IsValid())
 				{
-					PlatformCursor->SetCustomShape((void*)CefCursor);
+					PlatformCursor->SetTypeShape(EMouseCursor::Custom, (void*)CefCursor);
 					Cursor = EMouseCursor::Custom;
 				}
 			}
diff --git a/Engine/Source/Runtime/XmlParser/Private/XmlFile.cpp b/Engine/Source/Runtime/XmlParser/Private/XmlFile.cpp
--- a/Engine/Source/Runtime/XmlParser/Private/XmlFile.cpp
+++ b/Engine/Source/Runtime/XmlParser/Private/XmlFile.cpp
@@ -679,8 +679,10 @@ FXmlNode* FXmlFile::CreateNodeRecursive(const TArray<FString>& Tokens, int32 Sta
 					FXmlNode* NewNode = new FXmlNode();
 					NewNode->Tag = Tag;
 					NewNode->Attributes = Attributes;
-					check(OutNextIndex != nullptr);
-					*OutNextIndex = i + 1;
+					if (OutNextIndex != nullptr)
+					{
+						*OutNextIndex = i + 1;
+					}
 					return NewNode;
 				}
 
diff --git a/Engine/Source/Runtime/XmlParser/Public/XmlFile.h b/Engine/Source/Runtime/XmlParser/Public/XmlFile.h
--- a/Engine/Source/Runtime/XmlParser/Public/XmlFile.h
+++ b/Engine/Source/Runtime/XmlParser/Public/XmlFile.h
@@ -2,9 +2,6 @@
 
 #pragma once
 
-#ifndef __XmlFile_h__
-#define __XmlFile_h__
-
 #include "CoreMinimal.h"
 #include "XmlNode.h"
 
@@ -118,5 +115,3 @@ private:
 	/** Flag for whether the node creation process failed or not */
 	bool bCreationFailed;
 };
-
-#endif
diff --git a/Engine/Source/Runtime/XmlParser/Public/XmlNode.h b/Engine/Source/Runtime/XmlParser/Public/XmlNode.h
--- a/Engine/Source/Runtime/XmlParser/Public/XmlNode.h
+++ b/Engine/Source/Runtime/XmlParser/Public/XmlNode.h
@@ -2,9 +2,6 @@
 
 #pragma once
 
-#ifndef __XmlNode_h__
-#define __XmlNode_h__
-
 #include "CoreMinimal.h"
 
 /** Forward Decl for implementation */
@@ -99,5 +96,3 @@ private:
 	FXmlNode* NextNode;
 
 };
-
-#endif
diff --git a/Engine/Source/Runtime/XmlParser/Public/XmlParser.h b/Engine/Source/Runtime/XmlParser/Public/XmlParser.h
--- a/Engine/Source/Runtime/XmlParser/Public/XmlParser.h
+++ b/Engine/Source/Runtime/XmlParser/Public/XmlParser.h
@@ -2,10 +2,5 @@
 
 #pragma once
 
-#ifndef __XmlParser_h__
-#define __XmlParser_h__
-
 #include "XmlNode.h"
 #include "XmlFile.h"
-
-#endif
