Fix Dev-Editor build fails

diff --git a/Engine/Source/Runtime/SlateCore/Public/Rendering/RenderingCommon.h b/Engine/Source/Runtime/SlateCore/Public/Rendering/RenderingCommon.h
--- a/Engine/Source/Runtime/SlateCore/Public/Rendering/RenderingCommon.h
+++ b/Engine/Source/Runtime/SlateCore/Public/Rendering/RenderingCommon.h
@@ -229,41 +229,55 @@ struct SLATECORE_API FSlateVertex
 public:
 	
 	// These constructors have more or less been deprecated, you should use FSlateVertex::Make<ESlateVertexRounding::Enabled>(...) in order to have fine control over pixel snapping or not.
-
+	DEPRECATED(4.16, "FSlateVertex constructors have been deprecated, you should use FSlateVertex::Make" )
 	FSlateVertex(const FSlateRenderTransform& RenderTransform, const FVector2D& InLocalPosition, const FVector2D& InTexCoord, const FVector2D& InTexCoord2, const FColor& InColor, const FSlateRotatedClipRectType& InClipRect)
-		: TexCoords { InTexCoord.X, InTexCoord.Y, InTexCoord2.X, InTexCoord2.Y }
-		, Position(TransformPoint(RenderTransform, InLocalPosition))
+		: Position(TransformPoint(RenderTransform, InLocalPosition))
 		, ClipRect(InClipRect)
 		, Color(InColor)
 	{
 		Position.X = FMath::RoundToInt(Position.X);
 		Position.Y = FMath::RoundToInt(Position.Y);
+
+		TexCoords[0] = InTexCoord.X;
+		TexCoords[1] = InTexCoord.Y;
+		TexCoords[2] = InTexCoord2.X;
+		TexCoords[3] = InTexCoord2.Y;
 	}
 
+	DEPRECATED(4.16, "FSlateVertex constructors have been deprecated, you should use FSlateVertex::Make")
 	FSlateVertex(const FSlateRenderTransform& RenderTransform, const FVector2D& InLocalPosition, const FVector2D& InTexCoord, const FColor& InColor, const FSlateRotatedClipRectType& InClipRect)
-		: TexCoords { InTexCoord.X, InTexCoord.Y, 1.0f, 1.0f }
-		, Position(TransformPoint(RenderTransform, InLocalPosition))
+		: Position(TransformPoint(RenderTransform, InLocalPosition))
 		, ClipRect(InClipRect)
 		, Color(InColor)
 	{
 		Position.X = FMath::RoundToInt(Position.X);
 		Position.Y = FMath::RoundToInt(Position.Y);
+
+		TexCoords[0] = InTexCoord.X;
+		TexCoords[1] = InTexCoord.Y;
+		TexCoords[2] = 1.0f;
+		TexCoords[3] = 1.0f;
 	}
 
+	DEPRECATED(4.16, "FSlateVertex constructors have been deprecated, you should use FSlateVertex::Make")
 	FSlateVertex(const FSlateRenderTransform& RenderTransform, const FVector2D& InLocalPosition, const FVector4& InTexCoords, const FVector2D& InMaterialTexCoords, const FColor& InColor, const FSlateRotatedClipRectType& InClipRect)
-		: TexCoords { InTexCoords.X, InTexCoords.Y, InTexCoords.Z, InTexCoords.W }
-		, MaterialTexCoords(InMaterialTexCoords)
+		: MaterialTexCoords(InMaterialTexCoords)
 		, Position(TransformPoint(RenderTransform, InLocalPosition))
 		, ClipRect(InClipRect)
 		, Color(InColor)
 	{
 		Position.X = FMath::RoundToInt(Position.X);
 		Position.Y = FMath::RoundToInt(Position.Y);
+
+		TexCoords[0] = InTexCoords.X;
+		TexCoords[1] = InTexCoords.Y;
+		TexCoords[2] = InTexCoords.Z;
+		TexCoords[3] = InTexCoords.W;
 	}
 
+	DEPRECATED(4.16, "FSlateVertex constructors have been deprecated, you should use FSlateVertex::Make")
 	FSlateVertex(const FSlateRenderTransform& RenderTransform, const FVector2D& InLocalPosition, const FVector2D& InLocalSize, float Scale, const FVector4& InTexCoords, const FColor& InColor, const FSlateRotatedClipRectType& InClipRect)
-		: TexCoords { InTexCoords.X, InTexCoords.Y, InTexCoords.Z, InTexCoords.W }
-		, MaterialTexCoords(InLocalPosition.X / InLocalSize.X, InLocalPosition.Y / InLocalSize.Y)
+		: MaterialTexCoords(InLocalPosition.X / InLocalSize.X, InLocalPosition.Y / InLocalSize.Y)
 		, Position(TransformPoint(RenderTransform, InLocalPosition))
 		, ClipRect(InClipRect)
 		, Color(InColor)
@@ -273,6 +287,11 @@ public:
 
 		PixelSize[0] = FMath::RoundToInt(InLocalSize.X * Scale);
 		PixelSize[1] = FMath::RoundToInt(InLocalSize.Y * Scale);
+
+		TexCoords[0] = InTexCoords.X;
+		TexCoords[1] = InTexCoords.Y;
+		TexCoords[2] = InTexCoords.Z;
+		TexCoords[3] = InTexCoords.W;
 	}
 
 public:
