Added an ensure for condition that Orion is running into (the cells array is empty while NumCells is not zero!)

#tests ran several cycled matches of orion on PS4 Development

diff --git a/Engine/Source/Runtime/SlateCore/Private/Input/HittestGrid.cpp b/Engine/Source/Runtime/SlateCore/Private/Input/HittestGrid.cpp
--- a/Engine/Source/Runtime/SlateCore/Private/Input/HittestGrid.cpp
+++ b/Engine/Source/Runtime/SlateCore/Private/Input/HittestGrid.cpp
@@ -309,7 +309,7 @@ int32 FHittestGrid::InsertWidget(const int32 ParentHittestIndex, const EVisibili
 			(*WidgetsCachedThisFrame)[ParentHittestIndex].AddChild(WidgetIndex);
 		}
 
-		if (Visibility.IsHitTestVisible())
+		if (ensure(Cells.Num() > 0) && Visibility.IsHitTestVisible())
 		{
 			// Mark any cell that is overlapped by this widget.
 
