Backing out ensure-don't-crash change due to actual fix (3572777) being submitted

#tests none

diff --git a/Engine/Source/Runtime/SlateCore/Private/Input/HittestGrid.cpp b/Engine/Source/Runtime/SlateCore/Private/Input/HittestGrid.cpp
--- a/Engine/Source/Runtime/SlateCore/Private/Input/HittestGrid.cpp
+++ b/Engine/Source/Runtime/SlateCore/Private/Input/HittestGrid.cpp
@@ -309,7 +309,7 @@ int32 FHittestGrid::InsertWidget(const int32 ParentHittestIndex, const EVisibili
 			(*WidgetsCachedThisFrame)[ParentHittestIndex].AddChild(WidgetIndex);
 		}
 
-		if (ensure(Cells.Num() > 0) && Visibility.IsHitTestVisible())
+		if (Visibility.IsHitTestVisible())
 		{
 			// Mark any cell that is overlapped by this widget.
 
