Copying WEX-Staging @ (WEX/Main @ 3740665) to //UE4/Main

diff --git a/Engine/Source/Runtime/Renderer/Private/DrawingPolicy.cpp b/Engine/Source/Runtime/Renderer/Private/DrawingPolicy.cpp
--- a/Engine/Source/Runtime/Renderer/Private/DrawingPolicy.cpp
+++ b/Engine/Source/Runtime/Renderer/Private/DrawingPolicy.cpp
@@ -66,9 +66,9 @@ void FMeshDrawingPolicy::OnlyApplyDitheredLODTransitionState(FDrawingPolicyRende
 	}
 }
 
-void FMeshDrawingPolicy::SetInstanceParameters(FRHICommandList& RHICmdList, uint32 InInstanceOffset, uint32 InInstanceCount) const
+void FMeshDrawingPolicy::SetInstanceParameters(FRHICommandList& RHICmdList, uint32 InVertexOffset, uint32 InInstanceOffset, uint32 InInstanceCount) const
 {
-	BaseVertexShader->SetInstanceParameters(RHICmdList, InInstanceOffset, InInstanceCount);
+	BaseVertexShader->SetInstanceParameters(RHICmdList, InVertexOffset, InInstanceOffset, InInstanceCount);
 }
 
 void FMeshDrawingPolicy::DrawMesh(FRHICommandList& RHICmdList, const FMeshBatch& Mesh, int32 BatchElementIndex, const bool bIsInstancedStereo) const
@@ -96,7 +96,7 @@ void FMeshDrawingPolicy::DrawMesh(FRHICommandList& RHICmdList, const FMeshBatch&
 					for (uint32 Run = 0; Run < BatchElement.NumInstances; Run++)
 					{
 						const uint32 InstanceCount = (1 + BatchElement.InstanceRuns[Run * 2 + 1] - BatchElement.InstanceRuns[Run * 2]);
-						SetInstanceParameters(RHICmdList, 0, InstanceCount);
+						SetInstanceParameters(RHICmdList, BatchElement.BaseVertexIndex, 0, InstanceCount);
 						GetVertexFactory()->OffsetPositionInstanceStreams(RHICmdList, BatchElement.InstanceRuns[Run * 2]);
 
 						RHICmdList.DrawIndexedPrimitive(
@@ -116,7 +116,7 @@ void FMeshDrawingPolicy::DrawMesh(FRHICommandList& RHICmdList, const FMeshBatch&
 					for (uint32 Run = 0; Run < BatchElement.NumInstances; Run++)
 					{
 						uint32 InstanceCount = (1 + BatchElement.InstanceRuns[Run * 2 + 1] - BatchElement.InstanceRuns[Run * 2]);
-						SetInstanceParameters(RHICmdList, 0, InstanceCount);
+						SetInstanceParameters(RHICmdList, BatchElement.BaseVertexIndex, 0, InstanceCount);
 						GetVertexFactory()->OffsetInstanceStreams(RHICmdList, BatchElement.InstanceRuns[Run * 2]);
 
 						RHICmdList.DrawIndexedPrimitive(
@@ -138,7 +138,7 @@ void FMeshDrawingPolicy::DrawMesh(FRHICommandList& RHICmdList, const FMeshBatch&
 				{
 					const uint32 InstanceOffset = BatchElement.InstanceRuns[Run * 2];
 					const uint32 InstanceCount = (1 + BatchElement.InstanceRuns[Run * 2 + 1] - BatchElement.InstanceRuns[Run * 2]);
-					SetInstanceParameters(RHICmdList, InstanceOffset, InstanceCount);
+					SetInstanceParameters(RHICmdList, BatchElement.BaseVertexIndex, InstanceOffset, InstanceCount);
 
 					RHICmdList.DrawIndexedPrimitive(
 						BatchElement.IndexBuffer->IndexBufferRHI,
@@ -168,7 +168,7 @@ void FMeshDrawingPolicy::DrawMesh(FRHICommandList& RHICmdList, const FMeshBatch&
 			{
 				// Currently only supporting this path for instanced stereo.
 				const uint32 InstanceCount = ((bIsInstancedStereo && !BatchElement.bIsInstancedMesh) ? 2 : BatchElement.NumInstances);
-				SetInstanceParameters(RHICmdList, 0, InstanceCount);
+				SetInstanceParameters(RHICmdList, BatchElement.BaseVertexIndex, 0, InstanceCount);
 
 				RHICmdList.DrawIndexedPrimitive(
 					BatchElement.IndexBuffer->IndexBufferRHI,
@@ -185,7 +185,7 @@ void FMeshDrawingPolicy::DrawMesh(FRHICommandList& RHICmdList, const FMeshBatch&
 	}
 	else
 	{
-		SetInstanceParameters(RHICmdList, 0, 1);
+		SetInstanceParameters(RHICmdList, BatchElement.BaseVertexIndex + BatchElement.FirstIndex, 0, 1);
 
 		RHICmdList.DrawPrimitive(
 			Mesh.Type,
diff --git a/Engine/Source/Runtime/Renderer/Private/ShaderBaseClasses.cpp b/Engine/Source/Runtime/Renderer/Private/ShaderBaseClasses.cpp
--- a/Engine/Source/Runtime/Renderer/Private/ShaderBaseClasses.cpp
+++ b/Engine/Source/Runtime/Renderer/Private/ShaderBaseClasses.cpp
@@ -76,6 +76,7 @@ FMaterialShader::FMaterialShader(const FMaterialShaderType::CompiledShaderInitia
 
 	InstanceCount.Bind(Initializer.ParameterMap, TEXT("InstanceCount"));
 	InstanceOffset.Bind(Initializer.ParameterMap, TEXT("InstanceOffset"));
+	VertexOffset.Bind(Initializer.ParameterMap, TEXT("VertexOffset"));
 }
 
 FUniformBufferRHIParamRef FMaterialShader::GetParameterCollectionBuffer(const FGuid& Id, const FSceneInterface* SceneInterface) const
@@ -479,6 +480,7 @@ bool FMaterialShader::Serialize(FArchive& Ar)
 
 	Ar << InstanceCount;
 	Ar << InstanceOffset;
+	Ar << VertexOffset;
 
 	return bShaderHasOutdatedParameters;
 }
diff --git a/Engine/Source/Runtime/Renderer/Public/DrawingPolicy.h b/Engine/Source/Runtime/Renderer/Public/DrawingPolicy.h
--- a/Engine/Source/Runtime/Renderer/Public/DrawingPolicy.h
+++ b/Engine/Source/Runtime/Renderer/Public/DrawingPolicy.h
@@ -468,5 +468,5 @@ private:
 		return InstanceFactor;
 	}
 
-	void SetInstanceParameters(FRHICommandList& RHICmdList, uint32 InInstanceOffset, uint32 InInstanceCount) const;
+	void SetInstanceParameters(FRHICommandList& RHICmdList, uint32 InVertexOffset, uint32 InInstanceOffset, uint32 InInstanceCount) const;
 };
diff --git a/Engine/Source/Runtime/Renderer/Public/MaterialShader.h b/Engine/Source/Runtime/Renderer/Public/MaterialShader.h
--- a/Engine/Source/Runtime/Renderer/Public/MaterialShader.h
+++ b/Engine/Source/Runtime/Renderer/Public/MaterialShader.h
@@ -162,8 +162,9 @@ public:
 	virtual bool Serialize(FArchive& Ar) override;
 	virtual uint32 GetAllocatedSize() const override;
 
-	void SetInstanceParameters(FRHICommandList& RHICmdList, uint32 InInstanceOffset, uint32 InInstanceCount) const
+	void SetInstanceParameters(FRHICommandList& RHICmdList, uint32 InVertexOffset, uint32 InInstanceOffset, uint32 InInstanceCount) const
 	{
+		SetShaderValue(RHICmdList, GetVertexShader(), VertexOffset, InVertexOffset);
 		SetShaderValue(RHICmdList, GetVertexShader(), InstanceCount, InInstanceCount);
 		SetShaderValue(RHICmdList, GetVertexShader(), InstanceOffset, InInstanceOffset);
 	}
@@ -185,6 +186,7 @@ private:
 
 	FShaderParameter InstanceCount;
 	FShaderParameter InstanceOffset;
+	FShaderParameter VertexOffset;
 
 	FDebugUniformExpressionSet	DebugUniformExpressionSet;
 	FRHIUniformBufferLayout		DebugUniformExpressionUBLayout;
diff --git a/Engine/Source/Runtime/Slate/Private/Widgets/Layout/SScaleBox.cpp b/Engine/Source/Runtime/Slate/Private/Widgets/Layout/SScaleBox.cpp
--- a/Engine/Source/Runtime/Slate/Private/Widgets/Layout/SScaleBox.cpp
+++ b/Engine/Source/Runtime/Slate/Private/Widgets/Layout/SScaleBox.cpp
@@ -115,7 +115,15 @@ void SScaleBox::OnArrangeChildren( const FGeometry& AllottedGeometry, FArrangedC
 					break;
 				}
 
-				LastFinalScale = FinalScale;
+				// Force full layout calculations when the previously calculated final scale is zero
+				if (!FMath::IsNearlyZero(FinalScale) || !bSingleLayoutPass)
+				{
+					LastFinalScale = FinalScale;
+				}
+				else
+				{
+					LastFinalScale.Reset();
+				}
 			}
 			else
 			{
diff --git a/Engine/Source/Runtime/Slate/Private/Widgets/SInvalidationPanel.cpp b/Engine/Source/Runtime/Slate/Private/Widgets/SInvalidationPanel.cpp
--- a/Engine/Source/Runtime/Slate/Private/Widgets/SInvalidationPanel.cpp
+++ b/Engine/Source/Runtime/Slate/Private/Widgets/SInvalidationPanel.cpp
@@ -85,6 +85,8 @@ void SInvalidationPanel::Construct( const FArguments& InArgs )
 	RootCacheNode = nullptr;
 	LastUsedCachedNodeIndex = 0;
 	LastHitTestIndex = 0;
+	LastLayerId = 0;
+	LastClippingIntersectionSize = FVector2D::ZeroVector;
 
 	bCacheRelativeTransforms = InArgs._CacheRelativeTransforms;
 
@@ -118,8 +120,13 @@ bool SInvalidationPanel::IsCachingNeeded() const
 }
 #endif
 
-bool SInvalidationPanel::IsCachingNeeded(FSlateWindowElementList& OutDrawElements, const FGeometry& AllottedGeometry, const FSlateRect& MyCullingRect) const
+bool SInvalidationPanel::IsCachingNeeded(FSlateWindowElementList& OutDrawElements, const FGeometry& AllottedGeometry, const FSlateRect& MyCullingRect, int32 LayerId) const
 {
+	if (LastLayerId != LayerId)
+	{
+		return true;
+	}
+	
 	if (bCacheRelativeTransforms)
 	{
 		// If the container we're in has changed in either scale or the rotation matrix has changed, 
@@ -167,7 +174,17 @@ bool SInvalidationPanel::IsCachingNeeded(FSlateWindowElementList& OutDrawElement
 	{
 		return true;
 	}
-	
+
+	if (bCacheRelativeTransforms)
+	{
+		bool bOverlapping;
+		FVector2D IntersectionSize = AllottedGeometry.GetLayoutBoundingRect().IntersectionWith(MyCullingRect, bOverlapping).GetSize();
+		if (!LastClippingIntersectionSize.Equals(IntersectionSize, 1.0f))
+		{
+			return true;
+		}
+	}
+		
 	return false;
 }
 
@@ -267,7 +284,7 @@ int32 SInvalidationPanel::OnPaint( const FPaintArgs& Args, const FGeometry& Allo
 
 		//FPlatformMisc::BeginNamedEvent(FColor::Magenta, "Slate::InvalidationPanel::Paint");
 
-		const bool bWasCachingNeeded = IsCachingNeeded() || IsCachingNeeded(OutDrawElements, AllottedGeometry, MyCullingRect);
+		const bool bWasCachingNeeded = IsCachingNeeded() || IsCachingNeeded(OutDrawElements, AllottedGeometry, MyCullingRect, LayerId);
 
 		if ( bWasCachingNeeded )
 		{
@@ -354,6 +371,12 @@ int32 SInvalidationPanel::OnPaint( const FPaintArgs& Args, const FGeometry& Allo
 
 			LastAllottedGeometry = AllottedGeometry;
 			LastClipRectSize = MyCullingRect.GetSize().RoundToVector();
+			LastLayerId = LayerId;
+
+			if (bCacheRelativeTransforms)
+			{
+				LastClippingIntersectionSize = AllottedGeometry.GetLayoutBoundingRect().IntersectionWith(MyCullingRect).GetSize();
+			}
 
 			bIsInvalidating = false;
 
@@ -406,8 +429,17 @@ int32 SInvalidationPanel::OnPaint( const FPaintArgs& Args, const FGeometry& Allo
 
 			const TArray<TSharedPtr<FSlateWindowElementList::FVolatilePaint>>& VolatileElements = CachedWindowElements->GetVolatileElements();
 			INC_DWORD_STAT_BY(STAT_SlateNumVolatileWidgets, VolatileElements.Num());
-
-			int32 VolatileLayerId = CachedWindowElements->PaintVolatile(OutDrawElements, Args.GetCurrentTime(), Args.GetDeltaTime(), AbsoluteDeltaPosition * AllottedGeometry.Scale);
+			
+			int32 VolatileLayerId = 0;
+			if (bCacheRenderData)
+			{
+				VolatileLayerId = CachedWindowElements->PaintVolatile(OutDrawElements, Args.GetCurrentTime(), Args.GetDeltaTime(), AbsoluteDeltaPosition);
+			}
+			else
+			{
+				VolatileLayerId = CachedWindowElements->PaintVolatileRootLayer(OutDrawElements, Args.GetCurrentTime(), Args.GetDeltaTime(), AbsoluteDeltaPosition);
+			}
+			
 			OutMaxChildLayer = FMath::Max(OutMaxChildLayer, VolatileLayerId);
 
 			//FPlatformMisc::EndNamedEvent();
diff --git a/Engine/Source/Runtime/Slate/Private/Widgets/Views/STableViewBase.cpp b/Engine/Source/Runtime/Slate/Private/Widgets/Views/STableViewBase.cpp
--- a/Engine/Source/Runtime/Slate/Private/Widgets/Views/STableViewBase.cpp
+++ b/Engine/Source/Runtime/Slate/Private/Widgets/Views/STableViewBase.cpp
@@ -683,10 +683,12 @@ float STableViewBase::ScrollBy(const FGeometry& MyGeometry, float ScrollByAmount
 	return ScrollTo( ClampedScrollOffsetInItems );
 }
 
-float STableViewBase::ScrollTo( float InScrollOffset )
+float STableViewBase::ScrollTo( float InScrollOffset)
 {
 	const float NewScrollOffset = FMath::Clamp( InScrollOffset, -10.0f, GetNumItemsBeingObserved()+10.0f );
 	float AmountScrolled = FMath::Abs( ScrollOffset - NewScrollOffset );
+
+	EndInertialScrolling();
 	SetScrollOffset( NewScrollOffset );
 	
 	if ( bWasAtEndOfList && NewScrollOffset >= ScrollOffset )
@@ -712,6 +714,11 @@ void STableViewBase::SetScrollOffset( const float InScrollOffset )
 	}
 }
 
+void STableViewBase::EndInertialScrolling()
+{
+	InertialScrollManager.ClearScrollVelocity();
+}
+
 void STableViewBase::AddScrollOffset(const float InScrollOffsetDelta, bool RefreshList)
 {
 	if (FMath::IsNearlyEqual(InScrollOffsetDelta, 0.0f) == false)
@@ -828,13 +835,15 @@ float STableViewBase::GetScrollRateInItems() const
 
 void STableViewBase::ScrollToTop()
 {
-	SetScrollOffset( 0 );
+	EndInertialScrolling();
+	SetScrollOffset(0);
 	RequestListRefresh();
 }
 
 void STableViewBase::ScrollToBottom()
 {
-	SetScrollOffset( GetNumItemsBeingObserved() );
+	EndInertialScrolling();
+	SetScrollOffset(GetNumItemsBeingObserved());
 	RequestListRefresh();
 }
 
diff --git a/Engine/Source/Runtime/Slate/Public/Framework/Text/GenericPlatformTextField.h b/Engine/Source/Runtime/Slate/Public/Framework/Text/GenericPlatformTextField.h
--- a/Engine/Source/Runtime/Slate/Public/Framework/Text/GenericPlatformTextField.h
+++ b/Engine/Source/Runtime/Slate/Public/Framework/Text/GenericPlatformTextField.h
@@ -10,7 +10,7 @@ class FGenericPlatformTextField : public IPlatformTextField
 {
 public:
 	virtual void ShowVirtualKeyboard(bool bShow, int32 UserIndex, TSharedPtr<IVirtualKeyboardEntry> TextEntryWidget) override {};
-
+	virtual bool AllowMoveCursor() override { return true; };
 };
 
 typedef FGenericPlatformTextField FPlatformTextField;
diff --git a/Engine/Source/Runtime/Slate/Public/Framework/Text/IOS/IOSPlatformTextField.h b/Engine/Source/Runtime/Slate/Public/Framework/Text/IOS/IOSPlatformTextField.h
--- a/Engine/Source/Runtime/Slate/Public/Framework/Text/IOS/IOSPlatformTextField.h
+++ b/Engine/Source/Runtime/Slate/Public/Framework/Text/IOS/IOSPlatformTextField.h
@@ -19,6 +19,7 @@ public:
 	virtual ~FIOSPlatformTextField();
 
 	virtual void ShowVirtualKeyboard(bool bShow, int32 UserIndex, TSharedPtr<IVirtualKeyboardEntry> TextEntryWidget) override;
+	virtual bool AllowMoveCursor() override { return true; };
 
 private:
 	SlateTextField* TextField;
diff --git a/Engine/Source/Runtime/Slate/Public/Widgets/SInvalidationPanel.h b/Engine/Source/Runtime/Slate/Public/Widgets/SInvalidationPanel.h
--- a/Engine/Source/Runtime/Slate/Public/Widgets/SInvalidationPanel.h
+++ b/Engine/Source/Runtime/Slate/Public/Widgets/SInvalidationPanel.h
@@ -85,7 +85,7 @@ private:
 #endif
 	
 
-	bool IsCachingNeeded(FSlateWindowElementList& OutDrawElements, const FGeometry& AllottedGeometry, const FSlateRect& MyCullingRect) const;
+	bool IsCachingNeeded(FSlateWindowElementList& OutDrawElements, const FGeometry& AllottedGeometry, const FSlateRect& MyCullingRect, int32 LayerId) const;
 
 private:
 	mutable FGeometry LastAllottedGeometry;
@@ -109,9 +109,11 @@ private:
 	mutable int32 LastUsedCachedNodeIndex;
 	mutable int32 LastHitTestIndex;
 	mutable FVector2D LastClipRectSize;
+	mutable FVector2D LastClippingIntersectionSize;
 	mutable int32 LastClippingIndex;
 	mutable int32 LastClippingStateOffset;
 	mutable TOptional<FSlateClippingState> LastClippingState;
+	mutable int32 LastLayerId;
 
 	mutable int32 CachedMaxChildLayer;
 	mutable bool bNeedsCaching;
diff --git a/Engine/Source/Runtime/Slate/Public/Widgets/Views/SListView.h b/Engine/Source/Runtime/Slate/Public/Widgets/Views/SListView.h
--- a/Engine/Source/Runtime/Slate/Public/Widgets/Views/SListView.h
+++ b/Engine/Source/Runtime/Slate/Public/Widgets/Views/SListView.h
@@ -1307,6 +1307,8 @@ protected:
 						return EScrollIntoViewResult::Deferred;
 					}
 				}
+
+				EndInertialScrolling();
 				
 				// Only scroll the item into view if it's not already in the visible range
 				const double IndexPlusOne = IndexOfItem+1;
diff --git a/Engine/Source/Runtime/Slate/Public/Widgets/Views/STableViewBase.h b/Engine/Source/Runtime/Slate/Public/Widgets/Views/STableViewBase.h
--- a/Engine/Source/Runtime/Slate/Public/Widgets/Views/STableViewBase.h
+++ b/Engine/Source/Runtime/Slate/Public/Widgets/Views/STableViewBase.h
@@ -115,6 +115,9 @@ public:
 	/** Set the scroll offset of this view (in items) */
 	void SetScrollOffset( const float InScrollOffset );
 
+	/** Reset the inertial scroll velocity accumulated in the InertialScrollManager */
+	void EndInertialScrolling();
+
 	/** Add the scroll offset of this view (in items) */
 	void AddScrollOffset(const float InScrollOffsetDelta, bool RefreshList = false);
 
@@ -164,13 +167,13 @@ protected:
 	virtual float ScrollBy( const FGeometry& MyGeometry, float ScrollByAmount, EAllowOverscroll InAllowOverscroll );
 
 	/**
-	 * Scroll the view to an offset
+	 * Scroll the view to an offset and resets the inertial scroll velocity 
 	 *
-	 * @param InScrollOffset  Offset into the total list length to scroll down.
+	 * @param InScrollOffset       Offset into the total list length to scroll down.
 	 *
 	 * @return The amount actually scrolled
 	 */
-	virtual float ScrollTo( float InScrollOffset );
+	virtual float ScrollTo( float InScrollOffset);
 
 	/** Insert WidgetToInsert at the top of the view. */
 	void InsertWidget( const TSharedRef<ITableRow> & WidgetToInset );
diff --git a/Engine/Source/Runtime/Slate/Public/Widgets/Views/STileView.h b/Engine/Source/Runtime/Slate/Public/Widgets/Views/STileView.h
--- a/Engine/Source/Runtime/Slate/Public/Widgets/Views/STileView.h
+++ b/Engine/Source/Runtime/Slate/Public/Widgets/Views/STileView.h
@@ -420,6 +420,8 @@ protected:
 					}
 				}
 
+				this->EndInertialScrolling();
+
 				// Only scroll the item into view if it's not already in the visible range
 				const int32 NumItemsWide = GetNumItemsWide();
 				const float Index = (IndexOfItem / NumItemsWide) * NumItemsWide;
diff --git a/Engine/Source/Runtime/SlateCore/Private/Rendering/DrawElements.cpp b/Engine/Source/Runtime/SlateCore/Private/Rendering/DrawElements.cpp
--- a/Engine/Source/Runtime/SlateCore/Private/Rendering/DrawElements.cpp
+++ b/Engine/Source/Runtime/SlateCore/Private/Rendering/DrawElements.cpp
@@ -516,7 +516,6 @@ FVector2D FSlateDrawElement::GetRotationPoint(const FPaintGeometry& PaintGeometr
 void FSlateBatchData::Reset()
 {
 	RenderBatches.Reset();
-	DynamicOffset = FVector2D(0, 0);
 	
 	// note: LayerToElementBatches is not reset here as the same layers are 
 	// more than likely reused and we can save memory allocations by not resetting the map every frame
@@ -673,7 +672,7 @@ void FSlateBatchData::AddRenderBatch(uint32 InLayer, const FSlateElementBatch& I
 	NumBatchedIndices += InNumIndices;
 
 	const int32 Index = RenderBatches.Add(FSlateRenderBatch(InLayer, InElementBatch, RenderDataHandle, InNumVertices, InNumIndices, InVertexOffset, InIndexOffset));
-	RenderBatches[Index].DynamicOffset = DynamicOffset;
+	RenderBatches[Index].DynamicOffset = FVector2D::ZeroVector;
 }
 
 void FSlateBatchData::ResetVertexArray(FSlateVertexArray& InOutVertexArray)
@@ -714,8 +713,6 @@ void FSlateBatchData::Merge(FElementBatchMap& InLayerToElementBatches, uint32& V
 				{
 					if ( FSlateRenderDataHandle* RenderHandle = ElementBatch.GetCachedRenderHandle().Get() )
 					{
-						DynamicOffset += ElementBatch.GetCachedRenderDataOffset();
-
 						TArray<FSlateRenderBatch>* ForeignBatches = RenderHandle->GetRenderBatches();
 						//TArray<FSlateClippingState>* ForeignClipState = RenderHandle->GetClipStates();
 						if (ForeignBatches /*&& ForeignClipState*/)
@@ -739,14 +736,12 @@ void FSlateBatchData::Merge(FElementBatchMap& InLayerToElementBatches, uint32& V
 								else
 								{
 									const int32 Index = RenderBatches.Add(ForeignBatchesRef[i]);
-									RenderBatches[Index].DynamicOffset = DynamicOffset;
+									RenderBatches[Index].DynamicOffset = ElementBatch.GetCachedRenderDataOffset();
 									//RenderBatches[Index].ClippingIndex = ElementBatch.GetClippingIndex();
 								}
 							}
 						}
 
-						DynamicOffset -= ElementBatch.GetCachedRenderDataOffset();
-
 						continue;
 					}
 				}
@@ -797,25 +792,22 @@ void FSlateBatchData::Merge(FElementBatchMap& InLayerToElementBatches, uint32& V
 void FSlateWindowElementList::MergeElementList(FSlateWindowElementList* ElementList, FVector2D AbsoluteOffset)
 {
 	const bool bMoved = !AbsoluteOffset.IsZero();
-
-	const TArray< FSlateClippingState >& States = ElementList->ClippingManager.GetClippingStates();
-	const int32 ClippingStateOffset = ClippingManager.MergeClippingStates(States);
-
 	const TArray<FSlateDrawElement>& CachedElements = ElementList->GetDrawElements();
 	const int32 CachedElementCount = CachedElements.Num();
-	for (int32 Index = 0; Index < CachedElementCount; Index++)
-	{
-		const FSlateDrawElement& LocalElement = CachedElements[Index];
 
-		FSlateDrawElement AbsElement = LocalElement;
-		if (bMoved)
+	if (bMoved)
+	{
+		for (int32 Index = 0; Index < CachedElementCount; Index++)
 		{
+			const FSlateDrawElement& LocalElement = CachedElements[Index];
+			FSlateDrawElement AbsElement = LocalElement;
 			FSlateDrawElement::ApplyPositionOffset(AbsElement, AbsoluteOffset);
+			AddItem(AbsElement);
 		}
-		
-		AbsElement.SetClippingIndex(LocalElement.GetClippingIndex() == -1 ? GetClippingIndex() : (ClippingStateOffset + GetClippingIndex()));
-
-		AddItem(AbsElement);
+	}
+	else
+	{
+		AppendItems(CachedElements);
 	}
 }
 
@@ -981,6 +973,20 @@ int32 FSlateWindowElementList::PaintVolatile(FSlateWindowElementList& OutElement
 	return MaxLayerId;
 }
 
+int32 FSlateWindowElementList::PaintVolatileRootLayer(FSlateWindowElementList& OutElementList, double InCurrentTime, float InDeltaTime, const FVector2D& InDynamicOffset)
+{
+	int32 MaxLayerId = 0;
+
+	for (int32 VolatileIndex = 0; VolatileIndex < VolatilePaintList.Num(); ++VolatileIndex)
+	{
+		const TSharedPtr<FVolatilePaint>& Args = VolatilePaintList[VolatileIndex];
+		MaxLayerId = FMath::Max(MaxLayerId, Args->ExecutePaint(OutElementList, InCurrentTime, InDeltaTime, InDynamicOffset));
+	}
+
+	return MaxLayerId;
+}
+
+
 void FSlateWindowElementList::BeginLogicalLayer(const TSharedPtr<FSlateDrawLayerHandle, ESPMode::ThreadSafe>& LayerHandle)
 {
 	// Don't attempt to begin logical layers inside a cached view of the data.
diff --git a/Engine/Source/Runtime/SlateCore/Public/Layout/Geometry.h b/Engine/Source/Runtime/SlateCore/Public/Layout/Geometry.h
--- a/Engine/Source/Runtime/SlateCore/Public/Layout/Geometry.h
+++ b/Engine/Source/Runtime/SlateCore/Public/Layout/Geometry.h
@@ -4,11 +4,12 @@
 
 #include "CoreMinimal.h"
 #include "UObject/ObjectMacros.h"
+#include "Math/TransformCalculus2D.h"
 #include "Layout/SlateRect.h"
 #include "Layout/SlateRotatedRect.h"
-#include "Math/TransformCalculus2D.h"
-#include "Rendering/SlateLayoutTransform.h"
 #include "Layout/PaintGeometry.h"
+#include "Layout/SlateRotatedRect.h"
+#include "Rendering/SlateLayoutTransform.h"
 #include "Geometry.generated.h"
 
 class FArrangedWidget;
@@ -376,8 +377,9 @@ public:
 	 */
 	FORCEINLINE_DEBUGGABLE bool IsUnderLocation(const FVector2D& AbsoluteCoordinate) const
 	{
-		// this render transform invert is a little expensive. We might consider caching it.
-		return FSlateRect(FVector2D(0.0f, 0.0f), Size).ContainsPoint(TransformPoint(Inverse(GetAccumulatedRenderTransform()), AbsoluteCoordinate));
+		// this render transform invert is a little expensive. We might consider caching it?
+		FSlateRotatedRect Rect = TransformRect(GetAccumulatedRenderTransform(), FSlateRotatedRect(FSlateRect(FVector2D(0.0f, 0.0f), Size)));
+		return Rect.IsUnderLocation(AbsoluteCoordinate);
 	}
 
 	/** 
diff --git a/Engine/Source/Runtime/SlateCore/Public/Rendering/DrawElements.h b/Engine/Source/Runtime/SlateCore/Public/Rendering/DrawElements.h
--- a/Engine/Source/Runtime/SlateCore/Public/Rendering/DrawElements.h
+++ b/Engine/Source/Runtime/SlateCore/Public/Rendering/DrawElements.h
@@ -1198,8 +1198,7 @@ class FSlateBatchData
 {
 public:
 	FSlateBatchData()
-		: DynamicOffset(0, 0)
-		, NumBatchedVertices(0)
+		: NumBatchedVertices(0)
 		, NumBatchedIndices(0)
 		, NumLayers(0)
 	{}
@@ -1310,9 +1309,6 @@ private:
 	/**  */
 	TArray<FSlateClippingState> RenderClipStates;
 
-	/**  */
-	FVector2D DynamicOffset;
-
 	/**  */
 	int32 NumBatchedVertices;
 
@@ -1412,6 +1408,12 @@ public:
 		TArray<FSlateDrawElement>& ActiveDrawElements = DrawStack.Last()->DrawElements;
 		ActiveDrawElements.Add(InDrawElement);
 	}
+	
+	FORCEINLINE void AppendItems(const TArray<FSlateDrawElement>& InDrawElements)
+	{
+		TArray<FSlateDrawElement>& ActiveDrawElements = DrawStack.Last()->DrawElements;
+		ActiveDrawElements.Append(InDrawElements);
+	}
 
 	/** @return Get the window size that we will be painting */
 	FORCEINLINE FVector2D GetWindowSize() const
@@ -1512,6 +1514,7 @@ public:
 	SLATECORE_API void QueueVolatilePainting( const FVolatilePaint& InVolatilePaint );
 
 	SLATECORE_API int32 PaintVolatile(FSlateWindowElementList& OutElementList, double InCurrentTime, float InDeltaTime, const FVector2D& InDynamicOffset);
+	SLATECORE_API int32 PaintVolatileRootLayer(FSlateWindowElementList& OutElementList, double InCurrentTime, float InDeltaTime, const FVector2D& InDynamicOffset);
 
 	SLATECORE_API void BeginLogicalLayer(const TSharedPtr<FSlateDrawLayerHandle, ESPMode::ThreadSafe>& LayerHandle);
 	SLATECORE_API void EndLogicalLayer();
diff --git a/Engine/Source/Runtime/UMG/Private/Components/ComboBoxString.cpp b/Engine/Source/Runtime/UMG/Private/Components/ComboBoxString.cpp
--- a/Engine/Source/Runtime/UMG/Private/Components/ComboBoxString.cpp
+++ b/Engine/Source/Runtime/UMG/Private/Components/ComboBoxString.cpp
@@ -41,7 +41,8 @@ UComboBoxString::UComboBoxString(const FObjectInitializer& ObjectInitializer)
 	// We don't want to try and load fonts on the server.
 	if ( !IsRunningDedicatedServer() )
 	{
-		static ConstructorHelpers::FObjectFinder<UFont> RobotoFontObj(TEXT("/Engine/EngineFonts/Roboto"));
+		static ConstructorHelpers::FObjectFinder<UFont> RobotoFontObj(*UWidget::GetDefaultFontName());
+#include "Widgets/SNullWidget.h"
 		Font = FSlateFontInfo(RobotoFontObj.Object, 16, FName("Bold"));
 	}
 }
diff --git a/Engine/Source/Runtime/UMG/Private/Components/EditableText.cpp b/Engine/Source/Runtime/UMG/Private/Components/EditableText.cpp
--- a/Engine/Source/Runtime/UMG/Private/Components/EditableText.cpp
+++ b/Engine/Source/Runtime/UMG/Private/Components/EditableText.cpp
@@ -22,7 +22,7 @@ UEditableText::UEditableText(const FObjectInitializer& ObjectInitializer)
 
 	if (!IsRunningDedicatedServer())
 	{
-		static ConstructorHelpers::FObjectFinder<UFont> RobotoFontObj(TEXT("/Engine/EngineFonts/Roboto"));
+		static ConstructorHelpers::FObjectFinder<UFont> RobotoFontObj(*UWidget::GetDefaultFontName());
 		Font_DEPRECATED = FSlateFontInfo(RobotoFontObj.Object, 12, FName("Bold"));
 	}
 
diff --git a/Engine/Source/Runtime/UMG/Private/Components/EditableTextBox.cpp b/Engine/Source/Runtime/UMG/Private/Components/EditableTextBox.cpp
--- a/Engine/Source/Runtime/UMG/Private/Components/EditableTextBox.cpp
+++ b/Engine/Source/Runtime/UMG/Private/Components/EditableTextBox.cpp
@@ -20,7 +20,7 @@ UEditableTextBox::UEditableTextBox(const FObjectInitializer& ObjectInitializer)
 
 	if (!IsRunningDedicatedServer())
 	{
-		static ConstructorHelpers::FObjectFinder<UFont> RobotoFontObj(TEXT("/Engine/EngineFonts/Roboto"));
+		static ConstructorHelpers::FObjectFinder<UFont> RobotoFontObj(*UWidget::GetDefaultFontName());
 		Font_DEPRECATED = FSlateFontInfo(RobotoFontObj.Object, 12, FName("Bold"));
 	}
 
diff --git a/Engine/Source/Runtime/UMG/Private/Components/Widget.cpp b/Engine/Source/Runtime/UMG/Private/Components/Widget.cpp
--- a/Engine/Source/Runtime/UMG/Private/Components/Widget.cpp
+++ b/Engine/Source/Runtime/UMG/Private/Components/Widget.cpp
@@ -1,6 +1,8 @@
 // Copyright 1998-2017 Epic Games, Inc. All Rights Reserved.
 
 #include "Components/Widget.h"
+#include "ConfigCacheIni.h"
+#include "CoreGlobals.h"
 #include "Widgets/SNullWidget.h"
 #include "Types/NavigationMetaData.h"
 #include "Widgets/IToolTip.h"
@@ -1176,6 +1178,15 @@ UWidget* UWidget::FindChildContainingDescendant(UWidget* Root, UWidget* Descenda
 	return nullptr;
 }
 
+// TODO: Clean this up to, move it to a user interface setting, don't use a config. 
+FString UWidget::GetDefaultFontName()
+{
+	FString DefaultFontName = TEXT("/Engine/EngineFonts/Roboto");
+	GConfig->GetString(TEXT("SlateStyle"), TEXT("DefaultFontName"), DefaultFontName, GEngineIni);
+
+	return DefaultFontName;
+}
+
 //bool UWidget::BindProperty(const FName& DestinationProperty, UObject* SourceObject, const FName& SourceProperty)
 //{
 //	UDelegateProperty* DelegateProperty = FindField<UDelegateProperty>(GetClass(), FName(*( DestinationProperty.ToString() + TEXT("Delegate") )));
diff --git a/Engine/Source/Runtime/UMG/Public/Components/Widget.h b/Engine/Source/Runtime/UMG/Public/Components/Widget.h
--- a/Engine/Source/Runtime/UMG/Public/Components/Widget.h
+++ b/Engine/Source/Runtime/UMG/Public/Components/Widget.h
@@ -728,6 +728,8 @@ public:
 
 	static UWidget* FindChildContainingDescendant(UWidget* Root, UWidget* Descendant);
 
+	static FString GetDefaultFontName();
+
 protected:
 	virtual void OnBindingChanged(const FName& Property);
 
diff --git a/Engine/Source/Runtime/WebBrowser/Private/SWebBrowserView.cpp b/Engine/Source/Runtime/WebBrowser/Private/SWebBrowserView.cpp
--- a/Engine/Source/Runtime/WebBrowser/Private/SWebBrowserView.cpp
+++ b/Engine/Source/Runtime/WebBrowser/Private/SWebBrowserView.cpp
@@ -72,6 +72,7 @@ SWebBrowserView::~SWebBrowserView()
 	if (SlateParentWindowPtr.IsValid())
 	{
 		SlateParentWindowPtr.Pin()->GetOnWindowDeactivatedEvent().RemoveAll(this);
+		SlateParentWindowPtr.Pin()->GetOnWindowActivatedEvent().RemoveAll(this);
 	}
 }
 
@@ -201,6 +202,25 @@ void SWebBrowserView::Construct(const FArguments& InArgs, const TSharedPtr<IWebB
 	}
 }
 
+int32 SWebBrowserView::OnPaint(const FPaintArgs& Args, const FGeometry& AllottedGeometry, const FSlateRect& MyCullingRect, FSlateWindowElementList& OutDrawElements, int32 LayerId, const FWidgetStyle& InWidgetStyle, bool bParentEnabled) const
+{
+	int32 Layer = SCompoundWidget::OnPaint(Args, AllottedGeometry, MyCullingRect, OutDrawElements, LayerId, InWidgetStyle, bParentEnabled);
+	
+	// Cache a reference to our parent window, if we didn't already reference it.
+	if (!SlateParentWindowPtr.IsValid())
+	{
+		TSharedPtr<SWindow> ParentWindow = OutDrawElements.GetWindow();
+
+		SlateParentWindowPtr = ParentWindow;
+		if (BrowserWindow.IsValid())
+		{
+			BrowserWindow->SetParentWindow(ParentWindow);
+		}
+	}
+
+	return Layer;
+}
+
 void SWebBrowserView::HandleWindowDeactivated()
 {
 	if (BrowserViewport.IsValid())
@@ -342,11 +362,6 @@ bool SWebBrowserView::IsInitialized() const
 
 void SWebBrowserView::SetupParentWindowHandlers()
 {
-	if (!SlateParentWindowPtr.IsValid())
-	{
-		SlateParentWindowPtr = FSlateApplication::Get().FindWidgetWindow(SharedThis(this));
-	}
-
 	if (SlateParentWindowPtr.IsValid() && BrowserWindow.IsValid())
 	{
 		SlateParentWindowPtr.Pin()->GetOnWindowDeactivatedEvent().AddSP(this, &SWebBrowserView::HandleWindowDeactivated);
diff --git a/Engine/Source/Runtime/WebBrowser/Public/SWebBrowserView.h b/Engine/Source/Runtime/WebBrowser/Public/SWebBrowserView.h
--- a/Engine/Source/Runtime/WebBrowser/Public/SWebBrowserView.h
+++ b/Engine/Source/Runtime/WebBrowser/Public/SWebBrowserView.h
@@ -135,6 +135,8 @@ public:
 	 */
 	void Construct(const FArguments& InArgs, const TSharedPtr<IWebBrowserWindow>& InWebBrowserWindow = nullptr);
 
+	virtual int32 OnPaint(const FPaintArgs& Args, const FGeometry& AllottedGeometry, const FSlateRect& MyCullingRect, FSlateWindowElementList& OutDrawElements, int32 LayerId, const FWidgetStyle& InWidgetStyle, bool bParentEnabled) const override;
+
 	/**
 	 * Load the specified URL.
 	 *
@@ -309,7 +311,7 @@ private:
 	/** Interface for dealing with a web browser window. */
 	TSharedPtr<IWebBrowserWindow> BrowserWindow;
 	/** The slate window that contains this widget. This must be stored weak otherwise we create a circular reference. */
-	TWeakPtr<SWindow> SlateParentWindowPtr;
+	mutable TWeakPtr<SWindow> SlateParentWindowPtr;
 	/** Viewport interface for rendering the web page. */
 	TSharedPtr<FWebBrowserViewport> BrowserViewport;
 	/** Viewport interface for rendering popup menus. */
