Added missing Hiragana combining character range

Combining characters need to run through HarfBuzz

diff --git a/Engine/Source/Runtime/SlateCore/Private/Fonts/SlateTextShaper.cpp b/Engine/Source/Runtime/SlateCore/Private/Fonts/SlateTextShaper.cpp
--- a/Engine/Source/Runtime/SlateCore/Private/Fonts/SlateTextShaper.cpp
+++ b/Engine/Source/Runtime/SlateCore/Private/Fonts/SlateTextShaper.cpp
@@ -164,6 +164,7 @@ void FSlateTextShaper::PerformTextShaping(const TCHAR* InText, const int32 InTex
 						{ TEXT('\u1AB0'), TEXT('\u1AFF') },
 						{ TEXT('\u1DC0'), TEXT('\u1DFF') },
 						{ TEXT('\u20D0'), TEXT('\u20FF') },
+						{ TEXT('\u3099'), TEXT('\u309A') },
 						{ TEXT('\u31C0'), TEXT('\u31EF') },
 						{ TEXT('\uFE20'), TEXT('\uFE2F') },
 
