Added hover style for color picker slider

#rnx

diff --git a/Engine/Source/Runtime/SlateCore/Private/Styling/CoreStyle.cpp b/Engine/Source/Runtime/SlateCore/Private/Styling/CoreStyle.cpp
--- a/Engine/Source/Runtime/SlateCore/Private/Styling/CoreStyle.cpp
+++ b/Engine/Source/Runtime/SlateCore/Private/Styling/CoreStyle.cpp
@@ -709,6 +709,7 @@ TSharedRef<ISlateStyle> FCoreStyle::Create( const FName& InStyleSetName )
 		Style->Set( "ColorPicker.Slider", FSliderStyle()
 			.SetDisabledThumbImage( IMAGE_BRUSH( "Common/ColorPicker_SliderHandle", FVector2D(8.0f, 32.0f)) )
 			.SetNormalThumbImage( IMAGE_BRUSH( "Common/ColorPicker_SliderHandle", FVector2D(8.0f, 32.0f)) )
+			.SetHoveredThumbImage( IMAGE_BRUSH("Common/ColorPicker_SliderHandle", FVector2D(8.0f, 32.0f)) )
 			);
 	}
 
