Cast our new resize w/h to int before checking if we are already that size

diff --git a/Engine/Source/Runtime/SlateCore/Private/Widgets/SWindow.cpp b/Engine/Source/Runtime/SlateCore/Private/Widgets/SWindow.cpp
--- a/Engine/Source/Runtime/SlateCore/Private/Widgets/SWindow.cpp
+++ b/Engine/Source/Runtime/SlateCore/Private/Widgets/SWindow.cpp
@@ -914,17 +914,21 @@ void SWindow::Resize( FVector2D NewSize )
 
 	NewSize = GetWindowSizeFromClientSize(NewSize);
 
-	if ( Size != NewSize )
-	{
-		NewSize.X = FMath::Max(SizeLimits.GetMinWidth().Get(NewSize.X), NewSize.X);
-		NewSize.X = FMath::Min(SizeLimits.GetMaxWidth().Get(NewSize.X), NewSize.X);
+	NewSize.X = FMath::Max(SizeLimits.GetMinWidth().Get(NewSize.X), NewSize.X);
+	NewSize.X = FMath::Min(SizeLimits.GetMaxWidth().Get(NewSize.X), NewSize.X);
+
+	NewSize.Y = FMath::Max(SizeLimits.GetMinHeight().Get(NewSize.Y), NewSize.Y);
+	NewSize.Y = FMath::Min(SizeLimits.GetMaxHeight().Get(NewSize.Y), NewSize.Y);
 
-		NewSize.Y = FMath::Max(SizeLimits.GetMinHeight().Get(NewSize.Y), NewSize.Y);
-		NewSize.Y = FMath::Min(SizeLimits.GetMaxHeight().Get(NewSize.Y), NewSize.Y);
+	// ReshapeWindow W/H takes an int, so lets move our new W/H to int before checking if they are the same size
+	FIntPoint CurrentIntSize = FIntPoint(FMath::TruncToInt(Size.X), FMath::TruncToInt(Size.Y));
+	FIntPoint NewIntSize     = FIntPoint(FMath::CeilToInt(NewSize.X), FMath::CeilToInt(NewSize.Y));
 
+	if (CurrentIntSize != NewIntSize)
+	{
 		if (NativeWindow.IsValid())
 		{
-			NativeWindow->ReshapeWindow( FMath::TruncToInt(ScreenPosition.X), FMath::TruncToInt(ScreenPosition.Y), FMath::CeilToInt(NewSize.X), FMath::CeilToInt(NewSize.Y) );
+			NativeWindow->ReshapeWindow(FMath::TruncToInt(ScreenPosition.X), FMath::TruncToInt(ScreenPosition.Y), NewIntSize.X, NewIntSize.Y);
 		}
 		else
 		{
