Ceil not trunk for this compare

diff --git a/Engine/Source/Runtime/SlateCore/Private/Widgets/SWindow.cpp b/Engine/Source/Runtime/SlateCore/Private/Widgets/SWindow.cpp
--- a/Engine/Source/Runtime/SlateCore/Private/Widgets/SWindow.cpp
+++ b/Engine/Source/Runtime/SlateCore/Private/Widgets/SWindow.cpp
@@ -921,7 +921,7 @@ void SWindow::Resize( FVector2D NewSize )
 	NewSize.Y = FMath::Min(SizeLimits.GetMaxHeight().Get(NewSize.Y), NewSize.Y);
 
 	// ReshapeWindow W/H takes an int, so lets move our new W/H to int before checking if they are the same size
-	FIntPoint CurrentIntSize = FIntPoint(FMath::TruncToInt(Size.X), FMath::TruncToInt(Size.Y));
+	FIntPoint CurrentIntSize = FIntPoint(FMath::CeilToInt(Size.X), FMath::CeilToInt(Size.Y));
 	FIntPoint NewIntSize     = FIntPoint(FMath::CeilToInt(NewSize.X), FMath::CeilToInt(NewSize.Y));
 
 	if (CurrentIntSize != NewIntSize)
