Slate Fast Path  - Changed FSlateWindowElementList GetWindow to be thread safe.  For fast path, this is accessed on multiple threads so it needs to be safe
GetWindow is deprecated and GetPaintWindow should be used instead

Edigrate from source CL 4254611

diff --git a/Engine/Source/Developer/SlateReflector/Private/Widgets/SWidgetReflector.cpp b/Engine/Source/Developer/SlateReflector/Private/Widgets/SWidgetReflector.cpp
--- a/Engine/Source/Developer/SlateReflector/Private/Widgets/SWidgetReflector.cpp
+++ b/Engine/Source/Developer/SlateReflector/Private/Widgets/SWidgetReflector.cpp
@@ -1230,12 +1230,12 @@ int32 SWidgetReflector::VisualizeCursorAndKeys(FSlateWindowElementList& OutDrawE
 		static const float ClickFadeTime = 0.5f;
 		static const float PingScaleAmount = 3.0f;
 		static const FName CursorPingBrush("DemoRecording.CursorPing");
-		const TSharedPtr<SWindow> WindowBeingDrawn = OutDrawElements.GetWindow();
+		const SWindow* WindowBeingDrawn = OutDrawElements.GetPaintWindow();
 
 		// Normalized animation value for the cursor ping between 0 and 1.
 		const float AnimAmount = (FSlateApplication::Get().GetCurrentTime() - LastMouseClickTime) / ClickFadeTime;
 
-		if (WindowBeingDrawn.IsValid() && AnimAmount <= 1.0f)
+		if (WindowBeingDrawn && AnimAmount <= 1.0f)
 		{
 			const FVector2D CursorPosDesktopSpace = CursorPingPosition;
 			const FVector2D CursorSize = FSlateApplication::Get().GetCursorSize();
diff --git a/Engine/Source/Developer/StandaloneRenderer/Private/OpenGL/SlateOpenGLRenderer.cpp b/Engine/Source/Developer/StandaloneRenderer/Private/OpenGL/SlateOpenGLRenderer.cpp
--- a/Engine/Source/Developer/StandaloneRenderer/Private/OpenGL/SlateOpenGLRenderer.cpp
+++ b/Engine/Source/Developer/StandaloneRenderer/Private/OpenGL/SlateOpenGLRenderer.cpp
@@ -112,12 +112,12 @@ void FSlateOpenGLRenderer::DrawWindows( FSlateDrawBuffer& InWindowDrawBuffer )
 	{
 		FSlateWindowElementList& ElementList = *WindowElementLists[ListIndex];
 
-		if ( ElementList.GetWindow().IsValid() )
+		if ( ElementList.GetRenderWindow() )
 		{
-			TSharedRef<SWindow> WindowToDraw = ElementList.GetWindow().ToSharedRef();
+			SWindow* WindowToDraw = ElementList.GetRenderWindow();
 
 			const FVector2D WindowSize = WindowToDraw->GetSizeInScreen();
-			FSlateOpenGLViewport* Viewport = WindowToViewportMap.Find( &WindowToDraw.Get() );
+			FSlateOpenGLViewport* Viewport = WindowToViewportMap.Find( WindowToDraw );
 			check(Viewport);
 		
 			//@todo Slate OpenGL: Move this to ResizeViewport
diff --git a/Engine/Source/Developer/StandaloneRenderer/Private/Windows/D3D/SlateD3DRenderer.cpp b/Engine/Source/Developer/StandaloneRenderer/Private/Windows/D3D/SlateD3DRenderer.cpp
--- a/Engine/Source/Developer/StandaloneRenderer/Private/Windows/D3D/SlateD3DRenderer.cpp
+++ b/Engine/Source/Developer/StandaloneRenderer/Private/Windows/D3D/SlateD3DRenderer.cpp
@@ -520,11 +520,10 @@ void FSlateD3DRenderer::DrawWindows( FSlateDrawBuffer& InWindowDrawBuffer )
 	for( int32 ListIndex = 0; ListIndex < WindowElementLists.Num(); ++ListIndex )
 	{
 		FSlateWindowElementList& ElementList = *WindowElementLists[ListIndex];
-		SLATE_CYCLE_COUNTER_SCOPE_CUSTOM_DETAILED(SLATE_STATS_DETAIL_LEVEL_MED, GRendererDrawElementList, ElementList.GetWindow()->GetCreatedInLocation());
 
-		if ( ElementList.GetWindow().IsValid() )
+		if ( ElementList.GetRenderWindow() )
 		{
-			TSharedRef<SWindow> WindowToDraw = ElementList.GetWindow().ToSharedRef();
+			SWindow* WindowToDraw = ElementList.GetRenderWindow();
 
 			// Add all elements for this window to the element batcher
 			ElementBatcher->AddElements( ElementList );
@@ -534,7 +533,7 @@ void FSlateD3DRenderer::DrawWindows( FSlateDrawBuffer& InWindowDrawBuffer )
 
 			FVector2D WindowSize = WindowToDraw->GetSizeInScreen();
 
-			FSlateD3DViewport* Viewport = WindowToViewportMap.Find( &WindowToDraw.Get() );
+			FSlateD3DViewport* Viewport = WindowToViewportMap.Find( WindowToDraw );
 			check(Viewport);
 
 			FSlateBatchData& BatchData = ElementList.GetBatchData();
@@ -559,7 +558,6 @@ void FSlateD3DRenderer::DrawWindows( FSlateDrawBuffer& InWindowDrawBuffer )
 			GD3DDeviceContext->OMSetRenderTargets( 1, &RTV, NULL );
 
 			{
-				SLATE_CYCLE_COUNTER_SCOPE(GRendererDrawElements);
 				RenderingPolicy->DrawElements(ViewMatrix * Viewport->ProjectionMatrix, BatchData.GetRenderBatches());
 			}
 
diff --git a/Engine/Source/Editor/IntroTutorials/Private/STutorialButton.cpp b/Engine/Source/Editor/IntroTutorials/Private/STutorialButton.cpp
--- a/Engine/Source/Editor/IntroTutorials/Private/STutorialButton.cpp
+++ b/Engine/Source/Editor/IntroTutorials/Private/STutorialButton.cpp
@@ -122,7 +122,7 @@ int32 STutorialButton::OnPaint(const FPaintArgs& Args, const FGeometry& Allotted
 		const FLinearColor PulseColor = FEditorStyle::Get().GetColor(TEXT("TutorialLaunch.Circle.Color"));
 
 		// We should be clipped by the window size, not our containing widget, as we want to draw outside the widget
-		const FVector2D WindowSize = OutDrawElements.GetWindow()->GetSizeInScreen();
+		const FVector2D WindowSize = OutDrawElements.GetPaintWindow()->GetSizeInScreen();
 		FSlateRect WindowClippingRect(0.0f, 0.0f, WindowSize.X, WindowSize.Y);
 
 		{
diff --git a/Engine/Source/Editor/IntroTutorials/Private/STutorialContent.cpp b/Engine/Source/Editor/IntroTutorials/Private/STutorialContent.cpp
--- a/Engine/Source/Editor/IntroTutorials/Private/STutorialContent.cpp
+++ b/Engine/Source/Editor/IntroTutorials/Private/STutorialContent.cpp
@@ -282,7 +282,7 @@ void STutorialContent::GetAnimationValues(float& OutAlphaFactor, float& OutPulse
 int32 STutorialContent::OnPaint(const FPaintArgs& Args, const FGeometry& AllottedGeometry, const FSlateRect& MyCullingRect, FSlateWindowElementList& OutDrawElements, int32 LayerId, const FWidgetStyle& InWidgetStyle, bool bParentEnabled) const
 {
 	CachedContentGeometry = AllottedGeometry;
-	CachedContentGeometry.AppendTransform(FSlateLayoutTransform(OutDrawElements.GetWindow()->GetPositionInScreen()));
+	CachedContentGeometry.AppendTransform(FSlateLayoutTransform(OutDrawElements.GetPaintWindow()->GetPositionInScreen()));
 
 	if(bIsVisible && Anchor.Type != ETutorialAnchorIdentifier::None && Anchor.bDrawHighlight)
 	{
@@ -296,7 +296,7 @@ int32 STutorialContent::OnPaint(const FPaintArgs& Args, const FGeometry& Allotte
 		const FSlateBrush* BorderBrush = FCoreStyle::Get().GetBrush(TEXT("Tutorials.Border"));
 					
 		const FGeometry& WidgetGeometry = CachedGeometry;
-		const FVector2D WindowSize = OutDrawElements.GetWindow()->GetSizeInScreen();
+		const FVector2D WindowSize = OutDrawElements.GetPaintWindow()->GetSizeInScreen();
 
 		// We should be clipped by the window size, not our containing widget, as we want to draw outside the widget
 		FSlateRect WindowClippingRect(0.0f, 0.0f, WindowSize.X, WindowSize.Y);
diff --git a/Engine/Source/Runtime/Slate/Private/Framework/Application/SlateApplication.cpp b/Engine/Source/Runtime/Slate/Private/Framework/Application/SlateApplication.cpp
--- a/Engine/Source/Runtime/Slate/Private/Framework/Application/SlateApplication.cpp
+++ b/Engine/Source/Runtime/Slate/Private/Framework/Application/SlateApplication.cpp
@@ -1488,15 +1488,20 @@ void FSlateApplication::DrawPrepass( TSharedPtr<SWindow> DrawOnlyThisWindow )
 	}
 }
 
-TArray<TSharedRef<SWindow>> GatherAllDescendants(const TArray< TSharedRef<SWindow> >& InWindowList)
+TArray<SWindow*> GatherAllDescendants(const TArray< TSharedRef<SWindow> >& InWindowList)
 {
-	TArray<TSharedRef<SWindow>> GatheredDescendants(InWindowList);
+	TArray<SWindow*> GatheredDescendants;
+	GatheredDescendants.Reserve(InWindowList.Num());
+	for (const TSharedRef<SWindow>& Window : InWindowList)
+	{
+		GatheredDescendants.Add(&Window.Get());
+	}
 
 	for (const TSharedRef<SWindow>& SomeWindow : InWindowList)
 	{
-		GatheredDescendants.Append( GatherAllDescendants( SomeWindow->GetChildWindows() ) );
+		GatheredDescendants.Append(GatherAllDescendants(SomeWindow->GetChildWindows()));
 	}
-	
+
 	return GatheredDescendants;
 }
 
@@ -1576,7 +1581,7 @@ void FSlateApplication::PrivateDrawWindows( TSharedPtr<SWindow> DrawOnlyThisWind
 	{
 		// Some windows may have been destroyed/removed.
 		// Do not attempt to draw any windows that have been removed.
-		TArray<TSharedRef<SWindow>> AllWindows = GatherAllDescendants(SlateWindows);
+		TArray<SWindow*> AllWindows = GatherAllDescendants(SlateWindows);
 		DrawWindowArgs.OutDrawBuffer.RemoveUnusedWindowElement(AllWindows);
 	}
 
@@ -4515,7 +4520,7 @@ TSharedPtr< FSlateWindowElementList > FSlateApplication::FCacheElementPools::Get
 	// Remove inactive lists that don't belong to this window.
 	for ( int32 i = InactiveCachedElementListPool.Num() - 1; i >= 0; i-- )
 	{
-		if ( InactiveCachedElementListPool[i]->GetWindow() != CurrentWindow )
+		if (InactiveCachedElementListPool[i]->GetPaintWindow() != CurrentWindow.Get())
 		{
 			InactiveCachedElementListPool.RemoveAtSwap(i, 1, false);
 		}
@@ -6901,6 +6906,8 @@ EWindowZone::Type FSlateApplication::GetWindowZoneForPoint( const TSharedRef< FG
 
 void FSlateApplication::PrivateDestroyWindow( const TSharedRef<SWindow>& DestroyedWindow )
 {
+	WindowBeingDestroyedEvent.Broadcast(*DestroyedWindow);
+
 	// Notify the window that it is going to be destroyed.  The window must be completely intact when this is called 
 	// because delegates are allowed to leave Slate here
 	DestroyedWindow->NotifyWindowBeingDestroyed();
diff --git a/Engine/Source/Runtime/Slate/Private/Widgets/SInvalidationPanel.cpp b/Engine/Source/Runtime/Slate/Private/Widgets/SInvalidationPanel.cpp
--- a/Engine/Source/Runtime/Slate/Private/Widgets/SInvalidationPanel.cpp
+++ b/Engine/Source/Runtime/Slate/Private/Widgets/SInvalidationPanel.cpp
@@ -341,7 +341,8 @@ int32 SInvalidationPanel::OnPaint( const FPaintArgs& Args, const FGeometry& Allo
 				// Record a new maximum layer Id, throw in 10 so we have a bit of padding before we need to recache.
 				MaximumLayerIdCachedAt = LayerId + 10;
 
-				CachedWindowElements = FSlateApplication::Get().GetCachableElementList(OutDrawElements.GetWindow(), this);
+				SWindow* Window = OutDrawElements.GetPaintWindow();
+				CachedWindowElements = FSlateApplication::Get().GetCachableElementList(StaticCastSharedRef<SWindow>(Window->AsShared()), this);
 
 				// Reset the render data handle in case it was in use, and we're not overriding it this frame.
 				CachedRenderData.Reset();
diff --git a/Engine/Source/Runtime/Slate/Public/Framework/Application/SlateApplication.h b/Engine/Source/Runtime/Slate/Public/Framework/Application/SlateApplication.h
--- a/Engine/Source/Runtime/Slate/Public/Framework/Application/SlateApplication.h
+++ b/Engine/Source/Runtime/Slate/Public/Framework/Application/SlateApplication.h
@@ -523,6 +523,10 @@ public:
 	DECLARE_EVENT_OneParam(FSlateApplication, FUserRegisteredEvent, int32);
 	FUserRegisteredEvent& OnUserRegistered() { return UserRegisteredEvent; }
 
+	/** Delegate called when a window is about to be destroyed */
+	DECLARE_EVENT_OneParam(FSlateApplication, FOnWindowBeingDestroyed, const SWindow&);
+	FOnWindowBeingDestroyed& OnWindowBeingDestroyed() { return WindowBeingDestroyedEvent; }
+
 	/** 
 	 * Removes references to FViewportRHI's.  
 	 * This has to be done explicitly instead of using the FRenderResource mechanism because FViewportRHI's are managed by the game thread.
@@ -2136,6 +2140,9 @@ private:
 	/** Delegate for when a new user has been registered. */
 	FUserRegisteredEvent UserRegisteredEvent;
 
+	/** Delegate for when a window is in the process of being destroyed */
+	FOnWindowBeingDestroyed WindowBeingDestroyedEvent;
+
 	/** Delegate for slate Tick during modal dialogs */
 	FOnModalLoopTickEvent ModalLoopTickEvent;
 
diff --git a/Engine/Source/Runtime/SlateCore/Private/Rendering/DrawElements.cpp b/Engine/Source/Runtime/SlateCore/Private/Rendering/DrawElements.cpp
--- a/Engine/Source/Runtime/SlateCore/Private/Rendering/DrawElements.cpp
+++ b/Engine/Source/Runtime/SlateCore/Private/Rendering/DrawElements.cpp
@@ -30,7 +30,8 @@ static bool IsResourceObjectValid(UObject*& InObject)
 }
 
 FSlateWindowElementList::FSlateWindowElementList(TSharedPtr<SWindow> InPaintWindow)
-	: PaintWindow(InPaintWindow)
+	: WeakPaintWindow(InPaintWindow)
+	, RawPaintWindow(InPaintWindow.Get())
 	, RenderTargetWindow(nullptr)
 	, bNeedsDeferredResolve(false)
 	, ResolveToDeferredIndex()
diff --git a/Engine/Source/Runtime/SlateCore/Private/Rendering/ElementBatcher.cpp b/Engine/Source/Runtime/SlateCore/Private/Rendering/ElementBatcher.cpp
--- a/Engine/Source/Runtime/SlateCore/Private/Rendering/ElementBatcher.cpp
+++ b/Engine/Source/Runtime/SlateCore/Private/Rendering/ElementBatcher.cpp
@@ -63,7 +63,7 @@ void FSlateElementBatcher::AddElements(FSlateWindowElementList& WindowElementLis
 	BatchData = &WindowElementList.GetBatchData();
 	DrawLayer = &WindowElementList.GetRootDrawLayer();
 	
-	FVector2D ViewportSize = WindowElementList.GetWindow()->GetViewportSize();
+	FVector2D ViewportSize = WindowElementList.GetPaintWindow()->GetViewportSize();
 
 	ClippingStates = &WindowElementList.ClippingManager.GetClippingStates();
 
diff --git a/Engine/Source/Runtime/SlateCore/Private/Rendering/SlateDrawBuffer.cpp b/Engine/Source/Runtime/SlateCore/Private/Rendering/SlateDrawBuffer.cpp
--- a/Engine/Source/Runtime/SlateCore/Private/Rendering/SlateDrawBuffer.cpp
+++ b/Engine/Source/Runtime/SlateCore/Private/Rendering/SlateDrawBuffer.cpp
@@ -18,7 +18,7 @@ FSlateWindowElementList& FSlateDrawBuffer::AddWindowElementList(TSharedRef<SWind
 	{
 		TSharedRef<FSlateWindowElementList> ExistingElementList = WindowElementListsPool[WindowIndex];
 
-		if (ExistingElementList->GetWindow() == ForWindow )
+		if (ExistingElementList->GetPaintWindow() == &ForWindow.Get())
 		{
 			WindowElementLists.Add(ExistingElementList);
 			WindowElementListsPool.RemoveAtSwap(WindowIndex);
@@ -35,13 +35,13 @@ FSlateWindowElementList& FSlateDrawBuffer::AddWindowElementList(TSharedRef<SWind
 	return *WindowElements;
 }
 
-void FSlateDrawBuffer::RemoveUnusedWindowElement(const TArray<TSharedRef<SWindow>>& AllWindows)
+void FSlateDrawBuffer::RemoveUnusedWindowElement(const TArray<SWindow*>& AllWindows)
 {
 	// Remove any window elements that are no longer valid.
 	for (int32 WindowIndex = 0; WindowIndex < WindowElementLists.Num(); ++WindowIndex)
 	{
-		TSharedPtr<SWindow> CandidateWindow = WindowElementLists[WindowIndex]->GetWindow();
-		if (CandidateWindow.IsValid() == false || !AllWindows.Contains(CandidateWindow.ToSharedRef()))
+		SWindow* CandidateWindow = WindowElementLists[WindowIndex]->GetPaintWindow();
+		if (!CandidateWindow || !AllWindows.Contains(CandidateWindow))
 		{
 			WindowElementLists.RemoveAtSwap(WindowIndex);
 			--WindowIndex;
@@ -64,7 +64,7 @@ void FSlateDrawBuffer::ClearBuffer()
 	// Remove any window elements that are no longer valid.
 	for (int32 WindowIndex = 0; WindowIndex < WindowElementListsPool.Num(); ++WindowIndex)
 	{
-		if (WindowElementListsPool[WindowIndex]->GetWindow().IsValid() == false)
+		if (WindowElementListsPool[WindowIndex]->GetPaintWindow() == nullptr)
 		{
 			WindowElementListsPool.RemoveAtSwap(WindowIndex);
 			--WindowIndex;
@@ -74,7 +74,7 @@ void FSlateDrawBuffer::ClearBuffer()
 	// Move all the window elements back into the pool.
 	for (TSharedRef<FSlateWindowElementList> ExistingList : WindowElementLists)
 	{
-		if (ExistingList->GetWindow().IsValid())
+		if (ExistingList->GetPaintWindow() != nullptr)
 		{
 			WindowElementListsPool.Add(ExistingList);
 		}
@@ -83,6 +83,7 @@ void FSlateDrawBuffer::ClearBuffer()
 	WindowElementLists.Reset();
 }
 
+
 void FSlateDrawBuffer::UpdateResourceVersion(uint32 NewResourceVersion)
 {
 	if (IsInGameThread() && NewResourceVersion != ResourceVersion)
diff --git a/Engine/Source/Runtime/SlateCore/Public/Rendering/DrawElements.h b/Engine/Source/Runtime/SlateCore/Public/Rendering/DrawElements.h
--- a/Engine/Source/Runtime/SlateCore/Public/Rendering/DrawElements.h
+++ b/Engine/Source/Runtime/SlateCore/Public/Rendering/DrawElements.h
@@ -1653,18 +1653,27 @@ public:
 	SLATECORE_API ~FSlateWindowElementList();
 
 	/** @return Get the window that we will be painting */
+	DEPRECATED(4.21, "FSlateWindowElementList::GetWindow is not thread safe but window element lists are accessed on multiple threads.  Please call GetPaintWindow instead")
 	FORCEINLINE TSharedPtr<SWindow> GetWindow() const
 	{
 		// check that we are in game thread or are in slate/movie loading thread
-		check(IsInGameThread() || IsInSlateThread())
-		return PaintWindow.Pin();
+		check(IsInGameThread() || IsInSlateThread());
+		return WeakPaintWindow.Pin();
 	}
 
-	/** @return Get the window that we will be rendering to */
+	/** @return Get the window that we will be painting */
+	FORCEINLINE SWindow* GetPaintWindow() const
+	{
+		check(IsInGameThread() || IsInSlateThread());
+		return WeakPaintWindow.IsValid() ? RawPaintWindow : nullptr;
+	}
+
+	/** @return Get the window that we will be rendering to.  Unless you are a slate renderer you probably want to use GetPaintWindow() */
 	FORCEINLINE SWindow* GetRenderWindow() const
 	{
+		check(IsInGameThread() || IsInSlateThread());
 		// Note: This assumes that the PaintWindow is safe to pin and is not accessed by another thread
-		return RenderTargetWindow != nullptr ? RenderTargetWindow : PaintWindow.Pin().Get();
+		return RenderTargetWindow != nullptr ? RenderTargetWindow : GetPaintWindow();
 	}
 
 	/** @return Get the draw elements that we want to render into this window */
@@ -1952,11 +1961,12 @@ private:
 	/** Resources to report to the garbage collector */
 	TArray<UObject*> ResourcesToReport;
 
-	/** 
-	 * Window which owns the widgets that are being painted but not necessarily rendered to
-	 * Widgets are always rendered to the RenderTargetWindow
-	 */
-	TWeakPtr<SWindow> PaintWindow;
+	/**
+	* Window which owns the widgets that are being painted but not necessarily rendered to
+	* Widgets are always rendered to the RenderTargetWindow
+	*/
+	TWeakPtr<SWindow> WeakPaintWindow;
+	SWindow* RawPaintWindow;
 
 	/**
 	 * The window to render to.  This may be different from the paint window if we are displaying the contents of a window (or virtual window) onto another window
diff --git a/Engine/Source/Runtime/SlateCore/Public/Rendering/SlateDrawBuffer.h b/Engine/Source/Runtime/SlateCore/Public/Rendering/SlateDrawBuffer.h
--- a/Engine/Source/Runtime/SlateCore/Public/Rendering/SlateDrawBuffer.h
+++ b/Engine/Source/Runtime/SlateCore/Public/Rendering/SlateDrawBuffer.h
@@ -37,7 +37,7 @@ public:
 	FSlateWindowElementList& AddWindowElementList(TSharedRef<SWindow> ForWindow);
 
 	/** Removes any window from the draw buffer that's not in this list or whose window has become invalid. */
-	void RemoveUnusedWindowElement(const TArray<TSharedRef<SWindow>>& AllWindows);
+	void RemoveUnusedWindowElement(const TArray<SWindow*>& AllWindows);
 
 	/**
 	 * Gets all window element lists in this buffer.
diff --git a/Engine/Source/Runtime/SlateRHIRenderer/Private/Slate3DRenderer.cpp b/Engine/Source/Runtime/SlateRHIRenderer/Private/Slate3DRenderer.cpp
--- a/Engine/Source/Runtime/SlateRHIRenderer/Private/Slate3DRenderer.cpp
+++ b/Engine/Source/Runtime/SlateRHIRenderer/Private/Slate3DRenderer.cpp
@@ -78,16 +78,16 @@ void FSlate3DRenderer::DrawWindow_GameThread(FSlateDrawBuffer& DrawBuffer)
 
 	const TArray<TSharedRef<FSlateWindowElementList>>& WindowElementLists = DrawBuffer.GetWindowElementLists();
 
-	for ( int32 WindowIndex = 0; WindowIndex < WindowElementLists.Num(); WindowIndex++ )
+	for (int32 WindowIndex = 0; WindowIndex < WindowElementLists.Num(); WindowIndex++)
 	{
 		FSlateWindowElementList& ElementList = *WindowElementLists[WindowIndex];
 
-		TSharedPtr<SWindow> Window = ElementList.GetWindow();
+		SWindow* Window = ElementList.GetPaintWindow();
 
-		if ( Window.IsValid() )
+		if (Window)
 		{
 			const FVector2D WindowSize = Window->GetSizeInScreen();
-			if ( WindowSize.X > 0 && WindowSize.Y > 0 )
+			if (WindowSize.X > 0 && WindowSize.Y > 0)
 			{
 				// Add all elements for this window to the element batcher
 				ElementBatcher->AddElements(ElementList);
diff --git a/Engine/Source/Runtime/WebBrowser/Private/SWebBrowserView.cpp b/Engine/Source/Runtime/WebBrowser/Private/SWebBrowserView.cpp
--- a/Engine/Source/Runtime/WebBrowser/Private/SWebBrowserView.cpp
+++ b/Engine/Source/Runtime/WebBrowser/Private/SWebBrowserView.cpp
@@ -210,12 +210,13 @@ int32 SWebBrowserView::OnPaint(const FPaintArgs& Args, const FGeometry& Allotted
 	// Cache a reference to our parent window, if we didn't already reference it.
 	if (!SlateParentWindowPtr.IsValid())
 	{
-		TSharedPtr<SWindow> ParentWindow = OutDrawElements.GetWindow();
+		SWindow* ParentWindow = OutDrawElements.GetPaintWindow();
 
-		SlateParentWindowPtr = ParentWindow;
+		TSharedRef<SWindow> SlateParentWindowRef = StaticCastSharedRef<SWindow>(ParentWindow->AsShared());
+		SlateParentWindowPtr = SlateParentWindowRef;
 		if (BrowserWindow.IsValid())
 		{
-			BrowserWindow->SetParentWindow(ParentWindow);
+			BrowserWindow->SetParentWindow(SlateParentWindowRef);
 		}
 	}
 
