Added better check output to try and learn more about a crash in the wild.
Added some better const saftey while in there.

diff --git a/Engine/Source/Runtime/SlateCore/Private/Input/HittestGrid.cpp b/Engine/Source/Runtime/SlateCore/Private/Input/HittestGrid.cpp
--- a/Engine/Source/Runtime/SlateCore/Private/Input/HittestGrid.cpp
+++ b/Engine/Source/Runtime/SlateCore/Private/Input/HittestGrid.cpp
@@ -394,16 +394,16 @@ TSharedPtr<SWidget> FHittestGrid::FindFocusableWidget(FSlateRect WidgetRect, con
 
 		for (StrideCellPoint[StrideAxis] = StrideAxisMin; StrideCellPoint[StrideAxis] <= StrideAxisMax; ++StrideCellPoint[StrideAxis])
 		{
-			FHittestGrid::FCell& Cell = FHittestGrid::CellAt(StrideCellPoint.X, StrideCellPoint.Y);
+			const FHittestGrid::FCell& Cell = FHittestGrid::CellAt(StrideCellPoint.X, StrideCellPoint.Y);
 			const TArray<int32>& IndexesInCell = Cell.CachedWidgetIndexes;
 
 			for (int32 i = IndexesInCell.Num() - 1; i >= 0; --i)
 			{
-				int32 CurrentIndex = IndexesInCell[i];
+				const int32 CurrentIndex = IndexesInCell[i];
 				check(CurrentIndex < WidgetsCachedThisFrame.Num());
 
 				const FCachedWidget& TestCandidate = WidgetsCachedThisFrame[CurrentIndex];
-				FSlateRect TestCandidateRect = TestCandidate.CachedGeometry.GetRenderBoundingRect().OffsetBy(-GridOrigin);
+				const FSlateRect TestCandidateRect = TestCandidate.CachedGeometry.GetRenderBoundingRect().OffsetBy(-GridOrigin);
 
 				if (CompareFunc(DestSideFunc(TestCandidateRect), CurrentSourceSide) && FSlateRect::DoRectanglesIntersect(SweptRect, TestCandidateRect))
 				{
diff --git a/Engine/Source/Runtime/SlateCore/Public/Input/HittestGrid.h b/Engine/Source/Runtime/SlateCore/Public/Input/HittestGrid.h
--- a/Engine/Source/Runtime/SlateCore/Public/Input/HittestGrid.h
+++ b/Engine/Source/Runtime/SlateCore/Public/Input/HittestGrid.h
@@ -174,13 +174,13 @@ private:
 	/** Access a cell at coordinates X, Y. Coordinates are row and column indexes. */
 	FORCEINLINE_DEBUGGABLE FCell& CellAt(const int32 X, const int32 Y)
 	{
-		check(( Y*NumCells.X + X ) < Cells.Num());
+		checkf((Y*NumCells.X + X) < Cells.Num(), TEXT("HitTestGrid CellAt() failed: X= %d Y= %d NumCells.X= %d NumCells.Y= %d Cells.Num()= %d"), X, Y, NumCells.X, NumCells.Y, Cells.Num());
 		return Cells[Y*NumCells.X + X];
 	}
 
 	FORCEINLINE_DEBUGGABLE const FCell& CellAt( const int32 X, const int32 Y ) const
 	{
-		check(( Y*NumCells.X + X ) < Cells.Num());
+		checkf((Y*NumCells.X + X) < Cells.Num(), TEXT("HitTestGrid CellAt() failed: X= %d Y= %d NumCells.X= %d NumCells.Y= %d Cells.Num()= %d"), X, Y, NumCells.X, NumCells.Y, Cells.Num());
 		return Cells[Y*NumCells.X + X];
 	}
 
