Fix outline colors not updating when changing on the fly

diff --git a/Engine/Source/Runtime/SlateCore/Public/Fonts/SlateFontInfo.h b/Engine/Source/Runtime/SlateCore/Public/Fonts/SlateFontInfo.h
--- a/Engine/Source/Runtime/SlateCore/Public/Fonts/SlateFontInfo.h
+++ b/Engine/Source/Runtime/SlateCore/Public/Fonts/SlateFontInfo.h
@@ -71,7 +71,9 @@ struct SLATECORE_API FFontOutlineSettings
 	bool operator==(const FFontOutlineSettings& Other) const
 	{
 		return OutlineSize == Other.OutlineSize
-			&& bSeparateFillAlpha == Other.bSeparateFillAlpha;
+			&& bSeparateFillAlpha == Other.bSeparateFillAlpha
+			&& OutlineMaterial == Other.OutlineMaterial
+			&& OutlineColor == Other.OutlineColor;
 	}
 
 	friend inline uint32 GetTypeHash(const FFontOutlineSettings& OutlineSettings)
