Slate - Encountering crashes that appear to be from the widget hierarchy changing during prepass, which shouldn't be possible.  Adding ensures to catch it, and updating the loop to resample the number of widgets in case it changes.

Matt.Kuhlenschmidt

diff --git a/Engine/Source/Runtime/SlateCore/Private/Widgets/SWidget.cpp b/Engine/Source/Runtime/SlateCore/Private/Widgets/SWidget.cpp
--- a/Engine/Source/Runtime/SlateCore/Private/Widgets/SWidget.cpp
+++ b/Engine/Source/Runtime/SlateCore/Private/Widgets/SWidget.cpp
@@ -501,17 +501,18 @@ void SWidget::SlatePrepass(float InLayoutScaleMultiplier)
 		// a function of its children's sizes.
 		FChildren* MyChildren = this->GetChildren();
 		const int32 NumChildren = MyChildren->Num();
-		for ( int32 ChildIndex=0; ChildIndex < NumChildren; ++ChildIndex )
+		for (int32 ChildIndex = 0; ChildIndex < MyChildren->Num(); ++ChildIndex)
 		{
 			const TSharedRef<SWidget>& Child = MyChildren->GetChildAt(ChildIndex);
 
-			if ( GSlateLayoutCaching || Child->Visibility.Get() != EVisibility::Collapsed )
+			if (GSlateLayoutCaching || Child->Visibility.Get() != EVisibility::Collapsed)
 			{
 				const float ChildLayoutScaleMultiplier = GetRelativeLayoutScale(MyChildren->GetSlotAt(ChildIndex), InLayoutScaleMultiplier);
 				// Recur: Descend down the widget tree.
 				Child->SlatePrepass(InLayoutScaleMultiplier * ChildLayoutScaleMultiplier);
 			}
 		}
+		ensure(NumChildren == MyChildren->Num());
 	}
 
 	if(!GSlateLayoutCaching)
