[LLM] use correct LLM macro for tracking font caches

diff --git a/Engine/Source/Runtime/SlateCore/Private/Fonts/FontCacheCompositeFont.cpp b/Engine/Source/Runtime/SlateCore/Private/Fonts/FontCacheCompositeFont.cpp
--- a/Engine/Source/Runtime/SlateCore/Private/Fonts/FontCacheCompositeFont.cpp
+++ b/Engine/Source/Runtime/SlateCore/Private/Fonts/FontCacheCompositeFont.cpp
@@ -341,7 +341,7 @@ const FFontData& FCompositeFontCache::GetFontDataForCharacter(const FSlateFontIn
 
 TSharedPtr<FFreeTypeFace> FCompositeFontCache::GetFontFace(const FFontData& InFontData)
 {
-	LLM(ELLMTag::UI);
+	LLM_SCOPE(ELLMTag::UI);
 
 	TSharedPtr<FFreeTypeFace> FaceAndMemory = FontFaceMap.FindRef(InFontData);
 	if (!FaceAndMemory.IsValid() && InFontData.HasFont())
diff --git a/Engine/Source/Runtime/SlateCore/Private/Fonts/FontCacheFreeType.cpp b/Engine/Source/Runtime/SlateCore/Private/Fonts/FontCacheFreeType.cpp
--- a/Engine/Source/Runtime/SlateCore/Private/Fonts/FontCacheFreeType.cpp
+++ b/Engine/Source/Runtime/SlateCore/Private/Fonts/FontCacheFreeType.cpp
@@ -20,7 +20,7 @@ namespace FreeTypeMemory
 
 static void* Alloc(FT_Memory Memory, long Size)
 {
-	LLM(ELLMTag::UI);
+	LLM_SCOPE(ELLMTag::UI);
 	void* Result = FMemory::Malloc(Size);
 
 #if STATS
@@ -33,7 +33,7 @@ static void* Alloc(FT_Memory Memory, long Size)
 
 static void* Realloc(FT_Memory Memory, long CurSize, long NewSize, void* Block)
 {
-	LLM(ELLMTag::UI);
+	LLM_SCOPE(ELLMTag::UI);
 
 #if STATS
 	long DeltaNewSize = NewSize - CurSize;
diff --git a/Engine/Source/Runtime/SlateCore/Private/Fonts/FontCacheHarfBuzz.cpp b/Engine/Source/Runtime/SlateCore/Private/Fonts/FontCacheHarfBuzz.cpp
--- a/Engine/Source/Runtime/SlateCore/Private/Fonts/FontCacheHarfBuzz.cpp
+++ b/Engine/Source/Runtime/SlateCore/Private/Fonts/FontCacheHarfBuzz.cpp
@@ -13,13 +13,13 @@ extern "C"
 
 void* HarfBuzzMalloc(size_t InSizeBytes)
 {
-	LLM(ELLMTag::UI);
+	LLM_SCOPE(ELLMTag::UI);
 	return FMemory::Malloc(InSizeBytes);
 }
 
 void* HarfBuzzCalloc(size_t InNumItems, size_t InItemSizeBytes)
 {
-	LLM(ELLMTag::UI);
+	LLM_SCOPE(ELLMTag::UI);
 	const size_t AllocSizeBytes = InNumItems * InItemSizeBytes;
 	if (AllocSizeBytes > 0)
 	{
@@ -32,7 +32,7 @@ void* HarfBuzzCalloc(size_t InNumItems, size_t InItemSizeBytes)
 
 void* HarfBuzzRealloc(void* InPtr, size_t InSizeBytes)
 {
-	LLM(ELLMTag::UI);
+	LLM_SCOPE(ELLMTag::UI);
 	return FMemory::Realloc(InPtr, InSizeBytes);
 }
 
