Slate - Adding a simple setter for setting texture coordinates for slate vertices.

diff --git a/Engine/Source/Runtime/SlateCore/Public/Rendering/RenderingCommon.h b/Engine/Source/Runtime/SlateCore/Public/Rendering/RenderingCommon.h
--- a/Engine/Source/Runtime/SlateCore/Public/Rendering/RenderingCommon.h
+++ b/Engine/Source/Runtime/SlateCore/Public/Rendering/RenderingCommon.h
@@ -228,6 +228,14 @@ public:
 		return Vertex;
 	}
 
+	void SetTexCoords(const FVector4& InTexCoords)
+	{
+		TexCoords[0] = InTexCoords.X;
+		TexCoords[1] = InTexCoords.Y;
+		TexCoords[2] = InTexCoords.Z;
+		TexCoords[3] = InTexCoords.W;
+	}
+
 private:
 
 	template<ESlateVertexRounding Rounding>
