Slate - Adding the ability to stable sort TPanelChildren

diff --git a/Engine/Source/Runtime/SlateCore/Public/Layout/Children.h b/Engine/Source/Runtime/SlateCore/Public/Layout/Children.h
--- a/Engine/Source/Runtime/SlateCore/Public/Layout/Children.h
+++ b/Engine/Source/Runtime/SlateCore/Public/Layout/Children.h
@@ -393,6 +393,12 @@ public:
 		::Sort(TIndirectArray< SlotType >::GetData(), TIndirectArray<SlotType>::Num(), Predicate);
 	}
 
+	template <class PREDICATE_CLASS>
+	void StableSort(const PREDICATE_CLASS& Predicate)
+	{
+		::StableSort(TIndirectArray< SlotType >::GetData(), TIndirectArray< SlotType >::Num(), Predicate);
+	}
+
 	void Swap( int32 IndexA, int32 IndexB )
 	{
 		TIndirectArray< SlotType >::Swap(IndexA, IndexB);
