Stopping the slate debugger should stop the slate debugger.

diff --git a/Engine/Source/Runtime/SlateCore/Private/Debugging/ConsoleSlateDebugger.cpp b/Engine/Source/Runtime/SlateCore/Private/Debugging/ConsoleSlateDebugger.cpp
--- a/Engine/Source/Runtime/SlateCore/Private/Debugging/ConsoleSlateDebugger.cpp
+++ b/Engine/Source/Runtime/SlateCore/Private/Debugging/ConsoleSlateDebugger.cpp
@@ -31,7 +31,7 @@ FConsoleSlateDebugger::FConsoleSlateDebugger()
 	, StopDebuggingCommand(
 		TEXT("SlateDebugger.Stop"),
 		*LOCTEXT("StopDebugger", "Stops the debugger.").ToString(),
-		FConsoleCommandDelegate::CreateRaw(this, &FConsoleSlateDebugger::StartDebugging))
+		FConsoleCommandDelegate::CreateRaw(this, &FConsoleSlateDebugger::StopDebugging))
 	, CaptureStackVariable(
 		TEXT("SlateDebugger.CaptureStack"),
 		bCaptureStack,
