Adding another call option for GetWidgetDebugInfo that takes a reference.

diff --git a/Engine/Source/Runtime/SlateCore/Private/Types/ReflectionMetadata.cpp b/Engine/Source/Runtime/SlateCore/Private/Types/ReflectionMetadata.cpp
--- a/Engine/Source/Runtime/SlateCore/Private/Types/ReflectionMetadata.cpp
+++ b/Engine/Source/Runtime/SlateCore/Private/Types/ReflectionMetadata.cpp
@@ -9,8 +9,13 @@ FString FReflectionMetaData::GetWidgetDebugInfo(const SWidget* InWidget)
 		return TEXT("None");
 	}
 
+	return GetWidgetDebugInfo(*InWidget);
+}
+
+FString FReflectionMetaData::GetWidgetDebugInfo(const SWidget& InWidget)
+{
 	// UMG widgets have meta-data to help track them
-	TSharedPtr<FReflectionMetaData> MetaData = InWidget->GetMetaData<FReflectionMetaData>();
+	TSharedPtr<FReflectionMetaData> MetaData = InWidget.GetMetaData<FReflectionMetaData>();
 	if (MetaData.IsValid())
 	{
 		if (const UObject* AssetPtr = MetaData->Asset.Get())
@@ -22,7 +27,7 @@ FString FReflectionMetaData::GetWidgetDebugInfo(const SWidget* InWidget)
 		}
 	}
 
-	TSharedPtr<FReflectionMetaData> ParentMetadata = GetWidgetOrParentMetaData(InWidget);
+	TSharedPtr<FReflectionMetaData> ParentMetadata = GetWidgetOrParentMetaData(&InWidget);
 	if (ParentMetadata.IsValid())
 	{
 		if (const UObject* AssetPtr = ParentMetadata->Asset.Get())
@@ -30,11 +35,11 @@ FString FReflectionMetaData::GetWidgetDebugInfo(const SWidget* InWidget)
 			const FName AssetName = AssetPtr->GetFName();
 			const FName WidgetName = ParentMetadata->Name;
 
-			return FString::Printf(TEXT("%s [%s(%s)]"), *AssetName.ToString(), *WidgetName.ToString(), *InWidget->GetReadableLocation());
+			return FString::Printf(TEXT("%s [%s(%s)]"), *AssetName.ToString(), *WidgetName.ToString(), *InWidget.GetReadableLocation());
 		}
 	}
 
-	return InWidget->ToString();
+	return InWidget.ToString();
 }
 
 TSharedPtr<FReflectionMetaData> FReflectionMetaData::GetWidgetOrParentMetaData(const SWidget* InWidget)
diff --git a/Engine/Source/Runtime/SlateCore/Public/Types/ReflectionMetadata.h b/Engine/Source/Runtime/SlateCore/Public/Types/ReflectionMetadata.h
--- a/Engine/Source/Runtime/SlateCore/Public/Types/ReflectionMetadata.h
+++ b/Engine/Source/Runtime/SlateCore/Public/Types/ReflectionMetadata.h
@@ -41,6 +41,7 @@ public:
 public:
 
 	static FString GetWidgetDebugInfo(const SWidget* InWidget);
+	static FString GetWidgetDebugInfo(const SWidget& InWidget);
 
 	static TSharedPtr<FReflectionMetaData> GetWidgetOrParentMetaData(const SWidget* InWidget);
 };
