Making a few slate stats only trigger in WITH_VERY_VERBOSE_SLATE_STATS blocks to reduce stat spam overall.

diff --git a/Engine/Source/Runtime/SlateCore/Private/Widgets/SWidget.cpp b/Engine/Source/Runtime/SlateCore/Private/Widgets/SWidget.cpp
--- a/Engine/Source/Runtime/SlateCore/Private/Widgets/SWidget.cpp
+++ b/Engine/Source/Runtime/SlateCore/Private/Widgets/SWidget.cpp
@@ -480,7 +480,7 @@ void SWidget::SlatePrepass()
 
 void SWidget::SlatePrepass(float InLayoutScaleMultiplier)
 {
-#if SLATE_VERBOSE_NAMED_EVENTS
+#if WITH_VERY_VERBOSE_SLATE_STATS
 	SCOPED_NAMED_EVENT(SWidget_Prepass, FColor::Silver);
 #endif
 
@@ -1116,7 +1116,7 @@ float SWidget::GetRelativeLayoutScale(const FSlotBase& Child, float LayoutScaleM
 
 void SWidget::ArrangeChildren(const FGeometry& AllottedGeometry, FArrangedChildren& ArrangedChildren) const
 {
-#if SLATE_VERBOSE_NAMED_EVENTS
+#if WITH_VERY_VERBOSE_SLATE_STATS
 	SCOPED_NAMED_EVENT(SWidget_ArrangeChildren, FColor::Black);
 #endif
 	OnArrangeChildren(AllottedGeometry, ArrangedChildren);
