Fixing a bug with labeling the genesis of a navigation event in the slate console debugger.

diff --git a/Engine/Source/Runtime/SlateCore/Private/Debugging/ConsoleSlateDebugger.cpp b/Engine/Source/Runtime/SlateCore/Private/Debugging/ConsoleSlateDebugger.cpp
--- a/Engine/Source/Runtime/SlateCore/Private/Debugging/ConsoleSlateDebugger.cpp
+++ b/Engine/Source/Runtime/SlateCore/Private/Debugging/ConsoleSlateDebugger.cpp
@@ -239,20 +239,17 @@ void FConsoleSlateDebugger::OnFocusEvent(const FSlateDebuggingFocusEventArgs& Ev
 
 void FConsoleSlateDebugger::OnNavigationEvent(const FSlateDebuggingNavigationEventArgs& EventArgs)
 {
-	static const UEnum* UINavigationEnum = FindObjectChecked<UEnum>(ANY_PACKAGE, TEXT("EUINavigation"));
-	static const UEnum* NavigationGenesisEnum = FindObjectChecked<UEnum>(ANY_PACKAGE, TEXT("ENavigationGenesis"));
-
 	static const FText NavEventFormat = LOCTEXT("NavEventFormat", "Nav: {0}:{1} | {2} -> {3}");
 
 	const FText SourceWidget = FText::FromString(FReflectionMetaData::GetWidgetDebugInfo(&EventArgs.NavigationSource.GetLastWidget().Get()));
 	const FText DestinationWidget = FText::FromString(FReflectionMetaData::GetWidgetDebugInfo(EventArgs.DestinationWidget.Get()));
-	const FText NavigationTypeText = UINavigationEnum->GetDisplayNameTextByValue((int64)EventArgs.NavigationEvent.GetNavigationType());
-	const FText NavigationGenesisText = UINavigationEnum->GetDisplayNameTextByValue((int64)EventArgs.NavigationEvent.GetNavigationGenesis());
+	const FText NavigationTypeText = StaticEnum<EUINavigation>()->GetDisplayNameTextByValue((int64)EventArgs.NavigationEvent.GetNavigationType());
+	const FText NavigationGenesisText = StaticEnum<ENavigationGenesis>()->GetDisplayNameTextByValue((int64)EventArgs.NavigationEvent.GetNavigationGenesis());
 
 	FText EventText = FText::Format(
 		NavEventFormat,
-		NavigationTypeText,
 		NavigationGenesisText,
+		NavigationTypeText,
 		SourceWidget,
 		DestinationWidget
 	);
