Slate - Adding additional navigation information when we do perform navigation we report what system is performing the navigate, was it direct, a boundry condition, did it use the hittest grid...etc.

Slate - Also found a bug where we're not using the world switcher before calling the handler for some navigation options.

diff --git a/Engine/Source/Runtime/Slate/Private/Framework/Application/SlateApplication.cpp b/Engine/Source/Runtime/Slate/Private/Framework/Application/SlateApplication.cpp
--- a/Engine/Source/Runtime/Slate/Private/Framework/Application/SlateApplication.cpp
+++ b/Engine/Source/Runtime/Slate/Private/Framework/Application/SlateApplication.cpp
@@ -6484,19 +6484,40 @@ bool FSlateApplication::AttemptNavigation(const FWidgetPath& NavigationSource, c
 	TSharedPtr<SWidget> DestinationWidget = TSharedPtr<SWidget>();
 	bool bAlwaysHandleNavigationAttempt = false;
 
+#if WITH_SLATE_DEBUGGING
+	ESlateDebuggingNavigationMethod NavigationMethod = ESlateDebuggingNavigationMethod::Unknown;
+#endif
+
 	EUINavigation NavigationType = NavigationEvent.GetNavigationType();
 	if ( NavigationReply.GetBoundaryRule() == EUINavigationRule::Explicit )
 	{
 		DestinationWidget = NavigationReply.GetFocusRecipient();
 		bAlwaysHandleNavigationAttempt = true;
+
+#if WITH_SLATE_DEBUGGING
+		NavigationMethod = ESlateDebuggingNavigationMethod::Explicit;
+#endif
 	}
 	else if ( NavigationReply.GetBoundaryRule() == EUINavigationRule::Custom )
 	{
 		const FNavigationDelegate& FocusDelegate = NavigationReply.GetFocusDelegate();
 		if ( FocusDelegate.IsBound() )
 		{
+			// Switch worlds for widgets in the current path 
+			FScopedSwitchWorldHack SwitchWorld(NavigationSource);
+
 			DestinationWidget = FocusDelegate.Execute(NavigationType);
 			bAlwaysHandleNavigationAttempt = true;
+
+#if WITH_SLATE_DEBUGGING
+			NavigationMethod = ESlateDebuggingNavigationMethod::CustomDelegateBound;
+#endif
+		}
+		else
+		{
+#if WITH_SLATE_DEBUGGING
+			NavigationMethod = ESlateDebuggingNavigationMethod::CustomDelegateUnbound;
+#endif
 		}
 	}
 	else
@@ -6511,6 +6532,10 @@ bool FSlateApplication::AttemptNavigation(const FWidgetPath& NavigationSource, c
 			// Resolve the Widget Path
 			FArrangedWidget& NewFocusedArrangedWidget = NewFocusedWidgetPath.Widgets.Last();
 			DestinationWidget = NewFocusedArrangedWidget.Widget;
+
+#if WITH_SLATE_DEBUGGING
+			NavigationMethod = ESlateDebuggingNavigationMethod::NextOrPrevious;
+#endif
 		}
 		else
 		{
@@ -6521,11 +6546,15 @@ bool FSlateApplication::AttemptNavigation(const FWidgetPath& NavigationSource, c
 			FScopedSwitchWorldHack SwitchWorld(NavigationSource);
 
 			DestinationWidget = NavigationSource.GetWindow()->GetHittestGrid()->FindNextFocusableWidget(FocusedArrangedWidget, NavigationType, NavigationReply, BoundaryWidget);
+
+#if WITH_SLATE_DEBUGGING
+			NavigationMethod = ESlateDebuggingNavigationMethod::HitTestGrid;
+#endif
 		}
 	}
 
 #if WITH_SLATE_DEBUGGING
-	FSlateDebugging::BroadcastAttemptNavigation(NavigationEvent, NavigationReply, NavigationSource, DestinationWidget);
+	FSlateDebugging::BroadcastAttemptNavigation(NavigationEvent, NavigationReply, NavigationSource, DestinationWidget, NavigationMethod);
 #endif
 
 	return ExecuteNavigation(NavigationSource, DestinationWidget, NavigationEvent.GetUserIndex(), bAlwaysHandleNavigationAttempt);
@@ -6567,6 +6596,10 @@ bool FSlateApplication::ExecuteNavigation(const FWidgetPath& NavigationSource, T
 		}
 	}
 
+#if WITH_SLATE_DEBUGGING
+	FSlateDebugging::BroadcastExecuteNavigation();
+#endif
+
 	return bHandled;
 }
 
diff --git a/Engine/Source/Runtime/SlateCore/Private/Debugging/ConsoleSlateDebugger.cpp b/Engine/Source/Runtime/SlateCore/Private/Debugging/ConsoleSlateDebugger.cpp
--- a/Engine/Source/Runtime/SlateCore/Private/Debugging/ConsoleSlateDebugger.cpp
+++ b/Engine/Source/Runtime/SlateCore/Private/Debugging/ConsoleSlateDebugger.cpp
@@ -92,7 +92,7 @@ void FConsoleSlateDebugger::SetInputFilter(const TArray< FString >& Params)
 		return;
 	}
 
-	static const UEnum* SlateDebuggingInputEventEnum = FindObjectChecked<UEnum>(ANY_PACKAGE, TEXT("ESlateDebuggingInputEvent"));
+	static const UEnum* SlateDebuggingInputEventEnum = StaticEnum<ESlateDebuggingInputEvent>();
 	
 	const int64 InputEventEnumValue = SlateDebuggingInputEventEnum->GetValueByNameString(Params[0]);
 	if (InputEventEnumValue == INDEX_NONE)
@@ -115,7 +115,7 @@ void FConsoleSlateDebugger::RemoveListeners()
 	FSlateDebugging::Warning.RemoveAll(this);
 	FSlateDebugging::InputEvent.RemoveAll(this);
 	FSlateDebugging::FocusEvent.RemoveAll(this);
-	FSlateDebugging::NavigationEvent.RemoveAll(this);
+	FSlateDebugging::AttemptNavigationEvent.RemoveAll(this);
 	FSlateDebugging::MouseCaptureEvent.RemoveAll(this);
 #endif
 }
@@ -128,7 +128,8 @@ void FConsoleSlateDebugger::UpdateListeners()
 	FSlateDebugging::Warning.AddRaw(this, &FConsoleSlateDebugger::OnWarning);
 	FSlateDebugging::InputEvent.AddRaw(this, &FConsoleSlateDebugger::OnInputEvent);
 	FSlateDebugging::FocusEvent.AddRaw(this, &FConsoleSlateDebugger::OnFocusEvent);
-	FSlateDebugging::NavigationEvent.AddRaw(this, &FConsoleSlateDebugger::OnNavigationEvent);
+	FSlateDebugging::AttemptNavigationEvent.AddRaw(this, &FConsoleSlateDebugger::OnAttemptNavigationEvent);
+	FSlateDebugging::ExecuteNavigationEvent.AddRaw(this, &FConsoleSlateDebugger::OnExecuteNavigationEvent);
 	FSlateDebugging::MouseCaptureEvent.AddRaw(this, &FConsoleSlateDebugger::OnCaptureStateChangeEvent);
 #endif
 }
@@ -160,7 +161,7 @@ void FConsoleSlateDebugger::OnInputEvent(const FSlateDebuggingInputEventArgs& Ev
 
 	static const FText InputEventFormat = LOCTEXT("InputEventFormat", "{0} - ({1}) - [{2}]");
 
-	static const UEnum* SlateDebuggingInputEventEnum = FindObjectChecked<UEnum>(ANY_PACKAGE, TEXT("ESlateDebuggingInputEvent"));
+	static const UEnum* SlateDebuggingInputEventEnum = StaticEnum<ESlateDebuggingInputEvent>();
 	const FText InputEventTypeText = SlateDebuggingInputEventEnum->GetDisplayNameTextByValue((int64)EventArgs.InputEventType);
 	const FText AdditionalContent = FText::FromString(EventArgs.AdditionalContent);
 	const FText HandlerWidget = FText::FromString(FReflectionMetaData::GetWidgetDebugInfo(EventArgs.HandlerWidget.Get()));
@@ -237,21 +238,24 @@ void FConsoleSlateDebugger::OnFocusEvent(const FSlateDebuggingFocusEventArgs& Ev
 	OptionallyDumpCallStack();
 }
 
-void FConsoleSlateDebugger::OnNavigationEvent(const FSlateDebuggingNavigationEventArgs& EventArgs)
+void FConsoleSlateDebugger::OnAttemptNavigationEvent(const FSlateDebuggingNavigationEventArgs& EventArgs)
 {
-	static const FText NavEventFormat = LOCTEXT("NavEventFormat", "Nav: {0}:{1} | {2} -> {3}");
+	static const FText NavEventFormat = LOCTEXT("NavEventFormat", "Navigation User({4}) Source({0}:{1}) | {5} | {2} -> {3}");
 
 	const FText SourceWidget = FText::FromString(FReflectionMetaData::GetWidgetDebugInfo(&EventArgs.NavigationSource.GetLastWidget().Get()));
 	const FText DestinationWidget = FText::FromString(FReflectionMetaData::GetWidgetDebugInfo(EventArgs.DestinationWidget.Get()));
 	const FText NavigationTypeText = StaticEnum<EUINavigation>()->GetDisplayNameTextByValue((int64)EventArgs.NavigationEvent.GetNavigationType());
 	const FText NavigationGenesisText = StaticEnum<ENavigationGenesis>()->GetDisplayNameTextByValue((int64)EventArgs.NavigationEvent.GetNavigationGenesis());
+	const FText NavigationMethodText = StaticEnum<ESlateDebuggingNavigationMethod>()->GetDisplayNameTextByValue((int64)EventArgs.NavigationMethod);
 
 	FText EventText = FText::Format(
 		NavEventFormat,
 		NavigationGenesisText,
 		NavigationTypeText,
 		SourceWidget,
-		DestinationWidget
+		DestinationWidget,
+		EventArgs.NavigationEvent.GetUserIndex(),
+		NavigationMethodText
 	);
 
 	UE_LOG(LogSlateDebugger, Log, TEXT("%s"), *EventText.ToString());
@@ -259,6 +263,11 @@ void FConsoleSlateDebugger::OnNavigationEvent(const FSlateDebuggingNavigationEve
 	OptionallyDumpCallStack();
 }
 
+void FConsoleSlateDebugger::OnExecuteNavigationEvent(const FSlateDebuggingExecuteNavigationEventArgs& EventArgs)
+{
+	OptionallyDumpCallStack();
+}
+
 void FConsoleSlateDebugger::OnCaptureStateChangeEvent(const FSlateDebuggingMouseCaptureEventArgs& EventArgs)
 {
 	static const FText StateChangeEventFormat = LOCTEXT("StateChangeEventFormat", "{0}({1}:{2}) : {3}");
diff --git a/Engine/Source/Runtime/SlateCore/Private/Debugging/ConsoleSlateDebugger.h b/Engine/Source/Runtime/SlateCore/Private/Debugging/ConsoleSlateDebugger.h
--- a/Engine/Source/Runtime/SlateCore/Private/Debugging/ConsoleSlateDebugger.h
+++ b/Engine/Source/Runtime/SlateCore/Private/Debugging/ConsoleSlateDebugger.h
@@ -40,7 +40,8 @@ private:
 	void OnWarning(const FSlateDebuggingWarningEventArgs& EventArgs);
 	void OnInputEvent(const FSlateDebuggingInputEventArgs& EventArgs);
 	void OnFocusEvent(const FSlateDebuggingFocusEventArgs& EventArgs);
-	void OnNavigationEvent(const FSlateDebuggingNavigationEventArgs& EventArgs);
+	void OnAttemptNavigationEvent(const FSlateDebuggingNavigationEventArgs& EventArgs);
+	void OnExecuteNavigationEvent(const FSlateDebuggingExecuteNavigationEventArgs& EventArgs);
 	void OnCaptureStateChangeEvent(const FSlateDebuggingMouseCaptureEventArgs& EventArgs);
 
 	void OptionallyDumpCallStack();
diff --git a/Engine/Source/Runtime/SlateCore/Private/Debugging/SlateDebugging.cpp b/Engine/Source/Runtime/SlateCore/Private/Debugging/SlateDebugging.cpp
--- a/Engine/Source/Runtime/SlateCore/Private/Debugging/SlateDebugging.cpp
+++ b/Engine/Source/Runtime/SlateCore/Private/Debugging/SlateDebugging.cpp
@@ -33,11 +33,13 @@ FSlateDebuggingNavigationEventArgs::FSlateDebuggingNavigationEventArgs(
 	const FNavigationEvent& InNavigationEvent,
 	const FNavigationReply& InNavigationReply,
 	const FWidgetPath& InNavigationSource,
-	const TSharedPtr<SWidget>& InDestinationWidget)
+	const TSharedPtr<SWidget>& InDestinationWidget,
+	const ESlateDebuggingNavigationMethod InNavigationMethod)
 	: NavigationEvent(InNavigationEvent)
 	, NavigationReply(InNavigationReply)
 	, NavigationSource(InNavigationSource)
 	, DestinationWidget(InDestinationWidget)
+	, NavigationMethod(InNavigationMethod)
 {
 }
 
@@ -79,7 +81,9 @@ FSlateDebugging::FWidgetInputEvent FSlateDebugging::InputEvent;
 
 FSlateDebugging::FWidgetFocusEvent FSlateDebugging::FocusEvent;
 
-FSlateDebugging::FWidgetNavigationEvent FSlateDebugging::NavigationEvent;
+FSlateDebugging::FWidgetAttemptNavigationEvent FSlateDebugging::AttemptNavigationEvent;
+
+FSlateDebugging::FWidgetExecuteNavigationEvent FSlateDebugging::ExecuteNavigationEvent;
 
 FSlateDebugging::FWidgetMouseCaptureEvent FSlateDebugging::MouseCaptureEvent;
 
@@ -134,9 +138,14 @@ void FSlateDebugging::BroadcastFocusReceived(const FFocusEvent& InFocusEvent, co
 	FocusEvent.Broadcast(FSlateDebuggingFocusEventArgs(ESlateDebuggingFocusEvent::FocusReceived, InFocusEvent, InOldFocusedWidgetPath, InOldFocusedWidget, InNewFocusedWidgetPath, InNewFocusedWidget));
 }
 
-void FSlateDebugging::BroadcastAttemptNavigation(const FNavigationEvent& InNavigationEvent, const FNavigationReply& InNavigationReply, const FWidgetPath& InNavigationSource, const TSharedPtr<SWidget>& InDestinationWidget)
+void FSlateDebugging::BroadcastAttemptNavigation(const FNavigationEvent& InNavigationEvent, const FNavigationReply& InNavigationReply, const FWidgetPath& InNavigationSource, const TSharedPtr<SWidget>& InDestinationWidget, ESlateDebuggingNavigationMethod InNavigationMethod)
+{
+	AttemptNavigationEvent.Broadcast(FSlateDebuggingNavigationEventArgs(InNavigationEvent, InNavigationReply, InNavigationSource, InDestinationWidget, InNavigationMethod));
+}
+
+void FSlateDebugging::BroadcastExecuteNavigation()
 {
-	NavigationEvent.Broadcast(FSlateDebuggingNavigationEventArgs(InNavigationEvent, InNavigationReply, InNavigationSource, InDestinationWidget));
+	ExecuteNavigationEvent.Broadcast(FSlateDebuggingExecuteNavigationEventArgs());
 }
 
 void FSlateDebugging::BroadcastMouseCapture(uint32 UserIndex, uint32 PointerIndex, const TSharedPtr<SWidget>& InCapturingWidget)
diff --git a/Engine/Source/Runtime/SlateCore/Public/Debugging/SlateDebugging.h b/Engine/Source/Runtime/SlateCore/Public/Debugging/SlateDebugging.h
--- a/Engine/Source/Runtime/SlateCore/Public/Debugging/SlateDebugging.h
+++ b/Engine/Source/Runtime/SlateCore/Public/Debugging/SlateDebugging.h
@@ -56,6 +56,17 @@ enum class ESlateDebuggingStateChangeEvent : uint8
 	MouseCaptureLost,
 };
 
+UENUM()
+enum class ESlateDebuggingNavigationMethod : uint8
+{
+	Unknown,
+	Explicit,
+	CustomDelegateBound,
+	CustomDelegateUnbound,
+	NextOrPrevious,
+	HitTestGrid
+};
+
 struct SLATECORE_API FSlateDebuggingInputEventArgs
 {
 public:
@@ -102,13 +113,20 @@ public:
 		const FNavigationEvent& InNavigationEvent,
 		const FNavigationReply& InNavigationReply,
 		const FWidgetPath& InNavigationSource,
-		const TSharedPtr<SWidget>& InDestinationWidget
+		const TSharedPtr<SWidget>& InDestinationWidget,
+		const ESlateDebuggingNavigationMethod InNavigationMethod
 	);
 
 	const FNavigationEvent& NavigationEvent;
 	const FNavigationReply& NavigationReply;
 	const FWidgetPath& NavigationSource;
 	const TSharedPtr<SWidget>& DestinationWidget;
+	const ESlateDebuggingNavigationMethod NavigationMethod;
+};
+
+struct SLATECORE_API FSlateDebuggingExecuteNavigationEventArgs
+{
+public:
 };
 
 struct SLATECORE_API FSlateDebuggingWarningEventArgs
@@ -198,10 +216,16 @@ public:
 
 public:
 	/**  */
-	DECLARE_MULTICAST_DELEGATE_OneParam(FWidgetNavigationEvent, const FSlateDebuggingNavigationEventArgs& /*EventArgs*/);
-	static FWidgetNavigationEvent NavigationEvent;
+	DECLARE_MULTICAST_DELEGATE_OneParam(FWidgetAttemptNavigationEvent, const FSlateDebuggingNavigationEventArgs& /*EventArgs*/);
+	static FWidgetAttemptNavigationEvent AttemptNavigationEvent;
+
+	static void BroadcastAttemptNavigation(const FNavigationEvent& InNavigationEvent, const FNavigationReply& InNavigationReply, const FWidgetPath& InNavigationSource, const TSharedPtr<SWidget>& InDestinationWidget, ESlateDebuggingNavigationMethod InNavigationMethod);
+
+	/**  */
+	DECLARE_MULTICAST_DELEGATE_OneParam(FWidgetExecuteNavigationEvent, const FSlateDebuggingExecuteNavigationEventArgs& /*EventArgs*/);
+	static FWidgetExecuteNavigationEvent ExecuteNavigationEvent;
 
-	static void BroadcastAttemptNavigation(const FNavigationEvent& InNavigationEvent, const FNavigationReply& InNavigationReply, const FWidgetPath& InNavigationSource, const TSharedPtr<SWidget>& InDestinationWidget);
+	static void BroadcastExecuteNavigation();
 
 public:
 	/**  */
