Initialize AccessibleWidgetData in constructor instead of Slate args so that subclasses inherit properly. Fixes crash with non-accessible windows.

diff --git a/Engine/Source/Runtime/ApplicationCore/Private/Windows/Accessibility/WindowsUIAManager.cpp b/Engine/Source/Runtime/ApplicationCore/Private/Windows/Accessibility/WindowsUIAManager.cpp
--- a/Engine/Source/Runtime/ApplicationCore/Private/Windows/Accessibility/WindowsUIAManager.cpp
+++ b/Engine/Source/Runtime/ApplicationCore/Private/Windows/Accessibility/WindowsUIAManager.cpp
@@ -117,7 +117,9 @@ FWindowsUIAWindowProvider& FWindowsUIAManager::GetWindowProvider(TSharedRef<FWin
 		WindowsApplication.GetAccessibleMessageHandler()->SetActive(true);
 	}
 
-	return static_cast<FWindowsUIAWindowProvider&>(GetWidgetProvider(WindowsApplication.GetAccessibleMessageHandler()->GetAccessibleWindow(InWindow).ToSharedRef()));
+	TSharedPtr<IAccessibleWidget> AccessibleWindow = WindowsApplication.GetAccessibleMessageHandler()->GetAccessibleWindow(InWindow);
+	checkf(AccessibleWindow.IsValid(), TEXT("%s is not an accessible window. All windows must be accessible."), *InWindow->GetDefinition().Title);
+	return static_cast<FWindowsUIAWindowProvider&>(GetWidgetProvider(AccessibleWindow.ToSharedRef()));
 }
 
 void FWindowsUIAManager::OnWidgetProviderRemoved(TSharedRef<IAccessibleWidget> InWidget)
diff --git a/Engine/Source/Runtime/Slate/Private/Widgets/Input/SButton.cpp b/Engine/Source/Runtime/Slate/Private/Widgets/Input/SButton.cpp
--- a/Engine/Source/Runtime/Slate/Private/Widgets/Input/SButton.cpp
+++ b/Engine/Source/Runtime/Slate/Private/Widgets/Input/SButton.cpp
@@ -11,6 +11,13 @@
 
 static FName SButtonTypeName("SButton");
 
+SButton::SButton()
+{
+#if WITH_ACCESSIBILITY
+	AccessibleData = FAccessibleWidgetData(EAccessibleBehavior::Summary, EAccessibleBehavior::Auto, false);
+#endif
+}
+
 /**
  * Construct this widget
  *
diff --git a/Engine/Source/Runtime/Slate/Private/Widgets/Input/SCheckBox.cpp b/Engine/Source/Runtime/Slate/Private/Widgets/Input/SCheckBox.cpp
--- a/Engine/Source/Runtime/Slate/Private/Widgets/Input/SCheckBox.cpp
+++ b/Engine/Source/Runtime/Slate/Private/Widgets/Input/SCheckBox.cpp
@@ -12,6 +12,13 @@
 #include "Widgets/Accessibility/SlateAccessibleMessageHandler.h"
 #endif
 
+SCheckBox::SCheckBox()
+{
+#if WITH_ACCESSIBILITY
+	AccessibleData = FAccessibleWidgetData(EAccessibleBehavior::Summary, EAccessibleBehavior::Auto, false);
+#endif
+}
+
 /**
  * Construct this widget
  *
diff --git a/Engine/Source/Runtime/Slate/Private/Widgets/Input/SEditableText.cpp b/Engine/Source/Runtime/Slate/Private/Widgets/Input/SEditableText.cpp
--- a/Engine/Source/Runtime/Slate/Private/Widgets/Input/SEditableText.cpp
+++ b/Engine/Source/Runtime/Slate/Private/Widgets/Input/SEditableText.cpp
@@ -11,6 +11,9 @@
 
 SEditableText::SEditableText()
 {
+#if WITH_ACCESSIBILITY
+	AccessibleData = FAccessibleWidgetData(EAccessibleBehavior::Auto, EAccessibleBehavior::Auto, false);
+#endif
 }
 
 SEditableText::~SEditableText()
diff --git a/Engine/Source/Runtime/Slate/Private/Widgets/Input/SEditableTextBox.cpp b/Engine/Source/Runtime/Slate/Private/Widgets/Input/SEditableTextBox.cpp
--- a/Engine/Source/Runtime/Slate/Private/Widgets/Input/SEditableTextBox.cpp
+++ b/Engine/Source/Runtime/Slate/Private/Widgets/Input/SEditableTextBox.cpp
@@ -9,6 +9,13 @@
 #include "Widgets/Accessibility/SlateAccessibleWidgets.h"
 #endif
 
+SEditableTextBox::SEditableTextBox()
+{
+#if WITH_ACCESSIBILITY
+	AccessibleData = FAccessibleWidgetData(EAccessibleBehavior::Auto, EAccessibleBehavior::Auto, false);
+#endif
+}
+
 /**
  * Construct this widget
  *
@@ -86,11 +93,11 @@ void SEditableTextBox::Construct( const FArguments& InArgs )
 
 #if WITH_ACCESSIBILITY
 	// These functions need to be called again because SWidget::Construct is called before SEditableText is created.
-	if (!InArgs._AccessibleParams.AccessibleText.IsSet())
+	if (AccessibleData.AccessibleBehavior != EAccessibleBehavior::Custom)
 	{
 		SetDefaultAccessibleText(EAccessibleType::Main);
 	}
-	if (!InArgs._AccessibleParams.AccessibleSummaryText.IsSet())
+	if (AccessibleData.AccessibleSummaryBehavior != EAccessibleBehavior::Custom)
 	{
 		SetDefaultAccessibleText(EAccessibleType::Summary);
 	}
diff --git a/Engine/Source/Runtime/Slate/Private/Widgets/Input/SSlider.cpp b/Engine/Source/Runtime/Slate/Private/Widgets/Input/SSlider.cpp
--- a/Engine/Source/Runtime/Slate/Private/Widgets/Input/SSlider.cpp
+++ b/Engine/Source/Runtime/Slate/Private/Widgets/Input/SSlider.cpp
@@ -7,6 +7,13 @@
 #include "Widgets/Accessibility/SlateAccessibleWidgets.h"
 #endif
 
+SSlider::SSlider()
+{
+#if WITH_ACCESSIBILITY
+	AccessibleData = FAccessibleWidgetData(EAccessibleBehavior::Summary, EAccessibleBehavior::Auto, false);
+#endif
+}
+
 void SSlider::Construct( const SSlider::FArguments& InDeclaration )
 {
 	check(InDeclaration._Style);
diff --git a/Engine/Source/Runtime/Slate/Private/Widgets/Text/STextBlock.cpp b/Engine/Source/Runtime/Slate/Private/Widgets/Text/STextBlock.cpp
--- a/Engine/Source/Runtime/Slate/Private/Widgets/Text/STextBlock.cpp
+++ b/Engine/Source/Runtime/Slate/Private/Widgets/Text/STextBlock.cpp
@@ -22,6 +22,10 @@ STextBlock::STextBlock()
 	SetCanTick(false);
 	bCanSupportFocus = false;
 	bSimpleTextMode = false;
+
+#if WITH_ACCESSIBILITY
+	AccessibleData = FAccessibleWidgetData(EAccessibleBehavior::Auto, EAccessibleBehavior::Auto, false);
+#endif
 }
 
 STextBlock::~STextBlock()
diff --git a/Engine/Source/Runtime/Slate/Public/Widgets/Input/SButton.h b/Engine/Source/Runtime/Slate/Public/Widgets/Input/SButton.h
--- a/Engine/Source/Runtime/Slate/Public/Widgets/Input/SButton.h
+++ b/Engine/Source/Runtime/Slate/Public/Widgets/Input/SButton.h
@@ -49,7 +49,6 @@ public:
 		, _ForegroundColor( FCoreStyle::Get().GetSlateColor( "InvertedForeground" ) )
 		, _IsFocusable( true )
 		{
-			_AccessibleParams = FAccessibleWidgetData(EAccessibleBehavior::Summary, EAccessibleBehavior::Auto, false);
 		}
 
 		/** Slot for this button's content (optional) */
@@ -120,6 +119,7 @@ public:
 
 	SLATE_END_ARGS()
 
+	SButton();
 
 	/** @return An image that represents this button's border*/
 	virtual const FSlateBrush* GetBorder() const;
diff --git a/Engine/Source/Runtime/Slate/Public/Widgets/Input/SCheckBox.h b/Engine/Source/Runtime/Slate/Public/Widgets/Input/SCheckBox.h
--- a/Engine/Source/Runtime/Slate/Public/Widgets/Input/SCheckBox.h
+++ b/Engine/Source/Runtime/Slate/Public/Widgets/Input/SCheckBox.h
@@ -53,7 +53,6 @@ public:
 		, _UndeterminedHoveredImage( nullptr )
 		, _UndeterminedPressedImage( nullptr )
 	{
-		_AccessibleParams = FAccessibleWidgetData(EAccessibleBehavior::Summary, EAccessibleBehavior::Auto, false);
 	}
 
 		/** Content to be placed next to the check box, or for a toggle button, the content to be placed inside the button */
@@ -131,6 +130,7 @@ public:
 
 	SLATE_END_ARGS()
 
+	SCheckBox();
 
 	/**
 	 * Construct this widget
diff --git a/Engine/Source/Runtime/Slate/Public/Widgets/Input/SEditableText.h b/Engine/Source/Runtime/Slate/Public/Widgets/Input/SEditableText.h
--- a/Engine/Source/Runtime/Slate/Public/Widgets/Input/SEditableText.h
+++ b/Engine/Source/Runtime/Slate/Public/Widgets/Input/SEditableText.h
@@ -65,7 +65,6 @@ public:
 		, _TextFlowDirection()
 		{
 			_Clipping = EWidgetClipping::ClipToBounds;
-			_AccessibleParams = FAccessibleWidgetData(EAccessibleBehavior::Auto, EAccessibleBehavior::Auto, false);
 		}
 
 		/** Sets the text content for this editable text widget */
diff --git a/Engine/Source/Runtime/Slate/Public/Widgets/Input/SEditableTextBox.h b/Engine/Source/Runtime/Slate/Public/Widgets/Input/SEditableTextBox.h
--- a/Engine/Source/Runtime/Slate/Public/Widgets/Input/SEditableTextBox.h
+++ b/Engine/Source/Runtime/Slate/Public/Widgets/Input/SEditableTextBox.h
@@ -53,7 +53,6 @@ public:
 		, _VirtualKeyboardTrigger(EVirtualKeyboardTrigger::OnFocusByPointer)
 		, _VirtualKeyboardDismissAction(EVirtualKeyboardDismissAction::TextChangeOnDismiss)
 		{
-			_AccessibleParams = FAccessibleWidgetData(EAccessibleBehavior::Auto, EAccessibleBehavior::Auto, false);
 		}
 
 		/** The styling of the textbox */
@@ -150,6 +149,8 @@ public:
 		SLATE_ARGUMENT(TOptional<ETextFlowDirection>, TextFlowDirection)
 
 	SLATE_END_ARGS()
+
+	SEditableTextBox();
 	
 	/**
 	 * Construct this widget
diff --git a/Engine/Source/Runtime/Slate/Public/Widgets/Input/SSlider.h b/Engine/Source/Runtime/Slate/Public/Widgets/Input/SSlider.h
--- a/Engine/Source/Runtime/Slate/Public/Widgets/Input/SSlider.h
+++ b/Engine/Source/Runtime/Slate/Public/Widgets/Input/SSlider.h
@@ -42,7 +42,6 @@ public:
 		, _OnMouseCaptureEnd()
 		, _OnValueChanged()
 		{
-			_AccessibleParams = FAccessibleWidgetData(EAccessibleBehavior::Summary, EAccessibleBehavior::Auto, false);
 		}
 
 		/** Whether the slidable area should be indented to fit the handle. */
@@ -100,6 +99,8 @@ public:
 
 	SLATE_END_ARGS()
 
+	SSlider();
+
 	/**
 	 * Construct the widget.
 	 * 
diff --git a/Engine/Source/Runtime/Slate/Public/Widgets/Text/STextBlock.h b/Engine/Source/Runtime/Slate/Public/Widgets/Text/STextBlock.h
--- a/Engine/Source/Runtime/Slate/Public/Widgets/Text/STextBlock.h
+++ b/Engine/Source/Runtime/Slate/Public/Widgets/Text/STextBlock.h
@@ -65,7 +65,6 @@ public:
 		, _SimpleTextMode(false)
 		{
 			_Clipping = EWidgetClipping::OnDemand;
-			_AccessibleParams = FAccessibleWidgetData(EAccessibleBehavior::Auto, EAccessibleBehavior::Auto, false);
 		}
 
 		/** The text displayed in this text block */
diff --git a/Engine/Source/Runtime/SlateCore/Private/Widgets/SWidget.cpp b/Engine/Source/Runtime/SlateCore/Private/Widgets/SWidget.cpp
--- a/Engine/Source/Runtime/SlateCore/Private/Widgets/SWidget.cpp
+++ b/Engine/Source/Runtime/SlateCore/Private/Widgets/SWidget.cpp
@@ -168,7 +168,7 @@ void SWidget::Construct(
 	const bool InForceVolatile,
 	const EWidgetClipping InClipping,
 	const EFlowDirectionPreference InFlowPreference,
-	const FAccessibleWidgetData& InAccessibleData,
+	const TOptional<FAccessibleWidgetData>& InAccessibleData,
 	const TArray<TSharedRef<ISlateMetaData>>& InMetaData
 )
 {
@@ -201,24 +201,27 @@ void SWidget::Construct(
 	MetaData = InMetaData;
 
 #if WITH_ACCESSIBILITY
-	AccessibleData.bCanChildrenBeAccessible = InAccessibleData.bCanChildrenBeAccessible;
-	// If custom text is provided, force behavior to custom. Otherwise, use the passed-in behavior and set their default text.
-	if (!InAccessibleData.AccessibleText.IsSet())
+	if (InAccessibleData.IsSet())
+	{
+		SetCanChildrenBeAccessible(InAccessibleData->bCanChildrenBeAccessible);
+		// If custom text is provided, force behavior to custom. Otherwise, use the passed-in behavior and set their default text.
+		SetAccessibleBehavior(InAccessibleData->AccessibleText.IsSet() ? EAccessibleBehavior::Custom : InAccessibleData->AccessibleBehavior, InAccessibleData->AccessibleText, EAccessibleType::Main);
+		SetAccessibleBehavior(InAccessibleData->AccessibleSummaryText.IsSet() ? EAccessibleBehavior::Custom : InAccessibleData->AccessibleSummaryBehavior, InAccessibleData->AccessibleSummaryText, EAccessibleType::Summary);
+	}
+	if (AccessibleData.AccessibleBehavior != EAccessibleBehavior::Custom)
 	{
 		SetDefaultAccessibleText(EAccessibleType::Main);
 	}
-	if (!InAccessibleData.AccessibleSummaryText.IsSet())
+	if (AccessibleData.AccessibleSummaryBehavior != EAccessibleBehavior::Custom)
 	{
 		SetDefaultAccessibleText(EAccessibleType::Summary);
 	}
-	SetAccessibleBehavior(InAccessibleData.AccessibleText.IsSet() ? EAccessibleBehavior::Custom : InAccessibleData.AccessibleBehavior, InAccessibleData.AccessibleText, EAccessibleType::Main);
-	SetAccessibleBehavior(InAccessibleData.AccessibleSummaryText.IsSet() ? EAccessibleBehavior::Custom : InAccessibleData.AccessibleSummaryBehavior, InAccessibleData.AccessibleSummaryText, EAccessibleType::Summary);
 #endif
 }
 
 void SWidget::SWidgetConstruct(const TAttribute<FText>& InToolTipText, const TSharedPtr<IToolTip>& InToolTip, const TAttribute< TOptional<EMouseCursor::Type> >& InCursor, const TAttribute<bool>& InEnabledState,
 							   const TAttribute<EVisibility>& InVisibility, const float InRenderOpacity, const TAttribute<TOptional<FSlateRenderTransform>>& InTransform, const TAttribute<FVector2D>& InTransformPivot,
-							   const FName& InTag, const bool InForceVolatile, const EWidgetClipping InClipping, const EFlowDirectionPreference InFlowPreference, const FAccessibleWidgetData& InAccessibleData,
+							   const FName& InTag, const bool InForceVolatile, const EWidgetClipping InClipping, const EFlowDirectionPreference InFlowPreference, const TOptional<FAccessibleWidgetData>& InAccessibleData,
 							   const TArray<TSharedRef<ISlateMetaData>>& InMetaData)
 {
 	Construct(InToolTipText, InToolTip, InCursor, InEnabledState, InVisibility, InRenderOpacity, InTransform, InTransformPivot, InTag, InForceVolatile, InClipping, InFlowPreference, InAccessibleData, InMetaData);
diff --git a/Engine/Source/Runtime/SlateCore/Private/Widgets/SWindow.cpp b/Engine/Source/Runtime/SlateCore/Private/Widgets/SWindow.cpp
--- a/Engine/Source/Runtime/SlateCore/Private/Widgets/SWindow.cpp
+++ b/Engine/Source/Runtime/SlateCore/Private/Widgets/SWindow.cpp
@@ -1919,6 +1919,9 @@ SWindow::SWindow()
 	, TitleBar()
 	, bIsDrawingEnabled( true )
 {
+#if WITH_ACCESSIBILITY
+	AccessibleData = FAccessibleWidgetData(EAccessibleBehavior::Auto);
+#endif
 }
 
 
diff --git a/Engine/Source/Runtime/SlateCore/Public/Widgets/DeclarativeSyntaxSupport.h b/Engine/Source/Runtime/SlateCore/Public/Widgets/DeclarativeSyntaxSupport.h
--- a/Engine/Source/Runtime/SlateCore/Public/Widgets/DeclarativeSyntaxSupport.h
+++ b/Engine/Source/Runtime/SlateCore/Public/Widgets/DeclarativeSyntaxSupport.h
@@ -835,7 +835,7 @@ struct TSlateBaseNamedArgs
 	SLATE_ARGUMENT( bool, ForceVolatile )
 	SLATE_ARGUMENT( EWidgetClipping, Clipping )
 	SLATE_ARGUMENT( EFlowDirectionPreference, FlowDirectionPreference)
-	SLATE_ARGUMENT(FAccessibleWidgetData, AccessibleParams)
+	SLATE_ARGUMENT(TOptional<FAccessibleWidgetData>, AccessibleParams)
 	SLATE_ATTRIBUTE(FText, AccessibleText)
 
 	TArray<TSharedRef<ISlateMetaData>> MetaData;
diff --git a/Engine/Source/Runtime/SlateCore/Public/Widgets/SWidget.h b/Engine/Source/Runtime/SlateCore/Public/Widgets/SWidget.h
--- a/Engine/Source/Runtime/SlateCore/Public/Widgets/SWidget.h
+++ b/Engine/Source/Runtime/SlateCore/Public/Widgets/SWidget.h
@@ -291,7 +291,7 @@ public:
 		const bool InForceVolatile,
 		const EWidgetClipping InClipping,
 		const EFlowDirectionPreference InFlowPreference,
-		const FAccessibleWidgetData& InAccessibleData,
+		const TOptional<FAccessibleWidgetData>& InAccessibleData,
 		const TArray<TSharedRef<ISlateMetaData>>& InMetaData);
 
 	void SWidgetConstruct(const TAttribute<FText>& InToolTipText,
@@ -306,7 +306,7 @@ public:
 		const bool InForceVolatile,
 		const EWidgetClipping InClipping,
 		const EFlowDirectionPreference InFlowPreference,
-		const FAccessibleWidgetData& InAccessibleData,
+		const TOptional<FAccessibleWidgetData>& InAccessibleData,
 		const TArray<TSharedRef<ISlateMetaData>>& InMetaData);
 
 	//
diff --git a/Engine/Source/Runtime/SlateCore/Public/Widgets/SWindow.h b/Engine/Source/Runtime/SlateCore/Public/Widgets/SWindow.h
--- a/Engine/Source/Runtime/SlateCore/Public/Widgets/SWindow.h
+++ b/Engine/Source/Runtime/SlateCore/Public/Widgets/SWindow.h
@@ -158,7 +158,6 @@ public:
 		, _bManualManageDPI( false )
 
 	{
-		_AccessibleParams = FAccessibleWidgetData(EAccessibleBehavior::Auto);
 	}
 
 		/** Type of this window */
