Slate - Disabling the desired size invalidating when prepass is invalidating, it's causing a feedback loop, and might be causing an explosion in the font page filling, and jittering text.

diff --git a/Engine/Source/Runtime/SlateCore/Private/Widgets/SWidget.cpp b/Engine/Source/Runtime/SlateCore/Private/Widgets/SWidget.cpp
--- a/Engine/Source/Runtime/SlateCore/Private/Widgets/SWidget.cpp
+++ b/Engine/Source/Runtime/SlateCore/Private/Widgets/SWidget.cpp
@@ -525,7 +525,8 @@ void SWidget::SlatePrepass(float InLayoutScaleMultiplier)
 
 		// If the scale changed, that can affect the desired size of some elements that take it into
 		// account, such as text, so when the prepass size changes, so must we invalidate desired size.
-		bNeedsDesiredSize = true;
+		//bNeedsDesiredSize = true;
+		// Disabled: This was causing instability in the feedback loop created by the scale box
 	}
 
 	if ( bCanHaveChildren )
