Fix for font material not applied correctly on rich text (because SlateTextRun would reuse a cached ShapedGlyphSequence with the wrong FontMaterial because the CachedShapedTextKey was ignoring it).
Similar fix for outline material and color.

diff --git a/Engine/Source/Runtime/Slate/Public/Framework/Text/ShapedTextCache.h b/Engine/Source/Runtime/Slate/Public/Framework/Text/ShapedTextCache.h
--- a/Engine/Source/Runtime/Slate/Public/Framework/Text/ShapedTextCache.h
+++ b/Engine/Source/Runtime/Slate/Public/Framework/Text/ShapedTextCache.h
@@ -26,7 +26,7 @@ public:
 		return TextRange == Other.TextRange
 			&& Scale == Other.Scale
 			&& TextContext == Other.TextContext
-			&& FontInfo.IsIdentialToForCaching(Other.FontInfo);
+			&& FontInfo.IsIdenticalTo(Other.FontInfo);
 	}
 
 	FORCEINLINE bool operator!=(const FCachedShapedTextKey& Other) const
diff --git a/Engine/Source/Runtime/SlateCore/Public/Fonts/SlateFontInfo.h b/Engine/Source/Runtime/SlateCore/Public/Fonts/SlateFontInfo.h
--- a/Engine/Source/Runtime/SlateCore/Public/Fonts/SlateFontInfo.h
+++ b/Engine/Source/Runtime/SlateCore/Public/Fonts/SlateFontInfo.h
@@ -78,6 +78,7 @@ struct SLATECORE_API FFontOutlineSettings
 
 	inline bool IsIdenticalToForCaching(const FFontOutlineSettings& Other) const
 	{
+		// Ignore OutlineMaterial && OutlineColor because they do not affect the cached glyph.
 		return OutlineSize == Other.OutlineSize
 			&&  bSeparateFillAlpha == Other.bSeparateFillAlpha;
 	}
@@ -85,10 +86,11 @@ struct SLATECORE_API FFontOutlineSettings
 	inline bool IsIdenticalTo(const FFontOutlineSettings& Other) const
 	{
 		return
-			// Ignore OutlineMaterial && OutlineColor because they do not affect the cached glyph.
 			OutlineSize == Other.OutlineSize &&
-			bSeparateFillAlpha == Other.bSeparateFillAlpha;
-			
+			bSeparateFillAlpha == Other.bSeparateFillAlpha &&
+			bApplyOutlineToDropShadows == Other.bApplyOutlineToDropShadows &&
+			OutlineMaterial == Other.OutlineMaterial &&
+			OutlineColor == Other.OutlineColor;
 	}
 
 	friend inline uint32 GetTypeHash(const FFontOutlineSettings& OutlineSettings)
@@ -97,7 +99,6 @@ struct SLATECORE_API FFontOutlineSettings
 		// Ignore OutlineMaterial && OutlineColor because they do not affect the cached glyph.
 		Hash = HashCombine(Hash, GetTypeHash(OutlineSettings.OutlineSize));
 		Hash = HashCombine(Hash, GetTypeHash(OutlineSettings.bSeparateFillAlpha));
-
 		return Hash;
 	}
 
@@ -239,6 +240,7 @@ public:
 public:
 	inline bool IsIdentialToForCaching(const FSlateFontInfo& Other) const
 	{
+		// Ignore FontMaterial because it does not affect the cached glyph.
 		return FontObject == Other.FontObject
 			&& OutlineSettings.IsIdenticalToForCaching(Other.OutlineSettings)
 			&& CompositeFont == Other.CompositeFont
@@ -282,9 +284,9 @@ public:
 	 */
 	friend inline uint32 GetTypeHash( const FSlateFontInfo& FontInfo )
 	{
+		// Ignore FontMaterial because it does not affect the cached glyph.
 		uint32 Hash = 0;
 		Hash = HashCombine(Hash, GetTypeHash(FontInfo.FontObject));
-		Hash = HashCombine(Hash, GetTypeHash(FontInfo.FontMaterial));
 		Hash = HashCombine(Hash, GetTypeHash(FontInfo.OutlineSettings));
 		Hash = HashCombine(Hash, GetTypeHash(FontInfo.CompositeFont));
 		Hash = HashCombine(Hash, GetTypeHash(FontInfo.TypefaceFontName));
