Slate Global Invalidation System

diff --git a/Engine/Source/Runtime/SlateCore/Private/Input/HittestGrid.cpp b/Engine/Source/Runtime/SlateCore/Private/Input/HittestGrid.cpp
--- a/Engine/Source/Runtime/SlateCore/Private/Input/HittestGrid.cpp
+++ b/Engine/Source/Runtime/SlateCore/Private/Input/HittestGrid.cpp
@@ -4,9 +4,15 @@
 #include "Rendering/RenderingCommon.h"
 #include "SlateGlobals.h"
 #include "HAL/IConsoleManager.h"
+#include "Rendering/DrawElements.h"
+#include "Styling/CoreStyle.h"
 
 DEFINE_LOG_CATEGORY_STATIC(LogHittestDebug, Display, All);
 
+DECLARE_CYCLE_STAT(TEXT("HitTestGrid AddWidget"), STAT_SlateHTG_AddWidget, STATGROUP_Slate);
+DECLARE_CYCLE_STAT(TEXT("HitTestGrid RemoveWidget"), STAT_SlateHTG_RemoveWidget, STATGROUP_Slate);
+DECLARE_CYCLE_STAT(TEXT("HitTestGrid Clear"), STAT_SlateHTG_Clear, STATGROUP_Slate);
+
 int32 SlateVerifyHitTestVisibility = 0;
 static FAutoConsoleVariableRef CVarSlateVerifyHitTestVisibility(TEXT("Slate.VerifyHitTestVisibility"), SlateVerifyHitTestVisibility, TEXT("Should we double check the visibility of widgets during hit testing, in case previously resolved hit tests that same frame may have changed state?"), ECVF_Default);
 
@@ -73,12 +79,12 @@ bool ContainsInteractableWidget(const TArray<FWidgetAndPointer>& PathToTest)
 	return false;
 }
 
-//
-// FHittestGrid
-//
 
 const FVector2D CellSize(128.0f, 128.0f);
 
+//
+// FHittestGrid::FGridTestingParams
+//
 struct FHittestGrid::FGridTestingParams
 {
 	/** Ctor */
@@ -95,13 +101,65 @@ struct FHittestGrid::FGridTestingParams
 	bool bTestWidgetIsInteractive;
 };
 
-FHittestGrid::FHittestGrid()
-: WidgetsCachedThisFrame()
-, NumCells(0, 0)
-, NumCellsExcess(0, 0)
+//
+// FHittestGrid::FCell
+//
+
+FHittestGrid::FCell::FCell()
+	: bRequiresSort(false)
+{
+
+}
+
+void FHittestGrid::FCell::AddIndex(int32 WidgetIndex)
+{
+	CachedWidgetIndexes.Add(WidgetIndex);
+	bRequiresSort = true;
+}
+
+void FHittestGrid::FCell::RemoveIndex(int32 WidgetIndex)
+{
+	for (int32 CallArrayIndex = 0; CallArrayIndex < CachedWidgetIndexes.Num(); CallArrayIndex++)
+	{
+		if (CachedWidgetIndexes[CallArrayIndex] == WidgetIndex)
+		{
+			CachedWidgetIndexes.RemoveAt(CallArrayIndex);
+			break;
+		}
+	}
+}
+
+void FHittestGrid::FCell::Sort(const TSparseArray<FWidgetData>& InWidgetArray)
+{
+	if (bRequiresSort)
+	{
+		CachedWidgetIndexes.StableSort([&InWidgetArray](const int32& A, const int32& B)
+		{
+			return InWidgetArray[A].PrimarySort < InWidgetArray[B].PrimarySort || (InWidgetArray[A].PrimarySort == InWidgetArray[B].PrimarySort && InWidgetArray[A].SecondarySort < InWidgetArray[B].SecondarySort);
+		});
+
+		bRequiresSort = false;
+	}
+}
+
+const TArray<int32>& FHittestGrid::FCell::GetCachedWidgetIndexes() const
 {
+	return CachedWidgetIndexes;
 }
 
+//
+// FHittestGrid
+//
+
+FHittestGrid::FHittestGrid()
+	: WidgetMap()
+	, WidgetArray()
+	, Cells()
+	, NumCells(0, 0)
+	, GridOrigin(0, 0)
+	, GridSize(0, 0)
+{
+}
 
 FHittestGrid::~FHittestGrid()
 {
@@ -109,244 +167,151 @@ FHittestGrid::~FHittestGrid()
 
 TArray<FWidgetAndPointer> FHittestGrid::GetBubblePath(FVector2D DesktopSpaceCoordinate, float CursorRadius, bool bIgnoreEnabledStatus)
 {
-	//grab the radius, and if we are non-zero
-	const bool bDirectTestingOnly = (CursorRadius <= 0.0f);
-
-	//calculate the cursor position in the grid
+	checkSlow(IsInGameThread());
+	CursorRadius = 00.0f;
 	const FVector2D CursorPositionInGrid = DesktopSpaceCoordinate - GridOrigin;
 
-	if (WidgetsCachedThisFrame.Num() > 0 && Cells.Num() > 0)
+	if (WidgetArray.Num() > 0 && Cells.Num() > 0)
 	{
-		//grab the path for direct testing first
-		FGridTestingParams DirectTestingParams;
-		DirectTestingParams.CursorPositionInGrid = CursorPositionInGrid;
-		DirectTestingParams.CellCoord = GetCellCoordinate(CursorPositionInGrid);
-		DirectTestingParams.Radius = 0.0f;
-		DirectTestingParams.bTestWidgetIsInteractive = false;
+		TArray<FIndexAndDistance, TInlineAllocator<9>> Hits;
 
-		const FWidgetPathAndDist DirectBubblePathInfo = GetWidgetPathAndDist(DirectTestingParams, bIgnoreEnabledStatus);
+		FGridTestingParams TestingParams;
+		TestingParams.CursorPositionInGrid = CursorPositionInGrid;
+		TestingParams.CellCoord = GetCellCoordinate(CursorPositionInGrid);
+		TestingParams.Radius = 0.0f;
+		TestingParams.bTestWidgetIsInteractive = false;
 
-		//if we aren't doing a radius check, or we already have a direct path, use that
-		if (bDirectTestingOnly || ContainsInteractableWidget(DirectBubblePathInfo.BubblePath))
-		{
-			return DirectBubblePathInfo.BubblePath;
-		}
+		// First add the exact point test results
+		Hits.Add(GetHitIndexFromCellIndex(TestingParams));
 
-		//if we are here, we need to check other cells
-		const FVector2D RadiusVector(CursorRadius, CursorRadius);
-		const FIntPoint ULIndex = GetCellCoordinate(CursorPositionInGrid - RadiusVector);
-		const FIntPoint LRIndex = GetCellCoordinate(CursorPositionInGrid + RadiusVector);
+		// If we have a cursor radius greater than a pixel in size add any extra hit possibilities.
+		if (false) //CursorRadius >= 0.5f)
+		{
+			TestingParams.Radius = CursorRadius;
+			TestingParams.bTestWidgetIsInteractive = true;
 
-		//first, find all the overlapping cells
-		TArray<FIntPoint, TInlineAllocator<16>> CellIndexes;
+			const FVector2D RadiusVector(CursorRadius, CursorRadius);
+			const FIntPoint ULIndex = GetCellCoordinate(CursorPositionInGrid - RadiusVector);
+			const FIntPoint LRIndex = GetCellCoordinate(CursorPositionInGrid + RadiusVector);
 
-		for (int32 YIndex = ULIndex.Y; YIndex <= LRIndex.Y; ++YIndex)
-		{
-			for (int32 XIndex = ULIndex.X; XIndex <= LRIndex.X; ++XIndex)
+			for (int32 YIndex = ULIndex.Y; YIndex <= LRIndex.Y; ++YIndex)
 			{
-				const FIntPoint PointToTest(XIndex, YIndex);
-				if (IsValidCellCoord(PointToTest))
+				for (int32 XIndex = ULIndex.X; XIndex <= LRIndex.X; ++XIndex)
 				{
-					CellIndexes.Add(PointToTest);
+					const FIntPoint PointToTest(XIndex, YIndex);
+					if (IsValidCellCoord(PointToTest))
+					{
+						TestingParams.CellCoord = PointToTest;
+						Hits.Add(GetHitIndexFromCellIndex(TestingParams));
+					}
 				}
 			}
-		}
-
-		//setup the radius testing params
-		FGridTestingParams RadiusTestingParams;
-		RadiusTestingParams.CursorPositionInGrid = CursorPositionInGrid;
-		RadiusTestingParams.Radius = CursorRadius;
-		RadiusTestingParams.bTestWidgetIsInteractive = true;
 
-		//next, collect valid paths from all those cells
-		TArray<FWidgetPathAndDist> PathsAndDistances;
-		for (const FIntPoint& CellCoord : CellIndexes)
-		{
-			//update the cell coord property
-			RadiusTestingParams.CellCoord = CellCoord;
-
-			const FWidgetPathAndDist TestPath = GetWidgetPathAndDist(RadiusTestingParams, bIgnoreEnabledStatus);
-			if ( TestPath.IsValidPath() )
+			// sort the paths, and get the closest one that is valid
+			Hits.StableSort([](const FIndexAndDistance& A, const FIndexAndDistance& B)
 			{
-				PathsAndDistances.Add(TestPath);
-			}
+				return A.WidgetIndex != INDEX_NONE && A.DistanceSqToWidget < B.DistanceSqToWidget;
+			});
 		}
 
-		// We have paths from multiple cells; use the one that is closest to the cursor's center.
-		if (PathsAndDistances.Num() > 0)
+		const FIndexAndDistance& BestHit = Hits[0];
+		if (BestHit.WidgetIndex != INDEX_NONE)
 		{
-			//sort the paths, and get the closest one that is valid
-			PathsAndDistances.Sort([](const FWidgetPathAndDist& A, const FWidgetPathAndDist& B)
-			{
-				return A.DistToTopWidgetSq < B.DistToTopWidgetSq;
-			});
+			const FWidgetData& WidgetData = WidgetArray[BestHit.WidgetIndex];
 
-			const FWidgetPathAndDist* BestCandidateByDistance = nullptr;
-			for (const FWidgetPathAndDist& TestPath : PathsAndDistances)
+			const TSharedPtr<SWidget> FirstHitWidget = WidgetData.GetWidget();
+			if (FirstHitWidget.IsValid()) // Make Sure we landed on a valid widget
 			{
-				if (ContainsInteractableWidget(TestPath.BubblePath))
+				TArray<FWidgetAndPointer> Path;
+
+				TSharedPtr<SWidget> CurWidget = FirstHitWidget;
+				while (CurWidget.IsValid())
 				{
-					BestCandidateByDistance = &TestPath;
-					break;
+					FGeometry DesktopSpaceGeometry = CurWidget->GetPaintSpaceGeometry();
+					DesktopSpaceGeometry.AppendTransform(FSlateLayoutTransform(GridOrigin - GridWindowOrigin));
+
+					Path.Emplace(FArrangedWidget(CurWidget.ToSharedRef(), DesktopSpaceGeometry), TSharedPtr<FVirtualPointerPosition>());
+					CurWidget = CurWidget->Advanced_GetPaintParentWidget();
 				}
-			}
 
-			if (BestCandidateByDistance != nullptr)
-			{
-				const FGeometry& LeafmostGeometry = BestCandidateByDistance->BubblePath.Last().Geometry;
-				const int32 LeafmostLayerId = BestCandidateByDistance->LayerId;
-				if (DirectBubblePathInfo.LayerId > LeafmostLayerId)
+				if (!Path.Last().Widget->Advanced_IsWindow())
 				{
-					return DirectBubblePathInfo.BubblePath;
+					return TArray<FWidgetAndPointer>();
 				}
-				else
+
+				Algo::Reverse(Path);
+
+				if (WidgetData.CustomPath.IsValid())
 				{
-					return BestCandidateByDistance->BubblePath;
+					const TArray<FWidgetAndPointer> BubblePathExtension = WidgetData.CustomPath.Pin()->GetBubblePathAndVirtualCursors(FirstHitWidget->GetPaintSpaceGeometry(), DesktopSpaceCoordinate, bIgnoreEnabledStatus);
+					Path.Append(BubblePathExtension);
 				}
+				return Path;
 			}
 		}
-
-
-		return DirectBubblePathInfo.BubblePath;
 	}
-	else
-	{
-		// We didn't hit anything.
-		return TArray<FWidgetAndPointer>();
-	}
-}
-
-void FHittestGrid::ClearGridForNewFrame(const FSlateRect& HittestArea)
-{
-#if SLATE_VERBOSE_NAMED_EVENTS
-	SCOPED_NAMED_EVENT_TEXT("FHittestGrid::ClearGridForNewFrame", FColor::Magenta);
-#endif
-	//LogGrid();
-
-	GridOrigin = HittestArea.GetTopLeft();
-	const FVector2D GridSize = HittestArea.GetSize();
-	NumCells = FIntPoint(FMath::CeilToInt(GridSize.X / CellSize.X), FMath::CeilToInt(GridSize.Y / CellSize.Y));
-	WidgetsCachedThisFrame.Reset();
-
-	const int32 NewTotalCells = NumCells.X * NumCells.Y;
-	if (NewTotalCells != Cells.Num())
-	{
-		Cells.Reset(NewTotalCells);
-		Cells.SetNum(NewTotalCells);
-	}
-	else
-	{
-		// As an optimization, if the number of cells do not change then we will just reset the index list inside of them
-		// This will leave slack for indices to be re-added without reallocating.
-		for (FCell& Cell : Cells)
-		{
-			Cell.CachedWidgetIndexes.Reset();
-		}
-	}
-
-	ClippingManager.ResetClippingState();
-}
 
-void FHittestGrid::PushClip(const FSlateClippingZone& ClippingZone)
-{
-	ClippingManager.PushClip(ClippingZone);
+	return TArray<FWidgetAndPointer>();
 }
 
-void FHittestGrid::PopClip()
+bool FHittestGrid::SetHittestArea(const FVector2D& HittestPositionInDesktop, const FVector2D& HittestDimensions, const FVector2D& HitestOffsetInWindow)
 {
-	ClippingManager.PopClip();
-}
+	bool bWasCleared = false;
 
-int32 FHittestGrid::InsertWidget(const int32 ParentHittestIndex, const EVisibility& Visibility, const FArrangedWidget& ArrangedWidget, const FVector2D InWindowOffset, int32 LayerId)
-{
-	if (ensureMsgf(ParentHittestIndex < WidgetsCachedThisFrame.Num(), TEXT("Widget '%s' being drawn before its parent."), *ArrangedWidget.ToString()))
+	// If the size of the hit test area changes we need to clear it out
+	if (GridSize != HittestDimensions)
 	{
-		// Update the FGeometry to transform into desktop space.
-		FArrangedWidget WindowAdjustedWidget(ArrangedWidget);
-		WindowAdjustedWidget.Geometry.AppendTransform(FSlateLayoutTransform(InWindowOffset));
-
-#if 0
-		// Enable this code if you're trying to make sure a widget is only ever added to the hit test grid once.
-		const FCachedWidget* ExistingCachedWidget =
-			WidgetsCachedThisFrame.FindByPredicate([Widget](const FCachedWidget& CachedWidget) {
-			return CachedWidget.WidgetPtr.Pin() == Widget.Widget;
-		});
-
-		ensure(ExistingCachedWidget == nullptr);
-#endif
-
-		const int32 ClippingStateIndex = ClippingManager.GetClippingIndex();
-
-		// Remember this widget, its geometry, and its place in the logical hierarchy.
-		const int32 WidgetIndex = WidgetsCachedThisFrame.Add(FCachedWidget(ParentHittestIndex, WindowAdjustedWidget, ClippingStateIndex, LayerId));
-		check(WidgetIndex < WidgetsCachedThisFrame.Num());
-		if (ParentHittestIndex != INDEX_NONE)
-		{
-			WidgetsCachedThisFrame[ParentHittestIndex].AddChild(WidgetIndex);
-		}
-
-		if (Visibility.IsHitTestVisible())
+		GridSize = HittestDimensions;
+		NumCells = FIntPoint(FMath::CeilToInt(GridSize.X / CellSize.X), FMath::CeilToInt(GridSize.Y / CellSize.Y));
+		
+		const int32 NewTotalCells = NumCells.X * NumCells.Y;
+		if (NewTotalCells != Cells.Num())
 		{
-			// Mark any cell that is overlapped by this widget.
+			Cells.Reset(NewTotalCells);
+			Cells.SetNum(NewTotalCells);
 
-			// Compute the render space clipping rect, and compute it's aligned bounds so we can insert conservatively into the hit test grid.
-			FSlateRect GridRelativeBoundingClipRect =
-				WindowAdjustedWidget.Geometry.GetRenderBoundingRect()
-				.OffsetBy(-GridOrigin);
+			WidgetMap.Reset();
+			WidgetArray.Reset();
 
-			// Starting and ending cells covered by this widget.	
-			const FIntPoint UpperLeftCell = FIntPoint(
-				FMath::Max(0, FMath::FloorToInt(GridRelativeBoundingClipRect.Left / CellSize.X)),
-				FMath::Max(0, FMath::FloorToInt(GridRelativeBoundingClipRect.Top / CellSize.Y)));
-
-			const FIntPoint LowerRightCell = FIntPoint(
-				FMath::Min(NumCells.X - 1, FMath::FloorToInt(GridRelativeBoundingClipRect.Right / CellSize.X)),
-				FMath::Min(NumCells.Y - 1, FMath::FloorToInt(GridRelativeBoundingClipRect.Bottom / CellSize.Y)));
-
-			for (int32 XIndex = UpperLeftCell.X; XIndex <= LowerRightCell.X; ++XIndex)
-			{
-				for (int32 YIndex = UpperLeftCell.Y; YIndex <= LowerRightCell.Y; ++YIndex)
-				{
-					CellAt(XIndex, YIndex).CachedWidgetIndexes.Add(WidgetIndex);
-				}
-			}
+			bWasCleared = true;
 		}
-
-		return WidgetIndex;
-	}
-	else
-	{
-		return INDEX_NONE;
 	}
+
+	GridOrigin = HittestPositionInDesktop;
+	GridWindowOrigin = HitestOffsetInWindow;
+
+	return bWasCleared;
 }
 
-void FHittestGrid::InsertCustomHitTestPath( TSharedRef<ICustomHitTestPath> CustomHitTestPath, int32 WidgetIndex )
+void FHittestGrid::Clear()
 {
-	if( WidgetsCachedThisFrame.IsValidIndex( WidgetIndex ) )
-	{
-		WidgetsCachedThisFrame[WidgetIndex].CustomPath = CustomHitTestPath;
-	}
+	SCOPE_CYCLE_COUNTER(STAT_SlateHTG_Clear);
+	const int32 TotalCells = Cells.Num();
+	Cells.Reset(TotalCells);
+	Cells.SetNum(TotalCells);
+
+	WidgetMap.Reset();
+	WidgetArray.Reset();
 }
 
-bool FHittestGrid::IsDescendantOf(const TSharedRef<SWidget> Parent, const FCachedWidget& Child)
+bool FHittestGrid::IsDescendantOf(const TSharedRef<SWidget> Parent, const FWidgetData& ChildData)
 {
-	TSharedPtr<SWidget> ChildWidget = Child.WidgetPtr.Pin();
-
-	if (!ChildWidget.IsValid() || ChildWidget == Parent)
+	const TSharedPtr<SWidget> ChildWidgetPtr = ChildData.GetWidget();
+	if (ChildWidgetPtr == Parent)
 	{
 		return false;
 	}
 
-	int32 CurWidgetIndex = Child.ParentIndex;
-	while (CurWidgetIndex != INDEX_NONE)
-	{
-		const FCachedWidget& CurCachedWidget = WidgetsCachedThisFrame[CurWidgetIndex];
-		CurWidgetIndex = CurCachedWidget.ParentIndex;
+	const SWidget* ParentWidget = &Parent.Get();
+	SWidget* CurWidget = ChildWidgetPtr.Get();
 
-		if (Parent == CurCachedWidget.WidgetPtr.Pin())
+	while (CurWidget)
+	{
+		if (ParentWidget == CurWidget)
 		{
 			return true;
 		}
+		CurWidget = CurWidget->Advanced_GetPaintParentWidget().Get();
 	}
 
 	return false;
@@ -396,38 +361,43 @@ TSharedPtr<SWidget> FHittestGrid::FindFocusableWidget(FSlateRect WidgetRect, con
 
 		for (StrideCellPoint[StrideAxis] = StrideAxisMin; StrideCellPoint[StrideAxis] <= StrideAxisMax; ++StrideCellPoint[StrideAxis])
 		{
-			FHittestGrid::FCell& Cell = FHittestGrid::CellAt(StrideCellPoint.X, StrideCellPoint.Y);
-			const TArray<int32>& IndexesInCell = Cell.CachedWidgetIndexes;
+			const FHittestGrid::FCell& Cell = FHittestGrid::CellAt(StrideCellPoint.X, StrideCellPoint.Y);
+			const TArray<int32>& IndexesInCell = Cell.GetCachedWidgetIndexes();
 
 			for (int32 i = IndexesInCell.Num() - 1; i >= 0; --i)
 			{
-				int32 CurrentIndex = IndexesInCell[i];
-				check(CurrentIndex < WidgetsCachedThisFrame.Num());
+				const int32 CurrentIndex = IndexesInCell[i];
+				checkSlow(WidgetArray.IsValidIndex(CurrentIndex));
 
-				const FCachedWidget& TestCandidate = WidgetsCachedThisFrame[CurrentIndex];
-				FSlateRect TestCandidateRect = TestCandidate.CachedGeometry.GetRenderBoundingRect().OffsetBy(-GridOrigin);
-
-				if (CompareFunc(DestSideFunc(TestCandidateRect), CurrentSourceSide) && FSlateRect::DoRectanglesIntersect(SweptRect, TestCandidateRect))
+				const FWidgetData& TestCandidate = WidgetArray[CurrentIndex];
+				const TSharedPtr<SWidget> TestWidget = TestCandidate.GetWidget();
+				if (TestWidget.IsValid())
 				{
-					// If this found widget isn't closer then the previously found widget then keep looking.
-					if (BestWidget.IsValid() && !CompareFunc(DestSideFunc(BestWidgetRect), DestSideFunc(TestCandidateRect)))
-					{
-						continue;
-					}
-
-					// If we have a non escape boundary condition and this widget isn't a descendant of our boundary condition widget then it's invalid so we keep looking.
-					if (NavigationReply.GetBoundaryRule() != EUINavigationRule::Escape
-						&& NavigationReply.GetHandler().IsValid()
-						&& !IsDescendantOf(NavigationReply.GetHandler().ToSharedRef(), TestCandidate))
-					{
-						continue;
-					}
+					FGeometry TestCandidateGeo = TestWidget->GetPaintSpaceGeometry();
+					TestCandidateGeo.AppendTransform(FSlateLayoutTransform(-GridWindowOrigin));
+					const FSlateRect TestCandidateRect = TestCandidateGeo.GetRenderBoundingRect();
 
-					TSharedPtr<SWidget> Widget = TestCandidate.WidgetPtr.Pin();
-					if (Widget.IsValid() && Widget->IsEnabled() && Widget->SupportsKeyboardFocus())
+					if (CompareFunc(DestSideFunc(TestCandidateRect), CurrentSourceSide) && FSlateRect::DoRectanglesIntersect(SweptRect, TestCandidateRect))
 					{
-						BestWidgetRect = TestCandidateRect;
-						BestWidget = Widget;
+						// If this found widget isn't closer then the previously found widget then keep looking.
+						if (BestWidget.IsValid() && !CompareFunc(DestSideFunc(BestWidgetRect), DestSideFunc(TestCandidateRect)))
+						{
+							continue;
+						}
+
+						// If we have a non escape boundary condition and this widget isn't a descendant of our boundary condition widget then it's invalid so we keep looking.
+						if (NavigationReply.GetBoundaryRule() != EUINavigationRule::Escape
+							&& NavigationReply.GetHandler().IsValid()
+							&& !IsDescendantOf(NavigationReply.GetHandler().ToSharedRef(), TestCandidate))
+						{
+							continue;
+						}
+
+						if (TestWidget->IsEnabled() && TestWidget->SupportsKeyboardFocus())
+						{
+							BestWidgetRect = TestCandidateRect;
+							BestWidget = TestWidget;
+						}
 					}
 				}
 			}
@@ -497,27 +467,13 @@ TSharedPtr<SWidget> FHittestGrid::FindFocusableWidget(FSlateRect WidgetRect, con
 
 TSharedPtr<SWidget> FHittestGrid::FindNextFocusableWidget(const FArrangedWidget& StartingWidget, const EUINavigation Direction, const FNavigationReply& NavigationReply, const FArrangedWidget& RuleWidget)
 {
-	FSlateRect WidgetRect =
-		TransformRect(
-			Concatenate(
-				Inverse(StartingWidget.Geometry.GetAccumulatedLayoutTransform()),
-				StartingWidget.Geometry.GetAccumulatedRenderTransform()
-			),
-			FSlateRotatedRect(StartingWidget.Geometry.GetLayoutBoundingRect())
-		)
-		.ToBoundingRect()
-		.OffsetBy(-GridOrigin);
-
-	FSlateRect BoundingRuleRect =
-		TransformRect(
-			Concatenate(
-				Inverse(RuleWidget.Geometry.GetAccumulatedLayoutTransform()),
-				RuleWidget.Geometry.GetAccumulatedRenderTransform()
-			),
-			FSlateRotatedRect(RuleWidget.Geometry.GetLayoutBoundingRect())
-		)
-		.ToBoundingRect()
-		.OffsetBy(-GridOrigin);
+	FGeometry StartingWidgetGeo = StartingWidget.Widget->GetPaintSpaceGeometry();
+	StartingWidgetGeo.AppendTransform(FSlateLayoutTransform(-GridWindowOrigin));
+	FSlateRect WidgetRect = StartingWidgetGeo.GetRenderBoundingRect();
+
+	FGeometry BoundingRuleWidgetGeo = RuleWidget.Widget->GetPaintSpaceGeometry();
+	BoundingRuleWidgetGeo.AppendTransform(FSlateLayoutTransform(-GridWindowOrigin));
+	FSlateRect BoundingRuleRect = BoundingRuleWidgetGeo.GetRenderBoundingRect();
 
 	FSlateRect SweptWidgetRect = WidgetRect;
 
@@ -587,10 +543,40 @@ bool FHittestGrid::IsValidCellCoord(const FIntPoint& CellCoord) const
 
 bool FHittestGrid::IsValidCellCoord(const int32 XCoord, const int32 YCoord) const
 {
-	const int32 Index = (YCoord * NumCells.X) + XCoord;
-	return Cells.IsValidIndex(Index);
+	return XCoord >= 0 && XCoord < NumCells.X && YCoord >= 0 && YCoord < NumCells.Y;
 }
 
+void FHittestGrid::AppendGrid(FHittestGrid& OtherGrid)
+{
+	// The two grids must occupy the same space
+	if (ensure(this != &OtherGrid && GridOrigin == OtherGrid.GridOrigin && GridWindowOrigin == OtherGrid.GridWindowOrigin && GridSize == OtherGrid.GridSize))
+	{
+		// Index is not valid in the other grid so remap it 
+		for (auto& Pair : OtherGrid.WidgetMap)
+		{
+			int32& WidgetIndex = WidgetMap.Add(Pair.Key);
+			
+			// Get the already created widget data out of the other widget array
+			const FWidgetData& OtherData = OtherGrid.WidgetArray[Pair.Value];
+
+			WidgetIndex = WidgetArray.Add(OtherData); 
+
+			// Add the new widget index to the correct cells in this grid
+			for (int32 XIndex = OtherData.UpperLeftCell.X; XIndex <= OtherData.LowerRightCell.X; ++XIndex)
+			{
+				for (int32 YIndex = OtherData.UpperLeftCell.Y; YIndex <= OtherData.LowerRightCell.Y; ++YIndex)
+				{
+					if (IsValidCellCoord(XIndex, YIndex))
+					{
+						CellAt(XIndex, YIndex).AddIndex(WidgetIndex);
+					}
+				}
+			}
+
+		}
+	}
+}
+#if WITH_SLATE_DEBUGGING
 void FHittestGrid::LogGrid() const
 {
 	FString TempString;
@@ -600,7 +586,7 @@ void FHittestGrid::LogGrid() const
 		{
 			TempString += "\t";
 			TempString += "[";
-			for (int i : CellAt(x, y).CachedWidgetIndexes)
+			for (int i : CellAt(x, y).GetCachedWidgetIndexes())
 			{
 				TempString += FString::Printf(TEXT("%d,"), i);
 			}
@@ -613,112 +599,198 @@ void FHittestGrid::LogGrid() const
 
 	UE_LOG(LogHittestDebug, Warning, TEXT("\n%s"), *TempString);
 
-	for (int i = 0; i < WidgetsCachedThisFrame.Num(); ++i)
+	for (TSparseArray<FWidgetData>::TConstIterator It(WidgetArray); It; ++It)
+	{
+		const FWidgetData& CurWidgetData = *It;
+		const TSharedPtr<SWidget> CachedWidget = CurWidgetData.GetWidget();
+		UE_LOG(LogHittestDebug, Warning, TEXT("  [%d][%d][%d] => %s @ %s"),
+			It.GetIndex(),
+			CurWidgetData.PrimarySort,
+			CachedWidget.IsValid() ? *CachedWidget->ToString() : TEXT("Invalid WIdget"),
+			CachedWidget.IsValid() ? *CachedWidget->GetPaintSpaceGeometry().ToString() : TEXT("Invalid WIdget"));
+	}
+}
+
+void FHittestGrid::DisplayGrid(int32 InLayer, const FGeometry& AllottedGeometry, FSlateWindowElementList& WindowElementList) const
+{
+	static const FSlateBrush* WhiteBrush = FCoreStyle::Get().GetBrush(TEXT("FocusRectangle"));
+
+	//WindowElementList.PushAbsoluteBatchPriortyGroup(INT_MAX);
+	for (TSparseArray<FWidgetData>::TConstIterator It(WidgetArray); It; ++It)
 	{
-		if (WidgetsCachedThisFrame[i].ParentIndex == INDEX_NONE)
+		const FWidgetData& CurWidgetData = *It;
+		const TSharedPtr<SWidget> CachedWidget = CurWidgetData.GetWidget();
+		if (CachedWidget.IsValid())
 		{
-			LogChildren(i, 0, WidgetsCachedThisFrame);
+			FSlateDrawElement::MakeBox(
+				WindowElementList,
+				InLayer,
+				CachedWidget->GetPaintSpaceGeometry().ToPaintGeometry(),
+				WhiteBrush
+			);
 		}
 	}
+
+	//WindowElementList.PopBatchPriortyGroup();
 }
 
-void FHittestGrid::LogChildren(int32 Index, int32 IndentLevel, const TArray<FCachedWidget>& WidgetsCachedThisFrame)
+#endif // WITH_SLATE_DEBUGGING
+
+
+void FHittestGrid::AddWidget(const TSharedRef<SWidget>& InWidget, int32 InBatchPriorityGroup, int32 InLayerId, int32 InSecondarySort)
 {
-	FString IndentString;
-	for (int i = 0; i < IndentLevel; ++i)
+	if (!InWidget->GetVisibility().IsHitTestVisible())
 	{
-		IndentString += "|\t";
+		return;
 	}
 
-	const FCachedWidget& CachedWidget = WidgetsCachedThisFrame[Index];
-	const TSharedPtr<SWidget> CachedWidgetPtr = CachedWidget.WidgetPtr.Pin();
-	const FString WidgetString = CachedWidgetPtr.IsValid() ? CachedWidgetPtr->ToString() : TEXT("(null)");
-	UE_LOG(LogHittestDebug, Warning, TEXT("%s[%d] => %s @ %s"), *IndentString, Index, *WidgetString, *CachedWidget.CachedGeometry.ToString());
+	SCOPE_CYCLE_COUNTER(STAT_SlateHTG_AddWidget);
 
-	for (int i = 0; i < CachedWidget.Children.Num(); ++i)
+	if (WidgetMap.Contains(&*InWidget))
 	{
-		LogChildren(CachedWidget.Children[i], IndentLevel + 1, WidgetsCachedThisFrame);
+		// TODO: Update can be faster than a Remove and Add
+		RemoveWidget(InWidget);
 	}
-}
 
-FHittestGrid::FWidgetPathAndDist FHittestGrid::GetWidgetPathAndDist(const FGridTestingParams& Params, const bool bIgnoreEnabledStatus) const
-{
-	//Also grab the Hit Index, and the distance to the top hit
-	const FIndexAndDistance HitIndex = GetHitIndexFromCellIndex(Params);
+	const int64 PrimarySort = (((int64)InBatchPriorityGroup << 32) | InLayerId);
+
+	// Track the widget and identify it's Widget Index
+	FGeometry GridSpaceGeometry = InWidget->GetPaintSpaceGeometry();
+	GridSpaceGeometry.AppendTransform(FSlateLayoutTransform(-GridWindowOrigin));
+
+	// Currently using grid offset because the grid covers all desktop space.
+	const FSlateRect BoundingRect = GridSpaceGeometry.GetRenderBoundingRect();
+
+	// Starting and ending cells covered by this widget.	
+	const FIntPoint UpperLeftCell = GetCellCoordinate(BoundingRect.GetTopLeft());
+	const FIntPoint LowerRightCell = GetCellCoordinate(BoundingRect.GetBottomRight());
 
-	if (HitIndex.WidgetIndex != INDEX_NONE)
+	int32& WidgetIndex = WidgetMap.Add(&*InWidget);
+	FWidgetData Data(InWidget, UpperLeftCell, LowerRightCell, PrimarySort, InSecondarySort);
+	WidgetIndex = WidgetArray.Add(Data); // Bleh why doesn't Sparse Array have an emplace.
+
+	for (int32 XIndex = UpperLeftCell.X; XIndex <= LowerRightCell.X; ++XIndex)
 	{
-		// if we have a custom path, we want to do the testing for 3D Widgets
-		const FCachedWidget& PhysicallyHitWidget = WidgetsCachedThisFrame[HitIndex.WidgetIndex];
-		if (PhysicallyHitWidget.CustomPath.IsValid())
+		for (int32 YIndex = UpperLeftCell.Y; YIndex <= LowerRightCell.Y; ++YIndex)
 		{
-			const FVector2D DesktopSpaceCoordinate = Params.CursorPositionInGrid + GridOrigin;
-			TArray<FWidgetAndPointer> LogicalBubblePath = GetBubblePathFromHitIndex(HitIndex.WidgetIndex, bIgnoreEnabledStatus);
-			const TArray<FWidgetAndPointer> BubblePathExtension = PhysicallyHitWidget.CustomPath.Pin()->GetBubblePathAndVirtualCursors(PhysicallyHitWidget.CachedGeometry, DesktopSpaceCoordinate, bIgnoreEnabledStatus);
-			LogicalBubblePath.Append(BubblePathExtension);
-			return FWidgetPathAndDist(LogicalBubblePath, 0.0, PhysicallyHitWidget.LayerId);
+			if (IsValidCellCoord(XIndex, YIndex))
+			{
+				CellAt(XIndex, YIndex).AddIndex(WidgetIndex);
+			}
 		}
-		else
+	}
+}
+
+void FHittestGrid::RemoveWidget(const TSharedRef<SWidget>& InWidget)
+{
+	SCOPE_CYCLE_COUNTER(STAT_SlateHTG_RemoveWidget);
+
+	int32 WidgetIndex = INDEX_NONE;
+	if(WidgetMap.RemoveAndCopyValue(&*InWidget, WidgetIndex))
+	{
+		const FWidgetData& WidgetData = WidgetArray[WidgetIndex];
+
+		// Starting and ending cells covered by this widget.	
+		const FIntPoint& UpperLeftCell = WidgetData.UpperLeftCell;
+		const FIntPoint& LowerRightCell = WidgetData.LowerRightCell;
+
+		for (int32 XIndex = UpperLeftCell.X; XIndex <= LowerRightCell.X; ++XIndex)
 		{
-			//get the path from the hit index, check if anything came back
-			const TArray<FWidgetAndPointer> BubblePath = GetBubblePathFromHitIndex(HitIndex.WidgetIndex, bIgnoreEnabledStatus);
-			return FWidgetPathAndDist(BubblePath, BubblePath.Num() > 0 ? HitIndex.DistanceSqToWidget : -1.0f, PhysicallyHitWidget.LayerId);
+			for (int32 YIndex = UpperLeftCell.Y; YIndex <= LowerRightCell.Y; ++YIndex)
+			{
+				checkSlow(IsValidCellCoord(XIndex, YIndex));
+				CellAt(XIndex, YIndex).RemoveIndex(WidgetIndex);
+			}
 		}
+
+		WidgetArray.RemoveAt(WidgetIndex);
 	}
+}
 
-	//return if we have a valid path or not
-	return FWidgetPathAndDist();
+void FHittestGrid::InsertCustomHitTestPath(const TSharedRef<SWidget> InWidget, TSharedRef<ICustomHitTestPath> CustomHitTestPath)
+{
+	int32 WidgetIndex = WidgetMap.FindChecked(&*InWidget);
+	FWidgetData& WidgetData = WidgetArray[WidgetIndex];
+	WidgetData.CustomPath = CustomHitTestPath;
 }
 
-FHittestGrid::FIndexAndDistance FHittestGrid::GetHitIndexFromCellIndex(const FGridTestingParams& Params) const
+FHittestGrid::FIndexAndDistance FHittestGrid::GetHitIndexFromCellIndex(const FGridTestingParams& Params)
 {
 	//check if the cell coord 
 	if (IsValidCellCoord(Params.CellCoord))
 	{
-		const TArray<int32>& IndexesInCell = CellAt(Params.CellCoord.X, Params.CellCoord.Y).CachedWidgetIndexes;
-		const TArray< FSlateClippingState >& ClippingStates = ClippingManager.GetClippingStates();
+		// Get the cell and sort it 
+		FHittestGrid::FCell& Cell = CellAt(Params.CellCoord.X, Params.CellCoord.Y);
+		Cell.Sort(WidgetArray);
+
+		// Get the cell's index array for the search
+		const TArray<int32>& IndexesInCell = Cell.GetCachedWidgetIndexes();
+
+#if 0 //Unroll some data for debugging if necessary
+		struct FDebugData
+		{
+			int32 WidgetIndex;
+			int64 PrimarySort;
+			FName WidgetType;
+			FName WidgetLoc;
+			TSharedPtr<SWidget> Widget;
+		};
+
+		TArray<FDebugData> DebugData;
+		for (int32 i = 0; i < IndexesInCell.Num(); ++i)
+		{
+			FDebugData& Cur = DebugData.AddDefaulted_GetRef();
+			Cur.WidgetIndex = IndexesInCell[i];
+			Cur.PrimarySort = WidgetArray[Cur.WidgetIndex].PrimarySort;
+			Cur.Widget = WidgetArray[Cur.WidgetIndex].GetWidget();
+			Cur.WidgetType = Cur.Widget.IsValid() ? Cur.Widget->GetType() : NAME_None;
+			Cur.WidgetLoc = Cur.Widget.IsValid() ? Cur.Widget->GetCreatedInLocation() : NAME_None;
+		}
+
+#endif
 
 		// Consider front-most widgets first for hittesting.
 		for (int32 i = IndexesInCell.Num() - 1; i >= 0; --i)
 		{
 			const int32 WidgetIndex = IndexesInCell[i];
 
-			check(WidgetsCachedThisFrame.IsValidIndex(WidgetIndex));
+			checkSlow(WidgetArray.IsValidIndex(WidgetIndex));
 
-			const FCachedWidget& TestCandidate = WidgetsCachedThisFrame[WidgetIndex];
+			const FWidgetData& TestCandidate = WidgetArray[WidgetIndex];
+			const TSharedPtr<SWidget> TestWidget = TestCandidate.GetWidget();
 
 			// When performing a point hittest, accept all hittestable widgets.
 			// When performing a hittest with a radius, only grab interactive widgets.
-			const bool bIsValidWidget = !Params.bTestWidgetIsInteractive || (TestCandidate.WidgetPtr.IsValid() && TestCandidate.WidgetPtr.Pin()->IsInteractable());
+			const bool bIsValidWidget = TestWidget.IsValid() && (!Params.bTestWidgetIsInteractive || TestWidget->IsInteractable());
 			if (bIsValidWidget)
 			{
-				const FVector2D DesktopSpaceCoordinate = Params.CursorPositionInGrid + GridOrigin;
-				
+				const FVector2D WindowSpaceCoordinate = Params.CursorPositionInGrid + GridWindowOrigin;
+
+				const FGeometry& TestGeometry = TestWidget->GetPaintSpaceGeometry();
+
 				bool bPointInsideClipMasks = true;
-				if (TestCandidate.ClippingStateIndex != INDEX_NONE)
+				const TOptional<FSlateClippingState>& WidgetClippingState = TestWidget->GetCurrentClippingState();
+				if (WidgetClippingState.IsSet())
 				{
-					const FSlateClippingState& ClippingState = ClippingStates[TestCandidate.ClippingStateIndex];
-
-					// TODO Solve non-zero radius cursors?
-					bPointInsideClipMasks = ClippingState.IsPointInside(DesktopSpaceCoordinate);
+					// TODO: Solve non-zero radius cursors?
+					bPointInsideClipMasks = WidgetClippingState->IsPointInside(WindowSpaceCoordinate);
 				}
 
 				if (bPointInsideClipMasks)
 				{
 					// Compute the render space clipping rect (FGeometry exposes a layout space clipping rect).
-					const FSlateRotatedRect DesktopOrientedClipRect = TransformRect(
+					const FSlateRotatedRect WindowOrientedClipRect = TransformRect(
 						Concatenate(
-							Inverse(TestCandidate.CachedGeometry.GetAccumulatedLayoutTransform()),
-							TestCandidate.CachedGeometry.GetAccumulatedRenderTransform()),
-						FSlateRotatedRect(TestCandidate.CachedGeometry.GetLayoutBoundingRect())
+							Inverse(TestGeometry.GetAccumulatedLayoutTransform()),
+							TestGeometry.GetAccumulatedRenderTransform()),
+						FSlateRotatedRect(TestGeometry.GetLayoutBoundingRect())
 					);
 
-					if (IsOverlappingSlateRotatedRect(DesktopSpaceCoordinate, Params.Radius, DesktopOrientedClipRect))
+					if (IsOverlappingSlateRotatedRect(WindowSpaceCoordinate, Params.Radius, WindowOrientedClipRect))
 					{
-						// We are within the search radius!
-						const bool bNeedsDistanceSearch = Params.Radius > 0.0f;
 						// For non-0 radii also record the distance to cursor's center so that we can pick the closest hit from the results.
-						const float DistSq = (bNeedsDistanceSearch) ? DistanceSqToSlateRotatedRect(DesktopSpaceCoordinate, DesktopOrientedClipRect) : 0.0f;
+						const bool bNeedsDistanceSearch = Params.Radius > 0.0f;
+						const float DistSq = (bNeedsDistanceSearch) ? DistanceSqToSlateRotatedRect(WindowSpaceCoordinate, WindowOrientedClipRect) : 0.0f;
 						return FIndexAndDistance(WidgetIndex, DistSq);
 					}
 				}
@@ -727,97 +799,4 @@ FHittestGrid::FIndexAndDistance FHittestGrid::GetHitIndexFromCellIndex(const FGr
 	}
 
 	return FIndexAndDistance();
-}
-
-TArray<FWidgetAndPointer> FHittestGrid::GetBubblePathFromHitIndex(const int32 HitIndex, const bool bIgnoreEnabledStatus) const
-{
-	TArray<FWidgetAndPointer> BubblePath;
-
-	if (WidgetsCachedThisFrame.IsValidIndex(HitIndex))
-	{
-		int32 CurWidgetIndex = HitIndex;
-		do
-		{
-			check(CurWidgetIndex < WidgetsCachedThisFrame.Num());
-			const FCachedWidget& CurCachedWidget = WidgetsCachedThisFrame[CurWidgetIndex];
-			const TSharedPtr<SWidget> CachedWidgetPtr = CurCachedWidget.WidgetPtr.Pin();
-
-
-			const bool bPathInterrupted = !CachedWidgetPtr.IsValid();
-			if (bPathInterrupted)
-			{
-				// A widget in the path to the root has been removed, so anything
-				// we thought we had hittest so far is no longer actually in the hierarchy.
-				// Continue bubbling to the root of the hirarchy to find an unbroken chain to root.
-				// The leafmost widget in that chain will get first shot at the events.
-				BubblePath.Reset();
-			}
-			else
-			{
-				BubblePath.Insert(FWidgetAndPointer(FArrangedWidget(CachedWidgetPtr.ToSharedRef(), CurCachedWidget.CachedGeometry), TSharedPtr<FVirtualPointerPosition>()), 0);
-			}
-			CurWidgetIndex = CurCachedWidget.ParentIndex;
-		} while (CurWidgetIndex != INDEX_NONE);
-
-
-		if (SlateVerifyHitTestVisibility)
-		{
-			// Hit Test Invisible widgets effects all of the logical children.
-			// Normally this isn't a problem, but in the case of low framerate
-			// if multiple mouse events buffer up, and are consumed in one frame
-			// it's possible that in one frame, the first mouse event might change
-			// the hit-test ability of widgets.
-			{
-				const int32 HitTestInvisibleWidgetIndex = BubblePath.IndexOfByPredicate([](const FArrangedWidget& SomeWidget) { return !SomeWidget.Widget->GetVisibility().AreChildrenHitTestVisible(); });
-				if (HitTestInvisibleWidgetIndex != INDEX_NONE)
-				{
-					BubblePath.RemoveAt(HitTestInvisibleWidgetIndex, BubblePath.Num() - HitTestInvisibleWidgetIndex);
-				}
-			}
-
-			// Similar to the above check, this determines if a widget became hit test invisible
-			// directly, because rather than an entire set of children becoming invisible
-			// this addresses the problem of a specific widget changing visibility.
-			{
-				int32 FirstHitTestWidgetIndex = INDEX_NONE;
-				for (int32 Index = BubblePath.Num() - 1; Index >= 0; --Index)
-				{
-					const EVisibility Visibilty = BubblePath[Index].Widget->GetVisibility();
-
-					if (Visibilty.IsHitTestVisible())
-					{
-						FirstHitTestWidgetIndex = Index;
-						break;
-					}
-				}
-
-				if (FirstHitTestWidgetIndex != INDEX_NONE)
-				{
-					const int32 RemovalCount = (BubblePath.Num() - (FirstHitTestWidgetIndex + 1));
-					if (RemovalCount > 0)
-					{
-						BubblePath.RemoveAt(FirstHitTestWidgetIndex + 1, RemovalCount);
-					}
-				}
-				else
-				{
-					BubblePath.Reset();
-				}
-			}
-		}
-
-		// Disabling a widget disables all of its logical children
-		// This effect is achieved by truncating the path to the
-		// root-most enabled widget.
-		if (!bIgnoreEnabledStatus)
-		{
-			const int32 DisabledWidgetIndex = BubblePath.IndexOfByPredicate([](const FArrangedWidget& SomeWidget) { return !SomeWidget.Widget->IsEnabled(); });
-			if (DisabledWidgetIndex != INDEX_NONE)
-			{
-				BubblePath.RemoveAt(DisabledWidgetIndex, BubblePath.Num() - DisabledWidgetIndex);
-			}
-		}
-	}
-
-	return BubblePath;
 }
\ No newline at end of file
diff --git a/Engine/Source/Runtime/SlateCore/Private/Layout/Clipping.cpp b/Engine/Source/Runtime/SlateCore/Private/Layout/Clipping.cpp
--- a/Engine/Source/Runtime/SlateCore/Private/Layout/Clipping.cpp
+++ b/Engine/Source/Runtime/SlateCore/Private/Layout/Clipping.cpp
@@ -199,13 +199,24 @@ bool FSlateClippingZone::IsPointInside(const FVector2D& Point) const
 
 FSlateClippingState::FSlateClippingState(EClippingFlags InFlags /*= EClippingFlags::None*/)
 	: Flags(InFlags)
-#if !(UE_BUILD_SHIPPING || UE_BUILD_TEST)
+#if WITH_SLATE_DEBUGGING
 	, Debugging_StateIndex(INDEX_NONE)
 	, Debugging_StateIndexFromFrame(INDEX_NONE)
 #endif
 {
 }
 
+FSlateClippingState::FSlateClippingState(const FSlateClippingState& Other)
+	: StencilQuads(Other.StencilQuads)
+	, ScissorRect(Other.ScissorRect)
+	, Flags(Other.Flags)
+#if WITH_SLATE_DEBUGGING
+	, Debugging_StateIndex(Other.Debugging_StateIndex)
+	, Debugging_StateIndexFromFrame(Other.Debugging_StateIndexFromFrame)
+#endif
+{
+}
+
 bool FSlateClippingState::IsPointInside(const FVector2D& Point) const
 {
 	if (ScissorRect.IsSet())
@@ -231,7 +242,7 @@ FSlateClippingManager::FSlateClippingManager()
 {
 }
 
-const FSlateClippingState* FSlateClippingManager::GetPreviousClippnigState(bool bWillIntersectWithParent) const
+const FSlateClippingState* FSlateClippingManager::GetPreviousClippingState(bool bWillIntersectWithParent) const
 {
 	const FSlateClippingState* PreviousClippingState = nullptr;
 
@@ -258,7 +269,7 @@ const FSlateClippingState* FSlateClippingManager::GetPreviousClippnigState(bool
 
 FSlateClippingState FSlateClippingManager::CreateClippingState(const FSlateClippingZone& InClipRect) const
 {
-	const FSlateClippingState* PreviousClippingState = GetPreviousClippnigState(InClipRect.GetShouldIntersectParent());
+	const FSlateClippingState* PreviousClippingState = GetPreviousClippingState(InClipRect.GetShouldIntersectParent());
 
 	// Initialize the new clipping state
 	FSlateClippingState NewClippingState(InClipRect.GetAlwaysClip() ? EClippingFlags::AlwaysClip : EClippingFlags::None);
@@ -307,23 +318,17 @@ int32 FSlateClippingManager::PushClip(const FSlateClippingZone& InClipRect)
 
 int32 FSlateClippingManager::PushClippingState(const FSlateClippingState& NewClippingState)
 {
-	int32 NewClippingStateIndex = ClippingStates.Num();
+	int32 NewClippingStateIndex = ClippingStates.AddUnique(NewClippingState);
 
-#if !(UE_BUILD_SHIPPING || UE_BUILD_TEST)
+#if WITH_SLATE_DEBUGGING
 	NewClippingState.SetDebuggingStateIndex(NewClippingStateIndex);
 #endif
 
 	ClippingStack.Add(NewClippingStateIndex);
-	ClippingStates.Add(NewClippingState);
 
 	return NewClippingStateIndex;
 }
 
-int32 FSlateClippingManager::PushAndMergePartialClippingState(const FSlateClippingState& NewPartialClippingState)
-{
-	return PushClippingState(MergePartialClippingState(NewPartialClippingState));
-}
-
 int32 FSlateClippingManager::GetClippingIndex() const
 {
 	return ClippingStack.Num() > 0 ? ClippingStack.Top() : INDEX_NONE;
@@ -357,59 +362,17 @@ void FSlateClippingManager::PopClip()
 	}
 }
 
-int32 FSlateClippingManager::MergePartialClippingStates(const TArray<FSlateClippingState>& States)
-{
-	int32 Offset = ClippingStates.Num();
-
-	for (const FSlateClippingState& State : States)
-	{
-		ClippingStates.Add(MergePartialClippingState(State));
-	}
-
-	return Offset;
-}
-
-FSlateClippingState FSlateClippingManager::MergePartialClippingState(const FSlateClippingState& State) const
-{
-	if (State.GetClippingMethod() == EClippingMethod::Scissor)
-	{
-		return CreateClippingState(State.ScissorRect.GetValue());
-	}
-	else
-	{
-		const int32 CurrentIndex = GetClippingIndex();
-
-		if (CurrentIndex != INDEX_NONE)
-		{
-			const FSlateClippingState* PreviousClippingState = GetPreviousClippnigState(State.GetShouldIntersectParent());
-			if (PreviousClippingState)
-			{
-				FSlateClippingState NewState = State;
-				if (PreviousClippingState->GetClippingMethod() == EClippingMethod::Scissor)
-				{
-					NewState.StencilQuads.Insert(PreviousClippingState->ScissorRect.GetValue(), 0);
-				}
-				else
-				{
-					NewState.StencilQuads.Insert(PreviousClippingState->StencilQuads, 0);
-				}
-
-				return NewState;
-			}
-		}
-
-		return State;
-	}
-}
-
 void FSlateClippingManager::ResetClippingState()
 {
 	ClippingStates.Reset();
 	ClippingStack.Reset();
 }
 
-void FSlateClippingManager::CopyClippingStateTo(FSlateClippingManager& Other) const
+void FSlateClippingManager::PopToStackIndex(int32 Index)
 {
-	Other.ClippingStack = ClippingStack;
-	Other.ClippingStates = ClippingStates;
-}
\ No newline at end of file
+	const int32 StartIndexToPop = Index + 1;
+	if (ClippingStack.Num() > StartIndexToPop)
+	{
+		ClippingStack.RemoveAt(StartIndexToPop, ClippingStack.Num() - StartIndexToPop, false);
+	}
+}
