Slate Global Invalidation System

diff --git a/Engine/Source/Runtime/SlateCore/Private/Rendering/DrawElements.cpp b/Engine/Source/Runtime/SlateCore/Private/Rendering/DrawElements.cpp
--- a/Engine/Source/Runtime/SlateCore/Private/Rendering/DrawElements.cpp
+++ b/Engine/Source/Runtime/SlateCore/Private/Rendering/DrawElements.cpp
@@ -1,7 +1,6 @@
 // Copyright 1998-2019 Epic Games, Inc. All Rights Reserved.
 
 #include "Rendering/DrawElements.h"
-#include "Application/SlateApplicationBase.h"
 #include "Widgets/SWidget.h"
 #include "Widgets/SWindow.h"
 #include "HAL/IConsoleManager.h"
@@ -10,14 +9,14 @@
 #include "Fonts/FontCache.h"
 #include "Rendering/SlateObjectReferenceCollector.h"
 #include "Debugging/SlateDebugging.h"
+#include "Application/SlateApplicationBase.h"
 
 DECLARE_CYCLE_STAT(TEXT("FSlateDrawElement::Make Time"), STAT_SlateDrawElementMakeTime, STATGROUP_SlateVerbose);
 DECLARE_CYCLE_STAT(TEXT("FSlateDrawElement::MakeCustomVerts Time"), STAT_SlateDrawElementMakeCustomVertsTime, STATGROUP_Slate);
 DECLARE_CYCLE_STAT(TEXT("FSlateDrawElement::Prebatch Time"), STAT_SlateDrawElementPrebatchTime, STATGROUP_Slate);
 
 DEFINE_STAT(STAT_SlateBufferPoolMemory);
-
-FSlateShaderResourceManager* FSlateDataPayload::ResourceManager;
+DEFINE_STAT(STAT_SlateCachedDrawElementMemory);
 
 static bool IsResourceObjectValid(UObject*& InObject)
 {
@@ -30,114 +29,131 @@ static bool IsResourceObjectValid(UObject*& InObject)
 	return true;
 }
 
-FSlateWindowElementList::FSlateWindowElementList(const TSharedPtr<SWindow>& InPaintWindow)
-	: WeakPaintWindow(InPaintWindow)
-	, RawPaintWindow(InPaintWindow.Get())
-	, RenderTargetWindow(nullptr)
-	, bNeedsDeferredResolve(false)
-	, ResolveToDeferredIndex()
-	, MemManager(0)
-	, WindowSize(FVector2D(0.0f, 0.0f))
-	, bReportReferences(true)
+static bool ShouldCull(const FSlateWindowElementList& ElementList)
 {
-	DrawStack.Push(&RootDrawLayer);
-	if (InPaintWindow.IsValid())
+	const FSlateClippingManager& ClippingManager = ElementList.GetClippingManager();
+	const int32 CurrentIndex = ClippingManager.GetClippingIndex();
+	if (CurrentIndex != INDEX_NONE)
 	{
-		WindowSize = InPaintWindow->GetSizeInScreen();
+		const FSlateClippingState& ClippingState = ClippingManager.GetClippingStates()[CurrentIndex];
+		return ClippingState.HasZeroArea();
 	}
 
-	// Only keep UObject resources alive if this window element list is born on the game thread.
-	if (IsInGameThread())
-	{
-		ResourceGCRoot = MakeUnique<FWindowElementGCObject>(this);
-	}
+	return false;
 }
 
-FSlateWindowElementList::~FSlateWindowElementList()
+static bool ShouldCull(const FSlateWindowElementList& ElementList, const FPaintGeometry& PaintGeometry)
 {
-	if (ResourceGCRoot.IsValid())
+	const FVector2D& LocalSize = PaintGeometry.GetLocalSize();
+	if (LocalSize.X == 0 || LocalSize.Y == 0)
 	{
-		ResourceGCRoot->ClearOwner();
+		return true;
 	}
+
+	return ShouldCull(ElementList);
 }
 
-void FSlateWindowElementList::AppendItems(FSlateWindowElementList* Other)
+static bool ShouldCull(const FSlateWindowElementList& ElementList, const FPaintGeometry& PaintGeometry, const FSlateBrush* InBrush)
 {
-	FSlateDrawLayer* Layer = DrawStack.Last();
+	if (ShouldCull(ElementList, PaintGeometry))
+	{
+		return true;
+	}
+
+	if (InBrush->GetDrawType() == ESlateBrushDrawType::NoDrawType)
+	{
+		return true;
+	}
 
-	TArray<FSlateDrawElement>& ActiveDrawElements = Layer->DrawElements;
-	ActiveDrawElements.Append(Other->GetDrawElements());
+	UObject* ResourceObject = InBrush->GetResourceObject();
+	if (!IsResourceObjectValid(ResourceObject))
+	{
+		return true;
+	}
+
+	return false;
 }
 
-void FSlateDataPayload::SetShapedText(FSlateWindowElementList& ElementList, const FShapedGlyphSequencePtr& InShapedGlyphSequence, FLinearColor InOutlineTint)
+
+static bool ShouldCull(const FSlateWindowElementList& ElementList, const FPaintGeometry& PaintGeometry, const FLinearColor& InTint)
 {
-	ShapedGlyphSequence = InShapedGlyphSequence;
-	OutlineTint = InOutlineTint;
+	if (InTint.A == 0 || ShouldCull(ElementList, PaintGeometry))
+	{
+		return true;
+	}
 
-	FSlateObjectReferenceCollector Collector(ElementList.ResourcesToReport);
-	const_cast<FShapedGlyphSequence*>(ShapedGlyphSequence.Get())->AddReferencedObjects(Collector);
+	return false;
 }
 
-void FSlateDataPayload::SetText(FSlateWindowElementList& ElementList, const FString& InText, const FSlateFontInfo& InFontInfo, const int32 InStartIndex, const int32 InEndIndex)
+static bool ShouldCull(const FSlateWindowElementList& ElementList, const FPaintGeometry& PaintGeometry, const FLinearColor& InTint, const FString& InText)
 {
-	FontInfo = InFontInfo;
-	const int32 StartIndex = FMath::Min<int32>(InStartIndex, InText.Len());
-	const int32 EndIndex = FMath::Min<int32>(InEndIndex, InText.Len());
-	TextLength = (EndIndex > StartIndex) ? EndIndex - StartIndex : 0;
-	// Allocate memory and account for null terminator
-	ImmutableText = (TCHAR*)ElementList.Alloc(sizeof(TCHAR) * (TextLength + 1), alignof(TCHAR));
-	if (TextLength > 0)
+	if (InTint.A == 0 || InText.Len() == 0 || ShouldCull(ElementList, PaintGeometry))
 	{
-		FCString::Strncpy(ImmutableText, InText.GetCharArray().GetData() + StartIndex, TextLength + 1);
-		check(!ImmutableText[TextLength]);
+		return true;
 	}
-	else
+
+	return false;
+}
+
+static bool ShouldCull(const FSlateWindowElementList& ElementList, const FPaintGeometry& PaintGeometry, const FSlateBrush* InBrush, const FLinearColor& InTint)
+{
+	if (InTint.A == 0 || ShouldCull(ElementList, PaintGeometry, InBrush))
 	{
-		ImmutableText[0] = 0;
+		return true;
 	}
 
-	FSlateObjectReferenceCollector Collector(ElementList.ResourcesToReport);
-	FontInfo.AddReferencedObjects(Collector);
+	return false;
 }
 
-void FSlateDataPayload::SetLines(FSlateWindowElementList& ElementList, const TArray<FVector2D>& InPoints, bool bInAntialias, const TArray<FLinearColor>* InPointColors)
+
+
+FSlateWindowElementList::FSlateWindowElementList(const TSharedPtr<SWindow>& InPaintWindow)
+	: WeakPaintWindow(InPaintWindow)
+	, RawPaintWindow(InPaintWindow.Get())
+	, MemManager(0)
+#if STATS
+	, MemManagerAllocatedMemory(0)
+#endif
+	, RenderTargetWindow(nullptr)
+	, bNeedsDeferredResolve(false)
+	, ResolveToDeferredIndex()
+	, WindowSize(FVector2D(0.0f, 0.0f))
+	//, bReportReferences(true)
 {
-	bAntialias = bInAntialias;
+	if (InPaintWindow.IsValid())
+	{
+		WindowSize = InPaintWindow->GetSizeInScreen();
+	}
 
-	NumPoints = InPoints.Num();
-	if (NumPoints > 0)
+	// Only keep UObject resources alive if this window element list is born on the game thread.
+/*
+	if (IsInGameThread())
 	{
-		Points = (FVector2D*)ElementList.Alloc(sizeof(FVector2D) * NumPoints, alignof(FVector2D));
-		FMemory::Memcpy(Points, InPoints.GetData(), sizeof(FVector2D) * NumPoints);
+		ResourceGCRoot = MakeUnique<FWindowElementGCObject>(this);
+	}*/
+}
 
-		if (InPointColors && ensure(InPointColors->Num() == NumPoints))
-		{
-			PointColors = (FLinearColor*)ElementList.Alloc(sizeof(FLinearColor) * NumPoints, alignof(FLinearColor));
-			FMemory::Memcpy(PointColors, InPointColors->GetData(), sizeof(FLinearColor) * NumPoints);
-		}
-		else
-		{
-			PointColors = nullptr;
-		}
-	}
-	else
+FSlateWindowElementList::~FSlateWindowElementList()
+{
+	/*if (ResourceGCRoot.IsValid())
 	{
-		Points = nullptr;
-	}
+		ResourceGCRoot->ClearOwner();
+	}*/
 }
 
 void FSlateDrawElement::Init(FSlateWindowElementList& ElementList, EElementType InElementType, uint32 InLayer, const FPaintGeometry& PaintGeometry, ESlateDrawEffect InDrawEffects)
 {
-	ElementType = InElementType;
-
 	RenderTransform = PaintGeometry.GetAccumulatedRenderTransform();
 	Position = PaintGeometry.DrawPosition;
 	Scale = PaintGeometry.DrawScale;
 	LocalSize = PaintGeometry.GetLocalSize();
-	ClippingIndex = ElementList.GetClippingIndex();
-	Layer = InLayer;
-	DrawEffects = InDrawEffects;
+	ClipStateHandle.SetPreCachedClipIndex(ElementList.GetClippingIndex());
+
+	LayerId = InLayer;
 
+	ElementType = InElementType;
+	DrawEffects = InDrawEffects;
+	
 	// Calculate the layout to render transform as this is needed by several calculations downstream.
 	const FSlateLayoutTransform InverseLayoutTransform(Inverse(FSlateLayoutTransform(Scale, Position)));
 
@@ -148,16 +164,16 @@ void FSlateDrawElement::Init(FSlateWindowElementList& ElementList, EElementType
 	checkSlow(Renderer);
 	SceneIndex = Renderer->GetCurrentSceneIndex();
 
-	DataPayload.BatchFlags = ESlateBatchDrawFlag::None;
-	DataPayload.BatchFlags |= static_cast<ESlateBatchDrawFlag>( static_cast<uint32>( InDrawEffects ) & static_cast<uint32>( ESlateDrawEffect::NoBlending | ESlateDrawEffect::PreMultipliedAlpha | ESlateDrawEffect::NoGamma | ESlateDrawEffect::InvertAlpha ) );
+	BatchFlags = ESlateBatchDrawFlag::None;
+	BatchFlags |= static_cast<ESlateBatchDrawFlag>(static_cast<uint32>(InDrawEffects) & static_cast<uint32>(ESlateDrawEffect::NoBlending | ESlateDrawEffect::PreMultipliedAlpha | ESlateDrawEffect::NoGamma | ESlateDrawEffect::InvertAlpha));
 
-	static_assert( ( ( __underlying_type(ESlateDrawEffect) )ESlateDrawEffect::NoBlending ) == ( ( __underlying_type(ESlateBatchDrawFlag) )ESlateBatchDrawFlag::NoBlending ), "Must keep ESlateBatchDrawFlag and ESlateDrawEffect partial matches" );
-	static_assert( ( ( __underlying_type(ESlateDrawEffect) )ESlateDrawEffect::PreMultipliedAlpha ) == ( ( __underlying_type(ESlateBatchDrawFlag) )ESlateBatchDrawFlag::PreMultipliedAlpha ), "Must keep ESlateBatchDrawFlag and ESlateDrawEffect partial matches" );
-	static_assert( ( ( __underlying_type(ESlateDrawEffect) )ESlateDrawEffect::NoGamma ) == ( ( __underlying_type(ESlateBatchDrawFlag) )ESlateBatchDrawFlag::NoGamma ), "Must keep ESlateBatchDrawFlag and ESlateDrawEffect partial matches" );
-	static_assert( ( ( __underlying_type(ESlateDrawEffect) )ESlateDrawEffect::InvertAlpha ) == ( ( __underlying_type(ESlateBatchDrawFlag) )ESlateBatchDrawFlag::InvertAlpha ), "Must keep ESlateBatchDrawFlag and ESlateDrawEffect partial matches" );
+	static_assert(((__underlying_type(ESlateDrawEffect))ESlateDrawEffect::NoBlending) == ((__underlying_type(ESlateBatchDrawFlag))ESlateBatchDrawFlag::NoBlending), "Must keep ESlateBatchDrawFlag and ESlateDrawEffect partial matches");
+	static_assert(((__underlying_type(ESlateDrawEffect))ESlateDrawEffect::PreMultipliedAlpha) == ((__underlying_type(ESlateBatchDrawFlag))ESlateBatchDrawFlag::PreMultipliedAlpha), "Must keep ESlateBatchDrawFlag and ESlateDrawEffect partial matches");
+	static_assert(((__underlying_type(ESlateDrawEffect))ESlateDrawEffect::NoGamma) == ((__underlying_type(ESlateBatchDrawFlag))ESlateBatchDrawFlag::NoGamma), "Must keep ESlateBatchDrawFlag and ESlateDrawEffect partial matches");
+	static_assert(((__underlying_type(ESlateDrawEffect))ESlateDrawEffect::InvertAlpha) == ((__underlying_type(ESlateBatchDrawFlag))ESlateBatchDrawFlag::InvertAlpha), "Must keep ESlateBatchDrawFlag and ESlateDrawEffect partial matches");
 	if ((InDrawEffects & ESlateDrawEffect::ReverseGamma) != ESlateDrawEffect::None)
 	{
-		DataPayload.BatchFlags |= ESlateBatchDrawFlag::ReverseGamma;
+		BatchFlags |= ESlateBatchDrawFlag::ReverseGamma;
 	}
 }
 
@@ -170,79 +186,62 @@ void FSlateDrawElement::ApplyPositionOffset(const FVector2D& InOffset)
 	const FSlateLayoutTransform InverseLayoutTransform(Inverse(FSlateLayoutTransform(Scale, Position)));
 }
 
-bool FSlateDrawElement::ShouldCull(const FSlateWindowElementList& ElementList)
-{
-	const FSlateClippingManager& ClippingManager = ElementList.GetClippingManager();
-	const int32 CurrentIndex = ClippingManager.GetClippingIndex();
-	if (CurrentIndex != INDEX_NONE)
-	{
-		const FSlateClippingState& ClippingState = ClippingManager.GetClippingStates()[CurrentIndex];
-		return ClippingState.HasZeroArea();
-	}
-
-	return false;
-}
-
-bool FSlateDrawElement::ShouldCull(const FSlateWindowElementList& ElementList, const FPaintGeometry& PaintGeometry, const FSlateBrush* InBrush)
+void FSlateDrawElement::MakeDebugQuad( FSlateWindowElementList& ElementList, uint32 InLayer, const FPaintGeometry& PaintGeometry)
 {
-	if (ShouldCull(ElementList, PaintGeometry))
-	{
-		return true;
-	}
+	PaintGeometry.CommitTransformsIfUsingLegacyConstructor();
 
-	if (InBrush->GetDrawType() == ESlateBrushDrawType::NoDrawType)
+	if (ShouldCull(ElementList))
 	{
-		return true;
+		return;
 	}
 
-	UObject* ResourceObject = InBrush->GetResourceObject();
-	if (!IsResourceObjectValid(ResourceObject))
-	{
-		return true;
-	}
+	FSlateDrawElement& Element = ElementList.AddUninitialized();
+	ElementList.CreatePayload<FSlateBoxPayload>(Element);
 
-	return false;
+	Element.Init(ElementList, EElementType::ET_DebugQuad, InLayer, PaintGeometry, ESlateDrawEffect::None);
 }
 
-void FSlateDrawElement::MakeDebugQuad( FSlateWindowElementList& ElementList, uint32 InLayer, const FPaintGeometry& PaintGeometry)
+FSlateDrawElement& FSlateDrawElement::MakeBoxInternal(
+	FSlateWindowElementList& ElementList,
+	uint32 InLayer,
+	const FPaintGeometry& PaintGeometry,
+	const FSlateBrush* InBrush,
+	ESlateDrawEffect InDrawEffects,
+	const FLinearColor& InTint
+)
 {
 	PaintGeometry.CommitTransformsIfUsingLegacyConstructor();
 
-	if (ShouldCull(ElementList))
-	{
-		return;
-	}
+	EElementType ElementType = (InBrush->DrawAs == ESlateBrushDrawType::Border) ? EElementType::ET_Border : EElementType::ET_Box;
 
 	FSlateDrawElement& Element = ElementList.AddUninitialized();
-	Element.Init(ElementList, ET_DebugQuad, InLayer, PaintGeometry, ESlateDrawEffect::None);
+
+	const FMargin& Margin = InBrush->GetMargin();
+	FSlateBoxPayload& BoxPayload = ElementList.CreatePayload<FSlateBoxPayload>(Element);
+
+	Element.Init(ElementList, ElementType, InLayer, PaintGeometry, InDrawEffects);
+
+	BoxPayload.SetTint(InTint);
+	BoxPayload.SetBrush(InBrush);
+
+	return Element;
 }
 
-void FSlateDrawElement::MakeBox( 
+void FSlateDrawElement::MakeBox(
 	FSlateWindowElementList& ElementList,
 	uint32 InLayer, 
 	const FPaintGeometry& PaintGeometry, 
 	const FSlateBrush* InBrush,
 	ESlateDrawEffect InDrawEffects, 
-	const FLinearColor& InTint )
+	const FLinearColor& InTint)
 {
-	PaintGeometry.CommitTransformsIfUsingLegacyConstructor();
 
 	if (ShouldCull(ElementList, PaintGeometry, InBrush, InTint))
 	{
 		return;
 	}
 
-	EElementType ElementType = (InBrush->DrawAs == ESlateBrushDrawType::Border) ? ET_Border : ET_Box;
-
-	FSlateDrawElement& Element = ElementList.AddUninitialized();
-	Element.Init(ElementList, ElementType, InLayer, PaintGeometry, InDrawEffects);
-	Element.DataPayload.SetTint(InTint);
-	Element.DataPayload.SetBrush(InBrush);
-
-	if (UObject* ResourceObject = InBrush->GetResourceObject())
-	{
-		ElementList.ResourcesToReport.Add(InBrush->GetResourceObject());
-	}
+	MakeBoxInternal(ElementList, InLayer, PaintGeometry, InBrush, InDrawEffects, InTint);
 }
 
 void FSlateDrawElement::MakeBox( 
@@ -268,30 +267,19 @@ void FSlateDrawElement::MakeRotatedBox(
 	ERotationSpace RotationSpace,
 	const FLinearColor& InTint)
 {
-	PaintGeometry.CommitTransformsIfUsingLegacyConstructor();
 
 	if (ShouldCull(ElementList, PaintGeometry, InBrush, InTint))
 	{
 		return;
 	}
 
-	EElementType ElementType = (InBrush->DrawAs == ESlateBrushDrawType::Border) ? ET_Border : ET_Box;
-
-	FSlateDrawElement& Element = ElementList.AddUninitialized();
-	Element.Init(ElementList, ElementType, InLayer, PaintGeometry, InDrawEffects);
-	Element.DataPayload.SetTint(InTint);
-	Element.DataPayload.SetBrush(InBrush);
-
+	FSlateDrawElement& DrawElement = MakeBoxInternal(ElementList, InLayer, PaintGeometry, InBrush, InDrawEffects, InTint);
+	
 	if (Angle2D != 0.0f)
 	{
 		const FVector2D RotationPoint = GetRotationPoint(PaintGeometry, InRotationPoint, RotationSpace);
 		const FSlateRenderTransform RotationTransform = Concatenate(Inverse(RotationPoint), FQuat2D(Angle2D), RotationPoint);
-		Element.SetRenderTransform(Concatenate(RotationTransform, Element.GetRenderTransform()));
-	}
-
-	if (UObject* ResourceObject = InBrush->GetResourceObject())
-	{
-		ElementList.ResourcesToReport.Add(InBrush->GetResourceObject());
+		DrawElement.SetRenderTransform(Concatenate(RotationTransform, DrawElement.GetRenderTransform()));
 	}
 }
 
@@ -306,9 +294,12 @@ void FSlateDrawElement::MakeText( FSlateWindowElementList& ElementList, uint32 I
 	}
 
 	FSlateDrawElement& Element = ElementList.AddUninitialized();
-	Element.Init(ElementList, ET_Text, InLayer, PaintGeometry, InDrawEffects);
-	Element.DataPayload.SetTint(InTint);
-	Element.DataPayload.SetText(ElementList, InText, InFontInfo, StartIndex, EndIndex);
+	FSlateTextPayload& DataPayload = ElementList.CreatePayload<FSlateTextPayload>(Element);
+
+	DataPayload.SetTint(InTint);
+	DataPayload.SetText(InText, InFontInfo, StartIndex, EndIndex);
+
+	Element.Init(ElementList, EElementType::ET_Text, InLayer, PaintGeometry, InDrawEffects);
 }
 
 void FSlateDrawElement::MakeText( FSlateWindowElementList& ElementList, uint32 InLayer, const FPaintGeometry& PaintGeometry, const FString& InText, const FSlateFontInfo& InFontInfo, ESlateDrawEffect InDrawEffects, const FLinearColor& InTint )
@@ -334,12 +325,16 @@ void FSlateDrawElement::MakeText( FSlateWindowElementList& ElementList, uint32 I
 	}
 
 	FSlateDrawElement& Element = ElementList.AddUninitialized();
-	Element.Init(ElementList, ET_Text, InLayer, PaintGeometry, InDrawEffects);
-	Element.DataPayload.SetTint(InTint);
-	Element.DataPayload.SetText(ElementList, InText, InFontInfo);
+
+	FSlateTextPayload& DataPayload = ElementList.CreatePayload<FSlateTextPayload>(Element);
+
+	DataPayload.SetTint(InTint);
+	DataPayload.SetText(InText, InFontInfo);
+
+	Element.Init(ElementList, EElementType::ET_Text, InLayer, PaintGeometry, InDrawEffects);
 }
 
-void FSlateDrawElement::MakeShapedText( FSlateWindowElementList& ElementList, uint32 InLayer, const FPaintGeometry& PaintGeometry, const FShapedGlyphSequenceRef& InShapedGlyphSequence, ESlateDrawEffect InDrawEffects, const FLinearColor& BaseTint, const FLinearColor& OutlineTint )
+void FSlateDrawElement::MakeShapedText(FSlateWindowElementList& ElementList, uint32 InLayer, const FPaintGeometry& PaintGeometry, const FShapedGlyphSequenceRef& InShapedGlyphSequence, ESlateDrawEffect InDrawEffects, const FLinearColor& BaseTint, const FLinearColor& OutlineTint)
 {
 	SCOPE_CYCLE_COUNTER(STAT_SlateDrawElementMakeTime)
 	PaintGeometry.CommitTransformsIfUsingLegacyConstructor();
@@ -362,9 +357,12 @@ void FSlateDrawElement::MakeShapedText( FSlateWindowElementList& ElementList, ui
 	}
 
 	FSlateDrawElement& Element = ElementList.AddUninitialized();
-	Element.Init(ElementList, ET_ShapedText, InLayer, PaintGeometry, InDrawEffects);
-	Element.DataPayload.SetTint(BaseTint);
-	Element.DataPayload.SetShapedText(ElementList, InShapedGlyphSequence, OutlineTint);
+
+	FSlateShapedTextPayload& DataPayload = ElementList.CreatePayload<FSlateShapedTextPayload>(Element);
+	DataPayload.SetTint(BaseTint);
+	DataPayload.SetShapedText(ElementList, InShapedGlyphSequence, OutlineTint);
+
+	Element.Init(ElementList, EElementType::ET_ShapedText, InLayer, PaintGeometry, InDrawEffects);
 }
 
 void FSlateDrawElement::MakeGradient( FSlateWindowElementList& ElementList, uint32 InLayer, const FPaintGeometry& PaintGeometry, TArray<FSlateGradientStop> InGradientStops, EOrientation InGradientType, ESlateDrawEffect InDrawEffects )
@@ -377,8 +375,12 @@ void FSlateDrawElement::MakeGradient( FSlateWindowElementList& ElementList, uint
 	}
 
 	FSlateDrawElement& Element = ElementList.AddUninitialized();
-	Element.Init(ElementList, ET_Gradient, InLayer, PaintGeometry, InDrawEffects);
-	Element.DataPayload.SetGradientPayloadProperties( InGradientStops, InGradientType );
+
+	FSlateGradientPayload& DataPayload = ElementList.CreatePayload<FSlateGradientPayload>(Element);
+
+	DataPayload.SetGradient(InGradientStops, InGradientType);
+
+	Element.Init(ElementList, EElementType::ET_Gradient, InLayer, PaintGeometry, InDrawEffects);
 }
 
 void FSlateDrawElement::MakeSpline( FSlateWindowElementList& ElementList, uint32 InLayer, const FPaintGeometry& PaintGeometry, const FVector2D& InStart, const FVector2D& InStartDir, const FVector2D& InEnd, const FVector2D& InEndDir, float InThickness, ESlateDrawEffect InDrawEffects, const FLinearColor& InTint )
@@ -389,10 +391,13 @@ void FSlateDrawElement::MakeSpline( FSlateWindowElementList& ElementList, uint32
 	{
 		return;
 	}
-
 	FSlateDrawElement& Element = ElementList.AddUninitialized();
-	Element.Init(ElementList, ET_Spline, InLayer, PaintGeometry, InDrawEffects);
-	Element.DataPayload.SetHermiteSplinePayloadProperties( InStart, InStartDir, InEnd, InEndDir, InThickness, InTint );
+
+	FSlateSplinePayload& DataPayload = ElementList.CreatePayload<FSlateSplinePayload>(Element);
+
+	DataPayload.SetHermiteSpline(InStart, InStartDir, InEnd, InEndDir, InThickness, InTint);
+
+	Element.Init(ElementList, EElementType::ET_Spline, InLayer, PaintGeometry, InDrawEffects);
 }
 
 void FSlateDrawElement::MakeCubicBezierSpline(FSlateWindowElementList & ElementList, uint32 InLayer, const FPaintGeometry & PaintGeometry, const FVector2D & P0, const FVector2D & P1, const FVector2D & P2, const FVector2D & P3, float InThickness, ESlateDrawEffect InDrawEffects, const FLinearColor & InTint)
@@ -403,10 +408,13 @@ void FSlateDrawElement::MakeCubicBezierSpline(FSlateWindowElementList & ElementL
 	{
 		return;
 	}
-
 	FSlateDrawElement& Element = ElementList.AddUninitialized();
-	Element.Init(ElementList, ET_Spline, InLayer, PaintGeometry, InDrawEffects);
-	Element.DataPayload.SetCubicBezierPayloadProperties(P0, P1, P2, P3, InThickness, InTint);
+
+	FSlateSplinePayload& DataPayload = ElementList.CreatePayload<FSlateSplinePayload>(Element);
+
+	DataPayload.SetCubicBezier(P0, P1, P2, P3, InThickness, InTint);
+
+	Element.Init(ElementList, EElementType::ET_Spline, InLayer, PaintGeometry, InDrawEffects);
 }
 
 void FSlateDrawElement::MakeDrawSpaceSpline( FSlateWindowElementList& ElementList, uint32 InLayer, const FVector2D& InStart, const FVector2D& InStartDir, const FVector2D& InEnd, const FVector2D& InEndDir, float InThickness, ESlateDrawEffect InDrawEffects, const FLinearColor& InTint )
@@ -425,8 +433,11 @@ void FSlateDrawElement::MakeDrawSpaceGradientSpline(FSlateWindowElementList& Ele
 	}
 
 	FSlateDrawElement& Element = ElementList.AddUninitialized();
-	Element.Init(ElementList, ET_Spline, InLayer, PaintGeometry, InDrawEffects);
-	Element.DataPayload.SetGradientHermiteSplinePayloadProperties(InStart, InStartDir, InEnd, InEndDir, InThickness, InGradientStops);
+
+	FSlateSplinePayload& DataPayload = ElementList.CreatePayload<FSlateSplinePayload>(Element);
+	DataPayload.SetGradientHermiteSpline(InStart, InStartDir, InEnd, InEndDir, InThickness, InGradientStops);
+
+	Element.Init(ElementList, EElementType::ET_Spline, InLayer, PaintGeometry, InDrawEffects);
 }
 
 void FSlateDrawElement::MakeDrawSpaceGradientSpline(FSlateWindowElementList& ElementList, uint32 InLayer, const FVector2D& InStart, const FVector2D& InStartDir, const FVector2D& InEnd, const FVector2D& InEndDir, const FSlateRect InClippingRect, const TArray<FSlateGradientStop>& InGradientStops, float InThickness, ESlateDrawEffect InDrawEffects)
@@ -440,8 +451,11 @@ void FSlateDrawElement::MakeDrawSpaceGradientSpline(FSlateWindowElementList& Ele
 	}
 
 	FSlateDrawElement& Element = ElementList.AddUninitialized();
-	Element.Init(ElementList, ET_Spline, InLayer, PaintGeometry, InDrawEffects);
-	Element.DataPayload.SetGradientHermiteSplinePayloadProperties(InStart, InStartDir, InEnd, InEndDir, InThickness, InGradientStops);
+
+	FSlateSplinePayload& DataPayload = ElementList.CreatePayload<FSlateSplinePayload>(Element);
+	DataPayload.SetGradientHermiteSpline(InStart, InStartDir, InEnd, InEndDir, InThickness, InGradientStops);
+
+	Element.Init(ElementList, EElementType::ET_Spline, InLayer, PaintGeometry, InDrawEffects);
 }
 
 void FSlateDrawElement::MakeLines(FSlateWindowElementList& ElementList, uint32 InLayer, const FPaintGeometry& PaintGeometry, const TArray<FVector2D>& Points, ESlateDrawEffect InDrawEffects, const FLinearColor& InTint, bool bAntialias, float Thickness)
@@ -454,17 +468,22 @@ void FSlateDrawElement::MakeLines(FSlateWindowElementList& ElementList, uint32 I
 	}
 
 	FSlateDrawElement& Element = ElementList.AddUninitialized();
-	Element.Init(ElementList, ET_Line, InLayer, PaintGeometry, InDrawEffects);
-	Element.DataPayload.SetTint(InTint);
-	Element.DataPayload.SetThickness(Thickness);
-	Element.DataPayload.SetLines(ElementList, Points, bAntialias, nullptr);
 
+	FSlateLinePayload& DataPayload = ElementList.CreatePayload<FSlateLinePayload>(Element);
+
+	DataPayload.SetTint(InTint);
+	DataPayload.SetThickness(Thickness);
+	DataPayload.SetLines(Points, bAntialias, nullptr);
+
+	ESlateDrawEffect DrawEffects = InDrawEffects;
 	if (bAntialias)
 	{
 		// If the line is to be anti-aliased, we cannot reliably snap
-		// the generated vertexes.
-		Element.DrawEffects |= ESlateDrawEffect::NoPixelSnapping;
+		// the generated vertices.
+		DrawEffects |= ESlateDrawEffect::NoPixelSnapping;
 	}
+
+	Element.Init(ElementList, EElementType::ET_Line, InLayer, PaintGeometry, DrawEffects);
 }
 
 void FSlateDrawElement::MakeLines( FSlateWindowElementList& ElementList, uint32 InLayer, const FPaintGeometry& PaintGeometry, const TArray<FVector2D>& Points, const TArray<FLinearColor>& PointColors, ESlateDrawEffect InDrawEffects, const FLinearColor& InTint, bool bAntialias, float Thickness )
@@ -477,10 +496,14 @@ void FSlateDrawElement::MakeLines( FSlateWindowElementList& ElementList, uint32
 	}
 
 	FSlateDrawElement& Element = ElementList.AddUninitialized();
-	Element.Init(ElementList, ET_Line, InLayer, PaintGeometry, InDrawEffects);
-	Element.DataPayload.SetTint(InTint);
-	Element.DataPayload.SetThickness(Thickness);
-	Element.DataPayload.SetLines(ElementList, Points, bAntialias, &PointColors);
+
+	FSlateLinePayload& DataPayload = ElementList.CreatePayload<FSlateLinePayload>(Element);
+	DataPayload.SetTint(InTint);
+	DataPayload.SetThickness(Thickness);
+	DataPayload.SetLines(Points, bAntialias, &PointColors);
+
+	Element.Init(ElementList, EElementType::ET_Line, InLayer, PaintGeometry, InDrawEffects);
+
 }
 
 void FSlateDrawElement::MakeViewport( FSlateWindowElementList& ElementList, uint32 InLayer, const FPaintGeometry& PaintGeometry, TSharedPtr<const ISlateViewport> Viewport, ESlateDrawEffect InDrawEffects, const FLinearColor& InTint )
@@ -493,8 +516,11 @@ void FSlateDrawElement::MakeViewport( FSlateWindowElementList& ElementList, uint
 	}
 
 	FSlateDrawElement& Element = ElementList.AddUninitialized();
-	Element.Init(ElementList, ET_Viewport, InLayer, PaintGeometry, InDrawEffects);
-	Element.DataPayload.SetViewportPayloadProperties( Viewport, InTint );
+	FSlateViewportPayload& DataPayload = ElementList.CreatePayload<FSlateViewportPayload>(Element);
+
+	DataPayload.SetViewport(Viewport, InTint);
+
+	Element.Init(ElementList, EElementType::ET_Viewport, InLayer, PaintGeometry, InDrawEffects);
 }
 
 
@@ -506,9 +532,12 @@ void FSlateDrawElement::MakeCustom( FSlateWindowElementList& ElementList, uint32
 	}
 
 	FSlateDrawElement& Element = ElementList.AddUninitialized();
-	Element.Init(ElementList, ET_Custom, InLayer, FPaintGeometry(), ESlateDrawEffect::None);
+
+	FSlateCustomDrawerPayload& DataPayload = ElementList.CreatePayload<FSlateCustomDrawerPayload>(Element);
+	DataPayload.SetCustomDrawer(CustomDrawer);
+	
+	Element.Init(ElementList, EElementType::ET_Custom, InLayer, FPaintGeometry(), ESlateDrawEffect::None);
 	Element.RenderTransform = FSlateRenderTransform();
-	Element.DataPayload.SetCustomDrawerPayloadProperties( CustomDrawer );
 }
 
 
@@ -522,56 +551,51 @@ void FSlateDrawElement::MakeCustomVerts(FSlateWindowElementList& ElementList, ui
 	}
 
 	FSlateDrawElement& Element = ElementList.AddUninitialized();
-	Element.Init(ElementList, ET_CustomVerts, InLayer, FPaintGeometry(), InDrawEffects);
-	Element.RenderTransform = FSlateRenderTransform();
+	FSlateCustomVertsPayload& DataPayload = ElementList.CreatePayload<FSlateCustomVertsPayload>(Element);
 
 	const FSlateShaderResourceProxy* RenderingProxy = InRenderResourceHandle.GetResourceProxy();
+	DataPayload.SetCustomVerts(RenderingProxy, InVerts, InIndexes, InInstanceData, InInstanceOffset, InNumInstances);
 
-	Element.DataPayload.SetCustomVertsPayloadProperties(RenderingProxy, InVerts, InIndexes, InInstanceData, InInstanceOffset, InNumInstances);
+	Element.Init(ElementList, EElementType::ET_CustomVerts, InLayer, FPaintGeometry(), InDrawEffects);
+	Element.RenderTransform = FSlateRenderTransform();
 }
 
-void FSlateDrawElement::MakeCachedBuffer(FSlateWindowElementList& ElementList, uint32 InLayer, TSharedPtr<FSlateRenderDataHandle, ESPMode::ThreadSafe>& CachedRenderDataHandle, const FVector2D& Offset)
+void FSlateDrawElement::MakePostProcessPass(FSlateWindowElementList& ElementList, uint32 InLayer, const FPaintGeometry& PaintGeometry, const FVector4& Params, int32 DownsampleAmount)
 {
-	if (ShouldCull(ElementList))
-	{
-		return;
-	}
+	PaintGeometry.CommitTransformsIfUsingLegacyConstructor();
 
-	// Don't draw invalid render data handles.
-	if (!CachedRenderDataHandle.IsValid())
+	if (ShouldCull(ElementList))
 	{
 		return;
 	}
 
 	FSlateDrawElement& Element = ElementList.AddUninitialized();
-	Element.Init(ElementList, ET_CachedBuffer, InLayer, FPaintGeometry(), ESlateDrawEffect::None);
-	Element.DataPayload.SetCachedBuffer(CachedRenderDataHandle.Get(), Offset);
 
-	// Note that the buffer is currently in use, this avoid releasing it back into a pool.
-	ElementList.BeginUsingCachedBuffer(CachedRenderDataHandle);
+	FSlatePostProcessPayload& DataPayload = ElementList.CreatePayload<FSlatePostProcessPayload>(Element);
+	DataPayload.DownsampleAmount = DownsampleAmount;
+	DataPayload.PostProcessData = Params;
+
+	Element.Init(ElementList, EElementType::ET_PostProcessPass, InLayer, PaintGeometry, ESlateDrawEffect::None);
 }
 
-void FSlateDrawElement::MakeLayer(FSlateWindowElementList& ElementList, uint32 InLayer, TSharedPtr<FSlateDrawLayerHandle, ESPMode::ThreadSafe>& DrawLayerHandle)
+FSlateDrawElement::FSlateDrawElement()
+	: DataPayload(nullptr)
+	, bIsCached(false)
 {
-	FSlateDrawElement& Element = ElementList.AddUninitialized();
-	Element.Init(ElementList, ET_Layer, InLayer, FPaintGeometry(), ESlateDrawEffect::None);
-	Element.RenderTransform = FSlateRenderTransform();
-	Element.DataPayload.SetLayerPayloadProperties(DrawLayerHandle.Get());
+
 }
 
-void FSlateDrawElement::MakePostProcessPass(FSlateWindowElementList& ElementList, uint32 InLayer, const FPaintGeometry& PaintGeometry, const FVector4& Params, int32 DownsampleAmount)
+FSlateDrawElement::~FSlateDrawElement()
 {
-	PaintGeometry.CommitTransformsIfUsingLegacyConstructor();
-
-	if (ShouldCull(ElementList))
+	if (bIsCached)
 	{
-		return;
+		delete DataPayload;
+	}
+	else if (DataPayload)
+	{
+		// Allocated by a memstack so we just need to call the destructor manually
+		DataPayload->~FSlateDataPayload();
 	}
-
-	FSlateDrawElement& Element = ElementList.AddUninitialized();
-	Element.Init(ElementList, ET_PostProcessPass, InLayer, PaintGeometry, ESlateDrawEffect::None);
-	Element.DataPayload.DownsampleAmount = DownsampleAmount;
-	Element.DataPayload.PostProcessData = Params;
 }
 
 FVector2D FSlateDrawElement::GetRotationPoint(const FPaintGeometry& PaintGeometry, const TOptional<FVector2D>& UserRotationPoint, ERotationSpace RotationSpace)
@@ -602,314 +626,39 @@ FVector2D FSlateDrawElement::GetRotationPoint(const FPaintGeometry& PaintGeometr
 	return RotationPoint;
 }
 
-void FSlateBatchData::Reset()
-{
-	RenderBatches.Reset();
-	
-	// note: LayerToElementBatches is not reset here as the same layers are 
-	// more than likely reused and we can save memory allocations by not resetting the map every frame
-
-	NumBatchedVertices = 0;
-	NumBatchedIndices = 0;
-	NumLayers = 0;
-
-	bIsStencilBufferRequired = false;
-
-	RenderDataHandle.Reset();
-}
-
-#define MAX_VERT_ARRAY_RECYCLE (200)
-#define MAX_INDEX_ARRAY_RECYCLE (500)
-
-bool FSlateBatchData::IsStencilClippingRequired() const
-{
-	return bIsStencilBufferRequired;
-}
-
-void FSlateBatchData::DetermineIsStencilClippingRequired(const TArray<FSlateClippingState>& ClippingStates)
-{
-	bIsStencilBufferRequired = false;
-
-	for (const FSlateClippingState& Clipping : ClippingStates)
-	{
-		if (Clipping.GetClippingMethod() == EClippingMethod::Stencil)
-		{
-			bIsStencilBufferRequired = true;
-			return;
-		}
-	}
-}
-
-void FSlateBatchData::AssignVertexArrayToBatch( FSlateElementBatch& Batch )
+void FSlateDrawElement::AddReferencedObjects(FReferenceCollector& Collector)
 {
-	// Get a free vertex array
-	if (VertexArrayFreeList.Num() > 0)
+	if(DataPayload)
 	{
-		Batch.VertexArrayIndex = VertexArrayFreeList.Pop(/*bAllowShrinking=*/ false);
-	}
-	else
-	{
-		// There are no free vertex arrays so we must add one		
-		uint32 NewIndex = BatchVertexArrays.Add(FSlateVertexArray());
-		ResetVertexArray(BatchVertexArrays[NewIndex]);
-
-		Batch.VertexArrayIndex = NewIndex;
+		DataPayload->AddReferencedObjects(Collector);
 	}
 }
 
-void FSlateBatchData::AssignIndexArrayToBatch( FSlateElementBatch& Batch )
-{
-	// Get a free index array
-	if (IndexArrayFreeList.Num() > 0)
-	{
-		Batch.IndexArrayIndex = IndexArrayFreeList.Pop(/*bAllowShrinking=*/ false);
-	}
-	else
-	{
-		// There are no free index arrays so we must add one
-		uint32 NewIndex = BatchIndexArrays.Add(FSlateIndexArray());
-		ResetIndexArray(BatchIndexArrays[NewIndex]);
-
-		Batch.IndexArrayIndex = NewIndex;
-	}
-}
 
-void FSlateBatchData::FillVertexAndIndexBuffer(uint8* VertexBuffer, uint8* IndexBuffer, bool bAbsoluteIndices)
+FSlateDrawElement& FSlateWindowElementList::AddUninitialized()
 {
-	int32 IndexOffset = 0;
-	int32 VertexOffset = 0;
-	int32 BaseVertexIndex = 0;
-
-	const bool bValidBuffers = (nullptr != VertexBuffer) && (nullptr != IndexBuffer);
+	const bool bAllowCache = CachedElementDataListStack.Num() > 0 && WidgetDrawStack.Num() && !WidgetDrawStack.Top().bIsVolatile;
 
-	for ( const FSlateRenderBatch& Batch : RenderBatches )
+	if (bAllowCache)
 	{
-		// Ignore foreign batches that are inserted into our render set.
-		if ( RenderDataHandle != Batch.CachedRenderHandle )
-		{
-			continue;
-		}
-
-		if ( Batch.VertexArrayIndex != INDEX_NONE && Batch.IndexArrayIndex != INDEX_NONE )
-		{
-			FSlateVertexArray& Vertices = BatchVertexArrays[Batch.VertexArrayIndex];
-			FSlateIndexArray& Indices = BatchIndexArrays[Batch.IndexArrayIndex];
-	
-			if ( Vertices.Num() && Indices.Num() )
-			{
-				if (bValidBuffers)
-				{
-					uint32 RequiredVertexSize = Vertices.Num() * Vertices.GetTypeSize();
-					uint32 RequiredIndexSize = Indices.Num() * Indices.GetTypeSize();
-
-					FMemory::Memcpy(VertexBuffer + VertexOffset, Vertices.GetData(), RequiredVertexSize);
-					if (BaseVertexIndex == 0 || !bAbsoluteIndices)
-					{
-						FMemory::Memcpy(IndexBuffer + IndexOffset, Indices.GetData(), RequiredIndexSize);
-					}
-					else
-					{
-						SlateIndex* TargetIndexBuffer = (SlateIndex*)(IndexBuffer + IndexOffset);
-						for (int32 i = 0; i < Indices.Num(); ++i)
-						{
-							TargetIndexBuffer[i] = Indices[i] + BaseVertexIndex;
-						}
-					}
-
-					BaseVertexIndex += Vertices.Num();
-					IndexOffset += (Indices.Num() * sizeof(SlateIndex));
-					VertexOffset += (Vertices.Num() * sizeof(FSlateVertex));
-				}
-
-				Vertices.Reset();
-				Indices.Reset();
-
-				if ( Vertices.GetSlack() > MAX_VERT_ARRAY_RECYCLE )
-				{
-					ResetVertexArray(Vertices);
-				}
-
-				if ( Indices.GetSlack() > MAX_INDEX_ARRAY_RECYCLE )
-				{
-					ResetIndexArray(Indices);
-				}
-			}
-
-			VertexArrayFreeList.Add(Batch.VertexArrayIndex);
-			IndexArrayFreeList.Add(Batch.IndexArrayIndex);
-		}
-	}
-}
-
-void FSlateBatchData::CreateRenderBatches(FElementBatchMap& LayerToElementBatches)
-{
-	checkSlow(IsInRenderingThread());
-
-	uint32 VertexOffset = 0;
-	uint32 IndexOffset = 0;
-
-	{
-		SCOPED_NAMED_EVENT_TEXT("SlateRT::CreateRenderBatches", FColor::Magenta);
-		Merge(LayerToElementBatches, VertexOffset, IndexOffset);
+		// @todo get working with slate debugging
+		return AddCachedElement();
 	}
-
-	// 
-	if ( RenderDataHandle.IsValid() )
-	{
-		RenderDataHandle->SetRenderBatches(&RenderBatches);
-	}
-}
-
-void FSlateBatchData::AddRenderBatch(uint32 InLayer, const FSlateElementBatch& InElementBatch, int32 InNumVertices, int32 InNumIndices, int32 InVertexOffset, int32 InIndexOffset)
-{
-	NumBatchedVertices += InNumVertices;
-	NumBatchedIndices += InNumIndices;
-
-	const int32 Index = RenderBatches.Add(FSlateRenderBatch(InLayer, InElementBatch, RenderDataHandle, InNumVertices, InNumIndices, InVertexOffset, InIndexOffset));
-	RenderBatches[Index].DynamicOffset = FVector2D::ZeroVector;
-}
-
-void FSlateBatchData::ResetVertexArray(FSlateVertexArray& InOutVertexArray)
-{
-	InOutVertexArray.Empty(MAX_VERT_ARRAY_RECYCLE);
-}
-
-void FSlateBatchData::ResetIndexArray(FSlateIndexArray& InOutIndexArray)
-{
-	InOutIndexArray.Empty(MAX_INDEX_ARRAY_RECYCLE);
-}
-
-
-void FSlateBatchData::Merge(FElementBatchMap& InLayerToElementBatches, uint32& VertexOffset, uint32& IndexOffset)
-{
-	InLayerToElementBatches.Sort();
-
-	const bool bExpandLayersAndCachedHandles = RenderDataHandle.IsValid() == false;
-
-	InLayerToElementBatches.ForEachLayer([&] (uint32 Layer, FElementBatchArray& ElementBatches)
+	else
 	{
-		++NumLayers;
-		for ( FElementBatchArray::TIterator BatchIt(ElementBatches); BatchIt; ++BatchIt )
-		{
-			FSlateElementBatch& ElementBatch = *BatchIt;
-
-			if ( ElementBatch.GetCustomDrawer().IsValid() )
-			{
-				AddRenderBatch(Layer, ElementBatch, 0, 0, 0, 0);
-			}
-			else if(ElementBatch.GetShaderType() == ESlateShader::PostProcess)
-			{
-				AddRenderBatch(Layer, ElementBatch, 0, 0, 0, 0);
-			}
-			else
-			{
-				if ( bExpandLayersAndCachedHandles )
-				{
-					if ( FSlateRenderDataHandle* RenderHandle = ElementBatch.GetCachedRenderHandle().Get() )
-					{
-						TArray<FSlateRenderBatch>* ForeignBatches = RenderHandle->GetRenderBatches();
-						if (ForeignBatches)
-						{
-							TArray<FSlateRenderBatch>& ForeignBatchesRef = *ForeignBatches;
-							for ( int32 i = 0; i < ForeignBatches->Num(); i++ )
-							{
-								TSharedPtr<FSlateDrawLayerHandle, ESPMode::ThreadSafe> LayerHandle = ForeignBatchesRef[i].LayerHandle.Pin();
-								if ( LayerHandle.IsValid() )
-								{
-									// If a record was added for a layer, but nothing was ever drawn for it, the batch map will be null.
-									if ( LayerHandle->BatchMap )
-									{
-										Merge(*LayerHandle->BatchMap, VertexOffset, IndexOffset);
-										LayerHandle->BatchMap = nullptr;
-									}
-								}
-								else
-								{
-									const int32 Index = RenderBatches.Add(ForeignBatchesRef[i]);
-									RenderBatches[Index].DynamicOffset = ElementBatch.GetCachedRenderDataOffset();
-								}
-							}
-						}
-
-						continue;
-					}
-				}
-				else
-				{
-					// Insert if we're not expanding
-					if ( FSlateDrawLayerHandle* LayerHandle = ElementBatch.GetLayerHandle().Get() )
-					{
-						AddRenderBatch(Layer, ElementBatch, 0, 0, 0, 0);
-						continue;
-					}
-				}
-				
-				// This is the normal path, for draw buffers that just contain Vertices and Indices.
-				if ( ElementBatch.VertexArrayIndex != INDEX_NONE && ElementBatch.IndexArrayIndex != INDEX_NONE )
-				{
-					FSlateVertexArray& BatchVertices = GetBatchVertexList(ElementBatch);
-					FSlateIndexArray& BatchIndices = GetBatchIndexList(ElementBatch);
-
-					// We should have at least some vertices and indices in the batch or none at all
-					check(BatchVertices.Num() > 0 && BatchIndices.Num() > 0 || BatchVertices.Num() == 0 && BatchIndices.Num() == 0);
-
-					if ( BatchVertices.Num() > 0 && BatchIndices.Num() > 0 )
-					{
-						const int32 NumVertices = BatchVertices.Num();
-						const int32 NumIndices = BatchIndices.Num();
-
-						AddRenderBatch(Layer, ElementBatch, NumVertices, NumIndices, VertexOffset, IndexOffset);
-
-						VertexOffset += BatchVertices.Num();
-						IndexOffset += BatchIndices.Num();
-					}
-					else
-					{
-						VertexArrayFreeList.Add(ElementBatch.VertexArrayIndex);
-						IndexArrayFreeList.Add(ElementBatch.IndexArrayIndex);
-					}
-				}
-			}
-		}
-
-		ElementBatches.Reset();
-	});
+		FSlateDrawElementArray& Elements = UncachedDrawElements;
+		const int32 InsertIdx = Elements.AddDefaulted();
 
-	InLayerToElementBatches.Reset();
-}
+#if WITH_SLATE_DEBUGGING
+		FSlateDebugging::ElementAdded.Broadcast(*this, InsertIdx);
+#endif
 
-void FSlateWindowElementList::MergeResources(const TArray<UObject*>& AssociatedResources)
-{
-	for (UObject* AssociatedResource : AssociatedResources)
-	{
-		IsResourceObjectValid(AssociatedResource);
+		FSlateDrawElement& NewElement = Elements[InsertIdx];
+		//NewElement.SetBatchPriorityGroup(GetCurrentBatchPriorityGroup());
+		return Elements[InsertIdx];
 	}
-
-	ResourcesToReport.Append(AssociatedResources);
 }
 
-void FSlateWindowElementList::MergeElementList(FSlateWindowElementList* ElementList, FVector2D AbsoluteOffset)
-{
-	const bool bMoved = !AbsoluteOffset.IsZero();
-
-	if (bMoved)
-	{
-		const TArray<FSlateDrawElement>& CachedElements = ElementList->GetDrawElements();
-		const int32 CachedElementCount = CachedElements.Num();
-
-		for (int32 Index = 0; Index < CachedElementCount; Index++)
-		{
-			const FSlateDrawElement& LocalElement = CachedElements[Index];
-			FSlateDrawElement AbsElement = LocalElement;
-			AbsElement.ApplyPositionOffset(AbsoluteOffset);
-			AddItem(AbsElement);
-		}
-	}
-	else
-	{
-		AppendItems(ElementList);
-	}
-}
 
 FSlateWindowElementList::FDeferredPaint::FDeferredPaint( const TSharedRef<const SWidget>& InWidgetToPaint, const FPaintArgs& InArgs, const FGeometry InAllottedGeometry, const FWidgetStyle& InWidgetStyle, bool InParentEnabled )
 	: WidgetToPaintPtr( InWidgetToPaint )
@@ -980,151 +729,112 @@ void FSlateWindowElementList::EndDeferredGroup()
 	bNeedsDeferredResolve = true;
 }
 
-
-
-FSlateWindowElementList::FVolatilePaint::FVolatilePaint(const TSharedRef<const SWidget>& InWidgetToPaint, const FPaintArgs& InArgs, const FGeometry InAllottedGeometry, const FSlateRect InMyCullingRect, const TOptional<FSlateClippingState>& InClippingState, int32 InLayerId, const FWidgetStyle& InWidgetStyle, bool InParentEnabled)
-	: WidgetToPaintPtr(InWidgetToPaint)
-	, Args(InArgs.EnableCaching(InArgs.GetLayoutCache(), InArgs.GetParentCacheNode(), false, true))
-	, AllottedGeometry(InAllottedGeometry)
-	, MyCullingRect(InMyCullingRect)
-	, ClippingState(InClippingState)
-	, LayerId(InLayerId)
-	, WidgetStyle(InWidgetStyle)
-	, bParentEnabled(InParentEnabled)
-{
-}
-
-static const FName InvalidationPanelName(TEXT("SInvalidationPanel"));
-
-int32 FSlateWindowElementList::FVolatilePaint::ExecutePaint(FSlateWindowElementList& OutDrawElements, double InCurrentTime, float InDeltaTime, const FVector2D& InDynamicOffset) const
+void FSlateWindowElementList::PushPaintingWidget(const SWidget& CurrentWidget, int32 StartingLayerId, FSlateCachedElementListNode* CurrentCacheNode)
 {
-	TSharedPtr<const SWidget> WidgetToPaint = WidgetToPaintPtr.Pin();
-	if ( WidgetToPaint.IsValid() )
+	FSlateCachedElementData* CurrentCachedElementData = GetCurrentCachedElementData();
+	if (CurrentCachedElementData)
 	{
-#if SLATE_VERBOSE_NAMED_EVENTS
-		SCOPED_NAMED_EVENT_FSTRING(FReflectionMetaData::GetWidgetDebugInfo(WidgetToPaint.Get()), FColor::Orange);
-#endif
+		const FWidgetDrawElementState& PreviousState = WidgetDrawStack.Num() ? WidgetDrawStack.Top() : FWidgetDrawElementState(nullptr, false, nullptr);
+		const bool bParentVolatile = WidgetDrawStack.Num() ? PreviousState.bIsVolatile : false;
 
-		// Have to run a slate pre-pass for all volatile elements, some widgets cache information like 
-		// the STextBlock.  This may be all kinds of terrible an idea to do during paint.
-		SWidget* MutableWidget = const_cast<SWidget*>( WidgetToPaint.Get() );
-
-		if ( MutableWidget->GetType() != InvalidationPanelName )
-		{
-			MutableWidget->SlatePrepass(AllottedGeometry.Scale);
-		}
-
-		FPaintArgs PaintArgs = Args.WithNewTime(InCurrentTime, InDeltaTime);
-
-		if (ClippingState.IsSet())
-		{
-			FSlateClippingState ExistingClippingState = ClippingState.GetValue();
-			OutDrawElements.GetClippingManager().PushAndMergePartialClippingState(ExistingClippingState);
-		}
+		WidgetDrawStack.Emplace(CurrentCacheNode, bParentVolatile || CurrentWidget.IsVolatile(), &CurrentWidget);
 
-		int32 NewLayer = 0;
-		if (InDynamicOffset.IsZero())
+		// When a widget is pushed reset its draw elements.  They are being recached or possibly going away
+		if (CurrentCacheNode != nullptr)
 		{
-			NewLayer = WidgetToPaint->Paint(PaintArgs, AllottedGeometry, MyCullingRect, OutDrawElements, LayerId, WidgetStyle, bParentEnabled);
-		}
-		else
-		{
-			FSlateRect LocalRect = MyCullingRect.OffsetBy(InDynamicOffset);
-			FGeometry LocalGeometry = AllottedGeometry;
-			LocalGeometry.AppendTransform(FSlateLayoutTransform(InDynamicOffset));
-			
-			NewLayer = WidgetToPaint->Paint(PaintArgs, LocalGeometry, LocalRect, OutDrawElements, LayerId, WidgetStyle, bParentEnabled);
-		}
-
-		if (ClippingState.IsSet())
-		{
-			OutDrawElements.GetClippingManager().PopClip();
+#if WITH_SLATE_DEBUGGING
+			check(CurrentCacheNode->GetValue().Widget == &CurrentWidget);
+#endif
+			CurrentCachedElementData->ResetCache(CurrentCacheNode);
 		}
-
-		return NewLayer;
 	}
-
-	return LayerId;
 }
 
-void FSlateWindowElementList::QueueVolatilePainting(const FVolatilePaint& InVolatilePaint)
-{
-	TSharedPtr< FSlateDrawLayerHandle, ESPMode::ThreadSafe > LayerHandle = MakeShareable(new FSlateDrawLayerHandle());
-
-	FSlateDrawElement::MakeLayer(*this, InVolatilePaint.GetLayerId(), LayerHandle);
 
-	const int32 NewEntryIndex = VolatilePaintList.Add(MakeShareable(new FVolatilePaint(InVolatilePaint)));
-	VolatilePaintList[NewEntryIndex]->LayerHandle = LayerHandle;
-}
-
-int32 FSlateWindowElementList::PaintVolatile(FSlateWindowElementList& OutElementList, double InCurrentTime, float InDeltaTime, const FVector2D& InDynamicOffset)
+FSlateCachedElementListNode* FSlateWindowElementList::PopPaintingWidget()
 {
-	int32 MaxLayerId = 0;
-
-	for ( int32 VolatileIndex = 0; VolatileIndex < VolatilePaintList.Num(); ++VolatileIndex )
+	FSlateCachedElementData* CurrentCachedElementData = GetCurrentCachedElementData();
+	if (CurrentCachedElementData)
 	{
-		const TSharedPtr<FVolatilePaint>& Args = VolatilePaintList[VolatileIndex];
-
-		OutElementList.BeginLogicalLayer(Args->LayerHandle);
-		MaxLayerId = FMath::Max(MaxLayerId, Args->ExecutePaint(OutElementList, InCurrentTime, InDeltaTime, InDynamicOffset));
-		OutElementList.EndLogicalLayer();
+		return WidgetDrawStack.Pop().CacheNode;
 	}
 
-	return MaxLayerId;
+	return nullptr;
 }
 
-int32 FSlateWindowElementList::PaintVolatileRootLayer(FSlateWindowElementList& OutElementList, double InCurrentTime, float InDeltaTime, const FVector2D& InDynamicOffset)
+/*
+int32 FSlateWindowElementList::PushBatchPriortyGroup(const SWidget& CurrentWidget)
 {
-	int32 MaxLayerId = 0;
-
-	for (int32 VolatileIndex = 0; VolatileIndex < VolatilePaintList.Num(); ++VolatileIndex)
+	int32 NewPriorityGroup = 0;
+/ *
+	if (GSlateEnableGlobalInvalidation)
+	{
+		NewPriorityGroup = BatchDepthPriorityStack.Add_GetRef(CurrentWidget.FastPathProxyHandle.IsValid() ? CurrentWidget.FastPathProxyHandle.GetIndex() : 0);
+	}
+	else
 	{
-		const TSharedPtr<FVolatilePaint>& Args = VolatilePaintList[VolatileIndex];
-		MaxLayerId = FMath::Max(MaxLayerId, Args->ExecutePaint(OutElementList, InCurrentTime, InDeltaTime, InDynamicOffset));
+		NewPriorityGroup = BatchDepthPriorityStack.Add_GetRef(MaxPriorityGroup + 1);
+		//NewPriorityGroup = BatchDepthPriorityStack.Add_GetRef(0);
 	}
 
-	return MaxLayerId;
+	// Should be +1 or the first overlay slot will not appear on top of stuff below it?
+	// const int32 NewPriorityGroup = BatchDepthPriorityStack.Add_GetRef(BatchDepthPriorityStack.Num() ? BatchDepthPriorityStack.Top()+1 : 1);
+
+	MaxPriorityGroup = FMath::Max(NewPriorityGroup, MaxPriorityGroup);* /
+	return NewPriorityGroup;
 }
 
+int32 FSlateWindowElementList::PushAbsoluteBatchPriortyGroup(int32 BatchPriorityGroup)
+{
+	return 0;// return BatchDepthPriorityStack.Add_GetRef(BatchPriorityGroup);
+}
 
-void FSlateWindowElementList::BeginLogicalLayer(const TSharedPtr<FSlateDrawLayerHandle, ESPMode::ThreadSafe>& LayerHandle)
+void FSlateWindowElementList::PopBatchPriortyGroup()
 {
-	// Don't attempt to begin logical layers inside a cached view of the data.
-	checkSlow(!IsCachedRenderDataInUse());
+	//BatchDepthPriorityStack.Pop();
+}*/
 
-	TSharedPtr<FSlateDrawLayer> Layer;
-	{
-		//SCOPED_NAMED_EVENT(FindLayer, FColor::Orange);
-		Layer = DrawLayers.FindRef(LayerHandle);
-	}
+FSlateDrawElement& FSlateWindowElementList::AddCachedElement()
+{
+	FSlateCachedElementData* CurrentCachedElementData = GetCurrentCachedElementData();
+	check(CurrentCachedElementData);
 
-	if ( !Layer.IsValid() )
-	{
-		if ( DrawLayerPool.Num() > 0 )
-		{
-			Layer = DrawLayerPool.Pop(false);
-		}
-		else
-		{
-			Layer = MakeShareable(new FSlateDrawLayer());
-		}
+	FWidgetDrawElementState& CurrentWidgetState = WidgetDrawStack.Top();
+	check(!CurrentWidgetState.bIsVolatile);
 
-		//SCOPED_NAMED_EVENT(AddLayer, FColor::Orange);
-		DrawLayers.Add(LayerHandle, Layer);
+	if (CurrentWidgetState.CacheNode == nullptr)
+	{
+#if WITH_SLATE_DEBUGGING
+		CurrentWidgetState.CacheNode = CurrentCachedElementData->AddCache(CurrentWidgetState.Widget);
+#else
+		CurrentWidgetState.CacheNode = CurrentCachedElementData->AddCache(nullptr);
+#endif
 	}
 
-	//SCOPED_NAMED_EVENT(PushLayer, FColor::Orange);
-	DrawStack.Push(Layer.Get());
+#if WITH_SLATE_DEBUGGING
+	return CurrentCachedElementData->AddCachedElement(CurrentWidgetState.CacheNode, GetClippingManager(), CurrentWidgetState.Widget);
+#else
+	return CurrentCachedElementData->AddCachedElement(CurrentWidgetState.CacheNode, GetClippingManager(), nullptr);
+#endif
+}
+
+void FSlateWindowElementList::PushCachedElementData(FSlateCachedElementData& CachedElementData)
+{
+	const int32 Index = CachedElementDataList.Add(&CachedElementData);
+	CachedElementDataListStack.Push(Index);
 }
 
-void FSlateWindowElementList::EndLogicalLayer()
+void FSlateWindowElementList::PopCachedElementData()
 {
-	DrawStack.Pop();
+	CachedElementDataListStack.Pop();
 }
 
-void FSlateWindowElementList::PushClip(const FSlateClippingZone& InClipZone)
+
+int32 FSlateWindowElementList::PushClip(const FSlateClippingZone& InClipZone)
 {
-	ClippingManager.PushClip(InClipZone);
+	const int32 NewClipIndex = ClippingManager.PushClip(InClipZone);
+
+	return NewClipIndex;
 }
 
 int32 FSlateWindowElementList::GetClippingIndex() const
@@ -1142,173 +852,191 @@ void FSlateWindowElementList::PopClip()
 	ClippingManager.PopClip();
 }
 
-FSlateRenderDataHandle::FSlateRenderDataHandle(const ILayoutCache* InCacher, ISlateRenderDataManager* InManager)
-	: Cacher(InCacher)
-	, Manager(InManager)
-	, RenderBatches(nullptr)
-	, UsageCount(0)
+void FSlateWindowElementList::PopClipToStackIndex(int32 Index)
 {
+	ClippingManager.PopToStackIndex(Index);
 }
 
-FSlateRenderDataHandle::~FSlateRenderDataHandle()
-{
-	if ( Manager )
-	{
-		Manager->BeginReleasingRenderData(this);
-	}
-}
 
-void FSlateRenderDataHandle::Disconnect()
+void FSlateWindowElementList::SetRenderTargetWindow(SWindow* InRenderTargetWindow)
 {
-	Manager = nullptr;
-	RenderBatches = nullptr;
+	check(IsThreadSafeForSlateRendering());
+	RenderTargetWindow = InRenderTargetWindow;
 }
 
-TSharedRef<FSlateRenderDataHandle, ESPMode::ThreadSafe> FSlateWindowElementList::CacheRenderData(const ILayoutCache* Cacher)
+DECLARE_MEMORY_STAT(TEXT("FSlateWindowElementList MemManager"), STAT_FSlateWindowElementListMemManager, STATGROUP_SlateMemory);
+DECLARE_DWORD_COUNTER_STAT(TEXT("FSlateWindowElementList MemManager Count"), STAT_FSlateWindowElementListMemManagerCount, STATGROUP_SlateMemory);
+
+void FSlateWindowElementList::ResetElementList()
 {
-	// Don't attempt to use this slate window element list if the cache is still being used.
-	checkSlow(!IsCachedRenderDataInUse());
+	QUICK_SCOPE_CYCLE_COUNTER(Slate_ResetElementList);
 
-	FSlateRenderer* Renderer = FSlateApplicationBase::Get().GetRenderer();
+	check(IsThreadSafeForSlateRendering());
+
+	DeferredPaintList.Reset();
+
+	BatchData.ResetData();
+
+	ClippingManager.ResetClippingState();
 
-	TSharedRef<FSlateRenderDataHandle, ESPMode::ThreadSafe> CachedRenderDataHandleRef = Renderer->CacheElementRenderData(Cacher, *this);
-	CachedRenderDataHandle = CachedRenderDataHandleRef;
+	UncachedDrawElements.Reset();
 
-	return CachedRenderDataHandleRef;
+#if STATS
+	const int32 DeltaMemory = MemManager.GetByteCount() - MemManagerAllocatedMemory;
+	INC_DWORD_STAT(STAT_FSlateWindowElementListMemManagerCount);
+	INC_MEMORY_STAT_BY(STAT_FSlateWindowElementListMemManager, DeltaMemory);
+
+	MemManagerAllocatedMemory = MemManager.GetByteCount();
+#endif
+
+	MemManager.Flush();
+	
+	CachedElementDataList.Empty();
+
+	check(WidgetDrawStack.Num() == 0);
+
+	RenderTargetWindow = nullptr;
 }
 
-void FSlateWindowElementList::PreDraw_ParallelThread()
-{
-	check(IsInParallelRenderingThread());
 
-	for ( auto& Entry : DrawLayers )
+void FSlateWindowElementList::AddReferencedObjects(FReferenceCollector& Collector)
+{
+	for (FSlateDrawElement& Element : UncachedDrawElements)
 	{
-		checkSlow(Entry.Key->BatchMap == nullptr);
-		Entry.Key->BatchMap = &Entry.Value->GetElementBatchMap();
+		Element.AddReferencedObjects(Collector);
 	}
 }
 
-void FSlateWindowElementList::PostDraw_ParallelThread()
+static const FSlateClippingState* GetClipStateFromParent(const FSlateClippingManager& ParentClipManager)
 {
-	check(IsInParallelRenderingThread());
-
-	PostDraw_NonParallelRenderer();
-}
+	const int32 ClippingIndex = ParentClipManager.GetClippingIndex();
 
-void FSlateWindowElementList::PostDraw_NonParallelRenderer()
-{
-	for ( auto& Entry : DrawLayers )
+	if(ClippingIndex != INDEX_NONE)
 	{
-		Entry.Key->BatchMap = nullptr;
+		return &ParentClipManager.GetClippingStates()[ClippingIndex];
 	}
-
-	for ( TSharedPtr<FSlateRenderDataHandle, ESPMode::ThreadSafe>& Handle : CachedRenderHandlesInUse )
+	else
 	{
-		Handle->EndUsing();
+		return nullptr;
 	}
-
-	CachedRenderHandlesInUse.Reset();
-	bReportReferences = false;
-}
-
-SLATECORE_API void FSlateWindowElementList::SetRenderTargetWindow(SWindow* InRenderTargetWindow)
-{
-	check(IsThreadSafeForSlateRendering());
-	RenderTargetWindow = InRenderTargetWindow;
 }
 
-DECLARE_MEMORY_STAT(TEXT("FSlateWindowElementList MemManager"), STAT_FSlateWindowElementListMemManager, STATGROUP_SlateVerbose);
-DECLARE_DWORD_COUNTER_STAT(TEXT("FSlateWindowElementList MemManager Count"), STAT_FSlateWindowElementListMemManagerCount, STATGROUP_SlateVerbose);
-
-void FSlateWindowElementList::ResetElementBuffers()
+FSlateDrawElement& FSlateCachedElementData::AddCachedElement(FSlateCachedElementListNode* CacheNode, const FSlateClippingManager& ParentClipManager, const SWidget* CurrentWidget)
 {
-	// Don't attempt to use this slate window element list if the cache is still being used.
-	checkSlow(!IsCachedRenderDataInUse());
-	check(IsThreadSafeForSlateRendering());
+#if WITH_SLATE_DEBUGGING
+	check(CacheNode->GetValue().Widget == CurrentWidget);
+	check(CurrentWidget->GetParentWidget().IsValid());
+#endif
 
-	// Reset the Main Thread Resources, because we no longer need to keep these referenced objects alive.
-	ResourcesToReport.Reset();
+	FSlateCachedElementList& List = CacheNode->GetValue();
+	FSlateDrawElement& NewElement = List.DrawElements.AddDefaulted_GetRef();
+	NewElement.SetIsCached(true);
 
-	DeferredPaintList.Reset();
-	VolatilePaintList.Reset();
-	BatchData.Reset();
+	List.bNewData = true;
+	const FSlateClippingState* ExistingClipState = GetClipStateFromParent(ParentClipManager);
 
-	// Reset the draw elements on the root draw layer
-	RootDrawLayer.ResetLayer();
-	ClippingManager.ResetClippingState();
-
-	// Return child draw layers to the pool, and reset their draw elements.
-	for ( auto& Entry : DrawLayers )
+	if (ExistingClipState)
 	{
-		FSlateDrawLayer* Layer = Entry.Value.Get();
-		Layer->ResetLayer();
-		DrawLayerPool.Add(Entry.Value);
+		// We need to cache this clip state for the next time the element draws
+		FSlateCachedClipState& CachedClipState = FindOrAddCachedClipState(ExistingClipState);
+		List.AddCachedClipState(CachedClipState);
+		NewElement.SetCachedClippingState(&CachedClipState.ClippingState);
 	}
 
-	DrawLayers.Reset();
-
-	DrawStack.Reset();
-	DrawStack.Push(&RootDrawLayer);
-
-	INC_DWORD_STAT(STAT_FSlateWindowElementListMemManagerCount);
-	INC_MEMORY_STAT_BY(STAT_FSlateWindowElementListMemManager, MemManager.GetByteCount());
-
-	MemManager.Flush();
-
-	RenderTargetWindow = nullptr;
-
-	bReportReferences = true;
+	return NewElement;
 }
 
-void FSlateWindowElementList::SetShouldReportReferencesToGC(bool bInReportReferences)
+void FSlateCachedElementData::AddReferencedObjects(FReferenceCollector& Collector)
 {
-	bReportReferences = bInReportReferences;
+	for (FSlateCachedElementList& CachedElementList : CachedElementLists)
+	{
+		CachedElementList.AddReferencedObjects(Collector);
+	}
 }
 
-bool FSlateWindowElementList::ShouldReportUObjectReferences() const
+FSlateCachedClipState& FSlateCachedElementData::FindOrAddCachedClipState(const FSlateClippingState* RefClipState)
 {
-	return bReportReferences || IsCachedRenderDataInUse();
-}
+	for (auto& CachedState : CachedClipStates)
+	{
+		if (CachedState->ClippingState == *RefClipState)
+		{
+			CachedState->BeginUsingState();
+			return *CachedState;
+		}
+	}
 
-FSlateWindowElementList::FWindowElementGCObject::FWindowElementGCObject(FSlateWindowElementList* InOwner)
-	: Owner(InOwner)
-{
+	return *CachedClipStates.Emplace_GetRef(MakeUnique<FSlateCachedClipState>(*RefClipState));
 }
 
-void FSlateWindowElementList::FWindowElementGCObject::ClearOwner()
+void FSlateCachedElementData::CleanupUnusedClipStates()
 {
-	Owner = nullptr;
+	for(int32 CachedStateIdx = 0; CachedStateIdx < CachedClipStates.Num();)
+	{
+		const TUniquePtr<FSlateCachedClipState>& CachedState = CachedClipStates[CachedStateIdx];
+		if (CachedState->GetUsageCount() == 0)
+		{
+			CachedClipStates.RemoveAtSwap(CachedStateIdx);
+		}
+		else
+		{
+			++CachedStateIdx;
+		}
+	}
 }
 
-void FSlateWindowElementList::FWindowElementGCObject::AddReferencedObjects(FReferenceCollector& Collector)
+void FSlateCachedElementList::Reset()
 {
-	if (Owner && Owner->ShouldReportUObjectReferences())
+	DrawElements.Reset();
+
+	CachedBatches.Reset();
+
+
+	// Destroy vertex data in a thread safe way
+	DestroyCachedVertexData();
+
+	CachedRenderingData = new FSlateCachedFastPathRenderingData;
+
+#if WITH_SLATE_DEBUGGING
+	if (!ensure(!bNewData))
 	{
-		Owner->AddReferencedObjects(Collector);
+		UE_LOG(LogSlate, Log, TEXT("Cleared out data in cached ElementList for Widget: %s before it was batched"), *Widget->GetTag().ToString());
 	}
+#endif
+
+	bNewData = false;
 }
 
-FString FSlateWindowElementList::FWindowElementGCObject::GetReferencerName() const
+FSlateRenderBatch& FSlateCachedElementList::AddRenderBatch(int32 InLayer, const FShaderParams& InShaderParams, const FSlateShaderResource* InResource, ESlateDrawPrimitive InPrimitiveType, ESlateShader InShaderType, ESlateDrawEffect InDrawEffects, ESlateBatchDrawFlag InDrawFlags, int8 SceneIndex)
 {
-	return TEXT("FSlateWindowElementList");
+	return CachedBatches.Emplace_GetRef(InLayer, InShaderParams, InResource, InPrimitiveType, InShaderType, InDrawEffects, InDrawFlags, SceneIndex, &CachedRenderingData->Vertices, &CachedRenderingData->Indices, CachedRenderingData->Vertices.Num(), CachedRenderingData->Indices.Num());
 }
 
-void FSlateWindowElementList::AddReferencedObjects(FReferenceCollector& Collector)
+void FSlateCachedElementList::AddCachedClipState(FSlateCachedClipState& ClipStateToCache)
 {
-	Collector.AddReferencedObjects(ResourcesToReport);
+	CachedRenderingData->CachedClipStates.Add(&ClipStateToCache);
 }
 
-int32 FSlateWindowElementList::GetElementCount() const
+void FSlateCachedElementList::AddReferencedObjects(FReferenceCollector& Collector)
 {
-	int32 ElementTotal = RootDrawLayer.GetElementCount();
+	for (FSlateDrawElement& Element : DrawElements)
+	{
+		Element.AddReferencedObjects(Collector);
+	}
+}
 
-	for (auto& DrawLayerEntry : DrawLayers)
+void FSlateCachedElementList::DestroyCachedVertexData()
+{
+	if (CachedRenderingData)
 	{
-		if (FSlateDrawLayer* DrawLayer = DrawLayerEntry.Value.Get())
+		if (FSlateApplicationBase::IsInitialized())
+		{
+			FSlateApplicationBase::Get().GetRenderer()->DestroyCachedFastPathRenderingData(CachedRenderingData);
+		}
+		else
 		{
-			ElementTotal += DrawLayer->GetElementCount();
+			delete CachedRenderingData;
 		}
 	}
 
-	return ElementTotal;
+	CachedRenderingData = nullptr;
 }
