Slate Global Invalidation System

diff --git a/Engine/Source/Runtime/SlateCore/Private/Rendering/ElementBatcher.cpp b/Engine/Source/Runtime/SlateCore/Private/Rendering/ElementBatcher.cpp
--- a/Engine/Source/Runtime/SlateCore/Private/Rendering/ElementBatcher.cpp
+++ b/Engine/Source/Runtime/SlateCore/Private/Rendering/ElementBatcher.cpp
@@ -7,28 +7,371 @@
 #include "Rendering/RenderingPolicy.h"
 #include "Widgets/SWindow.h"
 #include "HAL/IConsoleManager.h"
+#include "ProfilingDebugging/ScopedTimers.h"
 
-DECLARE_CYCLE_STAT(TEXT("Find Batch For Element Time"), STAT_SlateFindBatchForElement, STATGROUP_SlateVerbose);
+DECLARE_CYCLE_STAT(TEXT("Slate RT: Create Batches"), STAT_SlateRTCreateBatches, STATGROUP_Slate);
 
-DECLARE_CYCLE_STAT(TEXT("Add Elements Time"), STAT_SlateAddElements, STATGROUP_Slate);
+DEFINE_STAT(STAT_SlateAddElements);
 
-DECLARE_DWORD_COUNTER_STAT(TEXT("Elements"), STAT_SlateElements, STATGROUP_Slate);
-DECLARE_DWORD_COUNTER_STAT(TEXT("Elements (Box)"), STAT_SlateElements_Box, STATGROUP_Slate);
-DECLARE_DWORD_COUNTER_STAT(TEXT("Elements (Border)"), STAT_SlateElements_Border, STATGROUP_Slate);
-DECLARE_DWORD_COUNTER_STAT(TEXT("Elements (Text)"), STAT_SlateElements_Text, STATGROUP_Slate);
-DECLARE_DWORD_COUNTER_STAT(TEXT("Elements (ShapedText)"), STAT_SlateElements_ShapedText, STATGROUP_Slate);
-DECLARE_DWORD_COUNTER_STAT(TEXT("Elements (Line)"), STAT_SlateElements_Line, STATGROUP_Slate);
-DECLARE_DWORD_COUNTER_STAT(TEXT("Elements (Cached Buffer)"), STAT_SlateElements_CachedBuffer, STATGROUP_Slate);
-DECLARE_DWORD_COUNTER_STAT(TEXT("Elements (Other)"), STAT_SlateElements_Other, STATGROUP_Slate);
+DEFINE_STAT(STAT_SlateElements);
+DEFINE_STAT(STAT_SlateElements_Box);
+DEFINE_STAT(STAT_SlateElements_Border);
+DEFINE_STAT(STAT_SlateElements_Text);
+DEFINE_STAT(STAT_SlateElements_ShapedText);
+DEFINE_STAT(STAT_SlateElements_Line);
+DEFINE_STAT(STAT_SlateElements_Other);
+DEFINE_STAT(STAT_SlateCachedElements);
 
 int32 GSlateFeathering = 0;
-static FAutoConsoleVariableRef CVarSlateFeathering(TEXT("Slate.Feathering"), GSlateFeathering, TEXT("Provides a means of doing Anti-Aliasing around the edges by feathering boxes."), ECVF_Default);
+
+
+FSlateElementBatch::FSlateElementBatch(const FSlateShaderResource* InShaderResource, const FShaderParams& InShaderParams, ESlateShader ShaderType, ESlateDrawPrimitive PrimitiveType, ESlateDrawEffect InDrawEffects, ESlateBatchDrawFlag InBatchFlags, const FSlateDrawElement& InDrawElement, int32 InstanceCount, uint32 InstanceOffset, ISlateUpdatableInstanceBuffer* InstanceData)
+	: BatchKey(InShaderParams, ShaderType, PrimitiveType, InDrawEffects, InBatchFlags, InDrawElement.GetClippingHandle(), InstanceCount, InstanceOffset, InstanceData, InDrawElement.GetSceneIndex())
+	, ShaderResource(InShaderResource)
+	, NumElementsInBatch(0)
+	, VertexArrayIndex(INDEX_NONE)
+	, IndexArrayIndex(INDEX_NONE)
+{
+}
+
+FSlateElementBatch::FSlateElementBatch(TWeakPtr<ICustomSlateElement, ESPMode::ThreadSafe> InCustomDrawer, const FSlateDrawElement& InDrawElement)
+	: BatchKey(InCustomDrawer, InDrawElement.GetClippingHandle())
+	, ShaderResource(nullptr)
+	, NumElementsInBatch(0)
+	, VertexArrayIndex(INDEX_NONE)
+	, IndexArrayIndex(INDEX_NONE)
+{
+}
+
+void FSlateElementBatch::SaveClippingState(const TArray<FSlateClippingState>& PrecachedClipStates)
+{
+	/*// Do cached first
+	if (BatchKey.ClipStateHandle.GetCachedClipState().IsSet())
+	{
+		const TSharedPtr<FSlateClippingState>& CachedState = BatchKey.ClipStateHandle.GetCachedClipState().GetValue();
+		if (CachedState.IsValid())
+		{
+			ClippingState = *CachedState;
+		}
+	}
+	else if (PrecachedClipStates.IsValidIndex(BatchKey.ClipStateHandle.GetPrecachedClipIndex()))
+	{
+		// Store the clipping state so we can use it later for rendering.
+		ClippingState = PrecachedClipStates[BatchKey.ClipStateHandle.GetPrecachedClipIndex()];
+	}*/
+}
+
+void FSlateBatchData::ResetData()
+{
+	RenderBatches.Reset();
+	UncachedSourceBatchIndices.Reset();
+	UncachedSourceBatchVertices.Reset();
+	FinalIndexData.Reset();
+	FinalVertexData.Reset();
+
+	// note: LayerToElementBatches is not reset here as the same layers are 
+	// more than likely reused and we can save memory allocations by not resetting the map every frame
+
+	NumBatches = 0;
+	NumLayers = 0;
+
+	bIsStencilBufferRequired = false;
+}
+
+#define MAX_VERT_ARRAY_RECYCLE (200)
+#define MAX_INDEX_ARRAY_RECYCLE (500)
+
+bool FSlateBatchData::IsStencilClippingRequired() const
+{
+	return bIsStencilBufferRequired;
+}
+
+void FSlateBatchData::FillVertexAndIndexBuffer(uint8* VertexBuffer, uint8* IndexBuffer, bool bAbsoluteIndices)
+{
+	/*int32 IndexOffset = 0;
+	int32 VertexOffset = 0;
+	int32 BaseVertexIndex = 0;
+
+	const bool bValidBuffers = (nullptr != VertexBuffer) && (nullptr != IndexBuffer);
+
+	for (const FSlateRenderBatch& Batch : RenderBatches)
+	{
+		// Ignore foreign batches that are inserted into our render set.
+		if (RenderDataHandle != Batch.CachedRenderHandle)
+		{
+			continue;
+		}
+
+		if (Batch.VertexArrayIndex != INDEX_NONE && Batch.IndexArrayIndex != INDEX_NONE)
+		{
+			FSlateVertexArray& Vertices = BatchVertexArrays[Batch.VertexArrayIndex];
+			FSlateIndexArray& Indices = BatchIndexArrays[Batch.IndexArrayIndex];
+
+			if (Vertices.Num() && Indices.Num())
+			{
+				if (bValidBuffers)
+				{
+					uint32 RequiredVertexSize = Vertices.Num() * Vertices.GetTypeSize();
+					uint32 RequiredIndexSize = Indices.Num() * Indices.GetTypeSize();
+
+					FMemory::Memcpy(VertexBuffer + VertexOffset, Vertices.GetData(), RequiredVertexSize);
+					if (BaseVertexIndex == 0 || !bAbsoluteIndices)
+					{
+						FMemory::Memcpy(IndexBuffer + IndexOffset, Indices.GetData(), RequiredIndexSize);
+					}
+					else
+					{
+						SlateIndex* TargetIndexBuffer = (SlateIndex*)(IndexBuffer + IndexOffset);
+						for (int32 i = 0; i < Indices.Num(); ++i)
+						{
+							TargetIndexBuffer[i] = Indices[i] + BaseVertexIndex;
+						}
+					}
+
+					BaseVertexIndex += Vertices.Num();
+					IndexOffset += (Indices.Num() * sizeof(SlateIndex));
+					VertexOffset += (Vertices.Num() * sizeof(FSlateVertex));
+				}
+
+				Vertices.Reset();
+				Indices.Reset();
+
+				if (Vertices.GetSlack() > MAX_VERT_ARRAY_RECYCLE)
+				{
+					ResetVertexArray(Vertices);
+				}
+
+				if (Indices.GetSlack() > MAX_INDEX_ARRAY_RECYCLE)
+				{
+					ResetIndexArray(Indices);
+				}
+			}
+		}
+	}*/
+}
+
+/*
+void FSlateBatchData::CreateRenderBatches(FElementBatchMap& LayerToElementBatches)
+{
+	checkSlow(IsInRenderingThread());
+
+	uint32 VertexOffset = 0;
+	uint32 IndexOffset = 0;
+
+	{
+		SCOPED_NAMED_EVENT_TEXT("SlateRT::CreateRenderBatches", FColor::Magenta);
+		Merge(LayerToElementBatches, VertexOffset, IndexOffset);
+	}
+
+	// 
+	if (RenderDataHandle.IsValid())
+	{
+		RenderDataHandle->SetRenderBatches(&RenderBatches);
+	}
+}
+*/
+
+FSlateRenderBatch& FSlateBatchData::AddRenderBatch(int32 InLayer, const FShaderParams& InShaderParams, const FSlateShaderResource* InResource, ESlateDrawPrimitive InPrimitiveType, ESlateShader InShaderType, ESlateDrawEffect InDrawEffects, ESlateBatchDrawFlag InDrawFlags, int8 SceneIndex)
+{
+
+/*
+	NumBatchedVertices += InNumVertices;
+	NumBatchedIndices += InNumIndices;
+
+	if (InElementBatch.ClippingState.IsSet())
+	{
+		bIsStencilBufferRequired |= InElementBatch.ClippingState->GetClippingMethod() == EClippingMethod::Stencil;
+	}
+
+	const int32 Index = RenderBatches.Add(FSlateRenderBatch(InLayer, InElementBatch, RenderDataHandle, InNumVertices, InNumIndices, InVertexOffset, InIndexOffset));
+	RenderBatches[Index].DynamicOffset = FVector2D::ZeroVector;*/
+
+	return RenderBatches.Emplace_GetRef(InLayer, InShaderParams, InResource, InPrimitiveType, InShaderType, InDrawEffects, InDrawFlags, SceneIndex, &UncachedSourceBatchVertices, &UncachedSourceBatchIndices, UncachedSourceBatchVertices.Num(), UncachedSourceBatchIndices.Num());
+}
+
+
+
+void FSlateBatchData::AddCachedBatches(const TArray<FSlateRenderBatch>& InCachedBatches)
+{
+	RenderBatches.Reserve(RenderBatches.Num() + InCachedBatches.Num());
+
+	RenderBatches.Append(InCachedBatches);
+}
+
+void FSlateBatchData::FillBuffersFromNewBatch(FSlateRenderBatch& Batch, FSlateVertexArray& FinalVertices, FSlateIndexArray& FinalIndices)
+{
+	if(Batch.HasVertexData())
+	{
+		const int32 SourceVertexOffset = Batch.VertexOffset;
+		const int32 SourceIndexOffset = Batch.IndexOffset;
+
+		// At the start of a new batch, just direct copy the verts
+		// todo: May need to change this to use absolute indices
+		Batch.VertexOffset = FinalVertices.Num();
+		Batch.IndexOffset = FinalIndices.Num();
+		
+		FinalVertices.Append(&(*Batch.SourceVertices)[SourceVertexOffset], Batch.NumVertices);
+		FinalIndices.Append(&(*Batch.SourceIndices)[SourceIndexOffset], Batch.NumIndices);
+	}
+}
+
+void FSlateBatchData::CombineBatches(FSlateRenderBatch& FirstBatch, FSlateRenderBatch& SecondBatch, FSlateVertexArray& FinalVertices, FSlateIndexArray& FinalIndices)
+{
+	check(!SecondBatch.bIsMerged);
+	if (FirstBatch.HasVertexData() || SecondBatch.HasVertexData())
+	{
+		// when merging verts we have to offset the indices in the second batch based on the first batches existing number of verts
+		const int32 BatchOffset = FirstBatch.NumVertices;
+
+		// Final vertices is assumed to have the first batch already in it
+		FirstBatch.NumVertices += SecondBatch.NumVertices;
+		FirstBatch.NumIndices += SecondBatch.NumIndices;
+
+		FinalVertices.Append(&(*SecondBatch.SourceVertices)[SecondBatch.VertexOffset], SecondBatch.NumVertices);
+
+		FinalIndices.Reserve(FinalIndices.Num()+SecondBatch.NumIndices);
+		
+		// Get source indices at the source index offset and shift each index by the batches current offset
+		for (int32 i = 0; i < SecondBatch.NumIndices; ++i)
+		{
+			const int32 FinalIndex = (*SecondBatch.SourceIndices)[i + SecondBatch.IndexOffset] + BatchOffset;
+			FinalIndices.Add(FinalIndex);
+		}
+	}
+
+	SecondBatch.bIsMerged = true;
+}
+
+
+/*
+static void AddNewBatch(FSlateRenderBatch& NewBatch, TArray<FSlateRenderBatch>& FinalRenderBatches, TArray<FSlateVertex>& Vertices, TArray<SlateIndex>& Indices)
+{
+	if (NewBatch.GetNumVertices() > 0 && NewBatch.GetNumIndices())
+	{
+		const int32 NewVertOffset = Vertices.Num();
+	
+		NewBatch.CopySourceVertices(Vertices);
+		NewBatch.VertexOffset = NewVertOffset;
+
+		// Starting a new batch does not require updating absolute index position
+		const int32 NewIndexOffset = Indices.Num();
+		NewBatch.CopySourceIndices(Indices);
+		NewBatch.IndexOffset = NewIndexOffset;
+
+		FinalRenderBatches.Add(NewBatch);
+	}
+}
+*/
+
+void FSlateBatchData::MergeRenderBatches()
+{
+	SCOPE_CYCLE_COUNTER(STAT_SlateRTCreateBatches);
+
+	if(RenderBatches.Num())
+	{
+		TArray<TPair<int32, int32>, TInlineAllocator<100, TMemStackAllocator<>>> BatchIndices;
+
+		{
+			SCOPED_NAMED_EVENT_TEXT("Slate::SortRenderBatches", FColor::Magenta);
+
+			// Sort an index array instead of the render batches since they are large and not trivially relocatable 
+			BatchIndices.AddUninitialized(RenderBatches.Num());
+			for (int32 Index = 0; Index < RenderBatches.Num(); ++Index)
+			{
+				BatchIndices[Index].Key = Index;
+				BatchIndices[Index].Value = RenderBatches[Index].GetLayer();
+			}
+
+			// Stable sort because order in the same layer should be preserved
+			BatchIndices.StableSort
+			(
+				[](const TPair<int32, int32>& A, const TPair<int32, int32>& B)
+				{
+					return A.Value < B.Value;
+				}
+			);
+		}
+
+
+		NumBatches = 0;
+		NumLayers = 0;
+
+#if STATS
+		int32 CurLayerId = INDEX_NONE;
+		int32 PrevLayerId = INDEX_NONE;
+#endif
+
+		FirstRenderBatchIndex = BatchIndices[0].Key;
+
+		FSlateRenderBatch* PrevBatch = nullptr;
+		for (int32 BatchIndex = 0; BatchIndex < BatchIndices.Num(); ++BatchIndex)
+		{
+			const TPair<int32, int32>& BatchIndexPair = BatchIndices[BatchIndex];
+
+			FSlateRenderBatch& CurBatch = RenderBatches[BatchIndexPair.Key];
+
+#if STATS
+			CurLayerId = CurBatch.GetLayer();
+			if (PrevLayerId != CurLayerId)
+			{
+				++NumLayers;
+			}
+			CurLayerId = PrevLayerId;
+#endif
+
+			if (CurBatch.bIsMerged)
+			{
+				// skip already merged batches
+				continue;
+			}
+
+			if (PrevBatch != nullptr)
+			{
+				PrevBatch->NextBatchIndex = BatchIndexPair.Key;
+			}
+
+			++NumBatches;
+
+			FillBuffersFromNewBatch(CurBatch, FinalVertexData, FinalIndexData);
+
+			if (CurBatch.ClippingState)
+			{
+				bIsStencilBufferRequired |= CurBatch.ClippingState->GetClippingMethod() == EClippingMethod::Stencil;
+			}
+
+#if 1  // Do batching at all?
+
+			if (CurBatch.bIsMergable)
+			{
+				for (int32 TestIndex = BatchIndex + 1; TestIndex < BatchIndices.Num(); ++TestIndex)
+				{
+					const TPair<int32, int32>& NextBatchIndexPair = BatchIndices[TestIndex];
+					FSlateRenderBatch& TestBatch = RenderBatches[NextBatchIndexPair.Key];
+					if (TestBatch.GetLayer() != CurBatch.GetLayer())
+					{
+						// none of the batches will be compatible since we encountered an incompatible layer
+						break;
+					}
+					else if (!TestBatch.bIsMerged && CurBatch.IsBatchableWith(TestBatch))
+					{
+						CombineBatches(CurBatch, TestBatch, FinalVertexData, FinalIndexData);
+
+						check(TestBatch.NextBatchIndex == INDEX_NONE);
+
+					}
+				}
+			}
+#endif
+			PrevBatch = &CurBatch;
+		}
+	}
+}
 
 FSlateElementBatcher::FSlateElementBatcher( TSharedRef<FSlateRenderingPolicy> InRenderingPolicy )
 	: BatchData( nullptr )
-	, DrawLayer( nullptr )
+	, CachedElementList( nullptr )
+	, PrecachedClippingStates(nullptr)
 	, RenderingPolicy( &InRenderingPolicy.Get() )
-	, ElmementStat_Other(0)
 	, NumPostProcessPasses(0)
 	, PixelCenterOffset( InRenderingPolicy->GetPixelCenterOffset() )
 	, bSRGBVertexColor( !InRenderingPolicy->IsVertexColorInLinearSpace() )
@@ -46,123 +389,156 @@ void FSlateElementBatcher::AddElements(FSlateWindowElementList& WindowElementLis
 
 	SCOPE_CYCLE_COUNTER(STAT_SlateAddElements);
 
-	ElmementStat_Other = 0;
-	ElmementStat_Boxes = 0;
-	ElmementStat_Borders = 0;
-	ElmementStat_Text = 0;
-	ElmementStat_ShapedText = 0;
-	ElmementStat_Line = 0;
-	ElmementStat_CachedBuffer = 0;
+#if STATS
+	ElementStat_Other = 0;
+	ElementStat_Boxes = 0;
+	ElementStat_Borders = 0;
+	ElementStat_Text = 0;
+	ElementStat_ShapedText = 0;
+	ElementStat_Line = 0;
+	ElementStat_CachedElements = 0;
+#endif
 
 	BatchData = &WindowElementList.GetBatchData();
-	DrawLayer = &WindowElementList.GetRootDrawLayer();
-	
+	check(BatchData->GetRenderBatches().Num() == 0);
+
+
 	FVector2D ViewportSize = WindowElementList.GetPaintWindow()->GetViewportSize();
 
-	ClippingStates = &WindowElementList.ClippingManager.GetClippingStates();
+	PrecachedClippingStates = &WindowElementList.ClippingManager.GetClippingStates();
 
-	BatchData->DetermineIsStencilClippingRequired(*ClippingStates);
-	
-	AddElementsInternal(WindowElementList.GetRootDrawLayer().DrawElements, ViewportSize);
+	AddElementsInternal(WindowElementList.GetUncachedDrawElements(), ViewportSize);
 
-	TMap< TSharedPtr<FSlateDrawLayerHandle, ESPMode::ThreadSafe>, TSharedPtr<FSlateDrawLayer> >& DrawLayers = WindowElementList.GetChildDrawLayers();
-	for ( auto& Entry : DrawLayers )
+	const TArrayView<FSlateCachedElementData* const> CachedElementDataList = WindowElementList.GetCachedElementDataList();
+
+
+	if(CachedElementDataList.Num())
 	{
-		DrawLayer = Entry.Value.Get();
-		AddElementsInternal(DrawLayer->DrawElements, ViewportSize);
+		SCOPED_NAMED_EVENT_TEXT("Slate::AddCachedElements", FColor::Magenta);
+
+		for (FSlateCachedElementData* CachedElementData : CachedElementDataList)
+		{
+			AddCachedElements(*CachedElementData, ViewportSize);
+		}
 	}
 
 	// Done with the element list
 	BatchData = nullptr;
-	DrawLayer = nullptr;
-	ClippingStates = nullptr;
-
-	const int32 ElmementStat_All =
-		ElmementStat_Boxes +
-		ElmementStat_Borders +
-		ElmementStat_Text +
-		ElmementStat_ShapedText +
-		ElmementStat_Line +
-		ElmementStat_CachedBuffer +
-		ElmementStat_Other;
-
-	INC_DWORD_STAT_BY(STAT_SlateElements, ElmementStat_All);
-	INC_DWORD_STAT_BY(STAT_SlateElements_Box, ElmementStat_Boxes);
-	INC_DWORD_STAT_BY(STAT_SlateElements_Border, ElmementStat_Borders);
-	INC_DWORD_STAT_BY(STAT_SlateElements_Text, ElmementStat_Text);
-	INC_DWORD_STAT_BY(STAT_SlateElements_ShapedText, ElmementStat_ShapedText);
-	INC_DWORD_STAT_BY(STAT_SlateElements_Line, ElmementStat_Line);
-	INC_DWORD_STAT_BY(STAT_SlateElements_CachedBuffer, ElmementStat_CachedBuffer);
-	INC_DWORD_STAT_BY(STAT_SlateElements_Other, ElmementStat_Other);
+	PrecachedClippingStates = nullptr;
+
+#if STATS
+	const int32 ElementStat_All =
+		ElementStat_Boxes +
+		ElementStat_Borders +
+		ElementStat_Text +
+		ElementStat_ShapedText +
+		ElementStat_Line +
+		ElementStat_Other;
+
+	INC_DWORD_STAT_BY(STAT_SlateElements, ElementStat_All);
+	INC_DWORD_STAT_BY(STAT_SlateElements_Box, ElementStat_Boxes);
+	INC_DWORD_STAT_BY(STAT_SlateElements_Border, ElementStat_Borders);
+	INC_DWORD_STAT_BY(STAT_SlateElements_Text, ElementStat_Text);
+	INC_DWORD_STAT_BY(STAT_SlateElements_ShapedText, ElementStat_ShapedText);
+	INC_DWORD_STAT_BY(STAT_SlateElements_Line, ElementStat_Line);
+	INC_DWORD_STAT_BY(STAT_SlateElements_Other, ElementStat_Other);
+	INC_DWORD_STAT_BY(STAT_SlateCachedElements, ElementStat_CachedElements);
+#endif
 }
 
-void FSlateElementBatcher::AddElementsInternal(const TArray<FSlateDrawElement>& DrawElements, const FVector2D& ViewportSize)
-{
-	checkSlow(DrawLayer);
 
-	for ( int32 DrawElementIndex = 0; DrawElementIndex < DrawElements.Num(); ++DrawElementIndex )
+void FSlateElementBatcher::AddElementsInternal(const FSlateDrawElementArray& DrawElements, const FVector2D& ViewportSize)
+{
+	for (const FSlateDrawElement& DrawElement : DrawElements)
 	{
-		const FSlateDrawElement& DrawElement = DrawElements[DrawElementIndex];
-
 		// Determine what type of element to add
 		switch ( DrawElement.GetElementType() )
 		{
-		case FSlateDrawElement::ET_Box:
-			ElmementStat_Boxes++;
+		case EElementType::ET_Box:
+		{
+			SCOPED_NAMED_EVENT_TEXT("Slate::AddBoxElement", FColor::Magenta);
+			STAT(ElementStat_Boxes++);
 			DrawElement.IsPixelSnapped() ? AddBoxElement<ESlateVertexRounding::Enabled>(DrawElement) : AddBoxElement<ESlateVertexRounding::Disabled>(DrawElement);
+		}
 			break;
-		case FSlateDrawElement::ET_Border:
-			ElmementStat_Borders++;
+		case EElementType::ET_Border:
+		{
+			SCOPED_NAMED_EVENT_TEXT("Slate::AddBorderElement", FColor::Magenta);
+			STAT(ElementStat_Borders++);
 			DrawElement.IsPixelSnapped() ? AddBorderElement<ESlateVertexRounding::Enabled>(DrawElement) : AddBorderElement<ESlateVertexRounding::Disabled>(DrawElement);
+		}
 			break;
-		case FSlateDrawElement::ET_Text:
-			ElmementStat_Text++;
+		case EElementType::ET_Text:
+		{
+			SCOPED_NAMED_EVENT_TEXT("Slate::AddTextElement", FColor::Magenta);
+			STAT(ElementStat_Text++);
 			DrawElement.IsPixelSnapped() ? AddTextElement<ESlateVertexRounding::Enabled>(DrawElement) : AddTextElement<ESlateVertexRounding::Disabled>(DrawElement);
+		}
 			break;
-		case FSlateDrawElement::ET_ShapedText:
-			ElmementStat_ShapedText++;
+		case EElementType::ET_ShapedText:
+		{
+			SCOPED_NAMED_EVENT_TEXT("Slate::AddShapedTextElement", FColor::Magenta);
+			STAT(ElementStat_ShapedText++);
 			DrawElement.IsPixelSnapped() ? AddShapedTextElement<ESlateVertexRounding::Enabled>(DrawElement) : AddShapedTextElement<ESlateVertexRounding::Disabled>(DrawElement);
+		}
 			break;
-		case FSlateDrawElement::ET_Line:
-			ElmementStat_Line++;
+		case EElementType::ET_Line:
+		{
+			SCOPED_NAMED_EVENT_TEXT("Slate::AddLineElement", FColor::Magenta);
+			STAT(ElementStat_Line++);
 			DrawElement.IsPixelSnapped() ? AddLineElement<ESlateVertexRounding::Enabled>(DrawElement) : AddLineElement<ESlateVertexRounding::Disabled>(DrawElement);
+		}
 			break;
-		case FSlateDrawElement::ET_DebugQuad:
-			ElmementStat_Other++;
+		case EElementType::ET_DebugQuad:
+		{
+			SCOPED_NAMED_EVENT_TEXT("Slate::AddDebugQuadElement", FColor::Magenta);
+			STAT(ElementStat_Other++);
 			DrawElement.IsPixelSnapped() ? AddQuadElement<ESlateVertexRounding::Enabled>(DrawElement) : AddQuadElement<ESlateVertexRounding::Disabled>(DrawElement);
+		}
 			break;
-		case FSlateDrawElement::ET_Spline:
+		case EElementType::ET_Spline:
+		{
+			SCOPED_NAMED_EVENT_TEXT("Slate::AddSplineElement", FColor::Magenta);
 			// Note that we ignore pixel snapping here; see implementation for more info.
-			ElmementStat_Other++;
+			STAT(ElementStat_Other++);
 			AddSplineElement(DrawElement);
+		}
 			break;
-		case FSlateDrawElement::ET_Gradient:
-			ElmementStat_Other++;
+		case EElementType::ET_Gradient:
+		{
+			SCOPED_NAMED_EVENT_TEXT("Slate::AddGradientElement", FColor::Magenta);
+			STAT(ElementStat_Other++);
 			DrawElement.IsPixelSnapped() ? AddGradientElement<ESlateVertexRounding::Enabled>(DrawElement) : AddGradientElement<ESlateVertexRounding::Disabled>(DrawElement);
+		}
 			break;
-		case FSlateDrawElement::ET_Viewport:
-			ElmementStat_Other++;
+		case EElementType::ET_Viewport:
+		{
+			SCOPED_NAMED_EVENT_TEXT("Slate::AddViewportElement", FColor::Magenta);
+			STAT(ElementStat_Other++);
 			DrawElement.IsPixelSnapped() ? AddViewportElement<ESlateVertexRounding::Enabled>(DrawElement) : AddViewportElement<ESlateVertexRounding::Disabled>(DrawElement);
+		}
 			break;
-		case FSlateDrawElement::ET_Custom:
-			ElmementStat_Other++;
+		case EElementType::ET_Custom:
+		{
+			SCOPED_NAMED_EVENT_TEXT("Slate::AddCustomElement", FColor::Magenta);
+			STAT(ElementStat_Other++);
 			AddCustomElement(DrawElement);
+		}
 			break;
-		case FSlateDrawElement::ET_CustomVerts:
-			ElmementStat_Other++;
+		case EElementType::ET_CustomVerts:
+		{
+			SCOPED_NAMED_EVENT_TEXT("Slate::AddCustomVertsElement", FColor::Magenta);
+			STAT(ElementStat_Other++);
 			AddCustomVerts(DrawElement);
+		}
 			break;
-		case FSlateDrawElement::ET_Layer:
-			ElmementStat_Other++;
-			AddLayer(DrawElement);
-			break;
-		case FSlateDrawElement::ET_CachedBuffer:
-			ElmementStat_CachedBuffer++;
-			AddCachedBuffer(DrawElement);
-			break;
-		case FSlateDrawElement::ET_PostProcessPass:
-			ElmementStat_Other++;
+	
+		case EElementType::ET_PostProcessPass:
+		{
+			SCOPED_NAMED_EVENT_TEXT("Slate::AddPostProcessElement", FColor::Magenta);
+			STAT(ElementStat_Other++);
 			AddPostProcessPass(DrawElement, ViewportSize);
+		}
 			break;
 		default:
 			checkf(0, TEXT("Invalid element type"));
@@ -171,18 +547,49 @@ void FSlateElementBatcher::AddElementsInternal(const TArray<FSlateDrawElement>&
 	}
 }
 
+void FSlateElementBatcher::AddCachedElements(FSlateCachedElementData& CachedElementData, const FVector2D& ViewportSize)
+{
+#define ALLOW_CACHED_RENDER_BATCHES 1
+	SCOPED_NAMED_EVENT_TEXT("Slate::AddCachedBatches", FColor::Magenta);
+	for ( FSlateCachedElementList& LocalCachedElementList : CachedElementData.CachedElementLists)
+	{
+		STAT(ElementStat_CachedElements += LocalCachedElementList.DrawElements.Num());
+		LocalCachedElementList.bNewData = false;
+#if ALLOW_CACHED_RENDER_BATCHES
+		if (LocalCachedElementList.CachedBatches.Num())
+		{
+			BatchData->AddCachedBatches(LocalCachedElementList.CachedBatches);
+		}
+		else
+		{
+			CachedElementList = &LocalCachedElementList;
+			{
+				SCOPED_NAMED_EVENT_TEXT("Slate::RecacheElements", FColor::Magenta);
+				AddElementsInternal(LocalCachedElementList.DrawElements, ViewportSize);
+			}
+
+			// Now apply the cached batches to the batch data
+			BatchData->AddCachedBatches(LocalCachedElementList.CachedBatches);
+		}
+
+		CachedElementList = nullptr;
+#else
+		AddElementsInternal(LocalCachedElementList.DrawElements, ViewportSize);
+#endif
+	}
+	CachedElementData.CleanupUnusedClipStates();
+}
+
 template<ESlateVertexRounding Rounding>
 void FSlateElementBatcher::AddQuadElement( const FSlateDrawElement& DrawElement, FColor Color )
 {
 	const FSlateRenderTransform& RenderTransform = DrawElement.GetRenderTransform();
 	const FVector2D& LocalSize = DrawElement.GetLocalSize();
 	ESlateDrawEffect InDrawEffects = DrawElement.GetDrawEffects();
-	uint32 Layer = DrawElement.GetLayer();
-
-	FSlateElementBatch& ElementBatch = FindBatchForElement( Layer, FShaderParams(), nullptr, ESlateDrawPrimitive::TriangleList, ESlateShader::Default, ESlateDrawEffect::None, ESlateBatchDrawFlag::Wireframe|ESlateBatchDrawFlag::NoBlending, DrawElement.GetClippingIndex(), DrawElement.GetSceneIndex());
-	FSlateVertexArray& BatchVertices = BatchData->GetBatchVertexList(ElementBatch);
-	FSlateIndexArray& BatchIndices = BatchData->GetBatchIndexList(ElementBatch);
+	const int32 Layer = DrawElement.GetLayer();
 
+	FSlateRenderBatch& RenderBatch = CreateRenderBatch(Layer, FShaderParams(), nullptr, ESlateDrawPrimitive::TriangleList, ESlateShader::Default, ESlateDrawEffect::None, ESlateBatchDrawFlag::Wireframe | ESlateBatchDrawFlag::NoBlending, DrawElement);
+	
 	// Determine the four corners of the quad
 	FVector2D TopLeft = FVector2D::ZeroVector;
 	FVector2D TopRight = FVector2D(LocalSize.X, 0);
@@ -190,41 +597,40 @@ void FSlateElementBatcher::AddQuadElement( const FSlateDrawElement& DrawElement,
 	FVector2D BotRight = FVector2D(LocalSize.X, LocalSize.Y);
 
 	// The start index of these vertices in the index buffer
-	uint32 IndexStart = BatchVertices.Num();
+	//const uint32 IndexStart = BatchVertices.Num();
+	const uint32 IndexStart = 0;
 
 	// Add four vertices to the list of verts to be added to the vertex buffer
-	BatchVertices.Add( FSlateVertex::Make<Rounding>( RenderTransform, TopLeft, FVector2D(0.0f,0.0f), Color ) );
-	BatchVertices.Add( FSlateVertex::Make<Rounding>( RenderTransform, TopRight, FVector2D(1.0f,0.0f), Color ) );
-	BatchVertices.Add( FSlateVertex::Make<Rounding>( RenderTransform, BotLeft, FVector2D(0.0f,1.0f), Color ) );
-	BatchVertices.Add( FSlateVertex::Make<Rounding>( RenderTransform, BotRight, FVector2D(1.0f,1.0f),Color ) );
+	RenderBatch.AddVertex(FSlateVertex::Make<Rounding>(RenderTransform, TopLeft, FVector2D(0.0f,0.0f),  Color));
+	RenderBatch.AddVertex(FSlateVertex::Make<Rounding>(RenderTransform, TopRight, FVector2D(1.0f,0.0f), Color));
+	RenderBatch.AddVertex(FSlateVertex::Make<Rounding>(RenderTransform, BotLeft, FVector2D(0.0f,1.0f),  Color));
+	RenderBatch.AddVertex(FSlateVertex::Make<Rounding>(RenderTransform, BotRight, FVector2D(1.0f,1.0f), Color));
 
-	// The offset into the index buffer where this quads indices start
-	uint32 IndexOffsetStart = BatchIndices.Num();
 	// Add 6 indices to the vertex buffer.  (2 tri's per quad, 3 indices per tri)
-	BatchIndices.Add( IndexStart + 0 );
-	BatchIndices.Add( IndexStart + 1 );
-	BatchIndices.Add( IndexStart + 2 );
+	RenderBatch.AddIndex(IndexStart + 0);
+	RenderBatch.AddIndex(IndexStart + 1);
+	RenderBatch.AddIndex(IndexStart + 2);
 
-	BatchIndices.Add( IndexStart + 2 );
-	BatchIndices.Add( IndexStart + 1 );
-	BatchIndices.Add( IndexStart + 3 );
+	RenderBatch.AddIndex(IndexStart + 2);
+	RenderBatch.AddIndex(IndexStart + 1);
+	RenderBatch.AddIndex(IndexStart + 3);
 }
 
-FORCEINLINE void IndexQuad(FSlateIndexArray& BatchIndices, int32 TopLeft, int32 TopRight, int32 BottomRight, int32 BottomLeft)
+FORCEINLINE void IndexQuad(FSlateRenderBatch& RenderBatch, int32 TopLeft, int32 TopRight, int32 BottomRight, int32 BottomLeft)
 {
-	BatchIndices.Add(TopLeft);
-	BatchIndices.Add(TopRight);
-	BatchIndices.Add(BottomRight);
+	RenderBatch.AddIndex(TopLeft);
+	RenderBatch.AddIndex(TopRight);
+	RenderBatch.AddIndex(BottomRight);
 
-	BatchIndices.Add(BottomRight);
-	BatchIndices.Add(BottomLeft);
-	BatchIndices.Add(TopLeft);
+	RenderBatch.AddIndex(BottomRight);
+	RenderBatch.AddIndex(BottomLeft);
+	RenderBatch.AddIndex(TopLeft);
 }
 
 template<ESlateVertexRounding Rounding>
 void FSlateElementBatcher::AddBoxElement(const FSlateDrawElement& DrawElement)
 {
-	const FSlateDataPayload& DrawElementPayload = DrawElement.GetDataPayload();
+	const FSlateBoxPayload& DrawElementPayload = DrawElement.GetDataPayload<FSlateBoxPayload>();
 
 	const FColor Tint = PackVertexColor(DrawElementPayload.GetTint());
 	const FSlateRenderTransform& ElementRenderTransform = DrawElement.GetRenderTransform();
@@ -232,7 +638,7 @@ void FSlateElementBatcher::AddBoxElement(const FSlateDrawElement& DrawElement)
 	const FVector2D& LocalSize = DrawElement.GetLocalSize();
 
 	const ESlateDrawEffect InDrawEffects = DrawElement.GetDrawEffects();
-	const uint32 Layer = DrawElement.GetLayer();
+	const int32 Layer = DrawElement.GetLayer();
 
 	const float DrawScale = DrawElement.GetScale();
 
@@ -295,12 +701,12 @@ void FSlateElementBatcher::AddBoxElement(const FSlateDrawElement& DrawElement)
 	const bool bMirrorVertical = (MirroringRule == ESlateBrushMirrorType::Both || MirroringRule == ESlateBrushMirrorType::Vertical);
 
 	// Pass the tiling information as a flag so we can pick the correct texture addressing mode
-	ESlateBatchDrawFlag DrawFlags = DrawElementPayload.GetBatchFlags();
+	ESlateBatchDrawFlag DrawFlags = DrawElement.GetBatchFlags();
 	DrawFlags |= ( ( bTileHorizontal ? ESlateBatchDrawFlag::TileU : ESlateBatchDrawFlag::None ) | ( bTileVertical ? ESlateBatchDrawFlag::TileV : ESlateBatchDrawFlag::None ) );
 
-	FSlateElementBatch& ElementBatch = FindBatchForElement( Layer, FShaderParams(), Resource, ESlateDrawPrimitive::TriangleList, ESlateShader::Default, InDrawEffects, DrawFlags, DrawElement.GetClippingIndex(), DrawElement.GetSceneIndex());
-	FSlateVertexArray& BatchVertices = BatchData->GetBatchVertexList(ElementBatch);
-	FSlateIndexArray& BatchIndices = BatchData->GetBatchIndexList(ElementBatch);
+
+	FSlateRenderBatch& RenderBatch = CreateRenderBatch( Layer, FShaderParams(), Resource, ESlateDrawPrimitive::TriangleList, ESlateShader::Default, InDrawEffects, DrawFlags, DrawElement);
+
 
 	float HorizontalTiling = bTileHorizontal ? LocalSize.X / TextureWidth : 1.0f;
 	float VerticalTiling = bTileVertical ? LocalSize.Y / TextureHeight : 1.0f;
@@ -308,9 +714,8 @@ void FSlateElementBatcher::AddBoxElement(const FSlateDrawElement& DrawElement)
 	const FVector2D Tiling( HorizontalTiling, VerticalTiling );
 
 	// The start index of these vertices in the index buffer
-	uint32 IndexStart = BatchVertices.Num();
-	// The offset into the index buffer where this elements indices start
-	uint32 IndexOffsetStart = BatchIndices.Num();
+	const uint32 IndexStart = 0;// BatchVertices.Num();
+
 
 	const FMargin& Margin = DrawElementPayload.GetBrushMargin();
 
@@ -395,147 +800,147 @@ void FSlateElementBatcher::AddBoxElement(const FSlateDrawElement& DrawElement)
 		FVector2D Position = TopLeft;
 		FVector2D EndPos = BotRight;
 
-		BatchVertices.Add( FSlateVertex::Make<Rounding>( RenderTransform, FVector2D( Position.X, Position.Y ),		LocalSize, DrawScale, FVector4(StartUV,										Tiling),	Tint ) ); //0
-		BatchVertices.Add( FSlateVertex::Make<Rounding>( RenderTransform, FVector2D( Position.X, TopMarginY ),		LocalSize, DrawScale, FVector4(FVector2D( StartUV.X, TopMarginV ),			Tiling),	Tint ) ); //1
-		BatchVertices.Add( FSlateVertex::Make<Rounding>( RenderTransform, FVector2D( LeftMarginX, Position.Y ),		LocalSize, DrawScale, FVector4(FVector2D( LeftMarginU, StartUV.Y ),			Tiling),	Tint ) ); //2
-		BatchVertices.Add( FSlateVertex::Make<Rounding>( RenderTransform, FVector2D( LeftMarginX, TopMarginY ),		LocalSize, DrawScale, FVector4(FVector2D( LeftMarginU, TopMarginV ),		Tiling),	Tint ) ); //3
-		BatchVertices.Add( FSlateVertex::Make<Rounding>( RenderTransform, FVector2D( RightMarginX, Position.Y ),	LocalSize, DrawScale, FVector4(FVector2D( RightMarginU, StartUV.Y ),		Tiling),	Tint ) ); //4
-		BatchVertices.Add( FSlateVertex::Make<Rounding>( RenderTransform, FVector2D( RightMarginX, TopMarginY ),	LocalSize, DrawScale, FVector4(FVector2D( RightMarginU,TopMarginV),			Tiling),	Tint ) ); //5
-		BatchVertices.Add( FSlateVertex::Make<Rounding>( RenderTransform, FVector2D( EndPos.X, Position.Y ),		LocalSize, DrawScale, FVector4(FVector2D( EndUV.X, StartUV.Y ),				Tiling),	Tint ) ); //6
-		BatchVertices.Add( FSlateVertex::Make<Rounding>( RenderTransform, FVector2D( EndPos.X, TopMarginY ),		LocalSize, DrawScale, FVector4(FVector2D( EndUV.X, TopMarginV),				Tiling),	Tint ) ); //7
-
-		BatchVertices.Add( FSlateVertex::Make<Rounding>( RenderTransform, FVector2D( Position.X, BottomMarginY ),	LocalSize, DrawScale, FVector4(FVector2D( StartUV.X, BottomMarginV ),		Tiling),	Tint ) ); //8
-		BatchVertices.Add( FSlateVertex::Make<Rounding>( RenderTransform, FVector2D( LeftMarginX, BottomMarginY ),	LocalSize, DrawScale, FVector4(FVector2D( LeftMarginU, BottomMarginV ),		Tiling),	Tint ) ); //9
-		BatchVertices.Add( FSlateVertex::Make<Rounding>( RenderTransform, FVector2D( RightMarginX, BottomMarginY ),	LocalSize, DrawScale, FVector4(FVector2D( RightMarginU, BottomMarginV ),	Tiling),	Tint ) ); //10
-		BatchVertices.Add( FSlateVertex::Make<Rounding>( RenderTransform, FVector2D( EndPos.X, BottomMarginY ),		LocalSize, DrawScale, FVector4(FVector2D( EndUV.X, BottomMarginV ),			Tiling),	Tint ) ); //11
-		BatchVertices.Add( FSlateVertex::Make<Rounding>( RenderTransform, FVector2D( Position.X, EndPos.Y ),		LocalSize, DrawScale, FVector4(FVector2D( StartUV.X, EndUV.Y ),				Tiling),	Tint ) ); //12
-		BatchVertices.Add( FSlateVertex::Make<Rounding>( RenderTransform, FVector2D( LeftMarginX, EndPos.Y ),		LocalSize, DrawScale, FVector4(FVector2D( LeftMarginU, EndUV.Y ),			Tiling),	Tint ) ); //13
-		BatchVertices.Add( FSlateVertex::Make<Rounding>( RenderTransform, FVector2D( RightMarginX, EndPos.Y ),		LocalSize, DrawScale, FVector4(FVector2D( RightMarginU, EndUV.Y ),			Tiling),	Tint ) ); //14
-		BatchVertices.Add( FSlateVertex::Make<Rounding>( RenderTransform, FVector2D( EndPos.X, EndPos.Y ),			LocalSize, DrawScale, FVector4(EndUV,										Tiling),	Tint ) ); //15
+		RenderBatch.AddVertex( FSlateVertex::Make<Rounding>( RenderTransform, FVector2D( Position.X, Position.Y ),		LocalSize, DrawScale, FVector4(StartUV,										Tiling),	Tint ) ); //0
+		RenderBatch.AddVertex( FSlateVertex::Make<Rounding>( RenderTransform, FVector2D( Position.X, TopMarginY ),		LocalSize, DrawScale, FVector4(FVector2D( StartUV.X, TopMarginV ),			Tiling),	Tint ) ); //1
+		RenderBatch.AddVertex( FSlateVertex::Make<Rounding>( RenderTransform, FVector2D( LeftMarginX, Position.Y ),		LocalSize, DrawScale, FVector4(FVector2D( LeftMarginU, StartUV.Y ),			Tiling),	Tint ) ); //2
+		RenderBatch.AddVertex( FSlateVertex::Make<Rounding>( RenderTransform, FVector2D( LeftMarginX, TopMarginY ),		LocalSize, DrawScale, FVector4(FVector2D( LeftMarginU, TopMarginV ),		Tiling),	Tint ) ); //3
+		RenderBatch.AddVertex( FSlateVertex::Make<Rounding>( RenderTransform, FVector2D( RightMarginX, Position.Y ),	LocalSize, DrawScale, FVector4(FVector2D( RightMarginU, StartUV.Y ),		Tiling),	Tint ) ); //4
+		RenderBatch.AddVertex( FSlateVertex::Make<Rounding>( RenderTransform, FVector2D( RightMarginX, TopMarginY ),	LocalSize, DrawScale, FVector4(FVector2D( RightMarginU,TopMarginV),			Tiling),	Tint ) ); //5
+		RenderBatch.AddVertex( FSlateVertex::Make<Rounding>( RenderTransform, FVector2D( EndPos.X, Position.Y ),		LocalSize, DrawScale, FVector4(FVector2D( EndUV.X, StartUV.Y ),				Tiling),	Tint ) ); //6
+		RenderBatch.AddVertex( FSlateVertex::Make<Rounding>( RenderTransform, FVector2D( EndPos.X, TopMarginY ),		LocalSize, DrawScale, FVector4(FVector2D( EndUV.X, TopMarginV),				Tiling),	Tint ) ); //7
+
+		RenderBatch.AddVertex( FSlateVertex::Make<Rounding>( RenderTransform, FVector2D( Position.X, BottomMarginY ),	LocalSize, DrawScale, FVector4(FVector2D( StartUV.X, BottomMarginV ),		Tiling),	Tint ) ); //8
+		RenderBatch.AddVertex( FSlateVertex::Make<Rounding>( RenderTransform, FVector2D( LeftMarginX, BottomMarginY ),	LocalSize, DrawScale, FVector4(FVector2D( LeftMarginU, BottomMarginV ),		Tiling),	Tint ) ); //9
+		RenderBatch.AddVertex( FSlateVertex::Make<Rounding>( RenderTransform, FVector2D( RightMarginX, BottomMarginY ),	LocalSize, DrawScale, FVector4(FVector2D( RightMarginU, BottomMarginV ),	Tiling),	Tint ) ); //10
+		RenderBatch.AddVertex( FSlateVertex::Make<Rounding>( RenderTransform, FVector2D( EndPos.X, BottomMarginY ),		LocalSize, DrawScale, FVector4(FVector2D( EndUV.X, BottomMarginV ),			Tiling),	Tint ) ); //11
+		RenderBatch.AddVertex( FSlateVertex::Make<Rounding>( RenderTransform, FVector2D( Position.X, EndPos.Y ),		LocalSize, DrawScale, FVector4(FVector2D( StartUV.X, EndUV.Y ),				Tiling),	Tint ) ); //12
+		RenderBatch.AddVertex( FSlateVertex::Make<Rounding>( RenderTransform, FVector2D( LeftMarginX, EndPos.Y ),		LocalSize, DrawScale, FVector4(FVector2D( LeftMarginU, EndUV.Y ),			Tiling),	Tint ) ); //13
+		RenderBatch.AddVertex( FSlateVertex::Make<Rounding>( RenderTransform, FVector2D( RightMarginX, EndPos.Y ),		LocalSize, DrawScale, FVector4(FVector2D( RightMarginU, EndUV.Y ),			Tiling),	Tint ) ); //14
+		RenderBatch.AddVertex( FSlateVertex::Make<Rounding>( RenderTransform, FVector2D( EndPos.X, EndPos.Y ),			LocalSize, DrawScale, FVector4(EndUV,										Tiling),	Tint ) ); //15
 
 		// Top
-		BatchIndices.Add( IndexStart + 0 );
-		BatchIndices.Add( IndexStart + 1 );
-		BatchIndices.Add( IndexStart + 2 );
-		BatchIndices.Add( IndexStart + 2 );
-		BatchIndices.Add( IndexStart + 1 );
-		BatchIndices.Add( IndexStart + 3 );
-
-		BatchIndices.Add( IndexStart + 2 );
-		BatchIndices.Add( IndexStart + 3 );
-		BatchIndices.Add( IndexStart + 4 );
-		BatchIndices.Add( IndexStart + 4 );
-		BatchIndices.Add( IndexStart + 3 );
-		BatchIndices.Add( IndexStart + 5 );
-
-		BatchIndices.Add( IndexStart + 4 );
-		BatchIndices.Add( IndexStart + 5 );
-		BatchIndices.Add( IndexStart + 6 );
-		BatchIndices.Add( IndexStart + 6 );
-		BatchIndices.Add( IndexStart + 5 );
-		BatchIndices.Add( IndexStart + 7 );
+		RenderBatch.AddIndex( IndexStart + 0 );
+		RenderBatch.AddIndex( IndexStart + 1 );
+		RenderBatch.AddIndex( IndexStart + 2 );
+		RenderBatch.AddIndex( IndexStart + 2 );
+		RenderBatch.AddIndex( IndexStart + 1 );
+		RenderBatch.AddIndex( IndexStart + 3 );
+
+		RenderBatch.AddIndex( IndexStart + 2 );
+		RenderBatch.AddIndex( IndexStart + 3 );
+		RenderBatch.AddIndex( IndexStart + 4 );
+		RenderBatch.AddIndex( IndexStart + 4 );
+		RenderBatch.AddIndex( IndexStart + 3 );
+		RenderBatch.AddIndex( IndexStart + 5 );
+
+		RenderBatch.AddIndex( IndexStart + 4 );
+		RenderBatch.AddIndex( IndexStart + 5 );
+		RenderBatch.AddIndex( IndexStart + 6 );
+		RenderBatch.AddIndex( IndexStart + 6 );
+		RenderBatch.AddIndex( IndexStart + 5 );
+		RenderBatch.AddIndex( IndexStart + 7 );
 
 		// Middle
-		BatchIndices.Add( IndexStart + 1 );
-		BatchIndices.Add( IndexStart + 8 );
-		BatchIndices.Add( IndexStart + 3 );
-		BatchIndices.Add( IndexStart + 3 );
-		BatchIndices.Add( IndexStart + 8 );
-		BatchIndices.Add( IndexStart + 9 );
-
-		BatchIndices.Add( IndexStart + 3 );
-		BatchIndices.Add( IndexStart + 9 );
-		BatchIndices.Add( IndexStart + 5 );
-		BatchIndices.Add( IndexStart + 5 );
-		BatchIndices.Add( IndexStart + 9 );
-		BatchIndices.Add( IndexStart + 10 );
-
-		BatchIndices.Add( IndexStart + 5 );
-		BatchIndices.Add( IndexStart + 10 );
-		BatchIndices.Add( IndexStart + 7 );
-		BatchIndices.Add( IndexStart + 7 );
-		BatchIndices.Add( IndexStart + 10 );
-		BatchIndices.Add( IndexStart + 11 );
+		RenderBatch.AddIndex( IndexStart + 1 );
+		RenderBatch.AddIndex( IndexStart + 8 );
+		RenderBatch.AddIndex( IndexStart + 3 );
+		RenderBatch.AddIndex( IndexStart + 3 );
+		RenderBatch.AddIndex( IndexStart + 8 );
+		RenderBatch.AddIndex( IndexStart + 9 );
+
+		RenderBatch.AddIndex( IndexStart + 3 );
+		RenderBatch.AddIndex( IndexStart + 9 );
+		RenderBatch.AddIndex( IndexStart + 5 );
+		RenderBatch.AddIndex( IndexStart + 5 );
+		RenderBatch.AddIndex( IndexStart + 9 );
+		RenderBatch.AddIndex( IndexStart + 10 );
+
+		RenderBatch.AddIndex( IndexStart + 5 );
+		RenderBatch.AddIndex( IndexStart + 10 );
+		RenderBatch.AddIndex( IndexStart + 7 );
+		RenderBatch.AddIndex( IndexStart + 7 );
+		RenderBatch.AddIndex( IndexStart + 10 );
+		RenderBatch.AddIndex( IndexStart + 11 );
 
 		// Bottom
-		BatchIndices.Add( IndexStart + 8 );
-		BatchIndices.Add( IndexStart + 12 );
-		BatchIndices.Add( IndexStart + 9 );
-		BatchIndices.Add( IndexStart + 9 );
-		BatchIndices.Add( IndexStart + 12 );
-		BatchIndices.Add( IndexStart + 13 );
-
-		BatchIndices.Add( IndexStart + 9 );
-		BatchIndices.Add( IndexStart + 13 );
-		BatchIndices.Add( IndexStart + 10 );
-		BatchIndices.Add( IndexStart + 10 );
-		BatchIndices.Add( IndexStart + 13 );
-		BatchIndices.Add( IndexStart + 14 );
-
-		BatchIndices.Add( IndexStart + 10 );
-		BatchIndices.Add( IndexStart + 14 );
-		BatchIndices.Add( IndexStart + 11 );
-		BatchIndices.Add( IndexStart + 11 );
-		BatchIndices.Add( IndexStart + 14 );
-		BatchIndices.Add( IndexStart + 15 );
+		RenderBatch.AddIndex( IndexStart + 8 );
+		RenderBatch.AddIndex( IndexStart + 12 );
+		RenderBatch.AddIndex( IndexStart + 9 );
+		RenderBatch.AddIndex( IndexStart + 9 );
+		RenderBatch.AddIndex( IndexStart + 12 );
+		RenderBatch.AddIndex( IndexStart + 13 );
+
+		RenderBatch.AddIndex( IndexStart + 9 );
+		RenderBatch.AddIndex( IndexStart + 13 );
+		RenderBatch.AddIndex( IndexStart + 10 );
+		RenderBatch.AddIndex( IndexStart + 10 );
+		RenderBatch.AddIndex( IndexStart + 13 );
+		RenderBatch.AddIndex( IndexStart + 14 );
+
+		RenderBatch.AddIndex( IndexStart + 10 );
+		RenderBatch.AddIndex( IndexStart + 14 );
+		RenderBatch.AddIndex( IndexStart + 11 );
+		RenderBatch.AddIndex( IndexStart + 11 );
+		RenderBatch.AddIndex( IndexStart + 14 );
+		RenderBatch.AddIndex( IndexStart + 15 );
 
 		if ( GSlateFeathering && Rounding == ESlateVertexRounding::Disabled )
 		{
-			const int32 FeatherStart = BatchVertices.Num();
+			const int32 FeatherStart = RenderBatch.GetNumVertices();
 
 			// Top
-			BatchVertices.Add(FSlateVertex::Make<Rounding>(RenderTransform, FVector2D(Position.X, Position.Y) + FVector2D(-1, -1) / DrawScale, LocalSize, DrawScale, FVector4(StartUV, Tiling), FeatherColor)); //0
-			BatchVertices.Add(FSlateVertex::Make<Rounding>(RenderTransform, FVector2D(LeftMarginX, Position.Y) + FVector2D(0, -1) / DrawScale, LocalSize, DrawScale, FVector4(FVector2D(LeftMarginU, StartUV.Y), Tiling), FeatherColor)); //1
-			BatchVertices.Add(FSlateVertex::Make<Rounding>(RenderTransform, FVector2D(RightMarginX, Position.Y) + FVector2D(0, -1) / DrawScale, LocalSize, DrawScale, FVector4(FVector2D(RightMarginU, StartUV.Y), Tiling), FeatherColor)); //2
-			BatchVertices.Add(FSlateVertex::Make<Rounding>(RenderTransform, FVector2D(EndPos.X, Position.Y) + FVector2D(1, -1) / DrawScale, LocalSize, DrawScale, FVector4(FVector2D(EndUV.X, StartUV.Y), Tiling), FeatherColor)); //3
+			RenderBatch.AddVertex(FSlateVertex::Make<Rounding>(RenderTransform, FVector2D(Position.X, Position.Y) + FVector2D(-1, -1) / DrawScale, LocalSize, DrawScale, FVector4(StartUV, Tiling), FeatherColor)); //0
+			RenderBatch.AddVertex(FSlateVertex::Make<Rounding>(RenderTransform, FVector2D(LeftMarginX, Position.Y) + FVector2D(0, -1) / DrawScale, LocalSize, DrawScale, FVector4(FVector2D(LeftMarginU, StartUV.Y), Tiling), FeatherColor)); //1
+			RenderBatch.AddVertex(FSlateVertex::Make<Rounding>(RenderTransform, FVector2D(RightMarginX, Position.Y) + FVector2D(0, -1) / DrawScale, LocalSize, DrawScale, FVector4(FVector2D(RightMarginU, StartUV.Y), Tiling), FeatherColor)); //2
+			RenderBatch.AddVertex(FSlateVertex::Make<Rounding>(RenderTransform, FVector2D(EndPos.X, Position.Y) + FVector2D(1, -1) / DrawScale, LocalSize, DrawScale, FVector4(FVector2D(EndUV.X, StartUV.Y), Tiling), FeatherColor)); //3
 
 			// Left
-			BatchVertices.Add(FSlateVertex::Make<Rounding>(RenderTransform, FVector2D(Position.X, TopMarginY) + FVector2D(-1, 0) / DrawScale, LocalSize, DrawScale, FVector4(FVector2D(StartUV.X, TopMarginV), Tiling), FeatherColor)); //4
-			BatchVertices.Add(FSlateVertex::Make<Rounding>(RenderTransform, FVector2D(Position.X, BottomMarginY) + FVector2D(-1, 0) / DrawScale, LocalSize, DrawScale, FVector4(FVector2D(StartUV.X, BottomMarginV), Tiling), FeatherColor)); //5
+			RenderBatch.AddVertex(FSlateVertex::Make<Rounding>(RenderTransform, FVector2D(Position.X, TopMarginY) + FVector2D(-1, 0) / DrawScale, LocalSize, DrawScale, FVector4(FVector2D(StartUV.X, TopMarginV), Tiling), FeatherColor)); //4
+			RenderBatch.AddVertex(FSlateVertex::Make<Rounding>(RenderTransform, FVector2D(Position.X, BottomMarginY) + FVector2D(-1, 0) / DrawScale, LocalSize, DrawScale, FVector4(FVector2D(StartUV.X, BottomMarginV), Tiling), FeatherColor)); //5
 
 			// Right
-			BatchVertices.Add(FSlateVertex::Make<Rounding>(RenderTransform, FVector2D(EndPos.X, TopMarginY) + FVector2D(1, 0) / DrawScale, LocalSize, DrawScale, FVector4(FVector2D(EndUV.X, TopMarginV), Tiling), FeatherColor)); //6
-			BatchVertices.Add(FSlateVertex::Make<Rounding>(RenderTransform, FVector2D(EndPos.X, BottomMarginY) + FVector2D(1, 0) / DrawScale, LocalSize, DrawScale, FVector4(FVector2D(EndUV.X, BottomMarginV), Tiling), FeatherColor)); //7
+			RenderBatch.AddVertex(FSlateVertex::Make<Rounding>(RenderTransform, FVector2D(EndPos.X, TopMarginY) + FVector2D(1, 0) / DrawScale, LocalSize, DrawScale, FVector4(FVector2D(EndUV.X, TopMarginV), Tiling), FeatherColor)); //6
+			RenderBatch.AddVertex(FSlateVertex::Make<Rounding>(RenderTransform, FVector2D(EndPos.X, BottomMarginY) + FVector2D(1, 0) / DrawScale, LocalSize, DrawScale, FVector4(FVector2D(EndUV.X, BottomMarginV), Tiling), FeatherColor)); //7
 
 			// Bottom
-			BatchVertices.Add(FSlateVertex::Make<Rounding>(RenderTransform, FVector2D(Position.X, EndPos.Y) + FVector2D(-1, 1) / DrawScale, LocalSize, DrawScale, FVector4(FVector2D(StartUV.X, EndUV.Y), Tiling), FeatherColor)); //8
-			BatchVertices.Add(FSlateVertex::Make<Rounding>(RenderTransform, FVector2D(LeftMarginX, EndPos.Y) + FVector2D(0, 1) / DrawScale, LocalSize, DrawScale, FVector4(FVector2D(LeftMarginU, EndUV.Y), Tiling), FeatherColor)); //9
-			BatchVertices.Add(FSlateVertex::Make<Rounding>(RenderTransform, FVector2D(RightMarginX, EndPos.Y) + FVector2D(0, 1) / DrawScale, LocalSize, DrawScale, FVector4(FVector2D(RightMarginU, EndUV.Y), Tiling), FeatherColor)); //10
-			BatchVertices.Add(FSlateVertex::Make<Rounding>(RenderTransform, FVector2D(EndPos.X, EndPos.Y) + FVector2D(1, 1) / DrawScale, LocalSize, DrawScale, FVector4(EndUV, Tiling), FeatherColor)); //11
+			RenderBatch.AddVertex(FSlateVertex::Make<Rounding>(RenderTransform, FVector2D(Position.X, EndPos.Y) + FVector2D(-1, 1) / DrawScale, LocalSize, DrawScale, FVector4(FVector2D(StartUV.X, EndUV.Y), Tiling), FeatherColor)); //8
+			RenderBatch.AddVertex(FSlateVertex::Make<Rounding>(RenderTransform, FVector2D(LeftMarginX, EndPos.Y) + FVector2D(0, 1) / DrawScale, LocalSize, DrawScale, FVector4(FVector2D(LeftMarginU, EndUV.Y), Tiling), FeatherColor)); //9
+			RenderBatch.AddVertex(FSlateVertex::Make<Rounding>(RenderTransform, FVector2D(RightMarginX, EndPos.Y) + FVector2D(0, 1) / DrawScale, LocalSize, DrawScale, FVector4(FVector2D(RightMarginU, EndUV.Y), Tiling), FeatherColor)); //10
+			RenderBatch.AddVertex(FSlateVertex::Make<Rounding>(RenderTransform, FVector2D(EndPos.X, EndPos.Y) + FVector2D(1, 1) / DrawScale, LocalSize, DrawScale, FVector4(EndUV, Tiling), FeatherColor)); //11
 
 			// Top Left
-			IndexQuad(BatchIndices, FeatherStart + 0, FeatherStart + 1, IndexStart + 2, IndexStart + 0);
+			IndexQuad(RenderBatch, FeatherStart + 0, FeatherStart + 1, IndexStart + 2, IndexStart + 0);
 			// Top Middle
-			IndexQuad(BatchIndices, FeatherStart + 1, FeatherStart + 2, IndexStart + 4, IndexStart + 2);
+			IndexQuad(RenderBatch, FeatherStart + 1, FeatherStart + 2, IndexStart + 4, IndexStart + 2);
 			// Top Right
-			IndexQuad(BatchIndices, FeatherStart + 2, FeatherStart + 3, IndexStart + 6, IndexStart + 4);
+			IndexQuad(RenderBatch, FeatherStart + 2, FeatherStart + 3, IndexStart + 6, IndexStart + 4);
 
 			//-----------------------------------------------------------
 
 			// Left Top
-			IndexQuad(BatchIndices, FeatherStart + 0, IndexStart + 0, IndexStart + 1, FeatherStart + 4);
+			IndexQuad(RenderBatch, FeatherStart + 0, IndexStart + 0, IndexStart + 1, FeatherStart + 4);
 			// Left Middle
-			IndexQuad(BatchIndices, FeatherStart + 4, IndexStart + 1, IndexStart + 8, FeatherStart + 5);
+			IndexQuad(RenderBatch, FeatherStart + 4, IndexStart + 1, IndexStart + 8, FeatherStart + 5);
 			// Left Bottom
-			IndexQuad(BatchIndices, FeatherStart + 5, IndexStart + 8, IndexStart + 12, FeatherStart + 8);
+			IndexQuad(RenderBatch, FeatherStart + 5, IndexStart + 8, IndexStart + 12, FeatherStart + 8);
 
 			//-----------------------------------------------------------
 
 			// Right Top
-			IndexQuad(BatchIndices, IndexStart + 6, FeatherStart + 3, FeatherStart + 6, IndexStart + 7);
+			IndexQuad(RenderBatch, IndexStart + 6, FeatherStart + 3, FeatherStart + 6, IndexStart + 7);
 			// Right Middle
-			IndexQuad(BatchIndices, IndexStart + 7, FeatherStart + 6, FeatherStart + 7, IndexStart + 11);
+			IndexQuad(RenderBatch, IndexStart + 7, FeatherStart + 6, FeatherStart + 7, IndexStart + 11);
 			// Right Bottom
-			IndexQuad(BatchIndices, IndexStart + 11, FeatherStart + 7, FeatherStart + 11, IndexStart + 15);
+			IndexQuad(RenderBatch, IndexStart + 11, FeatherStart + 7, FeatherStart + 11, IndexStart + 15);
 
 			//-----------------------------------------------------------
 
 			// Bottom Left
-			IndexQuad(BatchIndices, IndexStart + 12, IndexStart + 13, FeatherStart + 9, FeatherStart + 8);
+			IndexQuad(RenderBatch, IndexStart + 12, IndexStart + 13, FeatherStart + 9, FeatherStart + 8);
 			// Bottom Middle
-			IndexQuad(BatchIndices, IndexStart + 13, IndexStart + 14, FeatherStart + 10, FeatherStart + 9);
+			IndexQuad(RenderBatch, IndexStart + 13, IndexStart + 14, FeatherStart + 10, FeatherStart + 9);
 			// Bottom Right
-			IndexQuad(BatchIndices, IndexStart + 14, IndexStart + 15, FeatherStart + 11, FeatherStart + 10);
+			IndexQuad(RenderBatch, IndexStart + 14, IndexStart + 15, FeatherStart + 11, FeatherStart + 10);
 		}
 	}
 	else
@@ -558,18 +963,18 @@ void FSlateElementBatcher::AddBoxElement(const FSlateDrawElement& DrawElement)
 		}
 
 		// Add four vertices to the list of verts to be added to the vertex buffer
-		BatchVertices.Add(FSlateVertex::Make<Rounding>(RenderTransform, TopLeft, LocalSize, DrawScale, FVector4(StartUV, Tiling), Tint));
-		BatchVertices.Add(FSlateVertex::Make<Rounding>(RenderTransform, TopRight, LocalSize, DrawScale, FVector4(FVector2D(EndUV.X, StartUV.Y), Tiling), Tint));
-		BatchVertices.Add(FSlateVertex::Make<Rounding>(RenderTransform, BotLeft, LocalSize, DrawScale, FVector4(FVector2D(StartUV.X, EndUV.Y), Tiling), Tint));
-		BatchVertices.Add(FSlateVertex::Make<Rounding>(RenderTransform, BotRight, LocalSize, DrawScale, FVector4(EndUV, Tiling), Tint));
+		RenderBatch.AddVertex(FSlateVertex::Make<Rounding>(RenderTransform, TopLeft, LocalSize, DrawScale, FVector4(StartUV, Tiling), Tint));
+		RenderBatch.AddVertex(FSlateVertex::Make<Rounding>(RenderTransform, TopRight, LocalSize, DrawScale, FVector4(FVector2D(EndUV.X, StartUV.Y), Tiling), Tint));
+		RenderBatch.AddVertex(FSlateVertex::Make<Rounding>(RenderTransform, BotLeft, LocalSize, DrawScale, FVector4(FVector2D(StartUV.X, EndUV.Y), Tiling), Tint));
+		RenderBatch.AddVertex(FSlateVertex::Make<Rounding>(RenderTransform, BotRight, LocalSize, DrawScale, FVector4(EndUV, Tiling), Tint));
 
-		BatchIndices.Add( IndexStart + 0 );
-		BatchIndices.Add( IndexStart + 1 );
-		BatchIndices.Add( IndexStart + 2 );
+		RenderBatch.AddIndex( IndexStart + 0 );
+		RenderBatch.AddIndex( IndexStart + 1 );
+		RenderBatch.AddIndex( IndexStart + 2 );
 
-		BatchIndices.Add( IndexStart + 2 );
-		BatchIndices.Add( IndexStart + 1 );
-		BatchIndices.Add( IndexStart + 3 );
+		RenderBatch.AddIndex( IndexStart + 2 );
+		RenderBatch.AddIndex( IndexStart + 1 );
+		RenderBatch.AddIndex( IndexStart + 3 );
 
 		int32 TopLeftIndex = IndexStart + 0;
 		int32 TopRightIndex = IndexStart + 1;
@@ -578,52 +983,52 @@ void FSlateElementBatcher::AddBoxElement(const FSlateDrawElement& DrawElement)
 
 		if ( GSlateFeathering && Rounding == ESlateVertexRounding::Disabled )
 		{
-			const int32 FeatherStart = BatchVertices.Num();
+			const int32 FeatherStart = RenderBatch.GetNumVertices();
 
-			BatchVertices.Add(FSlateVertex::Make<Rounding>(RenderTransform, TopLeft + FVector2D(-1, -1) / DrawScale, LocalSize, DrawScale, FVector4(StartUV, Tiling), FeatherColor));
-			BatchVertices.Add(FSlateVertex::Make<Rounding>(RenderTransform, TopRight + FVector2D(1, -1) / DrawScale, LocalSize, DrawScale, FVector4(FVector2D(EndUV.X, StartUV.Y), Tiling), FeatherColor));
-			BatchVertices.Add(FSlateVertex::Make<Rounding>(RenderTransform, BotLeft + FVector2D(-1, 1) / DrawScale, LocalSize, DrawScale, FVector4(FVector2D(StartUV.X, EndUV.Y), Tiling), FeatherColor));
-			BatchVertices.Add(FSlateVertex::Make<Rounding>(RenderTransform, BotRight + FVector2D(1, 1) / DrawScale, LocalSize, DrawScale, FVector4(EndUV, Tiling), FeatherColor));
+			RenderBatch.AddVertex(FSlateVertex::Make<Rounding>(RenderTransform, TopLeft + FVector2D(-1, -1) / DrawScale, LocalSize, DrawScale, FVector4(StartUV, Tiling), FeatherColor));
+			RenderBatch.AddVertex(FSlateVertex::Make<Rounding>(RenderTransform, TopRight + FVector2D(1, -1) / DrawScale, LocalSize, DrawScale, FVector4(FVector2D(EndUV.X, StartUV.Y), Tiling), FeatherColor));
+			RenderBatch.AddVertex(FSlateVertex::Make<Rounding>(RenderTransform, BotLeft + FVector2D(-1, 1) / DrawScale, LocalSize, DrawScale, FVector4(FVector2D(StartUV.X, EndUV.Y), Tiling), FeatherColor));
+			RenderBatch.AddVertex(FSlateVertex::Make<Rounding>(RenderTransform, BotRight + FVector2D(1, 1) / DrawScale, LocalSize, DrawScale, FVector4(EndUV, Tiling), FeatherColor));
 
 			// Top-Top
-			BatchIndices.Add(FeatherStart + 0);
-			BatchIndices.Add(FeatherStart + 1);
-			BatchIndices.Add(TopRightIndex);
+			RenderBatch.AddIndex(FeatherStart + 0);
+			RenderBatch.AddIndex(FeatherStart + 1);
+			RenderBatch.AddIndex(TopRightIndex);
 
 			// Top-Bottom
-			BatchIndices.Add(FeatherStart + 0);
-			BatchIndices.Add(TopRightIndex);
-			BatchIndices.Add(TopLeftIndex);
+			RenderBatch.AddIndex(FeatherStart + 0);
+			RenderBatch.AddIndex(TopRightIndex);
+			RenderBatch.AddIndex(TopLeftIndex);
 
 			// Left-Top
-			BatchIndices.Add(FeatherStart + 0);
-			BatchIndices.Add(BottomLeftIndex);
-			BatchIndices.Add(FeatherStart + 2);
+			RenderBatch.AddIndex(FeatherStart + 0);
+			RenderBatch.AddIndex(BottomLeftIndex);
+			RenderBatch.AddIndex(FeatherStart + 2);
 
 			// Left-Bottom
-			BatchIndices.Add(FeatherStart + 0);
-			BatchIndices.Add(TopLeftIndex);
-			BatchIndices.Add(BottomLeftIndex);
+			RenderBatch.AddIndex(FeatherStart + 0);
+			RenderBatch.AddIndex(TopLeftIndex);
+			RenderBatch.AddIndex(BottomLeftIndex);
 
 			// Right-Top
-			BatchIndices.Add(TopRightIndex);
-			BatchIndices.Add(FeatherStart + 1);
-			BatchIndices.Add(FeatherStart + 3);
+			RenderBatch.AddIndex(TopRightIndex);
+			RenderBatch.AddIndex(FeatherStart + 1);
+			RenderBatch.AddIndex(FeatherStart + 3);
 
 			// Right-Bottom
-			BatchIndices.Add(TopRightIndex);
-			BatchIndices.Add(FeatherStart + 3);
-			BatchIndices.Add(BottomRightIndex);
+			RenderBatch.AddIndex(TopRightIndex);
+			RenderBatch.AddIndex(FeatherStart + 3);
+			RenderBatch.AddIndex(BottomRightIndex);
 
 			// Bottom-Top
-			BatchIndices.Add(BottomLeftIndex);
-			BatchIndices.Add(BottomRightIndex);
-			BatchIndices.Add(FeatherStart + 3);
+			RenderBatch.AddIndex(BottomLeftIndex);
+			RenderBatch.AddIndex(BottomRightIndex);
+			RenderBatch.AddIndex(FeatherStart + 3);
 
 			// Bottom-Bottom
-			BatchIndices.Add(FeatherStart + 3);
-			BatchIndices.Add(FeatherStart + 2);
-			BatchIndices.Add(BottomLeftIndex);
+			RenderBatch.AddIndex(FeatherStart + 3);
+			RenderBatch.AddIndex(FeatherStart + 2);
+			RenderBatch.AddIndex(BottomLeftIndex);
 		}
 	}
 }
@@ -631,7 +1036,7 @@ void FSlateElementBatcher::AddBoxElement(const FSlateDrawElement& DrawElement)
 template<ESlateVertexRounding Rounding>
 void FSlateElementBatcher::AddTextElement(const FSlateDrawElement& DrawElement)
 {
-	const FSlateDataPayload& DrawElementPayload = DrawElement.GetDataPayload();
+	const FSlateTextPayload& DrawElementPayload = DrawElement.GetDataPayload<FSlateTextPayload>();
 	FColor BaseTint = PackVertexColor(DrawElementPayload.GetTint());
 
 	const FFontOutlineSettings& OutlineSettings = DrawElementPayload.GetFontInfo().OutlineSettings;
@@ -640,7 +1045,8 @@ void FSlateElementBatcher::AddTextElement(const FSlateDrawElement& DrawElement)
 	ensure(Len > 0);
 
 	ESlateDrawEffect InDrawEffects = DrawElement.GetDrawEffects();
-	uint32 Layer = DrawElement.GetLayer();
+
+	const int32 Layer = DrawElement.GetLayer();
 
 	// extract the layout transform from the draw element
 	FSlateLayoutTransform LayoutTransform(DrawElement.GetScale(), DrawElement.GetPosition());
@@ -679,7 +1085,7 @@ void FSlateElementBatcher::AddTextElement(const FSlateDrawElement& DrawElement)
 		FSlateShaderResource* FontAtlasTexture = nullptr;
 		FSlateShaderResource* FontShaderResource = nullptr;
 
-		FSlateElementBatch* ElementBatch = nullptr;
+		FSlateRenderBatch* RenderBatch = nullptr;
 		FSlateVertexArray* BatchVertices = nullptr;
 		FSlateIndexArray* BatchIndices = nullptr;
 
@@ -740,14 +1146,14 @@ void FSlateElementBatcher::AddTextElement(const FSlateDrawElement& DrawElement)
 					FontShaderResource = ResourceManager.GetFontShaderResource( FontTextureIndex, FontAtlasTexture, DrawElementPayload.GetFontInfo().FontMaterial );
 					check(FontShaderResource);
 
-					ElementBatch = &FindBatchForElement(InLayer, FShaderParams(), FontShaderResource, ESlateDrawPrimitive::TriangleList, ESlateShader::Font, InDrawEffects, ESlateBatchDrawFlag::None, DrawElement.GetClippingIndex(), DrawElement.GetSceneIndex());
+					RenderBatch = &CreateRenderBatch(InLayer, FShaderParams(), FontShaderResource, ESlateDrawPrimitive::TriangleList, ESlateShader::Font, InDrawEffects, ESlateBatchDrawFlag::None, DrawElement);
+
+					const int32 GlyphsLeft = NumChars - CharIndex;
 
-					BatchVertices = &BatchData->GetBatchVertexList(*ElementBatch);
-					BatchIndices = &BatchData->GetBatchIndexList(*ElementBatch);
+					// Reserve memory for the glyphs.  This isn't perfect as the text could contain spaces and we might not render the rest of the text in this batch but its better than resizing constantly
+					RenderBatch->ReserveVertices(GlyphsLeft * 4);
+					RenderBatch->ReserveIndices(GlyphsLeft * 6);
 
-					VertexOffset = BatchVertices->Num();
-					IndexOffset = BatchIndices->Num();
-				
 					InvTextureSizeX = 1.0f / FontAtlasTexture->GetWidth();
 					InvTextureSizeY = 1.0f / FontAtlasTexture->GetHeight();
 				}
@@ -780,21 +1186,13 @@ void FSlateElementBatcher::AddTextElement(const FSlateDrawElement& DrawElement)
 					const float SizeV = Entry.VSize * InvTextureSizeY;
 
 					{
-						FSlateVertexArray& BatchVerticesRef = *BatchVertices;
-						FSlateIndexArray& BatchIndicesRef = *BatchIndices;
-
 						FVector2D UpperLeft( X, Y );
 						FVector2D UpperRight( X+SizeX, Y );
 						FVector2D LowerLeft( X, Y+SizeY );
 						FVector2D LowerRight( X+SizeX, Y+SizeY );
 
-						// Add four vertices for this quad
-						BatchVerticesRef.AddUninitialized( 4 );
-						// Add six indices for this quad
-						BatchIndicesRef.AddUninitialized( 6 );
-
 						// The start index of these vertices in the index buffer
-						uint32 IndexStart = VertexOffset;
+						const uint32 IndexStart = RenderBatch->GetNumVertices();
 
 						float Ut = 0.0f, Vt = 0.0f, UtMax = 0.0f, VtMax = 0.0f;
 						if( bIsFontMaterial )
@@ -811,17 +1209,17 @@ void FSlateElementBatcher::AddTextElement(const FSlateDrawElement& DrawElement)
 						}
 
 						// Add four vertices to the list of verts to be added to the vertex buffer
-						BatchVerticesRef[ VertexOffset++ ] = FSlateVertex::Make<Rounding>( RenderTransform, UpperLeft,								FVector4(U,V,				Ut,Vt),			FVector2D(0.0f,0.0f), InTint );
-						BatchVerticesRef[ VertexOffset++ ] = FSlateVertex::Make<Rounding>( RenderTransform, FVector2D(LowerRight.X,UpperLeft.Y),	FVector4(U+SizeU, V,		UtMax,Vt),		FVector2D(1.0f,0.0f), InTint );
-						BatchVerticesRef[ VertexOffset++ ] = FSlateVertex::Make<Rounding>( RenderTransform, FVector2D(UpperLeft.X,LowerRight.Y),	FVector4(U, V+SizeV,		Ut,VtMax),		FVector2D(0.0f,1.0f), InTint );
-						BatchVerticesRef[ VertexOffset++ ] = FSlateVertex::Make<Rounding>( RenderTransform, LowerRight,								FVector4(U+SizeU, V+SizeV,	UtMax,VtMax),	FVector2D(1.0f,1.0f), InTint );
-
-						BatchIndicesRef[IndexOffset++] = IndexStart + 0;
-						BatchIndicesRef[IndexOffset++] = IndexStart + 1;
-						BatchIndicesRef[IndexOffset++] = IndexStart + 2;
-						BatchIndicesRef[IndexOffset++] = IndexStart + 1;
-						BatchIndicesRef[IndexOffset++] = IndexStart + 3;
-						BatchIndicesRef[IndexOffset++] = IndexStart + 2;
+						RenderBatch->AddVertex(FSlateVertex::Make<Rounding>( RenderTransform, UpperLeft,								FVector4(U,V,Ut,Vt),						FVector2D(0.0f,0.0f), InTint ));
+						RenderBatch->AddVertex(FSlateVertex::Make<Rounding>( RenderTransform, FVector2D(LowerRight.X,UpperLeft.Y),	FVector4(U+SizeU, V, UtMax,Vt),				FVector2D(1.0f,0.0f), InTint ));
+						RenderBatch->AddVertex(FSlateVertex::Make<Rounding>( RenderTransform, FVector2D(UpperLeft.X,LowerRight.Y),	FVector4(U, V+SizeV, Ut,VtMax),				FVector2D(0.0f,1.0f), InTint ));
+						RenderBatch->AddVertex(FSlateVertex::Make<Rounding>( RenderTransform, LowerRight,							FVector4(U+SizeU, V+SizeV, UtMax,VtMax),	FVector2D(1.0f,1.0f), InTint ));
+
+						RenderBatch->AddIndex(IndexStart + 0);
+						RenderBatch->AddIndex(IndexStart + 1);
+						RenderBatch->AddIndex(IndexStart + 2);
+						RenderBatch->AddIndex(IndexStart + 1);
+						RenderBatch->AddIndex(IndexStart + 3);
+						RenderBatch->AddIndex(IndexStart + 2);
 					}
 				}
 
@@ -851,7 +1249,7 @@ void FSlateElementBatcher::AddTextElement(const FSlateDrawElement& DrawElement)
 template<ESlateVertexRounding Rounding>
 void FSlateElementBatcher::AddShapedTextElement( const FSlateDrawElement& DrawElement )
 {
-	const FSlateDataPayload& DrawElementPayload = DrawElement.GetDataPayload();
+	const FSlateShapedTextPayload& DrawElementPayload = DrawElement.GetDataPayload<FSlateShapedTextPayload>();
 	const FShapedGlyphSequence* ShapedGlyphSequence = DrawElementPayload.GetShapedGlyphSequence().Get();
 	checkSlow(ShapedGlyphSequence);
 
@@ -875,7 +1273,8 @@ void FSlateElementBatcher::AddShapedTextElement( const FSlateDrawElement& DrawEl
 	}
 
 	ESlateDrawEffect InDrawEffects = DrawElement.GetDrawEffects();
-	uint32 Layer = DrawElement.GetLayer();
+
+	const int32 Layer = DrawElement.GetLayer();
 
 	// extract the layout transform from the draw element
 	FSlateLayoutTransform LayoutTransform(DrawElement.GetScale(), DrawElement.GetPosition());
@@ -908,9 +1307,7 @@ void FSlateElementBatcher::AddShapedTextElement( const FSlateDrawElement& DrawEl
 		FSlateShaderResource* FontAtlasTexture = nullptr;
 		FSlateShaderResource* FontShaderResource = nullptr;
 
-		FSlateElementBatch* ElementBatch = nullptr;
-		FSlateVertexArray* BatchVertices = nullptr;
-		FSlateIndexArray* BatchIndices = nullptr;
+		FSlateRenderBatch* RenderBatch = nullptr;
 
 		uint32 VertexOffset = 0;
 		uint32 IndexOffset = 0;
@@ -921,22 +1318,20 @@ void FSlateElementBatcher::AddShapedTextElement( const FSlateDrawElement& DrawEl
 		const bool bIsFontMaterial = FontMaterial != nullptr;
 
 		// Optimize by culling
+		// Todo: this doesnt work with cached clipping
 		bool bEnableCulling = false;
 		float LocalClipBoundingBoxLeft = 0;
 		float LocalClipBoundingBoxRight = 0;
 		if (GlyphsToRender.Num() > 200)
 		{
-			int16 ClippingIndex = DrawElement.GetClippingIndex();
-			if (ClippingStates && ClippingStates->IsValidIndex(ClippingIndex))
+			const FSlateClippingState* ClippingState = ResolveClippingState(DrawElement);
+
+			if (ClippingState && ClippingState->ScissorRect.IsSet() && ClippingState->ScissorRect->IsAxisAligned() && RenderTransform.GetMatrix().IsIdentity())
 			{
-				const FSlateClippingState& ClippingState = (*ClippingStates)[ClippingIndex];
-				if (ClippingState.ScissorRect.IsSet() && ClippingState.ScissorRect->IsAxisAligned() && RenderTransform.GetMatrix().IsIdentity())
-				{
-					bEnableCulling = true;
-					const FSlateRect LocalClipBoundingBox = TransformRect(RenderTransform.Inverse(), ClippingState.ScissorRect->GetBoundingBox());
-					LocalClipBoundingBoxLeft = LocalClipBoundingBox.Left;
-					LocalClipBoundingBoxRight = LocalClipBoundingBox.Right;
-				}
+				bEnableCulling = true;
+				const FSlateRect LocalClipBoundingBox = TransformRect(RenderTransform.Inverse(), ClippingState->ScissorRect->GetBoundingBox());
+				LocalClipBoundingBoxLeft = LocalClipBoundingBox.Left;
+				LocalClipBoundingBoxRight = LocalClipBoundingBox.Right;
 			}
 		}
 
@@ -979,13 +1374,13 @@ void FSlateElementBatcher::AddShapedTextElement( const FSlateDrawElement& DrawEl
 						FontShaderResource = ResourceManager.GetFontShaderResource(FontTextureIndex, FontAtlasTexture, FontMaterial);
 						check(FontShaderResource);
 
-						ElementBatch = &FindBatchForElement(InLayer, FShaderParams(), FontShaderResource, ESlateDrawPrimitive::TriangleList, ESlateShader::Font, InDrawEffects, ESlateBatchDrawFlag::None, DrawElement.GetClippingIndex(), DrawElement.GetSceneIndex());
+						RenderBatch = &CreateRenderBatch(InLayer, FShaderParams(), FontShaderResource, ESlateDrawPrimitive::TriangleList, ESlateShader::Font, InDrawEffects, ESlateBatchDrawFlag::None, DrawElement);
 
-						BatchVertices = &BatchData->GetBatchVertexList(*ElementBatch);
-						BatchIndices = &BatchData->GetBatchIndexList(*ElementBatch);
+						const int32 GlyphsLeft = NumGlyphs - GlyphIndex;
 
-						VertexOffset = BatchVertices->Num();
-						IndexOffset = BatchIndices->Num();
+						// Reserve memory for the glyphs.  This isn't perfect as the text could contain spaces and we might not render the rest of the text in this batch but its better than resizing constantly
+						RenderBatch->ReserveVertices(GlyphsLeft*4);
+						RenderBatch->ReserveIndices(GlyphsLeft*6);
 
 						InvTextureSizeX = 1.0f / FontAtlasTexture->GetWidth();
 						InvTextureSizeY = 1.0f / FontAtlasTexture->GetHeight();
@@ -1001,21 +1396,15 @@ void FSlateElementBatcher::AddShapedTextElement( const FSlateDrawElement& DrawEl
 					const float SizeV = GlyphAtlasData.VSize * InvTextureSizeY;
 
 					{
-						FSlateVertexArray& BatchVerticesRef = *BatchVertices;
-						FSlateIndexArray& BatchIndicesRef = *BatchIndices;
 
 						const FVector2D UpperLeft(X, Y);
 						const FVector2D UpperRight(X + SizeX, Y);
 						const FVector2D LowerLeft(X, Y + SizeY);
 						const FVector2D LowerRight(X + SizeX, Y + SizeY);
 
-						// Add four vertices for this quad
-						BatchVerticesRef.AddUninitialized(4);
-						// Add six indices for this quad
-						BatchIndicesRef.AddUninitialized(6);
-
+		
 						// The start index of these vertices in the index buffer
-						uint32 IndexStart = VertexOffset;
+						const uint32 IndexStart = RenderBatch->GetNumVertices();
 
 						float Ut = 0.0f, Vt = 0.0f, UtMax = 0.0f, VtMax = 0.0f;
 						if (bIsFontMaterial)
@@ -1032,17 +1421,17 @@ void FSlateElementBatcher::AddShapedTextElement( const FSlateDrawElement& DrawEl
 						}
 
 						// Add four vertices to the list of verts to be added to the vertex buffer
-						BatchVerticesRef[VertexOffset++] = FSlateVertex::Make<Rounding>(RenderTransform, UpperLeft,								FVector4(U, V, Ut, Vt),							FVector2D(0.0f, 0.0f), InTint );
-						BatchVerticesRef[VertexOffset++] = FSlateVertex::Make<Rounding>(RenderTransform, FVector2D(LowerRight.X, UpperLeft.Y),	FVector4(U + SizeU, V, UtMax, Vt),				FVector2D(1.0f, 0.0f), InTint );
-						BatchVerticesRef[VertexOffset++] = FSlateVertex::Make<Rounding>(RenderTransform, FVector2D(UpperLeft.X, LowerRight.Y),	FVector4(U, V + SizeV, Ut, VtMax),				FVector2D(0.0f, 1.0f), InTint );
-						BatchVerticesRef[VertexOffset++] = FSlateVertex::Make<Rounding>(RenderTransform, LowerRight,							FVector4(U + SizeU, V + SizeV, UtMax, VtMax),	FVector2D(1.0f, 1.0f), InTint );
-
-						BatchIndicesRef[IndexOffset++] = IndexStart + 0;
-						BatchIndicesRef[IndexOffset++] = IndexStart + 1;
-						BatchIndicesRef[IndexOffset++] = IndexStart + 2;
-						BatchIndicesRef[IndexOffset++] = IndexStart + 1;
-						BatchIndicesRef[IndexOffset++] = IndexStart + 3;
-						BatchIndicesRef[IndexOffset++] = IndexStart + 2;
+						RenderBatch->AddVertex(FSlateVertex::Make<Rounding>(RenderTransform, UpperLeft,								FVector4(U, V, Ut, Vt),							FVector2D(0.0f, 0.0f), InTint ));
+						RenderBatch->AddVertex(FSlateVertex::Make<Rounding>(RenderTransform, FVector2D(LowerRight.X, UpperLeft.Y),	FVector4(U + SizeU, V, UtMax, Vt),				FVector2D(1.0f, 0.0f), InTint ));
+						RenderBatch->AddVertex(FSlateVertex::Make<Rounding>(RenderTransform, FVector2D(UpperLeft.X, LowerRight.Y),	FVector4(U, V + SizeV, Ut, VtMax),				FVector2D(0.0f, 1.0f), InTint ));
+						RenderBatch->AddVertex(FSlateVertex::Make<Rounding>(RenderTransform, LowerRight,								FVector4(U + SizeU, V + SizeV, UtMax, VtMax),	FVector2D(1.0f, 1.0f), InTint ));
+
+						RenderBatch->AddIndex(IndexStart + 0);
+						RenderBatch->AddIndex(IndexStart + 1);
+						RenderBatch->AddIndex(IndexStart + 2);
+						RenderBatch->AddIndex(IndexStart + 1);
+						RenderBatch->AddIndex(IndexStart + 3);
+						RenderBatch->AddIndex(IndexStart + 2);
 					}
 				}
 			}
@@ -1075,27 +1464,23 @@ void FSlateElementBatcher::AddGradientElement( const FSlateDrawElement& DrawElem
 {
 	const FSlateRenderTransform& RenderTransform = DrawElement.GetRenderTransform();
 	const FVector2D& LocalSize = DrawElement.GetLocalSize();
-	const FSlateDataPayload& InPayload = DrawElement.GetDataPayload();
-	ESlateDrawEffect InDrawEffects = DrawElement.GetDrawEffects();
-	uint32 Layer = DrawElement.GetLayer();
+	const FSlateGradientPayload& InPayload = DrawElement.GetDataPayload<FSlateGradientPayload>();
+	const ESlateDrawEffect InDrawEffects = DrawElement.GetDrawEffects();
+	const int32 Layer = DrawElement.GetLayer();
 
 	// There must be at least one gradient stop
 	check( InPayload.GradientStops.Num() > 0 );
 
-	FSlateElementBatch& ElementBatch = 
-		FindBatchForElement( 
+	FSlateRenderBatch& RenderBatch = 
+		CreateRenderBatch( 
 			Layer,
 			FShaderParams(),
 			nullptr,
 			ESlateDrawPrimitive::TriangleList,
 			ESlateShader::Default,
 			InDrawEffects,
-			InPayload.BatchFlags,
-			DrawElement.GetClippingIndex(),
-			DrawElement.GetSceneIndex());
-
-	FSlateVertexArray& BatchVertices = BatchData->GetBatchVertexList(ElementBatch);
-	FSlateIndexArray& BatchIndices = BatchData->GetBatchIndexList(ElementBatch);
+			DrawElement.GetBatchFlags(),
+			DrawElement);
 
 	// Determine the four corners of the quad containing the gradient
 	FVector2D TopLeft = FVector2D::ZeroVector;
@@ -1140,13 +1525,11 @@ void FSlateElementBatcher::AddGradientElement( const FSlateDrawElement& DrawElem
 		}
 	}
 
-	uint32 IndexOffsetStart = BatchIndices.Num();
-
 	// Add a pair of vertices for each gradient stop. Connecting them to the previous stop if necessary
 	// Assumes gradient stops are sorted by position left to right or top to bottom
 	for( int32 StopIndex = 0; StopIndex < GradientStops.Num(); ++StopIndex )
 	{
-		uint32 IndexStart = BatchVertices.Num();
+		const uint32 IndexStart = RenderBatch.GetNumVertices();
 
 		const FSlateGradientStop& CurStop = GradientStops[StopIndex];
 
@@ -1177,23 +1560,23 @@ void FSlateElementBatcher::AddGradientElement( const FSlateDrawElement& DrawElem
 		if( StopIndex == 0 )
 		{
 			// First stop does not have a full quad yet so do not create indices
-			BatchVertices.Add(FSlateVertex::Make<Rounding>(RenderTransform, StartPt, FVector2D::ZeroVector, FVector2D::ZeroVector, CurStop.Color.ToFColor(false) ) );
-			BatchVertices.Add(FSlateVertex::Make<Rounding>(RenderTransform, EndPt, FVector2D::ZeroVector, FVector2D::ZeroVector, CurStop.Color.ToFColor(false) ) );
+			RenderBatch.AddVertex(FSlateVertex::Make<Rounding>(RenderTransform, StartPt, FVector2D::ZeroVector, FVector2D::ZeroVector, CurStop.Color.ToFColor(false) ) );
+			RenderBatch.AddVertex(FSlateVertex::Make<Rounding>(RenderTransform, EndPt, FVector2D::ZeroVector, FVector2D::ZeroVector, CurStop.Color.ToFColor(false) ) );
 		}
 		else
 		{
 			// All stops after the first have indices and generate quads
-			BatchVertices.Add(FSlateVertex::Make<Rounding>(RenderTransform, StartPt, FVector2D::ZeroVector, FVector2D::ZeroVector, CurStop.Color.ToFColor(false) ) );
-			BatchVertices.Add(FSlateVertex::Make<Rounding>(RenderTransform, EndPt, FVector2D::ZeroVector, FVector2D::ZeroVector, CurStop.Color.ToFColor(false) ) );
+			RenderBatch.AddVertex(FSlateVertex::Make<Rounding>(RenderTransform, StartPt, FVector2D::ZeroVector, FVector2D::ZeroVector, CurStop.Color.ToFColor(false) ) );
+			RenderBatch.AddVertex(FSlateVertex::Make<Rounding>(RenderTransform, EndPt, FVector2D::ZeroVector, FVector2D::ZeroVector, CurStop.Color.ToFColor(false) ) );
 
 			// Connect the indices to the previous vertices
-			BatchIndices.Add( IndexStart - 2 );
-			BatchIndices.Add( IndexStart - 1 );
-			BatchIndices.Add( IndexStart + 0 );
+			RenderBatch.AddIndex(IndexStart - 2);
+			RenderBatch.AddIndex(IndexStart - 1);
+			RenderBatch.AddIndex(IndexStart + 0);
 
-			BatchIndices.Add( IndexStart + 0 );
-			BatchIndices.Add( IndexStart - 1 );
-			BatchIndices.Add( IndexStart + 1 );
+			RenderBatch.AddIndex(IndexStart + 0);
+			RenderBatch.AddIndex(IndexStart - 1);
+			RenderBatch.AddIndex(IndexStart + 1);
 		}
 	}
 }
@@ -1261,8 +1644,7 @@ struct FLineBuilder
 	//                                 o is CapOrigin
 
 	static void MakeCap(
-		FSlateVertexArray& OutBatchVertices,
-		FSlateIndexArray& OutBatchIndices,
+		FSlateRenderBatch& RenderBatch,
 		const FSlateRenderTransform& RenderTransform,
 		const FVector2D& CapOrigin,
 		const FVector2D& CapDirection,
@@ -1270,30 +1652,29 @@ struct FLineBuilder
 		const FColor& Color
 	)
 	{
-		const int FirstVertIndex = OutBatchVertices.Num();
+		const uint32 FirstVertIndex = RenderBatch.GetNumVertices();
 
-		OutBatchVertices.Add(FSlateVertex::Make<ESlateVertexRounding::Disabled>(RenderTransform, CapOrigin, FVector2D(0.5, 0.0f), FVector2D::ZeroVector, Color));
-		OutBatchVertices.Add(FSlateVertex::Make<ESlateVertexRounding::Disabled>(RenderTransform, CapOrigin + CapDirection + Up, FVector2D(0.0, 0.0f), FVector2D::ZeroVector, Color));
-		OutBatchVertices.Add(FSlateVertex::Make<ESlateVertexRounding::Disabled>(RenderTransform, CapOrigin + CapDirection - Up, FVector2D(0.0, 0.0f), FVector2D::ZeroVector, Color));
-		OutBatchVertices.Add(FSlateVertex::Make<ESlateVertexRounding::Disabled>(RenderTransform, CapOrigin + Up, FVector2D(0.0, 0.0f), FVector2D::ZeroVector, Color));
-		OutBatchVertices.Add(FSlateVertex::Make<ESlateVertexRounding::Disabled>(RenderTransform, CapOrigin - Up, FVector2D(0.0, 0.0f), FVector2D::ZeroVector, Color));
+		RenderBatch.AddVertex(FSlateVertex::Make<ESlateVertexRounding::Disabled>(RenderTransform, CapOrigin, FVector2D(0.5, 0.0f), FVector2D::ZeroVector, Color));
+		RenderBatch.AddVertex(FSlateVertex::Make<ESlateVertexRounding::Disabled>(RenderTransform, CapOrigin + CapDirection + Up, FVector2D(0.0, 0.0f), FVector2D::ZeroVector, Color));
+		RenderBatch.AddVertex(FSlateVertex::Make<ESlateVertexRounding::Disabled>(RenderTransform, CapOrigin + CapDirection - Up, FVector2D(0.0, 0.0f), FVector2D::ZeroVector, Color));
+		RenderBatch.AddVertex(FSlateVertex::Make<ESlateVertexRounding::Disabled>(RenderTransform, CapOrigin + Up, FVector2D(0.0, 0.0f), FVector2D::ZeroVector, Color));
+		RenderBatch.AddVertex(FSlateVertex::Make<ESlateVertexRounding::Disabled>(RenderTransform, CapOrigin - Up, FVector2D(0.0, 0.0f), FVector2D::ZeroVector, Color));
 
-		OutBatchIndices.Add(FirstVertIndex + 0);
-		OutBatchIndices.Add(FirstVertIndex + 3);
-		OutBatchIndices.Add(FirstVertIndex + 1);
+		RenderBatch.AddIndex(FirstVertIndex + 0);
+		RenderBatch.AddIndex(FirstVertIndex + 3);
+		RenderBatch.AddIndex(FirstVertIndex + 1);
 
-		OutBatchIndices.Add(FirstVertIndex + 0);
-		OutBatchIndices.Add(FirstVertIndex + 1);
-		OutBatchIndices.Add(FirstVertIndex + 2);
+		RenderBatch.AddIndex(FirstVertIndex + 0);
+		RenderBatch.AddIndex(FirstVertIndex + 1);
+		RenderBatch.AddIndex(FirstVertIndex + 2);
 
-		OutBatchIndices.Add(FirstVertIndex + 0);
-		OutBatchIndices.Add(FirstVertIndex + 2);
-		OutBatchIndices.Add(FirstVertIndex + 4);
+		RenderBatch.AddIndex(FirstVertIndex + 0);
+		RenderBatch.AddIndex(FirstVertIndex + 2);
+		RenderBatch.AddIndex(FirstVertIndex + 4);
 	}
 
-	FLineBuilder(FSlateVertexArray& InBatchVertices, FSlateIndexArray& InBatchIndices, const FVector2D StartPoint, float HalfThickness, const FSlateRenderTransform& InRenderTransform, const FColor& InColor)
-		: BatchVertices(InBatchVertices)
-		, BatchIndices(InBatchIndices)
+	FLineBuilder(FSlateRenderBatch& InRenderBatch, const FVector2D StartPoint, float HalfThickness, const FSlateRenderTransform& InRenderTransform, const FColor& InColor)
+		: RenderBatch(InRenderBatch)
 		, RenderTransform(InRenderTransform)
 		, LastPointAdded()
 		, LastNormal(FVector2D::ZeroVector)
@@ -1347,27 +1728,27 @@ private:
 			// Once we have two points, we have a normal, so we can generate the first bit of geometry.
 			const FVector2D LastUp = LastNormal*HalfLineThickness;
 
-			BatchVertices.Add(FSlateVertex::Make<ESlateVertexRounding::Disabled>(RenderTransform, LastPointAdded[1] + LastUp, FVector2D(1.0f, 0.0f), FVector2D::ZeroVector, InColor));
-			BatchVertices.Add(FSlateVertex::Make<ESlateVertexRounding::Disabled>(RenderTransform, LastPointAdded[1] - LastUp, FVector2D(0.0f, 0.0f), FVector2D::ZeroVector, InColor));
+			RenderBatch.AddVertex(FSlateVertex::Make<ESlateVertexRounding::Disabled>(RenderTransform, LastPointAdded[1] + LastUp, FVector2D(1.0f, 0.0f), FVector2D::ZeroVector, InColor));
+			RenderBatch.AddVertex(FSlateVertex::Make<ESlateVertexRounding::Disabled>(RenderTransform, LastPointAdded[1] - LastUp, FVector2D(0.0f, 0.0f), FVector2D::ZeroVector, InColor));
 		}
 
 		if (NumPointsAdded >= 2)
 		{
 			const FVector2D AveragedUp = (0.5f*(NewNormal + LastNormal)).GetSafeNormal()*HalfLineThickness;
 
-			BatchVertices.Add(FSlateVertex::Make<ESlateVertexRounding::Disabled>(RenderTransform, LastPointAdded[0] + AveragedUp, FVector2D(1.0f, 0.0f), FVector2D::ZeroVector, InColor));
-			BatchVertices.Add(FSlateVertex::Make<ESlateVertexRounding::Disabled>(RenderTransform, LastPointAdded[0] - AveragedUp, FVector2D(0.0f, 0.0f), FVector2D::ZeroVector, InColor));
+			RenderBatch.AddVertex(FSlateVertex::Make<ESlateVertexRounding::Disabled>(RenderTransform, LastPointAdded[0] + AveragedUp, FVector2D(1.0f, 0.0f), FVector2D::ZeroVector, InColor));
+			RenderBatch.AddVertex(FSlateVertex::Make<ESlateVertexRounding::Disabled>(RenderTransform, LastPointAdded[0] - AveragedUp, FVector2D(0.0f, 0.0f), FVector2D::ZeroVector, InColor));
 
-			const int32 NumVerts = BatchVertices.Num();
+			const int32 NumVerts = RenderBatch.GetNumVertices();
 
 			// Counterclockwise winding on triangles
-			BatchIndices.Add(NumVerts - 3);
-			BatchIndices.Add(NumVerts - 4);
-			BatchIndices.Add(NumVerts - 2);
+			RenderBatch.AddIndex(NumVerts - 3);
+			RenderBatch.AddIndex(NumVerts - 4);
+			RenderBatch.AddIndex(NumVerts - 2);
 
-			BatchIndices.Add(NumVerts - 3);
-			BatchIndices.Add(NumVerts - 2);
-			BatchIndices.Add(NumVerts - 1);
+			RenderBatch.AddIndex(NumVerts - 3);
+			RenderBatch.AddIndex(NumVerts - 2);
+			RenderBatch.AddIndex(NumVerts - 1);
 		}
 
 		LastPointAdded[1] = LastPointAdded[0];
@@ -1393,19 +1774,19 @@ private:
 			// last line segment.
 			const FVector2D LastUp = LastNormal*HalfLineThickness;
 
-			BatchVertices.Add(FSlateVertex::Make<ESlateVertexRounding::Disabled>(RenderTransform, LastPointAdded[0] + LastUp, FVector2D(1.0f, 0.0f), FVector2D::ZeroVector, InColor));
-			BatchVertices.Add(FSlateVertex::Make<ESlateVertexRounding::Disabled>(RenderTransform, LastPointAdded[0] - LastUp, FVector2D(0.0f, 0.0f), FVector2D::ZeroVector, InColor));
+			RenderBatch.AddVertex(FSlateVertex::Make<ESlateVertexRounding::Disabled>(RenderTransform, LastPointAdded[0] + LastUp, FVector2D(1.0f, 0.0f), FVector2D::ZeroVector, InColor));
+			RenderBatch.AddVertex(FSlateVertex::Make<ESlateVertexRounding::Disabled>(RenderTransform, LastPointAdded[0] - LastUp, FVector2D(0.0f, 0.0f), FVector2D::ZeroVector, InColor));
 
-			const int32 NumVerts = BatchVertices.Num();
+			const int32 NumVerts = RenderBatch.GetNumVertices();
 
 			// Counterclockwise winding on triangles
-			BatchIndices.Add(NumVerts - 3);
-			BatchIndices.Add(NumVerts - 4);
-			BatchIndices.Add(NumVerts - 2);
+			RenderBatch.AddIndex(NumVerts - 3);
+			RenderBatch.AddIndex(NumVerts - 4);
+			RenderBatch.AddIndex(NumVerts - 2);
 
-			BatchIndices.Add(NumVerts - 3);
-			BatchIndices.Add(NumVerts - 2);
-			BatchIndices.Add(NumVerts - 1);
+			RenderBatch.AddIndex(NumVerts - 3);
+			RenderBatch.AddIndex(NumVerts - 2);
+			RenderBatch.AddIndex(NumVerts - 1);
 		}
 	}
 
@@ -1536,8 +1917,7 @@ private:
 	}
 	
 private:
-	FSlateVertexArray& BatchVertices;
-	FSlateIndexArray& BatchIndices;
+	FSlateRenderBatch& RenderBatch;
 	const FSlateRenderTransform& RenderTransform;
 	FVector2D LastPointAdded[2];
 	FVector2D LastNormal;
@@ -1559,9 +1939,9 @@ void FSlateElementBatcher::AddSplineElement(const FSlateDrawElement& DrawElement
 	// a floating point line width.
 
 	const FSlateRenderTransform& RenderTransform = DrawElement.GetRenderTransform();
-	const FSlateDataPayload& InPayload = DrawElement.GetDataPayload();
+	const FSlateSplinePayload& InPayload = DrawElement.GetDataPayload<FSlateSplinePayload>();
 	ESlateDrawEffect InDrawEffects = DrawElement.GetDrawEffects();
-	uint32 Layer = DrawElement.GetLayer();
+	const int32 Layer = DrawElement.GetLayer();
 
 	// 1 is the minimum thickness we support for generating geometry.
 	// The shader takes care of sub-pixel line widths.
@@ -1583,11 +1963,10 @@ void FSlateElementBatcher::AddSplineElement(const FSlateDrawElement& DrawElement
 
 	const FColor SplineColor = (InPayload.GradientStops.Num()==1) ? PackVertexColor(InPayload.GradientStops[0].Color) : PackVertexColor(InPayload.GetTint());
 
-	FSlateElementBatch& ElementBatch = FindBatchForElement(Layer, FShaderParams::MakePixelShaderParams(FVector4(InPayload.GetThickness(), FilterScale, 0, 0)), nullptr, ESlateDrawPrimitive::TriangleList, ESlateShader::LineSegment, InDrawEffects, ESlateBatchDrawFlag::None, DrawElement.GetClippingIndex(), DrawElement.GetSceneIndex());
+	FSlateRenderBatch& RenderBatch = CreateRenderBatch(Layer, FShaderParams::MakePixelShaderParams(FVector4(InPayload.GetThickness(), FilterScale, 0, 0)), nullptr, ESlateDrawPrimitive::TriangleList, ESlateShader::LineSegment, InDrawEffects, ESlateBatchDrawFlag::None, DrawElement);
 
 	FLineBuilder LineBuilder(
-		BatchData->GetBatchVertexList(ElementBatch),
-		BatchData->GetBatchIndexList(ElementBatch),
+		RenderBatch,
 		InPayload.P0,
 		HalfThickness,
 		RenderTransform,
@@ -1604,25 +1983,26 @@ void FSlateElementBatcher::AddSplineElement(const FSlateDrawElement& DrawElement
 		// Deprecated scenario _WithColorGradient
 		LineBuilder.BuildBezierGeometry_WithColorGradient( InPayload.GradientStops, 1, InPayload.P0, InPayload.P1, InPayload.P2, InPayload.P3, *this);
 	}
-	
 }
 
 
 template<ESlateVertexRounding Rounding>
 void FSlateElementBatcher::AddLineElement( const FSlateDrawElement& DrawElement )
 {
-	const FSlateDataPayload& DrawElementPayload = DrawElement.GetDataPayload();
+	const FSlateLinePayload& DrawElementPayload = DrawElement.GetDataPayload<FSlateLinePayload>();
 	const FSlateRenderTransform& RenderTransform = DrawElement.GetRenderTransform();
 	ESlateDrawEffect DrawEffects = DrawElement.GetDrawEffects();
-	uint32 Layer = DrawElement.GetLayer();
+	const int32 Layer = DrawElement.GetLayer();
+
+	const TArray<FVector2D>& Points = DrawElementPayload.GetPoints();
+	const TArray<FLinearColor>& PointColors = DrawElementPayload.GetPointColors();
 
-	const uint16 NumPoints = DrawElementPayload.GetNumPoints();
-	if (NumPoints < 2 || !DrawElementPayload.GetPoints())
+	const int32 NumPoints = Points.Num();
+	if (NumPoints < 2)
 	{
 		return;
 	}
 	
-	const FVector2D* Points = DrawElementPayload.GetPoints();
 	const FColor FinalTint = PackVertexColor(DrawElementPayload.GetTint());
 
 	if(DrawElementPayload.IsAntialiased() )
@@ -1673,9 +2053,7 @@ void FSlateElementBatcher::AddLineElement( const FSlateDrawElement& DrawElement
 		const float HalfThickness = LineThickness * .5f + FilterScale;
 
 		// Find a batch for the element
-		FSlateElementBatch& ElementBatch = FindBatchForElement( Layer, FShaderParams::MakePixelShaderParams( FVector4(RequestedThickness, FilterScale,0,0) ), nullptr, ESlateDrawPrimitive::TriangleList, ESlateShader::LineSegment, DrawEffects, ESlateBatchDrawFlag::None, DrawElement.GetClippingIndex(), DrawElement.GetSceneIndex());
-		FSlateVertexArray& BatchVertices = BatchData->GetBatchVertexList(ElementBatch);
-		FSlateIndexArray& BatchIndices = BatchData->GetBatchIndexList(ElementBatch);
+		FSlateRenderBatch& RenderBatch = CreateRenderBatch( Layer, FShaderParams::MakePixelShaderParams( FVector4(RequestedThickness, FilterScale,0,0) ), nullptr, ESlateDrawPrimitive::TriangleList, ESlateShader::LineSegment, DrawEffects, ESlateBatchDrawFlag::None, DrawElement);
 
 		FVector2D StartPos = Points[0];
 		FVector2D EndPos = Points[1];
@@ -1683,16 +2061,16 @@ void FSlateElementBatcher::AddLineElement( const FSlateDrawElement& DrawElement
 		FVector2D Normal = FVector2D( StartPos.Y - EndPos.Y, EndPos.X - StartPos.X ).GetSafeNormal();
 		FVector2D Up = Normal * HalfThickness;
 
-		FColor StartColor = DrawElementPayload.GetPointColors() ? PackVertexColor(DrawElementPayload.GetPointColors()[0] * DrawElementPayload.GetTint()) : FinalTint;
-		FColor EndColor = DrawElementPayload.GetPointColors() ? PackVertexColor(DrawElementPayload.GetPointColors()[1] * DrawElementPayload.GetTint()) : FinalTint;
+		FColor StartColor = PointColors.Num() ? PackVertexColor(PointColors[0] * DrawElementPayload.GetTint()) : FinalTint;
+		FColor EndColor = PointColors.Num() ? PackVertexColor(PointColors[1] * DrawElementPayload.GetTint()) : FinalTint;
 	
 		const FVector2D StartCapDirection = HalfThickness*((StartPos - EndPos).GetSafeNormal());
-		FLineBuilder::MakeCap(BatchVertices, BatchIndices, RenderTransform, StartPos, StartCapDirection, Up, StartColor);
-		const uint32 IndexStart = BatchVertices.Num();
+		FLineBuilder::MakeCap(RenderBatch, RenderTransform, StartPos, StartCapDirection, Up, StartColor);
+		const uint32 IndexStart = RenderBatch.GetNumVertices();
 
 		// First two points in the line.
-		BatchVertices.Add(FSlateVertex::Make<Rounding>( RenderTransform, StartPos + Up, FVector2D(1.0, 0.0f), FVector2D::ZeroVector, StartColor ) );
-		BatchVertices.Add(FSlateVertex::Make<Rounding>( RenderTransform, StartPos - Up, FVector2D(0.0, 0.0f), FVector2D::ZeroVector, StartColor) );
+		RenderBatch.AddVertex(FSlateVertex::Make<Rounding>( RenderTransform, StartPos + Up, FVector2D(1.0, 0.0f), FVector2D::ZeroVector, StartColor ) );
+		RenderBatch.AddVertex(FSlateVertex::Make<Rounding>( RenderTransform, StartPos - Up, FVector2D(0.0, 0.0f), FVector2D::ZeroVector, StartColor) );
 
 		// Generate the rest of the segments
 		for( int32 Point = 1; Point < NumPoints; ++Point )
@@ -1708,7 +2086,7 @@ void FSlateElementBatcher::AddLineElement( const FSlateDrawElement& DrawElement
 			// Create the new vertices for the thick line segment
 			Up = Normal * HalfThickness;
 
-			FColor PointColor = DrawElementPayload.GetPointColors() ? PackVertexColor(DrawElementPayload.GetPointColors()[Point] * DrawElementPayload.GetTint()) : FinalTint;
+			FColor PointColor = PointColors.Num() ? PackVertexColor(PointColors[Point] * DrawElementPayload.GetTint()) : FinalTint;
 
 			FVector2D IntersectUpper = EndPos + Up;
 			FVector2D IntersectLower = EndPos - Up;
@@ -1744,17 +2122,17 @@ void FSlateElementBatcher::AddLineElement( const FSlateDrawElement& DrawElement
 				}
 			}
 
-			BatchVertices.Add(FSlateVertex::Make<Rounding>( RenderTransform, IntersectUpper, FVector2D(1.0, 0.0f), FVector2D::ZeroVector, PointColor ) );
-			BatchVertices.Add(FSlateVertex::Make<Rounding>( RenderTransform, IntersectLower, FVector2D(0.0, 0.0f), FVector2D::ZeroVector, PointColor ) );
+			RenderBatch.AddVertex(FSlateVertex::Make<Rounding>( RenderTransform, IntersectUpper, FVector2D(1.0, 0.0f), FVector2D::ZeroVector, PointColor ) );
+			RenderBatch.AddVertex(FSlateVertex::Make<Rounding>( RenderTransform, IntersectLower, FVector2D(0.0, 0.0f), FVector2D::ZeroVector, PointColor ) );
 			
 			// Counterclockwise winding on triangles
-			BatchIndices.Add(IndexStart + 2 * Point - 1);
-			BatchIndices.Add(IndexStart + 2 * Point - 2);
-			BatchIndices.Add(IndexStart + 2 * Point + 0);
+			RenderBatch.AddIndex(IndexStart + 2 * Point - 1);
+			RenderBatch.AddIndex(IndexStart + 2 * Point - 2);
+			RenderBatch.AddIndex(IndexStart + 2 * Point + 0);
 
-			BatchIndices.Add(IndexStart + 2 * Point - 1);
-			BatchIndices.Add(IndexStart + 2 * Point + 0);
-			BatchIndices.Add(IndexStart + 2 * Point + 1);
+			RenderBatch.AddIndex(IndexStart + 2 * Point - 1);
+			RenderBatch.AddIndex(IndexStart + 2 * Point + 0);
+			RenderBatch.AddIndex(IndexStart + 2 * Point + 1);
 
 			StartPos = EndPos;
 		}
@@ -1762,66 +2140,63 @@ void FSlateElementBatcher::AddLineElement( const FSlateDrawElement& DrawElement
 		EndPos = Points[NumPoints - 1];
 		StartPos = Points[NumPoints - 2];
 		const FVector2D EndCapDirection = HalfThickness*((EndPos-StartPos).GetSafeNormal());
-		FLineBuilder::MakeCap(BatchVertices, BatchIndices, RenderTransform, EndPos, EndCapDirection, Up, EndColor);
+		FLineBuilder::MakeCap(RenderBatch, RenderTransform, EndPos, EndCapDirection, Up, EndColor);
 	}
 	else
 	{
 		if (DrawElementPayload.GetThickness() == 1)
 		{
 			// Find a batch for the element
-			FSlateElementBatch& ElementBatch = FindBatchForElement(Layer, FShaderParams(), nullptr, ESlateDrawPrimitive::LineList, ESlateShader::Default, DrawEffects, ESlateBatchDrawFlag::None, DrawElement.GetClippingIndex(), DrawElement.GetSceneIndex());
-			FSlateVertexArray& BatchVertices = BatchData->GetBatchVertexList(ElementBatch);
-			FSlateIndexArray& BatchIndices = BatchData->GetBatchIndexList(ElementBatch);
+			FSlateRenderBatch& RenderBatch = CreateRenderBatch(Layer, FShaderParams(), nullptr, ESlateDrawPrimitive::LineList, ESlateShader::Default, DrawEffects, ESlateBatchDrawFlag::None, DrawElement);
 
 			// Generate the line segments using the native line rendering of the platform.
-			for (int32 Point = 0; Point < DrawElementPayload.GetNumPoints() - 1; ++Point)
+			for (int32 Point = 0; Point < NumPoints - 1; ++Point)
 			{
-				uint32 IndexStart = BatchVertices.Num();
+				const uint32 IndexStart = RenderBatch.GetNumVertices();
 				FVector2D StartPos = Points[Point];
 				FVector2D EndPos = Points[Point + 1];
 
-				FColor StartColor = DrawElementPayload.GetPointColors() ? PackVertexColor(DrawElementPayload.GetPointColors()[Point]   * DrawElementPayload.GetTint()) : FinalTint;
-				FColor EndColor = DrawElementPayload.GetPointColors() ? PackVertexColor(DrawElementPayload.GetPointColors()[Point+1] * DrawElementPayload.GetTint()) : FinalTint;
+				FColor StartColor = PointColors.Num() ? PackVertexColor(PointColors[Point]  * DrawElementPayload.GetTint()) : FinalTint;
+				FColor EndColor = PointColors.Num() ? PackVertexColor(PointColors[Point+1] * DrawElementPayload.GetTint()) : FinalTint;
 
-				BatchVertices.Add(FSlateVertex::Make<Rounding>(RenderTransform, StartPos, FVector2D::ZeroVector, StartColor));
-				BatchVertices.Add(FSlateVertex::Make<Rounding>(RenderTransform, EndPos, FVector2D::ZeroVector, EndColor));
+				RenderBatch.AddVertex(FSlateVertex::Make<Rounding>(RenderTransform, StartPos, FVector2D::ZeroVector, StartColor));
+				RenderBatch.AddVertex(FSlateVertex::Make<Rounding>(RenderTransform, EndPos, FVector2D::ZeroVector, EndColor));
 
-				BatchIndices.Add(IndexStart);
-				BatchIndices.Add(IndexStart + 1);
+				RenderBatch.AddIndex(IndexStart);
+				RenderBatch.AddIndex(IndexStart + 1);
 			}
 		}
 		else
 		{
 			// Find a batch for the element
-			FSlateElementBatch& ElementBatch = FindBatchForElement(Layer, FShaderParams(), nullptr, ESlateDrawPrimitive::TriangleList, ESlateShader::Default, DrawEffects, ESlateBatchDrawFlag::None, DrawElement.GetClippingIndex(), DrawElement.GetSceneIndex());
-			FSlateVertexArray& BatchVertices = BatchData->GetBatchVertexList(ElementBatch);
-			FSlateIndexArray& BatchIndices = BatchData->GetBatchIndexList(ElementBatch);
+			FSlateRenderBatch& RenderBatch = CreateRenderBatch(Layer, FShaderParams(), nullptr, ESlateDrawPrimitive::TriangleList, ESlateShader::Default, DrawEffects, ESlateBatchDrawFlag::None, DrawElement);
+
 
 			// Generate the line segments using non-aa'ed polylines.
-			for (int32 Point = 0; Point < DrawElementPayload.GetNumPoints() - 1; ++Point)
+			for (int32 Point = 0; Point < NumPoints - 1; ++Point)
 			{
-				uint32 IndexStart = BatchVertices.Num();
+				const uint32 IndexStart = RenderBatch.GetNumVertices();
 				const FVector2D StartPos = Points[Point];
 				const FVector2D EndPos = Points[Point + 1];
 
-				FColor StartColor	= DrawElementPayload.GetPointColors() ? PackVertexColor(DrawElementPayload.GetPointColors()[Point]   * DrawElementPayload.GetTint()) : FinalTint;
-				FColor EndColor		= DrawElementPayload.GetPointColors() ? PackVertexColor(DrawElementPayload.GetPointColors()[Point+1] * DrawElementPayload.GetTint()) : FinalTint;
+				FColor StartColor	= PointColors.Num() ? PackVertexColor(PointColors[Point]   * DrawElementPayload.GetTint()) : FinalTint;
+				FColor EndColor		= PointColors.Num() ? PackVertexColor(PointColors[Point+1] * DrawElementPayload.GetTint()) : FinalTint;
 	
 				const FVector2D SegmentNormal = (EndPos - StartPos).GetSafeNormal();
 				const FVector2D HalfThickNormal = SegmentNormal * (DrawElementPayload.GetThickness() * 0.5f);
 
-				BatchVertices.Add(FSlateVertex::Make<Rounding>(RenderTransform, StartPos + FVector2D(HalfThickNormal.Y, -HalfThickNormal.X), FVector2D::ZeroVector, FVector2D::ZeroVector, StartColor));
-				BatchVertices.Add(FSlateVertex::Make<Rounding>(RenderTransform, StartPos + FVector2D(-HalfThickNormal.Y, HalfThickNormal.X), FVector2D::ZeroVector, FVector2D::ZeroVector, StartColor));
-				BatchVertices.Add(FSlateVertex::Make<Rounding>(RenderTransform, EndPos + FVector2D(HalfThickNormal.Y, -HalfThickNormal.X), FVector2D::ZeroVector, FVector2D::ZeroVector, EndColor));
-				BatchVertices.Add(FSlateVertex::Make<Rounding>(RenderTransform, EndPos + FVector2D(-HalfThickNormal.Y, HalfThickNormal.X), FVector2D::ZeroVector, FVector2D::ZeroVector, EndColor));
+				RenderBatch.AddVertex(FSlateVertex::Make<Rounding>(RenderTransform, StartPos + FVector2D(HalfThickNormal.Y, -HalfThickNormal.X), FVector2D::ZeroVector, FVector2D::ZeroVector, StartColor));
+				RenderBatch.AddVertex(FSlateVertex::Make<Rounding>(RenderTransform, StartPos + FVector2D(-HalfThickNormal.Y, HalfThickNormal.X), FVector2D::ZeroVector, FVector2D::ZeroVector, StartColor));
+				RenderBatch.AddVertex(FSlateVertex::Make<Rounding>(RenderTransform, EndPos + FVector2D(HalfThickNormal.Y, -HalfThickNormal.X), FVector2D::ZeroVector, FVector2D::ZeroVector, EndColor));
+				RenderBatch.AddVertex(FSlateVertex::Make<Rounding>(RenderTransform, EndPos + FVector2D(-HalfThickNormal.Y, HalfThickNormal.X), FVector2D::ZeroVector, FVector2D::ZeroVector, EndColor));
 
-				BatchIndices.Add(IndexStart + 0);
-				BatchIndices.Add(IndexStart + 1);
-				BatchIndices.Add(IndexStart + 2);
+				RenderBatch.AddIndex(IndexStart + 0);
+				RenderBatch.AddIndex(IndexStart + 1);
+				RenderBatch.AddIndex(IndexStart + 2);
 
-				BatchIndices.Add(IndexStart + 2);
-				BatchIndices.Add(IndexStart + 1);
-				BatchIndices.Add(IndexStart + 3);
+				RenderBatch.AddIndex(IndexStart + 2);
+				RenderBatch.AddIndex(IndexStart + 1);
+				RenderBatch.AddIndex(IndexStart + 3);
 			}
 		}
 	}
@@ -1832,16 +2207,16 @@ void FSlateElementBatcher::AddViewportElement( const FSlateDrawElement& DrawElem
 {
 	const FSlateRenderTransform& RenderTransform = DrawElement.GetRenderTransform();
 	const FVector2D& LocalSize = DrawElement.GetLocalSize();
-	const FSlateDataPayload& DrawElementPayload = DrawElement.GetDataPayload();
+	const FSlateViewportPayload& DrawElementPayload = DrawElement.GetDataPayload<FSlateViewportPayload>();
 	ESlateDrawEffect InDrawEffects = DrawElement.GetDrawEffects();
-	uint32 Layer = DrawElement.GetLayer();
+	const int32 Layer = DrawElement.GetLayer();
 
 	const FColor FinalColor = PackVertexColor(DrawElementPayload.GetTint());
 
-	ESlateBatchDrawFlag DrawFlags = DrawElementPayload.BatchFlags;
+	ESlateBatchDrawFlag DrawFlags = DrawElement.GetBatchFlags();
 
 	FSlateShaderResource* ViewportResource = DrawElementPayload.RenderTargetResource;
-	ESlateShader::Type ShaderType = ESlateShader::Default;
+	ESlateShader ShaderType = ESlateShader::Default;
 
 	if(DrawElementPayload.bViewportTextureAlphaOnly )
 	{
@@ -1849,9 +2224,7 @@ void FSlateElementBatcher::AddViewportElement( const FSlateDrawElement& DrawElem
 		ShaderType = ESlateShader::Font;
 	}
 
-	FSlateElementBatch& ElementBatch = FindBatchForElement( Layer, FShaderParams(), ViewportResource, ESlateDrawPrimitive::TriangleList, ShaderType, InDrawEffects, DrawFlags, DrawElement.GetClippingIndex(), DrawElement.GetSceneIndex());
-	FSlateVertexArray& BatchVertices = BatchData->GetBatchVertexList(ElementBatch);
-	FSlateIndexArray& BatchIndices = BatchData->GetBatchIndexList(ElementBatch);
+	FSlateRenderBatch& RenderBatch = CreateRenderBatch( Layer, FShaderParams(), ViewportResource, ESlateDrawPrimitive::TriangleList, ShaderType, InDrawEffects, DrawFlags, DrawElement);
 
 	// Tag this batch as requiring vsync if the viewport requires it.
 	if( ViewportResource != nullptr && !DrawElementPayload.bAllowViewportScaling )
@@ -1874,35 +2247,33 @@ void FSlateElementBatcher::AddViewportElement( const FSlateDrawElement& DrawElem
 	FVector2D BotLeft =	 FVector2D( TopLeft.X, BotRight.Y);
 
 	// The start index of these vertices in the index buffer
-	uint32 IndexStart = BatchVertices.Num();
+	const uint32 IndexStart = 0;
 
 	// Add four vertices to the list of verts to be added to the vertex buffer
-	BatchVertices.Add( FSlateVertex::Make<Rounding>( RenderTransform, TopLeft,	FVector2D(0.0f,0.0f),	FinalColor ) );
-	BatchVertices.Add( FSlateVertex::Make<Rounding>( RenderTransform, TopRight,	FVector2D(1.0f,0.0f),	FinalColor ) );
-	BatchVertices.Add( FSlateVertex::Make<Rounding>( RenderTransform, BotLeft,	FVector2D(0.0f,1.0f),	FinalColor ) );
-	BatchVertices.Add( FSlateVertex::Make<Rounding>( RenderTransform, BotRight,	FVector2D(1.0f,1.0f),	FinalColor ) );
-
-	// The offset into the index buffer where this quads indices start
-	uint32 IndexOffsetStart = BatchIndices.Num();
+	RenderBatch.AddVertex( FSlateVertex::Make<Rounding>( RenderTransform, TopLeft,	FVector2D(0.0f,0.0f),	FinalColor ) );
+	RenderBatch.AddVertex( FSlateVertex::Make<Rounding>( RenderTransform, TopRight,	FVector2D(1.0f,0.0f),	FinalColor ) );
+	RenderBatch.AddVertex( FSlateVertex::Make<Rounding>( RenderTransform, BotLeft,	FVector2D(0.0f,1.0f),	FinalColor ) );
+	RenderBatch.AddVertex( FSlateVertex::Make<Rounding>( RenderTransform, BotRight,	FVector2D(1.0f,1.0f),	FinalColor ) );
 
 	// Add 6 indices to the vertex buffer.  (2 tri's per quad, 3 indices per tri)
-	BatchIndices.Add( IndexStart + 0 );
-	BatchIndices.Add( IndexStart + 1 );
-	BatchIndices.Add( IndexStart + 2 );
+	RenderBatch.AddIndex( IndexStart + 0 );
+	RenderBatch.AddIndex( IndexStart + 1 );
+	RenderBatch.AddIndex( IndexStart + 2 );
 
-	BatchIndices.Add( IndexStart + 2 );
-	BatchIndices.Add( IndexStart + 1 );
-	BatchIndices.Add( IndexStart + 3 );
+	RenderBatch.AddIndex( IndexStart + 2 );
+	RenderBatch.AddIndex( IndexStart + 1 );
+	RenderBatch.AddIndex( IndexStart + 3 );
 }
 
 template<ESlateVertexRounding Rounding>
 void FSlateElementBatcher::AddBorderElement( const FSlateDrawElement& DrawElement )
 {
-	const FSlateDataPayload& DrawElementPayload = DrawElement.GetDataPayload();
+	const FSlateBoxPayload& DrawElementPayload = DrawElement.GetDataPayload<FSlateBoxPayload>();
 	const FSlateRenderTransform& RenderTransform = DrawElement.GetRenderTransform();
 	const FVector2D& LocalSize = DrawElement.GetLocalSize();
-	ESlateDrawEffect InDrawEffects = DrawElement.GetDrawEffects();
-	uint32 Layer = DrawElement.GetLayer();
+	const ESlateDrawEffect InDrawEffects = DrawElement.GetDrawEffects();
+
+	const int32 Layer = DrawElement.GetLayer();
 
 	const float DrawScale = DrawElement.GetScale();
 
@@ -1986,9 +2357,7 @@ void FSlateElementBatcher::AddBorderElement( const FSlateDrawElement& DrawElemen
 	// Pass the tiling information as a flag so we can pick the correct texture addressing mode
 	ESlateBatchDrawFlag DrawFlags = (ESlateBatchDrawFlag::TileU|ESlateBatchDrawFlag::TileV);
 
-	FSlateElementBatch& ElementBatch = FindBatchForElement( Layer, ShaderParams, Resource, ESlateDrawPrimitive::TriangleList, ESlateShader::Border, InDrawEffects, DrawFlags, DrawElement.GetClippingIndex(), DrawElement.GetSceneIndex());
-	FSlateVertexArray& BatchVertices = BatchData->GetBatchVertexList(ElementBatch);
-	FSlateIndexArray& BatchIndices = BatchData->GetBatchIndexList(ElementBatch);
+	FSlateRenderBatch& RenderBatch = CreateRenderBatch( Layer, ShaderParams, Resource, ESlateDrawPrimitive::TriangleList, ESlateShader::Border, InDrawEffects, DrawFlags, DrawElement);
 
 	// Ensure tiling of at least 1.  
 	TopTiling = TopTiling >= 1.0f ? TopTiling : 1.0f;
@@ -2000,147 +2369,153 @@ void FSlateElementBatcher::AddBorderElement( const FSlateDrawElement& DrawElemen
 	FVector2D EndPos = BotRight;
 
 	// The start index of these vertices in the index buffer
-	uint32 IndexStart = BatchVertices.Num();
+	const uint32 IndexStart = RenderBatch.GetNumVertices();
 
 	// Zero for second UV indicates no tiling and to just pass the UV though (for the corner sections)
 	FVector2D Zero(0,0);
 
 	// Add all the vertices needed for this element.  Vertices are duplicated so that we can have some sections with no tiling and some with tiling.
-	BatchVertices.Add( FSlateVertex::Make<Rounding>( RenderTransform, Position,									LocalSize, DrawScale, FVector4( StartUV.X, StartUV.Y, 0.0f, 0.0f),				Tint ) ); //0
-	BatchVertices.Add( FSlateVertex::Make<Rounding>( RenderTransform, FVector2D( Position.X, TopMarginY ),		LocalSize, DrawScale, FVector4( StartUV.X, TopMarginV, 0.0f, 0.0f),				Tint ) ); //1
-	BatchVertices.Add( FSlateVertex::Make<Rounding>( RenderTransform, FVector2D( LeftMarginX, Position.Y ),		LocalSize, DrawScale, FVector4( LeftMarginU, StartUV.Y, 0.0f, 0.0f),			Tint ) ); //2
-	BatchVertices.Add( FSlateVertex::Make<Rounding>( RenderTransform, FVector2D( LeftMarginX, TopMarginY ),		LocalSize, DrawScale, FVector4( LeftMarginU, TopMarginV, 0.0f, 0.0f),			Tint ) ); //3
-
-	BatchVertices.Add( FSlateVertex::Make<Rounding>( RenderTransform, FVector2D( LeftMarginX, Position.Y ),		LocalSize, DrawScale, FVector4( StartUV.X, StartUV.Y, TopTiling, 0.0f),			Tint ) ); //4
-	BatchVertices.Add( FSlateVertex::Make<Rounding>( RenderTransform, FVector2D( LeftMarginX, TopMarginY ),		LocalSize, DrawScale, FVector4( StartUV.X, TopMarginV, TopTiling, 0.0f),		Tint ) ); //5
-	BatchVertices.Add( FSlateVertex::Make<Rounding>( RenderTransform, FVector2D( RightMarginX, Position.Y ),	LocalSize, DrawScale, FVector4( EndUV.X, StartUV.Y, TopTiling, 0.0f),			Tint ) ); //6
-	BatchVertices.Add( FSlateVertex::Make<Rounding>( RenderTransform, FVector2D( RightMarginX, TopMarginY ),	LocalSize, DrawScale, FVector4( EndUV.X, TopMarginV, TopTiling, 0.0f),			Tint ) ); //7
-
-	BatchVertices.Add( FSlateVertex::Make<Rounding>( RenderTransform, FVector2D( RightMarginX, Position.Y ),	LocalSize, DrawScale, FVector4( RightMarginU, StartUV.Y, 0.0f, 0.0f),			Tint ) ); //8
-	BatchVertices.Add( FSlateVertex::Make<Rounding>( RenderTransform, FVector2D( RightMarginX, TopMarginY ),	LocalSize, DrawScale, FVector4( RightMarginU, TopMarginV, 0.0f, 0.0f),			Tint ) ); //9
-	BatchVertices.Add( FSlateVertex::Make<Rounding>( RenderTransform, FVector2D( EndPos.X, Position.Y ),		LocalSize, DrawScale, FVector4( EndUV.X, StartUV.Y, 0.0f, 0.0f),				Tint ) ); //10
-	BatchVertices.Add( FSlateVertex::Make<Rounding>( RenderTransform, FVector2D( EndPos.X, TopMarginY ),		LocalSize, DrawScale, FVector4( EndUV.X, TopMarginV, 0.0f, 0.0f),				Tint ) ); //11
-
-	BatchVertices.Add( FSlateVertex::Make<Rounding>( RenderTransform, FVector2D( Position.X, TopMarginY ),		LocalSize, DrawScale, FVector4( StartUV.X, StartUV.Y, 0.0f, LeftTiling),		Tint ) ); //12
-	BatchVertices.Add( FSlateVertex::Make<Rounding>( RenderTransform, FVector2D( Position.X, BottomMarginY ),	LocalSize, DrawScale, FVector4( StartUV.X, EndUV.Y, 0.0f, LeftTiling),			Tint ) ); //13
-	BatchVertices.Add( FSlateVertex::Make<Rounding>( RenderTransform, FVector2D( LeftMarginX, TopMarginY ),		LocalSize, DrawScale, FVector4( LeftMarginU, StartUV.Y, 0.0f, LeftTiling),		Tint ) ); //14
-	BatchVertices.Add( FSlateVertex::Make<Rounding>( RenderTransform, FVector2D( LeftMarginX, BottomMarginY ),	LocalSize, DrawScale, FVector4( LeftMarginU, EndUV.Y, 0.0f, LeftTiling),		Tint ) ); //15
-
-	BatchVertices.Add( FSlateVertex::Make<Rounding>( RenderTransform, FVector2D( RightMarginX, TopMarginY ),	LocalSize, DrawScale, FVector4( RightMarginU, StartUV.Y, 0.0f, RightTiling),	Tint ) ); //16
-	BatchVertices.Add( FSlateVertex::Make<Rounding>( RenderTransform, FVector2D( RightMarginX, BottomMarginY ), LocalSize, DrawScale, FVector4( RightMarginU, EndUV.Y, 0.0f, RightTiling),		Tint ) ); //17
-	BatchVertices.Add( FSlateVertex::Make<Rounding>( RenderTransform, FVector2D( EndPos.X, TopMarginY ),		LocalSize, DrawScale, FVector4( EndUV.X, StartUV.Y, 0.0f, RightTiling),			Tint ) ); //18
-	BatchVertices.Add( FSlateVertex::Make<Rounding>( RenderTransform, FVector2D( EndPos.X, BottomMarginY ),		LocalSize, DrawScale, FVector4( EndUV.X, EndUV.Y, 0.0f, RightTiling),			Tint ) ); //19
-
-	BatchVertices.Add( FSlateVertex::Make<Rounding>( RenderTransform, FVector2D( Position.X, BottomMarginY ),	LocalSize, DrawScale, FVector4( StartUV.X, BottomMarginV, 0.0f, 0.0f),			Tint ) ); //20
-	BatchVertices.Add( FSlateVertex::Make<Rounding>( RenderTransform, FVector2D( Position.X, EndPos.Y ),		LocalSize, DrawScale, FVector4( StartUV.X, EndUV.Y, 0.0f, 0.0f),				Tint ) ); //21
-	BatchVertices.Add( FSlateVertex::Make<Rounding>( RenderTransform, FVector2D( LeftMarginX, BottomMarginY ),	LocalSize, DrawScale, FVector4( LeftMarginU, BottomMarginV, 0.0f, 0.0f),		Tint ) ); //22
-	BatchVertices.Add( FSlateVertex::Make<Rounding>( RenderTransform, FVector2D( LeftMarginX, EndPos.Y ),		LocalSize, DrawScale, FVector4( LeftMarginU, EndUV.Y, 0.0f, 0.0f),				Tint ) ); //23
-
-	BatchVertices.Add( FSlateVertex::Make<Rounding>( RenderTransform, FVector2D( LeftMarginX, BottomMarginY ),	LocalSize, DrawScale, FVector4( StartUV.X, BottomMarginV, BottomTiling, 0.0f),	Tint ) ); //24
-	BatchVertices.Add( FSlateVertex::Make<Rounding>( RenderTransform, FVector2D( LeftMarginX, EndPos.Y ),		LocalSize, DrawScale, FVector4( StartUV.X, EndUV.Y, BottomTiling, 0.0f),		Tint ) ); //25
-	BatchVertices.Add( FSlateVertex::Make<Rounding>( RenderTransform, FVector2D( RightMarginX,BottomMarginY ),	LocalSize, DrawScale, FVector4( EndUV.X, BottomMarginV, BottomTiling, 0.0f),	Tint ) ); //26
-	BatchVertices.Add( FSlateVertex::Make<Rounding>( RenderTransform, FVector2D( RightMarginX, EndPos.Y ),		LocalSize, DrawScale, FVector4( EndUV.X, EndUV.Y, BottomTiling, 0.0f),			Tint ) ); //27
-
-	BatchVertices.Add( FSlateVertex::Make<Rounding>( RenderTransform, FVector2D( RightMarginX, BottomMarginY ), LocalSize, DrawScale, FVector4( RightMarginU, BottomMarginV, 0.0f, 0.0f),		Tint ) ); //29
-	BatchVertices.Add( FSlateVertex::Make<Rounding>( RenderTransform, FVector2D( RightMarginX, EndPos.Y ),		LocalSize, DrawScale, FVector4( RightMarginU, EndUV.Y, 0.0f, 0.0f),				Tint ) ); //30
-	BatchVertices.Add( FSlateVertex::Make<Rounding>( RenderTransform, FVector2D( EndPos.X, BottomMarginY ),		LocalSize, DrawScale, FVector4( EndUV.X, BottomMarginV, 0.0f, 0.0f),			Tint ) ); //31
-	BatchVertices.Add( FSlateVertex::Make<Rounding>( RenderTransform, FVector2D( EndPos.X, EndPos.Y ),			LocalSize, DrawScale, FVector4( EndUV.X, EndUV.Y, 0.0f, 0.0f),					Tint ) ); //32
-
-
-	// The offset into the index buffer where this elements indices start
-	uint32 IndexOffsetStart = BatchIndices.Num();
+	RenderBatch.AddVertex( FSlateVertex::Make<Rounding>( RenderTransform, Position,									LocalSize, DrawScale, FVector4( StartUV.X, StartUV.Y, 0.0f, 0.0f),				Tint ) ); //0
+	RenderBatch.AddVertex( FSlateVertex::Make<Rounding>( RenderTransform, FVector2D( Position.X, TopMarginY ),		LocalSize, DrawScale, FVector4( StartUV.X, TopMarginV, 0.0f, 0.0f),				Tint ) ); //1
+	RenderBatch.AddVertex( FSlateVertex::Make<Rounding>( RenderTransform, FVector2D( LeftMarginX, Position.Y ),		LocalSize, DrawScale, FVector4( LeftMarginU, StartUV.Y, 0.0f, 0.0f),			Tint ) ); //2
+	RenderBatch.AddVertex( FSlateVertex::Make<Rounding>( RenderTransform, FVector2D( LeftMarginX, TopMarginY ),		LocalSize, DrawScale, FVector4( LeftMarginU, TopMarginV, 0.0f, 0.0f),			Tint ) ); //3
+
+	RenderBatch.AddVertex( FSlateVertex::Make<Rounding>( RenderTransform, FVector2D( LeftMarginX, Position.Y ),		LocalSize, DrawScale, FVector4( StartUV.X, StartUV.Y, TopTiling, 0.0f),			Tint ) ); //4
+	RenderBatch.AddVertex( FSlateVertex::Make<Rounding>( RenderTransform, FVector2D( LeftMarginX, TopMarginY ),		LocalSize, DrawScale, FVector4( StartUV.X, TopMarginV, TopTiling, 0.0f),		Tint ) ); //5
+	RenderBatch.AddVertex( FSlateVertex::Make<Rounding>( RenderTransform, FVector2D( RightMarginX, Position.Y ),	LocalSize, DrawScale, FVector4( EndUV.X, StartUV.Y, TopTiling, 0.0f),			Tint ) ); //6
+	RenderBatch.AddVertex( FSlateVertex::Make<Rounding>( RenderTransform, FVector2D( RightMarginX, TopMarginY ),	LocalSize, DrawScale, FVector4( EndUV.X, TopMarginV, TopTiling, 0.0f),			Tint ) ); //7
+
+	RenderBatch.AddVertex( FSlateVertex::Make<Rounding>( RenderTransform, FVector2D( RightMarginX, Position.Y ),	LocalSize, DrawScale, FVector4( RightMarginU, StartUV.Y, 0.0f, 0.0f),			Tint ) ); //8
+	RenderBatch.AddVertex( FSlateVertex::Make<Rounding>( RenderTransform, FVector2D( RightMarginX, TopMarginY ),	LocalSize, DrawScale, FVector4( RightMarginU, TopMarginV, 0.0f, 0.0f),			Tint ) ); //9
+	RenderBatch.AddVertex( FSlateVertex::Make<Rounding>( RenderTransform, FVector2D( EndPos.X, Position.Y ),		LocalSize, DrawScale, FVector4( EndUV.X, StartUV.Y, 0.0f, 0.0f),				Tint ) ); //10
+	RenderBatch.AddVertex( FSlateVertex::Make<Rounding>( RenderTransform, FVector2D( EndPos.X, TopMarginY ),		LocalSize, DrawScale, FVector4( EndUV.X, TopMarginV, 0.0f, 0.0f),				Tint ) ); //11
+
+	RenderBatch.AddVertex( FSlateVertex::Make<Rounding>( RenderTransform, FVector2D( Position.X, TopMarginY ),		LocalSize, DrawScale, FVector4( StartUV.X, StartUV.Y, 0.0f, LeftTiling),		Tint ) ); //12
+	RenderBatch.AddVertex( FSlateVertex::Make<Rounding>( RenderTransform, FVector2D( Position.X, BottomMarginY ),	LocalSize, DrawScale, FVector4( StartUV.X, EndUV.Y, 0.0f, LeftTiling),			Tint ) ); //13
+	RenderBatch.AddVertex( FSlateVertex::Make<Rounding>( RenderTransform, FVector2D( LeftMarginX, TopMarginY ),		LocalSize, DrawScale, FVector4( LeftMarginU, StartUV.Y, 0.0f, LeftTiling),		Tint ) ); //14
+	RenderBatch.AddVertex( FSlateVertex::Make<Rounding>( RenderTransform, FVector2D( LeftMarginX, BottomMarginY ),	LocalSize, DrawScale, FVector4( LeftMarginU, EndUV.Y, 0.0f, LeftTiling),		Tint ) ); //15
+
+	RenderBatch.AddVertex( FSlateVertex::Make<Rounding>( RenderTransform, FVector2D( RightMarginX, TopMarginY ),	LocalSize, DrawScale, FVector4( RightMarginU, StartUV.Y, 0.0f, RightTiling),	Tint ) ); //16
+	RenderBatch.AddVertex( FSlateVertex::Make<Rounding>( RenderTransform, FVector2D( RightMarginX, BottomMarginY ), LocalSize, DrawScale, FVector4( RightMarginU, EndUV.Y, 0.0f, RightTiling),		Tint ) ); //17
+	RenderBatch.AddVertex( FSlateVertex::Make<Rounding>( RenderTransform, FVector2D( EndPos.X, TopMarginY ),		LocalSize, DrawScale, FVector4( EndUV.X, StartUV.Y, 0.0f, RightTiling),			Tint ) ); //18
+	RenderBatch.AddVertex( FSlateVertex::Make<Rounding>( RenderTransform, FVector2D( EndPos.X, BottomMarginY ),		LocalSize, DrawScale, FVector4( EndUV.X, EndUV.Y, 0.0f, RightTiling),			Tint ) ); //19
+
+	RenderBatch.AddVertex( FSlateVertex::Make<Rounding>( RenderTransform, FVector2D( Position.X, BottomMarginY ),	LocalSize, DrawScale, FVector4( StartUV.X, BottomMarginV, 0.0f, 0.0f),			Tint ) ); //20
+	RenderBatch.AddVertex( FSlateVertex::Make<Rounding>( RenderTransform, FVector2D( Position.X, EndPos.Y ),		LocalSize, DrawScale, FVector4( StartUV.X, EndUV.Y, 0.0f, 0.0f),				Tint ) ); //21
+	RenderBatch.AddVertex( FSlateVertex::Make<Rounding>( RenderTransform, FVector2D( LeftMarginX, BottomMarginY ),	LocalSize, DrawScale, FVector4( LeftMarginU, BottomMarginV, 0.0f, 0.0f),		Tint ) ); //22
+	RenderBatch.AddVertex( FSlateVertex::Make<Rounding>( RenderTransform, FVector2D( LeftMarginX, EndPos.Y ),		LocalSize, DrawScale, FVector4( LeftMarginU, EndUV.Y, 0.0f, 0.0f),				Tint ) ); //23
+
+	RenderBatch.AddVertex( FSlateVertex::Make<Rounding>( RenderTransform, FVector2D( LeftMarginX, BottomMarginY ),	LocalSize, DrawScale, FVector4( StartUV.X, BottomMarginV, BottomTiling, 0.0f),	Tint ) ); //24
+	RenderBatch.AddVertex( FSlateVertex::Make<Rounding>( RenderTransform, FVector2D( LeftMarginX, EndPos.Y ),		LocalSize, DrawScale, FVector4( StartUV.X, EndUV.Y, BottomTiling, 0.0f),		Tint ) ); //25
+	RenderBatch.AddVertex( FSlateVertex::Make<Rounding>( RenderTransform, FVector2D( RightMarginX,BottomMarginY ),	LocalSize, DrawScale, FVector4( EndUV.X, BottomMarginV, BottomTiling, 0.0f),	Tint ) ); //26
+	RenderBatch.AddVertex( FSlateVertex::Make<Rounding>( RenderTransform, FVector2D( RightMarginX, EndPos.Y ),		LocalSize, DrawScale, FVector4( EndUV.X, EndUV.Y, BottomTiling, 0.0f),			Tint ) ); //27
+
+	RenderBatch.AddVertex( FSlateVertex::Make<Rounding>( RenderTransform, FVector2D( RightMarginX, BottomMarginY ), LocalSize, DrawScale, FVector4( RightMarginU, BottomMarginV, 0.0f, 0.0f),		Tint ) ); //29
+	RenderBatch.AddVertex( FSlateVertex::Make<Rounding>( RenderTransform, FVector2D( RightMarginX, EndPos.Y ),		LocalSize, DrawScale, FVector4( RightMarginU, EndUV.Y, 0.0f, 0.0f),				Tint ) ); //30
+	RenderBatch.AddVertex( FSlateVertex::Make<Rounding>( RenderTransform, FVector2D( EndPos.X, BottomMarginY ),		LocalSize, DrawScale, FVector4( EndUV.X, BottomMarginV, 0.0f, 0.0f),			Tint ) ); //31
+	RenderBatch.AddVertex( FSlateVertex::Make<Rounding>( RenderTransform, FVector2D( EndPos.X, EndPos.Y ),			LocalSize, DrawScale, FVector4( EndUV.X, EndUV.Y, 0.0f, 0.0f),					Tint ) ); //32
 
 	// Top
-	BatchIndices.Add( IndexStart + 0 );
-	BatchIndices.Add( IndexStart + 1 );
-	BatchIndices.Add( IndexStart + 2 );
-	BatchIndices.Add( IndexStart + 2 );
-	BatchIndices.Add( IndexStart + 1 );
-	BatchIndices.Add( IndexStart + 3 );
-
-	BatchIndices.Add( IndexStart + 4 );
-	BatchIndices.Add( IndexStart + 5 );
-	BatchIndices.Add( IndexStart + 6 );
-	BatchIndices.Add( IndexStart + 6 );
-	BatchIndices.Add( IndexStart + 5 );
-	BatchIndices.Add( IndexStart + 7 );
-
-	BatchIndices.Add( IndexStart + 8 );
-	BatchIndices.Add( IndexStart + 9 );
-	BatchIndices.Add( IndexStart + 10 );
-	BatchIndices.Add( IndexStart + 10 );
-	BatchIndices.Add( IndexStart + 9 );
-	BatchIndices.Add( IndexStart + 11 );
+	RenderBatch.AddIndex( IndexStart + 0 );
+	RenderBatch.AddIndex( IndexStart + 1 );
+	RenderBatch.AddIndex( IndexStart + 2 );
+	RenderBatch.AddIndex( IndexStart + 2 );
+	RenderBatch.AddIndex( IndexStart + 1 );
+	RenderBatch.AddIndex( IndexStart + 3 );
+
+	RenderBatch.AddIndex( IndexStart + 4 );
+	RenderBatch.AddIndex( IndexStart + 5 );
+	RenderBatch.AddIndex( IndexStart + 6 );
+	RenderBatch.AddIndex( IndexStart + 6 );
+	RenderBatch.AddIndex( IndexStart + 5 );
+	RenderBatch.AddIndex( IndexStart + 7 );
+
+	RenderBatch.AddIndex( IndexStart + 8 );
+	RenderBatch.AddIndex( IndexStart + 9 );
+	RenderBatch.AddIndex( IndexStart + 10 );
+	RenderBatch.AddIndex( IndexStart + 10 );
+	RenderBatch.AddIndex( IndexStart + 9 );
+	RenderBatch.AddIndex( IndexStart + 11 );
 
 	// Middle
-	BatchIndices.Add( IndexStart + 12 );
-	BatchIndices.Add( IndexStart + 13 );
-	BatchIndices.Add( IndexStart + 14 );
-	BatchIndices.Add( IndexStart + 14 );
-	BatchIndices.Add( IndexStart + 13 );
-	BatchIndices.Add( IndexStart + 15 );
-
-	BatchIndices.Add( IndexStart + 16 );
-	BatchIndices.Add( IndexStart + 17 );
-	BatchIndices.Add( IndexStart + 18 );
-	BatchIndices.Add( IndexStart + 18 );
-	BatchIndices.Add( IndexStart + 17 );
-	BatchIndices.Add( IndexStart + 19 );
+	RenderBatch.AddIndex( IndexStart + 12 );
+	RenderBatch.AddIndex( IndexStart + 13 );
+	RenderBatch.AddIndex( IndexStart + 14 );
+	RenderBatch.AddIndex( IndexStart + 14 );
+	RenderBatch.AddIndex( IndexStart + 13 );
+	RenderBatch.AddIndex( IndexStart + 15 );
+
+	RenderBatch.AddIndex( IndexStart + 16 );
+	RenderBatch.AddIndex( IndexStart + 17 );
+	RenderBatch.AddIndex( IndexStart + 18 );
+	RenderBatch.AddIndex( IndexStart + 18 );
+	RenderBatch.AddIndex( IndexStart + 17 );
+	RenderBatch.AddIndex( IndexStart + 19 );
 
 	// Bottom
-	BatchIndices.Add( IndexStart + 20 );
-	BatchIndices.Add( IndexStart + 21 );
-	BatchIndices.Add( IndexStart + 22 );
-	BatchIndices.Add( IndexStart + 22 );
-	BatchIndices.Add( IndexStart + 21 );
-	BatchIndices.Add( IndexStart + 23 );
-
-	BatchIndices.Add( IndexStart + 24 );
-	BatchIndices.Add( IndexStart + 25 );
-	BatchIndices.Add( IndexStart + 26 );
-	BatchIndices.Add( IndexStart + 26 );
-	BatchIndices.Add( IndexStart + 25 );
-	BatchIndices.Add( IndexStart + 27 );
-
-	BatchIndices.Add( IndexStart + 28 );
-	BatchIndices.Add( IndexStart + 29 );
-	BatchIndices.Add( IndexStart + 30 );
-	BatchIndices.Add( IndexStart + 30 );
-	BatchIndices.Add( IndexStart + 29 );
-	BatchIndices.Add( IndexStart + 31 );
+	RenderBatch.AddIndex( IndexStart + 20 );
+	RenderBatch.AddIndex( IndexStart + 21 );
+	RenderBatch.AddIndex( IndexStart + 22 );
+	RenderBatch.AddIndex( IndexStart + 22 );
+	RenderBatch.AddIndex( IndexStart + 21 );
+	RenderBatch.AddIndex( IndexStart + 23 );
+
+	RenderBatch.AddIndex( IndexStart + 24 );
+	RenderBatch.AddIndex( IndexStart + 25 );
+	RenderBatch.AddIndex( IndexStart + 26 );
+	RenderBatch.AddIndex( IndexStart + 26 );
+	RenderBatch.AddIndex( IndexStart + 25 );
+	RenderBatch.AddIndex( IndexStart + 27 );
+
+	RenderBatch.AddIndex( IndexStart + 28 );
+	RenderBatch.AddIndex( IndexStart + 29 );
+	RenderBatch.AddIndex( IndexStart + 30 );
+	RenderBatch.AddIndex( IndexStart + 30 );
+	RenderBatch.AddIndex( IndexStart + 29 );
+	RenderBatch.AddIndex( IndexStart + 31 );
 }
 
 void FSlateElementBatcher::AddCustomElement( const FSlateDrawElement& DrawElement )
 {
-	FElementBatchMap& LayerToElementBatches = DrawLayer->GetElementBatchMap();
+	const int32 Layer = DrawElement.GetLayer();
 
-	const FSlateDataPayload& InPayload = DrawElement.GetDataPayload();
-	uint32 Layer = DrawElement.GetLayer();
-
-	if( InPayload.CustomDrawer.IsValid() )
-	{
-		// See if the layer already exists.
-		TUniqueObj<FElementBatchArray>* ElementBatches = LayerToElementBatches.Find( Layer );
-		if( !ElementBatches )
-		{
-			// The layer doesn't exist so make it now
-			ElementBatches = &LayerToElementBatches.Add( Layer );
-		}
-		check( ElementBatches );
-
-		// Custom elements are not batched together 
-		(*ElementBatches)->Add( FSlateElementBatch( InPayload.CustomDrawer, DrawElement.GetClippingIndex(), *ClippingStates ) );
-	}
+	FSlateRenderBatch& RenderBatch = CreateRenderBatch(Layer, FShaderParams(), nullptr, ESlateDrawPrimitive::None, ESlateShader::Default, ESlateDrawEffect::None, ESlateBatchDrawFlag::None, DrawElement);
+	RenderBatch.CustomDrawer = DrawElement.GetDataPayload<FSlateCustomDrawerPayload>().CustomDrawer.Pin().Get();
+	RenderBatch.bIsMergable = false;
 }
 
 void FSlateElementBatcher::AddCustomVerts(const FSlateDrawElement& DrawElement)
 {
-	FElementBatchMap& LayerToElementBatches = DrawLayer->GetElementBatchMap();
+	const FSlateCustomVertsPayload& InPayload = DrawElement.GetDataPayload<FSlateCustomVertsPayload>();
+	const int32 Layer = DrawElement.GetLayer();
 
-	const FSlateDataPayload& InPayload = DrawElement.GetDataPayload();
-	uint32 Layer = DrawElement.GetLayer();
+	if (InPayload.Vertices.Num() > 0)
+	{
+		FSlateRenderBatch& RenderBatch = CreateRenderBatch(
+			Layer, 
+			FShaderParams(), 
+			InPayload.ResourceProxy != nullptr ? InPayload.ResourceProxy->Resource : nullptr, 
+			ESlateDrawPrimitive::TriangleList,
+			ESlateShader::Custom, 
+			DrawElement.GetDrawEffects(), 
+			DrawElement.GetBatchFlags(),
+			DrawElement);
+
+		RenderBatch.bIsMergable = false;
+		RenderBatch.InstanceCount = InPayload.NumInstances;
+		RenderBatch.InstanceOffset = InPayload.InstanceOffset;
+		RenderBatch.InstanceData = InPayload.InstanceData;
+
+		RenderBatch.AddVertices(InPayload.Vertices);
+		RenderBatch.AddIndices(InPayload.Indices);
+
+	}
+	/*FElementBatchMap& LayerToElementBatches = CurrentDrawLayer->GetElementBatchMap();
+
+	const FSlateCustomVertsPayload& InPayload = DrawElement.GetDataPayload<FSlateCustomVertsPayload>();
+	uint32 Layer = DrawElement.GetAbsoluteLayer();
 
-	if (InPayload.CustomVertsData.Num() >0)
+	if (InPayload.Vertices.Num() >0)
 	{
 		// See if the layer already exists.
 		TUniqueObj<FElementBatchArray>* ElementBatches = LayerToElementBatches.Find(Layer);
@@ -2152,20 +2527,20 @@ void FSlateElementBatcher::AddCustomVerts(const FSlateDrawElement& DrawElement)
 		check(ElementBatches);
 
 		FSlateElementBatch NewBatch(
-			InPayload.GetResourceProxy() != nullptr ? InPayload.GetResourceProxy()->Resource : nullptr,
+			InPayload.ResourceProxy != nullptr ? InPayload.ResourceProxy->Resource : nullptr,
 			FShaderParams(),
 			ESlateShader::Custom,
 			ESlateDrawPrimitive::TriangleList,
 			DrawElement.GetDrawEffects(),
-			InPayload.BatchFlags,
-			DrawElement.GetClippingIndex(),
-			*ClippingStates,
+			DrawElement.GetBatchFlags(),
+			DrawElement,
 			InPayload.NumInstances,
 			InPayload.InstanceOffset,
-			InPayload.InstanceData,
-			DrawElement.GetSceneIndex()
+			InPayload.InstanceData
 		);
 
+		NewBatch.SaveClippingState(*PrecachedClippingStates);
+
 		int32 Index = (*ElementBatches)->Add(NewBatch);
 		FSlateElementBatch* ElementBatch = &(**ElementBatches)[Index];
 
@@ -2176,54 +2551,11 @@ void FSlateElementBatcher::AddCustomVerts(const FSlateDrawElement& DrawElement)
 		FSlateIndexArray& BatchIndices = BatchData->GetBatchIndexList(*ElementBatch);
 
 		// Vertex Buffer since  it is already in slate format it is a straight copy
-		BatchVertices = InPayload.CustomVertsData;
-		BatchIndices = InPayload.CustomVertsIndexData;
-	}
-}
-
-void FSlateElementBatcher::AddCachedBuffer(const FSlateDrawElement& DrawElement)
-{
-	FElementBatchMap& LayerToElementBatches = DrawLayer->GetElementBatchMap();
-
-	uint32 Layer = DrawElement.GetLayer();
-
-		// See if the layer already exists.
-		TUniqueObj<FElementBatchArray>* ElementBatches = LayerToElementBatches.Find(Layer);
-		if ( !ElementBatches )
-		{
-			// The layer doesn't exist so make it now
-			ElementBatches = &LayerToElementBatches.Add(Layer);
-		}
-		check(ElementBatches);
-
-		// Custom elements are not batched together
-	TSharedPtr< FSlateRenderDataHandle, ESPMode::ThreadSafe > RenderData = DrawElement.GetDataPayload().GetRenderDataHandle()->AsShared();
-	(*ElementBatches)->Add(FSlateElementBatch(RenderData, DrawElement.GetDataPayload().GetRenderOffset(), DrawElement.GetClippingIndex(), *ClippingStates));
+		BatchVertices = InPayload.Vertices;
+		BatchIndices = InPayload.Indices;
+	}*/
 }
 
-void FSlateElementBatcher::AddLayer(const FSlateDrawElement& DrawElement)
-{
-	FElementBatchMap& LayerToElementBatches = DrawLayer->GetElementBatchMap();
-
-	const FSlateDataPayload& InPayload = DrawElement.GetDataPayload();
-	uint32 Layer = DrawElement.GetLayer();
-
-	if ( InPayload.LayerHandle )
-	{
-		// See if the layer already exists.
-		TUniqueObj<FElementBatchArray>* ElementBatches = LayerToElementBatches.Find(Layer);
-		if ( !ElementBatches )
-		{
-			// The layer doesn't exist so make it now
-			ElementBatches = &LayerToElementBatches.Add(Layer);
-		}
-		check(ElementBatches);
-
-		// Custom elements are not batched together
-		TSharedPtr< FSlateDrawLayerHandle, ESPMode::ThreadSafe > LayerHandle = InPayload.LayerHandle->AsShared();
-		(*ElementBatches)->Add(FSlateElementBatch(LayerHandle, DrawElement.GetClippingIndex(), *ClippingStates));
-	}
-}
 
 void FSlateElementBatcher::AddPostProcessPass(const FSlateDrawElement& DrawElement, const FVector2D& WindowSize)
 {
@@ -2232,12 +2564,12 @@ void FSlateElementBatcher::AddPostProcessPass(const FSlateDrawElement& DrawEleme
 	const FSlateRenderTransform& RenderTransform = DrawElement.GetRenderTransform();
 	const FVector2D& LocalSize = DrawElement.GetLocalSize();
 	
-	const FSlateDataPayload& Payload = DrawElement.GetDataPayload();
+	const FSlatePostProcessPayload& Payload = DrawElement.GetDataPayload<FSlatePostProcessPayload>();
 
 	//@todo doesn't work with rotated or skewed objects yet
 	const FVector2D& Position = DrawElement.GetPosition();
 
-	uint32 Layer = DrawElement.GetLayer();
+	const int32 Layer = DrawElement.GetLayer();
 
 	// Determine the four corners of the quad
 	FVector2D TopLeft = FVector2D::ZeroVector;
@@ -2259,38 +2591,61 @@ void FSlateElementBatcher::AddPostProcessPass(const FSlateDrawElement& DrawEleme
 	{
 		FShaderParams Params = FShaderParams::MakePixelShaderParams(FVector4(WorldTopLeft, WorldBotRight), FVector4(Payload.PostProcessData.X, Payload.PostProcessData.Y, Payload.DownsampleAmount, 0));
 
-		FElementBatchMap& LayerToElementBatches = DrawLayer->GetElementBatchMap();
+		CreateRenderBatch(Layer, Params, nullptr, ESlateDrawPrimitive::TriangleList, ESlateShader::PostProcess, ESlateDrawEffect::None, ESlateBatchDrawFlag::None, DrawElement);
+	}
+}
 
-		// See if the layer already exists.
-		TUniqueObj<FElementBatchArray>* ElementBatches = LayerToElementBatches.Find(Layer);
-		if (!ElementBatches)
-		{
-			// The layer doesn't exist so make it now
-			ElementBatches = &LayerToElementBatches.Add(Layer);
-		}
-		check(ElementBatches);
+FSlateRenderBatch& FSlateElementBatcher::CreateRenderBatch(
+	int32 Layer, 
+	const FShaderParams& ShaderParams,
+	const FSlateShaderResource* InResource,
+	ESlateDrawPrimitive PrimitiveType,
+	ESlateShader ShaderType,
+	ESlateDrawEffect DrawEffects,
+	ESlateBatchDrawFlag DrawFlags,
+	const FSlateDrawElement& DrawElement)
+{
+	FSlateRenderBatch& NewBatch = CachedElementList
+		? CachedElementList->AddRenderBatch(Layer, ShaderParams, InResource, PrimitiveType, ShaderType, DrawEffects, DrawFlags, DrawElement.GetSceneIndex())
+		: BatchData->AddRenderBatch(Layer, ShaderParams, InResource, PrimitiveType, ShaderType, DrawEffects, DrawFlags, DrawElement.GetSceneIndex());
 
-		// Custom elements are not batched together 
-		(*ElementBatches)->Add(FSlateElementBatch(nullptr, Params, ESlateShader::PostProcess, ESlateDrawPrimitive::TriangleList, ESlateDrawEffect::None, ESlateBatchDrawFlag::None, DrawElement.GetClippingIndex(), *ClippingStates));
-	}
+	NewBatch.ClippingState = ResolveClippingState(DrawElement);
+
+	return NewBatch;
 }
 
+const FSlateClippingState* FSlateElementBatcher::ResolveClippingState(const FSlateDrawElement& DrawElement) const
+{
+	const FClipStateHandle& ClipHandle = DrawElement.GetClippingHandle();
+	// Do cached first
+	if (ClipHandle.GetCachedClipState())
+	{
+		// We should be working with cached elements if we have a cached clip state
+		check(CachedElementList);
+		return ClipHandle.GetCachedClipState();
+	}
+	else if (PrecachedClippingStates->IsValidIndex(ClipHandle.GetPrecachedClipIndex()))
+	{
+		// Store the clipping state so we can use it later for rendering.
+		return &(*PrecachedClippingStates)[ClipHandle.GetPrecachedClipIndex()];
+	}
 
+	return nullptr;
+}
 
+/*
 FSlateElementBatch& FSlateElementBatcher::FindBatchForElement(
 	uint32 Layer, 
 	const FShaderParams& ShaderParams, 
 	const FSlateShaderResource* InTexture, 
-	ESlateDrawPrimitive::Type PrimitiveType,
-	ESlateShader::Type ShaderType, 
+	ESlateDrawPrimitive PrimitiveType,
+	ESlateShader ShaderType, 
 	ESlateDrawEffect DrawEffects, 
 	ESlateBatchDrawFlag DrawFlags,
-	int32 ClippingIndex,
-	int32 SceneIndex)
+	const FSlateDrawElement& DrawElement)
 {
-
 	SCOPE_CYCLE_COUNTER( STAT_SlateFindBatchForElement );
-	FElementBatchMap& LayerToElementBatches = DrawLayer->GetElementBatchMap();
+	FElementBatchMap& LayerToElementBatches = CurrentDrawLayer->GetElementBatchMap();
 
 	// See if the layer already exists.
 	TUniqueObj<FElementBatchArray>* ElementBatches = LayerToElementBatches.Find( Layer );
@@ -2303,7 +2658,7 @@ FSlateElementBatch& FSlateElementBatcher::FindBatchForElement(
 	checkSlow( ElementBatches );
 
 	// Create a temp batch so we can use it as our key to find if the same batch already exists
-	FSlateElementBatch TempBatch( InTexture, ShaderParams, ShaderType, PrimitiveType, DrawEffects, DrawFlags, ClippingIndex, *ClippingStates, 0, 0, nullptr, SceneIndex );
+	FSlateElementBatch TempBatch( InTexture, ShaderParams, ShaderType, PrimitiveType, DrawEffects, DrawFlags, DrawElement, 0, 0, nullptr );
 
 	FSlateElementBatch* ElementBatch = (*ElementBatches)->FindByKey( TempBatch );
 	if( !ElementBatch )
@@ -2314,6 +2669,8 @@ FSlateElementBatch& FSlateElementBatcher::FindBatchForElement(
 
 		BatchData->AssignVertexArrayToBatch(*ElementBatch);
 		BatchData->AssignIndexArrayToBatch(*ElementBatch);
+		ElementBatch->SaveClippingState(*PrecachedClippingStates);
+
 	}
 	check( ElementBatch );
 
@@ -2321,9 +2678,11 @@ FSlateElementBatch& FSlateElementBatcher::FindBatchForElement(
 	++ElementBatch->NumElementsInBatch;
 	return *ElementBatch;
 }
+*/
 
 void FSlateElementBatcher::ResetBatches()
 {
 	bRequiresVsync = false;
 	NumPostProcessPasses = 0;
 }
+
