Slate Global Invalidation System

diff --git a/Engine/Source/Runtime/SlateCore/Private/Rendering/SlateDrawBuffer.cpp b/Engine/Source/Runtime/SlateCore/Private/Rendering/SlateDrawBuffer.cpp
--- a/Engine/Source/Runtime/SlateCore/Private/Rendering/SlateDrawBuffer.cpp
+++ b/Engine/Source/Runtime/SlateCore/Private/Rendering/SlateDrawBuffer.cpp
@@ -3,6 +3,8 @@
 #include "Rendering/SlateDrawBuffer.h"
 #include "Rendering/DrawElements.h"
 #include "Application/SlateApplicationBase.h"
+#include "Widgets/SWindow.h"
+#include "Input/HittestGrid.h"
 
 
 /* FSlateDrawBuffer interface
@@ -23,13 +25,14 @@ FSlateWindowElementList& FSlateDrawBuffer::AddWindowElementList(TSharedRef<SWind
 			WindowElementLists.Add(ExistingElementList);
 			WindowElementListsPool.RemoveAtSwap(WindowIndex);
 
-			ExistingElementList->ResetElementBuffers();
+			ExistingElementList->ResetElementList();
 
 			return *ExistingElementList;
 		}
 	}
 
 	TSharedRef<FSlateWindowElementList> WindowElements = MakeShared<FSlateWindowElementList>(ForWindow);
+	WindowElements->ResetElementList();
 	WindowElementLists.Add(WindowElements);
 
 	return *WindowElements;
diff --git a/Engine/Source/Runtime/SlateCore/Private/Rendering/SlateRenderer.cpp b/Engine/Source/Runtime/SlateCore/Private/Rendering/SlateRenderer.cpp
--- a/Engine/Source/Runtime/SlateCore/Private/Rendering/SlateRenderer.cpp
+++ b/Engine/Source/Runtime/SlateCore/Private/Rendering/SlateRenderer.cpp
@@ -6,6 +6,7 @@
 #include "Fonts/FontMeasure.h"
 #include "Widgets/SWindow.h"
 
+DEFINE_STAT(STAT_SlatePreFullBufferRTTime);
 
 /* FSlateFontCacheProvider interface
  *****************************************************************************/
@@ -133,14 +134,10 @@ ISlateAtlasProvider* FSlateRenderer::GetFontAtlasProvider()
 	return &SlateFontServices->GetGameThreadFontCache().Get();
 }
 
-TSharedRef<FSlateRenderDataHandle, ESPMode::ThreadSafe> FSlateRenderer::CacheElementRenderData(const ILayoutCache* Cacher, FSlateWindowElementList& ElementList)
+void FSlateRenderer::DestroyCachedFastPathRenderingData(struct FSlateCachedFastPathRenderingData* FastPathRenderingData)
 {
-	return MakeShareable(new FSlateRenderDataHandle(Cacher, nullptr));
-}
-
-void FSlateRenderer::ReleaseCachingResourcesFor(const ILayoutCache* Cacher)
-{
-
+	check(FastPathRenderingData);
+	delete FastPathRenderingData;
 }
 
 /* Global functions
diff --git a/Engine/Source/Runtime/SlateCore/Private/SlateCoreClasses.cpp b/Engine/Source/Runtime/SlateCore/Private/SlateCoreClasses.cpp
--- a/Engine/Source/Runtime/SlateCore/Private/SlateCoreClasses.cpp
+++ b/Engine/Source/Runtime/SlateCore/Private/SlateCoreClasses.cpp
@@ -21,16 +21,6 @@ FAutoConsoleVariableRef CVarSlateContrast(
 );
 
 
-
-int32 GSlateLayoutCaching = 0;
-
-FAutoConsoleVariableRef CVarSlateLayoutCaching(
-	TEXT("Slate.EnableLayoutCaching"),
-	GSlateLayoutCaching,
-	TEXT("Whether or not dynamic prepass and layout caching is enabled")
-);
-
-
 // Enable fast widget paths outside the editor by default.  Only reason we don't enable them everywhere
 // is that the editor is more complex than a game, and there are likely a larger swath of edge cases.
 int32 GSlateFastWidgetPath = 0;
@@ -42,6 +32,33 @@ FAutoConsoleVariableRef CVarSlateFastWidgetPath(
 );
 
 
+int32 GSlateEnableGlobalInvalidation = 0;
+static FAutoConsoleVariableRef CVarSlateNewUpdateMethod(
+	TEXT("Slate.EnableGlobalInvalidation"), 
+	GSlateEnableGlobalInvalidation, 
+	TEXT("")
+);
+
+bool GSlateIsOnFastUpdatePath = false;
+bool GSlateIsInInvalidationSlowPath = false;
+
+#if WITH_SLATE_DEBUGGING
+int32 GSlateInvalidationDebugging = 0;
+/** True if we should allow widgets to be cached in the UI at all. */
+FAutoConsoleVariableRef CVarInvalidationDebugging(
+	TEXT("Slate.InvalidationDebugging"),
+	GSlateInvalidationDebugging,
+	TEXT("Whether to show invalidation debugging visualization"));
+
+
+int32 GSlateHitTestGridDebugging = 0;
+/** True if we should allow widgets to be cached in the UI at all. */
+FAutoConsoleVariableRef CVarHitTestGridDebugging(
+	TEXT("Slate.HitTestGridDebugging"),
+	GSlateHitTestGridDebugging,
+	TEXT("Whether to show a visualization of everything in the hit teest grid"));
+
+#endif
 
 FSlateWidgetStyle::FSlateWidgetStyle()
 { }
diff --git a/Engine/Source/Runtime/SlateCore/Private/SlotBase.cpp b/Engine/Source/Runtime/SlateCore/Private/SlotBase.cpp
--- a/Engine/Source/Runtime/SlateCore/Private/SlotBase.cpp
+++ b/Engine/Source/Runtime/SlateCore/Private/SlotBase.cpp
@@ -47,11 +47,6 @@ void FSlotBase::DetatchParentFromContent()
 
 void FSlotBase::AfterContentOrOwnerAssigned()
 {
-	if (GSlateLayoutCaching && RawParentPtr)
-	{
-		RawParentPtr->InvalidatePrepass();
-	}
-
 	if (RawParentPtr)
 	{
 		if (Widget != SNullWidget::NullWidget)
diff --git a/Engine/Source/Runtime/SlateCore/Private/Types/PaintArgs.cpp b/Engine/Source/Runtime/SlateCore/Private/Types/PaintArgs.cpp
--- a/Engine/Source/Runtime/SlateCore/Private/Types/PaintArgs.cpp
+++ b/Engine/Source/Runtime/SlateCore/Private/Types/PaintArgs.cpp
@@ -3,86 +3,29 @@
 #include "Types/PaintArgs.h"
 #include "Layout/ArrangedWidget.h"
 #include "Input/HittestGrid.h"
-#include "Layout/WidgetCaching.h"
-
-FPaintArgs::FPaintArgs(const SWidget& Parent, FHittestGrid& InHittestGrid, FVector2D InWindowOffset, double InCurrentTime, float InDeltaTime)
-: ParentPtr(Parent)
-, Grid(InHittestGrid)
-, LastHittestIndex(INDEX_NONE)
-, LastRecordedVisibility( EVisibility::Visible )
-, WindowOffset(InWindowOffset)
-, CurrentTime(InCurrentTime)
-, DeltaTime(InDeltaTime)
-, bIsCaching(false)
-, bIsVolatilityPass(false)
-, LayoutCache(nullptr)
-, ParentCacheNode(nullptr)
+#include "Widgets/SWindow.h"
+
+FPaintArgs::FPaintArgs(const SWidget* PaintParent, FHittestGrid& InRootHittestGrid, FHittestGrid& InCurrentHitTestGrid, FVector2D InWindowOffset, double InCurrentTime, float InDeltaTime)
+	: RootGrid(InRootHittestGrid)
+	, CurrentGrid(InCurrentHitTestGrid)
+	, WindowOffset(InWindowOffset)
+	, PaintParentPtr(PaintParent)
+	, CurrentTime(InCurrentTime)
+	, DeltaTime(InDeltaTime)
+	, bInheritedHittestability(true)
 {
 }
 
-FPaintArgs FPaintArgs::EnableCaching(const TWeakPtr<ILayoutCache>& InLayoutCache, FCachedWidgetNode* InParentCacheNode, bool bEnableCaching, bool bEnableVolatile) const
+FPaintArgs::FPaintArgs(const SWidget* PaintParent, FHittestGrid& InRootHittestGrid, FVector2D InWindowOffset, double InCurrentTime, float InDeltaTime)
+	: FPaintArgs(PaintParent, InRootHittestGrid, InRootHittestGrid, InWindowOffset, InCurrentTime, InDeltaTime)
 {
-	FPaintArgs UpdatedArgs(*this);
-	UpdatedArgs.LayoutCache = InLayoutCache;
-	UpdatedArgs.ParentCacheNode = InParentCacheNode;
-	UpdatedArgs.bIsCaching = bEnableCaching;
-	UpdatedArgs.bIsVolatilityPass = bEnableVolatile;
-
-	return UpdatedArgs;
-}
 
-FPaintArgs FPaintArgs::WithNewTime(double InCurrentTime, float InDeltaTime) const
-{
-	FPaintArgs UpdatedArgs(*this);
-	UpdatedArgs.CurrentTime = InCurrentTime;
-	UpdatedArgs.DeltaTime = InDeltaTime;
-	return UpdatedArgs;
 }
 
-FPaintArgs FPaintArgs::RecordHittestGeometry(const SWidget* Widget, const FGeometry& WidgetGeometry, int32 LayerId) const
+FPaintArgs FPaintArgs::InsertCustomHitTestPath(const SWidget* Widget, TSharedRef<ICustomHitTestPath> CustomHitTestPath) const
 {
-	FPaintArgs UpdatedArgs(*this);
-
-	if ( LastRecordedVisibility.AreChildrenHitTestVisible() )
-	{
-		TSharedRef<SWidget> SafeWidget = const_cast<SWidget*>(Widget)->AsShared();
-
-		if ( bIsCaching )
-		{
-			TSharedPtr<ILayoutCache> SharedLayoutCache = LayoutCache.Pin();
-			if (SharedLayoutCache.IsValid())
-			{
-				FCachedWidgetNode* CacheNode = SharedLayoutCache->CreateCacheNode();
-				CacheNode->Initialize(*this, SafeWidget, WidgetGeometry);
-				UpdatedArgs.ParentCacheNode->Children.Add(CacheNode);
+	TSharedRef<SWidget> SafeWidget = const_cast<SWidget*>(Widget)->AsShared();
 
-				UpdatedArgs.ParentCacheNode = CacheNode;
-			}
-		}
-
-		int32 RealLastHitTestIndex = LastHittestIndex;
-		if ( bIsVolatilityPass && !bIsCaching && ParentCacheNode )
-		{
-			RealLastHitTestIndex = ParentCacheNode->LastRecordedHittestIndex;
-			UpdatedArgs.ParentCacheNode = nullptr;
-		}
-
-		// When rendering volatile widgets, their parent widgets who have been cached 
-		const EVisibility RecordedVisibility = Widget->GetVisibility();
-		const int32 RecordedHittestIndex = Grid.InsertWidget(RealLastHitTestIndex, RecordedVisibility, FArrangedWidget(SafeWidget, WidgetGeometry), WindowOffset, LayerId);
-		UpdatedArgs.LastHittestIndex = RecordedHittestIndex;
-		UpdatedArgs.LastRecordedVisibility = RecordedVisibility;
-	}
-	else
-	{
-		UpdatedArgs.LastRecordedVisibility = LastRecordedVisibility;
-	}
-
-	return UpdatedArgs;
-}
-
-FPaintArgs FPaintArgs::InsertCustomHitTestPath( TSharedRef<ICustomHitTestPath> CustomHitTestPath, int32 InLastHittestIndex ) const
-{
-	const_cast<FHittestGrid&>(Grid).InsertCustomHitTestPath(CustomHitTestPath, InLastHittestIndex);
+	const_cast<FHittestGrid&>(CurrentGrid).InsertCustomHitTestPath(SafeWidget, CustomHitTestPath);
 	return *this;
 }
diff --git a/Engine/Source/Runtime/SlateCore/Private/Widgets/Accessibility/SlateCoreAccessibleWidgets.cpp b/Engine/Source/Runtime/SlateCore/Private/Widgets/Accessibility/SlateCoreAccessibleWidgets.cpp
--- a/Engine/Source/Runtime/SlateCore/Private/Widgets/Accessibility/SlateCoreAccessibleWidgets.cpp
+++ b/Engine/Source/Runtime/SlateCore/Private/Widgets/Accessibility/SlateCoreAccessibleWidgets.cpp
@@ -270,7 +270,7 @@ TSharedPtr<IAccessibleWidget> FSlateAccessibleWindow::GetChildAtPosition(int32 X
 		if (UseHitTestGrid)
 		{
 			TSharedPtr<SWindow> SlateWindow = StaticCastSharedPtr<SWindow>(Widget.Pin());
-			TArray<FWidgetAndPointer> Hits = SlateWindow->GetHittestGrid()->GetBubblePath(FVector2D(X, Y), 0.0f, false);
+			TArray<FWidgetAndPointer> Hits = SlateWindow->GetHittestGrid().GetBubblePath(FVector2D(X, Y), 0.0f, false);
 			TSharedPtr<SWidget> LastAccessibleWidget = nullptr;
 			for (int32 i = 0; i < Hits.Num(); ++i)
 			{
diff --git a/Engine/Source/Runtime/SlateCore/Private/Widgets/SOverlay.cpp b/Engine/Source/Runtime/SlateCore/Private/Widgets/SOverlay.cpp
--- a/Engine/Source/Runtime/SlateCore/Private/Widgets/SOverlay.cpp
+++ b/Engine/Source/Runtime/SlateCore/Private/Widgets/SOverlay.cpp
@@ -4,6 +4,7 @@
  #include "Types/PaintArgs.h"
  #include "Layout/ArrangedChildren.h"
  #include "Layout/LayoutUtils.h"
+#include "Rendering/DrawElements.h"
 
 
 SOverlay::SOverlay()
@@ -77,32 +78,49 @@ int32 SOverlay::OnPaint( const FPaintArgs& Args, const FGeometry& AllottedGeomet
 	// Because we paint multiple children, we must track the maximum layer id that they produced in case one of our parents
 	// wants to an overlay for all of its contents.
 	int32 MaxLayerId = LayerId;
-
-	const FPaintArgs NewArgs = Args.WithNewParent(this);
+	FPaintArgs NewArgs = Args.WithNewParent(this);
 	const bool bChildrenEnabled = ShouldBeEnabled(bParentEnabled);
 
+	//int32 NumPriorityGroupsPushed = 0;
+
 	for (int32 ChildIndex = 0; ChildIndex < ArrangedChildren.Num(); ++ChildIndex)
 	{
+		FArrangedWidget& CurArrangedWidget = ArrangedChildren[ChildIndex];
+
 		// We don't increment the first layer.
 		if (ChildIndex > 0)
 		{
+			//OutDrawElements.PushBatchPriortyGroup(CurArrangedWidget.Widget.Get());
+			//NumPriorityGroupsPushed++;
 			MaxLayerId++;
 		}
 
-		FArrangedWidget& CurWidget = ArrangedChildren[ChildIndex];
-
 		const int32 CurWidgetsMaxLayerId =
-			CurWidget.Widget->Paint(
+			CurArrangedWidget.Widget->Paint(
 				NewArgs,
-				CurWidget.Geometry,
+				CurArrangedWidget.Geometry,
 				MyCullingRect,
 				OutDrawElements,
 				MaxLayerId,
 				InWidgetStyle,
 				bChildrenEnabled);
+		
+		// This is a hack to account for widgets incrementing their layer id inside an overlay in global invalidation mode.  
+		// Overlay slots that do not update will not know about the new layer id.  This padding adds buffering to avoid that being a problem for now
+		// This is a temporary solution until we build a full rendering tree
+		const int32 OverlaySlotPadding = 10;
+		MaxLayerId = CurWidgetsMaxLayerId + FMath::Max((CurWidgetsMaxLayerId - MaxLayerId) / OverlaySlotPadding, 1) * OverlaySlotPadding;
+		//MaxLayerId = FMath::Max(CurWidgetsMaxLayerId, MaxLayerId);
+	
+	}
 
-		MaxLayerId = FMath::Max(MaxLayerId, CurWidgetsMaxLayerId);
+	/*// Remove all priority groups made
+	for (int32 i = 0; i < NumPriorityGroupsPushed; ++i)
+	{
+		OutDrawElements.PopBatchPriortyGroup();
 	}
+*/
+
 
 	return MaxLayerId;
 }
diff --git a/Engine/Source/Runtime/SlateCore/Private/Widgets/SWidget.cpp b/Engine/Source/Runtime/SlateCore/Private/Widgets/SWidget.cpp
--- a/Engine/Source/Runtime/SlateCore/Private/Widgets/SWidget.cpp
+++ b/Engine/Source/Runtime/SlateCore/Private/Widgets/SWidget.cpp
@@ -15,19 +15,22 @@
 #include "Application/ActiveTimerHandle.h"
 #include "Input/HittestGrid.h"
 #include "Debugging/SlateDebugging.h"
+#include "Widgets/SWindow.h"
+
 #if WITH_ACCESSIBILITY
 #include "Widgets/Accessibility/SlateCoreAccessibleWidgets.h"
 #include "Widgets/Accessibility/SlateAccessibleMessageHandler.h"
 #endif
 
-DECLARE_DWORD_COUNTER_STAT(TEXT("Widgets Created (Per Frame)"), STAT_SlateTotalWidgetsPerFrame, STATGROUP_Slate);
-DECLARE_DWORD_COUNTER_STAT(TEXT("SWidget::Paint (Count)"), STAT_SlateNumPaintedWidgets, STATGROUP_Slate);
-DECLARE_DWORD_COUNTER_STAT(TEXT("SWidget::Tick (Count)"), STAT_SlateNumTickedWidgets, STATGROUP_Slate);
-DECLARE_CYCLE_STAT(TEXT("TickWidgets"), STAT_SlateTickWidgets, STATGROUP_Slate);
-
-DECLARE_DWORD_ACCUMULATOR_STAT(TEXT("Total Widgets"), STAT_SlateTotalWidgets, STATGROUP_SlateMemory);
-DECLARE_MEMORY_STAT(TEXT("SWidget Total Allocated Size"), STAT_SlateSWidgetAllocSize, STATGROUP_SlateMemory);
-
+DEFINE_STAT(STAT_SlateVeryVerboseStatGroupTester);
+DEFINE_STAT(STAT_SlateTotalWidgetsPerFrame);
+DEFINE_STAT(STAT_SlateNumPaintedWidgets);
+DEFINE_STAT(STAT_SlateNumTickedWidgets);
+DEFINE_STAT(STAT_SlateExecuteActiveTimers);
+DEFINE_STAT(STAT_SlateTickWidgets);
+DEFINE_STAT(STAT_SlatePrepass);
+DEFINE_STAT(STAT_SlateTotalWidgets);
+DEFINE_STAT(STAT_SlateSWidgetAllocSize);
 
 
 #if SLATE_CULL_WIDGETS
@@ -89,7 +92,6 @@ struct FScopeCycleCounterSWidget
 
 #endif
 
-DEFINE_STAT(STAT_SlateVeryVerboseStatGroupTester);
 
 void SWidget::CreateStatID() const
 {
@@ -98,6 +100,35 @@ void SWidget::CreateStatID() const
 #endif
 }
 
+void SWidget::UpdateWidgetProxy(int32 NewLayerId, FSlateCachedElementListNode* CacheNode)
+{
+#if WITH_SLATE_DEBUGGING
+	check(!CacheNode || CacheNode->GetValue().Widget == this);
+#endif
+
+	PersistentState.CachedElementListNode = CacheNode;
+
+	if (FastPathProxyHandle.IsValid())
+	{
+		FWidgetProxy& MyProxy = FastPathProxyHandle.GetProxy();
+
+		MyProxy.Visibility = GetVisibility();
+
+		PersistentState.OutgoingLayerId = NewLayerId;
+
+		Advanced_InvalidateVolatility();
+		if ((IsVolatile() && !IsVolatileIndirectly()) || (Advanced_IsInvalidationRoot() && !Advanced_IsWindow()))
+		{
+			AddUpdateFlags(EWidgetUpdateFlags::NeedsVolatilePaint);
+		}
+		else
+		{
+			RemoveUpdateFlags(EWidgetUpdateFlags::NeedsVolatilePaint);
+		}
+		FastPathProxyHandle.MarkWidgetUpdatedThisFrame();
+	}
+}
+
 FName NAME_MouseButtonDown(TEXT("MouseButtonDown"));
 FName NAME_MouseButtonUp(TEXT("MouseButtonUp"));
 FName NAME_MouseMove(TEXT("MouseMove"));
@@ -112,11 +143,15 @@ SWidget::SWidget()
 	, bForceVolatile(false)
 	, bCachedVolatile(false)
 	, bInheritedVolatility(false)
+	, bInvisibleDueToParentOrSelfVisibility(false)
 	, bNeedsPrepass(true)
 	, bNeedsDesiredSize(true)
 	, bUpdatingDesiredSize(false)
+	, bHasCustomPrepass(false)
+	, bVolatilityAlwaysInvalidatesPrepass(false)
 	, Clipping(EWidgetClipping::Inherit)
 	, FlowDirectionPreference(EFlowDirectionPreference::Inherit)
+	// Note we are defaulting to tick for backwards compatibility
 	, UpdateFlags(EWidgetUpdateFlags::NeedsTick)
 	, DesiredSize()
 	, PrepassLayoutScaleMultiplier(1.0f)
@@ -128,7 +163,6 @@ SWidget::SWidget()
 	, RenderTransformPivot(FVector2D::ZeroVector)
 	, Cursor( TOptional<EMouseCursor::Type>() )
 	, ToolTip()
-	, LayoutCache(nullptr)	
 {
 	if (GIsRunning)
 	{
@@ -146,6 +180,33 @@ SWidget::~SWidget()
 		{
 			FSlateApplicationBase::Get().UnRegisterActiveTimer(ActiveTimerHandle);
 		}
+
+		if (FastPathProxyHandle.IsValid())
+		{
+			FWidgetProxy& Proxy = FastPathProxyHandle.GetProxy();
+			Proxy.Widget = nullptr;
+		}
+		else
+		{
+			ensure(FastPathProxyHandle.GetIndex() == INDEX_NONE);
+		}
+
+		if (FSlateInvalidationRoot* Root = FastPathProxyHandle.GetInvalidationRoot())
+		{
+			Root->InvalidateChildOrder();
+
+			if (PersistentState.CachedElementListNode)
+			{
+				Root->GetCachedElements().RemoveCache(PersistentState.CachedElementListNode);
+			}
+		
+		}
+
+		// Reset handle
+		FastPathProxyHandle = FWidgetProxyHandle();
+
+		check(!PersistentState.CachedElementListNode);
+
 #if WITH_ACCESSIBILITY
 		FSlateApplicationBase::Get().GetAccessibleMessageHandler()->OnWidgetRemoved(this);
 #endif
@@ -501,91 +562,43 @@ void SWidget::Tick( const FGeometry& AllottedGeometry, const double InCurrentTim
 
 void SWidget::SlatePrepass()
 {
-	if (!GSlateLayoutCaching)
-	{
-		SlatePrepass(FSlateApplicationBase::Get().GetApplicationScale());
-	}
+	SlatePrepass(FSlateApplicationBase::Get().GetApplicationScale());
 }
 
 void SWidget::SlatePrepass(float InLayoutScaleMultiplier)
 {
-#if WITH_VERY_VERBOSE_SLATE_STATS
-	SCOPED_NAMED_EVENT(SWidget_Prepass, FColor::Silver);
-#endif
+	SCOPE_CYCLE_COUNTER(STAT_SlatePrepass);
 
-	if(GSlateLayoutCaching)
+	if (!GSlateIsOnFastUpdatePath || bNeedsPrepass)
 	{
-		if (!bNeedsPrepass && PrepassLayoutScaleMultiplier == InLayoutScaleMultiplier)
-		{
-			return;
-		}
-
-		PrepassLayoutScaleMultiplier = InLayoutScaleMultiplier;
-		bNeedsPrepass = false;
-
 		// If the scale changed, that can affect the desired size of some elements that take it into
 		// account, such as text, so when the prepass size changes, so must we invalidate desired size.
-		//bNeedsDesiredSize = true;
-		// Disabled: This was causing instability in the feedback loop created by the scale box
-	}
-
-	if ( bCanHaveChildren )
-	{
-		// Cache child desired sizes first. This widget's desired size is
-		// a function of its children's sizes.
-		FChildren* MyChildren = this->GetChildren();
-		const int32 NumChildren = MyChildren->Num();
-		for (int32 ChildIndex = 0; ChildIndex < MyChildren->Num(); ++ChildIndex)
-		{
-			const TSharedRef<SWidget>& Child = MyChildren->GetChildAt(ChildIndex);
+		bNeedsDesiredSize = true;
 
-			if (GSlateLayoutCaching || Child->Visibility.Get() != EVisibility::Collapsed)
-			{
-				const float ChildLayoutScaleMultiplier = GetRelativeLayoutScale(MyChildren->GetSlotAt(ChildIndex), InLayoutScaleMultiplier);
-				// Recur: Descend down the widget tree.
-				Child->SlatePrepass(InLayoutScaleMultiplier * ChildLayoutScaleMultiplier);
-			}
-		}
-		ensure(NumChildren == MyChildren->Num());
+		Prepass_Internal(InLayoutScaleMultiplier);
 	}
 
-	if(!GSlateLayoutCaching)
-	{
-		// Cache this widget's desired size.
-		CacheDesiredSize(InLayoutScaleMultiplier);
-	}
 }
 
 void SWidget::InvalidatePrepass()
 {
-	SCOPED_NAMED_EVENT(SWidget_InvalidatePrepass, FColor::Orange);
-
 	bNeedsPrepass = true;
-	LayoutChanged(EInvalidateWidget::LayoutAndVolatility);
 }
 
-FVector2D SWidget::GetDesiredSize() const
+void SWidget::InvalidateChildRemovedFromTree(SWidget& Child)
 {
-	if(GSlateLayoutCaching)
-	{
-		if (bNeedsDesiredSize && ensureMsgf(!bUpdatingDesiredSize, TEXT("The layout is cyclically dependent.  A child widget can not ask the desired size of a parent while the parent is asking the desired size of its children.")))
-		{
-			bUpdatingDesiredSize = true;
-
-			// Cache this widget's desired size.
-			const_cast<SWidget*>(this)->CacheDesiredSize(PrepassLayoutScaleMultiplier);
-
-			bUpdatingDesiredSize = false;
-		}
-
-		return DesiredSize.GetValue();
-	}
-	else
+	if (Child.FastPathProxyHandle.IsValid())
 	{
-		return DesiredSize.Get(FVector2D::ZeroVector);
+		SCOPED_NAMED_EVENT(SWidget_InvalidateChildRemovedFromTree, FColor::Red);
+		Child.UpdateFastPathVisibility(false, true);
 	}
 }
 
+FVector2D SWidget::GetDesiredSize() const
+{
+	return DesiredSize.Get(FVector2D::ZeroVector);
+}
+
 
 void SWidget::AssignParentWidget(TSharedPtr<SWidget> InParent)
 {
@@ -595,6 +608,8 @@ void SWidget::AssignParentWidget(TSharedPtr<SWidget> InParent)
 	ensureMsgf(InParent.IsValid(), TEXT("Are you trying to detatch the parent of a widget?  Use ConditionallyDetatchParentWidget()."));
 #endif
 
+	//@todo We should update inherited visibility and volatility here but currently we are relying on ChildOrder invalidation to do that for us
+
 	ParentWidgetPtr = InParent;
 #if WITH_ACCESSIBILITY
 	if (FSlateApplicationBase::IsInitialized())
@@ -604,7 +619,7 @@ void SWidget::AssignParentWidget(TSharedPtr<SWidget> InParent)
 #endif
 	if (InParent.IsValid())
 	{
-		InParent->Invalidate(EInvalidateWidget::Layout);
+		InParent->Invalidate(EInvalidateWidget::ChildOrder);
 	}
 }
 
@@ -627,9 +642,10 @@ bool SWidget::ConditionallyDetatchParentWidget(SWidget* InExpectedParent)
 
 		if (Parent.IsValid())
 		{
-			Parent->Invalidate(EInvalidateWidget::Layout);
+			Parent->Invalidate(EInvalidateWidget::ChildOrder);
 		}
 
+		InvalidateChildRemovedFromTree(*this);
 		return true;
 	}
 
@@ -653,42 +669,145 @@ void SWidget::LayoutChanged(EInvalidateWidget InvalidateReason)
 
 void SWidget::ChildLayoutChanged(EInvalidateWidget InvalidateReason)
 {
-	if (!bNeedsDesiredSize || EnumHasAllFlags(InvalidateReason, EInvalidateWidget::Visibility) )
+	if (!bNeedsDesiredSize || InvalidateReason == EInvalidateWidget::Visibility )
 	{
 		LayoutChanged(InvalidateReason);
 	}
 }
 
-void SWidget::CacheDesiredSize(float InLayoutScaleMultiplier)
+void SWidget::AssignIndicesToChildren(FSlateInvalidationRoot& Root, int32 ParentIndex, TArray<FWidgetProxy, TMemStackAllocator<>>& FastPathList, bool bParentVisible, bool bParentVolatile)
 {
-#if SLATE_VERBOSE_NAMED_EVENTS
-	SCOPED_NAMED_EVENT(SWidget_CacheDesiredSize, FColor::Red);
+	FWidgetProxy MyProxy(this);
+	MyProxy.Index = FastPathList.Num();
+	MyProxy.ParentIndex = ParentIndex;
+	MyProxy.Visibility = GetVisibility();
+
+	check(ParentIndex != MyProxy.Index);
+
+	// If this method is being called, child order changed.  Initial visibility and volatility needs to be propagated
+	// Update visibility
+	const bool bParentAndSelfVisible = bParentVisible && MyProxy.Visibility.IsVisible();
+	const bool bWasInvisible = bInvisibleDueToParentOrSelfVisibility;
+	bInvisibleDueToParentOrSelfVisibility = !bParentAndSelfVisible;
+	MyProxy.bInvisibleDueToParentOrSelfVisibility = bInvisibleDueToParentOrSelfVisibility;
+
+	// Update volatility
+	bInheritedVolatility = bParentVolatile;
+
+	FastPathProxyHandle = FWidgetProxyHandle(Root, MyProxy.Index);
+
+
+	if (bInvisibleDueToParentOrSelfVisibility&& PersistentState.CachedElementListNode != nullptr)
+	{
+#if WITH_SLATE_DEBUGGING
+		check(PersistentState.CachedElementListNode->GetValue().Widget == this);
 #endif
-	// Cache this widget's desired size.
-	Advanced_SetDesiredSize(ComputeDesiredSize(InLayoutScaleMultiplier));
-}
+		Root.GetCachedElements().ResetCache(PersistentState.CachedElementListNode);
+	}
 
-void SWidget::CachePrepass(const TWeakPtr<ILayoutCache>& InLayoutCache)
-{
-	if ( bCanHaveChildren )
+	FastPathList.Add(MyProxy);
+
+	// Don't recur into children if we are at a different invalidation root(nested invalidation panels) than where we started and not at the root of the tree. Those children should belong to that roots tree.
+	if (!Advanced_IsInvalidationRoot() || ParentIndex == INDEX_NONE)
 	{
-		FChildren* MyChildren = this->GetChildren();
+		FChildren* MyChildren = GetAllChildren();
 		const int32 NumChildren = MyChildren->Num();
-		for ( int32 ChildIndex=0; ChildIndex < NumChildren; ++ChildIndex )
+
+		int32 NumChildrenValidForFastPath = 0;
+		for (int32 ChildIndex = 0; ChildIndex < NumChildren; ++ChildIndex)
 		{
+			// Because null widgets are a shared static widget they are not valid for the fast path and are treated as non-existent
 			const TSharedRef<SWidget>& Child = MyChildren->GetChildAt(ChildIndex);
-			if ( Child->GetVisibility().IsVisible() == false )
-			{
-				Child->LayoutCache = InLayoutCache;
-			}
-			else
+			if (Child != SNullWidget::NullWidget)
 			{
-				Child->CachePrepass(InLayoutCache);
+				++NumChildrenValidForFastPath;
+				Child->AssignIndicesToChildren(Root, MyProxy.Index, FastPathList, bParentAndSelfVisible, bParentVolatile || IsVolatile());
 			}
 		}
+
+		{
+			FWidgetProxy& MyProxyRef = FastPathList[MyProxy.Index];
+			MyProxyRef.NumChildren = NumChildrenValidForFastPath;
+			int32 LastIndex = FastPathList.Num() - 1;
+			MyProxyRef.LeafMostChildIndex = LastIndex != MyProxy.Index ? LastIndex : INDEX_NONE;
+		}
+
+	}
+}
+
+void SWidget::UpdateFastPathVisibility(bool bParentVisible, bool bWidgetRemoved)
+{
+	const EVisibility CurrentVisibility = GetVisibility();
+	const bool bParentAndSelfVisible = bParentVisible && CurrentVisibility.IsVisible();
+	const bool bWasInvisible = bInvisibleDueToParentOrSelfVisibility;
+	bInvisibleDueToParentOrSelfVisibility = !bParentAndSelfVisible;
+	const bool bVisibilityChanged = bWasInvisible != bInvisibleDueToParentOrSelfVisibility;
+
+	if (FastPathProxyHandle.IsValid())
+	{
+		FastPathProxyHandle.GetInvalidationRoot()->GetHittestGrid()->RemoveWidget(SharedThis(this));
+		FWidgetProxy& Proxy = FastPathProxyHandle.GetProxy();
+		Proxy.Visibility = CurrentVisibility;
+		Proxy.bInvisibleDueToParentOrSelfVisibility = bInvisibleDueToParentOrSelfVisibility;
+
+		if (bWidgetRemoved)
+		{
+			FastPathProxyHandle.GetInvalidationRoot()->RemoveWidgetFromFastPath(Proxy);
+		}
+		else if (PersistentState.CachedElementListNode != nullptr)
+		{
+			FastPathProxyHandle.GetInvalidationRoot()->GetCachedElements().ResetCache(PersistentState.CachedElementListNode);
+		}
+	}
+	else
+	{
+		ensure(FastPathProxyHandle.GetIndex() == INDEX_NONE);
+	}
+
+	FChildren* MyChildren = GetAllChildren();
+	const int32 NumChildren = MyChildren->Num();
+	for (int32 ChildIndex = 0; ChildIndex < NumChildren; ++ChildIndex)
+	{
+		const TSharedRef<SWidget>& Child = MyChildren->GetChildAt(ChildIndex);
+		Child->UpdateFastPathVisibility(bParentAndSelfVisible, bWidgetRemoved);
+	}
+}
+
+void SWidget::UpdateFastPathVolatility(bool bParentVolatile)
+{
+	bInheritedVolatility = bParentVolatile;
+
+	if (IsVolatile() && !IsVolatileIndirectly())
+	{
+		AddUpdateFlags(EWidgetUpdateFlags::NeedsVolatilePaint);
+	}
+	else
+	{
+		RemoveUpdateFlags(EWidgetUpdateFlags::NeedsVolatilePaint);
+	}
+
+	FChildren* MyChildren = GetAllChildren();
+	const int32 NumChildren = MyChildren->Num();
+	for (int32 ChildIndex = 0; ChildIndex < NumChildren; ++ChildIndex)
+	{
+		const TSharedRef<SWidget>& Child = MyChildren->GetChildAt(ChildIndex);
+		Child->UpdateFastPathVolatility(bParentVolatile || IsVolatile());
 	}
+
+
+}
+
+void SWidget::CacheDesiredSize(float InLayoutScaleMultiplier)
+{
+#if SLATE_VERBOSE_NAMED_EVENTS
+	SCOPED_NAMED_EVENT(SWidget_CacheDesiredSize, FColor::Red);
+#endif
+
+	// Cache this widget's desired size.
+	SetDesiredSize(ComputeDesiredSize(InLayoutScaleMultiplier));
 }
 
+
 bool SWidget::SupportsKeyboardFocus() const
 {
 	return false;
@@ -865,6 +984,21 @@ FSlateColor SWidget::GetForegroundColor() const
 	return NoColor;
 }
 
+const FGeometry& SWidget::GetCachedGeometry() const
+{
+	return GetTickSpaceGeometry();
+}
+
+const FGeometry& SWidget::GetTickSpaceGeometry() const
+{
+	return PersistentState.DesktopGeometry;
+}
+
+const FGeometry& SWidget::GetPaintSpaceGeometry() const
+{
+	return PersistentState.AllottedGeometry;
+}
+
 void SWidget::SetToolTipText(const TAttribute<FText>& ToolTipText)
 {
 	ToolTip = FSlateApplicationBase::Get().MakeToolTip(ToolTipText);
@@ -899,21 +1033,54 @@ bool SWidget::IsDirectlyHovered() const
 	return FSlateApplicationBase::Get().IsWidgetDirectlyHovered(SharedThis(this));
 }
 
+void SWidget::SetVisibility(TAttribute<EVisibility> InVisibility)
+{
+	if (!Visibility.IdenticalTo(InVisibility))
+	{
+		Visibility = InVisibility;
+
+		Invalidate(EInvalidateWidget::Visibility);
+	}
+}
+
 void SWidget::Invalidate(EInvalidateWidget InvalidateReason)
 {
 	SLATE_CROSS_THREAD_CHECK();
 
 	SCOPED_NAMED_EVENT_TEXT("SWidget::Invalidate", FColor::Orange);
-
 	const bool bWasVolatile = IsVolatileIndirectly() || IsVolatile();
+
 	const bool bVolatilityChanged = EnumHasAnyFlags(InvalidateReason, EInvalidateWidget::Volatility) ? Advanced_InvalidateVolatility() : false;
 
-	if (bWasVolatile == false || bVolatilityChanged)
+	if (EnumHasAnyFlags(InvalidateReason, EInvalidateWidget::ChildOrder))
 	{
-		Advanced_ForceInvalidateLayout();
+		InvalidatePrepass();
 	}
 
-	LayoutChanged(InvalidateReason);
+	if(FastPathProxyHandle.IsValid())
+	{
+		// Current thinking is that visibility and volatility should be updated right away, not during fast path invalidation processing next frame
+		if (EnumHasAnyFlags(InvalidateReason, EInvalidateWidget::Visibility))
+		{
+			SCOPED_NAMED_EVENT(SWidget_UpdateFastPathVisibility, FColor::Red);
+			TSharedPtr<SWidget> ParentWidget = GetParentWidget();
+
+			UpdateFastPathVisibility(ParentWidget.IsValid() ? !ParentWidget->bInvisibleDueToParentOrSelfVisibility : false, false);
+		}
+
+		if (bVolatilityChanged)
+		{
+			SCOPED_NAMED_EVENT(SWidget_UpdateFastPathVolatility, FColor::Red);
+
+			TSharedPtr<SWidget> ParentWidget = GetParentWidget();
+
+			UpdateFastPathVolatility(ParentWidget.IsValid() ? ParentWidget->IsVolatile() || ParentWidget->IsVolatileIndirectly() : false);
+
+			ensure(!IsVolatile() || EnumHasAnyFlags(UpdateFlags, EWidgetUpdateFlags::NeedsVolatilePaint));
+		}
+
+		FastPathProxyHandle.MarkWidgetDirty(InvalidateReason);
+	}
 }
 
 void SWidget::SetCursor( const TAttribute< TOptional<EMouseCursor::Type> >& InCursor )
@@ -989,71 +1156,51 @@ FSlateRect SWidget::CalculateCullingAndClippingRules(const FGeometry& AllottedGe
 
 int32 SWidget::Paint(const FPaintArgs& Args, const FGeometry& AllottedGeometry, const FSlateRect& MyCullingRect, FSlateWindowElementList& OutDrawElements, int32 LayerId, const FWidgetStyle& InWidgetStyle, bool bParentEnabled) const
 {
-#if WITH_VERY_VERBOSE_SLATE_STATS
-	FScopeCycleCounterSWidget WidgetScope(this);
-#endif
-
-	INC_DWORD_STAT(STAT_SlateNumPaintedWidgets);
-
 	// TODO, Maybe we should just make Paint non-const and keep OnPaint const.
-	SWidget* MutableThis = const_cast<SWidget*>(this);
+	TSharedRef<SWidget> MutableThis = ConstCastSharedRef<SWidget>(AsShared());
 
-	if (GSlateLayoutCaching)
-	{
-		MutableThis->SlatePrepass(AllottedGeometry.Scale);
-	}
+	INC_DWORD_STAT(STAT_SlateNumPaintedWidgets);
 
-	// Save the current layout cache we're associated with (if any)
-	LayoutCache = Args.GetLayoutCache();
+	const SWidget* PaintParent = Args.GetPaintParent();
+	//if (GSlateEnableGlobalInvalidation)
+	//{
+	//	bInheritedVolatility = PaintParent ? (PaintParent->IsVolatileIndirectly() || PaintParent->IsVolatile()) : false;
+	//}
 
-	// Record if we're part of a volatility pass, this is critical for ensuring we don't report a child
-	// of a volatile widget as non-volatile, causing the invalidation panel to do work that's not required.
-	//
-	// Note: We only do this if we're not also caching.  The retainer panel takes advantage of the fact that
-	// it can both send down it's caching & it's a volatile pass, implying everyone should render, everyone
-	// is getting cached.  So we don't want volatile widgets to wait to be drawn later, they won't get another
-	// chance.
-	bInheritedVolatility = Args.IsVolatilityPass() && !Args.IsCaching();
 
 	// If this widget clips to its bounds, then generate a new clipping rect representing the intersection of the bounding
 	// rectangle of the widget's geometry, and the current clipping rectangle.
 	bool bClipToBounds, bAlwaysClip, bIntersectClipBounds;
+
 	FSlateRect CullingBounds = CalculateCullingAndClippingRules(AllottedGeometry, MyCullingRect, bClipToBounds, bAlwaysClip, bIntersectClipBounds);
 
 	FWidgetStyle ContentWidgetStyle = FWidgetStyle(InWidgetStyle)
 		.BlendOpacity(RenderOpacity);
 
-	// If this paint pass is to cache off our geometry, but we're a volatile widget,
-	// record this widget as volatile in the draw elements so that we get our own tick/paint 
-	// pass later when the layout cache draws.
-	if (IsVolatile() && Args.IsCaching() && !Args.IsVolatilityPass())
-	{
-		const int32 VolatileLayerId = LayerId + 1;
-		OutDrawElements.QueueVolatilePainting(
-			FSlateWindowElementList::FVolatilePaint(SharedThis(this), Args, AllottedGeometry, CullingBounds, OutDrawElements.GetClippingState(), VolatileLayerId, ContentWidgetStyle, bParentEnabled));
-
-		return VolatileLayerId;
-	}
-
 	// Cache the geometry for tick to allow external users to get the last geometry that was used,
 	// or would have been used to tick the Widget.
-	CachedGeometry = AllottedGeometry;
-	CachedGeometry.AppendTransform(FSlateLayoutTransform(Args.GetWindowToDesktopTransform()));
+	FGeometry DesktopSpaceGeometry = AllottedGeometry;
+	DesktopSpaceGeometry.AppendTransform(FSlateLayoutTransform(Args.GetWindowToDesktopTransform()));
 
-	MutableThis->ExecuteActiveTimers(Args.GetCurrentTime(), Args.GetDeltaTime());
+	if (HasAnyUpdateFlags(EWidgetUpdateFlags::NeedsActiveTimerUpdate))
+	{
+		SCOPE_CYCLE_COUNTER(STAT_SlateExecuteActiveTimers);
+		MutableThis->ExecuteActiveTimers(Args.GetCurrentTime(), Args.GetDeltaTime());
+	}
 
 	if (HasAnyUpdateFlags(EWidgetUpdateFlags::NeedsTick))
 	{
 		INC_DWORD_STAT(STAT_SlateNumTickedWidgets);
 
 		SCOPE_CYCLE_COUNTER(STAT_SlateTickWidgets);
-		MutableThis->Tick(CachedGeometry, Args.GetCurrentTime(), Args.GetDeltaTime());
+		MutableThis->Tick(DesktopSpaceGeometry, Args.GetCurrentTime(), Args.GetDeltaTime());
 	}
 
-	// Record hit test geometry, but only if we're not caching.
-	const FPaintArgs UpdatedArgs = Args.RecordHittestGeometry(this, AllottedGeometry, LayerId);
+	// the rule our parent has set for us
+	const bool bInheritedHittestability = Args.GetInheritedHittestability();
+	const bool bOutgoingHittestability = bInheritedHittestability && GetVisibility().AreChildrenHitTestVisible();
 
-#if !(UE_BUILD_SHIPPING || UE_BUILD_TEST)
+#if WITH_SLATE_DEBUGGING
 	if (GDebugCulling)
 	{
 		// When we're debugging culling, don't actually clip, we'll just pretend to, so we can see the effects of
@@ -1062,20 +1209,58 @@ int32 SWidget::Paint(const FPaintArgs& Args, const FGeometry& AllottedGeometry,
 	}
 #endif
 
-	if ( bClipToBounds )
+	SWidget* PaintParentPtr = const_cast<SWidget*>(Args.GetPaintParent());
+	ensure(PaintParentPtr != this);
+	if (PaintParentPtr)
+	{
+		PersistentState.PaintParent = PaintParentPtr->AsShared();
+	}
+	else
+	{
+		PaintParentPtr = nullptr;
+	}
+	
+	// @todo This should not do this copy if the clipping state is unset
+	PersistentState.InitialClipState = OutDrawElements.GetClippingState();
+	PersistentState.LayerId = LayerId;
+	PersistentState.bParentEnabled = bParentEnabled;
+	PersistentState.bInheritedHittestability = bInheritedHittestability;
+	PersistentState.AllottedGeometry = AllottedGeometry;
+	PersistentState.DesktopGeometry = DesktopSpaceGeometry;
+	PersistentState.WidgetStyle = InWidgetStyle;
+	PersistentState.CullingBounds = MyCullingRect;
+	PersistentState.IncomingFlowDirection = GSlateFlowDirection;
+
+	FPaintArgs UpdatedArgs = Args.WithNewParent(this);
+	UpdatedArgs.SetInheritedHittestability(bOutgoingHittestability);
+
+
+	// test ensure that we are not the last thing holding this widget together
+	ensure(!MutableThis.IsUnique());
+
+
+	if (!FastPathProxyHandle.IsValid() && PersistentState.CachedElementListNode != nullptr)
+	{
+		ensure(!bInvisibleDueToParentOrSelfVisibility);
+	}
+
+	OutDrawElements.PushPaintingWidget(*this, LayerId, PersistentState.CachedElementListNode);
+
+	if (bOutgoingHittestability)
+	{
+		Args.GetHittestGrid().AddWidget(MutableThis, 0, LayerId, FastPathProxyHandle.GetIndex());
+	}
+
+	if (bClipToBounds)
 	{
+		// This sets up the clip state for any children NOT myself
 		FSlateClippingZone ClippingZone(AllottedGeometry);
 		ClippingZone.SetShouldIntersectParent(bIntersectClipBounds);
 		ClippingZone.SetAlwaysClip(bAlwaysClip);
 		OutDrawElements.PushClip(ClippingZone);
-
-		// The hit test grid records things in desktop space, so we use the tick geometry instead of the paint geometry.
-		FSlateClippingZone DesktopClippingZone(CachedGeometry);
-		DesktopClippingZone.SetShouldIntersectParent(bIntersectClipBounds);
-		DesktopClippingZone.SetAlwaysClip(bAlwaysClip);
-		Args.GetGrid().PushClip(DesktopClippingZone);
 	}
 
+
 #if WITH_SLATE_DEBUGGING
 	FSlateDebugging::BeginWidgetPaint.Broadcast(this, UpdatedArgs, AllottedGeometry, CullingBounds, OutDrawElements, LayerId);
 #endif
@@ -1086,13 +1271,14 @@ int32 SWidget::Paint(const FPaintArgs& Args, const FGeometry& AllottedGeometry,
 	TGuardValue<EFlowDirection> FlowGuard(GSlateFlowDirection, ComputeFlowDirection());
 	
 	// Paint the geometry of this widget.
-	int32 NewLayerID = OnPaint(UpdatedArgs, AllottedGeometry, CullingBounds, OutDrawElements, LayerId, ContentWidgetStyle, bParentEnabled);
+	int32 NewLayerId = OnPaint(UpdatedArgs, AllottedGeometry, CullingBounds, OutDrawElements, LayerId, ContentWidgetStyle, bParentEnabled);
+
+	// Just repainted
+	MutableThis->RemoveUpdateFlags(EWidgetUpdateFlags::NeedsRepaint);
 
 #if WITH_SLATE_DEBUGGING
-	FSlateDebugging::EndWidgetPaint.Broadcast(this, OutDrawElements, NewLayerID);
-#endif
+	FSlateDebugging::EndWidgetPaint.Broadcast(this, OutDrawElements, NewLayerId);
 
-#if !(UE_BUILD_SHIPPING || UE_BUILD_TEST)
 	if (GShowClipping && bClipToBounds)
 	{
 		FSlateClippingZone ClippingZone(AllottedGeometry);
@@ -1107,7 +1293,7 @@ int32 SWidget::Paint(const FPaintArgs& Args, const FGeometry& AllottedGeometry,
 		const bool bAntiAlias = true;
 		FSlateDrawElement::MakeLines(
 			OutDrawElements,
-			NewLayerID,
+			NewLayerId,
 			FPaintGeometry(),
 			Points,
 			ESlateDrawEffect::None,
@@ -1115,17 +1301,17 @@ int32 SWidget::Paint(const FPaintArgs& Args, const FGeometry& AllottedGeometry,
 			bAntiAlias,
 			2.0f);
 	}
-#endif
+#endif // WITH_SLATE_DEBUGGING
 
-	if ( bClipToBounds )
+	if (bClipToBounds)
 	{
 		OutDrawElements.PopClip();
-		Args.GetGrid().PopClip();
 	}
 
+
 #if PLATFORM_UI_NEEDS_FOCUS_OUTLINES
 	// Check if we need to show the keyboard focus ring, this is only necessary if the widget could be focused.
-	if ( bCanSupportFocus && SupportsKeyboardFocus() )
+	if (bCanSupportFocus && SupportsKeyboardFocus())
 	{
 		bool bShowUserFocus = FSlateApplicationBase::Get().ShowUserFocus(SharedThis(this));
 		if (bShowUserFocus)
@@ -1136,7 +1322,7 @@ int32 SWidget::Paint(const FPaintArgs& Args, const FGeometry& AllottedGeometry,
 			{
 				FSlateDrawElement::MakeBox(
 					OutDrawElements,
-					NewLayerID,
+					NewLayerId,
 					AllottedGeometry.ToPaintGeometry(),
 					BrushResource,
 					ESlateDrawEffect::None,
@@ -1147,12 +1333,16 @@ int32 SWidget::Paint(const FPaintArgs& Args, const FGeometry& AllottedGeometry,
 	}
 #endif
 
-	if ( OutDrawElements.ShouldResolveDeferred() )
+	FSlateCachedElementListNode* NewCacheNode = OutDrawElements.PopPaintingWidget();
+	if (OutDrawElements.ShouldResolveDeferred())
 	{
-		NewLayerID = OutDrawElements.PaintDeferred(NewLayerID, MyCullingRect);
+		NewLayerId = OutDrawElements.PaintDeferred(NewLayerId, MyCullingRect);
 	}
 
-	return NewLayerID;
+	MutableThis->UpdateWidgetProxy(NewLayerId, NewCacheNode);
+
+	return NewLayerId;
+
 }
 
 float SWidget::GetRelativeLayoutScale(const FSlotBase& Child, float LayoutScaleMultiplier) const
@@ -1168,6 +1358,43 @@ void SWidget::ArrangeChildren(const FGeometry& AllottedGeometry, FArrangedChildr
 	OnArrangeChildren(AllottedGeometry, ArrangedChildren);
 }
 
+void SWidget::Prepass_Internal(float InLayoutScaleMultiplier)
+{
+	PrepassLayoutScaleMultiplier = InLayoutScaleMultiplier;
+
+	bool bShouldPrepassChildren = true;
+	if (bHasCustomPrepass)
+	{
+		bShouldPrepassChildren = CustomPrepass(InLayoutScaleMultiplier);
+	}
+
+	if (bCanHaveChildren && bShouldPrepassChildren)
+	{
+		// Cache child desired sizes first. This widget's desired size is
+		// a function of its children's sizes.
+		FChildren* MyChildren = this->GetChildren();
+		const int32 NumChildren = MyChildren->Num();
+		for (int32 ChildIndex = 0; ChildIndex < MyChildren->Num(); ++ChildIndex)
+		{
+			const TSharedRef<SWidget>& Child = MyChildren->GetChildAt(ChildIndex);
+
+			if (Child->Visibility.Get() != EVisibility::Collapsed)
+			{
+				const float ChildLayoutScaleMultiplier = GetRelativeLayoutScale(MyChildren->GetSlotAt(ChildIndex), InLayoutScaleMultiplier);
+				// Recur: Descend down the widget tree.
+				Child->Prepass_Internal(InLayoutScaleMultiplier * ChildLayoutScaleMultiplier);
+			}
+		}
+		ensure(NumChildren == MyChildren->Num());
+	}
+
+	{
+		// Cache this widget's desired size.
+		CacheDesiredSize(PrepassLayoutScaleMultiplier);
+		bNeedsPrepass = false;
+	}
+}
+
 TSharedRef<FActiveTimerHandle> SWidget::RegisterActiveTimer(float TickPeriod, FWidgetActiveTimerDelegate TickFunction)
 {
 	TSharedRef<FActiveTimerHandle> ActiveTimerHandle = MakeShareable(new FActiveTimerHandle(TickPeriod, TickFunction, FSlateApplicationBase::Get().GetCurrentTime() + TickPeriod));
@@ -1418,6 +1645,8 @@ bool SWidget::CanChildrenBeAccessible() const
 
 bool SWidget::IsChildWidgetCulled(const FSlateRect& MyCullingRect, const FArrangedWidget& ArrangedChild) const
 {
+	QUICK_SCOPE_CYCLE_COUNTER(Slate_IsChildWidgetCulled);
+
 	// We add some slack fill to the culling rect to deal with the common occurrence
 	// of widgets being larger than their root level widget is.  Happens when nested child widgets
 	// inflate their rendering bounds to render beyond their parent (the child of this panel doing the culling), 
diff --git a/Engine/Source/Runtime/SlateCore/Private/Widgets/SWindow.cpp b/Engine/Source/Runtime/SlateCore/Private/Widgets/SWindow.cpp
--- a/Engine/Source/Runtime/SlateCore/Private/Widgets/SWindow.cpp
+++ b/Engine/Source/Runtime/SlateCore/Private/Widgets/SWindow.cpp
@@ -6,6 +6,12 @@
 #include "Layout/WidgetPath.h"
 #include "Input/HittestGrid.h"
 #include "HAL/PlatformApplicationMisc.h"
+#include "Widgets/Images/SImage.h"
+
+#if WITH_SLATE_DEBUGGING
+#include "ProfilingDebugging/CsvProfiler.h"
+#endif
+
 #if WITH_ACCESSIBILITY
 #include "Widgets/Accessibility/SlateCoreAccessibleWidgets.h"
 #endif
@@ -390,12 +396,12 @@ void SWindow::Construct(const FArguments& InArgs)
 	this->InitialDesiredScreenPosition = WindowPosition;
 	this->InitialDesiredSize = WindowSize;
 
+	FSlateApplicationBase::Get().OnGlobalInvalidationToggled().AddSP(this, &SWindow::OnGlobalInvalidationToggled);
+
+
 	// Resize adds extra borders / title bar if necessary, but this is already taken into account in WindowSize, so subtract them again first
 	Resize(WindowSize - DeltaSize);
 
-	// Window visibility is currently driven by whether the window is interactive.
-	this->Visibility = TAttribute<EVisibility>::Create( TAttribute<EVisibility>::FGetter::CreateRaw(this, &SWindow::GetWindowVisibility) );
-
 	this->ConstructWindowInternals();
 	this->SetContent( InArgs._Content.Widget );
 }
@@ -496,8 +502,6 @@ EHorizontalAlignment SWindow::GetTitleAlignment()
 	return TitleContentAlignment;
 }
 
-
-
 void SWindow::ConstructWindowInternals()
 {
 	ForegroundColor = FCoreStyle::Get().GetSlateColor("DefaultForeground");
@@ -507,6 +511,8 @@ void SWindow::ConstructWindowInternals()
 		SNew( SVerticalBox )
 		.Visibility( EVisibility::SelfHitTestInvisible );
 
+	TSharedRef<SWidget> TitleBarWidget = FSlateApplicationBase::Get().MakeWindowTitleBar(SharedThis(this), nullptr, GetTitleAlignment(), TitleBar);
+
 	if (bCreateTitleBar)
 	{
 		// @todo mainframe: Should be measured from actual title bar content widgets.  Don't use a hard-coded size!
@@ -515,7 +521,7 @@ void SWindow::ConstructWindowInternals()
 		MainWindowArea->AddSlot()
 		.AutoHeight()
 		[
-			MakeWindowTitleBar(SharedThis(this), nullptr, GetTitleAlignment())
+			TitleBarWidget
 		];
 	}
 	else
@@ -523,6 +529,8 @@ void SWindow::ConstructWindowInternals()
 		TitleBarSize = 0;
 	}
 
+	UpdateWindowContentVisibility();
+
 	// create window content slot
 	MainWindowArea->AddSlot()
 		.FillHeight(1.0f)
@@ -534,51 +542,52 @@ void SWindow::ConstructWindowInternals()
 	// create window
 	if (Type != EWindowType::ToolTip && Type != EWindowType::CursorDecorator && !bIsPopupWindow && !bHasOSWindowBorder)
 	{
-		TAttribute<EVisibility> WindowContentVisibility(this, &SWindow::GetWindowContentVisibility);
-		TAttribute<const FSlateBrush*> WindowBackgroundAttr(this, &SWindow::GetWindowBackground);
-		TAttribute<FSlateColor> WindowBackgroundColorAttr(this, &SWindow::GetWindowBackgroundColor);
-		TAttribute<const FSlateBrush*> WindowOutlineAttr(this, &SWindow::GetWindowOutline);
-		TAttribute<FSlateColor> WindowOutlineColorAttr(this, &SWindow::GetWindowOutlineColor);
+		WindowBackgroundImage =
+			FSlateApplicationBase::Get().MakeImage(
+				WindowBackground,
+				Style->BackgroundColor,
+				WindowContentVisibility
+			);
+
+		WindowBorder =
+			FSlateApplicationBase::Get().MakeImage(
+				&Style->BorderBrush,
+				FLinearColor::White,
+				WindowContentVisibility
+			);
+
+		WindowOutline = FSlateApplicationBase::Get().MakeImage(
+				&Style->OutlineBrush,
+				Style->OutlineColor,
+				WindowContentVisibility
+			);
 
 		this->ChildSlot
 		[
 			SAssignNew(WindowOverlay, SOverlay)
 			.Visibility(EVisibility::SelfHitTestInvisible)
-
 			// window background
 			+ SOverlay::Slot()
 			[
-				FSlateApplicationBase::Get().MakeImage(
-					WindowBackgroundAttr,
-					WindowBackgroundColorAttr,
-					WindowContentVisibility
-				)
+				WindowBackgroundImage.ToSharedRef()
 			]
 
 			// window border
 			+ SOverlay::Slot()
 			[
-				FSlateApplicationBase::Get().MakeImage(
-					&Style->BorderBrush,
-					FLinearColor::White,
-					WindowContentVisibility
-				)
+				WindowBorder.ToSharedRef()
 			]
 
 			// window outline
 			+ SOverlay::Slot()
 			[
-				FSlateApplicationBase::Get().MakeImage(
-					WindowOutlineAttr,
-					WindowOutlineColorAttr,
-					WindowContentVisibility
-				)
+				WindowOutline.ToSharedRef()
 			]
 
 			// main area
 			+ SOverlay::Slot()
 			[
-				SNew(SVerticalBox)
+				SAssignNew(ContentAreaVBox, SVerticalBox)
 				.Visibility(WindowContentVisibility)
 
 				+ SVerticalBox::Slot()
@@ -648,9 +657,9 @@ bool SWindow::HasActiveParent() const
 	return false;
 }
 
-TSharedRef<FHittestGrid> SWindow::GetHittestGrid()
+FHittestGrid& SWindow::GetHittestGrid()
 {
-	return HittestGrid;
+	return *HittestGrid;
 }
 
 FWindowSizeLimits SWindow::GetSizeLimits() const
@@ -658,6 +667,17 @@ FWindowSizeLimits SWindow::GetSizeLimits() const
 	return SizeLimits;
 }
 
+void SWindow::SetAllowFastUpdate(bool bInAllowFastUpdate)
+{
+	bAllowFastUpdate = bInAllowFastUpdate;
+
+	if (bAllowFastUpdate)
+	{
+		InvalidateChildOrder();
+	}
+}
+
+
 void SWindow::Tick( const FGeometry& AllottedGeometry, const double InCurrentTime, const float InDeltaTime )
 {
 	if( Morpher.bIsActive )
@@ -972,7 +992,12 @@ void SWindow::SetCachedSize( FVector2D NewSize )
 	{
 		NativeWindow->AdjustCachedSize( NewSize );
 	}
-	Size = NewSize;
+
+	if(Size != NewSize)
+	{
+		Size = NewSize;
+		InvalidateRoot();
+	}
 }
 
 bool SWindow::IsMorphing() const
@@ -1018,16 +1043,35 @@ void SWindow::StartMorph()
 	}
 }
 
-const FSlateBrush* SWindow::GetWindowBackground() const
+
+bool SWindow::CustomPrepass(float LayoutScaleMultiplier)
+{
+	if (bAllowFastUpdate && GSlateEnableGlobalInvalidation)
+	{
+		ProcessInvalidation();
+		return NeedsPrepass();
+	}
+	else
+	{
+		return true;
+	}
+}
+
+/*
+void SWindow::Advanced_InvalidateRoot()
 {
-	return WindowBackground;
+	InvalidateRoot();
 }
+*/
 
+/*
 FSlateColor SWindow::GetWindowBackgroundColor() const
 {
 	return Style->BackgroundColor;
 }
+*/
 
+/*
 const FSlateBrush* SWindow::GetWindowOutline() const
 {
 	return &Style->OutlineBrush;
@@ -1037,6 +1081,7 @@ FSlateColor SWindow::GetWindowOutlineColor() const
 {
 	return Style->OutlineColor;
 }
+*/
 
 EVisibility SWindow::GetWindowVisibility() const
 {
@@ -1173,7 +1218,7 @@ void SWindow::SetContent( TSharedRef<SWidget> InContent )
 	{
 		this->ContentSlot->operator[]( InContent );
 	}
-	Invalidate(EInvalidateWidget::LayoutAndVolatility);
+	Invalidate(EInvalidateWidget::ChildOrder);
 }
 
 TSharedRef<const SWidget> SWindow::GetContent() const
@@ -1645,6 +1690,9 @@ bool SWindow::SupportsKeyboardFocus() const
 	return Type != EWindowType::ToolTip && Type != EWindowType::CursorDecorator;
 }
 
+
+
+
 FReply SWindow::OnFocusReceived(const FGeometry& MyGeometry, const FFocusEvent& InFocusEvent)
 {
 	return FReply::Handled();
@@ -1695,6 +1743,18 @@ FVector2D SWindow::ComputeDesiredSize(float LayoutScaleMultiplier) const
 	return SCompoundWidget::ComputeDesiredSize(LayoutScaleMultiplier) * LayoutScaleMultiplier;
 }
 
+bool SWindow::ComputeVolatility() const
+{
+	// If the entire window is volatile in fast path that defeats the whole purpose.
+	return bAllowFastUpdate ? false : SWidget::ComputeVolatility();
+}
+
+void SWindow::OnGlobalInvalidationToggled(bool bGlobalInvalidationEnabled)
+{
+	InvalidateRoot();
+	UE_LOG(LogSlate, Log, TEXT("Toggling fast path.  New State: %d"), bGlobalInvalidationEnabled);
+}
+
 const TArray< TSharedRef<SWindow> >& SWindow::GetChildWindows() const
 {
 	return ChildWindows;
@@ -1716,7 +1776,7 @@ void SWindow::AddChildWindow( const TSharedRef<SWindow>& ChildWindow )
 	}
 
 	ChildWindow->ParentWindowPtr = SharedThis(this);
-	ChildWindow->WindowBackground = &Style->ChildBackgroundBrush;
+	ChildWindow->SetWindowBackground(&Style->ChildBackgroundBrush);
 
 	FSlateApplicationBase::Get().ArrangeWindowToFrontVirtual( ChildWindows, ChildWindow );
 }
@@ -1748,7 +1808,7 @@ bool SWindow::RemoveDescendantWindow( const TSharedRef<SWindow>& DescendantToRem
 		if ( ChildWindow->RemoveDescendantWindow( DescendantToRemove ))
 		{
 			// Reset to the non-child background style
-			ChildWindow->WindowBackground = &Style->BackgroundBrush;
+			ChildWindow->SetWindowBackground(&Style->BackgroundBrush);
 			return true;
 		}
 	}
@@ -1901,6 +1961,7 @@ SWindow::SWindow()
 	, bIsMirrorWindow( false )
 	, bShouldPreserveAspectRatio( false )
 	, bManualManageDPI( false )
+	, bAllowFastUpdate( false )
 	, WindowActivationPolicy( EWindowActivationPolicy::Always )
 	, InitialDesiredScreenPosition( FVector2D::ZeroVector )
 	, InitialDesiredSize( FVector2D::ZeroVector )
@@ -1912,41 +1973,104 @@ SWindow::SWindow()
 	, ContentSlot(nullptr)
 	, Style( &FCoreStyle::Get().GetWidgetStyle<FWindowStyle>("Window") )
 	, WindowBackground( &Style->BackgroundBrush )
-	, HittestGrid( MakeShareable(new FHittestGrid()) )
+	, HittestGrid(MakeUnique<FHittestGrid>())
 	, bShouldShowWindowContentDuringOverlay( false )
 	, ExpectedMaxWidth( INDEX_NONE )
 	, ExpectedMaxHeight( INDEX_NONE )
 	, TitleBar()
 	, bIsDrawingEnabled( true )
+	
 {
+	bHasCustomPrepass = true;
+	SetInvalidationRootWidget(*this);
+	SetInvalidationRootHittestGrid(*HittestGrid);
+
 #if WITH_ACCESSIBILITY
 	AccessibleData = FAccessibleWidgetData(EAccessibleBehavior::Auto);
 #endif
 }
 
+SWindow::~SWindow()
+{
+	check(IsInGameThread());
+}
+
 
-int32 SWindow::PaintWindow( const FPaintArgs& Args, const FGeometry& AllottedGeometry, const FSlateRect& MyCullingRect, FSlateWindowElementList& OutDrawElements, int32 LayerId, const FWidgetStyle& InWidgetStyle, bool bParentEnabled ) const
+int32 SWindow::PaintSlowPath(const FSlateInvalidationContext& Context)
 {
-	// Create initial culture specific layout direction
-	EFlowDirection NewFlowDirection = GSlateFlowDirection;
-	if (GetFlowDirectionPreference() == EFlowDirectionPreference::Inherit)
+	HittestGrid->Clear();
+
+	const FSlateRect WindowCullingBounds = GetClippingRectangleInWindow();
+	const int32 LayerId = 0;
+	const FGeometry WindowGeometry = GetWindowGeometryInWindow();
+
+	int32 MaxLayerId = 0;
+
+	//OutDrawElements.PushBatchPriortyGroup(*this);
 	{
-		NewFlowDirection = GSlateFlowDirectionShouldFollowCultureByDefault ? FLayoutLocalization::GetLocalizedLayoutDirection() : EFlowDirection::LeftToRight;
+		
+		MaxLayerId = Paint(*Context.PaintArgs, WindowGeometry, WindowCullingBounds, *Context.WindowElementList, LayerId, Context.WidgetStyle, Context.bParentEnabled);
 	}
 
-	TGuardValue<EFlowDirection> FlowGuard(GSlateFlowDirection, NewFlowDirection);
+	//OutDrawElements.PopBatchPriortyGroup();
+
 
-	LayerId = Paint(Args, AllottedGeometry, MyCullingRect, OutDrawElements, LayerId, InWidgetStyle, bParentEnabled);
 
-	return LayerId;
+	return MaxLayerId;
 }
 
-int32 SWindow::OnPaint(const FPaintArgs& Args, const FGeometry& AllottedGeometry, const FSlateRect& MyCullingRect, FSlateWindowElementList& OutDrawElements, int32 LayerId, const FWidgetStyle& InWidgetStyle, bool bParentEnabled) const
+int32 SWindow::PaintWindow( double CurrentTime, float DeltaTime, FSlateWindowElementList& OutDrawElements, const FWidgetStyle& InWidgetStyle, bool bParentEnabled )
 {
 	OutDrawElements.BeginDeferredGroup();
-	int32 MaxLayer = SCompoundWidget::OnPaint(Args, AllottedGeometry, MyCullingRect, OutDrawElements, LayerId, InWidgetStyle, bParentEnabled);
+
+	const bool HittestCleared = HittestGrid->SetHittestArea(GetPositionInScreen(), GetViewportSize());
+
+
+	FPaintArgs PaintArgs(nullptr, GetHittestGrid(), GetPositionInScreen(), CurrentTime, DeltaTime);
+
+	FSlateInvalidationContext Context(OutDrawElements, InWidgetStyle);
+	Context.bParentEnabled = bParentEnabled;
+	// Fast path at the window level should only be enabled if global invalidation is allowed
+	Context.bAllowFastPathUpdate = bAllowFastUpdate && GSlateEnableGlobalInvalidation;
+	Context.LayoutScaleMultiplier = FSlateApplicationBase::Get().GetApplicationScale() * GetNativeWindow()->GetDPIScaleFactor();
+	Context.PaintArgs = &PaintArgs;
+	Context.IncomingLayerId = 0;
+	Context.CullingRect = GetClippingRectangleInWindow();
+
+	// Always set the window geometry and visibility
+	PersistentState.AllottedGeometry = GetWindowGeometryInWindow();
+	PersistentState.CullingBounds = GetClippingRectangleInWindow();
+	if (!Visibility.IsBound())
+	{
+		SetVisibility(GetWindowVisibility());
+	}
+
+
+	FSlateInvalidationResult Result = PaintInvalidationRoot(Context);
+
+#if WITH_SLATE_DEBUGGING
+	if (GSlateHitTestGridDebugging)
+	{
+		const FGeometry& WindowGeometry = GetWindowGeometryInWindow();
+		HittestGrid->DisplayGrid(INT_MAX, WindowGeometry, OutDrawElements);
+		//HittestGrid->LogGrid
+	}
+#endif
+
 	OutDrawElements.EndDeferredGroup();
 
+	if (OutDrawElements.ShouldResolveDeferred())
+	{
+		Result.MaxLayerIdPainted = OutDrawElements.PaintDeferred(Result.MaxLayerIdPainted, Context.CullingRect);
+	}
+
+	return Result.MaxLayerIdPainted;
+
+}
+
+int32 SWindow::OnPaint(const FPaintArgs& Args, const FGeometry& AllottedGeometry, const FSlateRect& MyCullingRect, FSlateWindowElementList& OutDrawElements, int32 LayerId, const FWidgetStyle& InWidgetStyle, bool bParentEnabled) const
+{
+	int32 MaxLayer = SCompoundWidget::OnPaint(Args, AllottedGeometry, MyCullingRect, OutDrawElements, LayerId, InWidgetStyle, bParentEnabled);
 	return MaxLayer;
 }
 
@@ -1983,6 +2107,8 @@ void SWindow::SetFullWindowOverlayContent(TSharedPtr<SWidget> InContent)
 			InContent.ToSharedRef()
 		];
 	}
+
+	UpdateWindowContentVisibility();
 }
 
 /** Toggle window between fullscreen and normal mode */
@@ -2034,11 +2160,13 @@ bool SWindow::HasFullWindowOverlayContent() const
 void SWindow::BeginFullWindowOverlayTransition()
 {
 	bShouldShowWindowContentDuringOverlay = true;
+	UpdateWindowContentVisibility();
 }
 
 void SWindow::EndFullWindowOverlayTransition()
 {
 	bShouldShowWindowContentDuringOverlay = false;
+	UpdateWindowContentVisibility();
 }
 
 void SWindow::SetNativeWindowButtonsVisibility(bool bVisible)
@@ -2049,19 +2177,48 @@ void SWindow::SetNativeWindowButtonsVisibility(bool bVisible)
 	}
 }
 
-EVisibility SWindow::GetWindowContentVisibility() const
-{
-	// The content of the window should be visible unless we have a full window overlay content
-	// in which case the full window overlay content is visible but nothing under it
-	return (bShouldShowWindowContentDuringOverlay == true || !FullWindowOverlayWidget.IsValid()) ? EVisibility::SelfHitTestInvisible : EVisibility::Hidden;
-};
-
 EActiveTimerReturnType SWindow::TriggerPlayMorphSequence( double InCurrentTime, float InDeltaTime )
 {
 	Morpher.Sequence.Play( this->AsShared() );
 	return EActiveTimerReturnType::Stop;
 }
 
+void SWindow::SetWindowBackground(const FSlateBrush* InWindowBackground)
+{
+	WindowBackground = InWindowBackground;
+	if (WindowBackgroundImage)
+	{
+		WindowBackgroundImage->SetImage(WindowBackground);
+	}
+}
+
+void SWindow::UpdateWindowContentVisibility()
+{
+	// The content of the window should be visible unless we have a full window overlay content
+	// in which case the full window overlay content is visible but nothing under it
+	WindowContentVisibility = (bShouldShowWindowContentDuringOverlay == true || !FullWindowOverlayWidget.IsValid()) ? EVisibility::SelfHitTestInvisible : EVisibility::Hidden;
+
+	if (WindowBackgroundImage.IsValid())
+	{
+		WindowBackgroundImage->SetVisibility(WindowContentVisibility);
+	}
+
+	if (WindowBorder.IsValid())
+	{
+		WindowBorder->SetVisibility(WindowContentVisibility);
+	}
+
+	if (WindowOutline.IsValid())
+	{
+		WindowOutline->SetVisibility(WindowContentVisibility);
+	}
+
+	if (ContentAreaVBox.IsValid())
+	{
+		ContentAreaVBox->SetVisibility(WindowContentVisibility);
+	}
+}
+
 #if WITH_EDITOR
 FScopedSwitchWorldHack::FScopedSwitchWorldHack( const FWidgetPath& WidgetPath )
 	: Window( WidgetPath.TopLevelWindow )
diff --git a/Engine/Source/Runtime/SlateCore/Public/Animation/CurveSequence.h b/Engine/Source/Runtime/SlateCore/Public/Animation/CurveSequence.h
--- a/Engine/Source/Runtime/SlateCore/Public/Animation/CurveSequence.h
+++ b/Engine/Source/Runtime/SlateCore/Public/Animation/CurveSequence.h
@@ -96,8 +96,9 @@ public:
 	 * @param InOwnerWidget The widget that is being animated by this sequence.
 	 * @param bPlayLooped True if the curve sequence should play continually on a loop. Note that the active timer will persist until this sequence is paused or jumped to the start/end.
 	 * @param StartAtTime The relative time within the animation at which to begin playing (i.e. 0.0f is the beginning).
+	 * @param bRequiresActiveTimer	Whether or not we need to register an active timer on the widget to keep slate ticking while the animation is playing.  If that is not necessary in your use case, you can set it to false for a small performance boost
 	 */
-	void Play( const TSharedRef<SWidget>& InOwnerWidget, bool bPlayLooped = false, const float StartAtTime = 0.0f );
+	void Play( const TSharedRef<SWidget>& InOwnerWidget, bool bPlayLooped = false, const float StartAtTime = 0.0f, bool bRequiresActiveTimer = true);
 
 	/**
 	 * Start playing this curve sequence in reverse. Registers an active timer for the widget using the sequence.
@@ -105,8 +106,9 @@ public:
 	 * @param InOwnerWidget The widget that is being animated by this sequence.
 	 * @param bPlayLooped True if the curve sequence should play continually on a loop. Note that the active timer will persist until this sequence is paused or jumped to the start/end.
 	 * @param StartAtTime The relative time within the animation at which to begin playing (i.e. 0.0f is the beginning).
+	 * @param bRequiresActiveTimer	Whether or not we need to register an active timer on the widget to keep slate ticking while the animation is playing.  If that is not necessary in your use case, you can set it to false for a small performance boost
 	 */
-	void PlayReverse( const TSharedRef<SWidget>& InOwnerWidget, bool bPlayLooped = false, const float StartAtTime = 0.0f );
+	void PlayReverse( const TSharedRef<SWidget>& InOwnerWidget, bool bPlayLooped = false, const float StartAtTime = 0.0f, bool bRequiresActiveTimer = true);
 
 	/** Reverse the direction of an in-progress animation */
 	void Reverse( );
diff --git a/Engine/Source/Runtime/SlateCore/Public/Application/SlateApplicationBase.h b/Engine/Source/Runtime/SlateCore/Public/Application/SlateApplicationBase.h
--- a/Engine/Source/Runtime/SlateCore/Public/Application/SlateApplicationBase.h
+++ b/Engine/Source/Runtime/SlateCore/Public/Application/SlateApplicationBase.h
@@ -9,9 +9,11 @@
 #include "Layout/SlateRect.h"
 #include "Rendering/SlateRenderer.h"
 #include "Misc/CoreDelegates.h"
+#include "Async/TaskGraphInterfaces.h"
 
 DECLARE_MULTICAST_DELEGATE_TwoParams(FOnDebugSafeZoneChanged, const FMargin&, bool);
 
+
 class FActiveTimerHandle;
 #if WITH_ACCESSIBILITY
 class FSlateAccessibleMessageHandler;
@@ -21,6 +23,8 @@ class FWidgetPath;
 class IToolTip;
 class SWidget;
 class SWindow;
+class SImage;
+enum class EInvalidateWidgetReason : uint8;
 
 template< typename ObjectType > class TAttribute;
 
@@ -30,7 +34,6 @@ template< typename ObjectType > class TAttribute;
 class IWindowTitleBar
 {
 public:
-
 	virtual void Flash( ) = 0;
 };
 
@@ -129,6 +132,14 @@ public:
 	 */
 	virtual bool FindPathToWidget( TSharedRef<const SWidget> InWidget, FWidgetPath& OutWidgetPath, EVisibility VisibilityFilter = EVisibility::Visible ) = 0;
 
+	/**
+	 * Returns the window the provided widget is contained in 
+
+	 * @param  InWidget       Widget to find the window for
+	 * @return	returns the window the provided widget is contained in or nullptr if it's not currently in a window
+	 */
+	virtual TSharedPtr<SWindow> FindWidgetWindow(TSharedRef<const SWidget> InWidget) const = 0;
+
 	/**
 	 * Gets the active top-level window.
 	 *
@@ -335,7 +346,7 @@ public:
 	 *
 	 * @return The new image widget.
 	 */
-	virtual TSharedRef<SWidget> MakeImage( const TAttribute<const FSlateBrush*>& Image, const TAttribute<FSlateColor>& Color, const TAttribute<EVisibility>& Visibility ) const = 0;
+	virtual TSharedRef<SImage> MakeImage( const TAttribute<const FSlateBrush*>& Image, const TAttribute<FSlateColor>& Color, const TAttribute<EVisibility>& Visibility ) const = 0;
 
 	/**
 	 * Creates a tool tip with the specified text.
@@ -417,8 +428,13 @@ public:
 	/**
 	 * Gets a delegate that is invoked when a global invalidate of all widgets should occur
 	 */
-	DECLARE_EVENT(FSlateApplicationBase, FOnGlobalInvalidate);
-	FOnGlobalInvalidate& OnGlobalInvalidate()  { return OnGlobalInvalidateEvent; }
+	DECLARE_EVENT(FSlateApplicationBase, FOnInvalidateAllWidgets);
+	FOnInvalidateAllWidgets& OnInvalidateAllWidgets() { return OnInvalidateAllWidgetsEvent; }
+
+	DECLARE_EVENT_OneParam(FSlateApplicationBase, FOnGlobalInvalidationToggled, bool);
+	FOnGlobalInvalidationToggled& OnGlobalInvalidationToggled() { return OnGlobalInvalidationToggledEvent; }
+
+	void ToggleGlobalInvalidation(bool bIsGlobalInvalidationEnabled);
 
 	/**
 	 * Notifies all invalidation panels that they should invalidate their contents
@@ -598,7 +614,8 @@ protected:
 #endif
 
 	/** multicast delegate to broadcast when a global invalidate is requested */
-	FOnGlobalInvalidate OnGlobalInvalidateEvent;
+	FOnInvalidateAllWidgets OnInvalidateAllWidgetsEvent;
+	FOnGlobalInvalidationToggled OnGlobalInvalidationToggledEvent;
 
 	/** Critical section for active timer registration as it can be called from the movie thread and the game thread */
 	FCriticalSection ActiveTimerCS;
diff --git a/Engine/Source/Runtime/SlateCore/Public/Debugging/SlateDebugging.h b/Engine/Source/Runtime/SlateCore/Public/Debugging/SlateDebugging.h
--- a/Engine/Source/Runtime/SlateCore/Public/Debugging/SlateDebugging.h
+++ b/Engine/Source/Runtime/SlateCore/Public/Debugging/SlateDebugging.h
@@ -22,6 +22,7 @@ class FSlateRect;
 class FWeakWidgetPath;
 class FWidgetPath;
 class FNavigationReply;
+class FSlateInvalidationRoot;
 
 UENUM()
 enum class ESlateDebuggingInputEvent : uint8
@@ -86,6 +87,10 @@ enum class ESlateDebuggingFocusEvent : uint8
 	FocusReceived
 };
 
+
+#if WITH_SLATE_DEBUGGING
+
+
 struct SLATECORE_API FSlateDebuggingFocusEventArgs
 {
 public:
@@ -158,8 +163,6 @@ public:
 };
 
 
-#if WITH_SLATE_DEBUGGING
-
 /**
  * 
  */
@@ -240,10 +243,19 @@ public:
 	DECLARE_MULTICAST_DELEGATE_TwoParams(FUICommandRun, const FName& /*CommandName*/, const FText& /*CommandLabel*/);
 	static FUICommandRun CommandRun;
 
+	static void WidgetInvalidated(FSlateInvalidationRoot& InvalidationRoot, const class FWidgetProxy& WidgetProxy, const FLinearColor* CustomInvalidationColor = nullptr);
+
+	static void DrawInvalidationRoot(const SWidget& RootWidget, int32 LayerId, FSlateWindowElementList& OutDrawElements);
+
+	static void DrawInvalidatedWidgets(const FSlateInvalidationRoot& Root, const FPaintArgs& PaintArgs, FSlateWindowElementList& OutDrawElements);
+
+	static void ClearInvalidatedWidgets(const FSlateInvalidationRoot& Root);
 private:
 
 	// This class is only for namespace use
 	FSlateDebugging() {}
+
+	static TArray<struct FInvalidatedWidgetDrawer> InvalidatedWidgetDrawers;
 };
 
 #endif
\ No newline at end of file
diff --git a/Engine/Source/Runtime/SlateCore/Public/FastUpdate/SlateInvalidationRoot.h b/Engine/Source/Runtime/SlateCore/Public/FastUpdate/SlateInvalidationRoot.h
new file mode 100644
--- /dev/null
+++ b/Engine/Source/Runtime/SlateCore/Public/FastUpdate/SlateInvalidationRoot.h
@@ -0,0 +1,127 @@
+// Copyright 1998-2018 Epic Games, Inc. All Rights Reserved.
+
+#pragma once
+
+#include "CoreMinimal.h"
+#include "UObject/GCObject.h"
+#include "WidgetProxy.h"
+#include "Rendering/DrawElements.h"
+
+struct FSlateCachedElementData;
+class FSlateWindowElementList;
+class FWidgetStyle;
+
+struct FSlateInvalidationContext
+{
+	FSlateInvalidationContext(FSlateWindowElementList& InWindowElementList, const FWidgetStyle& InWidgetStyle)
+		: PaintArgs(nullptr)
+		, WidgetStyle(InWidgetStyle)
+		, WindowElementList(&InWindowElementList)
+		, LayoutScaleMultiplier(1.0f)
+		, IncomingLayerId(0)
+		, bParentEnabled(true)
+		, bAllowFastPathUpdate(false)
+	{
+	} 
+
+	FSlateRect CullingRect;
+	const FPaintArgs* PaintArgs;
+	const FWidgetStyle& WidgetStyle;
+	FSlateWindowElementList* WindowElementList;
+	float LayoutScaleMultiplier;
+	int32 IncomingLayerId;
+	bool bParentEnabled;
+	bool bAllowFastPathUpdate;
+};
+
+struct FSlateInvalidationResult
+{
+	FSlateInvalidationResult()
+		: MaxLayerIdPainted(0)
+		, bRepaintedWidgets(false)
+	{}
+
+	/** The max layer id painted or cached */
+	int32 MaxLayerIdPainted;
+	/** If we had to repaint any widget */
+	bool bRepaintedWidgets;
+};
+
+class FSlateInvalidationRoot : public FGCObject, public FNoncopyable
+{
+	friend class FSlateUpdateFastWidgetPathTask;
+	friend class FSlateUpdateFastPathAndHitTestGridTask;
+	friend class FWidgetProxyHandle;
+public:
+	SLATECORE_API FSlateInvalidationRoot();
+	SLATECORE_API virtual ~FSlateInvalidationRoot();
+
+	FHittestGrid* GetHittestGrid() const { return RootHittestGrid; }
+
+	/** FGCObject interface */
+	SLATECORE_API virtual void AddReferencedObjects(FReferenceCollector& Collector) override;
+	SLATECORE_API virtual FString GetReferencerName() const override;
+
+
+	SLATECORE_API void InvalidateRoot();
+
+	/** Call to notify that the ordering of children somewhere in the hierarchy below this root has changed and the fast path is no longer valid */
+	SLATECORE_API void InvalidateChildOrder();
+
+	bool NeedsSlowPath() const { return bNeedsSlowPath; }
+
+	void RemoveWidgetFromFastPath(FWidgetProxy& Proxy);
+
+	/** @return the cached draw elements for this window and its widget hierarchy*/
+	FSlateCachedElementData& GetCachedElements()
+	{
+		return CachedElementData;
+	}
+
+	int32 GetFastPathGenerationNumber() const { return FastPathGenerationNumber; }
+
+	SLATECORE_API FSlateInvalidationResult PaintInvalidationRoot(const FSlateInvalidationContext& Context);
+
+protected:
+	virtual int32 PaintSlowPath(const FSlateInvalidationContext& Context) = 0;
+
+	void SetInvalidationRootWidget(SWidget& InInvalidationRootWidget) { InvalidationRootWidget = &InInvalidationRootWidget; }
+	void SetInvalidationRootHittestGrid(FHittestGrid& InHittestGrid) { RootHittestGrid = &InHittestGrid; }
+
+	SLATECORE_API bool ProcessInvalidation();
+
+	SLATECORE_API void ClearAllFastPathData();
+
+private:
+
+	void OnInvalidateAllWidgets();
+
+	bool PaintFastPath(const FSlateInvalidationContext& Context);
+
+	void BuildFastPathList(SWidget* RootWidget);
+	void BuildNewFastPathList_Recursive(FSlateInvalidationRoot& Root, FWidgetProxy& Proxy, int32 ParentIndex, int32& NextTreeIndex, TArray<FWidgetProxy>& CurrentFastPathList, TArray<FWidgetProxy, TMemStackAllocator<>>& NewFastPathList);
+private:
+	TArray<FWidgetProxy> FastWidgetPathList;
+	/** Index to widgets which are dirty, volatile, or need some sort of per frame update (such as a tick or timer) */
+	FWidgetUpdateList WidgetsNeedingUpdate;
+
+	TArray<int32, TInlineAllocator<100>> FinalUpdateList;
+
+	FSlateCachedElementData CachedElementData;
+
+	SWidget* InvalidationRootWidget;
+
+	FHittestGrid* RootHittestGrid;
+
+	/**
+	 * The purpose of this number is as a unique Id for all widget proxy handles to validate themselves.
+	 * As widgets are added and removed, all their children are indirectly added or remove so it is necessary to invalidate all their handles. Bumping the generation number is an efficient way to do this compared to iterating all the handles
+	 * The generation number is always incrementing
+	 */
+	int32 FastPathGenerationNumber;
+
+	int32 CachedMaxLayerId;
+
+	bool bChildOrderInvalidated;
+	bool bNeedsSlowPath;
+};
\ No newline at end of file
diff --git a/Engine/Source/Runtime/SlateCore/Public/FastUpdate/WidgetProxy.h b/Engine/Source/Runtime/SlateCore/Public/FastUpdate/WidgetProxy.h
new file mode 100644
--- /dev/null
+++ b/Engine/Source/Runtime/SlateCore/Public/FastUpdate/WidgetProxy.h
@@ -0,0 +1,226 @@
+// Copyright 1998-2018 Epic Games, Inc. All Rights Reserved.
+
+#pragma once
+
+#include "CoreMinimal.h"
+#include "Types/PaintArgs.h"
+#include "Styling/WidgetStyle.h"
+#include "Misc/MemStack.h"
+#include "Layout/Clipping.h"
+#include "Layout/FlowDirection.h"
+#include "Rendering/DrawElements.h"
+
+class SWidget;
+class FPaintArgs;
+struct FFastPathPerFrameData;
+struct FSlateWidgetPersistentState;
+struct FWidgetUpdateList;
+class FSlateInvalidationRoot;
+
+
+/** Widget update flags for fast path.  Work in progress, do not modify */
+enum class EWidgetUpdateFlags : uint8
+{
+	None = 0,
+
+	/** Widget has a tick function */
+	NeedsTick = 1 << 2,
+
+	/** Widget has an active timer that needs to update */
+	NeedsActiveTimerUpdate = 1 << 3,
+
+	/** Needs repaint because the widget is dirty */
+	NeedsRepaint = 1 << 4,
+
+	/** Needs repaint because the widget is volatile */
+	NeedsVolatilePaint = 1 << 6,
+	
+	AnyUpdate = 0xff,
+};
+
+ENUM_CLASS_FLAGS(EWidgetUpdateFlags)
+
+enum class EInvalidateWidgetReason : uint8;
+
+class FWidgetProxy
+{
+public:
+	FWidgetProxy(SWidget* InWidget);
+
+	int32 Update(
+		const FPaintArgs& PaintArgs,
+		int32 MyIndex,
+		FSlateWindowElementList& OutDrawElements);
+
+	bool ProcessInvalidation(FWidgetUpdateList& UpdateList, TArray<FWidgetProxy>& FastPathWidgetList, FSlateInvalidationRoot& Root);
+
+	static void MarkProxyUpdatedThisFrame(FWidgetProxy& Proxy, FWidgetUpdateList& UpdateList);
+
+private:
+	int32 Repaint(const FPaintArgs& PaintArgs, int32 MyIndex, FSlateWindowElementList& OutDrawElements) const;
+public:
+	SWidget* Widget;
+	int32 Index;
+	int32 ParentIndex;
+	int32 NumChildren;
+	int32 LeafMostChildIndex;
+	EWidgetUpdateFlags UpdateFlags;
+	EInvalidateWidgetReason CurrentInvalidateReason;
+	/** The widgets own visibility */
+	EVisibility Visibility;
+	uint8 bUpdatedSinceLastInvalidate : 1;
+	uint8 bInUpdateList : 1;
+	uint8 bInvisibleDueToParentOrSelfVisibility : 1;
+	uint8 bChildOrderInvalid : 1;
+};
+
+static_assert(sizeof(FWidgetProxy) == 32, "FWidgetProxy should be 32 bytes");
+static_assert(TIsTriviallyDestructible<FWidgetProxy>::Value == true, "FWidgetProxy must be trivially destructible");
+
+template <> struct TIsPODType<FWidgetProxy> { enum { Value = true }; };
+
+struct FWidgetUpdateList
+{
+public:
+	inline void Push(FWidgetProxy& Proxy)
+	{
+		if (!Heap.Contains(Proxy.Index))
+		{
+			Proxy.bInUpdateList = true;
+			Heap.HeapPush(Proxy.Index, TGreater<>());
+		}
+	}
+
+	inline int32 Pop()
+	{
+		int32 Index;
+		Heap.HeapPop(Index, TGreater<>(), false);
+		return Index;
+	}
+
+	inline void Empty()
+	{
+		Heap.Empty();
+	}
+
+	inline void Reset()
+	{
+		Heap.Reset();
+	}
+
+	inline int32 Num() const
+	{
+		return Heap.Num();
+	}
+
+	bool Contains(FWidgetProxy& Proxy) const
+	{
+		return Heap.Contains(Proxy.Index);
+	}
+
+	const TArray<int32, TInlineAllocator<100>>& GetRawData() const
+	{
+		return Heap;
+	}
+private:
+	TArray<int32, TInlineAllocator<100>> Heap;
+};
+
+
+/**
+ * Represents the state of a widget from when it last had SWidget::Paint called on it. 
+ * This should contain everything needed to directly call Paint on a widget
+ */
+struct FSlateWidgetPersistentState
+{
+	FSlateWidgetPersistentState()
+		: CachedElementListNode(nullptr)
+		, LayerId(0)
+		, OutgoingLayerId(0)
+		, IncomingFlowDirection(EFlowDirection::LeftToRight)
+		, bParentEnabled(true)
+		, bInheritedHittestability(false)
+	{}
+
+	TWeakPtr<SWidget> PaintParent;
+	TOptional<FSlateClippingState> InitialClipState;
+	FGeometry AllottedGeometry;
+	FGeometry DesktopGeometry;
+	FSlateRect CullingBounds;
+	FWidgetStyle WidgetStyle;
+	FSlateCachedElementListNode* CachedElementListNode;
+	/** Starting layer id for drawing children **/
+	int32 LayerId;
+	int32 OutgoingLayerId;
+	EFlowDirection IncomingFlowDirection;
+	uint8 bParentEnabled : 1;
+	uint8 bInheritedHittestability : 1;
+
+	static const FSlateWidgetPersistentState NoState;
+};
+template <> struct TIsPODType<FSlateWidgetPersistentState> { enum { Value = true }; };
+
+struct FWidgetStackData
+{
+	FWidgetStackData()
+		: ClipStackIndex(INDEX_NONE)
+		, IncomingLayerId(0)
+		, CurrentMaxLayerId(0)
+	{}
+
+	/** Index into the current stack where the widgets clip state is stored */
+	int32 ClipStackIndex;
+
+	/** The incoming layer id for children and draw elements*/
+	int32 IncomingLayerId;
+
+	/**
+	* Current Max layer id being computed.
+	* This is the layer id that would be returned from the widget during a traditional recurive
+	* OnPaint call
+	*/
+	int32 CurrentMaxLayerId;
+};
+
+template <> struct TIsPODType<FWidgetStackData> { enum { Value = true }; };
+
+class FWidgetProxyHandle
+{
+	friend class SWidget;
+	friend class FSlateDebugging;
+	friend class FSlateInvalidationRoot;
+public:
+	FWidgetProxyHandle()
+		: InvalidationRoot(nullptr)
+		, MyIndex(INDEX_NONE)
+		, GenerationNumber(INDEX_NONE)
+
+	{}
+
+	SLATECORE_API bool IsValid() const;
+
+	FSlateInvalidationRoot* GetInvalidationRoot() const { return InvalidationRoot; }
+	FWidgetProxy& GetProxy();
+	const FWidgetProxy& GetProxy() const;
+
+	int32 GetIndex() const { return IsValid() ? MyIndex : INDEX_NONE; }
+
+	/**
+	 * Marks the widget as updated this frame
+	 * Note: If the widget still has update flags (e.g it ticks or is volatile or something during update added new flags)
+	 * it will remain in the update list
+	 */
+	void MarkWidgetUpdatedThisFrame();
+	
+	void MarkWidgetDirty(EInvalidateWidgetReason InvalidateReason);
+	void UpdateWidgetFlags(EWidgetUpdateFlags NewFlags);
+private:
+	FWidgetProxyHandle(FSlateInvalidationRoot& InInvalidationRoot, int32 InIndex);
+private:
+	/** The root of invalidation tree this proxy belongs to */
+	FSlateInvalidationRoot* InvalidationRoot;
+	/** Index to myself in the fast path list */
+	int32 MyIndex;
+	/** This serves as an efficient way to test for validity which does not require invalidating all handles directly*/
+	int32 GenerationNumber;
+};
\ No newline at end of file
diff --git a/Engine/Source/Runtime/SlateCore/Public/Input/HittestGrid.h b/Engine/Source/Runtime/SlateCore/Public/Input/HittestGrid.h
--- a/Engine/Source/Runtime/SlateCore/Public/Input/HittestGrid.h
+++ b/Engine/Source/Runtime/SlateCore/Public/Input/HittestGrid.h
@@ -23,35 +23,6 @@ public:
 	virtual TSharedPtr<struct FVirtualPointerPosition> TranslateMouseCoordinateFor3DChild( const TSharedRef<SWidget>& ChildWidget, const FGeometry& ViewportGeometry, const FVector2D& ScreenSpaceMouseCoordinate, const FVector2D& LastScreenSpaceMouseCoordinate ) const = 0;
 };
 
-
-struct FCachedWidget
-{
-	FCachedWidget(int32 InParentIndex, const FArrangedWidget& InWidget, int32 InClippingStateIndex, int32 InLayerId)
-		: WidgetPtr(InWidget.Widget)
-		, CachedGeometry(InWidget.Geometry)
-		, ClippingStateIndex(InClippingStateIndex)
-		, Children()
-		, ParentIndex(InParentIndex)
-		, LayerId(InLayerId)
-	{}
-
-	void AddChild(const int32 ChildIndex)
-	{
-		Children.Add(ChildIndex);
-	}
-
-	TWeakPtr<SWidget> WidgetPtr;
-	/** Allow widgets that implement this interface to insert widgets into the bubble path */
-	TWeakPtr<ICustomHitTestPath> CustomPath;
-	FGeometry CachedGeometry;
-	int32 ClippingStateIndex;
-	TArray<int32, TInlineAllocator<4> > Children;
-	int32 ParentIndex;
-	/** This is needed to be able to pick the best of the widgets within the virtual cursor's radius. */
-	int32 LayerId;
-};
-
-
 class SLATECORE_API FHittestGrid
 {
 public:
@@ -63,25 +34,29 @@ public:
 	 * Given a Slate Units coordinate in virtual desktop space, perform a hittest
 	 * and return the path along which the corresponding event would be bubbled.
 	 */
-	TArray<FWidgetAndPointer> GetBubblePath( FVector2D DesktopSpaceCoordinate, float CursorRadius, bool bIgnoreEnabledStatus );
+	TArray<FWidgetAndPointer> GetBubblePath(FVector2D DesktopSpaceCoordinate, float CursorRadius, bool bIgnoreEnabledStatus);
 
 	/**
-	 * Clear the hittesting area and prepare to execute a new frame.
-	 * Depending on monitor arrangement, the area to be hittested could begin
-	 * in the negative coordinates.
+	 * Set the position and size of the hittest area in desktop coordinates
+	 *
+	 * @param HittestPositionInDesktop	The position of this hit testing area in desktop coordinates.
+	 * @param HittestDimensions			The dimensions of this hit testing area.
 	 *
-	 * @param HittestArea       Size in Slate Units of the area we are considering for hittesting.
+	 * @return      Returns true if a clear of the hittest grid was required. 
 	 */
-	void ClearGridForNewFrame(const FSlateRect& HittestArea);
+	bool SetHittestArea(const FVector2D& HittestPositionInDesktop, const FVector2D& HittestDimensions, const FVector2D& HitestOffsetInWindow = FVector2D::ZeroVector);
 
-	void PushClip(const FSlateClippingZone& ClippingZone); 
-
-	void PopClip();
+	/**
+	 * Clear the grid
+	 */
+	void Clear();
 
-	/** Add Widget into the hittest data structure so that we can later make queries about it. */
-	int32 InsertWidget(const int32 ParentHittestIndex, const EVisibility& Visibility, const FArrangedWidget& Widget, const FVector2D InWindowOffset, int32 LayerId);
+	/** Add / Remove SWidget from the HitTest Grid */
+	void AddWidget(const TSharedRef<SWidget>& InWidget, int32 InBatchPriorityGroup, int32 InLayerId, int32 InSecondarySort);
+	void RemoveWidget(const TSharedRef<SWidget>& InWidget);
 
-	void InsertCustomHitTestPath( TSharedRef<ICustomHitTestPath> CustomHitTestPath, int32 WidgetIndex );
+	/** Insert custom hit test data for a widget already in the grid */
+	void InsertCustomHitTestPath(const TSharedRef<SWidget> InWidget, TSharedRef<ICustomHitTestPath> CustomHitTestPath);
 
 	/**
 	 * Finds the next focusable widget by searching through the hit test grid
@@ -93,27 +68,41 @@ public:
 	 */
 	TSharedPtr<SWidget> FindNextFocusableWidget(const FArrangedWidget& StartingWidget, const EUINavigation Direction, const FNavigationReply& NavigationReply, const FArrangedWidget& RuleWidget);
 
+	void AppendGrid(FHittestGrid& OtherGrid);
+
+	FVector2D GetGridSize() const { return GridSize; }
+	FVector2D GetGridOrigin() const { return GridOrigin; }
+	FVector2D GetGridWindowOrigin() const { return GridWindowOrigin; }
+
+#if WITH_SLATE_DEBUGGING
+	void LogGrid() const;
+
+	void DisplayGrid(int32 InLayer, const FGeometry& AllottedGeometry, FSlateWindowElementList& WindowElementList) const;
+#endif
+
+private:
+
 	/**
-	 * Get the size of the grid in cells.
-	 * @returns - The size of the grid in cells.
-	 */
-	FIntPoint GetNumCells() const
-	{
-		return NumCells;
-	}
-	
-	/**
-	 * Set an excess of cells which will increase the size of the grid. This
-	 * extends the desktop space beyond the norm and allows hits to take place
-	 * on areas outside the desktop.
-	 * @param InNumCellsExcess - The excess of cells.
+	 * Widget Data we maintain internally store along with the widget reference
 	 */
-	void SetNumCellsExcess(const FIntPoint& InNumCellsExcess)
+	struct FWidgetData
 	{
-		NumCellsExcess = InNumCellsExcess;
-	}
-
-private:
+		FWidgetData(TSharedRef<SWidget> InWidget, const FIntPoint& InUpperLeftCell, const FIntPoint& InLowerRightCell, int64 InPrimarySort, int32 InSecondarySort)
+			: WeakWidget(InWidget)
+			, UpperLeftCell(InUpperLeftCell)
+			, LowerRightCell(InLowerRightCell)
+			, PrimarySort(InPrimarySort)
+			, SecondarySort(InSecondarySort)
+		{}
+		TWeakPtr<SWidget> WeakWidget;
+		TWeakPtr<ICustomHitTestPath> CustomPath;
+		FIntPoint UpperLeftCell;
+		FIntPoint LowerRightCell;
+		int64 PrimarySort;
+		int32 SecondarySort;
+
+		TSharedPtr<SWidget> GetWidget() const { return WeakWidget.Pin(); }
+	};
 
 	/**
 	 * All the available space is partitioned into Cells.
@@ -122,67 +111,40 @@ private:
 	 */
 	struct FCell
 	{
-		TArray<int32> CachedWidgetIndexes;
-	};
-
-	/** Shared arguments to helper functions. */
-	struct FGridTestingParams;
-
-	/** Helper functions */
-	bool IsValidCellCoord(const FIntPoint& CellCoord) const;
-	bool IsValidCellCoord(const int32 XCoord, const int32 YCoord) const;
-
-	/** Bubble path, distance to leafmost widget from point of picking, and LayerId of leafmost widget. */
-	struct FWidgetPathAndDist
-	{
-		/** Ctor */
-		FWidgetPathAndDist( float InDistanceSq = -1.0f )
-		: DistToTopWidgetSq(InDistanceSq)
-		, LayerId(0)
-		{}
-
-		FWidgetPathAndDist(const TArray<FWidgetAndPointer>& InPath, float InDistanceSq, int32 InLayerId)
-		: BubblePath(InPath)
-		, DistToTopWidgetSq(InDistanceSq)
-		, LayerId(InLayerId)
-		{}
-
-		void Clear()
-		{
-			BubblePath.Reset();
-			DistToTopWidgetSq = -1.0f;
-		}
+		FCell();
 
-		bool IsValidPath() const
-		{
-			return DistToTopWidgetSq >= 0.0f && BubblePath.Num() > 0;
-		}
+		void AddIndex(int32 WidgetIndex);
+		void RemoveIndex(int32 WidgetIndex);
+		void Sort(const TSparseArray<FWidgetData>& InWidgetArray);
 
-		TArray<FWidgetAndPointer> BubblePath;
-		float DistToTopWidgetSq;
-		int32 LayerId;
+		const TArray<int32>& GetCachedWidgetIndexes() const;
+		
+	private:
+		bool bRequiresSort;
+		TArray<int32> CachedWidgetIndexes;
 	};
 
-	FWidgetPathAndDist GetWidgetPathAndDist(const FGridTestingParams& Params, const bool bIgnoreEnabledStatus) const;
+	struct FGridTestingParams;
 
-	/** Return Value for GetHitIndexFromCellIndex */
 	struct FIndexAndDistance
 	{
-		FIndexAndDistance( int32 InIndex = INDEX_NONE, float InDistanceSq = 0 )
-		: WidgetIndex(InIndex)
-		, DistanceSqToWidget(InDistanceSq)
+		FIndexAndDistance(int32 InIndex = INDEX_NONE, float InDistanceSq = 0)
+			: WidgetIndex(InIndex)
+			, DistanceSqToWidget(InDistanceSq)
 		{}
 		int32 WidgetIndex;
 		float DistanceSqToWidget;
 	};
-	FIndexAndDistance GetHitIndexFromCellIndex(const FGridTestingParams& Params) const;
 
-	TArray<FWidgetAndPointer> GetBubblePathFromHitIndex(const int32 HitIndex, const bool bIgnoreEnabledStatus) const;
-	
-	friend class SWidgetReflector;
+	/** Helper functions */
+	bool IsValidCellCoord(const FIntPoint& CellCoord) const;
+	bool IsValidCellCoord(const int32 XCoord, const int32 YCoord) const;
+
+	/** Return the Index and distance to a hit given the testing params */
+	FIndexAndDistance GetHitIndexFromCellIndex(const FGridTestingParams& Params);
 
-	/** @returns true if Child is a descendant of Parent. */
-	bool IsDescendantOf(const TSharedRef<SWidget> Parent, const FCachedWidget& Child);
+	/** @returns true if the child is a paint descendant of the provided Parent. */
+	bool IsDescendantOf(const TSharedRef<SWidget> Parent, const FWidgetData& ChildData);
 
 	/** Utility function for searching for the next focusable widget. */
 	template<typename TCompareFunc, typename TSourceSideFunc, typename TDestSideFunc>
@@ -194,35 +156,37 @@ private:
 	/** Access a cell at coordinates X, Y. Coordinates are row and column indexes. */
 	FORCEINLINE_DEBUGGABLE FCell& CellAt(const int32 X, const int32 Y)
 	{
-		checkf((Y*NumCells.X + X) < Cells.Num(), TEXT("HitTestGrid CellAt() failed: X= %d Y= %d NumCells.X= %d NumCells.Y= %d Cells.Num()= %d"), X, Y, NumCells.X, NumCells.Y, Cells.Num());
+		checkfSlow((Y*NumCells.X + X) < Cells.Num(), TEXT("HitTestGrid CellAt() failed: X= %d Y= %d NumCells.X= %d NumCells.Y= %d Cells.Num()= %d"), X, Y, NumCells.X, NumCells.Y, Cells.Num());
 		return Cells[Y*NumCells.X + X];
 	}
 
+	/** Access a cell at coordinates X, Y. Coordinates are row and column indexes. */
 	FORCEINLINE_DEBUGGABLE const FCell& CellAt( const int32 X, const int32 Y ) const
 	{
-		checkf((Y*NumCells.X + X) < Cells.Num(), TEXT("HitTestGrid CellAt() failed: X= %d Y= %d NumCells.X= %d NumCells.Y= %d Cells.Num()= %d"), X, Y, NumCells.X, NumCells.Y, Cells.Num());
+		checkfSlow((Y*NumCells.X + X) < Cells.Num(), TEXT("HitTestGrid CellAt() failed: X= %d Y= %d NumCells.X= %d NumCells.Y= %d Cells.Num()= %d"), X, Y, NumCells.X, NumCells.Y, Cells.Num());
 		return Cells[Y*NumCells.X + X];
 	}
 
-	/** All the widgets and their arranged geometries encountered this frame. */
-	TArray<FCachedWidget> WidgetsCachedThisFrame;
+	/** Map of all the widgets currently in the hit test grid to their stable index. */
+	TMap<SWidget*, int32> WidgetMap;
+
+	/** Stable indexed sparse array of all the widget data we track. */
+	TSparseArray<FWidgetData> WidgetArray;
 
 	/** The cells that make up the space partition. */
 	TArray<FCell> Cells;
 
-	/** Where the 0,0 of the upper-left-most cell corresponds to in desktop space. */
-	FVector2D GridOrigin;
-
 	/** The size of the grid in cells. */
 	FIntPoint NumCells;
 
-	/** The excess of cells which will increase the size of the grid. */
-	FIntPoint NumCellsExcess;
+	/** Where the 0,0 of the upper-left-most cell corresponds to in desktop space. */
+	FVector2D GridOrigin;
 
-	/** The clipping manager that manages any clipping for hit testable widgets. */
-	FSlateClippingManager ClippingManager;
+	/** Where the 0,0 of the upper-left-most cell corresponds to in window space. */
+	FVector2D GridWindowOrigin;
 
-	void LogGrid() const;
+	/** The Size of the current grid. */
+	FVector2D GridSize;
 
-	static void LogChildren(int32 Index, int32 IndentLevel, const TArray<FCachedWidget>& WidgetsCachedThisFrame);
+	//friend class SWidgetReflector;
 };
diff --git a/Engine/Source/Runtime/SlateCore/Public/Layout/Children.h b/Engine/Source/Runtime/SlateCore/Public/Layout/Children.h
--- a/Engine/Source/Runtime/SlateCore/Public/Layout/Children.h
+++ b/Engine/Source/Runtime/SlateCore/Public/Layout/Children.h
@@ -156,7 +156,7 @@ public:
 		WidgetPtr = InWidget;
 		if (Owner) 
 		{ 
-			Owner->InvalidatePrepass();
+			Owner->Invalidate(EInvalidateWidget::ChildOrder);
 
 			if (InWidget.IsValid() && InWidget != SNullWidget::NullWidget)
 			{
@@ -173,7 +173,7 @@ public:
 
 			if (Owner)
 			{
-				Owner->InvalidatePrepass();
+				Owner->Invalidate(EInvalidateWidget::ChildOrder);
 			}
 		}
 	}
@@ -312,12 +312,14 @@ public:
 			return INDEX_NONE;
 		}
 
+		int32 Index = TIndirectArray< SlotType >::Add(Slot);
+
 		if (Owner)
 		{
 			Slot->AttachWidgetParent(Owner);
 		}
 
-		return TIndirectArray< SlotType >::Add(Slot);
+		return Index;
 	}
 
 	void RemoveAt( int32 Index )
@@ -347,18 +349,20 @@ public:
 	{
 		if (!bEmptying)
 		{
+			TIndirectArray< SlotType >::Insert(Slot, Index);
+
 			// Don't do parent manipulation if this panel has no owner.
 			if (Owner)
 			{
 				Slot->AttachWidgetParent(Owner);
 			}
-
-			TIndirectArray< SlotType >::Insert(Slot, Index);
 		}
 	}
 
 	void Move(int32 IndexToMove, int32 IndexToDestination)
 	{
+		// @todo this is going to cause a problem for draw ordering
+
 		// Since we're dealing with an Indirect Array, we can't move an item already in the array.
 		if (IndexToMove > IndexToDestination) // going up
 		{
@@ -372,6 +376,11 @@ public:
 			TIndirectArray< SlotType >::Swap(IndexToMove, IndexToDestination + 1);
 			TIndirectArray< SlotType >::RemoveAt(IndexToMove);
 		}
+
+		if (Owner)
+		{
+			Owner->Invalidate(EInvalidateWidget::ChildOrder);
+		}
 	}
 
 	void Reserve( int32 NumToReserve )
@@ -391,17 +400,29 @@ public:
 	void Sort( const PREDICATE_CLASS& Predicate )
 	{
 		::Sort(TIndirectArray< SlotType >::GetData(), TIndirectArray<SlotType>::Num(), Predicate);
+		if (Owner)
+		{
+			Owner->Invalidate(EInvalidateWidget::ChildOrder);
+		}
 	}
 
 	template <class PREDICATE_CLASS>
 	void StableSort(const PREDICATE_CLASS& Predicate)
 	{
 		::StableSort(TIndirectArray< SlotType >::GetData(), TIndirectArray< SlotType >::Num(), Predicate);
+		if (Owner)
+		{
+			Owner->Invalidate(EInvalidateWidget::ChildOrder);
+		}
 	}
 
 	void Swap( int32 IndexA, int32 IndexB )
 	{
 		TIndirectArray< SlotType >::Swap(IndexA, IndexB);
+		if (Owner)
+		{
+			Owner->Invalidate(EInvalidateWidget::ChildOrder);
+		}
 	}
 
 private:
@@ -565,20 +586,23 @@ public:
 
 	int32 Add( const TSharedRef<ChildType>& Child )
 	{
+		if (Owner && bChangesInvalidatePrepass)
+		{ 
+			Owner->Invalidate(EInvalidateWidget::ChildOrder);
+		}
+
+		int32 Index = TArray< TSharedRef<ChildType> >::Add(Child);
+
 		if (Owner)
 		{
-			if(bChangesInvalidatePrepass)
-			{
-				Owner->InvalidatePrepass();
-			}
-
 			if (Child != SNullWidget::NullWidget)
 			{
 				Child->AssignParentWidget(Owner->AsShared());
 			}
 		}
 
-		return TArray< TSharedRef<ChildType> >::Add(Child);
+		return Index;
+
 	}
 
 	void Reset(int32 NewSize = 0)
@@ -611,20 +635,20 @@ public:
 
 	void Insert(const TSharedRef<ChildType>& Child, int32 Index)
 	{
+		if (Owner && bChangesInvalidatePrepass) 
+		{
+			Owner->Invalidate(EInvalidateWidget::ChildOrder);
+		}
+
+		TArray< TSharedRef<ChildType> >::Insert(Child, Index);
+
 		if (Owner)
 		{
-			if(bChangesInvalidatePrepass)
-			{
-				Owner->InvalidatePrepass();
-			}
-		
 			if (Child != SNullWidget::NullWidget)
 			{
 				Child->AssignParentWidget(Owner->AsShared());
 			}
 		}
-
-		TArray< TSharedRef<ChildType> >::Insert(Child, Index);
 	}
 
 	int32 Remove( const TSharedRef<ChildType>& Child )
@@ -668,11 +692,19 @@ public:
 	void Sort( const PREDICATE_CLASS& Predicate )
 	{
 		TArray< TSharedRef<ChildType> >::Sort( Predicate );
+		if (Owner && bChangesInvalidatePrepass)
+		{
+			Owner->Invalidate(EInvalidateWidget::ChildOrder);
+		}
 	}
 
 	void Swap( int32 IndexA, int32 IndexB )
 	{
 		TIndirectArray< ChildType >::Swap(IndexA, IndexB);
+		if (Owner && bChangesInvalidatePrepass)
+		{
+			Owner->Invalidate(EInvalidateWidget::ChildOrder);
+		}
 	}
 
 private:
diff --git a/Engine/Source/Runtime/SlateCore/Public/Layout/Clipping.h b/Engine/Source/Runtime/SlateCore/Public/Layout/Clipping.h
--- a/Engine/Source/Runtime/SlateCore/Public/Layout/Clipping.h
+++ b/Engine/Source/Runtime/SlateCore/Public/Layout/Clipping.h
@@ -71,6 +71,7 @@ public:
 	explicit FSlateClippingZone(const FGeometry& BoundingGeometry);
 	explicit FSlateClippingZone(const FPaintGeometry& PaintingGeometry);
 	FSlateClippingZone(const FVector2D& InTopLeft, const FVector2D& InTopRight, const FVector2D& InBottomLeft, const FVector2D& InBottomRight);
+	FSlateClippingZone() {}
 
 	/**  */
 	FORCEINLINE bool GetShouldIntersectParent() const
@@ -150,16 +151,25 @@ public:
 		);
 	}
 
+	FSlateClippingZone ConvertRelativeToAbsolute(const FVector2D& WindowOffset) const
+	{
+		FSlateClippingZone Absolute(TopLeft + WindowOffset, TopRight + WindowOffset, BottomLeft + WindowOffset, BottomRight + WindowOffset);
+		Absolute.bIsAxisAligned = bIsAxisAligned;
+		Absolute.bIntersect = bIntersect;
+		Absolute.bAlwaysClip = bAlwaysClip;
+
+		return Absolute;
+	}
 private:
 	void InitializeFromArbitraryPoints(const FVector2D& InTopLeft, const FVector2D& InTopRight, const FVector2D& InBottomLeft, const FVector2D& InBottomRight);
 
 private:
 	/** Is the clipping zone axis aligned?  Axis aligned clipping zones are much cheaper. */
-	bool bIsAxisAligned : 1;
+	uint8 bIsAxisAligned : 1;
 	/** Should this clipping zone intersect the current one? */
-	bool bIntersect : 1;
+	uint8 bIntersect : 1;
 	/** Should this clipping zone always clip, even if another zone wants to ignore intersection? */
-	bool bAlwaysClip : 1;
+	uint8 bAlwaysClip : 1;
 };
 
 FORCEINLINE uint32 GetTypeHash(const FSlateClippingZone& Zone)
@@ -198,6 +208,8 @@ class SLATECORE_API FSlateClippingState
 public:
 	FSlateClippingState(EClippingFlags InFlags = EClippingFlags::None);
 	
+	FSlateClippingState(const FSlateClippingState& Other);
+
 	/** Is a point inside the clipping state? */
 	bool IsPointInside(const FVector2D& Point) const;
 
@@ -275,13 +287,71 @@ private:
 	/** The specialized flags needed for this clipping state. */
 	EClippingFlags Flags;
 
-#if !(UE_BUILD_SHIPPING || UE_BUILD_TEST)
+#if WITH_SLATE_DEBUGGING
 	/** For a given frame, this is a unique index into a state array of clipping zones that have been registered for a window being drawn. */
 	mutable int32 Debugging_StateIndex;
 	mutable int32 Debugging_StateIndexFromFrame;
 #endif
 };
 
+struct FClipStateHandle
+{
+public:
+	FClipStateHandle()
+		: CachedClipState(nullptr)
+		, PrecachedClipIndex(INDEX_NONE)
+	{}
+
+	int32 GetPrecachedClipIndex() const { return PrecachedClipIndex; }
+	const FSlateClippingState* GetCachedClipState() const { return CachedClipState; }
+
+	bool operator==(const FClipStateHandle& Other) const
+	{
+		return CachedClipState == Other.CachedClipState && PrecachedClipIndex == Other.PrecachedClipIndex;
+	}
+
+	void SetPreCachedClipIndex(int32 InClipIndex)
+	{
+		PrecachedClipIndex = InClipIndex;
+	}
+
+	void SetCachedClipState(const FSlateClippingState* CachedState)
+	{
+		CachedClipState = CachedState;
+		PrecachedClipIndex = INDEX_NONE;
+	}
+private:
+	const FSlateClippingState* CachedClipState;
+	int32 PrecachedClipIndex;
+
+};
+
+
+class FSlateCachedClipState
+{
+public:
+	FSlateCachedClipState(const FSlateClippingState& InState)
+		: ClippingState(InState)
+		, UsageCount(1)
+	{}
+
+	FSlateClippingState ClippingState;
+
+	int32 GetUsageCount() const { return UsageCount; }
+	void BeginUsingState()
+	{
+		FPlatformAtomics::InterlockedIncrement(&UsageCount);
+	}
+
+	void EndUsingState()
+	{
+		FPlatformAtomics::InterlockedDecrement(&UsageCount);
+		check(UsageCount >= 0);
+	}
+private:
+	int32 UsageCount;
+};
+
 
 /**
  * The clipping manager maintain the running clip state.  This is used for both maintain and for hit testing.
@@ -293,21 +363,20 @@ public:
 
 	int32 PushClip(const FSlateClippingZone& InClippingZone);
 	int32 PushClippingState(const FSlateClippingState& InClipState);
-	int32 PushAndMergePartialClippingState(const FSlateClippingState& NewPartialClippingState);
 	int32 GetClippingIndex() const;
 	TOptional<FSlateClippingState> GetActiveClippingState() const;
-	const TArray< FSlateClippingState >& GetClippingStates() const;
+	const TArray<int32>& GetClippingStack() const { return ClippingStack; }
+	const TArray<FSlateClippingState>& GetClippingStates() const;
 	void PopClip();
-
-	int32 MergePartialClippingStates(const TArray< FSlateClippingState >& States);
+	void PopToStackIndex(int32 Index);
+	int32 GetClippingIndexAtStackIndex(int32 StackIndex) const { return ClippingStack.IsValidIndex(StackIndex) ? ClippingStack[StackIndex] : INDEX_NONE; }
+	int32 GetStackDepth() const { return ClippingStack.Num(); }
+	const FSlateClippingState* GetPreviousClippingState(bool bWillIntersectWithParent) const;
 
 	void ResetClippingState();
 
-	void CopyClippingStateTo( FSlateClippingManager& Other) const;
-
 private:
-	FSlateClippingState MergePartialClippingState(const FSlateClippingState& State) const;
-	const FSlateClippingState* GetPreviousClippnigState(bool bWillIntersectWithParent) const;
+
 	FSlateClippingState CreateClippingState(const FSlateClippingZone& InClipRect) const;
 
 private:
