Slate Global Invalidation System

diff --git a/Engine/Source/Runtime/SlateCore/Public/Rendering/DrawElementPayloads.h b/Engine/Source/Runtime/SlateCore/Public/Rendering/DrawElementPayloads.h
new file mode 100644
--- /dev/null
+++ b/Engine/Source/Runtime/SlateCore/Public/Rendering/DrawElementPayloads.h
@@ -0,0 +1,327 @@
+// Copyright 1998-2018 Epic Games, Inc. All Rights Reserved.
+
+#pragma once
+
+#include "SlateRenderBatch.h"
+
+struct FSlateGradientStop
+{
+	FVector2D Position;
+	FLinearColor Color;
+
+	/**
+	 * Construct a Gradient Stop from a Position and a Color.
+	 * @param InPosition - The position in widget space for this stop. Both X and Y are used for a single-axis gradient.
+						  A two stop gradient should go from (0,0), to (Width,Height).
+	 * @param InColor	- The color to lerp towards at this stop.
+	 */
+	FSlateGradientStop(const FVector2D& InPosition, const FLinearColor& InColor)
+		: Position(InPosition)
+		, Color(InColor)
+	{
+
+	}
+};
+template <> struct TIsPODType<FSlateGradientStop> { enum { Value = true }; };
+
+
+struct FSlateDataPayload
+{
+	virtual ~FSlateDataPayload() {}
+	virtual void AddReferencedObjects(FReferenceCollector& Collector) {}
+};
+
+struct FSlateTintableElement
+{
+	FLinearColor Tint;
+
+	FORCEINLINE void SetTint(const FLinearColor& InTint) { Tint = InTint; }
+	FORCEINLINE FLinearColor GetTint() const { return Tint; }
+};
+
+struct FSlateBoxPayload : public FSlateDataPayload, public FSlateTintableElement
+{
+	FMargin Margin;
+	FBox2D UVRegion;
+	const FSlateShaderResourceProxy* ResourceProxy;
+	ESlateBrushTileType::Type Tiling;
+	ESlateBrushMirrorType::Type Mirroring;
+	ESlateBrushDrawType::Type DrawType;
+
+	const FMargin& GetBrushMargin() const { return Margin; }
+	const FBox2D& GetBrushUVRegion() const { return UVRegion; }
+	ESlateBrushTileType::Type GetBrushTiling() const { return Tiling; }
+	ESlateBrushMirrorType::Type GetBrushMirroring() const { return Mirroring; }
+	ESlateBrushDrawType::Type GetBrushDrawType() const { return DrawType; }
+	const FSlateShaderResourceProxy* GetResourceProxy() const { return ResourceProxy; }
+
+	void SetBrush(const FSlateBrush* InBrush)
+	{
+		check(InBrush);
+		ensureMsgf(InBrush->GetDrawType() != ESlateBrushDrawType::NoDrawType, TEXT("This should have been filtered out earlier in the Make... call."));
+
+		// Note: Do not store the brush.  It is possible brushes are destroyed after an element is enqueued for rendering
+		Margin = InBrush->GetMargin();
+		UVRegion = InBrush->GetUVRegion();
+		Tiling = InBrush->GetTiling();
+		Mirroring = InBrush->GetMirroring();
+		DrawType = InBrush->GetDrawType();
+		FSlateResourceHandle Handle = InBrush->GetRenderingResource();
+		if (Handle.IsValid())
+		{
+			ResourceProxy = Handle.GetResourceProxy();
+		}
+		else
+		{
+			ResourceProxy = nullptr;
+		}
+	}
+
+	FORCENOINLINE virtual ~FSlateBoxPayload()
+	{
+	}
+};
+
+struct FSlateTextPayload : public FSlateDataPayload, public FSlateTintableElement
+{
+	// The font to use when rendering
+	FSlateFontInfo FontInfo;
+	// Basic text data
+	FString ImmutableText;
+
+	const FSlateFontInfo& GetFontInfo() const { return FontInfo; }
+	const TCHAR* GetText() const { return *ImmutableText; }
+	int32 GetTextLength() const { return ImmutableText.Len(); }
+
+	void SetText(const FString& InText, const FSlateFontInfo& InFontInfo, int32 InStartIndex, int32 InEndIndex)
+	{
+		FontInfo = InFontInfo;
+		const int32 StartIndex = FMath::Min<int32>(InStartIndex, InText.Len());
+		const int32 EndIndex = FMath::Min<int32>(InEndIndex, InText.Len());
+		const int32 TextLength = (EndIndex > StartIndex) ? EndIndex - StartIndex : 0;
+		if (TextLength > 0)
+		{
+			ImmutableText = InText.Mid(StartIndex, TextLength);
+		}
+	}
+
+	void SetText(const FString& InText, const FSlateFontInfo& InFontInfo)
+	{
+		FontInfo = InFontInfo;
+		ImmutableText = InText;
+	}
+
+
+	virtual void AddReferencedObjects(FReferenceCollector& Collector)
+	{
+		FontInfo.AddReferencedObjects(Collector);
+	}
+};
+
+
+struct FSlateShapedTextPayload : public FSlateDataPayload, public FSlateTintableElement
+{
+	// Shaped text data
+	FShapedGlyphSequencePtr ShapedGlyphSequence;
+
+	FLinearColor OutlineTint;
+
+	const FShapedGlyphSequencePtr& GetShapedGlyphSequence() const { return ShapedGlyphSequence; }
+	FLinearColor GetOutlineTint() const { return OutlineTint; }
+
+	void SetShapedText(FSlateWindowElementList& ElementList, const FShapedGlyphSequencePtr& InShapedGlyphSequence, const FLinearColor& InOutlineTint)
+	{
+		ShapedGlyphSequence = InShapedGlyphSequence;
+		OutlineTint = InOutlineTint;
+	}
+
+	virtual void AddReferencedObjects(FReferenceCollector& Collector)
+	{
+		if (ShapedGlyphSequence.IsValid())
+		{
+			const_cast<FShapedGlyphSequence*>(ShapedGlyphSequence.Get())->AddReferencedObjects(Collector);
+		}
+	}
+};
+
+struct FSlateGradientPayload : public FSlateDataPayload
+{
+	TArray<FSlateGradientStop> GradientStops;
+	EOrientation GradientType;
+
+	void SetGradient(const TArray<FSlateGradientStop>& InGradientStops, EOrientation InGradientType)
+	{
+		GradientStops = InGradientStops;
+		GradientType = InGradientType;
+	}
+};
+
+struct FSlateSplinePayload : public FSlateDataPayload, public FSlateTintableElement
+{
+	TArray<FSlateGradientStop> GradientStops;
+	// Bezier Spline Data points. E.g.
+	//
+	//       P1 + - - - - + P2                P1 +
+	//         /           \                    / \
+	//     P0 *             * P3            P0 *   \   * P3
+	//                                              \ /
+	//                                               + P2	
+	FVector2D P0;
+	FVector2D P1;
+	FVector2D P2;
+	FVector2D P3;
+
+	float Thickness;
+
+	// Thickness
+	void SetThickness(float InThickness) { Thickness = InThickness; }
+	float GetThickness() const { return Thickness; }
+
+	void SetCubicBezier(const FVector2D& InP0, const FVector2D& InP1, const FVector2D& InP2, const FVector2D& InP3, float InThickness, const FLinearColor& InTint)
+	{
+		Tint = InTint;
+		P0 = InP0;
+		P1 = InP1;
+		P2 = InP2;
+		P3 = InP3;
+		Thickness = InThickness;
+	}
+
+	void SetHermiteSpline(const FVector2D& InStart, const FVector2D& InStartDir, const FVector2D& InEnd, const FVector2D& InEndDir, float InThickness, const FLinearColor& InTint)
+	{
+		Tint = InTint;
+		P0 = InStart;
+		P1 = InStart + InStartDir / 3.0f;
+		P2 = InEnd - InEndDir / 3.0f;
+		P3 = InEnd;
+		Thickness = InThickness;
+	}
+
+	void SetGradientHermiteSpline(const FVector2D& InStart, const FVector2D& InStartDir, const FVector2D& InEnd, const FVector2D& InEndDir, float InThickness, const TArray<FSlateGradientStop>& InGradientStops)
+	{
+		P0 = InStart;
+		P1 = InStart + InStartDir / 3.0f;
+		P2 = InEnd - InEndDir / 3.0f;
+		P3 = InEnd;
+		Thickness = InThickness;
+		GradientStops = InGradientStops;
+	}
+};
+
+
+struct FSlateLinePayload : public FSlateDataPayload, public FSlateTintableElement
+{ 
+	TArray<FVector2D> Points;
+	TArray<FLinearColor> PointColors;
+	float Thickness;
+
+	bool bAntialias;
+
+	bool IsAntialiased() const { return bAntialias; }
+	const TArray<FVector2D>& GetPoints() const { return Points; }
+	const TArray<FLinearColor>& GetPointColors() const { return PointColors; }
+	float GetThickness() const { return Thickness; }
+
+	void SetThickness(float InThickness)
+	{
+		Thickness = InThickness;
+	}
+
+	void SetLines(const TArray<FVector2D>& InPoints, bool bInAntialias, const TArray<FLinearColor>* InPointColors = nullptr)
+	{
+		bAntialias = bInAntialias;
+		Points = InPoints;
+		if (InPointColors)
+		{
+			PointColors = *InPointColors;
+		}
+	}
+};
+
+struct FSlateViewportPayload : public FSlateDataPayload, public FSlateTintableElement
+{
+	FSlateShaderResource* RenderTargetResource;
+	uint8 bAllowViewportScaling : 1;
+	uint8 bViewportTextureAlphaOnly : 1;
+	uint8 bRequiresVSync : 1;
+
+	void SetViewport(const TSharedPtr<const ISlateViewport>& InViewport, const FLinearColor& InTint)
+	{
+		Tint = InTint;
+		RenderTargetResource = InViewport->GetViewportRenderTargetTexture();
+		bAllowViewportScaling = InViewport->AllowScaling();
+		bViewportTextureAlphaOnly = InViewport->IsViewportTextureAlphaOnly();
+		bRequiresVSync = InViewport->RequiresVsync();
+	}
+};
+
+struct FSlateCustomDrawerPayload : public FSlateDataPayload
+{
+	// Custom drawer data
+	TWeakPtr<ICustomSlateElement, ESPMode::ThreadSafe> CustomDrawer;
+
+	void SetCustomDrawer(const TSharedPtr<ICustomSlateElement, ESPMode::ThreadSafe>& InCustomDrawer)
+	{
+		CustomDrawer = InCustomDrawer;
+	}
+};
+
+struct FSlateLayerPayload : public FSlateDataPayload
+{
+	class FSlateDrawLayerHandle* LayerHandle;
+
+	void SetLayer(FSlateDrawLayerHandle* InLayerHandle)
+	{
+		LayerHandle = InLayerHandle;
+		checkSlow(LayerHandle);
+	}
+
+};
+
+struct FSlateCachedBufferPayload : public FSlateDataPayload
+{
+	// Cached render data
+	class FSlateRenderDataHandle* CachedRenderData;
+	FVector2D CachedRenderDataOffset;
+
+	// Cached Buffers
+	void SetCachedBuffer(FSlateRenderDataHandle* InRenderDataHandle, const FVector2D& Offset)
+	{
+		check(InRenderDataHandle);
+
+		CachedRenderData = InRenderDataHandle;
+		CachedRenderDataOffset = Offset;
+	}
+};
+
+struct FSlateCustomVertsPayload : public FSlateDataPayload
+{
+	const FSlateShaderResourceProxy* ResourceProxy;
+
+	TArray<FSlateVertex> Vertices;
+	TArray<SlateIndex> Indices;
+
+	// Instancing support
+	ISlateUpdatableInstanceBuffer* InstanceData;
+	uint32 InstanceOffset;
+	uint32 NumInstances;
+
+	void SetCustomVerts(const FSlateShaderResourceProxy* InRenderProxy, const TArray<FSlateVertex>& InVerts, const TArray<SlateIndex>& InIndices, ISlateUpdatableInstanceBuffer* InInstanceData, uint32 InInstanceOffset, uint32 InNumInstances)
+	{
+		ResourceProxy = InRenderProxy;
+
+		Vertices = InVerts;
+		Indices = InIndices;
+
+		InstanceData = InInstanceData;
+		InstanceOffset = InInstanceOffset;
+		NumInstances = InNumInstances;
+	}
+};
+
+struct FSlatePostProcessPayload : public FSlateDataPayload
+{
+	// Post Process Data
+	FVector4 PostProcessData;
+	int32 DownsampleAmount;
+};
\ No newline at end of file
diff --git a/Engine/Source/Runtime/SlateCore/Public/Rendering/DrawElements.h b/Engine/Source/Runtime/SlateCore/Public/Rendering/DrawElements.h
--- a/Engine/Source/Runtime/SlateCore/Public/Rendering/DrawElements.h
+++ b/Engine/Source/Runtime/SlateCore/Public/Rendering/DrawElements.h
@@ -14,289 +14,36 @@
 #include "Layout/Clipping.h"
 #include "Types/PaintArgs.h"
 #include "Layout/Geometry.h"
-#include "Containers/StaticArray.h"
 #include "Rendering/ShaderResourceManager.h"
 #include "Rendering/RenderingCommon.h"
 #include "Debugging/SlateDebugging.h"
+#include "DrawElementPayloads.h"
+#include "ElementBatcher.h"
+
 
-class FSlateDrawLayerHandle;
 class FSlateRenderBatch;
-class FSlateRenderDataHandle;
 class FSlateWindowElementList;
-class ILayoutCache;
 class SWidget;
 class SWindow;
 
-DECLARE_MEMORY_STAT_EXTERN(TEXT("Vertex/Index Buffer Pool Memory (CPU)"), STAT_SlateBufferPoolMemory, STATGROUP_SlateMemory, SLATECORE_API );
-
-
-struct FSlateGradientStop
-{
-	FVector2D Position;
-	FLinearColor Color;
-
-	/**
-	* Construct a Gradient Stop from a Position and a Color.
-	* @param InPosition - The position in widget space for this stop. Both X and Y are used for a single-axis gradient.
-						  A two stop gradient should go from (0,0), to (Width,Height).
-	* @param InColor	- The color to lerp towards at this stop.
-	*/
-	FSlateGradientStop( const FVector2D& InPosition, const FLinearColor& InColor )
-		: Position(InPosition)
-		, Color(InColor)
-	{
-
-	}
-};
-
-template <> struct TIsPODType<FSlateGradientStop> { enum { Value = true }; };
-
-
 class FSlateDrawLayerHandle;
 
-
-class FSlateDataPayload
+enum class EElementType : uint8
 {
-public:
-	SLATECORE_API static FSlateShaderResourceManager* ResourceManager;
-
-	FSlateDataPayload()
-		: Tint(FLinearColor::White)
-		, ResourceProxy(nullptr)
-		, RenderTargetResource(nullptr)
-		, BatchFlags(ESlateBatchDrawFlag::None)
-		, bViewportTextureAlphaOnly(false)
-		, bRequiresVSync(false)
-		, CustomDrawer()
-		, InstanceData(nullptr)
-		, InstanceOffset(0)
-		, NumInstances(0)
-	{ }
-
-	void SetGradientPayloadProperties( const TArray<FSlateGradientStop>& InGradientStops, EOrientation InGradientType )
-	{
-		GradientStops = InGradientStops;
-		GradientType = InGradientType;
-	}
-
-	void SetCubicBezierPayloadProperties(const FVector2D& InP0, const FVector2D& InP1, const FVector2D& InP2, const FVector2D& InP3, float InThickness, const FLinearColor& InTint)
-	{
-		Tint = InTint;
-		P0 = InP0;
-		P1 = InP1;
-		P2 = InP2;
-		P3 = InP3;
-		Thickness = InThickness;
-	}
-	
-	void SetHermiteSplinePayloadProperties( const FVector2D& InStart, const FVector2D& InStartDir, const FVector2D& InEnd, const FVector2D& InEndDir, float InThickness, const FLinearColor& InTint )
-	{
-		Tint = InTint;
-		P0 = InStart;
-		P1 = InStart + InStartDir / 3.0f;
-		P2 = InEnd - InEndDir / 3.0f;
-		P3 = InEnd;
-		Thickness = InThickness;
-	}
-
-	void SetGradientHermiteSplinePayloadProperties( const FVector2D& InStart, const FVector2D& InStartDir, const FVector2D& InEnd, const FVector2D& InEndDir, float InThickness, const TArray<FSlateGradientStop>& InGradientStops )
-	{
-		P0 = InStart;
-		P1 = InStart + InStartDir / 3.0f;
-		P2 = InEnd - InEndDir / 3.0f;
-		P3 = InEnd;
-		Thickness = InThickness;
-		GradientStops = InGradientStops;
-	}
-
-	void SetViewportPayloadProperties( const TSharedPtr<const ISlateViewport>& InViewport, const FLinearColor& InTint )
-	{
-		Tint = InTint;
-		RenderTargetResource = InViewport->GetViewportRenderTargetTexture();
-		bAllowViewportScaling = InViewport->AllowScaling();
-		bViewportTextureAlphaOnly = InViewport->IsViewportTextureAlphaOnly();
-		bRequiresVSync = InViewport->RequiresVsync();
-	}
-
-	void SetCustomDrawerPayloadProperties( const TSharedPtr<ICustomSlateElement, ESPMode::ThreadSafe>& InCustomDrawer )
-	{
-		CustomDrawer = InCustomDrawer;
-	}
-
-	void SetCustomVertsPayloadProperties(const FSlateShaderResourceProxy* InRenderProxy, const TArray<FSlateVertex>& InVerts, const TArray<SlateIndex>& InIndexes, ISlateUpdatableInstanceBuffer* InInstanceData, uint32 InInstanceOffset, uint32 InNumInstances)
-	{
-		ResourceProxy = InRenderProxy;
-		CustomVertsData = InVerts;
-		CustomVertsIndexData = InIndexes;
-		InstanceData = InInstanceData;
-		InstanceOffset = InInstanceOffset;
-		NumInstances = InNumInstances;
-	}
-
-	void SetLayerPayloadProperties(FSlateDrawLayerHandle* InLayerHandle)
-	{
-		LayerHandle = InLayerHandle;
-		checkSlow(LayerHandle);
-	}
-
-	// Batch Flags
-	FORCEINLINE const ESlateBatchDrawFlag GetBatchFlags() const { return BatchFlags; }
-
-	// Brush
-	void SetBrush(const FSlateBrush* InBrush)
-	{
-		check(InBrush);
-		ensureMsgf(InBrush->GetDrawType() != ESlateBrushDrawType::NoDrawType, TEXT("This should have been filtered out earlier in the Make... call."));
-
-		// The slate brush ptr can't be trusted after batch elements.
-		BrushMargin = InBrush->GetMargin();
-		BrushUVRegion = InBrush->GetUVRegion();
-		BrushTiling = InBrush->GetTiling();
-		BrushMirroring = InBrush->GetMirroring();
-		BrushDrawType = InBrush->GetDrawType();
-
-		const FSlateResourceHandle& ResourceHandle = InBrush->GetRenderingResource();
-		ResourceProxy = ResourceHandle.GetResourceProxy();
-	}
-
-	FORCEINLINE const FMargin& GetBrushMargin() const { return BrushMargin; }
-	FORCEINLINE const FBox2D& GetBrushUVRegion() const { return BrushUVRegion; }
-	FORCEINLINE ESlateBrushTileType::Type GetBrushTiling() const { return BrushTiling; }
-	FORCEINLINE ESlateBrushMirrorType::Type GetBrushMirroring() const { return BrushMirroring; }
-	FORCEINLINE ESlateBrushDrawType::Type GetBrushDrawType() const { return BrushDrawType; }
-	FORCEINLINE const FSlateShaderResourceProxy* GetResourceProxy() const { return ResourceProxy; }
-
-	// Tinting
-	FORCEINLINE void SetTint(FLinearColor InTint) { Tint = InTint; }
-	FORCEINLINE FLinearColor GetTint() const { return Tint; }
-
-	// Thickness
-	FORCEINLINE void SetThickness(float InThickness) { Thickness = InThickness; }
-	FORCEINLINE float GetThickness() const { return Thickness; }
-
-	// Text
-	void SetText(FSlateWindowElementList& ElementList, const FString& Text, const FSlateFontInfo& FontInfo, const int32 StartIndex = 0, const int32 EndIndex = MAX_int32);
-
-	FORCEINLINE const FSlateFontInfo& GetFontInfo() const { return FontInfo; }
-	FORCEINLINE const TCHAR* GetText() const { return ImmutableText; }
-	FORCEINLINE int32 GetTextLength() const { return TextLength; }
-
-	// Shaped Text
-	void SetShapedText(FSlateWindowElementList& ElementList, const FShapedGlyphSequencePtr& InShapedGlyphSequence, FLinearColor InOutlineTint);
-
-	FORCEINLINE FShapedGlyphSequencePtr GetShapedGlyphSequence() const { return ShapedGlyphSequence; }
-	FORCEINLINE FLinearColor GetOutlineTint() const { return OutlineTint; }
-
-	// Lines
-	void SetLines(FSlateWindowElementList& ElementList, const TArray<FVector2D>& InPoints, bool bInAntialias, const TArray<FLinearColor>* InPointColors = nullptr);
-
-	FORCEINLINE bool IsAntialiased() const { return bAntialias; }
-	FORCEINLINE uint16 GetNumPoints() const { return NumPoints; }
-	FORCEINLINE const FVector2D* GetPoints() const { return Points; }
-	FORCEINLINE const FLinearColor* GetPointColors() const { return PointColors; }
-
-	// Cached Buffers
-	void SetCachedBuffer(FSlateRenderDataHandle* InRenderDataHandle, const FVector2D& Offset)
-	{
-		check(InRenderDataHandle);
-
-		CachedRenderData = InRenderDataHandle;
-		CachedRenderDataOffset = Offset;
-	}
-
-	FORCEINLINE class FSlateRenderDataHandle* GetRenderDataHandle() const { return CachedRenderData; }
-	FORCEINLINE FVector2D GetRenderOffset() const { return CachedRenderDataOffset; }
-
-protected:
-	FLinearColor Tint;
-
-	// The slate brush ptr can't be trusted after batch elements, so we copy everything when slate brush is set.
-	FMargin BrushMargin;
-	FBox2D BrushUVRegion;
-	TEnumAsByte<enum ESlateBrushTileType::Type> BrushTiling;
-	TEnumAsByte<enum ESlateBrushMirrorType::Type> BrushMirroring;
-	TEnumAsByte<enum ESlateBrushDrawType::Type> BrushDrawType;
-
-	// The resource proxy that we actually render with
-	const FSlateShaderResourceProxy* ResourceProxy;
-
-public:
-	//TODO add getter
-	// Render target 'brush'
-	FSlateShaderResource* RenderTargetResource;
-
-protected:
-	// The font to use when rendering
-	FSlateFontInfo FontInfo;
-	// The null-terminated string data
-	TCHAR* ImmutableText;
-	// The length of the text (excluding the null-terminator, matching strlen)
-	int32 TextLength;
-
-	// Spline/Line Data
-	float Thickness;
-
-	// Shaped text data
-	FShapedGlyphSequencePtr ShapedGlyphSequence;
-	// Element outline tint
-	FLinearColor OutlineTint;
-
-	// Line data - allocated with FSlateWindowElementList::Alloc
-	uint16 NumPoints;
-
-	FVector2D* Points;
-	FLinearColor* PointColors;
-
-	// Cached render data
-	class FSlateRenderDataHandle* CachedRenderData;
-	FVector2D CachedRenderDataOffset;
-
-public:
-	// Misc data
-	ESlateBatchDrawFlag BatchFlags;
-
-public:
-	// Bezier Spline Data points. E.g.
-	//
-	//       P1 + - - - - + P2                P1 +
-	//         /           \                    / \
-	//     P0 *             * P3            P0 *   \   * P3
-	//                                              \ /
-	//                                               + P2	
-	FVector2D P0;
-	FVector2D P1;
-	FVector2D P2;
-	FVector2D P3;
-
-	// Gradient data (fixme, this should be allocated with FSlateWindowElementList::Alloc)
-	TArray<FSlateGradientStop> GradientStops;
-	EOrientation GradientType;
-
-	// Viewport data
-	bool bAllowViewportScaling : 1;
-	bool bViewportTextureAlphaOnly : 1;
-	bool bRequiresVSync : 1;
-	// Whether or not to anti-alias lines
-	bool bAntialias : 1;
-
-	// Custom drawer data
-	TWeakPtr<ICustomSlateElement, ESPMode::ThreadSafe> CustomDrawer;
-
-	// Custom verts data, TODO FSlateWindowElementList::Alloc)?
-	TArray<FSlateVertex> CustomVertsData;
-	TArray<SlateIndex> CustomVertsIndexData;
-
-	// Instancing support
-	ISlateUpdatableInstanceBuffer* InstanceData;
-	uint32 InstanceOffset;
-	uint32 NumInstances;
-
-	// Layer handle
-	FSlateDrawLayerHandle* LayerHandle;
-
-	// Post Process Data
-	FVector4 PostProcessData;
-	int32 DownsampleAmount;
+	ET_Box,
+	ET_DebugQuad,
+	ET_Text,
+	ET_ShapedText,
+	ET_Spline,
+	ET_Line,
+	ET_Gradient,
+	ET_Viewport,
+	ET_Border,
+	ET_Custom,
+	ET_CustomVerts,
+	ET_PostProcessPass,
+	/** Total number of draw commands */
+	ET_Count,
 };
 
 
@@ -306,40 +53,9 @@ public:
  */
 class FSlateDrawElement
 {
+	friend class FSlateWindowElementList;
 public:
-	enum EElementType
-	{
-		ET_Box,
-		ET_DebugQuad,
-		ET_Text,
-		ET_ShapedText,
-		ET_Spline,
-		ET_Line,
-		ET_Gradient,
-		ET_Viewport,
-		ET_Border,
-		ET_Custom,
-		ET_CustomVerts,
-		/**
-		* Used for Invalidation, these buffers represent a complete cached buffer of what we normally send to the GPU to be
-		* drawn for a series of widgets.  They're used to reduce draw overhead in situations where the UI is largely static.
-		*/
-		ET_CachedBuffer,
-		/**
-		* These layers are different from "layerId", they're symbolic layers, used when building up cached geometry.  They allow
-		* Slate to semantically differentiate between Layer A and Layer B, which may have completely different layerIds, which perhaps
-		* overlap, but because they are in logically separate layers they won't intersect, the contents of Layer B would always
-		* come after the contents of Layer A.
-		*/
-		ET_Layer,
-		/**
-		*
-		*/
-		ET_PostProcessPass,
-		/** Total number of draw commands */
-		ET_Count,
-	};
-
+	
 	enum ERotationSpace
 	{
 		/** Relative to the element.  (0,0) is the upper left corner of the element */
@@ -542,16 +258,21 @@ public:
 	
 	SLATECORE_API static void MakeCustomVerts(FSlateWindowElementList& ElementList, uint32 InLayer, const FSlateResourceHandle& InRenderResourceHandle, const TArray<FSlateVertex>& InVerts, const TArray<SlateIndex>& InIndexes, ISlateUpdatableInstanceBuffer* InInstanceData, uint32 InInstanceOffset, uint32 InNumInstances, ESlateDrawEffect InDrawEffects = ESlateDrawEffect::None);
 
-	SLATECORE_API static void MakeCachedBuffer(FSlateWindowElementList& ElementList, uint32 InLayer, TSharedPtr<FSlateRenderDataHandle, ESPMode::ThreadSafe>& CachedRenderDataHandle, const FVector2D& Offset);
+	SLATECORE_API static void MakePostProcessPass(FSlateWindowElementList& ElementList, uint32 InLayer, const FPaintGeometry& PaintGeometry, const FVector4& Params, int32 DownsampleAmount);
 
-	SLATECORE_API static void MakeLayer(FSlateWindowElementList& ElementList, uint32 InLayer, TSharedPtr<FSlateDrawLayerHandle, ESPMode::ThreadSafe>& DrawLayerHandle);
+	FSlateDrawElement();
+	~FSlateDrawElement();
 
-	SLATECORE_API static void MakePostProcessPass(FSlateWindowElementList& ElementList, uint32 InLayer, const FPaintGeometry& PaintGeometry, const FVector4& Params, int32 DownsampleAmount);
+	FORCEINLINE int32 GetLayer() const { return LayerId; }
 
 	FORCEINLINE EElementType GetElementType() const { return ElementType; }
-	FORCEINLINE const FSlateDataPayload& GetDataPayload() const { return DataPayload; }
 
-	FORCEINLINE int16 GetLayer() const { return Layer; }
+	template<typename PayloadType>
+	FORCEINLINE const PayloadType& GetDataPayload() const { return *(PayloadType*)DataPayload; }
+
+	template<typename PayloadType>
+	FORCEINLINE PayloadType& GetDataPayload() { return *(PayloadType*)DataPayload; }
+
 	FORCEINLINE const FSlateRenderTransform& GetRenderTransform() const { return RenderTransform; }
 	FORCEINLINE void SetRenderTransform(const FSlateRenderTransform& InRenderTransform) { RenderTransform = InRenderTransform; }
 	FORCEINLINE const FVector2D& GetPosition() const { return Position; }
@@ -559,16 +280,26 @@ public:
 	FORCEINLINE const FVector2D& GetLocalSize() const { return LocalSize; }
 	FORCEINLINE float GetScale() const { return Scale; }
 	FORCEINLINE ESlateDrawEffect GetDrawEffects() const { return DrawEffects; }
+	FORCEINLINE ESlateBatchDrawFlag GetBatchFlags() const { return BatchFlags; }
 	FORCEINLINE bool IsPixelSnapped() const { return !EnumHasAllFlags(DrawEffects, ESlateDrawEffect::NoPixelSnapping); }
-	FORCEINLINE const int16 GetClippingIndex() const { return ClippingIndex; }
-	FORCEINLINE void SetClippingIndex(const int32 InClippingIndex) { ClippingIndex = InClippingIndex; }
+
+	FORCEINLINE int32 GetPrecachedClippingIndex() const { return ClipStateHandle.GetPrecachedClipIndex(); }
+	FORCEINLINE void SetPrecachedClippingIndex(int32 InClippingIndex) { ClipStateHandle.SetPreCachedClipIndex(InClippingIndex); }
+
+	FORCEINLINE void SetCachedClippingState(const FSlateClippingState* CachedState) { ClipStateHandle.SetCachedClipState(CachedState); }
+	FORCEINLINE const FClipStateHandle& GetClippingHandle() const { return ClipStateHandle; }
 	FORCEINLINE const int8 GetSceneIndex() const { return SceneIndex; }
 
+	FORCEINLINE void SetIsCached(bool bInIsCached) { bIsCached = bInIsCached; }
+	FORCEINLINE bool IsCached() const { return bIsCached; }
+
 	FORCEINLINE FSlateLayoutTransform GetInverseLayoutTransform() const
 	{
 		return Inverse(FSlateLayoutTransform(Scale, Position));
 	}
 
+	void AddReferencedObjects(FReferenceCollector& Collector);
+
 	/**
 	* Update element cached position with an arbitrary offset
 	*
@@ -577,808 +308,170 @@ public:
 	*/
 	void ApplyPositionOffset(const FVector2D& InOffset);
 
-private:
-	void Init(FSlateWindowElementList& ElementList, EElementType InElementType, uint32 InLayer, const FPaintGeometry& PaintGeometry, ESlateDrawEffect InDrawEffects);
-
-	static bool ShouldCull(const FSlateWindowElementList& ElementList);
-
-	FORCEINLINE static bool ShouldCull(const FSlateWindowElementList& ElementList, const FPaintGeometry& PaintGeometry)
-	{
-		const FVector2D& LocalSize = PaintGeometry.GetLocalSize();
-		if (LocalSize.X == 0 || LocalSize.Y == 0)
-		{
-			return true;
-		}
-
-		return ShouldCull(ElementList);
-	}
-
-	static bool ShouldCull(const FSlateWindowElementList& ElementList, const FPaintGeometry& PaintGeometry, const FSlateBrush* InBrush);
-
-	FORCEINLINE static bool ShouldCull(const FSlateWindowElementList& ElementList, const FPaintGeometry& PaintGeometry, const FLinearColor& InTint)
-	{
-		if (InTint.A == 0 || ShouldCull(ElementList, PaintGeometry))
-		{
-			return true;
-		}
-
-		return false;
-	}
-
-	FORCEINLINE static bool ShouldCull(const FSlateWindowElementList& ElementList, const FPaintGeometry& PaintGeometry, const FLinearColor& InTint, const FString& InText)
-	{
-		if (InTint.A == 0 || InText.Len() == 0 || ShouldCull(ElementList, PaintGeometry))
-		{
-			return true;
-		}
-
-		return false;
-	}
+	UE_DEPRECATED(4.23, "GetClippingIndex has been deprecated.  If you were using this please use GetPrecachedClippingIndex instead.")
+	FORCEINLINE const int32 GetClippingIndex() const { return GetPrecachedClippingIndex(); }
 
-	FORCEINLINE static bool ShouldCull(const FSlateWindowElementList& ElementList, const FPaintGeometry& PaintGeometry, const FSlateBrush* InBrush, const FLinearColor& InTint)
-	{
-		if (InTint.A == 0 || ShouldCull(ElementList, PaintGeometry, InBrush))
-		{
-			return true;
-		}
+	UE_DEPRECATED(4.23, "SetClippingIndex has been deprecated.  If you were using this please use SetPrecachedClippingIndex instead.")
+	FORCEINLINE void SetClippingIndex(const int32 InClippingIndex) { SetPrecachedClippingIndex(InClippingIndex); }
 
-		return false;
-	}
+private:
+	void Init(FSlateWindowElementList& ElementList, EElementType InElementType, uint32 InLayer, const FPaintGeometry& PaintGeometry, ESlateDrawEffect InDrawEffects);
 
 	static FVector2D GetRotationPoint( const FPaintGeometry& PaintGeometry, const TOptional<FVector2D>& UserRotationPoint, ERotationSpace RotationSpace );
-
+	static FSlateDrawElement& MakeBoxInternal(FSlateWindowElementList& ElementList, uint32 InLayer, const FPaintGeometry& PaintGeometry, const FSlateBrush* InBrush, ESlateDrawEffect InDrawEffects, const FLinearColor& InTint);
 private:
+	FSlateDataPayload* DataPayload;
 	FSlateRenderTransform RenderTransform;
-
 	FVector2D Position;
 	FVector2D LocalSize;
+	int32 LayerId;
+	FClipStateHandle ClipStateHandle;
 	float Scale;
-
-	int16 Layer;
-	int16 ClippingIndex;
 	int8 SceneIndex;
 	ESlateDrawEffect DrawEffects;
 	EElementType ElementType;
-
-	FSlateDataPayload DataPayload;
+	// Misc data
+	ESlateBatchDrawFlag BatchFlags;
+	uint8 bIsCached : 1;
 };
 
-
-/**
- * Shader parameters for slate
- */
-struct FShaderParams
+struct FSlateCachedFastPathRenderingData
 {
-	/** Pixel shader parameters */
-	FVector4 PixelParams;
-	FVector4 PixelParams2;
-
-	FShaderParams()
-		: PixelParams( 0,0,0,0 )
-		, PixelParams2( 0,0,0,0 ) 
-	{}
-
-	FShaderParams( const FVector4& InPixelParams, const FVector4& InPixelParams2 = FVector4(0) )
-		: PixelParams( InPixelParams )
-		, PixelParams2( InPixelParams2 )
-	{}
-
-	bool operator==( const FShaderParams& Other ) const
-	{
-		return PixelParams == Other.PixelParams && PixelParams2 == Other.PixelParams2;
-	}
-
-	static FShaderParams MakePixelShaderParams( const FVector4& PixelShaderParams, const FVector4& InPixelShaderParams2 = FVector4(0) )
+	~FSlateCachedFastPathRenderingData()
 	{
-		return FShaderParams( PixelShaderParams, InPixelShaderParams2);
+		for (FSlateCachedClipState* State : CachedClipStates)
+		{
+			State->EndUsingState();
+		}
 	}
-};
 
-class FSlateRenderer;
-class FSlateRenderBatch;
-
-class ISlateRenderDataManager
-{
-public:
-	virtual void BeginReleasingRenderData(const FSlateRenderDataHandle* RenderHandle) = 0;
+	TArray<class FSlateCachedClipState*, TInlineAllocator<1>> CachedClipStates;
+	FSlateVertexArray Vertices;
+	FSlateIndexArray Indices;
 };
 
-class SLATECORE_API FSlateRenderDataHandle : public TSharedFromThis < FSlateRenderDataHandle, ESPMode::ThreadSafe >
+struct FSlateCachedElementList
 {
-public:
-	FSlateRenderDataHandle(const ILayoutCache* Cacher, ISlateRenderDataManager* InManager);
-
-	virtual ~FSlateRenderDataHandle();
-	
-	void Disconnect();
-
-	const ILayoutCache* GetCacher() const { return Cacher; }
-
-	void SetRenderBatches(TArray<FSlateRenderBatch>* InRenderBatches) { RenderBatches = InRenderBatches; }
-	TArray<FSlateRenderBatch>* GetRenderBatches() { return RenderBatches; }
-
-	void BeginUsing() { FPlatformAtomics::InterlockedIncrement(&UsageCount); }
-	void EndUsing() { FPlatformAtomics::InterlockedDecrement(&UsageCount); }
-
-	bool IsInUse() const { return UsageCount > 0; }
-
-private:
-	const ILayoutCache* Cacher;
-	ISlateRenderDataManager* Manager;
-	TArray<FSlateRenderBatch>* RenderBatches;
+#if WITH_SLATE_DEBUGGING
+	FSlateCachedElementList(const SWidget* InWidget)
+		: Widget(InWidget)
+		, CachedRenderingData(nullptr)
+		, bNewData(false)
+#else
+	FSlateCachedElementList()
+		: CachedRenderingData(nullptr)
+		, bNewData(false)
 
-	volatile int32 UsageCount;
-};
+#endif
+	{}
 
-/** 
- * Represents an element batch for rendering. 
- */
-class FSlateElementBatch
-{
-public:
-	FSlateElementBatch(const FSlateShaderResource* InShaderResource, const FShaderParams& InShaderParams, ESlateShader::Type ShaderType, ESlateDrawPrimitive::Type PrimitiveType, ESlateDrawEffect DrawEffects, ESlateBatchDrawFlag DrawFlags, const int32 InClippingIndex, const TArray<FSlateClippingState>& ClippingStates, int32 InstanceCount = 0, uint32 InstanceOffset = 0, ISlateUpdatableInstanceBuffer* InstanceData = nullptr, int32 SceneIndex = -1)
-		: BatchKey(InShaderParams, ShaderType, PrimitiveType, DrawEffects, DrawFlags, InClippingIndex, InstanceCount, InstanceOffset, InstanceData, SceneIndex)
-		, ShaderResource(InShaderResource)
-		, NumElementsInBatch(0)
-		, VertexArrayIndex(INDEX_NONE)
-		, IndexArrayIndex(INDEX_NONE)
+	void Initialize()
 	{
-		SaveClippingState(ClippingStates);
+		CachedRenderingData = new FSlateCachedFastPathRenderingData;
 	}
 
-	FSlateElementBatch( TWeakPtr<ICustomSlateElement, ESPMode::ThreadSafe> InCustomDrawer, const int32 InClippingIndex, const TArray<FSlateClippingState>& ClippingStates)
-		: BatchKey( InCustomDrawer, InClippingIndex)
-		, ShaderResource( nullptr )
-		, NumElementsInBatch(0)
-		, VertexArrayIndex(INDEX_NONE)
-		, IndexArrayIndex(INDEX_NONE)
+	~FSlateCachedElementList()
 	{
-		SaveClippingState(ClippingStates);
+		DestroyCachedVertexData();
 	}
 
-	FSlateElementBatch( TSharedPtr<FSlateRenderDataHandle, ESPMode::ThreadSafe> InCachedRenderHandle, FVector2D InCachedRenderDataOffset, const int32 InClippingIndex, const TArray<FSlateClippingState>& ClippingStates)
-		: BatchKey( InCachedRenderHandle, InCachedRenderDataOffset, InClippingIndex)
-		, ShaderResource( nullptr )
-		, NumElementsInBatch(0)
-		, VertexArrayIndex(INDEX_NONE)
-		, IndexArrayIndex(INDEX_NONE)
-	{
-		SaveClippingState(ClippingStates);
-	}
+	void Reset();
 
-	FSlateElementBatch( TSharedPtr<FSlateDrawLayerHandle, ESPMode::ThreadSafe> InLayerHandle, const int32 InClippingIndex, const TArray<FSlateClippingState>& ClippingStates)
-		: BatchKey( InLayerHandle, InClippingIndex )
-		, ShaderResource( nullptr )
-		, NumElementsInBatch(0)
-		, VertexArrayIndex(INDEX_NONE)
-		, IndexArrayIndex(INDEX_NONE)
-	{
-		SaveClippingState(ClippingStates);
-	}
+	FSlateRenderBatch& AddRenderBatch(int32 InLayer, const FShaderParams& InShaderParams, const FSlateShaderResource* InResource, ESlateDrawPrimitive InPrimitiveType, ESlateShader InShaderType, ESlateDrawEffect InDrawEffects, ESlateBatchDrawFlag InDrawFlags, int8 SceneIndex);
 
-	void SaveClippingState(const TArray<FSlateClippingState>& ClippingStates)
-	{
-		// Store the clipping state so we can use it later for rendering.
-		if (ClippingStates.IsValidIndex(GetClippingIndex()))
-		{
-			ClippingState = ClippingStates[GetClippingIndex()];
-		}
-	}
+	void AddCachedClipState(FSlateCachedClipState& ClipStateToCache);
 
-	bool operator==( const FSlateElementBatch& Other ) const
-	{	
-		return BatchKey == Other.BatchKey && ShaderResource == Other.ShaderResource;
-	}
+	void AddReferencedObjects(FReferenceCollector& Collector);
 
-	const FSlateShaderResource* GetShaderResource() const { return ShaderResource; }
-	const FShaderParams& GetShaderParams() const { return BatchKey.ShaderParams; }
-	ESlateBatchDrawFlag GetDrawFlags() const { return BatchKey.DrawFlags; }
-	ESlateDrawPrimitive::Type GetPrimitiveType() const { return BatchKey.DrawPrimitiveType; }
-	ESlateShader::Type GetShaderType() const { return BatchKey.ShaderType; } 
-	ESlateDrawEffect GetDrawEffects() const { return BatchKey.DrawEffects; }
-	const int32 GetClippingIndex() const { return BatchKey.ClippingIndex; }
-	const TWeakPtr<ICustomSlateElement, ESPMode::ThreadSafe> GetCustomDrawer() const { return BatchKey.CustomDrawer; }
-	const TSharedPtr<FSlateRenderDataHandle, ESPMode::ThreadSafe> GetCachedRenderHandle() const { return BatchKey.CachedRenderHandle; }
-	FVector2D GetCachedRenderDataOffset() const { return BatchKey.CachedRenderDataOffset; }
-	const TSharedPtr<FSlateDrawLayerHandle, ESPMode::ThreadSafe> GetLayerHandle() const { return BatchKey.LayerHandle; }
-	int32 GetInstanceCount() const{ return BatchKey.InstanceCount; }
-	uint32 GetInstanceOffset() const{ return BatchKey.InstanceOffset; }
-	const ISlateUpdatableInstanceBuffer* GetInstanceData() const { return BatchKey.InstanceData; }
-	int32 GetSceneIndex() const { return BatchKey.SceneIndex; }
 private:
-	struct FBatchKey
-	{
-		const TWeakPtr<ICustomSlateElement, ESPMode::ThreadSafe> CustomDrawer;
-		const TSharedPtr<FSlateRenderDataHandle, ESPMode::ThreadSafe> CachedRenderHandle;
-		const FVector2D CachedRenderDataOffset;
-		const TSharedPtr<FSlateDrawLayerHandle, ESPMode::ThreadSafe> LayerHandle;
-		const FShaderParams ShaderParams;
-		const ESlateBatchDrawFlag DrawFlags;	
-		const ESlateShader::Type ShaderType;
-		const ESlateDrawPrimitive::Type DrawPrimitiveType;
-		const ESlateDrawEffect DrawEffects;
-		const int32 ClippingIndex;
-		const int32 InstanceCount;
-		const uint32 InstanceOffset;
-		const ISlateUpdatableInstanceBuffer* InstanceData;
-		const int32 SceneIndex;
-
-		FBatchKey(const FShaderParams& InShaderParams, ESlateShader::Type InShaderType, ESlateDrawPrimitive::Type InDrawPrimitiveType, ESlateDrawEffect InDrawEffects, ESlateBatchDrawFlag InDrawFlags, const int32 InClippingIndex, int32 InInstanceCount, uint32 InInstanceOffset, ISlateUpdatableInstanceBuffer* InInstanceBuffer, int32 InSceneIndex)
-			: ShaderParams(InShaderParams)
-			, DrawFlags(InDrawFlags)
-			, ShaderType(InShaderType)
-			, DrawPrimitiveType(InDrawPrimitiveType)
-			, DrawEffects(InDrawEffects)
-			, ClippingIndex(InClippingIndex)
-			, InstanceCount(InInstanceCount)
-			, InstanceOffset(InInstanceOffset)
-			, InstanceData(InInstanceBuffer)
-			, SceneIndex(InSceneIndex)
-		{
-		}
-
-		FBatchKey( const FShaderParams& InShaderParams, ESlateShader::Type InShaderType, ESlateDrawPrimitive::Type InDrawPrimitiveType, ESlateDrawEffect InDrawEffects, ESlateBatchDrawFlag InDrawFlags, const int32 InClippingIndex, int32 InInstanceCount, uint32 InInstanceOffset, ISlateUpdatableInstanceBuffer* InInstanceBuffer)
-			: ShaderParams( InShaderParams )
-			, DrawFlags( InDrawFlags )
-			, ShaderType( InShaderType )
-			, DrawPrimitiveType( InDrawPrimitiveType )
-			, DrawEffects( InDrawEffects )
-			, ClippingIndex(InClippingIndex)
-			, InstanceCount( InInstanceCount )
-			, InstanceOffset( InInstanceOffset )
-			, InstanceData(InInstanceBuffer)
-			, SceneIndex(-1)
-		{
-		}
-
-		FBatchKey( TWeakPtr<ICustomSlateElement, ESPMode::ThreadSafe> InCustomDrawer, const int32 InClippingIndex)
-			: CustomDrawer( InCustomDrawer )
-			, ShaderParams()
-			, DrawFlags( ESlateBatchDrawFlag::None )
-			, ShaderType( ESlateShader::Default )
-			, DrawPrimitiveType( ESlateDrawPrimitive::TriangleList )
-			, DrawEffects( ESlateDrawEffect::None )
-			, ClippingIndex(InClippingIndex)
-			, InstanceCount(0)
-			, InstanceOffset(0)
-			, SceneIndex(-1)
-		{}
-
-		FBatchKey( TSharedPtr<FSlateRenderDataHandle, ESPMode::ThreadSafe> InCachedRenderHandle, FVector2D InCachedRenderDataOffset, const int32 InClippingIndex)
-			: CachedRenderHandle(InCachedRenderHandle)
-			, CachedRenderDataOffset(InCachedRenderDataOffset)
-			, ShaderParams()
-			, DrawFlags(ESlateBatchDrawFlag::None)
-			, ShaderType(ESlateShader::Default)
-			, DrawPrimitiveType(ESlateDrawPrimitive::TriangleList)
-			, DrawEffects(ESlateDrawEffect::None)
-			, ClippingIndex(InClippingIndex)
-			, InstanceCount(0)
-			, InstanceOffset(0)
-			, SceneIndex(-1)
-		{}
-
-		FBatchKey(TSharedPtr<FSlateDrawLayerHandle, ESPMode::ThreadSafe> InLayerHandle, const int32 InClippingIndex)
-			: LayerHandle(InLayerHandle)
-			, ShaderParams()
-			, DrawFlags(ESlateBatchDrawFlag::None)
-			, ShaderType(ESlateShader::Default)
-			, DrawPrimitiveType(ESlateDrawPrimitive::TriangleList)
-			, DrawEffects(ESlateDrawEffect::None)
-			, ClippingIndex(InClippingIndex)
-			, InstanceCount(0)
-			, InstanceOffset(0)
-			, SceneIndex(-1)
-		{}
-
-		bool operator==( const FBatchKey& Other ) const
-		{
-			return DrawFlags == Other.DrawFlags
-				&& ShaderType == Other.ShaderType
-				&& DrawPrimitiveType == Other.DrawPrimitiveType
-				&& DrawEffects == Other.DrawEffects
-				&& ShaderParams == Other.ShaderParams
-				&& ClippingIndex == Other.ClippingIndex
-				&& CustomDrawer == Other.CustomDrawer
-				&& CachedRenderHandle == Other.CachedRenderHandle
-				&& LayerHandle == Other.LayerHandle
-				&& InstanceCount == Other.InstanceCount
-				&& InstanceOffset == Other.InstanceOffset
-				&& InstanceData == Other.InstanceData
-				&& SceneIndex == Other.SceneIndex;
-		}
-	};
-
-	/** A secondary key which represents elements needed to make a batch unique */
-	FBatchKey BatchKey;
-
-	/** Shader resource to use with this batch.  Used as a primary key.  No batch can have multiple textures */
-	const FSlateShaderResource* ShaderResource;
-
+	void DestroyCachedVertexData();
 public:
-	/** Number of elements in the batch */
-	uint32 NumElementsInBatch;
-	/** Index into an array of vertex arrays where this batches vertices are found (before submitting to the vertex buffer)*/
-	int32 VertexArrayIndex;
-	/** Index into an array of index arrays where this batches indices are found (before submitting to the index buffer) */
-	int32 IndexArrayIndex;
-
-	/** The Stored clipping state for the corresponding clipping state index.  The indexes are not directly comparable later, so we need to expand it to the full state to be compared. */
-	TOptional<FSlateClippingState> ClippingState;
-};
-
-class FSlateRenderBatch
-{
-public:
-	FSlateRenderBatch(uint32 InLayer, const FSlateElementBatch& InBatch, TSharedPtr<FSlateRenderDataHandle, ESPMode::ThreadSafe> InRenderHandle, int32 InNumVertices, int32 InNumIndices, int32 InVertexOffset, int32 InIndexOffset)
-		: Layer( InLayer )
-		, ShaderParams( InBatch.GetShaderParams() )
-		, Texture( InBatch.GetShaderResource() )
-		, InstanceData( InBatch.GetInstanceData() )
-		, InstanceCount( InBatch.GetInstanceCount() )
-		, InstanceOffset(InBatch.GetInstanceOffset() )
-		, CustomDrawer( InBatch.GetCustomDrawer() )
-		, LayerHandle( InBatch.GetLayerHandle() )
-		, CachedRenderHandle( InRenderHandle )
-		, DrawFlags( InBatch.GetDrawFlags() )
-		, ShaderType( InBatch.GetShaderType() )
-		, DrawPrimitiveType( InBatch.GetPrimitiveType() )
-		, DrawEffects( InBatch.GetDrawEffects() )
-		, ClippingIndex( InBatch.GetClippingIndex() )
-		, ClippingState( InBatch.ClippingState )
-		, VertexArrayIndex( InBatch.VertexArrayIndex )
-		, IndexArrayIndex( InBatch.IndexArrayIndex )
-		, VertexOffset( InVertexOffset )
-		, IndexOffset( InIndexOffset )
-		, NumVertices( InNumVertices )
-		, NumIndices( InNumIndices )
-		, SceneIndex(InBatch.GetSceneIndex())
-	{}
-
-public:
-	/** The layer we need to sort by when  */
-	const uint32 Layer;
-
-	/** Dynamically modified offset that occurs when we have relative position stored render batches. */
-	FVector2D DynamicOffset;
-
-	const FShaderParams ShaderParams;
-
-	/** Texture to use with this batch.  */
-	const FSlateShaderResource* Texture;
-
-	const ISlateUpdatableInstanceBuffer* InstanceData;
-
-	const int32 InstanceCount;
-
-	const uint32 InstanceOffset;
-
-	const TWeakPtr<ICustomSlateElement, ESPMode::ThreadSafe> CustomDrawer;
-
-	const TWeakPtr<FSlateDrawLayerHandle, ESPMode::ThreadSafe> LayerHandle;
-
-	const TSharedPtr<FSlateRenderDataHandle, ESPMode::ThreadSafe> CachedRenderHandle;
-
-	const ESlateBatchDrawFlag DrawFlags;	
-
-	const ESlateShader::Type ShaderType;
-
-	const ESlateDrawPrimitive::Type DrawPrimitiveType;
-
-	const ESlateDrawEffect DrawEffects;
-
-	const int32 ClippingIndex;
-
-	/** The Stored clipping state for the corresponding clipping state index.  The indexes are not directly comparable later, so we need to expand it to the full state to be compared. */
-	TOptional<FSlateClippingState> ClippingState;
-	
-	/** Index into the vertex array pool */
-	const int32 VertexArrayIndex;
-	/** Index into the index array pool */
-	const int32 IndexArrayIndex;
-	/** How far into the vertex buffer is this batch*/
-	const uint32 VertexOffset;
-	/** How far into the index buffer this batch is*/
-	const uint32 IndexOffset;
-	/** Number of vertices in the batch */
-	const uint32 NumVertices;
-	/** Number of indices in the batch */
-	const uint32 NumIndices;
-
-	const int32 SceneIndex;
+#if WITH_SLATE_DEBUGGING
+	const SWidget* Widget;
+#endif
+	/** List of source draw elements to create batches from */
+	FSlateDrawElementArray DrawElements;
+	/** List of cached batches to submit for drawing */
+	TArray<FSlateRenderBatch> CachedBatches;
+	FSlateCachedFastPathRenderingData* CachedRenderingData;
+	bool bNewData;
+private:
 };
 
+typedef TDoubleLinkedList<FSlateCachedElementList>::TDoubleLinkedListNode FSlateCachedElementListNode;
 
-typedef TArray<FSlateElementBatch, TInlineAllocator<2>> FElementBatchArray;
 
-class FElementBatchMap
+struct FSlateCachedElementData
 {
-public:
-	FElementBatchMap()
-		: ResourceVersion(0)
-	{
-		Reset();
-	}
-
-	FORCEINLINE_DEBUGGABLE int32 Num() const { return Layers.Num() + OverflowLayers.Num(); }
-	
-	FORCEINLINE_DEBUGGABLE TUniqueObj<FElementBatchArray>* Find(uint32 Layer)
-	{
-		if ( Layer < (uint32)Layers.Num() )
-		{
-			if ( ActiveLayers[Layer] )
-			{
-				return &Layers[Layer];
-			}
-
-			return nullptr;
-		}
-		else
-		{
-			return OverflowLayers.Find(Layer);
-		}
-	}
-
-	FORCEINLINE_DEBUGGABLE TUniqueObj<FElementBatchArray>& Add(uint32 Layer)
-	{
-		if ( Layer < (uint32)Layers.Num() )
-		{
-			MinLayer = FMath::Min(Layer, MinLayer);
-			MaxLayer = FMath::Max(Layer, MaxLayer);
-			ActiveLayers[Layer] = true;
-			return Layers[Layer];
-		}
-		else
-		{
-			return OverflowLayers.Add(Layer);
-		}
-	}
-
-	FORCEINLINE_DEBUGGABLE void Sort()
+	void Empty()
 	{
-		OverflowLayers.KeySort(TLess<uint32>());
-	}
-
-	template <typename TFunc>
-	FORCEINLINE_DEBUGGABLE void ForEachLayer(const TFunc& Process)
-	{
-		if ( MinLayer < (uint32)Layers.Num() )
+		if (CachedElementLists.Num())
 		{
-			for ( TBitArray<>::FConstIterator BitIt(ActiveLayers, MinLayer); BitIt; ++BitIt )
-			{
-				if ( BitIt.GetValue() == 0 )
-				{
-					continue;
-				}
-
-				const int32 BitIndex = BitIt.GetIndex();
-				FElementBatchArray& ElementBatches = *Layers[BitIndex];
-
-				if ( ElementBatches.Num() > 0 )
-				{
-					Process(BitIndex, ElementBatches);
-				}
-
-				if ( ((uint32)BitIndex) >= MaxLayer )
-				{
-					break;
-				}
-			}
+			UE_LOG(LogSlate, Log, TEXT("Resetting cached element data.  Num: %d"), CachedElementLists.Num());
 		}
 
-		for ( TMap<uint32, TUniqueObj<FElementBatchArray>>::TIterator It(OverflowLayers); It; ++It )
-		{
-			uint32 Layer = It.Key();
-			FElementBatchArray& ElementBatches = *It.Value();
-
-			if ( ElementBatches.Num() > 0 )
-			{
-				Process(Layer, ElementBatches);
-			}
-		}
+		CachedElementLists.Empty();
 	}
 
-	FORCEINLINE_DEBUGGABLE void UpdateResourceVersion(uint32 NewResourceVersion)
+	void ResetCache(FSlateCachedElementListNode* ElementListNode)
 	{
-		if (ResourceVersion != NewResourceVersion)
-		{
-			OverflowLayers.Empty();
-
-			// Since the resource version changed, clean out all cached data in the element batch arrays
-			for (int32 LayerIdx = 0; LayerIdx < Layers.Num(); ++LayerIdx)
-			{
-				Layers[LayerIdx]->Empty();
-			}
-
-			MinLayer = UINT_MAX;
-			MaxLayer = 0;
-			ActiveLayers.Init(false, Layers.Num());
-
-			ResourceVersion = NewResourceVersion;
-		}
+		ElementListNode->GetValue().Reset();
 	}
 
-	FORCEINLINE_DEBUGGABLE void Reset()
+	void RemoveCache(FSlateCachedElementListNode*& ElementListNode)
 	{
-		MinLayer = UINT_MAX;
-		MaxLayer = 0;
-		ActiveLayers.Init(false, Layers.Num());
-		OverflowLayers.Reset();
+		CachedElementLists.RemoveNode(ElementListNode);
+		ElementListNode = nullptr;
 	}
 
-private:
-	TBitArray<> ActiveLayers;
-	TStaticArray<TUniqueObj<FElementBatchArray>, 256> Layers;
-	TMap<uint32, TUniqueObj<FElementBatchArray>> OverflowLayers;
-	uint32 MinLayer;
-	uint32 MaxLayer;
-	uint32 ResourceVersion;
-};
-
-#if STATS
-
-class FSlateStatTrackingMemoryAllocator : public FDefaultAllocator
-{
-public:
-	typedef FDefaultAllocator Super;
-
-	class ForAnyElementType : public FDefaultAllocator::ForAnyElementType
+	FSlateCachedElementListNode* AddCache(const SWidget* Widget)
 	{
-	public:
-		typedef FDefaultAllocator::ForAnyElementType Super;
-
-		ForAnyElementType()
-			: AllocatedSize(0)
-		{
-
-		}
-
-		/**
-		 * Moves the state of another allocator into this one.
-		 * Assumes that the allocator is currently empty, i.e. memory may be allocated but any existing elements have already been destructed (if necessary).
-		 * @param Other - The allocator to move the state from.  This allocator should be left in a valid empty state.
-		 */
-		FORCEINLINE void MoveToEmpty(ForAnyElementType& Other)
-		{
-			Super::MoveToEmpty(Other);
-
-			AllocatedSize = Other.AllocatedSize;
-			Other.AllocatedSize = 0;
-		}
-
-		/** Destructor. */
-		~ForAnyElementType()
-		{
-			if(AllocatedSize)
-			{
-				DEC_DWORD_STAT_BY(STAT_SlateBufferPoolMemory, AllocatedSize);
-			}
-		}
-
-		void ResizeAllocation(int32 PreviousNumElements, int32 NumElements, int32 NumBytesPerElement)
+#if WITH_SLATE_DEBUGGING
+		for (FSlateCachedElementList& CachedElementList : CachedElementLists)
 		{
-			const int32 NewSize = NumElements * NumBytesPerElement;
-			INC_DWORD_STAT_BY(STAT_SlateBufferPoolMemory, NewSize - AllocatedSize);
-			AllocatedSize = NewSize;
-
-			Super::ResizeAllocation(PreviousNumElements, NumElements, NumBytesPerElement);
+			ensure(CachedElementList.Widget != Widget);
 		}
 
-	private:
-		ForAnyElementType(const ForAnyElementType&);
-		ForAnyElementType& operator=(const ForAnyElementType&);
-	private:
-		int32 AllocatedSize;
-	};
-};
-
-template <>
-struct TAllocatorTraits<FSlateStatTrackingMemoryAllocator> : TAllocatorTraitsBase<FSlateStatTrackingMemoryAllocator>
-{
-	enum { SupportsMove    = TAllocatorTraits<FDefaultAllocator>::SupportsMove    };
-	enum { IsZeroConstruct = TAllocatorTraits<FDefaultAllocator>::IsZeroConstruct };
-};
-
-typedef TArray<FSlateVertex, FSlateStatTrackingMemoryAllocator> FSlateVertexArray;
-typedef TArray<SlateIndex, FSlateStatTrackingMemoryAllocator> FSlateIndexArray;
-
+		FSlateCachedElementListNode* NewNode = new FSlateCachedElementListNode(Widget);
 #else
-
-typedef TArray<FSlateVertex> FSlateVertexArray;
-typedef TArray<SlateIndex> FSlateIndexArray;
-
+		FSlateCachedElementListNode* NewNode = new FSlateCachedElementListNode(FSlateCachedElementList());
 #endif
 
 
-class FSlateBatchData
-{
-public:
-	FSlateBatchData()
-		: NumBatchedVertices(0)
-		, NumBatchedIndices(0)
-		, NumLayers(0)
-		, bIsStencilBufferRequired(false)
-	{}
-
-	void Reset();
-
-	/**
-	 * Returns a list of element batches for this window
-	 */
-	const TArray<FSlateRenderBatch>& GetRenderBatches() const { return RenderBatches; }
-
-	/**
-	 * 
-	 */
-	SLATECORE_API bool IsStencilClippingRequired() const;
-
-	void DetermineIsStencilClippingRequired(const TArray<FSlateClippingState>& ClippingStates);
-
-	/**
-	 * Assigns a vertex array from the pool which is appropriate for the batch.  Creates a new array if needed
-	 */
-	void AssignVertexArrayToBatch( FSlateElementBatch& Batch );
-
-	/**
-	 * Assigns an index array from the pool which is appropriate for the batch.  Creates a new array if needed
-	 */
-	void AssignIndexArrayToBatch( FSlateElementBatch& Batch );
-
-	/** @return the list of vertices for a batch */
-	FSlateVertexArray& GetBatchVertexList( FSlateElementBatch& Batch ) { return BatchVertexArrays[Batch.VertexArrayIndex]; }
-
-	/** @return the list of indices for a batch */
-	FSlateIndexArray& GetBatchIndexList( FSlateElementBatch& Batch ) { return BatchIndexArrays[Batch.IndexArrayIndex]; }
-
-	/** @return The total number of batched vertices */
-	int32 GetNumBatchedVertices() const { return NumBatchedVertices; }
-
-	/** @return The total number of batched indices */
-	int32 GetNumBatchedIndices() const { return NumBatchedIndices; }
-
-	/** @return Total number of batched layers */
-	int32 GetNumLayers() const { return NumLayers; }
-
-	/** Set the associated vertex/index buffer handle. */
-	void SetRenderDataHandle(TSharedPtr<FSlateRenderDataHandle, ESPMode::ThreadSafe> InRenderDataHandle) { RenderDataHandle = InRenderDataHandle; }
-
-	/** 
-	 * Fills batch data into the actual vertex and index buffer
-	 *
-	 * @param VertexBuffer	Pointer to the actual memory for the vertex buffer 
-	 * @param IndexBuffer	Pointer to the actual memory for an index buffer
-	 * @param bAbsoluteIndices	Whether to write absolute indices (simplifies draw call setup on RHIs that do not support BaseVertex)
-	 */
-	SLATECORE_API void FillVertexAndIndexBuffer(uint8* VertexBuffer, uint8* IndexBuffer, bool bAbsoluteIndices);
-
-	/** 
-	 * Creates rendering data from batched elements
-	 */
-	SLATECORE_API void CreateRenderBatches(FElementBatchMap& LayerToElementBatches);
-
-private:
-	/**  */
-	void Merge(FElementBatchMap& InLayerToElementBatches, uint32& VertexOffset, uint32& IndexOffset);
-
-	/** */
-	void AddRenderBatch(uint32 InLayer, const FSlateElementBatch& InElementBatch, int32 InNumVertices, int32 InNumIndices, int32 InVertexOffset, int32 InIndexOffset);
-
-	/**
-	 * Resets an array from the pool of vertex arrays
-	 * This will empty the array and give it a reasonable starting memory amount for when it is reused
-	 */
-	void ResetVertexArray(FSlateVertexArray& InOutVertexArray);
-
-	/**
-	* Resets an array from the pool of index arrays
-	* This will empty the array and give it a reasonable starting memory amount for when it is reused
-	*/
-	void ResetIndexArray(FSlateIndexArray& InOutIndexArray);
-
-private:
-
-	// The associated render data handle if these render batches are not in the default vertex/index buffer
-	TSharedPtr<FSlateRenderDataHandle, ESPMode::ThreadSafe> RenderDataHandle;
+		CachedElementLists.AddTail(NewNode);
+		NewNode->GetValue().Initialize();
 
-	// Array of vertex lists that are currently free (have no elements in them).
-	TArray<uint32> VertexArrayFreeList;
-
-	// Array of index lists that are currently free (have no elements in them).
-	TArray<uint32> IndexArrayFreeList;
-
-	// Array of vertex lists for batching vertices. We use this method for quickly resetting the arrays without deleting memory.
-	TArray<FSlateVertexArray> BatchVertexArrays;
-
-	// Array of vertex lists for batching indices. We use this method for quickly resetting the arrays without deleting memory.
-	TArray<FSlateIndexArray> BatchIndexArrays;
-
-	/** List of element batches sorted by later for use in rendering (for threaded renderers, can only be accessed from the render thread)*/
-	TArray<FSlateRenderBatch> RenderBatches;
-
-	/**  */
-	int32 NumBatchedVertices;
-
-	/**  */
-	int32 NumBatchedIndices;
-
-	/** */
-	int32 NumLayers;
-
-	/** */
-	bool bIsStencilBufferRequired;
-};
-
-class FSlateRenderer;
-
-/**
- * The draw layer represents a logical draw layer.  Since some drawn areas mights be from cached
- * buffers we can't depend on layer id to sort widget draw buffers from different frames being combined
- * together.
- */
-class FSlateDrawLayer
-{
-public:
-	FElementBatchMap& GetElementBatchMap() { return LayerToElementBatches; }
-
-	void ResetLayer()
-	{
-		DrawElements.Reset();
-	}
-
-	FORCEINLINE int32 GetElementCount() const
-	{
-		return DrawElements.Num();
+		return NewNode;
 	}
 
-	/** Apply Function to each draw element */ 
-	void ForEachElement(const TFunction<void(FSlateDrawElement&)>& InFunction)
-	{
-		for (FSlateDrawElement& Element : DrawElements)
-		{
-			InFunction(Element);
-		}
-	}
+	FSlateDrawElement& AddCachedElement(FSlateCachedElementListNode* CacheNode, const FSlateClippingManager& ParentClipManager, const SWidget* Widget);
 
-public:
-	// Element batch maps sorted by layer.
-	FElementBatchMap LayerToElementBatches;
+	void AddReferencedObjects(FReferenceCollector& Collector);
 
-	/** The elements drawn on this layer */
-	TArray<FSlateDrawElement> DrawElements;
-};
+	FSlateCachedClipState& FindOrAddCachedClipState(const FSlateClippingState* RefClipState);
+	void CleanupUnusedClipStates();
 
-/**
- * 
- */
-class FSlateDrawLayerHandle : public TSharedFromThis < FSlateDrawLayerHandle, ESPMode::ThreadSafe >
-{
 public:
-	FSlateDrawLayerHandle()
-		: BatchMap(nullptr)
-	{
-	}
+	TDoubleLinkedList<FSlateCachedElementList> CachedElementLists;
+private:
 
-	FElementBatchMap* BatchMap;
+	TArray<TUniquePtr<FSlateCachedClipState>> CachedClipStates;
 };
 
-
 /**
  * Represents a top level window and its draw elements.
  */
-class FSlateWindowElementList : public FNoncopyable
+class FSlateWindowElementList : public FGCObject, public FNoncopyable
 {
 	friend class FSlateElementBatcher;
-	friend class FSlateDrawElement;
-	friend class FSlateDataPayload;
-
 public:
-	/** 
+	/**
 	 * Construct a new list of elements with which to paint a window.
 	 *
 	 * @param InPaintWindow		The window that owns the widgets being painted.  This is almost most always the same window that is being rendered to
 	 * @param InRenderWindow	The window that we will be rendering to.
 	 */
-	SLATECORE_API explicit FSlateWindowElementList(const TSharedPtr<SWindow>& InPaintWindow = TSharedPtr<SWindow>());
+	SLATECORE_API explicit FSlateWindowElementList(const TSharedPtr<SWindow>& InPaintWindow);
 
 	SLATECORE_API ~FSlateWindowElementList();
 
@@ -1392,14 +485,14 @@ public:
 	}
 
 	/** @return Get the window that we will be painting */
-	FORCEINLINE SWindow* GetPaintWindow() const
+	SWindow* GetPaintWindow() const
 	{
 		check(IsInGameThread() || IsInSlateThread());
 		return WeakPaintWindow.IsValid() ? RawPaintWindow : nullptr;
 	}
 
 	/** @return Get the window that we will be rendering to.  Unless you are a slate renderer you probably want to use GetPaintWindow() */
-	FORCEINLINE SWindow* GetRenderWindow() const
+	SWindow* GetRenderWindow() const
 	{
 		check(IsInGameThread() || IsInSlateThread());
 		// Note: This assumes that the PaintWindow is safe to pin and is not accessed by another thread
@@ -1407,74 +500,64 @@ public:
 	}
 
 	/** @return Get the draw elements that we want to render into this window */
-	FORCEINLINE const TArray<FSlateDrawElement>& GetDrawElements() const
+	const FSlateDrawElementArray& GetUncachedDrawElements() const
 	{
-		return RootDrawLayer.DrawElements;
+		return UncachedDrawElements;
 	}
 
-	/** @return Get the draw elements that we want to render into this window */
-	FORCEINLINE TArray<FSlateDrawElement>& GetDrawElements()
-	{
-		return RootDrawLayer.DrawElements;
-	}
-
-	/** Apply Function to each draw element */ 
-	FORCEINLINE void ForEachElement(const TFunction<void(FSlateDrawElement&)>& InFunction)
-	{
-		RootDrawLayer.ForEachElement(InFunction);
-	}
-
-
-	/** @return the total number of elements that have been registered to be drawn. */
-	SLATECORE_API int32 GetElementCount() const;
-
 	/**
 	 * Add a draw element to the list
 	 *
 	 * @param InDrawElement  The draw element to add
 	 */
-	FORCEINLINE void AddItem(const FSlateDrawElement& InDrawElement)
+	UE_DEPRECATED(4.23, "AddItem is deprecated, use AddUninitialized instead")
+	void AddItem(const FSlateDrawElement& InDrawElement)
 	{
-		TArray<FSlateDrawElement>& ActiveDrawElements = DrawStack.Last()->DrawElements;
-		ActiveDrawElements.Add(InDrawElement);
+		UncachedDrawElements.Add(InDrawElement);
 	}
-	
-	void AppendItems(FSlateWindowElementList* Other);
-
 
 	/** @return Get the window size that we will be painting */
-	FORCEINLINE FVector2D GetWindowSize() const
+	FVector2D GetWindowSize() const
 	{
 		return WindowSize;
 	}
 
-	/**
-	 * Creates an uninitialized draw element
-	 */
-	FORCEINLINE FSlateDrawElement& AddUninitialized()
+	template<typename PayloadType>
+	PayloadType& CreatePayload(FSlateDrawElement& DrawElement)
 	{
-		TArray<FSlateDrawElement>& Elements = DrawStack.Last()->DrawElements;
-		const int32 InsertIdx = Elements.AddDefaulted();
-
-#if WITH_SLATE_DEBUGGING
-		FSlateDebugging::ElementAdded.Broadcast(*this, InsertIdx);
-#endif
+		// Elements are responsible for deleting their own payloads
+		PayloadType* Payload;
+		if (DrawElement.IsCached())
+		{
+			Payload = new PayloadType;
+		}
+		else
+		{
+			void* PayloadMem = MemManager.Alloc(sizeof(PayloadType), alignof(PayloadType));
+			Payload = new(PayloadMem) PayloadType();
+		}
 
-		return Elements[InsertIdx];
+		DrawElement.DataPayload = Payload;
+		return *Payload;
 	}
 
-	SLATECORE_API void MergeElementList(FSlateWindowElementList* ElementList, FVector2D AbsoluteOffset);
+	/**
+	 * Creates an uninitialized draw element
+	 */
+	SLATECORE_API FSlateDrawElement& AddUninitialized();
 
-	SLATECORE_API void MergeResources(const TArray<UObject*>& AssociatedResources);
 
 	//--------------------------------------------------------------------------
 	// CLIPPING
 	//--------------------------------------------------------------------------
 
-	SLATECORE_API void PushClip(const FSlateClippingZone& InClipZone);
+	SLATECORE_API int32 PushClip(const FSlateClippingZone& InClipZone);
 	SLATECORE_API int32 GetClippingIndex() const;
+	SLATECORE_API int32 GetClippingStackDepth() const { return ClippingManager.GetStackDepth(); }
 	SLATECORE_API TOptional<FSlateClippingState> GetClippingState() const;
 	SLATECORE_API void PopClip();
+	SLATECORE_API void PopClipToStackIndex(int32 Index);
+
 
 	FSlateClippingManager& GetClippingManager() { return ClippingManager; }
 	const FSlateClippingManager& GetClippingManager() const { return ClippingManager; }
@@ -1490,9 +573,9 @@ public:
 	struct SLATECORE_API FDeferredPaint
 	{
 	public:
-		FDeferredPaint( const TSharedRef<const SWidget>& InWidgetToPaint, const FPaintArgs& InArgs, const FGeometry InAllottedGeometry, const FWidgetStyle& InWidgetStyle, bool InParentEnabled );
+		FDeferredPaint(const TSharedRef<const SWidget>& InWidgetToPaint, const FPaintArgs& InArgs, const FGeometry InAllottedGeometry, const FWidgetStyle& InWidgetStyle, bool InParentEnabled);
 
-		int32 ExecutePaint( int32 LayerId, FSlateWindowElementList& OutDrawElements, const FSlateRect& MyCullingRect ) const;
+		int32 ExecutePaint(int32 LayerId, FSlateWindowElementList& OutDrawElements, const FSlateRect& MyCullingRect) const;
 
 		FDeferredPaint Copy(const FPaintArgs& InArgs);
 
@@ -1507,7 +590,7 @@ public:
 		const bool bParentEnabled;
 	};
 
-	SLATECORE_API void QueueDeferredPainting( const FDeferredPaint& InDeferredPaint );
+	SLATECORE_API void QueueDeferredPainting(const FDeferredPaint& InDeferredPaint);
 
 	int32 PaintDeferred(int32 LayerId, const FSlateRect& MyCullingRect);
 
@@ -1518,44 +601,27 @@ public:
 
 	TArray< TSharedPtr<FDeferredPaint> > GetDeferredPaintList() const { return DeferredPaintList; }
 
+
 	//--------------------------------------------------------------------------
-	// VOLATILE PAINTING
+	// FAST PATH
 	//--------------------------------------------------------------------------
 
-	struct FVolatilePaint
-	{
-	public:
-		SLATECORE_API FVolatilePaint(const TSharedRef<const SWidget>& InWidgetToPaint, const FPaintArgs& InArgs, const FGeometry InAllottedGeometry, const FSlateRect InMyCullingRect, const TOptional<FSlateClippingState>& ClippingState, int32 LayerId, const FWidgetStyle& InWidgetStyle, bool InParentEnabled);
-
-		int32 ExecutePaint(FSlateWindowElementList& OutDrawElements, double InCurrentTime, float InDeltaTime, const FVector2D& InDynamicOffset) const;
-
-		FORCEINLINE const SWidget* GetWidget() const { return WidgetToPaintPtr.Pin().Get(); }
-		FORCEINLINE FGeometry GetGeometry() const { return AllottedGeometry; }
-		FORCEINLINE int32 GetLayerId() const { return LayerId; }
-
-	public:
-		TSharedPtr< FSlateDrawLayerHandle, ESPMode::ThreadSafe > LayerHandle;
-	 
-	private:
-		const TWeakPtr<const SWidget> WidgetToPaintPtr;
-		const FPaintArgs Args;
-		const FGeometry AllottedGeometry;
-		const FSlateRect MyCullingRect;
-		const TOptional<FSlateClippingState> ClippingState;
-		const int32 LayerId;
-		const FWidgetStyle WidgetStyle;
-		const bool bParentEnabled;
-	};
-
-	SLATECORE_API void QueueVolatilePainting( const FVolatilePaint& InVolatilePaint );
-
-	SLATECORE_API int32 PaintVolatile(FSlateWindowElementList& OutElementList, double InCurrentTime, float InDeltaTime, const FVector2D& InDynamicOffset);
-	SLATECORE_API int32 PaintVolatileRootLayer(FSlateWindowElementList& OutElementList, double InCurrentTime, float InDeltaTime, const FVector2D& InDynamicOffset);
+	/**
+	 * Pushes the current widget that is painting onto the widget stack so we know what elements belong to each widget
+	 * This information is used for caching later.
+	 *
+	 */
+	SLATECORE_API void PushPaintingWidget(const SWidget& CurrentWidget, int32 StartingLayerId, FSlateCachedElementListNode* CurrentCacheNode);
 
-	SLATECORE_API void BeginLogicalLayer(const TSharedPtr<FSlateDrawLayerHandle, ESPMode::ThreadSafe>& LayerHandle);
-	SLATECORE_API void EndLogicalLayer();
+	/**
+	 * Pops the current painted widget off the stack
+	 * @return true if an element was added while the widget was pushed
+	 */
+	SLATECORE_API FSlateCachedElementListNode* PopPaintingWidget();
 
-	SLATECORE_API const TArray< TSharedPtr<FVolatilePaint> >& GetVolatileElements() const { return VolatilePaintList; }
+	/** Pushes cached element data onto the stack.  Any draw elements cached after will use this cached element data until popped */
+	void PushCachedElementData(FSlateCachedElementData& CachedElementData);
+	void PopCachedElementData();
 
 	//--------------------------------------------------------------------------
 	// OTHER
@@ -1564,73 +630,22 @@ public:
 	/**
 	 * Remove all the elements from this draw list.
 	 */
-	SLATECORE_API void ResetElementBuffers();
-
-	/** Allows manual hinting of reporting of references. */
-	SLATECORE_API void SetShouldReportReferencesToGC(bool bInReportReferences);
-
-	/** Will UObject references be reported to the GC? */
-	SLATECORE_API bool ShouldReportUObjectReferences() const;
-
-	/**
-	 * Allocate memory that remains valid until ResetBuffers is called.
-	 */
-	FORCEINLINE_DEBUGGABLE void* Alloc(int32 AllocSize, int32 Alignment = MIN_ALIGNMENT)
-	{
-		return MemManager.Alloc(AllocSize, Alignment);
-	}
+	SLATECORE_API void ResetElementList();
 
-	/**
-	 * Allocate memory for a type that remains valid until ResetBuffers is called.
-	 */
-	template <typename T>
-	FORCEINLINE_DEBUGGABLE void* Alloc()
-	{
-		return MemManager.Alloc(sizeof(T), alignof(T));
-	}
 
 	FSlateBatchData& GetBatchData() { return BatchData; }
 
-	FSlateDrawLayer& GetRootDrawLayer() { return RootDrawLayer; }
-
-	TMap < TSharedPtr<FSlateDrawLayerHandle, ESPMode::ThreadSafe>, TSharedPtr<FSlateDrawLayer> >& GetChildDrawLayers() { return DrawLayers; }
-
-	/**
-	 * Caches this element list on the renderer, generating all needed index and vertex buffers.
-	 */
-	SLATECORE_API TSharedRef<FSlateRenderDataHandle, ESPMode::ThreadSafe> CacheRenderData(const ILayoutCache* Cacher);
-
-	SLATECORE_API TSharedPtr<FSlateRenderDataHandle, ESPMode::ThreadSafe> GetCachedRenderDataHandle() const
-	{
-		return CachedRenderDataHandle.Pin();
-	}
-
-	void BeginUsingCachedBuffer(TSharedPtr<FSlateRenderDataHandle, ESPMode::ThreadSafe>& InCachedRenderDataHandle)
-	{
-		InCachedRenderDataHandle->BeginUsing();
-		CachedRenderHandlesInUse.Add(InCachedRenderDataHandle);
-	}
-
-	SLATECORE_API bool IsCachedRenderDataInUse() const
-	{
-		TSharedPtr<FSlateRenderDataHandle, ESPMode::ThreadSafe> SafeHandle = CachedRenderDataHandle.Pin();
-		return SafeHandle.IsValid() && SafeHandle->IsInUse();
-	}
-
-	SLATECORE_API void PreDraw_ParallelThread();
-
-	SLATECORE_API void PostDraw_ParallelThread();
-
-	SLATECORE_API void PostDraw_NonParallelRenderer();
-
 	SLATECORE_API void SetRenderTargetWindow(SWindow* InRenderTargetWindow);
 
-	SLATECORE_API void AddReferencedObjects(FReferenceCollector& Collector);
+	virtual void AddReferencedObjects(FReferenceCollector& Collector) override;
 
 private:
-	/** Resources to report to the garbage collector */
-	TArray<UObject*> ResourcesToReport;
+	FSlateDrawElement& AddCachedElement();
+
+	TArrayView<FSlateCachedElementData* const> GetCachedElementDataList() const { return MakeArrayView(CachedElementDataList.GetData(), CachedElementDataList.Num()); }
 
+	FSlateCachedElementData* GetCurrentCachedElementData() const { return CachedElementDataListStack.Num() ? CachedElementDataList[CachedElementDataListStack.Top()] : nullptr; }
+private:
 	/**
 	* Window which owns the widgets that are being painted but not necessarily rendered to
 	* Widgets are always rendered to the RenderTargetWindow
@@ -1638,6 +653,8 @@ private:
 	TWeakPtr<SWindow> WeakPaintWindow;
 	SWindow* RawPaintWindow;
 
+	FMemStackBase MemManager;
+	STAT(int32 MemManagerAllocatedMemory;)
 	/**
 	 * The window to render to.  This may be different from the paint window if we are displaying the contents of a window (or virtual window) onto another window
 	 * The primary use case of this is thread safe rendering of widgets during times when the main thread is blocked (e.g loading movies)
@@ -1648,24 +665,9 @@ private:
 	/** Batched data used for rendering */
 	FSlateBatchData BatchData;
 
-	/** The base draw layer/context. */
-	FSlateDrawLayer RootDrawLayer;
-
 	/**  */
 	FSlateClippingManager ClippingManager;
 
-	/** */
-	TMap< TSharedPtr<FSlateDrawLayerHandle, ESPMode::ThreadSafe>, TSharedPtr<FSlateDrawLayer> > DrawLayers;
-
-	/** */
-	TArray< TSharedPtr<FSlateDrawLayer> > DrawLayerPool;
-
-	/** */
-	TArray< FSlateDrawLayer* > DrawStack;
-
-	/** */
-	TArray< TSharedPtr<FSlateRenderDataHandle, ESPMode::ThreadSafe> > CachedRenderHandlesInUse;
-
 	/**
 	 * Some widgets want their logical children to appear at a different "layer" in the physical hierarchy.
 	 * We accomplish this by deferring their painting.
@@ -1675,40 +677,36 @@ private:
 	bool bNeedsDeferredResolve;
 	TArray<int32> ResolveToDeferredIndex;
 
-	/** The widgets be cached for a later paint pass when the invalidation host paints. */
-	TArray< TSharedPtr<FVolatilePaint> > VolatilePaintList;
-
-	/**
-	 * Handle to the cached render data associated with this element list.  Will only exist if 
-	 * this element list is being used for invalidation / caching.
-	 */
-	mutable TWeakPtr<FSlateRenderDataHandle, ESPMode::ThreadSafe> CachedRenderDataHandle;
-
-	// Mem stack for temp allocations
-	FMemStackBase MemManager; 
-
-	/** Store the size of the window being used to paint */
-	FVector2D WindowSize;
-
-	/** Should this element list report any references it knows about.  Defaults to true, once drawn it turns off on the rendering thread. */
-	bool bReportReferences;
-
-private:
+	// Begin Fast Path
 
-	class FWindowElementGCObject : public FGCObject
+	/** State of the current widget that is adding draw elements */
+	struct FWidgetDrawElementState
 	{
-	public:
-		FWindowElementGCObject(FSlateWindowElementList* InOwner);
+		FWidgetDrawElementState(FSlateCachedElementListNode* InCurrentCacheNode,  bool bInIsVolatile, const SWidget* InWidget)
+			: CacheNode(InCurrentCacheNode)
+			, bIsVolatile(bInIsVolatile)
+#if WITH_SLATE_DEBUGGING
+			, Widget(InWidget)
+#endif
+		{
+		}
+
+		FSlateCachedElementListNode* CacheNode;
+		bool bIsVolatile;
+#if WITH_SLATE_DEBUGGING
+		const SWidget* Widget;
+#endif
+	};
 
-		void ClearOwner();
 
-		virtual void AddReferencedObjects(FReferenceCollector& Collector) override;
-		virtual FString GetReferencerName() const override;
+	/** The uncached draw elements to be processed */
+	FSlateDrawElementArray UncachedDrawElements;
 
-	private:
-		FSlateWindowElementList* Owner;
-	};
+	TArray<FWidgetDrawElementState, TMemStackAllocator<>> WidgetDrawStack;
+	TArray<FSlateCachedElementData*, TMemStackAllocator<>> CachedElementDataList;
+	TArray<int32, TMemStackAllocator<>> CachedElementDataListStack;
+	// End Fast Path
 
-	/** This keeps drawn elements alive in case the UObject stops being referenced elsewhere that we need. */
-	TUniquePtr<FWindowElementGCObject> ResourceGCRoot;
+	/** Store the size of the window being used to paint */
+	FVector2D WindowSize;
 };
diff --git a/Engine/Source/Runtime/SlateCore/Public/Rendering/ElementBatcher.h b/Engine/Source/Runtime/SlateCore/Public/Rendering/ElementBatcher.h
--- a/Engine/Source/Runtime/SlateCore/Public/Rendering/ElementBatcher.h
+++ b/Engine/Source/Runtime/SlateCore/Public/Rendering/ElementBatcher.h
@@ -5,40 +5,255 @@
 #include "CoreMinimal.h"
 #include "Rendering/RenderingCommon.h"
 #include "Layout/Clipping.h"
+#include "Stats/Stats.h"
+#include "SlateGlobals.h"
+#include "Containers/StaticArray.h"
 
 class FSlateBatchData;
 class FSlateDrawElement;
-class FSlateDrawLayer;
-class FSlateElementBatch;
 class FSlateRenderingPolicy;
 class FSlateShaderResource;
 class FSlateWindowElementList;
 struct FShaderParams;
+struct FSlateCachedElementData;
+struct FSlateCachedElementList;
+
+DECLARE_CYCLE_STAT_EXTERN(TEXT("Add Elements Time"), STAT_SlateAddElements, STATGROUP_Slate, SLATECORE_API);
+
+DECLARE_DWORD_COUNTER_STAT_EXTERN(TEXT("Elements"), STAT_SlateElements, STATGROUP_Slate, SLATECORE_API);
+DECLARE_DWORD_COUNTER_STAT_EXTERN(TEXT("Elements (Box)"), STAT_SlateElements_Box, STATGROUP_Slate, SLATECORE_API);
+DECLARE_DWORD_COUNTER_STAT_EXTERN(TEXT("Elements (Border)"), STAT_SlateElements_Border, STATGROUP_Slate, SLATECORE_API);
+DECLARE_DWORD_COUNTER_STAT_EXTERN(TEXT("Elements (Text)"), STAT_SlateElements_Text, STATGROUP_Slate, SLATECORE_API);
+DECLARE_DWORD_COUNTER_STAT_EXTERN(TEXT("Elements (ShapedText)"), STAT_SlateElements_ShapedText, STATGROUP_Slate, SLATECORE_API);
+DECLARE_DWORD_COUNTER_STAT_EXTERN(TEXT("Elements (Line)"), STAT_SlateElements_Line, STATGROUP_Slate, SLATECORE_API);
+DECLARE_DWORD_COUNTER_STAT_EXTERN(TEXT("Elements (Other)"), STAT_SlateElements_Other, STATGROUP_Slate, SLATECORE_API);
+
+DECLARE_DWORD_COUNTER_STAT_EXTERN(TEXT("Cached Elements"), STAT_SlateCachedElements, STATGROUP_Slate, SLATECORE_API);
+
+extern int32 GSlateFeathering;
+
+/**
+* Represents an  a set of slate draw elements that are batched together
+* These later get converted into FSlateRenderBatches for final rendering
+*/
+class FSlateElementBatch
+{
+public:
+	FSlateElementBatch(const FSlateShaderResource* InShaderResource, const FShaderParams& InShaderParams, ESlateShader ShaderType, ESlateDrawPrimitive PrimitiveType, ESlateDrawEffect InDrawEffects, ESlateBatchDrawFlag InBatchFlags, const FSlateDrawElement& InDrawElement, int32 InstanceCount = 0, uint32 InstanceOffset = 0, ISlateUpdatableInstanceBuffer* InstanceData = nullptr);
+	FSlateElementBatch(TWeakPtr<ICustomSlateElement, ESPMode::ThreadSafe> InCustomDrawer, const FSlateDrawElement& InDrawElement);
+
+	void SaveClippingState(const TArray<FSlateClippingState>& PrecachedClipStates);
+
+	bool operator==(const FSlateElementBatch& Other) const
+	{
+		return BatchKey == Other.BatchKey && ShaderResource == Other.ShaderResource;
+	}
+
+	const FSlateShaderResource* GetShaderResource() const { return ShaderResource; }
+	const FShaderParams& GetShaderParams() const { return BatchKey.ShaderParams; }
+	ESlateBatchDrawFlag GetDrawFlags() const { return BatchKey.DrawFlags; }
+	ESlateDrawPrimitive GetPrimitiveType() const { return BatchKey.DrawPrimitiveType; }
+	ESlateShader GetShaderType() const { return BatchKey.ShaderType; }
+	ESlateDrawEffect GetDrawEffects() const { return BatchKey.DrawEffects; }
+	const TWeakPtr<ICustomSlateElement, ESPMode::ThreadSafe> GetCustomDrawer() const { return BatchKey.CustomDrawer; }
+	int32 GetInstanceCount() const { return BatchKey.InstanceCount; }
+	uint32 GetInstanceOffset() const { return BatchKey.InstanceOffset; }
+	const ISlateUpdatableInstanceBuffer* GetInstanceData() const { return BatchKey.InstanceData; }
+	int32 GetSceneIndex() const { return BatchKey.SceneIndex; }
+private:
+	struct FBatchKey
+	{
+		const TWeakPtr<ICustomSlateElement, ESPMode::ThreadSafe> CustomDrawer;
+		const FShaderParams ShaderParams;
+		const ESlateBatchDrawFlag DrawFlags;
+		const ESlateShader ShaderType;
+		const ESlateDrawPrimitive DrawPrimitiveType;
+		const ESlateDrawEffect DrawEffects;
+		const FClipStateHandle ClipStateHandle;
+		const int32 InstanceCount;
+		const uint32 InstanceOffset;
+		const ISlateUpdatableInstanceBuffer* InstanceData;
+		const int32 SceneIndex;
+
+		FBatchKey(const FShaderParams& InShaderParams, ESlateShader InShaderType, ESlateDrawPrimitive InDrawPrimitiveType, ESlateDrawEffect InDrawEffects, ESlateBatchDrawFlag InDrawFlags, const FClipStateHandle InClipStateHandle, int32 InInstanceCount, uint32 InInstanceOffset, ISlateUpdatableInstanceBuffer* InInstanceBuffer, int32 InSceneIndex)
+			: ShaderParams(InShaderParams)
+			, DrawFlags(InDrawFlags)
+			, ShaderType(InShaderType)
+			, DrawPrimitiveType(InDrawPrimitiveType)
+			, DrawEffects(InDrawEffects)
+			, ClipStateHandle(InClipStateHandle)
+			, InstanceCount(InInstanceCount)
+			, InstanceOffset(InInstanceOffset)
+			, InstanceData(InInstanceBuffer)
+			, SceneIndex(InSceneIndex)
+		{
+		}
+
+		FBatchKey(const FShaderParams& InShaderParams, ESlateShader InShaderType, ESlateDrawPrimitive InDrawPrimitiveType, ESlateDrawEffect InDrawEffects, ESlateBatchDrawFlag InDrawFlags, const FClipStateHandle InClipStateHandle, int32 InInstanceCount, uint32 InInstanceOffset, ISlateUpdatableInstanceBuffer* InInstanceBuffer)
+			: ShaderParams(InShaderParams)
+			, DrawFlags(InDrawFlags)
+			, ShaderType(InShaderType)
+			, DrawPrimitiveType(InDrawPrimitiveType)
+			, DrawEffects(InDrawEffects)
+			, ClipStateHandle(InClipStateHandle)
+			, InstanceCount(InInstanceCount)
+			, InstanceOffset(InInstanceOffset)
+			, InstanceData(InInstanceBuffer)
+			, SceneIndex(-1)
+		{
+		}
+
+		FBatchKey(TWeakPtr<ICustomSlateElement, ESPMode::ThreadSafe> InCustomDrawer, const FClipStateHandle InClipStateHandle)
+			: CustomDrawer(InCustomDrawer)
+			, ShaderParams()
+			, DrawFlags(ESlateBatchDrawFlag::None)
+			, ShaderType(ESlateShader::Default)
+			, DrawPrimitiveType(ESlateDrawPrimitive::TriangleList)
+			, DrawEffects(ESlateDrawEffect::None)
+			, ClipStateHandle(InClipStateHandle)
+			, InstanceCount(0)
+			, InstanceOffset(0)
+			, SceneIndex(-1)
+		{}
+
+		bool operator==(const FBatchKey& Other) const
+		{
+			return DrawFlags == Other.DrawFlags
+				&& ShaderType == Other.ShaderType
+				&& DrawPrimitiveType == Other.DrawPrimitiveType
+				&& DrawEffects == Other.DrawEffects
+				&& ShaderParams == Other.ShaderParams
+				&& ClipStateHandle == Other.ClipStateHandle
+				&& CustomDrawer == Other.CustomDrawer
+				&& InstanceCount == Other.InstanceCount
+				&& InstanceOffset == Other.InstanceOffset
+				&& InstanceData == Other.InstanceData
+				&& SceneIndex == Other.SceneIndex;
+		}
+	};
+
+	/** A secondary key which represents elements needed to make a batch unique */
+	FBatchKey BatchKey;
+
+	/** Shader resource to use with this batch.  Used as a primary key.  No batch can have multiple textures */
+	const FSlateShaderResource* ShaderResource;
+
+public:
+	/** Number of elements in the batch */
+	uint32 NumElementsInBatch;
+	/** Index into an array of vertex arrays where this batches vertices are found (before submitting to the vertex buffer)*/
+	int32 VertexArrayIndex;
+	/** Index into an array of index arrays where this batches indices are found (before submitting to the index buffer) */
+	int32 IndexArrayIndex;
+};
+
+
+class FSlateBatchData
+{
+public:
+	FSlateBatchData()
+		: FirstRenderBatchIndex(0)
+		, NumLayers(0)
+		, NumBatches(0)
+		, bIsStencilBufferRequired(false)
+	{}
+
+	SLATECORE_API void ResetData();
+
+	/**
+	* Returns a list of element batches for this window
+	*/
+	const TArray<FSlateRenderBatch>& GetRenderBatches() const { return RenderBatches; }
+
+	/**
+	*
+	*/
+	SLATECORE_API bool IsStencilClippingRequired() const;
+
+	int32 GetFirstRenderBatchIndex() const { return FirstRenderBatchIndex; }
+	/** @return Total number of batched layers */
+	int32 GetNumLayers() const { return NumLayers; }
+
+	int32 GetNumFinalBatches() const { return NumBatches; }
+
+	const FSlateVertexArray& GetFinalVertexData() const { return FinalVertexData; }
+	const FSlateIndexArray& GetFinalIndexData() const { return FinalIndexData; }
+
+	/**
+	* Fills batch data into the actual vertex and index buffer
+	*
+	* @param VertexBuffer	Pointer to the actual memory for the vertex buffer
+	* @param IndexBuffer	Pointer to the actual memory for an index buffer
+	* @param bAbsoluteIndices	Whether to write absolute indices (simplifies draw call setup on RHIs that do not support BaseVertex)
+	*/
+	SLATECORE_API void FillVertexAndIndexBuffer(uint8* VertexBuffer, uint8* IndexBuffer, bool bAbsoluteIndices);
+
+	/**
+	* Creates rendering data from batched elements
+	*/
+	//SLATECORE_API void CreateRenderBatches(class FElementBatchMap& LayerToElementBatches);
+
+	SLATECORE_API void MergeRenderBatches();
+
+	/** */
+	FSlateRenderBatch& AddRenderBatch(
+		int32 InLayer,
+		const FShaderParams& InShaderParams,
+		const FSlateShaderResource* InResource,
+		ESlateDrawPrimitive InPrimitiveType,
+		ESlateShader InShaderType,
+		ESlateDrawEffect InDrawEffects,
+		ESlateBatchDrawFlag InDrawFlags,
+		int8 SceneIndex);
+
+	void AddCachedBatches(const TArray<FSlateRenderBatch>& InCachedBatches);
+private:
+	void FillBuffersFromNewBatch(FSlateRenderBatch& Batch, FSlateVertexArray& FinalVertices, FSlateIndexArray& FinalIndices);
+	void CombineBatches(FSlateRenderBatch& FirstBatch, FSlateRenderBatch& SecondBatch, FSlateVertexArray& FinalVertices, FSlateIndexArray& FinalIndices);
+private:
+
+	/** List of element batches sorted by later for use in rendering*/
+	// todo allocator?
+	TArray<FSlateRenderBatch> RenderBatches;
+
+	/** Uncached source vertices and indices from volatile elements */
+	FSlateVertexArray UncachedSourceBatchVertices;
+	FSlateIndexArray UncachedSourceBatchIndices;
+
+	FSlateVertexArray FinalVertexData;
+	FSlateIndexArray  FinalIndexData;
+
+	int32 FirstRenderBatchIndex;
+
+	/** */
+	int32 NumLayers;
+
+	/** Number of final render batches.  it is not the same as RenderBatches.Num(); */
+	int32 NumBatches;
+
+	/** */
+	bool bIsStencilBufferRequired;
+};
 
-class FSlateDrawBox;
-class FSlateDrawText;
-class FSlateDrawShapedText;
-class FSlateDrawLines;
-class FSlateDrawCachedBuffer;
 
 /**
  * A class which batches Slate elements for rendering
  */
-class SLATECORE_API FSlateElementBatcher
+class FSlateElementBatcher
 {
 
 	friend struct FLineBuilder;
 public:
 
-	FSlateElementBatcher( TSharedRef<FSlateRenderingPolicy> InRenderingPolicy );
-	~FSlateElementBatcher();
+	SLATECORE_API FSlateElementBatcher( TSharedRef<FSlateRenderingPolicy> InRenderingPolicy );
+	SLATECORE_API ~FSlateElementBatcher();
 
 	/** 
 	 * Batches elements to be rendered 
 	 * 
 	 * @param DrawElements	The elements to batch
 	 */
-	void AddElements( FSlateWindowElementList& ElementList );
+	SLATECORE_API void AddElements( FSlateWindowElementList& ElementList );
 
 	/**
 	 * Returns true if the elements in this batcher require v-sync.
@@ -51,10 +266,7 @@ public:
 	/** 
 	 * Resets all stored data accumulated during the batching process
 	 */
-	void ResetBatches();
-
-private:
-	void AddElementsInternal(const TArray<FSlateDrawElement>& DrawElements, const FVector2D& ViewportSize);
+	SLATECORE_API void ResetBatches();
 
 	FORCEINLINE FColor PackVertexColor(const FLinearColor& InLinearColor) const
 	{
@@ -63,6 +275,10 @@ private:
 		return InLinearColor.ToFColor(bSRGBVertexColor);
 	}
 
+private:
+	void AddElementsInternal(const FSlateDrawElementArray& DrawElements, const FVector2D& ViewportSize);
+	void AddCachedElements(FSlateCachedElementData& CachedElementData, const FVector2D& ViewportSize);
+
 	/** 
 	 * Creates vertices necessary to draw a Quad element 
 	 */
@@ -116,81 +332,58 @@ private:
 	template<ESlateVertexRounding Rounding>
 	void AddBorderElement( const FSlateDrawElement& DrawElement );
 
-	/**
-	 * Batches a custom slate drawing element
-	 *
-	 * @param Position		The top left screen space position of the element
-	 * @param Size			The size of the element
-	 * @param Scale			The amount to scale the element by
-	 * @param InPayload		The data payload for this element
-	 * @param DrawEffects	DrawEffects to apply
-	 * @param Layer			The layer to draw this element in
-	 */
 	void AddCustomElement( const FSlateDrawElement& DrawElement );
 
 	void AddCustomVerts( const FSlateDrawElement& DrawElement );
 
-	void AddCachedBuffer( const FSlateDrawElement& DrawElement );
+	void AddPostProcessPass(const FSlateDrawElement& DrawElement, const FVector2D& WindowSize);
 
-	void AddLayer(const FSlateDrawElement& DrawElement);
+	FSlateRenderBatch& CreateRenderBatch(
+		int32 Layer,
+		const FShaderParams& ShaderParams,
+		const FSlateShaderResource* InResource,
+		ESlateDrawPrimitive PrimitiveType,
+		ESlateShader ShaderType,
+		ESlateDrawEffect DrawEffects,
+		ESlateBatchDrawFlag DrawFlags,
+		const FSlateDrawElement& DrawElement);
 
-	void AddPostProcessPass(const FSlateDrawElement& DrawElement, const FVector2D& WindowSize);
+	const FSlateClippingState* ResolveClippingState(const FSlateDrawElement& DrawElement) const;
 
-	/** 
-	 * Finds an batch for an element based on the passed in parameters
-	 * Elements with common parameters and layers will batched together.
-	 *
-	 * @param Layer			The layer where this element should be drawn (signifies draw order)
-	 * @param ShaderParams	The shader params for this element
-	 * @param InTexture		The texture to use in the batch
-	 * @param PrimitiveType	The primitive type( triangles, lines ) to use when drawing the batch
-	 * @param ShaderType	The shader to use when rendering this batch
-	 * @param DrawFlags		Any optional draw flags for this batch
-	 * @param ScissorRect   Optional scissor rectangle for this batch
-	 * @param SceneIndex    Index in the slate renderer's scenes array associated with this element.
-	 */
-	FSlateElementBatch& FindBatchForElement( uint32 Layer, 
-											 const FShaderParams& ShaderParams, 
-											 const FSlateShaderResource* InTexture, 
-											 ESlateDrawPrimitive::Type PrimitiveType, 
-											 ESlateShader::Type ShaderType, 
-											 ESlateDrawEffect DrawEffects, 
-											 ESlateBatchDrawFlag DrawFlags,
-											 int32 ClippingIndex,
-											 int32 SceneIndex = -1);
 private:
-	/** Batch data currently being filled in */
+	/** Uncached Batch data currently being filled in */
 	FSlateBatchData* BatchData;
 
-	/** The draw layer currently being accumulated */
-	FSlateDrawLayer* DrawLayer;
+	/** Cached batches currently being filled in */
+	FSlateCachedElementList* CachedElementList;
 
-	/** The draw layer currently being accumulated */
-	const TArray<FSlateClippingState>* ClippingStates;
+	/** Clipping states currently being used */
+	const TArray<FSlateClippingState>* PrecachedClippingStates;
 
 	/** Rendering policy we were created from */
 	FSlateRenderingPolicy* RenderingPolicy;
 
+#if STATS
 	/** Track the number of drawn boxes from the previous frame to report to stats. */
-	int32 ElmementStat_Boxes;
+	int32 ElementStat_Boxes;
 
 	/** Track the number of drawn borders from the previous frame to report to stats. */
-	int32 ElmementStat_Borders;
+	int32 ElementStat_Borders;
 
 	/** Track the number of drawn text from the previous frame to report to stats. */
-	int32 ElmementStat_Text;
+	int32 ElementStat_Text;
 
 	/** Track the number of drawn shaped text from the previous frame to report to stats. */
-	int32 ElmementStat_ShapedText;
+	int32 ElementStat_ShapedText;
 
 	/** Track the number of drawn lines from the previous frame to report to stats. */
-	int32 ElmementStat_Line;
-
-	/** Track the number of drawn cached buffers from the previous frame to report to stats. */
-	int32 ElmementStat_CachedBuffer;
+	int32 ElementStat_Line;
 
 	/** Track the number of drawn batches from the previous frame to report to stats. */
-	int32 ElmementStat_Other;
+	int32 ElementStat_Other;
+
+	int32 ElementStat_CachedElements;
+#endif
 
 	/** How many post process passes are needed */
 	int32 NumPostProcessPasses;
diff --git a/Engine/Source/Runtime/SlateCore/Public/Rendering/RenderingCommon.h b/Engine/Source/Runtime/SlateCore/Public/Rendering/RenderingCommon.h
--- a/Engine/Source/Runtime/SlateCore/Public/Rendering/RenderingCommon.h
+++ b/Engine/Source/Runtime/SlateCore/Public/Rendering/RenderingCommon.h
@@ -3,18 +3,28 @@
 #pragma once
 
 #include "CoreMinimal.h"
+#include "Stats/Stats.h"
 #include "Layout/SlateRect.h"
 #include "Layout/SlateRotatedRect.h"
 #include "Input/CursorReply.h"
 #include "Input/Reply.h"
 #include "Input/NavigationReply.h"
 #include "Input/PopupMethodReply.h"
+#include "SlateGlobals.h"
 #include "RenderingCommon.generated.h"
 
 class FSlateInstanceBufferUpdate;
 class FWidgetStyle;
 class SWidget;
-struct Rect;
+
+
+DECLARE_MEMORY_STAT_EXTERN(TEXT("Vertex/Index Buffer Pool Memory (CPU)"), STAT_SlateBufferPoolMemory, STATGROUP_SlateMemory, SLATECORE_API);
+DECLARE_MEMORY_STAT_EXTERN(TEXT("Cached Draw Element Memory (CPU)"), STAT_SlateCachedDrawElementMemory, STATGROUP_SlateMemory, SLATECORE_API);
+
+DECLARE_DWORD_COUNTER_STAT_EXTERN(TEXT("Num Cached Element Lists"), STAT_SlateNumCachedElementLists, STATGROUP_Slate, SLATECORE_API);
+DECLARE_DWORD_COUNTER_STAT_EXTERN(TEXT("Num Cached Elements"), STAT_SlateNumCachedElements, STATGROUP_Slate, SLATECORE_API);
+
+DECLARE_CYCLE_STAT_EXTERN(TEXT("PreFill Buffers RT"), STAT_SlatePreFullBufferRTTime, STATGROUP_Slate, SLATECORE_API);
 
 #define SLATE_USE_32BIT_INDICES !PLATFORM_USES_ES2
 
@@ -27,33 +37,31 @@ typedef uint16 SlateIndex;
 /**
  * Draw primitive types                   
  */
-namespace ESlateDrawPrimitive
+enum class ESlateDrawPrimitive : uint8
 {
-	typedef uint8 Type;
-
-	const Type LineList = 0;
-	const Type TriangleList = 1;
+	None,
+	LineList,
+	TriangleList,
 };
-
+ 
 /**
  * Shader types. NOTE: mirrored in the shader file   
  * If you add a type here you must also implement the proper shader type (TSlateElementPS).  See SlateShaders.h
  */
-namespace ESlateShader
+enum class ESlateShader : uint8
 {
-	typedef uint8 Type;
 	/** The default shader type.  Simple texture lookup */
-	const Type Default = 0;
+	Default = 0,
 	/** Border shader */
-	const Type Border = 1;
+	Border = 1,
 	/** Font shader, same as default except uses an alpha only texture */
-	const Type Font = 2;
+	Font = 2,
 	/** Line segment shader. For drawing anti-aliased lines */
-	const Type LineSegment = 3;
+	LineSegment = 3,
 	/** For completely customized materials.  Makes no assumptions on use*/
-	const Type Custom = 4;
+	Custom = 4,
 	/** For post processing passes */
-	const Type PostProcess = 5;
+	PostProcess = 5,
 };
 
 /**
@@ -145,6 +153,38 @@ enum class ESlateVertexRounding : uint8
 	Enabled
 };
 
+class FSlateRenderBatch;
+
+
+/**
+* Shader parameters for slate
+*/
+struct FShaderParams
+{
+	/** Pixel shader parameters */
+	FVector4 PixelParams;
+	FVector4 PixelParams2;
+
+	FShaderParams()
+		: PixelParams(0, 0, 0, 0)
+		, PixelParams2(0, 0, 0, 0)
+	{}
+
+	FShaderParams(const FVector4& InPixelParams, const FVector4& InPixelParams2 = FVector4(0))
+		: PixelParams(InPixelParams)
+		, PixelParams2(InPixelParams2)
+	{}
+
+	bool operator==(const FShaderParams& Other) const
+	{
+		return PixelParams == Other.PixelParams && PixelParams2 == Other.PixelParams2;
+	}
+
+	static FShaderParams MakePixelShaderParams(const FVector4& PixelShaderParams, const FVector4& InPixelShaderParams2 = FVector4(0))
+	{
+		return FShaderParams(PixelShaderParams, InPixelShaderParams2);
+	}
+};
 
 
 /** 
@@ -254,6 +294,7 @@ private:
 };
 
 template<> struct TIsPODType<FSlateVertex> { enum { Value = true }; };
+static_assert(TIsTriviallyDestructible<FSlateVertex>::Value == true, "FSlateVertex should be trivially destructible");
 
 /** Stores an aligned rect as shorts. */
 struct FShortRect
@@ -311,6 +352,110 @@ struct FShortRect
 	uint16 Bottom;
 };
 
+#if STATS
+
+struct FRenderingBufferStatTracker
+{
+	static void MemoryAllocated(int32 SizeBytes)
+	{
+		INC_DWORD_STAT_BY(STAT_SlateBufferPoolMemory, SizeBytes);
+	}
+
+	static void MemoryFreed(int32 SizeBytes)
+	{
+		DEC_DWORD_STAT_BY(STAT_SlateBufferPoolMemory, SizeBytes);
+	}
+};
+
+struct FDrawElementStatTracker
+{
+	static void MemoryAllocated(int32 SizeBytes)
+	{
+		INC_DWORD_STAT_BY(STAT_SlateCachedDrawElementMemory, SizeBytes);
+	}
+
+	static void MemoryFreed(int32 SizeBytes)
+	{
+		DEC_DWORD_STAT_BY(STAT_SlateCachedDrawElementMemory, SizeBytes);
+	}
+};
+template<typename StatTracker>
+class FSlateStatTrackingMemoryAllocator : public FDefaultAllocator
+{
+public:
+	typedef FDefaultAllocator Super;
+
+	class ForAnyElementType : public FDefaultAllocator::ForAnyElementType
+	{
+	public:
+		typedef FDefaultAllocator::ForAnyElementType Super;
+
+		ForAnyElementType()
+			: AllocatedSize(0)
+		{
+
+		}
+
+		/**
+		* Moves the state of another allocator into this one.
+		* Assumes that the allocator is currently empty, i.e. memory may be allocated but any existing elements have already been destructed (if necessary).
+		* @param Other - The allocator to move the state from.  This allocator should be left in a valid empty state.
+		*/
+		FORCEINLINE void MoveToEmpty(ForAnyElementType& Other)
+		{
+			Super::MoveToEmpty(Other);
+
+			AllocatedSize = Other.AllocatedSize;
+			Other.AllocatedSize = 0;
+		}
+
+		/** Destructor. */
+		~ForAnyElementType()
+		{
+			if (AllocatedSize)
+			{
+				StatTracker::MemoryFreed(AllocatedSize);
+
+			}
+		}
+
+		void ResizeAllocation(int32 PreviousNumElements, int32 NumElements, int32 NumBytesPerElement)
+		{
+			const int32 NewSize = NumElements * NumBytesPerElement;
+			StatTracker::MemoryAllocated(NewSize - AllocatedSize);
+
+			AllocatedSize = NewSize;
+
+			Super::ResizeAllocation(PreviousNumElements, NumElements, NumBytesPerElement);
+		}
+
+	private:
+		ForAnyElementType(const ForAnyElementType&);
+		ForAnyElementType& operator=(const ForAnyElementType&);
+	private:
+		int32 AllocatedSize;
+	};
+};
+
+template <typename T>
+struct TAllocatorTraits<FSlateStatTrackingMemoryAllocator<T>> : TAllocatorTraitsBase<FSlateStatTrackingMemoryAllocator<T>>
+{
+	enum { SupportsMove = TAllocatorTraits<FDefaultAllocator>::SupportsMove };
+	enum { IsZeroConstruct = TAllocatorTraits<FDefaultAllocator>::IsZeroConstruct };
+};
+
+typedef TArray<FSlateVertex, FSlateStatTrackingMemoryAllocator<FRenderingBufferStatTracker>> FSlateVertexArray;
+typedef TArray<SlateIndex, FSlateStatTrackingMemoryAllocator<FRenderingBufferStatTracker>> FSlateIndexArray;
+typedef TArray<FSlateDrawElement, FSlateStatTrackingMemoryAllocator<FDrawElementStatTracker>> FSlateDrawElementArray;
+
+#else
+
+typedef TArray<FSlateVertex> FSlateVertexArray;
+typedef TArray<SlateIndex> FSlateIndexArray;
+typedef TArray<FSlateDrawElement> FSlateDrawElementArray;
+
+#endif
+
 static FVector2D RoundToInt(const FVector2D& Vec)
 {
 	return FVector2D(FMath::RoundToInt(Vec.X), FMath::RoundToInt(Vec.Y));
@@ -813,3 +958,4 @@ private:
 	uint32 InstanceCount;
 	bool bWasCommitted;
 };
+
diff --git a/Engine/Source/Runtime/SlateCore/Public/Rendering/SlateRenderer.h b/Engine/Source/Runtime/SlateCore/Public/Rendering/SlateRenderer.h
--- a/Engine/Source/Runtime/SlateCore/Public/Rendering/SlateRenderer.h
+++ b/Engine/Source/Runtime/SlateCore/Public/Rendering/SlateRenderer.h
@@ -11,7 +11,6 @@
 class FRHITexture2D;
 class FSlateDrawBuffer;
 class FSlateUpdatableTexture;
-class ILayoutCache;
 class ISlateAtlasProvider;
 class ISlateStyle;
 class SWindow;
@@ -442,19 +441,6 @@ public:
 	 */
 	virtual ISlateAtlasProvider* GetFontAtlasProvider();
 
-	/**
-	 * Converts and caches the elements list data as the final rendered vertex and index buffer data,
-	 * returning a handle to it to allow issuing future draw commands using it.
-	 */
-	virtual TSharedRef<FSlateRenderDataHandle, ESPMode::ThreadSafe> CacheElementRenderData(const ILayoutCache* Cacher, FSlateWindowElementList& ElementList);
-
-	/**
-	 * Releases the caching resources used on the render thread for the provided ILayoutCache.  This is
-	 * should be the kind of operation you perform when the ILayoutCache is being destroyed and will no longer
-	 * be needing to draw anything again.
-	 */
-	virtual void ReleaseCachingResourcesFor(const ILayoutCache* Cacher);
-
 	/**
 	 * Copies all slate windows out to a buffer at half resolution with debug information
 	 * like the mouse cursor and any keypresses.
@@ -480,6 +466,8 @@ public:
 	/** Reset the internal Scene tracking.*/
 	virtual void ClearScenes() = 0;
 
+	virtual void DestroyCachedFastPathRenderingData(struct FSlateCachedFastPathRenderingData* VertexData);
+
 	virtual bool HasLostDevice() const { return false; }
 
 	/**
diff --git a/Engine/Source/Runtime/SlateCore/Public/SlateGlobals.h b/Engine/Source/Runtime/SlateCore/Public/SlateGlobals.h
--- a/Engine/Source/Runtime/SlateCore/Public/SlateGlobals.h
+++ b/Engine/Source/Runtime/SlateCore/Public/SlateGlobals.h
@@ -4,6 +4,7 @@
 
 #include "CoreMinimal.h"
 #include "Stats/Stats.h"
+#include "Debugging/SlateDebugging.h"
 
 #define SLATE_CHECK_UOBJECT_RENDER_RESOURCES !UE_BUILD_SHIPPING
 
@@ -48,12 +49,19 @@ DECLARE_STATS_GROUP(TEXT("Slate"), STATGROUP_Slate, STATCAT_Advanced);
 DECLARE_STATS_GROUP_VERBOSE(TEXT("SlateVerbose"), STATGROUP_SlateVerbose, STATCAT_Advanced);
 DECLARE_STATS_GROUP_MAYBE_COMPILED_OUT(TEXT("SlateVeryVerbose"), STATGROUP_SlateVeryVerbose, STATCAT_Advanced, WITH_VERY_VERBOSE_SLATE_STATS);
 
-/** Whether or not dynamic prepass and layout caching is enabled */
-extern SLATECORE_API int32 GSlateLayoutCaching;
-
 /** Whether or not we've enabled fast widget pathing which validates paths to widgets without arranging children. */
 extern SLATECORE_API int32 GSlateFastWidgetPath;
 
+extern SLATECORE_API int32 GSlateEnableGlobalInvalidation;
+
+extern SLATECORE_API bool GSlateIsOnFastUpdatePath;
+
+extern SLATECORE_API bool GSlateIsInInvalidationSlowPath;
+
+#if WITH_SLATE_DEBUGGING
+extern SLATECORE_API int32 GSlateInvalidationDebugging;
+extern SLATECORE_API int32 GSlateHitTestGridDebugging;
+#endif
 /* Forward declarations
 *****************************************************************************/
 class FActiveTimerHandle;
