Slate Global Invalidation System

diff --git a/Engine/Source/Runtime/SlateCore/Public/Types/PaintArgs.h b/Engine/Source/Runtime/SlateCore/Public/Types/PaintArgs.h
--- a/Engine/Source/Runtime/SlateCore/Public/Types/PaintArgs.h
+++ b/Engine/Source/Runtime/SlateCore/Public/Types/PaintArgs.h
@@ -4,14 +4,15 @@
 
 #include "CoreMinimal.h"
 #include "Layout/Visibility.h"
+#include "Layout/Geometry.h"
+#include "SlateGlobals.h"
 
-class FCachedWidgetNode;
 class FHittestGrid;
 class FSlateRect;
 class ICustomHitTestPath;
-class ILayoutCache;
 class SWidget;
-struct FGeometry;
+class SWindow;
+
 
 /**
  * SWidget::OnPaint and SWidget::Paint use FPaintArgs as their
@@ -20,54 +21,54 @@ struct FGeometry;
  */
 class SLATECORE_API FPaintArgs
 {
+	friend class SInvalidationPanel;
+	friend class SRetainerWidget;
 public:
-	FPaintArgs( const SWidget& Parent, FHittestGrid& InHittestGrid, FVector2D InWindowOffset, double InCurrentTime, float InDeltaTime );
-	
-	FORCEINLINE FPaintArgs WithNewParent(const SWidget* Parent) const
-	{
-		checkSlow(Parent);
-		return WithNewParent(*Parent);
-	}
+	FPaintArgs(const SWidget* PaintParent, FHittestGrid& InRootHittestGrid, FHittestGrid& InCurrentHitTestGrid, FVector2D InWindowOffset, double InCurrentTime, float InDeltaTime);
+	FPaintArgs(const SWidget* PaintParent, FHittestGrid& InRootHittestGrid, FVector2D InWindowOffset, double InCurrentTime, float InDeltaTime);
 
-	FORCEINLINE_DEBUGGABLE FPaintArgs WithNewParent(const SWidget& Parent) const
+	FORCEINLINE_DEBUGGABLE FPaintArgs WithNewParent(const SWidget* PaintParent) const
 	{
-		FPaintArgs Args = FPaintArgs(Parent, this->Grid, this->WindowOffset, this->CurrentTime, this->DeltaTime);
-		Args.LastHittestIndex = this->LastHittestIndex;
-		Args.LastRecordedVisibility = this->LastRecordedVisibility;
-		Args.LayoutCache = this->LayoutCache;
-		Args.ParentCacheNode = this->ParentCacheNode;
-		Args.bIsCaching = this->bIsCaching;
-		Args.bIsVolatilityPass = this->bIsVolatilityPass;
+		FPaintArgs Args(*this);
+		Args.PaintParentPtr = PaintParent;
 
 		return Args;
 	}
 
-	FPaintArgs EnableCaching(const TWeakPtr<ILayoutCache>& InLayoutCache, FCachedWidgetNode* InParentCacheNode, bool bEnableCaching, bool bEnableVolatility) const;
-	FPaintArgs WithNewTime(double InCurrentTime, float InDeltaTime) const;
-	FPaintArgs RecordHittestGeometry(const SWidget* Widget, const FGeometry& WidgetGeometry, int32 LayerId) const;
-	FPaintArgs InsertCustomHitTestPath( TSharedRef<ICustomHitTestPath> CustomHitTestPath, int32 HitTestIndex ) const;
+	FORCEINLINE_DEBUGGABLE FPaintArgs WithNewHitTestGrid(FHittestGrid& NewHitTestGrid) const
+	{
+		FPaintArgs NewArgs(PaintParentPtr, RootGrid, NewHitTestGrid, WindowOffset, CurrentTime, DeltaTime);
+		return NewArgs;
+	}
+
+	FPaintArgs InsertCustomHitTestPath(const SWidget* Widget, TSharedRef<ICustomHitTestPath> CustomHitTestPath) const;
+	void SetInheritedHittestability(bool InInheritedHittestability) { bInheritedHittestability = InInheritedHittestability; }
+
+	UE_DEPRECATED(4.23, "FHittestGrid::GetGrid is deprecated.  Use GetHittestGrid instead")
+	FHittestGrid& GetGrid() const { return GetHittestGrid(); }
 
-	FHittestGrid& GetGrid() const { return Grid; }
-	int32 GetLastHitTestIndex() const { return LastHittestIndex; }
-	EVisibility GetLastRecordedVisibility() const { return LastRecordedVisibility; }
+	FHittestGrid& GetHittestGrid() const { return CurrentGrid; }
+
+	const SWidget* GetPaintParent() const { return PaintParentPtr;  }
 	FVector2D GetWindowToDesktopTransform() const { return WindowOffset; }
 	double GetCurrentTime() const { return CurrentTime; }
 	float GetDeltaTime() const { return DeltaTime; }
-	bool IsCaching() const { return bIsCaching; }
-	bool IsVolatilityPass() const { return bIsVolatilityPass; }
-	const TWeakPtr<ILayoutCache>& GetLayoutCache() const { return LayoutCache; }
-	FCachedWidgetNode* GetParentCacheNode() const { return ParentCacheNode; }
 
+	bool GetInheritedHittestability() const { return bInheritedHittestability; }
+	
 private:
-	const SWidget& ParentPtr;
-	FHittestGrid& Grid;
-	int32 LastHittestIndex;
-	EVisibility LastRecordedVisibility;
+
+	/** The root most grid.  Only the window should set this and only invalidation panels should modify it */
+	FHittestGrid& RootGrid;
+
+	/** The current hit test grid.  Its possible that there is more than one grid when there is nested invalidation panels.  This is what widgets should add to always */
+	FHittestGrid& CurrentGrid;
+
 	FVector2D WindowOffset;
+	const SWidget* PaintParentPtr;
+
 	double CurrentTime;
 	float DeltaTime;
-	bool bIsCaching;
-	bool bIsVolatilityPass;
-	TWeakPtr<ILayoutCache> LayoutCache;
-	FCachedWidgetNode* ParentCacheNode;
+	uint8 bInheritedHittestability : 1;
+
 };
diff --git a/Engine/Source/Runtime/SlateCore/Public/Widgets/SWidget.h b/Engine/Source/Runtime/SlateCore/Public/Widgets/SWidget.h
--- a/Engine/Source/Runtime/SlateCore/Public/Widgets/SWidget.h
+++ b/Engine/Source/Runtime/SlateCore/Public/Widgets/SWidget.h
@@ -24,11 +24,12 @@
 #include "Types/WidgetActiveTimerDelegate.h"
 #include "Textures/SlateShaderResource.h"
 #include "SlateGlobals.h"
+#include "Types/PaintArgs.h"
+#include "FastUpdate/WidgetProxy.h"
 #include "Widgets/Accessibility/SlateWidgetAccessibleTypes.h"
 
 class FActiveTimerHandle;
 class FArrangedChildren;
-class FCachedWidgetNode;
 class FChildren;
 class FPaintArgs;
 class FSlateWindowElementList;
@@ -38,29 +39,21 @@ class FWidgetPath;
 class IToolTip;
 class SWidget;
 struct FSlateBrush;
+struct FSlatePaintElementLists;
 
-enum class EAccessibleType : uint8
-{
-	Main,
-	Summary
-};
 
-/** Widget update flags for fast path.  Work in progress, do not modify */
-enum class EWidgetUpdateFlags : uint8
-{
-	None = 0,
 
-	NeedsTick = 1 << 2,
+DECLARE_DWORD_COUNTER_STAT_EXTERN(TEXT("STAT_SlateVeryVerboseStatGroupTester"), STAT_SlateVeryVerboseStatGroupTester, STATGROUP_SlateVeryVerbose, SLATECORE_API);
+DECLARE_DWORD_COUNTER_STAT_EXTERN(TEXT("Widgets Created (Per Frame)"), STAT_SlateTotalWidgetsPerFrame, STATGROUP_Slate, SLATECORE_API);
+DECLARE_DWORD_COUNTER_STAT_EXTERN(TEXT("SWidget::Paint (Count)"), STAT_SlateNumPaintedWidgets, STATGROUP_Slate, SLATECORE_API);
+DECLARE_DWORD_COUNTER_STAT_EXTERN(TEXT("SWidget::Tick (Count)"), STAT_SlateNumTickedWidgets, STATGROUP_Slate, SLATECORE_API);
+DECLARE_CYCLE_STAT_EXTERN(TEXT("Execute Active Timers"), STAT_SlateExecuteActiveTimers, STATGROUP_Slate, SLATECORE_API);
+DECLARE_CYCLE_STAT_EXTERN(TEXT("Tick Widgets"), STAT_SlateTickWidgets, STATGROUP_Slate, SLATECORE_API);
+DECLARE_CYCLE_STAT_EXTERN(TEXT("SlatePrepass"), STAT_SlatePrepass, STATGROUP_Slate, SLATECORE_API);
 
-	NeedsActiveTimerUpdate = 1 << 3,
+DECLARE_DWORD_ACCUMULATOR_STAT_EXTERN(TEXT("Total Widgets"), STAT_SlateTotalWidgets, STATGROUP_SlateMemory, SLATECORE_API);
+DECLARE_MEMORY_STAT_EXTERN(TEXT("SWidget Total Allocated Size"), STAT_SlateSWidgetAllocSize, STATGROUP_SlateMemory, SLATECORE_API);
 
-	NeedsRepaint = 1 << 4,
-
-	All = 0xff
-};
-ENUM_CLASS_FLAGS(EWidgetUpdateFlags)
-
-DECLARE_DWORD_COUNTER_STAT_EXTERN(TEXT("STAT_SlateVeryVerboseStatGroupTester"), STAT_SlateVeryVerboseStatGroupTester, STATGROUP_SlateVeryVerbose, SLATECORE_API);
 
 /** Delegate type for handling mouse events */
 DECLARE_DELEGATE_RetVal_TwoParams(
@@ -128,13 +121,17 @@ public:
 	}
 };
 
+enum class EAccessibleType : uint8
+{
+	Main,
+	Summary
+};
 
 /**
  * The different types of invalidation that are possible for a widget.
  */
 enum class EInvalidateWidgetReason : uint8
 {
-	/** Not actually used, but defining it. */
 	None = 0,
 
 	/**
@@ -152,6 +149,19 @@ enum class EInvalidateWidgetReason : uint8
 	 */
 	Volatility = 1 << 2,
 
+	/**
+	 * A child was added or removed.   (this implies layout)
+	 */
+	ChildOrder = 1 << 3,
+
+	/** A Widgets render transform changed */
+	RenderTransform = 1 << 4,
+
+	/**
+	 * Changing visibility (this implies layout)
+	 */
+	Visibility = 1 << 5,
+
 	/**
 	 * Use Paint invalidation if you're changing a normal property involving painting or sizing.
 	 * Additionally if the property that was changed affects Volatility in anyway, it's important
@@ -165,10 +175,6 @@ enum class EInvalidateWidgetReason : uint8
 	 */
 	LayoutAndVolatility = Layout | Volatility,
 
-	/**
-	 * Changing visibility affects layout, paint, and volatility
-	 */
-	Visibility = Layout | Paint | Volatility,
 
 	/**
 	 * Do not use this ever unless you know what you are doing
@@ -181,18 +187,6 @@ ENUM_CLASS_FLAGS(EInvalidateWidgetReason)
 // This typedefed because EInvalidateWidget will be deprecated soon
 typedef EInvalidateWidgetReason EInvalidateWidget;
 
-/**
- * An ILayoutCache implementor is responsible for caching a the hierarchy of widgets it is drawing.
- * The shipped implementation of this is SInvalidationPanel.
- */
-class ILayoutCache
-{
-public:
-	virtual ~ILayoutCache() { }
-	virtual void InvalidateWidget(class SWidget* InvalidateWidget) = 0;
-	virtual FCachedWidgetNode* CreateCacheNode() const = 0;
-};
-
 
 /**
  * An FPopupLayer hosts the pop-up content which could be anything you want to appear on top of a widget.
@@ -269,12 +263,18 @@ class IToolTip;
  */
 class SLATECORE_API SWidget
 	: public FSlateControlledConstruction,
-	  public TSharedFromThis<SWidget>		// Enables 'this->AsShared()'
+	public TSharedFromThis<SWidget>		// Enables 'this->AsShared()'
 {
 	friend struct FCurveSequence;
-
+	friend class FWidgetProxy;
+	friend class FSlateInvalidationRoot;
+	friend class FSlateWindowElementList;
+	friend class SWindow;
+#if WITH_SLATE_DEBUGGING
+	friend struct FInvalidatedWidgetDrawer;
+#endif
 public:
-	
+
 	/**
 	 * Construct a SWidget based on initial parameters.
 	 */
@@ -315,8 +315,8 @@ public:
 
 	/**
 	 * Called to tell a widget to paint itself (and it's children).
-	 * 
-	 * The widget should respond by populating the OutDrawElements array with FDrawElements 
+	 *
+	 * The widget should respond by populating the OutDrawElements array with FDrawElements
 	 * that represent it and any of its children.
 	 *
 	 * @param Args              All the arguments necessary to paint this widget (@todo umg: move all params into this struct)
@@ -337,7 +337,7 @@ public:
 	 * @param  InCurrentTime  Current absolute real time
 	 * @param  InDeltaTime  Real time passed since last tick
 	 */
-	virtual void Tick( const FGeometry& AllottedGeometry, const double InCurrentTime, const float InDeltaTime );
+	virtual void Tick(const FGeometry& AllottedGeometry, const double InCurrentTime, const float InDeltaTime);
 
 	//
 	// KEY INPUT
@@ -404,7 +404,7 @@ public:
 	 * @param InKeyEvent  Key event
 	 * @return Returns whether the event was handled, along with other possible actions
 	 */
-	virtual FReply OnKeyUp( const FGeometry& MyGeometry, const FKeyEvent& InKeyEvent );
+	virtual FReply OnKeyUp(const FGeometry& MyGeometry, const FKeyEvent& InKeyEvent);
 
 	/**
 	 * Called when an analog value changes on a button that supports analog
@@ -413,7 +413,7 @@ public:
 	 * @param InAnalogInputEvent Analog input event
 	 * @return Returns whether the event was handled, along with other possible actions
 	 */
-	virtual FReply OnAnalogValueChanged( const FGeometry& MyGeometry, const FAnalogInputEvent& InAnalogInputEvent );
+	virtual FReply OnAnalogValueChanged(const FGeometry& MyGeometry, const FAnalogInputEvent& InAnalogInputEvent);
 
 	//
 	// MOUSE INPUT
@@ -431,12 +431,12 @@ public:
 	/**
 	 * Just like OnMouseButtonDown, but tunnels instead of bubbling.
 	 * If this even is handled, OnMouseButtonDown will not be sent.
-	 * 
+	 *
 	 * Use this event sparingly as preview events generally make UIs more
 	 * difficult to reason about.
 	 */
 	virtual FReply OnPreviewMouseButtonDown(const FGeometry& MyGeometry, const FPointerEvent& MouseEvent);
-	
+
 	/**
 	 * The system calls this method to notify the widget that a mouse button was release within it. This event is bubbled.
 	 *
@@ -445,7 +445,7 @@ public:
 	 * @return Whether the event was handled along with possible requests for the system to take action.
 	 */
 	virtual FReply OnMouseButtonUp(const FGeometry& MyGeometry, const FPointerEvent& MouseEvent);
-	
+
 	/**
 	 * The system calls this method to notify the widget that a mouse moved within it. This event is bubbled.
 	 *
@@ -454,7 +454,7 @@ public:
 	 * @return Whether the event was handled along with possible requests for the system to take action.
 	 */
 	virtual FReply OnMouseMove(const FGeometry& MyGeometry, const FPointerEvent& MouseEvent);
-	
+
 	/**
 	 * The system will use this event to notify a widget that the cursor has entered it. This event is uses a custom bubble strategy.
 	 *
@@ -462,14 +462,14 @@ public:
 	 * @param MouseEvent Information about the input event
 	 */
 	virtual void OnMouseEnter(const FGeometry& MyGeometry, const FPointerEvent& MouseEvent);
-	
+
 	/**
 	 * The system will use this event to notify a widget that the cursor has left it. This event is uses a custom bubble strategy.
 	 *
 	 * @param MouseEvent Information about the input event
 	 */
 	virtual void OnMouseLeave(const FPointerEvent& MouseEvent);
-	
+
 	/**
 	 * Called when the mouse wheel is spun. This event is bubbled.
 	 *
@@ -477,17 +477,17 @@ public:
 	 * @return  Returns whether the event was handled, along with other possible actions
 	 */
 	virtual FReply OnMouseWheel(const FGeometry& MyGeometry, const FPointerEvent& MouseEvent);
-	
+
 	/**
 	 * The system asks each widget under the mouse to provide a cursor. This event is bubbled.
-	 * 
+	 *
 	 * @return FCursorReply::Unhandled() if the event is not handled; return FCursorReply::Cursor() otherwise.
 	 */
 	virtual FCursorReply OnCursorQuery(const FGeometry& MyGeometry, const FPointerEvent& CursorEvent) const;
 
 	/**
 	 * After OnCursorQuery has specified a cursor type the system asks each widget under the mouse to map that cursor to a widget. This event is bubbled.
-	 * 
+	 *
 	 * @return TOptional<TSharedRef<SWidget>>() if you don't have a mapping otherwise return the Widget to show.
 	 */
 	virtual TOptional<TSharedRef<SWidget>> OnMapCursor(const FCursorReply& CursorReply) const;
@@ -505,7 +505,7 @@ public:
 	 * Called when Slate wants to visualize tooltip.
 	 * If nobody handles this event, Slate will use default tooltip visualization.
 	 * If you override this event, you should probably return true.
-	 * 
+	 *
 	 * @param  TooltipContent    The TooltipContent that I may want to visualize.
 	 * @return true if this widget visualized the tooltip content; i.e., the event is handled.
 	 */
@@ -513,9 +513,9 @@ public:
 
 	/**
 	 * Visualize a new pop-up if possible.  If it's not possible for this widget to host the pop-up
-	 * content you'll get back an invalid pointer to the layer.  The returned FPopupLayer allows you 
+	 * content you'll get back an invalid pointer to the layer.  The returned FPopupLayer allows you
 	 * to remove the pop-up when you're done with it
-	 * 
+	 *
 	 * @param PopupContent The widget to try and host overlaid on top of the widget.
 	 *
 	 * @return a valid FPopupLayer if this widget supported hosting it.  You can call Remove() on this to destroy the pop-up.
@@ -559,7 +559,7 @@ public:
 	 * @return A reply that indicated whether the contents of the DragDropEvent can potentially be processed by this widget.
 	 */
 	virtual void OnDragEnter(const FGeometry& MyGeometry, const FDragDropEvent& DragDropEvent);
-	
+
 	/**
 	 * Called during drag and drop when the drag leaves a widget.
 	 *
@@ -575,7 +575,7 @@ public:
 	 * @return A reply that indicated whether this event was handled.
 	 */
 	virtual FReply OnDragOver(const FGeometry& MyGeometry, const FDragDropEvent& DragDropEvent);
-	
+
 	/**
 	 * Called when the user is dropping something onto a widget; terminates drag and drop.
 	 *
@@ -596,38 +596,38 @@ public:
 	 * @return  Returns whether the event was handled, along with other possible actions
 	 */
 	virtual FReply OnTouchGesture(const FGeometry& MyGeometry, const FPointerEvent& GestureEvent);
-	
+
 	/**
 	 * Called when a touchpad touch is started (finger down)
-	 * 
+	 *
 	 * @param InTouchEvent	The touch event generated
 	 */
 	virtual FReply OnTouchStarted(const FGeometry& MyGeometry, const FPointerEvent& InTouchEvent);
 
 	/**
 	 * Called when a touchpad touch is moved  (finger moved)
-	 * 
+	 *
 	 * @param InTouchEvent	The touch event generated
 	 */
 	virtual FReply OnTouchMoved(const FGeometry& MyGeometry, const FPointerEvent& InTouchEvent);
 
 	/**
 	 * Called when a touchpad touch is ended (finger lifted)
-	 * 
+	 *
 	 * @param InTouchEvent	The touch event generated
 	 */
 	virtual FReply OnTouchEnded(const FGeometry& MyGeometry, const FPointerEvent& InTouchEvent);
 
 	/**
 	 * Called when a touchpad touch force changes
-	 * 
+	 *
 	 * @param InTouchEvent	The touch event generated
 	 */
 	virtual FReply OnTouchForceChanged(const FGeometry& MyGeometry, const FPointerEvent& TouchEvent);
 
 	/**
 	 * Called when a touchpad touch first moves after TouchStarted
-	 * 
+	 *
 	 * @param InTouchEvent	The touch event generated
 	 */
 	virtual FReply OnTouchFirstMove(const FGeometry& MyGeometry, const FPointerEvent& TouchEvent);
@@ -635,7 +635,7 @@ public:
 	/**
 	 * Called when motion is detected (controller or device)
 	 * e.g. Someone tilts or shakes their controller.
-	 * 
+	 *
 	 * @param InMotionEvent	The motion event generated
 	 */
 	virtual FReply OnMotionDetected(const FGeometry& MyGeometry, const FMotionEvent& InMotionEvent);
@@ -660,7 +660,7 @@ public:
 	virtual FPopupMethodReply OnQueryPopupMethod() const;
 
 	virtual TSharedPtr<FVirtualPointerPosition> TranslateMouseCoordinateFor3DChild(const TSharedRef<SWidget>& ChildWidget, const FGeometry& MyGeometry, const FVector2D& ScreenSpaceMouseCoordinate, const FVector2D& LastScreenSpaceMouseCoordinate) const;
-	
+
 	/**
 	 * All the pointer (mouse, touch, stylus, etc.) events from this frame have been routed.
 	 * This is a widget's chance to act on any accumulated data.
@@ -676,7 +676,7 @@ public:
 	/**
 	 * Called when navigation is requested
 	 * e.g. Left Joystick, Direction Pad, Arrow Keys can generate navigation events.
-	 * 
+	 *
 	 * @param InNavigationEvent	The navigation event generated
 	 */
 	virtual FNavigationReply OnNavigation(const FGeometry& MyGeometry, const FNavigationEvent& InNavigationEvent);
@@ -699,6 +699,7 @@ public:
 	// LAYOUT
 	//
 
+	bool NeedsPrepass() const { return bNeedsPrepass; }
 	/** DEPRECATED version of SlatePrepass that assumes no scaling beyond AppScale*/
 	//UE_DEPRECATED(4.20, "SlatePrepass requires a layout scale to be accurate.")
 	void SlatePrepass();
@@ -707,34 +708,40 @@ public:
 	 * Descends to leaf-most widgets in the hierarchy and gathers desired sizes on the way up.
 	 * i.e. Caches the desired size of all of this widget's children recursively, then caches desired size for itself.
 	 */
-	void SlatePrepass(float LayoutScaleMultiplier);
+	void SlatePrepass(float InLayoutScaleMultiplier);
 
 	void SetCanTick(bool bInCanTick) { bInCanTick ? AddUpdateFlags(EWidgetUpdateFlags::NeedsTick) : RemoveUpdateFlags(EWidgetUpdateFlags::NeedsTick); }
 	bool GetCanTick() const { return HasAnyUpdateFlags(EWidgetUpdateFlags::NeedsTick); }
 
 	virtual void ChildLayoutChanged(EInvalidateWidget InvalidateReason);
 
-public:
-	void InvalidatePrepass();
+	const FSlateWidgetPersistentState& GetPersistentState() const { return PersistentState; }
+	const FWidgetProxyHandle GetProxyHandle() const { return FastPathProxyHandle; }
 
 	/** @return the DesiredSize that was computed the last time CacheDesiredSize() was called. */
 	FVector2D GetDesiredSize() const;
 
-public:
 	void AssignParentWidget(TSharedPtr<SWidget> InParent);
 	bool ConditionallyDetatchParentWidget(SWidget* InExpectedParent);
+
 	/**  */
 	virtual bool ValidatePathToChild(SWidget* InChild) { return true; }
 
 	FORCEINLINE bool IsParentValid() const { return ParentWidgetPtr.IsValid(); }
 	FORCEINLINE TSharedPtr<SWidget> GetParentWidget() const { return ParentWidgetPtr.Pin(); }
 
+	FORCEINLINE TSharedPtr<SWidget> Advanced_GetPaintParentWidget() const { return PersistentState.PaintParent.Pin(); }
+
 	/**
-	* Calculates what if any clipping state changes need to happen when drawing this widget.
-	* @return the culling rect that should be used going forward.
-	*/
-	FSlateRect CalculateCullingAndClippingRules(const FGeometry& AllottedGeometry, const FSlateRect& IncomingCullingRect, bool& bClipToBounds, bool& bAlwaysClip, bool& bIntersectClipBounds) const;
+	 * Calculates what if any clipping state changes need to happen when drawing this widget.
+	 * @return the culling rect that should be used going forward.
+	 */
+	FSlateRect CalculateCullingAndClippingRules(const FGeometry& AllottedGeometry, const FSlateRect& MyCullingRect, bool& bClipToBounds, bool& bAlwaysClip, bool& bIntersectClipBounds) const;
+
 protected:
+	virtual bool CustomPrepass(float LayoutScaleMultiplier) { return false; }
+	void AssignIndicesToChildren(FSlateInvalidationRoot& Root, int32 ParentIndex, TArray<FWidgetProxy, TMemStackAllocator<>>& FastPathList, bool bParentVisible, bool bParentVolatile);
+
 	/**
 	 * The system calls this method. It performs a breadth-first traversal of every visible widget and asks
 	 * each widget to cache how big it needs to be in order to present all of its content.
@@ -762,23 +769,18 @@ protected:
 	}
 
 private:
+
+	void UpdateFastPathVisibility(bool bParentVisible, bool bWidgetRemoved);
+
+	void UpdateFastPathVolatility(bool bParentVolatile);
+
 	/**
 	 * Explicitly set the desired size. This is highly advanced functionality that is meant
 	 * to be used in conjunction with overriding CacheDesiredSize. Use ComputeDesiredSize() instead.
 	 */
-	void Advanced_SetDesiredSize(const FVector2D& InDesiredSize)
+	void SetDesiredSize(const FVector2D& InDesiredSize)
 	{
-		if (GSlateLayoutCaching)
-		{
-			DesiredSize = InDesiredSize;
-
-			bNeedsDesiredSize = false;
-		}
-		else
-		{
-			DesiredSize = InDesiredSize;
-			//ensure(DesiredSize.GetValue().Y >= 0);
-		}
+		DesiredSize = InDesiredSize;
 	}
 
 	void LayoutChanged(EInvalidateWidget InvalidateReason);
@@ -788,35 +790,41 @@ private:
 	void AddUpdateFlags(EWidgetUpdateFlags FlagsToAdd)
 	{
 		UpdateFlags |= FlagsToAdd;
+		if (FastPathProxyHandle.IsValid())
+		{
+			FastPathProxyHandle.UpdateWidgetFlags(UpdateFlags);
+		}
 	}
 
 	void RemoveUpdateFlags(EWidgetUpdateFlags FlagsToRemove)
 	{
 		UpdateFlags &= (~FlagsToRemove);
+		if (FastPathProxyHandle.IsValid())
+		{
+			FastPathProxyHandle.UpdateWidgetFlags(UpdateFlags);
+		}
 	}
 
+	void UpdateWidgetProxy(int32 NewLayerId, FSlateCachedElementListNode* CacheNode);
+
 public:
 
 	FORCEINLINE TStatId GetStatID() const
 	{
 #if STATS
 		// this is done to avoid even registering stats for a disabled group (unless we plan on using it later)
-        if (FThreadStats::IsCollectingData(GET_STATID(STAT_SlateVeryVerboseStatGroupTester)))
-        {
-            if (!StatID.IsValidStat())
-            {
-                CreateStatID();
-            }
-            return StatID;
-        }
+		if (FThreadStats::IsCollectingData(GET_STATID(STAT_SlateVeryVerboseStatGroupTester)))
+		{
+			if (!StatID.IsValidStat())
+			{
+				CreateStatID();
+			}
+			return StatID;
+		}
 #endif
 		return TStatId(); // not doing stats at the moment, or ever
 	}
 
-	/** What is the Child's scale relative to this widget. */
-	UE_DEPRECATED(4.15, "This version is no longer used, please use the new version which also provides the incoming prepass scale, in case that has bearing on the relative layout scale.")
-	virtual float GetRelativeLayoutScale(const FSlotBase& Child) const { return 1.0f; }
-
 	/** What is the Child's scale relative to this widget. */
 	virtual float GetRelativeLayoutScale(const FSlotBase& Child, float LayoutScaleMultiplier) const;
 
@@ -830,14 +838,16 @@ public:
 	 * @param ArrangedChildren    The array to which to add the WidgetGeometries that represent the arranged children.
 	 */
 	void ArrangeChildren(const FGeometry& AllottedGeometry, FArrangedChildren& ArrangedChildren) const;
-	
+
 	/**
 	 * Every widget that has children must implement this method. This allows for iteration over the Widget's
 	 * children regardless of how they are actually stored.
 	 */
-	// @todo Slate: Consider renaming to GetVisibleChildren  (not ALL children will be returned in all cases)
+	 // @todo Slate: Consider renaming to GetVisibleChildren  (not ALL children will be returned in all cases)
 	virtual FChildren* GetChildren() = 0;
 
+	virtual FChildren* GetAllChildren() { return GetChildren(); }
+
 	/**
 	 * Checks to see if this widget supports keyboard focus.  Override this in derived classes.
 	 *
@@ -911,9 +921,9 @@ public:
 	 *
 	 * @param InEnabledState	An attribute containing the enabled state or a delegate to call to get the enabled state.
 	 */
-	void SetEnabled( const TAttribute<bool>& InEnabledState )
+	void SetEnabled(const TAttribute<bool>& InEnabledState)
 	{
-		if ( !EnabledState.IdenticalTo(InEnabledState) )
+		if (!EnabledState.IdenticalTo(InEnabledState))
 		{
 			EnabledState = InEnabledState;
 			Invalidate(EInvalidateWidget::PaintAndVolatility);
@@ -921,7 +931,7 @@ public:
 	}
 
 	/** @return Whether or not this widget is enabled */
-	FORCEINLINE bool IsEnabled() const 
+	FORCEINLINE bool IsEnabled() const
 	{
 		return EnabledState.Get();
 	}
@@ -944,7 +954,7 @@ public:
 	 *
 	 * @param	bEnableForceField	True to enable tool tip force field for this widget
 	 */
-	void EnableToolTipForceField( const bool bEnableForceField );
+	void EnableToolTipForceField(const bool bEnableForceField);
 
 	/** @return True if a tool tip force field is active on this widget */
 	bool HasToolTipForceField() const
@@ -965,14 +975,7 @@ public:
 	FORCEINLINE EVisibility GetVisibility() const { return Visibility.Get(); }
 
 	/** @param InVisibility  should this widget be */
-	virtual void SetVisibility( TAttribute<EVisibility> InVisibility )
-	{
-		if ( !Visibility.IdenticalTo(InVisibility) )
-		{
-			Visibility = InVisibility;
-			Invalidate(EInvalidateWidget::Visibility);
-		}
-	}
+	virtual void SetVisibility(TAttribute<EVisibility> InVisibility);
 
 #if WITH_ACCESSIBILITY
 	/**
@@ -1038,18 +1041,11 @@ public:
 	virtual void SetDefaultAccessibleText(EAccessibleType AccessibleType = EAccessibleType::Main);
 #endif
 
-	/**
-	 * When performing a caching pass, volatile widgets are not cached as part of everything
-	 * else, instead they and their children are drawn as normal standard widgets and excluded
-	 * from the cache.  This is extremely useful for things like timers and text that change
-	 * every frame.
-	 */
+	/** Whether or not a widget is volatile and will update every frame without being invalidated */
 	FORCEINLINE bool IsVolatile() const { return bCachedVolatile; }
 
 	/**
-	 * Was this widget painted as part of a volatile pass previously.  This may mean it was the
-	 * widget directly responsible for making a hierarchy volatile, or it may mean it was simply
-	 * a child of a volatile panel.
+	 * This widget is volatile because its parent or some ancestor is volatile
 	 */
 	FORCEINLINE bool IsVolatileIndirectly() const { return bInheritedVolatility; }
 
@@ -1064,6 +1060,8 @@ public:
 		Invalidate(EInvalidateWidget::LayoutAndVolatility);
 	}
 
+	FORCEINLINE bool ShouldInvalidatePrepassDueToVolatility() { return bVolatilityAlwaysInvalidatesPrepass; }
+
 	/**
 	 * Invalidates the widget from the view of a layout caching widget that may own this widget.
 	 * will force the owning widget to redraw and cache children on the next paint pass.
@@ -1079,6 +1077,8 @@ public:
 		bCachedVolatile = bForceVolatile || ComputeVolatility();
 	}
 
+	void InvalidatePrepass();
+
 protected:
 
 #if SLATE_CULL_WIDGETS
@@ -1094,6 +1094,13 @@ protected:
 
 protected:
 
+
+	/**
+	 * Called when a child is removed from the tree parent's widget tree either by removing it from a slot. This can also be called manually if you've got some non-slot based what of no longer reporting children
+	 * An example of a widget that needs manual calling is SWidgetSwitcher.  It keeps all its children but only arranges and paints a single "active" one.  Once a child becomes inactive, its cached data should be removed.
+	 */
+	void InvalidateChildRemovedFromTree(SWidget& Child);
+
 	/**
 	 * Recalculates and caches volatility and returns 'true' if the volatility changed.
 	 */
@@ -1103,19 +1110,6 @@ protected:
 		CacheVolatility();
 		return bWasDirectlyVolatile != IsVolatile();
 	}
-
-	/**
-	 * Forces invalidation, doesn't check volatility.
-	 */
-	FORCEINLINE void Advanced_ForceInvalidateLayout()
-	{
-		TSharedPtr<ILayoutCache> SharedLayoutCache = LayoutCache.Pin();
-		if (SharedLayoutCache.IsValid() )
-		{
-			SharedLayoutCache->InvalidateWidget(this);
-		}
-	}
-
 public:
 
 	/** @return the render opacity of the widget. */
@@ -1170,7 +1164,7 @@ public:
 		if(!RenderTransform.IdenticalTo(InTransform))
 		{
 			RenderTransform = InTransform;
-			Invalidate(EInvalidateWidget::LayoutAndVolatility);
+			Invalidate(EInvalidateWidget::LayoutAndVolatility|EInvalidateWidget::RenderTransform);
 		}
 	}
 
@@ -1186,7 +1180,7 @@ public:
 		if (!RenderTransformPivot.IdenticalTo(InTransformPivot))
 		{
 			RenderTransformPivot = InTransformPivot;
-			Invalidate(EInvalidateWidget::LayoutAndVolatility);
+			Invalidate(EInvalidateWidget::LayoutAndVolatility | EInvalidateWidget::RenderTransform);
 		}
 	}
 
@@ -1199,6 +1193,7 @@ public:
 		{
 			Clipping = InClipping;
 			OnClippingChanged();
+			// @todo - Fast path should this be Paint?
 			Invalidate(EInvalidateWidget::Layout);
 		}
 	}
@@ -1218,6 +1213,7 @@ public:
 		if (CullingBoundsExtension != InCullingBoundsExtension)
 		{
 			CullingBoundsExtension = InCullingBoundsExtension;
+			// @todo - Fast path should this be Paint?
 			Invalidate(EInvalidateWidget::Layout);
 		}
 	}
@@ -1358,8 +1354,11 @@ public:
 	/** @return the Foreground color that this widget sets; unset options if the widget does not set a foreground color */
 	virtual FSlateColor GetForegroundColor() const;
 
+	//UE_DEPRECATED(4.23, "GetCachedGeometry has been deprecated, use GetTickSpaceGeometry instead")
+	const FGeometry& GetCachedGeometry() const;
+
 	/**
-	 * Gets the last geometry used to Tick the widget.  This data may not exist yet if this call happens prior to 
+	 * Gets the last geometry used to Tick the widget.  This data may not exist yet if this call happens prior to
 	 * the widget having been ticked/painted, or it may be out of date, or a frame behind.
 	 *
 	 * We recommend not to use this data unless there's no other way to solve your problem.  Normally in Slate we
@@ -1367,11 +1366,20 @@ public:
 	 * or what are referred to as hysteresis problems, both caused by depending on geometry from the previous frame
 	 * being used to advise how to layout a dependent object the current frame.
 	 */
-	const FGeometry& GetCachedGeometry() const { return CachedGeometry; }
+	const FGeometry& GetTickSpaceGeometry() const;
+
+	/**
+	 * Gets the last geometry used to Tick the widget.  This data may not exist yet if this call happens prior to
+	 * the widget having been ticked/painted, or it may be out of date, or a frame behind.
+ 	 */
+	const FGeometry& GetPaintSpaceGeometry() const;
+
+	/** Returns the clipping state to clip this widget against its parent */
+	const TOptional<FSlateClippingState>& GetCurrentClippingState() const { return PersistentState.InitialClipState; }
 
 	/** Is this widget derivative of SWindow */
 	virtual bool Advanced_IsWindow() const { return false; }
-
+	virtual bool Advanced_IsInvalidationRoot() const { return false; }
 protected:
 
 	/**
@@ -1485,15 +1493,11 @@ private:
 	 */
 	virtual void OnArrangeChildren(const FGeometry& AllottedGeometry, FArrangedChildren& ArrangedChildren) const = 0;
 
-protected:
+	void Prepass_Internal(float LayoutScaleMultiplier);
 
-	/**
-	 * Don't call this directly unless you're a layout cache, this is used to recursively set the layout cache on
-	 * on invisible children that never get the opportunity to paint and receive the layout cache through the normal
-	 * means.  That way if an invisible widget becomes visible, we still properly invalidate the hierarchy.
-	 */
-	void CachePrepass(const TWeakPtr<ILayoutCache>& LayoutCache);
+protected:
 
+	float GetPrepassLayoutScaleMultiplier() const { return PrepassLayoutScaleMultiplier; }
 public:
 	/**
 	 * Registers an "active timer" delegate that will execute at some regular interval. TickFunction will not be called until the specified interval has elapsed once.
@@ -1536,6 +1540,10 @@ protected:
 	/** Dtor ensures that active timer handles are UnRegistered with the SlateApplication. */
 	virtual ~SWidget();
 
+private:
+
+	/** Handle to the proxy when on the fast path */
+	mutable FWidgetProxyHandle FastPathProxyHandle;
 protected:
 	/** Is this widget hovered? */
 	uint8 bIsHovered : 1;
@@ -1570,7 +1578,10 @@ private:
 	uint8 bCachedVolatile : 1;
 
 	/** If we're owned by a volatile widget, we need inherit that volatility and use as part of our volatility, but don't cache it. */
-	mutable uint8 bInheritedVolatility : 1;
+	uint8 bInheritedVolatility : 1;
+
+	/** If the widget is hidden or collapsed to ancestor visibility */
+	uint8 bInvisibleDueToParentOrSelfVisibility : 1;
 
 	/** Are we currently updating the desired size? */
 	uint8 bNeedsPrepass : 1;
@@ -1581,6 +1592,11 @@ private:
 	mutable uint8 bUpdatingDesiredSize : 1;
 
 protected:
+	uint8 bHasCustomPrepass : 1;
+	
+	/** if this widget should always invalidate the prepass step when volatile */
+	uint8 bVolatilityAlwaysInvalidatesPrepass : 1;
+
 #if !(UE_BUILD_SHIPPING || UE_BUILD_TEST)
 	UE_DEPRECATED(4.21, "Setting bCanTick on a widget directly is deprecated and will not function.  Call SetCanTick instead")
 	uint8 bCanTick : 1;
@@ -1612,25 +1628,20 @@ private:
 	/** Flow direction preference */
 	EFlowDirectionPreference FlowDirectionPreference;
 
-private:
 	EWidgetUpdateFlags UpdateFlags;
 
-	/** Stores the ideal size this widget wants to be. */
-	TOptional<FVector2D> DesiredSize;
 
-	float PrepassLayoutScaleMultiplier;
+	mutable FSlateWidgetPersistentState PersistentState;
 
-	/**
-	 * Stores the cached Tick Geometry of the widget.  This information can and will be outdated, that's the
-	 * nature of it.  However, users were found to often need access to the geometry at times inconvenient to always
-	 * need to be located in Widget Tick.
-	 */
-	mutable FGeometry CachedGeometry;
+	/** Stores the ideal size this widget wants to be. */
+	TOptional<FVector2D> DesiredSize;
 
 	/** The list of active timer handles for this widget. */
 	TArray<TSharedRef<FActiveTimerHandle>> ActiveTimers;
 
 protected:
+
+	float PrepassLayoutScaleMultiplier;
 	/**
 	* Can be used to enlarge the culling bounds of this widget (pre-intersection), this can be useful if you've got
 	* children that you know are using rendering transforms to render outside their standard bounds, if that happens
@@ -1654,12 +1665,12 @@ protected:
 	/** Render transform pivot of this widget (in normalized local space) */
 	TAttribute< FVector2D > RenderTransformPivot;
 
+
 #if WITH_ACCESSIBILITY
 	/** All variables surrounding how this widget is exposed to the platform's accessibility API. */
 	FAccessibleWidgetData AccessibleData;
 #endif
 
-protected:
 	/** Debugging information on the type of widget we're creating for the Widget Reflector. */
 	FName TypeOfWidget;
 
@@ -1684,20 +1695,15 @@ private:
 	/** Tool tip content for this widget */
 	TSharedPtr<IToolTip> ToolTip;
 
-	/** The current layout cache that may need to invalidated by changes to this widget. */
-	mutable TWeakPtr<ILayoutCache> LayoutCache;
-
 	/** Pointer to this widgets parent widget.  If it is null this is a root widget or it is not in the widget tree */
 	TWeakPtr<SWidget> ParentWidgetPtr;
 
-private:
 	// Events
 	TArray<TPair<FName, FPointerEventHandler>> PointerEvents;
 
 	FNoReplyPointerEventHandler MouseEnterHandler;
 	FSimpleNoReplyPointerEventHandler MouseLeaveHandler;
 
-private:
 	STAT(mutable TStatId StatID;)
 	STAT(size_t AllocSize;)
 };
diff --git a/Engine/Source/Runtime/SlateCore/Public/Widgets/SWindow.h b/Engine/Source/Runtime/SlateCore/Public/Widgets/SWindow.h
--- a/Engine/Source/Runtime/SlateCore/Public/Widgets/SWindow.h
+++ b/Engine/Source/Runtime/SlateCore/Public/Widgets/SWindow.h
@@ -28,6 +28,7 @@
 #include "Widgets/SOverlay.h"
 #include "Styling/SlateTypes.h"
 #include "Styling/CoreStyle.h"
+#include "FastUpdate/SlateInvalidationRoot.h"
 
 class FActiveTimerHandle;
 class FHittestGrid;
@@ -37,6 +38,9 @@ class FWidgetPath;
 class IWindowTitleBar;
 class SPopupLayer;
 class SWindow;
+class SImage;
+
+enum class EUpdateFastPathReason : uint8;
 
 /** Notification that a window has been activated */
 DECLARE_DELEGATE( FOnWindowActivated );
@@ -125,7 +129,9 @@ private:
  */
 class SLATECORE_API SWindow
 	: public SCompoundWidget
+	, public FSlateInvalidationRoot
 {
+
 public:
 
 	SLATE_BEGIN_ARGS( SWindow )
@@ -273,6 +279,7 @@ public:
 	 * Default constructor. Use SNew(SWindow) instead.
 	 */
 	SWindow();
+	~SWindow();
 
 public:
 
@@ -340,7 +347,7 @@ public:
 	}
 
 	/** Paint the window and all of its contents. Not the same as Paint(). */
-	int32 PaintWindow( const FPaintArgs& Args, const FGeometry& AllottedGeometry, const FSlateRect& MyCullingRect, FSlateWindowElementList& OutDrawElements, int32 LayerId, const FWidgetStyle& InWidgetStyle, bool bParentEnabled ) const;
+	int32 PaintWindow( double CurrentTime, float DeltaTime, FSlateWindowElementList& OutDrawElements, const FWidgetStyle& InWidgetStyle, bool bParentEnabled );
 
 	/**
 	 * Returns the size of the title bar as a Slate size parameter.  Does not take into account application scale!
@@ -784,7 +791,8 @@ public:
 
 	bool IsDrawingEnabled() const { return bIsDrawingEnabled; }
 
-	virtual bool Advanced_IsWindow() const { return true; }
+	virtual bool Advanced_IsWindow() const override { return true; }
+	virtual bool Advanced_IsInvalidationRoot() const override { return bAllowFastUpdate; }
 
 #if WITH_ACCESSIBILITY
 	virtual TSharedRef<FSlateAccessibleWidget> CreateAccessibleWidget() override;
@@ -800,7 +808,9 @@ private:
 
 	/** The window's desired size takes into account the ratio between the slate units and the pixel size */
 	virtual FVector2D ComputeDesiredSize(float) const override;
+	virtual bool ComputeVolatility() const override;
 
+	void OnGlobalInvalidationToggled(bool bGlobalInvalidationEnabled);
 public:
 	// For a given client size, calculate the window size required to accomodate any potential non-OS borders and tilebars
 	FVector2D GetWindowSizeFromClientSize(FVector2D InClientSize);
@@ -904,11 +914,12 @@ public:
 	 *
 	 * @see FHittestGrid for more details.
 	 */
-	TSharedRef<FHittestGrid> GetHittestGrid();
+	FHittestGrid& GetHittestGrid();
 
 	/** Optional constraints on min and max sizes that this window can be. */
 	FWindowSizeLimits GetSizeLimits() const;
 
+	void SetAllowFastUpdate(bool bInAllowFastUpdate);
 public:
 
 	// SWidget overrides
@@ -923,24 +934,11 @@ protected:
 	virtual TSharedRef<SWidget> MakeWindowTitleBar(const TSharedRef<SWindow>& Window, const TSharedPtr<SWidget>& CenterContent, EHorizontalAlignment CenterContentAlignment);
 	/**Returns the alignment type for the titlebar's title text. */
 	virtual EHorizontalAlignment GetTitleAlignment();
-	/** Get the desired color of titlebar items. These change during flashing. */
-	FSlateColor GetWindowTitleContentColor() const;
 
 	/** Kick off a morph to whatever the target shape happens to be. */
 	void StartMorph();
 
-	/** Get the brush used to draw the window background */
-	const FSlateBrush* GetWindowBackground() const;
-
-	/** Get the color used to tint the window background */
-	FSlateColor GetWindowBackgroundColor() const;
-
-	/** Get the brush used to draw the window outline */
-	const FSlateBrush* GetWindowOutline() const;
-
-	/** Get the color used to tint the window outline */
-	FSlateColor GetWindowOutlineColor() const;
-
+	virtual bool CustomPrepass(float LayoutScaleMultiplier) override;
 protected:
 
 	/** Type of the window */
@@ -1019,6 +1017,9 @@ protected:
 
 	bool bManualManageDPI : 1;
 
+	/** True if this window allows global invalidation of its contents */
+	bool bAllowFastUpdate : 1;
+
 	/** When should the window be activated upon being shown */
 	EWindowActivationPolicy WindowActivationPolicy;
 
@@ -1092,10 +1093,17 @@ protected:
 		restore focus to a widget after the window regains focus. */
 	TWeakPtr< SWidget > WidgetFocusedOnDeactivate;
 
+private:
 	/** Style used to draw this window */
 	const FWindowStyle* Style;
+
 	const FSlateBrush* WindowBackground;
 
+	TSharedPtr<SImage> WindowBackgroundImage;
+	TSharedPtr<SImage> WindowBorder;
+	TSharedPtr<SImage> WindowOutline;
+	TSharedPtr<SWidget> ContentAreaVBox;
+	EVisibility WindowContentVisibility;
 protected:
 
 	/** Min and Max values for Width and Height; all optional. */
@@ -1105,7 +1113,7 @@ protected:
 	TSharedPtr<FGenericWindow> NativeWindow;
 
 	/** Each window has its own hittest grid for accelerated widget picking. */
-	TSharedRef<FHittestGrid> HittestGrid;
+	TUniquePtr<FHittestGrid> HittestGrid;
 	
 	/** Invoked when the window has been activated. */
 	FOnWindowActivated OnWindowActivated;
@@ -1174,20 +1182,15 @@ protected:
 	
 	void ConstructWindowInternals();
 
-private:
+	/** One-off active timer to trigger a the morph sequence to play */
+	EActiveTimerReturnType TriggerPlayMorphSequence( double InCurrentTime, float InDeltaTime );
 
-	/**
-	 * @return EVisibility::Visible if we are showing this viewports content.  EVisibility::Hidden otherwise (we hide the content during full screen overlays)
-	 */
-	EVisibility GetWindowContentVisibility() const;
+	void SetWindowBackground(const FSlateBrush* InWindowBackground);
 
-	/**
-	 * @return EVisibility::Visible if the window is flashing. Used to show/hide the white flash in the title area
-	 */
-	EVisibility GetWindowFlashVisibility() const;
+	void UpdateWindowContentVisibility();
 
-	/** One-off active timer to trigger a the morph sequence to play */
-	EActiveTimerReturnType TriggerPlayMorphSequence( double InCurrentTime, float InDeltaTime );
+	int32 PaintSlowPath(const FSlateInvalidationContext& InvalidationContext) override;
+private:
 
 	/** The handle to the active timer */
 	TWeakPtr<FActiveTimerHandle> ActiveTimerHandle;
diff --git a/Engine/Source/Runtime/SlateNullRenderer/Private/SlateNullRenderer.cpp b/Engine/Source/Runtime/SlateNullRenderer/Private/SlateNullRenderer.cpp
--- a/Engine/Source/Runtime/SlateNullRenderer/Private/SlateNullRenderer.cpp
+++ b/Engine/Source/Runtime/SlateNullRenderer/Private/SlateNullRenderer.cpp
@@ -43,11 +43,6 @@ void FSlateNullRenderer::OnWindowDestroyed( const TSharedRef<SWindow>& InWindow
 
 void FSlateNullRenderer::DrawWindows( FSlateDrawBuffer& WindowDrawBuffer )
 {
-	// Must call PostDraw to clean up resources built in FSlateWindowElementList.
-	for (const TSharedRef<FSlateWindowElementList>& ElementList : WindowDrawBuffer.GetWindowElementLists())
-	{
-		ElementList->PostDraw_NonParallelRenderer();
-	}
 }
 
 FIntPoint FSlateNullRenderer::GenerateDynamicImageResource(const FName InTextureName)
diff --git a/Engine/Source/Runtime/SlateNullRenderer/Private/SlateNullRendererModule.cpp b/Engine/Source/Runtime/SlateNullRenderer/Private/SlateNullRendererModule.cpp
--- a/Engine/Source/Runtime/SlateNullRenderer/Private/SlateNullRendererModule.cpp
+++ b/Engine/Source/Runtime/SlateNullRenderer/Private/SlateNullRendererModule.cpp
@@ -148,7 +148,6 @@ private:
 		if (!ResourceManager.IsValid())
 		{
 			ResourceManager = MakeShareable(new FSlateNullShaderResourceManager);
-			FSlateDataPayload::ResourceManager = ResourceManager.Get();
 		}
 	}
 
diff --git a/Engine/Source/Runtime/SlateRHIRenderer/Private/Slate3DRenderer.cpp b/Engine/Source/Runtime/SlateRHIRenderer/Private/Slate3DRenderer.cpp
--- a/Engine/Source/Runtime/SlateRHIRenderer/Private/Slate3DRenderer.cpp
+++ b/Engine/Source/Runtime/SlateRHIRenderer/Private/Slate3DRenderer.cpp
@@ -17,7 +17,7 @@ FSlate3DRenderer::FSlate3DRenderer( TSharedRef<FSlateFontServices> InSlateFontSe
 	RenderTargetPolicy = MakeShareable( new FSlateRHIRenderingPolicy( SlateFontServices, ResourceManager, InitialBufferSize ) );
 	RenderTargetPolicy->SetUseGammaCorrection( bUseGammaCorrection );
 
-	ElementBatcher = MakeShareable(new FSlateElementBatcher(RenderTargetPolicy.ToSharedRef()));
+	ElementBatcher = MakeUnique<FSlateElementBatcher>(RenderTargetPolicy.ToSharedRef());
 }
 
 void FSlate3DRenderer::Cleanup()
@@ -149,16 +149,11 @@ void FSlate3DRenderer::DrawWindowToTarget_RenderThread(FRHICommandListImmediate&
 			FSlateWindowElementList& WindowElementList = *WindowsToDraw[WindowIndex];
 
 			FSlateBatchData& BatchData = WindowElementList.GetBatchData();
-			FElementBatchMap& RootBatchMap = WindowElementList.GetRootDrawLayer().GetElementBatchMap();
-
-			WindowElementList.PreDraw_ParallelThread();
-
-			BatchData.CreateRenderBatches(RootBatchMap);
 
 			if (BatchData.GetRenderBatches().Num() > 0)
 			{
-				RenderTargetPolicy->UpdateVertexAndIndexBuffers(InRHICmdList, BatchData);
-
+				RenderTargetPolicy->BuildRenderingBuffers(InRHICmdList, BatchData);
+		
 				FVector2D DrawOffset = Context.WindowDrawBuffer->ViewOffset;
 
 				FMatrix ProjectionMatrix = FSlateRHIRenderer::CreateProjectionMatrix(RTTextureRHI->GetSizeX(), RTTextureRHI->GetSizeY());
@@ -192,6 +187,7 @@ void FSlate3DRenderer::DrawWindowToTarget_RenderThread(FRHICommandListImmediate&
 					BackBufferTarget,
 					ColorTarget,
 					DepthStencil,
+					BatchData.GetFirstRenderBatchIndex(),
 					BatchData.GetRenderBatches(),
 					DrawOptions
 				);
diff --git a/Engine/Source/Runtime/SlateRHIRenderer/Private/SlateElementIndexBuffer.cpp b/Engine/Source/Runtime/SlateRHIRenderer/Private/SlateElementIndexBuffer.cpp
--- a/Engine/Source/Runtime/SlateRHIRenderer/Private/SlateElementIndexBuffer.cpp
+++ b/Engine/Source/Runtime/SlateRHIRenderer/Private/SlateElementIndexBuffer.cpp
@@ -77,7 +77,7 @@ void FSlateElementIndexBuffer::ResizeBuffer( int32 NewSizeBytes )
 
 void FSlateElementIndexBuffer::PreFillBuffer(int32 RequiredIndexCount, bool bShrinkToMinSize)
 {
-	//SCOPE_CYCLE_COUNTER( STAT_SlatePreFullBufferRTTime );
+	SCOPE_CYCLE_COUNTER( STAT_SlatePreFullBufferRTTime );
 
 	checkSlow(IsInRenderingThread());
 
@@ -96,15 +96,29 @@ void FSlateElementIndexBuffer::PreFillBuffer(int32 RequiredIndexCount, bool bShr
 	}
 }
 
-void* FSlateElementIndexBuffer::LockBuffer_RenderThread(int32 NumIndices)
+void* FSlateElementIndexBuffer::LockBuffer(int32 NumIndices, bool bInRenderThread)
 {
 	uint32 RequiredBufferSize = NumIndices*sizeof(SlateIndex);		
-	return RHILockIndexBuffer( IndexBufferRHI, 0, RequiredBufferSize, RLM_WriteOnly );
+	if(bInRenderThread)
+	{
+		return RHILockIndexBuffer(IndexBufferRHI, 0, RequiredBufferSize, RLM_WriteOnly);
+	}
+	else
+	{
+		return GDynamicRHI->RHILockIndexBuffer(IndexBufferRHI, 0, RequiredBufferSize, RLM_WriteOnly);
+	}
 }
 
-void FSlateElementIndexBuffer::UnlockBuffer_RenderThread()
+void FSlateElementIndexBuffer::UnlockBuffer(bool bInRenderThread)
 {
-	RHIUnlockIndexBuffer( IndexBufferRHI );
+	if (bInRenderThread)
+	{
+		RHIUnlockIndexBuffer(IndexBufferRHI);
+	}
+	else
+	{
+		GDynamicRHI->RHIUnlockIndexBuffer(IndexBufferRHI);
+	}
 }
 
 /** Releases the index buffers RHI resource. */
diff --git a/Engine/Source/Runtime/SlateRHIRenderer/Private/SlateElementIndexBuffer.h b/Engine/Source/Runtime/SlateRHIRenderer/Private/SlateElementIndexBuffer.h
--- a/Engine/Source/Runtime/SlateRHIRenderer/Private/SlateElementIndexBuffer.h
+++ b/Engine/Source/Runtime/SlateRHIRenderer/Private/SlateElementIndexBuffer.h
@@ -40,8 +40,8 @@ public:
 	/** Resizes buffer, accumulates states safely on render thread */
 	void PreFillBuffer(int32 RequiredIndexCount, bool bShrinkToMinSize);
 
-	void* LockBuffer_RenderThread(int32 NumIndices);
-	void UnlockBuffer_RenderThread();
+	void* LockBuffer(int32 NumIndices, bool bInRenderThread);
+	void UnlockBuffer(bool bInRenderThread);
 
 	int32 GetMinBufferSize() const { return MinBufferSize; }
 
diff --git a/Engine/Source/Runtime/SlateRHIRenderer/Private/SlateElementVertexBuffer.h b/Engine/Source/Runtime/SlateRHIRenderer/Private/SlateElementVertexBuffer.h
--- a/Engine/Source/Runtime/SlateRHIRenderer/Private/SlateElementVertexBuffer.h
+++ b/Engine/Source/Runtime/SlateRHIRenderer/Private/SlateElementVertexBuffer.h
@@ -92,6 +92,8 @@ public:
 	/** Resizes buffer, accumulates states safely on render thread */
 	void PreFillBuffer(int32 RequiredVertexCount, bool bShrinkToMinSize)
 	{
+		SCOPE_CYCLE_COUNTER(STAT_SlatePreFullBufferRTTime);
+
 		checkSlow(IsInRenderingThread());
 
 		if (RequiredVertexCount > 0 )
@@ -115,16 +117,33 @@ public:
 
 	int32 GetMinBufferSize() const { return MinBufferSize; }
 
-	void* LockBuffer_RenderThread(int32 NumVertices)
+	void* LockBuffer(int32 NumVertices, bool bInRenderThread)
 	{
 		uint32 RequiredBufferSize = NumVertices*sizeof(VertexType);	
 
-		return RHILockVertexBuffer( VertexBufferRHI, 0, RequiredBufferSize, RLM_WriteOnly );
+		if (bInRenderThread)
+		{
+			// Render thread
+			return RHILockVertexBuffer(VertexBufferRHI, 0, RequiredBufferSize, RLM_WriteOnly);
+		}
+		else
+		{
+			// RHI thread
+			return GDynamicRHI->RHILockVertexBuffer(VertexBufferRHI, 0, RequiredBufferSize, RLM_WriteOnly);
+		}
+	
 	}
 
-	void UnlockBuffer_RenderThread()
+	void UnlockBuffer(bool bInRenderThread)
 	{
-		RHIUnlockVertexBuffer( VertexBufferRHI );
+		if (bInRenderThread)
+		{
+			RHIUnlockVertexBuffer(VertexBufferRHI);
+		}
+		else
+		{
+			return GDynamicRHI->RHIUnlockVertexBuffer(VertexBufferRHI);
+		}
 	}
 
 private:
diff --git a/Engine/Source/Runtime/SlateRHIRenderer/Private/SlateMaterialShader.h b/Engine/Source/Runtime/SlateRHIRenderer/Private/SlateMaterialShader.h
--- a/Engine/Source/Runtime/SlateRHIRenderer/Private/SlateMaterialShader.h
+++ b/Engine/Source/Runtime/SlateRHIRenderer/Private/SlateMaterialShader.h
@@ -105,7 +105,7 @@ public:
 	}
 };
 
-template<ESlateShader::Type ShaderType,bool bDrawDisabledEffect> 
+template<ESlateShader ShaderType,bool bDrawDisabledEffect> 
 class TSlateMaterialShaderPS : public FSlateMaterialShaderPS
 {
 public:
diff --git a/Engine/Source/Runtime/SlateRHIRenderer/Private/SlateRHIRenderer.cpp b/Engine/Source/Runtime/SlateRHIRenderer/Private/SlateRHIRenderer.cpp
--- a/Engine/Source/Runtime/SlateRHIRenderer/Private/SlateRHIRenderer.cpp
+++ b/Engine/Source/Runtime/SlateRHIRenderer/Private/SlateRHIRenderer.cpp
@@ -31,10 +31,10 @@
 #include "RenderTargetPool.h"
 #include "RendererUtils.h"
 #include "HAL/LowLevelMemTracker.h"
+#include "Rendering/RenderingCommon.h"
 
 DECLARE_CYCLE_STAT(TEXT("Slate RT: Rendering"), STAT_SlateRenderingRTTime, STATGROUP_Slate);
-DECLARE_CYCLE_STAT(TEXT("Slate RT: Create Batches"), STAT_SlateRTCreateBatches, STATGROUP_Slate);
-DECLARE_CYCLE_STAT(TEXT("Slate RT: Fill Vertex & Index Buffers"), STAT_SlateRTFillVertexIndexBuffers, STATGROUP_Slate);
+
 DECLARE_CYCLE_STAT(TEXT("Slate RT: Draw Batches"), STAT_SlateRTDrawBatches, STATGROUP_Slate);
 
 DECLARE_GPU_STAT_NAMED(SlateUI, TEXT("Slate UI"));
@@ -118,25 +118,20 @@ void FViewportInfo::ReleaseResource()
 
 void FViewportInfo::ConditionallyUpdateDepthBuffer(bool bInRequiresStencilTest, uint32 InWidth, uint32 InHeight)
 {
-	FViewportInfo* ViewportInfo = this;
-	ENQUEUE_RENDER_COMMAND(UpdateDepthBufferCommand)(
-		[ViewportInfo, bInRequiresStencilTest, InWidth, InHeight](FRHICommandListImmediate& RHICmdList)
-	{
-		bool bDepthStencilStale =
-			bInRequiresStencilTest &&
-			(!ViewportInfo->bRequiresStencilTest ||
-			(ViewportInfo->DepthStencil.IsValid() && (ViewportInfo->DepthStencil->GetSizeX() != InWidth || ViewportInfo->DepthStencil->GetSizeY() != InHeight))
-				);
+	check(IsInRenderingThread());
 
-		ViewportInfo->bRequiresStencilTest = bInRequiresStencilTest;
+	bool bDepthStencilStale =
+		bInRequiresStencilTest &&
+		(!bRequiresStencilTest ||
+		(DepthStencil.IsValid() && (DepthStencil->GetSizeX() != InWidth || DepthStencil->GetSizeY() != InHeight)));
 
-		// Allocate a stencil buffer if needed and not already allocated
-		if (bDepthStencilStale)
-		{
-			ViewportInfo->RecreateDepthBuffer_RenderThread();
-		}
+	bRequiresStencilTest = bInRequiresStencilTest;
+
+	// Allocate a stencil buffer if needed and not already allocated
+	if (bDepthStencilStale)
+	{
+		RecreateDepthBuffer_RenderThread();
 	}
-	);
 }
 
 void FViewportInfo::RecreateDepthBuffer_RenderThread()
@@ -201,7 +196,7 @@ bool FSlateRHIRenderer::Initialize()
 
 	RenderingPolicy = MakeShareable(new FSlateRHIRenderingPolicy(SlateFontServices.ToSharedRef(), ResourceManager.ToSharedRef()));
 
-	ElementBatcher = MakeShareable(new FSlateElementBatcher(RenderingPolicy.ToSharedRef()));
+	ElementBatcher = MakeUnique<FSlateElementBatcher>(RenderingPolicy.ToSharedRef());
 
 	CurrentSceneIndex = -1;
 	ActiveScenes.Empty();
@@ -223,7 +218,6 @@ void FSlateRHIRenderer::Destroy()
 	FlushRenderingCommands();
 
 
-	check(ElementBatcher.IsUnique());
 	ElementBatcher.Reset();
 	RenderingPolicy.Reset();
 	ResourceManager.Reset();
@@ -661,6 +655,8 @@ void FSlateRHIRenderer::DrawWindow_RenderThread(FRHICommandListImmediate& RHICmd
 {
 	LLM_SCOPE(ELLMTag::SceneRender);
 
+	FMemMark MemMark(FMemStack::Get());
+
 	static uint32 LastTimestamp = FPlatformTime::Cycles();
 	{
 		SCOPED_DRAW_EVENTF(RHICmdList, SlateUI, TEXT("SlateUI Title = %s"), DrawCommandParams.WindowTitle.IsEmpty() ? TEXT("<none>") : *DrawCommandParams.WindowTitle);
@@ -670,9 +666,6 @@ void FSlateRHIRenderer::DrawWindow_RenderThread(FRHICommandListImmediate& RHICmd
 		// Should only be called by the rendering thread
 		check(IsInRenderingThread());
 
-		// Any window that draws should be reporting referenced UObjects for the duration of the frame.
-		check(WindowElementList.ShouldReportUObjectReferences());
-
 		FMaterialRenderProxy::UpdateDeferredCachedUniformExpressions();
 		GetRendererModule().InitializeSystemTextures(RHICmdList);
 
@@ -711,22 +704,12 @@ void FSlateRHIRenderer::DrawWindow_RenderThread(FRHICommandListImmediate& RHICmd
 			SCOPE_CYCLE_COUNTER(STAT_SlateRenderingRTTime);
 
 			FSlateBatchData& BatchData = WindowElementList.GetBatchData();
-			FElementBatchMap& RootBatchMap = WindowElementList.GetRootDrawLayer().GetElementBatchMap();
 
-			WindowElementList.PreDraw_ParallelThread();
+			// Update the vertex and index buffer	
+			RenderingPolicy->BuildRenderingBuffers(RHICmdList, BatchData);
 
-			{
-				SCOPE_CYCLE_COUNTER(STAT_SlateRTCreateBatches);
-				// Update the vertex and index buffer	
-				BatchData.CreateRenderBatches(RootBatchMap);
-			}
-
-			RootBatchMap.UpdateResourceVersion(ResourceVersion);
-
-			{
-				SCOPE_CYCLE_COUNTER(STAT_SlateRTFillVertexIndexBuffers);
-				RenderingPolicy->UpdateVertexAndIndexBuffers(RHICmdList, BatchData);
-			}
+			// This must happen after rendering buffers are created
+			ViewportInfo.ConditionallyUpdateDepthBuffer(BatchData.IsStencilClippingRequired(), ViewportInfo.DesiredWidth, ViewportInfo.DesiredHeight);
 
 			// should have been created by the game thread
 			check(IsValidRef(ViewportInfo.ViewportRHI));
@@ -841,6 +824,7 @@ void FSlateRHIRenderer::DrawWindow_RenderThread(FRHICommandListImmediate& RHICmd
 							BackBufferTarget,
 							BackBuffer,
 							ViewportInfo.bRequiresStencilTest ? ViewportInfo.DepthStencil : EmptyTarget,
+							BatchData.GetFirstRenderBatchIndex(),
 							BatchData.GetRenderBatches(),
 							RenderParams
 						);
@@ -1050,9 +1034,6 @@ void FSlateRHIRenderer::DrawWindow_RenderThread(FRHICommandListImmediate& RHICmd
 			GWorkingRHIThreadStallTime = 0;
 		});
 	}
-
-	// Any window that draws should be reporting referenced UObjects for the duration of the frame.
-	check(WindowElementList.ShouldReportUObjectReferences());
 }
 
 void FSlateRHIRenderer::DrawWindows(FSlateDrawBuffer& WindowDrawBuffer)
@@ -1116,7 +1097,6 @@ void FSlateRHIRenderer::DrawWindows_Private(FSlateDrawBuffer& WindowDrawBuffer)
 				FontCache->UpdateCache();
 
 				bool bLockToVsync = ElementBatcher->RequiresVsync();
-				bool bRequiresStencilTest = ElementList.GetBatchData().IsStencilClippingRequired();
 
 				bool bForceVsyncFromCVar = false;
 				if (GIsEditor)
@@ -1144,13 +1124,6 @@ void FSlateRHIRenderer::DrawWindows_Private(FSlateDrawBuffer& WindowDrawBuffer)
 					ConditionalResizeViewport(ViewInfo, ViewInfo->DesiredWidth, ViewInfo->DesiredHeight, IsViewportFullscreen(*Window));
 				}
 
-				
-				if (bRequiresStencilTest || bRequiresStencilTest != ViewInfo->bRequiresStencilTest)
-				{
-					ViewInfo->ConditionallyUpdateDepthBuffer(bRequiresStencilTest, ViewInfo->DesiredWidth, ViewInfo->DesiredHeight);
-				}
-				
-
 				// Tell the rendering thread to draw the windows
 				{
 					FSlateDrawWindowCommandParams Params;
@@ -1503,6 +1476,46 @@ void FSlateRHIRenderer::ClearScenes()
 	}
 }
 
+
+struct FClearCachedRenderingDataCommand final : public FRHICommand < FClearCachedRenderingDataCommand >
+{
+public:
+	FClearCachedRenderingDataCommand(FSlateCachedFastPathRenderingData* InCachedRenderingData)
+		: CachedRenderingData(InCachedRenderingData)
+	{
+		
+	}
+
+	void Execute(FRHICommandListBase& CmdList)
+	{
+		delete CachedRenderingData;
+	}
+
+private:
+	FSlateCachedFastPathRenderingData* CachedRenderingData;
+};
+
+void FSlateRHIRenderer::DestroyCachedFastPathRenderingData(struct FSlateCachedFastPathRenderingData* CachedRenderingData)
+{
+	check(CachedRenderingData);
+
+	// Cached data should be destroyed in a thread safe way.  If there is an rhi thread it could be reading from the data to copy it into a vertex buffer
+	// so delete it on the rhi thread if necessary, otherwise delete it on the render thread
+	ENQUEUE_RENDER_COMMAND(ClearCachedRenderingData)(
+		[CachedRenderingData](FRHICommandListImmediate& RHICmdList)
+	{
+		if (!RHICmdList.Bypass())
+		{
+			new (RHICmdList.AllocCommand<FClearCachedRenderingDataCommand>()) FClearCachedRenderingDataCommand(CachedRenderingData);
+		}
+		else
+		{
+			FClearCachedRenderingDataCommand Cmd(CachedRenderingData);
+			Cmd.Execute(RHICmdList);
+		}
+	});
+}
+
 bool FSlateRHIRenderer::AreShadersInitialized() const
 {
 #if WITH_EDITORONLY_DATA
@@ -1539,7 +1552,7 @@ void FSlateRHIRenderer::ReleaseAccessedResources(bool bImmediatelyFlush)
 		// This should NOT be done unless flushing
 		RenderingPolicy->FlushGeneratedResources();
 
-		FlushCommands();
+		//FlushCommands();
 	}
 }
 
@@ -1583,75 +1596,6 @@ void FSlateRHIRenderer::AddWidgetRendererUpdate(const struct FRenderThreadUpdate
 	}
 }
 
-
-TSharedRef<FSlateRenderDataHandle, ESPMode::ThreadSafe> FSlateRHIRenderer::CacheElementRenderData(const ILayoutCache* Cacher, FSlateWindowElementList& ElementList)
-{
-	TSharedRef<FSlateRenderDataHandle, ESPMode::ThreadSafe> RenderDataHandle = MakeShareable(new FSlateRenderDataHandle(Cacher, ResourceManager.Get()));
-
-	checkSlow(ElementList.GetChildDrawLayers().Num() == 0);
-
-	// Add all elements for this window to the element batcher
-	ElementBatcher->AddElements(ElementList);
-
-	// All elements for this window have been batched and rendering data updated
-	ElementBatcher->ResetBatches();
-
-	struct FCacheElementBatchesContext
-	{
-		FSlateRHIRenderingPolicy* RenderPolicy;
-		FSlateWindowElementList* SlateElementList;
-		TSharedRef<FSlateRenderDataHandle, ESPMode::ThreadSafe> RenderDataHandle;
-		uint32 ResourceVersion;
-	};
-	FCacheElementBatchesContext CacheElementBatchesContext =
-	{
-		RenderingPolicy.Get(),
-		&ElementList,
-		RenderDataHandle,
-		ResourceVersion
-	};
-
-	RenderDataHandle->BeginUsing();
-
-	ENQUEUE_RENDER_COMMAND(CacheElementBatches)(
-		[CacheElementBatchesContext](FRHICommandListImmediate& RHICmdList)
-	{
-		FSlateBatchData& BatchData = CacheElementBatchesContext.SlateElementList->GetBatchData();
-		FElementBatchMap& RootBatchMap = CacheElementBatchesContext.SlateElementList->GetRootDrawLayer().GetElementBatchMap();
-
-		BatchData.SetRenderDataHandle(CacheElementBatchesContext.RenderDataHandle);
-		BatchData.CreateRenderBatches(RootBatchMap);
-		CacheElementBatchesContext.RenderPolicy->UpdateVertexAndIndexBuffers(RHICmdList, BatchData, CacheElementBatchesContext.RenderDataHandle);
-
-		RootBatchMap.UpdateResourceVersion(CacheElementBatchesContext.ResourceVersion);
-
-		CacheElementBatchesContext.RenderDataHandle->EndUsing();
-	}
-	);
-
-	return RenderDataHandle;
-}
-
-void FSlateRHIRenderer::ReleaseCachingResourcesFor(const ILayoutCache* Cacher)
-{
-	struct FReleaseCachingResourcesForContext
-	{
-		FSlateRHIRenderingPolicy* RenderPolicy;
-		const ILayoutCache* Cacher;
-	};
-	FReleaseCachingResourcesForContext MarshalContext =
-	{
-		RenderingPolicy.Get(),
-		Cacher,
-	};
-	ENQUEUE_RENDER_COMMAND(ReleaseCachingResourcesFor)(
-		[MarshalContext](FRHICommandListImmediate& RHICmdList)
-	{
-		MarshalContext.RenderPolicy->ReleaseCachingResourcesFor(RHICmdList, MarshalContext.Cacher);
-	}
-	);
-}
-
 FSlateEndDrawingWindowsCommand::FSlateEndDrawingWindowsCommand(FSlateRHIRenderingPolicy& InPolicy, FSlateDrawBuffer* InDrawBuffer)
 	: Policy(InPolicy)
 	, DrawBuffer(InDrawBuffer)
@@ -1659,11 +1603,6 @@ FSlateEndDrawingWindowsCommand::FSlateEndDrawingWindowsCommand(FSlateRHIRenderin
 
 void FSlateEndDrawingWindowsCommand::Execute(FRHICommandListBase& CmdList)
 {
-	for (auto& ElementList : DrawBuffer->GetWindowElementLists())
-	{
-		ElementList->PostDraw_ParallelThread();
-	}
-
 	DrawBuffer->Unlock();
 	Policy.EndDrawingWindows();
 }
diff --git a/Engine/Source/Runtime/SlateRHIRenderer/Private/SlateRHIRenderer.h b/Engine/Source/Runtime/SlateRHIRenderer/Private/SlateRHIRenderer.h
--- a/Engine/Source/Runtime/SlateRHIRenderer/Private/SlateRHIRenderer.h
+++ b/Engine/Source/Runtime/SlateRHIRenderer/Private/SlateRHIRenderer.h
@@ -187,11 +187,10 @@ public:
 	virtual ISlateAtlasProvider* GetTextureAtlasProvider() override;
 	virtual FCriticalSection* GetResourceCriticalSection() override;
 	virtual void ReleaseAccessedResources(bool bImmediatelyFlush) override;
-	virtual TSharedRef<FSlateRenderDataHandle, ESPMode::ThreadSafe> CacheElementRenderData(const ILayoutCache* Cacher, FSlateWindowElementList& ElementList) override;
-	virtual void ReleaseCachingResourcesFor(const ILayoutCache* Cacher) override;
 	virtual int32 RegisterCurrentScene(FSceneInterface* Scene) override;
 	virtual int32 GetCurrentSceneIndex() const override;
 	virtual void ClearScenes() override;
+	virtual void DestroyCachedFastPathRenderingData(struct FSlateCachedFastPathRenderingData* InRenderingData) override;
 	virtual void BeginFrame() const override;
 	virtual void EndFrame() const override;
 	virtual void AddWidgetRendererUpdate(const struct FRenderThreadUpdateContext& Context, bool bDeferredRenderTargetUpdate) override;
@@ -265,7 +264,7 @@ private:
 	uint8 FreeBufferIndex;
 
 	/** Element batcher which renders draw elements */
-	TSharedPtr<FSlateElementBatcher> ElementBatcher;
+	TUniquePtr<FSlateElementBatcher> ElementBatcher;
 
 	/** Texture manager for accessing textures on the game thread */
 	TSharedPtr<FSlateRHIResourceManager> ResourceManager;
diff --git a/Engine/Source/Runtime/SlateRHIRenderer/Private/SlateRHIRendererModule.cpp b/Engine/Source/Runtime/SlateRHIRenderer/Private/SlateRHIRendererModule.cpp
--- a/Engine/Source/Runtime/SlateRHIRenderer/Private/SlateRHIRendererModule.cpp
+++ b/Engine/Source/Runtime/SlateRHIRenderer/Private/SlateRHIRendererModule.cpp
@@ -114,7 +114,6 @@ private:
 		if( !ResourceManager.IsValid() )
 		{
 			ResourceManager = MakeShareable( new FSlateRHIResourceManager );
-			FSlateDataPayload::ResourceManager = ResourceManager.Get();
 		}
 
 		if( !SlateFontServices.IsValid() )
diff --git a/Engine/Source/Runtime/SlateRHIRenderer/Private/SlateRHIRenderingPolicy.cpp b/Engine/Source/Runtime/SlateRHIRenderer/Private/SlateRHIRenderingPolicy.cpp
--- a/Engine/Source/Runtime/SlateRHIRenderer/Private/SlateRHIRenderingPolicy.cpp
+++ b/Engine/Source/Runtime/SlateRHIRenderer/Private/SlateRHIRenderingPolicy.cpp
@@ -30,7 +30,8 @@
 extern void UpdateNoiseTextureParameters(FViewUniformShaderParameters& ViewUniformShaderParameters);
 
 DECLARE_CYCLE_STAT(TEXT("Update Buffers RT"), STAT_SlateUpdateBufferRTTime, STATGROUP_Slate);
-DECLARE_CYCLE_STAT(TEXT("PreFill Buffers RT"), STAT_SlatePreFullBufferRTTime, STATGROUP_Slate);
+
+
 DECLARE_DWORD_COUNTER_STAT(TEXT("Num Layers"), STAT_SlateNumLayers, STATGROUP_Slate);
 DECLARE_DWORD_COUNTER_STAT(TEXT("Num Batches"), STAT_SlateNumBatches, STATGROUP_Slate);
 DECLARE_DWORD_COUNTER_STAT(TEXT("Num Vertices"), STAT_SlateVertexCount, STATGROUP_Slate);
@@ -106,16 +107,16 @@ void FSlateRHIRenderingPolicy::InitResources()
 
 	UE_LOG(LogSlate, Verbose, TEXT("Allocating space for %d vertices"), NumVertices);
 
-	VertexBuffers.Init(NumVertices);
-	IndexBuffers.Init(NumVertices);
+	MasterVertexBuffer.Init(NumVertices);
+	MasterIndexBuffer.Init(NumVertices);
 
 	BeginInitResource(&StencilVertexBuffer);
 }
 
 void FSlateRHIRenderingPolicy::ReleaseResources()
 {
-	VertexBuffers.Destroy();
-	IndexBuffers.Destroy();
+	MasterVertexBuffer.Destroy();
+	MasterIndexBuffer.Destroy();
 
 	BeginReleaseResource(&StencilVertexBuffer);
 }
@@ -132,14 +133,14 @@ void FSlateRHIRenderingPolicy::EndDrawingWindows()
 
 struct FSlateUpdateVertexAndIndexBuffers final : public FRHICommand<FSlateUpdateVertexAndIndexBuffers>
 {
-	FVertexBufferRHIRef VertexBufferRHI;
-	FIndexBufferRHIRef IndexBufferRHI;
+	TSlateElementVertexBuffer<FSlateVertex>& VertexBuffer;
+	FSlateElementIndexBuffer& IndexBuffer;
 	FSlateBatchData& BatchData;
 	bool bAbsoluteIndices;
 
 	FSlateUpdateVertexAndIndexBuffers(TSlateElementVertexBuffer<FSlateVertex>& InVertexBuffer, FSlateElementIndexBuffer& InIndexBuffer, FSlateBatchData& InBatchData, bool bInAbsoluteIndices)
-		: VertexBufferRHI(InVertexBuffer.VertexBufferRHI)
-		, IndexBufferRHI(InIndexBuffer.IndexBufferRHI)
+		: VertexBuffer(InVertexBuffer)
+		, IndexBuffer(InIndexBuffer)
 		, BatchData(InBatchData)
 		, bAbsoluteIndices(bInAbsoluteIndices)
 	{
@@ -150,104 +151,81 @@ struct FSlateUpdateVertexAndIndexBuffers final : public FRHICommand<FSlateUpdate
 	{
 		SCOPE_CYCLE_COUNTER(STAT_SlateUpdateBufferRTTime);
 
-		const int32 NumBatchedVertices = BatchData.GetNumBatchedVertices();
-		const int32 NumBatchedIndices = BatchData.GetNumBatchedIndices();
+		const bool bIsInRenderingThread = !IsRunningRHIInSeparateThread() || CmdList.Bypass();
+
+		const FSlateVertexArray& FinalVertexData = BatchData.GetFinalVertexData();
+		const FSlateIndexArray& FinalIndexData = BatchData.GetFinalIndexData();
 
-		int32 RequiredVertexBufferSize = NumBatchedVertices*sizeof(FSlateVertex);
-		uint8* VertexBufferData = (uint8*)GDynamicRHI->RHILockVertexBuffer( VertexBufferRHI, 0, RequiredVertexBufferSize, RLM_WriteOnly );
+		const int32 NumBatchedVertices = FinalVertexData.Num();
+		const int32 NumBatchedIndices = FinalIndexData.Num();
 
-		uint32 RequiredIndexBufferSize = NumBatchedIndices*sizeof(SlateIndex);		
-		uint8* IndexBufferData = (uint8*)GDynamicRHI->RHILockIndexBuffer( IndexBufferRHI, 0, RequiredIndexBufferSize, RLM_WriteOnly );
+		//int32 RequiredVertexBufferSize = NumBatchedVertices*sizeof(FSlateVertex);
+		uint8* VertexBufferData = (uint8*)VertexBuffer.LockBuffer(NumBatchedVertices, bIsInRenderingThread);
+
+		//uint32 RequiredIndexBufferSize = NumBatchedIndices*sizeof(SlateIndex);		
+		uint8* IndexBufferData = (uint8*)IndexBuffer.LockBuffer(NumBatchedIndices, bIsInRenderingThread);
 
 		//Early out if we have an invalid buffer (might have lost context and now have invalid buffers)
 		if ((nullptr != VertexBufferData) && (nullptr != IndexBufferData))
 		{
-			BatchData.FillVertexAndIndexBuffer(VertexBufferData, IndexBufferData, bAbsoluteIndices);
+			FMemory::Memcpy(VertexBufferData, FinalVertexData.GetData(), FinalVertexData.Num() * sizeof(FSlateVertex));
+			FMemory::Memcpy(IndexBufferData, FinalIndexData.GetData(), FinalIndexData.Num() * sizeof(SlateIndex));
 		}
 
 		if (nullptr != VertexBufferData)
 		{
-			GDynamicRHI->RHIUnlockVertexBuffer(VertexBufferRHI);
+			VertexBuffer.UnlockBuffer(bIsInRenderingThread);
 		}
 
 		if (nullptr != IndexBufferData)
 		{
-			GDynamicRHI->RHIUnlockIndexBuffer(IndexBufferRHI);
+			IndexBuffer.UnlockBuffer(bIsInRenderingThread);
 		}
 	}
 };
 
-void FSlateRHIRenderingPolicy::UpdateVertexAndIndexBuffers(FRHICommandListImmediate& RHICmdList, FSlateBatchData& InBatchData)
+void FSlateRHIRenderingPolicy::BuildRenderingBuffers(FRHICommandListImmediate& RHICmdList, FSlateBatchData& InBatchData)
 {
-	UpdateVertexAndIndexBuffers(RHICmdList, InBatchData, VertexBuffers, IndexBuffers);
-}
+	SCOPE_CYCLE_COUNTER(STAT_SlateUpdateBufferRTTime);
 
-void FSlateRHIRenderingPolicy::UpdateVertexAndIndexBuffers(FRHICommandListImmediate& RHICmdList, FSlateBatchData& InBatchData, const TSharedRef<FSlateRenderDataHandle, ESPMode::ThreadSafe>& RenderHandle)
-{
 	// Should only be called by the rendering thread
 	check(IsInRenderingThread());
 
-	FCachedRenderBuffers* Buffers = ResourceManager->FindOrCreateCachedBuffersForHandle(RenderHandle);
-
-	UpdateVertexAndIndexBuffers(RHICmdList, InBatchData, Buffers->VertexBuffer, Buffers->IndexBuffer);
-}
-
-void FSlateRHIRenderingPolicy::ReleaseCachingResourcesFor(FRHICommandListImmediate& RHICmdList, const ILayoutCache* Cacher)
-{
-	ResourceManager->ReleaseCachingResourcesFor(RHICmdList, Cacher);
-}
-
-void FSlateRHIRenderingPolicy::UpdateVertexAndIndexBuffers(FRHICommandListImmediate& RHICmdList, FSlateBatchData& InBatchData, TSlateElementVertexBuffer<FSlateVertex>& VertexBuffer, FSlateElementIndexBuffer& IndexBuffer)
-{
-	SCOPE_CYCLE_COUNTER( STAT_SlateUpdateBufferRTTime );
+	// Merge together batches for less draw calls
+	InBatchData.MergeRenderBatches();
 
-	// Should only be called by the rendering thread
-	check(IsInRenderingThread());
+	const FSlateVertexArray& FinalVertexData = InBatchData.GetFinalVertexData();
+	const FSlateIndexArray& FinalIndexData = InBatchData.GetFinalIndexData();
 
-	const int32 NumVertices = InBatchData.GetNumBatchedVertices();
-	const int32 NumIndices = InBatchData.GetNumBatchedIndices();
+	const int32 NumVertices = FinalVertexData.Num();
+	const int32 NumIndices = FinalIndexData.Num();
 
-	if( InBatchData.GetRenderBatches().Num() > 0  && NumVertices > 0 && NumIndices > 0)
+	if (InBatchData.GetRenderBatches().Num() > 0 && NumVertices > 0 && NumIndices > 0)
 	{
 		bool bShouldShrinkResources = false;
 		bool bAbsoluteIndices = CVarSlateAbsoluteIndices.GetValueOnRenderThread() != 0;
 
-		VertexBuffer.PreFillBuffer(NumVertices, bShouldShrinkResources);
-		IndexBuffer.PreFillBuffer(NumIndices, bShouldShrinkResources);
+		MasterVertexBuffer.PreFillBuffer(NumVertices, bShouldShrinkResources);
+		MasterIndexBuffer.PreFillBuffer(NumIndices, bShouldShrinkResources);
 
-		if(!IsRunningRHIInSeparateThread() || RHICmdList.Bypass())
+		if (!IsRunningRHIInSeparateThread() || RHICmdList.Bypass())
 		{
-			uint8* VertexBufferData = (uint8*)VertexBuffer.LockBuffer_RenderThread(NumVertices);
-			uint8* IndexBufferData =  (uint8*)IndexBuffer.LockBuffer_RenderThread(NumIndices);
-			
-			//Check if the LockBuffer failed for these threads (might have lost context)
-			if ((nullptr != VertexBufferData) && (nullptr != IndexBufferData))
-			{
-				InBatchData.FillVertexAndIndexBuffer(VertexBufferData, IndexBufferData, bAbsoluteIndices);
-			}
+			FSlateUpdateVertexAndIndexBuffers UpdateBufferCommand(MasterVertexBuffer, MasterIndexBuffer, InBatchData, bAbsoluteIndices);
 
-			if (nullptr != VertexBufferData)
-			{
-				VertexBuffer.UnlockBuffer_RenderThread();
-			}
-			
-			if (nullptr != IndexBufferData)
-			{
-				IndexBuffer.UnlockBuffer_RenderThread();
-			}
+			UpdateBufferCommand.Execute(RHICmdList);
 		}
 		else
 		{
-			ALLOC_COMMAND_CL(RHICmdList, FSlateUpdateVertexAndIndexBuffers)(VertexBuffer, IndexBuffer, InBatchData, bAbsoluteIndices);
+			ALLOC_COMMAND_CL(RHICmdList, FSlateUpdateVertexAndIndexBuffers)(MasterVertexBuffer, MasterIndexBuffer, InBatchData, bAbsoluteIndices);
 		}
 	}
 
-	checkSlow(VertexBuffer.GetBufferUsageSize() <= VertexBuffer.GetBufferSize());
-	checkSlow(IndexBuffer.GetBufferUsageSize() <= IndexBuffer.GetBufferSize());
+	checkSlow(MasterVertexBuffer.GetBufferUsageSize() <= MasterVertexBuffer.GetBufferSize());
+	checkSlow(MasterIndexBuffer.GetBufferUsageSize() <= MasterIndexBuffer.GetBufferSize());
 
-	SET_DWORD_STAT( STAT_SlateNumLayers, InBatchData.GetNumLayers() );
-	SET_DWORD_STAT( STAT_SlateNumBatches, InBatchData.GetRenderBatches().Num() );
-	SET_DWORD_STAT( STAT_SlateVertexCount, InBatchData.GetNumBatchedVertices() );
+	SET_DWORD_STAT(STAT_SlateNumLayers, InBatchData.GetNumLayers());
+	SET_DWORD_STAT(STAT_SlateNumBatches, InBatchData.GetNumFinalBatches());
+	SET_DWORD_STAT(STAT_SlateVertexCount, InBatchData.GetFinalVertexData().Num());
 }
 
 static FSceneView* CreateSceneView( FSceneViewFamilyContext* ViewFamilyContext, FSlateBackBuffer& BackBuffer, const FMatrix& ViewProjectionMatrix )
@@ -322,7 +300,7 @@ static bool UpdateScissorRect(
 	const FSlateRenderBatch& RenderBatch, 
 	FTexture2DRHIRef& ColorTarget,
 	FTexture2DRHIRef& DepthStencilTarget,
-	TOptional<FSlateClippingState>& LastClippingState,
+	const FSlateClippingState*& LastClippingState,
 	const FVector2D& ViewTranslation2D, 
 	bool bSwitchVerticalAxis,
 	FGraphicsPipelineStateInitializer& InGraphicsPSOInit,
@@ -337,9 +315,9 @@ static bool UpdateScissorRect(
 
 	if (RenderBatch.ClippingState != LastClippingState || bForceStateChange)
 	{
-		if (RenderBatch.ClippingState.IsSet())
+		if (RenderBatch.ClippingState)
 		{
-			const FSlateClippingState& ClipState = RenderBatch.ClippingState.GetValue();
+			const FSlateClippingState& ClipState = *RenderBatch.ClippingState;
 			if (ClipState.GetClippingMethod() == EClippingMethod::Scissor)
 			{
 #if STATS
@@ -604,6 +582,7 @@ void FSlateRHIRenderingPolicy::DrawElements(
 	FSlateBackBuffer& BackBuffer,
 	FTexture2DRHIRef& ColorTarget,
 	FTexture2DRHIRef& DepthStencilTarget,
+	int32 FirstBatchIndex,
 	const TArray<FSlateRenderBatch>& RenderBatches,
 	const FSlateRenderingParams& Params)
 {
@@ -687,8 +666,8 @@ void FSlateRHIRenderingPolicy::DrawElements(
 
 	FSamplerStateRHIRef BilinearClamp = TStaticSamplerState<SF_Bilinear, AM_Clamp, AM_Clamp, AM_Clamp>::GetRHI();
 
-	TSlateElementVertexBuffer<FSlateVertex>* VertexBuffer = &VertexBuffers;
-	FSlateElementIndexBuffer* IndexBuffer = &IndexBuffers;
+	TSlateElementVertexBuffer<FSlateVertex>* VertexBufferPtr = &MasterVertexBuffer;
+	FSlateElementIndexBuffer* IndexBufferPtr = &MasterIndexBuffer;
 
 	FGraphicsPipelineStateInitializer GraphicsPSOInit;
 	RHICmdList.ApplyCachedRenderTargets(GraphicsPSOInit);
@@ -706,7 +685,8 @@ void FSlateRHIRenderingPolicy::DrawElements(
 	const bool bSwitchVerticalAxis = Params.bAllowSwitchVerticalAxis && RHINeedsToSwitchVerticalAxis(GShaderPlatformForFeatureLevel[GMaxRHIFeatureLevel]);
 
 	// This variable tracks the last clipping state, so that if multiple batches have the same clipping state, we don't have to do any work.
-	TOptional<FSlateClippingState> LastClippingState;
+	const FSlateClippingState* LastClippingState;
+
 	// This is the stenciling ref variable we set any time we draw, so that any stencil comparisons use the right mask id.
 	uint32 StencilRef = 0;
 	// This is an accumulating maskID that we use to track the between batch usage of the stencil buffer, when at 0, or over 255
@@ -720,12 +700,22 @@ void FSlateRHIRenderingPolicy::DrawElements(
 	FVector2D ViewTranslation2D = Params.ViewOffset;
 
 	// Draw each element
+#if WITH_SLATE_DEBUGGING
+	int32 NextRenderBatchIndex = BatchToDraw == -1 ? FirstBatchIndex : BatchToDraw;
+#else
+	int32 NextRenderBatchIndex = FirstBatchIndex;
+#endif
+
+
 	/*
 		#todo-renderpasses This loop ends up with ugly logic.
 		CustomDrawers will draw in their own renderpass. So we must remember to reopen the renderpass with the passed in Color/DepthStencil targets.
 	*/
-	for (int32 BatchIndex = 0; BatchIndex < RenderBatches.Num(); ++BatchIndex)
+	while (NextRenderBatchIndex != INDEX_NONE)
 	{
+		VertexBufferPtr = &MasterVertexBuffer;
+		IndexBufferPtr = &MasterIndexBuffer;
+
 		if (!RHICmdList.IsInsideRenderPass())
 		{
 			// Restart the renderpass since the CustomDrawer or post-process may have changed it in last iteration
@@ -745,55 +735,27 @@ void FSlateRHIRenderingPolicy::DrawElements(
 
 			// Something may have messed with the viewport size so set it back to the full target.
 			RHICmdList.SetViewport(0, 0, 0, BackBuffer.GetSizeXY().X, BackBuffer.GetSizeXY().Y, 0.0f);
-			RHICmdList.SetStreamSource(0, VertexBuffer->VertexBufferRHI, 0);
+			RHICmdList.SetStreamSource(0, VertexBufferPtr->VertexBufferRHI, 0);
 		}
 				
 #if WITH_SLATE_VISUALIZERS
 		FLinearColor BatchColor = FLinearColor(BatchColors.GetUnitVector());
 #endif
+		const FSlateRenderBatch& RenderBatch = RenderBatches[NextRenderBatchIndex];
 
-		const FSlateRenderBatch& RenderBatch = RenderBatches[BatchIndex];
-		const FSlateRenderDataHandle* RenderHandle = RenderBatch.CachedRenderHandle.Get();
+		NextRenderBatchIndex = RenderBatch.NextBatchIndex;
 
-#if !(UE_BUILD_SHIPPING || UE_BUILD_TEST)
-		if (BatchIndex != RenderBatches.Num() - 1)
+#if WITH_SLATE_DEBUGGING
+		if (BatchToDraw != -1)
 		{
-			if (BatchToDraw != -1 && BatchToDraw != BatchIndex)
-			{
-				continue;
-			}
+			break;
 		}
 #endif
 
-		if (RenderHandle != LastHandle)
-		{
-			SLATE_DRAW_EVENT(RHICmdList, ChangingRenderBuffers);
-
-			LastHandle = RenderHandle;
-
-			if (RenderHandle)
-			{
-				FCachedRenderBuffers* Buffers = ResourceManager->FindCachedBuffersForHandle(RenderHandle);
-				if (Buffers != nullptr)
-				{
-					VertexBuffer = &Buffers->VertexBuffer;
-					IndexBuffer = &Buffers->IndexBuffer;
-				}
-			}
-			else
-			{
-				VertexBuffer = &VertexBuffers;
-				IndexBuffer = &IndexBuffers;
-			}
-
-			checkSlow(VertexBuffer);
-			checkSlow(IndexBuffer);
-		}
-
-		const FSlateShaderResource* ShaderResource = RenderBatch.Texture;
+		const FSlateShaderResource* ShaderResource = RenderBatch.ShaderResource;
 		const ESlateBatchDrawFlag DrawFlags = RenderBatch.DrawFlags;
 		const ESlateDrawEffect DrawEffects = RenderBatch.DrawEffects;
-		const ESlateShader::Type ShaderType = RenderBatch.ShaderType;
+		const ESlateShader ShaderType = RenderBatch.ShaderType;
 		const FShaderParams& ShaderParams = RenderBatch.ShaderParams;
 
 		if (EnumHasAllFlags(DrawFlags, ESlateBatchDrawFlag::Wireframe))
@@ -805,7 +767,7 @@ void FSlateRHIRenderingPolicy::DrawElements(
 			GraphicsPSOInit.RasterizerState = TStaticRasterizerState<FM_Solid, CM_None, false>::GetRHI();
 		}
 
-		if (!RenderBatch.CustomDrawer.IsValid())
+		if (!RenderBatch.CustomDrawer)
 		{
 			FMatrix DynamicOffset = FTranslationMatrix::Make(FVector(RenderBatch.DynamicOffset.X, RenderBatch.DynamicOffset.Y, 0));
 			const FMatrix ViewProjection = DynamicOffset * Params.ViewProjectionMatrix;
@@ -1034,16 +996,16 @@ void FSlateRHIRenderingPolicy::DrawElements(
 					SCOPE_CYCLE_COUNTER(STAT_SlateRTTextureDrawCall);
 
 					// for RHIs that can't handle VertexOffset, we need to offset the stream source each time
-					if (!GRHISupportsBaseVertexIndex && !bAbsoluteIndices)
+					if (1 || (!GRHISupportsBaseVertexIndex && !bAbsoluteIndices))
 					{
-						RHICmdList.SetStreamSource(0, VertexBuffer->VertexBufferRHI, RenderBatch.VertexOffset * sizeof(FSlateVertex));
-						RHICmdList.DrawIndexedPrimitive(IndexBuffer->IndexBufferRHI, 0, 0, RenderBatch.NumVertices, RenderBatch.IndexOffset, PrimitiveCount, RenderBatch.InstanceCount);
+						RHICmdList.SetStreamSource(0, VertexBufferPtr->VertexBufferRHI, RenderBatch.VertexOffset * sizeof(FSlateVertex));
+						RHICmdList.DrawIndexedPrimitive(IndexBufferPtr->IndexBufferRHI, 0, 0, RenderBatch.NumVertices, RenderBatch.IndexOffset, PrimitiveCount, RenderBatch.InstanceCount);
 					}
 					else
 					{
 						uint32 VertexOffset = bAbsoluteIndices ? 0 : RenderBatch.VertexOffset;
-						RHICmdList.SetStreamSource(0, VertexBuffer->VertexBufferRHI, 0);
-						RHICmdList.DrawIndexedPrimitive(IndexBuffer->IndexBufferRHI, VertexOffset, 0, RenderBatch.NumVertices, RenderBatch.IndexOffset, PrimitiveCount, RenderBatch.InstanceCount);
+						RHICmdList.SetStreamSource(0, VertexBufferPtr->VertexBufferRHI, 0);
+						RHICmdList.DrawIndexedPrimitive(IndexBufferPtr->IndexBufferRHI, VertexOffset, 0, RenderBatch.NumVertices, RenderBatch.IndexOffset, PrimitiveCount, RenderBatch.InstanceCount);
 					}
 				}
 			}
@@ -1070,10 +1032,8 @@ void FSlateRHIRenderingPolicy::DrawElements(
 					// MaterialParameterCollections will not be correct for this scene, should they be
 					// used.
 					ActiveSceneIndex = NumScenes - 1;
-#if UE_BUILD_DEBUG
-#if WITH_EDITOR
+#if UE_BUILD_DEBUG && WITH_EDITOR
 					UE_LOG(LogSlate, Error, TEXT("Invalid scene index in batch: %d of %d known scenes!"), RenderBatch.SceneIndex, ResourceManager->GetSceneCount());
-#endif
 #endif
 				}
 
@@ -1175,16 +1135,16 @@ void FSlateRHIRenderingPolicy::DrawElements(
 									InstanceBuffer->BindStreamSource(RHICmdList, 1, RenderBatch.InstanceOffset);
 
 									// for RHIs that can't handle VertexOffset, we need to offset the stream source each time
-									if (!GRHISupportsBaseVertexIndex && !bAbsoluteIndices)
+									if (1 || (!GRHISupportsBaseVertexIndex && !bAbsoluteIndices))
 									{
-										RHICmdList.SetStreamSource(0, VertexBuffer->VertexBufferRHI, RenderBatch.VertexOffset * sizeof(FSlateVertex));
-										RHICmdList.DrawIndexedPrimitive(IndexBuffer->IndexBufferRHI, 0, 0, RenderBatch.NumVertices, RenderBatch.IndexOffset, PrimitiveCount, InstanceCount);
+										RHICmdList.SetStreamSource(0, VertexBufferPtr->VertexBufferRHI, RenderBatch.VertexOffset * sizeof(FSlateVertex));
+										RHICmdList.DrawIndexedPrimitive(IndexBufferPtr->IndexBufferRHI, 0, 0, RenderBatch.NumVertices, RenderBatch.IndexOffset, PrimitiveCount, InstanceCount);
 									}
 									else
 									{
 										uint32 VertexOffset = bAbsoluteIndices ? 0 : RenderBatch.VertexOffset;
-										RHICmdList.SetStreamSource(0, VertexBuffer->VertexBufferRHI, 0);
-										RHICmdList.DrawIndexedPrimitive(IndexBuffer->IndexBufferRHI, VertexOffset, 0, RenderBatch.NumVertices, RenderBatch.IndexOffset, PrimitiveCount, InstanceCount);
+										RHICmdList.SetStreamSource(0, VertexBufferPtr->VertexBufferRHI, 0);
+										RHICmdList.DrawIndexedPrimitive(IndexBufferPtr->IndexBufferRHI, VertexOffset, 0, RenderBatch.NumVertices, RenderBatch.IndexOffset, PrimitiveCount, InstanceCount);
 									}
 								}
 							}
@@ -1193,16 +1153,16 @@ void FSlateRHIRenderingPolicy::DrawElements(
 								RHICmdList.SetStreamSource(1, nullptr, 0);
 
 								// for RHIs that can't handle VertexOffset, we need to offset the stream source each time
-								if (!GRHISupportsBaseVertexIndex && !bAbsoluteIndices)
+								if (1 || (!GRHISupportsBaseVertexIndex && !bAbsoluteIndices))
 								{
-									RHICmdList.SetStreamSource(0, VertexBuffer->VertexBufferRHI, RenderBatch.VertexOffset * sizeof(FSlateVertex));
-									RHICmdList.DrawIndexedPrimitive(IndexBuffer->IndexBufferRHI, 0, 0, RenderBatch.NumVertices, RenderBatch.IndexOffset, PrimitiveCount, 1);
+									RHICmdList.SetStreamSource(0, VertexBufferPtr->VertexBufferRHI, RenderBatch.VertexOffset * sizeof(FSlateVertex));
+									RHICmdList.DrawIndexedPrimitive(IndexBufferPtr->IndexBufferRHI, 0, 0, RenderBatch.NumVertices, RenderBatch.IndexOffset, PrimitiveCount, 1);
 								}
 								else
 								{
 									uint32 VertexOffset = bAbsoluteIndices ? 0 : RenderBatch.VertexOffset;
-									RHICmdList.SetStreamSource(0, VertexBuffer->VertexBufferRHI, 0);
-									RHICmdList.DrawIndexedPrimitive(IndexBuffer->IndexBufferRHI, VertexOffset, 0, RenderBatch.NumVertices, RenderBatch.IndexOffset, PrimitiveCount, 1);
+									RHICmdList.SetStreamSource(0, VertexBufferPtr->VertexBufferRHI, 0);
+									RHICmdList.DrawIndexedPrimitive(IndexBufferPtr->IndexBufferRHI, VertexOffset, 0, RenderBatch.NumVertices, RenderBatch.IndexOffset, PrimitiveCount, 1);
 								}
 							}
 						}
@@ -1261,8 +1221,8 @@ void FSlateRHIRenderingPolicy::DrawElements(
 		}
 		else
 		{
-			TSharedPtr<ICustomSlateElement, ESPMode::ThreadSafe> CustomDrawer = RenderBatch.CustomDrawer.Pin();
-			if (CustomDrawer.IsValid())
+			ICustomSlateElement* CustomDrawer = RenderBatch.CustomDrawer;
+			if (CustomDrawer)
 			{
 				// CustomDrawers will change the rendertarget. So we must close any outstanding renderpasses.
 				// Render pass for slate elements will be restarted on a next loop iteration if any
@@ -1272,7 +1232,7 @@ void FSlateRHIRenderingPolicy::DrawElements(
 
 				// Disable scissor rect. A previous draw element may have had one
 				RHICmdList.SetScissorRect(false, 0, 0, 0, 0);
-				LastClippingState.Reset();
+				LastClippingState = nullptr;
 
 				// This element is custom and has no Slate geometry.  Tell it to render itself now
 				CustomDrawer->DrawRenderThread(RHICmdList, &BackBuffer.GetRenderTargetTexture());
@@ -1360,7 +1320,7 @@ ETextureSamplerFilter FSlateRHIRenderingPolicy::GetSamplerFilter(const UTexture*
 	return Filter;
 }
 
-FSlateElementPS* FSlateRHIRenderingPolicy::GetTexturePixelShader( TShaderMap<FGlobalShaderType>* ShaderMap, ESlateShader::Type ShaderType, ESlateDrawEffect DrawEffects )
+FSlateElementPS* FSlateRHIRenderingPolicy::GetTexturePixelShader( TShaderMap<FGlobalShaderType>* ShaderMap, ESlateShader ShaderType, ESlateDrawEffect DrawEffects )
 {
 	FSlateElementPS* PixelShader = nullptr;
 	QUICK_SCOPE_CYCLE_COUNTER(STAT_Slate_GetTexturePixelShader);
@@ -1449,7 +1409,7 @@ FSlateElementPS* FSlateRHIRenderingPolicy::GetTexturePixelShader( TShaderMap<FGl
 	return PixelShader;
 }
 
-FSlateMaterialShaderPS* FSlateRHIRenderingPolicy::GetMaterialPixelShader( const FMaterial* Material, ESlateShader::Type ShaderType, ESlateDrawEffect DrawEffects )
+FSlateMaterialShaderPS* FSlateRHIRenderingPolicy::GetMaterialPixelShader( const FMaterial* Material, ESlateShader ShaderType, ESlateDrawEffect DrawEffects )
 {
 	const FMaterialShaderMap* MaterialShaderMap = Material->GetRenderingThreadShaderMap();
 
@@ -1519,7 +1479,7 @@ FSlateMaterialShaderVS* FSlateRHIRenderingPolicy::GetMaterialVertexShader( const
 	return FoundShader ? (FSlateMaterialShaderVS*)FoundShader->GetShader() : nullptr;
 }
 
-EPrimitiveType FSlateRHIRenderingPolicy::GetRHIPrimitiveType(ESlateDrawPrimitive::Type SlateType)
+EPrimitiveType FSlateRHIRenderingPolicy::GetRHIPrimitiveType(ESlateDrawPrimitive SlateType)
 {
 	switch(SlateType)
 	{
diff --git a/Engine/Source/Runtime/SlateRHIRenderer/Private/SlateRHIRenderingPolicy.h b/Engine/Source/Runtime/SlateRHIRenderer/Private/SlateRHIRenderingPolicy.h
--- a/Engine/Source/Runtime/SlateRHIRenderer/Private/SlateRHIRenderingPolicy.h
+++ b/Engine/Source/Runtime/SlateRHIRenderer/Private/SlateRHIRenderingPolicy.h
@@ -18,7 +18,6 @@
 class FSlateFontServices;
 class FSlateRHIResourceManager;
 class FSlatePostProcessor;
-class ILayoutCache;
 class UDeviceProfile;
 
 struct FSlateRenderingParams
@@ -48,12 +47,9 @@ class FSlateRHIRenderingPolicy : public FSlateRenderingPolicy
 public:
 	FSlateRHIRenderingPolicy(TSharedRef<FSlateFontServices> InSlateFontServices, TSharedRef<FSlateRHIResourceManager> InResourceManager, TOptional<int32> InitialBufferSize = TOptional<int32>());
 
-	void UpdateVertexAndIndexBuffers(FRHICommandListImmediate& RHICmdList, FSlateBatchData& BatchData);
-	void UpdateVertexAndIndexBuffers(FRHICommandListImmediate& RHICmdList, FSlateBatchData& BatchData, const TSharedRef<FSlateRenderDataHandle, ESPMode::ThreadSafe>& RenderHandle);
+	void BuildRenderingBuffers(FRHICommandListImmediate& RHICmdList, FSlateBatchData& InBatchData);
 
-	void ReleaseCachingResourcesFor(FRHICommandListImmediate& RHICmdList, const ILayoutCache* Cacher);
-
-	void DrawElements(FRHICommandListImmediate& RHICmdList, class FSlateBackBuffer& BackBuffer, FTexture2DRHIRef& ColorTarget, FTexture2DRHIRef& DepthStencilTarget, const TArray<FSlateRenderBatch>& RenderBatches, const FSlateRenderingParams& Params);
+	void DrawElements(FRHICommandListImmediate& RHICmdList, class FSlateBackBuffer& BackBuffer, FTexture2DRHIRef& ColorTarget, FTexture2DRHIRef& DepthStencilTarget, int32 FirstBatchIndex, const TArray<FSlateRenderBatch>& RenderBatches, const FSlateRenderingParams& Params);
 
 	virtual TSharedRef<FSlateShaderResourceManager> GetResourceManager() const override { return ResourceManager; }
 	virtual bool IsVertexColorInLinearSpace() const override { return false; }
@@ -72,9 +68,6 @@ public:
 
 	virtual void FlushGeneratedResources();
 
-protected:
-	void UpdateVertexAndIndexBuffers(FRHICommandListImmediate& RHICmdList, FSlateBatchData& BatchData, TSlateElementVertexBuffer<FSlateVertex>& VertexBuffer, FSlateElementIndexBuffer& IndexBuffer);
-
 private:
 	ETextureSamplerFilter GetSamplerFilter(const UTexture* Texture) const;
 
@@ -85,17 +78,17 @@ private:
 	 * @param DrawEffects	Draw effects being used
 	 * @return The pixel shader for use with the shader type and draw effects
 	 */
-	class FSlateElementPS* GetTexturePixelShader( TShaderMap<FGlobalShaderType>* ShaderMap, ESlateShader::Type ShaderType, ESlateDrawEffect DrawEffects );
-	class FSlateMaterialShaderPS* GetMaterialPixelShader( const class FMaterial* Material, ESlateShader::Type ShaderType, ESlateDrawEffect DrawEffects );
+	class FSlateElementPS* GetTexturePixelShader( TShaderMap<FGlobalShaderType>* ShaderMap, ESlateShader ShaderType, ESlateDrawEffect DrawEffects );
+	class FSlateMaterialShaderPS* GetMaterialPixelShader( const class FMaterial* Material, ESlateShader ShaderType, ESlateDrawEffect DrawEffects );
 	class FSlateMaterialShaderVS* GetMaterialVertexShader( const class FMaterial* Material, bool bUseInstancing );
 
 	/** @return The RHI primitive type from the Slate primitive type */
-	EPrimitiveType GetRHIPrimitiveType(ESlateDrawPrimitive::Type SlateType);
+	EPrimitiveType GetRHIPrimitiveType(ESlateDrawPrimitive SlateType);
 
 private:
 	/** Buffers used for rendering */
-	TSlateElementVertexBuffer<FSlateVertex> VertexBuffers;
-	FSlateElementIndexBuffer IndexBuffers;
+	TSlateElementVertexBuffer<FSlateVertex> MasterVertexBuffer;
+	FSlateElementIndexBuffer MasterIndexBuffer;
 
 	FSlateStencilClipVertexBuffer StencilVertexBuffer;
 
diff --git a/Engine/Source/Runtime/SlateRHIRenderer/Private/SlateRHIResourceManager.cpp b/Engine/Source/Runtime/SlateRHIRenderer/Private/SlateRHIResourceManager.cpp
--- a/Engine/Source/Runtime/SlateRHIRenderer/Private/SlateRHIResourceManager.cpp
+++ b/Engine/Source/Runtime/SlateRHIRenderer/Private/SlateRHIResourceManager.cpp
@@ -290,38 +290,6 @@ bool FSlateRHIResourceManager::IsAtlasPageResourceAlphaOnly() const
 void FSlateRHIResourceManager::Tick(float DeltaSeconds)
 {
 	TryToCleanupExpiredResources(false);
-
-	// Don't need to do this if there's no RHI thread.
-	if (IsRunningRHIInSeparateThread())
-	{
-		struct FDeleteCachedRenderDataContext
-		{
-			FSlateRHIResourceManager* ResourceManager;
-		};
-		FDeleteCachedRenderDataContext Context =
-		{
-			this,
-		};
-		ENQUEUE_RENDER_COMMAND(DeleteCachedRenderData)(
-			[Context](FRHICommandListImmediate& RHICmdList)
-			{
-				// Go through the pending delete buffers and see if any of their fences has cleared
-				// the RHI thread, if so, they should be safe to delete now.
-				for ( int32 BufferIndex = Context.ResourceManager->PooledBuffersPendingRelease.Num() - 1; BufferIndex >= 0; BufferIndex-- )
-				{
-					FCachedRenderBuffers* PooledBuffer = Context.ResourceManager->PooledBuffersPendingRelease[BufferIndex];
-
-					if ( PooledBuffer->ReleaseResourcesFence->IsComplete() )
-					{
-						PooledBuffer->VertexBuffer.Destroy();
-						PooledBuffer->IndexBuffer.Destroy();
-						delete PooledBuffer;
-
-						Context.ResourceManager->PooledBuffersPendingRelease.RemoveAt(BufferIndex);
-					}
-				}
-			});
-	}
 }
 
 void FSlateRHIResourceManager::CreateTextures( const TArray< const FSlateBrush* >& Resources )
@@ -928,116 +896,6 @@ void FSlateRHIResourceManager::UpdateTextureAtlases()
 	}
 }
 
-FCachedRenderBuffers* FSlateRHIResourceManager::FindOrCreateCachedBuffersForHandle(const TSharedRef<FSlateRenderDataHandle, ESPMode::ThreadSafe>& RenderHandle)
-{
-	// Should only be called by the rendering thread
-	check(IsInRenderingThread());
-
-	FCachedRenderBuffers* Buffers = CachedBuffers.FindRef(&RenderHandle.Get());
-	if ( Buffers == nullptr )
-	{
-		// Rather than having a global pool, we associate the pools with a particular layout cacher.
-		// If we don't do this, all buffers eventually become as larger as the largest buffer, and it
-		// would be much better to keep the pools coherent with the sizes typically associated with
-		// a particular caching panel.
-		const ILayoutCache* LayoutCacher = RenderHandle->GetCacher();
-		TArray< FCachedRenderBuffers* >& Pool = CachedBufferPool.FindOrAdd(LayoutCacher);
-
-		// If the cached buffer pool is empty, time to create a new one!
-		if ( Pool.Num() == 0 )
-		{
-			Buffers = new FCachedRenderBuffers();
-			Buffers->VertexBuffer.Init(100);
-			Buffers->IndexBuffer.Init(100);
-		}
-		else
-		{
-			// If we found one in the pool, lets use it!
-			Buffers = Pool[0];
-			Pool.RemoveAtSwap(0, 1, false);
-		}
-
-		CachedBuffers.Add(&RenderHandle.Get(), Buffers);
-	}
-
-	return Buffers;
-}
-
-void FSlateRHIResourceManager::BeginReleasingRenderData(const FSlateRenderDataHandle* RenderHandle)
-{
-	struct FReleaseCachedRenderDataContext
-	{
-		FSlateRHIResourceManager* ResourceManager;
-		const FSlateRenderDataHandle* RenderDataHandle;
-		const ILayoutCache* LayoutCacher;
-	};
-	FReleaseCachedRenderDataContext Context =
-	{
-		this,
-		RenderHandle,
-		RenderHandle->GetCacher()
-	};
-	ENQUEUE_RENDER_COMMAND(ReleaseCachedRenderData)(
-		[Context](FRHICommandListImmediate& RHICmdList)
-		{
-			Context.ResourceManager->ReleaseCachedRenderData(RHICmdList, Context.RenderDataHandle, Context.LayoutCacher);
-		});
-}
-
-void FSlateRHIResourceManager::ReleaseCachedRenderData(FRHICommandListImmediate& RHICmdList, const FSlateRenderDataHandle* RenderHandle, const ILayoutCache* LayoutCacher)
-{
-	check(IsInRenderingThread());
-	check(RenderHandle);
-
-	FCachedRenderBuffers* PooledBuffer = CachedBuffers.FindRef(RenderHandle);
-	if ( ensure(PooledBuffer != nullptr) )
-	{
-		TArray< FCachedRenderBuffers* >* Pool = CachedBufferPool.Find(LayoutCacher);
-		if ( Pool )
-		{
-			Pool->Add(PooledBuffer);
-		}
-		else
-		{
-			ReleaseCachedBuffer(RHICmdList, PooledBuffer);
-		}
-
-		CachedBuffers.Remove(RenderHandle);
-	}
-}
-
-void FSlateRHIResourceManager::ReleaseCachingResourcesFor(FRHICommandListImmediate& RHICmdList, const ILayoutCache* Cacher)
-{
-	check(IsInRenderingThread());
-
-	TArray< FCachedRenderBuffers* >* Pool = CachedBufferPool.Find(Cacher);
-	if ( Pool )
-	{
-		for ( FCachedRenderBuffers* PooledBuffer : *Pool )
-		{
-			ReleaseCachedBuffer(RHICmdList, PooledBuffer);
-		}
-
-		CachedBufferPool.Remove(Cacher);
-	}
-}
-
-void FSlateRHIResourceManager::ReleaseCachedBuffer(FRHICommandListImmediate& RHICmdList, FCachedRenderBuffers* PooledBuffer)
-{
-	check(IsInRenderingThread());
-
-	if (IsRunningRHIInSeparateThread())
-	{
-		PooledBuffersPendingRelease.Add(PooledBuffer);
-		PooledBuffer->ReleaseResourcesFence = RHICmdList.RHIThreadFence();
-	}
-	else
-	{
-		PooledBuffer->VertexBuffer.Destroy();
-		PooledBuffer->IndexBuffer.Destroy();
-		delete PooledBuffer;
-	}
-}
 
 void FSlateRHIResourceManager::ReleaseResources()
 {
@@ -1055,27 +913,6 @@ void FSlateRHIResourceManager::ReleaseResources()
 
 	DynamicResourceMap.ReleaseResources();
 
-	for ( TCachedBufferMap::TIterator BufferIt(CachedBuffers); BufferIt; ++BufferIt )
-	{
-		FSlateRenderDataHandle* Handle = BufferIt.Key();
-		FCachedRenderBuffers* Buffer = BufferIt.Value();
-
-		Handle->Disconnect();
-
-		Buffer->VertexBuffer.Destroy();
-		Buffer->IndexBuffer.Destroy();
-	}
-
-	for ( TCachedBufferPoolMap::TIterator BufferIt(CachedBufferPool); BufferIt; ++BufferIt )
-	{
-		TArray< FCachedRenderBuffers* >& Pool = BufferIt.Value();
-		for ( FCachedRenderBuffers* PooledBuffer : Pool )
-		{
-			PooledBuffer->VertexBuffer.Destroy();
-			PooledBuffer->IndexBuffer.Destroy();
-		}
-	}
-
 	// Note the base class has texture proxies only which do not need to be released
 }
 
@@ -1115,7 +952,6 @@ void FSlateRHIResourceManager::DeleteResources()
 
 	DeleteUObjectBrushResources();
 
-	DeleteCachedBuffers();
 }
 
 void FSlateRHIResourceManager::DeleteUObjectBrushResources()
@@ -1125,28 +961,6 @@ void FSlateRHIResourceManager::DeleteUObjectBrushResources()
 	UTextureFreeList.Empty();
 }
 
-void FSlateRHIResourceManager::DeleteCachedBuffers()
-{
-	for (TCachedBufferMap::TIterator BufferIt(CachedBuffers); BufferIt; ++BufferIt)
-	{
-		FCachedRenderBuffers* Buffer = BufferIt.Value();
-		delete Buffer;
-	}
-
-	CachedBuffers.Empty();
-
-	for (TCachedBufferPoolMap::TIterator BufferIt(CachedBufferPool); BufferIt; ++BufferIt)
-	{
-		TArray< FCachedRenderBuffers* >& Pool = BufferIt.Value();
-		for (FCachedRenderBuffers* PooledBuffer : Pool)
-		{
-			delete PooledBuffer;
-		}
-	}
-
-	CachedBufferPool.Empty();
-}
-
 void FSlateRHIResourceManager::ReloadTextures()
 {
 	checkSlow( IsThreadSafeForSlateRendering() );
