Fix copyright notices

diff --git a/Engine/Source/Runtime/SlateCore/Private/FastUpdate/SlateInvalidationRoot.cpp b/Engine/Source/Runtime/SlateCore/Private/FastUpdate/SlateInvalidationRoot.cpp
--- a/Engine/Source/Runtime/SlateCore/Private/FastUpdate/SlateInvalidationRoot.cpp
+++ b/Engine/Source/Runtime/SlateCore/Private/FastUpdate/SlateInvalidationRoot.cpp
@@ -1,4 +1,4 @@
-// Copyright 1998-2018 Epic Games, Inc. All Rights Reserved.
+// Copyright 1998-2019 Epic Games, Inc. All Rights Reserved.
 
 #include "FastUpdate/SlateInvalidationRoot.h"
 #include "Async/TaskGraphInterfaces.h"
diff --git a/Engine/Source/Runtime/SlateCore/Private/FastUpdate/WidgetProxy.cpp b/Engine/Source/Runtime/SlateCore/Private/FastUpdate/WidgetProxy.cpp
--- a/Engine/Source/Runtime/SlateCore/Private/FastUpdate/WidgetProxy.cpp
+++ b/Engine/Source/Runtime/SlateCore/Private/FastUpdate/WidgetProxy.cpp
@@ -1,4 +1,4 @@
-// Copyright 1998-2018 Epic Games, Inc. All Rights Reserved.
+// Copyright 1998-2019 Epic Games, Inc. All Rights Reserved.
 
 #include "FastUpdate/WidgetProxy.h"
 #include "Widgets/SWidget.h"
diff --git a/Engine/Source/Runtime/SlateCore/Private/Rendering/SlateRenderBatch.cpp b/Engine/Source/Runtime/SlateCore/Private/Rendering/SlateRenderBatch.cpp
--- a/Engine/Source/Runtime/SlateCore/Private/Rendering/SlateRenderBatch.cpp
+++ b/Engine/Source/Runtime/SlateCore/Private/Rendering/SlateRenderBatch.cpp
@@ -1,4 +1,4 @@
-// Copyright 1998-2018 Epic Games, Inc. All Rights Reserved.
+// Copyright 1998-2019 Epic Games, Inc. All Rights Reserved.
 
 #include "Rendering/SlateRenderBatch.h"
 #include "Rendering/ElementBatcher.h"
diff --git a/Engine/Source/Runtime/SlateCore/Public/FastUpdate/SlateInvalidationRoot.h b/Engine/Source/Runtime/SlateCore/Public/FastUpdate/SlateInvalidationRoot.h
--- a/Engine/Source/Runtime/SlateCore/Public/FastUpdate/SlateInvalidationRoot.h
+++ b/Engine/Source/Runtime/SlateCore/Public/FastUpdate/SlateInvalidationRoot.h
@@ -1,4 +1,4 @@
-// Copyright 1998-2018 Epic Games, Inc. All Rights Reserved.
+// Copyright 1998-2019 Epic Games, Inc. All Rights Reserved.
 
 #pragma once
 
diff --git a/Engine/Source/Runtime/SlateCore/Public/FastUpdate/WidgetProxy.h b/Engine/Source/Runtime/SlateCore/Public/FastUpdate/WidgetProxy.h
--- a/Engine/Source/Runtime/SlateCore/Public/FastUpdate/WidgetProxy.h
+++ b/Engine/Source/Runtime/SlateCore/Public/FastUpdate/WidgetProxy.h
@@ -1,4 +1,4 @@
-// Copyright 1998-2018 Epic Games, Inc. All Rights Reserved.
+// Copyright 1998-2019 Epic Games, Inc. All Rights Reserved.
 
 #pragma once
 
diff --git a/Engine/Source/Runtime/SlateCore/Public/Rendering/DrawElementPayloads.h b/Engine/Source/Runtime/SlateCore/Public/Rendering/DrawElementPayloads.h
--- a/Engine/Source/Runtime/SlateCore/Public/Rendering/DrawElementPayloads.h
+++ b/Engine/Source/Runtime/SlateCore/Public/Rendering/DrawElementPayloads.h
@@ -1,4 +1,4 @@
-// Copyright 1998-2018 Epic Games, Inc. All Rights Reserved.
+// Copyright 1998-2019 Epic Games, Inc. All Rights Reserved.
 
 #pragma once
 
diff --git a/Engine/Source/Runtime/SlateCore/Public/Rendering/SlateRenderBatch.h b/Engine/Source/Runtime/SlateCore/Public/Rendering/SlateRenderBatch.h
--- a/Engine/Source/Runtime/SlateCore/Public/Rendering/SlateRenderBatch.h
+++ b/Engine/Source/Runtime/SlateCore/Public/Rendering/SlateRenderBatch.h
@@ -1,4 +1,4 @@
-// Copyright 1998-2018 Epic Games, Inc. All Rights Reserved.
+// Copyright 1998-2019 Epic Games, Inc. All Rights Reserved.
 
 #pragma once
 
