Fix race condition with cached widget data when invalidation panels are destroyed

diff --git a/Engine/Source/Runtime/SlateCore/Private/FastUpdate/SlateInvalidationRoot.cpp b/Engine/Source/Runtime/SlateCore/Private/FastUpdate/SlateInvalidationRoot.cpp
--- a/Engine/Source/Runtime/SlateCore/Private/FastUpdate/SlateInvalidationRoot.cpp
+++ b/Engine/Source/Runtime/SlateCore/Private/FastUpdate/SlateInvalidationRoot.cpp
@@ -27,7 +27,8 @@ static FAutoConsoleCommand HandleDumpUpdateListCommand(
 
 
 FSlateInvalidationRoot::FSlateInvalidationRoot()
-	: InvalidationRootWidget(nullptr)
+	: CachedElementData(new FSlateCachedElementData)
+	, InvalidationRootWidget(nullptr)
 	, RootHittestGrid(nullptr)
 	, FastPathGenerationNumber(INDEX_NONE)
 	, bChildOrderInvalidated(false)
@@ -44,12 +45,18 @@ FSlateInvalidationRoot::~FSlateInvalidationRoot()
 	if (FSlateApplicationBase::IsInitialized())
 	{
 		FSlateApplicationBase::Get().OnInvalidateAllWidgets().RemoveAll(this);
+
+		FSlateApplicationBase::Get().GetRenderer()->DestroyCachedFastPathElementData(CachedElementData);
+	}
+	else
+	{
+		delete CachedElementData;
 	}
 }
 
 void FSlateInvalidationRoot::AddReferencedObjects(FReferenceCollector& Collector)
 {
-	CachedElementData.AddReferencedObjects(Collector);
+	CachedElementData->AddReferencedObjects(Collector);
 }
 
 FString FSlateInvalidationRoot::GetReferencerName() const
@@ -108,7 +115,7 @@ void FSlateInvalidationRoot::RemoveWidgetFromFastPath(FWidgetProxy& Proxy)
 {
 	if (Proxy.Widget->PersistentState.CachedElementListNode)
 	{
-		CachedElementData.RemoveCache(Proxy.Widget->PersistentState.CachedElementListNode);
+		CachedElementData->RemoveCache(Proxy.Widget->PersistentState.CachedElementListNode);
 	}
 
 	if (Proxy.Index == 0)
@@ -145,7 +152,7 @@ FSlateInvalidationResult FSlateInvalidationRoot::PaintInvalidationRoot(const FSl
 
 	if (Context.bAllowFastPathUpdate)
 	{
-		Context.WindowElementList->PushCachedElementData(CachedElementData);
+		Context.WindowElementList->PushCachedElementData(*CachedElementData);
 	}
 
 	SWidget* RootWidget = InvalidationRootWidget->Advanced_IsWindow() ? InvalidationRootWidget : &(*InvalidationRootWidget->GetAllChildren()->GetChildAt(0));
@@ -535,7 +542,7 @@ void FSlateInvalidationRoot::ClearAllFastPathData()
 
 	FastWidgetPathList.Empty();
 	WidgetsNeedingUpdate.Empty();
-	CachedElementData.Empty();
+	CachedElementData->Empty();
 	FinalUpdateList.Empty();
 }
 
diff --git a/Engine/Source/Runtime/SlateCore/Private/Rendering/DrawElements.cpp b/Engine/Source/Runtime/SlateCore/Private/Rendering/DrawElements.cpp
--- a/Engine/Source/Runtime/SlateCore/Private/Rendering/DrawElements.cpp
+++ b/Engine/Source/Runtime/SlateCore/Private/Rendering/DrawElements.cpp
@@ -654,7 +654,6 @@ FSlateDrawElement& FSlateWindowElementList::AddUninitialized()
 #endif
 
 		FSlateDrawElement& NewElement = Elements[InsertIdx];
-		//NewElement.SetBatchPriorityGroup(GetCurrentBatchPriorityGroup());
 		return Elements[InsertIdx];
 	}
 }
@@ -996,8 +995,8 @@ void FSlateCachedElementList::Reset()
 
 	CachedRenderingData = new FSlateCachedFastPathRenderingData;
 
-#if WITH_SLATE_DEBUGGING
-	if (!ensure(!bNewData))
+#if 0 // enable this if you want to know why a widget is invalidated after it has been drawn but before it has been batched (probably a child or parent invalidating a relation)
+	if (ensure(!bNewData))
 	{
 		UE_LOG(LogSlate, Log, TEXT("Cleared out data in cached ElementList for Widget: %s before it was batched"), *Widget->GetTag().ToString());
 	}
diff --git a/Engine/Source/Runtime/SlateCore/Private/Rendering/SlateRenderer.cpp b/Engine/Source/Runtime/SlateCore/Private/Rendering/SlateRenderer.cpp
--- a/Engine/Source/Runtime/SlateCore/Private/Rendering/SlateRenderer.cpp
+++ b/Engine/Source/Runtime/SlateCore/Private/Rendering/SlateRenderer.cpp
@@ -140,6 +140,12 @@ void FSlateRenderer::DestroyCachedFastPathRenderingData(struct FSlateCachedFastP
 	delete FastPathRenderingData;
 }
 
+void FSlateRenderer::DestroyCachedFastPathElementData(struct FSlateCachedElementData* ElementData)
+{
+	check(ElementData);
+	delete ElementData;
+}
+
 /* Global functions
  *****************************************************************************/
 
diff --git a/Engine/Source/Runtime/SlateCore/Public/FastUpdate/SlateInvalidationRoot.h b/Engine/Source/Runtime/SlateCore/Public/FastUpdate/SlateInvalidationRoot.h
--- a/Engine/Source/Runtime/SlateCore/Public/FastUpdate/SlateInvalidationRoot.h
+++ b/Engine/Source/Runtime/SlateCore/Public/FastUpdate/SlateInvalidationRoot.h
@@ -75,7 +75,7 @@ public:
 	/** @return the cached draw elements for this window and its widget hierarchy*/
 	FSlateCachedElementData& GetCachedElements()
 	{
-		return CachedElementData;
+		return *CachedElementData;
 	}
 
 	int32 GetFastPathGenerationNumber() const { return FastPathGenerationNumber; }
@@ -107,7 +107,7 @@ private:
 
 	TArray<int32, TInlineAllocator<100>> FinalUpdateList;
 
-	FSlateCachedElementData CachedElementData;
+	FSlateCachedElementData* CachedElementData;
 
 	SWidget* InvalidationRootWidget;
 
diff --git a/Engine/Source/Runtime/SlateCore/Public/Layout/Clipping.h b/Engine/Source/Runtime/SlateCore/Public/Layout/Clipping.h
--- a/Engine/Source/Runtime/SlateCore/Public/Layout/Clipping.h
+++ b/Engine/Source/Runtime/SlateCore/Public/Layout/Clipping.h
@@ -345,8 +345,8 @@ public:
 
 	void EndUsingState()
 	{
-		FPlatformAtomics::InterlockedDecrement(&UsageCount);
-		check(UsageCount >= 0);
+		int32 NewCount = FPlatformAtomics::InterlockedDecrement(&UsageCount);
+		check(NewCount >= 0);
 	}
 private:
 	int32 UsageCount;
diff --git a/Engine/Source/Runtime/SlateCore/Public/Rendering/SlateRenderer.h b/Engine/Source/Runtime/SlateCore/Public/Rendering/SlateRenderer.h
--- a/Engine/Source/Runtime/SlateCore/Public/Rendering/SlateRenderer.h
+++ b/Engine/Source/Runtime/SlateCore/Public/Rendering/SlateRenderer.h
@@ -467,6 +467,7 @@ public:
 	virtual void ClearScenes() = 0;
 
 	virtual void DestroyCachedFastPathRenderingData(struct FSlateCachedFastPathRenderingData* VertexData);
+	virtual void DestroyCachedFastPathElementData(struct FSlateCachedElementData* ElementData);
 
 	virtual bool HasLostDevice() const { return false; }
 
diff --git a/Engine/Source/Runtime/SlateRHIRenderer/Private/SlateRHIRenderer.cpp b/Engine/Source/Runtime/SlateRHIRenderer/Private/SlateRHIRenderer.cpp
--- a/Engine/Source/Runtime/SlateRHIRenderer/Private/SlateRHIRenderer.cpp
+++ b/Engine/Source/Runtime/SlateRHIRenderer/Private/SlateRHIRenderer.cpp
@@ -1495,7 +1495,25 @@ private:
 	FSlateCachedFastPathRenderingData* CachedRenderingData;
 };
 
-void FSlateRHIRenderer::DestroyCachedFastPathRenderingData(struct FSlateCachedFastPathRenderingData* CachedRenderingData)
+struct FClearCachedElementDataCommand final : public FRHICommand < FClearCachedElementDataCommand >
+{
+public:
+	FClearCachedElementDataCommand(FSlateCachedElementData* InCachedElementData)
+		: CachedElementData(InCachedElementData)
+	{
+
+	}
+
+	void Execute(FRHICommandListBase& CmdList)
+	{
+		delete CachedElementData;
+	}
+
+private:
+	FSlateCachedElementData* CachedElementData;
+};
+
+void FSlateRHIRenderer::DestroyCachedFastPathRenderingData(FSlateCachedFastPathRenderingData* CachedRenderingData)
 {
 	check(CachedRenderingData);
 
@@ -1516,6 +1534,27 @@ void FSlateRHIRenderer::DestroyCachedFastPathRenderingData(struct FSlateCachedFa
 	});
 }
 
+void FSlateRHIRenderer::DestroyCachedFastPathElementData(FSlateCachedElementData* CachedElementData)
+{
+	check(CachedElementData);
+
+	// Cached data should be destroyed in a thread safe way.  If there is an rhi thread it could be reading from the data to copy it into a vertex buffer
+	// so delete it on the rhi thread if necessary, otherwise delete it on the render thread
+	ENQUEUE_RENDER_COMMAND(ClearCachedElementData)(
+		[CachedElementData](FRHICommandListImmediate& RHICmdList)
+	{
+		if (!RHICmdList.Bypass())
+		{
+			new (RHICmdList.AllocCommand<FClearCachedElementDataCommand>()) FClearCachedElementDataCommand(CachedElementData);
+		}
+		else
+		{
+			FClearCachedElementDataCommand Cmd(CachedElementData);
+			Cmd.Execute(RHICmdList);
+		}
+	});
+}
+
 bool FSlateRHIRenderer::AreShadersInitialized() const
 {
 #if WITH_EDITORONLY_DATA
diff --git a/Engine/Source/Runtime/SlateRHIRenderer/Private/SlateRHIRenderer.h b/Engine/Source/Runtime/SlateRHIRenderer/Private/SlateRHIRenderer.h
--- a/Engine/Source/Runtime/SlateRHIRenderer/Private/SlateRHIRenderer.h
+++ b/Engine/Source/Runtime/SlateRHIRenderer/Private/SlateRHIRenderer.h
@@ -191,6 +191,7 @@ public:
 	virtual int32 GetCurrentSceneIndex() const override;
 	virtual void ClearScenes() override;
 	virtual void DestroyCachedFastPathRenderingData(struct FSlateCachedFastPathRenderingData* InRenderingData) override;
+	virtual void DestroyCachedFastPathElementData(FSlateCachedElementData* InCachedElementData) override;
 	virtual void BeginFrame() const override;
 	virtual void EndFrame() const override;
 	virtual void AddWidgetRendererUpdate(const struct FRenderThreadUpdateContext& Context, bool bDeferredRenderTargetUpdate) override;
