Fix CIS

diff --git a/Engine/Source/Runtime/SlateCore/Public/Rendering/DrawElements.h b/Engine/Source/Runtime/SlateCore/Public/Rendering/DrawElements.h
--- a/Engine/Source/Runtime/SlateCore/Public/Rendering/DrawElements.h
+++ b/Engine/Source/Runtime/SlateCore/Public/Rendering/DrawElements.h
@@ -261,7 +261,7 @@ public:
 	SLATECORE_API static void MakePostProcessPass(FSlateWindowElementList& ElementList, uint32 InLayer, const FPaintGeometry& PaintGeometry, const FVector4& Params, int32 DownsampleAmount);
 
 	FSlateDrawElement();
-	~FSlateDrawElement();
+	SLATECORE_API ~FSlateDrawElement();
 
 	FORCEINLINE int32 GetLayer() const { return LayerId; }
 
@@ -384,7 +384,7 @@ struct FSlateCachedElementList
 	void AddReferencedObjects(FReferenceCollector& Collector);
 
 private:
-	void DestroyCachedVertexData();
+	SLATECORE_API void DestroyCachedVertexData();
 public:
 #if WITH_SLATE_DEBUGGING
 	const SWidget* Widget;
