Fix Win32 static assert

diff --git a/Engine/Source/Runtime/SlateCore/Public/FastUpdate/WidgetProxy.h b/Engine/Source/Runtime/SlateCore/Public/FastUpdate/WidgetProxy.h
--- a/Engine/Source/Runtime/SlateCore/Public/FastUpdate/WidgetProxy.h
+++ b/Engine/Source/Runtime/SlateCore/Public/FastUpdate/WidgetProxy.h
@@ -74,7 +74,8 @@ public:
 	uint8 bChildOrderInvalid : 1;
 };
 
-static_assert(sizeof(FWidgetProxy) == 32, "FWidgetProxy should be 32 bytes");
+static_assert(sizeof(FWidgetProxy) <= 32, "FWidgetProxy should be 32 bytes");
+
 static_assert(TIsTriviallyDestructible<FWidgetProxy>::Value == true, "FWidgetProxy must be trivially destructible");
 
 template <> struct TIsPODType<FWidgetProxy> { enum { Value = true }; };
