Fix pie shutdown crash in slate global invalidation

diff --git a/Engine/Source/Runtime/Slate/Private/Widgets/SInvalidationPanel.cpp b/Engine/Source/Runtime/Slate/Private/Widgets/SInvalidationPanel.cpp
--- a/Engine/Source/Runtime/Slate/Private/Widgets/SInvalidationPanel.cpp
+++ b/Engine/Source/Runtime/Slate/Private/Widgets/SInvalidationPanel.cpp
@@ -101,7 +101,7 @@ void SInvalidationPanel::OnGlobalInvalidationToggled(bool bGlobalInvalidationEna
 {
 	InvalidateRoot();
 
-	ClearAllFastPathData();
+	ClearAllFastPathData(false);
 }
 
 bool SInvalidationPanel::UpdateCachePrequisites(FSlateWindowElementList& OutDrawElements, const FGeometry& AllottedGeometry, const FSlateRect& MyCullingRect, int32 LayerId) const
diff --git a/Engine/Source/Runtime/SlateCore/Private/FastUpdate/SlateInvalidationRoot.cpp b/Engine/Source/Runtime/SlateCore/Private/FastUpdate/SlateInvalidationRoot.cpp
--- a/Engine/Source/Runtime/SlateCore/Private/FastUpdate/SlateInvalidationRoot.cpp
+++ b/Engine/Source/Runtime/SlateCore/Private/FastUpdate/SlateInvalidationRoot.cpp
@@ -40,7 +40,7 @@ FSlateInvalidationRoot::FSlateInvalidationRoot()
 
 FSlateInvalidationRoot::~FSlateInvalidationRoot()
 {
-	ClearAllFastPathData();
+	ClearAllFastPathData(true);
 
 	if (FSlateApplicationBase::IsInitialized())
 	{
@@ -174,7 +174,7 @@ FSlateInvalidationResult FSlateInvalidationRoot::PaintInvalidationRoot(const FSl
 #if WITH_SLATE_DEBUGGING
 		FSlateDebugging::ClearInvalidatedWidgets(*this);
 #endif
-		ClearAllFastPathData();
+		ClearAllFastPathData(false);
 
 		GSlateIsOnFastUpdatePath = false;
 		bNeedsSlowPath = false;
@@ -530,13 +530,17 @@ bool FSlateInvalidationRoot::ProcessInvalidation()
 	return bWidgetsNeedRepaint;
 }
 
-void FSlateInvalidationRoot::ClearAllFastPathData()
+void FSlateInvalidationRoot::ClearAllFastPathData(bool bInvalidationRootBeingDestroyed)
 {
 	for (const FWidgetProxy& Proxy : FastWidgetPathList)
 	{
 		if (Proxy.Widget)
 		{
 			Proxy.Widget->PersistentState.CachedElementListNode = nullptr;
+			if (bInvalidationRootBeingDestroyed)
+			{
+				Proxy.Widget->FastPathProxyHandle = FWidgetProxyHandle();
+			}
 		}
 	}
 
diff --git a/Engine/Source/Runtime/SlateCore/Public/FastUpdate/SlateInvalidationRoot.h b/Engine/Source/Runtime/SlateCore/Public/FastUpdate/SlateInvalidationRoot.h
--- a/Engine/Source/Runtime/SlateCore/Public/FastUpdate/SlateInvalidationRoot.h
+++ b/Engine/Source/Runtime/SlateCore/Public/FastUpdate/SlateInvalidationRoot.h
@@ -90,7 +90,7 @@ protected:
 
 	SLATECORE_API bool ProcessInvalidation();
 
-	SLATECORE_API void ClearAllFastPathData();
+	SLATECORE_API void ClearAllFastPathData(bool bInvalidationRootBeingDestroyed);
 
 private:
 
diff --git a/Engine/Source/Runtime/UMG/Private/Slate/SRetainerWidget.cpp b/Engine/Source/Runtime/UMG/Private/Slate/SRetainerWidget.cpp
--- a/Engine/Source/Runtime/UMG/Private/Slate/SRetainerWidget.cpp
+++ b/Engine/Source/Runtime/UMG/Private/Slate/SRetainerWidget.cpp
@@ -543,16 +543,9 @@ void SRetainerWidget::OnGlobalInvalidationToggled(bool bGlobalInvalidationEnable
 {
 	InvalidateRoot();
 
-	ClearAllFastPathData();
+	ClearAllFastPathData(false);
 }
 
-/*
-void SRetainerWidget::Advanced_InvalidateRoot()
-{
-	InvalidateRoot();
-}
-*/
-
 bool SRetainerWidget::CustomPrepass(float LayoutScaleMultiplier)
 {
 	if (bEnableRetainedRendering)
