Fix crashes with empty slate render batches

diff --git a/Engine/Source/Developer/StandaloneRenderer/Private/OpenGL/SlateOpenGLRenderingPolicy.cpp b/Engine/Source/Developer/StandaloneRenderer/Private/OpenGL/SlateOpenGLRenderingPolicy.cpp
--- a/Engine/Source/Developer/StandaloneRenderer/Private/OpenGL/SlateOpenGLRenderingPolicy.cpp
+++ b/Engine/Source/Developer/StandaloneRenderer/Private/OpenGL/SlateOpenGLRenderingPolicy.cpp
@@ -243,6 +243,7 @@ void FSlateOpenGLRenderingPolicy::DrawElements( const FMatrix& ViewProjectionMat
 		check( RenderBatch.GetNumIndices() > 0 );
 
 		const uint32 IndexCount = RenderBatch.GetNumIndices();
+		check(IndexCount > 0);
 
 		uint32 Offset = 0;
 		// The offset into the vertex buffer where this batches vertices are located
diff --git a/Engine/Source/Runtime/SlateCore/Private/Rendering/ElementBatcher.cpp b/Engine/Source/Runtime/SlateCore/Private/Rendering/ElementBatcher.cpp
--- a/Engine/Source/Runtime/SlateCore/Private/Rendering/ElementBatcher.cpp
+++ b/Engine/Source/Runtime/SlateCore/Private/Rendering/ElementBatcher.cpp
@@ -310,6 +310,13 @@ void FSlateBatchData::MergeRenderBatches()
 
 			FSlateRenderBatch& CurBatch = RenderBatches[BatchIndexPair.Key];
 
+
+			if (CurBatch.bIsMerged || CurBatch.GetNumIndices() == 0 || CurBatch.GetNumVertices() == 0)
+			{
+				// skip already merged batches or batches with invalid data (e.g text with pure whitespace)
+				continue;
+			}
+
 #if STATS
 			CurLayerId = CurBatch.GetLayer();
 			if (PrevLayerId != CurLayerId)
@@ -319,12 +326,6 @@ void FSlateBatchData::MergeRenderBatches()
 			CurLayerId = PrevLayerId;
 #endif
 
-			if (CurBatch.bIsMerged)
-			{
-				// skip already merged batches
-				continue;
-			}
-
 			if (PrevBatch != nullptr)
 			{
 				PrevBatch->NextBatchIndex = BatchIndexPair.Key;
diff --git a/Engine/Source/Runtime/SlateCore/Public/Rendering/ElementBatcher.h b/Engine/Source/Runtime/SlateCore/Public/Rendering/ElementBatcher.h
--- a/Engine/Source/Runtime/SlateCore/Public/Rendering/ElementBatcher.h
+++ b/Engine/Source/Runtime/SlateCore/Public/Rendering/ElementBatcher.h
@@ -152,7 +152,7 @@ class FSlateBatchData
 {
 public:
 	FSlateBatchData()
-		: FirstRenderBatchIndex(0)
+		: FirstRenderBatchIndex(INDEX_NONE)
 		, NumLayers(0)
 		, NumBatches(0)
 		, bIsStencilBufferRequired(false)
