Fix VS 2019 CIS issue

diff --git a/Engine/Source/Runtime/SlateCore/Private/Application/SlateApplicationBase.cpp b/Engine/Source/Runtime/SlateCore/Private/Application/SlateApplicationBase.cpp
--- a/Engine/Source/Runtime/SlateCore/Private/Application/SlateApplicationBase.cpp
+++ b/Engine/Source/Runtime/SlateCore/Private/Application/SlateApplicationBase.cpp
@@ -156,7 +156,7 @@ void FSlateApplicationBase::ToggleGlobalInvalidation(bool bIsGlobalInvalidationE
 {
 	if ((!!GSlateEnableGlobalInvalidation) != bIsGlobalInvalidationEnabled)
 	{
-		GSlateEnableGlobalInvalidation = bIsGlobalInvalidationEnabled;
+		GSlateEnableGlobalInvalidation = bIsGlobalInvalidationEnabled ? 1 : 0;
 		OnGlobalInvalidationToggledEvent.Broadcast(GSlateEnableGlobalInvalidation);
 	}
 }
