Fix static analysis

diff --git a/Engine/Source/Runtime/SlateCore/Private/Input/HittestGrid.cpp b/Engine/Source/Runtime/SlateCore/Private/Input/HittestGrid.cpp
--- a/Engine/Source/Runtime/SlateCore/Private/Input/HittestGrid.cpp
+++ b/Engine/Source/Runtime/SlateCore/Private/Input/HittestGrid.cpp
@@ -168,7 +168,7 @@ FHittestGrid::~FHittestGrid()
 TArray<FWidgetAndPointer> FHittestGrid::GetBubblePath(FVector2D DesktopSpaceCoordinate, float CursorRadius, bool bIgnoreEnabledStatus)
 {
 	checkSlow(IsInGameThread());
-	CursorRadius = 00.0f;
+
 	const FVector2D CursorPositionInGrid = DesktopSpaceCoordinate - GridOrigin;
 
 	if (WidgetArray.Num() > 0 && Cells.Num() > 0)
