Fix Linux/Mac link  error

diff --git a/Engine/Source/Runtime/SlateCore/Public/FastUpdate/WidgetProxy.h b/Engine/Source/Runtime/SlateCore/Public/FastUpdate/WidgetProxy.h
--- a/Engine/Source/Runtime/SlateCore/Public/FastUpdate/WidgetProxy.h
+++ b/Engine/Source/Runtime/SlateCore/Public/FastUpdate/WidgetProxy.h
@@ -214,7 +214,7 @@ public:
 	void MarkWidgetUpdatedThisFrame();
 	
 	void MarkWidgetDirty(EInvalidateWidgetReason InvalidateReason);
-	void UpdateWidgetFlags(EWidgetUpdateFlags NewFlags);
+	SLATECORE_API void UpdateWidgetFlags(EWidgetUpdateFlags NewFlags);
 private:
 	FWidgetProxyHandle(FSlateInvalidationRoot& InInvalidationRoot, int32 InIndex);
 private:
