Fix for crashes with widgets not knowing their cached geometry was destoryed

diff --git a/Engine/Source/Runtime/SlateCore/Private/FastUpdate/SlateInvalidationRoot.cpp b/Engine/Source/Runtime/SlateCore/Private/FastUpdate/SlateInvalidationRoot.cpp
--- a/Engine/Source/Runtime/SlateCore/Private/FastUpdate/SlateInvalidationRoot.cpp
+++ b/Engine/Source/Runtime/SlateCore/Private/FastUpdate/SlateInvalidationRoot.cpp
@@ -42,6 +42,10 @@ FSlateInvalidationRoot::~FSlateInvalidationRoot()
 {
 	ClearAllFastPathData(true);
 
+#if WITH_SLATE_DEBUGGING
+	FSlateDebugging::ClearInvalidatedWidgets(*this);
+#endif
+
 	if (FSlateApplicationBase::IsInitialized())
 	{
 		FSlateApplicationBase::Get().OnInvalidateAllWidgets().RemoveAll(this);
diff --git a/Engine/Source/Runtime/SlateCore/Private/Rendering/DrawElements.cpp b/Engine/Source/Runtime/SlateCore/Private/Rendering/DrawElements.cpp
--- a/Engine/Source/Runtime/SlateCore/Private/Rendering/DrawElements.cpp
+++ b/Engine/Source/Runtime/SlateCore/Private/Rendering/DrawElements.cpp
@@ -803,18 +803,10 @@ FSlateDrawElement& FSlateWindowElementList::AddCachedElement()
 
 	if (CurrentWidgetState.CacheNode == nullptr)
 	{
-#if WITH_SLATE_DEBUGGING
 		CurrentWidgetState.CacheNode = CurrentCachedElementData->AddCache(CurrentWidgetState.Widget);
-#else
-		CurrentWidgetState.CacheNode = CurrentCachedElementData->AddCache(nullptr);
-#endif
 	}
 
-#if WITH_SLATE_DEBUGGING
 	return CurrentCachedElementData->AddCachedElement(CurrentWidgetState.CacheNode, GetClippingManager(), CurrentWidgetState.Widget);
-#else
-	return CurrentCachedElementData->AddCachedElement(CurrentWidgetState.CacheNode, GetClippingManager(), nullptr);
-#endif
 }
 
 void FSlateWindowElementList::PushCachedElementData(FSlateCachedElementData& CachedElementData)
@@ -983,6 +975,12 @@ void FSlateCachedElementData::CleanupUnusedClipStates()
 	}
 }
 
+FSlateCachedElementList::~FSlateCachedElementList()
+{
+	DestroyCachedVertexData();
+	Widget->PersistentState.CachedElementListNode = nullptr;
+}
+
 void FSlateCachedElementList::Reset()
 {
 	DrawElements.Reset();
diff --git a/Engine/Source/Runtime/SlateCore/Public/Rendering/DrawElements.h b/Engine/Source/Runtime/SlateCore/Public/Rendering/DrawElements.h
--- a/Engine/Source/Runtime/SlateCore/Public/Rendering/DrawElements.h
+++ b/Engine/Source/Runtime/SlateCore/Public/Rendering/DrawElements.h
@@ -352,17 +352,10 @@ struct FSlateCachedFastPathRenderingData
 
 struct FSlateCachedElementList
 {
-#if WITH_SLATE_DEBUGGING
 	FSlateCachedElementList(const SWidget* InWidget)
 		: Widget(InWidget)
 		, CachedRenderingData(nullptr)
 		, bNewData(false)
-#else
-	FSlateCachedElementList()
-		: CachedRenderingData(nullptr)
-		, bNewData(false)
-
-#endif
 	{}
 
 	void Initialize()
@@ -370,10 +363,7 @@ struct FSlateCachedElementList
 		CachedRenderingData = new FSlateCachedFastPathRenderingData;
 	}
 
-	~FSlateCachedElementList()
-	{
-		DestroyCachedVertexData();
-	}
+	SLATECORE_API ~FSlateCachedElementList();
 
 	void Reset();
 
@@ -386,9 +376,7 @@ struct FSlateCachedElementList
 private:
 	SLATECORE_API void DestroyCachedVertexData();
 public:
-#if WITH_SLATE_DEBUGGING
 	const SWidget* Widget;
-#endif
 	/** List of source draw elements to create batches from */
 	FSlateDrawElementArray DrawElements;
 	/** List of cached batches to submit for drawing */
@@ -431,11 +419,9 @@ struct FSlateCachedElementData
 		{
 			ensure(CachedElementList.Widget != Widget);
 		}
+#endif
 
 		FSlateCachedElementListNode* NewNode = new FSlateCachedElementListNode(Widget);
-#else
-		FSlateCachedElementListNode* NewNode = new FSlateCachedElementListNode(FSlateCachedElementList());
-#endif
 
 
 		CachedElementLists.AddTail(NewNode);
@@ -684,18 +670,15 @@ private:
 	{
 		FWidgetDrawElementState(FSlateCachedElementListNode* InCurrentCacheNode,  bool bInIsVolatile, const SWidget* InWidget)
 			: CacheNode(InCurrentCacheNode)
-			, bIsVolatile(bInIsVolatile)
-#if WITH_SLATE_DEBUGGING
 			, Widget(InWidget)
-#endif
+			, bIsVolatile(bInIsVolatile)
 		{
 		}
 
 		FSlateCachedElementListNode* CacheNode;
-		bool bIsVolatile;
-#if WITH_SLATE_DEBUGGING
 		const SWidget* Widget;
-#endif
+		bool bIsVolatile;
+	
 	};
 
 
diff --git a/Engine/Source/Runtime/SlateCore/Public/Widgets/SWidget.h b/Engine/Source/Runtime/SlateCore/Public/Widgets/SWidget.h
--- a/Engine/Source/Runtime/SlateCore/Public/Widgets/SWidget.h
+++ b/Engine/Source/Runtime/SlateCore/Public/Widgets/SWidget.h
@@ -270,6 +270,7 @@ class SLATECORE_API SWidget
 	friend class FSlateInvalidationRoot;
 	friend class FSlateWindowElementList;
 	friend class SWindow;
+	friend struct FSlateCachedElementList;
 #if WITH_SLATE_DEBUGGING
 	friend struct FInvalidatedWidgetDrawer;
 #endif
