Fix shutdown crash that happens because a widget is destroyed and isnt removed from an invalidation root because the invalidation root been disabled from using the fast path

diff --git a/Engine/Source/Runtime/SlateCore/Private/FastUpdate/SlateInvalidationRoot.cpp b/Engine/Source/Runtime/SlateCore/Private/FastUpdate/SlateInvalidationRoot.cpp
--- a/Engine/Source/Runtime/SlateCore/Private/FastUpdate/SlateInvalidationRoot.cpp
+++ b/Engine/Source/Runtime/SlateCore/Private/FastUpdate/SlateInvalidationRoot.cpp
@@ -56,6 +56,8 @@ FSlateInvalidationRoot::~FSlateInvalidationRoot()
 	{
 		delete CachedElementData;
 	}
+
+	CachedElementData = nullptr;
 }
 
 void FSlateInvalidationRoot::AddReferencedObjects(FReferenceCollector& Collector)
@@ -238,6 +240,27 @@ FSlateInvalidationResult FSlateInvalidationRoot::PaintInvalidationRoot(const FSl
 	return Result;
 }
 
+void FSlateInvalidationRoot::OnWidgetDestroyed(const SWidget* Widget)
+{
+	InvalidateChildOrder();
+
+	// We need the index even if we've invalidated this root.  We need to clear out its proxy regardless
+	const bool bEvenIfInvalid = true;
+	const int32 ProxyIndex = Widget->FastPathProxyHandle.GetIndex(bEvenIfInvalid);
+	if (FastWidgetPathList.IsValidIndex(ProxyIndex) && FastWidgetPathList[ProxyIndex].Widget == Widget)
+	{
+		FastWidgetPathList[ProxyIndex].Widget = nullptr;
+	
+	}
+
+	if (Widget->PersistentState.CachedElementListNode)
+	{
+		CachedElementData->RemoveCache(Widget->PersistentState.CachedElementListNode);
+	}
+
+	Widget->PersistentState.CachedElementListNode = nullptr;
+}
+
 bool FSlateInvalidationRoot::PaintFastPath(const FSlateInvalidationContext& Context)
 {
 	SCOPED_NAMED_EVENT(SWidget_FastPathUpdate, FColor::Green);
diff --git a/Engine/Source/Runtime/SlateCore/Private/Widgets/SWidget.cpp b/Engine/Source/Runtime/SlateCore/Private/Widgets/SWidget.cpp
--- a/Engine/Source/Runtime/SlateCore/Private/Widgets/SWidget.cpp
+++ b/Engine/Source/Runtime/SlateCore/Private/Widgets/SWidget.cpp
@@ -181,25 +181,9 @@ SWidget::~SWidget()
 			FSlateApplicationBase::Get().UnRegisterActiveTimer(ActiveTimerHandle);
 		}
 
-		if (FastPathProxyHandle.IsValid())
-		{
-			FWidgetProxy& Proxy = FastPathProxyHandle.GetProxy();
-			Proxy.Widget = nullptr;
-		}
-		else
-		{
-			ensure(FastPathProxyHandle.GetIndex() == INDEX_NONE);
-		}
-
 		if (FSlateInvalidationRoot* Root = FastPathProxyHandle.GetInvalidationRoot())
 		{
-			Root->InvalidateChildOrder();
-
-			if (PersistentState.CachedElementListNode)
-			{
-				Root->GetCachedElements().RemoveCache(PersistentState.CachedElementListNode);
-			}
-		
+			Root->OnWidgetDestroyed(this);
 		}
 
 		// Reset handle
diff --git a/Engine/Source/Runtime/SlateCore/Public/FastUpdate/SlateInvalidationRoot.h b/Engine/Source/Runtime/SlateCore/Public/FastUpdate/SlateInvalidationRoot.h
--- a/Engine/Source/Runtime/SlateCore/Public/FastUpdate/SlateInvalidationRoot.h
+++ b/Engine/Source/Runtime/SlateCore/Public/FastUpdate/SlateInvalidationRoot.h
@@ -82,6 +82,7 @@ public:
 
 	SLATECORE_API FSlateInvalidationResult PaintInvalidationRoot(const FSlateInvalidationContext& Context);
 
+	void OnWidgetDestroyed(const SWidget* Widget);
 protected:
 	virtual int32 PaintSlowPath(const FSlateInvalidationContext& Context) = 0;
 
diff --git a/Engine/Source/Runtime/SlateCore/Public/FastUpdate/WidgetProxy.h b/Engine/Source/Runtime/SlateCore/Public/FastUpdate/WidgetProxy.h
--- a/Engine/Source/Runtime/SlateCore/Public/FastUpdate/WidgetProxy.h
+++ b/Engine/Source/Runtime/SlateCore/Public/FastUpdate/WidgetProxy.h
@@ -204,7 +204,7 @@ public:
 	FWidgetProxy& GetProxy();
 	const FWidgetProxy& GetProxy() const;
 
-	int32 GetIndex() const { return IsValid() ? MyIndex : INDEX_NONE; }
+	int32 GetIndex(bool bEvenIfInvalid = false) const { return bEvenIfInvalid || IsValid() ? MyIndex : INDEX_NONE; }
 
 	/**
 	 * Marks the widget as updated this frame
