fix bug in overlay layer padding causing overlflow of layer Id's and causing batch sorting isssues

diff --git a/Engine/Source/Runtime/SlateCore/Private/Widgets/SOverlay.cpp b/Engine/Source/Runtime/SlateCore/Private/Widgets/SOverlay.cpp
--- a/Engine/Source/Runtime/SlateCore/Private/Widgets/SOverlay.cpp
+++ b/Engine/Source/Runtime/SlateCore/Private/Widgets/SOverlay.cpp
@@ -81,7 +81,6 @@ int32 SOverlay::OnPaint( const FPaintArgs& Args, const FGeometry& AllottedGeomet
 	FPaintArgs NewArgs = Args.WithNewParent(this);
 	const bool bChildrenEnabled = ShouldBeEnabled(bParentEnabled);
 
-	//int32 NumPriorityGroupsPushed = 0;
 
 	for (int32 ChildIndex = 0; ChildIndex < ArrangedChildren.Num(); ++ChildIndex)
 	{
@@ -90,8 +89,6 @@ int32 SOverlay::OnPaint( const FPaintArgs& Args, const FGeometry& AllottedGeomet
 		// We don't increment the first layer.
 		if (ChildIndex > 0)
 		{
-			//OutDrawElements.PushBatchPriortyGroup(CurArrangedWidget.Widget.Get());
-			//NumPriorityGroupsPushed++;
 			MaxLayerId++;
 		}
 
@@ -109,17 +106,13 @@ int32 SOverlay::OnPaint( const FPaintArgs& Args, const FGeometry& AllottedGeomet
 		// Overlay slots that do not update will not know about the new layer id.  This padding adds buffering to avoid that being a problem for now
 		// This is a temporary solution until we build a full rendering tree
 		const int32 OverlaySlotPadding = 10;
-		MaxLayerId = CurWidgetsMaxLayerId + FMath::Max((CurWidgetsMaxLayerId - MaxLayerId) / OverlaySlotPadding, 1) * OverlaySlotPadding;
+		MaxLayerId = CurWidgetsMaxLayerId + FMath::Min(FMath::Max((CurWidgetsMaxLayerId - MaxLayerId) / OverlaySlotPadding, 1) * OverlaySlotPadding,100);
+
+		// Non padding method
 		//MaxLayerId = FMath::Max(CurWidgetsMaxLayerId, MaxLayerId);
 	
 	}
 
-	/*// Remove all priority groups made
-	for (int32 i = 0; i < NumPriorityGroupsPushed; ++i)
-	{
-		OutDrawElements.PopBatchPriortyGroup();
-	}
-*/
 
 
 	return MaxLayerId;
