Copy of CL 4566240

Initialize FHittestGrid::NumCells before use

diff --git a/Engine/Source/Runtime/SlateCore/Private/Input/HittestGrid.cpp b/Engine/Source/Runtime/SlateCore/Private/Input/HittestGrid.cpp
--- a/Engine/Source/Runtime/SlateCore/Private/Input/HittestGrid.cpp
+++ b/Engine/Source/Runtime/SlateCore/Private/Input/HittestGrid.cpp
@@ -97,6 +97,7 @@ struct FHittestGrid::FGridTestingParams
 
 FHittestGrid::FHittestGrid()
 : WidgetsCachedThisFrame()
+, NumCells(0, 0)
 , NumCellsExcess(0, 0)
 {
 }
