Change default font size for SlateFontInfo from 0 to 24

diff --git a/Engine/Source/Runtime/SlateCore/Private/Fonts/SlateFontInfo.cpp b/Engine/Source/Runtime/SlateCore/Private/Fonts/SlateFontInfo.cpp
--- a/Engine/Source/Runtime/SlateCore/Private/Fonts/SlateFontInfo.cpp
+++ b/Engine/Source/Runtime/SlateCore/Private/Fonts/SlateFontInfo.cpp
@@ -36,7 +36,7 @@ FSlateFontInfo::FSlateFontInfo( )
 	, FontMaterial(nullptr)
 	, CompositeFont()
 	, TypefaceFontName()
-	, Size(0)
+	, Size(24)
 	, FontFallback(EFontFallback::FF_Max)
 #if WITH_EDITORONLY_DATA
 	, Hinting_DEPRECATED(EFontHinting::Default)
