Previous and next now support wrap navigation rules, also Stop navigation now works properly.

diff --git a/Engine/Source/Runtime/SlateCore/Private/Layout/WidgetPath.cpp b/Engine/Source/Runtime/SlateCore/Private/Layout/WidgetPath.cpp
--- a/Engine/Source/Runtime/SlateCore/Private/Layout/WidgetPath.cpp
+++ b/Engine/Source/Runtime/SlateCore/Private/Layout/WidgetPath.cpp
@@ -154,12 +154,13 @@ struct FFocusableWidgetMatcher
  * Move focus either forward on backward in the path level specified by PathLevel.
  * That is, this movement of focus will modify the subtree under Widgets(PathLevel).
  *
- * @param PathLevel       The level in this WidgetPath whose focus to move.
- * @param MoveDirectin    Move focus forward or backward?
+ * @param PathLevel					The level in this WidgetPath whose focus to move.
+ * @param MoveDirectin				Move focus forward or backward?
+ * @param bSearchFromPathWidget		if set false the search for the next will simply start at the beginning or end of the list of children dependant on the direction
  *
  * @return true if the focus moved successfully, false if we were unable to move focus
  */
-bool FWidgetPath::MoveFocus(int32 PathLevel, EUINavigation NavigationType)
+bool FWidgetPath::MoveFocus(int32 PathLevel, EUINavigation NavigationType, bool bSearchFromPathWidget)
 {
 	check(NavigationType == EUINavigation::Next || NavigationType == EUINavigation::Previous);
 
@@ -194,9 +195,13 @@ bool FWidgetPath::MoveFocus(int32 PathLevel, EUINavigation NavigationType)
 		// Don't continue if there are no children in the widget.
 		if (ArrangedChildren.Num() > 0)
 		{
-			// Find the currently focused child among the children.
-			int32 FocusedChildIndex = ArrangedChildren.FindItemIndex(Widgets[PathLevel + 1]);
-			FocusedChildIndex = (FocusedChildIndex) % ArrangedChildren.Num() + MoveDirectionAsInt;
+			int32 FocusedChildIndex = NavigationType == EUINavigation::Next ? 0 : ArrangedChildren.Num() - 1;
+			if (bSearchFromPathWidget)
+			{
+				// Find the currently focused child among the children.
+				FocusedChildIndex = ArrangedChildren.FindItemIndex(Widgets[PathLevel + 1]);
+				FocusedChildIndex = (FocusedChildIndex) % ArrangedChildren.Num() + MoveDirectionAsInt;
+			}
 
 			// Now actually search for the widget.
 			for (; FocusedChildIndex < ArrangedChildren.Num() && FocusedChildIndex >= 0; FocusedChildIndex += MoveDirectionAsInt)
@@ -426,15 +431,24 @@ FWidgetPath FWeakWidgetPath::ToNextFocusedPath(EUINavigation NavigationType, con
 	// Attempt to move the focus starting at the leafmost widget and bubbling up to the root (i.e. the window)
 	for (int32 FocusNodeIndex=NewFocusPath.Widgets.Num()-1; !bMovedFocus && FocusNodeIndex >= 0; --FocusNodeIndex)
 	{
-		// We've reached the stop boundary and not yet moved focus, so don't advance.
-		if ( NavigationReply.GetBoundaryRule() == EUINavigationRule::Stop && RuleWidget.Widget == NewFocusPath.Widgets[FocusNodeIndex].Widget )
-		{
-			break;
-		}
+		bMovedFocus = NewFocusPath.MoveFocus(FocusNodeIndex, NavigationType);
 
-		//TODO Slate Navigation Handle Wrap.
+		// If we didn't move and we hit our rule widget consider stop and wrap
+		if (!bMovedFocus && RuleWidget.Widget == NewFocusPath.Widgets[FocusNodeIndex].Widget)
+		{
+			// We've reached the stop boundary and not yet moved focus, so don't advance.
+			if (NavigationReply.GetBoundaryRule() == EUINavigationRule::Stop)
+			{
+				break;
+			}
 
-		bMovedFocus = NewFocusPath.MoveFocus(FocusNodeIndex, NavigationType);
+			// We've reached the wrap boundary and not yet moved focus so move to the first or last element in BoundaryWidget
+			if (NavigationReply.GetBoundaryRule() == EUINavigationRule::Wrap)
+			{
+				NewFocusPath.MoveFocus(FocusNodeIndex, NavigationType, false);
+				break;
+			}
+		}
 	}
 
 	return NewFocusPath;
diff --git a/Engine/Source/Runtime/SlateCore/Public/Layout/WidgetPath.h b/Engine/Source/Runtime/SlateCore/Public/Layout/WidgetPath.h
--- a/Engine/Source/Runtime/SlateCore/Public/Layout/WidgetPath.h
+++ b/Engine/Source/Runtime/SlateCore/Public/Layout/WidgetPath.h
@@ -138,12 +138,13 @@ public:
 	 * Move focus either forward on backward in the path level specified by PathLevel.
 	 * That is, this movement of focus will modify the subtree under Widgets(PathLevel).
 	 *
-	 * @param PathLevel       The level in this WidgetPath whose focus to move.
-	 * @param MoveDirectin    Move focus forward or backward?
+	 * @param PathLevel					The level in this WidgetPath whose focus to move.
+	 * @param MoveDirectin				Move focus forward or backward?
+	 * @param bSearchFromPathWidget		if set false the search for the next will simply start at the beginning or end of the list of children dependant on the direction
 	 *
 	 * @return true if the focus moved successfully, false if we were unable to move focus
 	 */
-	bool MoveFocus(int32 PathLevel, EUINavigation NavigationType);
+	bool MoveFocus(int32 PathLevel, EUINavigation NavigationType, bool bSearchFromPathWidget = true);
 
 	/** Get the last (leaf-most) widget in this path; assumes path is valid */
 	TSharedRef< SWidget > GetLastWidget() const
