PR #5241: Prevent popup from unexpectedly overlapping mouse position (Contributed by [REDACTED_PERSON])

diff --git a/Engine/Source/Runtime/SlateCore/Private/Layout/LayoutUtils.cpp b/Engine/Source/Runtime/SlateCore/Private/Layout/LayoutUtils.cpp
--- a/Engine/Source/Runtime/SlateCore/Private/Layout/LayoutUtils.cpp
+++ b/Engine/Source/Runtime/SlateCore/Private/Layout/LayoutUtils.cpp
@@ -2,7 +2,7 @@
 
 #include "Layout/LayoutUtils.h"
 
-FVector2D ComputePopupFitInRect(const FSlateRect& InAnchor, const FSlateRect& PopupRect, const EOrientation Orientation, const FSlateRect RectToFit)
+FVector2D ComputePopupFitInRect(const FSlateRect& InAnchor, const FSlateRect& PopupRect, const EOrientation& Orientation, const FSlateRect& RectToFit)
 {
 	const bool bAdjustmentNeeded = PopupRect.IntersectionWith(RectToFit) != PopupRect;
 	if (bAdjustmentNeeded)
@@ -46,7 +46,7 @@ FVector2D ComputePopupFitInRect(const FSlateRect& InAnchor, const FSlateRect& Po
 
 			if (!bFitsDown && !bFitsUp)
 			{
-				NewPosition.X = InAnchor.Right;
+				return ComputePopupFitInRect(InAnchor, PopupRect, EOrientation::Orient_Horizontal, RectToFit);
 			}
 		}
 
diff --git a/Engine/Source/Runtime/SlateCore/Public/Layout/LayoutUtils.h b/Engine/Source/Runtime/SlateCore/Public/Layout/LayoutUtils.h
--- a/Engine/Source/Runtime/SlateCore/Public/Layout/LayoutUtils.h
+++ b/Engine/Source/Runtime/SlateCore/Public/Layout/LayoutUtils.h
@@ -167,4 +167,4 @@ void ArrangeSingleChild( const FGeometry& AllottedGeometry, FArrangedChildren& A
 *
 * @return A best position within the RectToFit such that none of the popup clips outside of the RectToFit.
 */
-SLATECORE_API FVector2D ComputePopupFitInRect(const FSlateRect& InAnchor, const FSlateRect& PopupRect, const EOrientation Orientation, const FSlateRect RectToFit);
+SLATECORE_API FVector2D ComputePopupFitInRect(const FSlateRect& InAnchor, const FSlateRect& PopupRect, const EOrientation& Orientation, const FSlateRect& RectToFit);
