Optimize a few FNames being created repeatedly in slate

diff --git a/Engine/Source/Runtime/SlateCore/Private/Rendering/ElementBatcher.cpp b/Engine/Source/Runtime/SlateCore/Private/Rendering/ElementBatcher.cpp
--- a/Engine/Source/Runtime/SlateCore/Private/Rendering/ElementBatcher.cpp
+++ b/Engine/Source/Runtime/SlateCore/Private/Rendering/ElementBatcher.cpp
@@ -100,12 +100,12 @@ void FSlateElementBatcher::AddElementsInternal(const TArray<FSlateDrawElement>&
 {
 	checkSlow(DrawLayer);
 
-	FName Elements_Boxes(TEXT("Boxes"));
-	FName Elements_Borders(TEXT("Borders"));
-	FName Elements_Text(TEXT("Text"));
-	FName Elements_ShapedText(TEXT("ShapedText"));
-	FName Elements_Lines(TEXT("Lines"));
-	FName Elements_CachedBuffer(TEXT("CachedBuffer"));
+	static const FName Elements_Boxes(TEXT("Boxes"));
+	static const FName Elements_Borders(TEXT("Borders"));
+	static const FName Elements_Text(TEXT("Text"));
+	static const FName Elements_ShapedText(TEXT("ShapedText"));
+	static const FName Elements_Lines(TEXT("Lines"));
+	static const FName Elements_CachedBuffer(TEXT("CachedBuffer"));
 
 	for ( int32 DrawElementIndex = 0; DrawElementIndex < DrawElements.Num(); ++DrawElementIndex )
 	{
