Don't early out of complex shaping detection when detecting ASCII characters

We still need to process the rest of the string, as there may be characters later on that require complex shaping

diff --git a/Engine/Source/Runtime/SlateCore/Private/Fonts/SlateTextShaper.cpp b/Engine/Source/Runtime/SlateCore/Private/Fonts/SlateTextShaper.cpp
--- a/Engine/Source/Runtime/SlateCore/Private/Fonts/SlateTextShaper.cpp
+++ b/Engine/Source/Runtime/SlateCore/Private/Fonts/SlateTextShaper.cpp
@@ -220,7 +220,7 @@ void FSlateTextShaper::PerformTextShaping(const TCHAR* InText, const int32 InTex
 
 					if (Char <= TEXT('\u007F'))
 					{
-						return false;
+						continue;
 					}
 
 					if (CharRequiresFullShaping(Char))
