Unrolled a loop and stopped some temporary small arrays from allocating on the heap for their common case

diff --git a/Engine/Source/Runtime/SlateCore/Private/Fonts/SlateTextShaper.cpp b/Engine/Source/Runtime/SlateCore/Private/Fonts/SlateTextShaper.cpp
--- a/Engine/Source/Runtime/SlateCore/Private/Fonts/SlateTextShaper.cpp
+++ b/Engine/Source/Runtime/SlateCore/Private/Fonts/SlateTextShaper.cpp
@@ -157,58 +157,52 @@ void FSlateTextShaper::PerformTextShaping(const TCHAR* InText, const int32 InTex
 				// Note: We deliberately avoid using HarfBuzz/ICU here as we don't care about the script itself, only that the character is within a shaped script range (and testing that is much faster!)
 				auto CharRequiresFullShaping = [](const TCHAR InChar) -> bool
 				{
-					// This isn't an exhaustive list, as it omits some "dead" or uncommon languages, and ranges outside the BMP
-					static const TCHAR FullShapingScriptRanges[][2] = {
-						// Combining characters
-						{ TEXT('\u0300'), TEXT('\u036F') },
-						{ TEXT('\u1AB0'), TEXT('\u1AFF') },
-						{ TEXT('\u1DC0'), TEXT('\u1DFF') },
-						{ TEXT('\u20D0'), TEXT('\u20FF') },
-						{ TEXT('\u3099'), TEXT('\u309A') },
-						{ TEXT('\u31C0'), TEXT('\u31EF') },
-						{ TEXT('\uFE20'), TEXT('\uFE2F') },
-
-						// Devanagari
-						{ TEXT('\u0900'), TEXT('\u097F') },
-						{ TEXT('\uA8E0'), TEXT('\uA8FF') },
-						{ TEXT('\u1CD0'), TEXT('\u1CFF') },
-
-						// Telugu
-						{ TEXT('\u0C00'), TEXT('\u0C7F') },
-
-						// Thai
-						{ TEXT('\u0E00'), TEXT('\u0E7F') },
-						
-						// Tibetan
-						{ TEXT('\u0F00'), TEXT('\u0FFF') },
-
-						// Khmer
-						{ TEXT('\u1780'), TEXT('\u17FF') },
-						{ TEXT('\u19E0'), TEXT('\u19FF') },
-
-						// Sinhala
-						{ TEXT('\u0D80'), TEXT('\u0DFF') },
-
-						// Limbu
-						{ TEXT('\u1900'), TEXT('\u194F') },
-
-						// Tai Tham
-						{ TEXT('\u1A20'), TEXT('\u1AAF') },
-
-						// Tai Viet
-						{ TEXT('\uAA80'), TEXT('\uAADF') },
-
-						// Batak
-						{ TEXT('\u1BC0'), TEXT('\u1BFF') },
-					};
-
-					for (const TCHAR* FullShapingScriptRange : FullShapingScriptRanges)
-					{
-						if (InChar >= FullShapingScriptRange[0] && InChar <= FullShapingScriptRange[1])
-						{
-							return true;
-						}
-					}
+					// Note: This isn't an exhaustive list, as it omits some "dead" or uncommon languages, and ranges outside the BMP
+					#define RETURN_TRUE_IF_CHAR_WITHIN_RANGE(LOWER, UPPER) if (InChar >= (LOWER) && InChar <= (UPPER)) return true
+
+					// Combining characters
+					RETURN_TRUE_IF_CHAR_WITHIN_RANGE(TEXT('\u0300'), TEXT('\u036F'));
+					RETURN_TRUE_IF_CHAR_WITHIN_RANGE(TEXT('\u1AB0'), TEXT('\u1AFF'));
+					RETURN_TRUE_IF_CHAR_WITHIN_RANGE(TEXT('\u1DC0'), TEXT('\u1DFF'));
+					RETURN_TRUE_IF_CHAR_WITHIN_RANGE(TEXT('\u20D0'), TEXT('\u20FF'));
+					RETURN_TRUE_IF_CHAR_WITHIN_RANGE(TEXT('\u3099'), TEXT('\u309A'));
+					RETURN_TRUE_IF_CHAR_WITHIN_RANGE(TEXT('\u31C0'), TEXT('\u31EF'));
+					RETURN_TRUE_IF_CHAR_WITHIN_RANGE(TEXT('\uFE20'), TEXT('\uFE2F'));
+
+					// Devanagari
+					RETURN_TRUE_IF_CHAR_WITHIN_RANGE(TEXT('\u0900'), TEXT('\u097F'));
+					RETURN_TRUE_IF_CHAR_WITHIN_RANGE(TEXT('\uA8E0'), TEXT('\uA8FF'));
+					RETURN_TRUE_IF_CHAR_WITHIN_RANGE(TEXT('\u1CD0'), TEXT('\u1CFF'));
+
+					// Telugu
+					RETURN_TRUE_IF_CHAR_WITHIN_RANGE(TEXT('\u0C00'), TEXT('\u0C7F'));
+
+					// Thai
+					RETURN_TRUE_IF_CHAR_WITHIN_RANGE(TEXT('\u0E00'), TEXT('\u0E7F'));
+					
+					// Tibetan
+					RETURN_TRUE_IF_CHAR_WITHIN_RANGE(TEXT('\u0F00'), TEXT('\u0FFF'));
+
+					// Khmer
+					RETURN_TRUE_IF_CHAR_WITHIN_RANGE(TEXT('\u1780'), TEXT('\u17FF'));
+					RETURN_TRUE_IF_CHAR_WITHIN_RANGE(TEXT('\u19E0'), TEXT('\u19FF'));
+
+					// Sinhala
+					RETURN_TRUE_IF_CHAR_WITHIN_RANGE(TEXT('\u0D80'), TEXT('\u0DFF'));
+
+					// Limbu
+					RETURN_TRUE_IF_CHAR_WITHIN_RANGE(TEXT('\u1900'), TEXT('\u194F'));
+
+					// Tai Tham
+					RETURN_TRUE_IF_CHAR_WITHIN_RANGE(TEXT('\u1A20'), TEXT('\u1AAF'));
+
+					// Tai Viet
+					RETURN_TRUE_IF_CHAR_WITHIN_RANGE(TEXT('\uAA80'), TEXT('\uAADF'));
+
+					// Batak
+					RETURN_TRUE_IF_CHAR_WITHIN_RANGE(TEXT('\u1BC0'), TEXT('\u1BFF'));
+
+					#undef RETURN_TRUE_IF_CHAR_WITHIN_RANGE
 
 					return false;
 				};
@@ -284,7 +278,7 @@ FShapedGlyphSequenceRef FSlateTextShaper::FinalizeTextShaping(TArray<FShapedGlyp
 void FSlateTextShaper::PerformKerningOnlyTextShaping(const TCHAR* InText, const int32 InTextStart, const int32 InTextLen, const FSlateFontInfo &InFontInfo, const float InFontScale, TArray<FShapedGlyphEntry>& OutGlyphsToRender) const
 {
 	// We need to work out the correct FFontData for everything so that we can build accurate FShapedGlyphFaceData for rendering later on
-	TArray<FKerningOnlyTextSequenceEntry> KerningOnlyTextSequence;
+	TArray<FKerningOnlyTextSequenceEntry, TInlineAllocator<4>> KerningOnlyTextSequence;
 
 	// Step 1) Split the text into sections that are using the same font face (composite fonts may contain different faces for different character ranges)
 	{
@@ -438,7 +432,7 @@ void FSlateTextShaper::PerformKerningOnlyTextShaping(const TCHAR* InText, const
 void FSlateTextShaper::PerformHarfBuzzTextShaping(const TCHAR* InText, const int32 InTextStart, const int32 InTextLen, const FSlateFontInfo &InFontInfo, const float InFontScale, const TextBiDi::ETextDirection InTextDirection, TArray<FShapedGlyphEntry>& OutGlyphsToRender) const
 {
 	// HarfBuzz can only shape data that uses the same font face, reads in the same direction, and uses the same script so we need to split the given text...
-	TArray<FHarfBuzzTextSequenceEntry> HarfBuzzTextSequence;
+	TArray<FHarfBuzzTextSequenceEntry, TInlineAllocator<4>> HarfBuzzTextSequence;
 	hb_unicode_funcs_t* HarfBuzzUnicodeFuncs = hb_unicode_funcs_get_default();
 
 	// Step 1) Split the text into sections that are using the same font face (composite fonts may contain different faces for different character ranges)
