Properly compile out slate debugging code unless WITH_SLATE_DEBUGGING is true

diff --git a/Engine/Source/Runtime/SlateCore/Private/Debugging/ConsoleSlateDebugger.cpp b/Engine/Source/Runtime/SlateCore/Private/Debugging/ConsoleSlateDebugger.cpp
--- a/Engine/Source/Runtime/SlateCore/Private/Debugging/ConsoleSlateDebugger.cpp
+++ b/Engine/Source/Runtime/SlateCore/Private/Debugging/ConsoleSlateDebugger.cpp
@@ -1,6 +1,9 @@
 // Copyright 1998-2019 Epic Games, Inc. All Rights Reserved.
 
 #include "ConsoleSlateDebugger.h"
+
+#if WITH_SLATE_DEBUGGING
+
 #include "Debugging/SlateDebugging.h"
 #include "Widgets/SBoxPanel.h"
 #include "Modules/ModuleManager.h"
@@ -16,9 +19,7 @@ DEFINE_LOG_CATEGORY_STATIC(LogSlateDebugger, Log, All);
 
 #define LOCTEXT_NAMESPACE "ConsoleSlateDebugger"
 
-#if WITH_SLATE_DEBUGGING
 static FConsoleSlateDebugger SlateConsoleDebugger;
-#endif
 
 FConsoleSlateDebugger::FConsoleSlateDebugger()
 	: bCaptureStack(0)
@@ -303,3 +304,5 @@ void FConsoleSlateDebugger::OptionallyDumpCallStack()
 }
 
 #undef LOCTEXT_NAMESPACE
+
+#endif // WITH_SLATE_DEBUGGING
diff --git a/Engine/Source/Runtime/SlateCore/Private/Debugging/ConsoleSlateDebugger.h b/Engine/Source/Runtime/SlateCore/Private/Debugging/ConsoleSlateDebugger.h
--- a/Engine/Source/Runtime/SlateCore/Private/Debugging/ConsoleSlateDebugger.h
+++ b/Engine/Source/Runtime/SlateCore/Private/Debugging/ConsoleSlateDebugger.h
@@ -2,10 +2,14 @@
 
 #pragma once
 
+
+#include "Debugging/SlateDebugging.h"
+
+#if WITH_SLATE_DEBUGGING
+
 #include "CoreMinimal.h"
 #include "Delegates/Delegate.h"
 #include "Input/Reply.h"
-#include "Debugging/SlateDebugging.h"
 #include "HAL/IConsoleManager.h"
 
 /**
@@ -56,4 +60,6 @@ private:
 	FAutoConsoleCommand StopDebuggingCommand;
 	FAutoConsoleVariableRef CaptureStackVariable;
 	FAutoConsoleCommand SetInputFilterCommand;
-};
\ No newline at end of file
+};
+
+#endif //WITH_SLATE_DEBUGGING
\ No newline at end of file
