Save 64 bytes per widget by removing tmap of pointer events and replacing with array.

Remove overlapping events on SImage that exist on SWidget

diff --git a/Engine/Source/Runtime/Slate/Private/Widgets/Images/SImage.cpp b/Engine/Source/Runtime/Slate/Private/Widgets/Images/SImage.cpp
--- a/Engine/Source/Runtime/Slate/Private/Widgets/Images/SImage.cpp
+++ b/Engine/Source/Runtime/Slate/Private/Widgets/Images/SImage.cpp
@@ -7,7 +7,8 @@ void SImage::Construct( const FArguments& InArgs )
 {
 	Image = InArgs._Image;
 	ColorAndOpacity = InArgs._ColorAndOpacity;
-	OnMouseButtonDownHandler = InArgs._OnMouseButtonDown;
+
+	SetOnMouseButtonDown(InArgs._OnMouseButtonDown);
 }
 
 int32 SImage::OnPaint( const FPaintArgs& Args, const FGeometry& AllottedGeometry, const FSlateRect& MyCullingRect, FSlateWindowElementList& OutDrawElements, int32 LayerId, const FWidgetStyle& InWidgetStyle, bool bParentEnabled ) const
@@ -26,20 +27,6 @@ int32 SImage::OnPaint( const FPaintArgs& Args, const FGeometry& AllottedGeometry
 	return LayerId;
 }
 
-FReply SImage::OnMouseButtonDown( const FGeometry& MyGeometry, const FPointerEvent& MouseEvent )
-{
-	if ( OnMouseButtonDownHandler.IsBound() )
-	{
-		// If a handler is assigned, call it.
-		return OnMouseButtonDownHandler.Execute(MyGeometry, MouseEvent);
-	}
-	else
-	{
-		// otherwise the event is unhandled.
-		return FReply::Unhandled();
-	}	
-}
-
 FVector2D SImage::ComputeDesiredSize( float ) const
 {
 	const FSlateBrush* ImageBrush = Image.Get();
@@ -76,8 +63,3 @@ void SImage::SetImage(TAttribute<const FSlateBrush*> InImage)
 		Invalidate(EInvalidateWidget::LayoutAndVolatility);
 	}
 }
-
-void SImage::SetOnMouseButtonDown(FPointerEventHandler EventHandler)
-{
-	OnMouseButtonDownHandler = EventHandler;
-}
diff --git a/Engine/Source/Runtime/Slate/Private/Widgets/Images/SSpinningImage.cpp b/Engine/Source/Runtime/Slate/Private/Widgets/Images/SSpinningImage.cpp
--- a/Engine/Source/Runtime/Slate/Private/Widgets/Images/SSpinningImage.cpp
+++ b/Engine/Source/Runtime/Slate/Private/Widgets/Images/SSpinningImage.cpp
@@ -8,7 +8,11 @@ void SSpinningImage::Construct(const FArguments& InArgs)
 {
 	Image = InArgs._Image;
 	ColorAndOpacity = InArgs._ColorAndOpacity;
-	OnMouseButtonDownHandler = InArgs._OnMouseButtonDown;
+
+	if (InArgs._OnMouseButtonDown.IsBound())
+	{
+		SetOnMouseButtonDown(InArgs._OnMouseButtonDown);
+	}
 	
 	SpinAnimationSequence = FCurveSequence( 0.f, InArgs._Period );
 	SpinAnimationSequence.Play( this->AsShared(), true );
diff --git a/Engine/Source/Runtime/Slate/Private/Widgets/Layout/SBorder.cpp b/Engine/Source/Runtime/Slate/Private/Widgets/Layout/SBorder.cpp
--- a/Engine/Source/Runtime/Slate/Private/Widgets/Layout/SBorder.cpp
+++ b/Engine/Source/Runtime/Slate/Private/Widgets/Layout/SBorder.cpp
@@ -30,10 +30,26 @@ void SBorder::Construct( const SBorder::FArguments& InArgs )
 	BorderImage = InArgs._BorderImage;
 	BorderBackgroundColor = InArgs._BorderBackgroundColor;
 	ForegroundColor = InArgs._ForegroundColor;
-	SetOnMouseButtonDown(InArgs._OnMouseButtonDown);
-	SetOnMouseButtonUp(InArgs._OnMouseButtonUp);
-	SetOnMouseMove(InArgs._OnMouseMove);
-	SetOnMouseDoubleClick(InArgs._OnMouseDoubleClick);
+
+	if (InArgs._OnMouseButtonDown.IsBound())
+	{
+		SetOnMouseButtonDown(InArgs._OnMouseButtonDown);
+	}
+
+	if (InArgs._OnMouseButtonUp.IsBound())
+	{
+		SetOnMouseButtonUp(InArgs._OnMouseButtonUp);
+	}
+
+	if (InArgs._OnMouseMove.IsBound())
+	{
+		SetOnMouseMove(InArgs._OnMouseMove);
+	}
+
+	if (InArgs._OnMouseDoubleClick.IsBound())
+	{
+		SetOnMouseDoubleClick(InArgs._OnMouseDoubleClick);
+	}
 
 	ChildSlot
 	.HAlign(InArgs._HAlign)
diff --git a/Engine/Source/Runtime/Slate/Public/Widgets/Images/SImage.h b/Engine/Source/Runtime/Slate/Public/Widgets/Images/SImage.h
--- a/Engine/Source/Runtime/Slate/Public/Widgets/Images/SImage.h
+++ b/Engine/Source/Runtime/Slate/Public/Widgets/Images/SImage.h
@@ -20,20 +20,19 @@ class SLATE_API SImage
 	: public SLeafWidget
 {
 public:
-	SLATE_BEGIN_ARGS( SImage )
-		: _Image( FCoreStyle::Get().GetDefaultBrush() )
-		, _ColorAndOpacity( FLinearColor::White )
-		, _OnMouseButtonDown()
-		{}
+	SLATE_BEGIN_ARGS(SImage)
+		: _Image(FCoreStyle::Get().GetDefaultBrush())
+		, _ColorAndOpacity(FLinearColor::White)
+	{}
 
-		/** Image resource */
-		SLATE_ATTRIBUTE( const FSlateBrush*, Image )
+	/** Image resource */
+	SLATE_ATTRIBUTE(const FSlateBrush*, Image)
 
 		/** Color and opacity */
-		SLATE_ATTRIBUTE( FSlateColor, ColorAndOpacity )
+		SLATE_ATTRIBUTE(FSlateColor, ColorAndOpacity)
 
 		/** Invoked when the mouse is pressed in the widget. */
-		SLATE_EVENT( FPointerEventHandler, OnMouseButtonDown )
+		SLATE_EVENT(FPointerEventHandler, OnMouseButtonDown)
 
 	SLATE_END_ARGS()
 
@@ -49,27 +48,23 @@ public:
 	 *
 	 * @param	InArgs	The declaration data for this widget
 	 */
-	void Construct( const FArguments& InArgs );
+	void Construct(const FArguments& InArgs);
 
 public:
 
 	/** See the ColorAndOpacity attribute */
-	void SetColorAndOpacity( const TAttribute<FSlateColor>& InColorAndOpacity );
-	
+	void SetColorAndOpacity(const TAttribute<FSlateColor>& InColorAndOpacity);
+
 	/** See the ColorAndOpacity attribute */
-	void SetColorAndOpacity( FLinearColor InColorAndOpacity );
+	void SetColorAndOpacity(FLinearColor InColorAndOpacity);
 
 	/** See the Image attribute */
 	void SetImage(TAttribute<const FSlateBrush*> InImage);
-	
-	/** See OnMouseButtonDown event */
-	void SetOnMouseButtonDown(FPointerEventHandler EventHandler);
 
 public:
 
 	// SWidget overrides
-	virtual int32 OnPaint( const FPaintArgs& Args, const FGeometry& AllottedGeometry, const FSlateRect& MyCullingRect, FSlateWindowElementList& OutDrawElements, int32 LayerId, const FWidgetStyle& InWidgetStyle, bool bParentEnabled ) const override;
-	virtual FReply OnMouseButtonDown( const FGeometry& MyGeometry, const FPointerEvent& MouseEvent ) override;
+	virtual int32 OnPaint(const FPaintArgs& Args, const FGeometry& AllottedGeometry, const FSlateRect& MyCullingRect, FSlateWindowElementList& OutDrawElements, int32 LayerId, const FWidgetStyle& InWidgetStyle, bool bParentEnabled) const override;
 
 protected:
 	// Begin SWidget overrides.
@@ -83,7 +78,4 @@ protected:
 
 	/** Color and opacity scale for this image */
 	TAttribute<FSlateColor> ColorAndOpacity;
-
-	/** Invoked when the mouse is pressed in the image */
-	FPointerEventHandler OnMouseButtonDownHandler;
 };
diff --git a/Engine/Source/Runtime/SlateCore/Private/Widgets/SWidget.cpp b/Engine/Source/Runtime/SlateCore/Private/Widgets/SWidget.cpp
--- a/Engine/Source/Runtime/SlateCore/Private/Widgets/SWidget.cpp
+++ b/Engine/Source/Runtime/SlateCore/Private/Widgets/SWidget.cpp
@@ -266,7 +266,7 @@ FReply SWidget::OnPreviewMouseButtonDown( const FGeometry& MyGeometry, const FPo
 
 FReply SWidget::OnMouseButtonDown(const FGeometry& MyGeometry, const FPointerEvent& MouseEvent)
 {
-	if ( FPointerEventHandler* Event = PointerEvents.Find(NAME_MouseButtonDown) )
+	if (const FPointerEventHandler* Event = GetPointerEvent(NAME_MouseButtonDown))
 	{
 		if ( Event->IsBound() )
 		{
@@ -278,7 +278,7 @@ FReply SWidget::OnMouseButtonDown(const FGeometry& MyGeometry, const FPointerEve
 
 FReply SWidget::OnMouseButtonUp(const FGeometry& MyGeometry, const FPointerEvent& MouseEvent)
 {
-	if ( FPointerEventHandler* Event = PointerEvents.Find(NAME_MouseButtonUp) )
+	if (const FPointerEventHandler* Event = GetPointerEvent(NAME_MouseButtonUp) )
 	{
 		if ( Event->IsBound() )
 		{
@@ -290,7 +290,7 @@ FReply SWidget::OnMouseButtonUp(const FGeometry& MyGeometry, const FPointerEvent
 
 FReply SWidget::OnMouseMove(const FGeometry& MyGeometry, const FPointerEvent& MouseEvent)
 {
-	if ( FPointerEventHandler* Event = PointerEvents.Find(NAME_MouseMove) )
+	if (const FPointerEventHandler* Event = GetPointerEvent(NAME_MouseMove) )
 	{
 		if ( Event->IsBound() )
 		{
@@ -302,7 +302,7 @@ FReply SWidget::OnMouseMove(const FGeometry& MyGeometry, const FPointerEvent& Mo
 
 FReply SWidget::OnMouseButtonDoubleClick(const FGeometry& MyGeometry, const FPointerEvent& MouseEvent)
 {
-	if ( FPointerEventHandler* Event = PointerEvents.Find(NAME_MouseDoubleClick) )
+	if ( const FPointerEventHandler* Event = GetPointerEvent(NAME_MouseDoubleClick) )
 	{
 		if ( Event->IsBound() )
 		{
@@ -1167,24 +1167,49 @@ void SWidget::ExecuteActiveTimers(double CurrentTime, float DeltaTime)
 	}
 }
 
+const FPointerEventHandler* SWidget::GetPointerEvent(const FName EventName) const
+{
+	auto* Pair = PointerEvents.FindByPredicate([&EventName](const auto& Pair) {return Pair.Key == EventName; });
+	if (Pair)
+	{
+		return &Pair->Value;
+	}
+	return nullptr;
+}
+
+void SWidget::SetPointerEvent(const FName EventName, FPointerEventHandler& InEvent)
+{
+	// Find the event name and if found, replace the delegate
+	auto* Pair = PointerEvents.FindByPredicate([&EventName](const auto& Pair) {return Pair.Key == EventName; });
+	if (Pair)
+	{
+		Pair->Value = InEvent;
+	}
+	else
+	{
+		PointerEvents.Emplace(EventName, InEvent);
+	}
+
+}
+
 void SWidget::SetOnMouseButtonDown(FPointerEventHandler EventHandler)
 {
-	PointerEvents.Add(NAME_MouseButtonDown, EventHandler);
+	SetPointerEvent(NAME_MouseButtonDown, EventHandler);
 }
 
 void SWidget::SetOnMouseButtonUp(FPointerEventHandler EventHandler)
 {
-	PointerEvents.Add(NAME_MouseButtonUp, EventHandler);
+	SetPointerEvent(NAME_MouseButtonUp, EventHandler);
 }
 
 void SWidget::SetOnMouseMove(FPointerEventHandler EventHandler)
 {
-	PointerEvents.Add(NAME_MouseMove, EventHandler);
+	SetPointerEvent(NAME_MouseMove, EventHandler);
 }
 
 void SWidget::SetOnMouseDoubleClick(FPointerEventHandler EventHandler)
 {
-	PointerEvents.Add(NAME_MouseDoubleClick, EventHandler);
+	SetPointerEvent(NAME_MouseDoubleClick, EventHandler);
 }
 
 void SWidget::SetOnMouseEnter(FNoReplyPointerEventHandler EventHandler)
diff --git a/Engine/Source/Runtime/SlateCore/Public/Widgets/SWidget.h b/Engine/Source/Runtime/SlateCore/Public/Widgets/SWidget.h
--- a/Engine/Source/Runtime/SlateCore/Public/Widgets/SWidget.h
+++ b/Engine/Source/Runtime/SlateCore/Public/Widgets/SWidget.h
@@ -1418,6 +1418,8 @@ private:
 	/** Iterates over the active timer handles on the widget and executes them if their interval has elapsed. */
 	void ExecuteActiveTimers(double CurrentTime, float DeltaTime);
 
+	const FPointerEventHandler* GetPointerEvent(const FName EventName) const;
+	void SetPointerEvent(const FName EventName, FPointerEventHandler& InEvent);
 protected:
 	/** Dtor ensures that active timer handles are UnRegistered with the SlateApplication. */
 	virtual ~SWidget();
@@ -1553,7 +1555,7 @@ private:
 
 private:
 	// Events
-	TMap<FName, FPointerEventHandler> PointerEvents;
+	TArray<TPair<FName, FPointerEventHandler>> PointerEvents;
 
 	FNoReplyPointerEventHandler MouseEnterHandler;
 	FSimpleNoReplyPointerEventHandler MouseLeaveHandler;
