Fix android cis

diff --git a/Engine/Source/Runtime/SlateCore/Private/Widgets/SWidget.cpp b/Engine/Source/Runtime/SlateCore/Private/Widgets/SWidget.cpp
--- a/Engine/Source/Runtime/SlateCore/Private/Widgets/SWidget.cpp
+++ b/Engine/Source/Runtime/SlateCore/Private/Widgets/SWidget.cpp
@@ -1169,10 +1169,10 @@ void SWidget::ExecuteActiveTimers(double CurrentTime, float DeltaTime)
 
 const FPointerEventHandler* SWidget::GetPointerEvent(const FName EventName) const
 {
-	auto* Pair = PointerEvents.FindByPredicate([&EventName](const auto& Pair) {return Pair.Key == EventName; });
-	if (Pair)
+	auto* FoundPair = PointerEvents.FindByPredicate([&EventName](const auto& TestPair) {return TestPair.Key == EventName; });
+	if (FoundPair)
 	{
-		return &Pair->Value;
+		return &FoundPair->Value;
 	}
 	return nullptr;
 }
@@ -1180,10 +1180,10 @@ const FPointerEventHandler* SWidget::GetPointerEvent(const FName EventName) cons
 void SWidget::SetPointerEvent(const FName EventName, FPointerEventHandler& InEvent)
 {
 	// Find the event name and if found, replace the delegate
-	auto* Pair = PointerEvents.FindByPredicate([&EventName](const auto& Pair) {return Pair.Key == EventName; });
-	if (Pair)
+	auto* FoundPair = PointerEvents.FindByPredicate([&EventName](const auto& TestPair) {return TestPair.Key == EventName; });
+	if (FoundPair)
 	{
-		Pair->Value = InEvent;
+		FoundPair->Value = InEvent;
 	}
 	else
 	{
