Moved SImage to SlateCore.  SWindow needs SImage and is accessing one through an interface and attribute binding instead of directly accessing it.  This is causing unecessary volatility for the entire SWindow

diff --git a/Engine/Source/Runtime/Slate/Private/Widgets/Images/SImage.cpp b/Engine/Source/Runtime/Slate/Private/Widgets/Images/SImage.cpp
deleted file mode 100644
diff --git a/Engine/Source/Runtime/Slate/Public/Widgets/Images/SImage.h b/Engine/Source/Runtime/Slate/Public/Widgets/Images/SImage.h
deleted file mode 100644
diff --git a/Engine/Source/Runtime/SlateCore/Private/Widgets/Images/SImage.cpp b/Engine/Source/Runtime/SlateCore/Private/Widgets/Images/SImage.cpp
new file mode 100644
--- /dev/null
+++ b/Engine/Source/Runtime/SlateCore/Private/Widgets/Images/SImage.cpp
@@ -0,0 +1,65 @@
+// Copyright 1998-2019 Epic Games, Inc. All Rights Reserved.
+
+#include "Widgets/Images/SImage.h"
+#include "Rendering/DrawElements.h"
+
+void SImage::Construct( const FArguments& InArgs )
+{
+	Image = InArgs._Image;
+	ColorAndOpacity = InArgs._ColorAndOpacity;
+
+	SetOnMouseButtonDown(InArgs._OnMouseButtonDown);
+}
+
+int32 SImage::OnPaint( const FPaintArgs& Args, const FGeometry& AllottedGeometry, const FSlateRect& MyCullingRect, FSlateWindowElementList& OutDrawElements, int32 LayerId, const FWidgetStyle& InWidgetStyle, bool bParentEnabled ) const
+{
+	const FSlateBrush* ImageBrush = Image.Get();
+
+	if ((ImageBrush != nullptr) && (ImageBrush->DrawAs != ESlateBrushDrawType::NoDrawType))
+	{
+		const bool bIsEnabled = ShouldBeEnabled(bParentEnabled);
+		const ESlateDrawEffect DrawEffects = bIsEnabled ? ESlateDrawEffect::None : ESlateDrawEffect::DisabledEffect;
+
+		const FLinearColor FinalColorAndOpacity( InWidgetStyle.GetColorAndOpacityTint() * ColorAndOpacity.Get().GetColor(InWidgetStyle) * ImageBrush->GetTint( InWidgetStyle ) );
+
+		FSlateDrawElement::MakeBox(OutDrawElements, LayerId, AllottedGeometry.ToPaintGeometry(), ImageBrush, DrawEffects, FinalColorAndOpacity );
+	}
+	return LayerId;
+}
+
+FVector2D SImage::ComputeDesiredSize( float ) const
+{
+	const FSlateBrush* ImageBrush = Image.Get();
+	if (ImageBrush != nullptr)
+	{
+		return ImageBrush->ImageSize;
+	}
+	return FVector2D::ZeroVector;
+}
+
+void SImage::SetColorAndOpacity( const TAttribute<FSlateColor>& InColorAndOpacity )
+{
+	if (!ColorAndOpacity.IdenticalTo(InColorAndOpacity))
+	{
+		ColorAndOpacity = InColorAndOpacity;
+		Invalidate(EInvalidateWidget::PaintAndVolatility);
+	}
+}
+
+void SImage::SetColorAndOpacity( FLinearColor InColorAndOpacity )
+{
+	if (!ColorAndOpacity.IdenticalTo(InColorAndOpacity))
+	{
+		ColorAndOpacity = InColorAndOpacity;
+		Invalidate(EInvalidateWidget::PaintAndVolatility);
+	}
+}
+
+void SImage::SetImage(TAttribute<const FSlateBrush*> InImage)
+{
+	if (!Image.IdenticalTo(InImage))
+	{
+		Image = InImage;
+		Invalidate(EInvalidateWidget::LayoutAndVolatility);
+	}
+}
diff --git a/Engine/Source/Runtime/SlateCore/Public/Widgets/Images/SImage.h b/Engine/Source/Runtime/SlateCore/Public/Widgets/Images/SImage.h
new file mode 100644
--- /dev/null
+++ b/Engine/Source/Runtime/SlateCore/Public/Widgets/Images/SImage.h
@@ -0,0 +1,81 @@
+// Copyright 1998-2019 Epic Games, Inc. All Rights Reserved.
+
+#pragma once
+
+#include "CoreMinimal.h"
+#include "Misc/Attribute.h"
+#include "Input/Reply.h"
+#include "Widgets/DeclarativeSyntaxSupport.h"
+#include "Styling/SlateColor.h"
+#include "Styling/CoreStyle.h"
+#include "Widgets/SLeafWidget.h"
+
+class FPaintArgs;
+class FSlateWindowElementList;
+
+/**
+ * Implements a widget that displays an image with a desired width and height.
+ */
+class SLATECORE_API SImage
+	: public SLeafWidget
+{
+public:
+	SLATE_BEGIN_ARGS(SImage)
+		: _Image(FCoreStyle::Get().GetDefaultBrush())
+		, _ColorAndOpacity(FLinearColor::White)
+	{}
+
+	/** Image resource */
+	SLATE_ATTRIBUTE(const FSlateBrush*, Image)
+
+		/** Color and opacity */
+		SLATE_ATTRIBUTE(FSlateColor, ColorAndOpacity)
+
+		/** Invoked when the mouse is pressed in the widget. */
+		SLATE_EVENT(FPointerEventHandler, OnMouseButtonDown)
+
+	SLATE_END_ARGS()
+
+	/** Constructor */
+	SImage()
+	{
+		SetCanTick(false);
+		bCanSupportFocus = false;
+	}
+
+	/**
+	 * Construct this widget
+	 *
+	 * @param	InArgs	The declaration data for this widget
+	 */
+	void Construct(const FArguments& InArgs);
+
+public:
+
+	/** See the ColorAndOpacity attribute */
+	void SetColorAndOpacity(const TAttribute<FSlateColor>& InColorAndOpacity);
+
+	/** See the ColorAndOpacity attribute */
+	void SetColorAndOpacity(FLinearColor InColorAndOpacity);
+
+	/** See the Image attribute */
+	void SetImage(TAttribute<const FSlateBrush*> InImage);
+
+public:
+
+	// SWidget overrides
+	virtual int32 OnPaint(const FPaintArgs& Args, const FGeometry& AllottedGeometry, const FSlateRect& MyCullingRect, FSlateWindowElementList& OutDrawElements, int32 LayerId, const FWidgetStyle& InWidgetStyle, bool bParentEnabled) const override;
+
+protected:
+	// Begin SWidget overrides.
+	virtual FVector2D ComputeDesiredSize(float) const override;
+	// End SWidget overrides.
+
+protected:
+
+	/** The slate brush to draw for the image, or a bound delegate to a brush. */
+	TAttribute< const FSlateBrush* > Image;
+
+	/** Color and opacity scale for this image */
+	TAttribute<FSlateColor> ColorAndOpacity;
+};
