(Right to Left) Layout Support

Slate - Adding Left To Right | Right To Left support to the flow direction of widgets.  Widgets in both Slate and UMG can now specify a desired flow direction for their children widgets.  This will affect padding, and if we lay out from left to right, or right to left depending on th culture.  By default, widgets will just inherit the state of things, but some widgets may choose to enact the culture's preference in a localized area of the game, so there's no automatic flipping of all your game's UI.  Additionally SBorder and SImage, and their UMG counterparts now have a bFlipForRightToLeftFlowDirection.  Use this if you have images with directionality that you want flipped when the layout flips.

Progress bars now flip Left-Right filling, if inside a culuturally aware and right to left layout.

diff --git a/Engine/Source/Editor/UMGEditor/Private/Customizations/UMGDetailCustomizations.cpp b/Engine/Source/Editor/UMGEditor/Private/Customizations/UMGDetailCustomizations.cpp
--- a/Engine/Source/Editor/UMGEditor/Private/Customizations/UMGDetailCustomizations.cpp
+++ b/Engine/Source/Editor/UMGEditor/Private/Customizations/UMGDetailCustomizations.cpp
@@ -281,6 +281,9 @@ void FBlueprintWidgetCustomization::CreateMulticastEventCustomization(IDetailLay
 void FBlueprintWidgetCustomization::CustomizeDetails( IDetailLayoutBuilder& DetailLayout )
 {
 	static const FName LayoutCategoryKey(TEXT("Layout"));
+	static const FName LocalizationCategoryKey(TEXT("Localization"));
+
+	DetailLayout.EditCategory(LocalizationCategoryKey, FText::GetEmpty(), ECategoryPriority::Uncommon);
 
 	TArray< TWeakObjectPtr<UObject> > OutObjects;
 	DetailLayout.GetObjectsBeingCustomized(OutObjects);
diff --git a/Engine/Source/Runtime/Core/Public/Math/TransformCalculus2D.h b/Engine/Source/Runtime/Core/Public/Math/TransformCalculus2D.h
--- a/Engine/Source/Runtime/Core/Public/Math/TransformCalculus2D.h
+++ b/Engine/Source/Runtime/Core/Public/Math/TransformCalculus2D.h
@@ -657,6 +657,8 @@ public:
 	/** Access to the translation */
 	const FVector2D& GetTranslation() const { return Trans; }
 
+	void SetTranslation(const FVector2D& InTrans) { Trans = InTrans; }
+
 	/**
 	 * Specialized function to determine if a transform is precisely the identity transform. Uses exact float comparison, so rounding error is not considered.
 	 */
diff --git a/Engine/Source/Runtime/Slate/Private/Widgets/Input/SMenuAnchor.cpp b/Engine/Source/Runtime/Slate/Private/Widgets/Input/SMenuAnchor.cpp
--- a/Engine/Source/Runtime/Slate/Private/Widgets/Input/SMenuAnchor.cpp
+++ b/Engine/Source/Runtime/Slate/Private/Widgets/Input/SMenuAnchor.cpp
@@ -174,7 +174,7 @@ bool SMenuAnchor::ComputeVolatility() const
 
 void SMenuAnchor::OnArrangeChildren( const FGeometry& AllottedGeometry, FArrangedChildren& ArrangedChildren ) const
 {
-	ArrangeSingleChild( AllottedGeometry, ArrangedChildren, Children[0], FVector2D::UnitVector );
+	ArrangeSingleChild(GSlateFlowDirection, AllottedGeometry, ArrangedChildren, Children[0], FVector2D::UnitVector);
 	const TSharedPtr<SWindow> PresentingWindow = PopupWindowPtr.Pin();
 	if (IsOpenAndReusingWindow() && PresentingWindow.IsValid())
 	{
diff --git a/Engine/Source/Runtime/Slate/Private/Widgets/Layout/SBorder.cpp b/Engine/Source/Runtime/Slate/Private/Widgets/Layout/SBorder.cpp
--- a/Engine/Source/Runtime/Slate/Private/Widgets/Layout/SBorder.cpp
+++ b/Engine/Source/Runtime/Slate/Private/Widgets/Layout/SBorder.cpp
@@ -27,6 +27,8 @@ void SBorder::Construct( const SBorder::FArguments& InArgs )
 
 	ShowDisabledEffect = InArgs._ShowEffectWhenDisabled;
 
+	bFlipForRightToLeftFlowDirection = InArgs._FlipForRightToLeftFlowDirection;
+
 	BorderImage = InArgs._BorderImage;
 	BorderBackgroundColor = InArgs._BorderBackgroundColor;
 	ForegroundColor = InArgs._ForegroundColor;
@@ -89,14 +91,29 @@ int32 SBorder::OnPaint( const FPaintArgs& Args, const FGeometry& AllottedGeometr
 		const bool bShowDisabledEffect = ShowDisabledEffect.Get();
 		const ESlateDrawEffect DrawEffects = (bShowDisabledEffect && !bEnabled) ? ESlateDrawEffect::DisabledEffect : ESlateDrawEffect::None;
 
-		FSlateDrawElement::MakeBox(
-			OutDrawElements,
-			LayerId,
-			AllottedGeometry.ToPaintGeometry(),
-			BrushResource,
-			DrawEffects,
-			BrushResource->GetTint(InWidgetStyle) * InWidgetStyle.GetColorAndOpacityTint() * BorderBackgroundColor.Get().GetColor(InWidgetStyle)
-		);
+		if (bFlipForRightToLeftFlowDirection && GSlateFlowDirection == EFlowDirection::RightToLeft)
+		{
+			const FGeometry FlippedGeometry = AllottedGeometry.MakeChild(FSlateRenderTransform(FScale2D(-1, 1)));
+			FSlateDrawElement::MakeBox(
+				OutDrawElements,
+				LayerId,
+				FlippedGeometry.ToPaintGeometry(),
+				BrushResource,
+				DrawEffects,
+				BrushResource->GetTint(InWidgetStyle) * InWidgetStyle.GetColorAndOpacityTint() * BorderBackgroundColor.Get().GetColor(InWidgetStyle)
+			);
+		}
+		else
+		{
+			FSlateDrawElement::MakeBox(
+				OutDrawElements,
+				LayerId,
+				AllottedGeometry.ToPaintGeometry(),
+				BrushResource,
+				DrawEffects,
+				BrushResource->GetTint(InWidgetStyle) * InWidgetStyle.GetColorAndOpacityTint() * BorderBackgroundColor.Get().GetColor(InWidgetStyle)
+			);
+		}
 	}
 
 	return SCompoundWidget::OnPaint(Args, AllottedGeometry, MyCullingRect, OutDrawElements, LayerId, InWidgetStyle, bEnabled );
diff --git a/Engine/Source/Runtime/Slate/Private/Widgets/Layout/SWidgetSwitcher.cpp b/Engine/Source/Runtime/Slate/Private/Widgets/Layout/SWidgetSwitcher.cpp
--- a/Engine/Source/Runtime/Slate/Private/Widgets/Layout/SWidgetSwitcher.cpp
+++ b/Engine/Source/Runtime/Slate/Private/Widgets/Layout/SWidgetSwitcher.cpp
@@ -140,7 +140,7 @@ void SWidgetSwitcher::OnArrangeChildren( const FGeometry& AllottedGeometry, FArr
 		const FSlot* ActiveSlotPtr = GetActiveSlot();	// Returns null if unsafe attribute WidgetIndex is out-of-bounds
 		if (ActiveSlotPtr != nullptr)
 		{
-			ArrangeSingleChild(AllottedGeometry, ArrangedChildren, *ActiveSlotPtr, ContentScale);
+			ArrangeSingleChild(GSlateFlowDirection, AllottedGeometry, ArrangedChildren, *ActiveSlotPtr, ContentScale);
 		}
 	}
 }
diff --git a/Engine/Source/Runtime/Slate/Private/Widgets/Notifications/SProgressBar.cpp b/Engine/Source/Runtime/Slate/Private/Widgets/Notifications/SProgressBar.cpp
--- a/Engine/Source/Runtime/Slate/Private/Widgets/Notifications/SProgressBar.cpp
+++ b/Engine/Source/Runtime/Slate/Private/Widgets/Notifications/SProgressBar.cpp
@@ -174,8 +174,22 @@ int32 SProgressBar::OnPaint( const FPaintArgs& Args, const FGeometry& AllottedGe
 	
 	if( ProgressFraction.IsSet() )
 	{
+		EProgressBarFillType::Type ComputedBarFillType = BarFillType;
+		if (GSlateFlowDirection == EFlowDirection::RightToLeft)
+		{
+			switch (ComputedBarFillType)
+			{
+			case EProgressBarFillType::LeftToRight:
+				ComputedBarFillType = EProgressBarFillType::RightToLeft;
+				break;
+			case EProgressBarFillType::RightToLeft:
+				ComputedBarFillType = EProgressBarFillType::LeftToRight;
+				break;
+			}
+		}
+
 		const float ClampedFraction = FMath::Clamp(ProgressFraction.GetValue(), 0.0f, 1.0f);
-		switch (BarFillType)
+		switch (ComputedBarFillType)
 		{
 			case EProgressBarFillType::RightToLeft:
 			{
diff --git a/Engine/Source/Runtime/Slate/Private/Widgets/Views/SHeaderRow.cpp b/Engine/Source/Runtime/Slate/Private/Widgets/Views/SHeaderRow.cpp
--- a/Engine/Source/Runtime/Slate/Private/Widgets/Views/SHeaderRow.cpp
+++ b/Engine/Source/Runtime/Slate/Private/Widgets/Views/SHeaderRow.cpp
@@ -48,7 +48,7 @@ public:
 	{
 		check(InArgs._Style);
 
-		SWidget::Construct( InArgs._ToolTipText, InArgs._ToolTip, InArgs._Cursor, InArgs._IsEnabled, InArgs._Visibility, InArgs._RenderOpacity, InArgs._RenderTransform, InArgs._RenderTransformPivot, InArgs._Tag, InArgs._ForceVolatile, InArgs._Clipping, InArgs.MetaData );
+		SWidget::Construct( InArgs._ToolTipText, InArgs._ToolTip, InArgs._Cursor, InArgs._IsEnabled, InArgs._Visibility, InArgs._RenderOpacity, InArgs._RenderTransform, InArgs._RenderTransformPivot, InArgs._Tag, InArgs._ForceVolatile, InArgs._Clipping, InArgs._FlowDirectionPreference, InArgs.MetaData );
 
 		Style = InArgs._Style;
 		ColumnId = Column.ColumnId;
@@ -414,7 +414,7 @@ void SHeaderRow::Construct( const FArguments& InArgs )
 {
 	check(InArgs._Style);
 
-	SWidget::Construct( InArgs._ToolTipText, InArgs._ToolTip, InArgs._Cursor, InArgs._IsEnabled, InArgs._Visibility, InArgs._RenderOpacity, InArgs._RenderTransform, InArgs._RenderTransformPivot, InArgs._Tag, InArgs._ForceVolatile, InArgs._Clipping, InArgs.MetaData );
+	SWidget::Construct( InArgs._ToolTipText, InArgs._ToolTip, InArgs._Cursor, InArgs._IsEnabled, InArgs._Visibility, InArgs._RenderOpacity, InArgs._RenderTransform, InArgs._RenderTransformPivot, InArgs._Tag, InArgs._ForceVolatile, InArgs._Clipping, InArgs._FlowDirectionPreference, InArgs.MetaData );
 
 	ScrollBarThickness = FVector2D::ZeroVector;
 	ScrollBarVisibility = EVisibility::Collapsed;
diff --git a/Engine/Source/Runtime/Slate/Public/Widgets/Layout/SBorder.h b/Engine/Source/Runtime/Slate/Public/Widgets/Layout/SBorder.h
--- a/Engine/Source/Runtime/Slate/Public/Widgets/Layout/SBorder.h
+++ b/Engine/Source/Runtime/Slate/Public/Widgets/Layout/SBorder.h
@@ -40,6 +40,7 @@ public:
 		, _BorderBackgroundColor( FLinearColor::White )
 		, _ForegroundColor( FSlateColor::UseForeground() )
 		, _ShowEffectWhenDisabled( true )
+		, _FlipForRightToLeftFlowDirection(false)
 		{}
 
 		SLATE_DEFAULT_SLOT( FArguments, Content )
@@ -67,6 +68,9 @@ public:
 		SLATE_ATTRIBUTE( FSlateColor, ForegroundColor )
 		/** Whether or not to show the disabled effect when this border is disabled */
 		SLATE_ATTRIBUTE( bool, ShowEffectWhenDisabled )
+
+		/** Flips the background image if the localization's flow direction is RightToLeft */
+		SLATE_ARGUMENT(bool, FlipForRightToLeftFlowDirection)
 	SLATE_END_ARGS()
 
 	/**
@@ -136,4 +140,7 @@ protected:
 	TAttribute<FVector2D> DesiredSizeScale;
 	/** Whether or not to show the disabled effect when this border is disabled */
 	TAttribute<bool> ShowDisabledEffect;
+
+	/** Flips the image if the localization's flow direction is RightToLeft */
+	bool bFlipForRightToLeftFlowDirection;
  };
diff --git a/Engine/Source/Runtime/SlateCore/Private/Layout/FlowDirection.cpp b/Engine/Source/Runtime/SlateCore/Private/Layout/FlowDirection.cpp
new file mode 100644
--- /dev/null
+++ b/Engine/Source/Runtime/SlateCore/Private/Layout/FlowDirection.cpp
@@ -0,0 +1,43 @@
+// Copyright 1998-2019 Epic Games, Inc. All Rights Reserved.
+
+#include "Layout/FlowDirection.h"
+#include "Internationalization/Internationalization.h"
+#include "Internationalization/Culture.h"
+#include "HAL/IConsoleManager.h"
+#include "Misc/FrameValue.h"
+
+EFlowDirection GSlateFlowDirection = EFlowDirection::LeftToRight;
+
+static int32 GSlateEnableLayoutLocalization = 1;
+static FAutoConsoleVariableRef CVarSlateEnableLayoutLocalization(
+	TEXT("Slate.EnableLayoutLocalization"),
+	GSlateEnableLayoutLocalization,
+	TEXT("Controls if we enable or disable localized layout, which affects left to right or right to left detection for cultures."),
+	ECVF_Default
+);
+
+int32 GSlateFlowDirectionShouldFollowCultureByDefault = 0;
+static FAutoConsoleVariableRef CVarSlateFlowDirectionShouldFollowCultureByDefault(
+	TEXT("Slate.ShouldFollowCultureByDefault"),
+	GSlateFlowDirectionShouldFollowCultureByDefault,
+	TEXT("Should we initially follow the culture's flow direction at the window level."),
+	ECVF_Default
+);
+
+EFlowDirection FLayoutLocalization::GetLocalizedLayoutDirection()
+{
+	if (!GSlateEnableLayoutLocalization)
+	{
+		return EFlowDirection::LeftToRight;
+	}
+
+	// Only do this once per frame max no matter how many times people call the function.
+	static TFrameValue<EFlowDirection> FrameCoherentDirection;
+	if (!FrameCoherentDirection.IsSet())
+	{
+		//HACK: Normally we'd get this from the culture, but our cultures don't tell us if they're RightToLeft.  Newer versions of ICU tell you, but we need to upgrade.
+		FrameCoherentDirection = FInternationalization::Get().GetCurrentCulture()->GetTwoLetterISOLanguageName().StartsWith(TEXT("ar")) ? EFlowDirection::RightToLeft : EFlowDirection::LeftToRight;
+	}
+
+	return FrameCoherentDirection.GetValue();
+}
diff --git a/Engine/Source/Runtime/SlateCore/Private/Widgets/Images/SImage.cpp b/Engine/Source/Runtime/SlateCore/Private/Widgets/Images/SImage.cpp
--- a/Engine/Source/Runtime/SlateCore/Private/Widgets/Images/SImage.cpp
+++ b/Engine/Source/Runtime/SlateCore/Private/Widgets/Images/SImage.cpp
@@ -7,8 +7,8 @@ void SImage::Construct( const FArguments& InArgs )
 {
 	Image = InArgs._Image;
 	ColorAndOpacity = InArgs._ColorAndOpacity;
-
-	SetOnMouseButtonDown(InArgs._OnMouseButtonDown);
+	bFlipForRightToLeftFlowDirection = InArgs._FlipForRightToLeftFlowDirection;
+	OnMouseButtonDownHandler = InArgs._OnMouseButtonDown;
 }
 
 int32 SImage::OnPaint( const FPaintArgs& Args, const FGeometry& AllottedGeometry, const FSlateRect& MyCullingRect, FSlateWindowElementList& OutDrawElements, int32 LayerId, const FWidgetStyle& InWidgetStyle, bool bParentEnabled ) const
@@ -22,11 +22,34 @@ int32 SImage::OnPaint( const FPaintArgs& Args, const FGeometry& AllottedGeometry
 
 		const FLinearColor FinalColorAndOpacity( InWidgetStyle.GetColorAndOpacityTint() * ColorAndOpacity.Get().GetColor(InWidgetStyle) * ImageBrush->GetTint( InWidgetStyle ) );
 
-		FSlateDrawElement::MakeBox(OutDrawElements, LayerId, AllottedGeometry.ToPaintGeometry(), ImageBrush, DrawEffects, FinalColorAndOpacity );
+		if (bFlipForRightToLeftFlowDirection && GSlateFlowDirection == EFlowDirection::RightToLeft)
+		{
+			const FGeometry FlippedGeometry = AllottedGeometry.MakeChild(FSlateRenderTransform(FScale2D(-1, 1)));
+			FSlateDrawElement::MakeBox(OutDrawElements, LayerId, FlippedGeometry.ToPaintGeometry(), ImageBrush, DrawEffects, FinalColorAndOpacity);
+		}
+		else
+		{
+			FSlateDrawElement::MakeBox(OutDrawElements, LayerId, AllottedGeometry.ToPaintGeometry(), ImageBrush, DrawEffects, FinalColorAndOpacity);
+		}
 	}
+
 	return LayerId;
 }
 
+FReply SImage::OnMouseButtonDown( const FGeometry& MyGeometry, const FPointerEvent& MouseEvent )
+{
+	if ( OnMouseButtonDownHandler.IsBound() )
+	{
+		// If a handler is assigned, call it.
+		return OnMouseButtonDownHandler.Execute(MyGeometry, MouseEvent);
+	}
+	else
+	{
+		// otherwise the event is unhandled.
+		return FReply::Unhandled();
+	}	
+}
+
 FVector2D SImage::ComputeDesiredSize( float ) const
 {
 	const FSlateBrush* ImageBrush = Image.Get();
@@ -63,3 +86,8 @@ void SImage::SetImage(TAttribute<const FSlateBrush*> InImage)
 		Invalidate(EInvalidateWidget::LayoutAndVolatility);
 	}
 }
+
+void SImage::SetOnMouseButtonDown(FPointerEventHandler EventHandler)
+{
+	OnMouseButtonDownHandler = EventHandler;
+}
diff --git a/Engine/Source/Runtime/SlateCore/Private/Widgets/SBoxPanel.cpp b/Engine/Source/Runtime/SlateCore/Private/Widgets/SBoxPanel.cpp
--- a/Engine/Source/Runtime/SlateCore/Private/Widgets/SBoxPanel.cpp
+++ b/Engine/Source/Runtime/SlateCore/Private/Widgets/SBoxPanel.cpp
@@ -1,4 +1,4 @@
-﻿// Copyright 1998-2019 Epic Games, Inc. All Rights Reserved.
+// Copyright 1998-2018 Epic Games, Inc. All Rights Reserved.
 
 #include "Widgets/SBoxPanel.h"
 #include "Layout/LayoutUtils.h"
@@ -34,7 +34,7 @@ void SVerticalBox::Construct( const SVerticalBox::FArguments& InArgs )
 }
 
 template<EOrientation Orientation>
-static void ArrangeChildrenAlong( const TPanelChildren<SBoxPanel::FSlot>& Children, const FGeometry& AllottedGeometry, FArrangedChildren& ArrangedChildren )
+static void ArrangeChildrenAlong(EFlowDirection InLayoutFlow, const TPanelChildren<SBoxPanel::FSlot>& Children, const FGeometry& AllottedGeometry, FArrangedChildren& ArrangedChildren )
 {
 	// Allotted space will be given to fixed-size children first.
 	// Remaining space will be proportionately divided between stretch children (SizeRule_Stretch)
@@ -74,7 +74,7 @@ static void ArrangeChildrenAlong( const TPanelChildren<SBoxPanel::FSlot>& Childr
 
 					// Clamp to the max size if it was specified
 					float MaxSize = CurChild.MaxSize.Get();
-					FixedTotal += MaxSize > 0 ? FMath::Min( MaxSize, ChildSize ) : ChildSize ;
+					FixedTotal += MaxSize > 0 ? FMath::Min( MaxSize, ChildSize ) : ChildSize;
 				}
 			}
 		}
@@ -93,9 +93,9 @@ static void ArrangeChildrenAlong( const TPanelChildren<SBoxPanel::FSlot>& Childr
 
 		// Now that we have the total fixed-space requirement and the total stretch coefficients we can
 		// arrange widgets top-to-bottom or left-to-right (depending on the orientation).
-		for( int32 ChildIndex=0; ChildIndex < Children.Num(); ++ChildIndex )
+		for (TPanelChildrenConstIterator<SBoxPanel::FSlot> It(Children, Orientation, InLayoutFlow); It; ++It)
 		{
-			const SBoxPanel::FSlot& CurChild = Children[ChildIndex];
+			const SBoxPanel::FSlot& CurChild = *It;
 			const EVisibility ChildVisibility = CurChild.GetWidget()->GetVisibility();
 
 			// Figure out the area allocated to the child in the direction of BoxPanel
@@ -130,14 +130,14 @@ static void ArrangeChildrenAlong( const TPanelChildren<SBoxPanel::FSlot>& Childr
 				}
 			}
 
-			const FMargin SlotPadding(CurChild.SlotPadding.Get());
+			const FMargin SlotPadding(LayoutPaddingWithFlow(InLayoutFlow, CurChild.SlotPadding.Get()));
 
 			FVector2D SlotSize = (Orientation == Orient_Vertical)
 				? FVector2D(AllottedGeometry.GetLocalSize().X, ChildSize + SlotPadding.GetTotalSpaceAlong<Orient_Vertical>())
 				: FVector2D(ChildSize + SlotPadding.GetTotalSpaceAlong<Orient_Horizontal>(), AllottedGeometry.GetLocalSize().Y);
 
 			// Figure out the size and local position of the child within the slot			
-			AlignmentArrangeResult XAlignmentResult = AlignChild<Orient_Horizontal>( SlotSize.X, CurChild, SlotPadding);
+			AlignmentArrangeResult XAlignmentResult = AlignChild<Orient_Horizontal>(InLayoutFlow, SlotSize.X, CurChild, SlotPadding);
 			AlignmentArrangeResult YAlignmentResult = AlignChild<Orient_Vertical>(SlotSize.Y, CurChild, SlotPadding);
 
 			const FVector2D LocalPosition = (Orientation == Orient_Vertical)
@@ -176,11 +176,11 @@ void SBoxPanel::OnArrangeChildren( const FGeometry& AllottedGeometry, FArrangedC
 {
 	if ( this->Orientation == Orient_Horizontal )
 	{
-		ArrangeChildrenAlong<Orient_Horizontal>(this->Children, AllottedGeometry, ArrangedChildren );
+		ArrangeChildrenAlong<Orient_Horizontal>(GSlateFlowDirection, this->Children, AllottedGeometry, ArrangedChildren );
 	}
 	else
 	{
-		ArrangeChildrenAlong<Orient_Vertical>(this->Children, AllottedGeometry, ArrangedChildren );
+		ArrangeChildrenAlong<Orient_Vertical>(GSlateFlowDirection, this->Children, AllottedGeometry, ArrangedChildren );
 	}
 }
 
diff --git a/Engine/Source/Runtime/SlateCore/Private/Widgets/SCompoundWidget.cpp b/Engine/Source/Runtime/SlateCore/Private/Widgets/SCompoundWidget.cpp
--- a/Engine/Source/Runtime/SlateCore/Private/Widgets/SCompoundWidget.cpp
+++ b/Engine/Source/Runtime/SlateCore/Private/Widgets/SCompoundWidget.cpp
@@ -58,7 +58,7 @@ FVector2D SCompoundWidget::ComputeDesiredSize( float ) const
 
 void SCompoundWidget::OnArrangeChildren( const FGeometry& AllottedGeometry, FArrangedChildren& ArrangedChildren ) const
 {
-	ArrangeSingleChild( AllottedGeometry, ArrangedChildren, ChildSlot, ContentScale );
+	ArrangeSingleChild(GSlateFlowDirection, AllottedGeometry, ArrangedChildren, ChildSlot, ContentScale);
 }
 
 FSlateColor SCompoundWidget::GetForegroundColor() const
diff --git a/Engine/Source/Runtime/SlateCore/Private/Widgets/SOverlay.cpp b/Engine/Source/Runtime/SlateCore/Private/Widgets/SOverlay.cpp
--- a/Engine/Source/Runtime/SlateCore/Private/Widgets/SOverlay.cpp
+++ b/Engine/Source/Runtime/SlateCore/Private/Widgets/SOverlay.cpp
@@ -30,8 +30,8 @@ void SOverlay::OnArrangeChildren( const FGeometry& AllottedGeometry, FArrangedCh
 		const EVisibility ChildVisibility = CurChild.GetWidget()->GetVisibility();
 		if ( ArrangedChildren.Accepts(ChildVisibility) )
 		{
-			const FMargin SlotPadding(CurChild.SlotPadding.Get());
-			AlignmentArrangeResult XResult = AlignChild<Orient_Horizontal>(AllottedGeometry.GetLocalSize().X, CurChild, SlotPadding);
+			const FMargin SlotPadding(LayoutPaddingWithFlow(GSlateFlowDirection, CurChild.SlotPadding.Get()));
+			AlignmentArrangeResult XResult = AlignChild<Orient_Horizontal>(GSlateFlowDirection, AllottedGeometry.GetLocalSize().X, CurChild, SlotPadding);
 			AlignmentArrangeResult YResult = AlignChild<Orient_Vertical>(AllottedGeometry.GetLocalSize().Y, CurChild, SlotPadding);
 
 			ArrangedChildren.AddWidget( ChildVisibility, AllottedGeometry.MakeChild(
diff --git a/Engine/Source/Runtime/SlateCore/Private/Widgets/SWidget.cpp b/Engine/Source/Runtime/SlateCore/Private/Widgets/SWidget.cpp
--- a/Engine/Source/Runtime/SlateCore/Private/Widgets/SWidget.cpp
+++ b/Engine/Source/Runtime/SlateCore/Private/Widgets/SWidget.cpp
@@ -112,6 +112,7 @@ SWidget::SWidget()
 	, bNeedsDesiredSize(true)
 	, bUpdatingDesiredSize(false)
 	, Clipping(EWidgetClipping::Inherit)
+	, FlowDirectionPreference(EFlowDirectionPreference::Inherit)
 	, UpdateFlags(EWidgetUpdateFlags::NeedsTick)
 	, DesiredSize()
 	, PrepassLayoutScaleMultiplier(1.0f)
@@ -159,6 +160,7 @@ void SWidget::Construct(
 	const FName& InTag,
 	const bool InForceVolatile,
 	const EWidgetClipping InClipping,
+	const EFlowDirectionPreference InFlowPreference,
 	const TArray<TSharedRef<ISlateMetaData>>& InMetaData
 )
 {
@@ -187,12 +189,13 @@ void SWidget::Construct(
 	Tag = InTag;
 	bForceVolatile = InForceVolatile;
 	Clipping = InClipping;
+	FlowDirectionPreference = InFlowPreference;
 	MetaData = InMetaData;
 }
 
-void SWidget::SWidgetConstruct(const TAttribute<FText>& InToolTipText, const TSharedPtr<IToolTip>& InToolTip, const TAttribute< TOptional<EMouseCursor::Type> >& InCursor, const TAttribute<bool>& InEnabledState, const TAttribute<EVisibility>& InVisibility, const float InRenderOpacity, const TAttribute<TOptional<FSlateRenderTransform>>& InTransform, const TAttribute<FVector2D>& InTransformPivot, const FName& InTag, const bool InForceVolatile, const EWidgetClipping InClipping, const TArray<TSharedRef<ISlateMetaData>>& InMetaData)
+void SWidget::SWidgetConstruct(const TAttribute<FText>& InToolTipText, const TSharedPtr<IToolTip>& InToolTip, const TAttribute< TOptional<EMouseCursor::Type> >& InCursor, const TAttribute<bool>& InEnabledState, const TAttribute<EVisibility>& InVisibility, const float InRenderOpacity, const TAttribute<TOptional<FSlateRenderTransform>>& InTransform, const TAttribute<FVector2D>& InTransformPivot, const FName& InTag, const bool InForceVolatile, const EWidgetClipping InClipping, const EFlowDirectionPreference InFlowPreference, const TArray<TSharedRef<ISlateMetaData>>& InMetaData)
 {
-	Construct(InToolTipText, InToolTip, InCursor, InEnabledState, InVisibility, InRenderOpacity, InTransform, InTransformPivot, InTag, InForceVolatile, InClipping, InMetaData);
+	Construct(InToolTipText, InToolTip, InCursor, InEnabledState, InVisibility, InRenderOpacity, InTransform, InTransformPivot, InTag, InForceVolatile, InClipping, InFlowPreference, InMetaData);
 }
 
 FReply SWidget::OnFocusReceived(const FGeometry& MyGeometry, const FFocusEvent& InFocusEvent)
@@ -1029,6 +1032,11 @@ int32 SWidget::Paint(const FPaintArgs& Args, const FGeometry& AllottedGeometry,
 	FSlateDebugging::BeginWidgetPaint.Broadcast(this, UpdatedArgs, AllottedGeometry, CullingBounds, OutDrawElements, LayerId);
 #endif
 
+	// Establish the flow direction if we're changing from inherit.
+	// FOR RB mode, this should first set GSlateFlowDirection to the incoming state that was cached for the widget, then paint
+	// will override it here to reflow is needed.
+	TGuardValue<EFlowDirection> FlowGuard(GSlateFlowDirection, ComputeFlowDirection());
+	
 	// Paint the geometry of this widget.
 	int32 NewLayerID = OnPaint(UpdatedArgs, AllottedGeometry, CullingBounds, OutDrawElements, LayerId, ContentWidgetStyle, bParentEnabled);
 
diff --git a/Engine/Source/Runtime/SlateCore/Private/Widgets/SWindow.cpp b/Engine/Source/Runtime/SlateCore/Private/Widgets/SWindow.cpp
--- a/Engine/Source/Runtime/SlateCore/Private/Widgets/SWindow.cpp
+++ b/Engine/Source/Runtime/SlateCore/Private/Widgets/SWindow.cpp
@@ -1922,12 +1922,20 @@ SWindow::SWindow()
 int32 SWindow::PaintWindow( const FPaintArgs& Args, const FGeometry& AllottedGeometry, const FSlateRect& MyCullingRect, FSlateWindowElementList& OutDrawElements, int32 LayerId, const FWidgetStyle& InWidgetStyle, bool bParentEnabled ) const
 {
 	LayerId = Paint(Args, AllottedGeometry, MyCullingRect, OutDrawElements, LayerId, InWidgetStyle, bParentEnabled);
-	//LayerId = OutDrawElements.PaintDeferred( LayerId );
+
 	return LayerId;
 }
 
 int32 SWindow::OnPaint(const FPaintArgs& Args, const FGeometry& AllottedGeometry, const FSlateRect& MyCullingRect, FSlateWindowElementList& OutDrawElements, int32 LayerId, const FWidgetStyle& InWidgetStyle, bool bParentEnabled) const
 {
+	EFlowDirection NewFlowDirection = GSlateFlowDirection;
+	if (GetFlowDirectionPreference() == EFlowDirectionPreference::Inherit)
+	{
+		NewFlowDirection = GSlateFlowDirectionShouldFollowCultureByDefault ? FLayoutLocalization::GetLocalizedLayoutDirection() : EFlowDirection::LeftToRight;
+	}
+
+	TGuardValue<EFlowDirection> FlowGuard(GSlateFlowDirection, NewFlowDirection);
+
 	OutDrawElements.BeginDeferredGroup();
 	int32 MaxLayer = SCompoundWidget::OnPaint(Args, AllottedGeometry, MyCullingRect, OutDrawElements, LayerId, InWidgetStyle, bParentEnabled);
 	OutDrawElements.EndDeferredGroup();
diff --git a/Engine/Source/Runtime/SlateCore/Public/Layout/Children.h b/Engine/Source/Runtime/SlateCore/Public/Layout/Children.h
--- a/Engine/Source/Runtime/SlateCore/Public/Layout/Children.h
+++ b/Engine/Source/Runtime/SlateCore/Public/Layout/Children.h
@@ -9,6 +9,7 @@
 #include "Widgets/SNullWidget.h"
 #include "SlotBase.h"
 #include "Widgets/SWidget.h"
+#include "FlowDirection.h"
 
 /**
  * FChildren is an interface that must be implemented by all child containers.
@@ -223,6 +224,8 @@ public:
 };
 
 
+
+
 /**
  * A generic FChildren that stores children along with layout-related information.
  * The type containing Widget* and layout info is specified by ChildType.
@@ -350,6 +353,121 @@ private:
 	bool bEmptying;
 };
 
+
+
+template<typename SlotType>
+class TPanelChildrenConstIterator
+{
+public:
+	TPanelChildrenConstIterator(const TPanelChildren<SlotType>& InContainer, EFlowDirection InLayoutFlow)
+		: Container(InContainer)
+		, LayoutFlow(InLayoutFlow)
+	{
+		Reset();
+	}
+
+	TPanelChildrenConstIterator(const TPanelChildren<SlotType>& InContainer, EOrientation InOrientation, EFlowDirection InLayoutFlow)
+		: Container(InContainer)
+		, LayoutFlow(InOrientation == Orient_Vertical ? EFlowDirection::LeftToRight : InLayoutFlow)
+	{
+		Reset();
+	}
+
+	/** Advances iterator to the next element in the container. */
+	TPanelChildrenConstIterator<SlotType>& operator++()
+	{
+		switch (LayoutFlow)
+		{
+		default:
+		case EFlowDirection::LeftToRight:
+			++Index;
+			break;
+		case EFlowDirection::RightToLeft:
+			--Index;
+			break;
+		}
+
+		return *this;
+	}
+
+	/** Moves iterator to the previous element in the container. */
+	TPanelChildrenConstIterator<SlotType>& operator--()
+	{
+		switch (LayoutFlow)
+		{
+		default:
+		case EFlowDirection::LeftToRight:
+			--Index;
+			break;
+		case EFlowDirection::RightToLeft:
+			++Index;
+			break;
+		}
+
+		return *this;
+	}
+
+	const SlotType& operator* () const
+	{
+		return Container[Index];
+	}
+
+	const SlotType* operator->() const
+	{
+		return &Container[Index];
+	}
+
+	/** conversion to "bool" returning true if the iterator has not reached the last element. */
+	FORCEINLINE explicit operator bool() const
+	{
+		return Container.IsValidIndex(Index);
+	}
+
+	/** Returns an index to the current element. */
+	int32 GetIndex() const
+	{
+		return Index;
+	}
+
+	/** Resets the iterator to the first element. */
+	void Reset()
+	{
+		switch (LayoutFlow)
+		{
+		default:
+		case EFlowDirection::LeftToRight:
+			Index = 0;
+			break;
+		case EFlowDirection::RightToLeft:
+			Index = Container.Num() - 1;
+			break;
+		}
+	}
+
+	/** Sets iterator to the last element. */
+	void SetToEnd()
+	{
+		switch (LayoutFlow)
+		{
+		default:
+		case EFlowDirection::LeftToRight:
+			Index = Container.Num() - 1;
+			break;
+		case EFlowDirection::RightToLeft:
+			Index = 0;
+			break;
+		}
+	}
+
+private:
+
+	const TPanelChildren<SlotType>& Container;
+	int32 Index;
+	EFlowDirection LayoutFlow;
+};
+
+
+
 /**
  * Some advanced widgets contain no layout information, and do not require slots.
  * Those widgets may wish to store a specialized type of child widget.
diff --git a/Engine/Source/Runtime/SlateCore/Public/Layout/FlowDirection.h b/Engine/Source/Runtime/SlateCore/Public/Layout/FlowDirection.h
new file mode 100644
--- /dev/null
+++ b/Engine/Source/Runtime/SlateCore/Public/Layout/FlowDirection.h
@@ -0,0 +1,51 @@
+// Copyright 1998-2019 Epic Games, Inc. All Rights Reserved.
+
+#pragma once
+
+#include "CoreMinimal.h"
+#include "UObject/ObjectMacros.h"
+#include "UObject/Class.h"
+#include "FlowDirection.generated.h"
+
+/**
+ * Widgets may need to flow left or right depending upon the current culture/localization that's active.
+ * This enum is used to request a specific layout flow.
+ */
+enum class EFlowDirection : uint8
+{
+	/** Desires content flows using a LTR layout */
+	LeftToRight,
+	/** Desires content flows using a RTL layout */
+	RightToLeft,
+};
+
+
+/**
+ * 
+ */
+UENUM()
+enum class EFlowDirectionPreference : uint8
+{
+	/** */
+	Inherit,
+	/** */
+	Culture,
+	/** */
+	LeftToRight,
+	/**  */
+	RightToLeft
+};
+
+
+extern SLATECORE_API EFlowDirection GSlateFlowDirection;
+extern SLATECORE_API int32 GSlateFlowDirectionShouldFollowCultureByDefault;
+
+/**
+ * Class containing utilities for getting layout localization information.
+ */
+class SLATECORE_API FLayoutLocalization
+{
+public:
+	/** Gets the current expected flow direction based on localization. */
+	static EFlowDirection GetLocalizedLayoutDirection();
+};
diff --git a/Engine/Source/Runtime/SlateCore/Public/Layout/Geometry.h b/Engine/Source/Runtime/SlateCore/Public/Layout/Geometry.h
--- a/Engine/Source/Runtime/SlateCore/Public/Layout/Geometry.h
+++ b/Engine/Source/Runtime/SlateCore/Public/Layout/Geometry.h
@@ -229,6 +229,11 @@ public:
 		return FGeometry(InLocalSize, LayoutTransform, GetAccumulatedLayoutTransform(), GetAccumulatedRenderTransform(), bHasRenderTransform);
 	}
 
+	FORCEINLINE_DEBUGGABLE FGeometry MakeChild(const FSlateRenderTransform& RenderTransform, const FVector2D& RenderTransformPivot = FVector2D(0.5f, 0.5)) const
+	{
+		return FGeometry(GetLocalSize(), FSlateLayoutTransform(), RenderTransform, RenderTransformPivot, GetAccumulatedLayoutTransform(), GetAccumulatedRenderTransform());
+	}
+
 	/**
 	 * Create a child geometry+widget relative to this one using the given LayoutGeometry.
 	 *
diff --git a/Engine/Source/Runtime/SlateCore/Public/Layout/LayoutUtils.h b/Engine/Source/Runtime/SlateCore/Public/Layout/LayoutUtils.h
--- a/Engine/Source/Runtime/SlateCore/Public/Layout/LayoutUtils.h
+++ b/Engine/Source/Runtime/SlateCore/Public/Layout/LayoutUtils.h
@@ -22,21 +22,36 @@ struct AlignmentArrangeResult
 
 namespace ArrangeUtils
 {
-	/** Gets the alignment of an axis-agnostic int32 so that we can do aligment on an axis without caring about its orientation */
+	/** Gets the alignment of an axis-agnostic int32 so that we can do alignment on an axis without caring about its orientation */
 	template<EOrientation Orientation>
 	struct GetChildAlignment
 	{
 		template<typename SlotType>
-		static int32 AsInt( const SlotType& InSlot );
+		static int32 AsInt(EFlowDirection InFlowDirection, const SlotType& InSlot );
 	};
 
 	template<>
 	struct GetChildAlignment<Orient_Horizontal>
 	{
 		template<typename SlotType>
-		static int32 AsInt( const SlotType& InSlot )
+		static int32 AsInt(EFlowDirection InFlowDirection, const SlotType& InSlot )
 		{
-			return static_cast<int32>(InSlot.HAlignment);
+			switch (InFlowDirection)
+			{
+			default:
+			case EFlowDirection::LeftToRight:
+				return static_cast<int32>(InSlot.HAlignment);
+			case EFlowDirection::RightToLeft:
+				switch (InSlot.HAlignment)
+				{
+				case HAlign_Left:
+					return static_cast<int32>(HAlign_Right);
+				case HAlign_Right:
+					return static_cast<int32>(HAlign_Left);
+				default:
+					return static_cast<int32>(InSlot.HAlignment);
+				}
+			}
 		}
 	};
 
@@ -44,15 +59,16 @@ namespace ArrangeUtils
 	struct GetChildAlignment<Orient_Vertical>
 	{
 		template<typename SlotType>
-		static int32 AsInt( const SlotType& InSlot )
+		static int32 AsInt(EFlowDirection InFlowDirection, const SlotType& InSlot )
 		{
+			// InFlowDirection has no effect in vertical orientations.
 			return static_cast<int32>(InSlot.VAlignment);
 		}
 	};
-
 }
 
 
+
 /**
  * Helper method to BoxPanel::ArrangeChildren.
  * 
@@ -65,21 +81,21 @@ namespace ArrangeUtils
  * @return  Offset and Size of widget
  */
 template<EOrientation Orientation, typename SlotType>
-static AlignmentArrangeResult AlignChild(float AllottedSize, float ChildDesiredSize, const SlotType& ChildToArrange, const FMargin& SlotPadding, const float& ContentScale = 1.0f, bool bClampToParent = true)
+static AlignmentArrangeResult AlignChild(EFlowDirection InLayoutFlow, float AllottedSize, float ChildDesiredSize, const SlotType& ChildToArrange, const FMargin& SlotPadding, const float& ContentScale = 1.0f, bool bClampToParent = true)
 {
 	const FMargin& Margin = SlotPadding;
 	const float TotalMargin = Margin.GetTotalSpaceAlong<Orientation>();
 	const float MarginPre = ( Orientation == Orient_Horizontal ) ? Margin.Left : Margin.Top;
 	const float MarginPost = ( Orientation == Orient_Horizontal ) ? Margin.Right : Margin.Bottom;
 
-	const int32 Alignment = ArrangeUtils::GetChildAlignment<Orientation>::AsInt(ChildToArrange);
+	const int32 Alignment = ArrangeUtils::GetChildAlignment<Orientation>::AsInt(InLayoutFlow, ChildToArrange);
 
-	switch ( Alignment )
+	switch (Alignment)
 	{
 	case HAlign_Fill:
-		return AlignmentArrangeResult(MarginPre, ( AllottedSize - TotalMargin ) * ContentScale);
+		return AlignmentArrangeResult(MarginPre, (AllottedSize - TotalMargin) * ContentScale);
 	}
-
+	
 	const float ChildSize = bClampToParent ? FMath::Min(ChildDesiredSize, AllottedSize - TotalMargin) : ChildDesiredSize;
 
 	switch( Alignment )
@@ -97,19 +113,25 @@ static AlignmentArrangeResult AlignChild(float AllottedSize, float ChildDesiredS
 }
 
 template<EOrientation Orientation, typename SlotType>
-static AlignmentArrangeResult AlignChild(float AllottedSize, const SlotType& ChildToArrange, const FMargin& SlotPadding, const float& ContentScale = 1.0f, bool bClampToParent = true)
+static AlignmentArrangeResult AlignChild(float AllottedSize, float ChildDesiredSize, const SlotType& ChildToArrange, const FMargin& SlotPadding, const float& ContentScale = 1.0f, bool bClampToParent = true)
+{
+	return AlignChild<Orientation, SlotType>(EFlowDirection::LeftToRight, AllottedSize, ChildDesiredSize, ChildToArrange, SlotPadding, ContentScale, bClampToParent);
+}
+
+template<EOrientation Orientation, typename SlotType>
+static AlignmentArrangeResult AlignChild(EFlowDirection InLayoutFlow, float AllottedSize, const SlotType& ChildToArrange, const FMargin& SlotPadding, const float& ContentScale = 1.0f, bool bClampToParent = true)
 {
 	const FMargin& Margin = SlotPadding;
 	const float TotalMargin = Margin.GetTotalSpaceAlong<Orientation>();
 	const float MarginPre = ( Orientation == Orient_Horizontal ) ? Margin.Left : Margin.Top;
 	const float MarginPost = ( Orientation == Orient_Horizontal ) ? Margin.Right : Margin.Bottom;
 
-	const int32 Alignment = ArrangeUtils::GetChildAlignment<Orientation>::AsInt(ChildToArrange);
+	const int32 Alignment = ArrangeUtils::GetChildAlignment<Orientation>::AsInt(InLayoutFlow, ChildToArrange);
 
-	switch ( Alignment )
+	switch (Alignment)
 	{
 	case HAlign_Fill:
-		return AlignmentArrangeResult(MarginPre, ( AllottedSize - TotalMargin ) * ContentScale);
+		return AlignmentArrangeResult(MarginPre, (AllottedSize - TotalMargin) * ContentScale);
 	}
 
 	const float ChildDesiredSize = ( Orientation == Orient_Horizontal )
@@ -132,29 +154,52 @@ static AlignmentArrangeResult AlignChild(float AllottedSize, const SlotType& Chi
 	return AlignmentArrangeResult(MarginPre, ( AllottedSize - TotalMargin ) * ContentScale);
 }
 
+template<EOrientation Orientation, typename SlotType>
+static AlignmentArrangeResult AlignChild(float AllottedSize, const SlotType& ChildToArrange, const FMargin& SlotPadding, const float& ContentScale = 1.0f, bool bClampToParent = true)
+{
+	return AlignChild<Orientation, SlotType>(EFlowDirection::LeftToRight, AllottedSize, ChildToArrange, SlotPadding, ContentScale, bClampToParent);
+}
+
 
 /**
  * Arrange a ChildSlot within the AllottedGeometry and populate ArrangedChildren with the arranged result.
  * The code makes certain assumptions about the type of ChildSlot.
  */
 template<typename SlotType>
-void ArrangeSingleChild( const FGeometry& AllottedGeometry, FArrangedChildren& ArrangedChildren, const SlotType& ChildSlot, const TAttribute<FVector2D>& ContentScale )
+static void ArrangeSingleChild(const FGeometry& AllottedGeometry, FArrangedChildren& ArrangedChildren, const SlotType& ChildSlot, const TAttribute<FVector2D>& ContentScale)
+{
+	ArrangeSingleChild<SlotType>(EFlowDirection::LeftToRight, AllottedGeometry, ArrangedChildren, ChildSlot, ContentScale);
+}
+
+template<typename SlotType>
+static void ArrangeSingleChild(EFlowDirection InFlowDirection, const FGeometry& AllottedGeometry, FArrangedChildren& ArrangedChildren, const SlotType& ChildSlot, const TAttribute<FVector2D>& ContentScale)
 {
 	const EVisibility ChildVisibility = ChildSlot.GetWidget()->GetVisibility();
 	if ( ArrangedChildren.Accepts(ChildVisibility) )
 	{
 		const FVector2D ThisContentScale = ContentScale.Get();
-		const FMargin SlotPadding(ChildSlot.SlotPadding.Get());
-		AlignmentArrangeResult XResult = AlignChild<Orient_Horizontal>(AllottedGeometry.GetLocalSize().X, ChildSlot, SlotPadding, ThisContentScale.X);
-		AlignmentArrangeResult YResult = AlignChild<Orient_Vertical>(AllottedGeometry.GetLocalSize().Y, ChildSlot, SlotPadding, ThisContentScale.Y);
+		const FMargin SlotPadding(LayoutPaddingWithFlow(InFlowDirection, ChildSlot.SlotPadding.Get()));
+		const AlignmentArrangeResult XResult = AlignChild<Orient_Horizontal>(InFlowDirection, AllottedGeometry.GetLocalSize().X, ChildSlot, SlotPadding, ThisContentScale.X);
+		const AlignmentArrangeResult YResult = AlignChild<Orient_Vertical>(AllottedGeometry.GetLocalSize().Y, ChildSlot, SlotPadding, ThisContentScale.Y);
 
 		ArrangedChildren.AddWidget( ChildVisibility, AllottedGeometry.MakeChild(
 				ChildSlot.GetWidget(),
-				FVector2D(XResult.Offset,YResult.Offset),
+				FVector2D(XResult.Offset, YResult.Offset),
 				FVector2D(XResult.Size, YResult.Size)
 		) );
 	}
+}
 
+static FMargin LayoutPaddingWithFlow(EFlowDirection InLayoutFlow, const FMargin& InPadding)
+{
+	FMargin ReturnPadding(InPadding);
+	if (InLayoutFlow == EFlowDirection::RightToLeft)
+	{
+		float Temp = ReturnPadding.Left;
+		ReturnPadding.Left = ReturnPadding.Right;
+		ReturnPadding.Right = Temp;
+	}
+	return ReturnPadding;
 }
 
 /**
diff --git a/Engine/Source/Runtime/SlateCore/Public/Widgets/DeclarativeSyntaxSupport.h b/Engine/Source/Runtime/SlateCore/Public/Widgets/DeclarativeSyntaxSupport.h
--- a/Engine/Source/Runtime/SlateCore/Public/Widgets/DeclarativeSyntaxSupport.h
+++ b/Engine/Source/Runtime/SlateCore/Public/Widgets/DeclarativeSyntaxSupport.h
@@ -6,6 +6,7 @@
 #include "Misc/Attribute.h"
 #include "Layout/Visibility.h"
 #include "Layout/Clipping.h"
+#include "Layout/FlowDirection.h"
 #include "Rendering/SlateRenderTransform.h"
 #include "GenericPlatform/ICursor.h"
 #include "Types/ISlateMetaData.h"
@@ -774,6 +775,7 @@ struct TSlateBaseNamedArgs
 	, _RenderTransformPivot( FVector2D::ZeroVector )
 	, _ForceVolatile( false )
 	, _Clipping( EWidgetClipping::Inherit )
+	, _FlowDirectionPreference( EFlowDirectionPreference::Inherit )
 	{
 	}
 
@@ -817,6 +819,7 @@ struct TSlateBaseNamedArgs
 	SLATE_ARGUMENT( FName, Tag )
 	SLATE_ARGUMENT( bool, ForceVolatile )
 	SLATE_ARGUMENT( EWidgetClipping, Clipping )
+	SLATE_ARGUMENT( EFlowDirectionPreference, FlowDirectionPreference)
 
 	TArray<TSharedRef<ISlateMetaData>> MetaData;
 };
@@ -1089,6 +1092,7 @@ struct TDecl
 			InArgs._Tag,
 			InArgs._ForceVolatile,
 			InArgs._Clipping,
+			InArgs._FlowDirectionPreference,
 			InArgs.MetaData );
 
 		_RequiredArgs.CallConstruct(_Widget, InArgs);
diff --git a/Engine/Source/Runtime/SlateCore/Public/Widgets/Images/SImage.h b/Engine/Source/Runtime/SlateCore/Public/Widgets/Images/SImage.h
--- a/Engine/Source/Runtime/SlateCore/Public/Widgets/Images/SImage.h
+++ b/Engine/Source/Runtime/SlateCore/Public/Widgets/Images/SImage.h
@@ -20,19 +20,24 @@ class SLATECORE_API SImage
 	: public SLeafWidget
 {
 public:
-	SLATE_BEGIN_ARGS(SImage)
-		: _Image(FCoreStyle::Get().GetDefaultBrush())
-		, _ColorAndOpacity(FLinearColor::White)
-	{}
+	SLATE_BEGIN_ARGS( SImage )
+		: _Image( FCoreStyle::Get().GetDefaultBrush() )
+		, _ColorAndOpacity( FLinearColor::White )
+		, _FlipForRightToLeftFlowDirection( false )
+		, _OnMouseButtonDown()
+		{}
 
-	/** Image resource */
-	SLATE_ATTRIBUTE(const FSlateBrush*, Image)
+		/** Image resource */
+		SLATE_ATTRIBUTE( const FSlateBrush*, Image )
 
 		/** Color and opacity */
-		SLATE_ATTRIBUTE(FSlateColor, ColorAndOpacity)
+		SLATE_ATTRIBUTE( FSlateColor, ColorAndOpacity )
+
+		/** Flips the image if the localization's flow direction is RightToLeft */
+		SLATE_ARGUMENT( bool, FlipForRightToLeftFlowDirection )
 
 		/** Invoked when the mouse is pressed in the widget. */
-		SLATE_EVENT(FPointerEventHandler, OnMouseButtonDown)
+		SLATE_EVENT( FPointerEventHandler, OnMouseButtonDown )
 
 	SLATE_END_ARGS()
 
@@ -48,23 +53,27 @@ public:
 	 *
 	 * @param	InArgs	The declaration data for this widget
 	 */
-	void Construct(const FArguments& InArgs);
+	void Construct( const FArguments& InArgs );
 
 public:
 
 	/** See the ColorAndOpacity attribute */
-	void SetColorAndOpacity(const TAttribute<FSlateColor>& InColorAndOpacity);
-
+	void SetColorAndOpacity( const TAttribute<FSlateColor>& InColorAndOpacity );
+	
 	/** See the ColorAndOpacity attribute */
-	void SetColorAndOpacity(FLinearColor InColorAndOpacity);
+	void SetColorAndOpacity( FLinearColor InColorAndOpacity );
 
 	/** See the Image attribute */
 	void SetImage(TAttribute<const FSlateBrush*> InImage);
+	
+	/** See OnMouseButtonDown event */
+	void SetOnMouseButtonDown(FPointerEventHandler EventHandler);
 
 public:
 
 	// SWidget overrides
-	virtual int32 OnPaint(const FPaintArgs& Args, const FGeometry& AllottedGeometry, const FSlateRect& MyCullingRect, FSlateWindowElementList& OutDrawElements, int32 LayerId, const FWidgetStyle& InWidgetStyle, bool bParentEnabled) const override;
+	virtual int32 OnPaint( const FPaintArgs& Args, const FGeometry& AllottedGeometry, const FSlateRect& MyCullingRect, FSlateWindowElementList& OutDrawElements, int32 LayerId, const FWidgetStyle& InWidgetStyle, bool bParentEnabled ) const override;
+	virtual FReply OnMouseButtonDown( const FGeometry& MyGeometry, const FPointerEvent& MouseEvent ) override;
 
 protected:
 	// Begin SWidget overrides.
@@ -78,4 +87,10 @@ protected:
 
 	/** Color and opacity scale for this image */
 	TAttribute<FSlateColor> ColorAndOpacity;
+
+	/** Flips the image if the localization's flow direction is RightToLeft */
+	bool bFlipForRightToLeftFlowDirection;
+
+	/** Invoked when the mouse is pressed in the image */
+	FPointerEventHandler OnMouseButtonDownHandler;
 };
diff --git a/Engine/Source/Runtime/SlateCore/Public/Widgets/SWidget.h b/Engine/Source/Runtime/SlateCore/Public/Widgets/SWidget.h
--- a/Engine/Source/Runtime/SlateCore/Public/Widgets/SWidget.h
+++ b/Engine/Source/Runtime/SlateCore/Public/Widgets/SWidget.h
@@ -14,6 +14,7 @@
 #include "Layout/ArrangedWidget.h"
 #include "Layout/LayoutGeometry.h"
 #include "Layout/Margin.h"
+#include "Layout/FlowDirection.h"
 #include "Rendering/SlateLayoutTransform.h"
 #include "Input/CursorReply.h"
 #include "Input/Reply.h"
@@ -282,14 +283,13 @@ public:
 		const FName& InTag,
 		const bool InForceVolatile,
 		const EWidgetClipping InClipping,
+		const EFlowDirectionPreference InFlowPreference,
 		const TArray<TSharedRef<ISlateMetaData>>& InMetaData);
 
 	void SWidgetConstruct(const TAttribute<FText>& InToolTipText,
 		const TSharedPtr<IToolTip>& InToolTip,
 		const TAttribute< TOptional<EMouseCursor::Type> >& InCursor,
 		const TAttribute<bool>& InEnabledState,
-
-
 		const TAttribute<EVisibility>& InVisibility,
 		const float InRenderOpacity,
 		const TAttribute<TOptional<FSlateRenderTransform>>& InTransform,
@@ -297,6 +297,7 @@ public:
 		const FName& InTag,
 		const bool InForceVolatile,
 		const EWidgetClipping InClipping,
+		const EFlowDirectionPreference InFlowPreference,
 		const TArray<TSharedRef<ISlateMetaData>>& InMetaData);
 
 	//
@@ -1068,6 +1069,25 @@ public:
 		return RenderTransform.Get();
 	}
 
+	FORCEINLINE TOptional<FSlateRenderTransform> GetRenderTransformWithRespectToFlowDirection() const
+	{
+		if (LIKELY(GSlateFlowDirection == EFlowDirection::LeftToRight))
+		{
+			return RenderTransform.Get();
+		}
+		else
+		{
+			// If we're going right to left, flip the X translation on render transforms.
+			TOptional<FSlateRenderTransform> Transform = RenderTransform.Get();
+			if (Transform.IsSet())
+			{
+				FVector2D Translation = Transform.GetValue().GetTranslation();
+				Transform.GetValue().SetTranslation(FVector2D(-Translation.X, Translation.Y));
+			}
+			return Transform;
+		}
+	}
+
 	/** @param InTransform the render transform to set for the widget (transforms from widget's local space). TOptional<> to allow code to skip expensive overhead if there is no render transform applied. */
 	FORCEINLINE void SetRenderTransform(TAttribute<TOptional<FSlateRenderTransform>> InTransform)
 	{
@@ -1079,7 +1099,7 @@ public:
 	}
 
 	/** @return the pivot point of the render transform. */
-	FORCEINLINE const FVector2D& GetRenderTransformPivot() const
+	FORCEINLINE FVector2D GetRenderTransformPivot() const
 	{
 		return RenderTransformPivot.Get();
 	}
@@ -1132,6 +1152,22 @@ public:
 		return CullingBoundsExtension;
 	}
 
+	/**
+	 * Sets how content should flow in this panel, based on the current culture.  By default all panels inherit 
+	 * the state of the widget above.  If they set a new flow direction it will be inherited down the tree.
+	 */
+	void SetFlowDirectionPreference(EFlowDirectionPreference InFlowDirectionPreference)
+	{
+		if (FlowDirectionPreference != InFlowDirectionPreference)
+		{
+			FlowDirectionPreference = InFlowDirectionPreference;
+			Invalidate(EInvalidateWidget::Paint);
+		}
+	}
+
+	/** Gets the desired flow direction for the layout. */
+	EFlowDirectionPreference GetFlowDirectionPreference() const { return FlowDirectionPreference; }
+
 	/**
 	 * Set the tool tip that should appear when this widget is hovered.
 	 *
@@ -1479,6 +1515,26 @@ protected:
 	 */
 	EWidgetClipping Clipping;
 
+protected:
+	/** Establishes a new flow direction potentially, if this widget has a particular preference for it and all its children. */
+	EFlowDirection ComputeFlowDirection() const
+	{
+		switch (FlowDirectionPreference)
+		{
+		case EFlowDirectionPreference::Culture:
+			return FLayoutLocalization::GetLocalizedLayoutDirection();
+		case EFlowDirectionPreference::LeftToRight:
+			return EFlowDirection::LeftToRight;
+		case EFlowDirectionPreference::RightToLeft:
+			return EFlowDirection::RightToLeft;
+		}
+
+		return GSlateFlowDirection;
+	}
+
+private:
+	/** Flow direction preference */
+	EFlowDirectionPreference FlowDirectionPreference;
 
 private:
 	EWidgetUpdateFlags UpdateFlags;
@@ -1573,9 +1629,10 @@ FORCEINLINE_DEBUGGABLE FArrangedWidget FGeometry::MakeChild(const TSharedRef<SWi
 {
 	// If there is no render transform set, use the simpler MakeChild call that doesn't bother concatenating the render transforms.
 	// This saves a significant amount of overhead since every widget does this, and most children don't have a render transform.
-	if ( ChildWidget->GetRenderTransform().IsSet() )
+	TOptional<FSlateRenderTransform> RenderTransform = ChildWidget->GetRenderTransformWithRespectToFlowDirection();
+	if (RenderTransform.IsSet() )
 	{
-		return FArrangedWidget(ChildWidget, MakeChild(InLocalSize, LayoutTransform, ChildWidget->GetRenderTransform().GetValue(), ChildWidget->GetRenderTransformPivot()));
+		return FArrangedWidget(ChildWidget, MakeChild(InLocalSize, LayoutTransform, RenderTransform.GetValue(), ChildWidget->GetRenderTransformPivot()));
 	}
 	else
 	{
diff --git a/Engine/Source/Runtime/UMG/Private/Components/Border.cpp b/Engine/Source/Runtime/UMG/Private/Components/Border.cpp
--- a/Engine/Source/Runtime/UMG/Private/Components/Border.cpp
+++ b/Engine/Source/Runtime/UMG/Private/Components/Border.cpp
@@ -43,7 +43,8 @@ void UBorder::ReleaseSlateResources(bool bReleaseChildren)
 
 TSharedRef<SWidget> UBorder::RebuildWidget()
 {
-	MyBorder = SNew(SBorder);
+	MyBorder = SNew(SBorder)
+		.FlipForRightToLeftFlowDirection(bFlipForRightToLeftFlowDirection);
 	
 	if ( GetChildrenCount() > 0 )
 	{
diff --git a/Engine/Source/Runtime/UMG/Private/Components/Image.cpp b/Engine/Source/Runtime/UMG/Private/Components/Image.cpp
--- a/Engine/Source/Runtime/UMG/Private/Components/Image.cpp
+++ b/Engine/Source/Runtime/UMG/Private/Components/Image.cpp
@@ -44,7 +44,9 @@ void UImage::ReleaseSlateResources(bool bReleaseChildren)
 
 TSharedRef<SWidget> UImage::RebuildWidget()
 {
-	MyImage = SNew(SImage);
+	MyImage = SNew(SImage)
+			.FlipForRightToLeftFlowDirection(bFlipForRightToLeftFlowDirection);
+
 	return MyImage.ToSharedRef();
 }
 
diff --git a/Engine/Source/Runtime/UMG/Private/Components/Widget.cpp b/Engine/Source/Runtime/UMG/Private/Components/Widget.cpp
--- a/Engine/Source/Runtime/UMG/Private/Components/Widget.cpp
+++ b/Engine/Source/Runtime/UMG/Private/Components/Widget.cpp
@@ -1121,6 +1121,8 @@ void UWidget::SynchronizeProperties()
 	SafeWidget->SetClipping(Clipping);
 #endif
 
+	SafeWidget->SetFlowDirectionPreference(FlowDirectionPreference);
+
 	SafeWidget->ForceVolatile(bIsVolatile);
 
 	SafeWidget->SetRenderOpacity(RenderOpacity);
diff --git a/Engine/Source/Runtime/UMG/Public/Components/Border.h b/Engine/Source/Runtime/UMG/Public/Components/Border.h
--- a/Engine/Source/Runtime/UMG/Public/Components/Border.h
+++ b/Engine/Source/Runtime/UMG/Public/Components/Border.h
@@ -1,4 +1,4 @@
-﻿// Copyright 1998-2019 Epic Games, Inc. All Rights Reserved.
+// Copyright 1998-2018 Epic Games, Inc. All Rights Reserved.
 
 #pragma once
 
@@ -81,6 +81,10 @@ public:
 	UPROPERTY(EditAnywhere, BlueprintReadOnly, Category = Appearance)
 	FVector2D DesiredSizeScale;
 
+	/** Flips the background image if the localization's flow direction is RightToLeft */
+	UPROPERTY(EditAnywhere, Category = "Localization")
+	bool bFlipForRightToLeftFlowDirection;
+
 public:
 
 	UPROPERTY(EditAnywhere, Category=Events, meta=( IsBindableEvent="True" ))
diff --git a/Engine/Source/Runtime/UMG/Public/Components/Image.h b/Engine/Source/Runtime/UMG/Public/Components/Image.h
--- a/Engine/Source/Runtime/UMG/Public/Components/Image.h
+++ b/Engine/Source/Runtime/UMG/Public/Components/Image.h
@@ -55,6 +55,10 @@ public:
 	UPROPERTY()
 	FGetLinearColor ColorAndOpacityDelegate;
 
+	/** Flips the image if the localization's flow direction is RightToLeft */
+	UPROPERTY(EditAnywhere, Category = "Localization")
+	bool bFlipForRightToLeftFlowDirection;
+
 public:
 
 	UPROPERTY(EditAnywhere, Category=Events, meta=( IsBindableEvent="True" ))
diff --git a/Engine/Source/Runtime/UMG/Public/Components/Widget.h b/Engine/Source/Runtime/UMG/Public/Components/Widget.h
--- a/Engine/Source/Runtime/UMG/Public/Components/Widget.h
+++ b/Engine/Source/Runtime/UMG/Public/Components/Widget.h
@@ -371,6 +371,10 @@ public:
 	UPROPERTY(Instanced, EditAnywhere, BlueprintReadOnly, Category="Navigation")
 	class UWidgetNavigation* Navigation;
 
+	/** Allows you to set a new flow direction */
+	UPROPERTY(EditAnywhere, Category = "Localization")
+	EFlowDirectionPreference FlowDirectionPreference;
+
 #if WITH_EDITORONLY_DATA
 
 	/** Stores a reference to the asset responsible for this widgets construction. */
