Fix non-unity builds due to missing include directive

diff --git a/Engine/Source/Runtime/SlateCore/Public/Layout/LayoutUtils.h b/Engine/Source/Runtime/SlateCore/Public/Layout/LayoutUtils.h
--- a/Engine/Source/Runtime/SlateCore/Public/Layout/LayoutUtils.h
+++ b/Engine/Source/Runtime/SlateCore/Public/Layout/LayoutUtils.h
@@ -7,6 +7,7 @@
 #include "Layout/Visibility.h"
 #include "Layout/SlateRect.h"
 #include "Layout/ArrangedChildren.h"
+#include "Layout/FlowDirection.h"
 
 struct AlignmentArrangeResult
 {
