Fixing some build warnings from main.

diff --git a/Engine/Source/Runtime/SlateCore/Private/Widgets/SBoxPanel.cpp b/Engine/Source/Runtime/SlateCore/Private/Widgets/SBoxPanel.cpp
--- a/Engine/Source/Runtime/SlateCore/Private/Widgets/SBoxPanel.cpp
+++ b/Engine/Source/Runtime/SlateCore/Private/Widgets/SBoxPanel.cpp
@@ -1,4 +1,4 @@
-// Copyright 1998-2018 Epic Games, Inc. All Rights Reserved.
+// Copyright 1998-2019 Epic Games, Inc. All Rights Reserved.
 
 #include "Widgets/SBoxPanel.h"
 #include "Layout/LayoutUtils.h"
diff --git a/Engine/Source/Runtime/SlateCore/Public/Layout/FlowDirection.h b/Engine/Source/Runtime/SlateCore/Public/Layout/FlowDirection.h
--- a/Engine/Source/Runtime/SlateCore/Public/Layout/FlowDirection.h
+++ b/Engine/Source/Runtime/SlateCore/Public/Layout/FlowDirection.h
@@ -5,6 +5,7 @@
 #include "CoreMinimal.h"
 #include "UObject/ObjectMacros.h"
 #include "UObject/Class.h"
+
 #include "FlowDirection.generated.h"
 
 /**
@@ -26,13 +27,13 @@ enum class EFlowDirection : uint8
 UENUM()
 enum class EFlowDirectionPreference : uint8
 {
-	/** */
+	/** Inherits the flow direction set by the parent widget. */
 	Inherit,
-	/** */
+	/** Begins laying out widgets using the current cultures layout direction preference, flipping the directionality of flows. */
 	Culture,
-	/** */
+	/** Forces a Left to Right layout flow. */
 	LeftToRight,
-	/**  */
+	/** Forces a Right to Left layout flow. */
 	RightToLeft
 };
 
diff --git a/Engine/Source/Runtime/SlateCore/Public/Layout/LayoutUtils.h b/Engine/Source/Runtime/SlateCore/Public/Layout/LayoutUtils.h
--- a/Engine/Source/Runtime/SlateCore/Public/Layout/LayoutUtils.h
+++ b/Engine/Source/Runtime/SlateCore/Public/Layout/LayoutUtils.h
@@ -3,11 +3,11 @@
 #pragma once
 
 #include "CoreMinimal.h"
-#include "Layout/Margin.h"
-#include "Layout/Visibility.h"
-#include "Layout/SlateRect.h"
-#include "Layout/ArrangedChildren.h"
-#include "Layout/FlowDirection.h"
+#include "Margin.h"
+#include "Visibility.h"
+#include "SlateRect.h"
+#include "ArrangedChildren.h"
+#include "FlowDirection.h"
 
 struct AlignmentArrangeResult
 {
diff --git a/Engine/Source/Runtime/UMG/Public/Components/Border.h b/Engine/Source/Runtime/UMG/Public/Components/Border.h
--- a/Engine/Source/Runtime/UMG/Public/Components/Border.h
+++ b/Engine/Source/Runtime/UMG/Public/Components/Border.h
@@ -1,4 +1,4 @@
-// Copyright 1998-2018 Epic Games, Inc. All Rights Reserved.
+// Copyright 1998-2019 Epic Games, Inc. All Rights Reserved.
 
 #pragma once
 
