Fix for NAN ensures in AddBorderElement()
#rnx

diff --git a/Engine/Source/Runtime/SlateCore/Private/Rendering/ElementBatcher.cpp b/Engine/Source/Runtime/SlateCore/Private/Rendering/ElementBatcher.cpp
--- a/Engine/Source/Runtime/SlateCore/Private/Rendering/ElementBatcher.cpp
+++ b/Engine/Source/Runtime/SlateCore/Private/Rendering/ElementBatcher.cpp
@@ -1930,8 +1930,18 @@ void FSlateElementBatcher::AddBorderElement( const FSlateDrawElement& DrawElemen
 	RightMarginU += HalfTexel.X;
 
 	// Determine the amount of tiling needed for the texture in this element.  The formula is number of pixels covered by the tiling portion of the texture / the number number of texels corresponding to the tiled portion of the texture.
-	float TopTiling = (RightMarginX-LeftMarginX)/(TextureSizeLocalSpace.X * ( 1 - Margin.GetTotalSpaceAlong<Orient_Horizontal>() ));
-	float LeftTiling = (BottomMarginY-TopMarginY)/(TextureSizeLocalSpace.Y * ( 1 - Margin.GetTotalSpaceAlong<Orient_Vertical>() ));
+	float TopTiling = 1.0f;
+	float LeftTiling = 1.0f;
+	float Denom = TextureSizeLocalSpace.X * (1.0f - Margin.GetTotalSpaceAlong<Orient_Horizontal>());
+	if (!FMath::IsNearlyZero(Denom))
+	{
+		TopTiling = (RightMarginX - LeftMarginX) / Denom;
+	}
+	Denom = TextureSizeLocalSpace.Y * (1.0f - Margin.GetTotalSpaceAlong<Orient_Vertical>());
+	if (!FMath::IsNearlyZero(Denom))
+	{
+		LeftTiling = (BottomMarginY - TopMarginY) / Denom;
+	}
 	
 	FShaderParams ShaderParams = FShaderParams::MakePixelShaderParams( FVector4(LeftMarginU,RightMarginU,TopMarginV,BottomMarginV) );
 
