Reduced memory overhead of slate texture atlases for slate styles.  Outside of editor the texture atlases will no longer keep system memory around so that the atlas can be updated later.  This is an editor only feature by default.

diff --git a/Engine/Source/Runtime/SlateCore/Private/Fonts/FontTypes.cpp b/Engine/Source/Runtime/SlateCore/Private/Fonts/FontTypes.cpp
--- a/Engine/Source/Runtime/SlateCore/Private/Fonts/FontTypes.cpp
+++ b/Engine/Source/Runtime/SlateCore/Private/Fonts/FontTypes.cpp
@@ -4,7 +4,7 @@
 
 
 FSlateFontAtlas::FSlateFontAtlas( uint32 InWidth, uint32 InHeight )
-	: FSlateTextureAtlas( InWidth, InHeight, sizeof(uint8), ESlateTextureAtlasPaddingStyle::PadWithZero )
+	: FSlateTextureAtlas( InWidth, InHeight, sizeof(uint8), ESlateTextureAtlasPaddingStyle::PadWithZero, true )
 {
 }
 
@@ -26,7 +26,7 @@ const FAtlasedTextureSlot* FSlateFontAtlas::AddCharacter( const FCharacterRender
 void FSlateFontAtlas::Flush()
 {
 	// Empty the atlas
-	Empty();
+	EmptyAtlasData();
 
 	// Recreate the data
 	InitAtlasData();
diff --git a/Engine/Source/Runtime/SlateCore/Private/Textures/TextureAtlas.cpp b/Engine/Source/Runtime/SlateCore/Private/Textures/TextureAtlas.cpp
--- a/Engine/Source/Runtime/SlateCore/Private/Textures/TextureAtlas.cpp
+++ b/Engine/Source/Runtime/SlateCore/Private/Textures/TextureAtlas.cpp
@@ -5,6 +5,7 @@
 #include "Textures/SlateShaderResource.h"
 #include "Textures/SlateTextureData.h"
 #include "HAL/LowLevelMemTracker.h"
+#include "Misc/MemStack.h"
 
 DEFINE_STAT(STAT_SlateTextureGPUMemory);
 DEFINE_STAT(STAT_SlateTextureDataMemory);
@@ -32,17 +33,19 @@ ESlateTextureAtlasThreadId GetCurrentSlateTextureAtlasThreadId()
 
 FSlateTextureAtlas::~FSlateTextureAtlas()
 {
-	Empty();
+	EmptyAtlasData();
 }
 
 
 /* FSlateTextureAtlas interface
  *****************************************************************************/
 
-void FSlateTextureAtlas::Empty()
+void FSlateTextureAtlas::EmptyAtlasData()
 {
+	FMemMark Mark(FMemStack::Get());
+
 	// Remove all nodes
-	TArray<FAtlasedTextureSlot*> DeleteSlots;
+	TArray<FAtlasedTextureSlot*, TMemStackAllocator<>> DeleteSlots;
 
 	for (FAtlasedTextureSlot::TIterator SlotIt(AtlasUsedSlots); SlotIt; SlotIt.Next())
 	{
diff --git a/Engine/Source/Runtime/SlateCore/Public/Textures/TextureAtlas.h b/Engine/Source/Runtime/SlateCore/Public/Textures/TextureAtlas.h
--- a/Engine/Source/Runtime/SlateCore/Public/Textures/TextureAtlas.h
+++ b/Engine/Source/Runtime/SlateCore/Public/Textures/TextureAtlas.h
@@ -76,7 +76,7 @@ struct FAtlasedTextureSlot : public TIntrusiveLinkedList<FAtlasedTextureSlot>
 class SLATECORE_API FSlateTextureAtlas
 {
 public:
-	FSlateTextureAtlas( uint32 InWidth, uint32 InHeight, uint32 InBytesPerPixel, ESlateTextureAtlasPaddingStyle InPaddingStyle )
+	FSlateTextureAtlas( uint32 InWidth, uint32 InHeight, uint32 InBytesPerPixel, ESlateTextureAtlasPaddingStyle InPaddingStyle, bool bInUpdatesAfterInitialization )
 		: AtlasData()
 		, AtlasUsedSlots(NULL)
 		, AtlasEmptySlots(NULL)
@@ -85,6 +85,7 @@ public:
 		, BytesPerPixel( InBytesPerPixel )
 		, PaddingStyle( InPaddingStyle )
 		, bNeedsUpdate( false )
+		, bUpdatesAfterInitialization(bInUpdatesAfterInitialization)
 #if !(UE_BUILD_SHIPPING || UE_BUILD_TEST)
 		, AtlasOwnerThread( ESlateTextureAtlasThreadId::Unknown )
 #endif
@@ -95,9 +96,9 @@ public:
 	virtual ~FSlateTextureAtlas();
 
 	/**
-	 * Clears atlas data
+	 * Clears atlas cpu data.  It does not clear rendering data
 	 */
-	void Empty();
+	void EmptyAtlasData();
 
 	/**
 	 * Adds a texture to the atlas
@@ -201,7 +202,8 @@ protected:
 
 	/** True if this texture needs to have its rendering resources updated */
 	bool bNeedsUpdate;
-
+	/** True if this texture can update after initialziation and we should preserve the atlas slots and cpu memory */
+	bool bUpdatesAfterInitialization;
 #if !(UE_BUILD_SHIPPING || UE_BUILD_TEST)
 	/** 
 	 * The type of thread that owns this atlas - this is the only thread that can safely update it 
diff --git a/Engine/Source/Runtime/SlateRHIRenderer/Private/SlateRHIResourceManager.cpp b/Engine/Source/Runtime/SlateRHIRenderer/Private/SlateRHIResourceManager.cpp
--- a/Engine/Source/Runtime/SlateRHIRenderer/Private/SlateRHIResourceManager.cpp
+++ b/Engine/Source/Runtime/SlateRHIRenderer/Private/SlateRHIResourceManager.cpp
@@ -470,7 +470,8 @@ FSlateShaderResourceProxy* FSlateRHIResourceManager::GenerateTextureResource( co
 		{
 			INC_DWORD_STAT_BY(STAT_SlateNumTextureAtlases, 1);
 
-			Atlas = new FSlateTextureAtlasRHI( AtlasSize, AtlasSize, ESlateTextureAtlasPaddingStyle::DilateBorder );
+			bool bCanUpdateAfterInitialization = GIsEditor;
+			Atlas = new FSlateTextureAtlasRHI( AtlasSize, AtlasSize, ESlateTextureAtlasPaddingStyle::DilateBorder, bCanUpdateAfterInitialization);
 			TextureAtlases.Add( Atlas );
 			NewSlot = TextureAtlases.Last()->AddTexture( Width, Height, Info.TextureData->GetRawBytes() );
 		}
diff --git a/Engine/Source/Runtime/SlateRHIRenderer/Private/SlateRHITextureAtlas.cpp b/Engine/Source/Runtime/SlateRHIRenderer/Private/SlateRHITextureAtlas.cpp
--- a/Engine/Source/Runtime/SlateRHIRenderer/Private/SlateRHITextureAtlas.cpp
+++ b/Engine/Source/Runtime/SlateRHIRenderer/Private/SlateRHITextureAtlas.cpp
@@ -11,8 +11,8 @@
 /* FSlateTextureAtlasRHI structors
  *****************************************************************************/
 
-FSlateTextureAtlasRHI::FSlateTextureAtlasRHI( uint32 InWidth, uint32 InHeight, ESlateTextureAtlasPaddingStyle PaddingStyle )
-	: FSlateTextureAtlas(InWidth, InHeight, GPixelFormats[PF_B8G8R8A8].BlockBytes, PaddingStyle)
+FSlateTextureAtlasRHI::FSlateTextureAtlasRHI( uint32 InWidth, uint32 InHeight, ESlateTextureAtlasPaddingStyle PaddingStyle, bool bUpdatesAfterInitialization)
+	: FSlateTextureAtlas(InWidth, InHeight, GPixelFormats[PF_B8G8R8A8].BlockBytes, PaddingStyle, bUpdatesAfterInitialization)
 	, AtlasTexture(new FSlateTexture2DRHIRef(InWidth, InHeight, PF_B8G8R8A8, NULL, TexCreate_SRGB, true)) 
 {
 }
@@ -79,5 +79,10 @@ void FSlateTextureAtlasRHI::ConditionalUpdateTexture( )
 		});
 
 		bNeedsUpdate = false;
+
+		if (!bUpdatesAfterInitialization)
+		{
+			EmptyAtlasData();
+		}
 	}
 }
diff --git a/Engine/Source/Runtime/SlateRHIRenderer/Private/SlateRHITextureAtlas.h b/Engine/Source/Runtime/SlateRHIRenderer/Private/SlateRHITextureAtlas.h
--- a/Engine/Source/Runtime/SlateRHIRenderer/Private/SlateRHITextureAtlas.h
+++ b/Engine/Source/Runtime/SlateRHIRenderer/Private/SlateRHITextureAtlas.h
@@ -23,7 +23,7 @@ public:
 	 * @param InHeight
 	 * @param PaddingStyle
 	 */
-	FSlateTextureAtlasRHI( uint32 InWidth, uint32 InHeight, ESlateTextureAtlasPaddingStyle PaddingStyle );
+	FSlateTextureAtlasRHI( uint32 InWidth, uint32 InHeight, ESlateTextureAtlasPaddingStyle PaddingStyle, bool bUpdatesAfterInitialization);
 
 	/**
 	 * Destructor.
