Move slate layout flow direction to PaintWindow for compatibilty with fast math mode

diff --git a/Engine/Source/Runtime/SlateCore/Private/Widgets/SWindow.cpp b/Engine/Source/Runtime/SlateCore/Private/Widgets/SWindow.cpp
--- a/Engine/Source/Runtime/SlateCore/Private/Widgets/SWindow.cpp
+++ b/Engine/Source/Runtime/SlateCore/Private/Widgets/SWindow.cpp
@@ -1921,13 +1921,7 @@ SWindow::SWindow()
 
 int32 SWindow::PaintWindow( const FPaintArgs& Args, const FGeometry& AllottedGeometry, const FSlateRect& MyCullingRect, FSlateWindowElementList& OutDrawElements, int32 LayerId, const FWidgetStyle& InWidgetStyle, bool bParentEnabled ) const
 {
-	LayerId = Paint(Args, AllottedGeometry, MyCullingRect, OutDrawElements, LayerId, InWidgetStyle, bParentEnabled);
-
-	return LayerId;
-}
-
-int32 SWindow::OnPaint(const FPaintArgs& Args, const FGeometry& AllottedGeometry, const FSlateRect& MyCullingRect, FSlateWindowElementList& OutDrawElements, int32 LayerId, const FWidgetStyle& InWidgetStyle, bool bParentEnabled) const
-{
+	// Create initial culture specific layout direction
 	EFlowDirection NewFlowDirection = GSlateFlowDirection;
 	if (GetFlowDirectionPreference() == EFlowDirectionPreference::Inherit)
 	{
@@ -1936,6 +1930,13 @@ int32 SWindow::OnPaint(const FPaintArgs& Args, const FGeometry& AllottedGeometry
 
 	TGuardValue<EFlowDirection> FlowGuard(GSlateFlowDirection, NewFlowDirection);
 
+	LayerId = Paint(Args, AllottedGeometry, MyCullingRect, OutDrawElements, LayerId, InWidgetStyle, bParentEnabled);
+
+	return LayerId;
+}
+
+int32 SWindow::OnPaint(const FPaintArgs& Args, const FGeometry& AllottedGeometry, const FSlateRect& MyCullingRect, FSlateWindowElementList& OutDrawElements, int32 LayerId, const FWidgetStyle& InWidgetStyle, bool bParentEnabled) const
+{
 	OutDrawElements.BeginDeferredGroup();
 	int32 MaxLayer = SCompoundWidget::OnPaint(Args, AllottedGeometry, MyCullingRect, OutDrawElements, LayerId, InWidgetStyle, bParentEnabled);
 	OutDrawElements.EndDeferredGroup();
