(Part 1/2) - Slate side, Tweaking when SNullContentWidget is used.  Want to avoid using it on any situation that involves it being temporarily incremented in the loading screen, otherwise you run the risk of incrementing a thread unsafe shared pointer on multiple threads at once.

diff --git a/Engine/Source/Runtime/Engine/Private/Slate/SGameLayerManager.cpp b/Engine/Source/Runtime/Engine/Private/Slate/SGameLayerManager.cpp
--- a/Engine/Source/Runtime/Engine/Private/Slate/SGameLayerManager.cpp
+++ b/Engine/Source/Runtime/Engine/Private/Slate/SGameLayerManager.cpp
@@ -259,7 +259,7 @@ int32 SGameLayerManager::OnPaint(const FPaintArgs& Args, const FGeometry& Allott
 
 bool SGameLayerManager::OnVisualizeTooltip(const TSharedPtr<SWidget>& TooltipContent)
 {
-	TooltipPresenter->SetContent(TooltipContent.IsValid() ? TooltipContent.ToSharedRef() : SNullWidget::NullWidget);
+	TooltipPresenter->SetContent(TooltipContent);
 
 	return true;
 }
diff --git a/Engine/Source/Runtime/Slate/Private/Widgets/LayerManager/STooltipPresenter.cpp b/Engine/Source/Runtime/Slate/Private/Widgets/LayerManager/STooltipPresenter.cpp
--- a/Engine/Source/Runtime/Slate/Private/Widgets/LayerManager/STooltipPresenter.cpp
+++ b/Engine/Source/Runtime/Slate/Private/Widgets/LayerManager/STooltipPresenter.cpp
@@ -10,44 +10,54 @@ void STooltipPresenter::Construct(const FArguments& InArgs)
 	SetCanTick(false);
 }
 
-void STooltipPresenter::SetContent(const TSharedRef<SWidget>& InWidget)
+void STooltipPresenter::SetContent(TSharedPtr<SWidget> InWidget)
 {
-	ChildSlot.AttachWidget(InWidget);
+	if (InWidget.IsValid())
+	{
+		ChildSlot.AttachWidget(InWidget.ToSharedRef());
+	}
+	else
+	{
+		ChildSlot.DetachWidget();
+	}
 }
 
 void STooltipPresenter::OnArrangeChildren(const FGeometry& AllottedGeometry, FArrangedChildren& ArrangedChildren) const
 {
 	static const FVector2D CursorSize = FVector2D(12, 12);
 
-	// Cached geometry is in desktop space.  We need to convert from desktop space where the mouse is to local space so use CachedGeometry.
-	const FVector2D LocalCursorPosition = GetCachedGeometry().AbsoluteToLocal(FSlateApplication::Get().GetCursorPos());
+	if (ChildSlot.Num() > 0)
+	{
+		TSharedRef<SWidget> TooltipWidget = ChildSlot.GetWidget();
 
-	const FSlateRect CursorAnchorRect(LocalCursorPosition, LocalCursorPosition + CursorSize);
-	const FSlateRect TooltipPopup(LocalCursorPosition + CursorSize, LocalCursorPosition + CursorSize + ChildSlot.GetWidget()->GetDesiredSize());
-	
-	const FVector2D TooltipPosition = ComputePopupFitInRect(CursorAnchorRect, TooltipPopup, EOrientation::Orient_Vertical, FSlateRect(FVector2D::ZeroVector, AllottedGeometry.GetLocalSize()));
-	
-	// We round the final tooltip position so that our tooltip doesn't begin at a half pixel offset, which avoids the contents of the tooltip
-	// jittering in relation to one another.
-	const FVector2D TooltipPositionRounded = AllottedGeometry.LocalToRoundedLocal(TooltipPosition);
-
-	ArrangedChildren.AddWidget(AllottedGeometry.MakeChild(
-		ChildSlot.GetWidget(),
-		ChildSlot.GetWidget()->GetDesiredSize(),
-		FSlateLayoutTransform(TooltipPositionRounded)
-	));
+		// Cached geometry is in desktop space.  We need to convert from desktop space where the mouse is to local space so use CachedGeometry.
+		const FVector2D LocalCursorPosition = GetCachedGeometry().AbsoluteToLocal(FSlateApplication::Get().GetCursorPos());
+
+		const FSlateRect CursorAnchorRect(LocalCursorPosition, LocalCursorPosition + CursorSize);
+		const FSlateRect TooltipPopup(LocalCursorPosition + CursorSize, LocalCursorPosition + CursorSize + TooltipWidget->GetDesiredSize());
+
+		const FVector2D TooltipPosition = ComputePopupFitInRect(CursorAnchorRect, TooltipPopup, EOrientation::Orient_Vertical, FSlateRect(FVector2D::ZeroVector, AllottedGeometry.GetLocalSize()));
+
+		// We round the final tooltip position so that our tooltip doesn't begin at a half pixel offset, which avoids the contents of the tooltip
+		// jittering in relation to one another.
+		const FVector2D TooltipPositionRounded = AllottedGeometry.LocalToRoundedLocal(TooltipPosition);
+
+		ArrangedChildren.AddWidget(AllottedGeometry.MakeChild(
+			ChildSlot.GetWidget(),
+			ChildSlot.GetWidget()->GetDesiredSize(),
+			FSlateLayoutTransform(TooltipPositionRounded)
+		));
+	}
 }
 
 FVector2D STooltipPresenter::ComputeDesiredSize( float ) const
 {
-	if (ChildSlot.GetWidget() != SNullWidget::NullWidget)
+	if (ChildSlot.Num() > 0)
 	{
 		return ChildSlot.GetWidget()->GetDesiredSize();
 	}
-	else
-	{
-		return FVector2D(0, 0);
-	}
+	
+	return FVector2D(0, 0);
 }
 
 FChildren* STooltipPresenter::GetChildren()
diff --git a/Engine/Source/Runtime/Slate/Private/Widgets/SVirtualWindow.cpp b/Engine/Source/Runtime/Slate/Private/Widgets/SVirtualWindow.cpp
--- a/Engine/Source/Runtime/Slate/Private/Widgets/SVirtualWindow.cpp
+++ b/Engine/Source/Runtime/Slate/Private/Widgets/SVirtualWindow.cpp
@@ -39,7 +39,7 @@ FPopupMethodReply SVirtualWindow::OnQueryPopupMethod() const
 
 bool SVirtualWindow::OnVisualizeTooltip(const TSharedPtr<SWidget>& TooltipContent)
 {
-	TooltipPresenter->SetContent(TooltipContent.IsValid() ? TooltipContent.ToSharedRef() : SNullWidget::NullWidget);
+	TooltipPresenter->SetContent(TooltipContent);
 
 	return true;
 }
diff --git a/Engine/Source/Runtime/Slate/Public/Widgets/LayerManager/STooltipPresenter.h b/Engine/Source/Runtime/Slate/Public/Widgets/LayerManager/STooltipPresenter.h
--- a/Engine/Source/Runtime/Slate/Public/Widgets/LayerManager/STooltipPresenter.h
+++ b/Engine/Source/Runtime/Slate/Public/Widgets/LayerManager/STooltipPresenter.h
@@ -33,7 +33,7 @@ public:
 
 	void Construct(const FArguments& InArgs);
 
-	void SetContent(const TSharedRef<SWidget>& InWidget);
+	void SetContent(TSharedPtr<SWidget> InWidget);
 
 private:
 	virtual void OnArrangeChildren(const FGeometry& AllottedGeometry, FArrangedChildren& ArrangedChildren) const override;
diff --git a/Engine/Source/Runtime/Slate/Public/Widgets/Layout/SWidgetSwitcher.h b/Engine/Source/Runtime/Slate/Public/Widgets/Layout/SWidgetSwitcher.h
--- a/Engine/Source/Runtime/Slate/Public/Widgets/Layout/SWidgetSwitcher.h
+++ b/Engine/Source/Runtime/Slate/Public/Widgets/Layout/SWidgetSwitcher.h
@@ -168,28 +168,5 @@ private:
 	TPanelChildren<FSlot> AllChildren;
 
 	/** Required to implement GetChildren() in a way that can dynamically return the currently active child. */
-	class SLATE_API FOneDynamicChild : public FChildren
-	{
-	public:
-
-		FOneDynamicChild(SWidget* InOwner, TPanelChildren<FSlot>* InAllChildren, const TAttribute<int32>* InWidgetIndex )
-			: FChildren(InOwner)
-			, AllChildren( InAllChildren )
-			, WidgetIndex( InWidgetIndex )
-		{ }
-		
-		virtual int32 Num() const override { return AllChildren->Num() > 0 ? 1 : 0; }
-		
-		virtual TSharedRef<SWidget> GetChildAt( int32 Index ) override { check(Index == 0); return AllChildren->GetChildAt(WidgetIndex->Get()); }
-		
-		virtual TSharedRef<const SWidget> GetChildAt( int32 Index ) const override { check(Index == 0); return AllChildren->GetChildAt(WidgetIndex->Get()); }
-		
-	private:
-
-		virtual const FSlotBase& GetSlotAt(int32 ChildIndex) const override { return (*AllChildren)[ChildIndex]; }
-
-		TPanelChildren<FSlot>* AllChildren;
-		const TAttribute<int32>* WidgetIndex;
-
-	} OneDynamicChild;
+	TOneDynamicChild<FSlot> OneDynamicChild;
 };
diff --git a/Engine/Source/Runtime/SlateCore/Private/Rendering/DrawElements.cpp b/Engine/Source/Runtime/SlateCore/Private/Rendering/DrawElements.cpp
--- a/Engine/Source/Runtime/SlateCore/Private/Rendering/DrawElements.cpp
+++ b/Engine/Source/Runtime/SlateCore/Private/Rendering/DrawElements.cpp
@@ -30,7 +30,7 @@ static bool IsResourceObjectValid(UObject*& InObject)
 	return true;
 }
 
-FSlateWindowElementList::FSlateWindowElementList(TSharedPtr<SWindow> InPaintWindow)
+FSlateWindowElementList::FSlateWindowElementList(const TSharedPtr<SWindow>& InPaintWindow)
 	: WeakPaintWindow(InPaintWindow)
 	, RawPaintWindow(InPaintWindow.Get())
 	, RenderTargetWindow(nullptr)
diff --git a/Engine/Source/Runtime/SlateCore/Public/Layout/Children.h b/Engine/Source/Runtime/SlateCore/Public/Layout/Children.h
--- a/Engine/Source/Runtime/SlateCore/Public/Layout/Children.h
+++ b/Engine/Source/Runtime/SlateCore/Public/Layout/Children.h
@@ -65,7 +65,7 @@ public:
 		// Nobody should be getting a child when there aren't any children.
 		// We expect this to crash!
 		check( false );
-		return TSharedPtr<SWidget>(nullptr).ToSharedRef();
+		return SNullWidget::NullWidget;
 	}
 	
 	virtual TSharedRef<const SWidget> GetChildAt( int32 ) const override
@@ -73,7 +73,7 @@ public:
 		// Nobody should be getting a child when there aren't any children.
 		// We expect this to crash!
 		check( false );
-		return TSharedPtr<const SWidget>(nullptr).ToSharedRef();
+		return SNullWidget::NullWidget;
 	}
 
 private:
@@ -101,8 +101,19 @@ public:
 	}
 
 	virtual int32 Num() const override { return 1; }
-	virtual TSharedRef<SWidget> GetChildAt( int32 ChildIndex ) override { check(ChildIndex == 0); return FSlotBase::GetWidget(); }
-	virtual TSharedRef<const SWidget> GetChildAt( int32 ChildIndex ) const override { check(ChildIndex == 0); return FSlotBase::GetWidget(); }
+
+	virtual TSharedRef<SWidget> GetChildAt( int32 ChildIndex ) override
+	{
+		check(ChildIndex == 0);
+		return FSlotBase::GetWidget();
+	}
+
+	virtual TSharedRef<const SWidget> GetChildAt( int32 ChildIndex ) const override
+	{
+		check(ChildIndex == 0);
+		return FSlotBase::GetWidget();
+	}
+
 private:
 	virtual const FSlotBase& GetSlotAt(int32 ChildIndex) const override { check(ChildIndex == 0); return *this; }
 };
@@ -123,8 +134,18 @@ public:
 	}
 
 	virtual int32 Num() const override { return WidgetPtr.IsValid() ? 1 : 0 ; }
-	virtual TSharedRef<SWidget> GetChildAt( int32 ChildIndex ) override { check(ChildIndex == 0); return WidgetPtr.Pin().ToSharedRef(); }
-	virtual TSharedRef<const SWidget> GetChildAt( int32 ChildIndex ) const override { check(ChildIndex == 0); return WidgetPtr.Pin().ToSharedRef(); }
+
+	virtual TSharedRef<SWidget> GetChildAt( int32 ChildIndex ) override
+	{
+		check(ChildIndex == 0);
+		return GetWidget();
+	}
+
+	virtual TSharedRef<const SWidget> GetChildAt( int32 ChildIndex ) const override
+	{
+		check(ChildIndex == 0);
+		return GetWidget();
+	}
 
 private:
 	virtual const FSlotBase& GetSlotAt(int32 ChildIndex) const override { static FSlotBase NullSlot; check(ChildIndex == 0); return NullSlot; }
@@ -144,12 +165,34 @@ public:
 		}
 	}
 
+	void DetachWidget()
+	{
+		if (WidgetPtr.IsValid())
+		{
+			WidgetPtr.Reset();
+
+			if (Owner)
+			{
+				Owner->InvalidatePrepass();
+			}
+		}
+	}
+
 	TSharedRef<SWidget> GetWidget() const
 	{
+		ensure(Num() > 0);
 		TSharedPtr<SWidget> Widget = WidgetPtr.Pin();
 		return (Widget.IsValid()) ? Widget.ToSharedRef() : SNullWidget::NullWidget;
 	}
 
+private:
+	SWidget& GetWidgetRef() const
+	{
+		ensure(Num() > 0);
+		SWidget* Widget = WidgetPtr.Pin().Get();
+		return Widget ? *Widget : SNullWidget::NullWidget.Get();
+	}
+
 private:
 	TWeakPtr<ChildType> WidgetPtr;
 };
@@ -609,3 +652,36 @@ public:
 private:
 	bool bChangesInvalidatePrepass;
 };
+
+
+/** Required to implement GetChildren() in a way that can dynamically return the currently active child. */
+template<typename SlotType>
+class TOneDynamicChild : public FChildren
+{
+public:
+	TOneDynamicChild(SWidget* InOwner, TPanelChildren<SlotType>* InAllChildren, const TAttribute<int32>* InWidgetIndex)
+		: FChildren(InOwner)
+		, AllChildren(InAllChildren)
+		, WidgetIndex(InWidgetIndex)
+	{ }
+
+	virtual int32 Num() const override { return AllChildren->Num() > 0 ? 1 : 0; }
+
+	virtual TSharedRef<SWidget> GetChildAt(int32 Index) override
+	{
+		check(Index == 0); return AllChildren->GetChildAt(WidgetIndex->Get());
+	}
+
+	virtual TSharedRef<const SWidget> GetChildAt(int32 Index) const override
+	{
+		check(Index == 0);
+		return AllChildren->GetChildAt(WidgetIndex->Get());
+	}
+
+private:
+
+	virtual const FSlotBase& GetSlotAt(int32 ChildIndex) const override { return (*AllChildren)[ChildIndex]; }
+
+	TPanelChildren<SlotType>* AllChildren;
+	const TAttribute<int32>* WidgetIndex;
+};
\ No newline at end of file
diff --git a/Engine/Source/Runtime/SlateCore/Public/Rendering/DrawElements.h b/Engine/Source/Runtime/SlateCore/Public/Rendering/DrawElements.h
--- a/Engine/Source/Runtime/SlateCore/Public/Rendering/DrawElements.h
+++ b/Engine/Source/Runtime/SlateCore/Public/Rendering/DrawElements.h
@@ -1372,7 +1372,7 @@ public:
 	 * @param InPaintWindow		The window that owns the widgets being painted.  This is almost most always the same window that is being rendered to
 	 * @param InRenderWindow	The window that we will be rendering to.
 	 */
-	SLATECORE_API explicit FSlateWindowElementList(TSharedPtr<SWindow> InPaintWindow = nullptr);
+	SLATECORE_API explicit FSlateWindowElementList(const TSharedPtr<SWindow>& InPaintWindow = TSharedPtr<SWindow>());
 
 	SLATECORE_API ~FSlateWindowElementList();
 
