Avoiding stencil clipping on widgets that are just barely not axis aligned.

diff --git a/Engine/Source/Runtime/SlateCore/Private/Layout/Clipping.cpp b/Engine/Source/Runtime/SlateCore/Private/Layout/Clipping.cpp
--- a/Engine/Source/Runtime/SlateCore/Private/Layout/Clipping.cpp
+++ b/Engine/Source/Runtime/SlateCore/Private/Layout/Clipping.cpp
@@ -77,11 +77,19 @@ FSlateClippingZone::FSlateClippingZone(const FVector2D& InTopLeft, const FVector
 
 void FSlateClippingZone::InitializeFromArbitraryPoints(const FVector2D& InTopLeft, const FVector2D& InTopRight, const FVector2D& InBottomLeft, const FVector2D& InBottomRight)
 {
-	bIsAxisAligned =
-		FMath::RoundToInt(InTopLeft.X) == FMath::RoundToInt(InBottomLeft.X) &&
-		FMath::RoundToInt(InTopRight.X) == FMath::RoundToInt(InBottomRight.X) &&
-		FMath::RoundToInt(InTopLeft.Y) == FMath::RoundToInt(InTopRight.Y) &&
-		FMath::RoundToInt(InBottomLeft.Y) == FMath::RoundToInt(InBottomRight.Y);
+	FVector2D Points[4] = { InTopLeft, InTopRight, InBottomLeft, InBottomRight };
+
+	FBox2D AABB(Points, 4);
+
+	bIsAxisAligned = true;
+	for (int32 i = 0; i < 4; ++i)
+	{
+		if (!AABB.GetClosestPointTo(Points[i]).Equals(Points[i], 0.1f))
+		{
+			bIsAxisAligned = false;
+			break;
+		}
+	}
 
 	if ( bIsAxisAligned )
 	{
