Optimized memory and performance rendering very long text strings by culling glyphs outside clipping area

diff --git a/Engine/Source/Runtime/SlateCore/Private/Rendering/ElementBatcher.cpp b/Engine/Source/Runtime/SlateCore/Private/Rendering/ElementBatcher.cpp
--- a/Engine/Source/Runtime/SlateCore/Private/Rendering/ElementBatcher.cpp
+++ b/Engine/Source/Runtime/SlateCore/Private/Rendering/ElementBatcher.cpp
@@ -920,6 +920,26 @@ void FSlateElementBatcher::AddShapedTextElement( const FSlateDrawElement& DrawEl
 
 		const bool bIsFontMaterial = FontMaterial != nullptr;
 
+		// Optimize by culling
+		bool bEnableCulling = false;
+		float LocalClipBoundingBoxLeft = 0;
+		float LocalClipBoundingBoxRight = 0;
+		if (GlyphsToRender.Num() > 200)
+		{
+			int16 ClippingIndex = DrawElement.GetClippingIndex();
+			if (ClippingStates && ClippingStates->IsValidIndex(ClippingIndex))
+			{
+				const FSlateClippingState& ClippingState = (*ClippingStates)[ClippingIndex];
+				if (ClippingState.ScissorRect.IsSet() && ClippingState.ScissorRect->IsAxisAligned() && RenderTransform.GetMatrix().IsIdentity())
+				{
+					bEnableCulling = true;
+					const FSlateRect LocalClipBoundingBox = TransformRect(RenderTransform.Inverse(), ClippingState.ScissorRect->GetBoundingBox());
+					LocalClipBoundingBoxLeft = LocalClipBoundingBox.Left;
+					LocalClipBoundingBoxRight = LocalClipBoundingBox.Right;
+				}
+			}
+		}
+
 		const int32 NumGlyphs = GlyphsToRender.Num();
 		for (int32 GlyphIndex = 0; GlyphIndex < NumGlyphs; ++GlyphIndex)
 		{
@@ -931,6 +951,23 @@ void FSlateElementBatcher::AddShapedTextElement( const FSlateDrawElement& DrawEl
 				 
 				if (GlyphAtlasData.Valid)
 				{
+					const float X = LineX + GlyphAtlasData.HorizontalOffset + GlyphToRender.XOffset;
+					// Note PosX,PosY is the upper left corner of the bounding box representing the string.  This computes the Y position of the baseline where text will sit
+
+					if (bEnableCulling)
+					{
+						if (X + GlyphAtlasData.USize < LocalClipBoundingBoxLeft)
+						{
+							LineX += GlyphToRender.XAdvance;
+							LineY += GlyphToRender.YAdvance;
+							continue;
+						}
+						else if (X > LocalClipBoundingBoxRight)
+						{
+							break;
+						}
+					}
+
 					if (FontAtlasTexture == nullptr || GlyphAtlasData.TextureIndex != FontTextureIndex)
 					{
 						// Font has a new texture for this glyph. Refresh the batch we use and the index we are currently using
@@ -954,8 +991,6 @@ void FSlateElementBatcher::AddShapedTextElement( const FSlateDrawElement& DrawEl
 						InvTextureSizeY = 1.0f / FontAtlasTexture->GetHeight();
 					}
 
-					const float X = LineX + GlyphAtlasData.HorizontalOffset + GlyphToRender.XOffset;
-					// Note PosX,PosY is the upper left corner of the bounding box representing the string.  This computes the Y position of the baseline where text will sit
 
 					const float Y = LineY - GlyphAtlasData.VerticalOffset + GlyphToRender.YOffset + MaxHeight + TextBaseline;
 					const float U = GlyphAtlasData.StartU * InvTextureSizeX;
