Added CVar to disable legacy localized fallback font support, and added a warning if the legacy fallback fonts are loaded.

Slate.EnableLegacyLocalizedFallbackFont is the CVar. It currently defaults to enabled to preserve old behavior, but will be set to disabled in the next version, and removed in the one after that.

Please ensure your composite fonts use localized sub-font families: https://docs.unrealengine.com/en-US/Engine/UMG/UserGuide/Fonts/Overview

diff --git a/Engine/Source/Runtime/SlateCore/Private/Fonts/LegacySlateFontInfoCache.cpp b/Engine/Source/Runtime/SlateCore/Private/Fonts/LegacySlateFontInfoCache.cpp
--- a/Engine/Source/Runtime/SlateCore/Private/Fonts/LegacySlateFontInfoCache.cpp
+++ b/Engine/Source/Runtime/SlateCore/Private/Fonts/LegacySlateFontInfoCache.cpp
@@ -1,12 +1,43 @@
 // Copyright 1998-2019 Epic Games, Inc. All Rights Reserved.
 
 #include "Fonts/LegacySlateFontInfoCache.h"
+#include "HAL/IConsoleManager.h"
+#include "SlateGlobals.h"
 #include "Misc/Paths.h"
 #include "Misc/ScopeLock.h"
 #include "Misc/FileHelper.h"
 #include "Fonts/FontProviderInterface.h"
 #include "Fonts/UnicodeBlockRange.h"
 
+static int32 GSlateEnableLegacyLocalizedFallbackFont = 1;
+static FAutoConsoleVariableRef CVarSlateEnableLocalizedFallbackFont(TEXT("Slate.EnableLegacyLocalizedFallbackFont"), GSlateEnableLegacyLocalizedFallbackFont, TEXT("Enable the legacy localized fallback fonts? (0/1)."), ECVF_Default);
+
+FString FLegacySlateFontInfoCache::FFallbackContext::ToString() const
+{
+	FString FontName;
+	if (FontData)
+	{
+		const UObject* FontAsset = FontData->GetFontFaceAsset();
+		FontName = (FontAsset) ? FontAsset->GetPathName() : FontData->GetFontFilename();
+	}
+	else
+	{
+		FontName = TEXT("<unknown>");
+	}
+
+	FString CharacterInfo;
+	if (Char)
+	{
+		CharacterInfo = FString::Printf(TEXT("%c (0x%04x)"), Char, (int32)Char);
+	}
+	else
+	{
+		CharacterInfo = TEXT("<unknown>");
+	}
+
+	return FString::Printf(TEXT("Font: '%s', Character: '%s'"), *FontName, *CharacterInfo);
+}
+
 TSharedPtr<FLegacySlateFontInfoCache> FLegacySlateFontInfoCache::Singleton;
 
 FLegacySlateFontInfoCache& FLegacySlateFontInfoCache::Get()
@@ -202,7 +233,12 @@ TSharedPtr<const FCompositeFont> FLegacySlateFontInfoCache::GetSystemFont()
 	return SystemFont;
 }
 
-const FFontData& FLegacySlateFontInfoCache::GetLocalizedFallbackFontData()
+bool FLegacySlateFontInfoCache::IsLocalizedFallbackFontAvailable() const
+{
+	return !!GSlateEnableLegacyLocalizedFallbackFont;
+}
+
+const FFontData& FLegacySlateFontInfoCache::GetLocalizedFallbackFontData(const FFallbackContext& InContext)
 {
 	// GetLocalizedFallbackFontData is called directly from the font cache, so may be called from multiple threads at once
 	FScopeLock Lock(&LocalizedFallbackFontDataCS);
@@ -225,6 +261,8 @@ const FFontData& FLegacySlateFontInfoCache::GetLocalizedFallbackFontData()
 
 		if (!LocalizedFallbackFontData.IsValid())
 		{
+			UE_LOG(LogSlate, Warning, TEXT("Legacy localized fallback font '%s' was requested. %s\nLegacy localized fallback fonts were deprecated in 4.22 and will be removed in a future version!\nPlease update your composite fonts to use localized sub-font families: https://docs.unrealengine.com/en-US/Engine/UMG/UserGuide/Fonts/Overview"), *FallbackFontPath, *InContext.ToString());
+
 			LocalizedFallbackFontData = MakeShared<FFontData>(FallbackFontPath, EFontHinting::Default, EFontLoadingPolicy::LazyLoad);
 			AllLocalizedFallbackFontData.Add(FallbackFontPath, LocalizedFallbackFontData);
 		}
@@ -256,20 +294,22 @@ TSharedPtr<const FCompositeFont> FLegacySlateFontInfoCache::GetLastResortFont()
 
 	if (!LastResortFont.IsValid() && bIsLastResortFontAvailable)
 	{
-		const FFontData& FontData = GetLastResortFontData();
+		const FFontData& FontData = GetLastResortFontData(FFallbackContext());
 		LastResortFont = MakeShared<FStandaloneCompositeFont>(NAME_None, FontData.GetFontFilename(), FontData.GetHinting(), FontData.GetLoadingPolicy());
 	}
 
 	return LastResortFont;
 }
 
-const FFontData& FLegacySlateFontInfoCache::GetLastResortFontData()
+const FFontData& FLegacySlateFontInfoCache::GetLastResortFontData(const FFallbackContext& InContext)
 {
 	// GetLastResortFontData is called directly from the font cache, so may be called from multiple threads at once
 	FScopeLock Lock(&LastResortFontDataCS);
 
 	if (!LastResortFontData.IsValid())
 	{
+		UE_LOG(LogSlate, Log, TEXT("Last resort fallback font was requested. %s"), *InContext.ToString());
+
 		LastResortFontData = MakeShared<FFontData>(bIsLastResortFontAvailable ? LastResortFontPath : FString(), EFontHinting::Default, EFontLoadingPolicy::LazyLoad);
 	}
 
diff --git a/Engine/Source/Runtime/SlateCore/Private/Fonts/LegacySlateFontInfoCache.h b/Engine/Source/Runtime/SlateCore/Private/Fonts/LegacySlateFontInfoCache.h
--- a/Engine/Source/Runtime/SlateCore/Private/Fonts/LegacySlateFontInfoCache.h
+++ b/Engine/Source/Runtime/SlateCore/Private/Fonts/LegacySlateFontInfoCache.h
@@ -12,6 +12,27 @@
 class FLegacySlateFontInfoCache : public FGCObject, public TSharedFromThis<FLegacySlateFontInfoCache>
 {
 public:
+	/**
+	 * Context used to help debug font fallback requests
+	 */
+	struct FFallbackContext
+	{
+	public:
+		FFallbackContext() = default;
+
+		FFallbackContext(const FFontData* InFontData, const TCHAR InChar)
+			: FontData(InFontData)
+			, Char(InChar)
+		{
+		}
+
+		FString ToString() const;
+
+	private:
+		const FFontData* FontData = nullptr;
+		TCHAR Char = 0;
+	};
+
 	/**
 	 * Get (or create) the singleton instance of this cache
 	 */
@@ -32,10 +53,15 @@ public:
 	 */
 	TSharedPtr<const FCompositeFont> GetSystemFont();
 
+	/**
+	 * Is the localized fallback font available (can be controlled via a CVar).
+	 */
+	bool IsLocalizedFallbackFontAvailable() const;
+
 	/**
 	 * Get (or create) the culture specific fallback font
 	 */
-	const FFontData& GetLocalizedFallbackFontData();
+	const FFontData& GetLocalizedFallbackFontData(const FFallbackContext& InContext);
 
 	/**
 	 * Get the revision index of the currently active localized fallback font.
@@ -43,7 +69,7 @@ public:
 	uint16 GetLocalizedFallbackFontRevision() const;
 
 	/**
-	 * Is that last resort fallback font available (not all builds have it).
+	 * Is the last resort fallback font available (not all builds have it).
 	 */
 	bool IsLastResortFontAvailable() const;
 
@@ -55,7 +81,7 @@ public:
 	/**
 	 * Get (or create) the last resort fallback font
 	 */
-	const FFontData& GetLastResortFontData();
+	const FFontData& GetLastResortFontData(const FFallbackContext& InContext);
 
 	// FGCObject interface
 	virtual void AddReferencedObjects(FReferenceCollector& Collector) override;
diff --git a/Engine/Source/Runtime/SlateCore/Private/Fonts/SlateFontRenderer.cpp b/Engine/Source/Runtime/SlateCore/Private/Fonts/SlateFontRenderer.cpp
--- a/Engine/Source/Runtime/SlateCore/Private/Fonts/SlateFontRenderer.cpp
+++ b/Engine/Source/Runtime/SlateCore/Private/Fonts/SlateFontRenderer.cpp
@@ -242,9 +242,9 @@ FFreeTypeFaceGlyphData FSlateFontRenderer::GetFontFaceForCharacter(const FFontDa
 	{
 		const bool bCanFallback = bOverrideFallback || MaxFallbackLevel >= EFontFallback::FF_LocalizedFallback;
 
-		if (bCanFallback)
+		if (bCanFallback && FLegacySlateFontInfoCache::Get().IsLocalizedFallbackFontAvailable())
 		{
-			ReturnVal.FaceAndMemory = CompositeFontCache->GetFontFace(FLegacySlateFontInfoCache::Get().GetLocalizedFallbackFontData());
+			ReturnVal.FaceAndMemory = CompositeFontCache->GetFontFace(FLegacySlateFontInfoCache::Get().GetLocalizedFallbackFontData(FLegacySlateFontInfoCache::FFallbackContext(&InFontData, Char)));
 
 			if (ReturnVal.FaceAndMemory.IsValid())
 			{	
@@ -266,7 +266,7 @@ FFreeTypeFaceGlyphData FSlateFontRenderer::GetFontFaceForCharacter(const FFontDa
 
 		if (bCanFallback && FLegacySlateFontInfoCache::Get().IsLastResortFontAvailable())
 		{
-			ReturnVal.FaceAndMemory = CompositeFontCache->GetFontFace(FLegacySlateFontInfoCache::Get().GetLastResortFontData());
+			ReturnVal.FaceAndMemory = CompositeFontCache->GetFontFace(FLegacySlateFontInfoCache::Get().GetLastResortFontData(FLegacySlateFontInfoCache::FFallbackContext(&InFontData, Char)));
 
 			if (ReturnVal.FaceAndMemory.IsValid())
 			{
