Removed dependency on internal HarfBuzz API

diff --git a/Engine/Source/Runtime/SlateCore/Private/Fonts/FontCacheHarfBuzz.cpp b/Engine/Source/Runtime/SlateCore/Private/Fonts/FontCacheHarfBuzz.cpp
--- a/Engine/Source/Runtime/SlateCore/Private/Fonts/FontCacheHarfBuzz.cpp
+++ b/Engine/Source/Runtime/SlateCore/Private/Fonts/FontCacheHarfBuzz.cpp
@@ -5,36 +5,6 @@
 #include "Fonts/FontCacheFreeType.h"
 #include "Fonts/SlateFontRenderer.h"
 
-#if WITH_HARFBUZZ
-	#if PLATFORM_WINDOWS
-		#include "Windows/WindowsHWrapper.h"
-		#include "Windows/AllowWindowsPlatformAtomics.h"
-	#endif
-	#if PLATFORM_WINDOWS || PLATFORM_XBOXONE
-		#pragma warning(push)
-		#pragma warning(disable:4996) // warning C4996: 'strncpy': This function or variable may be unsafe. Consider using strncpy_s instead.
-		#pragma warning(disable:28113) // warning C28113: Accessing a local variable dummy via an Interlocked function:  This is an unusual usage which could be reconsidered.
-		#ifndef snprintf
-			#define snprintf _snprintf
-		#endif
-	#endif // #if PLATFORM_WINDOWS
-
-	// Include some private headers needed for our font implementation
-	PRAGMA_DISABLE_DEPRECATION_WARNINGS
-	THIRD_PARTY_INCLUDES_START
-	#include "hb-private.hh"
-	#include "hb-font-private.hh"
-	THIRD_PARTY_INCLUDES_END
-	PRAGMA_ENABLE_DEPRECATION_WARNINGS
-
-	#if PLATFORM_WINDOWS || PLATFORM_XBOXONE
-		#pragma warning(pop)
-	#endif // #if PLATFORM_WINDOWS
-	#if PLATFORM_WINDOWS
-		#include "Windows/HideWindowsPlatformAtomics.h"
-	#endif
-#endif // #if WITH_HARFBUZZ
-
 #if WITH_HARFBUZZ
 
 extern "C"
@@ -159,14 +129,16 @@ namespace Internal
 
 FORCEINLINE FT_Face get_ft_face(hb_font_t* InFont)
 {
-	check(InFont->parent);
-	return hb_ft_font_get_face(InFont->parent);
+	hb_font_t* FontParent = hb_font_get_parent(InFont);
+	check(FontParent);
+	return hb_ft_font_get_face(FontParent);
 }
 
 FORCEINLINE int32 get_ft_flags(hb_font_t* InFont)
 {
-	check(InFont->parent);
-	return hb_ft_font_get_load_flags(InFont->parent);
+	hb_font_t* FontParent = hb_font_get_parent(InFont);
+	check(FontParent);
+	return hb_ft_font_get_load_flags(FontParent);
 }
 
 hb_bool_t get_nominal_glyph(hb_font_t* InFont, void* InFontData, hb_codepoint_t InUnicodeChar, hb_codepoint_t *OutGlyphIndex, void *InUserData)
@@ -228,12 +200,16 @@ hb_bool_t get_glyph_v_origin(hb_font_t* InFont, void* InFontData, hb_codepoint_t
 		*OutX = CachedGlyphData.GlyphMetrics.horiBearingX -   CachedGlyphData.GlyphMetrics.vertBearingX;
 		*OutY = CachedGlyphData.GlyphMetrics.horiBearingY - (-CachedGlyphData.GlyphMetrics.vertBearingY);
 
-		if (InFont->x_scale < 0)
+		int FontXScale = 0;
+		int FontYScale = 0;
+		hb_font_get_scale(InFont, &FontXScale, &FontYScale);
+
+		if (FontXScale < 0)
 		{
 			*OutX = -*OutX;
 		}
 
-		if (InFont->y_scale < 0)
+		if (FontYScale < 0)
 		{
 			*OutY = -*OutY;
 		}
@@ -352,9 +328,17 @@ hb_font_t* FHarfBuzzFontFactory::CreateFont(const FFreeTypeFace& InFace, const u
 		hb_ft_font_set_load_flags(HarfBuzzFTFont, InGlyphFlags);
 
 		// The default FreeType implementation doesn't apply the font scale, so we have to do that ourselves (in 16.16 space for maximum precision)
-		const FT_Long FixedFontScale = FreeTypeUtils::ConvertPixelTo16Dot16<FT_Long>(InFontScale);
-		HarfBuzzFTFont->x_scale = FT_MulFix(HarfBuzzFTFont->x_scale, FixedFontScale);
-		HarfBuzzFTFont->y_scale = FT_MulFix(HarfBuzzFTFont->y_scale, FixedFontScale);
+		{
+			int HarfBuzzFTFontXScale = 0;
+			int HarfBuzzFTFontYScale = 0;
+			hb_font_get_scale(HarfBuzzFTFont, &HarfBuzzFTFontXScale, &HarfBuzzFTFontYScale);
+
+			const FT_Long FixedFontScale = FreeTypeUtils::ConvertPixelTo16Dot16<FT_Long>(InFontScale);
+			HarfBuzzFTFontXScale = FT_MulFix(HarfBuzzFTFontXScale, FixedFontScale);
+			HarfBuzzFTFontYScale = FT_MulFix(HarfBuzzFTFontYScale, FixedFontScale);
+
+			hb_font_set_scale(HarfBuzzFTFont, HarfBuzzFTFontXScale, HarfBuzzFTFontYScale);
+		}
 
 		HarfBuzzFont = hb_font_create_sub_font(HarfBuzzFTFont);
 		
