Slate: Fix non-axis aligned stencil clipping still being a problem

diff --git a/Engine/Source/Runtime/SlateCore/Private/Layout/Clipping.cpp b/Engine/Source/Runtime/SlateCore/Private/Layout/Clipping.cpp
--- a/Engine/Source/Runtime/SlateCore/Private/Layout/Clipping.cpp
+++ b/Engine/Source/Runtime/SlateCore/Private/Layout/Clipping.cpp
@@ -77,17 +77,24 @@ FSlateClippingZone::FSlateClippingZone(const FVector2D& InTopLeft, const FVector
 
 void FSlateClippingZone::InitializeFromArbitraryPoints(const FVector2D& InTopLeft, const FVector2D& InTopRight, const FVector2D& InBottomLeft, const FVector2D& InBottomRight)
 {
-	FVector2D Points[4] = { InTopLeft, InTopRight, InBottomLeft, InBottomRight };
+	bIsAxisAligned = false;
 
-	FBox2D AABB(Points, 4);
+	// Clipping is in pixel space, accept a very high tolerance
+	const float Tolerance = .1;
 
-	bIsAxisAligned = true;
-	for (int32 i = 0; i < 4; ++i)
+	// Since this is a rectangle check to edges.  If their points are equal they are aligned to the same axis and thus the whole rect is aligned
+	if (FMath::IsNearlyEqual(InTopLeft.X, InBottomLeft.X, Tolerance))
 	{
-		if (!AABB.GetClosestPointTo(Points[i]).Equals(Points[i], 0.1f))
+		if (FMath::IsNearlyEqual(InBottomLeft.Y, InBottomLeft.Y, Tolerance))
 		{
-			bIsAxisAligned = false;
-			break;
+			bIsAxisAligned = true;
+		}
+	}
+	else if (FMath::IsNearlyEqual(InTopLeft.Y, InBottomLeft.Y, Tolerance))
+	{
+		if (FMath::IsNearlyEqual(InBottomLeft.X, InBottomRight.X, Tolerance))
+		{
+			bIsAxisAligned = true;
 		}
 	}
 
