Fixed infinite recursion that could happen when logging font load warnings with the Output Log open

This would log before setting the data, resulting in it constantly being requested by each log addition.

diff --git a/Engine/Source/Runtime/SlateCore/Private/Fonts/LegacySlateFontInfoCache.cpp b/Engine/Source/Runtime/SlateCore/Private/Fonts/LegacySlateFontInfoCache.cpp
--- a/Engine/Source/Runtime/SlateCore/Private/Fonts/LegacySlateFontInfoCache.cpp
+++ b/Engine/Source/Runtime/SlateCore/Private/Fonts/LegacySlateFontInfoCache.cpp
@@ -28,7 +28,7 @@ FString FLegacySlateFontInfoCache::FFallbackContext::ToString() const
 	FString CharacterInfo;
 	if (Char)
 	{
-		CharacterInfo = FString::Printf(TEXT("%c (0x%04x)"), Char, (int32)Char);
+		CharacterInfo = FString::Printf(TEXT("%c (U+%04x)"), Char, (int32)Char);
 	}
 	else
 	{
@@ -261,10 +261,9 @@ const FFontData& FLegacySlateFontInfoCache::GetLocalizedFallbackFontData(const F
 
 		if (!LocalizedFallbackFontData.IsValid())
 		{
-			UE_LOG(LogSlate, Warning, TEXT("Legacy localized fallback font '%s' was requested. %s\nLegacy localized fallback fonts were deprecated in 4.22 and will be removed in a future version!\nPlease update your composite fonts to use localized sub-font families: https://docs.unrealengine.com/en-US/Engine/UMG/UserGuide/Fonts/Overview"), *FallbackFontPath, *InContext.ToString());
-
 			LocalizedFallbackFontData = MakeShared<FFontData>(FallbackFontPath, EFontHinting::Default, EFontLoadingPolicy::LazyLoad);
 			AllLocalizedFallbackFontData.Add(FallbackFontPath, LocalizedFallbackFontData);
+			UE_LOG(LogSlate, Warning, TEXT("Legacy localized fallback font '%s' was requested. %s\nLegacy localized fallback fonts were deprecated in 4.22 and will be removed in a future version!\nPlease update your composite fonts to use localized sub-font families: https://docs.unrealengine.com/en-US/Engine/UMG/UserGuide/Fonts/Overview"), *FallbackFontPath, *InContext.ToString());
 		}
 
 		if (LocalizedFallbackFontData != PreviousLocalizedFallbackFontData)
@@ -308,9 +307,8 @@ const FFontData& FLegacySlateFontInfoCache::GetLastResortFontData(const FFallbac
 
 	if (!LastResortFontData.IsValid())
 	{
-		UE_LOG(LogSlate, Log, TEXT("Last resort fallback font was requested. %s"), *InContext.ToString());
-
 		LastResortFontData = MakeShared<FFontData>(bIsLastResortFontAvailable ? LastResortFontPath : FString(), EFontHinting::Default, EFontLoadingPolicy::LazyLoad);
+		UE_LOG(LogSlate, Log, TEXT("Last resort fallback font was requested. %s"), *InContext.ToString());
 	}
 
 	return *LastResortFontData;
