Redo removal of redundant delegates on SImage which was lost by 4743892

diff --git a/Engine/Source/Runtime/SlateCore/Private/Widgets/Images/SImage.cpp b/Engine/Source/Runtime/SlateCore/Private/Widgets/Images/SImage.cpp
--- a/Engine/Source/Runtime/SlateCore/Private/Widgets/Images/SImage.cpp
+++ b/Engine/Source/Runtime/SlateCore/Private/Widgets/Images/SImage.cpp
@@ -8,7 +8,9 @@ void SImage::Construct( const FArguments& InArgs )
 	Image = InArgs._Image;
 	ColorAndOpacity = InArgs._ColorAndOpacity;
 	bFlipForRightToLeftFlowDirection = InArgs._FlipForRightToLeftFlowDirection;
-	OnMouseButtonDownHandler = InArgs._OnMouseButtonDown;
+	SetOnMouseButtonDown(InArgs._OnMouseButtonDown);
+
+
 }
 
 int32 SImage::OnPaint( const FPaintArgs& Args, const FGeometry& AllottedGeometry, const FSlateRect& MyCullingRect, FSlateWindowElementList& OutDrawElements, int32 LayerId, const FWidgetStyle& InWidgetStyle, bool bParentEnabled ) const
@@ -36,20 +38,6 @@ int32 SImage::OnPaint( const FPaintArgs& Args, const FGeometry& AllottedGeometry
 	return LayerId;
 }
 
-FReply SImage::OnMouseButtonDown( const FGeometry& MyGeometry, const FPointerEvent& MouseEvent )
-{
-	if ( OnMouseButtonDownHandler.IsBound() )
-	{
-		// If a handler is assigned, call it.
-		return OnMouseButtonDownHandler.Execute(MyGeometry, MouseEvent);
-	}
-	else
-	{
-		// otherwise the event is unhandled.
-		return FReply::Unhandled();
-	}	
-}
-
 FVector2D SImage::ComputeDesiredSize( float ) const
 {
 	const FSlateBrush* ImageBrush = Image.Get();
@@ -86,8 +74,3 @@ void SImage::SetImage(TAttribute<const FSlateBrush*> InImage)
 		Invalidate(EInvalidateWidget::LayoutAndVolatility);
 	}
 }
-
-void SImage::SetOnMouseButtonDown(FPointerEventHandler EventHandler)
-{
-	OnMouseButtonDownHandler = EventHandler;
-}
diff --git a/Engine/Source/Runtime/SlateCore/Public/Widgets/Images/SImage.h b/Engine/Source/Runtime/SlateCore/Public/Widgets/Images/SImage.h
--- a/Engine/Source/Runtime/SlateCore/Public/Widgets/Images/SImage.h
+++ b/Engine/Source/Runtime/SlateCore/Public/Widgets/Images/SImage.h
@@ -24,7 +24,6 @@ public:
 		: _Image( FCoreStyle::Get().GetDefaultBrush() )
 		, _ColorAndOpacity( FLinearColor::White )
 		, _FlipForRightToLeftFlowDirection( false )
-		, _OnMouseButtonDown()
 		{}
 
 		/** Image resource */
@@ -37,8 +36,7 @@ public:
 		SLATE_ARGUMENT( bool, FlipForRightToLeftFlowDirection )
 
 		/** Invoked when the mouse is pressed in the widget. */
-		SLATE_EVENT( FPointerEventHandler, OnMouseButtonDown )
-
+		SLATE_EVENT(FPointerEventHandler, OnMouseButtonDown)
 	SLATE_END_ARGS()
 
 	/** Constructor */
@@ -65,16 +63,11 @@ public:
 
 	/** See the Image attribute */
 	void SetImage(TAttribute<const FSlateBrush*> InImage);
-	
-	/** See OnMouseButtonDown event */
-	void SetOnMouseButtonDown(FPointerEventHandler EventHandler);
 
 public:
 
 	// SWidget overrides
 	virtual int32 OnPaint( const FPaintArgs& Args, const FGeometry& AllottedGeometry, const FSlateRect& MyCullingRect, FSlateWindowElementList& OutDrawElements, int32 LayerId, const FWidgetStyle& InWidgetStyle, bool bParentEnabled ) const override;
-	virtual FReply OnMouseButtonDown( const FGeometry& MyGeometry, const FPointerEvent& MouseEvent ) override;
-
 protected:
 	// Begin SWidget overrides.
 	virtual FVector2D ComputeDesiredSize(float) const override;
