Slate - Adding a note as a reminder how the parents end up getting invalidated when you remove slots from the children set.

diff --git a/Engine/Source/Runtime/SlateCore/Public/Layout/Children.h b/Engine/Source/Runtime/SlateCore/Public/Layout/Children.h
--- a/Engine/Source/Runtime/SlateCore/Public/Layout/Children.h
+++ b/Engine/Source/Runtime/SlateCore/Public/Layout/Children.h
@@ -324,6 +324,9 @@ public:
 	{
 		if ( !bEmptying )
 		{
+			// NOTE:
+			// We don't do any invalidating here, that's handled by the FSlotBase, which eventually calls ConditionallyDetatchParentWidget
+
 			TIndirectArray< SlotType >::RemoveAt(Index);
 		}
 	}
@@ -332,6 +335,9 @@ public:
 	{
 		if ( !bEmptying )
 		{
+			// NOTE:
+			// We don't do any invalidating here, that's handled by the FSlotBase, which eventually calls ConditionallyDetatchParentWidget
+
 			TGuardValue<bool> GuardEmptying(bEmptying, true);
 			TIndirectArray< SlotType >::Empty();
 		}
