Accessibility optimization pass, final major API improvements before experimental release

diff --git a/Engine/Source/Runtime/ApplicationCore/Private/IOS/Accessibility/IOSAccessibilityCache.cpp b/Engine/Source/Runtime/ApplicationCore/Private/IOS/Accessibility/IOSAccessibilityCache.cpp
--- a/Engine/Source/Runtime/ApplicationCore/Private/IOS/Accessibility/IOSAccessibilityCache.cpp
+++ b/Engine/Source/Runtime/ApplicationCore/Private/IOS/Accessibility/IOSAccessibilityCache.cpp
@@ -11,9 +11,25 @@
 
 @implementation FIOSAccessibilityCache
 
+static void DumpAccessibilityStatsForwarder()
+{
+	[[FIOSAccessibilityCache AccessibilityElementCache] DumpAccessibilityStats];
+}
+
 - (id)init
 {
 	Cache = [[NSMutableDictionary alloc] init];
+
+#if !UE_BUILD_SHIPPING
+	IConsoleManager::Get().RegisterConsoleCommand
+	(
+		TEXT("DumpAccessibilityStatsIOS"),
+		TEXT("Writes to LogAccessibility the memory stats for the platform-level accessibility data (AccessibilityElements) required for IOS support."),
+		FConsoleCommandDelegate::CreateStatic(&DumpAccessibilityStatsForwarder),
+		ECVF_Default
+	);
+#endif
+
 	return self;
 }
 
@@ -29,26 +45,11 @@
 	{
 		return nil;
 	}
-
+	
 	FIOSAccessibilityContainer* ExistingElement = [Cache objectForKey:[NSNumber numberWithInt:Id]];
 	if (ExistingElement == nil)
 	{
-		AccessibleWidgetId ParentId = IAccessibleWidget::InvalidAccessibleWidgetId;
-		// All IAccessibleWidget functions must be run on Game Thread
-		[IOSAppDelegate WaitAndRunOnGameThread : [Id, &ParentId]()
-		{
-			TSharedPtr<IAccessibleWidget> Widget = [IOSAppDelegate GetDelegate].IOSApplication->GetAccessibleMessageHandler()->GetAccessibleWidgetFromId(Id);
-			if (Widget.IsValid())
-			{
-				TSharedPtr<IAccessibleWidget> ParentWidget = Widget->GetParent();
-				if (ParentWidget.IsValid())
-				{
-					ParentId = ParentWidget->GetId();
-				}
-			}
-		}];
-
-		ExistingElement = [[[FIOSAccessibilityContainer alloc] initWithId:Id AndParentId:ParentId] autorelease];
+		ExistingElement = [[[FIOSAccessibilityContainer alloc] initWithId:Id] autorelease];
 		[Cache setObject:ExistingElement forKey:[NSNumber numberWithInt:Id]];
 	}
 	return ExistingElement;
@@ -134,6 +135,35 @@
 	}
 }
 
+#if !UE_BUILD_SHIPPING
+-(void)DumpAccessibilityStats
+{
+	const uint32 NumContainers = [Cache count];
+	uint32 SizeOfContainer = 0;
+	uint32 SizeOfLeaf = 0;
+	uint32 CacheSize = 0;
+
+	NSArray* Keys = [Cache allKeys];
+	for (NSString* Key in Keys)
+	{
+		FIOSAccessibilityContainer* Container = [Cache objectForKey : Key];
+		FIOSAccessibilityLeaf* Leaf = [Container GetLeaf];
+
+		SizeOfContainer = malloc_size(Container);
+		SizeOfLeaf = malloc_size(Leaf);
+		CacheSize += sizeof(NSString*) + sizeof(FIOSAccessibilityContainer*) + malloc_size(Key) + SizeOfContainer + SizeOfLeaf
+			+ malloc_size([NSString stringWithFString : Leaf.Label])
+			+ malloc_size([NSString stringWithFString : Leaf.Hint])
+			+ malloc_size([NSString stringWithFString : Leaf.Value]);
+	}
+
+	UE_LOG(LogAccessibility, Log, TEXT("Number of Accessibility Elements: %i"), NumContainers * 2);
+	UE_LOG(LogAccessibility, Log, TEXT("Size of FIOSAccessibilityContainer: %u"), SizeOfContainer);
+	UE_LOG(LogAccessibility, Log, TEXT("Size of FIOSAccessibilityLeaf: %u"), SizeOfLeaf);
+	UE_LOG(LogAccessibility, Log, TEXT("Memory stored in cache: %u kb"), CacheSize / 1000);
+}
+#endif
+
 @end
 
 #endif
diff --git a/Engine/Source/Runtime/ApplicationCore/Private/IOS/Accessibility/IOSAccessibilityElement.cpp b/Engine/Source/Runtime/ApplicationCore/Private/IOS/Accessibility/IOSAccessibilityElement.cpp
--- a/Engine/Source/Runtime/ApplicationCore/Private/IOS/Accessibility/IOSAccessibilityElement.cpp
+++ b/Engine/Source/Runtime/ApplicationCore/Private/IOS/Accessibility/IOSAccessibilityElement.cpp
@@ -18,23 +18,33 @@
 @synthesize Bounds;
 @synthesize bIsVisible;
 
--(id)initWithId:(AccessibleWidgetId)InId AndParentId:(AccessibleWidgetId)InParentId
+-(id)initWithId:(AccessibleWidgetId)InId
 {
-	id UIParent = nil;
-	if (InParentId != IAccessibleWidget::InvalidAccessibleWidgetId)
-	{
-		UIParent = [[FIOSAccessibilityCache AccessibilityElementCache] GetAccessibilityElement:InParentId];
-	}
-	else
-	{
-		UIParent = [IOSAppDelegate GetDelegate].IOSView;
-	}
-
-	if (self = [super initWithAccessibilityContainer:UIParent])
+	if (self = [super initWithAccessibilityContainer:[IOSAppDelegate GetDelegate].IOSView])
 	{
 		self.Id = InId;
 		self.bIsVisible = true;
 		Leaf = [[FIOSAccessibilityLeaf alloc] initWithParent:self];
+
+		// Retrieve parent ID in the background. Things probably won't work quite right until
+		// this finishes but it's better than locking up the application for a while.
+		FFunctionGraphTask::CreateAndDispatchWhenReady([InId]()
+		{
+			TSharedPtr<IAccessibleWidget> Widget = [IOSAppDelegate GetDelegate].IOSApplication->GetAccessibleMessageHandler()->GetAccessibleWidgetFromId(InId);
+			if (Widget.IsValid())
+			{
+				TSharedPtr<IAccessibleWidget> Parent = Widget->GetParent();
+				if (Parent.IsValid())
+				{
+					const AccessibleWidgetId ParentId = Parent->GetId();
+					// All UIKit functions must be run on Main Thread
+					dispatch_async(dispatch_get_main_queue(), ^
+					{
+						[[[FIOSAccessibilityCache AccessibilityElementCache] GetAccessibilityElement:InId] SetParent:ParentId];
+					});
+				}
+			}
+		}, TStatId(), NULL, ENamedThreads::GameThread);
 	}
 
 	return self;
@@ -131,7 +141,7 @@
 		const TSharedPtr<IAccessibleWidget> Widget = [IOSAppDelegate GetDelegate].IOSApplication->GetAccessibleMessageHandler()->GetAccessibleWidgetFromId(TempId);
 		if (Widget.IsValid())
 		{
-			const TSharedPtr<IAccessibleWidget> HitWidget = Widget->GetTopLevelWindow()->AsWindow()->GetChildAtPosition(X, Y);
+			const TSharedPtr<IAccessibleWidget> HitWidget = Widget->GetWindow()->AsWindow()->GetChildAtPosition(X, Y);
 			if (HitWidget.IsValid())
 			{
 				FoundId = HitWidget->GetId();
diff --git a/Engine/Source/Runtime/ApplicationCore/Private/IOS/IOSAppDelegate.cpp b/Engine/Source/Runtime/ApplicationCore/Private/IOS/IOSAppDelegate.cpp
--- a/Engine/Source/Runtime/ApplicationCore/Private/IOS/IOSAppDelegate.cpp
+++ b/Engine/Source/Runtime/ApplicationCore/Private/IOS/IOSAppDelegate.cpp
@@ -1297,7 +1297,7 @@ static FAutoConsoleVariableRef CVarGEnableThermalsReport(
 				{
 					// Start caching accessibility data so that it can be returned instantly to IOS. If not cached, the data takes too long
 					// to retrieve due to cross-thread waiting and IOS will timeout.
-					Delegate.AccessibilityCacheTimer = [NSTimer scheduledTimerWithTimeInterval:0.1f target:[FIOSAccessibilityCache AccessibilityElementCache] selector:@selector(UpdateAllCachedProperties) userInfo:nil repeats:YES];
+					Delegate.AccessibilityCacheTimer = [NSTimer scheduledTimerWithTimeInterval:0.25f target:[FIOSAccessibilityCache AccessibilityElementCache] selector:@selector(UpdateAllCachedProperties) userInfo:nil repeats:YES];
 				}
 			});
 		}, TStatId(), NULL, ENamedThreads::GameThread);
diff --git a/Engine/Source/Runtime/ApplicationCore/Private/IOS/IOSApplication.cpp b/Engine/Source/Runtime/ApplicationCore/Private/IOS/IOSApplication.cpp
--- a/Engine/Source/Runtime/ApplicationCore/Private/IOS/IOSApplication.cpp
+++ b/Engine/Source/Runtime/ApplicationCore/Private/IOS/IOSApplication.cpp
@@ -260,22 +260,14 @@ void FIOSApplication::OnAccessibleEventRaised(TSharedRef<IAccessibleWidget> Widg
 	const AccessibleWidgetId Id = Widget->GetId();
 	switch (Event)
 	{
-	case EAccessibleEvent::BeforeRemoveFromParent:
-		dispatch_async(dispatch_get_main_queue(), ^{
-			[[[FIOSAccessibilityCache AccessibilityElementCache] GetAccessibilityElement:Id] SetParent:IAccessibleWidget::InvalidAccessibleWidgetId];
-		});
-		break;
-	case EAccessibleEvent::AfterAddToParent:
+	case EAccessibleEvent::ParentChanged:
 	{
-		AccessibleWidgetId ParentId = IAccessibleWidget::InvalidAccessibleWidgetId;
-		TSharedPtr<IAccessibleWidget> ParentWidget = Widget->GetParent();
-		if (ParentWidget.IsValid())
-		{
-			ParentId = ParentWidget->GetId();
-		}
 		dispatch_async(dispatch_get_main_queue(), ^{
-			[[[FIOSAccessibilityCache AccessibilityElementCache] GetAccessibilityElement:Id] SetParent:ParentId];
+			[[[FIOSAccessibilityCache AccessibilityElementCache] GetAccessibilityElement:Id] SetParent:NewValue.GetValue<AccessibleWidgetId>()];
 		});
+		// LayoutChanged is to indicate things like "a widget became visible or hidden" while
+		// ScreenChanged is for large-scale UI changes. It can potentially take an NSString to read
+		// to the user when this happens, if we choose to support that.
 		UIAccessibilityPostNotification(UIAccessibilityLayoutChangedNotification, nil);
 		break;
 	}
diff --git a/Engine/Source/Runtime/ApplicationCore/Private/Windows/Accessibility/WindowsUIAManager.cpp b/Engine/Source/Runtime/ApplicationCore/Private/Windows/Accessibility/WindowsUIAManager.cpp
--- a/Engine/Source/Runtime/ApplicationCore/Private/Windows/Accessibility/WindowsUIAManager.cpp
+++ b/Engine/Source/Runtime/ApplicationCore/Private/Windows/Accessibility/WindowsUIAManager.cpp
@@ -37,6 +37,16 @@ FWindowsUIAManager::FWindowsUIAManager(const FWindowsApplication& InApplication)
 {
 	OnAccessibleMessageHandlerChanged();
 
+#if !UE_BUILD_SHIPPING
+	IConsoleManager::Get().RegisterConsoleCommand
+	(
+		TEXT("DumpAccessibilityStatsWindows"),
+		TEXT("Writes to LogAccessibility the memory stats for the platform-level accessibility data (Providers) required for Windows support."),
+		FConsoleCommandDelegate::CreateRaw(this, &FWindowsUIAManager::DumpAccessibilityStats),
+		ECVF_Default
+	);
+#endif
+
 	if (WidgetTypeToWindowsTypeMap.Num() == 0)
 	{
 			WidgetTypeToWindowsTypeMap.Add(EAccessibleWidgetType::Button, UIA_ButtonControlTypeId);
@@ -160,49 +170,94 @@ void FWindowsUIAManager::OnEventRaised(TSharedRef<IAccessibleWidget> Widget, EAc
 		case EAccessibleEvent::Notification:
 		{
 			typedef HRESULT(WINAPI* UiaRaiseNotificationEventFunc)(IRawElementProviderSimple*, NotificationKind, NotificationProcessing, BSTR, BSTR);
-			HMODULE Handle = GetModuleHandle(TEXT("Uiautomationcore.dll"));
-			if (Handle)
+			static UiaRaiseNotificationEventFunc NotificationFunc = nullptr;
+			static bool bDoOnce = true;
+			if (bDoOnce)
 			{
-				// Cast to intermediary void* to avoid compiler warning 4191, since GetProcAddress doesn't know function arguments
-				UiaRaiseNotificationEventFunc NotificationFunc = (UiaRaiseNotificationEventFunc)(void*)GetProcAddress(Handle, "UiaRaiseNotificationEvent");
-				if (NotificationFunc)
+				HMODULE Handle = GetModuleHandle(TEXT("Uiautomationcore.dll"));
+				bDoOnce = false;
+				if (Handle)
 				{
-					NotificationFunc(&ScopedProvider.Provider, NotificationKind_ActionCompleted, NotificationProcessing_All, SysAllocString(*NewValue.GetValue<FString>()), SysAllocString(TEXT("")));
+					// Cast to intermediary void* to avoid compiler warning 4191, since GetProcAddress doesn't know function arguments
+					NotificationFunc = (UiaRaiseNotificationEventFunc)(void*)GetProcAddress(Handle, "UiaRaiseNotificationEvent");
 				}
 			}
+			if (NotificationFunc)
+			{
+				NotificationFunc(&ScopedProvider.Provider, NotificationKind_ActionCompleted, NotificationProcessing_All, SysAllocString(*NewValue.GetValue<FString>()), SysAllocString(TEXT("")));
+			}
 			break;
 		}
-		case EAccessibleEvent::BeforeRemoveFromParent:
-		{
-			// ChildRemoved events require an ID passed along with them to identify which child (which no longer exists) was removed
-			int id[2] = { UiaAppendRuntimeId, Widget->GetId() };
-			FScopedWidgetProvider ParentProvider(GetWidgetProvider(Widget->GetParent().ToSharedRef()));
-			UiaRaiseStructureChangedEvent(&ParentProvider.Provider, StructureChangeType_ChildRemoved, id, ARRAYSIZE(id));
-			break;
-		}
-		case EAccessibleEvent::AfterAddToParent:
-		{
-			FScopedWidgetProvider ParentProvider(GetWidgetProvider(Widget->GetParent().ToSharedRef()));
-			UiaRaiseStructureChangedEvent(&ParentProvider.Provider, StructureChangeType_ChildAdded, nullptr, 0);
-			break;
-		}
+		// IMPORTANT: Calling UiaRaiseStructureChangedEvent seems to raise our per-frame timing for accessibility by
+		// over 10x (.04ms to .7ms, tested by clicking on the "All Classes" button in the modes panel of the Editor).
+		// For now, I'm disabling this until we figure out if it's absolutely necessary.
+		//case EAccessibleEvent::ParentChanged:
+		//{
+		//	const AccessibleWidgetId OldId = OldValue.GetValue<AccessibleWidgetId>();
+		//	const AccessibleWidgetId NewId = NewValue.GetValue<AccessibleWidgetId>();
+		//	if (OldId != IAccessibleWidget::InvalidAccessibleWidgetId)
+		//	{
+		//		GetWidgetProvider(WindowsApplication.GetAccessibleMessageHandler()->GetAccessibleWidgetFromId(OldId).ToSharedRef());
+		//		int id[2] = { UiaAppendRuntimeId, Widget->GetId() };
+		//		FScopedWidgetProvider ParentProvider(GetWidgetProvider(WindowsApplication.GetAccessibleMessageHandler()->GetAccessibleWidgetFromId(OldId).ToSharedRef()));
+		//		UiaRaiseStructureChangedEvent(&ParentProvider.Provider, StructureChangeType_ChildRemoved, id, ARRAYSIZE(id));
+		//	}
+		//	if (NewId != IAccessibleWidget::InvalidAccessibleWidgetId)
+		//	{
+		//		FScopedWidgetProvider ParentProvider(GetWidgetProvider(WindowsApplication.GetAccessibleMessageHandler()->GetAccessibleWidgetFromId(NewId).ToSharedRef()));
+		//		UiaRaiseStructureChangedEvent(&ParentProvider.Provider, StructureChangeType_ChildAdded, nullptr, 0);
+		//	}
+		//	break;
+		//}
 		case EAccessibleEvent::WidgetRemoved:
+		{
 			typedef HRESULT(WINAPI* UiaDisconnectProviderFunc)(IRawElementProviderSimple*);
-			HMODULE Handle = GetModuleHandle(TEXT("Uiautomationcore.dll"));
-			if (Handle)
+			static UiaDisconnectProviderFunc DisconnectFunc = nullptr;
+			static bool bDoOnce = true;
+			if (bDoOnce)
 			{
-				// Cast to intermediary void* to avoid compiler warning 4191, since GetProcAddress doesn't know function arguments
-				UiaDisconnectProviderFunc DisconnectFunc = (UiaDisconnectProviderFunc)(void*)GetProcAddress(Handle, "UiaDisconnectProvider");
-				if (DisconnectFunc)
+				bDoOnce = false;
+				HMODULE Handle = GetModuleHandle(TEXT("Uiautomationcore.dll"));
+				if (Handle)
 				{
-					DisconnectFunc(&ScopedProvider.Provider);
+					// Cast to intermediary void* to avoid compiler warning 4191, since GetProcAddress doesn't know function arguments
+					DisconnectFunc = (UiaDisconnectProviderFunc)(void*)GetProcAddress(Handle, "UiaDisconnectProvider");
 				}
 			}
+			if (DisconnectFunc)
+			{
+				DisconnectFunc(&ScopedProvider.Provider);
+			}
 			break;
 		}
+		}
 	}
 }
 
+#if !UE_BUILD_SHIPPING
+void FWindowsUIAManager::DumpAccessibilityStats() const
+{
+	const uint32 NumWidgetProviders = CachedWidgetProviders.Num();
+	// This isn't exactly right since some ControlProviders will be TextRangeProviders, but it should be close
+	const uint32 NumControlProviders = ProviderList.Num() - NumWidgetProviders;
+
+	const uint32 SizeOfWidgetProvider = sizeof(FWindowsUIAWidgetProvider);
+	const uint32 SizeOfControlProvider = sizeof(FWindowsUIAControlProvider);
+	const uint32 SizeOfCachedWidgetProviders = CachedWidgetProviders.GetAllocatedSize();
+	const uint32 SizeOfProviderList = ProviderList.GetAllocatedSize();
+	const uint32 CacheSize = NumWidgetProviders * SizeOfWidgetProvider + NumControlProviders * SizeOfControlProvider + SizeOfCachedWidgetProviders + SizeOfProviderList;
+
+	UE_LOG(LogAccessibility, Log, TEXT("Dumping Windows accessibility stats:"));
+	UE_LOG(LogAccessibility, Log, TEXT("Number of Widget Providers: %i"), NumWidgetProviders);
+	UE_LOG(LogAccessibility, Log, TEXT("Number of non-Widget Providers: %i"), NumControlProviders);
+	UE_LOG(LogAccessibility, Log, TEXT("Size of FWindowsUIAWidgetProvider: %u"), SizeOfWidgetProvider);
+	UE_LOG(LogAccessibility, Log, TEXT("Size of FWindowsUIAControlProvider: %u"), SizeOfControlProvider);
+	UE_LOG(LogAccessibility, Log, TEXT("Size of WidgetProvider* map: %u"), SizeOfCachedWidgetProviders);
+	UE_LOG(LogAccessibility, Log, TEXT("Size of all Provider* set: %u"), SizeOfProviderList);
+	UE_LOG(LogAccessibility, Log, TEXT("Memory stored in cache: %u kb"), CacheSize / 1000);
+}
+#endif
+
 #undef LOCTEXT_NAMESPACE
 
 #endif
diff --git a/Engine/Source/Runtime/ApplicationCore/Private/Windows/Accessibility/WindowsUIAWidgetProvider.cpp b/Engine/Source/Runtime/ApplicationCore/Private/Windows/Accessibility/WindowsUIAWidgetProvider.cpp
--- a/Engine/Source/Runtime/ApplicationCore/Private/Windows/Accessibility/WindowsUIAWidgetProvider.cpp
+++ b/Engine/Source/Runtime/ApplicationCore/Private/Windows/Accessibility/WindowsUIAWidgetProvider.cpp
@@ -539,7 +539,7 @@ HRESULT STDCALL FWindowsUIAWidgetProvider::get_FragmentRoot(IRawElementProviderF
 {
 	if (IsValid())
 	{
-		TSharedPtr<IAccessibleWidget> Window = Widget->GetTopLevelWindow();
+		TSharedPtr<IAccessibleWidget> Window = Widget->GetWindow();
 		if (Window.IsValid())
 		{
 			*pRetVal = static_cast<IRawElementProviderFragmentRoot*>(&static_cast<FWindowsUIAWindowProvider&>(UIAManager->GetWidgetProvider(Window.ToSharedRef())));
diff --git a/Engine/Source/Runtime/ApplicationCore/Public/GenericPlatform/GenericAccessibleInterfaces.h b/Engine/Source/Runtime/ApplicationCore/Public/GenericPlatform/GenericAccessibleInterfaces.h
--- a/Engine/Source/Runtime/ApplicationCore/Public/GenericPlatform/GenericAccessibleInterfaces.h
+++ b/Engine/Source/Runtime/ApplicationCore/Public/GenericPlatform/GenericAccessibleInterfaces.h
@@ -19,7 +19,7 @@ enum class EAccessibleBehavior : uint8
 	ToolTip
 };
 
-DECLARE_LOG_CATEGORY_EXTERN(LogAccessibility, Log, All);
+APPLICATIONCORE_API DECLARE_LOG_CATEGORY_EXTERN(LogAccessibility, Log, All);
 
 #if WITH_ACCESSIBILITY
 
@@ -51,17 +51,36 @@ enum class EAccessibleWidgetType : uint8
 /** Events that can be raised from accessible widgets to report back to the platform */
 enum class EAccessibleEvent : uint8
 {
-	/** A widget has become focused or unfocused. */
+	/**
+	 * A widget has become focused or unfocused.
+	 * OldValue - The old focus state of the widget
+	 * NewValue - The new focus state of the widget
+	 */
 	FocusChange,
-	/** A widget has been clicked, checked, or otherwise activated */
+	/**
+	 * A widget has been clicked, checked, or otherwise activated.
+	 * OldValue - N/A
+	 * NewValue - N/A
+	 */
 	Activate,
-	/** Notify the user that something has happened. The user is not guaranteed to get this message. */
+	/**
+	 * Warning: Partial implementation
+	 * Notify the user that something has happened. The user is not guaranteed to get this message.
+	 * OldValue - N/A
+	 * NewValue - An FString of the message to read
+	 */
 	Notification,
-	/** Right before a widget is removed from its parent widget. */
-	BeforeRemoveFromParent,
-	/** Right after a widget is added to its parent widget. */
-	AfterAddToParent,
-	/** The widget was removed from the UI tree or deleted. */
+	/**
+	 * A widget's parent is about to be changed.
+	 * OldValue - The AccessibleWidgetId of the old parent, or InvalidAccessibleWidgetId if there was none
+	 * NewValue - The AccessibleWidgetId of the new parent, or InvalidAccessibleWidgetId if there was none
+	 */
+	ParentChanged,
+	/**
+	 * The widget was removed from the UI tree or deleted.
+	 * OldValue - N/A
+	 * NewValue - N/A
+	 */
 	WidgetRemoved
 };
 
@@ -264,7 +283,7 @@ public:
 	 *
 	 * @return The root window in this widget's widget tree, or nullptr if there is no window.
 	 */
-	virtual TSharedPtr<IAccessibleWidget> GetTopLevelWindow() const = 0;
+	virtual TSharedPtr<IAccessibleWidget> GetWindow() const = 0;
 	/**
 	 * Retrieving the bounding rect in absolute coordinates for this widget. On some platforms this may be used for hit testing.
 	 *
@@ -490,24 +509,13 @@ public:
 	 *
 	 * @param Widget The widget raising the event
 	 * @param Event The type of event being raised
-	 * @param OldValue If this is a property changed event, the old value of the property
-	 * @param NewValue If this is a property changed event, the current value of the property. This may also be used by non-property events that require data.
+	 * @param OldValue See EAccessibleEvent documentation for more details.
+	 * @param NewValue See EAccessibleEvent documentation for more details.
 	 */
-	void RaiseEvent(TSharedRef<IAccessibleWidget> Widget, EAccessibleEvent Event, FVariant OldValue, FVariant NewValue)
+	void RaiseEvent(TSharedRef<IAccessibleWidget> Widget, EAccessibleEvent Event, FVariant OldValue = FVariant(), FVariant NewValue = FVariant())
 	{
 		AccessibleEventDelegate.ExecuteIfBound(Widget, Event, OldValue, NewValue);
 	}
-	/**
-	 * Push an event from an accessible widget back to the platform layer.
-	 * This is a convenience function if OldValue and NewValue are irrelevant for this event type.
-	 *
-	 * @param Widget The widget raising the event
-	 * @param Event The type of event being raised
-	 */
-	void RaiseEvent(TSharedRef<IAccessibleWidget> Widget, EAccessibleEvent Event)
-	{
-		RaiseEvent(Widget, Event, FVariant(), FVariant());
-	}
 
 	/**
 	 * Assign a function to be called whenever an accessible event is raised.
diff --git a/Engine/Source/Runtime/ApplicationCore/Public/IOS/Accessibility/IOSAccessibilityCache.h b/Engine/Source/Runtime/ApplicationCore/Public/IOS/Accessibility/IOSAccessibilityCache.h
--- a/Engine/Source/Runtime/ApplicationCore/Public/IOS/Accessibility/IOSAccessibilityCache.h
+++ b/Engine/Source/Runtime/ApplicationCore/Public/IOS/Accessibility/IOSAccessibilityCache.h
@@ -39,6 +39,10 @@
 /** Singleton accessor */
 +(id)AccessibilityElementCache;
 
+#if !UE_BUILD_SHIPPING
+-(void)DumpAccessibilityStats;
+#endif
+
 @end
 
 #endif
diff --git a/Engine/Source/Runtime/ApplicationCore/Public/IOS/Accessibility/IOSAccessibilityElement.h b/Engine/Source/Runtime/ApplicationCore/Public/IOS/Accessibility/IOSAccessibilityElement.h
--- a/Engine/Source/Runtime/ApplicationCore/Public/IOS/Accessibility/IOSAccessibilityElement.h
+++ b/Engine/Source/Runtime/ApplicationCore/Public/IOS/Accessibility/IOSAccessibilityElement.h
@@ -29,7 +29,7 @@
 }
 
 /** This must be used instead of initWithAccessibilityContainer in order to work properly. */
--(id)initWithId:(AccessibleWidgetId)InId AndParentId:(AccessibleWidgetId)InParentId;
+-(id)initWithId:(AccessibleWidgetId)InId;
 /** Updates the accessibilityContainer property from UIAccessibilityElement. */
 -(void)SetParent:(AccessibleWidgetId)InParentId;
 /** Get the accessible version of this widget. */
diff --git a/Engine/Source/Runtime/ApplicationCore/Public/Windows/Accessibility/WindowsUIAManager.h b/Engine/Source/Runtime/ApplicationCore/Public/Windows/Accessibility/WindowsUIAManager.h
--- a/Engine/Source/Runtime/ApplicationCore/Public/Windows/Accessibility/WindowsUIAManager.h
+++ b/Engine/Source/Runtime/ApplicationCore/Public/Windows/Accessibility/WindowsUIAManager.h
@@ -71,6 +71,10 @@ public:
 	static TMap<EAccessibleWidgetType, ULONG> WidgetTypeToWindowsTypeMap;
 	static TMap<EAccessibleWidgetType, FText> WidgetTypeToTextMap;
 
+#if !UE_BUILD_SHIPPING
+	void DumpAccessibilityStats() const;
+#endif
+
 private:
 	/** Callback function for processing events raised from the AccessibleMessageHandler */
 	void OnEventRaised(TSharedRef<IAccessibleWidget> Widget, EAccessibleEvent Event, FVariant OldValue, FVariant NewValue);
diff --git a/Engine/Source/Runtime/Slate/Private/Framework/Application/SlateApplication.cpp b/Engine/Source/Runtime/Slate/Private/Framework/Application/SlateApplication.cpp
--- a/Engine/Source/Runtime/Slate/Private/Framework/Application/SlateApplication.cpp
+++ b/Engine/Source/Runtime/Slate/Private/Framework/Application/SlateApplication.cpp
@@ -899,7 +899,7 @@ void FSlateApplication::SetPlatformApplication(const TSharedRef<class GenericApp
 	PlatformApplication = InPlatformApplication;
 	PlatformApplication->SetMessageHandler(CurrentApplication.ToSharedRef());
 #if WITH_ACCESSIBILITY
-	PlatformApplication->SetAccessibleMessageHandler(CurrentApplication->GetAccessibleMessageHandler().ToSharedRef());
+	PlatformApplication->SetAccessibleMessageHandler(CurrentApplication->GetAccessibleMessageHandler());
 #endif
 }
 
@@ -927,7 +927,7 @@ TSharedRef<FSlateApplication> FSlateApplication::Create(const TSharedRef<class G
 	PlatformApplication = InPlatformApplication;
 	PlatformApplication->SetMessageHandler( CurrentApplication.ToSharedRef() );
 #if WITH_ACCESSIBILITY
-	PlatformApplication->SetAccessibleMessageHandler(CurrentApplication->GetAccessibleMessageHandler().ToSharedRef());
+	PlatformApplication->SetAccessibleMessageHandler(CurrentApplication->GetAccessibleMessageHandler());
 #endif
 
 	// The grid needs to know the size and coordinate system of the desktop.
@@ -981,9 +981,6 @@ FSlateApplication::FSlateApplication()
 	, bSlateWindowActive(true)
 	, Scale( 1.0f )
 	, DragTriggerDistance( 0 )
-#if WITH_ACCESSIBILITY
-	, AccessibleMessageHandler(new FSlateAccessibleMessageHandler())
-#endif
 	, CursorRadius( 0.0f )
 	, LastUserInteractionTime( 0.0 )
 	, LastUserInteractionTimeForThrottling( 0.0 )
@@ -1879,6 +1876,10 @@ void FSlateApplication::TickApplication(ESlateTickType TickType, float DeltaTime
 
 			// Draw all windows
 			DrawWindows();
+
+#if WITH_ACCESSIBILITY
+			AccessibleMessageHandler->Tick();
+#endif
 		}
 
 		PostTickEvent.Broadcast(DeltaTime);
diff --git a/Engine/Source/Runtime/Slate/Private/Widgets/Input/SButton.cpp b/Engine/Source/Runtime/Slate/Private/Widgets/Input/SButton.cpp
--- a/Engine/Source/Runtime/Slate/Private/Widgets/Input/SButton.cpp
+++ b/Engine/Source/Runtime/Slate/Private/Widgets/Input/SButton.cpp
@@ -521,7 +521,7 @@ void SButton::SetPressMethod(EButtonPressMethod::Type InPressMethod)
 }
 
 #if WITH_ACCESSIBILITY
-TSharedPtr<FSlateAccessibleWidget> SButton::CreateAccessibleWidget()
+TSharedRef<FSlateAccessibleWidget> SButton::CreateAccessibleWidget()
 {
 	return MakeShareable<FSlateAccessibleWidget>(new FSlateAccessibleButton(SharedThis(this)));
 }
diff --git a/Engine/Source/Runtime/Slate/Private/Widgets/Input/SCheckBox.cpp b/Engine/Source/Runtime/Slate/Private/Widgets/Input/SCheckBox.cpp
--- a/Engine/Source/Runtime/Slate/Private/Widgets/Input/SCheckBox.cpp
+++ b/Engine/Source/Runtime/Slate/Private/Widgets/Input/SCheckBox.cpp
@@ -518,7 +518,7 @@ const FSlateBrush* SCheckBox::GetUndeterminedPressedImage() const
 }
 
 #if WITH_ACCESSIBILITY
-TSharedPtr<FSlateAccessibleWidget> SCheckBox::CreateAccessibleWidget()
+TSharedRef<FSlateAccessibleWidget> SCheckBox::CreateAccessibleWidget()
 {
 	return MakeShareable<FSlateAccessibleWidget>(new FSlateAccessibleCheckBox(SharedThis(this)));
 }
diff --git a/Engine/Source/Runtime/Slate/Private/Widgets/Input/SEditableText.cpp b/Engine/Source/Runtime/Slate/Private/Widgets/Input/SEditableText.cpp
--- a/Engine/Source/Runtime/Slate/Private/Widgets/Input/SEditableText.cpp
+++ b/Engine/Source/Runtime/Slate/Private/Widgets/Input/SEditableText.cpp
@@ -603,7 +603,7 @@ float SEditableText::UpdateAndClampVerticalScrollBar(const float InViewOffset, c
 }
 
 #if WITH_ACCESSIBILITY
-TSharedPtr<FSlateAccessibleWidget> SEditableText::CreateAccessibleWidget()
+TSharedRef<FSlateAccessibleWidget> SEditableText::CreateAccessibleWidget()
 {
 	return MakeShareable<FSlateAccessibleWidget>(new FSlateAccessibleEditableText(SharedThis(this)));
 }
diff --git a/Engine/Source/Runtime/Slate/Private/Widgets/Input/SEditableTextBox.cpp b/Engine/Source/Runtime/Slate/Private/Widgets/Input/SEditableTextBox.cpp
--- a/Engine/Source/Runtime/Slate/Private/Widgets/Input/SEditableTextBox.cpp
+++ b/Engine/Source/Runtime/Slate/Private/Widgets/Input/SEditableTextBox.cpp
@@ -413,7 +413,7 @@ void SEditableTextBox::SetVirtualKeyboardDismissAction(TAttribute<EVirtualKeyboa
 }
 
 #if WITH_ACCESSIBILITY
-TSharedPtr<FSlateAccessibleWidget> SEditableTextBox::CreateAccessibleWidget()
+TSharedRef<FSlateAccessibleWidget> SEditableTextBox::CreateAccessibleWidget()
 {
 	return MakeShareable<FSlateAccessibleWidget>(new FSlateAccessibleEditableTextBox(SharedThis(this)));
 }
diff --git a/Engine/Source/Runtime/Slate/Private/Widgets/Input/SSlider.cpp b/Engine/Source/Runtime/Slate/Private/Widgets/Input/SSlider.cpp
--- a/Engine/Source/Runtime/Slate/Private/Widgets/Input/SSlider.cpp
+++ b/Engine/Source/Runtime/Slate/Private/Widgets/Input/SSlider.cpp
@@ -512,7 +512,7 @@ void SSlider::SetRequiresControllerLock(bool RequiresControllerLock) {
 }
 
 #if WITH_ACCESSIBILITY
-TSharedPtr<FSlateAccessibleWidget> SSlider::CreateAccessibleWidget()
+TSharedRef<FSlateAccessibleWidget> SSlider::CreateAccessibleWidget()
 {
 	return MakeShareable<FSlateAccessibleWidget>(new FSlateAccessibleSlider(SharedThis(this)));
 }
diff --git a/Engine/Source/Runtime/Slate/Private/Widgets/Layout/SWidgetSwitcher.cpp b/Engine/Source/Runtime/Slate/Private/Widgets/Layout/SWidgetSwitcher.cpp
--- a/Engine/Source/Runtime/Slate/Private/Widgets/Layout/SWidgetSwitcher.cpp
+++ b/Engine/Source/Runtime/Slate/Private/Widgets/Layout/SWidgetSwitcher.cpp
@@ -142,9 +142,8 @@ void SWidgetSwitcher::OnArrangeChildren(const FGeometry& AllottedGeometry, FArra
 		// tree to update itself.
 		if (LastActiveWidget != ActiveSlotPtr->GetWidget())
 		{
-			SWidgetSwitcher* UnconstThis = const_cast<SWidgetSwitcher*>(this);
-			UnconstThis->LastActiveWidget = ActiveSlotPtr->GetWidget();
-			FSlateApplication::Get().GetAccessibleMessageHandler()->OnWidgetChildrenChanged(SharedThis(UnconstThis));
+			const_cast<SWidgetSwitcher*>(this)->LastActiveWidget = ActiveSlotPtr->GetWidget();
+			FSlateApplication::Get().GetAccessibleMessageHandler()->MarkDirty();
 		}
 #endif
 	}
diff --git a/Engine/Source/Runtime/Slate/Private/Widgets/Text/STextBlock.cpp b/Engine/Source/Runtime/Slate/Private/Widgets/Text/STextBlock.cpp
--- a/Engine/Source/Runtime/Slate/Private/Widgets/Text/STextBlock.cpp
+++ b/Engine/Source/Runtime/Slate/Private/Widgets/Text/STextBlock.cpp
@@ -478,7 +478,7 @@ FTextBlockStyle STextBlock::GetComputedTextStyle() const
 }
 
 #if WITH_ACCESSIBILITY
-TSharedPtr<FSlateAccessibleWidget> STextBlock::CreateAccessibleWidget()
+TSharedRef<FSlateAccessibleWidget> STextBlock::CreateAccessibleWidget()
 {
 	return MakeShareable<FSlateAccessibleWidget>(new FSlateAccessibleTextBlock(SharedThis(this)));
 }
diff --git a/Engine/Source/Runtime/Slate/Public/Framework/Application/SlateApplication.h b/Engine/Source/Runtime/Slate/Public/Framework/Application/SlateApplication.h
--- a/Engine/Source/Runtime/Slate/Public/Framework/Application/SlateApplication.h
+++ b/Engine/Source/Runtime/Slate/Public/Framework/Application/SlateApplication.h
@@ -1419,9 +1419,6 @@ public:
 	virtual void SetAllUserFocus(const FWidgetPath& InFocusPath, const EFocusCause InCause) override;
 	virtual void SetAllUserFocusAllowingDescendantFocus(const FWidgetPath& InFocusPath, const EFocusCause InCause) override;
 	virtual TSharedPtr<SWidget> GetUserFocusedWidget(uint32 UserIndex) const override;
-#if WITH_ACCESSIBILITY
-	virtual TSharedPtr<FSlateAccessibleMessageHandler> GetAccessibleMessageHandler() const override { return AccessibleMessageHandler; }
-#endif
 	virtual const TArray<TSharedRef<SWindow>> GetTopLevelWindows() const override { return SlateWindows; }
 
 	DECLARE_EVENT_OneParam(FSlateApplication, FApplicationActivationStateChangedEvent, const bool /*IsActive*/)
@@ -1751,10 +1748,7 @@ private:
 
 	/** These windows will be destroyed next tick. */
 	TArray< TSharedRef<SWindow> > WindowDestroyQueue;
-#if WITH_ACCESSIBILITY
-	/** Manager for widgets and application to interact with accessibility API */
-	TSharedRef<FSlateAccessibleMessageHandler> AccessibleMessageHandler;
-#endif
+
 	/** The stack of menus that are open */
 	FMenuStack MenuStack;
 
diff --git a/Engine/Source/Runtime/Slate/Public/Widgets/Input/SButton.h b/Engine/Source/Runtime/Slate/Public/Widgets/Input/SButton.h
--- a/Engine/Source/Runtime/Slate/Public/Widgets/Input/SButton.h
+++ b/Engine/Source/Runtime/Slate/Public/Widgets/Input/SButton.h
@@ -184,7 +184,7 @@ public:
 	virtual bool IsInteractable() const override;
 	virtual bool ComputeVolatility() const override;
 #if WITH_ACCESSIBILITY
-	virtual TSharedPtr<FSlateAccessibleWidget> CreateAccessibleWidget() override;
+	virtual TSharedRef<FSlateAccessibleWidget> CreateAccessibleWidget() override;
 #endif
 	// SWidget
 
diff --git a/Engine/Source/Runtime/Slate/Public/Widgets/Input/SCheckBox.h b/Engine/Source/Runtime/Slate/Public/Widgets/Input/SCheckBox.h
--- a/Engine/Source/Runtime/Slate/Public/Widgets/Input/SCheckBox.h
+++ b/Engine/Source/Runtime/Slate/Public/Widgets/Input/SCheckBox.h
@@ -149,7 +149,7 @@ public:
 	virtual void OnMouseLeave( const FPointerEvent& MouseEvent ) override;
 	virtual bool IsInteractable() const override;
 #if WITH_ACCESSIBILITY
-	virtual TSharedPtr<FSlateAccessibleWidget> CreateAccessibleWidget() override;
+	virtual TSharedRef<FSlateAccessibleWidget> CreateAccessibleWidget() override;
 #endif
 	// End of SWidget interface
 
diff --git a/Engine/Source/Runtime/Slate/Public/Widgets/Input/SEditableText.h b/Engine/Source/Runtime/Slate/Public/Widgets/Input/SEditableText.h
--- a/Engine/Source/Runtime/Slate/Public/Widgets/Input/SEditableText.h
+++ b/Engine/Source/Runtime/Slate/Public/Widgets/Input/SEditableText.h
@@ -354,7 +354,7 @@ protected:
 	virtual bool IsInteractable() const override;
 	virtual bool ComputeVolatility() const override;
 #if WITH_ACCESSIBILITY
-	virtual TSharedPtr<FSlateAccessibleWidget> CreateAccessibleWidget() override;
+	virtual TSharedRef<FSlateAccessibleWidget> CreateAccessibleWidget() override;
 	virtual void SetDefaultAccessibleText(EAccessibleType AccessibleType = EAccessibleType::Main) override;
 #endif
 	//~ End SWidget Interface
diff --git a/Engine/Source/Runtime/Slate/Public/Widgets/Input/SEditableTextBox.h b/Engine/Source/Runtime/Slate/Public/Widgets/Input/SEditableTextBox.h
--- a/Engine/Source/Runtime/Slate/Public/Widgets/Input/SEditableTextBox.h
+++ b/Engine/Source/Runtime/Slate/Public/Widgets/Input/SEditableTextBox.h
@@ -343,7 +343,7 @@ public:
 
 protected:
 #if WITH_ACCESSIBILITY
-	virtual TSharedPtr<FSlateAccessibleWidget> CreateAccessibleWidget() override;
+	virtual TSharedRef<FSlateAccessibleWidget> CreateAccessibleWidget() override;
 	virtual void SetDefaultAccessibleText(EAccessibleType AccessibleType = EAccessibleType::Main) override;
 #endif
 
diff --git a/Engine/Source/Runtime/Slate/Public/Widgets/Input/SHyperlink.h b/Engine/Source/Runtime/Slate/Public/Widgets/Input/SHyperlink.h
--- a/Engine/Source/Runtime/Slate/Public/Widgets/Input/SHyperlink.h
+++ b/Engine/Source/Runtime/Slate/Public/Widgets/Input/SHyperlink.h
@@ -85,7 +85,7 @@ public:
 		return FCursorReply::Cursor( EMouseCursor::Hand );
 	}
 #if WITH_ACCESSIBILITY
-	virtual TSharedPtr<FSlateAccessibleWidget> CreateAccessibleWidget() override
+	virtual TSharedRef<FSlateAccessibleWidget> CreateAccessibleWidget() override
 	{
 		return MakeShareable<FSlateAccessibleWidget>(new FSlateAccessibleHyperlink(SharedThis(this)));
 	}
diff --git a/Engine/Source/Runtime/Slate/Public/Widgets/Input/SSlider.h b/Engine/Source/Runtime/Slate/Public/Widgets/Input/SSlider.h
--- a/Engine/Source/Runtime/Slate/Public/Widgets/Input/SSlider.h
+++ b/Engine/Source/Runtime/Slate/Public/Widgets/Input/SSlider.h
@@ -172,7 +172,7 @@ public:
 	virtual bool SupportsKeyboardFocus() const override;
 	virtual bool IsInteractable() const override;
 #if WITH_ACCESSIBILITY
-	virtual TSharedPtr<FSlateAccessibleWidget> CreateAccessibleWidget() override;
+	virtual TSharedRef<FSlateAccessibleWidget> CreateAccessibleWidget() override;
 #endif
 
 	/** @return Is the handle locked or not? Defaults to false */
diff --git a/Engine/Source/Runtime/Slate/Public/Widgets/Text/STextBlock.h b/Engine/Source/Runtime/Slate/Public/Widgets/Text/STextBlock.h
--- a/Engine/Source/Runtime/Slate/Public/Widgets/Text/STextBlock.h
+++ b/Engine/Source/Runtime/Slate/Public/Widgets/Text/STextBlock.h
@@ -241,7 +241,7 @@ public:
 	virtual int32 OnPaint( const FPaintArgs& Args, const FGeometry& AllottedGeometry, const FSlateRect& MyCullingRect, FSlateWindowElementList& OutDrawElements, int32 LayerId, const FWidgetStyle& InWidgetStyle, bool bParentEnabled ) const override;
 	virtual FVector2D ComputeDesiredSize(float) const override;
 #if WITH_ACCESSIBILITY
-	virtual TSharedPtr<FSlateAccessibleWidget> CreateAccessibleWidget() override;
+	virtual TSharedRef<FSlateAccessibleWidget> CreateAccessibleWidget() override;
 	virtual void SetDefaultAccessibleText(EAccessibleType AccessibleType = EAccessibleType::Main) override;
 #endif
 	// End of SWidget interface
diff --git a/Engine/Source/Runtime/SlateCore/Private/Application/SlateApplicationBase.cpp b/Engine/Source/Runtime/SlateCore/Private/Application/SlateApplicationBase.cpp
--- a/Engine/Source/Runtime/SlateCore/Private/Application/SlateApplicationBase.cpp
+++ b/Engine/Source/Runtime/SlateCore/Private/Application/SlateApplicationBase.cpp
@@ -5,6 +5,9 @@
 #include "Layout/WidgetPath.h"
 #include "Application/ActiveTimerHandle.h"
 #include "Misc/ScopeLock.h"
+#if WITH_ACCESSIBILITY
+#include "Widgets/Accessibility/SlateAccessibleMessageHandler.h"
+#endif
 
 
 /* Static initialization
@@ -25,6 +28,9 @@ FWidgetPath FHitTesting::LocateWidgetInWindow(FVector2D ScreenspaceMouseCoordina
 FSlateApplicationBase::FSlateApplicationBase()
 : Renderer()
 , HitTesting(this)
+#if WITH_ACCESSIBILITY
+, AccessibleMessageHandler(new FSlateAccessibleMessageHandler())
+#endif
 , bIsSlateAsleep(false)
 {
 
diff --git a/Engine/Source/Runtime/SlateCore/Private/Widgets/Accessibility/SlateAccessibleMessageHandler.cpp b/Engine/Source/Runtime/SlateCore/Private/Widgets/Accessibility/SlateAccessibleMessageHandler.cpp
--- a/Engine/Source/Runtime/SlateCore/Private/Widgets/Accessibility/SlateAccessibleMessageHandler.cpp
+++ b/Engine/Source/Runtime/SlateCore/Private/Widgets/Accessibility/SlateAccessibleMessageHandler.cpp
@@ -10,19 +10,23 @@
 #include "Widgets/SWindow.h"
 #include "Input/HittestGrid.h"
 
-DECLARE_CYCLE_STAT(TEXT("Slate Accessibility: Parent Updated"), STAT_AccessibilitySlateParentUpdated, STATGROUP_Accessibility);
-DECLARE_CYCLE_STAT(TEXT("Slate Accessibility: Children Updated"), STAT_AccessibilitySlateChildrenUpdated, STATGROUP_Accessibility);
-DECLARE_CYCLE_STAT(TEXT("Slate Accessibility: Behavior Changed"), STAT_AccessibilitySlateBehaviorChanged, STATGROUP_Accessibility);
+DECLARE_CYCLE_STAT(TEXT("Slate Accessibility: Tick"), STAT_AccessibilitySlateTick, STATGROUP_Accessibility);
 DECLARE_CYCLE_STAT(TEXT("Slate Accessibility: Event Raised"), STAT_AccessibilitySlateEventRaised, STATGROUP_Accessibility);
 
+FSlateAccessibleMessageHandler::FSlateAccessibleMessageHandler()
+	: FGenericAccessibleMessageHandler()
+	, bDirty(false)
+{
+}
+
 void FSlateAccessibleMessageHandler::OnActivate()
 {
-	// widgets are initialized when their accessible window is created
+	bDirty = true;
 }
 
 void FSlateAccessibleMessageHandler::OnDeactivate()
 {
-	FSlateAccessibleWidgetCache::Get().ClearAll();
+	FSlateAccessibleWidgetCache::ClearAll();
 }
 
 TSharedPtr<IAccessibleWidget> FSlateAccessibleMessageHandler::GetAccessibleWindow(const TSharedRef<FGenericWindow>& InWindow) const
@@ -30,10 +34,7 @@ TSharedPtr<IAccessibleWidget> FSlateAccessibleMessageHandler::GetAccessibleWindo
 	if (IsActive())
 	{
 		TSharedPtr<SWindow> SlateWindow = FSlateWindowHelper::FindWindowByPlatformWindow(FSlateApplicationBase::Get().GetTopLevelWindows(), InWindow);
-		if (SlateWindow.IsValid())
-		{
-			return FSlateAccessibleWidgetCache::Get().GetAccessibleWidget(SlateWindow);
-		}
+		return FSlateAccessibleWidgetCache::GetAccessibleWidgetChecked(SlateWindow);
 	}
 	return nullptr;
 }
@@ -50,132 +51,107 @@ AccessibleWidgetId FSlateAccessibleMessageHandler::GetAccessibleWindowId(const T
 
 TSharedPtr<IAccessibleWidget> FSlateAccessibleMessageHandler::GetAccessibleWidgetFromId(AccessibleWidgetId Id) const
 {
-	return FSlateAccessibleWidgetCache::Get().GetAccessibleWidgetFromId(Id);
+	return FSlateAccessibleWidgetCache::GetAccessibleWidgetFromId(Id);
 }
 
 void FSlateAccessibleMessageHandler::OnWidgetRemoved(SWidget* Widget)
 {
 	if (IsActive())
 	{
-		TSharedPtr<FSlateAccessibleWidget> RemovedWidget = FSlateAccessibleWidgetCache::Get().RemoveWidget(Widget);
+		TSharedPtr<FSlateAccessibleWidget> RemovedWidget = FSlateAccessibleWidgetCache::RemoveWidget(Widget);
 		if (RemovedWidget.IsValid())
 		{
-			RaiseEvent(StaticCastSharedPtr<IAccessibleWidget>(RemovedWidget).ToSharedRef(), EAccessibleEvent::WidgetRemoved);
+			RaiseEvent(RemovedWidget.ToSharedRef(), EAccessibleEvent::WidgetRemoved);
+			// If this ensure fails, bDirty = true must be called to ensure the tree is kept up to date.
+			ensureMsgf(!Widget->GetParentWidget().IsValid(), TEXT("A widget was unexpectedly deleted before detaching from its parent."));
 		}
 	}
 }
 
-void FSlateAccessibleMessageHandler::OnWidgetParentChanged(TSharedRef<SWidget> Widget)
+void FSlateAccessibleMessageHandler::OnWidgetEventRaised(TSharedRef<SWidget> Widget, EAccessibleEvent Event, FVariant OldValue, FVariant NewValue)
 {
 	if (IsActive())
 	{
-		SCOPE_CYCLE_COUNTER(STAT_AccessibilitySlateParentUpdated);
-
-		TSharedPtr<SWidget> Parent = Widget->GetParentWidget();
-		while (Parent.IsValid() && !Parent->IsAccessible())
-		{
-			Parent = Parent->GetParentWidget();
-		}
-		
+		SCOPE_CYCLE_COUNTER(STAT_AccessibilitySlateEventRaised);
+		// todo: not sure what to do for a case like focus changed to not-accessible widget. maybe pass through a nullptr?
 		if (Widget->IsAccessible())
 		{
-			if (Parent.IsValid())
-			{
-				FSlateAccessibleWidgetCache::Get().GetAccessibleWidget(Widget)->UpdateParent(FSlateAccessibleWidgetCache::Get().GetAccessibleWidget(Parent));
-			}
-			else
-			{
-				FSlateAccessibleWidgetCache::Get().GetAccessibleWidget(Widget)->UpdateParent(nullptr);
-			}
-		}
-		else if (Parent.IsValid() && Parent->CanChildrenBeAccessible())
-		{
-			TSharedPtr<IAccessibleWidget> AccessibleParent = FSlateAccessibleWidgetCache::Get().GetAccessibleWidget(Parent);
-			TArray<TSharedRef<SWidget>> AccessibleChildren = FSlateAccessibleWidget::GetAccessibleChildren(Widget);
-			for (int32 i = 0; i < AccessibleChildren.Num(); ++i)
-			{
-				FSlateAccessibleWidgetCache::Get().GetAccessibleWidget(AccessibleChildren[i])->UpdateParent(AccessibleParent);
-			}
+			FSlateAccessibleMessageHandler::RaiseEvent(FSlateAccessibleWidgetCache::GetAccessibleWidget(Widget), Event, OldValue, NewValue);
 		}
 	}
 }
 
-void FSlateAccessibleMessageHandler::OnWidgetChildrenChanged(TSharedRef<SWidget> Widget)
-{
-	if (IsActive())
-	{
-		SCOPE_CYCLE_COUNTER(STAT_AccessibilitySlateChildrenUpdated);
+int32 GAccessibleWidgetsProcessedPerTick = 100;
+FAutoConsoleVariableRef AccessibleWidgetsProcessedPerTickRef(
+	TEXT("Slate.AccessibleWidgetsProcessedPerTick"),
+	GAccessibleWidgetsProcessedPerTick,
+	TEXT("To reduce performance spikes, generating the accessible widget tree is limited to this many widgets per tick to update.")
+);
 
-		TSharedPtr<SWidget> Parent = Widget;
-		while (Parent.IsValid() && !Parent->IsAccessible())
-		{
-			Parent = Parent->GetParentWidget();
-		}
-		if (Parent.IsValid())
-		{
-			TSharedPtr<IAccessibleWidget> AccessibleParent = FSlateAccessibleWidgetCache::Get().GetAccessibleWidget(Parent);
-			StaticCastSharedPtr<FSlateAccessibleWidget>(AccessibleParent)->MarkChildrenDirty();
-		}
-	}
-}
-
-void FSlateAccessibleMessageHandler::OnWidgetAccessibleBehaviorChanged(TSharedRef<SWidget> Widget)
+void FSlateAccessibleMessageHandler::Tick()
 {
 	if (IsActive())
 	{
-		SCOPE_CYCLE_COUNTER(STAT_AccessibilitySlateBehaviorChanged);
-
-		TSharedPtr<SWidget> Parent = Widget->GetParentWidget();
-		while (Parent.IsValid())
+		SCOPE_CYCLE_COUNTER(STAT_AccessibilitySlateTick);
+		if (bDirty && ToProcess.Num() == 0)
 		{
-			if (Parent->IsAccessible())
+			bDirty = false;
+			// Process ALL windows, not just the top level ones. Otherwise we miss things like combo boxes.
+			TArray<TSharedRef<SWindow>> SlateWindows = FSlateApplicationBase::Get().GetTopLevelWindows();
+			while (SlateWindows.Num() > 0)
 			{
-				break;
+				const TSharedRef<SWindow> CurrentWindow = SlateWindows.Pop(false);
+				ToProcess.Emplace(CurrentWindow, FSlateAccessibleWidgetCache::GetAccessibleWidget(CurrentWindow));
+				SlateWindows.Append(CurrentWindow->GetChildWindows());
 			}
-
-			Parent = Parent->GetParentWidget();
 		}
 
-		if (Parent.IsValid())
+		if (ToProcess.Num() > 0)
 		{
-			TSharedPtr<IAccessibleWidget> AccessibleParent = FSlateAccessibleWidgetCache::Get().GetAccessibleWidget(Parent);
-			TArray<TSharedRef<SWidget>> AccessibleChildren = FSlateAccessibleWidget::GetAccessibleChildren(Widget);
-			if (Widget->IsAccessible())
+			for (int32 Counter = 0; ToProcess.Num() > 0 && Counter < GAccessibleWidgetsProcessedPerTick; ++Counter)
 			{
-				TSharedPtr<FSlateAccessibleWidget> AccessibleWidget = StaticCastSharedPtr<FSlateAccessibleWidget>(FSlateAccessibleWidgetCache::Get().GetAccessibleWidget(Widget));
-				for (int32 i = 0; i < AccessibleChildren.Num(); ++i)
+				FWidgetAndParent WidgetAndParent = ToProcess.Pop(false);
+				if (WidgetAndParent.Widget.IsValid())
 				{
-					FSlateAccessibleWidgetCache::Get().GetAccessibleWidget(AccessibleChildren[i])->UpdateParent(AccessibleWidget);
+					TSharedPtr<SWidget> SharedWidget = WidgetAndParent.Widget.Pin();
+					if (SharedWidget->CanChildrenBeAccessible())
+					{
+						FChildren* SharedChildren = SharedWidget->GetChildren();
+						for (int32 i = 0; i < SharedChildren->Num(); ++i)
+						{
+							TSharedRef<SWidget> Child = SharedChildren->GetChildAt(i);
+							if (Child->GetAccessibleBehavior() != EAccessibleBehavior::NotAccessible)
+							{
+								TSharedRef<FSlateAccessibleWidget> AccessibleChild = FSlateAccessibleWidgetCache::GetAccessibleWidget(Child);
+								AccessibleChild->SiblingIndex = WidgetAndParent.Parent->ChildrenBuffer.Num();
+								AccessibleChild->UpdateParent(WidgetAndParent.Parent);
+								// A separate children buffer is filled instead of the children array itself
+								// so that accessibility queries still work (using the old data) while updating
+								// accessible widget data.
+								WidgetAndParent.Parent->ChildrenBuffer.Add(AccessibleChild);
+								ToProcess.Emplace(Child, AccessibleChild);
+							}
+							else
+							{
+								// Keep a reference to the last-known accessible parent
+								ToProcess.Emplace(Child, WidgetAndParent.Parent);
+							}
+						}
+					}
 				}
-				AccessibleWidget->UpdateParent(AccessibleParent);
 			}
-			else
+
+			// Once processing is finished, update each widget's children array with its children buffer
+			if (ToProcess.Num() == 0)
 			{
-				for (int32 i = 0; i < AccessibleChildren.Num(); ++i)
+				for (auto WidgetIterator = FSlateAccessibleWidgetCache::GetAllWidgets(); WidgetIterator; ++WidgetIterator)
 				{
-					FSlateAccessibleWidgetCache::Get().GetAccessibleWidget(AccessibleChildren[i])->UpdateParent(AccessibleParent);
+					TSharedRef<FSlateAccessibleWidget> AccessibleWidget = WidgetIterator.Value();
+					AccessibleWidget->Children = MoveTemp(AccessibleWidget->ChildrenBuffer);
 				}
 			}
 		}
 	}
 }
 
-void FSlateAccessibleMessageHandler::OnWidgetEventRaised(TSharedRef<SWidget> Widget, EAccessibleEvent Event)
-{
-	OnWidgetEventRaised(Widget, Event, FVariant(), FVariant());
-}
-
-void FSlateAccessibleMessageHandler::OnWidgetEventRaised(TSharedRef<SWidget> Widget, EAccessibleEvent Event, FVariant OldValue, FVariant NewValue)
-{
-	if (IsActive())
-	{
-		SCOPE_CYCLE_COUNTER(STAT_AccessibilitySlateEventRaised);
-		// todo: not sure what to do for a case like focus changed to not-accessible widget. maybe pass through a nullptr?
-		if (Widget->IsAccessible())
-		{
-			FSlateAccessibleMessageHandler::RaiseEvent(FSlateAccessibleWidgetCache::Get().GetAccessibleWidget(Widget).ToSharedRef(), Event, OldValue, NewValue);
-		}
-	}
-}
-
 #endif
diff --git a/Engine/Source/Runtime/SlateCore/Private/Widgets/Accessibility/SlateAccessibleWidgetCache.cpp b/Engine/Source/Runtime/SlateCore/Private/Widgets/Accessibility/SlateAccessibleWidgetCache.cpp
new file mode 100644
--- /dev/null
+++ b/Engine/Source/Runtime/SlateCore/Private/Widgets/Accessibility/SlateAccessibleWidgetCache.cpp
@@ -0,0 +1,97 @@
+// Copyright 1998-2019 Epic Games, Inc. All Rights Reserved.
+
+#if WITH_ACCESSIBILITY
+
+#include "Widgets/Accessibility/SlateAccessibleWidgetCache.h"
+#include "Application/SlateApplicationBase.h"
+#include "Widgets/Accessibility/SlateAccessibleMessageHandler.h"
+#include "Widgets/SWidget.h"
+
+TMap<SWidget*, TSharedRef<FSlateAccessibleWidget>> FSlateAccessibleWidgetCache::AccessibleWidgetMap;
+TMap<AccessibleWidgetId, TSharedRef<FSlateAccessibleWidget>> FSlateAccessibleWidgetCache::AccessibleIdMap;
+
+void FSlateAccessibleWidgetCache::ClearAll()
+{
+	AccessibleWidgetMap.Empty();
+	AccessibleIdMap.Empty();
+}
+
+TSharedPtr<FSlateAccessibleWidget> FSlateAccessibleWidgetCache::RemoveWidget(SWidget* Widget)
+{
+	TSharedRef<FSlateAccessibleWidget>* AccessibleWidgetPtr = AccessibleWidgetMap.Find(Widget);
+	if (AccessibleWidgetPtr)
+	{
+		TSharedRef<FSlateAccessibleWidget> AccessibleWidget = *AccessibleWidgetPtr;
+		AccessibleWidgetMap.Remove(Widget);
+		AccessibleIdMap.Remove(AccessibleWidget->GetId());
+		return AccessibleWidget;
+	}
+	return nullptr;
+}
+
+TSharedPtr<FSlateAccessibleWidget> FSlateAccessibleWidgetCache::GetAccessibleWidgetChecked(const TSharedPtr<SWidget>& Widget)
+{
+	if (!Widget.IsValid() || !Widget->IsAccessible())
+	{
+		return nullptr;
+	}
+
+	return GetAccessibleWidget(Widget.ToSharedRef());
+}
+
+TSharedRef<FSlateAccessibleWidget> FSlateAccessibleWidgetCache::GetAccessibleWidget(const TSharedRef<SWidget>& Widget)
+{
+	TSharedRef<FSlateAccessibleWidget>* AccessibleWidget = AccessibleWidgetMap.Find(&Widget.Get());
+	if (AccessibleWidget)
+	{
+		return *AccessibleWidget;
+	}
+	else
+	{
+		TSharedRef<FSlateAccessibleWidget> NewWidget = AccessibleWidgetMap.Add(&Widget.Get(), Widget->CreateAccessibleWidget());
+		AccessibleIdMap.Add(NewWidget->GetId(), NewWidget);
+		return NewWidget;
+	}
+}
+
+TSharedPtr<FSlateAccessibleWidget> FSlateAccessibleWidgetCache::GetAccessibleWidgetFromId(AccessibleWidgetId Id)
+{
+	TSharedRef<FSlateAccessibleWidget>* AccessibleWidget = AccessibleIdMap.Find(Id);
+	if (AccessibleWidget)
+	{
+		return *AccessibleWidget;
+	}
+	return nullptr;
+}
+
+TMap<SWidget*, TSharedRef<FSlateAccessibleWidget>>::TConstIterator FSlateAccessibleWidgetCache::GetAllWidgets()
+{
+	return AccessibleWidgetMap.CreateConstIterator();
+}
+
+#if !UE_BUILD_SHIPPING
+void FSlateAccessibleWidgetCache::DumpAccessibilityStats()
+{
+	const int32 NumWidgets = AccessibleWidgetMap.Num();
+	const int32 NumIds = AccessibleIdMap.Num();
+	const uint32 SizeOfWidget = sizeof(FSlateAccessibleWidget);
+	const uint32 SizeOfWidgetMap = AccessibleWidgetMap.GetAllocatedSize();
+	const uint32 SizeOfIdMap = AccessibleIdMap.GetAllocatedSize();
+	const uint32 CacheSize = NumWidgets * SizeOfWidget + SizeOfWidgetMap + SizeOfIdMap;
+
+	UE_LOG(LogAccessibility, Log, TEXT("Dumping Slate accessibility stats:"));
+	UE_LOG(LogAccessibility, Log, TEXT("Number of cached widgets and Ids: %i, %i"), NumWidgets, NumIds);
+	UE_LOG(LogAccessibility, Log, TEXT("Size of FSlateAccessibleWidget: %u"), SizeOfWidget);
+	UE_LOG(LogAccessibility, Log, TEXT("Size of id map: %u"), SizeOfIdMap);
+	UE_LOG(LogAccessibility, Log, TEXT("Approximate memory stored in cache: %u kb"), CacheSize / 1000);
+}
+
+static FAutoConsoleCommand DumpAccessibilityStatsSlateCommand
+(
+	TEXT("DumpAccessibilityStatsSlate"),
+	TEXT("Writes memory stats for Slate's accessibility data stored to LogAccessibility."),
+	FConsoleCommandDelegate::CreateStatic(&FSlateAccessibleWidgetCache::DumpAccessibilityStats)
+);
+#endif
+
+#endif
diff --git a/Engine/Source/Runtime/SlateCore/Private/Widgets/Accessibility/SlateCoreAccessibleWidgets.cpp b/Engine/Source/Runtime/SlateCore/Private/Widgets/Accessibility/SlateCoreAccessibleWidgets.cpp
--- a/Engine/Source/Runtime/SlateCore/Private/Widgets/Accessibility/SlateCoreAccessibleWidgets.cpp
+++ b/Engine/Source/Runtime/SlateCore/Private/Widgets/Accessibility/SlateCoreAccessibleWidgets.cpp
@@ -20,7 +20,6 @@ FSlateAccessibleWidget::FSlateAccessibleWidget(TWeakPtr<SWidget> InWidget, EAcce
 	: Widget(InWidget)
 	, WidgetType(InWidgetType)
 	, SiblingIndex(INDEX_NONE)
-	, bChildrenDirty(true)
 {
 	static AccessibleWidgetId RuntimeIdCounter = 0;
 	if (RuntimeIdCounter == TNumericLimits<AccessibleWidgetId>::Max())
@@ -48,12 +47,11 @@ bool FSlateAccessibleWidget::IsValid() const
 	return Widget.IsValid();
 }
 
-TSharedPtr<SWindow> FSlateAccessibleWidget::GetTopLevelSlateWindow() const
+TSharedPtr<SWindow> FSlateAccessibleWidget::GetSlateWindow() const
 {
 	if (Widget.IsValid())
 	{
 		TSharedPtr<SWidget> WindowWidget = Widget.Pin();
-		// todo: fix this for nested windows
 		while (WindowWidget.IsValid())
 		{
 			if (WindowWidget->Advanced_IsWindow())
@@ -66,9 +64,9 @@ TSharedPtr<SWindow> FSlateAccessibleWidget::GetTopLevelSlateWindow() const
 	return nullptr;
 }
 
-TSharedPtr<IAccessibleWidget> FSlateAccessibleWidget::GetTopLevelWindow() const
+TSharedPtr<IAccessibleWidget> FSlateAccessibleWidget::GetWindow() const
 {
-	return FSlateAccessibleWidgetCache::Get().GetAccessibleWidget(GetTopLevelSlateWindow());
+	return FSlateAccessibleWidgetCache::GetAccessibleWidgetChecked(GetSlateWindow());
 }
 
 FBox2D FSlateAccessibleWidget::GetBounds() const
@@ -95,10 +93,11 @@ FString FSlateAccessibleWidget::GetWidgetName() const
 {
 	if (Widget.IsValid())
 	{
-		FText AccessibleText = Widget.Pin()->GetAccessibleText();
+		TSharedPtr<SWidget> SharedWidget = Widget.Pin();
+		FText AccessibleText = SharedWidget->GetAccessibleText();
 		if (AccessibleText.IsEmpty())
 		{
-			TSharedPtr<FTagMetaData> Tag = Widget.Pin()->GetMetaData<FTagMetaData>();
+			TSharedPtr<FTagMetaData> Tag = SharedWidget->GetMetaData<FTagMetaData>();
 			if (Tag.IsValid())
 			{
 				return Tag->Tag.ToString();
@@ -120,10 +119,15 @@ FString FSlateAccessibleWidget::GetHelpText() const
 {
 	if (Widget.IsValid())
 	{
-		TSharedPtr<IToolTip> ToolTip = Widget.Pin()->GetToolTip();
-		if (ToolTip.IsValid())
+		TSharedPtr<SWidget> SharedWidget = Widget.Pin();
+		// If the accessible text is already the tooltip, don't duplicate it for the help text.
+		if (SharedWidget->GetAccessibleBehavior() != EAccessibleBehavior::ToolTip)
 		{
-			return ToolTip->GetContentWidget()->GetAccessibleText().ToString();
+			TSharedPtr<IToolTip> ToolTip = SharedWidget->GetToolTip();
+			if (ToolTip.IsValid())
+			{
+				return ToolTip->GetContentWidget()->GetAccessibleText().ToString();
+			}
 		}
 	}
 	return FString();
@@ -169,7 +173,7 @@ void FSlateAccessibleWidget::SetFocus()
 {
 	if (SupportsFocus())
 	{
-		TSharedPtr<SWindow> WidgetWindow = GetTopLevelSlateWindow();
+		TSharedPtr<SWindow> WidgetWindow = GetSlateWindow();
 		if (WidgetWindow.IsValid())
 		{
 			TArray<TSharedRef<SWindow>> WindowArray;
@@ -183,81 +187,21 @@ void FSlateAccessibleWidget::SetFocus()
 	}
 }
 
-void FSlateAccessibleWidget::MarkChildrenDirty()
-{
-	for (int32 i = 0; i < Children.Num(); ++i)
-	{
-		if (Children[i].IsValid())
-		{
-			Children[i].Pin()->SiblingIndex = INDEX_NONE;
-		}
-	}
-
-	Children.Reset();
-	bChildrenDirty = true;
-}
-
-void FSlateAccessibleWidget::UpdateAllChildren(bool bUpdateRecursively)
-{
-	if (bChildrenDirty)
-	{
-		bChildrenDirty = false;
-		if (Widget.IsValid())
-		{
-			TArray<TSharedRef<SWidget>> AccessibleChildren = GetAccessibleChildren(Widget.Pin().ToSharedRef());
-			Children.Reset(AccessibleChildren.Num());
-			for (int32 i = 0; i < AccessibleChildren.Num(); ++i)
-			{
-				TSharedPtr<FSlateAccessibleWidget> Child = FSlateAccessibleWidgetCache::Get().GetAccessibleWidget(AccessibleChildren[i]);
-				Children.Add(Child);
-				Child->Parent = StaticCastSharedRef<FSlateAccessibleWidget>(AsShared());
-				Child->SiblingIndex = i;
-
-				if (bUpdateRecursively)
-				{
-					Child->UpdateAllChildren(true);
-				}
-			}
-		}
-	}
-}
-
 void FSlateAccessibleWidget::UpdateParent(TSharedPtr<IAccessibleWidget> NewParent)
 {
-	if (Parent == NewParent)
-	{
-		return;
-	}
-
-	if (Parent.IsValid())
-	{
-		// Even though we're storing SiblingIndex, we have no guarantee
-		// that GetChildren will return the same order after a widget is
-		// added or removed, so we have to re-update all indices.
-		Parent.Pin()->MarkChildrenDirty();
-		FSlateApplicationBase::Get().GetAccessibleMessageHandler()->RaiseEvent(AsShared(), EAccessibleEvent::BeforeRemoveFromParent);
-	}
-
-	Parent = StaticCastSharedPtr<FSlateAccessibleWidget>(NewParent);
-
-	if (Parent.IsValid())
+	if (Parent != NewParent)
 	{
-		Parent.Pin()->MarkChildrenDirty();
-		FSlateApplicationBase::Get().GetAccessibleMessageHandler()->RaiseEvent(AsShared(), EAccessibleEvent::AfterAddToParent);
-	}
-	else
-	{
-		SiblingIndex = INDEX_NONE;
+		FSlateApplicationBase::Get().GetAccessibleMessageHandler()->RaiseEvent(
+			AsShared(), EAccessibleEvent::ParentChanged,
+			Parent.IsValid() ? Parent.Pin()->GetId() : IAccessibleWidget::InvalidAccessibleWidgetId,
+			NewParent.IsValid() ? NewParent->GetId() : IAccessibleWidget::InvalidAccessibleWidgetId);
+		Parent = StaticCastSharedPtr<FSlateAccessibleWidget>(NewParent);
 	}
 }
 
 TSharedPtr<IAccessibleWidget> FSlateAccessibleWidget::GetParent()
 {
-	if (Parent.IsValid())
-	{
-		return Parent.Pin();
-	}
-	return nullptr;
+	return Parent.Pin();
 }
 
 TSharedPtr<IAccessibleWidget> FSlateAccessibleWidget::GetNextSibling()
@@ -265,14 +209,9 @@ TSharedPtr<IAccessibleWidget> FSlateAccessibleWidget::GetNextSibling()
 	if (Parent.IsValid())
 	{
 		TSharedPtr<FSlateAccessibleWidget> SharedParent = Parent.Pin();
-		SharedParent->UpdateAllChildren();
 		if (SiblingIndex >= 0 && SiblingIndex < SharedParent->Children.Num() - 1)
 		{
-			const TWeakPtr<FSlateAccessibleWidget>& Child = SharedParent->Children[SiblingIndex + 1];
-			if (Child.IsValid())
-			{
-				return Child.Pin();
-			}
+			return SharedParent->Children[SiblingIndex + 1].Pin();
 		}
 	}
 	return nullptr;
@@ -283,14 +222,9 @@ TSharedPtr<IAccessibleWidget> FSlateAccessibleWidget::GetPreviousSibling()
 	if (Parent.IsValid())
 	{
 		TSharedPtr<FSlateAccessibleWidget> SharedParent = Parent.Pin();
-		SharedParent->UpdateAllChildren();
 		if (SiblingIndex >= 1 && SiblingIndex < SharedParent->Children.Num())
 		{
-			const TWeakPtr<FSlateAccessibleWidget>& Child = SharedParent->Children[SiblingIndex - 1];
-			if (Child.IsValid())
-			{
-				return Child.Pin();
-			}
+			return SharedParent->Children[SiblingIndex - 1].Pin();
 		}
 	}
 	return nullptr;
@@ -298,79 +232,23 @@ TSharedPtr<IAccessibleWidget> FSlateAccessibleWidget::GetPreviousSibling()
 
 TSharedPtr<IAccessibleWidget> FSlateAccessibleWidget::GetChildAt(int32 Index)
 {
-	UpdateAllChildren();
-	if (Index >= 0 && Index < Children.Num() && Widget.IsValid() && Widget.Pin()->CanChildrenBeAccessible())
+	check(Index >= 0 && Index < Children.Num());
+	if (Widget.IsValid())
 	{
-		const TWeakPtr<FSlateAccessibleWidget>& Child = Children[Index];
-		if (Child.IsValid())
-		{
-			return Child.Pin();
-		}
-
+		return Children[Index].Pin();
 	}
 	return nullptr;
 }
 
 int32 FSlateAccessibleWidget::GetNumberOfChildren()
 {
-	UpdateAllChildren();
-	if (Widget.IsValid() && Widget.Pin()->CanChildrenBeAccessible())
+	if (Widget.IsValid())
 	{
 		return Children.Num();
 	}
 	return 0;
 }
 
-TArray<TSharedRef<SWidget>> FSlateAccessibleWidget::GetAccessibleChildren(TSharedRef<SWidget> Widget)
-{
-	TArray<TSharedRef<SWidget>> AccessibleChildren;
-	if (!Widget->CanChildrenBeAccessible())
-	{
-		return AccessibleChildren;
-	}
-
-	FChildren* Children = Widget->GetChildren();
-	if (Children)
-	{
-		for (int32 i = 0; i < Children->Num(); ++i)
-		{
-			TSharedRef<SWidget> Child = Children->GetChildAt(i);
-			if (Child->GetAccessibleBehavior() != EAccessibleBehavior::NotAccessible)
-			{
-				AccessibleChildren.Add(Child);
-			}
-			else
-			{
-				AccessibleChildren.Append(GetAccessibleChildren(Child));
-			}
-		}
-	}
-	return AccessibleChildren;
-}
-
-TSharedPtr<IAccessibleWidget> FSlateAccessibleWidget::GetChildAtUsingGeometry(int32 X, int32 Y)
-{
-	// This is slow and we should use the HitTest grid when possible.
-	if (!IsHidden() && GetBounds().IsInside(FVector2D(X, Y)))
-	{
-		UpdateAllChildren();
-		// Traverse the hierarchy backwards in order to handle the case where widgets are overlaid on top of each other.
-		for (int32 i = Children.Num() - 1; i >= 0; --i)
-		{
-			if (Children[i].IsValid())
-			{
-				TSharedPtr<IAccessibleWidget> Child = Children[i].Pin()->GetChildAtUsingGeometry(X, Y);
-				if (Child.IsValid())
-				{
-					return Child;
-				}
-			}
-		}
-		return AsShared();
-	}
-	return nullptr;
-}
-
 // SWindow
 TSharedPtr<FGenericWindow> FSlateAccessibleWindow::GetNativeWindow() const
 {
@@ -405,11 +283,37 @@ TSharedPtr<IAccessibleWidget> FSlateAccessibleWindow::GetChildAtPosition(int32 X
 					break;
 				}
 			}
-			HitWidget = FSlateAccessibleWidgetCache::Get().GetAccessibleWidget(LastAccessibleWidget);
+			if (LastAccessibleWidget.IsValid())
+			{
+				HitWidget = FSlateAccessibleWidgetCache::GetAccessibleWidget(LastAccessibleWidget.ToSharedRef());
+			}
 		}
 		else
 		{
-			HitWidget = GetChildAtUsingGeometry(X, Y);
+			TArray<TSharedPtr<IAccessibleWidget>> ToProcess;
+			ToProcess.Reserve(10);
+			ToProcess.Add(AsShared());
+
+			while (ToProcess.Num() > 0)
+			{
+				const TSharedPtr<IAccessibleWidget> Current = ToProcess.Pop(false);
+				// Because children are weak pointers, Current could be invalid in the case where the SWidget and all
+				// shared pointers were deleted while in the middle of FSlateAccessibleMessageHandler refreshing the data.
+				if (Current.IsValid() && !Current->IsHidden() && Current->GetBounds().IsInside(FVector2D(X, Y)))
+				{
+					// The widgets are being traversed in reverse render order, so usually if a widget is rendered
+					// on top of another this will return the rendered one. But it's not 100% guarantee, and opacity
+					// screws things up sometimes. ToProcess can safely be reset because once we go down a branch
+					// we no longer care about any other branches.
+					ToProcess.Reset();
+					HitWidget = Current;
+					const int32 NumChildren = Current->GetNumberOfChildren();
+					for (int32 i = 0; i < NumChildren; ++i)
+					{
+						ToProcess.Add(Current->GetChildAt(i));
+					}
+				}
+			}
 		}
 	}
 
@@ -418,7 +322,7 @@ TSharedPtr<IAccessibleWidget> FSlateAccessibleWindow::GetChildAtPosition(int32 X
 
 TSharedPtr<IAccessibleWidget> FSlateAccessibleWindow::GetFocusedWidget() const
 {
-	return FSlateAccessibleWidgetCache::Get().GetAccessibleWidget(FSlateApplicationBase::Get().GetKeyboardFocusedWidget());
+	return FSlateAccessibleWidgetCache::GetAccessibleWidgetChecked(FSlateApplicationBase::Get().GetKeyboardFocusedWidget());
 }
 
 FString FSlateAccessibleWindow::GetWidgetName() const
diff --git a/Engine/Source/Runtime/SlateCore/Private/Widgets/Images/SImage.cpp b/Engine/Source/Runtime/SlateCore/Private/Widgets/Images/SImage.cpp
--- a/Engine/Source/Runtime/SlateCore/Private/Widgets/Images/SImage.cpp
+++ b/Engine/Source/Runtime/SlateCore/Private/Widgets/Images/SImage.cpp
@@ -80,7 +80,7 @@ void SImage::SetImage(TAttribute<const FSlateBrush*> InImage)
 }
 
 #if WITH_ACCESSIBILITY
-TSharedPtr<FSlateAccessibleWidget> SImage::CreateAccessibleWidget()
+TSharedRef<FSlateAccessibleWidget> SImage::CreateAccessibleWidget()
 {
 	return MakeShareable<FSlateAccessibleWidget>(new FSlateAccessibleImage(SharedThis(this)));
 }
diff --git a/Engine/Source/Runtime/SlateCore/Private/Widgets/SWidget.cpp b/Engine/Source/Runtime/SlateCore/Private/Widgets/SWidget.cpp
--- a/Engine/Source/Runtime/SlateCore/Private/Widgets/SWidget.cpp
+++ b/Engine/Source/Runtime/SlateCore/Private/Widgets/SWidget.cpp
@@ -595,7 +595,7 @@ void SWidget::AssignParentWidget(TSharedPtr<SWidget> InParent)
 #if WITH_ACCESSIBILITY
 	if (FSlateApplicationBase::IsInitialized())
 	{
-		FSlateApplicationBase::Get().GetAccessibleMessageHandler()->OnWidgetParentChanged(AsShared());
+		FSlateApplicationBase::Get().GetAccessibleMessageHandler()->MarkDirty();
 	}
 #endif
 	if (InParent.IsValid())
@@ -617,7 +617,7 @@ bool SWidget::ConditionallyDetatchParentWidget(SWidget* InExpectedParent)
 #if WITH_ACCESSIBILITY
 		if (FSlateApplicationBase::IsInitialized())
 		{
-			FSlateApplicationBase::Get().GetAccessibleMessageHandler()->OnWidgetParentChanged(AsShared());
+			FSlateApplicationBase::Get().GetAccessibleMessageHandler()->MarkDirty();
 		}
 #endif
 
@@ -1275,7 +1275,7 @@ void SWidget::SetOnMouseLeave(FSimpleNoReplyPointerEventHandler EventHandler)
 }
 
 #if WITH_ACCESSIBILITY
-TSharedPtr<FSlateAccessibleWidget> SWidget::CreateAccessibleWidget()
+TSharedRef<FSlateAccessibleWidget> SWidget::CreateAccessibleWidget()
 {
 	return MakeShareable<FSlateAccessibleWidget>(new FSlateAccessibleWidget(AsShared()));
 }
@@ -1295,19 +1295,22 @@ void SWidget::SetAccessibleBehavior(EAccessibleBehavior InBehavior, const TAttri
 			TAttribute<FText>& Text = (AccessibleType == EAccessibleType::Main) ? AccessibleData.AccessibleText : AccessibleData.AccessibleSummaryText;
 			Text = InText;
 		}
-		const bool bWasAccessible = IsAccessible();
+		const bool bWasAccessible = Behavior != EAccessibleBehavior::NotAccessible;
 		Behavior = InBehavior;
-		if (AccessibleType == EAccessibleType::Main && bWasAccessible != IsAccessible())
+		if (AccessibleType == EAccessibleType::Main && bWasAccessible != (Behavior != EAccessibleBehavior::NotAccessible))
 		{
-			FSlateApplicationBase::Get().GetAccessibleMessageHandler()->OnWidgetAccessibleBehaviorChanged(AsShared());
+			FSlateApplicationBase::Get().GetAccessibleMessageHandler()->MarkDirty();
 		}
 	}
 }
 
 void SWidget::SetCanChildrenBeAccessible(bool InCanChildrenBeAccessible)
 {
-	AccessibleData.bCanChildrenBeAccessible = InCanChildrenBeAccessible;
-	// todo: emit notifications
+	if (AccessibleData.bCanChildrenBeAccessible != InCanChildrenBeAccessible)
+	{
+		AccessibleData.bCanChildrenBeAccessible = InCanChildrenBeAccessible;
+		FSlateApplicationBase::Get().GetAccessibleMessageHandler()->MarkDirty();
+	}
 }
 
 void SWidget::SetDefaultAccessibleText(EAccessibleType AccessibleType)
diff --git a/Engine/Source/Runtime/SlateCore/Private/Widgets/SWindow.cpp b/Engine/Source/Runtime/SlateCore/Private/Widgets/SWindow.cpp
--- a/Engine/Source/Runtime/SlateCore/Private/Widgets/SWindow.cpp
+++ b/Engine/Source/Runtime/SlateCore/Private/Widgets/SWindow.cpp
@@ -2071,7 +2071,7 @@ FScopedSwitchWorldHack::FScopedSwitchWorldHack( const FWidgetPath& WidgetPath )
 #endif
 
 #if WITH_ACCESSIBILITY
-TSharedPtr<FSlateAccessibleWidget> SWindow::CreateAccessibleWidget()
+TSharedRef<FSlateAccessibleWidget> SWindow::CreateAccessibleWidget()
 {
 	return MakeShareable<FSlateAccessibleWidget>(new FSlateAccessibleWindow(SharedThis(this)));
 }
diff --git a/Engine/Source/Runtime/SlateCore/Public/Application/SlateApplicationBase.h b/Engine/Source/Runtime/SlateCore/Public/Application/SlateApplicationBase.h
--- a/Engine/Source/Runtime/SlateCore/Public/Application/SlateApplicationBase.h
+++ b/Engine/Source/Runtime/SlateCore/Public/Application/SlateApplicationBase.h
@@ -13,6 +13,9 @@
 DECLARE_MULTICAST_DELEGATE_TwoParams(FOnDebugSafeZoneChanged, const FMargin&, bool);
 
 class FActiveTimerHandle;
+#if WITH_ACCESSIBILITY
+class FSlateAccessibleMessageHandler;
+#endif
 class FSlateApplicationBase;
 class FWidgetPath;
 class IToolTip;
@@ -238,11 +241,11 @@ public:
 
 #if WITH_ACCESSIBILITY
 	/** 
-	 * Subclasses should store their own accessibility handler and return it through this function.
+	 * Accessor for the accessible message handler. One must always exist, even if it's never activated.
 	 *
-	 * @return A pointer to the accessibility handler, if one exists
+	 * @return A reference to the assigned accessible message handler
 	 */
-	virtual TSharedPtr<class FSlateAccessibleMessageHandler> GetAccessibleMessageHandler() const = 0;
+	TSharedRef<FSlateAccessibleMessageHandler> GetAccessibleMessageHandler() const { return AccessibleMessageHandler; }
 #endif
 protected:
 	/**
@@ -589,6 +592,11 @@ public:
 #endif
 
 protected:
+#if WITH_ACCESSIBILITY
+	/** Manager for widgets and application to interact with accessibility API */
+	TSharedRef<FSlateAccessibleMessageHandler> AccessibleMessageHandler;
+#endif
+
 	/** multicast delegate to broadcast when a global invalidate is requested */
 	FOnGlobalInvalidate OnGlobalInvalidateEvent;
 
diff --git a/Engine/Source/Runtime/SlateCore/Public/Widgets/Accessibility/SlateAccessibleMessageHandler.h b/Engine/Source/Runtime/SlateCore/Public/Widgets/Accessibility/SlateAccessibleMessageHandler.h
--- a/Engine/Source/Runtime/SlateCore/Public/Widgets/Accessibility/SlateAccessibleMessageHandler.h
+++ b/Engine/Source/Runtime/SlateCore/Public/Widgets/Accessibility/SlateAccessibleMessageHandler.h
@@ -7,14 +7,19 @@
 #include "GenericPlatform/GenericAccessibleInterfaces.h"
 #include "Misc/Variant.h"
 
+class FSlateAccessibleWidget;
 class SWidget;
+class SWindow;
 
 /**
  * Message handling system for Slate Accessibility API, dealing with both receiving events and pushing them back to the platform layer.
+ * The message handler is also responsible for processing the Slate widget tree and 
  */
 class SLATECORE_API FSlateAccessibleMessageHandler : public FGenericAccessibleMessageHandler
 {
 public:
+	FSlateAccessibleMessageHandler();
+
 	// FGenericAccessibleMessageHandler
 	virtual void OnActivate() override;
 	virtual void OnDeactivate() override;
@@ -30,34 +35,7 @@ public:
 	 * @param Widget The widget that is being deleted.
 	 */
 	void OnWidgetRemoved(SWidget* Widget);
-	/**
-	 * Callback for a Slate widget's parent changing.
-	 *
-	 * @param Widget The widget whose parent changed.
-	 */
-	void OnWidgetParentChanged(TSharedRef<SWidget> Widget);
-	/**
-	 * Callback for a Slate widget's children changing. Although somewhat rare, widgets are not required
-	 * to return a child even if the child was parented to it (for instance, SWidgetSwitcher).
-	 *
-	 * Note: This should not be called if the children are already calling OnWidgetParentChanged().
-	 *
-	 * @param Widget The widget whose children changed.
-	 */
-	void OnWidgetChildrenChanged(TSharedRef<SWidget> Widget);
-	/**
-	 * Callback for a Slate widget's main accessible behavior changing. 
-	 *
-	 * @param Widget The widget whose behavior changed.
-	 */
-	void OnWidgetAccessibleBehaviorChanged(TSharedRef<SWidget> Widget);
-	/**
-	 * Callback for a Slate widget indicating that some event happened to it.
-	 *
-	 * @param Widget The widget raising the event.
-	 * @param Event The type of event being raised.
-	 */
-	void OnWidgetEventRaised(TSharedRef<SWidget> Widget, EAccessibleEvent Event);
+
 	/**
 	 * Callback for a Slate widget indicating that a property change occurred. This may also be used by certain events
 	 * such as Notification which don't have an 'OldValue'. Only NewValue should be set for those types of events.
@@ -67,7 +45,33 @@ public:
 	 * @param OldValue The value of the property being changed before the change occurred.
 	 * @param NewValue The value of the property being changed after the change occurred, or any miscellaneous data for the event.
 	 */
-	void OnWidgetEventRaised(TSharedRef<SWidget> Widget, EAccessibleEvent Event, FVariant OldValue, FVariant NewValue);
+	void OnWidgetEventRaised(TSharedRef<SWidget> Widget, EAccessibleEvent Event, FVariant OldValue = FVariant(), FVariant NewValue = FVariant());
+
+	/**
+	 * Refresh the accessible widget tree next available tick. This should be called any time the Slate tree changes.
+	 */
+	void MarkDirty() { bDirty = true; }
+
+	/**
+	 * Process any pending Slate widgets and update the accessible widget tree.
+	 */
+	void Tick();
+
+private:
+	struct FWidgetAndParent {
+		FWidgetAndParent(TWeakPtr<SWidget> InWidget, TSharedRef<FSlateAccessibleWidget> InParent)
+			: Widget(InWidget), Parent(InParent)
+		{
+		}
+
+		TWeakPtr<SWidget> Widget;
+		TSharedRef<FSlateAccessibleWidget> Parent;
+	};
+
+	/** A list of widgets waiting to be processed in order to keep the accessible widget tree up to date. */
+	TArray<FWidgetAndParent> ToProcess;
+	/** If true, Tick() will begin the update process to the accessible widget tree. Use MarkDirty() to set. */
+	bool bDirty;
 };
 
 #endif
diff --git a/Engine/Source/Runtime/SlateCore/Public/Widgets/Accessibility/SlateAccessibleWidgetCache.h b/Engine/Source/Runtime/SlateCore/Public/Widgets/Accessibility/SlateAccessibleWidgetCache.h
--- a/Engine/Source/Runtime/SlateCore/Public/Widgets/Accessibility/SlateAccessibleWidgetCache.h
+++ b/Engine/Source/Runtime/SlateCore/Public/Widgets/Accessibility/SlateAccessibleWidgetCache.h
@@ -6,7 +6,8 @@
 
 #include "CoreMinimal.h"
 #include "Widgets/Accessibility/SlateCoreAccessibleWidgets.h"
-#include "Widgets/SWidget.h"
+
+class SWidget;
 
 /**
  * Singleton used to retrieve accessible widgets for a given Slate widget. Accessible widgets will persist
@@ -15,106 +16,59 @@
 class SLATECORE_API FSlateAccessibleWidgetCache
 {
 public:
-	static FSlateAccessibleWidgetCache& Get()
-	{
-		static FSlateAccessibleWidgetCache AccessibleWidgetCache;
-		return AccessibleWidgetCache;
-	}
-
 	/** Empty the cache and release all cached accessible widgets. This should generally only be done with accessibility is turned off. */
-	void ClearAll()
-	{
-		AccessibleWidgetMap.Empty();
-	}
+	static void ClearAll();
 
 	/**
 	 * Callback for when an SWidget is deleted. FSlateAccessibleMessageHandler should be the only thing to call this.
-	 * The widget is removed from the cache and removed from the accessible tree.
+	 * The widget will be removed from the cache, and it should have already detached from its parent.
 	 *
 	 * @param Widget the Slate widget that is being deleted
 	 */
-	TSharedPtr<FSlateAccessibleWidget> RemoveWidget(SWidget* Widget)
-	{
-		TSharedPtr<FSlateAccessibleWidget> RemovedWidget;
-		TSharedPtr<FSlateAccessibleWidget>* AccessibleWidget = AccessibleWidgetMap.Find(Widget);
-		if (AccessibleWidget)
-		{
-			// Make sure to disconnect this widget from the tree
-			(*AccessibleWidget)->UpdateParent(nullptr);
-			AccessibleWidgetMap.RemoveAndCopyValue(Widget, RemovedWidget);
-			if (RemovedWidget.IsValid())
-			{
-				AccessibleIdMap.Remove(RemovedWidget->GetId());
-			}
-		}
-		return RemovedWidget;
-	}
+	static TSharedPtr<FSlateAccessibleWidget> RemoveWidget(SWidget* Widget);
 
 	/**
-	 * Get a cached accessible widget for a given Slate widget, or create a new one is bCreateIfMissing is true.
-	 * This may return nullptr in the case where the Slate widget is not a valid accessible widget. A non-valid
-	 * accessible widget includes the case where the widget has accessible flags set but its parent cannot have
-	 * accessible children.
+	 * Get a cached accessible widget for a given Slate widget. This may return nullptr in the case where the
+	 * Slate widget is not a valid accessible widget, which is either because the widget's behavior is not
+	 * accessible or it has a parent that cannot have accessible children.
 	 *
 	 * @param Widget The Slate widget to get the accessible widget for
-	 * @param bCreateIfMissing If the widget is not already in the cache, whether or not it should be created
 	 * @return The accessible widget for the given Slate widget, or nullptr if the widget is not accessible.
 	 */
-	TSharedPtr<FSlateAccessibleWidget> GetAccessibleWidget(const TSharedPtr<SWidget> Widget, bool bCreateIfMissing = true)
-	{
-		if (!Widget.IsValid() || !Widget->IsAccessible())
-		{
-			return nullptr;
-		}
-
-		TSharedPtr<FSlateAccessibleWidget>* AccessibleWidget = AccessibleWidgetMap.Find(Widget.Get());
-		if (AccessibleWidget)
-		{
-			return *AccessibleWidget;
-		}
-		else if (bCreateIfMissing)
-		{
-			TSharedPtr<FSlateAccessibleWidget> NewWidget = AccessibleWidgetMap.Add(Widget.Get(), Widget->CreateAccessibleWidget());
-			AccessibleIdMap.Add(NewWidget->GetId(), NewWidget);
-			// todo: fix this for nested windows
-			if (NewWidget->AsWindow())
-			{
-				NewWidget->UpdateAllChildren(true);
-			}
-			return NewWidget;
-		}
-		return nullptr;
-	}
+	static TSharedPtr<FSlateAccessibleWidget> GetAccessibleWidgetChecked(const TSharedPtr<SWidget>& Widget);
+	/**
+	 * Get a cached accessible widget for a given Slate widget. The caller is responsible for making sure
+	 * that the widget is fully accessible in the current Slate widget tree.
+	 * 
+	 * @param Widget The Slate widget to get the accessible widget for
+	 * @return The accessible widget for the given Slate widget
+	 */
+	static TSharedRef<FSlateAccessibleWidget> GetAccessibleWidget(const TSharedRef<SWidget>& Widget);
 
 	/**
 	 * Get a cached accessible widget for an identifier that matches the accessible widget's GetId() return value.
 	 * This will NOT create the widget if it does not exist yet, since Ids can only be generated by the widget itself.
-	 * Note that InvalidAccessibleWidgetId will always return null.
+	 * InvalidAccessibleWidgetId will always return null.
 	 *
 	 * @param Id The Id of a widget that corresponds to that widget's GetId() function
 	 * @return The accessible widget for the given Id, or nullptr if there is no widget with that Id.
 	 */
-	TSharedPtr<FSlateAccessibleWidget> GetAccessibleWidgetFromId(AccessibleWidgetId Id)
-	{
-		TSharedPtr<FSlateAccessibleWidget>* AccessibleWidget = AccessibleIdMap.Find(Id);
-		if (AccessibleWidget)
-		{
-			return *AccessibleWidget;
-		}
-		return nullptr;
-	}
+	static TSharedPtr<FSlateAccessibleWidget> GetAccessibleWidgetFromId(AccessibleWidgetId Id);
+
+	static TMap<SWidget*, TSharedRef<FSlateAccessibleWidget>>::TConstIterator GetAllWidgets();
+
+#if !UE_BUILD_SHIPPING
+	static void DumpAccessibilityStats();
+#endif
 
 private:
 	/**
-	 * Hold on to a shared pointer for each accessible widget until the SWidget is deleted. A raw pointer is used
+	 * Hold on to a shared ref for each accessible widget until the SWidget is deleted. A raw pointer is used
 	 * because widgets are removed from the cache via the SWidget destructor.
 	 */
-	TMap<SWidget*, TSharedPtr<FSlateAccessibleWidget>> AccessibleWidgetMap;
+	static TMap<SWidget*, TSharedRef<FSlateAccessibleWidget>> AccessibleWidgetMap;
 	/** Map each accessible widget to its ID to provide a second way to do lookups. */
-	TMap<AccessibleWidgetId, TSharedPtr<FSlateAccessibleWidget>> AccessibleIdMap;
-
-	FSlateAccessibleWidgetCache() {}
-	~FSlateAccessibleWidgetCache() {}
+	static TMap<AccessibleWidgetId, TSharedRef<FSlateAccessibleWidget>> AccessibleIdMap;
 };
 
 #endif
diff --git a/Engine/Source/Runtime/SlateCore/Public/Widgets/Accessibility/SlateCoreAccessibleWidgets.h b/Engine/Source/Runtime/SlateCore/Public/Widgets/Accessibility/SlateCoreAccessibleWidgets.h
--- a/Engine/Source/Runtime/SlateCore/Public/Widgets/Accessibility/SlateCoreAccessibleWidgets.h
+++ b/Engine/Source/Runtime/SlateCore/Public/Widgets/Accessibility/SlateCoreAccessibleWidgets.h
@@ -17,6 +17,7 @@ class SWindow;
  */
 class SLATECORE_API FSlateAccessibleWidget : public IAccessibleWidget
 {
+	friend class FSlateAccessibleMessageHandler;
 public:
 	FSlateAccessibleWidget(TWeakPtr<SWidget> InWidget, EAccessibleWidgetType InWidgetType = EAccessibleWidgetType::Unknown);
 	virtual ~FSlateAccessibleWidget();
@@ -24,7 +25,7 @@ public:
 	// IAccessibleWidget
 	virtual AccessibleWidgetId GetId() const override final;
 	virtual bool IsValid() const override final;
-	virtual TSharedPtr<IAccessibleWidget> GetTopLevelWindow() const override final;
+	virtual TSharedPtr<IAccessibleWidget> GetWindow() const override final;
 	virtual FBox2D GetBounds() const override final;
 	virtual TSharedPtr<IAccessibleWidget> GetParent() override final;
 	virtual TSharedPtr<IAccessibleWidget> GetNextSibling() override final;
@@ -43,60 +44,28 @@ public:
 	virtual FString GetHelpText() const override;
 	// ~
 
-	/** Tell this widget to recompute its children the next time they are requested. */
-	void MarkChildrenDirty();
 	/**
 	 * Detach this widget from its current parent and attach it to a new parent. This will emit notifications back to the accessible message handler.
 	 *
 	 * @param NewParent The widget to assign as the new parent widget.
 	 */
 	void UpdateParent(TSharedPtr<IAccessibleWidget> NewParent);
-	/**
-	 * If MarkChildrenDirty() has been called, recalculate the list of all accessible widgets below this one.
-	 * Because SWidget->GetChildren() has no guarantees about what it returns and how it returns it, we can
-	 * never truly 100% guarantee that the accessible tree will be in sync with the Slate tree.
-	 *
-	 * We make a reasonable assumption that widgets are smart about implementing this function to return the
-	 * same widgets every time. However, we can't assume anything about when a child gets added or removed
-	 * with respect to the ordering of the children. Because of this, we have to recompute their indices
-	 * any time we suspect the hierarchy may have changed.
-	 *
-	 * @param bUpdateRecursively If true, calls UpdateAllChildren() on any children found.
-	 */
-	void UpdateAllChildren(bool bUpdateRecursively = false);
-
-	/**
-	 * Search the Slate hierarchy recursively and generate a list of all accessible widgets whose parent is this widget.
-	 *
-	 * @param AccessibleWidget The root widget to find children for.
-	 * @return All Slate widgets whose accessible parent is the passed-in widget.
-	 */
-	static TArray<TSharedRef<SWidget>> GetAccessibleChildren(TSharedRef<SWidget> AccessibleWidget);
 
 protected:
-	/**
-	 * Recursively find the accessible widget under the specified X,Y coordinates.
-	 *
-	 * @param X The X coordinate to search in absolute screen space.
-	 * @param Y The Y coordinate to search in absolute screen space.
-	 * @return The deepest accessible widget found.
-	 */
-	TSharedPtr<IAccessibleWidget> GetChildAtUsingGeometry(int32 X, int32 Y);
-
 	/** The underlying Slate widget backing this accessible widget. */
 	TWeakPtr<SWidget> Widget;
 	/** What type of widget the platform's accessibility API should treat this as. */
 	EAccessibleWidgetType WidgetType;
 	/** The accessible parent to this widget. This should usually be valid on widgets in the hierarchy, except for SWindows. */
 	TWeakPtr<FSlateAccessibleWidget> Parent;
-	/** All accessible widgets whose parent is this widget. This is not necessarily correct unless UpdateAllChildren() is called first. */
+	/** All accessible widgets whose parent is this widget. */
 	TArray<TWeakPtr<FSlateAccessibleWidget>> Children;
+	/** A temporary array filled when processing the widget tree, which will eventually be moved to the Children array. */
+	TArray<TWeakPtr<FSlateAccessibleWidget>> ChildrenBuffer;
 	/** The index of this widget in its parent's list of children. */
 	int32 SiblingIndex;
 	/** An application-unique identifier for GetId(). */
 	AccessibleWidgetId Id;
-	/** Whether the contents of the Children array has changed and UpdateAllChildren() needs to be called. */
-	bool bChildrenDirty;
 
 private:
 	/**
@@ -104,7 +73,7 @@ private:
 	 *
 	 * @return The parent SWindow for the Slate widget referenced by this accessible widget.
 	 */
-	TSharedPtr<SWindow> GetTopLevelSlateWindow() const;
+	TSharedPtr<SWindow> GetSlateWindow() const;
 };
 
 // SWindow
diff --git a/Engine/Source/Runtime/SlateCore/Public/Widgets/Images/SImage.h b/Engine/Source/Runtime/SlateCore/Public/Widgets/Images/SImage.h
--- a/Engine/Source/Runtime/SlateCore/Public/Widgets/Images/SImage.h
+++ b/Engine/Source/Runtime/SlateCore/Public/Widgets/Images/SImage.h
@@ -69,7 +69,7 @@ public:
 	// SWidget overrides
 	virtual int32 OnPaint( const FPaintArgs& Args, const FGeometry& AllottedGeometry, const FSlateRect& MyCullingRect, FSlateWindowElementList& OutDrawElements, int32 LayerId, const FWidgetStyle& InWidgetStyle, bool bParentEnabled ) const override;
 #if WITH_ACCESSIBILITY
-	virtual TSharedPtr<FSlateAccessibleWidget> CreateAccessibleWidget() override;
+	virtual TSharedRef<FSlateAccessibleWidget> CreateAccessibleWidget() override;
 #endif
 
 protected:
diff --git a/Engine/Source/Runtime/SlateCore/Public/Widgets/SWidget.h b/Engine/Source/Runtime/SlateCore/Public/Widgets/SWidget.h
--- a/Engine/Source/Runtime/SlateCore/Public/Widgets/SWidget.h
+++ b/Engine/Source/Runtime/SlateCore/Public/Widgets/SWidget.h
@@ -691,7 +691,7 @@ public:
 	virtual EWindowZone::Type GetWindowZoneOverride() const;
 
 #if WITH_ACCESSIBILITY
-	virtual TSharedPtr<class FSlateAccessibleWidget> CreateAccessibleWidget();
+	virtual TSharedRef<class FSlateAccessibleWidget> CreateAccessibleWidget();
 #endif
 
 public:
diff --git a/Engine/Source/Runtime/SlateCore/Public/Widgets/SWindow.h b/Engine/Source/Runtime/SlateCore/Public/Widgets/SWindow.h
--- a/Engine/Source/Runtime/SlateCore/Public/Widgets/SWindow.h
+++ b/Engine/Source/Runtime/SlateCore/Public/Widgets/SWindow.h
@@ -788,7 +788,7 @@ public:
 	virtual bool Advanced_IsWindow() const { return true; }
 
 #if WITH_ACCESSIBILITY
-	virtual TSharedPtr<FSlateAccessibleWidget> CreateAccessibleWidget() override;
+	virtual TSharedRef<FSlateAccessibleWidget> CreateAccessibleWidget() override;
 	virtual void SetDefaultAccessibleText(EAccessibleType AccessibleType = EAccessibleType::Main) override;
 #endif
 private:
