Add missing headers

diff --git a/Engine/Source/Runtime/SlateCore/Private/Widgets/Accessibility/SlateAccessibleMessageHandler.cpp b/Engine/Source/Runtime/SlateCore/Private/Widgets/Accessibility/SlateAccessibleMessageHandler.cpp
--- a/Engine/Source/Runtime/SlateCore/Private/Widgets/Accessibility/SlateAccessibleMessageHandler.cpp
+++ b/Engine/Source/Runtime/SlateCore/Private/Widgets/Accessibility/SlateAccessibleMessageHandler.cpp
@@ -9,6 +9,7 @@
 #include "Widgets/SWidget.h"
 #include "Widgets/SWindow.h"
 #include "Input/HittestGrid.h"
+#include "HAL/IConsoleManager.h"
 
 DECLARE_CYCLE_STAT(TEXT("Slate Accessibility: Tick"), STAT_AccessibilitySlateTick, STATGROUP_Accessibility);
 DECLARE_CYCLE_STAT(TEXT("Slate Accessibility: Event Raised"), STAT_AccessibilitySlateEventRaised, STATGROUP_Accessibility);
diff --git a/Engine/Source/Runtime/SlateCore/Private/Widgets/Accessibility/SlateAccessibleWidgetCache.cpp b/Engine/Source/Runtime/SlateCore/Private/Widgets/Accessibility/SlateAccessibleWidgetCache.cpp
--- a/Engine/Source/Runtime/SlateCore/Private/Widgets/Accessibility/SlateAccessibleWidgetCache.cpp
+++ b/Engine/Source/Runtime/SlateCore/Private/Widgets/Accessibility/SlateAccessibleWidgetCache.cpp
@@ -6,6 +6,7 @@
 #include "Application/SlateApplicationBase.h"
 #include "Widgets/Accessibility/SlateAccessibleMessageHandler.h"
 #include "Widgets/SWidget.h"
+#include "HAL/IConsoleManager.h"
 
 TMap<SWidget*, TSharedRef<FSlateAccessibleWidget>> FSlateAccessibleWidgetCache::AccessibleWidgetMap;
 TMap<AccessibleWidgetId, TSharedRef<FSlateAccessibleWidget>> FSlateAccessibleWidgetCache::AccessibleIdMap;
