Fix dashed outlines always appearing around retainers.  Should have been enabled only when invalidation debugging is on.

diff --git a/Engine/Source/Runtime/SlateCore/Private/Debugging/SlateDebugging.cpp b/Engine/Source/Runtime/SlateCore/Private/Debugging/SlateDebugging.cpp
--- a/Engine/Source/Runtime/SlateCore/Private/Debugging/SlateDebugging.cpp
+++ b/Engine/Source/Runtime/SlateCore/Private/Debugging/SlateDebugging.cpp
@@ -283,17 +283,19 @@ void FSlateDebugging::WidgetInvalidated(FSlateInvalidationRoot& InvalidationRoot
 
 void FSlateDebugging::DrawInvalidationRoot(const SWidget& RootWidget, int32 LayerId, FSlateWindowElementList& OutDrawElements)
 {
-	static const FSlateBrush* WhiteBrush = FCoreStyle::Get().GetBrush(TEXT("FocusRectangle"));
-
+	if(GSlateInvalidationDebugging)
+	{
+		static const FSlateBrush* WhiteBrush = FCoreStyle::Get().GetBrush(TEXT("FocusRectangle"));
 
-	FSlateDrawElement::MakeBox(
-		OutDrawElements,
-		LayerId,
-		RootWidget.GetPaintSpaceGeometry().ToPaintGeometry(),
-		WhiteBrush,
-		ESlateDrawEffect::None,
-		FLinearColor(128,0,128)
-	);
+		FSlateDrawElement::MakeBox(
+			OutDrawElements,
+			LayerId,
+			RootWidget.GetPaintSpaceGeometry().ToPaintGeometry(),
+			WhiteBrush,
+			ESlateDrawEffect::None,
+			FLinearColor(128, 0, 128)
+		);
+	}
 }
 
 void FSlateDebugging::DrawInvalidatedWidgets(const FSlateInvalidationRoot& Root, const FPaintArgs& PaintArgs, FSlateWindowElementList& OutDrawElements)
