Fix crash opening inventory in creative

diff --git a/Engine/Source/Runtime/SlateCore/Private/Rendering/DrawElements.cpp b/Engine/Source/Runtime/SlateCore/Private/Rendering/DrawElements.cpp
--- a/Engine/Source/Runtime/SlateCore/Private/Rendering/DrawElements.cpp
+++ b/Engine/Source/Runtime/SlateCore/Private/Rendering/DrawElements.cpp
@@ -912,6 +912,23 @@ static const FSlateClippingState* GetClipStateFromParent(const FSlateClippingMan
 	}
 }
 
+FSlateCachedElementListNode* FSlateCachedElementData::AddCache(const SWidget* Widget)
+{
+#if WITH_SLATE_DEBUGGING
+	for (FSlateCachedElementList& CachedElementList : CachedElementLists)
+	{
+		ensure(CachedElementList.Widget != Widget);
+	}
+#endif
+
+	FSlateCachedElementListNode* NewNode = new FSlateCachedElementListNode(FSlateCachedElementList(this, Widget));
+
+	CachedElementLists.AddTail(NewNode);
+	NewNode->GetValue().Initialize();
+
+	return NewNode;
+}
+
 FSlateDrawElement& FSlateCachedElementData::AddCachedElement(FSlateCachedElementListNode* CacheNode, const FSlateClippingManager& ParentClipManager, const SWidget* CurrentWidget)
 {
 #if WITH_SLATE_DEBUGGING
diff --git a/Engine/Source/Runtime/SlateCore/Private/Widgets/SWidget.cpp b/Engine/Source/Runtime/SlateCore/Private/Widgets/SWidget.cpp
--- a/Engine/Source/Runtime/SlateCore/Private/Widgets/SWidget.cpp
+++ b/Engine/Source/Runtime/SlateCore/Private/Widgets/SWidget.cpp
@@ -174,9 +174,9 @@ SWidget::SWidget()
 SWidget::~SWidget()
 {
 	// Unregister all ActiveTimers so they aren't left stranded in the Application's list.
-	if ( FSlateApplicationBase::IsInitialized() )
+	if (FSlateApplicationBase::IsInitialized())
 	{
-		for ( const auto& ActiveTimerHandle : ActiveTimers )
+		for (const auto& ActiveTimerHandle : ActiveTimers)
 		{
 			FSlateApplicationBase::Get().UnRegisterActiveTimer(ActiveTimerHandle);
 		}
@@ -189,7 +189,12 @@ SWidget::~SWidget()
 		// Reset handle
 		FastPathProxyHandle = FWidgetProxyHandle();
 
-		check(!PersistentState.CachedElementListNode);
+		// Note: this would still be valid if a widget was painted and then destroyed in the same frame.  
+		// In that case invalidation hasn't taken place for added widgets so the invalidation panel doesn't know about their cached element data to clean it up
+		if (PersistentState.CachedElementListNode)
+		{
+			PersistentState.CachedElementListNode->GetValue().GetOwningData()->RemoveCache(PersistentState.CachedElementListNode);
+		}
 
 #if WITH_ACCESSIBILITY
 		FSlateApplicationBase::Get().GetAccessibleMessageHandler()->OnWidgetRemoved(this);
diff --git a/Engine/Source/Runtime/SlateCore/Public/Rendering/DrawElements.h b/Engine/Source/Runtime/SlateCore/Public/Rendering/DrawElements.h
--- a/Engine/Source/Runtime/SlateCore/Public/Rendering/DrawElements.h
+++ b/Engine/Source/Runtime/SlateCore/Public/Rendering/DrawElements.h
@@ -350,10 +350,13 @@ struct FSlateCachedFastPathRenderingData
 	FSlateIndexArray Indices;
 };
 
+struct FSlateCachedElementData;
+
 struct FSlateCachedElementList
 {
-	FSlateCachedElementList(const SWidget* InWidget)
-		: Widget(InWidget)
+	FSlateCachedElementList(FSlateCachedElementData* InOwningData, const SWidget* InWidget)
+		: OwningData(InOwningData)
+		, Widget(InWidget)
 		, CachedRenderingData(nullptr)
 		, bNewData(false)
 	{}
@@ -367,6 +370,8 @@ struct FSlateCachedElementList
 
 	void Reset();
 
+	FSlateCachedElementData* GetOwningData() { return OwningData; }
+
 	FSlateRenderBatch& AddRenderBatch(int32 InLayer, const FShaderParams& InShaderParams, const FSlateShaderResource* InResource, ESlateDrawPrimitive InPrimitiveType, ESlateShader InShaderType, ESlateDrawEffect InDrawEffects, ESlateBatchDrawFlag InDrawFlags, int8 SceneIndex);
 
 	void AddCachedClipState(FSlateCachedClipState& ClipStateToCache);
@@ -376,7 +381,9 @@ struct FSlateCachedElementList
 private:
 	SLATECORE_API void DestroyCachedVertexData();
 public:
+	FSlateCachedElementData* OwningData;
 	const SWidget* Widget;
+
 	/** List of source draw elements to create batches from */
 	FSlateDrawElementArray DrawElements;
 	/** List of cached batches to submit for drawing */
@@ -412,23 +419,7 @@ struct FSlateCachedElementData
 		ElementListNode = nullptr;
 	}
 
-	FSlateCachedElementListNode* AddCache(const SWidget* Widget)
-	{
-#if WITH_SLATE_DEBUGGING
-		for (FSlateCachedElementList& CachedElementList : CachedElementLists)
-		{
-			ensure(CachedElementList.Widget != Widget);
-		}
-#endif
-
-		FSlateCachedElementListNode* NewNode = new FSlateCachedElementListNode(Widget);
-
-
-		CachedElementLists.AddTail(NewNode);
-		NewNode->GetValue().Initialize();
-
-		return NewNode;
-	}
+	FSlateCachedElementListNode* AddCache(const SWidget* Widget);
 
 	FSlateDrawElement& AddCachedElement(FSlateCachedElementListNode* CacheNode, const FSlateClippingManager& ParentClipManager, const SWidget* Widget);
 
