Remove ensure in SWidget::paint that is causing hitches but isnt really detecting anything harmful

diff --git a/Engine/Source/Runtime/SlateCore/Private/Widgets/SWidget.cpp b/Engine/Source/Runtime/SlateCore/Private/Widgets/SWidget.cpp
--- a/Engine/Source/Runtime/SlateCore/Private/Widgets/SWidget.cpp
+++ b/Engine/Source/Runtime/SlateCore/Private/Widgets/SWidget.cpp
@@ -1224,8 +1224,10 @@ int32 SWidget::Paint(const FPaintArgs& Args, const FGeometry& AllottedGeometry,
 	UpdatedArgs.SetInheritedHittestability(bOutgoingHittestability);
 
 
+#if 0
 	// test ensure that we are not the last thing holding this widget together
 	ensure(!MutableThis.IsUnique());
+#endif
 
 
 	if (!FastPathProxyHandle.IsValid() && PersistentState.CachedElementListNode != nullptr)
