Fix draw elements in menu anchors hanging around after the menu was closed

diff --git a/Engine/Source/Runtime/SlateCore/Private/Rendering/DrawElements.cpp b/Engine/Source/Runtime/SlateCore/Private/Rendering/DrawElements.cpp
--- a/Engine/Source/Runtime/SlateCore/Private/Rendering/DrawElements.cpp
+++ b/Engine/Source/Runtime/SlateCore/Private/Rendering/DrawElements.cpp
@@ -734,9 +734,8 @@ void FSlateWindowElementList::PushPaintingWidget(const SWidget& CurrentWidget, i
 	if (CurrentCachedElementData)
 	{
 		const FWidgetDrawElementState& PreviousState = WidgetDrawStack.Num() ? WidgetDrawStack.Top() : FWidgetDrawElementState(nullptr, false, nullptr);
-		const bool bParentVolatile = WidgetDrawStack.Num() ? PreviousState.bIsVolatile : false;
 
-		WidgetDrawStack.Emplace(CurrentCacheNode, bParentVolatile || CurrentWidget.IsVolatile(), &CurrentWidget);
+		WidgetDrawStack.Emplace(CurrentCacheNode, CurrentWidget.IsVolatileIndirectly() || CurrentWidget.IsVolatile(), &CurrentWidget);
 
 		// When a widget is pushed reset its draw elements.  They are being recached or possibly going away
 		if (CurrentCacheNode != nullptr)
