Adding an option for ReflectionMetadata to print the path to a given widget, makes it easier to understand the hierarchy this widget is part of.

Matt.Kuhlenschmidt

diff --git a/Engine/Source/Runtime/SlateCore/Private/Types/ReflectionMetadata.cpp b/Engine/Source/Runtime/SlateCore/Private/Types/ReflectionMetadata.cpp
--- a/Engine/Source/Runtime/SlateCore/Private/Types/ReflectionMetadata.cpp
+++ b/Engine/Source/Runtime/SlateCore/Private/Types/ReflectionMetadata.cpp
@@ -1,6 +1,42 @@
 // Copyright 1998-2019 Epic Games, Inc. All Rights Reserved.
 
 #include "Types/ReflectionMetadata.h"
+#include "Widgets/SWidget.h"
+
+FString FReflectionMetaData::GetWidgetPath(const SWidget* InWidget)
+{
+	if (!InWidget)
+	{
+		return TEXT("None");
+	}
+
+	return GetWidgetPath(*InWidget);
+}
+
+FString FReflectionMetaData::GetWidgetPath(const SWidget& InWidget)
+{
+	FString WidgetPath;
+
+	const SWidget* CurrentWidget = &InWidget;
+	while (CurrentWidget)
+	{
+		WidgetPath.InsertAt(0, TEXT("/"));
+
+		TSharedPtr<FReflectionMetaData> MetaData = CurrentWidget->GetMetaData<FReflectionMetaData>();
+		if (MetaData.IsValid())
+		{
+			WidgetPath.InsertAt(0, MetaData->Name.ToString());
+		}
+		else
+		{
+			WidgetPath.InsertAt(0, CurrentWidget->GetReadableLocation());
+		}
+
+		CurrentWidget = CurrentWidget->GetParentWidget().Get();
+	}
+
+	return WidgetPath;
+}
 
 FString FReflectionMetaData::GetWidgetDebugInfo(const SWidget* InWidget)
 {
diff --git a/Engine/Source/Runtime/SlateCore/Public/Types/ReflectionMetadata.h b/Engine/Source/Runtime/SlateCore/Public/Types/ReflectionMetadata.h
--- a/Engine/Source/Runtime/SlateCore/Public/Types/ReflectionMetadata.h
+++ b/Engine/Source/Runtime/SlateCore/Public/Types/ReflectionMetadata.h
@@ -6,7 +6,8 @@
 #include "UObject/Object.h"
 #include "UObject/WeakObjectPtr.h"
 #include "Types/ISlateMetaData.h"
-#include "Widgets/SWidget.h"
+
+class SWidget;
 
 /**
  * Reflection meta-data that can be used by the widget reflector to determine
@@ -40,6 +41,9 @@ public:
 
 public:
 
+	static FString GetWidgetPath(const SWidget* InWidget);
+	static FString GetWidgetPath(const SWidget& InWidget);
+
 	static FString GetWidgetDebugInfo(const SWidget* InWidget);
 	static FString GetWidgetDebugInfo(const SWidget& InWidget);
 
