Fix layout scale not being propagated to newly added or newly visible widgets

diff --git a/Engine/Source/Runtime/Engine/Private/Slate/SGameLayerManager.cpp b/Engine/Source/Runtime/Engine/Private/Slate/SGameLayerManager.cpp
--- a/Engine/Source/Runtime/Engine/Private/Slate/SGameLayerManager.cpp
+++ b/Engine/Source/Runtime/Engine/Private/Slate/SGameLayerManager.cpp
@@ -25,6 +25,7 @@ SGameLayerManager::SGameLayerManager()
 :	DefaultWindowTitleBarHeight(64.0f)
 ,	bIsGameUsingBorderlessWindow(false)
 ,	bUseScaledDPI(false)
+,	CachedInverseDPIScale(0.0f)
 {
 }
 
@@ -435,7 +436,14 @@ void SGameLayerManager::AddOrUpdatePlayerLayers(const FGeometry& AllottedGeometr
 	ESplitScreenType::Type SplitType = ViewportClient->GetCurrentSplitscreenConfiguration();
 	TArray<struct FSplitscreenData>& SplitInfo = ViewportClient->SplitscreenInfo;
 
-	float InverseDPIScale = ViewportClient->Viewport ? 1.0f / GetGameViewportDPIScale() : 1.0f;
+	const float InverseDPIScale = ViewportClient->Viewport ? 1.0f / GetGameViewportDPIScale() : 1.0f;
+
+	if (CachedInverseDPIScale != InverseDPIScale)
+	{
+		InvalidatePrepass();
+		Invalidate(EInvalidateWidget::Layout);
+		CachedInverseDPIScale = InverseDPIScale;
+	}
 
 	// Add and Update Player Layers
 	for ( int32 PlayerIndex = 0; PlayerIndex < GamePlayers.Num(); PlayerIndex++ )
diff --git a/Engine/Source/Runtime/Engine/Public/Slate/SGameLayerManager.h b/Engine/Source/Runtime/Engine/Public/Slate/SGameLayerManager.h
--- a/Engine/Source/Runtime/Engine/Public/Slate/SGameLayerManager.h
+++ b/Engine/Source/Runtime/Engine/Public/Slate/SGameLayerManager.h
@@ -217,4 +217,5 @@ private:
 	bool bIsGameUsingBorderlessWindow;
 	FIntPoint ScaledDPIViewportReference;
 	bool bUseScaledDPI;
+	float CachedInverseDPIScale;
 };
diff --git a/Engine/Source/Runtime/SlateCore/Private/FastUpdate/SlateInvalidationRoot.cpp b/Engine/Source/Runtime/SlateCore/Private/FastUpdate/SlateInvalidationRoot.cpp
--- a/Engine/Source/Runtime/SlateCore/Private/FastUpdate/SlateInvalidationRoot.cpp
+++ b/Engine/Source/Runtime/SlateCore/Private/FastUpdate/SlateInvalidationRoot.cpp
@@ -581,5 +581,7 @@ void FSlateInvalidationRoot::OnInvalidateAllWidgets()
 {
 	InvalidateChildOrder();
 
+	InvalidationRootWidget->InvalidatePrepass();
+
 	bNeedsSlowPath = true;
 }
diff --git a/Engine/Source/Runtime/SlateCore/Private/FastUpdate/WidgetProxy.cpp b/Engine/Source/Runtime/SlateCore/Private/FastUpdate/WidgetProxy.cpp
--- a/Engine/Source/Runtime/SlateCore/Private/FastUpdate/WidgetProxy.cpp
+++ b/Engine/Source/Runtime/SlateCore/Private/FastUpdate/WidgetProxy.cpp
@@ -61,75 +61,91 @@ int32 FWidgetProxy::Update(const FPaintArgs& PaintArgs, int32 MyIndex, FSlateWin
 bool FWidgetProxy::ProcessInvalidation(FWidgetUpdateList& UpdateList, TArray<FWidgetProxy>& FastWidgetPathList, FSlateInvalidationRoot& Root)
 {
 	bool bWidgetNeedsRepaint = false;
-	if (EnumHasAnyFlags(CurrentInvalidateReason, EInvalidateWidget::RenderTransform))
+	if (!bInvisibleDueToParentOrSelfVisibility && ParentIndex != INDEX_NONE && !Widget->PrepassLayoutScaleMultiplier.IsSet())
 	{
-		UpdateFlags |= EWidgetUpdateFlags::NeedsRepaint;
-
+		// If this widget has never been prepassed make sure the parent prepasses it to set the correct multiplier
 		FWidgetProxy& ParentProxy = FastWidgetPathList[ParentIndex];
-		ParentProxy.CurrentInvalidateReason |= EInvalidateWidget::Layout;
-		UpdateList.Push(ParentProxy);
-
+		if (ParentProxy.Widget)
+		{
+			ParentProxy.Widget->InvalidatePrepass();
+			ParentProxy.CurrentInvalidateReason |= EInvalidateWidget::Layout;
+			//UpdateFlags |= EWidgetUpdateFlags::NeedsRepaint;
+			UpdateList.Push(ParentProxy);
+		}
 		bWidgetNeedsRepaint = true;
 	}
-	else if (EnumHasAnyFlags(CurrentInvalidateReason, EInvalidateWidget::Layout|EInvalidateWidget::Visibility|EInvalidateWidget::ChildOrder))
+	else
 	{
-		// When layout changes compute a new desired size for this widget
-		FVector2D CurrentDesiredSize = Widget->GetDesiredSize();
-		FVector2D NewDesiredSize = FVector2D::ZeroVector;
-		if (Visibility != EVisibility::Collapsed)
+		if (EnumHasAnyFlags(CurrentInvalidateReason, EInvalidateWidget::RenderTransform))
 		{
-			if (Widget->NeedsPrepass())
-			{
-				Widget->SlatePrepass(Widget->PrepassLayoutScaleMultiplier);
-			}
-			else
-			{
-				Widget->CacheDesiredSize(Widget->PrepassLayoutScaleMultiplier);
-			}
+			UpdateFlags |= EWidgetUpdateFlags::NeedsRepaint;
 
-			NewDesiredSize = Widget->GetDesiredSize();
-		}
+			FWidgetProxy& ParentProxy = FastWidgetPathList[ParentIndex];
+			ParentProxy.CurrentInvalidateReason |= EInvalidateWidget::Layout;
+			UpdateList.Push(ParentProxy);
 
-		// Note even if volatile we need to recompute desired size. We do not need to invalidate parents though if they are volatile since they will naturally redraw this widget
-		if (!Widget->IsVolatileIndirectly() && Visibility.IsVisible())
-		{
-			UpdateFlags |= EWidgetUpdateFlags::NeedsRepaint;
+			bWidgetNeedsRepaint = true;
 		}
-
-		// If the desired size changed, invalidate the parent if it is visible
-		if (NewDesiredSize != CurrentDesiredSize || EnumHasAnyFlags(CurrentInvalidateReason, EInvalidateWidget::Visibility))
+		else if (EnumHasAnyFlags(CurrentInvalidateReason, EInvalidateWidget::Layout | EInvalidateWidget::Visibility | EInvalidateWidget::ChildOrder))
 		{
-			if(ParentIndex != INDEX_NONE)
+			// When layout changes compute a new desired size for this widget
+			FVector2D CurrentDesiredSize = Widget->GetDesiredSize();
+			FVector2D NewDesiredSize = FVector2D::ZeroVector;
+			if (Visibility != EVisibility::Collapsed)
 			{
-				FWidgetProxy& ParentProxy = FastWidgetPathList[ParentIndex];
-				if (ParentIndex == 0)
+				if (Widget->NeedsPrepass())
 				{
-					// root of the invalidation panel just invalidate the whole thing
-					Root.InvalidateRoot();
+					Widget->SlatePrepass(Widget->PrepassLayoutScaleMultiplier.Get(1.0f));
 				}
-				else if (ParentProxy.Visibility.IsVisible())
+				else
 				{
-					ParentProxy.CurrentInvalidateReason |= EInvalidateWidget::Layout;
-					UpdateList.Push(ParentProxy);
+					Widget->CacheDesiredSize(Widget->PrepassLayoutScaleMultiplier.Get(1.0f));
 				}
+
+				NewDesiredSize = Widget->GetDesiredSize();
 			}
-			else if(!GSlateEnableGlobalInvalidation && Widget->IsParentValid())
+
+			// Note even if volatile we need to recompute desired size. We do not need to invalidate parents though if they are volatile since they will naturally redraw this widget
+			if (!Widget->IsVolatileIndirectly() && Visibility.IsVisible())
+			{
+				UpdateFlags |= EWidgetUpdateFlags::NeedsRepaint;
+			}
+
+			// If the desired size changed, invalidate the parent if it is visible
+			if (NewDesiredSize != CurrentDesiredSize || EnumHasAnyFlags(CurrentInvalidateReason, EInvalidateWidget::Visibility))
 			{
-				TSharedPtr<SWidget> ParentWidget = Widget->GetParentWidget();
-				if (ParentWidget->Advanced_IsInvalidationRoot())
+				if (ParentIndex != INDEX_NONE)
 				{
-					Root.InvalidateRoot();
+					FWidgetProxy& ParentProxy = FastWidgetPathList[ParentIndex];
+					if (ParentIndex == 0)
+					{
+						// root of the invalidation panel just invalidate the whole thing
+						Root.InvalidateRoot();
+					}
+					else if (ParentProxy.Visibility.IsVisible())
+					{
+						ParentProxy.CurrentInvalidateReason |= EInvalidateWidget::Layout;
+						UpdateList.Push(ParentProxy);
+					}
+				}
+				else if (!GSlateEnableGlobalInvalidation && Widget->IsParentValid())
+				{
+					TSharedPtr<SWidget> ParentWidget = Widget->GetParentWidget();
+					if (ParentWidget->Advanced_IsInvalidationRoot())
+					{
+						Root.InvalidateRoot();
+					}
 				}
 			}
+
+			bWidgetNeedsRepaint = true;
 		}
-		
-		bWidgetNeedsRepaint = true;
-	}
-	else if (EnumHasAnyFlags(CurrentInvalidateReason, EInvalidateWidget::Paint) && !Widget->IsVolatileIndirectly())
-	{
-		UpdateFlags |= EWidgetUpdateFlags::NeedsRepaint;
+		else if (EnumHasAnyFlags(CurrentInvalidateReason, EInvalidateWidget::Paint) && !Widget->IsVolatileIndirectly())
+		{
+			UpdateFlags |= EWidgetUpdateFlags::NeedsRepaint;
 
-		bWidgetNeedsRepaint = true;
+			bWidgetNeedsRepaint = true;
+		}
 	}
 
 
diff --git a/Engine/Source/Runtime/SlateCore/Private/Widgets/SWidget.cpp b/Engine/Source/Runtime/SlateCore/Private/Widgets/SWidget.cpp
--- a/Engine/Source/Runtime/SlateCore/Private/Widgets/SWidget.cpp
+++ b/Engine/Source/Runtime/SlateCore/Private/Widgets/SWidget.cpp
@@ -154,7 +154,6 @@ SWidget::SWidget()
 	// Note we are defaulting to tick for backwards compatibility
 	, UpdateFlags(EWidgetUpdateFlags::NeedsTick)
 	, DesiredSize()
-	, PrepassLayoutScaleMultiplier(1.0f)
 	, CullingBoundsExtension()
 	, EnabledState(true)
 	, Visibility(EVisibility::Visible)
@@ -1041,7 +1040,7 @@ void SWidget::Invalidate(EInvalidateWidget InvalidateReason)
 
 	const bool bVolatilityChanged = EnumHasAnyFlags(InvalidateReason, EInvalidateWidget::Volatility) ? Advanced_InvalidateVolatility() : false;
 
-	if (EnumHasAnyFlags(InvalidateReason, EInvalidateWidget::ChildOrder))
+	if (EnumHasAnyFlags(InvalidateReason, EInvalidateWidget::ChildOrder) || !PrepassLayoutScaleMultiplier.IsSet())
 	{
 		InvalidatePrepass();
 	}
@@ -1381,7 +1380,7 @@ void SWidget::Prepass_Internal(float InLayoutScaleMultiplier)
 
 	{
 		// Cache this widget's desired size.
-		CacheDesiredSize(PrepassLayoutScaleMultiplier);
+		CacheDesiredSize(PrepassLayoutScaleMultiplier.Get(1.0f));
 		bNeedsPrepass = false;
 	}
 }
diff --git a/Engine/Source/Runtime/SlateCore/Public/Widgets/SWidget.h b/Engine/Source/Runtime/SlateCore/Public/Widgets/SWidget.h
--- a/Engine/Source/Runtime/SlateCore/Public/Widgets/SWidget.h
+++ b/Engine/Source/Runtime/SlateCore/Public/Widgets/SWidget.h
@@ -1498,7 +1498,7 @@ private:
 
 protected:
 
-	float GetPrepassLayoutScaleMultiplier() const { return PrepassLayoutScaleMultiplier; }
+	float GetPrepassLayoutScaleMultiplier() const { return PrepassLayoutScaleMultiplier.Get(1.0f); }
 public:
 	/**
 	 * Registers an "active timer" delegate that will execute at some regular interval. TickFunction will not be called until the specified interval has elapsed once.
@@ -1642,7 +1642,7 @@ private:
 
 protected:
 
-	float PrepassLayoutScaleMultiplier;
+	TOptional<float> PrepassLayoutScaleMultiplier;
 	/**
 	* Can be used to enlarge the culling bounds of this widget (pre-intersection), this can be useful if you've got
 	* children that you know are using rendering transforms to render outside their standard bounds, if that happens
