Fix vs2019 warnings

diff --git a/Engine/Source/Runtime/Slate/Private/Framework/Application/SlateApplication.cpp b/Engine/Source/Runtime/Slate/Private/Framework/Application/SlateApplication.cpp
--- a/Engine/Source/Runtime/Slate/Private/Framework/Application/SlateApplication.cpp
+++ b/Engine/Source/Runtime/Slate/Private/Framework/Application/SlateApplication.cpp
@@ -1096,7 +1096,7 @@ FSlateApplication::FSlateApplication()
 	{
 		CVarGlobalInvalidation->SetOnChangedCallback(FConsoleVariableDelegate::CreateLambda([this](IConsoleVariable* Variable)
 		{
-			OnGlobalInvalidationToggledEvent.Broadcast(GSlateEnableGlobalInvalidation);
+			OnGlobalInvalidationToggledEvent.Broadcast(GSlateEnableGlobalInvalidation != 0 ? true : false);
 		}));
 	}
 }
diff --git a/Engine/Source/Runtime/SlateCore/Private/Application/SlateApplicationBase.cpp b/Engine/Source/Runtime/SlateCore/Private/Application/SlateApplicationBase.cpp
--- a/Engine/Source/Runtime/SlateCore/Private/Application/SlateApplicationBase.cpp
+++ b/Engine/Source/Runtime/SlateCore/Private/Application/SlateApplicationBase.cpp
@@ -157,7 +157,7 @@ void FSlateApplicationBase::ToggleGlobalInvalidation(bool bIsGlobalInvalidationE
 	if ((!!GSlateEnableGlobalInvalidation) != bIsGlobalInvalidationEnabled)
 	{
 		GSlateEnableGlobalInvalidation = bIsGlobalInvalidationEnabled ? 1 : 0;
-		OnGlobalInvalidationToggledEvent.Broadcast(GSlateEnableGlobalInvalidation);
+		OnGlobalInvalidationToggledEvent.Broadcast(bIsGlobalInvalidationEnabled);
 	}
 }
 
