Fix widgets which are invalidated for render transform not considering layout transform and desired size changes.  Calling ComputeDesiredSize is critical for some widgets to function

diff --git a/Engine/Source/Runtime/SlateCore/Private/FastUpdate/WidgetProxy.cpp b/Engine/Source/Runtime/SlateCore/Private/FastUpdate/WidgetProxy.cpp
--- a/Engine/Source/Runtime/SlateCore/Private/FastUpdate/WidgetProxy.cpp
+++ b/Engine/Source/Runtime/SlateCore/Private/FastUpdate/WidgetProxy.cpp
@@ -74,80 +74,66 @@ bool FWidgetProxy::ProcessInvalidation(FWidgetUpdateList& UpdateList, TArray<FWi
 		}
 		bWidgetNeedsRepaint = true;
 	}
-	else
+	else if (EnumHasAnyFlags(CurrentInvalidateReason, EInvalidateWidget::RenderTransform | EInvalidateWidget::Layout | EInvalidateWidget::Visibility | EInvalidateWidget::ChildOrder))
 	{
-		if (EnumHasAnyFlags(CurrentInvalidateReason, EInvalidateWidget::RenderTransform))
+		// When layout changes compute a new desired size for this widget
+		FVector2D CurrentDesiredSize = Widget->GetDesiredSize();
+		FVector2D NewDesiredSize = FVector2D::ZeroVector;
+		if (Visibility != EVisibility::Collapsed)
 		{
-			UpdateFlags |= EWidgetUpdateFlags::NeedsRepaint;
+			if (Widget->NeedsPrepass())
+			{
+				Widget->SlatePrepass(Widget->PrepassLayoutScaleMultiplier.Get(1.0f));
+			}
+			else
+			{
+				Widget->CacheDesiredSize(Widget->PrepassLayoutScaleMultiplier.Get(1.0f));
+			}
 
-			FWidgetProxy& ParentProxy = FastWidgetPathList[ParentIndex];
-			ParentProxy.CurrentInvalidateReason |= EInvalidateWidget::Layout;
-			UpdateList.Push(ParentProxy);
+			NewDesiredSize = Widget->GetDesiredSize();
+		}
 
-			bWidgetNeedsRepaint = true;
+		// Note even if volatile we need to recompute desired size. We do not need to invalidate parents though if they are volatile since they will naturally redraw this widget
+		if (!Widget->IsVolatileIndirectly() && Visibility.IsVisible())
+		{
+			UpdateFlags |= EWidgetUpdateFlags::NeedsRepaint;
 		}
-		else if (EnumHasAnyFlags(CurrentInvalidateReason, EInvalidateWidget::Layout | EInvalidateWidget::Visibility | EInvalidateWidget::ChildOrder))
+
+		// If the desired size changed, invalidate the parent if it is visible
+		if (NewDesiredSize != CurrentDesiredSize || EnumHasAnyFlags(CurrentInvalidateReason, EInvalidateWidget::Visibility|EInvalidateWidget::RenderTransform))
 		{
-			// When layout changes compute a new desired size for this widget
-			FVector2D CurrentDesiredSize = Widget->GetDesiredSize();
-			FVector2D NewDesiredSize = FVector2D::ZeroVector;
-			if (Visibility != EVisibility::Collapsed)
+			if (ParentIndex != INDEX_NONE)
 			{
-				if (Widget->NeedsPrepass())
+				FWidgetProxy& ParentProxy = FastWidgetPathList[ParentIndex];
+				if (ParentIndex == 0)
 				{
-					Widget->SlatePrepass(Widget->PrepassLayoutScaleMultiplier.Get(1.0f));
+					// root of the invalidation panel just invalidate the whole thing
+					Root.InvalidateRoot();
 				}
-				else
+				else if (ParentProxy.Visibility.IsVisible())
 				{
-					Widget->CacheDesiredSize(Widget->PrepassLayoutScaleMultiplier.Get(1.0f));
+					ParentProxy.CurrentInvalidateReason |= EInvalidateWidget::Layout;
+					UpdateList.Push(ParentProxy);
 				}
-
-				NewDesiredSize = Widget->GetDesiredSize();
 			}
-
-			// Note even if volatile we need to recompute desired size. We do not need to invalidate parents though if they are volatile since they will naturally redraw this widget
-			if (!Widget->IsVolatileIndirectly() && Visibility.IsVisible())
+			else if (!GSlateEnableGlobalInvalidation && Widget->IsParentValid())
 			{
-				UpdateFlags |= EWidgetUpdateFlags::NeedsRepaint;
-			}
-
-			// If the desired size changed, invalidate the parent if it is visible
-			if (NewDesiredSize != CurrentDesiredSize || EnumHasAnyFlags(CurrentInvalidateReason, EInvalidateWidget::Visibility))
-			{
-				if (ParentIndex != INDEX_NONE)
+				TSharedPtr<SWidget> ParentWidget = Widget->GetParentWidget();
+				if (ParentWidget->Advanced_IsInvalidationRoot())
 				{
-					FWidgetProxy& ParentProxy = FastWidgetPathList[ParentIndex];
-					if (ParentIndex == 0)
-					{
-						// root of the invalidation panel just invalidate the whole thing
-						Root.InvalidateRoot();
-					}
-					else if (ParentProxy.Visibility.IsVisible())
-					{
-						ParentProxy.CurrentInvalidateReason |= EInvalidateWidget::Layout;
-						UpdateList.Push(ParentProxy);
-					}
-				}
-				else if (!GSlateEnableGlobalInvalidation && Widget->IsParentValid())
-				{
-					TSharedPtr<SWidget> ParentWidget = Widget->GetParentWidget();
-					if (ParentWidget->Advanced_IsInvalidationRoot())
-					{
-						Root.InvalidateRoot();
-					}
+					Root.InvalidateRoot();
 				}
 			}
-
-			bWidgetNeedsRepaint = true;
 		}
-		else if (EnumHasAnyFlags(CurrentInvalidateReason, EInvalidateWidget::Paint) && !Widget->IsVolatileIndirectly())
-		{
-			UpdateFlags |= EWidgetUpdateFlags::NeedsRepaint;
 
-			bWidgetNeedsRepaint = true;
-		}
+		bWidgetNeedsRepaint = true;
 	}
+	else if (EnumHasAnyFlags(CurrentInvalidateReason, EInvalidateWidget::Paint) && !Widget->IsVolatileIndirectly())
+	{
+		UpdateFlags |= EWidgetUpdateFlags::NeedsRepaint;
 
+		bWidgetNeedsRepaint = true;
+	}
 
 	CurrentInvalidateReason = EInvalidateWidget::None;
 
