Fix a number of slate issues related to volatility, deferred paint, and scrolling

diff --git a/Engine/Source/Runtime/Slate/Private/Widgets/Layout/SScrollBox.cpp b/Engine/Source/Runtime/Slate/Private/Widgets/Layout/SScrollBox.cpp
--- a/Engine/Source/Runtime/Slate/Private/Widgets/Layout/SScrollBox.cpp
+++ b/Engine/Source/Runtime/Slate/Private/Widgets/Layout/SScrollBox.cpp
@@ -130,6 +130,7 @@ void SScrollBox::Construct( const FArguments& InArgs )
 	NavigationScrollPadding = InArgs._NavigationScrollPadding;
 	NavigationDestination = InArgs._NavigationDestination;
 	bTouchPanningCapture = false;
+	bVolatilityAlwaysInvalidatesPrepass = true;
 
 	if (InArgs._ExternalScrollbar.IsValid())
 	{
@@ -548,6 +549,7 @@ EActiveTimerReturnType SScrollBox::UpdateInertialScroll(double InCurrentTime, fl
 	{
 		bIsScrolling = false;
 		bIsScrollingActiveTimerRegistered = false;
+		Invalidate(EInvalidateWidget::LayoutAndVolatility);
 		UpdateInertialScrollHandle.Reset();
 	}
 
@@ -1104,7 +1106,7 @@ void SScrollBox::EndInertialScrolling()
 {
 	bIsScrolling = false;
 	bIsScrollingActiveTimerRegistered = false;
-
+	Invalidate(EInvalidateWidget::LayoutAndVolatility);
 	if ( UpdateInertialScrollHandle.IsValid() )
 	{
 		UnRegisterActiveTimer(UpdateInertialScrollHandle.ToSharedRef());
diff --git a/Engine/Source/Runtime/SlateCore/Private/Rendering/DrawElements.cpp b/Engine/Source/Runtime/SlateCore/Private/Rendering/DrawElements.cpp
--- a/Engine/Source/Runtime/SlateCore/Private/Rendering/DrawElements.cpp
+++ b/Engine/Source/Runtime/SlateCore/Private/Rendering/DrawElements.cpp
@@ -810,7 +810,7 @@ FSlateDrawElement& FSlateWindowElementList::AddCachedElement()
 
 void FSlateWindowElementList::PushCachedElementData(FSlateCachedElementData& CachedElementData)
 {
-	const int32 Index = CachedElementDataList.Add(&CachedElementData);
+	const int32 Index = CachedElementDataList.AddUnique(&CachedElementData);
 	CachedElementDataListStack.Push(Index);
 }
 
diff --git a/Engine/Source/Runtime/SlateCore/Private/Widgets/SWidget.cpp b/Engine/Source/Runtime/SlateCore/Private/Widgets/SWidget.cpp
--- a/Engine/Source/Runtime/SlateCore/Private/Widgets/SWidget.cpp
+++ b/Engine/Source/Runtime/SlateCore/Private/Widgets/SWidget.cpp
@@ -106,6 +106,12 @@ void SWidget::UpdateWidgetProxy(int32 NewLayerId, FSlateCachedElementListNode* C
 	check(!CacheNode || CacheNode->GetValue().Widget == this);
 #endif
 
+	if (PersistentState.CachedElementListNode != nullptr && PersistentState.CachedElementListNode != CacheNode)
+	{
+	//	ensure(false);
+		PersistentState.CachedElementListNode->GetValue().GetOwningData()->RemoveCache(PersistentState.CachedElementListNode);
+	}
+
 	PersistentState.CachedElementListNode = CacheNode;
 
 	if (FastPathProxyHandle.IsValid())
@@ -116,7 +122,6 @@ void SWidget::UpdateWidgetProxy(int32 NewLayerId, FSlateCachedElementListNode* C
 
 		PersistentState.OutgoingLayerId = NewLayerId;
 
-		Advanced_InvalidateVolatility();
 		if ((IsVolatile() && !IsVolatileIndirectly()) || (Advanced_IsInvalidationRoot() && !Advanced_IsWindow()))
 		{
 			AddUpdateFlags(EWidgetUpdateFlags::NeedsVolatilePaint);
@@ -1059,7 +1064,7 @@ void SWidget::Invalidate(EInvalidateWidget InvalidateReason)
 
 			UpdateFastPathVolatility(ParentWidget.IsValid() ? ParentWidget->IsVolatile() || ParentWidget->IsVolatileIndirectly() : false);
 
-			ensure(!IsVolatile() || EnumHasAnyFlags(UpdateFlags, EWidgetUpdateFlags::NeedsVolatilePaint));
+			ensure(!IsVolatile() || IsVolatileIndirectly() || EnumHasAnyFlags(UpdateFlags, EWidgetUpdateFlags::NeedsVolatilePaint));
 		}
 
 		FastPathProxyHandle.MarkWidgetDirty(InvalidateReason);
diff --git a/Engine/Source/Runtime/SlateCore/Private/Widgets/SWindow.cpp b/Engine/Source/Runtime/SlateCore/Private/Widgets/SWindow.cpp
--- a/Engine/Source/Runtime/SlateCore/Private/Widgets/SWindow.cpp
+++ b/Engine/Source/Runtime/SlateCore/Private/Widgets/SWindow.cpp
@@ -2059,11 +2059,21 @@ int32 SWindow::PaintWindow( double CurrentTime, float DeltaTime, FSlateWindowEle
 
 	OutDrawElements.EndDeferredGroup();
 
+	if (Context.bAllowFastPathUpdate)
+	{
+		OutDrawElements.PushCachedElementData(GetCachedElements());
+	}
+
 	if (OutDrawElements.ShouldResolveDeferred())
 	{
 		Result.MaxLayerIdPainted = OutDrawElements.PaintDeferred(Result.MaxLayerIdPainted, Context.CullingRect);
 	}
 
+	if (Context.bAllowFastPathUpdate)
+	{
+		OutDrawElements.PopCachedElementData();
+	}
+
 	return Result.MaxLayerIdPainted;
 
 }
diff --git a/Engine/Source/Runtime/SlateCore/Public/Layout/Children.h b/Engine/Source/Runtime/SlateCore/Public/Layout/Children.h
--- a/Engine/Source/Runtime/SlateCore/Public/Layout/Children.h
+++ b/Engine/Source/Runtime/SlateCore/Public/Layout/Children.h
@@ -169,12 +169,14 @@ public:
 	{
 		if (WidgetPtr.IsValid())
 		{
-			WidgetPtr.Reset();
+			TSharedPtr<SWidget> Widget = WidgetPtr.Pin();
 
-			if (Owner)
+			if (Widget != SNullWidget::NullWidget)
 			{
-				Owner->Invalidate(EInvalidateWidget::ChildOrder);
+				Widget->ConditionallyDetatchParentWidget(Owner);
 			}
+
+			WidgetPtr.Reset();
 		}
 	}
 
