Fix another slate crash when loading maps

diff --git a/Engine/Source/Runtime/SlateCore/Private/FastUpdate/SlateInvalidationRoot.cpp b/Engine/Source/Runtime/SlateCore/Private/FastUpdate/SlateInvalidationRoot.cpp
--- a/Engine/Source/Runtime/SlateCore/Private/FastUpdate/SlateInvalidationRoot.cpp
+++ b/Engine/Source/Runtime/SlateCore/Private/FastUpdate/SlateInvalidationRoot.cpp
@@ -553,14 +553,14 @@ bool FSlateInvalidationRoot::ProcessInvalidation()
 	return bWidgetsNeedRepaint;
 }
 
-void FSlateInvalidationRoot::ClearAllFastPathData(bool bInvalidationRootBeingDestroyed)
+void FSlateInvalidationRoot::ClearAllFastPathData(bool bClearResourcesImmediately)
 {
 	for (const FWidgetProxy& Proxy : FastWidgetPathList)
 	{
 		if (Proxy.Widget)
 		{
 			Proxy.Widget->PersistentState.CachedElementListNode = nullptr;
-			if (bInvalidationRootBeingDestroyed)
+			if (bClearResourcesImmediately)
 			{
 				Proxy.Widget->FastPathProxyHandle = FWidgetProxyHandle();
 			}
@@ -581,7 +581,7 @@ void FSlateInvalidationRoot::OnInvalidateAllWidgets(bool bClearResourcesImmediat
 
 	if (bClearResourcesImmediately)
 	{
-		ClearAllFastPathData(false);
+		ClearAllFastPathData(true);
 	}
 	bNeedsSlowPath = true;
 }
diff --git a/Engine/Source/Runtime/SlateCore/Public/FastUpdate/SlateInvalidationRoot.h b/Engine/Source/Runtime/SlateCore/Public/FastUpdate/SlateInvalidationRoot.h
--- a/Engine/Source/Runtime/SlateCore/Public/FastUpdate/SlateInvalidationRoot.h
+++ b/Engine/Source/Runtime/SlateCore/Public/FastUpdate/SlateInvalidationRoot.h
@@ -91,7 +91,7 @@ protected:
 
 	SLATECORE_API bool ProcessInvalidation();
 
-	SLATECORE_API void ClearAllFastPathData(bool bInvalidationRootBeingDestroyed);
+	SLATECORE_API void ClearAllFastPathData(bool bClearResourcesImmediately);
 
 private:
 
