Fix invalidating volatility not invalidating paint.  When volality changes we must always redraw to cache elements or clear the cache

diff --git a/Engine/Source/Runtime/SlateCore/Private/Widgets/SWidget.cpp b/Engine/Source/Runtime/SlateCore/Private/Widgets/SWidget.cpp
--- a/Engine/Source/Runtime/SlateCore/Private/Widgets/SWidget.cpp
+++ b/Engine/Source/Runtime/SlateCore/Private/Widgets/SWidget.cpp
@@ -1038,6 +1038,12 @@ void SWidget::Invalidate(EInvalidateWidget InvalidateReason)
 	SCOPED_NAMED_EVENT_TEXT("SWidget::Invalidate", FColor::Orange);
 	const bool bWasVolatile = IsVolatileIndirectly() || IsVolatile();
 
+	// Backwards compatibility fix:  Its no longer valid to just invalidate volatility since we need to repaint to cache elements if a widget becoems non-volatile. So after volatility changes force repaint
+	if (InvalidateReason == EInvalidateWidget::Volatility)
+	{
+		InvalidateReason = EInvalidateWidget::PaintAndVolatility;
+	}
+
 	const bool bVolatilityChanged = EnumHasAnyFlags(InvalidateReason, EInvalidateWidget::Volatility) ? Advanced_InvalidateVolatility() : false;
 
 	if (EnumHasAnyFlags(InvalidateReason, EInvalidateWidget::ChildOrder) || !PrepassLayoutScaleMultiplier.IsSet())
diff --git a/Engine/Source/Runtime/UMG/Private/Slate/SObjectWidget.cpp b/Engine/Source/Runtime/UMG/Private/Slate/SObjectWidget.cpp
--- a/Engine/Source/Runtime/UMG/Private/Slate/SObjectWidget.cpp
+++ b/Engine/Source/Runtime/UMG/Private/Slate/SObjectWidget.cpp
@@ -130,17 +130,6 @@ int32 SObjectWidget::OnPaint(const FPaintArgs& Args, const FGeometry& AllottedGe
 	return MaxLayer;
 }
 
-bool SObjectWidget::ComputeVolatility() const
-{
-	// @todo eliminate this
-	if (!GSlateEnableGlobalInvalidation && CanRouteEvent())
-	{
-		return WidgetObject->IsPlayingAnimation() || SCompoundWidget::ComputeVolatility();
-	}
-
-	return SCompoundWidget::ComputeVolatility();
-}
-
 FVector2D SObjectWidget::ComputeDesiredSize(float LayoutScaleMultiplier) const
 {
 	const FVector2D BaseDesiredSize = SCompoundWidget::ComputeDesiredSize(LayoutScaleMultiplier);
diff --git a/Engine/Source/Runtime/UMG/Public/Slate/SObjectWidget.h b/Engine/Source/Runtime/UMG/Public/Slate/SObjectWidget.h
--- a/Engine/Source/Runtime/UMG/Public/Slate/SObjectWidget.h
+++ b/Engine/Source/Runtime/UMG/Public/Slate/SObjectWidget.h
@@ -56,7 +56,6 @@ class UMG_API SObjectWidget : public SCompoundWidget, public FGCObject
 	virtual void Tick( const FGeometry& AllottedGeometry, const double InCurrentTime, const float InDeltaTime ) override;
 	virtual int32 OnPaint(const FPaintArgs& Args, const FGeometry& AllottedGeometry, const FSlateRect& MyCullingRect, FSlateWindowElementList& OutDrawElements, int32 LayerId, const FWidgetStyle& InWidgetStyle, bool bParentEnabled) const override;
 
-	virtual bool ComputeVolatility() const override;
 	virtual FVector2D ComputeDesiredSize(float LayoutScaleMultiplier) const override;
 
 	virtual bool IsInteractable() const override;
