Slate - Slots now have a similar function for setting attributes and invalidating the widget it owns when slot properties are modified that affect layout or painting, most all of them affect layout, will need to increase coverage here, right now pretty low coverage for invalidating slots when their properties are changed.

Matt.Kuhlenschmidt

diff --git a/Engine/Source/Runtime/SlateCore/Private/SlotBase.cpp b/Engine/Source/Runtime/SlateCore/Private/SlotBase.cpp
--- a/Engine/Source/Runtime/SlateCore/Private/SlotBase.cpp
+++ b/Engine/Source/Runtime/SlateCore/Private/SlotBase.cpp
@@ -37,6 +37,11 @@ const TSharedPtr<SWidget> FSlotBase::DetachWidget()
 	}
 }
 
+void FSlotBase::Invalidate(EInvalidateWidgetReason InvalidateReason)
+{
+	Widget->Invalidate(InvalidateReason);
+}
+
 void FSlotBase::DetatchParentFromContent()
 {
 	if (Widget != SNullWidget::NullWidget)
diff --git a/Engine/Source/Runtime/SlateCore/Public/SlotBase.h b/Engine/Source/Runtime/SlateCore/Public/SlotBase.h
--- a/Engine/Source/Runtime/SlateCore/Public/SlotBase.h
+++ b/Engine/Source/Runtime/SlateCore/Public/SlotBase.h
@@ -3,6 +3,7 @@
 #pragma once
 
 #include "CoreMinimal.h"
+#include "Widgets/InvalidateWidgetReason.h"
 
 class SWidget;
 
@@ -54,6 +55,35 @@ public:
 	 */
 	const TSharedPtr<SWidget> DetachWidget();
 
+protected:
+	void Invalidate(EInvalidateWidgetReason InvalidateReason);
+
+	/**
+	 * Performs the attribute assignment and invalidates the widget minimally based on what actually changed.  So if the boundness of the attribute didn't change
+	 * volatility won't need to be recalculated.  Returns true if the value changed.
+	 */
+	template<typename TargetValueType, typename SourceValueType>
+	bool SetAttribute(TAttribute<TargetValueType>& TargetValue, const TAttribute<SourceValueType>& SourceValue, EInvalidateWidgetReason BaseInvalidationReason)
+	{
+		if (!TargetValue.IdenticalTo(SourceValue))
+		{
+			const bool bWasBound = TargetValue.IsBound();
+			const bool bBoundnessChanged = bWasBound != SourceValue.IsBound();
+			TargetValue = SourceValue;
+
+			EInvalidateWidgetReason InvalidateReason = BaseInvalidationReason;
+			if (bBoundnessChanged)
+			{
+				InvalidateReason |= EInvalidateWidgetReason::Volatility;
+			}
+
+			Invalidate(InvalidateReason);
+			return true;
+		}
+
+		return false;
+	}
+
 protected:
 	/** The parent and owner of the slot. */
 	SWidget* RawParentPtr;
diff --git a/Engine/Source/Runtime/SlateCore/Public/Widgets/InvalidateWidgetReason.h b/Engine/Source/Runtime/SlateCore/Public/Widgets/InvalidateWidgetReason.h
new file mode 100644
--- /dev/null
+++ b/Engine/Source/Runtime/SlateCore/Public/Widgets/InvalidateWidgetReason.h
@@ -0,0 +1,65 @@
+// Copyright 1998-2019 Epic Games, Inc. All Rights Reserved.
+
+#pragma once
+
+#include "CoreMinimal.h"
+
+/**
+ * The different types of invalidation that are possible for a widget.
+ */
+enum class EInvalidateWidgetReason : uint8
+{
+	None = 0,
+
+	/**
+	 * Use Layout invalidation if your widget needs to change desired size.  This is an expensive invalidation so do not use if all you need to do is redraw a widget
+	 */
+	Layout = 1 << 0,
+
+	/**
+	 * Use when the painting of widget has been altered, but nothing affecting sizing.
+	 */
+	Paint = 1 << 1,
+
+	/**
+	 * Use if just the volatility of the widget has been adjusted.
+	 */
+	Volatility = 1 << 2,
+
+	/**
+	 * A child was added or removed.   (this implies layout)
+	 */
+	ChildOrder = 1 << 3,
+
+	/** A Widgets render transform changed */
+	RenderTransform = 1 << 4,
+
+	/**
+	 * Changing visibility (this implies layout)
+	 */
+	Visibility = 1 << 5,
+
+	/**
+	 * Use Paint invalidation if you're changing a normal property involving painting or sizing.
+	 * Additionally if the property that was changed affects Volatility in anyway, it's important
+	 * that you invalidate volatility so that it can be recalculated and cached.
+	 */
+	PaintAndVolatility = Paint | Volatility,
+	/**
+	 * Use Layout invalidation if you're changing a normal property involving painting or sizing.
+	 * Additionally if the property that was changed affects Volatility in anyway, it's important
+	 * that you invalidate volatility so that it can be recalculated and cached.
+	 */
+	LayoutAndVolatility = Layout | Volatility,
+
+
+	/**
+	 * Do not use this ever unless you know what you are doing
+	 */
+	All UE_DEPRECATED(4.22, "EInvalidateWidget::All has been deprecated.  You probably wanted EInvalidateWidget::Layout but if you need more than that then use bitwise or to combine them") = 0xff
+};
+
+ENUM_CLASS_FLAGS(EInvalidateWidgetReason)
+
+// This typedefed because EInvalidateWidget will be deprecated soon
+typedef EInvalidateWidgetReason EInvalidateWidget;
