Fix render batch index not resetting to a default value after batching.  Without this if nothing is batched the next frame there will be a crash due to the first batch index being invalid

diff --git a/Engine/Source/Runtime/SlateCore/Private/Rendering/ElementBatcher.cpp b/Engine/Source/Runtime/SlateCore/Private/Rendering/ElementBatcher.cpp
--- a/Engine/Source/Runtime/SlateCore/Private/Rendering/ElementBatcher.cpp
+++ b/Engine/Source/Runtime/SlateCore/Private/Rendering/ElementBatcher.cpp
@@ -69,8 +69,7 @@ void FSlateBatchData::ResetData()
 	FinalIndexData.Reset();
 	FinalVertexData.Reset();
 
-	// note: LayerToElementBatches is not reset here as the same layers are 
-	// more than likely reused and we can save memory allocations by not resetting the map every frame
+	FirstRenderBatchIndex = INDEX_NONE;
 
 	NumBatches = 0;
 	NumLayers = 0;
