#UE4 Nonunity fix

diff --git a/Engine/Source/Runtime/SlateCore/Public/SlotBase.h b/Engine/Source/Runtime/SlateCore/Public/SlotBase.h
--- a/Engine/Source/Runtime/SlateCore/Public/SlotBase.h
+++ b/Engine/Source/Runtime/SlateCore/Public/SlotBase.h
@@ -3,6 +3,7 @@
 #pragma once
 
 #include "CoreMinimal.h"
+#include "Misc/Attribute.h"
 #include "Widgets/InvalidateWidgetReason.h"
 
 class SWidget;
