When SWindows move they now update the tick geometry of cached widgets, fixes a lot of cases of people referencing desktop relative geometry that had gone stale.

diff --git a/Engine/Source/Runtime/SlateCore/Private/FastUpdate/SlateInvalidationRoot.cpp b/Engine/Source/Runtime/SlateCore/Private/FastUpdate/SlateInvalidationRoot.cpp
--- a/Engine/Source/Runtime/SlateCore/Private/FastUpdate/SlateInvalidationRoot.cpp
+++ b/Engine/Source/Runtime/SlateCore/Private/FastUpdate/SlateInvalidationRoot.cpp
@@ -33,6 +33,7 @@ FSlateInvalidationRoot::FSlateInvalidationRoot()
 	, FastPathGenerationNumber(INDEX_NONE)
 	, bChildOrderInvalidated(false)
 	, bNeedsSlowPath(true)
+	, bNeedScreenPositionShift(false)
 {
 	FSlateApplicationBase::Get().OnInvalidateAllWidgets().AddRaw(this, &FSlateInvalidationRoot::OnInvalidateAllWidgets);
 
@@ -101,6 +102,11 @@ void FSlateInvalidationRoot::InvalidateChildOrder()
 	}
 }
 
+void FSlateInvalidationRoot::InvalidateScreenPosition()
+{
+	bNeedScreenPositionShift = true;
+}
+
 int32 RecursiveFindParentWithChildOrderChange(const TArray<FWidgetProxy>& FastWidgetPathList, const FWidgetProxy& Proxy)
 {
 	if (Proxy.bChildOrderInvalid)
@@ -163,18 +169,22 @@ FSlateInvalidationResult FSlateInvalidationRoot::PaintInvalidationRoot(const FSl
 
 	SWidget* RootWidget = InvalidationRootWidget->Advanced_IsWindow() ? InvalidationRootWidget : &(*InvalidationRootWidget->GetAllChildren()->GetChildAt(0));
 
-	//AssignFastPathIndices(RootWidget, true);
+	if (bNeedScreenPositionShift)
+	{
+		SCOPED_NAMED_EVENT(Slate_InvalidateScreenPosition, FColor::Red);
+		AdjustWidgetsDesktopGeometry(Context.PaintArgs->GetWindowToDesktopTransform());
+		bNeedScreenPositionShift = false;
+	}
 
 	EFlowDirection NewFlowDirection = GSlateFlowDirection;
 	if (RootWidget->GetFlowDirectionPreference() == EFlowDirectionPreference::Inherit)
 	{
 		NewFlowDirection = GSlateFlowDirectionShouldFollowCultureByDefault ? FLayoutLocalization::GetLocalizedLayoutDirection() : EFlowDirection::LeftToRight;
 	}
-
 	TGuardValue<EFlowDirection> FlowGuard(GSlateFlowDirection, NewFlowDirection);
 	if (!Context.bAllowFastPathUpdate || bNeedsSlowPath || GSlateIsInInvalidationSlowPath)
 	{
-		SCOPED_NAMED_EVENT(SWidget_PaintSlowPath, FColor::Red);
+		SCOPED_NAMED_EVENT(Slate_PaintSlowPath, FColor::Red);
 
 		//CSV_EVENT(Basic, "Slate Slow Path update");
 #if WITH_SLATE_DEBUGGING
@@ -336,7 +346,7 @@ bool FSlateInvalidationRoot::PaintFastPath(const FSlateInvalidationContext& Cont
 
 	if (bNeedsSlowPath)
 	{
-		SCOPED_NAMED_EVENT(SWidget_PaintSlowPath, FColor::Red);
+		SCOPED_NAMED_EVENT(Slate_PaintSlowPath, FColor::Red);
 		CachedMaxLayerId = PaintSlowPath(Context);
 	}
 
@@ -388,6 +398,20 @@ void FSlateInvalidationRoot::BuildNewFastPathList_Recursive(FSlateInvalidationRo
 	
 }
 
+void FSlateInvalidationRoot::AdjustWidgetsDesktopGeometry(FVector2D WindowToDesktopTransform)
+{
+	FSlateLayoutTransform WindowToDesktop(WindowToDesktopTransform);
+
+	for (FWidgetProxy& Proxy : FastWidgetPathList)
+	{
+		if (Proxy.Widget)
+		{
+			Proxy.Widget->PersistentState.DesktopGeometry = Proxy.Widget->PersistentState.AllottedGeometry;
+			Proxy.Widget->PersistentState.DesktopGeometry.AppendTransform(WindowToDesktop);
+		}
+	}
+}
+
 #define VERIFY_CHILD_ORDER 0
 void FSlateInvalidationRoot::BuildFastPathList(SWidget* RootWidget)
 {
diff --git a/Engine/Source/Runtime/SlateCore/Private/Widgets/SWindow.cpp b/Engine/Source/Runtime/SlateCore/Private/Widgets/SWindow.cpp
--- a/Engine/Source/Runtime/SlateCore/Private/Widgets/SWindow.cpp
+++ b/Engine/Source/Runtime/SlateCore/Private/Widgets/SWindow.cpp
@@ -983,6 +983,9 @@ FSlateRect SWindow::GetFullScreenInfo() const
 void SWindow::SetCachedScreenPosition(FVector2D NewPosition)
 {
 	ScreenPosition = NewPosition;
+
+	InvalidateScreenPosition();
+
 	OnWindowMoved.ExecuteIfBound( SharedThis( this ) );
 }
 
diff --git a/Engine/Source/Runtime/SlateCore/Public/FastUpdate/SlateInvalidationRoot.h b/Engine/Source/Runtime/SlateCore/Public/FastUpdate/SlateInvalidationRoot.h
--- a/Engine/Source/Runtime/SlateCore/Public/FastUpdate/SlateInvalidationRoot.h
+++ b/Engine/Source/Runtime/SlateCore/Public/FastUpdate/SlateInvalidationRoot.h
@@ -68,6 +68,8 @@ public:
 	/** Call to notify that the ordering of children somewhere in the hierarchy below this root has changed and the fast path is no longer valid */
 	SLATECORE_API void InvalidateChildOrder();
 
+	SLATECORE_API void InvalidateScreenPosition();
+
 	bool NeedsSlowPath() const { return bNeedsSlowPath; }
 
 	void RemoveWidgetFromFastPath(FWidgetProxy& Proxy);
@@ -101,6 +103,8 @@ private:
 
 	void BuildFastPathList(SWidget* RootWidget);
 	void BuildNewFastPathList_Recursive(FSlateInvalidationRoot& Root, FWidgetProxy& Proxy, int32 ParentIndex, int32& NextTreeIndex, TArray<FWidgetProxy>& CurrentFastPathList, TArray<FWidgetProxy, TMemStackAllocator<>>& NewFastPathList);
+
+	void AdjustWidgetsDesktopGeometry(FVector2D WindowToDesktopTransform);
 private:
 	TArray<FWidgetProxy> FastWidgetPathList;
 	/** Index to widgets which are dirty, volatile, or need some sort of per frame update (such as a tick or timer) */
@@ -125,4 +129,5 @@ private:
 
 	bool bChildOrderInvalidated;
 	bool bNeedsSlowPath;
+	bool bNeedScreenPositionShift;
 };
\ No newline at end of file
