When WidgetProxys Update, they need to recompute the desktop geometry and not use the allocatedgeometry which is the paintspace, rather than desktop space.

diff --git a/Engine/Source/Runtime/SlateCore/Private/FastUpdate/WidgetProxy.cpp b/Engine/Source/Runtime/SlateCore/Private/FastUpdate/WidgetProxy.cpp
--- a/Engine/Source/Runtime/SlateCore/Private/FastUpdate/WidgetProxy.cpp
+++ b/Engine/Source/Runtime/SlateCore/Private/FastUpdate/WidgetProxy.cpp
@@ -50,9 +50,11 @@ int32 FWidgetProxy::Update(const FPaintArgs& PaintArgs, int32 MyIndex, FSlateWin
 
 			INC_DWORD_STAT(STAT_SlateNumTickedWidgets);
 			SCOPE_CYCLE_COUNTER(STAT_SlateTickWidgets);
-			Widget->Tick(MyState.AllottedGeometry, PaintArgs.GetCurrentTime(), PaintArgs.GetDeltaTime());
-		}
 
+			FGeometry DesktopSpaceGeometry = MyState.AllottedGeometry;
+			DesktopSpaceGeometry.AppendTransform(FSlateLayoutTransform(PaintArgs.GetWindowToDesktopTransform()));
+			Widget->Tick(DesktopSpaceGeometry, PaintArgs.GetCurrentTime(), PaintArgs.GetDeltaTime());
+		}
 	}
 
 	return OutgoingLayerId;
