When slots invalidate, they need to invalidate the widget parent, rather than the actual content of the slot.

diff --git a/Engine/Source/Runtime/SlateCore/Private/SlotBase.cpp b/Engine/Source/Runtime/SlateCore/Private/SlotBase.cpp
--- a/Engine/Source/Runtime/SlateCore/Private/SlotBase.cpp
+++ b/Engine/Source/Runtime/SlateCore/Private/SlotBase.cpp
@@ -39,7 +39,12 @@ const TSharedPtr<SWidget> FSlotBase::DetachWidget()
 
 void FSlotBase::Invalidate(EInvalidateWidgetReason InvalidateReason)
 {
-	Widget->Invalidate(InvalidateReason);
+	// If a slot invalidates it needs to invalidate the parent of widget of its content.
+	const TSharedPtr<SWidget>& ParentWidget = Widget->GetParentWidget();
+	if (ParentWidget.IsValid())
+	{
+		ParentWidget->Invalidate(InvalidateReason);
+	}
 }
 
 void FSlotBase::DetatchParentFromContent()
