Fixed auto-wrapping not being recomputed and the widget invalidated if the geometry of the widget was 0 when auto wrap was initially computed

diff --git a/Engine/Source/Runtime/Slate/Private/Widgets/Text/SMultiLineEditableText.cpp b/Engine/Source/Runtime/Slate/Private/Widgets/Text/SMultiLineEditableText.cpp
--- a/Engine/Source/Runtime/Slate/Private/Widgets/Text/SMultiLineEditableText.cpp
+++ b/Engine/Source/Runtime/Slate/Private/Widgets/Text/SMultiLineEditableText.cpp
@@ -531,6 +531,13 @@ int32 SMultiLineEditableText::OnPaint( const FPaintArgs& Args, const FGeometry&
 	FWidgetStyle TextWidgetStyle = FWidgetStyle(InWidgetStyle)
 		.SetForegroundColor(ForegroundColor);
 
+	const bool bAutoWrap = EditableTextLayout->GetAutoWrapText();
+	if (bAutoWrap && EditableTextLayout->GetComputedWrappingWidth() == 0)
+	{
+		QUICK_SCOPE_CYCLE_COUNTER(STAT_Slate_AutoWrapRecompute)
+		const_cast<SMultiLineEditableText*>(this)->CacheDesiredSize(GetPrepassLayoutScaleMultiplier());
+	}
+
 	LayerId = EditableTextLayout->OnPaint(Args, AllottedGeometry, MyCullingRect, OutDrawElements, LayerId, TextWidgetStyle, ShouldBeEnabled(bParentEnabled));
 
 	if (bIsSoftwareCursor)
diff --git a/Engine/Source/Runtime/Slate/Private/Widgets/Text/SlateEditableTextLayout.cpp b/Engine/Source/Runtime/Slate/Private/Widgets/Text/SlateEditableTextLayout.cpp
--- a/Engine/Source/Runtime/Slate/Private/Widgets/Text/SlateEditableTextLayout.cpp
+++ b/Engine/Source/Runtime/Slate/Private/Widgets/Text/SlateEditableTextLayout.cpp
@@ -663,6 +663,17 @@ FVector2D FSlateEditableTextLayout::GetScrollOffset() const
 	return ScrollOffset;
 }
 
+
+float FSlateEditableTextLayout::GetComputedWrappingWidth() const
+{
+	return TextLayout->GetWrappingWidth();
+}
+
+bool FSlateEditableTextLayout::GetAutoWrapText() const
+{
+	return AutoWrapText.Get();
+}
+
 bool FSlateEditableTextLayout::HandleFocusReceived(const FFocusEvent& InFocusEvent)
 {
 	if (ActiveContextMenu.IsValid())
diff --git a/Engine/Source/Runtime/Slate/Public/Widgets/Text/SlateEditableTextLayout.h b/Engine/Source/Runtime/Slate/Public/Widgets/Text/SlateEditableTextLayout.h
--- a/Engine/Source/Runtime/Slate/Public/Widgets/Text/SlateEditableTextLayout.h
+++ b/Engine/Source/Runtime/Slate/Public/Widgets/Text/SlateEditableTextLayout.h
@@ -138,6 +138,12 @@ public:
 	/** Get the absolute scroll offset value */
 	FVector2D GetScrollOffset() const;
 
+	/** Returns the computed wrap location for this layout */
+	float GetComputedWrappingWidth() const;
+
+	/** Returns whether or not we are auto wrapping text */
+	bool GetAutoWrapText() const;
+
 	/** Called when our parent widget receives focus */
 	bool HandleFocusReceived(const FFocusEvent& InFocusEvent);
 
diff --git a/Engine/Source/Runtime/SlateCore/Private/FastUpdate/WidgetProxy.cpp b/Engine/Source/Runtime/SlateCore/Private/FastUpdate/WidgetProxy.cpp
--- a/Engine/Source/Runtime/SlateCore/Private/FastUpdate/WidgetProxy.cpp
+++ b/Engine/Source/Runtime/SlateCore/Private/FastUpdate/WidgetProxy.cpp
@@ -241,11 +241,6 @@ void FWidgetProxyHandle::MarkWidgetDirty(EInvalidateWidget InvalidateReason)
 
 	if (EnumHasAnyFlags(InvalidateReason, EInvalidateWidget::ChildOrder))
 	{
-		/*
-				CSV_EVENT_GLOBAL(TEXT("Slow Path Needed"));
-		#if WITH_SLATE_DEBUGGING
-				UE_LOG(LogSlate, Log, TEXT("Slow Widget Path Needed: %s %s"), *Proxy.Widget->ToString(), *Proxy.Widget->GetTag().ToString());
-		#endif*/
 		Proxy.bChildOrderInvalid = true;
 		InvalidationRoot->InvalidateChildOrder();
 	}
@@ -254,7 +249,7 @@ void FWidgetProxyHandle::MarkWidgetDirty(EInvalidateWidget InvalidateReason)
 	{
 		InvalidationRoot->WidgetsNeedingUpdate.Push(Proxy);
 	}
-#if WITH_SLATE_DEBUGGING
+#if 0
 	else
 	{
 		ensure(InvalidationRoot->WidgetsNeedingUpdate.Contains(Proxy));
