SlateGI - No longer computing the desktop geometry during the fast path update, we can now safely use the cached desktop geometry.

diff --git a/Engine/Source/Runtime/SlateCore/Private/FastUpdate/WidgetProxy.cpp b/Engine/Source/Runtime/SlateCore/Private/FastUpdate/WidgetProxy.cpp
--- a/Engine/Source/Runtime/SlateCore/Private/FastUpdate/WidgetProxy.cpp
+++ b/Engine/Source/Runtime/SlateCore/Private/FastUpdate/WidgetProxy.cpp
@@ -51,9 +51,7 @@ int32 FWidgetProxy::Update(const FPaintArgs& PaintArgs, int32 MyIndex, FSlateWin
 			INC_DWORD_STAT(STAT_SlateNumTickedWidgets);
 			SCOPE_CYCLE_COUNTER(STAT_SlateTickWidgets);
 
-			FGeometry DesktopSpaceGeometry = MyState.AllottedGeometry;
-			DesktopSpaceGeometry.AppendTransform(FSlateLayoutTransform(PaintArgs.GetWindowToDesktopTransform()));
-			Widget->Tick(DesktopSpaceGeometry, PaintArgs.GetCurrentTime(), PaintArgs.GetDeltaTime());
+			Widget->Tick(MyState.DesktopGeometry, PaintArgs.GetCurrentTime(), PaintArgs.GetDeltaTime());
 		}
 	}
 
@@ -241,6 +239,11 @@ void FWidgetProxyHandle::MarkWidgetDirty(EInvalidateWidget InvalidateReason)
 
 	if (EnumHasAnyFlags(InvalidateReason, EInvalidateWidget::ChildOrder))
 	{
+		/*
+				CSV_EVENT_GLOBAL(TEXT("Slow Path Needed"));
+		#if WITH_SLATE_DEBUGGING
+				UE_LOG(LogSlate, Log, TEXT("Slow Widget Path Needed: %s %s"), *Proxy.Widget->ToString(), *Proxy.Widget->GetTag().ToString());
+		#endif*/
 		Proxy.bChildOrderInvalid = true;
 		InvalidationRoot->InvalidateChildOrder();
 	}
@@ -249,7 +252,7 @@ void FWidgetProxyHandle::MarkWidgetDirty(EInvalidateWidget InvalidateReason)
 	{
 		InvalidationRoot->WidgetsNeedingUpdate.Push(Proxy);
 	}
-#if 0
+#if WITH_SLATE_DEBUGGING
 	else
 	{
 		ensure(InvalidationRoot->WidgetsNeedingUpdate.Contains(Proxy));
