SlateGI - Providing another FChildren container that can support multiple FChildren linked together.

diff --git a/Engine/Source/Runtime/SlateCore/Public/Layout/Children.h b/Engine/Source/Runtime/SlateCore/Public/Layout/Children.h
--- a/Engine/Source/Runtime/SlateCore/Public/Layout/Children.h
+++ b/Engine/Source/Runtime/SlateCore/Public/Layout/Children.h
@@ -36,6 +36,7 @@ public:
 
 protected:
 	friend class SWidget;
+	friend class FCombinedChildren;
 	/** @return the const reference to the slot at the specified Index */
 	virtual const FSlotBase& GetSlotAt(int32 ChildIndex) const = 0;
 
@@ -47,6 +48,69 @@ protected:
 };
 
 
+/**
+ * Occasionally you may need to keep multiple descrete sets of children with differing slot requirements.
+ * This datastructure can be used to link multiple FChildren under a single accessor so you can always return
+ * all children from GetChildren, but internally manage them in their own child lists.
+ */
+class FCombinedChildren : public FChildren
+{
+public:
+	FCombinedChildren(SWidget* InOwner)
+		: FChildren(InOwner)
+	{
+	}
+
+	void AddChildren(FChildren& InLinkedChildren)
+	{
+		LinkedChildren.Add(&InLinkedChildren);
+	}
+
+	/** @return the number of children */
+	virtual int32 Num() const override
+	{
+		int32 TotalNum = 0;
+		for (const FChildren* Children : LinkedChildren)
+		{
+			TotalNum += Children->Num();
+		}
+
+		return TotalNum;
+	}
+
+	virtual TSharedRef<SWidget> GetChildAt(int32 Index) override
+	{
+		return GetSlotAt(Index).GetWidget();
+	}
+
+	virtual TSharedRef<const SWidget> GetChildAt(int32 Index) const override
+	{
+		return GetSlotAt(Index).GetWidget();
+	}
+
+protected:
+	virtual const FSlotBase& GetSlotAt(int32 ChildIndex) const override
+	{
+		int32 TotalNum = 0;
+		for (const FChildren* Children : LinkedChildren)
+		{
+			const int32 NewTotal = TotalNum + Children->Num();
+			if (NewTotal > ChildIndex)
+			{
+				return Children->GetSlotAt(ChildIndex - TotalNum);
+			}
+			TotalNum = NewTotal;
+		}
+
+		// This result should never occur users should always access a valid index for child slots.
+		static FSlotBase NullSlot; check(false); return NullSlot;
+	}
+
+protected:
+	TArray<FChildren*> LinkedChildren;
+};
+
+
 /**
  * Widgets with no Children can return an instance of FNoChildren.
  * For convenience a shared instance SWidget::NoChildrenInstance can be used.
