SlateGI - Disabling culling when we're on the fast path, adding some TODOs for clearing content in the futre instead.

diff --git a/Engine/Source/Editor/GraphEditor/Public/SNodePanel.h b/Engine/Source/Editor/GraphEditor/Public/SNodePanel.h
--- a/Engine/Source/Editor/GraphEditor/Public/SNodePanel.h
+++ b/Engine/Source/Editor/GraphEditor/Public/SNodePanel.h
@@ -421,7 +421,12 @@ public:
 					const int32 CurWidgetsMaxLayerId = CurWidget.Widget->Paint(Args.WithNewParent(this), CurWidget.Geometry, MyCullingRect, OutDrawElements, LayerId, InWidgetStyle, ShouldBeEnabled(bParentEnabled));
 					MaxLayerId = FMath::Max(MaxLayerId, CurWidgetsMaxLayerId);
 				}
+				else
+				{
+					//SlateGI - RemoveContent
+				}
 			}
+
 			return MaxLayerId;
 		}
 		// End of SPanel Interface
diff --git a/Engine/Source/Runtime/Slate/Private/Widgets/Layout/SConstraintCanvas.cpp b/Engine/Source/Runtime/Slate/Private/Widgets/Layout/SConstraintCanvas.cpp
--- a/Engine/Source/Runtime/Slate/Private/Widgets/Layout/SConstraintCanvas.cpp
+++ b/Engine/Source/Runtime/Slate/Private/Widgets/Layout/SConstraintCanvas.cpp
@@ -225,6 +225,10 @@ int32 SConstraintCanvas::OnPaint( const FPaintArgs& Args, const FGeometry& Allot
 
 			MaxLayerId = FMath::Max(MaxLayerId, CurWidgetsMaxLayerId);
 		}
+		else
+		{
+			//SlateGI - RemoveContent
+		}
 	}
 
 	return MaxLayerId;
diff --git a/Engine/Source/Runtime/Slate/Private/Widgets/Layout/SGridPanel.cpp b/Engine/Source/Runtime/Slate/Private/Widgets/Layout/SGridPanel.cpp
--- a/Engine/Source/Runtime/Slate/Private/Widgets/Layout/SGridPanel.cpp
+++ b/Engine/Source/Runtime/Slate/Private/Widgets/Layout/SGridPanel.cpp
@@ -103,6 +103,10 @@ int32 SGridPanel::OnPaint( const FPaintArgs& Args, const FGeometry& AllottedGeom
 
 					MaxLayerId = FMath::Max(MaxLayerId, CurWidgetsMaxLayerId);
 				}
+				else
+				{
+					//SlateGI - RemoveContent
+				}
 			}
 		}
 	}
diff --git a/Engine/Source/Runtime/Slate/Private/Widgets/SCanvas.cpp b/Engine/Source/Runtime/Slate/Private/Widgets/SCanvas.cpp
--- a/Engine/Source/Runtime/Slate/Private/Widgets/SCanvas.cpp
+++ b/Engine/Source/Runtime/Slate/Private/Widgets/SCanvas.cpp
@@ -115,6 +115,10 @@ int32 SCanvas::OnPaint( const FPaintArgs& Args, const FGeometry& AllottedGeometr
 
 			MaxLayerId = FMath::Max(MaxLayerId, CurWidgetsMaxLayerId);
 		}
+		else
+		{
+
+		}
 	}
 
 	return MaxLayerId;
diff --git a/Engine/Source/Runtime/SlateCore/Private/Widgets/SPanel.cpp b/Engine/Source/Runtime/SlateCore/Private/Widgets/SPanel.cpp
--- a/Engine/Source/Runtime/SlateCore/Private/Widgets/SPanel.cpp
+++ b/Engine/Source/Runtime/SlateCore/Private/Widgets/SPanel.cpp
@@ -30,6 +30,10 @@ int32 SPanel::PaintArrangedChildren( const FPaintArgs& Args, const FArrangedChil
 			const int32 CurWidgetsMaxLayerId = CurWidget.Widget->Paint(NewArgs, CurWidget.Geometry, MyCullingRect, OutDrawElements, LayerId, InWidgetStyle, bShouldBeEnabled);
 			MaxLayerId = FMath::Max(MaxLayerId, CurWidgetsMaxLayerId);
 		}
+		else
+		{
+			//SlateGI - RemoveContent
+		}
 	}
 	
 	return MaxLayerId;
diff --git a/Engine/Source/Runtime/SlateCore/Private/Widgets/SWidget.cpp b/Engine/Source/Runtime/SlateCore/Private/Widgets/SWidget.cpp
--- a/Engine/Source/Runtime/SlateCore/Private/Widgets/SWidget.cpp
+++ b/Engine/Source/Runtime/SlateCore/Private/Widgets/SWidget.cpp
@@ -1618,7 +1618,13 @@ bool SWidget::CanChildrenBeAccessible() const
 
 bool SWidget::IsChildWidgetCulled(const FSlateRect& MyCullingRect, const FArrangedWidget& ArrangedChild) const
 {
-	QUICK_SCOPE_CYCLE_COUNTER(Slate_IsChildWidgetCulled);
+	// If we've enabled global invalidation it's safe to run the culling logic and just 'stop' drawing
+	// a widget, that widget has to be given an opportunity to paint, as wlel as all its children, the
+	// only correct way is to remove the widget from the tree, or to change the visibility of it.
+	if (GSlateIsOnFastUpdatePath)
+	{
+		return false;
+	}
 
 	// We add some slack fill to the culling rect to deal with the common occurrence
 	// of widgets being larger than their root level widget is.  Happens when nested child widgets
