Fix social panel disappearing when switching social panel tabs.
Social panel was incorrectly marked as having volatlity inherited by a parent

diff --git a/Engine/Source/Runtime/SlateCore/Private/FastUpdate/SlateInvalidationRoot.cpp b/Engine/Source/Runtime/SlateCore/Private/FastUpdate/SlateInvalidationRoot.cpp
--- a/Engine/Source/Runtime/SlateCore/Private/FastUpdate/SlateInvalidationRoot.cpp
+++ b/Engine/Source/Runtime/SlateCore/Private/FastUpdate/SlateInvalidationRoot.cpp
@@ -361,7 +361,7 @@ void FSlateInvalidationRoot::BuildNewFastPathList_Recursive(FSlateInvalidationRo
 	if (Proxy.bChildOrderInvalid)
 	{
 		NextTreeIndex = Proxy.LeafMostChildIndex != INDEX_NONE ? Proxy.LeafMostChildIndex + 1 : NextTreeIndex + 1;
-		Proxy.Widget->AssignIndicesToChildren(*this, ParentIndex, NewFastPathList, !Proxy.bInvisibleDueToParentOrSelfVisibility, Proxy.Widget->IsVolatile() || Proxy.Widget->IsVolatileIndirectly());
+		Proxy.Widget->AssignIndicesToChildren(*this, ParentIndex, NewFastPathList, !Proxy.bInvisibleDueToParentOrSelfVisibility, Proxy.Widget->IsVolatileIndirectly());
 	}
 	else
 	{ 
