SWidget - Adding Slate.EnsureAllVisibleWidgetsPaint, it's off by default, but can be turned on to find cases where users are earlying out drawing a child that's visible along their paint path.  When we're on the Fast drawing path, ALL children must be drawn otherwise their cached draw elements don't get cleared.  Enabling this mode vets that after each widget paints all of its visible children painted.

Matt.Kuhlenschmidt

diff --git a/Engine/Source/Runtime/SlateCore/Private/Rendering/DrawElements.cpp b/Engine/Source/Runtime/SlateCore/Private/Rendering/DrawElements.cpp
--- a/Engine/Source/Runtime/SlateCore/Private/Rendering/DrawElements.cpp
+++ b/Engine/Source/Runtime/SlateCore/Private/Rendering/DrawElements.cpp
@@ -666,6 +666,11 @@ FSlateWindowElementList::FDeferredPaint::FDeferredPaint( const TSharedRef<const
 	, WidgetStyle( InWidgetStyle )
 	, bParentEnabled( InParentEnabled )
 {
+	// We need to perform this update here, because otherwise we'll warn that this widget
+	// was not painted along the fast path, which, it will be, but later because it's deferred,
+	// but we need to go ahead and update the painted frame to match the current one, so
+	// that we don't think this widget was forgotten.
+	const_cast<SWidget&>(InWidgetToPaint.Get()).Debug_UpdateLastPaintFrame();
 }
 
 FSlateWindowElementList::FDeferredPaint::FDeferredPaint(const FDeferredPaint& Copy, const FPaintArgs& InArgs)
diff --git a/Engine/Source/Runtime/SlateCore/Private/Types/ReflectionMetadata.cpp b/Engine/Source/Runtime/SlateCore/Private/Types/ReflectionMetadata.cpp
--- a/Engine/Source/Runtime/SlateCore/Private/Types/ReflectionMetadata.cpp
+++ b/Engine/Source/Runtime/SlateCore/Private/Types/ReflectionMetadata.cpp
@@ -3,27 +3,27 @@
 #include "Types/ReflectionMetadata.h"
 #include "Widgets/SWidget.h"
 
-FString FReflectionMetaData::GetWidgetPath(const SWidget* InWidget)
+FString FReflectionMetaData::GetWidgetPath(const SWidget* InWidget, bool bShort, bool bNativePathOnly)
 {
 	if (!InWidget)
 	{
 		return TEXT("None");
 	}
 
-	return GetWidgetPath(*InWidget);
+	return GetWidgetPath(*InWidget, bShort, bNativePathOnly);
 }
 
-FString FReflectionMetaData::GetWidgetPath(const SWidget& InWidget)
+FString FReflectionMetaData::GetWidgetPath(const SWidget& InWidget, bool bShort, bool bNativePathOnly)
 {
 	FString WidgetPath;
 
+	int32 bWidgetsInPath = 0;
+
 	const SWidget* CurrentWidget = &InWidget;
 	while (CurrentWidget)
 	{
-		WidgetPath.InsertAt(0, TEXT("/"));
-
 		TSharedPtr<FReflectionMetaData> MetaData = CurrentWidget->GetMetaData<FReflectionMetaData>();
-		if (MetaData.IsValid())
+		if (!bNativePathOnly && MetaData.IsValid())
 		{
 			WidgetPath.InsertAt(0, MetaData->Name.ToString());
 		}
@@ -33,6 +33,19 @@ FString FReflectionMetaData::GetWidgetPath(const SWidget& InWidget)
 		}
 
 		CurrentWidget = CurrentWidget->GetParentWidget().Get();
+
+		if (CurrentWidget)
+		{
+			WidgetPath.InsertAt(0, TEXT("/"));
+
+			bWidgetsInPath++;
+
+			if (bShort && bWidgetsInPath >= 5)
+			{
+				WidgetPath.InsertAt(0, TEXT("..."));
+				break;
+			}
+		}
 	}
 
 	return WidgetPath;
diff --git a/Engine/Source/Runtime/SlateCore/Private/Widgets/SWidget.cpp b/Engine/Source/Runtime/SlateCore/Private/Widgets/SWidget.cpp
--- a/Engine/Source/Runtime/SlateCore/Private/Widgets/SWidget.cpp
+++ b/Engine/Source/Runtime/SlateCore/Private/Widgets/SWidget.cpp
@@ -16,6 +16,7 @@
 #include "Input/HittestGrid.h"
 #include "Debugging/SlateDebugging.h"
 #include "Widgets/SWindow.h"
+#include "Types/ReflectionMetadata.h"
 
 #if WITH_ACCESSIBILITY
 #include "Widgets/Accessibility/SlateCoreAccessibleWidgets.h"
@@ -40,7 +41,7 @@ static FAutoConsoleVariableRef CVarCullingSlackFillPercent(TEXT("Slate.CullingSl
 
 #endif
 
-#if !(UE_BUILD_SHIPPING || UE_BUILD_TEST)
+#if WITH_SLATE_DEBUGGING
 
 int32 GShowClipping = 0;
 static FAutoConsoleVariableRef CVarSlateShowClipRects(TEXT("Slate.ShowClipping"), GShowClipping, TEXT("Controls whether we should render a clipping zone outline.  Yellow = Axis Scissor Rect Clipping (cheap).  Red = Stencil Clipping (expensive)."), ECVF_Default);
@@ -48,6 +49,9 @@ static FAutoConsoleVariableRef CVarSlateShowClipRects(TEXT("Slate.ShowClipping")
 int32 GDebugCulling = 0;
 static FAutoConsoleVariableRef CVarSlateDebugCulling(TEXT("Slate.DebugCulling"), GDebugCulling, TEXT("Controls whether we should ignore clip rects, and just use culling."), ECVF_Default);
 
+int32 GSlateEnsureAllVisibleWidgetsPaint = 0;
+static FAutoConsoleVariableRef CVarSlateEnsureAllVisibleWidgetsPaint(TEXT("Slate.EnsureAllVisibleWidgetsPaint"), GSlateEnsureAllVisibleWidgetsPaint, TEXT("Ensures that if a child widget is visible before OnPaint, that it was painted this frame after OnPaint, if still marked as visible.  Only works if we're on the FastPaintPath."), ECVF_Default);
+
 #endif
 
 #if STATS
@@ -1242,6 +1246,29 @@ int32 SWidget::Paint(const FPaintArgs& Args, const FGeometry& AllottedGeometry,
 	// FOR RB mode, this should first set GSlateFlowDirection to the incoming state that was cached for the widget, then paint
 	// will override it here to reflow is needed.
 	TGuardValue<EFlowDirection> FlowGuard(GSlateFlowDirection, ComputeFlowDirection());
+
+#if WITH_SLATE_DEBUGGING
+	TArray<TWeakPtr<const SWidget>> DebugChildWidgetsToPaint;
+
+	if (GSlateIsOnFastUpdatePath && GSlateEnsureAllVisibleWidgetsPaint)
+	{
+		// Don't check things that are invalidation roots, or volatile, or volatile indirectly, a completely different set
+		// of rules apply to those widgets.
+		if (!IsVolatile() && !IsVolatileIndirectly() && !Advanced_IsInvalidationRoot())
+		{
+			const FChildren* MyChildren = MutableThis->GetChildren();
+			const int32 NumChildren = MyChildren->Num();
+			for (int32 ChildIndex = 0; ChildIndex < MyChildren->Num(); ++ChildIndex)
+			{
+				TSharedRef<const SWidget> Child = MyChildren->GetChildAt(ChildIndex);
+				if (Child->GetVisibility().IsVisible())
+				{
+					DebugChildWidgetsToPaint.Add(Child);
+				}
+			}
+		}
+	}
+#endif
 	
 	// Paint the geometry of this widget.
 	int32 NewLayerId = OnPaint(UpdatedArgs, AllottedGeometry, CullingBounds, OutDrawElements, LayerId, ContentWidgetStyle, bParentEnabled);
@@ -1249,6 +1276,25 @@ int32 SWidget::Paint(const FPaintArgs& Args, const FGeometry& AllottedGeometry,
 	// Just repainted
 	MutableThis->RemoveUpdateFlags(EWidgetUpdateFlags::NeedsRepaint);
 
+	// Detect children that should have been painted, but were skipped during the paint process.
+	// this will result in geometry being left on screen and not cleared, because it's visible, yet wasn't painted.
+#if WITH_SLATE_DEBUGGING
+	if (GSlateIsOnFastUpdatePath && GSlateEnsureAllVisibleWidgetsPaint)
+	{
+		for (TWeakPtr<const SWidget>& DebugChildThatShouldHaveBeenPaintedPtr : DebugChildWidgetsToPaint)
+		{
+			if (TSharedPtr<const SWidget> DebugChild = DebugChildThatShouldHaveBeenPaintedPtr.Pin())
+			{
+				if (DebugChild->GetVisibility().IsVisible())
+				{
+					ensureMsgf(DebugChild->Debug_GetLastPaintFrame() == GFrameNumber, TEXT("The Widget '%s' was visible, but never painted.  This means it was skipped during painting, without alerting the fast path."), *FReflectionMetaData::GetWidgetPath(DebugChild.Get()));
+				}
+			}
+		}
+	}
+#endif
+
+	// Draw the clipping zone if we've got clipping enabled
 #if WITH_SLATE_DEBUGGING
 	FSlateDebugging::EndWidgetPaint.Broadcast(this, OutDrawElements, NewLayerId);
 
diff --git a/Engine/Source/Runtime/SlateCore/Public/Types/ReflectionMetadata.h b/Engine/Source/Runtime/SlateCore/Public/Types/ReflectionMetadata.h
--- a/Engine/Source/Runtime/SlateCore/Public/Types/ReflectionMetadata.h
+++ b/Engine/Source/Runtime/SlateCore/Public/Types/ReflectionMetadata.h
@@ -41,8 +41,8 @@ public:
 
 public:
 
-	static FString GetWidgetPath(const SWidget* InWidget);
-	static FString GetWidgetPath(const SWidget& InWidget);
+	static FString GetWidgetPath(const SWidget* InWidget, bool bShort = true, bool bNativePathOnly = false);
+	static FString GetWidgetPath(const SWidget& InWidget, bool bShort = true, bool bNativePathOnly = false);
 
 	static FString GetWidgetDebugInfo(const SWidget* InWidget);
 	static FString GetWidgetDebugInfo(const SWidget& InWidget);
diff --git a/Engine/Source/Runtime/SlateCore/Public/Widgets/SWidget.h b/Engine/Source/Runtime/SlateCore/Public/Widgets/SWidget.h
--- a/Engine/Source/Runtime/SlateCore/Public/Widgets/SWidget.h
+++ b/Engine/Source/Runtime/SlateCore/Public/Widgets/SWidget.h
@@ -741,10 +741,22 @@ private:
 		{
 			FastPathProxyHandle.UpdateWidgetFlags(UpdateFlags);
 		}
+
+#if WITH_SLATE_DEBUGGING
+		if (EnumHasAnyFlags(FlagsToRemove, EWidgetUpdateFlags::NeedsRepaint))
+		{
+			Debug_UpdateLastPaintFrame();
+		}
+#endif
 	}
 
 	void UpdateWidgetProxy(int32 NewLayerId, FSlateCachedElementListNode* CacheNode);
 
+#if WITH_SLATE_DEBUGGING
+	uint32 Debug_GetLastPaintFrame() const { return LastPaintFrame; }
+	void Debug_UpdateLastPaintFrame() { LastPaintFrame = GFrameNumber; }
+#endif
+
 public:
 
 	FORCEINLINE TStatId GetStatID() const
@@ -1585,8 +1597,13 @@ private:
 	/** Flow direction preference */
 	EFlowDirectionPreference FlowDirectionPreference;
 
+	/** The different updates this widget needs next frame. */
 	EWidgetUpdateFlags UpdateFlags;
 
+#if WITH_SLATE_DEBUGGING
+	/** The last time this widget got painted. */
+	uint32 LastPaintFrame = 0;
+#endif
 
 	mutable FSlateWidgetPersistentState PersistentState;
 
diff --git a/Engine/Source/Runtime/UMG/Private/Components/Widget.cpp b/Engine/Source/Runtime/UMG/Private/Components/Widget.cpp
--- a/Engine/Source/Runtime/UMG/Private/Components/Widget.cpp
+++ b/Engine/Source/Runtime/UMG/Private/Components/Widget.cpp
@@ -1215,10 +1215,6 @@ void UWidget::SynchronizeProperties()
 		SafeWidget->SetToolTipText(PROPERTY_BINDING(FText, ToolTipText));
 	}
 
-#if WITH_SLATE_DEBUGGING
-	SafeWidget->AddMetadata<FReflectionMetaData>(MakeShared<FReflectionMetaData>(GetFName(), GetClass(), this, GetSourceAssetOrClass()));
-#endif
-
 #if WITH_ACCESSIBILITY
 	TSharedPtr<SWidget> AccessibleWidget = GetAccessibleWidget();
 	if (AccessibleWidget.IsValid())
