Build - Fixing shipping/test builds.

diff --git a/Engine/Source/Runtime/SlateCore/Private/Rendering/DrawElements.cpp b/Engine/Source/Runtime/SlateCore/Private/Rendering/DrawElements.cpp
--- a/Engine/Source/Runtime/SlateCore/Private/Rendering/DrawElements.cpp
+++ b/Engine/Source/Runtime/SlateCore/Private/Rendering/DrawElements.cpp
@@ -666,11 +666,13 @@ FSlateWindowElementList::FDeferredPaint::FDeferredPaint( const TSharedRef<const
 	, WidgetStyle( InWidgetStyle )
 	, bParentEnabled( InParentEnabled )
 {
+#if WITH_SLATE_DEBUGGING
 	// We need to perform this update here, because otherwise we'll warn that this widget
 	// was not painted along the fast path, which, it will be, but later because it's deferred,
 	// but we need to go ahead and update the painted frame to match the current one, so
 	// that we don't think this widget was forgotten.
 	const_cast<SWidget&>(InWidgetToPaint.Get()).Debug_UpdateLastPaintFrame();
+#endif
 }
 
 FSlateWindowElementList::FDeferredPaint::FDeferredPaint(const FDeferredPaint& Copy, const FPaintArgs& InArgs)
