Disable ensure that will always happen during prepass.  Its not a valid place to ensure.

diff --git a/Engine/Source/Runtime/SlateCore/Private/FastUpdate/WidgetProxy.cpp b/Engine/Source/Runtime/SlateCore/Private/FastUpdate/WidgetProxy.cpp
--- a/Engine/Source/Runtime/SlateCore/Private/FastUpdate/WidgetProxy.cpp
+++ b/Engine/Source/Runtime/SlateCore/Private/FastUpdate/WidgetProxy.cpp
@@ -252,7 +252,7 @@ void FWidgetProxyHandle::MarkWidgetDirty(EInvalidateWidget InvalidateReason)
 	{
 		InvalidationRoot->WidgetsNeedingUpdate.Push(Proxy);
 	}
-#if WITH_SLATE_DEBUGGING
+#if 0
 	else
 	{
 		ensure(InvalidationRoot->WidgetsNeedingUpdate.Contains(Proxy));
