Fix hit test grid returning disabled widgets

diff --git a/Engine/Source/Runtime/SlateCore/Private/Input/HittestGrid.cpp b/Engine/Source/Runtime/SlateCore/Private/Input/HittestGrid.cpp
--- a/Engine/Source/Runtime/SlateCore/Private/Input/HittestGrid.cpp
+++ b/Engine/Source/Runtime/SlateCore/Private/Input/HittestGrid.cpp
@@ -217,9 +217,9 @@ TArray<FWidgetAndPointer> FHittestGrid::GetBubblePath(FVector2D DesktopSpaceCoor
 		const FIndexAndDistance& BestHit = Hits[0];
 		if (BestHit.WidgetIndex != INDEX_NONE)
 		{
-			const FWidgetData& WidgetData = WidgetArray[BestHit.WidgetIndex];
+			const FWidgetData& BestHitWidgetData = WidgetArray[BestHit.WidgetIndex];
 
-			const TSharedPtr<SWidget> FirstHitWidget = WidgetData.GetWidget();
+			const TSharedPtr<SWidget> FirstHitWidget = BestHitWidgetData.GetWidget();
 			if (FirstHitWidget.IsValid()) // Make Sure we landed on a valid widget
 			{
 				TArray<FWidgetAndPointer> Path;
@@ -241,11 +241,29 @@ TArray<FWidgetAndPointer> FHittestGrid::GetBubblePath(FVector2D DesktopSpaceCoor
 
 				Algo::Reverse(Path);
 
-				if (WidgetData.CustomPath.IsValid())
+				bool bRemovedDisabledWidgets = false;
+				if (!bIgnoreEnabledStatus)
 				{
-					const TArray<FWidgetAndPointer> BubblePathExtension = WidgetData.CustomPath.Pin()->GetBubblePathAndVirtualCursors(FirstHitWidget->GetPaintSpaceGeometry(), DesktopSpaceCoordinate, bIgnoreEnabledStatus);
-					Path.Append(BubblePathExtension);
+					// @todo It might be more correct to remove all disabled widgets and non-hit testable widgets.  It doesn't make sense to have a hit test invisible widget as a leaf in the path
+					// and that can happen if we remove a disabled widget. Furthermore if we did this we could then append custom paths in all cases since the leaf most widget would be hit testable
+					// For backwards compatibility changing this could be risky
+					const int32 DisabledWidgetIndex = Path.IndexOfByPredicate([](const FArrangedWidget& SomeWidget) { return !SomeWidget.Widget->IsEnabled(); });
+					if (DisabledWidgetIndex != INDEX_NONE)
+					{
+						bRemovedDisabledWidgets = true;
+						Path.RemoveAt(DisabledWidgetIndex, Path.Num() - DisabledWidgetIndex);
+					}
+				}
+
+				if (!bRemovedDisabledWidgets && Path.Num() > 0)
+				{
+					if (BestHitWidgetData.CustomPath.IsValid())
+					{
+						const TArray<FWidgetAndPointer> BubblePathExtension = BestHitWidgetData.CustomPath.Pin()->GetBubblePathAndVirtualCursors(FirstHitWidget->GetPaintSpaceGeometry(), DesktopSpaceCoordinate, bIgnoreEnabledStatus);
+						Path.Append(BubblePathExtension);
+					}
 				}
+	
 				return Path;
 			}
 		}
