Slate - Removing some log spam from emptying slate cached element data.

diff --git a/Engine/Source/Runtime/SlateCore/Public/Rendering/DrawElements.h b/Engine/Source/Runtime/SlateCore/Public/Rendering/DrawElements.h
--- a/Engine/Source/Runtime/SlateCore/Public/Rendering/DrawElements.h
+++ b/Engine/Source/Runtime/SlateCore/Public/Rendering/DrawElements.h
@@ -400,11 +400,6 @@ struct FSlateCachedElementData
 {
 	void Empty()
 	{
-		if (CachedElementLists.Num())
-		{
-			UE_LOG(LogSlate, Log, TEXT("Resetting cached element data.  Num: %d"), CachedElementLists.Num());
-		}
-
 		CachedElementLists.Empty();
 	}
 
