Fix situations where slate draw elements get stuck on screen because they are drawn once and registered as cached and never put on the fast path for some reason (usually a list scrolling and clearing widgets the same frame it adds them)

diff --git a/Engine/Source/Runtime/SlateCore/Private/Widgets/SWidget.cpp b/Engine/Source/Runtime/SlateCore/Private/Widgets/SWidget.cpp
--- a/Engine/Source/Runtime/SlateCore/Private/Widgets/SWidget.cpp
+++ b/Engine/Source/Runtime/SlateCore/Private/Widgets/SWidget.cpp
@@ -728,16 +728,17 @@ void SWidget::UpdateFastPathVisibility(bool bParentVisible, bool bWidgetRemoved)
 		{
 			FastPathProxyHandle.GetInvalidationRoot()->RemoveWidgetFromFastPath(Proxy);
 		}
-		else if (PersistentState.CachedElementListNode != nullptr)
-		{
-			FastPathProxyHandle.GetInvalidationRoot()->GetCachedElements().ResetCache(PersistentState.CachedElementListNode);
-		}
 	}
 	else
 	{
 		ensure(FastPathProxyHandle.GetIndex() == INDEX_NONE);
 	}
 
+	if (PersistentState.CachedElementListNode)
+	{
+		PersistentState.CachedElementListNode->GetValue().GetOwningData()->RemoveCache(PersistentState.CachedElementListNode);
+	}
+
 	FChildren* MyChildren = GetAllChildren();
 	const int32 NumChildren = MyChildren->Num();
 	for (int32 ChildIndex = 0; ChildIndex < NumChildren; ++ChildIndex)
