Slate - Changing the ForceVolatile flag now checks if it's changed before invalidating, and it now only invalidates Volatility.

Matt.Kuhlenschmidt

diff --git a/Engine/Source/Runtime/SlateCore/Private/Widgets/SWidget.cpp b/Engine/Source/Runtime/SlateCore/Private/Widgets/SWidget.cpp
--- a/Engine/Source/Runtime/SlateCore/Private/Widgets/SWidget.cpp
+++ b/Engine/Source/Runtime/SlateCore/Private/Widgets/SWidget.cpp
@@ -1013,7 +1013,7 @@ void SWidget::SetVisibility(TAttribute<EVisibility> InVisibility)
 	SetAttribute(Visibility, InVisibility, EInvalidateWidgetReason::Visibility);
 }
 
-void SWidget::Invalidate(EInvalidateWidget InvalidateReason)
+void SWidget::Invalidate(EInvalidateWidgetReason InvalidateReason)
 {
 	SLATE_CROSS_THREAD_CHECK();
 
diff --git a/Engine/Source/Runtime/SlateCore/Public/Widgets/SWidget.h b/Engine/Source/Runtime/SlateCore/Public/Widgets/SWidget.h
--- a/Engine/Source/Runtime/SlateCore/Public/Widgets/SWidget.h
+++ b/Engine/Source/Runtime/SlateCore/Public/Widgets/SWidget.h
@@ -1002,8 +1002,11 @@ public:
 	 */
 	FORCEINLINE void ForceVolatile(bool bForce)
 	{
-		bForceVolatile = bForce;
-		Invalidate(EInvalidateWidget::LayoutAndVolatility);
+		if (bForceVolatile != bForce)
+		{
+			bForceVolatile = bForce;
+			Invalidate(EInvalidateWidgetReason::Volatility);
+		}
 	}
 
 	FORCEINLINE bool ShouldInvalidatePrepassDueToVolatility() { return bVolatilityAlwaysInvalidatesPrepass; }
@@ -1012,7 +1015,7 @@ public:
 	 * Invalidates the widget from the view of a layout caching widget that may own this widget.
 	 * will force the owning widget to redraw and cache children on the next paint pass.
 	 */
-	void Invalidate(EInvalidateWidget InvalidateReason);
+	void Invalidate(EInvalidateWidgetReason InvalidateReason);
 
 	/**
 	 * Recalculates volatility of the widget and caches the result.  Should be called any time 
