Slate - Check the calling thread when someone invalidates prepass on a widget.

diff --git a/Engine/Source/Runtime/SlateCore/Private/Widgets/SWidget.cpp b/Engine/Source/Runtime/SlateCore/Private/Widgets/SWidget.cpp
--- a/Engine/Source/Runtime/SlateCore/Private/Widgets/SWidget.cpp
+++ b/Engine/Source/Runtime/SlateCore/Private/Widgets/SWidget.cpp
@@ -574,11 +574,12 @@ void SWidget::SlatePrepass(float InLayoutScaleMultiplier)
 
 		Prepass_Internal(InLayoutScaleMultiplier);
 	}
-
 }
 
 void SWidget::InvalidatePrepass()
 {
+	SLATE_CROSS_THREAD_CHECK();
+
 	bNeedsPrepass = true;
 }
 
@@ -596,7 +597,6 @@ FVector2D SWidget::GetDesiredSize() const
 	return DesiredSize.Get(FVector2D::ZeroVector);
 }
 
-
 void SWidget::AssignParentWidget(TSharedPtr<SWidget> InParent)
 {
 #if !UE_BUILD_SHIPPING
