Slate: Protect against re-entrancy initiated panel children destructors when emptying panel children

diff --git a/Engine/Source/Runtime/SlateCore/Public/Layout/Children.h b/Engine/Source/Runtime/SlateCore/Public/Layout/Children.h
--- a/Engine/Source/Runtime/SlateCore/Public/Layout/Children.h
+++ b/Engine/Source/Runtime/SlateCore/Public/Layout/Children.h
@@ -407,7 +407,16 @@ public:
 			// We don't do any invalidating here, that's handled by the FSlotBase, which eventually calls ConditionallyDetatchParentWidget
 
 			TGuardValue<bool> GuardEmptying(bEmptying, true);
+
+			// We empty children by first transferring them onto a stack-owned array, then freeing the elements.
+			// This alleviates issues where (misbehaving) destructors on the children may call back into this class and query children while they are being destroyed.
+			// By storing the children on the stack first, we defer the destruction of children until after we have emptied our owned container.
+			TIndirectArray< SlotType > ChildrenCopy = MoveTemp(static_cast<TIndirectArray< SlotType >&>(*this));
+
+			// Explicitly calling Empty is not really necessary (it is already empty/moved-from now), but we call it for safety
 			TIndirectArray< SlotType >::Empty();
+
+			// ChildrenCopy will now be destroyed 
 		}
 	}
 
