Fix retainers nested in invalidation panels not bubbling layout to their parents.

diff --git a/Engine/Source/Runtime/SlateCore/Private/FastUpdate/SlateInvalidationRoot.cpp b/Engine/Source/Runtime/SlateCore/Private/FastUpdate/SlateInvalidationRoot.cpp
--- a/Engine/Source/Runtime/SlateCore/Private/FastUpdate/SlateInvalidationRoot.cpp
+++ b/Engine/Source/Runtime/SlateCore/Private/FastUpdate/SlateInvalidationRoot.cpp
@@ -535,13 +535,10 @@ bool FSlateInvalidationRoot::ProcessInvalidation()
 			{
 				if (!GSlateEnableGlobalInvalidation && !InvalidationRootWidget->NeedsPrepass() && WidgetProxy.Widget->Advanced_IsInvalidationRoot())
 				{
-					// Prepass on child invalidation panels to give them a chance to process their invalidations since our panel is not going to prepass down all the children
-					WidgetProxy.Widget->SlatePrepass(WidgetProxy.Widget->GetPrepassLayoutScaleMultiplier());
-				}
-				else
-				{
-					bWidgetsNeedRepaint |= WidgetProxy.ProcessInvalidation(WidgetsNeedingUpdate, FastWidgetPathList, *this);
+					WidgetProxy.CurrentInvalidateReason |= EInvalidateWidget::Layout;
 				}
+			
+				bWidgetsNeedRepaint |= WidgetProxy.ProcessInvalidation(WidgetsNeedingUpdate, FastWidgetPathList, *this);
 
 #if WITH_SLATE_DEBUGGING
 				if (GSlateInvalidationDebugging)
diff --git a/Engine/Source/Runtime/UMG/Private/Slate/SRetainerWidget.cpp b/Engine/Source/Runtime/UMG/Private/Slate/SRetainerWidget.cpp
--- a/Engine/Source/Runtime/UMG/Private/Slate/SRetainerWidget.cpp
+++ b/Engine/Source/Runtime/UMG/Private/Slate/SRetainerWidget.cpp
@@ -448,7 +448,7 @@ bool SRetainerWidget::PaintRetainedContent(const FSlateInvalidationContext& Cont
 				}
 #endif
 
-				bRenderRequested = true;
+				bRenderRequested = false;
 				Shared_WaitingToRender.Remove(this);
 
 				LastDrawTime = FApp::GetCurrentTime();
@@ -537,13 +537,6 @@ FVector2D SRetainerWidget::ComputeDesiredSize(float LayoutScaleMuliplier) const
 	}
 }
 
-bool SRetainerWidget::ComputeVolatility() const
-{
-	// We need to be volatile when global invalidation is turned on so that we have a chance to update
-	// @todo instead of being volatile it may be possible to move this to tick
-	return GSlateEnableGlobalInvalidation != 0;
-}
-
 void SRetainerWidget::OnGlobalInvalidationToggled(bool bGlobalInvalidationEnabled)
 {
 	InvalidateRoot();
diff --git a/Engine/Source/Runtime/UMG/Public/Slate/SRetainerWidget.h b/Engine/Source/Runtime/UMG/Public/Slate/SRetainerWidget.h
--- a/Engine/Source/Runtime/UMG/Public/Slate/SRetainerWidget.h
+++ b/Engine/Source/Runtime/UMG/Public/Slate/SRetainerWidget.h
@@ -92,7 +92,6 @@ protected:
 	/** SCompoundWidget interface */
 	virtual int32 OnPaint(const FPaintArgs& Args, const FGeometry& AllottedGeometry, const FSlateRect& MyCullingRect, FSlateWindowElementList& OutDrawElements, int32 LayerId, const FWidgetStyle& InWidgetStyle, bool bParentEnabled) const override;
 	virtual FVector2D ComputeDesiredSize(float Scale) const override;
-	virtual bool ComputeVolatility() const override;
 	virtual bool Advanced_IsInvalidationRoot() const { return bEnableRetainedRendering; }
 	virtual bool CustomPrepass(float LayoutScaleMultiplier) override;
 
