Added missing file from last checkin

diff --git a/Engine/Source/Runtime/SlateCore/Private/FastUpdate/WidgetProxy.cpp b/Engine/Source/Runtime/SlateCore/Private/FastUpdate/WidgetProxy.cpp
--- a/Engine/Source/Runtime/SlateCore/Private/FastUpdate/WidgetProxy.cpp
+++ b/Engine/Source/Runtime/SlateCore/Private/FastUpdate/WidgetProxy.cpp
@@ -81,7 +81,7 @@ bool FWidgetProxy::ProcessInvalidation(FWidgetUpdateList& UpdateList, TArray<FWi
 		FVector2D NewDesiredSize = FVector2D::ZeroVector;
 		if (Visibility != EVisibility::Collapsed)
 		{
-			if (Widget->NeedsPrepass())
+			if (Widget->NeedsPrepass() || (!GSlateEnableGlobalInvalidation && Widget->Advanced_IsInvalidationRoot()))
 			{
 				Widget->SlatePrepass(Widget->PrepassLayoutScaleMultiplier.Get(1.0f));
 			}
