Guard against widgets in the tree more than once in ricky bobby.  This is illegal and can cause crashes

diff --git a/Engine/Source/Runtime/SlateCore/Private/Widgets/SWidget.cpp b/Engine/Source/Runtime/SlateCore/Private/Widgets/SWidget.cpp
--- a/Engine/Source/Runtime/SlateCore/Private/Widgets/SWidget.cpp
+++ b/Engine/Source/Runtime/SlateCore/Private/Widgets/SWidget.cpp
@@ -651,6 +651,11 @@ bool SWidget::ConditionallyDetatchParentWidget(SWidget* InExpectedParent)
 
 void SWidget::AssignIndicesToChildren(FSlateInvalidationRoot& Root, int32 ParentIndex, TArray<FWidgetProxy, TMemStackAllocator<>>& FastPathList, bool bParentVisible, bool bParentVolatile)
 {
+	if (FastPathProxyHandle.IsValid())
+	{
+		ensureAlwaysMsgf(!FastPathProxyHandle.IsValid(), TEXT("Widget %s was already assigned a proxy handle. If this is being hit this widget is in the active slate tree more than once.  This is illegal and at best will result in UI corruption."), *FReflectionMetaData::GetWidgetDebugInfo(*this));
+		return;
+	}
 	FWidgetProxy MyProxy(this);
 	MyProxy.Index = FastPathList.Num();
 	MyProxy.ParentIndex = ParentIndex;
@@ -670,7 +675,6 @@ void SWidget::AssignIndicesToChildren(FSlateInvalidationRoot& Root, int32 Parent
 
 	FastPathProxyHandle = FWidgetProxyHandle(Root, MyProxy.Index);
 
-
 	if (bInvisibleDueToParentOrSelfVisibility&& PersistentState.CachedElementListNode != nullptr)
 	{
 #if WITH_SLATE_DEBUGGING
