Slate - Switching several bool console vars over to be bools, not ints now that they're fully supported.

diff --git a/Engine/Source/Runtime/SlateCore/Private/SlateCoreClasses.cpp b/Engine/Source/Runtime/SlateCore/Private/SlateCoreClasses.cpp
--- a/Engine/Source/Runtime/SlateCore/Private/SlateCoreClasses.cpp
+++ b/Engine/Source/Runtime/SlateCore/Private/SlateCoreClasses.cpp
@@ -43,7 +43,8 @@ bool GSlateIsOnFastUpdatePath = false;
 bool GSlateIsInInvalidationSlowPath = false;
 
 #if WITH_SLATE_DEBUGGING
-int32 GSlateInvalidationDebugging = 0;
+
+bool GSlateInvalidationDebugging = false;
 /** True if we should allow widgets to be cached in the UI at all. */
 FAutoConsoleVariableRef CVarInvalidationDebugging(
 	TEXT("Slate.InvalidationDebugging"),
@@ -51,7 +52,7 @@ FAutoConsoleVariableRef CVarInvalidationDebugging(
 	TEXT("Whether to show invalidation debugging visualization"));
 
 
-int32 GSlateHitTestGridDebugging = 0;
+bool GSlateHitTestGridDebugging = false;
 /** True if we should allow widgets to be cached in the UI at all. */
 FAutoConsoleVariableRef CVarHitTestGridDebugging(
 	TEXT("Slate.HitTestGridDebugging"),
diff --git a/Engine/Source/Runtime/SlateCore/Private/Widgets/SWidget.cpp b/Engine/Source/Runtime/SlateCore/Private/Widgets/SWidget.cpp
--- a/Engine/Source/Runtime/SlateCore/Private/Widgets/SWidget.cpp
+++ b/Engine/Source/Runtime/SlateCore/Private/Widgets/SWidget.cpp
@@ -43,13 +43,13 @@ static FAutoConsoleVariableRef CVarCullingSlackFillPercent(TEXT("Slate.CullingSl
 
 #if WITH_SLATE_DEBUGGING
 
-int32 GShowClipping = 0;
+bool GShowClipping = false;
 static FAutoConsoleVariableRef CVarSlateShowClipRects(TEXT("Slate.ShowClipping"), GShowClipping, TEXT("Controls whether we should render a clipping zone outline.  Yellow = Axis Scissor Rect Clipping (cheap).  Red = Stencil Clipping (expensive)."), ECVF_Default);
 
-int32 GDebugCulling = 0;
+bool GDebugCulling = false;
 static FAutoConsoleVariableRef CVarSlateDebugCulling(TEXT("Slate.DebugCulling"), GDebugCulling, TEXT("Controls whether we should ignore clip rects, and just use culling."), ECVF_Default);
 
-int32 GSlateEnsureAllVisibleWidgetsPaint = 0;
+bool GSlateEnsureAllVisibleWidgetsPaint = false;
 static FAutoConsoleVariableRef CVarSlateEnsureAllVisibleWidgetsPaint(TEXT("Slate.EnsureAllVisibleWidgetsPaint"), GSlateEnsureAllVisibleWidgetsPaint, TEXT("Ensures that if a child widget is visible before OnPaint, that it was painted this frame after OnPaint, if still marked as visible.  Only works if we're on the FastPaintPath."), ECVF_Default);
 
 #endif
diff --git a/Engine/Source/Runtime/SlateCore/Public/SlateGlobals.h b/Engine/Source/Runtime/SlateCore/Public/SlateGlobals.h
--- a/Engine/Source/Runtime/SlateCore/Public/SlateGlobals.h
+++ b/Engine/Source/Runtime/SlateCore/Public/SlateGlobals.h
@@ -59,8 +59,8 @@ extern SLATECORE_API bool GSlateIsOnFastUpdatePath;
 extern SLATECORE_API bool GSlateIsInInvalidationSlowPath;
 
 #if WITH_SLATE_DEBUGGING
-extern SLATECORE_API int32 GSlateInvalidationDebugging;
-extern SLATECORE_API int32 GSlateHitTestGridDebugging;
+extern SLATECORE_API bool GSlateInvalidationDebugging;
+extern SLATECORE_API bool GSlateHitTestGridDebugging;
 #endif
 /* Forward declarations
 *****************************************************************************/
