Fix hit test grid not resetting cells if the total amount of cells remains the same after a resize but the width and height differ.  This also fixes a crash with retainer panels in that situation

diff --git a/Engine/Source/Runtime/SlateCore/Private/Input/HittestGrid.cpp b/Engine/Source/Runtime/SlateCore/Private/Input/HittestGrid.cpp
--- a/Engine/Source/Runtime/SlateCore/Private/Input/HittestGrid.cpp
+++ b/Engine/Source/Runtime/SlateCore/Private/Input/HittestGrid.cpp
@@ -283,16 +283,13 @@ bool FHittestGrid::SetHittestArea(const FVector2D& HittestPositionInDesktop, con
 		NumCells = FIntPoint(FMath::CeilToInt(GridSize.X / CellSize.X), FMath::CeilToInt(GridSize.Y / CellSize.Y));
 		
 		const int32 NewTotalCells = NumCells.X * NumCells.Y;
-		if (NewTotalCells != Cells.Num())
-		{
-			Cells.Reset(NewTotalCells);
-			Cells.SetNum(NewTotalCells);
+		Cells.Reset(NewTotalCells);
+		Cells.SetNum(NewTotalCells);
 
-			WidgetMap.Reset();
-			WidgetArray.Reset();
+		WidgetMap.Reset();
+		WidgetArray.Reset();
 
-			bWasCleared = true;
-		}
+		bWasCleared = true;
 	}
 
 	GridOrigin = HittestPositionInDesktop;
diff --git a/Engine/Source/Runtime/UMG/Private/Slate/SRetainerWidget.cpp b/Engine/Source/Runtime/UMG/Private/Slate/SRetainerWidget.cpp
--- a/Engine/Source/Runtime/UMG/Private/Slate/SRetainerWidget.cpp
+++ b/Engine/Source/Runtime/UMG/Private/Slate/SRetainerWidget.cpp
@@ -473,6 +473,10 @@ int32 SRetainerWidget::OnPaint(const FPaintArgs& Args, const FGeometry& Allotted
 
 		// Copy hit test grid settings from the root
 		const bool bHittestCleared = HittestGrid.SetHittestArea(Args.RootGrid.GetGridOrigin(), Args.RootGrid.GetGridSize(), Args.RootGrid.GetGridWindowOrigin());
+		if (bHittestCleared)
+		{
+			MutableThis->RequestRender();
+		}
 		FPaintArgs NewArgs = Args.WithNewHitTestGrid(HittestGrid);
 
 		FSlateInvalidationContext Context(OutDrawElements, InWidgetStyle);
