Fix hit testable widgets not always being removed from the hit test grid when a parents visibility changes.

diff --git a/Engine/Source/Runtime/SlateCore/Private/Widgets/SWidget.cpp b/Engine/Source/Runtime/SlateCore/Private/Widgets/SWidget.cpp
--- a/Engine/Source/Runtime/SlateCore/Private/Widgets/SWidget.cpp
+++ b/Engine/Source/Runtime/SlateCore/Private/Widgets/SWidget.cpp
@@ -588,7 +588,7 @@ void SWidget::InvalidateChildRemovedFromTree(SWidget& Child)
 	if (Child.FastPathProxyHandle.IsValid())
 	{
 		SCOPED_NAMED_EVENT(SWidget_InvalidateChildRemovedFromTree, FColor::Red);
-		Child.UpdateFastPathVisibility(false, true);
+		Child.UpdateFastPathVisibility(false, true, Child.FastPathProxyHandle.GetInvalidationRoot()->GetHittestGrid());
 	}
 }
 
@@ -713,7 +713,7 @@ void SWidget::AssignIndicesToChildren(FSlateInvalidationRoot& Root, int32 Parent
 	}
 }
 
-void SWidget::UpdateFastPathVisibility(bool bParentVisible, bool bWidgetRemoved)
+void SWidget::UpdateFastPathVisibility(bool bParentVisible, bool bWidgetRemoved, FHittestGrid* ParentHittestGrid)
 {
 	const EVisibility CurrentVisibility = GetVisibility();
 	const bool bParentAndSelfVisible = bParentVisible && CurrentVisibility.IsVisible();
@@ -721,9 +721,11 @@ void SWidget::UpdateFastPathVisibility(bool bParentVisible, bool bWidgetRemoved)
 	bInvisibleDueToParentOrSelfVisibility = !bParentAndSelfVisible;
 	const bool bVisibilityChanged = bWasInvisible != bInvisibleDueToParentOrSelfVisibility;
 
+	FHittestGrid* HittestGridToRemoveFrom = ParentHittestGrid;
 	if (FastPathProxyHandle.IsValid())
-	{
-		FastPathProxyHandle.GetInvalidationRoot()->GetHittestGrid()->RemoveWidget(SharedThis(this));
+	{	
+		// Try and remove this from the current handles hit test grid.  If we are in a nested invalidation situation the hittest grid may have changed
+		HittestGridToRemoveFrom = FastPathProxyHandle.GetInvalidationRoot()->GetHittestGrid();
 		FWidgetProxy& Proxy = FastPathProxyHandle.GetProxy();
 		Proxy.Visibility = CurrentVisibility;
 		Proxy.bInvisibleDueToParentOrSelfVisibility = bInvisibleDueToParentOrSelfVisibility;
@@ -738,6 +740,11 @@ void SWidget::UpdateFastPathVisibility(bool bParentVisible, bool bWidgetRemoved)
 		ensure(FastPathProxyHandle.GetIndex() == INDEX_NONE);
 	}
 
+	if (HittestGridToRemoveFrom)
+	{
+		HittestGridToRemoveFrom->RemoveWidget(SharedThis(this));
+	}
+
 	if (PersistentState.CachedElementListNode)
 	{
 		PersistentState.CachedElementListNode->GetValue().GetOwningData()->RemoveCache(PersistentState.CachedElementListNode);
@@ -748,7 +755,7 @@ void SWidget::UpdateFastPathVisibility(bool bParentVisible, bool bWidgetRemoved)
 	for (int32 ChildIndex = 0; ChildIndex < NumChildren; ++ChildIndex)
 	{
 		const TSharedRef<SWidget>& Child = MyChildren->GetChildAt(ChildIndex);
-		Child->UpdateFastPathVisibility(bParentAndSelfVisible, bWidgetRemoved);
+		Child->UpdateFastPathVisibility(bParentAndSelfVisible, bWidgetRemoved, HittestGridToRemoveFrom);
 	}
 }
 
@@ -1044,8 +1051,7 @@ void SWidget::Invalidate(EInvalidateWidgetReason InvalidateReason)
 		{
 			SCOPED_NAMED_EVENT(SWidget_UpdateFastPathVisibility, FColor::Red);
 			TSharedPtr<SWidget> ParentWidget = GetParentWidget();
-
-			UpdateFastPathVisibility(ParentWidget.IsValid() ? !ParentWidget->bInvisibleDueToParentOrSelfVisibility : false, false);
+			UpdateFastPathVisibility(ParentWidget.IsValid() ? !ParentWidget->bInvisibleDueToParentOrSelfVisibility : false, false, FastPathProxyHandle.GetInvalidationRoot()->GetHittestGrid());
 		}
 
 		if (bVolatilityChanged)
diff --git a/Engine/Source/Runtime/SlateCore/Public/Widgets/SWidget.h b/Engine/Source/Runtime/SlateCore/Public/Widgets/SWidget.h
--- a/Engine/Source/Runtime/SlateCore/Public/Widgets/SWidget.h
+++ b/Engine/Source/Runtime/SlateCore/Public/Widgets/SWidget.h
@@ -710,7 +710,7 @@ protected:
 
 private:
 
-	void UpdateFastPathVisibility(bool bParentVisible, bool bWidgetRemoved);
+	void UpdateFastPathVisibility(bool bParentVisible, bool bWidgetRemoved, FHittestGrid* ParentHittestGrid);
 
 	void UpdateFastPathVolatility(bool bParentVolatile);
 
