Text Shaping Improvements:
 - The font selection process for HarfBuzz shaping now processes text as grapheme clusters, which ensures that things like zero-width joiners are applied using the correct font.
 - The font ranges are now selected using codepoints (UTF32CHAR) rather than platform specific characters (TCHAR).
   - This allows characters outside of the BMP to choose the correct font, even on platforms using UTF-16 strings (surrogate pairs are combined when processing the grapheme clusters).
 - The logic for calculating the number of characters and grapheme clusters in a glyph have been combined into a single optimized loop (separately for LTR and RTL text) which fixes some counting issues in the old code.

diff --git a/Engine/Source/Runtime/SlateCore/Private/Fonts/FontCache.cpp b/Engine/Source/Runtime/SlateCore/Private/Fonts/FontCache.cpp
--- a/Engine/Source/Runtime/SlateCore/Private/Fonts/FontCache.cpp
+++ b/Engine/Source/Runtime/SlateCore/Private/Fonts/FontCache.cpp
@@ -568,9 +568,9 @@ bool FCharacterList::CanCacheCharacter(TCHAR Character, const EFontFallback MaxF
 	else
 	{
 		float SubFontScalingFactor = 1.0f;
-		const FFontData& FontData = FontCache.CompositeFontCache->GetFontDataForCharacter(FontKey.GetFontInfo(), Character, SubFontScalingFactor);
+		const FFontData& FontData = FontCache.CompositeFontCache->GetFontDataForCodepoint(FontKey.GetFontInfo(), Character, SubFontScalingFactor);
 
-		bReturnVal = FontCache.FontRenderer->CanLoadCharacter(FontData, Character, MaxFontFallback);
+		bReturnVal = FontCache.FontRenderer->CanLoadCodepoint(FontData, Character, MaxFontFallback);
 	}
 
 	return bReturnVal;
@@ -660,8 +660,8 @@ FCharacterList::FCharacterListEntry FCharacterList::CacheCharacter(TCHAR Charact
 
 		// Get the data needed to render this character
 		float SubFontScalingFactor = 1.0f;
-		const FFontData* FontDataPtr = &FontCache.CompositeFontCache->GetFontDataForCharacter(FontInfo, Character, SubFontScalingFactor);
-		FFreeTypeFaceGlyphData FaceGlyphData = FontCache.FontRenderer->GetFontFaceForCharacter(*FontDataPtr, Character, FontInfo.FontFallback);
+		const FFontData* FontDataPtr = &FontCache.CompositeFontCache->GetFontDataForCodepoint(FontInfo, Character, SubFontScalingFactor);
+		FFreeTypeFaceGlyphData FaceGlyphData = FontCache.FontRenderer->GetFontFaceForCodepoint(*FontDataPtr, Character, FontInfo.FontFallback);
 
 		// Found a valid font face?
 		if (FaceGlyphData.FaceAndMemory.IsValid())
@@ -1031,9 +1031,9 @@ const FFontData& FSlateFontCache::GetDefaultFontData( const FSlateFontInfo& InFo
 	return CompositeFontCache->GetDefaultFontData(InFontInfo);
 }
 
-const FFontData& FSlateFontCache::GetFontDataForCharacter( const FSlateFontInfo& InFontInfo, const TCHAR InChar, float& OutScalingFactor ) const
+const FFontData& FSlateFontCache::GetFontDataForCodepoint( const FSlateFontInfo& InFontInfo, const UTF32CHAR InCodepoint, float& OutScalingFactor ) const
 {
-	return CompositeFontCache->GetFontDataForCharacter(InFontInfo, InChar, OutScalingFactor);
+	return CompositeFontCache->GetFontDataForCodepoint(InFontInfo, InCodepoint, OutScalingFactor);
 }
 
 uint16 FSlateFontCache::GetMaxCharacterHeight( const FSlateFontInfo& InFontInfo, float FontScale ) const
diff --git a/Engine/Source/Runtime/SlateCore/Private/Fonts/FontCacheCompositeFont.cpp b/Engine/Source/Runtime/SlateCore/Private/Fonts/FontCacheCompositeFont.cpp
--- a/Engine/Source/Runtime/SlateCore/Private/Fonts/FontCacheCompositeFont.cpp
+++ b/Engine/Source/Runtime/SlateCore/Private/Fonts/FontCacheCompositeFont.cpp
@@ -81,9 +81,9 @@ FCachedCompositeFontData::FCachedCompositeFontData(const FCompositeFont& InCompo
 	RefreshFontRanges();
 }
 
-const FCachedTypefaceData* FCachedCompositeFontData::GetTypefaceForCharacter(const TCHAR InChar) const
+const FCachedTypefaceData* FCachedCompositeFontData::GetTypefaceForCodepoint(const UTF32CHAR InCodepoint) const
 {
-	const int32 CharIndex = static_cast<int32>(InChar);
+	const int32 CharIndex = static_cast<int32>(InCodepoint);
 
 	auto GetTypefaceFromRange = [CharIndex](const TArray<FCachedFontRange>& InFontRanges) -> const FCachedTypefaceData*
 	{
@@ -256,15 +256,15 @@ const FFontData& FCompositeFontCache::GetDefaultFontData(const FSlateFontInfo& I
 	return DummyFontData;
 }
 
-const FFontData& FCompositeFontCache::GetFontDataForCharacter(const FSlateFontInfo& InFontInfo, const TCHAR InChar, float& OutScalingFactor)
+const FFontData& FCompositeFontCache::GetFontDataForCodepoint(const FSlateFontInfo& InFontInfo, const UTF32CHAR InCodepoint, float& OutScalingFactor)
 {
 	static const FFontData DummyFontData;
 
-	auto GetFontDataForCharacterInTypeface = [this, InChar, &InFontInfo](const FCachedTypefaceData* InCachedTypefaceData, const FCachedTypefaceData* InCachedDefaultTypefaceData, const bool InSkipCharacterCheck) -> const FFontData*
+	auto GetFontDataForCharacterInTypeface = [this, InCodepoint, &InFontInfo](const FCachedTypefaceData* InCachedTypefaceData, const FCachedTypefaceData* InCachedDefaultTypefaceData, const bool InSkipCharacterCheck) -> const FFontData*
 	{
 		// Try to find the correct font from the typeface
 		const FFontData* FoundFontData = InCachedTypefaceData->GetFontData(InFontInfo.TypefaceFontName);
-		if (FoundFontData && (InSkipCharacterCheck || DoesFontDataSupportCharacter(*FoundFontData, InChar)))
+		if (FoundFontData && (InSkipCharacterCheck || DoesFontDataSupportCodepoint(*FoundFontData, InCodepoint)))
 		{
 			return FoundFontData;
 		}
@@ -278,7 +278,7 @@ const FFontData& FCompositeFontCache::GetFontDataForCharacter(const FSlateFontIn
 			{
 				const TSet<FName>& DefaultFontAttributes = GetFontAttributes(*FoundDefaultFontData);
 				FoundFontData = GetBestMatchFontForAttributes(InCachedTypefaceData, DefaultFontAttributes);
-				if (FoundFontData && (InSkipCharacterCheck || DoesFontDataSupportCharacter(*FoundFontData, InChar)))
+				if (FoundFontData && (InSkipCharacterCheck || DoesFontDataSupportCodepoint(*FoundFontData, InCodepoint)))
 				{
 					return FoundFontData;
 				}
@@ -287,7 +287,7 @@ const FFontData& FCompositeFontCache::GetFontDataForCharacter(const FSlateFontIn
 
 		// Failing that, try the primary font
 		FoundFontData = InCachedTypefaceData->GetPrimaryFontData();
-		if (FoundFontData && (InSkipCharacterCheck || DoesFontDataSupportCharacter(*FoundFontData, InChar)))
+		if (FoundFontData && (InSkipCharacterCheck || DoesFontDataSupportCodepoint(*FoundFontData, InCodepoint)))
 		{
 			return FoundFontData;
 		}
@@ -296,7 +296,7 @@ const FFontData& FCompositeFontCache::GetFontDataForCharacter(const FSlateFontIn
 	};
 
 	const FCompositeFont* const ResolvedCompositeFont = InFontInfo.GetCompositeFont();
-	const FCachedTypefaceData* const CachedTypefaceData = GetCachedTypefaceForCharacter(ResolvedCompositeFont, InChar);
+	const FCachedTypefaceData* const CachedTypefaceData = GetCachedTypefaceForCodepoint(ResolvedCompositeFont, InCodepoint);
 	if (CachedTypefaceData)
 	{
 		const FCachedTypefaceData* const CachedDefaultTypefaceData = GetDefaultCachedTypeface(ResolvedCompositeFont);
@@ -487,11 +487,11 @@ const FFontData* FCompositeFontCache::GetBestMatchFontForAttributes(const FCache
 	return BestMatchFont;
 }
 
-bool FCompositeFontCache::DoesFontDataSupportCharacter(const FFontData& InFontData, const TCHAR InChar)
+bool FCompositeFontCache::DoesFontDataSupportCodepoint(const FFontData& InFontData, const UTF32CHAR InCodepoint)
 {
 #if WITH_FREETYPE
 	TSharedPtr<FFreeTypeFace> FaceAndMemory = GetFontFace(InFontData);
-	return FaceAndMemory.IsValid() && FT_Get_Char_Index(FaceAndMemory->GetFace(), InChar) != 0;
+	return FaceAndMemory.IsValid() && FT_Get_Char_Index(FaceAndMemory->GetFace(), InCodepoint) != 0;
 #else  // WITH_FREETYPE
 	return false;
 #endif // WITH_FREETYPE
diff --git a/Engine/Source/Runtime/SlateCore/Private/Fonts/FontCacheCompositeFont.h b/Engine/Source/Runtime/SlateCore/Private/Fonts/FontCacheCompositeFont.h
--- a/Engine/Source/Runtime/SlateCore/Private/Fonts/FontCacheCompositeFont.h
+++ b/Engine/Source/Runtime/SlateCore/Private/Fonts/FontCacheCompositeFont.h
@@ -129,8 +129,8 @@ public:
 		return CachedTypefaces[CachedTypeface_FallbackIndex].Get();
 	}
 
-	/** Get the typeface that should be used for the given character */
-	const FCachedTypefaceData* GetTypefaceForCharacter(const TCHAR InChar) const;
+	/** Get the typeface that should be used for the given codepoint */
+	const FCachedTypefaceData* GetTypefaceForCodepoint(const UTF32CHAR InCodepoint) const;
 
 	/** Get all the font data cached by this entry */
 	void GetCachedFontData(TArray<const FFontData*>& OutFontData) const;
@@ -210,8 +210,8 @@ public:
 	/** Get the default font data to use for the given font info */
 	const FFontData& GetDefaultFontData(const FSlateFontInfo& InFontInfo);
 
-	/** Get the font data to use for the given font info and character */
-	const FFontData& GetFontDataForCharacter(const FSlateFontInfo& InFontInfo, const TCHAR InChar, float& OutScalingFactor);
+	/** Get the font data to use for the given font info and codepoint */
+	const FFontData& GetFontDataForCodepoint(const FSlateFontInfo& InFontInfo, const UTF32CHAR InCodepoint, float& OutScalingFactor);
 
 	/** Gets or loads a FreeType font face */
 	TSharedPtr<FFreeTypeFace> GetFontFace(const FFontData& InFontData);
@@ -248,18 +248,18 @@ private:
 		return (CachedCompositeFont) ? CachedCompositeFont->GetFallbackTypeface() : nullptr;
 	}
 
-	/** Get the typeface that should be used for the given character */
-	FORCEINLINE const FCachedTypefaceData* GetCachedTypefaceForCharacter(const FCompositeFont* const InCompositeFont, const TCHAR InChar)
+	/** Get the typeface that should be used for the given codepoint */
+	FORCEINLINE const FCachedTypefaceData* GetCachedTypefaceForCodepoint(const FCompositeFont* const InCompositeFont, const UTF32CHAR InCodepoint)
 	{
 		const FCachedCompositeFontData* const CachedCompositeFont = GetCachedCompositeFont(InCompositeFont);
-		return (CachedCompositeFont) ? CachedCompositeFont->GetTypefaceForCharacter(InChar) : nullptr;
+		return (CachedCompositeFont) ? CachedCompositeFont->GetTypefaceForCodepoint(InCodepoint) : nullptr;
 	}
 
 	/** Try and find some font data that best matches the given attributes */
 	const FFontData* GetBestMatchFontForAttributes(const FCachedTypefaceData* const InCachedTypefaceData, const TSet<FName>& InFontAttributes);
 
-	/** Check to see whether the given font data supports rendering the given character */
-	bool DoesFontDataSupportCharacter(const FFontData& InFontData, const TCHAR InChar);
+	/** Check to see whether the given font data supports rendering the given codepoint */
+	bool DoesFontDataSupportCodepoint(const FFontData& InFontData, const UTF32CHAR InCodepoint);
 
 	/** Called after the active culture has changed */
 	void HandleCultureChanged();
diff --git a/Engine/Source/Runtime/SlateCore/Private/Fonts/LegacySlateFontInfoCache.cpp b/Engine/Source/Runtime/SlateCore/Private/Fonts/LegacySlateFontInfoCache.cpp
--- a/Engine/Source/Runtime/SlateCore/Private/Fonts/LegacySlateFontInfoCache.cpp
+++ b/Engine/Source/Runtime/SlateCore/Private/Fonts/LegacySlateFontInfoCache.cpp
@@ -26,9 +26,9 @@ FString FLegacySlateFontInfoCache::FFallbackContext::ToString() const
 	}
 
 	FString CharacterInfo;
-	if (Char)
+	if (Codepoint)
 	{
-		CharacterInfo = FString::Printf(TEXT("%c (U+%04x)"), Char, (int32)Char);
+		CharacterInfo = FString::Printf(TEXT("%c (U+%04x)"), (TCHAR)Codepoint, (int32)Codepoint);
 	}
 	else
 	{
diff --git a/Engine/Source/Runtime/SlateCore/Private/Fonts/LegacySlateFontInfoCache.h b/Engine/Source/Runtime/SlateCore/Private/Fonts/LegacySlateFontInfoCache.h
--- a/Engine/Source/Runtime/SlateCore/Private/Fonts/LegacySlateFontInfoCache.h
+++ b/Engine/Source/Runtime/SlateCore/Private/Fonts/LegacySlateFontInfoCache.h
@@ -20,9 +20,9 @@ public:
 	public:
 		FFallbackContext() = default;
 
-		FFallbackContext(const FFontData* InFontData, const TCHAR InChar)
+		FFallbackContext(const FFontData* InFontData, const UTF32CHAR InCodepoint)
 			: FontData(InFontData)
-			, Char(InChar)
+			, Codepoint(InCodepoint)
 		{
 		}
 
@@ -30,7 +30,7 @@ public:
 
 	private:
 		const FFontData* FontData = nullptr;
-		TCHAR Char = 0;
+		UTF32CHAR Codepoint = 0;
 	};
 
 	/**
diff --git a/Engine/Source/Runtime/SlateCore/Private/Fonts/SlateFontRenderer.cpp b/Engine/Source/Runtime/SlateCore/Private/Fonts/SlateFontRenderer.cpp
--- a/Engine/Source/Runtime/SlateCore/Private/Fonts/SlateFontRenderer.cpp
+++ b/Engine/Source/Runtime/SlateCore/Private/Fonts/SlateFontRenderer.cpp
@@ -82,7 +82,7 @@ uint16 FSlateFontRenderer::GetMaxHeight(const FSlateFontInfo& InFontInfo, const
 	// Just get info for the null character 
 	TCHAR Char = 0;
 	const FFontData& FontData = CompositeFontCache->GetDefaultFontData(InFontInfo);
-	const FFreeTypeFaceGlyphData FaceGlyphData = GetFontFaceForCharacter(FontData, Char, InFontInfo.FontFallback);
+	const FFreeTypeFaceGlyphData FaceGlyphData = GetFontFaceForCodepoint(FontData, Char, InFontInfo.FontFallback);
 
 	if (FaceGlyphData.FaceAndMemory.IsValid())
 	{
@@ -105,7 +105,7 @@ int16 FSlateFontRenderer::GetBaseline(const FSlateFontInfo& InFontInfo, const fl
 	// Just get info for the null character 
 	TCHAR Char = 0;
 	const FFontData& FontData = CompositeFontCache->GetDefaultFontData(InFontInfo);
-	const FFreeTypeFaceGlyphData FaceGlyphData = GetFontFaceForCharacter(FontData, Char, InFontInfo.FontFallback);
+	const FFreeTypeFaceGlyphData FaceGlyphData = GetFontFaceForCodepoint(FontData, Char, InFontInfo.FontFallback);
 
 	if (FaceGlyphData.FaceAndMemory.IsValid())
 	{
@@ -207,12 +207,12 @@ int8 FSlateFontRenderer::GetKerning(const FFontData& InFontData, const int32 InS
 #endif // WITH_FREETYPE
 }
 
-bool FSlateFontRenderer::CanLoadCharacter(const FFontData& InFontData, TCHAR Char, EFontFallback MaxFallbackLevel) const
+bool FSlateFontRenderer::CanLoadCodepoint(const FFontData& InFontData, const UTF32CHAR InCodepoint, EFontFallback MaxFallbackLevel) const
 {
 	bool bReturnVal = false;
 
 #if WITH_FREETYPE
-	const FFreeTypeFaceGlyphData FaceGlyphData = GetFontFaceForCharacter(InFontData, Char, MaxFallbackLevel);
+	const FFreeTypeFaceGlyphData FaceGlyphData = GetFontFaceForCodepoint(InFontData, InCodepoint, MaxFallbackLevel);
 	bReturnVal = FaceGlyphData.FaceAndMemory.IsValid() && FaceGlyphData.GlyphIndex != 0;
 #endif // WITH_FREETYPE
 
@@ -221,10 +221,10 @@ bool FSlateFontRenderer::CanLoadCharacter(const FFontData& InFontData, TCHAR Cha
 
 #if WITH_FREETYPE
 
-FFreeTypeFaceGlyphData FSlateFontRenderer::GetFontFaceForCharacter(const FFontData& InFontData, TCHAR Char, EFontFallback MaxFallbackLevel) const 
+FFreeTypeFaceGlyphData FSlateFontRenderer::GetFontFaceForCodepoint(const FFontData& InFontData, const UTF32CHAR InCodepoint, EFontFallback MaxFallbackLevel) const
 {
 	FFreeTypeFaceGlyphData ReturnVal;
-	const bool bOverrideFallback = Char == SlateFontRendererUtils::InvalidSubChar;
+	const bool bOverrideFallback = InCodepoint == SlateFontRendererUtils::InvalidSubChar;
 
 	// Try the requested font first
 	{
@@ -232,23 +232,23 @@ FFreeTypeFaceGlyphData FSlateFontRenderer::GetFontFaceForCharacter(const FFontDa
 
 		if (ReturnVal.FaceAndMemory.IsValid())
 		{
-			ReturnVal.GlyphIndex = FT_Get_Char_Index(ReturnVal.FaceAndMemory->GetFace(), Char);
+			ReturnVal.GlyphIndex = FT_Get_Char_Index(ReturnVal.FaceAndMemory->GetFace(), InCodepoint);
 			ReturnVal.CharFallbackLevel = EFontFallback::FF_NoFallback;
 		}
 	}
 
 	// If the requested glyph doesn't exist, use the localization fallback font
-	if (!ReturnVal.FaceAndMemory.IsValid() || (Char != 0 && ReturnVal.GlyphIndex == 0))
+	if (!ReturnVal.FaceAndMemory.IsValid() || (InCodepoint != 0 && ReturnVal.GlyphIndex == 0))
 	{
 		const bool bCanFallback = bOverrideFallback || MaxFallbackLevel >= EFontFallback::FF_LocalizedFallback;
 
 		if (bCanFallback && FLegacySlateFontInfoCache::Get().IsLocalizedFallbackFontAvailable())
 		{
-			ReturnVal.FaceAndMemory = CompositeFontCache->GetFontFace(FLegacySlateFontInfoCache::Get().GetLocalizedFallbackFontData(FLegacySlateFontInfoCache::FFallbackContext(&InFontData, Char)));
+			ReturnVal.FaceAndMemory = CompositeFontCache->GetFontFace(FLegacySlateFontInfoCache::Get().GetLocalizedFallbackFontData(FLegacySlateFontInfoCache::FFallbackContext(&InFontData, InCodepoint)));
 
 			if (ReturnVal.FaceAndMemory.IsValid())
 			{	
-				ReturnVal.GlyphIndex = FT_Get_Char_Index(ReturnVal.FaceAndMemory->GetFace(), Char);
+				ReturnVal.GlyphIndex = FT_Get_Char_Index(ReturnVal.FaceAndMemory->GetFace(), InCodepoint);
 
 				if (ReturnVal.GlyphIndex != 0)
 				{
@@ -260,17 +260,17 @@ FFreeTypeFaceGlyphData FSlateFontRenderer::GetFontFaceForCharacter(const FFontDa
 	}
 
 	// If the requested glyph doesn't exist, use the last resort fallback font
-	if (!ReturnVal.FaceAndMemory.IsValid() || (Char != 0 && ReturnVal.GlyphIndex == 0))
+	if (!ReturnVal.FaceAndMemory.IsValid() || (InCodepoint != 0 && ReturnVal.GlyphIndex == 0))
 	{
 		const bool bCanFallback = bOverrideFallback || MaxFallbackLevel >= EFontFallback::FF_LastResortFallback;
 
 		if (bCanFallback && FLegacySlateFontInfoCache::Get().IsLastResortFontAvailable())
 		{
-			ReturnVal.FaceAndMemory = CompositeFontCache->GetFontFace(FLegacySlateFontInfoCache::Get().GetLastResortFontData(FLegacySlateFontInfoCache::FFallbackContext(&InFontData, Char)));
+			ReturnVal.FaceAndMemory = CompositeFontCache->GetFontFace(FLegacySlateFontInfoCache::Get().GetLastResortFontData(FLegacySlateFontInfoCache::FFallbackContext(&InFontData, InCodepoint)));
 
 			if (ReturnVal.FaceAndMemory.IsValid())
 			{
-				ReturnVal.GlyphIndex = FT_Get_Char_Index(ReturnVal.FaceAndMemory->GetFace(), Char);
+				ReturnVal.GlyphIndex = FT_Get_Char_Index(ReturnVal.FaceAndMemory->GetFace(), InCodepoint);
 
 				if (ReturnVal.GlyphIndex != 0)
 				{
@@ -282,7 +282,7 @@ FFreeTypeFaceGlyphData FSlateFontRenderer::GetFontFaceForCharacter(const FFontDa
 	}
 
 	// Found an invalid glyph?
-	if (Char != 0 && ReturnVal.GlyphIndex == 0)
+	if (InCodepoint != 0 && ReturnVal.GlyphIndex == 0)
 	{
 		ReturnVal.FaceAndMemory.Reset();
 	}
diff --git a/Engine/Source/Runtime/SlateCore/Private/Fonts/SlateFontRenderer.h b/Engine/Source/Runtime/SlateCore/Private/Fonts/SlateFontRenderer.h
--- a/Engine/Source/Runtime/SlateCore/Private/Fonts/SlateFontRenderer.h
+++ b/Engine/Source/Runtime/SlateCore/Private/Fonts/SlateFontRenderer.h
@@ -107,22 +107,22 @@ public:
 	 * Whether or not the specified character, within the specified font, can be loaded with the specified maximum font fallback level
 	 *
 	 * @param InFontData		Information about the font to load
-	 * @param Char				The character being loaded
+	 * @param InCodepoint		The codepoint being loaded
 	 * @param MaxFallbackLevel	The maximum fallback level to try for the font
 	 * @return					Whether or not the character can be loaded
 	 */
-	bool CanLoadCharacter(const FFontData& InFontData, TCHAR Char, EFontFallback MaxFallbackLevel) const;
+	bool CanLoadCodepoint(const FFontData& InFontData, const UTF32CHAR InCodepoint, EFontFallback MaxFallbackLevel) const;
 
 #if WITH_FREETYPE
 	/**
 	 * Wrapper for GetFontFace, which reverts to fallback or last resort fonts if the face could not be loaded
 	 *
 	 * @param InFontData		Information about the font to load
-	 * @param Char				The character being loaded (required for checking if a fallback font is needed)
+	 * @param InCodepoint		The codepoint being loaded (required for checking if a fallback font is needed)
 	 * @param MaxFallbackLevel	The maximum fallback level to try for the font
 	 * @return					Returns the character font face data
 	 */
-	FFreeTypeFaceGlyphData GetFontFaceForCharacter(const FFontData& InFontData, TCHAR Char, EFontFallback MaxFallbackLevel) const;
+	FFreeTypeFaceGlyphData GetFontFaceForCodepoint(const FFontData& InFontData, const UTF32CHAR InCodepoint, EFontFallback MaxFallbackLevel) const;
 #endif // WITH_FREETYPE
 
 	/** 
diff --git a/Engine/Source/Runtime/SlateCore/Private/Fonts/SlateTextShaper.cpp b/Engine/Source/Runtime/SlateCore/Private/Fonts/SlateTextShaper.cpp
--- a/Engine/Source/Runtime/SlateCore/Private/Fonts/SlateTextShaper.cpp
+++ b/Engine/Source/Runtime/SlateCore/Private/Fonts/SlateTextShaper.cpp
@@ -160,6 +160,12 @@ void FSlateTextShaper::PerformTextShaping(const TCHAR* InText, const int32 InTex
 					// Note: This isn't an exhaustive list, as it omits some "dead" or uncommon languages, and ranges outside the BMP
 					#define RETURN_TRUE_IF_CHAR_WITHIN_RANGE(LOWER, UPPER) if (InChar >= (LOWER) && InChar <= (UPPER)) return true
 
+					// Zero-width joiner
+					if (InChar == TEXT('\u200D'))
+					{
+						return true;
+					}
+
 					// Combining characters
 					RETURN_TRUE_IF_CHAR_WITHIN_RANGE(TEXT('\u0300'), TEXT('\u036F'));
 					RETURN_TRUE_IF_CHAR_WITHIN_RANGE(TEXT('\u1AB0'), TEXT('\u1AFF'));
@@ -258,7 +264,7 @@ FShapedGlyphSequenceRef FSlateTextShaper::FinalizeTextShaping(TArray<FShapedGlyp
 		// Just get info for the null character
 		TCHAR Char = 0;
 		const FFontData& FontData = CompositeFontCache->GetDefaultFontData(InFontInfo);
-		const FFreeTypeFaceGlyphData FaceGlyphData = FontRenderer->GetFontFaceForCharacter(FontData, Char, InFontInfo.FontFallback);
+		const FFreeTypeFaceGlyphData FaceGlyphData = FontRenderer->GetFontFaceForCodepoint(FontData, Char, InFontInfo.FontFallback);
 
 		if (FaceGlyphData.FaceAndMemory.IsValid())
 		{
@@ -320,14 +326,14 @@ void FSlateTextShaper::PerformKerningOnlyTextShaping(const TCHAR* InText, const
 
 			// First try with the actual character
 			float SubFontScalingFactor = 1.0f;
-			const FFontData* FontDataPtr = &CompositeFontCache->GetFontDataForCharacter(InFontInfo, CurrentChar, SubFontScalingFactor);
-			FFreeTypeFaceGlyphData FaceGlyphData = FontRenderer->GetFontFaceForCharacter(*FontDataPtr, CurrentChar, InFontInfo.FontFallback);
+			const FFontData* FontDataPtr = &CompositeFontCache->GetFontDataForCodepoint(InFontInfo, CurrentChar, SubFontScalingFactor);
+			FFreeTypeFaceGlyphData FaceGlyphData = FontRenderer->GetFontFaceForCodepoint(*FontDataPtr, CurrentChar, InFontInfo.FontFallback);
 
 			// If none of our fonts can render that character (as the fallback font may be missing), try again with the fallback character
 			if (!FaceGlyphData.FaceAndMemory.IsValid())
 			{
-				FontDataPtr = &CompositeFontCache->GetFontDataForCharacter(InFontInfo, SlateFontRendererUtils::InvalidSubChar, SubFontScalingFactor);
-				FaceGlyphData = FontRenderer->GetFontFaceForCharacter(*FontDataPtr, SlateFontRendererUtils::InvalidSubChar, InFontInfo.FontFallback);
+				FontDataPtr = &CompositeFontCache->GetFontDataForCodepoint(InFontInfo, SlateFontRendererUtils::InvalidSubChar, SubFontScalingFactor);
+				FaceGlyphData = FontRenderer->GetFontFaceForCodepoint(*FontDataPtr, SlateFontRendererUtils::InvalidSubChar, InFontInfo.FontFallback);
 			}
 
 			if (!RunningFontDataPtr || RunningFontDataPtr != FontDataPtr || RunningFaceAndMemory != FaceGlyphData.FaceAndMemory || RunningSubFontScalingFactor != SubFontScalingFactor)
@@ -435,6 +441,8 @@ void FSlateTextShaper::PerformHarfBuzzTextShaping(const TCHAR* InText, const int
 	TArray<FHarfBuzzTextSequenceEntry, TInlineAllocator<4>> HarfBuzzTextSequence;
 	hb_unicode_funcs_t* HarfBuzzUnicodeFuncs = hb_unicode_funcs_get_default();
 
+	GraphemeBreakIterator->SetString(InText + InTextStart, InTextLen);
+
 	// Step 1) Split the text into sections that are using the same font face (composite fonts may contain different faces for different character ranges)
 	{
 		// Data used while detecting font face boundaries
@@ -462,27 +470,38 @@ void FSlateTextShaper::PerformHarfBuzzTextShaping(const TCHAR* InText, const int
 			}
 		};
 
-		const int32 TextEndIndex = InTextStart + InTextLen;
-		for (; RunningTextIndex < TextEndIndex; ++RunningTextIndex)
+		// Process the text as its grapheme clusters, and have each cluster use the font of its first codepoint (eg, for Emoji combination)
+		for (int32 PreviousBreak = 0, CurrentBreak = 0; (CurrentBreak = GraphemeBreakIterator->MoveToNext()) != INDEX_NONE; PreviousBreak = CurrentBreak)
 		{
-			TCHAR CurrentChar = InText[RunningTextIndex];
+			const int32 ClusterSize = CurrentBreak - PreviousBreak;
+			
+			// Combine any surrogate pairs into a complete codepoint
+			UTF32CHAR CurrentCodepoint = InText[RunningTextIndex];
+			if (ClusterSize > 1)
+			{
+				const int32 NextTextIndex = RunningTextIndex + 1;
+				if (SurrogatePairUtil::IsSurrogatePair(InText[RunningTextIndex], InText[NextTextIndex]))
+				{
+					CurrentCodepoint = ((InText[RunningTextIndex] - 0xD800) << 10) + (InText[NextTextIndex] - 0xDC00) + 0x10000;
+				}
+			}
 
 			// Substitute whitespace characters with spaces since not all fonts support all kinds of whitespace characters (but we don't care since we don't render them anyway)
-			if (FText::IsWhitespace(CurrentChar))
+			if (FText::IsWhitespace(CurrentCodepoint))
 			{
-				CurrentChar = TEXT(' ');
+				CurrentCodepoint = TEXT(' ');
 			}
 
 			// First try with the actual character
 			float SubFontScalingFactor = 1.0f;
-			const FFontData* FontDataPtr = &CompositeFontCache->GetFontDataForCharacter(InFontInfo, CurrentChar, SubFontScalingFactor);
-			FFreeTypeFaceGlyphData FaceGlyphData = FontRenderer->GetFontFaceForCharacter(*FontDataPtr, CurrentChar, InFontInfo.FontFallback);
+			const FFontData* FontDataPtr = &CompositeFontCache->GetFontDataForCodepoint(InFontInfo, CurrentCodepoint, SubFontScalingFactor);
+			FFreeTypeFaceGlyphData FaceGlyphData = FontRenderer->GetFontFaceForCodepoint(*FontDataPtr, CurrentCodepoint, InFontInfo.FontFallback);
 
 			// If none of our fonts can render that character (as the fallback font may be missing), try again with the fallback character
 			if (!FaceGlyphData.FaceAndMemory.IsValid())
 			{
-				FontDataPtr = &CompositeFontCache->GetFontDataForCharacter(InFontInfo, SlateFontRendererUtils::InvalidSubChar, SubFontScalingFactor);
-				FaceGlyphData = FontRenderer->GetFontFaceForCharacter(*FontDataPtr, SlateFontRendererUtils::InvalidSubChar, InFontInfo.FontFallback);
+				FontDataPtr = &CompositeFontCache->GetFontDataForCodepoint(InFontInfo, SlateFontRendererUtils::InvalidSubChar, SubFontScalingFactor);
+				FaceGlyphData = FontRenderer->GetFontFaceForCodepoint(*FontDataPtr, SlateFontRendererUtils::InvalidSubChar, InFontInfo.FontFallback);
 			}
 
 			if (!RunningFontDataPtr || RunningFontDataPtr != FontDataPtr || RunningFaceAndMemory != FaceGlyphData.FaceAndMemory || RunningSubFontScalingFactor != SubFontScalingFactor)
@@ -494,6 +513,8 @@ void FSlateTextShaper::PerformHarfBuzzTextShaping(const TCHAR* InText, const int
 				RunningFaceAndMemory = FaceGlyphData.FaceAndMemory;
 				RunningSubFontScalingFactor = SubFontScalingFactor;
 			}
+
+			RunningTextIndex += ClusterSize;
 		}
 
 		AppendPendingFontDataToSequence();
@@ -670,88 +691,111 @@ void FSlateTextShaper::PerformHarfBuzzTextShaping(const TCHAR* InText, const int
 		hb_buffer_destroy(HarfBuzzTextBuffer);
 	}
 
-	const int32 NumGlyphsRendered = OutGlyphsToRender.Num() - InitialNumGlyphsToRender;
-
-	// Step 4) Count the characters that belong to each glyph if they haven't already been set
-	if (NumGlyphsRendered > 0)
+	// Step 4) Count the characters and grapheme clusters that belong to each glyph if they haven't already been set
 	{
-		const int32 CurrentNumGlyphsToRender = OutGlyphsToRender.Num();
+		const int32 NumGlyphsRendered = OutGlyphsToRender.Num() - InitialNumGlyphsToRender;
+		if (NumGlyphsRendered > 0)
+		{
+			auto ConditionalUpdateGlyphCountsForRange = [this, InTextStart](FShapedGlyphEntry& ShapedGlyphEntry, const int32 TextStartIndex, const int32 TextEndIndex)
+			{
+				check(TextStartIndex <= TextEndIndex);
 
-		// The glyphs in the array are in render order, so LTR and RTL text use different start and end points in the source string
-		const int32 FirstGlyphPrevSourceIndex = (InTextDirection == TextBiDi::ETextDirection::LeftToRight) ? InTextStart - 1 : InTextStart + InTextLen;
-		const int32 LastGlyphNextSourceIndex  = (InTextDirection == TextBiDi::ETextDirection::LeftToRight) ? InTextStart + InTextLen : InTextStart - 1;
+				if (ShapedGlyphEntry.NumCharactersInGlyph == 0 && ShapedGlyphEntry.NumGraphemeClustersInGlyph == 0)
+				{
+					ShapedGlyphEntry.NumCharactersInGlyph = TextEndIndex - TextStartIndex;
 
-		// Start of the loop; process against the "start" of the string range
-		{
-			FShapedGlyphEntry& ShapedGlyphEntry = OutGlyphsToRender[InitialNumGlyphsToRender];
-			ShapedGlyphEntry.NumCharactersInGlyph = FMath::Abs(FirstGlyphPrevSourceIndex - ShapedGlyphEntry.SourceIndex);
-		}
+					if (ShapedGlyphEntry.NumCharactersInGlyph > 0)
+					{
+						const int32 FirstCharacterIndex = TextStartIndex - InTextStart;
+						const int32 LastCharacterIndex = TextEndIndex - InTextStart;
 
-		// Body of the loop; this will process the initial character again, but won't change its value and will walk past its entire cluster
-		for (int32 GlyphToRenderIndex = InitialNumGlyphsToRender; GlyphToRenderIndex < CurrentNumGlyphsToRender;)
-		{
-			FShapedGlyphEntry& ShapedGlyphEntry = OutGlyphsToRender[GlyphToRenderIndex];
+						// Only count grapheme clusters if this glyph starts on a grapheme boundary
+						int32 PreviousBreak = FirstCharacterIndex;
+						{
+							GraphemeBreakIterator->MoveToCandidateAfter(FirstCharacterIndex);
+							PreviousBreak = GraphemeBreakIterator->MoveToPrevious();
+						}
 
-			// Walk forward to find the first glyph in the next cluster; the number of characters in this glyph is the difference between their two source indices
-			int32 NextGlyphToRenderIndex = GlyphToRenderIndex + 1;
-			for (; NextGlyphToRenderIndex < CurrentNumGlyphsToRender; ++NextGlyphToRenderIndex)
-			{
-				const FShapedGlyphEntry& NextShapedGlyphEntry = OutGlyphsToRender[NextGlyphToRenderIndex];
-				if (ShapedGlyphEntry.SourceIndex != NextShapedGlyphEntry.SourceIndex)
-				{
-					break;
+						if (PreviousBreak == FirstCharacterIndex)
+						{
+							int32 CurrentBreak = LastCharacterIndex;
+							for (CurrentBreak = GraphemeBreakIterator->MoveToCandidateAfter(FirstCharacterIndex);
+								CurrentBreak != INDEX_NONE;
+								CurrentBreak = GraphemeBreakIterator->MoveToNext()
+								)
+							{
+								++ShapedGlyphEntry.NumGraphemeClustersInGlyph;
+								if (CurrentBreak >= LastCharacterIndex)
+								{
+									break;
+								}
+							}
+
+							// Only count grapheme clusters if this glyph ends on a grapheme boundary
+							if (CurrentBreak != LastCharacterIndex)
+							{
+								ShapedGlyphEntry.NumGraphemeClustersInGlyph = 0;
+							}
+						}
+					}
 				}
-			}
+			};
 
-			if (NextGlyphToRenderIndex < CurrentNumGlyphsToRender)
+			auto GetNextGlyphToRenderIndex = [&OutGlyphsToRender](const int32 GlyphToRenderIndex) -> int32
 			{
-				FShapedGlyphEntry& NextShapedGlyphEntry = OutGlyphsToRender[NextGlyphToRenderIndex];
+				FShapedGlyphEntry& ShapedGlyphEntry = OutGlyphsToRender[GlyphToRenderIndex];
+				int32 NextGlyphToRenderIndex = GlyphToRenderIndex + 1;
 
-				// For LTR text we update ourself based on the next glyph cluster, for RTL text we update the next glyph cluster based on us
-				FShapedGlyphEntry& ShapedGlyphEntryToUpdate = (InTextDirection == TextBiDi::ETextDirection::LeftToRight) ? ShapedGlyphEntry : NextShapedGlyphEntry;
-				if (ShapedGlyphEntryToUpdate.NumCharactersInGlyph == 0)
+				// Walk forward to find the first glyph in the next cluster; the number of characters in this glyph is the difference between their two source indices
+				for (; NextGlyphToRenderIndex < OutGlyphsToRender.Num(); ++NextGlyphToRenderIndex)
 				{
-					ShapedGlyphEntryToUpdate.NumCharactersInGlyph = FMath::Abs(NextShapedGlyphEntry.SourceIndex - ShapedGlyphEntry.SourceIndex);
+					const FShapedGlyphEntry& NextShapedGlyphEntry = OutGlyphsToRender[NextGlyphToRenderIndex];
+					if (ShapedGlyphEntry.SourceIndex != NextShapedGlyphEntry.SourceIndex)
+					{
+						break;
+					}
 				}
-			}
 
-			GlyphToRenderIndex = NextGlyphToRenderIndex;
-		}
+				return NextGlyphToRenderIndex;
+			};
 
-		// End of the loop; process against the "end" of the string range (RTL text is implicitly handled as part of the loop above)
-		if (InTextDirection == TextBiDi::ETextDirection::LeftToRight)
-		{
-			FShapedGlyphEntry& ShapedGlyphEntry = OutGlyphsToRender[CurrentNumGlyphsToRender - 1];
-			ShapedGlyphEntry.NumCharactersInGlyph = FMath::Abs(LastGlyphNextSourceIndex - ShapedGlyphEntry.SourceIndex);
-		}
-	}
+			// The glyphs in the array are in render order, so LTR and RTL text use different start and end points in the source string
+			if (InTextDirection == TextBiDi::ETextDirection::LeftToRight)
+			{
+				for (int32 GlyphToRenderIndex = InitialNumGlyphsToRender; GlyphToRenderIndex < OutGlyphsToRender.Num();)
+				{
+					FShapedGlyphEntry& ShapedGlyphEntry = OutGlyphsToRender[GlyphToRenderIndex];
 
-	// Step 5) Count the grapheme clusters for any entries that haven't been set yet
-	if (NumGlyphsRendered > 0)
-	{
-		GraphemeBreakIterator->SetString(InText + InTextStart, InTextLen);
+					const int32 NextGlyphToRenderIndex = GetNextGlyphToRenderIndex(GlyphToRenderIndex);
+					if (NextGlyphToRenderIndex < OutGlyphsToRender.Num())
+					{
+						const FShapedGlyphEntry& NextShapedGlyphEntry = OutGlyphsToRender[NextGlyphToRenderIndex];
+						ConditionalUpdateGlyphCountsForRange(ShapedGlyphEntry, ShapedGlyphEntry.SourceIndex, NextShapedGlyphEntry.SourceIndex);
+					}
+					else
+					{
+						ConditionalUpdateGlyphCountsForRange(ShapedGlyphEntry, ShapedGlyphEntry.SourceIndex, InTextStart + InTextLen);
+					}
 
-		const int32 CurrentNumGlyphsToRender = OutGlyphsToRender.Num();
-		for (int32 GlyphToRenderIndex = InitialNumGlyphsToRender; GlyphToRenderIndex < CurrentNumGlyphsToRender; ++GlyphToRenderIndex)
-		{
-			FShapedGlyphEntry& ShapedGlyphEntry = OutGlyphsToRender[GlyphToRenderIndex];
-			if (ShapedGlyphEntry.NumCharactersInGlyph > 0 && ShapedGlyphEntry.NumGraphemeClustersInGlyph == 0)
+					GlyphToRenderIndex = NextGlyphToRenderIndex;
+				}
+			}
+			else
 			{
-				const int32 FirstCharacterIndex = ShapedGlyphEntry.SourceIndex - InTextStart;
-				const int32 LastCharacterIndex = (ShapedGlyphEntry.SourceIndex + ShapedGlyphEntry.NumCharactersInGlyph) - InTextStart;
-
-				for (int32 GraphemeIndex = GraphemeBreakIterator->MoveToCandidateAfter(FirstCharacterIndex);
-					GraphemeIndex != INDEX_NONE && GraphemeIndex <= LastCharacterIndex;
-					GraphemeIndex = GraphemeBreakIterator->MoveToNext()
-					)
+				int32 PreviousSourceIndex = InTextStart + InTextLen;
+				for (int32 GlyphToRenderIndex = InitialNumGlyphsToRender; GlyphToRenderIndex < OutGlyphsToRender.Num();)
 				{
-					++ShapedGlyphEntry.NumGraphemeClustersInGlyph;
+					FShapedGlyphEntry& ShapedGlyphEntry = OutGlyphsToRender[GlyphToRenderIndex];
+
+					ConditionalUpdateGlyphCountsForRange(ShapedGlyphEntry, ShapedGlyphEntry.SourceIndex, PreviousSourceIndex);
+					GlyphToRenderIndex = GetNextGlyphToRenderIndex(GlyphToRenderIndex);
+					PreviousSourceIndex = ShapedGlyphEntry.SourceIndex;
 				}
 			}
 		}
-
-		GraphemeBreakIterator->ClearString();
 	}
+
+	GraphemeBreakIterator->ClearString();
 }
 
 #endif // WITH_HARFBUZZ
diff --git a/Engine/Source/Runtime/SlateCore/Public/Fonts/FontCache.h b/Engine/Source/Runtime/SlateCore/Public/Fonts/FontCache.h
--- a/Engine/Source/Runtime/SlateCore/Public/Fonts/FontCache.h
+++ b/Engine/Source/Runtime/SlateCore/Public/Fonts/FontCache.h
@@ -762,15 +762,15 @@ public:
 	const FFontData& GetDefaultFontData( const FSlateFontInfo& InFontInfo ) const;
 
 	/**
-	 * Returns the font to use from the typeface associated with the given character
+	 * Returns the font to use from the typeface associated with the given codepoint
 	 *
 	 * @param InFontInfo		A descriptor of the font to get the typeface for
-	 * @param InChar			The character to get the typeface associated with
+	 * @param InCodepoint		The codepoint to get the typeface associated with
 	 * @param OutScalingFactor	The scaling factor applied to characters rendered with the given font
 	 * 
 	 * @return The raw font data
 	 */
-	const FFontData& GetFontDataForCharacter( const FSlateFontInfo& InFontInfo, const TCHAR InChar, float& OutScalingFactor ) const;
+	const FFontData& GetFontDataForCodepoint( const FSlateFontInfo& InFontInfo, const UTF32CHAR InCodepoint, float& OutScalingFactor ) const;
 
 	/**
 	 * Returns the height of the largest character in the font. 
