Added support for rendering TTF/OTF fonts containing bitmaps, including color fonts (like Emoji)

Note: Bitmap fonts require FreeType 2.10, and support is compiled out if using an older version.

Font Rendering:
 - ApplySizeAndScale now computes the desired pixel size manually (using our desired DPI) rather than calling FT_Set_Char_Size.
   - This also includes the font scale, which avoids the call to FT_Set_Transform and lets us remove some workarounds for inconsistent scaling from FreeType.
   - For scalable fonts the desired pixel size is passed to FT_Set_Pixel_Sizes.
   - For bitmap fonts the desired pixel size is used to find the most suitable strike size (set via FT_Select_Size) and the ratio needed to scale that strike to the desired size (see GetBitmapAtlasScale and GetBitmapRenderScale).
 - AppendGlyphFlags has been updated to load color data (when available), and to only exclude bitmap data for fonts that also have scalable data.
 - FSlateFontRenderer::GetRenderDataInternal has been updated to handle color bitmap rendering (BGRA, sRGB) in addition to the existing grayscale rendering.
   - It will also handle scaling down bitmap font strikes that are larger than the desired size, which saves space in the atlas (see GetBitmapAtlasScale).
   - It won't scale up any bitmaps, as that is handled by the transform applied when batching the glyph quads for rendering (see GetBitmapRenderScale).
 - Added support for disabling outline and tint for bitmap fonts (in Slate and Canvas).
 - Added support for batching both grayscale (8-bit alpha-only) and color (8-bit per-channel BGRA) font quads from a font texture atlas (in RHI, D3D, and OpenGL).
 - Removed unused data from FCharacterRenderData.

diff --git a/Engine/Shaders/Private/SlateElementPixelShader.usf b/Engine/Shaders/Private/SlateElementPixelShader.usf
--- a/Engine/Shaders/Private/SlateElementPixelShader.usf
+++ b/Engine/Shaders/Private/SlateElementPixelShader.usf
@@ -48,7 +48,7 @@ half3 GammaCorrect(half3 InColor)
 #if USE_MATERIALS
 void FillTexCoords( inout FMaterialPixelParameters Parameters, VertexToPixelInterpolants InVertex, float2 InDefaultUV )
 {
-#if SHADER_TYPE == ST_Font
+#if SHADER_TYPE == ST_GrayscaleFont || SHADER_TYPE == ST_ColorFont
 	#if NUM_MATERIAL_TEXCOORDS == 1
 		// Note: The first channel is used to look up into the texture atlas. The second channel is the custom 0-1 UV set for the text run
 		Parameters.TexCoords[0] = GetUV( InVertex, 1 );
@@ -155,7 +155,7 @@ half4 GetColor( VertexToPixelInterpolants InVertex, float2 UV )
 	return FinalColor;
 }
 
-#if SHADER_TYPE == ST_Font
+#if SHADER_TYPE == ST_GrayscaleFont || SHADER_TYPE == ST_ColorFont
 half4 GetFontElementColor( VertexToPixelInterpolants InVertex )
 {
 	float2 AtlasTextureCoordinates = GetUV( InVertex, 0 );
@@ -166,7 +166,11 @@ half4 GetFontElementColor( VertexToPixelInterpolants InVertex )
 #endif
 	half4 OutColor = BaseColor;
 
+#if SHADER_TYPE == ST_GrayscaleFont
 	OutColor.a *= Texture2DSample_A8( ElementTexture, ElementTextureSampler, AtlasTextureCoordinates );
+#else
+	OutColor *= Texture2DSample( ElementTexture, ElementTextureSampler, AtlasTextureCoordinates );
+#endif
 
 	return OutColor;
 }
@@ -250,7 +254,7 @@ half4 Main( VertexToPixelInterpolants VIn ) : SV_Target0
 	OutColor = GetDefaultElementColor( VIn );
 #elif SHADER_TYPE == ST_Border
 	OutColor = GetBorderElementColor( VIn );
-#elif SHADER_TYPE == ST_Font
+#elif SHADER_TYPE == ST_GrayscaleFont || SHADER_TYPE == ST_ColorFont
 	OutColor = GetFontElementColor( VIn );
 #else
 	OutColor = GetLineSegmentElementColor( VIn );
diff --git a/Engine/Shaders/Private/SlateShaderCommon.ush b/Engine/Shaders/Private/SlateShaderCommon.ush
--- a/Engine/Shaders/Private/SlateShaderCommon.ush
+++ b/Engine/Shaders/Private/SlateShaderCommon.ush
@@ -24,11 +24,12 @@
 
 
 /** Shader types (mirrored from ESlateShader::Type in RenderingCommon.h */
-#define ST_Default		0
-#define ST_Border		1
-#define ST_Font			2
-#define ST_LineSegment	3
-#define ST_Custom		4
+#define ST_Default			0
+#define ST_Border			1
+#define ST_GrayscaleFont	2
+#define ST_ColorFont		3
+#define ST_LineSegment		4
+#define ST_Custom			5
 
 struct VertexToPixelInterpolants
 {
diff --git a/Engine/Shaders/StandaloneRenderer/D3D/SlateElementPixelShader.hlsl b/Engine/Shaders/StandaloneRenderer/D3D/SlateElementPixelShader.hlsl
--- a/Engine/Shaders/StandaloneRenderer/D3D/SlateElementPixelShader.hlsl
+++ b/Engine/Shaders/StandaloneRenderer/D3D/SlateElementPixelShader.hlsl
@@ -5,8 +5,9 @@
 // Shader types
 #define ESlateShader::Default		0
 #define ESlateShader::Border		1
-#define ESlateShader::Font			2
-#define ESlateShader::LineSegment	3
+#define ESlateShader::GrayscaleFont	2
+#define ESlateShader::ColorFont		3
+#define ESlateShader::LineSegment	4
 
 Texture2D ElementTexture;
 SamplerState ElementTextureSampler;
@@ -53,7 +54,7 @@ float3 GammaCorrect(float3 InColor)
 	return CorrectedColor;
 }
 
-float4 GetFontElementColor( VertexOut InVertex )
+float4 GetGrayscaleFontElementColor( VertexOut InVertex )
 {
 	float4 OutColor = InVertex.Color;
 
@@ -62,6 +63,15 @@ float4 GetFontElementColor( VertexOut InVertex )
 	return OutColor;
 }
 
+float4 GetColorFontElementColor(VertexOut InVertex)
+{
+	float4 OutColor = InVertex.Color;
+
+	OutColor *= ElementTexture.Sample(ElementTextureSampler, InVertex.TextureCoordinates.xy);
+
+	return OutColor;
+}
+
 float4 GetColor( VertexOut InVertex, float2 UV )
 {
 	float4 FinalColor;
@@ -150,9 +160,13 @@ float4 Main( VertexOut InVertex ) : SV_Target
 	{
 		OutColor = GetBorderElementColor( InVertex );
 	}
-	else if( ShaderType == ESlateShader::Font )
+	else if( ShaderType == ESlateShader::GrayscaleFont )
+	{
+		OutColor = GetGrayscaleFontElementColor( InVertex );
+	}
+	else if (ShaderType == ESlateShader::ColorFont)
 	{
-		OutColor = GetFontElementColor( InVertex );
+		OutColor = GetColorFontElementColor(InVertex);
 	}
 	else
 	{
diff --git a/Engine/Shaders/StandaloneRenderer/OpenGL/SlateElementPixelShader.glsl b/Engine/Shaders/StandaloneRenderer/OpenGL/SlateElementPixelShader.glsl
--- a/Engine/Shaders/StandaloneRenderer/OpenGL/SlateElementPixelShader.glsl
+++ b/Engine/Shaders/StandaloneRenderer/OpenGL/SlateElementPixelShader.glsl
@@ -13,10 +13,11 @@ precision highp float;
 #endif // USE_709
 
 // Shader types
-#define ST_Default	0
-#define ST_Border	1
-#define ST_Font		2
-#define ST_Line		3
+#define ST_Default			0
+#define ST_Border			1
+#define ST_GrayscaleFont	2
+#define ST_ColorFont		3
+#define ST_Line				4
 
 /** Display gamma x:gamma curve adjustment, y:inverse gamma (1/GEngine->DisplayGamma) */
 uniform vec2 GammaValues = vec2(1, 1/2.2);
@@ -130,7 +131,7 @@ vec3 GammaCorrect(vec3 InColor)
 	return CorrectedColor;
 }
 
-vec4 GetFontElementColor()
+vec4 GetGrayscaleFontElementColor()
 {
 	vec4 OutColor = Color;
 #if PLATFORM_LINUX
@@ -142,6 +143,15 @@ vec4 GetFontElementColor()
 	return OutColor;
 }
 
+vec4 GetColorFontElementColor()
+{
+	vec4 OutColor = Color;
+
+	OutColor *= texture2D(ElementTexture, TexCoords.xy);
+
+	return OutColor;
+}
+
 vec4 GetDefaultElementColor()
 {
 	vec4 OutColor = Color;
@@ -236,9 +246,13 @@ void main()
 	{
 		OutColor = GetBorderElementColor();
 	}
-	else if( ShaderType == ST_Font )
+	else if( ShaderType == ST_GrayscaleFont )
+	{
+		OutColor = GetGrayscaleFontElementColor();
+	}
+	else if( ShaderType == ST_ColorFont )
 	{
-		OutColor = GetFontElementColor();
+		OutColor = GetColorFontElementColor();
 	}
 	else
 	{
diff --git a/Engine/Source/Developer/SlateReflector/Private/Widgets/SAtlasVisualizer.cpp b/Engine/Source/Developer/SlateReflector/Private/Widgets/SAtlasVisualizer.cpp
--- a/Engine/Source/Developer/SlateReflector/Private/Widgets/SAtlasVisualizer.cpp
+++ b/Engine/Source/Developer/SlateReflector/Private/Widgets/SAtlasVisualizer.cpp
@@ -241,14 +241,10 @@ void SAtlasVisualizer::Construct( const FArguments& InArgs )
 	AtlasProvider = InArgs._AtlasProvider;
 	check(AtlasProvider);
 
-	const bool IsAlphaOnly = AtlasProvider->IsAtlasPageResourceAlphaOnly();
-
 	SelectedAtlasPage = 0;
 	bDisplayCheckerboard = false;
 
-	const FIntPoint DesiredViewportSize = GetSize();
 	TSharedPtr<SViewport> Viewport;
-
 	ChildSlot
 	[
 		SNew(SBorder)
@@ -290,7 +286,7 @@ void SAtlasVisualizer::Construct( const FArguments& InArgs )
 				.Padding(0.0,2.0f,2.0f,2.0f)
 				[
 					SNew( STextBlock )
-					.Text( FText::Format( LOCTEXT("PageSizeXY", "({0} x {1})"), FText::AsNumber(DesiredViewportSize.X), FText::AsNumber(DesiredViewportSize.Y) ) )
+					.Text( this, &SAtlasVisualizer::GetViewportSizeText )
 				]
 				+ SHorizontalBox::Slot()
 				.Padding(20.0f,2.0f)
@@ -298,7 +294,7 @@ void SAtlasVisualizer::Construct( const FArguments& InArgs )
 				.VAlign( VAlign_Center )
 				[
 					SNew( SCheckBox )
-					.Visibility( (IsAlphaOnly) ? EVisibility::Collapsed : EVisibility::Visible )
+					.Visibility( this, &SAtlasVisualizer::OnGetDisplayCheckerboardVisibility )
 					.OnCheckStateChanged( this, &SAtlasVisualizer::OnDisplayCheckerboardStateChanged )
 					.IsChecked( this, &SAtlasVisualizer::OnGetCheckerboardState )
 					.Content()
@@ -362,7 +358,7 @@ void SAtlasVisualizer::Construct( const FArguments& InArgs )
 					+ SOverlay::Slot()
 					[
 						SAssignNew( Viewport, SViewport )
-						.ViewportSize(FVector2D(DesiredViewportSize.X, DesiredViewportSize.Y))
+						.ViewportSize(this, &SAtlasVisualizer::GetViewportWidgetSize)
 						.IgnoreTextureAlpha(false)
 						.EnableBlending(true)
 						.PreMultipliedAlpha(false)
@@ -377,7 +373,12 @@ void SAtlasVisualizer::Construct( const FArguments& InArgs )
 
 FIntPoint SAtlasVisualizer::GetSize() const
 {
-	return AtlasProvider->GetAtlasPageSize();
+	if (FSlateShaderResource* TargetTexture = GetViewportRenderTargetTexture())
+	{
+		return FIntPoint(TargetTexture->GetWidth(), TargetTexture->GetHeight());
+	}
+
+	return FIntPoint(0, 0);
 }
 
 bool SAtlasVisualizer::RequiresVsync() const
@@ -397,7 +398,24 @@ FSlateShaderResource* SAtlasVisualizer::GetViewportRenderTargetTexture() const
 
 bool SAtlasVisualizer::IsViewportTextureAlphaOnly() const
 {
-	return AtlasProvider->IsAtlasPageResourceAlphaOnly();
+	if (SelectedAtlasPage < AtlasProvider->GetNumAtlasPages())
+	{
+		return AtlasProvider->IsAtlasPageResourceAlphaOnly(SelectedAtlasPage);
+	}
+
+	return false;
+}
+
+FText SAtlasVisualizer::GetViewportSizeText() const
+{
+	const FIntPoint ViewportSize = GetSize();
+	return FText::Format(LOCTEXT("PageSizeXY", "({0} x {1})"), ViewportSize.X, ViewportSize.Y);
+}
+
+FVector2D SAtlasVisualizer::GetViewportWidgetSize() const
+{
+	const FIntPoint ViewportSize = GetSize();
+	return FVector2D(ViewportSize.X, ViewportSize.Y);
 }
 
 FText SAtlasVisualizer::GetZoomLevelPercentText() const
@@ -445,6 +463,11 @@ FReply SAtlasVisualizer::OnActualSizeClicked()
 	return FReply::Handled();
 }
 
+EVisibility SAtlasVisualizer::OnGetDisplayCheckerboardVisibility() const
+{
+	return IsViewportTextureAlphaOnly() ? EVisibility::Collapsed : EVisibility::Visible;
+}
+
 void SAtlasVisualizer::OnDisplayCheckerboardStateChanged( ECheckBoxState NewState )
 {
 	bDisplayCheckerboard = NewState == ECheckBoxState::Checked;
@@ -457,7 +480,7 @@ ECheckBoxState SAtlasVisualizer::OnGetCheckerboardState() const
 
 EVisibility SAtlasVisualizer::OnGetCheckerboardVisibility() const
 {
-	return bDisplayCheckerboard ? EVisibility::Visible : EVisibility::Collapsed;
+	return bDisplayCheckerboard && !IsViewportTextureAlphaOnly() ? EVisibility::Visible : EVisibility::Collapsed;
 }
 
 void SAtlasVisualizer::OnComboOpening()
diff --git a/Engine/Source/Developer/SlateReflector/Private/Widgets/SAtlasVisualizer.h b/Engine/Source/Developer/SlateReflector/Private/Widgets/SAtlasVisualizer.h
--- a/Engine/Source/Developer/SlateReflector/Private/Widgets/SAtlasVisualizer.h
+++ b/Engine/Source/Developer/SlateReflector/Private/Widgets/SAtlasVisualizer.h
@@ -29,15 +29,18 @@ public:
 	virtual FIntPoint GetSize() const override;
 	virtual bool RequiresVsync() const override;
 	virtual FSlateShaderResource* GetViewportRenderTargetTexture() const override;
-	bool IsViewportTextureAlphaOnly() const override;
+	virtual bool IsViewportTextureAlphaOnly() const override;
 
 	void Construct( const FArguments& InArgs );
 
 protected:
+	FText GetViewportSizeText() const;
+	FVector2D GetViewportWidgetSize() const;
 	FText GetZoomLevelPercentText() const;
 	void OnFitToWindowStateChanged( ECheckBoxState NewState );
 	ECheckBoxState OnGetFitToWindowState() const;
 	FReply OnActualSizeClicked();
+	EVisibility OnGetDisplayCheckerboardVisibility() const;
 	void OnDisplayCheckerboardStateChanged( ECheckBoxState NewState );
 	ECheckBoxState OnGetCheckerboardState() const;
 	EVisibility OnGetCheckerboardVisibility() const;
diff --git a/Engine/Source/Developer/StandaloneRenderer/Private/OpenGL/SlateOpenGLRenderer.cpp b/Engine/Source/Developer/StandaloneRenderer/Private/OpenGL/SlateOpenGLRenderer.cpp
--- a/Engine/Source/Developer/StandaloneRenderer/Private/OpenGL/SlateOpenGLRenderer.cpp
+++ b/Engine/Source/Developer/StandaloneRenderer/Private/OpenGL/SlateOpenGLRenderer.cpp
@@ -21,20 +21,24 @@ public:
 	{
 	}
 
-	virtual FIntPoint GetAtlasSize() const override
+	virtual FIntPoint GetAtlasSize(const bool InIsGrayscale) const override
 	{
-		return FIntPoint(TextureSize, TextureSize);
+		return InIsGrayscale 
+			? FIntPoint(GrayscaleTextureSize, GrayscaleTextureSize)
+			: FIntPoint(ColorTextureSize, ColorTextureSize);
 	}
 
-	virtual TSharedRef<FSlateFontAtlas> CreateFontAtlas() const override
+	virtual TSharedRef<FSlateFontAtlas> CreateFontAtlas(const bool InIsGrayscale) const override
 	{
-		TSharedRef<FSlateFontTextureOpenGL> FontTexture = MakeShareable( new FSlateFontTextureOpenGL( TextureSize, TextureSize ) );
+		const FIntPoint AtlasSize = GetAtlasSize(InIsGrayscale);
+
+		TSharedRef<FSlateFontTextureOpenGL> FontTexture = MakeShareable(new FSlateFontTextureOpenGL(AtlasSize.X, AtlasSize.Y, InIsGrayscale));
 		FontTexture->CreateFontTexture();
 
 		return FontTexture;
 	}
 
-	virtual TSharedPtr<ISlateFontTexture> CreateNonAtlasedTexture(const uint32 InWidth, const uint32 InHeight, const TArray<uint8>& InRawData) const override
+	virtual TSharedPtr<ISlateFontTexture> CreateNonAtlasedTexture(const uint32 InWidth, const uint32 InHeight, const bool InIsGrayscale, const TArray<uint8>& InRawData) const override
 	{
 		return nullptr;
 	}
@@ -42,7 +46,8 @@ public:
 private:
 
 	/** Size of each font texture, width and height */
-	static const uint32 TextureSize = 1024;
+	static const uint32 GrayscaleTextureSize = 1024;
+	static const uint32 ColorTextureSize = 256;
 };
 
 TSharedRef<FSlateFontServices> CreateOpenGLFontServices()
diff --git a/Engine/Source/Developer/StandaloneRenderer/Private/OpenGL/SlateOpenGLTextures.cpp b/Engine/Source/Developer/StandaloneRenderer/Private/OpenGL/SlateOpenGLTextures.cpp
--- a/Engine/Source/Developer/StandaloneRenderer/Private/OpenGL/SlateOpenGLTextures.cpp
+++ b/Engine/Source/Developer/StandaloneRenderer/Private/OpenGL/SlateOpenGLTextures.cpp
@@ -110,11 +110,10 @@ void FSlateOpenGLTexture::UpdateTextureRaw(const void* Buffer, const FIntRect& D
 }
 
 
-FSlateFontTextureOpenGL::FSlateFontTextureOpenGL( uint32 Width, uint32 Height )
-	: FSlateFontAtlas( Width, Height ) 
+FSlateFontTextureOpenGL::FSlateFontTextureOpenGL(uint32 Width, uint32 Height, const bool InIsGrayscale)
+	: FSlateFontAtlas(Width, Height, InIsGrayscale) 
 	, FontTexture(nullptr)
 {
-
 }
 
 FSlateFontTextureOpenGL::~FSlateFontTextureOpenGL()
@@ -133,23 +132,16 @@ void FSlateFontTextureOpenGL::CreateFontTexture()
 	glTexParameteri(GL_TEXTURE_2D, GL_TEXTURE_MIN_FILTER, GL_NEAREST);
 	glTexParameteri(GL_TEXTURE_2D, GL_TEXTURE_MAG_FILTER, GL_NEAREST);
 
-	// Font textures use an alpha texture only
-	GLint Format =
-#if USE_DEPRECATED_OPENGL_FUNCTIONALITY
-		GL_ALPHA
-#else
-		GL_RED
-#endif // USE_DEPRECATED_OPENGL_FUNCTIONALITY
-	;
+	GLint InternalFormat = GetGLTextureInternalFormat();
+	GLint Format = GetGLTextureFormat();
+	GLint Type = GetGLTextureType();
 
 	// Upload the data to the texture
-	glTexImage2D( GL_TEXTURE_2D, 0, Format, AtlasWidth, AtlasHeight, 0, Format, GL_UNSIGNED_BYTE, NULL );
+	glTexImage2D( GL_TEXTURE_2D, 0, InternalFormat, AtlasWidth, AtlasHeight, 0, Format, Type, NULL );
 
 	// Create a new slate texture for use in rendering
 	FontTexture = new FSlateOpenGLTexture( AtlasWidth, AtlasHeight );
 	FontTexture->Init( TextureID );
-
-
 }
 
 void FSlateFontTextureOpenGL::ConditionalUpdateTexture()
@@ -161,13 +153,6 @@ void FSlateFontTextureOpenGL::ConditionalUpdateTexture()
 
 		// Completely the texture data each time characters are added
 		glBindTexture(GL_TEXTURE_2D, FontTexture->GetTypedResource() );
-		GLint Format =
-#if USE_DEPRECATED_OPENGL_FUNCTIONALITY
-			GL_ALPHA
-#else
-			GL_RED
-#endif // USE_DEPRECATED_OPENGL_FUNCTIONALITY
-		;
 
 #if PLATFORM_MAC // Make this texture use a DMA'd client storage backing store on OS X, where these extensions always exist
 				 // This avoids a problem on Intel & Nvidia cards that makes characters disappear as well as making the texture updates
@@ -176,7 +161,10 @@ void FSlateFontTextureOpenGL::ConditionalUpdateTexture()
 		glTexParameteri(GL_TEXTURE_2D, GL_TEXTURE_STORAGE_HINT_APPLE, GL_STORAGE_CACHED_APPLE);
 		glPixelStorei(GL_UNPACK_CLIENT_STORAGE_APPLE, GL_TRUE);
 #endif
-		glTexImage2D( GL_TEXTURE_2D, 0, Format, AtlasWidth, AtlasHeight, 0, Format, GL_UNSIGNED_BYTE, AtlasData.GetData() );
+		GLint InternalFormat = GetGLTextureInternalFormat();
+		GLint Format = GetGLTextureFormat();
+		GLint Type = GetGLTextureType();
+		glTexImage2D( GL_TEXTURE_2D, 0, InternalFormat, AtlasWidth, AtlasHeight, 0, Format, Type, AtlasData.GetData() );
 #if PLATFORM_MAC
 		glPixelStorei(GL_UNPACK_CLIENT_STORAGE_APPLE, GL_FALSE);
 #endif
@@ -184,3 +172,38 @@ void FSlateFontTextureOpenGL::ConditionalUpdateTexture()
 		bNeedsUpdate = false;
 	}
 }
+
+GLint FSlateFontTextureOpenGL::GetGLTextureInternalFormat() const
+{
+	if (IsGrayscale())
+	{
+#if USE_DEPRECATED_OPENGL_FUNCTIONALITY
+		return GL_ALPHA;
+#else
+		return GL_RED;
+#endif // USE_DEPRECATED_OPENGL_FUNCTIONALITY
+	}
+	return GL_SRGB8_ALPHA8_EXT;
+}
+
+GLint FSlateFontTextureOpenGL::GetGLTextureFormat() const
+{
+	if (IsGrayscale())
+	{
+#if USE_DEPRECATED_OPENGL_FUNCTIONALITY
+		return GL_ALPHA;
+#else
+		return GL_RED;
+#endif // USE_DEPRECATED_OPENGL_FUNCTIONALITY
+	}
+	return GL_RGBA;
+}
+
+GLint FSlateFontTextureOpenGL::GetGLTextureType() const
+{
+	if (IsGrayscale())
+	{
+		return GL_UNSIGNED_BYTE;
+	}
+	return GL_UNSIGNED_INT_8_8_8_8_REV;
+}
diff --git a/Engine/Source/Developer/StandaloneRenderer/Private/OpenGL/SlateOpenGLTextures.h b/Engine/Source/Developer/StandaloneRenderer/Private/OpenGL/SlateOpenGLTextures.h
--- a/Engine/Source/Developer/StandaloneRenderer/Private/OpenGL/SlateOpenGLTextures.h
+++ b/Engine/Source/Developer/StandaloneRenderer/Private/OpenGL/SlateOpenGLTextures.h
@@ -62,7 +62,7 @@ private:
 class FSlateFontTextureOpenGL : public FSlateFontAtlas
 {
 public:
-	FSlateFontTextureOpenGL( uint32 Width, uint32 Height );
+	FSlateFontTextureOpenGL(uint32 Width, uint32 Height, const bool InIsGrayscale);
 	~FSlateFontTextureOpenGL();
 
 	void CreateFontTexture();
@@ -72,6 +72,10 @@ public:
 	virtual class FSlateShaderResource* GetSlateTexture() override { return FontTexture; }
 	virtual class FTextureResource* GetEngineTexture() override { return nullptr; }
 private:
+	GLint GetGLTextureInternalFormat() const;
+	GLint GetGLTextureFormat() const;
+	GLint GetGLTextureType() const;
+
 	FSlateOpenGLTexture* FontTexture;
 };
 
diff --git a/Engine/Source/Developer/StandaloneRenderer/Private/Windows/D3D/SlateD3DRenderer.cpp b/Engine/Source/Developer/StandaloneRenderer/Private/Windows/D3D/SlateD3DRenderer.cpp
--- a/Engine/Source/Developer/StandaloneRenderer/Private/Windows/D3D/SlateD3DRenderer.cpp
+++ b/Engine/Source/Developer/StandaloneRenderer/Private/Windows/D3D/SlateD3DRenderer.cpp
@@ -83,17 +83,20 @@ public:
 	{
 	}
 
-	virtual FIntPoint GetAtlasSize() const override
+	virtual FIntPoint GetAtlasSize(const bool InIsGrayscale) const override
 	{
-		return FIntPoint(TextureSize, TextureSize);
+		return InIsGrayscale
+			? FIntPoint(GrayscaleTextureSize, GrayscaleTextureSize)
+			: FIntPoint(ColorTextureSize, ColorTextureSize);
 	}
 
-	virtual TSharedRef<FSlateFontAtlas> CreateFontAtlas() const override
+	virtual TSharedRef<FSlateFontAtlas> CreateFontAtlas(const bool InIsGrayscale) const override
 	{
-		return MakeShareable( new FSlateFontAtlasD3D( TextureSize, TextureSize ) );
+		const FIntPoint AtlasSize = GetAtlasSize(InIsGrayscale);
+		return MakeShareable(new FSlateFontAtlasD3D(AtlasSize.X, AtlasSize.Y, InIsGrayscale));
 	}
 
-	virtual TSharedPtr<ISlateFontTexture> CreateNonAtlasedTexture(const uint32 InWidth, const uint32 InHeight, const TArray<uint8>& InRawData) const override
+	virtual TSharedPtr<ISlateFontTexture> CreateNonAtlasedTexture(const uint32 InWidth, const uint32 InHeight, const bool InIsGrayscale, const TArray<uint8>& InRawData) const override
 	{
 		return nullptr;
 	}
@@ -101,7 +104,8 @@ public:
 private:
 
 	/** Size of each font texture, width and height */
-	static const uint32 TextureSize = 1024;
+	static const uint32 GrayscaleTextureSize = 1024;
+	static const uint32 ColorTextureSize = 256;
 };
 
 TSharedRef<FSlateFontServices> CreateD3DFontServices()
diff --git a/Engine/Source/Developer/StandaloneRenderer/Private/Windows/D3D/SlateD3DTextureManager.cpp b/Engine/Source/Developer/StandaloneRenderer/Private/Windows/D3D/SlateD3DTextureManager.cpp
--- a/Engine/Source/Developer/StandaloneRenderer/Private/Windows/D3D/SlateD3DTextureManager.cpp
+++ b/Engine/Source/Developer/StandaloneRenderer/Private/Windows/D3D/SlateD3DTextureManager.cpp
@@ -57,17 +57,12 @@ int32 FSlateD3DTextureManager::GetNumAtlasPages() const
 	return TextureAtlases.Num();
 }
 
-FIntPoint FSlateD3DTextureManager::GetAtlasPageSize() const
-{
-	return FIntPoint(1024, 1024);
-}
-
 FSlateShaderResource* FSlateD3DTextureManager::GetAtlasPageResource(const int32 InIndex) const
 {
 	return TextureAtlases[InIndex]->GetAtlasTexture();
 }
 
-bool FSlateD3DTextureManager::IsAtlasPageResourceAlphaOnly() const
+bool FSlateD3DTextureManager::IsAtlasPageResourceAlphaOnly(const int32 InIndex) const
 {
 	return false;
 }
diff --git a/Engine/Source/Developer/StandaloneRenderer/Private/Windows/D3D/SlateD3DTextureManager.h b/Engine/Source/Developer/StandaloneRenderer/Private/Windows/D3D/SlateD3DTextureManager.h
--- a/Engine/Source/Developer/StandaloneRenderer/Private/Windows/D3D/SlateD3DTextureManager.h
+++ b/Engine/Source/Developer/StandaloneRenderer/Private/Windows/D3D/SlateD3DTextureManager.h
@@ -19,9 +19,8 @@ public:
 
 	/** ISlateAtlasProvider */
 	virtual int32 GetNumAtlasPages() const override;
-	virtual FIntPoint GetAtlasPageSize() const override;
 	virtual FSlateShaderResource* GetAtlasPageResource(const int32 InIndex) const override;
-	virtual bool IsAtlasPageResourceAlphaOnly() const override;
+	virtual bool IsAtlasPageResourceAlphaOnly(const int32 InIndex) const override;
 
 	/**
 	 * Loads and creates rendering resources for all used textures.  
diff --git a/Engine/Source/Developer/StandaloneRenderer/Private/Windows/D3D/SlateD3DTextures.cpp b/Engine/Source/Developer/StandaloneRenderer/Private/Windows/D3D/SlateD3DTextures.cpp
--- a/Engine/Source/Developer/StandaloneRenderer/Private/Windows/D3D/SlateD3DTextures.cpp
+++ b/Engine/Source/Developer/StandaloneRenderer/Private/Windows/D3D/SlateD3DTextures.cpp
@@ -195,11 +195,11 @@ void FSlateTextureAtlasD3D::ConditionalUpdateTexture()
 	// Not yet supported
 }
 
-FSlateFontAtlasD3D::FSlateFontAtlasD3D( uint32 Width, uint32 Height )
-	: FSlateFontAtlas( Width, Height ) 
+FSlateFontAtlasD3D::FSlateFontAtlasD3D(uint32 Width, uint32 Height, const bool InIsGrayscale)
+	: FSlateFontAtlas(Width, Height, InIsGrayscale)
 {
-	FontTexture = new FSlateD3DTexture( Width, Height );
-	FontTexture->Init(DXGI_FORMAT_A8_UNORM, NULL, true, false);
+	FontTexture = new FSlateD3DTexture(Width, Height);
+	FontTexture->Init(InIsGrayscale ? DXGI_FORMAT_A8_UNORM : DXGI_FORMAT_B8G8R8A8_UNORM_SRGB, NULL, true, false);
 }
 
 FSlateFontAtlasD3D::~FSlateFontAtlasD3D()
diff --git a/Engine/Source/Developer/StandaloneRenderer/Private/Windows/D3D/SlateD3DTextures.h b/Engine/Source/Developer/StandaloneRenderer/Private/Windows/D3D/SlateD3DTextures.h
--- a/Engine/Source/Developer/StandaloneRenderer/Private/Windows/D3D/SlateD3DTextures.h
+++ b/Engine/Source/Developer/StandaloneRenderer/Private/Windows/D3D/SlateD3DTextures.h
@@ -76,7 +76,7 @@ private:
 class FSlateFontAtlasD3D : public FSlateFontAtlas
 {
 public:
-	FSlateFontAtlasD3D( uint32 Width, uint32 Height );
+	FSlateFontAtlasD3D(uint32 Width, uint32 Height, const bool InIsGrayscale);
 	~FSlateFontAtlasD3D();
 
 	/** FSlateFontAtlas interface */
diff --git a/Engine/Source/Editor/FontEditor/Private/SCompositeFontEditor.cpp b/Engine/Source/Editor/FontEditor/Private/SCompositeFontEditor.cpp
--- a/Engine/Source/Editor/FontEditor/Private/SCompositeFontEditor.cpp
+++ b/Engine/Source/Editor/FontEditor/Private/SCompositeFontEditor.cpp
@@ -1761,7 +1761,7 @@ void SFontScalingFactorEditor::Construct(const FArguments& InArgs)
 		.AutoWidth()
 		[
 			SNew(SNumericEntryBox<float>)
-			.ToolTipText(LOCTEXT("ScalingFactorTooltip", "The scaling factor will adjust the size of the rendered glyphs so that you can tweak their size to match that of the default font family"))
+			.ToolTipText(LOCTEXT("ScalingFactorTooltip", "The scaling factor will adjust the size of the rendered glyphs so that you can tweak their size to match that of the default font family\nNote: Only applies to scalable font formats (ie, not to bitmap fonts)"))
 			.Value(this, &SFontScalingFactorEditor::GetScalingFactorAsOptional)
 			.OnValueCommitted(this, &SFontScalingFactorEditor::OnScalingFactorCommittedAsNumeric)
 		]
diff --git a/Engine/Source/Editor/FontEditor/Private/SFontEditorViewport.cpp b/Engine/Source/Editor/FontEditor/Private/SFontEditorViewport.cpp
--- a/Engine/Source/Editor/FontEditor/Private/SFontEditorViewport.cpp
+++ b/Engine/Source/Editor/FontEditor/Private/SFontEditorViewport.cpp
@@ -262,11 +262,14 @@ void FFontEditorViewportClient::Draw(FViewport* Viewport, FCanvas* Canvas)
 							if (GlyphToRender.bIsVisible)
 							{
 								const FShapedGlyphFontAtlasData GlyphAtlasData = FontCache->GetShapedGlyphFontAtlasData(GlyphToRender, FFontOutlineSettings::NoOutline);
+								
+								const float BitmapRenderScale = GlyphToRender.GetBitmapRenderScale();
+								const float InvBitmapRenderScale = 1.0f / BitmapRenderScale;
 
 								const float X = CurPos.X + LineX + GlyphAtlasData.HorizontalOffset + GlyphToRender.XOffset;
-								const float Y = CurPos.Y - GlyphAtlasData.VerticalOffset + GlyphToRender.YOffset + ShapedPreviewText->GetTextBaseline() + ShapedPreviewText->GetMaxTextHeight();
+								const float Y = CurPos.Y - GlyphAtlasData.VerticalOffset + GlyphToRender.YOffset + ((ShapedPreviewText->GetTextBaseline() + ShapedPreviewText->GetMaxTextHeight()) * InvBitmapRenderScale);
 
-								FCanvasBoxItem BoundingBoxItem(FVector2D(X, Y), FVector2D(GlyphAtlasData.USize, GlyphAtlasData.VSize));
+								FCanvasBoxItem BoundingBoxItem(FVector2D(X, Y), FVector2D(GlyphAtlasData.USize * BitmapRenderScale, GlyphAtlasData.VSize * BitmapRenderScale));
 								BoundingBoxItem.SetColor(FColor::Orange);
 								Canvas->DrawItem(BoundingBoxItem);
 							}
@@ -292,9 +295,12 @@ void FFontEditorViewportClient::Draw(FViewport* Viewport, FCanvas* Canvas)
 								if (GlyphToRender.bIsVisible)
 								{
 									const FShapedGlyphFontAtlasData GlyphAtlasData = FontCache->GetShapedGlyphFontAtlasData(GlyphToRender,FFontOutlineSettings::NoOutline);
+									
+									const float BitmapRenderScale = GlyphToRender.GetBitmapRenderScale();
+									const float InvBitmapRenderScale = 1.0f / BitmapRenderScale;
 
 									const float X = CurPos.X + LineX + GlyphClusterAdvance + GlyphAtlasData.HorizontalOffset + GlyphToRender.XOffset;
-									const float Y = CurPos.Y - GlyphAtlasData.VerticalOffset + GlyphToRender.YOffset + ShapedPreviewText->GetTextBaseline() + ShapedPreviewText->GetMaxTextHeight();
+									const float Y = CurPos.Y - GlyphAtlasData.VerticalOffset + GlyphToRender.YOffset + ((ShapedPreviewText->GetTextBaseline() + ShapedPreviewText->GetMaxTextHeight()) * InvBitmapRenderScale);
 
 									FVector2D ExtraWidth = FVector2D(ForceInitToZero);
 									if (GlyphClusterBounds.bIsValid)
@@ -303,7 +309,7 @@ void FFontEditorViewportClient::Draw(FViewport* Viewport, FCanvas* Canvas)
 										ExtraWidth.Y = (GlyphClusterBounds.Min.Y > Y) ? FMath::Abs(GlyphClusterBounds.Min.Y - Y) : 0.0f;
 									}
 
-									GlyphClusterBounds += FBox2D(FVector2D(X, Y), FVector2D(GlyphAtlasData.USize, GlyphAtlasData.VSize));
+									GlyphClusterBounds += FBox2D(FVector2D(X, Y), FVector2D(GlyphAtlasData.USize * BitmapRenderScale, GlyphAtlasData.VSize * BitmapRenderScale));
 									GlyphClusterBounds.Max += ExtraWidth;
 								}
 
diff --git a/Engine/Source/Runtime/Core/Private/Internationalization/ICUUtilities.h b/Engine/Source/Runtime/Core/Private/Internationalization/ICUUtilities.h
--- a/Engine/Source/Runtime/Core/Private/Internationalization/ICUUtilities.h
+++ b/Engine/Source/Runtime/Core/Private/Internationalization/ICUUtilities.h
@@ -11,6 +11,10 @@ THIRD_PARTY_INCLUDES_START
 	#include <unicode/unistr.h>
 THIRD_PARTY_INCLUDES_END
 
+#ifndef WITH_ICU_V64
+	#define WITH_ICU_V64 0
+#endif
+
 namespace ICUUtilities
 {
 	/**
diff --git a/Engine/Source/Runtime/Engine/Private/UserInterface/CanvasItem.cpp b/Engine/Source/Runtime/Engine/Private/UserInterface/CanvasItem.cpp
--- a/Engine/Source/Runtime/Engine/Private/UserInterface/CanvasItem.cpp
+++ b/Engine/Source/Runtime/Engine/Private/UserInterface/CanvasItem.cpp
@@ -1153,6 +1153,8 @@ void FCanvasTextItem::DrawStringInternal_RuntimeCache( FCanvas* InCanvas, const
 
 	uint32 FontTextureIndex = 0;
 	FTextureResource* FontTexture = nullptr;
+	ESimpleElementBlendMode FontTextureBlendMode = BlendMode;
+	FLinearColor FontColor = InColor;
 
 	float InvTextureSizeX = 0;
 	float InvTextureSizeY = 0;
@@ -1199,18 +1201,26 @@ void FCanvasTextItem::DrawStringInternal_RuntimeCache( FCanvas* InCanvas, const
 			{
 				// Font has a new texture for this glyph. Refresh the batch we use and the index we are currently using
 				FontTextureIndex = Entry.TextureIndex;
-				FontTexture = FontCache->GetEngineTextureResource( FontTextureIndex );
+
+				ISlateFontTexture* SlateFontTexture = FontCache->GetFontTexture(FontTextureIndex);
+				check(SlateFontTexture);
+
+				FontTexture = SlateFontTexture->GetEngineTexture();
 				check(FontTexture);
 
+				const bool bIsGrayscale = SlateFontTexture->IsGrayscale();
+				FontTextureBlendMode = bIsGrayscale ? BlendMode : SE_BLEND_Translucent;
+				FontColor = bIsGrayscale ? InColor : FLinearColor::White;
+
 				FBatchedElementParameters* BatchedElementParams = nullptr;
-				BatchedElements = InCanvas->GetBatchedElements(FCanvas::ET_Triangle, BatchedElementParams, FontTexture, BlendMode, FontRenderInfo.GlowInfo, false);
+				BatchedElements = InCanvas->GetBatchedElements(FCanvas::ET_Triangle, BatchedElementParams, FontTexture, FontTextureBlendMode, FontRenderInfo.GlowInfo, false);
 				check(BatchedElements);
 
 				// Trade-off to use memory for performance by pre-allocating more reserved space 
 				// for the triangles/vertices of the batched elements used to render the text tiles.
 				// Only reserve initial batch, allow growth afterwards in case there are multiple repeated calls.
 				// Reserving exactly the added amount each time would essentially force an alloc each time on subsequent calls.
-				BatchedElements->ReserveTriangles(TextLen*2, FontTexture, BlendMode);
+				BatchedElements->ReserveTriangles(TextLen*2, FontTexture, FontTextureBlendMode);
 				BatchedElements->ReserveVertices(TextLen*4);
 
 				InvTextureSizeX = 1.0f/FontTexture->GetSizeX();
@@ -1230,14 +1240,16 @@ void FCanvasTextItem::DrawStringInternal_RuntimeCache( FCanvas* InCanvas, const
 
 			if( !bIsWhitespace )
 			{
+				const float InvBitmapRenderScale = 1.0f / Entry.BitmapRenderScale;
+
 				const float X = DrawPos.X + LineX + (Entry.HorizontalOffset * Scale.X);
 				// Note PosX,PosY is the upper left corner of the bounding box representing the string.  This computes the Y position of the baseline where text will sit
 
-				const float Y = DrawPos.Y + PosY - (Entry.VerticalOffset * Scale.Y) + (Entry.GlobalDescender * Scale.Y) + ScaledMaxHeight;
+				const float Y = DrawPos.Y + PosY - (Entry.VerticalOffset * Scale.Y) + (((Entry.GlobalDescender * Scale.Y) + ScaledMaxHeight) * InvBitmapRenderScale);
 				const float U = Entry.StartU * InvTextureSizeX;
 				const float V = Entry.StartV * InvTextureSizeY;
-				const float SizeX = Entry.USize * Scale.X;
-				const float SizeY = Entry.VSize * Scale.Y;
+				const float SizeX = Entry.USize * Scale.X * Entry.BitmapRenderScale;
+				const float SizeY = Entry.VSize * Scale.Y * Entry.BitmapRenderScale;
 				const float SizeU = Entry.USize * InvTextureSizeX;
 				const float SizeV = Entry.VSize * InvTextureSizeY;
 
@@ -1249,26 +1261,26 @@ void FCanvasTextItem::DrawStringInternal_RuntimeCache( FCanvas* InCanvas, const
 				int32 V00 = BatchedElements->AddVertex(
 					FVector4( Left, Top, 0.f, Depth ),
 					FVector2D( U, V ),
-					InColor,
+					FontColor,
 					HitProxyId );
 				int32 V10 = BatchedElements->AddVertex(
 					FVector4( Right, Top, 0.0f, Depth ),
 					FVector2D( U + SizeU, V ),			
-					InColor,
+					FontColor,
 					HitProxyId );
 				int32 V01 = BatchedElements->AddVertex(
 					FVector4( Left, Bottom, 0.0f, Depth ),
 					FVector2D( U, V + SizeV ),	
-					InColor,
+					FontColor,
 					HitProxyId);
 				int32 V11 = BatchedElements->AddVertex(
 					FVector4( Right, Bottom, 0.0f, Depth ),
 					FVector2D( U + SizeU, V + SizeV ),
-					InColor,
+					FontColor,
 					HitProxyId);
 
-				BatchedElements->AddTriangle(V00, V10, V11, FontTexture, BlendMode, FontRenderInfo.GlowInfo);
-				BatchedElements->AddTriangle(V00, V11, V01, FontTexture, BlendMode, FontRenderInfo.GlowInfo);
+				BatchedElements->AddTriangle(V00, V10, V11, FontTexture, FontTextureBlendMode, FontRenderInfo.GlowInfo);
+				BatchedElements->AddTriangle(V00, V11, V01, FontTexture, FontTextureBlendMode, FontRenderInfo.GlowInfo);
 			}
 
 			LineX += Entry.XAdvance * Scale.X;
@@ -1319,6 +1331,8 @@ void FCanvasShapedTextItem::DrawStringInternal( FCanvas* InCanvas, const FVector
 
 	uint32 FontTextureIndex = 0;
 	FTextureResource* FontTexture = nullptr;
+	ESimpleElementBlendMode FontTextureBlendMode = BlendMode;
+	FLinearColor FontColor = InColor;
 
 	float InvTextureSizeX = 0;
 	float InvTextureSizeY = 0;
@@ -1352,29 +1366,40 @@ void FCanvasShapedTextItem::DrawStringInternal( FCanvas* InCanvas, const FVector
 				{
 					// Font has a new texture for this glyph. Refresh the batch we use and the index we are currently using
 					FontTextureIndex = GlyphAtlasData.TextureIndex;
-					FontTexture = FontCache->GetEngineTextureResource( FontTextureIndex );
+
+					ISlateFontTexture* SlateFontTexture = FontCache->GetFontTexture(FontTextureIndex);
+					check(SlateFontTexture);
+
+					FontTexture = SlateFontTexture->GetEngineTexture();
 					check(FontTexture);
 
+					const bool bIsGrayscale = SlateFontTexture->IsGrayscale();
+					FontTextureBlendMode = bIsGrayscale ? BlendMode : SE_BLEND_Translucent;
+					FontColor = bIsGrayscale ? InColor : FLinearColor::White;
+
 					FBatchedElementParameters* BatchedElementParams = nullptr;
-					BatchedElements = InCanvas->GetBatchedElements(FCanvas::ET_Triangle, BatchedElementParams, FontTexture, BlendMode, FontRenderInfo.GlowInfo);
+					BatchedElements = InCanvas->GetBatchedElements(FCanvas::ET_Triangle, BatchedElementParams, FontTexture, FontTextureBlendMode, FontRenderInfo.GlowInfo);
 					check(BatchedElements);
 
 					const int32 NumGlyphs = ShapedGlyphSequence->GetGlyphsToRender().Num();
 					BatchedElements->ReserveVertices(4 * NumGlyphs);
-					BatchedElements->ReserveTriangles(2 * NumGlyphs, FontTexture, BlendMode);
+					BatchedElements->ReserveTriangles(2 * NumGlyphs, FontTexture, FontTextureBlendMode);
 
 					InvTextureSizeX = 1.0f/FontTexture->GetSizeX();
 					InvTextureSizeY = 1.0f/FontTexture->GetSizeY();
 				}
 
-				const float X = DrawPos.X + LineX + (GlyphAtlasData.HorizontalOffset * Scale.X) + (GlyphToRender.XOffset  * Scale.X);
+				const float BitmapRenderScale = GlyphToRender.GetBitmapRenderScale();
+				const float InvBitmapRenderScale = 1.0f / BitmapRenderScale;
+
+				const float X = DrawPos.X + LineX + (GlyphAtlasData.HorizontalOffset * Scale.X) + (GlyphToRender.XOffset * Scale.X);
 				// Note PosX,PosY is the upper left corner of the bounding box representing the string.  This computes the Y position of the baseline where text will sit
 
-				const float Y = DrawPos.Y + PosY - (GlyphAtlasData.VerticalOffset * Scale.Y) + (GlyphToRender.YOffset * Scale.Y) + ScaledBaseline + ScaledMaxHeight;
+				const float Y = DrawPos.Y + PosY - (GlyphAtlasData.VerticalOffset * Scale.Y) + (GlyphToRender.YOffset * Scale.Y) + ((ScaledBaseline + ScaledMaxHeight) * InvBitmapRenderScale);
 				const float U = GlyphAtlasData.StartU * InvTextureSizeX;
 				const float V = GlyphAtlasData.StartV * InvTextureSizeY;
-				const float SizeX = GlyphAtlasData.USize * Scale.X;
-				const float SizeY = GlyphAtlasData.VSize * Scale.Y;
+				const float SizeX = GlyphAtlasData.USize * Scale.X * BitmapRenderScale;
+				const float SizeY = GlyphAtlasData.VSize * Scale.Y * BitmapRenderScale;
 				const float SizeU = GlyphAtlasData.USize * InvTextureSizeX;
 				const float SizeV = GlyphAtlasData.VSize * InvTextureSizeY;
 
@@ -1386,26 +1411,26 @@ void FCanvasShapedTextItem::DrawStringInternal( FCanvas* InCanvas, const FVector
 				int32 V00 = BatchedElements->AddVertex(
 					FVector4( Left, Top, 0.f, Depth ),
 					FVector2D( U, V ),
-					InColor,
+					FontColor,
 					HitProxyId );
 				int32 V10 = BatchedElements->AddVertex(
 					FVector4( Right, Top, 0.0f, Depth ),
 					FVector2D( U + SizeU, V ),			
-					InColor,
+					FontColor,
 					HitProxyId );
 				int32 V01 = BatchedElements->AddVertex(
 					FVector4( Left, Bottom, 0.0f, Depth ),
 					FVector2D( U, V + SizeV ),	
-					InColor,
+					FontColor,
 					HitProxyId);
 				int32 V11 = BatchedElements->AddVertex(
 					FVector4( Right, Bottom, 0.0f, Depth ),
 					FVector2D( U + SizeU, V + SizeV ),
-					InColor,
+					FontColor,
 					HitProxyId);
 
-				BatchedElements->AddTriangle(V00, V10, V11, FontTexture, BlendMode, FontRenderInfo.GlowInfo);
-				BatchedElements->AddTriangle(V00, V11, V01, FontTexture, BlendMode, FontRenderInfo.GlowInfo);
+				BatchedElements->AddTriangle(V00, V10, V11, FontTexture, FontTextureBlendMode, FontRenderInfo.GlowInfo);
+				BatchedElements->AddTriangle(V00, V11, V01, FontTexture, FontTextureBlendMode, FontRenderInfo.GlowInfo);
 			}
 		}
 
diff --git a/Engine/Source/Runtime/SlateCore/Private/Fonts/FontCache.cpp b/Engine/Source/Runtime/SlateCore/Private/Fonts/FontCache.cpp
--- a/Engine/Source/Runtime/SlateCore/Private/Fonts/FontCache.cpp
+++ b/Engine/Source/Runtime/SlateCore/Private/Fonts/FontCache.cpp
@@ -74,6 +74,11 @@ ETextShapingMethod GetDefaultTextShapingMethod()
 }
 
 
+float FShapedGlyphEntry::GetBitmapRenderScale() const
+{
+	return FontFaceData ? FontFaceData->BitmapRenderScale : 1.0f;
+}
+
 FShapedGlyphEntryKey::FShapedGlyphEntryKey(const FShapedGlyphFaceData& InFontFaceData, uint32 InGlyphIndex, const FFontOutlineSettings& InOutlineSettings)
 	: FontFace(InFontFaceData.FontFace)
 	, FontSize(InFontFaceData.FontSize)
@@ -666,10 +671,16 @@ FCharacterList::FCharacterListEntry FCharacterList::CacheCharacter(TCHAR Charact
 		// Found a valid font face?
 		if (FaceGlyphData.FaceAndMemory.IsValid())
 		{
+			// Only scalable font types can use sub-font scaling
+			if (!FT_IS_SCALABLE(FaceGlyphData.FaceAndMemory->GetFace()))
+			{
+				SubFontScalingFactor = 1.0f;
+			}
+
 			const float FinalFontScale = FontKey.GetScale() * SubFontScalingFactor;
 
 			uint32 GlyphFlags = 0;
-			SlateFontRendererUtils::AppendGlyphFlags(*FontDataPtr, GlyphFlags);
+			SlateFontRendererUtils::AppendGlyphFlags(*FaceGlyphData.FaceAndMemory, *FontDataPtr, GlyphFlags);
 
 			const bool bHasKerning = FT_HAS_KERNING(FaceGlyphData.FaceAndMemory->GetFace()) != 0;
 
@@ -733,6 +744,7 @@ FCharacterEntry FCharacterList::MakeCharacterEntry(TCHAR Character, const FChara
 			CharEntry.GlyphIndex = InternalEntry.ShapedGlyphEntry.GlyphIndex;
 			CharEntry.FontData = InternalEntry.FontData;
 			CharEntry.FontScale = InternalEntry.ShapedGlyphEntry.FontFaceData->FontScale;
+			CharEntry.BitmapRenderScale = InternalEntry.ShapedGlyphEntry.FontFaceData->BitmapRenderScale;
 			CharEntry.StartU = ShapedGlyphFontAtlasData.StartU;
 			CharEntry.StartV = ShapedGlyphFontAtlasData.StartV;
 			CharEntry.USize = ShapedGlyphFontAtlasData.USize;
@@ -743,6 +755,7 @@ FCharacterEntry FCharacterList::MakeCharacterEntry(TCHAR Character, const FChara
 			CharEntry.XAdvance = InternalEntry.ShapedGlyphEntry.XAdvance;
 			CharEntry.TextureIndex = ShapedGlyphFontAtlasData.TextureIndex;
 			CharEntry.HasKerning = InternalEntry.HasKerning;
+			CharEntry.SupportsOutline = ShapedGlyphFontAtlasData.SupportsOutline;
 			CharEntry.FallbackLevel = InternalEntry.FallbackLevel;
 		}
 	}
@@ -760,7 +773,8 @@ FSlateFontCache::FSlateFontCache( TSharedRef<ISlateFontAtlasFactory> InFontAtlas
 	, TextShaper( new FSlateTextShaper( FTGlyphCache.Get(), FTAdvanceCache.Get(), FTKerningPairCache.Get(), CompositeFontCache.Get(), FontRenderer.Get(), this ) )
 	, FontAtlasFactory( InFontAtlasFactory )
 	, bFlushRequested( false )
-	, CurrentMaxAtlasPagesBeforeFlushRequest(InitialMaxAtlasPagesBeforeFlushRequest)
+	, CurrentMaxGrayscaleAtlasPagesBeforeFlushRequest(InitialMaxAtlasPagesBeforeFlushRequest)
+	, CurrentMaxColorAtlasPagesBeforeFlushRequest(InitialMaxAtlasPagesBeforeFlushRequest)
 	, CurrentMaxNonAtlasedTexturesBeforeFlushRequest(InitialMaxNonAtlasedTexturesBeforeFlushRequest)
 	, FrameCounterLastFlushRequest( 0 )
 	, OwningThread(InOwningThread)
@@ -773,6 +787,8 @@ FSlateFontCache::~FSlateFontCache()
 	FInternationalization::Get().OnCultureChanged().RemoveAll(this);
 
 	// Make sure things get destroyed in the correct order
+	FontToCharacterListCache.Empty();
+	ShapedGlyphToAtlasData.Empty();
 	TextShaper.Reset();
 	FontRenderer.Reset();
 	CompositeFontCache.Reset();
@@ -784,22 +800,19 @@ FSlateFontCache::~FSlateFontCache()
 
 int32 FSlateFontCache::GetNumAtlasPages() const
 {
-	return FontAtlases.Num();
-}
-
-FIntPoint FSlateFontCache::GetAtlasPageSize() const
-{
-	return FontAtlasFactory->GetAtlasSize();
+	return GrayscaleFontAtlasIndices.Num() + ColorFontAtlasIndices.Num();
 }
 
 FSlateShaderResource* FSlateFontCache::GetAtlasPageResource(const int32 InIndex) const
 {
-	return FontAtlases[InIndex]->GetSlateTexture();
+	const int32 AllFontTexturesIndex = InIndex < GrayscaleFontAtlasIndices.Num() ? GrayscaleFontAtlasIndices[InIndex] : ColorFontAtlasIndices[InIndex - GrayscaleFontAtlasIndices.Num()];
+	return AllFontTextures[AllFontTexturesIndex]->GetSlateTexture();
 }
 
-bool FSlateFontCache::IsAtlasPageResourceAlphaOnly() const
+bool FSlateFontCache::IsAtlasPageResourceAlphaOnly(const int32 InIndex) const
 {
-	return true;
+	const int32 AllFontTexturesIndex = InIndex < GrayscaleFontAtlasIndices.Num() ? GrayscaleFontAtlasIndices[InIndex] : ColorFontAtlasIndices[InIndex - GrayscaleFontAtlasIndices.Num()];
+	return AllFontTextures[AllFontTexturesIndex]->IsGrayscale();
 }
 
 bool FSlateFontCache::AddNewEntry(const FShapedGlyphEntry& InShapedGlyph, const FFontOutlineSettings& InOutlineSettings, FShapedGlyphFontAtlasData& OutAtlasData)
@@ -811,8 +824,9 @@ bool FSlateFontCache::AddNewEntry(const FShapedGlyphEntry& InShapedGlyph, const
 	OutAtlasData.Valid = bDidRender && AddNewEntry(RenderData, OutAtlasData.TextureIndex, OutAtlasData.StartU, OutAtlasData.StartV, OutAtlasData.USize, OutAtlasData.VSize);
 	if (OutAtlasData.Valid)
 	{
-		OutAtlasData.VerticalOffset = RenderData.MeasureInfo.VerticalOffset;
-		OutAtlasData.HorizontalOffset = RenderData.MeasureInfo.HorizontalOffset;
+		OutAtlasData.VerticalOffset = RenderData.VerticalOffset;
+		OutAtlasData.HorizontalOffset = RenderData.HorizontalOffset;
+		OutAtlasData.SupportsOutline = RenderData.bSupportsOutline;
 	}
 
 	return OutAtlasData.Valid;
@@ -823,31 +837,33 @@ bool FSlateFontCache::AddNewEntry( const FCharacterRenderData InRenderData, uint
 	const uint64 LastFlushRequestFrameDelta = GFrameCounter - FrameCounterLastFlushRequest;
 
 	// Will this entry fit within any atlas texture?
-	if (InRenderData.MeasureInfo.SizeX > FontAtlasFactory->GetAtlasSize().X || InRenderData.MeasureInfo.SizeY > FontAtlasFactory->GetAtlasSize().Y)
+	const FIntPoint FontAtlasSize = FontAtlasFactory->GetAtlasSize(InRenderData.bIsGrayscale);
+	if (InRenderData.SizeX > FontAtlasSize.X || InRenderData.SizeY > FontAtlasSize.Y)
 	{
-		TSharedPtr<ISlateFontTexture> NonAtlasedTexture = FontAtlasFactory->CreateNonAtlasedTexture(InRenderData.MeasureInfo.SizeX, InRenderData.MeasureInfo.SizeY, InRenderData.RawPixels);
+		TSharedPtr<ISlateFontTexture> NonAtlasedTexture = FontAtlasFactory->CreateNonAtlasedTexture(InRenderData.SizeX, InRenderData.SizeY, InRenderData.bIsGrayscale, InRenderData.RawPixels);
 		if (NonAtlasedTexture.IsValid())
 		{
 			INC_DWORD_STAT_BY(STAT_SlateNumFontNonAtlasedTextures, 1);
 
 			UE_LOG(LogSlate, Warning, TEXT("SlateFontCache - Glyph texture is too large to store in the font atlas, so we're falling back to a non-atlased texture for this glyph. This may have SERIOUS performance implications. Atlas page size: { %d, %d }. Glyph render size: { %d, %d }"),
-				FontAtlasFactory->GetAtlasSize().X, FontAtlasFactory->GetAtlasSize().Y,
-				InRenderData.MeasureInfo.SizeX, InRenderData.MeasureInfo.SizeY
+				FontAtlasSize.X, FontAtlasSize.Y,
+				InRenderData.SizeX, InRenderData.SizeY
 				);
 
-			NonAtlasedTextures.Add(NonAtlasedTexture.ToSharedRef());
 			OutTextureIndex = AllFontTextures.Add(NonAtlasedTexture.ToSharedRef());
+			NonAtlasedTextureIndices.Add(OutTextureIndex);
+
 			OutGlyphX = 0;
 			OutGlyphY = 0;
-			OutGlyphWidth = InRenderData.MeasureInfo.SizeX;
-			OutGlyphHeight = InRenderData.MeasureInfo.SizeY;
+			OutGlyphWidth = InRenderData.SizeX;
+			OutGlyphHeight = InRenderData.SizeY;
 
-			if (NonAtlasedTextures.Num() > CurrentMaxNonAtlasedTexturesBeforeFlushRequest && !bFlushRequested)
+			if (NonAtlasedTextureIndices.Num() > CurrentMaxNonAtlasedTexturesBeforeFlushRequest && !bFlushRequested)
 			{
 				// The InitialMaxNonAtlasedTexturesBeforeFlushRequest may have changed since last flush,
 				// so double check that we're below the current max or initial, if we're under it,
 				// update the current to the initial.
-				if (NonAtlasedTextures.Num() <= InitialMaxNonAtlasedTexturesBeforeFlushRequest)
+				if (NonAtlasedTextureIndices.Num() <= InitialMaxNonAtlasedTexturesBeforeFlushRequest)
 				{
 					CurrentMaxNonAtlasedTexturesBeforeFlushRequest = InitialMaxNonAtlasedTexturesBeforeFlushRequest;
 				}
@@ -855,13 +871,13 @@ bool FSlateFontCache::AddNewEntry( const FCharacterRenderData InRenderData, uint
 				// a lot of font data cached. We should update CurrentMaxNonAtlasedTexturesBeforeFlushRequest to give us a bit more flexibility before the next flush request
 				else if (LastFlushRequestFrameDelta <= GrowFontNonAtlasFrameWindow)
 				{
-					CurrentMaxNonAtlasedTexturesBeforeFlushRequest = NonAtlasedTextures.Num();
+					CurrentMaxNonAtlasedTexturesBeforeFlushRequest = NonAtlasedTextureIndices.Num();
 					UE_LOG(LogSlate, Warning, TEXT("SlateFontCache - Setting the threshold to trigger a flush to %d non-atlased textures as there is a lot of font data being cached."), CurrentMaxNonAtlasedTexturesBeforeFlushRequest);
 				}
 				else
 				{
 					// We've grown beyond our current stable limit - try and request a flush
-					RequestFlushCache(FString::Printf(TEXT("Large glyph font atlases out of space; %d/%d Textures; frames since last flush: %llu"), NonAtlasedTextures.Num(), CurrentMaxNonAtlasedTexturesBeforeFlushRequest, LastFlushRequestFrameDelta));
+					RequestFlushCache(FString::Printf(TEXT("Large glyph font atlases out of space; %d/%d Textures; frames since last flush: %llu"), NonAtlasedTextureIndices.Num(), CurrentMaxNonAtlasedTexturesBeforeFlushRequest, LastFlushRequestFrameDelta));
 				}
 			}
 
@@ -869,8 +885,8 @@ bool FSlateFontCache::AddNewEntry( const FCharacterRenderData InRenderData, uint
 		}
 
 		UE_LOG(LogSlate, Warning, TEXT("SlateFontCache - Glyph texture is too large to store in the font atlas, but we cannot support rendering such a large texture. Atlas page size: { %d, %d }. Glyph render size: { %d, %d }"),
-			FontAtlasFactory->GetAtlasSize().X, FontAtlasFactory->GetAtlasSize().Y,
-			InRenderData.MeasureInfo.SizeX, InRenderData.MeasureInfo.SizeY
+			FontAtlasSize.X, FontAtlasSize.Y,
+			InRenderData.SizeX, InRenderData.SizeY
 			);
 		return false;
 	}
@@ -883,18 +899,29 @@ bool FSlateFontCache::AddNewEntry( const FCharacterRenderData InRenderData, uint
 		OutGlyphHeight = AtlasedTextureSlot.Height - (2 * AtlasedTextureSlot.Padding);
 	};
 
-	for( OutTextureIndex = 0; OutTextureIndex < FontAtlases.Num(); ++OutTextureIndex ) 
+	TArray<int32>& FontAtlasIndices = InRenderData.bIsGrayscale ? GrayscaleFontAtlasIndices : ColorFontAtlasIndices;
+	int32& CurrentMaxAtlasPagesBeforeFlushRequest = InRenderData.bIsGrayscale ? CurrentMaxGrayscaleAtlasPagesBeforeFlushRequest : CurrentMaxColorAtlasPagesBeforeFlushRequest;
+
+	for (const int32 FontAtlasIndex : FontAtlasIndices)
 	{
+		FSlateFontAtlas& FontAtlas = static_cast<FSlateFontAtlas&>(AllFontTextures[FontAtlasIndex].Get());
+		checkSlow(FontAtlas.IsGrayscale() == InRenderData.bIsGrayscale);
+
 		// Add the character to the texture
-		const FAtlasedTextureSlot* NewSlot = FontAtlases[OutTextureIndex]->AddCharacter(InRenderData);
+		const FAtlasedTextureSlot* NewSlot = FontAtlas.AddCharacter(InRenderData);
 		if( NewSlot )
 		{
+			OutTextureIndex = FontAtlasIndex;
 			FillOutputParamsFromAtlasedTextureSlot(*NewSlot);
 			return true;
 		}
 	}
 
-	TSharedRef<FSlateFontAtlas> FontAtlas = FontAtlasFactory->CreateFontAtlas();
+	TSharedRef<FSlateFontAtlas> FontAtlas = FontAtlasFactory->CreateFontAtlas(InRenderData.bIsGrayscale);
+	OutTextureIndex = AllFontTextures.Add(FontAtlas);
+	FontAtlasIndices.Add(OutTextureIndex);
+
+	INC_DWORD_STAT_BY(STAT_SlateNumFontAtlases, 1);
 
 	// Add the character to the texture
 	const FAtlasedTextureSlot* NewSlot = FontAtlas->AddCharacter(InRenderData);
@@ -903,17 +930,12 @@ bool FSlateFontCache::AddNewEntry( const FCharacterRenderData InRenderData, uint
 		FillOutputParamsFromAtlasedTextureSlot(*NewSlot);
 	}
 
-	FontAtlases.Add( FontAtlas );
-	OutTextureIndex = AllFontTextures.Add(FontAtlas);
-
-	INC_DWORD_STAT_BY( STAT_SlateNumFontAtlases, 1 );
-
-	if( FontAtlases.Num() > CurrentMaxAtlasPagesBeforeFlushRequest && !bFlushRequested )
+	if (FontAtlasIndices.Num() > CurrentMaxAtlasPagesBeforeFlushRequest && !bFlushRequested)
 	{
 		// The InitialMaxNonAtlasedTexturesBeforeFlushRequest may have changed since last flush,
 		// so double check that we're below the current max or initial, if we're under it,
 		// update the current to the initial.
-		if (FontAtlases.Num() <= InitialMaxAtlasPagesBeforeFlushRequest)
+		if (FontAtlasIndices.Num() <= InitialMaxAtlasPagesBeforeFlushRequest)
 		{
 			CurrentMaxAtlasPagesBeforeFlushRequest = InitialMaxAtlasPagesBeforeFlushRequest;
 		}
@@ -921,13 +943,13 @@ bool FSlateFontCache::AddNewEntry( const FCharacterRenderData InRenderData, uint
 		// a lot of font data cached. We should update MaxAtlasPagesBeforeFlushRequest to give us a bit more flexibility before the next flush request
 		else if (LastFlushRequestFrameDelta <= GrowFontAtlasFrameWindow)
 		{
-			CurrentMaxAtlasPagesBeforeFlushRequest = FontAtlases.Num();
+			CurrentMaxAtlasPagesBeforeFlushRequest = FontAtlasIndices.Num();
 			UE_LOG(LogSlate, Warning, TEXT("SlateFontCache - Setting the threshold to trigger a flush to %d atlas pages as there is a lot of font data being cached."), CurrentMaxAtlasPagesBeforeFlushRequest);
 		}
 		else
 		{
 			// We've grown beyond our current stable limit - try and request a flush
-			RequestFlushCache(FString::Printf(TEXT("Font Atlases Full; %d/%d Pages; frames since last flush: %llu"), FontAtlases.Num(), CurrentMaxAtlasPagesBeforeFlushRequest, LastFlushRequestFrameDelta));
+			RequestFlushCache(FString::Printf(TEXT("Font Atlases Full; %d/%d Pages; frames since last flush: %llu"), FontAtlasIndices.Num(), CurrentMaxAtlasPagesBeforeFlushRequest, LastFlushRequestFrameDelta));
 		}
 	}
 
@@ -1097,7 +1119,8 @@ void FSlateFontCache::RequestFlushCache(const FString& FlushReason)
 #endif
 
 		bFlushRequested = true;
-		CurrentMaxAtlasPagesBeforeFlushRequest = InitialMaxAtlasPagesBeforeFlushRequest;
+		CurrentMaxGrayscaleAtlasPagesBeforeFlushRequest = InitialMaxAtlasPagesBeforeFlushRequest;
+		CurrentMaxColorAtlasPagesBeforeFlushRequest = InitialMaxAtlasPagesBeforeFlushRequest;
 		CurrentMaxNonAtlasedTexturesBeforeFlushRequest = InitialMaxNonAtlasedTexturesBeforeFlushRequest;
 		FrameCounterLastFlushRequest = GFrameCounter;
 	}
@@ -1141,22 +1164,24 @@ bool FSlateFontCache::ConditionalFlushCache()
 
 void FSlateFontCache::UpdateCache()
 {
-	for (const TSharedRef<FSlateFontAtlas>& FontAtlas : FontAtlases)
+	auto UpdateFontAtlasTextures = [this](const TArray<int32>& FontAtlasIndices)
 	{
-		FontAtlas->ConditionalUpdateTexture();
-	}
+		for (const int32 FontAtlasIndex : FontAtlasIndices)
+		{
+			FSlateFontAtlas& FontAtlas = static_cast<FSlateFontAtlas&>(AllFontTextures[FontAtlasIndex].Get());
+			FontAtlas.ConditionalUpdateTexture();
+		}
+	};
+
+	UpdateFontAtlasTextures(GrayscaleFontAtlasIndices);
+	UpdateFontAtlasTextures(ColorFontAtlasIndices);
 }
 
 void FSlateFontCache::ReleaseResources()
 {
-	for (const TSharedRef<FSlateFontAtlas>& FontAtlas : FontAtlases)
-	{
-		FontAtlas->ReleaseResources();
-	}
-
-	for (const TSharedRef<ISlateFontTexture>& NonAtlasedTexture : NonAtlasedTextures)
+	for (const TSharedRef<ISlateFontTexture>& FontTexture : AllFontTextures)
 	{
-		NonAtlasedTexture->ReleaseResources();
+		FontTexture->ReleaseResources();
 	}
 }
 
@@ -1175,8 +1200,9 @@ bool FSlateFontCache::FlushCache()
 		SET_DWORD_STAT(STAT_SlateNumFontAtlases, 0);
 		SET_DWORD_STAT(STAT_SlateNumFontNonAtlasedTextures, 0);
 
-		FontAtlases.Empty();
-		NonAtlasedTextures.Empty();
+		GrayscaleFontAtlasIndices.Empty();
+		ColorFontAtlasIndices.Empty();
+		NonAtlasedTextureIndices.Empty();
 		AllFontTextures.Empty();
 
 		{
diff --git a/Engine/Source/Runtime/SlateCore/Private/Fonts/FontCacheFreeType.cpp b/Engine/Source/Runtime/SlateCore/Private/Fonts/FontCacheFreeType.cpp
--- a/Engine/Source/Runtime/SlateCore/Private/Fonts/FontCacheFreeType.cpp
+++ b/Engine/Source/Runtime/SlateCore/Private/Fonts/FontCacheFreeType.cpp
@@ -5,9 +5,13 @@
 #include "HAL/PlatformFile.h"
 #include "HAL/PlatformFilemanager.h"
 #include "HAL/LowLevelMemTracker.h"
+#include "HAL/IConsoleManager.h"
 
 #if WITH_FREETYPE
 
+static int32 GSlateEnableLegacyFontHinting = 0;
+static FAutoConsoleVariableRef CVarSlateEnableLegacyFontHinting(TEXT("Slate.EnableLegacyFontHinting"), GSlateEnableLegacyFontHinting, TEXT("Enable the legacy font hinting? (0/1)."), ECVF_Default);
+
 // The total amount of memory freetype allocates internally
 DECLARE_MEMORY_STAT(TEXT("FreeType Total Allocated Memory"), STAT_SlateFreetypeAllocatedMemory, STATGROUP_SlateMemory);
 
@@ -64,25 +68,104 @@ namespace FreeTypeUtils
 
 void ApplySizeAndScale(FT_Face InFace, const int32 InFontSize, const float InFontScale)
 {
-	FT_Error Error = FT_Set_Char_Size(InFace, 0, ConvertPixelTo26Dot6<FT_F26Dot6>(InFontSize), FreeTypeConstants::HorizontalDPI, FreeTypeConstants::VerticalDPI);
-
-	check(Error == 0);
+	// Convert to fixed scale for maximum precision
+	const FT_F26Dot6 FixedFontSize = ConvertPixelTo26Dot6<FT_F26Dot6>(InFontSize);
+	const FT_Long FixedFontScale = ConvertPixelTo16Dot16<FT_Long>(InFontScale);
 
-	if (InFontScale != 1.0f)
+	// Convert the requested font size to a pixel size based on our render DPI and the requested scale
+	// This result is left in 26.6 space as it is the common format used by both FT_IS_SCALABLE and FT_HAS_FIXED_SIZES fonts
+	FT_F26Dot6 RequiredFixedFontPixelSize = 0;
 	{
-		FT_Matrix ScaleMatrix;
-		ScaleMatrix.xy = 0;
-		ScaleMatrix.xx = ConvertPixelTo16Dot16<FT_Fixed>(InFontScale);
-		ScaleMatrix.yy = ConvertPixelTo16Dot16<FT_Fixed>(InFontScale);
-		ScaleMatrix.yx = 0;
-		FT_Set_Transform(InFace, &ScaleMatrix, nullptr);
+		// Convert the 26.6 character size to the unscaled 26.6 pixel size
+		// Note: The logic here mirrors what FT_REQUEST_WIDTH and FT_REQUEST_HEIGHT do internally when using FT_Set_Char_Size
+		RequiredFixedFontPixelSize = ((FixedFontSize * (FT_Pos)FreeTypeConstants::RenderDPI) + 36) / 72;
+
+		// Scale the 26.6 pixel size by the desired 16.16 fractional scaling value
+		RequiredFixedFontPixelSize = FT_MulFix(RequiredFixedFontPixelSize, FixedFontScale);
 	}
-	else
+
+	if (FT_IS_SCALABLE(InFace))
 	{
-		FT_Set_Transform(InFace, nullptr, nullptr);
+		// Convert the 26.6 scaled pixel size back into pixel space
+		const uint32 RequiredFontPixelSize = Convert26Dot6ToRoundedPixel<uint32>(RequiredFixedFontPixelSize);
+
+		// Set the pixel size
+		FT_Error Error = FT_Set_Pixel_Sizes(InFace, RequiredFontPixelSize, RequiredFontPixelSize);
+		check(Error == 0);
 	}
+	else if (FT_HAS_FIXED_SIZES(InFace))
+	{
+		// Find the best bitmap strike for our required pixel size
+		// Height is the most important metric, so that is the one we search for
+		int32 BestStrikeIndex = INDEX_NONE;
+		{
+			FT_F26Dot6 RunningBestFixedStrikeHeight = 0;
+			for (int32 PotentialStrikeIndex = 0; PotentialStrikeIndex < InFace->num_fixed_sizes; ++PotentialStrikeIndex)
+			{
+				const FT_F26Dot6 PotentialFixedStrikeHeight = InFace->available_sizes[PotentialStrikeIndex].y_ppem;
+
+				// An exact match always wins
+				if (PotentialFixedStrikeHeight == RequiredFixedFontPixelSize)
+				{
+					BestStrikeIndex = PotentialStrikeIndex;
+					break;
+				}
+
+				// First loop?
+				if (BestStrikeIndex == INDEX_NONE)
+				{
+					BestStrikeIndex = PotentialStrikeIndex;
+					RunningBestFixedStrikeHeight = PotentialFixedStrikeHeight;
+					continue;
+				}
+
+				// Our running strike is smaller than our required size, so choose this as our running strike... 
+				if (RunningBestFixedStrikeHeight < RequiredFixedFontPixelSize)
+				{
+					// ... but only if it grows the size of our running strike
+					if (PotentialFixedStrikeHeight > RunningBestFixedStrikeHeight)
+					{
+						BestStrikeIndex = PotentialStrikeIndex;
+						RunningBestFixedStrikeHeight = PotentialFixedStrikeHeight;
+					}
+					continue;
+				}
+				
+				// Our running strike is larger than our required size, so choose this as our running strike... 
+				if (RunningBestFixedStrikeHeight > RequiredFixedFontPixelSize)
+				{
+					// ... but only if it shrinks the size of our running strike to no less than the required size
+					if (PotentialFixedStrikeHeight < RunningBestFixedStrikeHeight && PotentialFixedStrikeHeight > RequiredFixedFontPixelSize)
+					{
+						BestStrikeIndex = PotentialStrikeIndex;
+						RunningBestFixedStrikeHeight = PotentialFixedStrikeHeight;
+					}
+					continue;
+				}
+			}
+		}
+
+		// Set the strike index
+		check(BestStrikeIndex != INDEX_NONE);
+		FT_Error Error = FT_Select_Size(InFace, BestStrikeIndex);
+		check(Error == 0);
+
+		// Work out the scale required to convert from the chosen strike to the desired pixel size
+		FT_Long FixedStrikeScale = 0;
+		{
+			// Convert the 26.6 values to 16.16 space so that we can use FT_DivFix
+			const FT_Long RequiredFixedFontPixelSize16Dot16 = RequiredFixedFontPixelSize << 10;
+			const FT_Long BestFixedStrikeHeight16Dot16 = InFace->available_sizes[BestStrikeIndex].y_ppem << 10;
+			FixedStrikeScale = FT_DivFix(RequiredFixedFontPixelSize16Dot16, BestFixedStrikeHeight16Dot16);
+		}
 
-	check(Error == 0);
+		// Fixed size fonts don't use the x_scale/y_scale values so we use them to store the scaling adjustment we 
+		// need to apply to the bitmap (as a 16.16 fractional scaling value) to scale it to our desired pixel size
+		// Note: Technically metrics are supposed to be read-only, so if this causes issues then we'll have to add
+		// the information to FFreeTypeFace instead and pass it through to everything that calls ApplySizeAndScale
+		InFace->size->metrics.x_scale = FixedStrikeScale;
+		InFace->size->metrics.y_scale = FixedStrikeScale;
+	}
 }
 
 FT_Error LoadGlyph(FT_Face InFace, const uint32 InGlyphIndex, const int32 InLoadFlags, const int32 InFontSize, const float InFontScale)
@@ -94,22 +177,92 @@ FT_Error LoadGlyph(FT_Face InFace, const uint32 InGlyphIndex, const int32 InLoad
 
 FT_Pos GetHeight(FT_Face InFace, const EFontLayoutMethod InLayoutMethod)
 {
-	return (InLayoutMethod == EFontLayoutMethod::Metrics) ? (FT_Pos)InFace->height : (InFace->bbox.yMax - InFace->bbox.yMin);
+	if (FT_IS_SCALABLE(InFace))
+	{
+		// Scalable fonts use the unscaled value, as the metrics have had scaling and rounding applied to them
+		return (InLayoutMethod == EFontLayoutMethod::Metrics) ? (FT_Pos)InFace->height : (InFace->bbox.yMax - InFace->bbox.yMin);
+	}
+	else if (FT_HAS_FIXED_SIZES(InFace))
+	{
+		// Fixed size fonts don't support scaling, so the metrics value is already unscaled
+		return InFace->size->metrics.height;
+	}
+	return 0;
 }
 
 FT_Pos GetScaledHeight(FT_Face InFace, const EFontLayoutMethod InLayoutMethod)
 {
-	return (InLayoutMethod == EFontLayoutMethod::Metrics) ? InFace->size->metrics.height : FT_MulFix(InFace->bbox.yMax - InFace->bbox.yMin, InFace->size->metrics.y_scale);
+	if (FT_IS_SCALABLE(InFace))
+	{
+		// Scalable fonts use the unscaled value (and apply the scale manually), as the metrics have had rounding applied to them
+		return FT_MulFix((InLayoutMethod == EFontLayoutMethod::Metrics) ? InFace->height : (InFace->bbox.yMax - InFace->bbox.yMin), InFace->size->metrics.y_scale);
+	}
+	else if (FT_HAS_FIXED_SIZES(InFace))
+	{
+		// Fixed size fonts don't support scaling, but we calculated the scale to use for the glyph in ApplySizeAndScale
+		return FT_MulFix(InFace->size->metrics.height, InFace->size->metrics.y_scale);
+	}
+	return 0;
 }
 
 FT_Pos GetAscender(FT_Face InFace, const EFontLayoutMethod InLayoutMethod)
 {
-	return (InLayoutMethod == EFontLayoutMethod::Metrics) ? InFace->size->metrics.ascender : FT_MulFix(InFace->bbox.yMax, InFace->size->metrics.y_scale);
+	if (FT_IS_SCALABLE(InFace))
+	{
+		// Scalable fonts use the unscaled value (and apply the scale manually), as the metrics have had rounding applied to them
+		return FT_MulFix((InLayoutMethod == EFontLayoutMethod::Metrics) ? InFace->ascender : InFace->bbox.yMax, InFace->size->metrics.y_scale);
+	}
+	else if (FT_HAS_FIXED_SIZES(InFace))
+	{
+		// Fixed size fonts don't support scaling, but we calculated the scale to use for the glyph in ApplySizeAndScale
+		return FT_MulFix(InFace->size->metrics.ascender, InFace->size->metrics.y_scale);
+	}
+	return 0;
 }
 
 FT_Pos GetDescender(FT_Face InFace, const EFontLayoutMethod InLayoutMethod)
 {
-	return (InLayoutMethod == EFontLayoutMethod::Metrics) ? InFace->size->metrics.descender : FT_MulFix(InFace->bbox.yMin, InFace->size->metrics.y_scale);
+	if (FT_IS_SCALABLE(InFace))
+	{
+		// Scalable fonts use the unscaled value (and apply the scale manually), as the metrics have had rounding applied to them
+		return FT_MulFix((InLayoutMethod == EFontLayoutMethod::Metrics) ? InFace->descender : InFace->bbox.yMin, InFace->size->metrics.y_scale);
+	}
+	else if (FT_HAS_FIXED_SIZES(InFace))
+	{
+		// Fixed size fonts don't support scaling, but we calculated the scale to use for the glyph in ApplySizeAndScale
+		return FT_MulFix(InFace->size->metrics.descender, InFace->size->metrics.y_scale);
+	}
+	return 0;
+}
+
+float GetBitmapAtlasScale(FT_Face InFace)
+{
+	if (!FT_IS_SCALABLE(InFace) && FT_HAS_FIXED_SIZES(InFace))
+	{
+		// We only scale images down to fit into the atlas
+		// If they're smaller than our desired size we just let them be scaled on the GPU when rendering (see GetBitmapRenderScale)
+		if (InFace->size->metrics.x_scale < ConvertPixelTo16Dot16<FT_Fixed>(1))
+		{
+			// Fixed size fonts don't support scaling, but we calculated the scale to use for the glyph in ApplySizeAndScale
+			return InFace->size->metrics.x_scale / 65536.0f; // 16.16 to pixel scale
+		}
+	}
+	return 1.0f;
+}
+
+float GetBitmapRenderScale(FT_Face InFace)
+{
+	if (!FT_IS_SCALABLE(InFace) && FT_HAS_FIXED_SIZES(InFace))
+	{
+		// We only scale images up when rendering
+		// If they're larger than our desired size we scale them before they go into the atlas (see GetBitmapAtlasScale)
+		if (InFace->size->metrics.x_scale > ConvertPixelTo16Dot16<FT_Fixed>(1))
+		{
+			// Fixed size fonts don't support scaling, but we calculated the scale to use for the glyph in ApplySizeAndScale
+			return InFace->size->metrics.x_scale / 65536.0f; // 16.16 to pixel scale
+		}
+	}
+	return 1.0f;
 }
 
 #endif // WITH_FREETYPE
@@ -137,6 +290,14 @@ FFreeTypeLibrary::FFreeTypeLibrary()
 		
 	FT_Add_Default_Modules(FTLibrary);
 
+#if WITH_FREETYPE_V210
+	// Set the interpreter version based on the hinting method we'd like
+	{
+		FT_UInt FTInterpreterVersion = GSlateEnableLegacyFontHinting ? TT_INTERPRETER_VERSION_35 : TT_INTERPRETER_VERSION_40;
+		FT_Property_Set(FTLibrary, "truetype", "interpreter-version", &FTInterpreterVersion);
+	}
+#endif	// WITH_FREETYPE_V210
+
 	static bool bLoggedVersion = false;
 
 	// Write out the version of freetype for debugging purposes
@@ -175,7 +336,7 @@ FFreeTypeFace::FFreeTypeFace(const FFreeTypeLibrary* InFTLibrary, FFontFaceDataC
 	if (Error)
 	{
 		// You can look these error codes up in the FreeType docs: https://www.freetype.org/freetype2/docs/reference/ft2-error_code_values.html
-		ensureAlwaysMsgf(0, TEXT("FT_New_Memory_Face failed with error code %i"), Error);
+		//ensureAlwaysMsgf(0, TEXT("FT_New_Memory_Face failed with error code %i"), Error);
 
 		// We assume the face is null if the function errored
 		ensureAlwaysMsgf(FTFace == nullptr, TEXT("FT_New_Memory_Face failed but also returned a non-null FT_Face. This is unexpected and may leak memory!"));
@@ -467,15 +628,17 @@ bool FFreeTypeAdvanceCache::FindOrCache(FT_Face InFace, const uint32 InGlyphInde
 		}
 	}
 
-	FreeTypeUtils::ApplySizeAndScale(InFace, InFontSize, 1.0f);
+	FreeTypeUtils::ApplySizeAndScale(InFace, InFontSize, InFontScale);
 
 	// No cached data, go ahead and add an entry for it...
 	FT_Error Error = FT_Get_Advance(InFace, InGlyphIndex, InLoadFlags, &OutCachedAdvance);
 	if (Error == 0)
 	{
-		// We apply our own scaling as FreeType doesn't always produce the correct results for all fonts when applying the scale via the transform matrix
-		const FT_Long FixedFontScale = FreeTypeUtils::ConvertPixelTo16Dot16<FT_Long>(InFontScale);
-		OutCachedAdvance = FT_MulFix(OutCachedAdvance, FixedFontScale);
+		if (FT_HAS_FIXED_SIZES(InFace))
+		{
+			// Fixed size fonts don't support scaling, but we calculated the scale to use for the glyph in ApplySizeAndScale
+			OutCachedAdvance = FT_MulFix(OutCachedAdvance, ((InLoadFlags & FT_LOAD_VERTICAL_LAYOUT) ? InFace->size->metrics.y_scale : InFace->size->metrics.x_scale));
+		}
 
 		CachedAdvanceMap.Add(CachedAdvanceKey, OutCachedAdvance);
 		return true;
@@ -518,7 +681,7 @@ bool FFreeTypeKerningPairCache::FindOrCache(FT_Face InFace, const FKerningPair&
 		}
 	}
 
-	FreeTypeUtils::ApplySizeAndScale(InFace, InFontSize, 1.0f);
+	FreeTypeUtils::ApplySizeAndScale(InFace, InFontSize, InFontScale);
 
 	// No cached data, go ahead and add an entry for it...
 	FT_Error Error = FT_Get_Kerning(InFace, InKerningPair.FirstGlyphIndex, InKerningPair.SecondGlyphIndex, InKerningFlags, &OutKerning);
@@ -526,10 +689,12 @@ bool FFreeTypeKerningPairCache::FindOrCache(FT_Face InFace, const FKerningPair&
 	{
 		if (InKerningFlags != FT_KERNING_UNSCALED)
 		{
-			// We apply our own scaling as FreeType doesn't always produce the correct results for all fonts when applying the scale via the transform matrix
-			const FT_Long FixedFontScale = FreeTypeUtils::ConvertPixelTo16Dot16<FT_Long>(InFontScale);
-			OutKerning.x = FT_MulFix(OutKerning.x, FixedFontScale);
-			OutKerning.y = FT_MulFix(OutKerning.y, FixedFontScale);
+			if (FT_HAS_FIXED_SIZES(InFace))
+			{
+				// Fixed size fonts don't support scaling, but we calculated the scale to use for the glyph in ApplySizeAndScale
+				OutKerning.x = FT_MulFix(OutKerning.x, InFace->size->metrics.x_scale);
+				OutKerning.y = FT_MulFix(OutKerning.y, InFace->size->metrics.y_scale);
+			}
 		}
 
 		CachedKerningPairMap.Add(CachedKerningPairKey, OutKerning);
diff --git a/Engine/Source/Runtime/SlateCore/Private/Fonts/FontCacheFreeType.h b/Engine/Source/Runtime/SlateCore/Private/Fonts/FontCacheFreeType.h
--- a/Engine/Source/Runtime/SlateCore/Private/Fonts/FontCacheFreeType.h
+++ b/Engine/Source/Runtime/SlateCore/Private/Fonts/FontCacheFreeType.h
@@ -4,12 +4,19 @@
 
 #include "CoreMinimal.h"
 #include "Fonts/FontCache.h"
+#include "Templates/IsIntegral.h"
+#include "Templates/IsFloatingPoint.h"
 
 #ifndef WITH_FREETYPE
 	#define WITH_FREETYPE	0
 #endif // WITH_FREETYPE
 
 
+#ifndef WITH_FREETYPE_V210
+	#define WITH_FREETYPE_V210	0
+#endif // WITH_FREETYPE_V210
+
+
 #if PLATFORM_COMPILER_HAS_GENERIC_KEYWORD
 	#define generic __identifier(generic)
 #endif	//PLATFORM_COMPILER_HAS_GENERIC_KEYWORD
@@ -21,6 +28,9 @@
 
 	// FreeType style include
 	#include FT_FREETYPE_H
+#if WITH_FREETYPE_V210
+	#include FT_DRIVER_H
+#endif	// WITH_FREETYPE_V210
 	#include FT_GLYPH_H
 	#include FT_MODULE_H
 	#include FT_BITMAP_H
@@ -37,10 +47,8 @@
 
 namespace FreeTypeConstants
 {
-	/** The horizontal DPI we render at */
-	const uint32 HorizontalDPI = 96;
-	/** The vertical DPI we render at */
-	const uint32 VerticalDPI = 96;
+	/** The horizontal DPI we render at (horizontal and vertical) */
+	const uint32 RenderDPI = 96;
 } // namespace FreeTypeConstants
 
 
@@ -83,25 +91,58 @@ FT_Pos GetAscender(FT_Face InFace, const EFontLayoutMethod InLayoutMethod);
  */
 FT_Pos GetDescender(FT_Face InFace, const EFontLayoutMethod InLayoutMethod);
 
+/**
+ * Get any additional scale that should be applied when producing the atlas glyphs from this face.
+ * @note ApplySizeAndScale must have been called prior to this function to prepare the face.
+ */
+float GetBitmapAtlasScale(FT_Face InFace);
+
+/**
+ * Get any additional scale that should be applied when rendering glyphs from this face.
+ * @note ApplySizeAndScale must have been called prior to this function to prepare the face.
+ */
+float GetBitmapRenderScale(FT_Face InFace);
+
 #endif // WITH_FREETYPE
 
 /** Convert the given value from 26.6 space into rounded pixel space */
 template <typename TRetType, typename TParamType>
-FORCEINLINE TRetType Convert26Dot6ToRoundedPixel(TParamType InValue)
+FORCEINLINE typename TEnableIf<TIsIntegral<TParamType>::Value, TRetType>::Type Convert26Dot6ToRoundedPixel(TParamType InValue)
+{
+	return static_cast<TRetType>((InValue + (1<<5)) >> 6);
+}
+
+/** Convert the given value from 26.6 space into rounded pixel space */
+template <typename TRetType, typename TParamType>
+FORCEINLINE typename TEnableIf<TIsFloatingPoint<TParamType>::Value, TRetType>::Type Convert26Dot6ToRoundedPixel(TParamType InValue)
 {
 	return static_cast<TRetType>(FMath::RoundToInt(InValue / 64.0f));
 }
 
 /** Convert the given value from pixel space into 26.6 space */
 template <typename TRetType, typename TParamType>
-FORCEINLINE TRetType ConvertPixelTo26Dot6(TParamType InValue)
+FORCEINLINE typename TEnableIf<TIsIntegral<TParamType>::Value, TRetType>::Type ConvertPixelTo26Dot6(TParamType InValue)
+{
+	return static_cast<TRetType>(InValue << 6);
+}
+
+/** Convert the given value from pixel space into 26.6 space */
+template <typename TRetType, typename TParamType>
+FORCEINLINE typename TEnableIf<TIsFloatingPoint<TParamType>::Value, TRetType>::Type ConvertPixelTo26Dot6(TParamType InValue)
 {
 	return static_cast<TRetType>(InValue * 64);
 }
 
 /** Convert the given value from pixel space into 16.16 space */
 template <typename TRetType, typename TParamType>
-FORCEINLINE TRetType ConvertPixelTo16Dot16(TParamType InValue)
+FORCEINLINE typename TEnableIf<TIsIntegral<TParamType>::Value, TRetType>::Type ConvertPixelTo16Dot16(TParamType InValue)
+{
+	return static_cast<TRetType>(InValue << 16);
+}
+
+/** Convert the given value from pixel space into 16.16 space */
+template <typename TRetType, typename TParamType>
+FORCEINLINE typename TEnableIf<TIsFloatingPoint<TParamType>::Value, TRetType>::Type ConvertPixelTo16Dot16(TParamType InValue)
 {
 	return static_cast<TRetType>(InValue * 65536);
 }
@@ -183,6 +224,16 @@ public:
 	{
 		return FreeTypeUtils::GetDescender(FTFace, LayoutMethod);
 	}
+
+	FORCEINLINE float GetBitmapAtlasScale() const
+	{
+		return FreeTypeUtils::GetBitmapAtlasScale(FTFace);
+	}
+
+	FORCEINLINE float GetBitmapRenderScale() const
+	{
+		return FreeTypeUtils::GetBitmapRenderScale(FTFace);
+	}
 #endif // WITH_FREETYPE
 
 	FORCEINLINE const TSet<FName>& GetAttributes() const
diff --git a/Engine/Source/Runtime/SlateCore/Private/Fonts/FontCacheHarfBuzz.cpp b/Engine/Source/Runtime/SlateCore/Private/Fonts/FontCacheHarfBuzz.cpp
--- a/Engine/Source/Runtime/SlateCore/Private/Fonts/FontCacheHarfBuzz.cpp
+++ b/Engine/Source/Runtime/SlateCore/Private/Fonts/FontCacheHarfBuzz.cpp
@@ -101,12 +101,13 @@ hb_user_data_key_t UserDataKey;
 
 struct FUserData
 {
-	FUserData(const int32 InFontSize, const float InFontScale, FFreeTypeGlyphCache* InFTGlyphCache, FFreeTypeAdvanceCache* InFTAdvanceCache, FFreeTypeKerningPairCache* InFTKerningPairCache)
+	FUserData(const int32 InFontSize, const float InFontScale, FFreeTypeGlyphCache* InFTGlyphCache, FFreeTypeAdvanceCache* InFTAdvanceCache, FFreeTypeKerningPairCache* InFTKerningPairCache, const hb_font_extents_t& InHarfBuzzFontExtents)
 		: FontSize(InFontSize)
 		, FontScale(InFontScale)
 		, FTGlyphCache(InFTGlyphCache)
 		, FTAdvanceCache(InFTAdvanceCache)
 		, FTKerningPairCache(InFTKerningPairCache)
+		, HarfBuzzFontExtents(InHarfBuzzFontExtents)
 	{
 	}
 
@@ -115,11 +116,12 @@ struct FUserData
 	FFreeTypeGlyphCache* FTGlyphCache;
 	FFreeTypeAdvanceCache* FTAdvanceCache;
 	FFreeTypeKerningPairCache* FTKerningPairCache;
+	hb_font_extents_t HarfBuzzFontExtents;
 };
 
-void* CreateUserData(const int32 InFontSize, const float InFontScale, FFreeTypeGlyphCache* InFTGlyphCache, FFreeTypeAdvanceCache* InFTAdvanceCache, FFreeTypeKerningPairCache* InFTKerningPairCache)
+void* CreateUserData(const int32 InFontSize, const float InFontScale, FFreeTypeGlyphCache* InFTGlyphCache, FFreeTypeAdvanceCache* InFTAdvanceCache, FFreeTypeKerningPairCache* InFTKerningPairCache, const hb_font_extents_t& InHarfBuzzFontExtents)
 {
-	return new FUserData(InFontSize, InFontScale, InFTGlyphCache, InFTAdvanceCache, InFTKerningPairCache);
+	return new FUserData(InFontSize, InFontScale, InFTGlyphCache, InFTAdvanceCache, InFTKerningPairCache, InHarfBuzzFontExtents);
 }
 
 void DestroyUserData(void* UserData)
@@ -145,50 +147,150 @@ FORCEINLINE int32 get_ft_flags(hb_font_t* InFont)
 	return hb_ft_font_get_load_flags(FontParent);
 }
 
-hb_bool_t get_nominal_glyph(hb_font_t* InFont, void* InFontData, hb_codepoint_t InUnicodeChar, hb_codepoint_t *OutGlyphIndex, void *InUserData)
+hb_bool_t get_font_h_extents(hb_font_t* InFont, void* InFontData, hb_font_extents_t* OutMetrics, void* InUserData)
+{
+	const FUserData* UserDataPtr = static_cast<FUserData*>(hb_font_get_user_data(InFont, &UserDataKey));
+	*OutMetrics = UserDataPtr->HarfBuzzFontExtents;
+	return true;
+}
+
+unsigned int get_nominal_glyphs(hb_font_t* InFont, void* InFontData, unsigned int InCount, const hb_codepoint_t *InUnicodeCharBuffer, unsigned int InUnicodeCharBufferStride, hb_codepoint_t *OutGlyphIndexBuffer, unsigned int InGlyphIndexBufferStride, void* InUserData)
 {
 	FT_Face FreeTypeFace = get_ft_face(InFont);
 
-	*OutGlyphIndex = FT_Get_Char_Index(FreeTypeFace, InUnicodeChar);
+	const uint8* UnicodeCharRawBuffer = (const uint8*)InUnicodeCharBuffer;
+	uint8* GlyphIndexRawBuffer = (uint8*)OutGlyphIndexBuffer;
 
-	// If the given font can't render that character (as the fallback font may be missing), try again with the fallback character
-	if (InUnicodeChar != 0 && *OutGlyphIndex == 0)
+	for (unsigned int ItemIndex = 0; ItemIndex < InCount; ++ItemIndex)
 	{
-		*OutGlyphIndex = FT_Get_Char_Index(FreeTypeFace, SlateFontRendererUtils::InvalidSubChar);
+		const hb_codepoint_t* UnicodeCharPtr = (const hb_codepoint_t*)UnicodeCharRawBuffer;
+		hb_codepoint_t* OutGlyphIndexPtr = (hb_codepoint_t*)GlyphIndexRawBuffer;
+
+		*OutGlyphIndexPtr = FT_Get_Char_Index(FreeTypeFace, *UnicodeCharPtr);
+
+		// If the given font can't render that character (as the fallback font may be missing), try again with the fallback character
+		if (*UnicodeCharPtr != 0 && *OutGlyphIndexPtr == 0)
+		{
+			*OutGlyphIndexPtr = FT_Get_Char_Index(FreeTypeFace, SlateFontRendererUtils::InvalidSubChar);
+		}
+
+		// If this resolution failed, return the number if items we managed to process
+		if (*UnicodeCharPtr != 0 && *OutGlyphIndexPtr == 0)
+		{
+			return ItemIndex;
+		}
+
+		// Advance the buffers
+		UnicodeCharRawBuffer += InUnicodeCharBufferStride;
+		GlyphIndexRawBuffer += InGlyphIndexBufferStride;
 	}
 
-	return InUnicodeChar == 0 || *OutGlyphIndex != 0;
+	// Processed everything - return the count
+	return InCount;
 }
 
-hb_position_t get_glyph_h_advance(hb_font_t* InFont, void* InFontData, hb_codepoint_t InGlyphIndex, void* InUserData)
+hb_bool_t get_nominal_glyph(hb_font_t* InFont, void* InFontData, hb_codepoint_t InUnicodeChar, hb_codepoint_t* OutGlyphIndex, void* InUserData)
+{
+	return get_nominal_glyphs(InFont, InFontData, 1, &InUnicodeChar, sizeof(hb_codepoint_t), OutGlyphIndex, sizeof(hb_codepoint_t), InUserData) == 1;
+}
+
+void get_glyph_h_advances(hb_font_t* InFont, void* InFontData, unsigned int InCount, const hb_codepoint_t* InGlyphIndexBuffer, unsigned int InGlyphIndexBufferStride, hb_position_t* OutAdvanceBuffer, unsigned int InAdvanceBufferStride, void* InUserData)
 {
 	FT_Face FreeTypeFace = get_ft_face(InFont);
 	const int32 FreeTypeFlags = get_ft_flags(InFont);
 	const FUserData* UserDataPtr = static_cast<FUserData*>(hb_font_get_user_data(InFont, &UserDataKey));
 
-	FT_Fixed CachedAdvanceData = 0;
-	if (UserDataPtr->FTAdvanceCache->FindOrCache(FreeTypeFace, InGlyphIndex, FreeTypeFlags, UserDataPtr->FontSize, UserDataPtr->FontScale, CachedAdvanceData))
+	int ScaleMultiplier = 1;
+	{
+		int FontXScale = 0;
+		int FontYScale = 0;
+		hb_font_get_scale(InFont, &FontXScale, &FontYScale);
+
+		if (FontXScale < 0)
+		{
+			ScaleMultiplier = -1;
+		}
+	}
+
+	const uint8* GlyphIndexRawBuffer = (const uint8*)InGlyphIndexBuffer;
+	uint8* AdvanceRawBuffer = (uint8*)OutAdvanceBuffer;
+
+	for (unsigned int ItemIndex = 0; ItemIndex < InCount; ++ItemIndex)
 	{
-		return (CachedAdvanceData + (1<<9)) >> 10;
+		const hb_codepoint_t* GlyphIndexPtr = (const hb_codepoint_t*)GlyphIndexRawBuffer;
+		hb_position_t* OutAdvancePtr = (hb_position_t*)AdvanceRawBuffer;
+
+		FT_Fixed CachedAdvanceData = 0;
+		if (UserDataPtr->FTAdvanceCache->FindOrCache(FreeTypeFace, *GlyphIndexPtr, FreeTypeFlags, UserDataPtr->FontSize, UserDataPtr->FontScale, CachedAdvanceData))
+		{
+			*OutAdvancePtr = ((CachedAdvanceData * ScaleMultiplier) + (1<<9)) >> 10;
+		}
+		else
+		{
+			*OutAdvancePtr = 0;
+		}
+
+		// Advance the buffers
+		GlyphIndexRawBuffer += InGlyphIndexBufferStride;
+		AdvanceRawBuffer += InAdvanceBufferStride;
 	}
+}
 
-	return 0;
+hb_position_t get_glyph_h_advance(hb_font_t* InFont, void* InFontData, hb_codepoint_t InGlyphIndex, void* InUserData)
+{
+	hb_position_t Advance = 0;
+	get_glyph_h_advances(InFont, InFontData, 1, &InGlyphIndex, sizeof(hb_codepoint_t), &Advance, sizeof(hb_position_t), InUserData);
+	return Advance;
 }
 
-hb_position_t get_glyph_v_advance(hb_font_t* InFont, void* InFontData, hb_codepoint_t InGlyphIndex, void* InUserData)
+void get_glyph_v_advances(hb_font_t* InFont, void* InFontData, unsigned int InCount, const hb_codepoint_t* InGlyphIndexBuffer, unsigned int InGlyphIndexBufferStride, hb_position_t* OutAdvanceBuffer, unsigned int InAdvanceBufferStride, void* InUserData)
 {
 	FT_Face FreeTypeFace = get_ft_face(InFont);
 	const int32 FreeTypeFlags = get_ft_flags(InFont);
 	const FUserData* UserDataPtr = static_cast<FUserData*>(hb_font_get_user_data(InFont, &UserDataKey));
 
-	FT_Fixed CachedAdvanceData = 0;
-	if (UserDataPtr->FTAdvanceCache->FindOrCache(FreeTypeFace, InGlyphIndex, FreeTypeFlags | FT_LOAD_VERTICAL_LAYOUT, UserDataPtr->FontSize, UserDataPtr->FontScale, CachedAdvanceData))
+	int ScaleMultiplier = 1;
 	{
-		// Note: FreeType's vertical metrics grows downward while other FreeType coordinates have a Y growing upward.  Hence the extra negation.
-		return (-CachedAdvanceData + (1<<9)) >> 10;
+		int FontXScale = 0;
+		int FontYScale = 0;
+		hb_font_get_scale(InFont, &FontXScale, &FontYScale);
+
+		if (FontYScale < 0)
+		{
+			ScaleMultiplier = -1;
+		}
 	}
 
-	return 0;
+	const uint8* GlyphIndexRawBuffer = (const uint8*)InGlyphIndexBuffer;
+	uint8* AdvanceRawBuffer = (uint8*)OutAdvanceBuffer;
+
+	for (unsigned int ItemIndex = 0; ItemIndex < InCount; ++ItemIndex)
+	{
+		const hb_codepoint_t* GlyphIndexPtr = (const hb_codepoint_t*)GlyphIndexRawBuffer;
+		hb_position_t* OutAdvancePtr = (hb_position_t*)AdvanceRawBuffer;
+
+		FT_Fixed CachedAdvanceData = 0;
+		if (UserDataPtr->FTAdvanceCache->FindOrCache(FreeTypeFace, *GlyphIndexPtr, FreeTypeFlags | FT_LOAD_VERTICAL_LAYOUT, UserDataPtr->FontSize, UserDataPtr->FontScale, CachedAdvanceData))
+		{
+			// Note: FreeType's vertical metrics grows downward while other FreeType coordinates have a Y growing upward. Hence the extra negation.
+			*OutAdvancePtr = ((-CachedAdvanceData * ScaleMultiplier) + (1<<9)) >> 10;
+		}
+		else
+		{
+			*OutAdvancePtr = 0;
+		}
+
+		// Advance the buffers
+		GlyphIndexRawBuffer += InGlyphIndexBufferStride;
+		AdvanceRawBuffer += InAdvanceBufferStride;
+	}
+}
+
+hb_position_t get_glyph_v_advance(hb_font_t* InFont, void* InFontData, hb_codepoint_t InGlyphIndex, void* InUserData)
+{
+	hb_position_t Advance = 0;
+	get_glyph_v_advances(InFont, InFontData, 1, &InGlyphIndex, sizeof(hb_codepoint_t), &Advance, sizeof(hb_position_t), InUserData);
+	return Advance;
 }
 
 hb_bool_t get_glyph_v_origin(hb_font_t* InFont, void* InFontData, hb_codepoint_t InGlyphIndex, hb_position_t* OutX, hb_position_t* OutY, void* InUserData)
@@ -200,7 +302,7 @@ hb_bool_t get_glyph_v_origin(hb_font_t* InFont, void* InFontData, hb_codepoint_t
 	FFreeTypeGlyphCache::FCachedGlyphData CachedGlyphData;
 	if (UserDataPtr->FTGlyphCache->FindOrCache(FreeTypeFace, InGlyphIndex, FreeTypeFlags, UserDataPtr->FontSize, UserDataPtr->FontScale, CachedGlyphData))
 	{
-		// Note: FreeType's vertical metrics grows downward while other FreeType coordinates have a Y growing upward.  Hence the extra negation.
+		// Note: FreeType's vertical metrics grows downward while other FreeType coordinates have a Y growing upward. Hence the extra negation.
 		*OutX = CachedGlyphData.GlyphMetrics.horiBearingX -   CachedGlyphData.GlyphMetrics.vertBearingX;
 		*OutY = CachedGlyphData.GlyphMetrics.horiBearingY - (-CachedGlyphData.GlyphMetrics.vertBearingY);
 
@@ -251,6 +353,23 @@ hb_bool_t get_glyph_extents(hb_font_t* InFont, void* InFontData, hb_codepoint_t
 		OutExtents->y_bearing	=  CachedGlyphData.GlyphMetrics.horiBearingY;
 		OutExtents->width		=  CachedGlyphData.GlyphMetrics.width;
 		OutExtents->height		= -CachedGlyphData.GlyphMetrics.height;
+
+		int FontXScale = 0;
+		int FontYScale = 0;
+		hb_font_get_scale(InFont, &FontXScale, &FontYScale);
+
+		if (FontXScale < 0)
+		{
+			OutExtents->x_bearing = -OutExtents->x_bearing;
+			OutExtents->width = -OutExtents->width;
+		}
+
+		if (FontYScale < 0)
+		{
+			OutExtents->y_bearing = -OutExtents->y_bearing;
+			OutExtents->height = -OutExtents->height;
+		}
+
 		return true;
 	}
 
@@ -295,9 +414,19 @@ FHarfBuzzFontFactory::FHarfBuzzFontFactory(FFreeTypeGlyphCache* InFTGlyphCache,
 #if WITH_HARFBUZZ
 	CustomHarfBuzzFuncs = hb_font_funcs_create();
 
+	hb_font_funcs_set_font_h_extents_func(CustomHarfBuzzFuncs, &HarfBuzzFontFunctions::Internal::get_font_h_extents, nullptr, nullptr);
 	hb_font_funcs_set_nominal_glyph_func(CustomHarfBuzzFuncs, &HarfBuzzFontFunctions::Internal::get_nominal_glyph, nullptr, nullptr);
+#if WITH_HARFBUZZ_V24
+	hb_font_funcs_set_nominal_glyphs_func(CustomHarfBuzzFuncs, &HarfBuzzFontFunctions::Internal::get_nominal_glyphs, nullptr, nullptr);
+#endif // WITH_HARFBUZZ_V24
 	hb_font_funcs_set_glyph_h_advance_func(CustomHarfBuzzFuncs, &HarfBuzzFontFunctions::Internal::get_glyph_h_advance, nullptr, nullptr);
+#if WITH_HARFBUZZ_V24
+	hb_font_funcs_set_glyph_h_advances_func(CustomHarfBuzzFuncs, &HarfBuzzFontFunctions::Internal::get_glyph_h_advances, nullptr, nullptr);
+#endif // WITH_HARFBUZZ_V24
 	hb_font_funcs_set_glyph_v_advance_func(CustomHarfBuzzFuncs, &HarfBuzzFontFunctions::Internal::get_glyph_v_advance, nullptr, nullptr);
+#if WITH_HARFBUZZ_V24
+	hb_font_funcs_set_glyph_v_advances_func(CustomHarfBuzzFuncs, &HarfBuzzFontFunctions::Internal::get_glyph_v_advances, nullptr, nullptr);
+#endif // WITH_HARFBUZZ_V24
 	hb_font_funcs_set_glyph_v_origin_func(CustomHarfBuzzFuncs, &HarfBuzzFontFunctions::Internal::get_glyph_v_origin, nullptr, nullptr);
 	hb_font_funcs_set_glyph_h_kerning_func(CustomHarfBuzzFuncs, &HarfBuzzFontFunctions::Internal::get_glyph_h_kerning, nullptr, nullptr);
 	hb_font_funcs_set_glyph_extents_func(CustomHarfBuzzFuncs, &HarfBuzzFontFunctions::Internal::get_glyph_extents, nullptr, nullptr);
@@ -323,9 +452,11 @@ hb_font_t* FHarfBuzzFontFactory::CreateFont(const FFreeTypeFace& InFace, const u
 
 #if WITH_FREETYPE
 	FT_Face FreeTypeFace = InFace.GetFace();
-
 	FreeTypeUtils::ApplySizeAndScale(FreeTypeFace, InFontSize, InFontScale);
 
+	hb_font_extents_t HarfBuzzFontExtents;
+	FMemory::Memzero(HarfBuzzFontExtents);
+
 	// Create a sub-font from the default FreeType implementation so we can override some font functions to provide low-level caching
 	{
 		hb_font_t* HarfBuzzFTFont = hb_ft_font_create(FreeTypeFace, nullptr);
@@ -337,11 +468,16 @@ hb_font_t* FHarfBuzzFontFactory::CreateFont(const FFreeTypeFace& InFace, const u
 			int HarfBuzzFTFontYScale = 0;
 			hb_font_get_scale(HarfBuzzFTFont, &HarfBuzzFTFontXScale, &HarfBuzzFTFontYScale);
 
-			const FT_Long FixedFontScale = FreeTypeUtils::ConvertPixelTo16Dot16<FT_Long>(InFontScale);
-			HarfBuzzFTFontXScale = FT_MulFix(HarfBuzzFTFontXScale, FixedFontScale);
-			HarfBuzzFTFontYScale = FT_MulFix(HarfBuzzFTFontYScale, FixedFontScale);
-
-			hb_font_set_scale(HarfBuzzFTFont, HarfBuzzFTFontXScale, HarfBuzzFTFontYScale);
+			// Cache the font extents
+			HarfBuzzFontExtents.ascender = InFace.GetAscender();
+			HarfBuzzFontExtents.descender = InFace.GetDescender();
+			HarfBuzzFontExtents.line_gap = InFace.GetScaledHeight() - (HarfBuzzFontExtents.ascender - HarfBuzzFontExtents.descender);
+			if (HarfBuzzFTFontYScale < 0)
+			{
+				HarfBuzzFontExtents.ascender = -HarfBuzzFontExtents.ascender;
+				HarfBuzzFontExtents.descender = -HarfBuzzFontExtents.descender;
+				HarfBuzzFontExtents.line_gap = -HarfBuzzFontExtents.line_gap;
+			}
 		}
 
 		HarfBuzzFont = hb_font_create_sub_font(HarfBuzzFTFont);
@@ -354,7 +490,7 @@ hb_font_t* FHarfBuzzFontFactory::CreateFont(const FFreeTypeFace& InFace, const u
 	hb_font_set_user_data(
 		HarfBuzzFont, 
 		&HarfBuzzFontFunctions::UserDataKey, 
-		HarfBuzzFontFunctions::CreateUserData(InFontSize, InFontScale, FTGlyphCache, FTAdvanceCache, FTKerningPairCache), 
+		HarfBuzzFontFunctions::CreateUserData(InFontSize, InFontScale, FTGlyphCache, FTAdvanceCache, FTKerningPairCache, HarfBuzzFontExtents),
 		&HarfBuzzFontFunctions::DestroyUserData, 
 		true
 		);
diff --git a/Engine/Source/Runtime/SlateCore/Private/Fonts/FontCacheHarfBuzz.h b/Engine/Source/Runtime/SlateCore/Private/Fonts/FontCacheHarfBuzz.h
--- a/Engine/Source/Runtime/SlateCore/Private/Fonts/FontCacheHarfBuzz.h
+++ b/Engine/Source/Runtime/SlateCore/Private/Fonts/FontCacheHarfBuzz.h
@@ -10,9 +10,13 @@ class FFreeTypeGlyphCache;
 class FFreeTypeKerningPairCache;
 
 #ifndef WITH_HARFBUZZ
-	#define WITH_HARFBUZZ	0
+	#define WITH_HARFBUZZ 0
 #endif // WITH_HARFBUZZ
 
+#ifndef WITH_HARFBUZZ_V24
+	#define WITH_HARFBUZZ_V24 0
+#endif // WITH_HARFBUZZ_V24
+
 #if PLATFORM_COMPILER_HAS_GENERIC_KEYWORD
 	#define generic __identifier(generic)
 #endif	//PLATFORM_COMPILER_HAS_GENERIC_KEYWORD
diff --git a/Engine/Source/Runtime/SlateCore/Private/Fonts/FontTypes.cpp b/Engine/Source/Runtime/SlateCore/Private/Fonts/FontTypes.cpp
--- a/Engine/Source/Runtime/SlateCore/Private/Fonts/FontTypes.cpp
+++ b/Engine/Source/Runtime/SlateCore/Private/Fonts/FontTypes.cpp
@@ -3,24 +3,24 @@
 #include "Fonts/FontTypes.h"
 
 
-FSlateFontAtlas::FSlateFontAtlas( uint32 InWidth, uint32 InHeight )
-	: FSlateTextureAtlas( InWidth, InHeight, sizeof(uint8), ESlateTextureAtlasPaddingStyle::PadWithZero, true )
+FSlateFontAtlas::FSlateFontAtlas(uint32 InWidth, uint32 InHeight, const bool InIsGrayscale)
+	: FSlateTextureAtlas(InWidth, InHeight, InIsGrayscale ? 1 : 4, ESlateTextureAtlasPaddingStyle::PadWithZero, true)
 {
 }
 
 FSlateFontAtlas::~FSlateFontAtlas()
 {
-
 }	
 
-/** 
- * Adds a character to the texture.
- *
- * @param CharInfo	Information about the size of the character
- */
+bool FSlateFontAtlas::IsGrayscale() const
+{
+	return BytesPerPixel == 1;
+}
+
 const FAtlasedTextureSlot* FSlateFontAtlas::AddCharacter( const FCharacterRenderData& RenderData )
 {
-	return AddTexture( RenderData.MeasureInfo.SizeX, RenderData.MeasureInfo.SizeY, RenderData.RawPixels );
+	check(RenderData.bIsGrayscale == IsGrayscale());
+	return AddTexture( RenderData.SizeX, RenderData.SizeY, RenderData.RawPixels );
 }
 
 void FSlateFontAtlas::Flush()
diff --git a/Engine/Source/Runtime/SlateCore/Private/Fonts/SlateFontRenderer.cpp b/Engine/Source/Runtime/SlateCore/Private/Fonts/SlateFontRenderer.cpp
--- a/Engine/Source/Runtime/SlateCore/Private/Fonts/SlateFontRenderer.cpp
+++ b/Engine/Source/Runtime/SlateCore/Private/Fonts/SlateFontRenderer.cpp
@@ -36,14 +36,20 @@ namespace SlateFontRendererUtils
 
 #if WITH_FREETYPE
 
-void AppendGlyphFlags(const FFontData& InFontData, uint32& InOutGlyphFlags)
+void AppendGlyphFlags(const FFreeTypeFace& InFace, const FFontData& InFontData, uint32& InOutGlyphFlags)
 {
 	// Setup additional glyph flags
-	InOutGlyphFlags |= GlobalGlyphFlags;
+#if WITH_FREETYPE_V210
+	InOutGlyphFlags |= FT_LOAD_COLOR;
+#endif	// WITH_FREETYPE_V210
+	if (FT_IS_SCALABLE(InFace.GetFace()))
+	{
+		InOutGlyphFlags |= FT_LOAD_NO_BITMAP;
+	}
 
-	if(EnableFontAntiAliasing)
+	if (EnableFontAntiAliasing)
 	{
-		switch(InFontData.GetHinting())
+		switch (InFontData.GetHinting())
 		{
 		case EFontHinting::Auto:		InOutGlyphFlags |= FT_LOAD_FORCE_AUTOHINT; break;
 		case EFontHinting::AutoLight:	InOutGlyphFlags |= FT_LOAD_TARGET_LIGHT; break;
@@ -61,6 +67,105 @@ void AppendGlyphFlags(const FFontData& InFontData, uint32& InOutGlyphFlags)
 
 #endif // WITH_FREETYPE
 
+template <int32 BytesPerPixel>
+void ResizeFontBitmap_SampleInputPixel(const TArray<uint8>& SrcData, const int32 PixelIndex, FVector4& OutStepColor);
+
+template <>
+void ResizeFontBitmap_SampleInputPixel<1>(const TArray<uint8>& SrcData, const int32 PixelIndex, FVector4& OutStepColor)
+{
+	OutStepColor.X += (float)SrcData[PixelIndex];
+}
+
+template <>
+void ResizeFontBitmap_SampleInputPixel<4>(const TArray<uint8>& SrcData, const int32 PixelIndex, FVector4& OutStepColor)
+{
+	const int32 SrcDataIndex = PixelIndex * 4;
+	OutStepColor.X += (float)SrcData[SrcDataIndex+0];
+	OutStepColor.Y += (float)SrcData[SrcDataIndex+1];
+	OutStepColor.Z += (float)SrcData[SrcDataIndex+2];
+	OutStepColor.W += (float)SrcData[SrcDataIndex+3];
+}
+
+template <int32 BytesPerPixel>
+void ResizeFontBitmap_StoreAverageOutputPixel(TArray<uint8>& DstData, const int32 PixelIndex, const FVector4& StepColor, const int32 PixelCount);
+
+template <>
+void ResizeFontBitmap_StoreAverageOutputPixel<1>(TArray<uint8>& DstData, const int32 PixelIndex, const FVector4& StepColor, const int32 PixelCount)
+{
+	DstData[PixelIndex] = FMath::Clamp(FMath::TruncToInt(StepColor.X / (float)PixelCount), 0, 255);
+}
+
+template <>
+void ResizeFontBitmap_StoreAverageOutputPixel<4>(TArray<uint8>& DstData, const int32 PixelIndex, const FVector4& StepColor, const int32 PixelCount)
+{
+	const int32 DstDataIndex = PixelIndex * 4;
+	DstData[DstDataIndex+0] = FMath::Clamp(FMath::TruncToInt(StepColor.X / (float)PixelCount), 0, 255);
+	DstData[DstDataIndex+1] = FMath::Clamp(FMath::TruncToInt(StepColor.Y / (float)PixelCount), 0, 255);
+	DstData[DstDataIndex+2] = FMath::Clamp(FMath::TruncToInt(StepColor.Z / (float)PixelCount), 0, 255);
+	DstData[DstDataIndex+3] = FMath::Clamp(FMath::TruncToInt(StepColor.W / (float)PixelCount), 0, 255);
+}
+
+// Note: Copied from FImageUtils and updated to work efficiently with 1 or 4-bytes per-pixel, and include the alpha channel (which the FImageUtils version just clobbers)
+template <int32 BytesPerPixel>
+void ResizeFontBitmap(const int32 SrcWidth, const int32 SrcHeight, const TArray<uint8>& SrcData, const int32 DstWidth, const int32 DstHeight, TArray<uint8>& DstData)
+{
+	static_assert(BytesPerPixel == 1 || BytesPerPixel == 4, "BytesPerPixel must be 1 or 4!");
+
+	check(SrcData.Num() >= SrcWidth * SrcHeight * BytesPerPixel);
+	DstData.SetNumZeroed(DstWidth * DstHeight * BytesPerPixel);
+
+	float SrcX = 0;
+	float SrcY = 0;
+
+	const float StepSizeX = SrcWidth / (float)DstWidth;
+	const float StepSizeY = SrcHeight / (float)DstHeight;
+
+	for (int32 Y = 0; Y < DstHeight; ++Y)
+	{
+		int32 PixelPos = Y * DstWidth;
+		SrcX = 0.0f;
+
+		for (int32 X = 0; X < DstWidth; ++X)
+		{
+			float EndX = SrcX + StepSizeX;
+			float EndY = SrcY + StepSizeY;
+
+			// Generate a rectangular region of pixels and then find the average color of the region.
+			int32 PosY = FMath::TruncToInt(SrcY + 0.5f);
+			PosY = FMath::Clamp<int32>(PosY, 0, (SrcHeight - 1));
+
+			int32 PosX = FMath::TruncToInt(SrcX + 0.5f);
+			PosX = FMath::Clamp<int32>(PosX, 0, (SrcWidth - 1));
+
+			int32 EndPosY = FMath::TruncToInt(EndY + 0.5f);
+			EndPosY = FMath::Clamp<int32>(EndPosY, 0, (SrcHeight - 1));
+
+			int32 EndPosX = FMath::TruncToInt(EndX + 0.5f);
+			EndPosX = FMath::Clamp<int32>(EndPosX, 0, (SrcWidth - 1));
+
+			int32 PixelCount = 0;
+			FVector4 StepColor(0, 0, 0, 0);
+			for (int32 PixelX = PosX; PixelX <= EndPosX; ++PixelX)
+			{
+				for (int32 PixelY = PosY; PixelY <= EndPosY; ++PixelY)
+				{
+					int32 StartPixel = PixelX + PixelY * SrcWidth;
+					ResizeFontBitmap_SampleInputPixel<BytesPerPixel>(SrcData, StartPixel, StepColor);
+					++PixelCount;
+				}
+			}
+
+			// Store the final averaged pixel color value.
+			ResizeFontBitmap_StoreAverageOutputPixel<BytesPerPixel>(DstData, PixelPos, StepColor, PixelCount);
+
+			SrcX = EndX;
+			++PixelPos;
+		}
+
+		SrcY += StepSizeY;
+	}
+}
+
 } // namespace SlateFontRendererUtils
 
 
@@ -89,8 +194,8 @@ uint16 FSlateFontRenderer::GetMaxHeight(const FSlateFontInfo& InFontInfo, const
 		FreeTypeUtils::ApplySizeAndScale(FaceGlyphData.FaceAndMemory->GetFace(), InFontInfo.Size, InScale);
 
 		// Adjust the height by the size of the outline that was applied.  
-		const float HeightAdjustment = InFontInfo.OutlineSettings.OutlineSize;
-		return static_cast<uint16>((FreeTypeUtils::Convert26Dot6ToRoundedPixel<int32>(FaceGlyphData.FaceAndMemory->GetScaledHeight()) + HeightAdjustment) * InScale);
+		const float HeightAdjustment = InFontInfo.OutlineSettings.OutlineSize * InScale;
+		return static_cast<uint16>(FreeTypeUtils::Convert26Dot6ToRoundedPixel<int32>(FaceGlyphData.FaceAndMemory->GetScaledHeight()) + HeightAdjustment);
 	}
 
 	return 0;
@@ -111,7 +216,7 @@ int16 FSlateFontRenderer::GetBaseline(const FSlateFontInfo& InFontInfo, const fl
 	{
 		FreeTypeUtils::ApplySizeAndScale(FaceGlyphData.FaceAndMemory->GetFace(), InFontInfo.Size, InScale);
 
-		return static_cast<int16>((FreeTypeUtils::Convert26Dot6ToRoundedPixel<int32>(FaceGlyphData.FaceAndMemory->GetDescender())) * InScale);
+		return FreeTypeUtils::Convert26Dot6ToRoundedPixel<int16>(FaceGlyphData.FaceAndMemory->GetDescender());
 	}
 
 	return 0;
@@ -130,8 +235,8 @@ void FSlateFontRenderer::GetUnderlineMetrics(const FSlateFontInfo& InFontInfo, c
 	{
 		FreeTypeUtils::ApplySizeAndScale(FontFace, InFontInfo.Size, InScale);
 
-		OutUnderlinePos = static_cast<int16>(FreeTypeUtils::Convert26Dot6ToRoundedPixel<int32>(FT_MulFix(FontFace->underline_position, FontFace->size->metrics.y_scale)) * InScale);
-		OutUnderlineThickness = static_cast<int16>(FreeTypeUtils::Convert26Dot6ToRoundedPixel<int32>(FT_MulFix(FontFace->underline_thickness, FontFace->size->metrics.y_scale)) * InScale);
+		OutUnderlinePos = FreeTypeUtils::Convert26Dot6ToRoundedPixel<int16>(FT_MulFix(FontFace->underline_position, FontFace->size->metrics.y_scale));
+		OutUnderlineThickness = FreeTypeUtils::Convert26Dot6ToRoundedPixel<int16>(FT_MulFix(FontFace->underline_thickness, FontFace->size->metrics.y_scale));
 	}
 	else
 #endif // WITH_FREETYPE
@@ -152,8 +257,8 @@ void FSlateFontRenderer::GetStrikeMetrics(const FSlateFontInfo& InFontInfo, cons
 		FreeTypeUtils::ApplySizeAndScale(FontFace, InFontInfo.Size, InScale);
 
 		// Place the strike 2/5th down the text by height (the code below does 3/5th as it counts from the bottom, not the top)
-		OutStrikeLinePos = static_cast<int16>(FreeTypeUtils::Convert26Dot6ToRoundedPixel<int32>(FT_MulFix(FT_MulFix(FT_DivFix(FontFace->height, 5), 3), FontFace->size->metrics.y_scale)) * InScale);
-		OutStrikeLineThickness = static_cast<int16>(FreeTypeUtils::Convert26Dot6ToRoundedPixel<int32>(FT_MulFix(FontFace->underline_thickness, FontFace->size->metrics.y_scale)) * InScale);
+		OutStrikeLinePos = FreeTypeUtils::Convert26Dot6ToRoundedPixel<int16>(FT_MulFix(FT_MulFix(FT_DivFix(FontFace->height, 5), 3), FontFace->size->metrics.y_scale));
+		OutStrikeLineThickness = FreeTypeUtils::Convert26Dot6ToRoundedPixel<int16>(FT_MulFix(FontFace->underline_thickness, FontFace->size->metrics.y_scale));
 	}
 	else
 #endif // WITH_FREETYPE
@@ -309,7 +414,6 @@ bool FSlateFontRenderer::GetRenderData(const FShapedGlyphEntry& InShapedGlyph, c
 		FT_Error Error = FreeTypeUtils::LoadGlyph(FaceGlyphData.FaceAndMemory->GetFace(), FaceGlyphData.GlyphIndex, FaceGlyphData.GlyphFlags, InShapedGlyph.FontFaceData->FontSize, InShapedGlyph.FontFaceData->FontScale);
 		check(Error == 0);
 
-		OutRenderData.Char = 0;
 		return GetRenderDataInternal(FaceGlyphData, InShapedGlyph.FontFaceData->FontScale, InOutlineSettings, OutRenderData);
 	}
 #endif // WITH_FREETYPE
@@ -384,12 +488,17 @@ void RenderOutlineRows(FT_Library Library, FT_Outline* Outline, FRasterizerSpanL
 
 	FT_Raster_Params RasterParams;
 	FMemory::Memzero<FT_Raster_Params>(RasterParams);
-	RasterParams.flags = FT_RASTER_FLAG_AA | FT_RASTER_FLAG_DIRECT;
+	RasterParams.flags = FT_RASTER_FLAG_AA | FT_RASTER_FLAG_DIRECT | FT_RASTER_FLAG_CLIP;
 	RasterParams.gray_spans = RasterizerCallback;
 	RasterParams.user = &OutSpansList;
 
-	FT_Outline_Render(Library, Outline, &RasterParams);
+	// Unbounding clipping box
+	RasterParams.clip_box.xMin = -32768;
+	RasterParams.clip_box.yMin = -32768;
+	RasterParams.clip_box.xMax = +32767;
+	RasterParams.clip_box.yMax = +32767;
 
+	FT_Outline_Render(Library, Outline, &RasterParams);
 }
 
 bool FSlateFontRenderer::GetRenderDataInternal(const FFreeTypeFaceGlyphData& InFaceGlyphData, const float InScale, const FFontOutlineSettings& InOutlineSettings, FCharacterRenderData& OutRenderData) const
@@ -400,9 +509,13 @@ bool FSlateFontRenderer::GetRenderDataInternal(const FFreeTypeFaceGlyphData& InF
 	// Get the lot for the glyph.  This contains measurement info
 	FT_GlyphSlot Slot = Face->glyph;
 
+	const float BitmapAtlasScale = FreeTypeUtils::GetBitmapAtlasScale(Face);
 	float ScaledOutlineSize = FMath::RoundToFloat(InOutlineSettings.OutlineSize * InScale);
 
-	if((ScaledOutlineSize > 0 || OutlineFontRenderMethod == 1) && Slot->format == FT_GLYPH_FORMAT_OUTLINE)
+	OutRenderData.bIsGrayscale = true;
+	OutRenderData.bSupportsOutline = FT_IS_SCALABLE(Face);
+
+	if((ScaledOutlineSize > 0.0f || OutlineFontRenderMethod == 1) && Slot->format == FT_GLYPH_FORMAT_OUTLINE)
 	{
 		// Render the filled area first
 		FRasterizerSpanList FillSpans;
@@ -439,8 +552,8 @@ bool FSlateFontRenderer::GetRenderDataInternal(const FFreeTypeFaceGlyphData& InF
 
 		OutRenderData.RawPixels.Reset();
 
-		OutRenderData.MeasureInfo.SizeX = Width;
-		OutRenderData.MeasureInfo.SizeY = Height;
+		OutRenderData.SizeX = Width;
+		OutRenderData.SizeY = Height;
 
 		OutRenderData.RawPixels.AddZeroed(Width*Height);
 
@@ -512,7 +625,15 @@ bool FSlateFontRenderer::GetRenderDataInternal(const FFreeTypeFaceGlyphData& InF
 
 		FT_Bitmap* Bitmap = nullptr;
 		FT_Bitmap TmpBitmap;
-		if(Slot->bitmap.pixel_mode != FT_PIXEL_MODE_GRAY)
+#if WITH_FREETYPE_V210
+		if (Slot->bitmap.pixel_mode == FT_PIXEL_MODE_BGRA)
+		{
+			OutRenderData.bIsGrayscale = false;
+			Bitmap = &Slot->bitmap;
+		}
+		else
+#endif	// WITH_FREETYPE_V210
+		if (Slot->bitmap.pixel_mode != FT_PIXEL_MODE_GRAY)
 		{
 			// Convert the bitmap to 8bpp grayscale
 			FT_Bitmap_New(&TmpBitmap);
@@ -523,24 +644,28 @@ bool FSlateFontRenderer::GetRenderDataInternal(const FFreeTypeFaceGlyphData& InF
 		{
 			Bitmap = &Slot->bitmap;
 		}
-		check(Bitmap && Bitmap->pixel_mode == FT_PIXEL_MODE_GRAY);
+#if WITH_FREETYPE_V210
+		check(Bitmap && ((Bitmap->pixel_mode == FT_PIXEL_MODE_GRAY && OutRenderData.bIsGrayscale) || (Slot->bitmap.pixel_mode == FT_PIXEL_MODE_BGRA && !OutRenderData.bIsGrayscale)));
+#else	// WITH_FREETYPE_V210
+		check(Bitmap && Bitmap->pixel_mode == FT_PIXEL_MODE_GRAY && OutRenderData.bIsGrayscale);
+#endif	// WITH_FREETYPE_V210
 
+		const uint32 BytesPerPixel = OutRenderData.bIsGrayscale ? 1 : 4;
 		OutRenderData.RawPixels.Reset();
-		OutRenderData.RawPixels.AddUninitialized(Bitmap->rows * Bitmap->width);
+		OutRenderData.RawPixels.AddUninitialized(Bitmap->rows * Bitmap->width * BytesPerPixel);
 
 		// Nothing to do for zero width or height glyphs
 		if(OutRenderData.RawPixels.Num())
 		{
 			// Copy the rendered bitmap to our raw pixels array
-			// This code assumes we're dealing with an 8bpp grayscale image (as asserted above)
 			for(uint32 Row = 0; Row < (uint32)Bitmap->rows; ++Row)
 			{
 				// Copy a single row. Note Bitmap.pitch contains the offset (in bytes) between rows.  Not always equal to Bitmap.width!
-				FMemory::Memcpy(&OutRenderData.RawPixels[Row*Bitmap->width], &Bitmap->buffer[Row*Bitmap->pitch], Bitmap->width);
+				FMemory::Memcpy(&OutRenderData.RawPixels[Row*Bitmap->width*BytesPerPixel], &Bitmap->buffer[Row*Bitmap->pitch], Bitmap->width*BytesPerPixel);
 			}
 
 			// Grayscale images not using 256 colors need to convert their gray range to a 0-255 range
-			if(Bitmap->num_grays != 256)
+			if(OutRenderData.bIsGrayscale && Bitmap->num_grays != 256)
 			{
 				const int32 GrayBoost = 255 / (Bitmap->num_grays - 1);
 				for(uint8& RawPixel : OutRenderData.RawPixels)
@@ -550,8 +675,8 @@ bool FSlateFontRenderer::GetRenderDataInternal(const FFreeTypeFaceGlyphData& InF
 			}
 		}
 
-		OutRenderData.MeasureInfo.SizeX = Bitmap->width;
-		OutRenderData.MeasureInfo.SizeY = Bitmap->rows;
+		OutRenderData.SizeX = Bitmap->width;
+		OutRenderData.SizeY = Bitmap->rows;
 
 		if(Bitmap == &TmpBitmap)
 		{
@@ -559,22 +684,36 @@ bool FSlateFontRenderer::GetRenderDataInternal(const FFreeTypeFaceGlyphData& InF
 			Bitmap = nullptr;
 		}
 
+		// Resize the generated bitmap if required
+		if (BitmapAtlasScale < 1.0f)
+		{
+			const int32 ScaledWidth = FMath::TruncToInt(OutRenderData.SizeX * BitmapAtlasScale);
+			const int32 ScaledHeight = FMath::TruncToInt(OutRenderData.SizeY * BitmapAtlasScale);
+			TArray<uint8> ScaledRawPixels;
+
+			if (OutRenderData.bIsGrayscale)
+			{
+				check(BytesPerPixel == 1);
+				SlateFontRendererUtils::ResizeFontBitmap<1>(OutRenderData.SizeX, OutRenderData.SizeY, OutRenderData.RawPixels, ScaledWidth, ScaledHeight, ScaledRawPixels);
+			}
+			else
+			{
+				check(BytesPerPixel == 4);
+				SlateFontRendererUtils::ResizeFontBitmap<4>(OutRenderData.SizeX, OutRenderData.SizeY, OutRenderData.RawPixels, ScaledWidth, ScaledHeight, ScaledRawPixels);
+			}
+
+			OutRenderData.SizeX = ScaledWidth;
+			OutRenderData.SizeY = ScaledHeight;
+			OutRenderData.RawPixels = MoveTemp(ScaledRawPixels);
+		}
+
 		// Reset the outline to zero.  If we are in this path, either the outline failed to generate because the font supported or there is no outline
 		// We do not want to take it into account if it failed to generate
-		ScaledOutlineSize = 0;
+		ScaledOutlineSize = 0.0f;
 	}
 
-	// Set measurement info for this character
-	OutRenderData.GlyphIndex = InFaceGlyphData.GlyphIndex;
-	OutRenderData.HasKerning = FT_HAS_KERNING(Face) != 0;
-
-	OutRenderData.MaxHeight = static_cast<int32>(FreeTypeUtils::Convert26Dot6ToRoundedPixel<int32>(InFaceGlyphData.FaceAndMemory->GetScaledHeight()) * InScale);
-	OutRenderData.MeasureInfo.GlobalAscender = static_cast<int16>(FreeTypeUtils::Convert26Dot6ToRoundedPixel<int32>(InFaceGlyphData.FaceAndMemory->GetAscender()) * InScale);
-	OutRenderData.MeasureInfo.GlobalDescender = static_cast<int16>(FreeTypeUtils::Convert26Dot6ToRoundedPixel<int32>(InFaceGlyphData.FaceAndMemory->GetDescender()) * InScale);
-	// Note we use Slot->advance instead of Slot->metrics.horiAdvance because Slot->Advance contains transformed position (needed if we scale)
-	OutRenderData.MeasureInfo.XAdvance = FreeTypeUtils::Convert26Dot6ToRoundedPixel<int16>(Slot->advance.x);
-	OutRenderData.MeasureInfo.HorizontalOffset = Slot->bitmap_left;
-	OutRenderData.MeasureInfo.VerticalOffset = Slot->bitmap_top + ScaledOutlineSize;
+	OutRenderData.HorizontalOffset = FMath::RoundToInt(Slot->bitmap_left * BitmapAtlasScale);
+	OutRenderData.VerticalOffset = FMath::RoundToInt(Slot->bitmap_top + ScaledOutlineSize) * BitmapAtlasScale;
 
 	return true;
 }
diff --git a/Engine/Source/Runtime/SlateCore/Private/Fonts/SlateFontRenderer.h b/Engine/Source/Runtime/SlateCore/Private/Fonts/SlateFontRenderer.h
--- a/Engine/Source/Runtime/SlateCore/Private/Fonts/SlateFontRenderer.h
+++ b/Engine/Source/Runtime/SlateCore/Private/Fonts/SlateFontRenderer.h
@@ -42,10 +42,8 @@ const TCHAR InvalidSubChar = TEXT('\uFFFD');
 
 #if WITH_FREETYPE
 
-const uint32 GlobalGlyphFlags = FT_LOAD_NO_BITMAP;
-
 /** Append the flags needed by the given font data to the given flags variable */
-void AppendGlyphFlags(const FFontData& InFontData, uint32& InOutGlyphFlags);
+void AppendGlyphFlags(const FFreeTypeFace& InFace, const FFontData& InFontData, uint32& InOutGlyphFlags);
 
 #endif // WITH_FREETYPE
 
diff --git a/Engine/Source/Runtime/SlateCore/Private/Fonts/SlateTextShaper.cpp b/Engine/Source/Runtime/SlateCore/Private/Fonts/SlateTextShaper.cpp
--- a/Engine/Source/Runtime/SlateCore/Private/Fonts/SlateTextShaper.cpp
+++ b/Engine/Source/Runtime/SlateCore/Private/Fonts/SlateTextShaper.cpp
@@ -111,7 +111,7 @@ FSlateTextShaper::FSlateTextShaper(FFreeTypeGlyphCache* InFTGlyphCache, FFreeTyp
 	check(FontCache);
 }
 
-FShapedGlyphSequenceRef FSlateTextShaper::ShapeBidirectionalText(const TCHAR* InText, const int32 InTextStart, const int32 InTextLen, const FSlateFontInfo &InFontInfo, const float InFontScale, const TextBiDi::ETextDirection InBaseDirection, const ETextShapingMethod TextShapingMethod) const
+FShapedGlyphSequenceRef FSlateTextShaper::ShapeBidirectionalText(const TCHAR* InText, const int32 InTextStart, const int32 InTextLen, const FSlateFontInfo& InFontInfo, const float InFontScale, const TextBiDi::ETextDirection InBaseDirection, const ETextShapingMethod TextShapingMethod) const
 {
 	SCOPE_CYCLE_COUNTER(STAT_SlateShapeBidirectionalText);
 
@@ -127,7 +127,7 @@ FShapedGlyphSequenceRef FSlateTextShaper::ShapeBidirectionalText(const TCHAR* In
 	return FinalizeTextShaping(MoveTemp(GlyphsToRender), InFontInfo, InFontScale, FShapedGlyphSequence::FSourceTextRange(InTextStart, InTextLen));
 }
 
-FShapedGlyphSequenceRef FSlateTextShaper::ShapeUnidirectionalText(const TCHAR* InText, const int32 InTextStart, const int32 InTextLen, const FSlateFontInfo &InFontInfo, const float InFontScale, const TextBiDi::ETextDirection InTextDirection, const ETextShapingMethod TextShapingMethod) const
+FShapedGlyphSequenceRef FSlateTextShaper::ShapeUnidirectionalText(const TCHAR* InText, const int32 InTextStart, const int32 InTextLen, const FSlateFontInfo& InFontInfo, const float InFontScale, const TextBiDi::ETextDirection InTextDirection, const ETextShapingMethod TextShapingMethod) const
 {
 	SCOPE_CYCLE_COUNTER(STAT_SlateShapeUnidirectionalText);
 
@@ -136,7 +136,7 @@ FShapedGlyphSequenceRef FSlateTextShaper::ShapeUnidirectionalText(const TCHAR* I
 	return FinalizeTextShaping(MoveTemp(GlyphsToRender), InFontInfo, InFontScale, FShapedGlyphSequence::FSourceTextRange(InTextStart, InTextLen));
 }
 
-void FSlateTextShaper::PerformTextShaping(const TCHAR* InText, const int32 InTextStart, const int32 InTextLen, const FSlateFontInfo &InFontInfo, const float InFontScale, const TextBiDi::ETextDirection InTextDirection, const ETextShapingMethod TextShapingMethod, TArray<FShapedGlyphEntry>& OutGlyphsToRender) const
+void FSlateTextShaper::PerformTextShaping(const TCHAR* InText, const int32 InTextStart, const int32 InTextLen, const FSlateFontInfo& InFontInfo, const float InFontScale, const TextBiDi::ETextDirection InTextDirection, const ETextShapingMethod TextShapingMethod, TArray<FShapedGlyphEntry>& OutGlyphsToRender) const
 {
 	check(InTextDirection != TextBiDi::ETextDirection::Mixed);
 
@@ -254,7 +254,7 @@ void FSlateTextShaper::PerformTextShaping(const TCHAR* InText, const int32 InTex
 #endif // WITH_FREETYPE
 }
 
-FShapedGlyphSequenceRef FSlateTextShaper::FinalizeTextShaping(TArray<FShapedGlyphEntry> InGlyphsToRender, const FSlateFontInfo &InFontInfo, const float InFontScale, const FShapedGlyphSequence::FSourceTextRange& InSourceTextRange) const
+FShapedGlyphSequenceRef FSlateTextShaper::FinalizeTextShaping(TArray<FShapedGlyphEntry> InGlyphsToRender, const FSlateFontInfo& InFontInfo, const float InFontScale, const FShapedGlyphSequence::FSourceTextRange& InSourceTextRange) const
 {
 	int16 TextBaseline = 0;
 	uint16 MaxHeight = 0;
@@ -269,9 +269,9 @@ FShapedGlyphSequenceRef FSlateTextShaper::FinalizeTextShaping(TArray<FShapedGlyp
 		if (FaceGlyphData.FaceAndMemory.IsValid())
 		{
 			FreeTypeUtils::ApplySizeAndScale(FaceGlyphData.FaceAndMemory->GetFace(), InFontInfo.Size, InFontScale);
-
-			TextBaseline = static_cast<int16>(FreeTypeUtils::Convert26Dot6ToRoundedPixel<int32>(FaceGlyphData.FaceAndMemory->GetDescender()) * InFontScale);
-			MaxHeight = static_cast<uint16>(FreeTypeUtils::Convert26Dot6ToRoundedPixel<int32>(FaceGlyphData.FaceAndMemory->GetScaledHeight()) * InFontScale);
+			
+			TextBaseline = FreeTypeUtils::Convert26Dot6ToRoundedPixel<int16>(FaceGlyphData.FaceAndMemory->GetDescender());
+			MaxHeight = FreeTypeUtils::Convert26Dot6ToRoundedPixel<uint16>(FaceGlyphData.FaceAndMemory->GetScaledHeight());
 		}
 	}
 #endif // WITH_FREETYPE
@@ -281,7 +281,7 @@ FShapedGlyphSequenceRef FSlateTextShaper::FinalizeTextShaping(TArray<FShapedGlyp
 
 #if WITH_FREETYPE
 
-void FSlateTextShaper::PerformKerningOnlyTextShaping(const TCHAR* InText, const int32 InTextStart, const int32 InTextLen, const FSlateFontInfo &InFontInfo, const float InFontScale, TArray<FShapedGlyphEntry>& OutGlyphsToRender) const
+void FSlateTextShaper::PerformKerningOnlyTextShaping(const TCHAR* InText, const int32 InTextStart, const int32 InTextLen, const FSlateFontInfo& InFontInfo, const float InFontScale, TArray<FShapedGlyphEntry>& OutGlyphsToRender) const
 {
 	// We need to work out the correct FFontData for everything so that we can build accurate FShapedGlyphFaceData for rendering later on
 	TArray<FKerningOnlyTextSequenceEntry, TInlineAllocator<4>> KerningOnlyTextSequence;
@@ -336,6 +336,12 @@ void FSlateTextShaper::PerformKerningOnlyTextShaping(const TCHAR* InText, const
 				FaceGlyphData = FontRenderer->GetFontFaceForCodepoint(*FontDataPtr, SlateFontRendererUtils::InvalidSubChar, InFontInfo.FontFallback);
 			}
 
+			// Only scalable font types can use sub-font scaling
+			if (FaceGlyphData.FaceAndMemory.IsValid() && !FT_IS_SCALABLE(FaceGlyphData.FaceAndMemory->GetFace()))
+			{
+				SubFontScalingFactor = 1.0f;
+			}
+
 			if (!RunningFontDataPtr || RunningFontDataPtr != FontDataPtr || RunningFaceAndMemory != FaceGlyphData.FaceAndMemory || RunningSubFontScalingFactor != SubFontScalingFactor)
 			{
 				AppendPendingFontDataToSequence();
@@ -355,13 +361,6 @@ void FSlateTextShaper::PerformKerningOnlyTextShaping(const TCHAR* InText, const
 		OutGlyphsToRender.Reserve(OutGlyphsToRender.Num() + InTextLen);
 		for (const FKerningOnlyTextSequenceEntry& KerningOnlyTextSequenceEntry : KerningOnlyTextSequence)
 		{
-			const float FinalFontScale = InFontScale * KerningOnlyTextSequenceEntry.SubFontScalingFactor;
-
-			uint32 GlyphFlags = 0;
-			SlateFontRendererUtils::AppendGlyphFlags(*KerningOnlyTextSequenceEntry.FontDataPtr, GlyphFlags);
-
-			TSharedRef<FShapedGlyphFaceData> ShapedGlyphFaceData = MakeShared<FShapedGlyphFaceData>(KerningOnlyTextSequenceEntry.FaceAndMemory, GlyphFlags, InFontInfo.Size, FinalFontScale);
-
 			if (!KerningOnlyTextSequenceEntry.FaceAndMemory.IsValid())
 			{
 				continue;
@@ -369,12 +368,19 @@ void FSlateTextShaper::PerformKerningOnlyTextShaping(const TCHAR* InText, const
 
 			const bool bHasKerning = FT_HAS_KERNING(KerningOnlyTextSequenceEntry.FaceAndMemory->GetFace()) != 0;
 
+			uint32 GlyphFlags = 0;
+			SlateFontRendererUtils::AppendGlyphFlags(*KerningOnlyTextSequenceEntry.FaceAndMemory, *KerningOnlyTextSequenceEntry.FontDataPtr, GlyphFlags);
+			const float FinalFontScale = InFontScale * KerningOnlyTextSequenceEntry.SubFontScalingFactor;
+
+			FreeTypeUtils::ApplySizeAndScale(KerningOnlyTextSequenceEntry.FaceAndMemory->GetFace(), InFontInfo.Size, FinalFontScale);
+			TSharedRef<FShapedGlyphFaceData> ShapedGlyphFaceData = MakeShared<FShapedGlyphFaceData>(KerningOnlyTextSequenceEntry.FaceAndMemory, GlyphFlags, InFontInfo.Size, FinalFontScale);
+
 			for (int32 SequenceCharIndex = 0; SequenceCharIndex < KerningOnlyTextSequenceEntry.TextLength; ++SequenceCharIndex)
 			{
 				const int32 CurrentCharIndex = KerningOnlyTextSequenceEntry.TextStartIndex + SequenceCharIndex;
 				const TCHAR CurrentChar = InText[CurrentCharIndex];
 
-				if (!InsertSubstituteGlyphs(InText, CurrentCharIndex, InFontInfo, InFontScale, ShapedGlyphFaceData, OutGlyphsToRender))
+				if (!InsertSubstituteGlyphs(InText, CurrentCharIndex, ShapedGlyphFaceData, OutGlyphsToRender))
 				{
 					const bool bIsZeroWidthSpace = CurrentChar == TEXT('\u200B');
 					const bool bIsWhitespace = bIsZeroWidthSpace || FText::IsWhitespace(CurrentChar);
@@ -435,7 +441,7 @@ void FSlateTextShaper::PerformKerningOnlyTextShaping(const TCHAR* InText, const
 
 #if WITH_HARFBUZZ
 
-void FSlateTextShaper::PerformHarfBuzzTextShaping(const TCHAR* InText, const int32 InTextStart, const int32 InTextLen, const FSlateFontInfo &InFontInfo, const float InFontScale, const TextBiDi::ETextDirection InTextDirection, TArray<FShapedGlyphEntry>& OutGlyphsToRender) const
+void FSlateTextShaper::PerformHarfBuzzTextShaping(const TCHAR* InText, const int32 InTextStart, const int32 InTextLen, const FSlateFontInfo& InFontInfo, const float InFontScale, const TextBiDi::ETextDirection InTextDirection, TArray<FShapedGlyphEntry>& OutGlyphsToRender) const
 {
 	// HarfBuzz can only shape data that uses the same font face, reads in the same direction, and uses the same script so we need to split the given text...
 	TArray<FHarfBuzzTextSequenceEntry, TInlineAllocator<4>> HarfBuzzTextSequence;
@@ -504,6 +510,12 @@ void FSlateTextShaper::PerformHarfBuzzTextShaping(const TCHAR* InText, const int
 				FaceGlyphData = FontRenderer->GetFontFaceForCodepoint(*FontDataPtr, SlateFontRendererUtils::InvalidSubChar, InFontInfo.FontFallback);
 			}
 
+			// Only scalable font types can use sub-font scaling
+			if (FaceGlyphData.FaceAndMemory.IsValid() && !FT_IS_SCALABLE(FaceGlyphData.FaceAndMemory->GetFace()))
+			{
+				SubFontScalingFactor = 1.0f;
+			}
+
 			if (!RunningFontDataPtr || RunningFontDataPtr != FontDataPtr || RunningFaceAndMemory != FaceGlyphData.FaceAndMemory || RunningSubFontScalingFactor != SubFontScalingFactor)
 			{
 				AppendPendingFontDataToSequence();
@@ -601,13 +613,6 @@ void FSlateTextShaper::PerformHarfBuzzTextShaping(const TCHAR* InText, const int
 
 		for (const FHarfBuzzTextSequenceEntry& HarfBuzzTextSequenceEntry : HarfBuzzTextSequence)
 		{
-			const float FinalFontScale = InFontScale * HarfBuzzTextSequenceEntry.SubFontScalingFactor;
-
-			uint32 GlyphFlags = 0;
-			SlateFontRendererUtils::AppendGlyphFlags(*HarfBuzzTextSequenceEntry.FontDataPtr, GlyphFlags);
-
-			TSharedRef<FShapedGlyphFaceData> ShapedGlyphFaceData = MakeShared<FShapedGlyphFaceData>(HarfBuzzTextSequenceEntry.FaceAndMemory, GlyphFlags, InFontInfo.Size, FinalFontScale);
-
 			if (!HarfBuzzTextSequenceEntry.FaceAndMemory.IsValid())
 			{
 				continue;
@@ -623,7 +628,12 @@ void FSlateTextShaper::PerformHarfBuzzTextShaping(const TCHAR* InText, const int
 			};
 			const int32 HarfBuzzFeaturesCount = ARRAY_COUNT(HarfBuzzFeatures);
 
+			uint32 GlyphFlags = 0;
+			SlateFontRendererUtils::AppendGlyphFlags(*HarfBuzzTextSequenceEntry.FaceAndMemory, *HarfBuzzTextSequenceEntry.FontDataPtr, GlyphFlags);
+			const float FinalFontScale = InFontScale * HarfBuzzTextSequenceEntry.SubFontScalingFactor;
+
 			hb_font_t* HarfBuzzFont = HarfBuzzFontFactory.CreateFont(*HarfBuzzTextSequenceEntry.FaceAndMemory, GlyphFlags, InFontInfo.Size, FinalFontScale);
+			TSharedRef<FShapedGlyphFaceData> ShapedGlyphFaceData = MakeShared<FShapedGlyphFaceData>(HarfBuzzTextSequenceEntry.FaceAndMemory, GlyphFlags, InFontInfo.Size, FinalFontScale);
 
 			for (const FHarfBuzzTextSequenceEntry::FSubSequenceEntry& HarfBuzzTextSubSequenceEntry : HarfBuzzTextSequenceEntry.SubSequence)
 			{
@@ -646,7 +656,7 @@ void FSlateTextShaper::PerformHarfBuzzTextShaping(const TCHAR* InText, const int
 
 					const int32 CurrentCharIndex = static_cast<int32>(HarfBuzzGlyphInfo.cluster);
 					const TCHAR CurrentChar = InText[CurrentCharIndex];
-					if (!InsertSubstituteGlyphs(InText, CurrentCharIndex, InFontInfo, InFontScale, ShapedGlyphFaceData, OutGlyphsToRender))
+					if (!InsertSubstituteGlyphs(InText, CurrentCharIndex, ShapedGlyphFaceData, OutGlyphsToRender))
 					{
 						const bool bIsZeroWidthSpace = CurrentChar == TEXT('\u200B');
 						const bool bIsWhitespace = bIsZeroWidthSpace || FText::IsWhitespace(CurrentChar);
@@ -800,15 +810,14 @@ void FSlateTextShaper::PerformHarfBuzzTextShaping(const TCHAR* InText, const int
 
 #endif // WITH_HARFBUZZ
 
-bool FSlateTextShaper::InsertSubstituteGlyphs(const TCHAR* InText, const int32 InCharIndex, const FSlateFontInfo &InFontInfo, const float InFontScale, const TSharedRef<FShapedGlyphFaceData>& InShapedGlyphFaceData, TArray<FShapedGlyphEntry>& OutGlyphsToRender) const
+bool FSlateTextShaper::InsertSubstituteGlyphs(const TCHAR* InText, const int32 InCharIndex, const TSharedRef<FShapedGlyphFaceData>& InShapedGlyphFaceData, TArray<FShapedGlyphEntry>& OutGlyphsToRender) const
 {
 	const TCHAR Char = InText[InCharIndex];
 
 	if (TextBiDi::IsControlCharacter(Char))
 	{
 		// We insert a stub entry for control characters to avoid them being drawn as a visual glyph with size
-		const int32 CurrentGlyphEntryIndex = OutGlyphsToRender.AddDefaulted();
-		FShapedGlyphEntry& ShapedGlyphEntry = OutGlyphsToRender[CurrentGlyphEntryIndex];
+		FShapedGlyphEntry& ShapedGlyphEntry = OutGlyphsToRender.AddDefaulted_GetRef();
 		ShapedGlyphEntry.FontFaceData = InShapedGlyphFaceData;
 		ShapedGlyphEntry.GlyphIndex = 0;
 		ShapedGlyphEntry.SourceIndex = InCharIndex;
@@ -848,8 +857,7 @@ bool FSlateTextShaper::InsertSubstituteGlyphs(const TCHAR* InText, const int32 I
 		const int32 NumSpacesToInsert = 4 - (OutGlyphsToRender.Num() % 4);
 		if (NumSpacesToInsert > 0)
 		{
-			const int32 CurrentGlyphEntryIndex = OutGlyphsToRender.AddDefaulted();
-			FShapedGlyphEntry& ShapedGlyphEntry = OutGlyphsToRender[CurrentGlyphEntryIndex];
+			FShapedGlyphEntry& ShapedGlyphEntry = OutGlyphsToRender.AddDefaulted_GetRef();
 			ShapedGlyphEntry.FontFaceData = InShapedGlyphFaceData;
 			ShapedGlyphEntry.GlyphIndex = SpaceGlyphIndex;
 			ShapedGlyphEntry.SourceIndex = InCharIndex;
diff --git a/Engine/Source/Runtime/SlateCore/Private/Fonts/SlateTextShaper.h b/Engine/Source/Runtime/SlateCore/Private/Fonts/SlateTextShaper.h
--- a/Engine/Source/Runtime/SlateCore/Private/Fonts/SlateTextShaper.h
+++ b/Engine/Source/Runtime/SlateCore/Private/Fonts/SlateTextShaper.h
@@ -22,7 +22,14 @@ public:
 		, GlyphFlags(InGlyphFlags)
 		, FontSize(InFontSize)
 		, FontScale(InFontScale)
+		, BitmapRenderScale(1.0f)
 	{
+#if WITH_FREETYPE
+		if (TSharedPtr<FFreeTypeFace> FontFacePin = FontFace.Pin())
+		{
+			BitmapRenderScale = FontFacePin->GetBitmapRenderScale();
+		}
+#endif	// WITH_FREETYPE
 	}
 
 	/** Weak pointer to the FreeType face to render with */
@@ -33,6 +40,8 @@ public:
 	int32 FontSize;
 	/** Provides the final scale used to render to the font */
 	float FontScale;
+	/** Any additional scale that should be applied when rendering this glyph */
+	float BitmapRenderScale;
 };
 
 
@@ -45,24 +54,24 @@ class FSlateTextShaper
 public:
 	FSlateTextShaper(FFreeTypeGlyphCache* InFTGlyphCache, FFreeTypeAdvanceCache* InFTAdvanceCache, FFreeTypeKerningPairCache* InFTKerningPairCache, FCompositeFontCache* InCompositeFontCache, FSlateFontRenderer* InFontRenderer, FSlateFontCache* InFontCache);
 
-	FShapedGlyphSequenceRef ShapeBidirectionalText(const TCHAR* InText, const int32 InTextStart, const int32 InTextLen, const FSlateFontInfo &InFontInfo, const float InFontScale, const TextBiDi::ETextDirection InBaseDirection, const ETextShapingMethod TextShapingMethod) const;
+	FShapedGlyphSequenceRef ShapeBidirectionalText(const TCHAR* InText, const int32 InTextStart, const int32 InTextLen, const FSlateFontInfo& InFontInfo, const float InFontScale, const TextBiDi::ETextDirection InBaseDirection, const ETextShapingMethod TextShapingMethod) const;
 
-	FShapedGlyphSequenceRef ShapeUnidirectionalText(const TCHAR* InText, const int32 InTextStart, const int32 InTextLen, const FSlateFontInfo &InFontInfo, const float InFontScale, const TextBiDi::ETextDirection InTextDirection, const ETextShapingMethod TextShapingMethod) const;
+	FShapedGlyphSequenceRef ShapeUnidirectionalText(const TCHAR* InText, const int32 InTextStart, const int32 InTextLen, const FSlateFontInfo& InFontInfo, const float InFontScale, const TextBiDi::ETextDirection InTextDirection, const ETextShapingMethod TextShapingMethod) const;
 
 private:
-	void PerformTextShaping(const TCHAR* InText, const int32 InTextStart, const int32 InTextLen, const FSlateFontInfo &InFontInfo, const float InFontScale, const TextBiDi::ETextDirection InTextDirection, const ETextShapingMethod TextShapingMethod, TArray<FShapedGlyphEntry>& OutGlyphsToRender) const;
+	void PerformTextShaping(const TCHAR* InText, const int32 InTextStart, const int32 InTextLen, const FSlateFontInfo& InFontInfo, const float InFontScale, const TextBiDi::ETextDirection InTextDirection, const ETextShapingMethod TextShapingMethod, TArray<FShapedGlyphEntry>& OutGlyphsToRender) const;
 
-	FShapedGlyphSequenceRef FinalizeTextShaping(TArray<FShapedGlyphEntry> InGlyphsToRender, const FSlateFontInfo &InFontInfo, const float InFontScale, const FShapedGlyphSequence::FSourceTextRange& InSourceTextRange) const;
+	FShapedGlyphSequenceRef FinalizeTextShaping(TArray<FShapedGlyphEntry> InGlyphsToRender, const FSlateFontInfo& InFontInfo, const float InFontScale, const FShapedGlyphSequence::FSourceTextRange& InSourceTextRange) const;
 
 #if WITH_FREETYPE
-	void PerformKerningOnlyTextShaping(const TCHAR* InText, const int32 InTextStart, const int32 InTextLen, const FSlateFontInfo &InFontInfo, const float InFontScale, TArray<FShapedGlyphEntry>& OutGlyphsToRender) const;
+	void PerformKerningOnlyTextShaping(const TCHAR* InText, const int32 InTextStart, const int32 InTextLen, const FSlateFontInfo& InFontInfo, const float InFontScale, TArray<FShapedGlyphEntry>& OutGlyphsToRender) const;
 #endif // WITH_FREETYPE
 
 #if WITH_HARFBUZZ
-	void PerformHarfBuzzTextShaping(const TCHAR* InText, const int32 InTextStart, const int32 InTextLen, const FSlateFontInfo &InFontInfo, const float InFontScale, const TextBiDi::ETextDirection InTextDirection, TArray<FShapedGlyphEntry>& OutGlyphsToRender) const;
+	void PerformHarfBuzzTextShaping(const TCHAR* InText, const int32 InTextStart, const int32 InTextLen, const FSlateFontInfo& InFontInfo, const float InFontScale, const TextBiDi::ETextDirection InTextDirection, TArray<FShapedGlyphEntry>& OutGlyphsToRender) const;
 #endif // WITH_HARFBUZZ
 
-	bool InsertSubstituteGlyphs(const TCHAR* InText, const int32 InCharIndex, const FSlateFontInfo &InFontInfo, const float InFontScale, const TSharedRef<FShapedGlyphFaceData>& InShapedGlyphFaceData, TArray<FShapedGlyphEntry>& OutGlyphsToRender) const;
+	bool InsertSubstituteGlyphs(const TCHAR* InText, const int32 InCharIndex, const TSharedRef<FShapedGlyphFaceData>& InShapedGlyphFaceData, TArray<FShapedGlyphEntry>& OutGlyphsToRender) const;
 
 	FFreeTypeGlyphCache* FTGlyphCache;
 	FFreeTypeAdvanceCache* FTAdvanceCache;
diff --git a/Engine/Source/Runtime/SlateCore/Private/Rendering/ElementBatcher.cpp b/Engine/Source/Runtime/SlateCore/Private/Rendering/ElementBatcher.cpp
--- a/Engine/Source/Runtime/SlateCore/Private/Rendering/ElementBatcher.cpp
+++ b/Engine/Source/Runtime/SlateCore/Private/Rendering/ElementBatcher.cpp
@@ -1084,6 +1084,7 @@ void FSlateElementBatcher::AddTextElement(const FSlateDrawElement& DrawElement)
 		uint32 FontTextureIndex = 0;
 		FSlateShaderResource* FontAtlasTexture = nullptr;
 		FSlateShaderResource* FontShaderResource = nullptr;
+		FColor FontTint = InTint;
 
 		FSlateRenderBatch* RenderBatch = nullptr;
 		FSlateVertexArray* BatchVertices = nullptr;
@@ -1109,6 +1110,7 @@ void FSlateElementBatcher::AddTextElement(const FSlateDrawElement& DrawElement)
 		LineX = PosX;
 
 		const bool bIsFontMaterial = FontMaterial != nullptr;
+		const bool bEnableOutline = InOutlineSettings.OutlineSize > 0.0f;
 
 		uint32 NumChars = Len;
 
@@ -1140,17 +1142,22 @@ void FSlateElementBatcher::AddTextElement(const FSlateDrawElement& DrawElement)
 					// Font has a new texture for this glyph. Refresh the batch we use and the index we are currently using
 					FontTextureIndex = Entry.TextureIndex;
 
-					FontAtlasTexture = FontCache.GetSlateTextureResource( FontTextureIndex );
+					ISlateFontTexture* SlateFontTexture = FontCache.GetFontTexture(FontTextureIndex);
+					check(SlateFontTexture);
+
+					FontAtlasTexture = SlateFontTexture->GetSlateTexture();
 					check(FontAtlasTexture);
 
 					FontShaderResource = ResourceManager.GetFontShaderResource( FontTextureIndex, FontAtlasTexture, DrawElementPayload.GetFontInfo().FontMaterial );
 					check(FontShaderResource);
 
-					RenderBatch = &CreateRenderBatch(InLayer, FShaderParams(), FontShaderResource, ESlateDrawPrimitive::TriangleList, ESlateShader::Font, InDrawEffects, ESlateBatchDrawFlag::None, DrawElement);
+					const bool bIsGrayscale = SlateFontTexture->IsGrayscale();
+					FontTint = bIsGrayscale ? InTint : FColor::White;
 
-					const int32 GlyphsLeft = NumChars - CharIndex;
+					RenderBatch = &CreateRenderBatch(InLayer, FShaderParams(), FontShaderResource, ESlateDrawPrimitive::TriangleList, bIsGrayscale ? ESlateShader::GrayscaleFont : ESlateShader::ColorFont, InDrawEffects, ESlateBatchDrawFlag::None, DrawElement);
 
 					// Reserve memory for the glyphs.  This isn't perfect as the text could contain spaces and we might not render the rest of the text in this batch but its better than resizing constantly
+					const int32 GlyphsLeft = NumChars - CharIndex;
 					RenderBatch->ReserveVertices(GlyphsLeft * 4);
 					RenderBatch->ReserveIndices(GlyphsLeft * 6);
 
@@ -1158,7 +1165,7 @@ void FSlateElementBatcher::AddTextElement(const FSlateDrawElement& DrawElement)
 					InvTextureSizeY = 1.0f / FontAtlasTexture->GetHeight();
 				}
 
-				const bool bIsWhitespace = !Entry.Valid || FChar::IsWhitespace(CurrentChar);
+				const bool bIsWhitespace = !Entry.Valid || (bEnableOutline && !Entry.SupportsOutline) || FChar::IsWhitespace(CurrentChar);
 
 				if( !bIsWhitespace && PreviousCharEntry.Valid )
 				{
@@ -1174,14 +1181,16 @@ void FSlateElementBatcher::AddTextElement(const FSlateDrawElement& DrawElement)
 
 				if( !bIsWhitespace )
 				{
+					const float InvBitmapRenderScale = 1.0f / Entry.BitmapRenderScale;
+
 					const float X = LineX + Entry.HorizontalOffset+InOutlineHorizontalOffset;
 					// Note PosX,PosY is the upper left corner of the bounding box representing the string.  This computes the Y position of the baseline where text will sit
 
-					const float Y = PosY - Entry.VerticalOffset+MaxHeight+Entry.GlobalDescender;
+					const float Y = PosY - Entry.VerticalOffset + ((MaxHeight + Entry.GlobalDescender) * InvBitmapRenderScale);
 					const float U = Entry.StartU * InvTextureSizeX;
 					const float V = Entry.StartV * InvTextureSizeY;
-					const float SizeX = Entry.USize;
-					const float SizeY = Entry.VSize;
+					const float SizeX = Entry.USize * Entry.BitmapRenderScale;
+					const float SizeY = Entry.VSize * Entry.BitmapRenderScale;
 					const float SizeU = Entry.USize * InvTextureSizeX;
 					const float SizeV = Entry.VSize * InvTextureSizeY;
 
@@ -1209,10 +1218,10 @@ void FSlateElementBatcher::AddTextElement(const FSlateDrawElement& DrawElement)
 						}
 
 						// Add four vertices to the list of verts to be added to the vertex buffer
-						RenderBatch->AddVertex(FSlateVertex::Make<Rounding>( RenderTransform, UpperLeft,								FVector4(U,V,Ut,Vt),						FVector2D(0.0f,0.0f), InTint ));
-						RenderBatch->AddVertex(FSlateVertex::Make<Rounding>( RenderTransform, FVector2D(LowerRight.X,UpperLeft.Y),	FVector4(U+SizeU, V, UtMax,Vt),				FVector2D(1.0f,0.0f), InTint ));
-						RenderBatch->AddVertex(FSlateVertex::Make<Rounding>( RenderTransform, FVector2D(UpperLeft.X,LowerRight.Y),	FVector4(U, V+SizeV, Ut,VtMax),				FVector2D(0.0f,1.0f), InTint ));
-						RenderBatch->AddVertex(FSlateVertex::Make<Rounding>( RenderTransform, LowerRight,							FVector4(U+SizeU, V+SizeV, UtMax,VtMax),	FVector2D(1.0f,1.0f), InTint ));
+						RenderBatch->AddVertex(FSlateVertex::Make<Rounding>( RenderTransform, UpperLeft,								FVector4(U,V,Ut,Vt),						FVector2D(0.0f,0.0f), FontTint ));
+						RenderBatch->AddVertex(FSlateVertex::Make<Rounding>( RenderTransform, FVector2D(LowerRight.X,UpperLeft.Y),	FVector4(U+SizeU, V, UtMax,Vt),				FVector2D(1.0f,0.0f), FontTint ));
+						RenderBatch->AddVertex(FSlateVertex::Make<Rounding>( RenderTransform, FVector2D(UpperLeft.X,LowerRight.Y),	FVector4(U, V+SizeV, Ut,VtMax),				FVector2D(0.0f,1.0f), FontTint ));
+						RenderBatch->AddVertex(FSlateVertex::Make<Rounding>( RenderTransform, LowerRight,							FVector4(U+SizeU, V+SizeV, UtMax,VtMax),	FVector2D(1.0f,1.0f), FontTint ));
 
 						RenderBatch->AddIndex(IndexStart + 0);
 						RenderBatch->AddIndex(IndexStart + 1);
@@ -1306,6 +1315,7 @@ void FSlateElementBatcher::AddShapedTextElement( const FSlateDrawElement& DrawEl
 		int32 FontTextureIndex = -1;
 		FSlateShaderResource* FontAtlasTexture = nullptr;
 		FSlateShaderResource* FontShaderResource = nullptr;
+		FColor FontTint = InTint;
 
 		FSlateRenderBatch* RenderBatch = nullptr;
 
@@ -1316,6 +1326,7 @@ void FSlateElementBatcher::AddShapedTextElement( const FSlateDrawElement& DrawEl
 		float InvTextureSizeY = 0;
 
 		const bool bIsFontMaterial = FontMaterial != nullptr;
+		const bool bEnableOutline = InOutlineSettings.OutlineSize > 0.0f;
 
 		// Optimize by culling
 		// Todo: this doesnt work with cached clipping
@@ -1344,7 +1355,7 @@ void FSlateElementBatcher::AddShapedTextElement( const FSlateDrawElement& DrawEl
 			{
 				const FShapedGlyphFontAtlasData GlyphAtlasData = FontCache.GetShapedGlyphFontAtlasData(GlyphToRender, InOutlineSettings);
 				 
-				if (GlyphAtlasData.Valid)
+				if (GlyphAtlasData.Valid && (!bEnableOutline || GlyphAtlasData.SupportsOutline))
 				{
 					const float X = LineX + GlyphAtlasData.HorizontalOffset + GlyphToRender.XOffset;
 					// Note PosX,PosY is the upper left corner of the bounding box representing the string.  This computes the Y position of the baseline where text will sit
@@ -1368,17 +1379,22 @@ void FSlateElementBatcher::AddShapedTextElement( const FSlateDrawElement& DrawEl
 						// Font has a new texture for this glyph. Refresh the batch we use and the index we are currently using
 						FontTextureIndex = GlyphAtlasData.TextureIndex;
 
-						FontAtlasTexture = FontCache.GetSlateTextureResource(FontTextureIndex);
+						ISlateFontTexture* SlateFontTexture = FontCache.GetFontTexture(FontTextureIndex);
+						check(SlateFontTexture);
+
+						FontAtlasTexture = SlateFontTexture->GetSlateTexture();
 						check(FontAtlasTexture);
 
 						FontShaderResource = ResourceManager.GetFontShaderResource(FontTextureIndex, FontAtlasTexture, FontMaterial);
 						check(FontShaderResource);
 
-						RenderBatch = &CreateRenderBatch(InLayer, FShaderParams(), FontShaderResource, ESlateDrawPrimitive::TriangleList, ESlateShader::Font, InDrawEffects, ESlateBatchDrawFlag::None, DrawElement);
+						const bool bIsGrayscale = SlateFontTexture->IsGrayscale();
+						FontTint = bIsGrayscale ? InTint : FColor::White;
 
-						const int32 GlyphsLeft = NumGlyphs - GlyphIndex;
+						RenderBatch = &CreateRenderBatch(InLayer, FShaderParams(), FontShaderResource, ESlateDrawPrimitive::TriangleList, bIsGrayscale ? ESlateShader::GrayscaleFont : ESlateShader::ColorFont, InDrawEffects, ESlateBatchDrawFlag::None, DrawElement);
 
 						// Reserve memory for the glyphs.  This isn't perfect as the text could contain spaces and we might not render the rest of the text in this batch but its better than resizing constantly
+						const int32 GlyphsLeft = NumGlyphs - GlyphIndex;
 						RenderBatch->ReserveVertices(GlyphsLeft*4);
 						RenderBatch->ReserveIndices(GlyphsLeft*6);
 
@@ -1386,12 +1402,14 @@ void FSlateElementBatcher::AddShapedTextElement( const FSlateDrawElement& DrawEl
 						InvTextureSizeY = 1.0f / FontAtlasTexture->GetHeight();
 					}
 
+					const float BitmapRenderScale = GlyphToRender.GetBitmapRenderScale();
+					const float InvBitmapRenderScale = 1.0f / BitmapRenderScale;
 
-					const float Y = LineY - GlyphAtlasData.VerticalOffset + GlyphToRender.YOffset + MaxHeight + TextBaseline;
+					const float Y = LineY - GlyphAtlasData.VerticalOffset + GlyphToRender.YOffset + ((MaxHeight + TextBaseline) * InvBitmapRenderScale);
 					const float U = GlyphAtlasData.StartU * InvTextureSizeX;
 					const float V = GlyphAtlasData.StartV * InvTextureSizeY;
-					const float SizeX = GlyphAtlasData.USize;
-					const float SizeY = GlyphAtlasData.VSize;
+					const float SizeX = GlyphAtlasData.USize * BitmapRenderScale;
+					const float SizeY = GlyphAtlasData.VSize * BitmapRenderScale;
 					const float SizeU = GlyphAtlasData.USize * InvTextureSizeX;
 					const float SizeV = GlyphAtlasData.VSize * InvTextureSizeY;
 
@@ -1421,10 +1439,10 @@ void FSlateElementBatcher::AddShapedTextElement( const FSlateDrawElement& DrawEl
 						}
 
 						// Add four vertices to the list of verts to be added to the vertex buffer
-						RenderBatch->AddVertex(FSlateVertex::Make<Rounding>(RenderTransform, UpperLeft,								FVector4(U, V, Ut, Vt),							FVector2D(0.0f, 0.0f), InTint ));
-						RenderBatch->AddVertex(FSlateVertex::Make<Rounding>(RenderTransform, FVector2D(LowerRight.X, UpperLeft.Y),	FVector4(U + SizeU, V, UtMax, Vt),				FVector2D(1.0f, 0.0f), InTint ));
-						RenderBatch->AddVertex(FSlateVertex::Make<Rounding>(RenderTransform, FVector2D(UpperLeft.X, LowerRight.Y),	FVector4(U, V + SizeV, Ut, VtMax),				FVector2D(0.0f, 1.0f), InTint ));
-						RenderBatch->AddVertex(FSlateVertex::Make<Rounding>(RenderTransform, LowerRight,								FVector4(U + SizeU, V + SizeV, UtMax, VtMax),	FVector2D(1.0f, 1.0f), InTint ));
+						RenderBatch->AddVertex(FSlateVertex::Make<Rounding>(RenderTransform, UpperLeft,								FVector4(U, V, Ut, Vt),							FVector2D(0.0f, 0.0f), FontTint ));
+						RenderBatch->AddVertex(FSlateVertex::Make<Rounding>(RenderTransform, FVector2D(LowerRight.X, UpperLeft.Y),	FVector4(U + SizeU, V, UtMax, Vt),				FVector2D(1.0f, 0.0f), FontTint ));
+						RenderBatch->AddVertex(FSlateVertex::Make<Rounding>(RenderTransform, FVector2D(UpperLeft.X, LowerRight.Y),	FVector4(U, V + SizeV, Ut, VtMax),				FVector2D(0.0f, 1.0f), FontTint ));
+						RenderBatch->AddVertex(FSlateVertex::Make<Rounding>(RenderTransform, LowerRight,								FVector4(U + SizeU, V + SizeV, UtMax, VtMax),	FVector2D(1.0f, 1.0f), FontTint ));
 
 						RenderBatch->AddIndex(IndexStart + 0);
 						RenderBatch->AddIndex(IndexStart + 1);
@@ -2220,8 +2238,8 @@ void FSlateElementBatcher::AddViewportElement( const FSlateDrawElement& DrawElem
 
 	if(DrawElementPayload.bViewportTextureAlphaOnly )
 	{
-		// This is a slight hack, but the font shader is the same as the general shader except it reads alpha only textures
-		ShaderType = ESlateShader::Font;
+		// This is a slight hack, but the grayscale font shader is the same as the general shader except it reads alpha only textures and doesn't support tiling
+		ShaderType = ESlateShader::GrayscaleFont;
 	}
 
 	FSlateRenderBatch& RenderBatch = CreateRenderBatch( Layer, FShaderParams(), ViewportResource, ESlateDrawPrimitive::TriangleList, ShaderType, InDrawEffects, DrawFlags, DrawElement);
diff --git a/Engine/Source/Runtime/SlateCore/Public/Fonts/FontCache.h b/Engine/Source/Runtime/SlateCore/Public/Fonts/FontCache.h
--- a/Engine/Source/Runtime/SlateCore/Public/Fonts/FontCache.h
+++ b/Engine/Source/Runtime/SlateCore/Public/Fonts/FontCache.h
@@ -69,33 +69,23 @@ SLATECORE_API ETextShapingMethod GetDefaultTextShapingMethod();
 struct FShapedGlyphFontAtlasData
 {
 	/** The vertical distance from the baseline to the topmost border of the glyph bitmap */
-	int16 VerticalOffset;
+	int16 VerticalOffset = 0;
 	/** The horizontal distance from the origin to the leftmost border of the glyph bitmap */
-	int16 HorizontalOffset;
+	int16 HorizontalOffset = 0;
 	/** Start X location of the glyph in the texture */
-	uint16 StartU;
+	uint16 StartU = 0;
 	/** Start Y location of the glyph in the texture */
-	uint16 StartV;
+	uint16 StartV = 0;
 	/** X Size of the glyph in the texture */
-	uint16 USize;
+	uint16 USize = 0;
 	/** Y Size of the glyph in the texture */
-	uint16 VSize;
+	uint16 VSize = 0;
 	/** Index to a specific texture in the font cache. */
-	uint8 TextureIndex;
+	uint8 TextureIndex = 0;
+	/** True if this entry supports outline rendering, false otherwise. */
+	bool SupportsOutline = false;
 	/** True if this entry is valid, false otherwise. */
-	bool Valid;
-
-	FShapedGlyphFontAtlasData()
-		: VerticalOffset(0)
-		, HorizontalOffset(0)
-		, StartU(0.0f)
-		, StartV(0.0f)
-		, USize(0.0f)
-		, VSize(0.0f)
-		, TextureIndex(0)
-		, Valid(false)
-	{
-	}
+	bool Valid = false;
 };
 
 /** Information for rendering one glyph in a shaped text sequence */
@@ -106,58 +96,45 @@ struct FShapedGlyphEntry
 	/** Provides access to the FreeType face for this glyph (not available publicly) */
 	TSharedPtr<FShapedGlyphFaceData> FontFaceData;
 	/** The index of this glyph in the FreeType face */
-	uint32 GlyphIndex;
+	uint32 GlyphIndex = 0;
 	/** The index of this glyph from the source text. The source indices may skip characters if the sequence contains ligatures, additionally, some characters produce multiple glyphs leading to duplicate source indices */
-	int32 SourceIndex;
+	int32 SourceIndex = 0;
 	/** The amount to advance in X before drawing the next glyph in the sequence */
-	int16 XAdvance;
+	int16 XAdvance = 0;
 	/** The amount to advance in Y before drawing the next glyph in the sequence */
-	int16 YAdvance;
+	int16 YAdvance = 0;
 	/** The offset to apply in X when drawing this glyph */
-	int16 XOffset;
+	int16 XOffset = 0;
 	/** The offset to apply in Y when drawing this glyph */
-	int16 YOffset;
+	int16 YOffset = 0;
 	/** 
 	 * The "kerning" between this glyph and the next one in the sequence
 	 * @note This value is included in the XAdvance so you never usually need it unless you're manually combining two sets of glyphs together.
 	 * @note This value isn't strictly the kerning value - it's simply the difference between the glyphs horizontal advance, and the shaped horizontal advance (so will contain any accumulated advance added by the shaper)
 	 */
-	int8 Kerning;
+	int8 Kerning = 0;
 	/**
 	 * The number of source characters represented by this glyph
 	 * This is typically 1, however will be greater for ligatures, or may be 0 if a single character produces multiple glyphs
 	 */
-	uint8 NumCharactersInGlyph;
+	uint8 NumCharactersInGlyph = 0;
 	/**
 	 * The number of source grapheme clusters represented by this glyph
 	 * This is typically 1, however will be greater for ligatures, or may be 0 if a single character produces multiple glyphs
 	 */
-	uint8 NumGraphemeClustersInGlyph;
+	uint8 NumGraphemeClustersInGlyph = 0;
 	/**
 	 * The reading direction of the text this glyph was shaped from
 	 */
-	TextBiDi::ETextDirection TextDirection;
+	TextBiDi::ETextDirection TextDirection = TextBiDi::ETextDirection::LeftToRight;
 	/**
 	 * True if this is a visible glyph that should be drawn.
 	 * False if the glyph is invisible (eg, whitespace or a control code) and should skip drawing, but still include its advance amount.
 	 */
-	bool bIsVisible;
-
-	FShapedGlyphEntry()
-		: FontFaceData()
-		, GlyphIndex(0)
-		, SourceIndex(0)
-		, XAdvance(0)
-		, YAdvance(0)
-		, XOffset(0)
-		, YOffset(0)
-		, Kerning(0)
-		, NumCharactersInGlyph(0)
-		, NumGraphemeClustersInGlyph(0)
-		, TextDirection(TextBiDi::ETextDirection::LeftToRight)
-		, bIsVisible(false)
-	{
-	}
+	bool bIsVisible = false;
+
+	/** Get any additional scale that should be applied when rendering this glyph */
+	SLATECORE_API float GetBitmapRenderScale() const;
 
 private:
 	/** 
@@ -469,42 +446,41 @@ private:
 struct SLATECORE_API FCharacterEntry
 {
 	/** The character this entry is for */
-	TCHAR Character;
+	TCHAR Character = 0;
 	/** The index of the glyph from the FreeType face that this entry is for */
-	uint32 GlyphIndex;
+	uint32 GlyphIndex = 0;
 	/** The raw font data this character was rendered with */
-	const FFontData* FontData;
+	const FFontData* FontData = nullptr;
 	/** Scale that was applied when rendering this character */
-	float FontScale;
+	float FontScale = 0.0f;
+	/** Any additional scale that should be applied when rendering this glyph */
+	float BitmapRenderScale = 0.0f;
 	/** Start X location of the character in the texture */
-	uint16 StartU;
+	uint16 StartU = 0;
 	/** Start Y location of the character in the texture */
-	uint16 StartV;
+	uint16 StartV = 0;
 	/** X Size of the character in the texture */
-	uint16 USize;
+	uint16 USize = 0;
 	/** Y Size of the character in the texture */
-	uint16 VSize;
+	uint16 VSize = 0;
 	/** The vertical distance from the baseline to the topmost border of the character */
-	int16 VerticalOffset;
+	int16 VerticalOffset = 0;
 	/** The vertical distance from the origin to the left most border of the character */
-	int16 HorizontalOffset;
+	int16 HorizontalOffset = 0;
 	/** The largest vertical distance below the baseline for any character in the font */
-	int16 GlobalDescender;
+	int16 GlobalDescender = 0;
 	/** The amount to advance in X before drawing the next character in a string */
-	int16 XAdvance;
+	int16 XAdvance = 0;
 	/** Index to a specific texture in the font cache. */
-	uint8 TextureIndex;
-	/** 1 if this entry has kerning, 0 otherwise. */
-	bool HasKerning;
+	uint8 TextureIndex = 0;
 	/** The fallback level this character represents */
-	EFontFallback FallbackLevel;
+	EFontFallback FallbackLevel = EFontFallback::FF_Max;
+	/** 1 if this entry has kerning, 0 otherwise. */
+	bool HasKerning = false;
+	/** 1 if this entry supports outline rendering, 0 otherwise. */
+	bool SupportsOutline = false;
 	/** 1 if this entry is valid, 0 otherwise. */
-	bool Valid;
-
-	FCharacterEntry()
-	{
-		FMemory::Memzero( this, sizeof(FCharacterEntry) );
-	}
+	bool Valid = false;
 };
 
 /**
@@ -574,25 +550,16 @@ private:
 	/** Maintains a fake shaped glyph for each character in the character list */
 	struct FCharacterListEntry
 	{
-		FCharacterListEntry()
-			: ShapedGlyphEntry()
-			, FontData(nullptr)
-			, FallbackLevel(EFontFallback::FF_Max)
-			, HasKerning(false)
-			, Valid(false)
-		{
-		}
-
 		/** The shaped glyph data for this character */
 		FShapedGlyphEntry ShapedGlyphEntry;
 		/** Font data this character was rendered with */
-		const FFontData* FontData;
+		const FFontData* FontData = nullptr;
 		/** The fallback level this character represents */
-		EFontFallback FallbackLevel;
+		EFontFallback FallbackLevel = EFontFallback::FF_Max;
 		/** Does this character have kerning? */
-		bool HasKerning;
+		bool HasKerning = false;
 		/** Has this entry been initialized? */
-		bool Valid;
+		bool Valid = false;
 	};
 
 	/**
@@ -656,9 +623,8 @@ public:
 
 	/** ISlateAtlasProvider */
 	virtual int32 GetNumAtlasPages() const override;
-	virtual FIntPoint GetAtlasPageSize() const override;
 	virtual FSlateShaderResource* GetAtlasPageResource(const int32 InIndex) const override;
-	virtual bool IsAtlasPageResourceAlphaOnly() const override;
+	virtual bool IsAtlasPageResourceAlphaOnly(const int32 InIndex) const override;
 
 	/** 
 	 * Performs text shaping on the given string using the given font info. Returns you the shaped text sequence to use for text rendering via FSlateDrawElement::MakeShapedText.
@@ -749,8 +715,7 @@ public:
 	 * @param Index	The index of the texture 
 	 * @return Handle to the texture resource
 	 */
-	class FSlateShaderResource* GetSlateTextureResource( uint32 Index ) { return AllFontTextures[Index]->GetSlateTexture(); }
-	class FTextureResource* GetEngineTextureResource( uint32 Index ) { return AllFontTextures[Index]->GetEngineTexture(); }
+	ISlateFontTexture* GetFontTexture( uint32 Index ) { return &AllFontTextures[Index].Get(); }
 
 	/**
 	 * Returns the font to use from the default typeface
@@ -911,14 +876,17 @@ private:
 	/** Mapping shaped glyphs to their cached atlas data */
 	TMap<FShapedGlyphEntryKey, TSharedRef<FShapedGlyphFontAtlasData>> ShapedGlyphToAtlasData;
 
-	/** Array of all font atlases */
-	TArray< TSharedRef<FSlateFontAtlas> > FontAtlases;
+	/** Array of grayscale font atlas indices for use with AllFontTextures (cast the element to FSlateFontAtlas) */
+	TArray<int32> GrayscaleFontAtlasIndices;
 
-	/** Array of any non-atlased font textures */
-	TArray< TSharedRef<ISlateFontTexture> > NonAtlasedTextures;
+	/** Array of color font atlas indices for use with AllFontTextures (cast the element to FSlateFontAtlas) */
+	TArray<int32> ColorFontAtlasIndices;
+
+	/** Array of any non-atlased font texture indices for use with AllFontTextures */
+	TArray<int32> NonAtlasedTextureIndices;
 
 	/** Array of all font textures - both atlased and non-atlased */
-	TArray< TSharedRef<ISlateFontTexture> > AllFontTextures;
+	TArray<TSharedRef<ISlateFontTexture>> AllFontTextures;
 
 	/** Factory for creating new font atlases */
 	TSharedRef<ISlateFontAtlasFactory> FontAtlasFactory;
@@ -926,8 +894,11 @@ private:
 	/** Whether or not we have a pending request to flush the cache when it is safe to do so */
 	volatile bool bFlushRequested;
 
-	/** Number of atlas pages we can have before we request that the cache be flushed */
-	int32 CurrentMaxAtlasPagesBeforeFlushRequest;
+	/** Number of grayscale atlas pages we can have before we request that the cache be flushed */
+	int32 CurrentMaxGrayscaleAtlasPagesBeforeFlushRequest;
+
+	/** Number of color atlas pages we can have before we request that the cache be flushed */
+	int32 CurrentMaxColorAtlasPagesBeforeFlushRequest;
 
 	/** Number of non-atlased textures we can have before we request that the cache be flushed */
 	int32 CurrentMaxNonAtlasedTexturesBeforeFlushRequest;
diff --git a/Engine/Source/Runtime/SlateCore/Public/Fonts/FontTypes.h b/Engine/Source/Runtime/SlateCore/Public/Fonts/FontTypes.h
--- a/Engine/Source/Runtime/SlateCore/Public/Fonts/FontTypes.h
+++ b/Engine/Source/Runtime/SlateCore/Public/Fonts/FontTypes.h
@@ -83,44 +83,23 @@ struct FSlateFontKeyFuncs : BaseKeyFuncs<TPair<FSlateFontKey, ValueType>, FSlate
 };
 
 
-/** Measurement details for a specific character */
-struct FCharacterMeasureInfo
+/** Contains pixel data for a character rendered from freetype as well as measurement info */
+struct FCharacterRenderData
 {
+	/** Raw pixels of the rendered character */
+	TArray<uint8> RawPixels;
 	/** Width of the character in pixels */
-	int16 SizeX;
+	int16 SizeX = 0;
 	/** Height of the character in pixels */
-	int16 SizeY;
-	/** How many pixels to advance in X after drawing this character (when its part of a string)*/
-	int16 XAdvance;
-	/** The largest vertical distance above the baseline for any character in the font */
-	int16 GlobalAscender;
-	/** The largest vertical distance below the baseline for any character in the font */
-	int16 GlobalDescender;
+	int16 SizeY = 0;
 	/** The vertical distance from the baseline to the topmost border of the glyph bitmap */
-	int16 VerticalOffset;
+	int16 VerticalOffset = 0;
 	/** The horizontal distance from the origin to the leftmost border of the character */
-	int16 HorizontalOffset;
-	FCharacterMeasureInfo( int16 InSizeX = 0, int16 InSizeY = 0, int16 InXAdvance = 0, int16 InVerticalOffset = 0, int16 InHorizontalOffset = 0 )
-		: SizeX(InSizeX),SizeY(InSizeY),XAdvance(InXAdvance),VerticalOffset(InVerticalOffset), HorizontalOffset(InHorizontalOffset)
-	{
-	}
-};
-
-/** Contains pixel data for a character rendered from freetype as well as measurement info */
-struct FCharacterRenderData
-{
-	FCharacterMeasureInfo MeasureInfo;
-	/** Measurement data for this character */
-	/** Raw pixels created by freetype */
-	TArray<uint8> RawPixels;
-	/** @todo Doesnt belong here. */
-	uint16 MaxHeight;
-	/** The character that was rendered */
-	TCHAR Char;
-	/** The index of the glyph from the FreeType face that was rendered */
-	uint32 GlyphIndex;
-	/** Whether or not the character has kerning */
-	bool HasKerning;
+	int16 HorizontalOffset = 0;
+	/** True if the rendered character is 8-bit grayscale, or false if it's 8-bit per-channel BGRA color */
+	bool bIsGrayscale = true;
+	/** True if the rendered character supports outlines, false otherwise */
+	bool bSupportsOutline = false;
 };
 
 /**
@@ -141,6 +120,11 @@ public:
 	 */
 	virtual class FTextureResource* GetEngineTexture() = 0;
 
+	/**
+	 * Returns whether the texture resource is 8-bit grayscale or 8-bit per-channel BGRA color
+	 */
+	virtual bool IsGrayscale() const = 0;
+
 	/**
 	 * Releases rendering resources of this texture
 	 */
@@ -153,9 +137,12 @@ public:
 class SLATECORE_API FSlateFontAtlas : public ISlateFontTexture, public FSlateTextureAtlas
 {
 public:
-	FSlateFontAtlas( uint32 InWidth, uint32 InHeight );
+	FSlateFontAtlas(uint32 InWidth, uint32 InHeight, const bool InIsGrayscale);
 	virtual ~FSlateFontAtlas();
 
+	//~ ISlateFontTexture interface
+	virtual bool IsGrayscale() const override final;
+
 	/**
 	 * Flushes all cached data.
 	 */
@@ -172,7 +159,7 @@ public:
 class ISlateFontAtlasFactory
 {
 public:
-	virtual FIntPoint GetAtlasSize() const = 0;
-	virtual TSharedRef<FSlateFontAtlas> CreateFontAtlas() const = 0;
-	virtual TSharedPtr<ISlateFontTexture> CreateNonAtlasedTexture(const uint32 InWidth, const uint32 InHeight, const TArray<uint8>& InRawData) const = 0;
+	virtual FIntPoint GetAtlasSize(const bool InIsGrayscale) const = 0;
+	virtual TSharedRef<FSlateFontAtlas> CreateFontAtlas(const bool InIsGrayscale) const = 0;
+	virtual TSharedPtr<ISlateFontTexture> CreateNonAtlasedTexture(const uint32 InWidth, const uint32 InHeight, const bool InIsGrayscale, const TArray<uint8>& InRawData) const = 0;
 };
diff --git a/Engine/Source/Runtime/SlateCore/Public/Rendering/RenderingCommon.h b/Engine/Source/Runtime/SlateCore/Public/Rendering/RenderingCommon.h
--- a/Engine/Source/Runtime/SlateCore/Public/Rendering/RenderingCommon.h
+++ b/Engine/Source/Runtime/SlateCore/Public/Rendering/RenderingCommon.h
@@ -54,14 +54,16 @@ enum class ESlateShader : uint8
 	Default = 0,
 	/** Border shader */
 	Border = 1,
-	/** Font shader, same as default except uses an alpha only texture */
-	Font = 2,
+	/** Grayscale font shader. Uses an alpha only texture */
+	GrayscaleFont = 2,
+	/** Grayscale font shader. Uses an color texture */
+	ColorFont = 3,
 	/** Line segment shader. For drawing anti-aliased lines */
-	LineSegment = 3,
+	LineSegment = 4,
 	/** For completely customized materials.  Makes no assumptions on use*/
-	Custom = 4,
+	Custom = 5,
 	/** For post processing passes */
-	PostProcess = 5,
+	PostProcess = 6,
 };
 
 /**
diff --git a/Engine/Source/Runtime/SlateCore/Public/Textures/TextureAtlas.h b/Engine/Source/Runtime/SlateCore/Public/Textures/TextureAtlas.h
--- a/Engine/Source/Runtime/SlateCore/Public/Textures/TextureAtlas.h
+++ b/Engine/Source/Runtime/SlateCore/Public/Textures/TextureAtlas.h
@@ -225,12 +225,9 @@ public:
 	/** Get the number of atlas pages this atlas provider has available when calling GetAtlasPageResource */
 	virtual int32 GetNumAtlasPages() const = 0;
 
-	/** Get the size of each atlas page */
-	virtual FIntPoint GetAtlasPageSize() const = 0;
-
 	/** Get the page resource for the given index (verify with GetNumAtlasPages) */ 
 	virtual class FSlateShaderResource* GetAtlasPageResource(const int32 InIndex) const = 0;
 
-	/** Do the atlas page resources only contain alpha information? This affects how the atlas visualizer will sample them */
-	virtual bool IsAtlasPageResourceAlphaOnly() const = 0;
+	/** Does the page resources for the given index only contain alpha information? This affects how the atlas visualizer will sample them (verify with GetNumAtlasPages) */
+	virtual bool IsAtlasPageResourceAlphaOnly(const int32 InIndex) const = 0;
 };
diff --git a/Engine/Source/Runtime/SlateNullRenderer/Private/SlateNullRendererModule.cpp b/Engine/Source/Runtime/SlateNullRenderer/Private/SlateNullRendererModule.cpp
--- a/Engine/Source/Runtime/SlateNullRenderer/Private/SlateNullRendererModule.cpp
+++ b/Engine/Source/Runtime/SlateNullRenderer/Private/SlateNullRendererModule.cpp
@@ -18,9 +18,8 @@ class FSlateNullShaderResourceManager : public ISlateAtlasProvider, public FSlat
 public:
 	// ISlateAtlasProvider interface
 	virtual int32 GetNumAtlasPages() const override { return 0; }
-	virtual FIntPoint GetAtlasPageSize() const override { return FIntPoint(0, 0); }
 	virtual class FSlateShaderResource* GetAtlasPageResource(const int32 InIndex) const override { return nullptr; }
-	virtual bool IsAtlasPageResourceAlphaOnly() const override { return false; }
+	virtual bool IsAtlasPageResourceAlphaOnly(const int32 InIndex) const override { return false; }
 
 	// FSlateShaderResourceManager interface
 	virtual FSlateShaderResourceProxy* GetShaderResource( const FSlateBrush& InBrush ) override { return nullptr; }
@@ -62,7 +61,7 @@ class FSlateFontAtlasNull : public FSlateFontAtlas
 {
 public:
 	FSlateFontAtlasNull(float AtlasSize)
-		: FSlateFontAtlas(AtlasSize, AtlasSize)
+		: FSlateFontAtlas(AtlasSize, AtlasSize, true)
 	{}
 
 	virtual class FSlateShaderResource* GetSlateTexture() override { return &NullFontTexture; }
@@ -91,17 +90,17 @@ public:
 
 	virtual ~FSlateNullFontAtlasFactory() {}
 
-	virtual FIntPoint GetAtlasSize() const override
+	virtual FIntPoint GetAtlasSize(const bool InIsGrayscale) const override
 	{
 		return FIntPoint(AtlasSize, AtlasSize);
 	}
 
-	virtual TSharedRef<FSlateFontAtlas> CreateFontAtlas() const override
+	virtual TSharedRef<FSlateFontAtlas> CreateFontAtlas(const bool InIsGrayscale) const override
 	{
 		return MakeShareable(new FSlateFontAtlasNull(AtlasSize));
 	}
 
-	virtual TSharedPtr<ISlateFontTexture> CreateNonAtlasedTexture(const uint32 InWidth, const uint32 InHeight, const TArray<uint8>& InRawData) const override
+	virtual TSharedPtr<ISlateFontTexture> CreateNonAtlasedTexture(const uint32 InWidth, const uint32 InHeight, const bool InIsGrayscale, const TArray<uint8>& InRawData) const override
 	{
 		return nullptr;
 	}
diff --git a/Engine/Source/Runtime/SlateRHIRenderer/Private/SlateMaterialShader.cpp b/Engine/Source/Runtime/SlateRHIRenderer/Private/SlateMaterialShader.cpp
--- a/Engine/Source/Runtime/SlateRHIRenderer/Private/SlateMaterialShader.cpp
+++ b/Engine/Source/Runtime/SlateRHIRenderer/Private/SlateMaterialShader.cpp
@@ -169,5 +169,7 @@ IMPLEMENT_SLATE_MATERIALSHADER_TYPE(Default, true);
 IMPLEMENT_SLATE_MATERIALSHADER_TYPE(Default, false);
 IMPLEMENT_SLATE_MATERIALSHADER_TYPE(Border, true);
 IMPLEMENT_SLATE_MATERIALSHADER_TYPE(Border, false);
-IMPLEMENT_SLATE_MATERIALSHADER_TYPE(Font, true);
-IMPLEMENT_SLATE_MATERIALSHADER_TYPE(Font, false);
+IMPLEMENT_SLATE_MATERIALSHADER_TYPE(GrayscaleFont, true);
+IMPLEMENT_SLATE_MATERIALSHADER_TYPE(GrayscaleFont, false);
+IMPLEMENT_SLATE_MATERIALSHADER_TYPE(ColorFont, true);
+IMPLEMENT_SLATE_MATERIALSHADER_TYPE(ColorFont, false);
diff --git a/Engine/Source/Runtime/SlateRHIRenderer/Private/SlateRHIFontTexture.cpp b/Engine/Source/Runtime/SlateRHIRenderer/Private/SlateRHIFontTexture.cpp
--- a/Engine/Source/Runtime/SlateRHIRenderer/Private/SlateRHIFontTexture.cpp
+++ b/Engine/Source/Runtime/SlateRHIRenderer/Private/SlateRHIFontTexture.cpp
@@ -5,9 +5,10 @@
 #include "RenderingThread.h"
 #include "RenderUtils.h"
 
-FSlateFontTextureRHIResource::FSlateFontTextureRHIResource(uint32 InWidth, uint32 InHeight)
-	: Width( InWidth )
-	, Height( InHeight )
+FSlateFontTextureRHIResource::FSlateFontTextureRHIResource(uint32 InWidth, uint32 InHeight, const bool InIsGrayscale)
+	: Width(InWidth)
+	, Height(InHeight)
+	, bIsGrayscale(InIsGrayscale)
 {
 }
 
@@ -18,9 +19,12 @@ void FSlateFontTextureRHIResource::InitDynamicRHI()
 	// Create the texture
 	if( Width > 0 && Height > 0 )
 	{
+		const EPixelFormat PixelFormat = GetRHIPixelFormat();
+		const uint32 TextCreateFlags = TexCreate_Dynamic | (bIsGrayscale ? TexCreate_None : TexCreate_SRGB);
+
 		check( !IsValidRef( ShaderResource) );
 		FRHIResourceCreateInfo CreateInfo;
-		ShaderResource = RHICreateTexture2D( Width, Height, PF_A8, 1, 1, TexCreate_Dynamic, CreateInfo );
+		ShaderResource = RHICreateTexture2D( Width, Height, PixelFormat, 1, 1, TextCreateFlags, CreateInfo );
 		check( IsValidRef( ShaderResource ) );
 
 		// Also assign the reference to the FTextureResource variable so that the Engine can access it
@@ -54,7 +58,7 @@ void FSlateFontTextureRHIResource::InitDynamicRHI()
 		);
 		DeferredPassSamplerStateRHI = RHICreateSamplerState(DeferredPassSamplerStateInitializer);
 
-		INC_MEMORY_STAT_BY(STAT_SlateTextureGPUMemory, Width*Height*GPixelFormats[PF_A8].BlockBytes);
+		INC_MEMORY_STAT_BY(STAT_SlateTextureGPUMemory, Width*Height*GPixelFormats[PixelFormat].BlockBytes);
 	}
 }
 
@@ -65,15 +69,24 @@ void FSlateFontTextureRHIResource::ReleaseDynamicRHI()
 	// Release the texture
 	if( IsValidRef(ShaderResource) )
 	{
-		DEC_MEMORY_STAT_BY(STAT_SlateTextureGPUMemory, Width*Height*GPixelFormats[PF_A8].BlockBytes);
+		const EPixelFormat PixelFormat = GetRHIPixelFormat();
+
+		DEC_MEMORY_STAT_BY(STAT_SlateTextureGPUMemory, Width*Height*GPixelFormats[PixelFormat].BlockBytes);
 	}
 
 	ShaderResource.SafeRelease();
 }
 
-FSlateFontAtlasRHI::FSlateFontAtlasRHI( uint32 Width, uint32 Height )
-	: FSlateFontAtlas( Width, Height ) 
-	, FontTexture( new FSlateFontTextureRHIResource( Width, Height ) )
+EPixelFormat FSlateFontTextureRHIResource::GetRHIPixelFormat() const
+{
+	return bIsGrayscale
+		? PF_A8
+		: PF_B8G8R8A8;
+}
+
+FSlateFontAtlasRHI::FSlateFontAtlasRHI(uint32 Width, uint32 Height, const bool InIsGrayscale)
+	: FSlateFontAtlas(Width, Height, InIsGrayscale) 
+	, FontTexture(new FSlateFontTextureRHIResource(Width, Height, InIsGrayscale))
 {
 }
 
@@ -124,8 +137,8 @@ void FSlateFontAtlasRHI::ConditionalUpdateTexture()
 	}
 }
 
-FSlateFontTextureRHI::FSlateFontTextureRHI(const uint32 InWidth, const uint32 InHeight, const TArray<uint8>& InRawData)
-	: FontTexture(new FSlateFontTextureRHIResource(InWidth, InHeight))
+FSlateFontTextureRHI::FSlateFontTextureRHI(const uint32 InWidth, const uint32 InHeight, const bool InIsGrayscale, const TArray<uint8>& InRawData)
+	: FontTexture(new FSlateFontTextureRHIResource(InWidth, InHeight, InIsGrayscale))
 {
 	if (IsInRenderingThread())
 	{
@@ -163,7 +176,7 @@ void FSlateFontTextureRHI::ReleaseResources()
 
 void FSlateFontTextureRHI::UpdateTextureFromSource(const uint32 SourceWidth, const uint32 SourceHeight, const TArray<uint8>& SourceData)
 {
-	static const uint32 BytesPerPixel = sizeof(uint8);
+	const uint32 BytesPerPixel = FontTexture->IsGrayscale() ? 1 : 4;
 
 	uint32 DestStride;
 	uint8* LockedTextureData = static_cast<uint8*>(RHILockTexture2D(FontTexture->GetTypedResource(), 0, RLM_WriteOnly, /*out*/ DestStride, false));
diff --git a/Engine/Source/Runtime/SlateRHIRenderer/Private/SlateRHIFontTexture.h b/Engine/Source/Runtime/SlateRHIRenderer/Private/SlateRHIFontTexture.h
--- a/Engine/Source/Runtime/SlateRHIRenderer/Private/SlateRHIFontTexture.h
+++ b/Engine/Source/Runtime/SlateRHIRenderer/Private/SlateRHIFontTexture.h
@@ -20,7 +20,7 @@ public:
 	 * @param InWidth The width of the texture
 	 * @param InHeight The height of the texture
 	 */
-	FSlateFontTextureRHIResource(uint32 InWidth, uint32 InHeight);
+	FSlateFontTextureRHIResource(uint32 InWidth, uint32 InHeight, const bool InIsGrayscale);
 
 	/** FSlateShaderResource interface */
 	virtual uint32 GetWidth() const override { return Width; }
@@ -34,11 +34,19 @@ public:
 	/** FRenderResource interface */
 	virtual void InitDynamicRHI() override;
 	virtual void ReleaseDynamicRHI() override;
+	
+	/** Returns whether the texture resource is 8-bit grayscale or 8-bit per-channel BGRA color */
+	bool IsGrayscale() const { return bIsGrayscale; }
+
 private:
+	EPixelFormat GetRHIPixelFormat() const;
+
 	/** Width of this texture */
 	uint32 Width;
 	/** Height of this texture */
 	uint32 Height;
+	/** Whether this texture is 8-bit grayscale or 8-bit per-channel BGRA color */
+	bool bIsGrayscale;
 	/** Temporary data stored between Release and InitDynamicRHI */
 	TArray<uint8> TempData;
 };
@@ -49,7 +57,7 @@ private:
 class FSlateFontAtlasRHI : public FSlateFontAtlas
 {
 public:
-	FSlateFontAtlasRHI( uint32 Width, uint32 Height );
+	FSlateFontAtlasRHI(uint32 Width, uint32 Height, const bool InIsGrayscale);
 	~FSlateFontAtlasRHI();
 
 	/**
@@ -69,7 +77,7 @@ private:
 class FSlateFontTextureRHI : public ISlateFontTexture
 {
 public:
-	FSlateFontTextureRHI(const uint32 InWidth, const uint32 InHeight, const TArray<uint8>& InRawData);
+	FSlateFontTextureRHI(const uint32 InWidth, const uint32 InHeight, const bool InIsGrayscale, const TArray<uint8>& InRawData);
 	~FSlateFontTextureRHI();
 
 	/**
@@ -77,6 +85,7 @@ public:
 	 */
 	virtual class FSlateShaderResource* GetSlateTexture() override { return FontTexture.Get(); }
 	virtual class FTextureResource* GetEngineTexture() override { return FontTexture.Get(); }
+	virtual bool IsGrayscale() const override { return FontTexture->IsGrayscale(); }
 	virtual void ReleaseResources() override;
 private:
 	void UpdateTextureFromSource(const uint32 SourceWidth, const uint32 SourceHeight, const TArray<uint8>& SourceData);
diff --git a/Engine/Source/Runtime/SlateRHIRenderer/Private/SlateRHIRendererModule.cpp b/Engine/Source/Runtime/SlateRHIRenderer/Private/SlateRHIRendererModule.cpp
--- a/Engine/Source/Runtime/SlateRHIRenderer/Private/SlateRHIRendererModule.cpp
+++ b/Engine/Source/Runtime/SlateRHIRenderer/Private/SlateRHIRendererModule.cpp
@@ -23,41 +23,46 @@ public:
 	{
 		if (GIsEditor)
 		{
-			AtlasSize = 2048;
+			GrayscaleAtlasSize = 2048;
 		}
 		else
 		{
-			AtlasSize = 1024;
+			GrayscaleAtlasSize = 1024;
 			if (GConfig)
 			{
-				GConfig->GetInt(TEXT("SlateRenderer"), TEXT("FontAtlasSize"), AtlasSize, GEngineIni);
-				AtlasSize = FMath::Clamp(AtlasSize, 0, 2048);
+				GConfig->GetInt(TEXT("SlateRenderer"), TEXT("FontAtlasSize"), GrayscaleAtlasSize, GEngineIni);
+				GrayscaleAtlasSize = FMath::Clamp(GrayscaleAtlasSize, 0, 2048);
 			}
 		}
+		ColorAtlasSize = GrayscaleAtlasSize / 4;
 	}
 
 	virtual ~FSlateRHIFontAtlasFactory()
 	{
 	}
 
-	virtual FIntPoint GetAtlasSize() const override
+	virtual FIntPoint GetAtlasSize(const bool InIsGrayscale) const override
 	{
-		return FIntPoint(AtlasSize, AtlasSize);
+		return InIsGrayscale
+			? FIntPoint(GrayscaleAtlasSize, GrayscaleAtlasSize)
+			: FIntPoint(ColorAtlasSize, ColorAtlasSize);
 	}
 
-	virtual TSharedRef<FSlateFontAtlas> CreateFontAtlas() const override
+	virtual TSharedRef<FSlateFontAtlas> CreateFontAtlas(const bool InIsGrayscale) const override
 	{
-		return MakeShareable(new FSlateFontAtlasRHI(AtlasSize, AtlasSize));
+		const FIntPoint AtlasSize = GetAtlasSize(InIsGrayscale);
+		return MakeShareable(new FSlateFontAtlasRHI(AtlasSize.X, AtlasSize.Y, InIsGrayscale));
 	}
 
-	virtual TSharedPtr<ISlateFontTexture> CreateNonAtlasedTexture(const uint32 InWidth, const uint32 InHeight, const TArray<uint8>& InRawData) const override
+	virtual TSharedPtr<ISlateFontTexture> CreateNonAtlasedTexture(const uint32 InWidth, const uint32 InHeight, const bool InIsGrayscale, const TArray<uint8>& InRawData) const override
 	{
 		if (GIsEditor)
 		{
-			const uint32 MaxFontTextureDimension = FMath::Min(AtlasSize * 4u, GetMax2DTextureDimension()); // Don't allow textures greater than 4x our atlas size, but still honor the platform limit
+			const FIntPoint AtlasSize = GetAtlasSize(InIsGrayscale);
+			const uint32 MaxFontTextureDimension = FMath::Min(AtlasSize.Y * 4u, GetMax2DTextureDimension()); // Don't allow textures greater than 4x our atlas size, but still honor the platform limit
 			if (InWidth <= MaxFontTextureDimension && InHeight <= MaxFontTextureDimension)
 			{
-				return MakeShareable(new FSlateFontTextureRHI(InWidth, InHeight, InRawData));
+				return MakeShareable(new FSlateFontTextureRHI(InWidth, InHeight, InIsGrayscale, InRawData));
 			}
 		}
 		return nullptr;
@@ -65,7 +70,8 @@ public:
 
 private:
 	/** Size of each font texture, width and height */
-	int32 AtlasSize;
+	int32 GrayscaleAtlasSize;
+	int32 ColorAtlasSize;
 };
 
 
diff --git a/Engine/Source/Runtime/SlateRHIRenderer/Private/SlateRHIRenderingPolicy.cpp b/Engine/Source/Runtime/SlateRHIRenderer/Private/SlateRHIRenderingPolicy.cpp
--- a/Engine/Source/Runtime/SlateRHIRenderer/Private/SlateRHIRenderingPolicy.cpp
+++ b/Engine/Source/Runtime/SlateRHIRenderer/Private/SlateRHIRenderingPolicy.cpp
@@ -1366,8 +1366,11 @@ FSlateElementPS* FSlateRHIRenderingPolicy::GetTexturePixelShader( TShaderMap<FGl
 				PixelShader = *TShaderMapRef<TSlateElementPS<ESlateShader::Border, true, false> >(ShaderMap);
 			}
 			break;
-		case ESlateShader::Font:
-			PixelShader = *TShaderMapRef<TSlateElementPS<ESlateShader::Font, true> >(ShaderMap);
+		case ESlateShader::GrayscaleFont:
+			PixelShader = *TShaderMapRef<TSlateElementPS<ESlateShader::GrayscaleFont, true> >(ShaderMap);
+			break;
+		case ESlateShader::ColorFont:
+			PixelShader = *TShaderMapRef<TSlateElementPS<ESlateShader::ColorFont, true> >(ShaderMap);
 			break;
 		case ESlateShader::LineSegment:
 			PixelShader = *TShaderMapRef<TSlateElementPS<ESlateShader::LineSegment, true> >(ShaderMap);
@@ -1399,8 +1402,11 @@ FSlateElementPS* FSlateRHIRenderingPolicy::GetTexturePixelShader( TShaderMap<FGl
 				PixelShader = *TShaderMapRef<TSlateElementPS<ESlateShader::Border, false, false> >(ShaderMap);
 			}
 			break;
-		case ESlateShader::Font:
-			PixelShader = *TShaderMapRef<TSlateElementPS<ESlateShader::Font, false> >(ShaderMap);
+		case ESlateShader::GrayscaleFont:
+			PixelShader = *TShaderMapRef<TSlateElementPS<ESlateShader::GrayscaleFont, false> >(ShaderMap);
+			break;
+		case ESlateShader::ColorFont:
+			PixelShader = *TShaderMapRef<TSlateElementPS<ESlateShader::ColorFont, false> >(ShaderMap);
 			break;
 		case ESlateShader::LineSegment:
 			PixelShader = *TShaderMapRef<TSlateElementPS<ESlateShader::LineSegment, false> >(ShaderMap);
@@ -1444,14 +1450,24 @@ FSlateMaterialShaderPS* FSlateRHIRenderingPolicy::GetMaterialPixelShader( const
 			FoundShader = MaterialShaderMap->GetShader(&TSlateMaterialShaderPS<ESlateShader::Border, false>::StaticType);
 		}
 		break;
-	case ESlateShader::Font:
+	case ESlateShader::GrayscaleFont:
 		if(bDrawDisabled)
 		{
-			FoundShader = MaterialShaderMap->GetShader(&TSlateMaterialShaderPS<ESlateShader::Font, true>::StaticType);
+			FoundShader = MaterialShaderMap->GetShader(&TSlateMaterialShaderPS<ESlateShader::GrayscaleFont, true>::StaticType);
+		}
+		else
+		{
+			FoundShader = MaterialShaderMap->GetShader(&TSlateMaterialShaderPS<ESlateShader::GrayscaleFont, false>::StaticType);
+		}
+		break;
+	case ESlateShader::ColorFont:
+		if (bDrawDisabled)
+		{
+			FoundShader = MaterialShaderMap->GetShader(&TSlateMaterialShaderPS<ESlateShader::ColorFont, true>::StaticType);
 		}
 		else
 		{
-			FoundShader = MaterialShaderMap->GetShader(&TSlateMaterialShaderPS<ESlateShader::Font, false>::StaticType);
+			FoundShader = MaterialShaderMap->GetShader(&TSlateMaterialShaderPS<ESlateShader::ColorFont, false>::StaticType);
 		}
 		break;
 	case ESlateShader::Custom:
diff --git a/Engine/Source/Runtime/SlateRHIRenderer/Private/SlateRHIResourceManager.cpp b/Engine/Source/Runtime/SlateRHIRenderer/Private/SlateRHIResourceManager.cpp
--- a/Engine/Source/Runtime/SlateRHIRenderer/Private/SlateRHIResourceManager.cpp
+++ b/Engine/Source/Runtime/SlateRHIRenderer/Private/SlateRHIResourceManager.cpp
@@ -272,17 +272,12 @@ int32 FSlateRHIResourceManager::GetNumAtlasPages() const
 	return TextureAtlases.Num();
 }
 
-FIntPoint FSlateRHIResourceManager::GetAtlasPageSize() const
-{
-	return FIntPoint(1024, 1024);
-}
-
 FSlateShaderResource* FSlateRHIResourceManager::GetAtlasPageResource(const int32 InIndex) const
 {
 	return TextureAtlases[InIndex]->GetAtlasTexture();
 }
 
-bool FSlateRHIResourceManager::IsAtlasPageResourceAlphaOnly() const
+bool FSlateRHIResourceManager::IsAtlasPageResourceAlphaOnly(const int32 InIndex) const
 {
 	return false;
 }
diff --git a/Engine/Source/Runtime/SlateRHIRenderer/Private/SlateRHIResourceManager.h b/Engine/Source/Runtime/SlateRHIRenderer/Private/SlateRHIResourceManager.h
--- a/Engine/Source/Runtime/SlateRHIRenderer/Private/SlateRHIResourceManager.h
+++ b/Engine/Source/Runtime/SlateRHIRenderer/Private/SlateRHIResourceManager.h
@@ -128,9 +128,8 @@ public:
 
 	/** ISlateAtlasProvider interface */
 	virtual int32 GetNumAtlasPages() const override;
-	virtual FIntPoint GetAtlasPageSize() const override;
 	virtual FSlateShaderResource* GetAtlasPageResource(const int32 InIndex) const override;
-	virtual bool IsAtlasPageResourceAlphaOnly() const override;
+	virtual bool IsAtlasPageResourceAlphaOnly(const int32 InIndex) const override;
 
 	/** FTickableGameObject interface */
 	virtual ETickableTickType GetTickableTickType() const override { return ETickableTickType::Always; }
diff --git a/Engine/Source/Runtime/SlateRHIRenderer/Private/SlateShaders.cpp b/Engine/Source/Runtime/SlateRHIRenderer/Private/SlateShaders.cpp
--- a/Engine/Source/Runtime/SlateRHIRenderer/Private/SlateShaders.cpp
+++ b/Engine/Source/Runtime/SlateRHIRenderer/Private/SlateShaders.cpp
@@ -42,9 +42,13 @@ IMPLEMENT_SLATE_PIXELSHADER_TYPE(Border, false, false);
 IMPLEMENT_SLATE_PIXELSHADER_TYPE(Default, true, false);
 IMPLEMENT_SLATE_PIXELSHADER_TYPE(Border, true, false);
 
-IMPLEMENT_SLATE_PIXELSHADER_TYPE(Font, false, true);
+IMPLEMENT_SLATE_PIXELSHADER_TYPE(GrayscaleFont, false, true);
+IMPLEMENT_SLATE_PIXELSHADER_TYPE(GrayscaleFont, true, true);
+
+IMPLEMENT_SLATE_PIXELSHADER_TYPE(ColorFont, false, true);
+IMPLEMENT_SLATE_PIXELSHADER_TYPE(ColorFont, true, true);
+
 IMPLEMENT_SLATE_PIXELSHADER_TYPE(LineSegment, false, true);
-IMPLEMENT_SLATE_PIXELSHADER_TYPE(Font, true, true);
 IMPLEMENT_SLATE_PIXELSHADER_TYPE(LineSegment, true, true);
 
 /** The Slate vertex declaration. */
