Ensure font size and scale aren't set to negative values within FreeType

#rnx

diff --git a/Engine/Source/Runtime/SlateCore/Private/Fonts/FontCacheFreeType.cpp b/Engine/Source/Runtime/SlateCore/Private/Fonts/FontCacheFreeType.cpp
--- a/Engine/Source/Runtime/SlateCore/Private/Fonts/FontCacheFreeType.cpp
+++ b/Engine/Source/Runtime/SlateCore/Private/Fonts/FontCacheFreeType.cpp
@@ -69,8 +69,8 @@ namespace FreeTypeUtils
 void ApplySizeAndScale(FT_Face InFace, const int32 InFontSize, const float InFontScale)
 {
 	// Convert to fixed scale for maximum precision
-	const FT_F26Dot6 FixedFontSize = ConvertPixelTo26Dot6<FT_F26Dot6>(InFontSize);
-	const FT_Long FixedFontScale = ConvertPixelTo16Dot16<FT_Long>(InFontScale);
+	const FT_F26Dot6 FixedFontSize = ConvertPixelTo26Dot6<FT_F26Dot6>(FMath::Max<int32>(0, InFontSize));
+	const FT_Long FixedFontScale = ConvertPixelTo16Dot16<FT_Long>(FMath::Max<float>(0.0f, InFontScale));
 
 	// Convert the requested font size to a pixel size based on our render DPI and the requested scale
 	// This result is left in 26.6 space as it is the common format used by both FT_IS_SCALABLE and FT_HAS_FIXED_SIZES fonts
