Fixed spline arrows and graph execution bubbles not rendering in Slate
#rnx

diff --git a/Engine/Source/Runtime/SlateCore/Private/Rendering/DrawElements.cpp b/Engine/Source/Runtime/SlateCore/Private/Rendering/DrawElements.cpp
--- a/Engine/Source/Runtime/SlateCore/Private/Rendering/DrawElements.cpp
+++ b/Engine/Source/Runtime/SlateCore/Private/Rendering/DrawElements.cpp
@@ -210,8 +210,6 @@ FSlateDrawElement& FSlateDrawElement::MakeBoxInternal(
 	const FLinearColor& InTint
 )
 {
-	PaintGeometry.CommitTransformsIfUsingLegacyConstructor();
-
 	EElementType ElementType = (InBrush->DrawAs == ESlateBrushDrawType::Border) ? EElementType::ET_Border : EElementType::ET_Box;
 
 	FSlateDrawElement& Element = ElementList.AddUninitialized();
@@ -235,6 +233,7 @@ void FSlateDrawElement::MakeBox(
 	ESlateDrawEffect InDrawEffects, 
 	const FLinearColor& InTint)
 {
+	PaintGeometry.CommitTransformsIfUsingLegacyConstructor();
 
 	if (ShouldCull(ElementList, PaintGeometry, InBrush, InTint))
 	{
@@ -267,6 +266,7 @@ void FSlateDrawElement::MakeRotatedBox(
 	ERotationSpace RotationSpace,
 	const FLinearColor& InTint)
 {
+	PaintGeometry.CommitTransformsIfUsingLegacyConstructor();
 
 	if (ShouldCull(ElementList, PaintGeometry, InBrush, InTint))
 	{
