Build fix - forward declare SetWidgetAttribute, and move it to the bottom of the SWidget.h header.

diff --git a/Engine/Source/Runtime/SlateCore/Public/Widgets/SWidget.h b/Engine/Source/Runtime/SlateCore/Public/Widgets/SWidget.h
--- a/Engine/Source/Runtime/SlateCore/Public/Widgets/SWidget.h
+++ b/Engine/Source/Runtime/SlateCore/Public/Widgets/SWidget.h
@@ -161,26 +161,7 @@ private:
  * volatility won't need to be recalculated.  Returns true if the value changed.
  */
 template<typename TargetValueType, typename SourceValueType>
-static bool SetWidgetAttribute(SWidget& ThisWidget, TAttribute<TargetValueType>& TargetValue, const TAttribute<SourceValueType>& SourceValue, EInvalidateWidgetReason BaseInvalidationReason)
-{
-	if (!TargetValue.IdenticalTo(SourceValue))
-	{
-		const bool bWasBound = TargetValue.IsBound();
-		const bool bBoundnessChanged = bWasBound != SourceValue.IsBound();
-		TargetValue = SourceValue;
-
-		EInvalidateWidgetReason InvalidateReason = BaseInvalidationReason;
-		if (bBoundnessChanged)
-		{
-			InvalidateReason |= EInvalidateWidgetReason::Volatility;
-		}
-
-		ThisWidget.Invalidate(InvalidateReason);
-		return true;
-	}
-
-	return false;
-}
+static bool SetWidgetAttribute(SWidget& ThisWidget, TAttribute<TargetValueType>& TargetValue, const TAttribute<SourceValueType>& SourceValue, EInvalidateWidgetReason BaseInvalidationReason);
 
 class IToolTip;
 
@@ -1725,3 +1706,26 @@ FORCEINLINE_DEBUGGABLE FArrangedWidget FGeometry::MakeChild(const TSharedRef<SWi
 	// The extra TransformPoint below does this by converting the local offset to an offset in parent space.
 	return MakeChild(ChildWidget, InLocalSize, FSlateLayoutTransform(ChildScale, TransformPoint(ChildScale, ChildOffset)));
 }
+
+
+template<typename TargetValueType, typename SourceValueType>
+bool SetWidgetAttribute(SWidget& ThisWidget, TAttribute<TargetValueType>& TargetValue, const TAttribute<SourceValueType>& SourceValue, EInvalidateWidgetReason BaseInvalidationReason)
+{
+	if (!TargetValue.IdenticalTo(SourceValue))
+	{
+		const bool bWasBound = TargetValue.IsBound();
+		const bool bBoundnessChanged = bWasBound != SourceValue.IsBound();
+		TargetValue = SourceValue;
+
+		EInvalidateWidgetReason InvalidateReason = BaseInvalidationReason;
+		if (bBoundnessChanged)
+		{
+			InvalidateReason |= EInvalidateWidgetReason::Volatility;
+		}
+
+		ThisWidget.Invalidate(InvalidateReason);
+		return true;
+	}
+
+	return false;
+}
\ No newline at end of file
