Fix regression where custom title bars would not appear in slate windows

diff --git a/Engine/Source/Runtime/SlateCore/Private/Widgets/SWindow.cpp b/Engine/Source/Runtime/SlateCore/Private/Widgets/SWindow.cpp
--- a/Engine/Source/Runtime/SlateCore/Private/Widgets/SWindow.cpp
+++ b/Engine/Source/Runtime/SlateCore/Private/Widgets/SWindow.cpp
@@ -511,7 +511,7 @@ void SWindow::ConstructWindowInternals()
 		SNew( SVerticalBox )
 		.Visibility( EVisibility::SelfHitTestInvisible );
 
-	TSharedRef<SWidget> TitleBarWidget = FSlateApplicationBase::Get().MakeWindowTitleBar(SharedThis(this), nullptr, GetTitleAlignment(), TitleBar);
+	TSharedRef<SWidget> TitleBarWidget = MakeWindowTitleBar(SharedThis(this), nullptr, GetTitleAlignment());
 
 	if (bCreateTitleBar)
 	{
