Enable an Emoji font for the UE4 editor

diff --git a/Engine/Source/Runtime/SlateCore/Private/Fonts/LegacySlateFontInfoCache.cpp b/Engine/Source/Runtime/SlateCore/Private/Fonts/LegacySlateFontInfoCache.cpp
--- a/Engine/Source/Runtime/SlateCore/Private/Fonts/LegacySlateFontInfoCache.cpp
+++ b/Engine/Source/Runtime/SlateCore/Private/Fonts/LegacySlateFontInfoCache.cpp
@@ -205,6 +205,21 @@ TSharedRef<const FCompositeFont> FLegacySlateFontInfoCache::GetDefaultFont()
 			APPEND_EDITOR_FONT(SubFont.Typeface, "VeryLight", "NanumGothic.ttf", EFontHinting::Default);
 		}
 
+		// Emoji (editor-only)
+		if (GIsEditor)
+		{
+			FCompositeSubFont& SubFont = MutableDefaultFont->SubTypefaces.AddDefaulted_GetRef();
+			APPEND_RANGE(SubFont, EmoticonsEmoji);
+			APPEND_RANGE(SubFont, MiscellaneousSymbols);
+			APPEND_RANGE(SubFont, MiscellaneousSymbolsAndArrows);
+			APPEND_RANGE(SubFont, MiscellaneousSymbolsAndPictographs);
+			APPEND_RANGE(SubFont, SupplementalSymbolsAndPictographs);
+			APPEND_RANGE(SubFont, TransportAndMapSymbols);
+			APPEND_RANGE(SubFont, Dingbats);
+			APPEND_RANGE(SubFont, EnclosedAlphanumericSupplement);
+			APPEND_EDITOR_FONT(SubFont.Typeface, "Regular", "NotoColorEmoji.ttf", EFontHinting::Default);
+		}
+
 		#undef APPEND_FONT
 		#undef APPEND_EDITOR_FONT
 		#undef APPEND_RANGE
