Slate - Creating an alternative play of Slate's CurveSequence that manages the lifetime of the animation through a ticker rather than a specific widget.

diff --git a/Engine/Source/Runtime/SlateCore/Private/Animation/CurveSequence.cpp b/Engine/Source/Runtime/SlateCore/Private/Animation/CurveSequence.cpp
--- a/Engine/Source/Runtime/SlateCore/Private/Animation/CurveSequence.cpp
+++ b/Engine/Source/Runtime/SlateCore/Private/Animation/CurveSequence.cpp
@@ -37,6 +37,11 @@ FCurveSequence::~FCurveSequence()
 	{
 		OwnerWidget.Pin()->UnRegisterActiveTimer(ActiveTimerHandle.Pin().ToSharedRef());
 	}
+
+	if (TickerHandle.IsValid())
+	{
+		FTicker::GetCoreTicker().RemoveTicker(TickerHandle);
+	}
 }
 
 FCurveHandle FCurveSequence::AddCurve( const float InStartTimeSeconds, const float InDurationSeconds, const ECurveEaseFunction InEaseFunction )
@@ -75,6 +80,21 @@ void FCurveSequence::Play(const TSharedRef<SWidget>& InOwnerWidget, bool bPlayLo
 	SetStartTime(FSlateApplicationBase::Get().GetCurrentTime() - StartAtTime);
 }
 
+void FCurveSequence::Play(const FTickerDelegate& InDelegate, bool bPlayLooped, const float StartAtTime)
+{
+	bIsLooping = bPlayLooped;
+	bIsPaused = false;
+
+	// Playing forward
+	bInReverse = false;
+
+	// We start playing NOW.
+	SetStartTime(FSlateApplicationBase::Get().GetCurrentTime() - StartAtTime);
+
+	// Register the core ticker so we can start giving callbacks
+	TickerHandle = FTicker::GetCoreTicker().AddTicker(FTickerDelegate::CreateRaw(this, &FCurveSequence::TickPlay, InDelegate));
+}
+
 void FCurveSequence::Reverse( )
 {	
 	// We've played this far into the animation.
@@ -227,6 +247,11 @@ void FCurveSequence::RegisterActiveTimerIfNeeded(TSharedRef<SWidget> InOwnerWidg
 	}
 }
 
+bool FCurveSequence::TickPlay(float InDeltaTime, FTickerDelegate InUserDelegate)
+{
+	return (InUserDelegate.IsBound() && InUserDelegate.Execute(InDeltaTime)) && IsPlaying();
+}
+
 EActiveTimerReturnType FCurveSequence::EnsureSlateTickDuringAnimation( double InCurrentTime, float InDeltaTime )
 {
 	// Keep Slate ticking until play stops
diff --git a/Engine/Source/Runtime/SlateCore/Public/Animation/CurveSequence.h b/Engine/Source/Runtime/SlateCore/Public/Animation/CurveSequence.h
--- a/Engine/Source/Runtime/SlateCore/Public/Animation/CurveSequence.h
+++ b/Engine/Source/Runtime/SlateCore/Public/Animation/CurveSequence.h
@@ -4,6 +4,7 @@
 
 #include "CoreMinimal.h"
 #include "Animation/CurveHandle.h"
+#include "Containers/Ticker.h"
 
 class FActiveTimerHandle;
 class SWidget;
@@ -100,6 +101,11 @@ public:
 	 */
 	void Play( const TSharedRef<SWidget>& InOwnerWidget, bool bPlayLooped = false, const float StartAtTime = 0.0f, bool bRequiresActiveTimer = true);
 
+	/**
+	 * Plays the curve sequence for a ticker, rather than a widget.
+	 */
+	void Play(const FTickerDelegate& InDelegate, bool bPlayLooped = false, const float StartAtTime = 0.0f);
+
 	/**
 	 * Start playing this curve sequence in reverse. Registers an active timer for the widget using the sequence.
 	 *
@@ -174,6 +180,9 @@ private:
 	/** Helper to take care of registering the active timer */
 	void RegisterActiveTimerIfNeeded(TSharedRef<SWidget> InOwnerWidget);
 
+	/** Tick callback from the Ticker based plays. */
+	bool TickPlay(float InDeltaTime, FTickerDelegate InUserDelegate);
+
 	/** Hollow active timer to ensure a Slate Tick/Paint pass while the sequence is playing */
 	EActiveTimerReturnType EnsureSlateTickDuringAnimation( double InCurrentTime, float InDeltaTime );
 
@@ -188,6 +197,9 @@ private:
 	/** The handle to the active timer */
 	TWeakPtr<FActiveTimerHandle> ActiveTimerHandle;
 
+	/** Handle to the ticker based player. */
+	FDelegateHandle TickerHandle;
+
 	/** All the curves in this sequence. */
 	TArray<FSlateCurve> Curves;
 
