Slate - When doing WidgetProxy fast path update, we now return true correctly if any of the updated proxies required a repaint.  Also removed some stats and variables that were not being used.  This fixes the issue with retainer widgets using invalidation method of redrawing to correctly repaint when widgets invalidate layout and paint.

Matt.Kuhlenschmidt

diff --git a/Engine/Source/Runtime/SlateCore/Private/FastUpdate/SlateInvalidationRoot.cpp b/Engine/Source/Runtime/SlateCore/Private/FastUpdate/SlateInvalidationRoot.cpp
--- a/Engine/Source/Runtime/SlateCore/Private/FastUpdate/SlateInvalidationRoot.cpp
+++ b/Engine/Source/Runtime/SlateCore/Private/FastUpdate/SlateInvalidationRoot.cpp
@@ -7,9 +7,6 @@
 #include "Layout/Children.h"
 #include "ProfilingDebugging/CsvProfiler.h"
 
-DECLARE_DWORD_COUNTER_STAT(TEXT("Num Dirty Widgets"), STAT_SlateNumDirtyWidgets, STATGROUP_Slate);
-DECLARE_DWORD_COUNTER_STAT(TEXT("Num Volatile Widgets"), STAT_SlateNumVolatileWidgets, STATGROUP_Slate);
-
 #if WITH_SLATE_DEBUGGING
 bool GDumpUpdateList = false;
 void HandleDumpUpdateList(const TArray<FString>& Args)
@@ -278,12 +275,6 @@ bool FSlateInvalidationRoot::PaintFastPath(const FSlateInvalidationContext& Cont
 	{
 		TGuardValue<bool> OnFastPathGuard(GSlateIsOnFastUpdatePath, true);
 
-
-		bool bDrawElementsUpdated = false;
-
-		int32 NumDirtyWidgets = 0;
-		int32 NumVolatileWidgets = 0;
-
 		int32 LastParentIndex = 0;
 
 #if WITH_SLATE_DEBUGGING
@@ -330,6 +321,11 @@ bool FSlateInvalidationRoot::PaintFastPath(const FSlateInvalidationContext& Cont
 				// Check visibility, it may have been in the update list but a parent who was also in the update list already updated it.
 				if (!WidgetProxy.bInvisibleDueToParentOrSelfVisibility && !WidgetProxy.bUpdatedSinceLastInvalidate && ensure(WidgetProxy.Widget))
 				{
+					if (EnumHasAnyFlags(WidgetProxy.UpdateFlags, EWidgetUpdateFlags::NeedsRepaint | EWidgetUpdateFlags::NeedsVolatilePaint))
+					{
+						bWidgetsNeededRepaint = true;
+					}
+
 					const int32 NewLayerId = WidgetProxy.Update(*Context.PaintArgs, MyIndex, *Context.WindowElementList);
 					CachedMaxLayerId = FMath::Max(NewLayerId, CachedMaxLayerId);
 
@@ -351,9 +347,6 @@ bool FSlateInvalidationRoot::PaintFastPath(const FSlateInvalidationContext& Cont
 	}
 
 	return bWidgetsNeededRepaint;
-
-	INC_DWORD_STAT_BY(STAT_SlateNumDirtyWidgets, FinalUpdateList.Num());
-	//SET_DWORD_STAT(STAT_SlateNumVolatileWidgets, NumVolatileWidgets);
 }
 
 void FSlateInvalidationRoot::BuildNewFastPathList_Recursive(FSlateInvalidationRoot& Root, FWidgetProxy& Proxy, int32 ParentIndex, int32& NextTreeIndex, TArray<FWidgetProxy>& CurrentFastPathList, TArray<FWidgetProxy, TMemStackAllocator<>>& NewFastPathList)
