Merge CL #9455766 from Dev-Editor-  Fixes tab hover regression

- also fixes mode toolbar palette tabs to match dock tab colors

diff --git a/Engine/Source/Editor/EditorStyle/Private/SlateEditorStyle.cpp b/Engine/Source/Editor/EditorStyle/Private/SlateEditorStyle.cpp
--- a/Engine/Source/Editor/EditorStyle/Private/SlateEditorStyle.cpp
+++ b/Engine/Source/Editor/EditorStyle/Private/SlateEditorStyle.cpp
@@ -4943,9 +4943,9 @@ void FSlateEditorStyle::FStyle::SetupLevelEditorStyle()
 			.SetCheckedImage(             FSlateBoxBrush( RootToCoreContentDir("Docking/Tab_Shape.png"),  2.f/8.0f, FColor(62, 62, 62) ) )
 			.SetCheckedHoveredImage(      FSlateBoxBrush( RootToCoreContentDir("Docking/Tab_Shape.png"),  2.f/8.0f, FColor(62, 62, 62) ) )
 			.SetCheckedPressedImage(      FSlateBoxBrush( RootToCoreContentDir("Docking/Tab_Shape.png"),  2.f/8.0f, FColor(62, 62, 62) ) )
-			.SetUncheckedImage(           FSlateNoResource() )
-			.SetUncheckedHoveredImage(    FSlateNoResource() )
-			.SetUncheckedPressedImage(    FSlateNoResource() )
+			.SetUncheckedImage(           FSlateBoxBrush( RootToCoreContentDir("Docking/Tab_Shape.png"),  2.f/8.0f, FColor(45, 45, 45) ) )
+			.SetUncheckedHoveredImage(    FSlateBoxBrush( RootToCoreContentDir("Docking/Tab_Shape.png"),  2.f/8.0f, FColor(54, 54, 54) ) )    			
+			.SetUncheckedPressedImage(    FSlateBoxBrush( RootToCoreContentDir("Docking/Tab_Shape.png"),  2.f/8.0f, FColor(54, 54, 54) ) )    			   
 			.SetUndeterminedImage(        FSlateNoResource() )
 			.SetUndeterminedHoveredImage( FSlateNoResource() )
 			.SetUndeterminedPressedImage( FSlateNoResource() )
diff --git a/Engine/Source/Editor/UnrealEd/Private/EditorModeManager.cpp b/Engine/Source/Editor/UnrealEd/Private/EditorModeManager.cpp
--- a/Engine/Source/Editor/UnrealEd/Private/EditorModeManager.cpp
+++ b/Engine/Source/Editor/UnrealEd/Private/EditorModeManager.cpp
@@ -734,7 +734,7 @@ void FEditorModeTools::RebuildModeToolBar()
 					{
 						PaletteTabBox->AddSlot()
 						.AutoWidth()
-						.Padding(FMargin(3.f, 1.0f, 3.0f, 0.0f))
+						.Padding(FMargin(0.f, 1.0f, 1.0f, 0.0f))
 						[
 							SNew(SCheckBox)
 							.Style( FEditorStyle::Get(),  "ToolPalette.DockingTab" )
diff --git a/Engine/Source/Runtime/SlateCore/Private/Styling/CoreStyle.cpp b/Engine/Source/Runtime/SlateCore/Private/Styling/CoreStyle.cpp
--- a/Engine/Source/Runtime/SlateCore/Private/Styling/CoreStyle.cpp
+++ b/Engine/Source/Runtime/SlateCore/Private/Styling/CoreStyle.cpp
@@ -1110,19 +1110,20 @@ TSharedRef<ISlateStyle> FCoreStyle::Create( const FName& InStyleSetName )
 			.SetHovered( IMAGE_BRUSH( "/Docking/CloseApp_Hovered", Icon16x16 ) );
 
 
-		FLinearColor DockColor(FColor(62, 62, 62));
-		FLinearColor InactiveDockColor(FColor(50, 50, 50));
+		FLinearColor DockColor_Inactive(FColor(45, 45, 45));
+		FLinearColor DockColor_Hovered(FColor(54, 54, 54));
+		FLinearColor DockColor_Active(FColor(62, 62, 62));
 
 		// Panel Tab
 		Style->Set( "Docking.Tab", FDockTabStyle()
 			.SetCloseButtonStyle( CloseButton )
-			.SetNormalBrush( BOX_BRUSH( "/Docking/Tab_Shape", 2.f /8.0f, InactiveDockColor ) )
+			.SetNormalBrush( BOX_BRUSH( "/Docking/Tab_Shape", 2.f /8.0f, DockColor_Inactive ) )
+			.SetHoveredBrush( BOX_BRUSH( "/Docking/Tab_Shape", 2.f /8.0f, DockColor_Hovered ) )
+			.SetForegroundBrush( BOX_BRUSH( "/Docking/Tab_Shape", 2.f /8.0f, DockColor_Active ) )
 			.SetActiveBrush( BOX_BRUSH( "/Docking/Tab_Active", 4./8.0f) ) // 
 			.SetColorOverlayTabBrush( BOX_BRUSH( "/Docking/Tab_ColorOverlay", 4 / 16.0f))
 			.SetColorOverlayIconBrush( BOX_BRUSH( "/Docking/Tab_ColorOverlayIcon", 4/16.0f ) )
-			.SetForegroundBrush( BOX_BRUSH( "/Docking/Tab_Shape", 2.f /8.0f, DockColor ) )
-			.SetHoveredBrush( BOX_BRUSH( "/Docking/Tab_Shape", 2.f /8.0f, InactiveDockColor ) )
-			.SetContentAreaBrush( FSlateColorBrush( DockColor) )
+			.SetContentAreaBrush( FSlateColorBrush( DockColor_Active ) )
 			.SetTabWellBrush( FSlateNoResource() )
 			.SetTabPadding( FMargin(5, 2, 5, 2) )
 			.SetOverlapWidth( -1.0f )
