Fix alt-tab on DX12

diff --git a/Engine/Source/Runtime/SlateCore/Private/Rendering/SlateRenderer.cpp b/Engine/Source/Runtime/SlateCore/Private/Rendering/SlateRenderer.cpp
--- a/Engine/Source/Runtime/SlateCore/Private/Rendering/SlateRenderer.cpp
+++ b/Engine/Source/Runtime/SlateCore/Private/Rendering/SlateRenderer.cpp
@@ -112,6 +112,12 @@ bool FSlateRenderer::IsViewportFullscreen( const SWindow& Window ) const
 		else
 		{
 			bFullscreen = Window.GetWindowMode() == EWindowMode::Fullscreen;
+#if PLATFORM_WINDOWS
+			// When we are in fullscreen mode but the user alt-tabs out we need to temporarily drop out of fullscreen while the window has lost focus, otherwise DXGI will eventually
+			// forcibly throw us out of fullscreen mode with device loss and crash as typical result. By returning false here we'll trigger a mode switch to windowed when the user
+			// alt-tabs, and back to fullscreen again once the window comes back in focus, through the regular path. DXGI will never need to intervene and everyone is happy.
+			bFullscreen = bFullscreen && Window.GetNativeWindow()->IsForegroundWindow();
+#endif
 		}
 	}
 	else
