Bound clipping rect to the outline control bounds when rendering outlined fonts

The unbounded rect caused a serious performance regression in FreeType 2.10 when compared to FreeType 2.6 (which did this clamping internally).

#rnx

diff --git a/Engine/Source/Runtime/SlateCore/Private/Fonts/SlateFontRenderer.cpp b/Engine/Source/Runtime/SlateCore/Private/Fonts/SlateFontRenderer.cpp
--- a/Engine/Source/Runtime/SlateCore/Private/Fonts/SlateFontRenderer.cpp
+++ b/Engine/Source/Runtime/SlateCore/Private/Fonts/SlateFontRenderer.cpp
@@ -470,11 +470,12 @@ void RenderOutlineRows(FT_Library Library, FT_Outline* Outline, FRasterizerSpanL
 	RasterParams.gray_spans = RasterizerCallback;
 	RasterParams.user = &OutSpansList;
 
-	// Unbounding clipping box
-	RasterParams.clip_box.xMin = -32768;
-	RasterParams.clip_box.yMin = -32768;
-	RasterParams.clip_box.xMax = +32767;
-	RasterParams.clip_box.yMax = +32767;
+	// Bound clipping to the outline control bounds, truncated to integer pixels
+	FT_Outline_Get_CBox(Outline, &RasterParams.clip_box);
+	RasterParams.clip_box.xMin = RasterParams.clip_box.xMin >> 6;
+	RasterParams.clip_box.yMin = RasterParams.clip_box.yMin >> 6;
+	RasterParams.clip_box.xMax = (RasterParams.clip_box.xMax + 63) >> 6;
+	RasterParams.clip_box.yMax = (RasterParams.clip_box.yMax + 63) >> 6;
 
 	FT_Outline_Render(Library, Outline, &RasterParams);
 }
