Improved display scaling support during window creation and resize.

diff --git a/Engine/Source/Editor/UnrealEd/Private/PlayLevel.cpp b/Engine/Source/Editor/UnrealEd/Private/PlayLevel.cpp
--- a/Engine/Source/Editor/UnrealEd/Private/PlayLevel.cpp
+++ b/Engine/Source/Editor/UnrealEd/Private/PlayLevel.cpp
@@ -3367,6 +3367,7 @@ UGameInstance* UEditorEngine::CreatePIEGameInstance(int32 InPIEInstance, bool bI
 						.Title(ViewportName)
 						.ScreenPosition(FVector2D(PosX, PosY))
 						.ClientSize(FVector2D(NewWindowWidth, NewWindowHeight))
+						.AdjustInitialSizeAndPositionForDPIScale(false)
 						.AutoCenter(CenterNewWindow ? EAutoCenter::PreferredWorkArea : EAutoCenter::None)
 						.UseOSWindowBorder(bUseOSWndBorder)
 						.SaneWindowPlacement(!CenterNewWindow)
diff --git a/Engine/Source/Runtime/SlateCore/Private/Widgets/SWindow.cpp b/Engine/Source/Runtime/SlateCore/Private/Widgets/SWindow.cpp
--- a/Engine/Source/Runtime/SlateCore/Private/Widgets/SWindow.cpp
+++ b/Engine/Source/Runtime/SlateCore/Private/Widgets/SWindow.cpp
@@ -204,7 +204,7 @@ private:
 };
 
 
-FVector2D SWindow::GetWindowSizeFromClientSize(FVector2D InClientSize)
+FVector2D SWindow::GetWindowSizeFromClientSize(FVector2D InClientSize, TOptional<float> DPIScale)
 {
 	// If this is a regular non-OS window, we need to compensate for the border and title bar area that we will add
 	// Note: Windows with an OS border do this in ReshapeWindow
@@ -212,12 +212,34 @@ FVector2D SWindow::GetWindowSizeFromClientSize(FVector2D InClientSize)
 	{
 		const FMargin BorderSize = GetWindowBorderSize();
 
-		InClientSize.X += BorderSize.Left + BorderSize.Right;
-		InClientSize.Y += BorderSize.Bottom + BorderSize.Top;
+		// Get DPIScale for border and title if not already supplied
+		if (!DPIScale.IsSet())
+		{
+			if (NativeWindow.IsValid())
+			{
+				DPIScale = FPlatformApplicationMisc::GetDPIScaleFactorAtPoint(ScreenPosition.X, ScreenPosition.Y);
+			}
+			else
+			{
+				DPIScale = FPlatformApplicationMisc::GetDPIScaleFactorAtPoint(InitialDesiredScreenPosition.X, InitialDesiredScreenPosition.Y);
+			}
+		}
+
+		// Apply DPIScale if not already applied by GetWindowBorderSize()
+		if (!(NativeWindow.IsValid() && NativeWindow->IsMaximized()))
+		{
+			InClientSize.X += ((BorderSize.Left + BorderSize.Right) * DPIScale.GetValue());
+			InClientSize.Y += ((BorderSize.Bottom + BorderSize.Top) * DPIScale.GetValue());
+		}
+		else
+		{
+			InClientSize.X += BorderSize.Left + BorderSize.Right;
+			InClientSize.Y += BorderSize.Bottom + BorderSize.Top;
+		}
 
 		if (bCreateTitleBar)
 		{
-			InClientSize.Y += SWindowDefs::DefaultTitleBarSize;
+			InClientSize.Y += (SWindowDefs::DefaultTitleBarSize * DPIScale.GetValue());
 		}
 	}
 
@@ -263,12 +285,21 @@ void SWindow::Construct(const FArguments& InArgs)
 	// calculate window size from client size
 	bCreateTitleBar = InArgs._CreateTitleBar && !bIsPopupWindow && Type != EWindowType::CursorDecorator && !bHasOSWindowBorder;
 
-	// If the window has no OS border, simulate it ourselves, enlarging window by the size that OS border would have.
-	FVector2D WindowSize = GetWindowSizeFromClientSize(InArgs._ClientSize);
-
 	// calculate initial window position
 	FVector2D WindowPosition = InArgs._ScreenPosition;
 
+	const bool bAnchorWindowWindowPositionTopLeft = FPlatformApplicationMisc::AnchorWindowWindowPositionTopLeft();
+	if (bAnchorWindowWindowPositionTopLeft)
+	{
+		WindowPosition.X = WindowPosition.Y = 0;
+	}
+	else if (InArgs._AdjustInitialSizeAndPositionForDPIScale && !WindowPosition.IsZero())
+	{
+		// Will need to add additional logic to walk over multiple monitors at various DPIs to determine correct WindowPosition
+		const float InitialDPIScale = FPlatformApplicationMisc::GetDPIScaleFactorAtPoint(WindowPosition.X, WindowPosition.Y);
+		WindowPosition *= InitialDPIScale;
+	}
+
 	AutoCenterRule = InArgs._AutoCenter;
 
 	// Get desktop metrics
@@ -310,24 +341,17 @@ void SWindow::Construct(const FArguments& InArgs)
 		{
 			AutoCenterRule = EAutoCenter::PreferredWorkArea;
 		}
-
-		float PrimaryWidthPadding = DisplayMetrics.PrimaryDisplayWidth -
-			(PrimaryDisplayRect.Right - PrimaryDisplayRect.Left);
-		float PrimaryHeightPadding = DisplayMetrics.PrimaryDisplayHeight -
-			(PrimaryDisplayRect.Bottom - PrimaryDisplayRect.Top);
-
-		float VirtualWidth = (VirtualDisplayRect.Right - VirtualDisplayRect.Left);
-		float VirtualHeight = (VirtualDisplayRect.Bottom - VirtualDisplayRect.Top);
-
-		// Make sure that the window size is no larger than the virtual display area.
-		WindowSize.X = FMath::Clamp(WindowSize.X, 0.0f, VirtualWidth - PrimaryWidthPadding);
-		WindowSize.Y = FMath::Clamp(WindowSize.Y, 0.0f, VirtualHeight - PrimaryHeightPadding);
 	}
 
-	if( AutoCenterRule != EAutoCenter::None )
+	FSlateRect AutoCenterRect(0, 0, 0, 0);
+	float DPIScale = 1.0f;
+	if (bAnchorWindowWindowPositionTopLeft)
+	{
+		WindowPosition.X = WindowPosition.Y = 0;
+		DPIScale = FPlatformApplicationMisc::GetDPIScaleFactorAtPoint(WindowPosition.X, WindowPosition.Y);
+	}
+	else if (AutoCenterRule != EAutoCenter::None)
 	{
-		FSlateRect AutoCenterRect;
-
 		switch( AutoCenterRule )
 		{
 		default:
@@ -342,13 +366,38 @@ void SWindow::Construct(const FArguments& InArgs)
 			AutoCenterRect = FSlateApplicationBase::Get().GetPreferredWorkArea();
 			break;
 		}
+		DPIScale = FPlatformApplicationMisc::GetDPIScaleFactorAtPoint(AutoCenterRect.Left, AutoCenterRect.Top);
+	}
+	else
+	{
+		DPIScale = FPlatformApplicationMisc::GetDPIScaleFactorAtPoint(WindowPosition.X, WindowPosition.Y);
+	}
+
+	// If the window has no OS border, simulate it ourselves, enlarging window by the size that OS border would have.
+	const FVector2D DPIScaledClientSize = InArgs._AdjustInitialSizeAndPositionForDPIScale ? InArgs._ClientSize * DPIScale: InArgs._ClientSize;
+	FVector2D WindowSize = GetWindowSizeFromClientSize(DPIScaledClientSize, DPIScale);
 
-		float RectDPIScale = 1.0f;
-		if (InArgs._AdjustInitialSizeAndPositionForDPIScale)
+	// If we're manually positioning the window we need to check if it's outside
+	// of the virtual bounds of the current displays or too large.
+	if (AutoCenterRule == EAutoCenter::None || bAnchorWindowWindowPositionTopLeft)
+	{
+		if (InArgs._SaneWindowPlacement)
 		{
-			RectDPIScale = FPlatformApplicationMisc::GetDPIScaleFactorAtPoint(PrimaryDisplayRect.Left, PrimaryDisplayRect.Top);
-		}
+			float PrimaryWidthPadding = DisplayMetrics.PrimaryDisplayWidth -
+				(PrimaryDisplayRect.Right - PrimaryDisplayRect.Left);
+			float PrimaryHeightPadding = DisplayMetrics.PrimaryDisplayHeight -
+				(PrimaryDisplayRect.Bottom - PrimaryDisplayRect.Top);
+
+			float VirtualWidth = (VirtualDisplayRect.Right - VirtualDisplayRect.Left);
+			float VirtualHeight = (VirtualDisplayRect.Bottom - VirtualDisplayRect.Top);
 
+			// Make sure that the window size is no larger than the virtual display area.
+			WindowSize.X = FMath::Clamp(WindowSize.X, 0.0f, VirtualWidth - PrimaryWidthPadding);
+			WindowSize.Y = FMath::Clamp(WindowSize.Y, 0.0f, VirtualHeight - PrimaryHeightPadding);
+		}
+	}
+	else
+	{
 		if (InArgs._SaneWindowPlacement)
 		{
 			// Clamp window size to be no greater than the work area size
@@ -359,47 +408,21 @@ void SWindow::Construct(const FArguments& InArgs)
 		// Setup a position and size for the main frame window that's centered in the desktop work area
 		const FVector2D DisplayTopLeft( AutoCenterRect.Left, AutoCenterRect.Top );
 		const FVector2D DisplaySize( AutoCenterRect.Right - AutoCenterRect.Left, AutoCenterRect.Bottom - AutoCenterRect.Top );
-		WindowPosition = DisplayTopLeft + ( DisplaySize - (WindowSize * RectDPIScale)) * 0.5f;
+
+		WindowPosition = DisplayTopLeft + (DisplaySize - WindowSize) * 0.5f;
 
 		// Don't allow the window to center to outside of the work area
 		WindowPosition.X = FMath::Max(WindowPosition.X, AutoCenterRect.Left);
 		WindowPosition.Y = FMath::Max(WindowPosition.Y, AutoCenterRect.Top);
 	}
 
-	FVector2D DeltaSize;
-	if(InArgs._AdjustInitialSizeAndPositionForDPIScale)
-	{
-		const float DPIScale = FPlatformApplicationMisc::GetDPIScaleFactorAtPoint(WindowPosition.X, WindowPosition.Y);
-
-		// Auto centering code will have taken care of the adjustment earlier
-		if (AutoCenterRule == EAutoCenter::None)
-		{
-			WindowPosition *= DPIScale;
-		}
-
-		WindowSize *= DPIScale;
-
-		// Get change in size resulting from the above call
-		DeltaSize = WindowSize - InArgs._ClientSize*DPIScale;
-	}
-	else
-	{
-		DeltaSize = WindowSize - InArgs._ClientSize;
-	}
-
-	if (FPlatformApplicationMisc::AnchorWindowWindowPositionTopLeft())
-	{
-		WindowPosition.X = WindowPosition.Y = 0;
-	}
-
 	this->InitialDesiredScreenPosition = WindowPosition;
 	this->InitialDesiredSize = WindowSize;
 
 	FSlateApplicationBase::Get().OnGlobalInvalidationToggled().AddSP(this, &SWindow::OnGlobalInvalidationToggled);
 
-
-	// Resize adds extra borders / title bar if necessary, but this is already taken into account in WindowSize, so subtract them again first
-	Resize(WindowSize - DeltaSize);
+	// Resize without adding extra borders / title because they are already included in WindowSize
+	ResizeWindowSize(WindowSize);
 
 	this->ConstructWindowInternals();
 	this->SetContent( InArgs._Content.Widget );
@@ -931,21 +954,24 @@ void SWindow::ReshapeWindow( const FSlateRect& InNewShape )
 	ReshapeWindow(FVector2D(InNewShape.Left, InNewShape.Top), FVector2D(InNewShape.Right - InNewShape.Left, InNewShape.Bottom - InNewShape.Top));
 }
 
-void SWindow::Resize( FVector2D NewSize )
+void SWindow::Resize( FVector2D NewClientSize )
 {
-	Morpher.Sequence.JumpToEnd();
+	ResizeWindowSize(GetWindowSizeFromClientSize(NewClientSize));
+}
 
-	NewSize = GetWindowSizeFromClientSize(NewSize);
+void SWindow::ResizeWindowSize( FVector2D NewWindowSize )
+{
+	Morpher.Sequence.JumpToEnd();
 
-	NewSize.X = FMath::Max(SizeLimits.GetMinWidth().Get(NewSize.X), NewSize.X);
-	NewSize.X = FMath::Min(SizeLimits.GetMaxWidth().Get(NewSize.X), NewSize.X);
+	NewWindowSize.X = FMath::Max(SizeLimits.GetMinWidth().Get(NewWindowSize.X), NewWindowSize.X);
+	NewWindowSize.X = FMath::Min(SizeLimits.GetMaxWidth().Get(NewWindowSize.X), NewWindowSize.X);
 
-	NewSize.Y = FMath::Max(SizeLimits.GetMinHeight().Get(NewSize.Y), NewSize.Y);
-	NewSize.Y = FMath::Min(SizeLimits.GetMaxHeight().Get(NewSize.Y), NewSize.Y);
+	NewWindowSize.Y = FMath::Max(SizeLimits.GetMinHeight().Get(NewWindowSize.Y), NewWindowSize.Y);
+	NewWindowSize.Y = FMath::Min(SizeLimits.GetMaxHeight().Get(NewWindowSize.Y), NewWindowSize.Y);
 
 	// ReshapeWindow W/H takes an int, so lets move our new W/H to int before checking if they are the same size
 	FIntPoint CurrentIntSize = FIntPoint(FMath::CeilToInt(Size.X), FMath::CeilToInt(Size.Y));
-	FIntPoint NewIntSize     = FIntPoint(FMath::CeilToInt(NewSize.X), FMath::CeilToInt(NewSize.Y));
+	FIntPoint NewIntSize     = FIntPoint(FMath::CeilToInt(NewWindowSize.X), FMath::CeilToInt(NewWindowSize.Y));
 
 	if (CurrentIntSize != NewIntSize)
 	{
@@ -955,10 +981,10 @@ void SWindow::Resize( FVector2D NewSize )
 		}
 		else
 		{
-			InitialDesiredSize = NewSize;
+			InitialDesiredSize = NewWindowSize;
 		}
 	}
-	SetCachedSize(NewSize);
+	SetCachedSize(NewWindowSize);
 }
 
 FSlateRect SWindow::GetFullScreenInfo() const
diff --git a/Engine/Source/Runtime/SlateCore/Public/Widgets/SWindow.h b/Engine/Source/Runtime/SlateCore/Public/Widgets/SWindow.h
--- a/Engine/Source/Runtime/SlateCore/Public/Widgets/SWindow.h
+++ b/Engine/Source/Runtime/SlateCore/Public/Widgets/SWindow.h
@@ -409,8 +409,12 @@ public:
 	/** Relocate the window to a screenspace position specified by NewPosition and resize it to NewSize */
 	void ReshapeWindow( FVector2D NewPosition, FVector2D NewSize );
 	void ReshapeWindow( const FSlateRect& InNewShape );
-	/** Resize the window to be NewSize immediately */
-	void Resize( FVector2D NewSize );
+	/**
+	 * Resize the window to be dpi scaled NewClientSize immediately
+	 *
+	 * @param NewClientSize: Client size with DPI scaling already applied that does not include border or title bars.
+	 */
+	void Resize( FVector2D NewClientSize );
 
 	/** Returns the rectangle of the screen the window is associated with */
 	FSlateRect GetFullScreenInfo() const;
@@ -817,10 +821,18 @@ private:
 	virtual FVector2D ComputeDesiredSize(float) const override;
 	virtual bool ComputeVolatility() const override;
 
+	/** Resize using already dpi scaled window size including borders/title bar */
+	void ResizeWindowSize( FVector2D NewWindowSize );
+
 	void OnGlobalInvalidationToggled(bool bGlobalInvalidationEnabled);
 public:
-	// For a given client size, calculate the window size required to accomodate any potential non-OS borders and tilebars
-	FVector2D GetWindowSizeFromClientSize(FVector2D InClientSize);
+	/**
+	 * For a given client size, calculate the window size required to accommodate any potential non-OS borders and title bars
+	 *
+	 * @param InClientSize: Client size with DPI scaling already applied
+	 * @param DPIScale: Scale that will be applied for border and title. When not supplied detects DPIScale using native or initial position.
+	 */
+	FVector2D GetWindowSizeFromClientSize(FVector2D InClientSize, TOptional<float> DPIScale = TOptional<float>());
 
 	/** @return true if this window will be focused when it is first shown */
 	inline bool IsFocusedInitially() const
