Fix for stale safe zone updates

diff --git a/Engine/Source/Runtime/Engine/Private/UserInterface/Canvas.cpp b/Engine/Source/Runtime/Engine/Private/UserInterface/Canvas.cpp
--- a/Engine/Source/Runtime/Engine/Private/UserInterface/Canvas.cpp
+++ b/Engine/Source/Runtime/Engine/Private/UserInterface/Canvas.cpp
@@ -1388,7 +1388,7 @@ void UCanvas::UpdateSafeZoneData()
 	else if(FSlateApplication::IsInitialized())
 	{
 		FDisplayMetrics DisplayMetrics;
-		FSlateApplication::Get().GetCachedDisplayMetrics(DisplayMetrics);
+		FSlateApplication::Get().GetDisplayMetrics(DisplayMetrics);
  		CachedDisplayWidth = DisplayMetrics.PrimaryDisplayWidth;
  		CachedDisplayHeight = DisplayMetrics.PrimaryDisplayHeight;
 
diff --git a/Engine/Source/Runtime/SlateCore/Public/Application/SlateApplicationBase.h b/Engine/Source/Runtime/SlateCore/Public/Application/SlateApplicationBase.h
--- a/Engine/Source/Runtime/SlateCore/Public/Application/SlateApplicationBase.h
+++ b/Engine/Source/Runtime/SlateCore/Public/Application/SlateApplicationBase.h
@@ -558,14 +558,10 @@ protected:
 	}
 	void SwapSafeZoneTypes()
 	{
-		if (FDisplayMetrics::GetDebugTitleSafeZoneRatio() != CachedDebugTitleSafeRatio)
-		{
-			FDisplayMetrics DisplayMetrics;
-			GetDisplayMetrics(DisplayMetrics);
-			CustomSafeZoneRatio = FMargin();
-			OnDebugSafeZoneChanged.Broadcast(FMargin(), false);
-		}
-
+		FDisplayMetrics DisplayMetrics;
+		GetDisplayMetrics(DisplayMetrics);
+		CustomSafeZoneRatio = FMargin();
+		OnDebugSafeZoneChanged.Broadcast(FMargin(), false);
 	}
 #endif
 
