Exposed CanLoadCodepoint on FSlateFontCache and HasValidGlyph on FShapedGlyphEntry

These can be used to validate whether a font supports certain characters and glyphs.

diff --git a/Engine/Source/Runtime/SlateCore/Private/Fonts/FontCache.cpp b/Engine/Source/Runtime/SlateCore/Private/Fonts/FontCache.cpp
--- a/Engine/Source/Runtime/SlateCore/Private/Fonts/FontCache.cpp
+++ b/Engine/Source/Runtime/SlateCore/Private/Fonts/FontCache.cpp
@@ -74,11 +74,26 @@ ETextShapingMethod GetDefaultTextShapingMethod()
 }
 
 
+bool FShapedGlyphEntry::HasValidGlyph() const
+{
+#if WITH_FREETYPE
+	TSharedPtr<FFreeTypeFace> FontFacePin = FontFaceData ? FontFaceData->FontFace.Pin() : nullptr;
+	if (FontFacePin && GlyphIndex != 0)
+	{
+		const uint32 InvalidSubCharGlyphIndex = FT_Get_Char_Index(FontFacePin->GetFace(), SlateFontRendererUtils::InvalidSubChar);
+		return GlyphIndex != InvalidSubCharGlyphIndex;
+	}
+#endif	// WITH_FREETYPE
+
+	return false;
+}
+
 float FShapedGlyphEntry::GetBitmapRenderScale() const
 {
 	return FontFaceData ? FontFaceData->BitmapRenderScale : 1.0f;
 }
 
+
 FShapedGlyphEntryKey::FShapedGlyphEntryKey(const FShapedGlyphFaceData& InFontFaceData, uint32 InGlyphIndex, const FFontOutlineSettings& InOutlineSettings)
 	: FontFace(InFontFaceData.FontFace)
 	, FontSize(InFontFaceData.FontSize)
@@ -1088,6 +1103,11 @@ bool FSlateFontCache::HasKerning( const FFontData& InFontData ) const
 	return FontRenderer->HasKerning(InFontData);
 }
 
+bool FSlateFontCache::CanLoadCodepoint(const FFontData& InFontData, const UTF32CHAR InCodepoint, EFontFallback MaxFallbackLevel) const
+{
+	return FontRenderer->CanLoadCodepoint(InFontData, InCodepoint, MaxFallbackLevel);
+}
+
 const TSet<FName>& FSlateFontCache::GetFontAttributes( const FFontData& InFontData ) const
 {
 	return CompositeFontCache->GetFontAttributes(InFontData);
diff --git a/Engine/Source/Runtime/SlateCore/Public/Fonts/FontCache.h b/Engine/Source/Runtime/SlateCore/Public/Fonts/FontCache.h
--- a/Engine/Source/Runtime/SlateCore/Public/Fonts/FontCache.h
+++ b/Engine/Source/Runtime/SlateCore/Public/Fonts/FontCache.h
@@ -132,6 +132,9 @@ struct FShapedGlyphEntry
 	 * False if the glyph is invisible (eg, whitespace or a control code) and should skip drawing, but still include its advance amount.
 	 */
 	bool bIsVisible = false;
+	
+	/** Check whether this entry contains a valid glyph (non-zero, and not the SlateFontRendererUtils::InvalidSubChar glyph) */
+	SLATECORE_API bool HasValidGlyph() const;
 
 	/** Get any additional scale that should be applied when rendering this glyph */
 	SLATECORE_API float GetBitmapRenderScale() const;
@@ -793,6 +796,16 @@ public:
 	 */
 	bool HasKerning( const FFontData& InFontData ) const;
 
+	/**
+	 * Whether or not the specified character, within the specified font, can be loaded with the specified maximum font fallback level
+	 *
+	 * @param InFontData		Information about the font to load
+	 * @param InCodepoint		The codepoint being loaded
+	 * @param MaxFallbackLevel	The maximum fallback level to try for the font
+	 * @return					Whether or not the character can be loaded
+	 */
+	bool CanLoadCodepoint(const FFontData& InFontData, const UTF32CHAR InCodepoint, EFontFallback MaxFallbackLevel = EFontFallback::FF_NoFallback) const;
+
 	/**
 	 * Returns the font attributes for the specified font.
 	 *
