FReply now stores the target widgets for various requests as WeakPtrs.
- Resolves issues where the FReply living on ULocalPlayer could cause SObjectWidgets and their associated UUserWidgets to have a delayed destruction process.
- Also just generally stands to reason that if the FReply reference is the only thing keeping a widget alive, there's no chance that SlateApplication will actually be able to do anything meaningful with that widget anyway (it can't be focused, captured, dragged, etc.)
- Resolved a lurking potential crash in FSlateApplication that was made more possible by this change

diff --git a/Engine/Source/Runtime/Slate/Private/Framework/Application/SlateApplication.cpp b/Engine/Source/Runtime/Slate/Private/Framework/Application/SlateApplication.cpp
--- a/Engine/Source/Runtime/Slate/Private/Framework/Application/SlateApplication.cpp
+++ b/Engine/Source/Runtime/Slate/Private/Framework/Application/SlateApplication.cpp
@@ -3154,7 +3154,7 @@ void FSlateApplication::ProcessReply( const FWidgetPath& CurrentEventPath, const
 
 	// Set focus if requested.
 	TSharedPtr<SWidget> RequestedFocusRecepient = TheReply.GetUserFocusRecepient();
-	if (TheReply.ShouldSetUserFocus() || RequestedFocusRecepient.IsValid())
+	if (TheReply.ShouldSetUserFocus() && RequestedFocusRecepient.IsValid())
 	{
 		if (TheReply.AffectsAllUsers())
 		{
diff --git a/Engine/Source/Runtime/SlateCore/Public/Input/Reply.h b/Engine/Source/Runtime/SlateCore/Public/Input/Reply.h
--- a/Engine/Source/Runtime/SlateCore/Public/Input/Reply.h
+++ b/Engine/Source/Runtime/SlateCore/Public/Input/Reply.h
@@ -177,16 +177,16 @@ public:
 	bool ShouldThrottle() const { return !bPreventThrottling; }
 
 	/** Returns the widget that the mouse should be locked to (if any) */
-	const TSharedPtr<SWidget>& GetMouseLockWidget() const { return MouseLockWidget; }
+	TSharedPtr<SWidget> GetMouseLockWidget() const { return MouseLockWidget.Pin(); }
 
 	/** When not nullptr, user focus has been requested to be set on the FocusRecipient. */
-	const TSharedPtr<SWidget>& GetUserFocusRecepient() const { return FocusRecipient; }
+	TSharedPtr<SWidget> GetUserFocusRecepient() const { return FocusRecipient.Pin(); }
 
 	/** Get the reason that a focus change is being requested. */
 	EFocusCause GetFocusCause() const { return FocusChangeReason; }
 
 	/** If the event replied with a request to capture the mouse, this returns the desired mouse captor. Otherwise returns an invalid pointer. */
-	const TSharedPtr<SWidget>& GetMouseCaptor() const { return MouseCaptor; }	
+	TSharedPtr<SWidget> GetMouseCaptor() const { return MouseCaptor.Pin(); }
 
 	/** Get the navigation type (Invalid if no navigation is requested). */
 	EUINavigation GetNavigationType() const { return NavigationType; }
@@ -198,7 +198,7 @@ public:
 	ENavigationSource GetNavigationSource() const { return NavigationSource; }
 
 	/** Get the widget that is the navigation destination. */
-	const TSharedPtr<SWidget>& GetNavigationDestination() const { return NavigationDestination; }
+	TSharedPtr<SWidget> GetNavigationDestination() const { return NavigationDestination.Pin(); }
 
 	/** @return the Content that we should use for the Drag and Drop operation; Invalid SharedPtr if a drag and drop operation is not requested*/
 	const TSharedPtr<FDragDropOperation>& GetDragDropContent() const { return DragDropContent; }
@@ -207,7 +207,7 @@ public:
 	bool ShouldEndDragDrop() const { return bEndDragDrop; }
 
 	/** @return a widget for which to detect a drag; Invalid SharedPtr if no drag detection requested */
-	const TSharedPtr<SWidget>& GetDetectDragRequest() const { return DetectDragForWidget; }  
+	TSharedPtr<SWidget> GetDetectDragRequest() const { return DetectDragForWidget.Pin(); }
 
 	/** @return the mouse button for which we are detecting a drag */
 	FKey GetDetectDragRequestButton() const { return DetectDragForMouseButton; }
@@ -272,12 +272,12 @@ private:
 private:
 
 	TOptional<FIntPoint> RequestedMousePos;
-	TSharedPtr<SWidget> EventHandler;
-	TSharedPtr<SWidget> MouseCaptor;
-	TSharedPtr<SWidget> FocusRecipient;
-	TSharedPtr<SWidget> MouseLockWidget;
-	TSharedPtr<SWidget> DetectDragForWidget;
-	TSharedPtr<SWidget> NavigationDestination;
+	TWeakPtr<SWidget> EventHandler;
+	TWeakPtr<SWidget> MouseCaptor;
+	TWeakPtr<SWidget> FocusRecipient;
+	TWeakPtr<SWidget> MouseLockWidget;
+	TWeakPtr<SWidget> DetectDragForWidget;
+	TWeakPtr<SWidget> NavigationDestination;
 	FKey DetectDragForMouseButton;
 	TSharedPtr<FDragDropOperation> DragDropContent;
 	EFocusCause FocusChangeReason;
