Adding a count to detect long WidgetPath that could lead to an out of memory exception.

#rnx

diff --git a/Engine/Source/Runtime/SlateCore/Private/Layout/WidgetPath.cpp b/Engine/Source/Runtime/SlateCore/Private/Layout/WidgetPath.cpp
--- a/Engine/Source/Runtime/SlateCore/Private/Layout/WidgetPath.cpp
+++ b/Engine/Source/Runtime/SlateCore/Private/Layout/WidgetPath.cpp
@@ -263,17 +263,29 @@ FWeakWidgetPath::EPathResolutionResult::Result FWeakWidgetPath::ToWidgetPath( FW
 {
 	SCOPE_CYCLE_COUNTER(STAT_WeakToStrong_WidgetPath);
 
-	if (GSlateFastWidgetPath)
-	{
-		TArray<FWidgetAndPointer> PathWithGeometries;
-		TArray< TSharedPtr<SWidget> > WidgetPtrs;
+	TArray< TSharedPtr<SWidget> > WidgetPtrs;
 
-		// Convert the weak pointers into shared pointers because we are about to do something with this path instead of just observe it.
-		TSharedPtr<SWindow> TopLevelWindowPtr = Window.Pin();
-		for (TArray< TWeakPtr<SWidget> >::TConstIterator SomeWeakWidgetPtr(Widgets); SomeWeakWidgetPtr; ++SomeWeakWidgetPtr)
+	// Convert the weak pointers into shared pointers because we are about to do something with this path instead of just observe it.
+	TSharedPtr<SWindow> TopLevelWindowPtr = Window.Pin();
+	int PathSize = 0;
+	for (TArray< TWeakPtr<SWidget> >::TConstIterator SomeWeakWidgetPtr(Widgets); SomeWeakWidgetPtr; ++SomeWeakWidgetPtr)
+	{
+		const int MaxWidgetPath = 1000;
+		++PathSize;
+		if (ensureMsgf(PathSize < MaxWidgetPath, TEXT("Converting a Widget Path of more that 1000 Widget deep.")))
 		{
 			WidgetPtrs.Add(SomeWeakWidgetPtr->Pin());
 		}
+		else
+		{
+			WidgetPath = FWidgetPath();
+			return EPathResolutionResult::Truncated;
+		}
+	}
+
+	if (GSlateFastWidgetPath)
+	{
+		TArray<FWidgetAndPointer> PathWithGeometries;
 
 		// The path can get interrupted if some subtree of widgets disappeared, but we still maintain weak references to it.
 		bool bPathUninterrupted = false;
@@ -329,15 +341,7 @@ FWeakWidgetPath::EPathResolutionResult::Result FWeakWidgetPath::ToWidgetPath( FW
 	}
 	else
 	{
-		TArray<FWidgetAndPointer> PathWithGeometries;
-		TArray< TSharedPtr<SWidget> > WidgetPtrs;
-
-		// Convert the weak pointers into shared pointers because we are about to do something with this path instead of just observe it.
-		TSharedPtr<SWindow> TopLevelWindowPtr = Window.Pin();
-		for (TArray< TWeakPtr<SWidget> >::TConstIterator SomeWeakWidgetPtr(Widgets); SomeWeakWidgetPtr; ++SomeWeakWidgetPtr)
-		{
-			WidgetPtrs.Add(SomeWeakWidgetPtr->Pin());
-		}
+		TArray<FWidgetAndPointer> PathWithGeometries;		
 
 		// The path can get interrupted if some subtree of widgets disappeared, but we still maintain weak references to it.
 		bool bPathUninterrupted = false;
