Give FHitttestGrid::GetBubblePath  default value for UserIndex to avoid backwards compatibility issues

diff --git a/Engine/Source/Runtime/SlateCore/Public/Input/HittestGrid.h b/Engine/Source/Runtime/SlateCore/Public/Input/HittestGrid.h
--- a/Engine/Source/Runtime/SlateCore/Public/Input/HittestGrid.h
+++ b/Engine/Source/Runtime/SlateCore/Public/Input/HittestGrid.h
@@ -34,7 +34,7 @@ public:
 	 * Given a Slate Units coordinate in virtual desktop space, perform a hittest
 	 * and return the path along which the corresponding event would be bubbled.
 	 */
-	TArray<FWidgetAndPointer> GetBubblePath(FVector2D DesktopSpaceCoordinate, float CursorRadius, bool bIgnoreEnabledStatus, int32 UserIndex);
+	TArray<FWidgetAndPointer> GetBubblePath(FVector2D DesktopSpaceCoordinate, float CursorRadius, bool bIgnoreEnabledStatus, int32 UserIndex = INDEX_NONE);
 
 	/**
 	 * Set the position and size of the hittest area in desktop coordinates
