Dont fullly destroy widgets cached state when visibility changes

diff --git a/Engine/Source/Runtime/SlateCore/Private/Widgets/SWidget.cpp b/Engine/Source/Runtime/SlateCore/Private/Widgets/SWidget.cpp
--- a/Engine/Source/Runtime/SlateCore/Private/Widgets/SWidget.cpp
+++ b/Engine/Source/Runtime/SlateCore/Private/Widgets/SWidget.cpp
@@ -799,7 +799,7 @@ void SWidget::UpdateFastPathVisibility(bool bParentVisible, bool bWidgetRemoved,
 
 	if (PersistentState.CachedElementListNode)
 	{
-		PersistentState.CachedElementListNode->GetValue().GetOwningData()->RemoveCache(PersistentState.CachedElementListNode);
+		PersistentState.CachedElementListNode->GetValue().Reset();
 	}
 
 	FChildren* MyChildren = GetAllChildren();
