Disabled legacy localized fonts by default for 4.23

These can be enabled again by setting Slate.EnableLegacyLocalizedFallbackFont to 1, however legacy fallback fonts will be removed for 4.24 so anything relying on this behavior should be converted to use composite fonts with sub-font families instead.

diff --git a/Engine/Source/Runtime/SlateCore/Private/Fonts/LegacySlateFontInfoCache.cpp b/Engine/Source/Runtime/SlateCore/Private/Fonts/LegacySlateFontInfoCache.cpp
--- a/Engine/Source/Runtime/SlateCore/Private/Fonts/LegacySlateFontInfoCache.cpp
+++ b/Engine/Source/Runtime/SlateCore/Private/Fonts/LegacySlateFontInfoCache.cpp
@@ -9,7 +9,7 @@
 #include "Fonts/FontProviderInterface.h"
 #include "Fonts/UnicodeBlockRange.h"
 
-static int32 GSlateEnableLegacyLocalizedFallbackFont = 1;
+static int32 GSlateEnableLegacyLocalizedFallbackFont = 0;
 static FAutoConsoleVariableRef CVarSlateEnableLocalizedFallbackFont(TEXT("Slate.EnableLegacyLocalizedFallbackFont"), GSlateEnableLegacyLocalizedFallbackFont, TEXT("Enable the legacy localized fallback fonts? (0/1)."), ECVF_Default);
 
 FString FLegacySlateFontInfoCache::FFallbackContext::ToString() const
