Unlink UMG default colors from editor settings colors (CoreStyle).

diff --git a/Engine/Source/Runtime/SlateCore/Public/Styling/SlateBrush.h b/Engine/Source/Runtime/SlateCore/Public/Styling/SlateBrush.h
--- a/Engine/Source/Runtime/SlateCore/Public/Styling/SlateBrush.h
+++ b/Engine/Source/Runtime/SlateCore/Public/Styling/SlateBrush.h
@@ -209,6 +209,15 @@ public:
 		return TintColor.GetColor(InWidgetStyle);
 	}
 
+	/**
+	 * Unlinks all colors in this brush.
+	 * @see FSlateColor::Unlink
+	 */
+	void UnlinkColors()
+	{
+		TintColor.Unlink();
+	}
+
 	/**
 	 * Checks whether this brush has a UTexture object
 	 *
diff --git a/Engine/Source/Runtime/SlateCore/Public/Styling/SlateColor.h b/Engine/Source/Runtime/SlateCore/Public/Styling/SlateColor.h
--- a/Engine/Source/Runtime/SlateCore/Public/Styling/SlateColor.h
+++ b/Engine/Source/Runtime/SlateCore/Public/Styling/SlateColor.h
@@ -136,6 +136,24 @@ public:
 		return (ColorUseRule == ESlateColorStylingMode::UseColor_Specified) || (ColorUseRule == ESlateColorStylingMode::UseColor_Specified_Link);
 	}
 
+	/**
+	 * If the color rule is set to UseColor_Specified_Link, this will copy the linked color internally,
+	 * unlink it and set the color rule to UseColor_Specified.
+	 * Nothing happens if the color rule is not set to UseColor_Specified_Link.
+	 */
+	void Unlink()
+	{
+		if (ColorUseRule == ESlateColorStylingMode::UseColor_Specified_Link)
+		{
+			ColorUseRule = ESlateColorStylingMode::UseColor_Specified;
+			if (ensure(LinkedSpecifiedColor.IsValid()))
+			{
+				SpecifiedColor = *LinkedSpecifiedColor;
+				LinkedSpecifiedColor.Reset();
+			}
+		}
+	}
+
 	/**
 	 * Compares this color with another for equality.
 	 *
diff --git a/Engine/Source/Runtime/SlateCore/Public/Styling/SlateTypes.h b/Engine/Source/Runtime/SlateCore/Public/Styling/SlateTypes.h
--- a/Engine/Source/Runtime/SlateCore/Public/Styling/SlateTypes.h
+++ b/Engine/Source/Runtime/SlateCore/Public/Styling/SlateTypes.h
@@ -185,6 +185,25 @@ struct SLATECORE_API FCheckBoxStyle : public FSlateWidgetStyle
 	 */	
 	void PostSerialize(const FArchive& Ar);
 #endif
+
+	/**
+	 * Unlinks all colors in this style.
+	 * @see FSlateColor::Unlink
+	 */
+	void UnlinkColors()
+	{
+		UncheckedImage.UnlinkColors();
+		UncheckedHoveredImage.UnlinkColors();
+		UncheckedPressedImage.UnlinkColors();
+		CheckedImage.UnlinkColors();
+		CheckedHoveredImage.UnlinkColors();
+		CheckedPressedImage.UnlinkColors();
+		UndeterminedImage.UnlinkColors();
+		UndeterminedHoveredImage.UnlinkColors();
+		UndeterminedPressedImage.UnlinkColors();
+		ForegroundColor.Unlink();
+		BorderBackgroundColor.Unlink();
+	}
 };
 
 #if WITH_EDITORONLY_DATA
@@ -273,6 +292,19 @@ struct SLATECORE_API FTextBlockStyle : public FSlateWidgetStyle
 	UPROPERTY(EditAnywhere, BlueprintReadWrite, Category=Appearance, AdvancedDisplay)
 	FSlateBrush UnderlineBrush;
 	FTextBlockStyle& SetUnderlineBrush( const FSlateBrush& InUnderlineBrush ){ UnderlineBrush = InUnderlineBrush; return *this; }
+
+	/**
+	 * Unlinks all colors in this style.
+	 * @see FSlateColor::Unlink
+	 */
+	void UnlinkColors()
+	{
+		ColorAndOpacity.Unlink();
+		SelectedBackgroundColor.Unlink();
+		HighlightShape.UnlinkColors();
+		StrikeBrush.UnlinkColors();
+		UnderlineBrush.UnlinkColors();
+	}
 };
 
 /**
@@ -356,6 +388,18 @@ struct SLATECORE_API FButtonStyle : public FSlateWidgetStyle
 	 */	
 	void PostSerialize(const FArchive& Ar);
 #endif
+
+	/**
+	 * Unlinks all colors in this style.
+	 * @see FSlateColor::Unlink
+	 */
+	void UnlinkColors()
+	{
+		Normal.UnlinkColors();
+		Hovered.UnlinkColors();
+		Pressed.UnlinkColors();
+		Disabled.UnlinkColors();
+	}
 };
 
 template<>
@@ -416,6 +460,17 @@ struct SLATECORE_API FComboButtonStyle : public FSlateWidgetStyle
 	UPROPERTY(EditAnywhere, BlueprintReadWrite, Category=Appearance)
 	FMargin MenuBorderPadding;
 	FComboButtonStyle& SetMenuBorderPadding( const FMargin& InMenuBorderPadding ){ MenuBorderPadding = InMenuBorderPadding; return *this; }
+
+	/**
+	* Unlinks all colors in this style.
+	* @see FSlateColor::Unlink
+	 */
+	void UnlinkColors()
+	{
+		ButtonStyle.UnlinkColors();
+		DownArrowImage.UnlinkColors();
+		MenuBorderBrush.UnlinkColors();
+	}
 };
 
 
@@ -470,6 +525,16 @@ struct SLATECORE_API FComboBoxStyle : public FSlateWidgetStyle
 	 */	
 	void PostSerialize(const FArchive& Ar);
 #endif
+
+	/**
+	 * Unlinks all colors in this style.
+	 * @see FSlateColor::Unlink
+	 */
+	void UnlinkColors()
+	{
+		ComboButtonStyle.UnlinkColors();
+	}
+
 };
 
 #if WITH_EDITORONLY_DATA
@@ -562,6 +627,18 @@ struct SLATECORE_API FEditableTextStyle : public FSlateWidgetStyle
 	UPROPERTY(EditAnywhere, BlueprintReadWrite, Category=Appearance)
 	FSlateBrush CaretImage;
 	FEditableTextStyle& SetCaretImage( const FSlateBrush& InCaretImage ){ CaretImage = InCaretImage; return *this; }
+
+	/**
+	 * Unlinks all colors in this style.
+	 * @see FSlateColor::Unlink
+	 */
+	void UnlinkColors()
+	{
+		ColorAndOpacity.Unlink();
+		BackgroundImageSelected.UnlinkColors();
+		BackgroundImageComposing.UnlinkColors();
+		CaretImage.UnlinkColors();
+	}
 };
 
 
@@ -628,6 +705,23 @@ struct SLATECORE_API FScrollBarStyle : public FSlateWidgetStyle
 	UPROPERTY(EditAnywhere, BlueprintReadWrite, Category=Appearance)
 	FSlateBrush DraggedThumbImage;
 	FScrollBarStyle& SetDraggedThumbImage( const FSlateBrush& InDraggedThumbImage ){ DraggedThumbImage = InDraggedThumbImage; return *this; }
+
+	/**
+	 * Unlinks all colors in this style.
+	 * @see FSlateColor::Unlink
+	 */
+	void UnlinkColors()
+	{
+		HorizontalBackgroundImage.UnlinkColors();
+		VerticalBackgroundImage.UnlinkColors();
+		VerticalTopSlotImage.UnlinkColors();
+		HorizontalTopSlotImage.UnlinkColors();
+		VerticalBottomSlotImage.UnlinkColors();
+		HorizontalBottomSlotImage.UnlinkColors();
+		NormalThumbImage.UnlinkColors();
+		HoveredThumbImage.UnlinkColors();
+		DraggedThumbImage.UnlinkColors();
+	}
 };
 
 
@@ -710,6 +804,22 @@ struct SLATECORE_API FEditableTextBoxStyle : public FSlateWidgetStyle
 	UPROPERTY(EditAnywhere, BlueprintReadWrite, Category=Appearance)
 	FScrollBarStyle ScrollBarStyle;
 	FEditableTextBoxStyle& SetScrollBarStyle( const FScrollBarStyle& InScrollBarStyle ){ ScrollBarStyle = InScrollBarStyle; return *this; }
+
+	/**
+	 * Unlinks all colors in this style.
+	 * @see FSlateColor::Unlink
+	 */
+	void UnlinkColors()
+	{
+		BackgroundImageNormal.UnlinkColors();
+		BackgroundImageHovered.UnlinkColors();
+		BackgroundImageFocused.UnlinkColors();
+		BackgroundImageReadOnly.UnlinkColors();
+		ForegroundColor.Unlink();
+		BackgroundColor.Unlink();
+		ReadOnlyForegroundColor.Unlink();
+		ScrollBarStyle.UnlinkColors();
+	}
 };
 
 
@@ -777,6 +887,17 @@ struct SLATECORE_API FProgressBarStyle : public FSlateWidgetStyle
 	UPROPERTY(EditAnywhere, BlueprintReadWrite, Category=Appearance)
 	FSlateBrush MarqueeImage;
 	FProgressBarStyle& SetMarqueeImage( const FSlateBrush& InMarqueeImage ){ MarqueeImage = InMarqueeImage; return *this; }
+
+	/**
+	* Unlinks all colors in this style.
+	* @see FSlateColor::Unlink
+	*/
+	void UnlinkColors()
+	{
+		BackgroundImage.UnlinkColors();
+		FillImage.UnlinkColors();
+		MarqueeImage.UnlinkColors();
+	}
 };
 
 
@@ -813,6 +934,16 @@ struct SLATECORE_API FExpandableAreaStyle : public FSlateWidgetStyle
 	UPROPERTY(EditAnywhere, Category = Appearance)
 	float RolloutAnimationSeconds;
 	FExpandableAreaStyle& SetRolloutAnimationSeconds(float InLengthSeconds) { RolloutAnimationSeconds = InLengthSeconds; return *this; }
+
+	/**
+	 * Unlinks all colors in this style.
+	 * @see FSlateColor::Unlink
+	 */
+	void UnlinkColors()
+	{
+		CollapsedImage.UnlinkColors();
+		ExpandedImage.UnlinkColors();
+	}
 };
 
 
@@ -929,6 +1060,20 @@ struct SLATECORE_API FSliderStyle : public FSlateWidgetStyle
 	UPROPERTY(EditAnywhere, BlueprintReadWrite, Category=Appearance)
 	float BarThickness;
 	FSliderStyle& SetBarThickness(float InBarThickness) { BarThickness = InBarThickness; return *this; }
+
+	/**
+	 * Unlinks all colors in this style.
+	 * @see FSlateColor::Unlink
+	 */
+	void UnlinkColors()
+	{
+		NormalBarImage.UnlinkColors();
+		HoveredBarImage.UnlinkColors();
+		DisabledBarImage.UnlinkColors();
+		NormalThumbImage.UnlinkColors();
+		HoveredThumbImage.UnlinkColors();
+		DisabledThumbImage.UnlinkColors();
+	}
 };
 
 
@@ -1065,6 +1210,20 @@ struct SLATECORE_API FSpinBoxStyle : public FSlateWidgetStyle
 	UPROPERTY(EditAnywhere, BlueprintReadWrite, Category=Appearance)
 	FMargin TextPadding;
 	FSpinBoxStyle& SetTextPadding( const FMargin& InTextPadding ){ TextPadding = InTextPadding; return *this; }
+
+	/**
+	 * Unlinks all colors in this style.
+	 * @see FSlateColor::Unlink
+	 */
+	void UnlinkColors()
+	{
+		BackgroundBrush.UnlinkColors();
+		HoveredBackgroundBrush.UnlinkColors();
+		ActiveFillBrush.UnlinkColors();
+		InactiveFillBrush.UnlinkColors();
+		ArrowsImage.UnlinkColors();
+		ForegroundColor.Unlink();
+	}
 };
 
 
@@ -1187,6 +1346,28 @@ struct SLATECORE_API FTableRowStyle : public FSlateWidgetStyle
 	UPROPERTY(EditAnywhere, BlueprintReadWrite, Category = Appearance)
 	FSlateBrush DropIndicator_Below;
 	FTableRowStyle& SetDropIndicator_Below(const FSlateBrush& InValue){ DropIndicator_Below = InValue; return *this; }
+
+	/**
+	 * Unlinks all colors in this style.
+	 * @see FSlateColor::Unlink
+	 */
+	void UnlinkColors()
+	{
+		SelectorFocusedBrush.UnlinkColors();
+		ActiveHoveredBrush.UnlinkColors();
+		ActiveBrush.UnlinkColors();
+		InactiveHoveredBrush.UnlinkColors();
+		InactiveBrush.UnlinkColors();
+		EvenRowBackgroundHoveredBrush.UnlinkColors();
+		EvenRowBackgroundBrush.UnlinkColors();
+		OddRowBackgroundHoveredBrush.UnlinkColors();
+		OddRowBackgroundBrush.UnlinkColors();
+		TextColor.Unlink();
+		SelectedTextColor.Unlink();
+		DropIndicator_Above.UnlinkColors();
+		DropIndicator_Onto.UnlinkColors();
+		DropIndicator_Below.UnlinkColors();
+	}
 };
 
 
@@ -1430,6 +1611,18 @@ struct SLATECORE_API FScrollBoxStyle : public FSlateWidgetStyle
 		RightShadowBrush = InRightShadowBrush;
 		return *this;
 	}
+
+	/**
+	 * Unlinks all colors in this style.
+	 * @see FSlateColor::Unlink
+	 */
+	void UnlinkColors()
+	{
+		TopShadowBrush.UnlinkColors();
+		BottomShadowBrush.UnlinkColors();
+		LeftShadowBrush.UnlinkColors();
+		RightShadowBrush.UnlinkColors();
+	}
 };
 
 
diff --git a/Engine/Source/Runtime/UMG/Private/Components/Button.cpp b/Engine/Source/Runtime/UMG/Private/Components/Button.cpp
--- a/Engine/Source/Runtime/UMG/Private/Components/Button.cpp
+++ b/Engine/Source/Runtime/UMG/Private/Components/Button.cpp
@@ -11,12 +11,21 @@
 /////////////////////////////////////////////////////
 // UButton
 
+static FButtonStyle* DefaultButtonStyle = nullptr;
+
 UButton::UButton(const FObjectInitializer& ObjectInitializer)
 	: Super(ObjectInitializer)
 {
-	// HACK: THIS SHOULD NOT COME FROM CORESTYLE AND SHOULD INSTEAD BY DEFINED BY ENGINE TEXTURES/PROJECT SETTINGS
-	static const FButtonStyle StaticButtonStyle = FCoreStyle::Get().GetWidgetStyle< FButtonStyle >("Button");
-	WidgetStyle = StaticButtonStyle;
+	if (DefaultButtonStyle == nullptr)
+	{
+		// HACK: THIS SHOULD NOT COME FROM CORESTYLE AND SHOULD INSTEAD BE DEFINED BY ENGINE TEXTURES/PROJECT SETTINGS
+		DefaultButtonStyle = new FButtonStyle(FCoreStyle::Get().GetWidgetStyle<FButtonStyle>("Button"));
+
+		// Unlink UMG default colors from the editor settings colors.
+		DefaultButtonStyle->UnlinkColors();
+	}
+
+	WidgetStyle = *DefaultButtonStyle;
 
 	ColorAndOpacity = FLinearColor::White;
 	BackgroundColor = FLinearColor::White;
diff --git a/Engine/Source/Runtime/UMG/Private/Components/CheckBox.cpp b/Engine/Source/Runtime/UMG/Private/Components/CheckBox.cpp
--- a/Engine/Source/Runtime/UMG/Private/Components/CheckBox.cpp
+++ b/Engine/Source/Runtime/UMG/Private/Components/CheckBox.cpp
@@ -11,12 +11,21 @@
 /////////////////////////////////////////////////////
 // UCheckBox
 
+static FCheckBoxStyle* DefaultCheckboxStyle = nullptr;
+
 UCheckBox::UCheckBox(const FObjectInitializer& ObjectInitializer)
 	: Super(ObjectInitializer)
 {
-	// HACK: THIS SHOULD NOT COME FROM CORESTYLE AND SHOULD INSTEAD BY DEFINED BY ENGINE TEXTURES/PROJECT SETTINGS
-	static const FCheckBoxStyle StaticCheckboxStyle = FCoreStyle::Get().GetWidgetStyle< FCheckBoxStyle >("Checkbox");
-	WidgetStyle = StaticCheckboxStyle;
+	if (DefaultCheckboxStyle == nullptr)
+	{
+		// HACK: THIS SHOULD NOT COME FROM CORESTYLE AND SHOULD INSTEAD BE DEFINED BY ENGINE TEXTURES/PROJECT SETTINGS
+		DefaultCheckboxStyle = new FCheckBoxStyle(FCoreStyle::Get().GetWidgetStyle<FCheckBoxStyle>("Checkbox"));
+
+		// Unlink UMG default colors from the editor settings colors.
+		DefaultCheckboxStyle->UnlinkColors();
+	}
+
+	WidgetStyle = *DefaultCheckboxStyle;
 
 	CheckedState = ECheckBoxState::Unchecked;
 
diff --git a/Engine/Source/Runtime/UMG/Private/Components/CircularThrobber.cpp b/Engine/Source/Runtime/UMG/Private/Components/CircularThrobber.cpp
--- a/Engine/Source/Runtime/UMG/Private/Components/CircularThrobber.cpp
+++ b/Engine/Source/Runtime/UMG/Private/Components/CircularThrobber.cpp
@@ -11,13 +11,22 @@
 /////////////////////////////////////////////////////
 // UCircularThrobber
 
+static FSlateBrush* DefaultCircularThrobberBrushStyle = nullptr;
+
 UCircularThrobber::UCircularThrobber(const FObjectInitializer& ObjectInitializer)
 	: Super(ObjectInitializer)
 	, bEnableRadius(true)
 {
-	// HACK: THIS SHOULD NOT COME FROM CORESTYLE AND SHOULD INSTEAD BY DEFINED BY ENGINE TEXTURES/PROJECT SETTINGS
-	static const FSlateBrush StaticBrushStyle = *FCoreStyle::Get().GetBrush("Throbber.CircleChunk");
-	Image = StaticBrushStyle;
+	if (DefaultCircularThrobberBrushStyle == nullptr)
+	{
+		// HACK: THIS SHOULD NOT COME FROM CORESTYLE AND SHOULD INSTEAD BE DEFINED BY ENGINE TEXTURES/PROJECT SETTINGS
+		DefaultCircularThrobberBrushStyle = new FSlateBrush(*FCoreStyle::Get().GetBrush("Throbber.CircleChunk"));
+
+		// Unlink UMG default colors from the editor settings colors.
+		DefaultCircularThrobberBrushStyle->UnlinkColors();
+	}
+
+	Image = *DefaultCircularThrobberBrushStyle;
 
 	NumberOfPieces = 6;
 	Period = 0.75f;
diff --git a/Engine/Source/Runtime/UMG/Private/Components/ComboBoxString.cpp b/Engine/Source/Runtime/UMG/Private/Components/ComboBoxString.cpp
--- a/Engine/Source/Runtime/UMG/Private/Components/ComboBoxString.cpp
+++ b/Engine/Source/Runtime/UMG/Private/Components/ComboBoxString.cpp
@@ -10,15 +10,32 @@
 /////////////////////////////////////////////////////
 // UComboBoxString
 
+static FComboBoxStyle* DefaultComboBoxStyle = nullptr;
+static FTableRowStyle* DefaultComboBoxRowStyle = nullptr;
+
 UComboBoxString::UComboBoxString(const FObjectInitializer& ObjectInitializer)
 	: Super(ObjectInitializer)
 {
-	// HACK: THIS SHOULD NOT COME FROM CORESTYLE AND SHOULD INSTEAD BY DEFINED BY ENGINE TEXTURES/PROJECT SETTINGS
-	static const FComboBoxStyle StaticComboboxStyle = FCoreStyle::Get().GetWidgetStyle< FComboBoxStyle >("ComboBox");
-	static const FTableRowStyle StaticRowStyle = FCoreStyle::Get().GetWidgetStyle< FTableRowStyle >("TableView.Row");
+	if (DefaultComboBoxStyle == nullptr)
+	{
+		// HACK: THIS SHOULD NOT COME FROM CORESTYLE AND SHOULD INSTEAD BE DEFINED BY ENGINE TEXTURES/PROJECT SETTINGS
+		DefaultComboBoxStyle = new FComboBoxStyle(FCoreStyle::Get().GetWidgetStyle<FComboBoxStyle>("ComboBox"));
+
+		// Unlink UMG default colors from the editor settings colors.
+		DefaultComboBoxStyle->UnlinkColors();
+	}
+
+	if (DefaultComboBoxRowStyle == nullptr)
+	{
+		// HACK: THIS SHOULD NOT COME FROM CORESTYLE AND SHOULD INSTEAD BE DEFINED BY ENGINE TEXTURES/PROJECT SETTINGS
+		DefaultComboBoxRowStyle = new FTableRowStyle(FCoreStyle::Get().GetWidgetStyle<FTableRowStyle>("TableView.Row"));
+
+		// Unlink UMG default colors from the editor settings colors.
+		DefaultComboBoxRowStyle->UnlinkColors();
+	}
 
-	WidgetStyle = StaticComboboxStyle;
-	ItemStyle = StaticRowStyle;
+	WidgetStyle = *DefaultComboBoxStyle;
+	ItemStyle = *DefaultComboBoxRowStyle;
 	ItemStyle.SelectorFocusedBrush.TintColor = ItemStyle.SelectorFocusedBrush.TintColor.GetSpecifiedColor();
 	ItemStyle.ActiveHoveredBrush.TintColor = ItemStyle.ActiveHoveredBrush.TintColor.GetSpecifiedColor();
 	ItemStyle.ActiveBrush.TintColor = ItemStyle.ActiveBrush.TintColor.GetSpecifiedColor();
diff --git a/Engine/Source/Runtime/UMG/Private/Components/EditableText.cpp b/Engine/Source/Runtime/UMG/Private/Components/EditableText.cpp
--- a/Engine/Source/Runtime/UMG/Private/Components/EditableText.cpp
+++ b/Engine/Source/Runtime/UMG/Private/Components/EditableText.cpp
@@ -12,12 +12,21 @@
 /////////////////////////////////////////////////////
 // UEditableText
 
+static FEditableTextStyle* DefaultEditableTextStyle = nullptr;
+
 UEditableText::UEditableText(const FObjectInitializer& ObjectInitializer)
 	: Super(ObjectInitializer)
 {
-	// HACK: THIS SHOULD NOT COME FROM CORESTYLE AND SHOULD INSTEAD BY DEFINED BY ENGINE TEXTURES/PROJECT SETTINGS
-	static const FEditableTextStyle StaticEditableTextStyle = FCoreStyle::Get().GetWidgetStyle< FEditableTextStyle >("NormalEditableText");
-	WidgetStyle = StaticEditableTextStyle;
+	if (DefaultEditableTextStyle == nullptr)
+	{
+		// HACK: THIS SHOULD NOT COME FROM CORESTYLE AND SHOULD INSTEAD BE DEFINED BY ENGINE TEXTURES/PROJECT SETTINGS
+		DefaultEditableTextStyle = new FEditableTextStyle(FCoreStyle::Get().GetWidgetStyle<FEditableTextStyle>("NormalEditableText"));
+
+		// Unlink UMG default colors from the editor settings colors.
+		DefaultEditableTextStyle->UnlinkColors();
+	}
+
+	WidgetStyle = *DefaultEditableTextStyle;
 
 	ColorAndOpacity_DEPRECATED = FLinearColor::Black;
 
diff --git a/Engine/Source/Runtime/UMG/Private/Components/EditableTextBox.cpp b/Engine/Source/Runtime/UMG/Private/Components/EditableTextBox.cpp
--- a/Engine/Source/Runtime/UMG/Private/Components/EditableTextBox.cpp
+++ b/Engine/Source/Runtime/UMG/Private/Components/EditableTextBox.cpp
@@ -11,6 +11,8 @@
 /////////////////////////////////////////////////////
 // UEditableTextBox
 
+static FEditableTextBoxStyle* DefaultEditableTextBoxStyle = nullptr;
+
 UEditableTextBox::UEditableTextBox(const FObjectInitializer& ObjectInitializer)
 	: Super(ObjectInitializer)
 {
@@ -36,9 +38,16 @@ UEditableTextBox::UEditableTextBox(const FObjectInitializer& ObjectInitializer)
 	AllowContextMenu = true;
 	VirtualKeyboardDismissAction = EVirtualKeyboardDismissAction::TextChangeOnDismiss;
 
-	// HACK: THIS SHOULD NOT COME FROM CORESTYLE AND SHOULD INSTEAD BY DEFINED BY ENGINE TEXTURES/PROJECT SETTINGS
-	static const FEditableTextBoxStyle StaticNormalEditableTextBox = FCoreStyle::Get().GetWidgetStyle< FEditableTextBoxStyle >("NormalEditableTextBox");
-	WidgetStyle = StaticNormalEditableTextBox;
+	if (DefaultEditableTextBoxStyle == nullptr)
+	{
+		// HACK: THIS SHOULD NOT COME FROM CORESTYLE AND SHOULD INSTEAD BE DEFINED BY ENGINE TEXTURES/PROJECT SETTINGS
+		DefaultEditableTextBoxStyle = new FEditableTextBoxStyle(FCoreStyle::Get().GetWidgetStyle<FEditableTextBoxStyle>("NormalEditableTextBox"));
+
+		// Unlink UMG default colors from the editor settings colors.
+		DefaultEditableTextBoxStyle->UnlinkColors();
+	}
+
+	WidgetStyle = *DefaultEditableTextBoxStyle;
 
 	AccessibleBehavior = ESlateAccessibleBehavior::Auto;
 	bCanChildrenBeAccessible = false;
diff --git a/Engine/Source/Runtime/UMG/Private/Components/ExpandableArea.cpp b/Engine/Source/Runtime/UMG/Private/Components/ExpandableArea.cpp
--- a/Engine/Source/Runtime/UMG/Private/Components/ExpandableArea.cpp
+++ b/Engine/Source/Runtime/UMG/Private/Components/ExpandableArea.cpp
@@ -14,18 +14,35 @@
 static const FName HeaderName(TEXT("Header"));
 static const FName BodyName(TEXT("Body"));
 
+static FExpandableAreaStyle* DefaultExpandableAreaStyle = nullptr;
+static FSlateBrush* DefaultExpandableAreaBorderBrush = nullptr;
+
 UExpandableArea::UExpandableArea(const FObjectInitializer& ObjectInitializer)
 	: Super(ObjectInitializer)
 	, bIsExpanded(false)
 {
 	bIsVariable = true;
 
-	// HACK: THIS SHOULD NOT COME FROM CORESTYLE AND SHOULD INSTEAD BY DEFINED BY ENGINE TEXTURES/PROJECT SETTINGS
-	static const FExpandableAreaStyle StaticExpandableArea = FCoreStyle::Get().GetWidgetStyle<FExpandableAreaStyle>("ExpandableArea");
-	static const FSlateBrush StaticBorderBrush = *FCoreStyle::Get().GetBrush("ExpandableArea.Border");
+	if (DefaultExpandableAreaStyle == nullptr)
+	{
+		// HACK: THIS SHOULD NOT COME FROM CORESTYLE AND SHOULD INSTEAD BE DEFINED BY ENGINE TEXTURES/PROJECT SETTINGS
+		DefaultExpandableAreaStyle = new FExpandableAreaStyle(FCoreStyle::Get().GetWidgetStyle<FExpandableAreaStyle>("ExpandableArea"));
+
+		// Unlink UMG default colors from the editor settings colors.
+		DefaultExpandableAreaStyle->UnlinkColors();
+	}
+
+	if (DefaultExpandableAreaBorderBrush == nullptr)
+	{
+		// HACK: THIS SHOULD NOT COME FROM CORESTYLE AND SHOULD INSTEAD BE DEFINED BY ENGINE TEXTURES/PROJECT SETTINGS
+		DefaultExpandableAreaBorderBrush = new FSlateBrush(*FCoreStyle::Get().GetBrush("ExpandableArea.Border"));
+
+		// Unlink UMG default colors from the editor settings colors.
+		DefaultExpandableAreaBorderBrush->UnlinkColors();
+	}
 
-	Style = StaticExpandableArea;
-	BorderBrush = StaticBorderBrush;
+	Style = *DefaultExpandableAreaStyle;
+	BorderBrush = *DefaultExpandableAreaBorderBrush;
 
 	BorderColor = FLinearColor::White;
 	AreaPadding = FMargin(1);
diff --git a/Engine/Source/Runtime/UMG/Private/Components/InputKeySelector.cpp b/Engine/Source/Runtime/UMG/Private/Components/InputKeySelector.cpp
--- a/Engine/Source/Runtime/UMG/Private/Components/InputKeySelector.cpp
+++ b/Engine/Source/Runtime/UMG/Private/Components/InputKeySelector.cpp
@@ -10,14 +10,32 @@
 
 #define LOCTEXT_NAMESPACE "UMG"
 
+static FButtonStyle* DefaultInputKeySelectorButtonStyle = nullptr;
+static FTextBlockStyle* DefaultInputKeySelectorTextStyle = nullptr;
+
 UInputKeySelector::UInputKeySelector( const FObjectInitializer& ObjectInitializer )
 	: Super(ObjectInitializer)
 {
-	// HACK: THIS SHOULD NOT COME FROM CORESTYLE AND SHOULD INSTEAD BY DEFINED BY ENGINE TEXTURES/PROJECT SETTINGS
-	static const FButtonStyle StaticButtonStyle = FCoreStyle::Get().GetWidgetStyle<FButtonStyle>("Button");
-	static const FTextBlockStyle StaticNormalTextStyle = FCoreStyle::Get().GetWidgetStyle< FTextBlockStyle >("NormalText");
-	WidgetStyle = StaticButtonStyle;
-	TextStyle = StaticNormalTextStyle;
+	if (DefaultInputKeySelectorButtonStyle == nullptr)
+	{
+		// HACK: THIS SHOULD NOT COME FROM CORESTYLE AND SHOULD INSTEAD BE DEFINED BY ENGINE TEXTURES/PROJECT SETTINGS
+		DefaultInputKeySelectorButtonStyle = new FButtonStyle(FCoreStyle::Get().GetWidgetStyle<FButtonStyle>("Button"));
+
+		// Unlink UMG default colors from the editor settings colors.
+		DefaultInputKeySelectorButtonStyle->UnlinkColors();
+	}
+
+	if (DefaultInputKeySelectorTextStyle == nullptr)
+	{
+		// HACK: THIS SHOULD NOT COME FROM CORESTYLE AND SHOULD INSTEAD BE DEFINED BY ENGINE TEXTURES/PROJECT SETTINGS
+		DefaultInputKeySelectorTextStyle = new FTextBlockStyle(FCoreStyle::Get().GetWidgetStyle<FTextBlockStyle>("NormalText"));
+
+		// Unlink UMG default colors from the editor settings colors.
+		DefaultInputKeySelectorTextStyle->UnlinkColors();
+	}
+
+	WidgetStyle = *DefaultInputKeySelectorButtonStyle;
+	TextStyle = *DefaultInputKeySelectorTextStyle;
 
 	KeySelectionText = NSLOCTEXT("InputKeySelector", "DefaultKeySelectionText", "...");
 	NoKeySpecifiedText = NSLOCTEXT("InputKeySelector", "DefaultEmptyText", "Empty");
diff --git a/Engine/Source/Runtime/UMG/Private/Components/MultiLineEditableText.cpp b/Engine/Source/Runtime/UMG/Private/Components/MultiLineEditableText.cpp
--- a/Engine/Source/Runtime/UMG/Private/Components/MultiLineEditableText.cpp
+++ b/Engine/Source/Runtime/UMG/Private/Components/MultiLineEditableText.cpp
@@ -11,12 +11,21 @@
 /////////////////////////////////////////////////////
 // UMultiLineEditableText
 
+static FTextBlockStyle* DefaultMultiLineEditableTextStyle = nullptr;
+
 UMultiLineEditableText::UMultiLineEditableText(const FObjectInitializer& ObjectInitializer)
 	: Super(ObjectInitializer)
 {
-	// HACK: THIS SHOULD NOT COME FROM CORESTYLE AND SHOULD INSTEAD BY DEFINED BY ENGINE TEXTURES/PROJECT SETTINGS
-	static const FTextBlockStyle StaticNormalText = FCoreStyle::Get().GetWidgetStyle<FTextBlockStyle>("NormalText");
-	WidgetStyle = StaticNormalText;
+	if (DefaultMultiLineEditableTextStyle == nullptr)
+	{
+		// HACK: THIS SHOULD NOT COME FROM CORESTYLE AND SHOULD INSTEAD BE DEFINED BY ENGINE TEXTURES/PROJECT SETTINGS
+		DefaultMultiLineEditableTextStyle = new FTextBlockStyle(FCoreStyle::Get().GetWidgetStyle<FTextBlockStyle>("NormalText"));
+
+		// Unlink UMG default colors from the editor settings colors.
+		DefaultMultiLineEditableTextStyle->UnlinkColors();
+	}
+
+	WidgetStyle = *DefaultMultiLineEditableTextStyle;
 	
 	bIsReadOnly = false;
 	SelectAllTextWhenFocused = false;
diff --git a/Engine/Source/Runtime/UMG/Private/Components/MultiLineEditableTextBox.cpp b/Engine/Source/Runtime/UMG/Private/Components/MultiLineEditableTextBox.cpp
--- a/Engine/Source/Runtime/UMG/Private/Components/MultiLineEditableTextBox.cpp
+++ b/Engine/Source/Runtime/UMG/Private/Components/MultiLineEditableTextBox.cpp
@@ -11,6 +11,9 @@
 /////////////////////////////////////////////////////
 // UMultiLineEditableTextBox
 
+static FEditableTextBoxStyle* DefaultMultiLineEditableTextBoxStyle = nullptr;
+static FTextBlockStyle* DefaultMultiLineEditableTextBoxTextStyle = nullptr;
+
 UMultiLineEditableTextBox::UMultiLineEditableTextBox(const FObjectInitializer& ObjectInitializer)
 	: Super(ObjectInitializer)
 {
@@ -18,11 +21,26 @@ UMultiLineEditableTextBox::UMultiLineEditableTextBox(const FObjectInitializer& O
 	BackgroundColor_DEPRECATED = FLinearColor::White;
 	ReadOnlyForegroundColor_DEPRECATED = FLinearColor::Black;
 
-	// HACK: THIS SHOULD NOT COME FROM CORESTYLE AND SHOULD INSTEAD BY DEFINED BY ENGINE TEXTURES/PROJECT SETTINGS
-	static const FEditableTextBoxStyle StaticNormalEditableTextBox = FCoreStyle::Get().GetWidgetStyle<FEditableTextBoxStyle>("NormalEditableTextBox");
-	static const FTextBlockStyle StaticNormalText = FCoreStyle::Get().GetWidgetStyle<FTextBlockStyle>("NormalText");
-	WidgetStyle = StaticNormalEditableTextBox;
-	TextStyle = StaticNormalText;
+	if (DefaultMultiLineEditableTextBoxStyle == nullptr)
+	{
+		// HACK: THIS SHOULD NOT COME FROM CORESTYLE AND SHOULD INSTEAD BE DEFINED BY ENGINE TEXTURES/PROJECT SETTINGS
+		DefaultMultiLineEditableTextBoxStyle = new FEditableTextBoxStyle(FCoreStyle::Get().GetWidgetStyle<FEditableTextBoxStyle>("NormalEditableTextBox"));
+
+		// Unlink UMG default colors from the editor settings colors.
+		DefaultMultiLineEditableTextBoxStyle->UnlinkColors();
+	}
+
+	if (DefaultMultiLineEditableTextBoxTextStyle == nullptr)
+	{
+		// HACK: THIS SHOULD NOT COME FROM CORESTYLE AND SHOULD INSTEAD BE DEFINED BY ENGINE TEXTURES/PROJECT SETTINGS
+		DefaultMultiLineEditableTextBoxTextStyle = new FTextBlockStyle(FCoreStyle::Get().GetWidgetStyle<FTextBlockStyle>("NormalText"));
+
+		// Unlink UMG default colors from the editor settings colors.
+		DefaultMultiLineEditableTextBoxTextStyle->UnlinkColors();
+	}
+	
+	WidgetStyle = *DefaultMultiLineEditableTextBoxStyle;
+	TextStyle = *DefaultMultiLineEditableTextBoxTextStyle;
 
 	bIsReadOnly = false;
 	AllowContextMenu = true;
diff --git a/Engine/Source/Runtime/UMG/Private/Components/ProgressBar.cpp b/Engine/Source/Runtime/UMG/Private/Components/ProgressBar.cpp
--- a/Engine/Source/Runtime/UMG/Private/Components/ProgressBar.cpp
+++ b/Engine/Source/Runtime/UMG/Private/Components/ProgressBar.cpp
@@ -8,13 +8,21 @@
 /////////////////////////////////////////////////////
 // UProgressBar
 
+static FProgressBarStyle* DefaultProgressBarStyle = nullptr;
+
 UProgressBar::UProgressBar(const FObjectInitializer& ObjectInitializer)
 	: Super(ObjectInitializer)
 {
-	// HACK: THIS SHOULD NOT COME FROM CORESTYLE AND SHOULD INSTEAD BY DEFINED BY ENGINE TEXTURES/PROJECT SETTINGS
-	static const FProgressBarStyle StaticProgressBar = FCoreStyle::Get().GetWidgetStyle<FProgressBarStyle>("ProgressBar");
+	if (DefaultProgressBarStyle == nullptr)
+	{
+		// HACK: THIS SHOULD NOT COME FROM CORESTYLE AND SHOULD INSTEAD BE DEFINED BY ENGINE TEXTURES/PROJECT SETTINGS
+		DefaultProgressBarStyle = new FProgressBarStyle(FCoreStyle::Get().GetWidgetStyle<FProgressBarStyle>("ProgressBar"));
+
+		// Unlink UMG default colors from the editor settings colors.
+		DefaultProgressBarStyle->UnlinkColors();
+	}
 
-	WidgetStyle = StaticProgressBar;
+	WidgetStyle = *DefaultProgressBarStyle;
 	WidgetStyle.FillImage.TintColor = FLinearColor::White;
 
 	BarFillType = EProgressBarFillType::LeftToRight;
diff --git a/Engine/Source/Runtime/UMG/Private/Components/ScrollBar.cpp b/Engine/Source/Runtime/UMG/Private/Components/ScrollBar.cpp
--- a/Engine/Source/Runtime/UMG/Private/Components/ScrollBar.cpp
+++ b/Engine/Source/Runtime/UMG/Private/Components/ScrollBar.cpp
@@ -8,6 +8,8 @@
 /////////////////////////////////////////////////////
 // UScrollBar
 
+static FScrollBarStyle* DefaultScrollBarStyle = nullptr;
+
 UScrollBar::UScrollBar(const FObjectInitializer& ObjectInitializer)
 	: Super(ObjectInitializer)
 {
@@ -19,9 +21,16 @@ UScrollBar::UScrollBar(const FObjectInitializer& ObjectInitializer)
 	Thickness = FVector2D(16.0f, 16.0f);
 	Padding = FMargin(2.0f);
 
-	// HACK: THIS SHOULD NOT COME FROM CORESTYLE AND SHOULD INSTEAD BY DEFINED BY ENGINE TEXTURES/PROJECT SETTINGS
-	static const FScrollBarStyle StaticScrollbar = FCoreStyle::Get().GetWidgetStyle<FScrollBarStyle>("Scrollbar");
-	WidgetStyle = StaticScrollbar;
+	if (DefaultScrollBarStyle == nullptr)
+	{
+		// HACK: THIS SHOULD NOT COME FROM CORESTYLE AND SHOULD INSTEAD BE DEFINED BY ENGINE TEXTURES/PROJECT SETTINGS
+		DefaultScrollBarStyle = new FScrollBarStyle(FCoreStyle::Get().GetWidgetStyle<FScrollBarStyle>("Scrollbar"));
+
+		// Unlink UMG default colors from the editor settings colors.
+		DefaultScrollBarStyle->UnlinkColors();
+	}
+	
+	WidgetStyle = *DefaultScrollBarStyle;
 }
 
 void UScrollBar::ReleaseSlateResources(bool bReleaseChildren)
diff --git a/Engine/Source/Runtime/UMG/Private/Components/ScrollBox.cpp b/Engine/Source/Runtime/UMG/Private/Components/ScrollBox.cpp
--- a/Engine/Source/Runtime/UMG/Private/Components/ScrollBox.cpp
+++ b/Engine/Source/Runtime/UMG/Private/Components/ScrollBox.cpp
@@ -10,6 +10,9 @@
 /////////////////////////////////////////////////////
 // UScrollBox
 
+static FScrollBoxStyle* DefaultScrollBoxStyle = nullptr;
+static FScrollBarStyle* DefaultScrollBoxBarStyle = nullptr;
+
 UScrollBox::UScrollBox(const FObjectInitializer& ObjectInitializer)
 	: Super(ObjectInitializer)
 	, Orientation(Orient_Vertical)
@@ -28,14 +31,26 @@ UScrollBox::UScrollBox(const FObjectInitializer& ObjectInitializer)
 	Visibility = ESlateVisibility::Visible;
 	Clipping = EWidgetClipping::ClipToBounds;
 
+	if (DefaultScrollBoxStyle == nullptr)
+	{
+		// HACK: THIS SHOULD NOT COME FROM CORESTYLE AND SHOULD INSTEAD BE DEFINED BY ENGINE TEXTURES/PROJECT SETTINGS
+		DefaultScrollBoxStyle = new FScrollBoxStyle(FCoreStyle::Get().GetWidgetStyle<FScrollBoxStyle>("ScrollBox"));
+
+		// Unlink UMG default colors from the editor settings colors.
+		DefaultScrollBoxStyle->UnlinkColors();
+	}
+
+	if (DefaultScrollBoxBarStyle == nullptr)
 	{
-		// HACK: THIS SHOULD NOT COME FROM CORESTYLE AND SHOULD INSTEAD BY DEFINED BY ENGINE TEXTURES/PROJECT SETTINGS
-		static const FScrollBoxStyle StaticScrollBoxStyle = FCoreStyle::Get().GetWidgetStyle<FScrollBoxStyle>("ScrollBox");
-		static const FScrollBarStyle StaticScrollBarStyle = FCoreStyle::Get().GetWidgetStyle<FScrollBarStyle>("ScrollBar");
+		// HACK: THIS SHOULD NOT COME FROM CORESTYLE AND SHOULD INSTEAD BE DEFINED BY ENGINE TEXTURES/PROJECT SETTINGS
+		DefaultScrollBoxBarStyle = new FScrollBarStyle(FCoreStyle::Get().GetWidgetStyle<FScrollBarStyle>("ScrollBar"));
 
-		WidgetStyle = StaticScrollBoxStyle;
-		WidgetBarStyle = StaticScrollBarStyle;
+		// Unlink UMG default colors from the editor settings colors.
+		DefaultScrollBoxBarStyle->UnlinkColors();
 	}
+	
+	WidgetStyle = *DefaultScrollBoxStyle;
+	WidgetBarStyle = *DefaultScrollBoxBarStyle;
 
 	bAllowRightClickDragScrolling = true;
 }
diff --git a/Engine/Source/Runtime/UMG/Private/Components/Slider.cpp b/Engine/Source/Runtime/UMG/Private/Components/Slider.cpp
--- a/Engine/Source/Runtime/UMG/Private/Components/Slider.cpp
+++ b/Engine/Source/Runtime/UMG/Private/Components/Slider.cpp
@@ -9,6 +9,8 @@
 /////////////////////////////////////////////////////
 // USlider
 
+static FSliderStyle* DefaultSliderStyle = nullptr;
+
 USlider::USlider(const FObjectInitializer& ObjectInitializer)
 	: Super(ObjectInitializer)
 {
@@ -22,9 +24,16 @@ USlider::USlider(const FObjectInitializer& ObjectInitializer)
 	MouseUsesStep = false;
 	RequiresControllerLock = true;
 
-	// HACK: THIS SHOULD NOT COME FROM CORESTYLE AND SHOULD INSTEAD BY DEFINED BY ENGINE TEXTURES/PROJECT SETTINGS
-	static const FSliderStyle StaticSlider = FCoreStyle::Get().GetWidgetStyle<FSliderStyle>("Slider");
-	WidgetStyle = StaticSlider;
+	if (DefaultSliderStyle == nullptr)
+	{
+		// HACK: THIS SHOULD NOT COME FROM CORESTYLE AND SHOULD INSTEAD BE DEFINED BY ENGINE TEXTURES/PROJECT SETTINGS
+		DefaultSliderStyle = new FSliderStyle(FCoreStyle::Get().GetWidgetStyle<FSliderStyle>("Slider"));
+
+		// Unlink UMG default colors from the editor settings colors.
+		DefaultSliderStyle->UnlinkColors();
+	}
+
+	WidgetStyle = *DefaultSliderStyle;
 
 	AccessibleBehavior = ESlateAccessibleBehavior::Summary;
 	bCanChildrenBeAccessible = false;
diff --git a/Engine/Source/Runtime/UMG/Private/Components/SpinBox.cpp b/Engine/Source/Runtime/UMG/Private/Components/SpinBox.cpp
--- a/Engine/Source/Runtime/UMG/Private/Components/SpinBox.cpp
+++ b/Engine/Source/Runtime/UMG/Private/Components/SpinBox.cpp
@@ -9,6 +9,8 @@
 /////////////////////////////////////////////////////
 // USpinBox
 
+static FSpinBoxStyle* DefaultSpinBoxStyle = nullptr;
+
 USpinBox::USpinBox(const FObjectInitializer& ObjectInitializer)
 	: Super(ObjectInitializer)
 {
@@ -30,9 +32,16 @@ USpinBox::USpinBox(const FObjectInitializer& ObjectInitializer)
 	SelectAllTextOnCommit = true;
 	ForegroundColor = FSlateColor(FLinearColor::Black);
 
-	// HACK: THIS SHOULD NOT COME FROM CORESTYLE AND SHOULD INSTEAD BY DEFINED BY ENGINE TEXTURES/PROJECT SETTINGS
-	static const FSpinBoxStyle StaticSpinBox = FCoreStyle::Get().GetWidgetStyle<FSpinBoxStyle>("SpinBox");
-	WidgetStyle = StaticSpinBox;
+	if (DefaultSpinBoxStyle == nullptr)
+	{
+		// HACK: THIS SHOULD NOT COME FROM CORESTYLE AND SHOULD INSTEAD BE DEFINED BY ENGINE TEXTURES/PROJECT SETTINGS
+		DefaultSpinBoxStyle = new FSpinBoxStyle(FCoreStyle::Get().GetWidgetStyle<FSpinBoxStyle>("SpinBox"));
+
+		// Unlink UMG default colors from the editor settings colors.
+		DefaultSpinBoxStyle->UnlinkColors();
+	}
+
+	WidgetStyle = *DefaultSpinBoxStyle;
 }
 
 void USpinBox::ReleaseSlateResources(bool bReleaseChildren)
diff --git a/Engine/Source/Runtime/UMG/Private/Components/Throbber.cpp b/Engine/Source/Runtime/UMG/Private/Components/Throbber.cpp
--- a/Engine/Source/Runtime/UMG/Private/Components/Throbber.cpp
+++ b/Engine/Source/Runtime/UMG/Private/Components/Throbber.cpp
@@ -9,6 +9,8 @@
 /////////////////////////////////////////////////////
 // UThrobber
 
+static FSlateBrush* DefaultThrobberBrush = nullptr;
+
 UThrobber::UThrobber(const FObjectInitializer& ObjectInitializer)
 	: Super(ObjectInitializer)
 {
@@ -18,9 +20,16 @@ UThrobber::UThrobber(const FObjectInitializer& ObjectInitializer)
 	bAnimateHorizontally = true;
 	bAnimateOpacity = true;
 
-	// HACK: THIS SHOULD NOT COME FROM CORESTYLE AND SHOULD INSTEAD BY DEFINED BY ENGINE TEXTURES/PROJECT SETTINGS
-	static const FSlateBrush StaticImage = *FCoreStyle::Get().GetBrush("Throbber.Chunk");
-	Image = StaticImage;
+	if (DefaultThrobberBrush == nullptr)
+	{
+		// HACK: THIS SHOULD NOT COME FROM CORESTYLE AND SHOULD INSTEAD BE DEFINED BY ENGINE TEXTURES/PROJECT SETTINGS
+		DefaultThrobberBrush = new FSlateBrush(*FCoreStyle::Get().GetBrush("Throbber.Chunk"));
+
+		// Unlink UMG default colors from the editor settings colors.
+		DefaultThrobberBrush->UnlinkColors();
+	}
+
+	Image = *DefaultThrobberBrush;
 }
 
 void UThrobber::ReleaseSlateResources(bool bReleaseChildren)
