Reverting unintended SWindow position changes

diff --git a/Engine/Source/Runtime/SlateCore/Private/Widgets/SWindow.cpp b/Engine/Source/Runtime/SlateCore/Private/Widgets/SWindow.cpp
--- a/Engine/Source/Runtime/SlateCore/Private/Widgets/SWindow.cpp
+++ b/Engine/Source/Runtime/SlateCore/Private/Widgets/SWindow.cpp
@@ -908,7 +908,7 @@ void SWindow::ReshapeWindow( FVector2D NewPosition, FVector2D NewSize )
 
 void SWindow::ReshapeWindow( const FSlateRect& InNewShape )
 {
-	ReshapeWindow( FVector2D(500, 500), FVector2D( InNewShape.Right - InNewShape.Left,  InNewShape.Bottom - InNewShape.Top) );
+	ReshapeWindow(FVector2D(InNewShape.Left, InNewShape.Top), FVector2D(InNewShape.Right - InNewShape.Left, InNewShape.Bottom - InNewShape.Top));
 }
 
 void SWindow::Resize( FVector2D NewSize )
@@ -1319,7 +1319,7 @@ void SWindow::ShowWindow()
 		{
 			SlatePrepass( FSlateApplicationBase::Get().GetApplicationScale() * NativeWindow->GetDPIScaleFactor() );
 			const FVector2D WindowDesiredSizePixels = GetDesiredSizeDesktopPixels();
-			ReshapeWindow( FVector2D(200, 200), WindowDesiredSizePixels);
+			ReshapeWindow(InitialDesiredScreenPosition - (WindowDesiredSizePixels * 0.5f), WindowDesiredSizePixels);
 		}
 
 		// Set the window to be maximized if we need to.  Note that this won't actually show the window if its not
