Added default values for some uninitialized variables

diff --git a/Engine/Source/Runtime/SlateCore/Public/Input/Events.h b/Engine/Source/Runtime/SlateCore/Public/Input/Events.h
--- a/Engine/Source/Runtime/SlateCore/Public/Input/Events.h
+++ b/Engine/Source/Runtime/SlateCore/Public/Input/Events.h
@@ -631,6 +631,8 @@ public:
 		, GestureType(EGestureEvent::None)
 		, WheelOrGestureDelta(0.0f, 0)
 		, bIsDirectionInvertedFromDevice(false)
+		, bIsTouchForceChanged(false)
+		, bIsTouchFirstMove(false)
 	{ }
 
 	/** Events are immutable once constructed. */
@@ -656,6 +658,8 @@ public:
 		, GestureType(EGestureEvent::None)
 		, WheelOrGestureDelta(0.0f, InWheelDelta)
 		, bIsDirectionInvertedFromDevice(false)
+		, bIsTouchForceChanged(false)
+		, bIsTouchFirstMove(false)
 	{ }
 
 	FPointerEvent(
@@ -681,6 +685,8 @@ public:
 		, GestureType(EGestureEvent::None)
 		, WheelOrGestureDelta(0.0f, InWheelDelta)
 		, bIsDirectionInvertedFromDevice(false)
+		, bIsTouchForceChanged(false)
+		, bIsTouchFirstMove(false)
 	{ }
 
 	/** A constructor for raw mouse events */
@@ -704,6 +710,8 @@ public:
 		, GestureType(EGestureEvent::None)
 		, WheelOrGestureDelta(0.0f, 0.0f)
 		, bIsDirectionInvertedFromDevice(false)
+		, bIsTouchForceChanged(false)
+		, bIsTouchFirstMove(false)
 	{ }
 
 	/** A constructor for touch events */
@@ -770,6 +778,8 @@ public:
 		, GestureType(InGestureType)
 		, WheelOrGestureDelta(InGestureDelta)
 		, bIsDirectionInvertedFromDevice(bInIsDirectionInvertedFromDevice)
+		, bIsTouchForceChanged(false)
+		, bIsTouchFirstMove(false)
 	{ }
 	
 public:
diff --git a/Engine/Source/Runtime/SlateCore/Public/Widgets/SWindow.h b/Engine/Source/Runtime/SlateCore/Public/Widgets/SWindow.h
--- a/Engine/Source/Runtime/SlateCore/Public/Widgets/SWindow.h
+++ b/Engine/Source/Runtime/SlateCore/Public/Widgets/SWindow.h
@@ -155,6 +155,7 @@ public:
 		, _SaneWindowPlacement( true )
 		, _LayoutBorder(FMargin(5, 5, 5, 5))
 		, _UserResizeBorder(FMargin(5, 5, 5, 5))
+		, _bManualManageDPI( false )
 
 	{
 		_AccessibleParams = FAccessibleWidgetData(EAccessibleBehavior::Auto);
