Fixed SWindow overlay windows flickering when Windows UI scaling was set to 125%.

diff --git a/Engine/Source/Runtime/SlateCore/Private/Widgets/SWindow.cpp b/Engine/Source/Runtime/SlateCore/Private/Widgets/SWindow.cpp
--- a/Engine/Source/Runtime/SlateCore/Private/Widgets/SWindow.cpp
+++ b/Engine/Source/Runtime/SlateCore/Private/Widgets/SWindow.cpp
@@ -882,7 +882,7 @@ void SWindow::ReshapeWindow( FVector2D NewPosition, FVector2D NewSize )
 	const FVector2D CurrentSize = GetSizeInScreen();
 
 	const FVector2D NewPositionTruncated = FVector2D(FMath::TruncToInt(NewPosition.X), FMath::TruncToInt(NewPosition.Y));
-	const FVector2D NewSizeRounded = FVector2D(FMath::TruncToInt(NewSize.X), FMath::TruncToInt(NewSize.Y));
+	const FVector2D NewSizeRounded = FVector2D(FMath::CeilToInt(NewSize.X), FMath::CeilToInt(NewSize.Y));
 
 	if ( CurrentPosition != NewPositionTruncated || CurrentSize != NewSizeRounded )
 	{
