Prevent an infinite loop in hit test grid when wrapping in a panel from a child widget outside of the bounds of the panel

diff --git a/Engine/Source/Runtime/SlateCore/Private/Input/HittestGrid.cpp b/Engine/Source/Runtime/SlateCore/Private/Input/HittestGrid.cpp
--- a/Engine/Source/Runtime/SlateCore/Private/Input/HittestGrid.cpp
+++ b/Engine/Source/Runtime/SlateCore/Private/Input/HittestGrid.cpp
@@ -475,6 +475,11 @@ TSharedPtr<SWidget> FHittestGrid::FindFocusableWidget(FSlateRect WidgetRect, con
 		{
 			if (NavigationReply.GetBoundaryRule() == EUINavigationRule::Wrap)
 			{
+				if (bWrapped)
+				{
+					// If we've already wrapped, unfortunately it must be that the starting widget wasn't within the boundary
+					break;
+				}
 				CurrentSourceSide = DestSideFunc(SweptRect);
 				FVector2D SampleSpot = WidgetRect.GetCenter();
 				SampleSpot[AxisIndex] = CurrentSourceSide;
