Fixed the Drag drop tool tip message from being cut off the edge of the screen.

diff --git a/Engine/Source/Runtime/Slate/Public/Framework/Application/SlateApplication.h b/Engine/Source/Runtime/Slate/Public/Framework/Application/SlateApplication.h
--- a/Engine/Source/Runtime/Slate/Public/Framework/Application/SlateApplication.h
+++ b/Engine/Source/Runtime/Slate/Public/Framework/Application/SlateApplication.h
@@ -860,7 +860,7 @@ public:
 	 * @param InSize The size of the tooltip window.
 	 * @return The suggested position.
 	 */
-	FVector2D CalculateTooltipWindowPosition( const FSlateRect& InAnchorRect, const FVector2D& InSize, bool bAutoAdjustForDPIScale) const;
+	virtual FVector2D CalculateTooltipWindowPosition( const FSlateRect& InAnchorRect, const FVector2D& InSize, bool bAutoAdjustForDPIScale) const;
 
 	/**
 	 * Is the window in the app's destroy queue? If so it will be destroyed next tick.
diff --git a/Engine/Source/Runtime/SlateCore/Private/Input/DragAndDrop.cpp b/Engine/Source/Runtime/SlateCore/Private/Input/DragAndDrop.cpp
--- a/Engine/Source/Runtime/SlateCore/Private/Input/DragAndDrop.cpp
+++ b/Engine/Source/Runtime/SlateCore/Private/Input/DragAndDrop.cpp
@@ -32,7 +32,17 @@ void FDragDropOperation::OnDragged( const class FDragDropEvent& DragDropEvent )
 {
 	if (CursorDecoratorWindow.IsValid())
 	{
-		CursorDecoratorWindow->MoveWindowTo(DragDropEvent.GetScreenSpacePosition() + FSlateApplicationBase::Get().GetCursorSize());
+		FSlateApplicationBase& SlateApplication = FSlateApplicationBase::Get();
+						
+		FVector2D Position = DragDropEvent.GetScreenSpacePosition() + FSlateApplicationBase::Get().GetCursorSize();
+
+		if (Position != CursorDecoratorWindow->GetPositionInScreen())
+		{
+			FSlateRect Anchor(Position.X, Position.Y, Position.X, Position.Y);
+			Position = SlateApplication.CalculateTooltipWindowPosition(Anchor, CursorDecoratorWindow->GetDesiredSizeDesktopPixels(), false);
+		}
+
+		CursorDecoratorWindow->MoveWindowTo(Position);
 	}
 }
 
diff --git a/Engine/Source/Runtime/SlateCore/Public/Application/SlateApplicationBase.h b/Engine/Source/Runtime/SlateCore/Public/Application/SlateApplicationBase.h
--- a/Engine/Source/Runtime/SlateCore/Public/Application/SlateApplicationBase.h
+++ b/Engine/Source/Runtime/SlateCore/Public/Application/SlateApplicationBase.h
@@ -339,6 +339,15 @@ public:
 	 */
 	virtual FWidgetPath LocateWindowUnderMouse( FVector2D ScreenspaceMouseCoordinate, const TArray< TSharedRef<SWindow > >& Windows, bool bIgnoreEnabledStatus = false ) = 0;
 
+	/**
+	 * Calculates the tooltip window position.
+	 *
+	 * @param InAnchorRect The current(suggested) window position and size of an area which may not be covered by the popup.
+	 * @param InSize The size of the tooltip window.
+	 * @return The suggested position.
+	 */
+	virtual FVector2D CalculateTooltipWindowPosition(const FSlateRect& InAnchorRect, const FVector2D& InSize, bool bAutoAdjustForDPIScale) const = 0;
+
 	/** @return true if 'WindowToTest' is being used to display the current tooltip and the tooltip is interactive. */
 	virtual bool IsWindowHousingInteractiveTooltip(const TSharedRef<const SWindow>& WindowToTest) const = 0;
 
