Fixed widget component mouse input

diff --git a/Engine/Source/Runtime/SlateCore/Private/Input/HittestGrid.cpp b/Engine/Source/Runtime/SlateCore/Private/Input/HittestGrid.cpp
--- a/Engine/Source/Runtime/SlateCore/Private/Input/HittestGrid.cpp
+++ b/Engine/Source/Runtime/SlateCore/Private/Input/HittestGrid.cpp
@@ -259,7 +259,7 @@ TArray<FWidgetAndPointer> FHittestGrid::GetBubblePath(FVector2D DesktopSpaceCoor
 				{
 					if (BestHitWidgetData.CustomPath.IsValid())
 					{
-						const TArray<FWidgetAndPointer> BubblePathExtension = BestHitWidgetData.CustomPath.Pin()->GetBubblePathAndVirtualCursors(FirstHitWidget->GetPaintSpaceGeometry(), DesktopSpaceCoordinate, bIgnoreEnabledStatus);
+						const TArray<FWidgetAndPointer> BubblePathExtension = BestHitWidgetData.CustomPath.Pin()->GetBubblePathAndVirtualCursors(FirstHitWidget->GetTickSpaceGeometry(), DesktopSpaceCoordinate, bIgnoreEnabledStatus);
 						Path.Append(BubblePathExtension);
 					}
 				}
diff --git a/Engine/Source/Runtime/SlateCore/Private/Layout/WidgetPath.cpp b/Engine/Source/Runtime/SlateCore/Private/Layout/WidgetPath.cpp
--- a/Engine/Source/Runtime/SlateCore/Private/Layout/WidgetPath.cpp
+++ b/Engine/Source/Runtime/SlateCore/Private/Layout/WidgetPath.cpp
@@ -305,7 +305,7 @@ FWeakWidgetPath::EPathResolutionResult::Result FWeakWidgetPath::ToWidgetPath( FW
 					{
 						if (PointerEvent && !VirtualPointerPos.IsValid())
 						{
-							VirtualPointerPos = CurWidget->TranslateMouseCoordinateFor3DChild(CurWidgetRef, ParentGeometry, PointerEvent->GetScreenSpacePosition(), PointerEvent->GetLastScreenSpacePosition());
+							VirtualPointerPos = CurWidget->TranslateMouseCoordinateFor3DChild(ChildWidgetPtr.ToSharedRef(), ParentGeometry, PointerEvent->GetScreenSpacePosition(), PointerEvent->GetLastScreenSpacePosition());
 						}
 
 						bFoundChild = true;
