PR #6154: Remove obsolete clipping rect comments (Contributed by [REDACTED_PERSON])

diff --git a/Engine/Source/Runtime/SlateCore/Public/Rendering/DrawElements.h b/Engine/Source/Runtime/SlateCore/Public/Rendering/DrawElements.h
--- a/Engine/Source/Runtime/SlateCore/Public/Rendering/DrawElements.h
+++ b/Engine/Source/Runtime/SlateCore/Public/Rendering/DrawElements.h
@@ -91,7 +91,6 @@ public:
 	 * @param InLayer               The layer to draw the element on
 	 * @param PaintGeometry         DrawSpace position and dimensions; see FPaintGeometry
 	 * @param InBrush               Brush to apply to this element
-	 * @param InClippingRect        Parts of the element are clipped if it falls outside of this rectangle
 	 * @param InDrawEffects         Optional draw effects to apply
 	 * @param InTint                Color to tint the element
 	 */
@@ -135,7 +134,6 @@ public:
 	 * @param StartIndex            Inclusive index to start rendering from on the specified text
 	 * @param EndIndex				Exclusive index to stop rendering on the specified text
 	 * @param InFontInfo            The font to draw the string with
-	 * @param InClippingRect        Parts of the element are clipped if it falls outside of this rectangle
 	 * @param InDrawEffects         Optional draw effects to apply
 	 * @param InTint                Color to tint the element
 	 */
@@ -155,7 +153,6 @@ public:
 	 * @param InLayer               The layer to draw the element on
 	 * @param PaintGeometry         DrawSpace position and dimensions; see FPaintGeometry
 	 * @param InShapedGlyphSequence The shaped glyph sequence to draw
-	 * @param InClippingRect        Parts of the element are clipped if it falls outside of this rectangle
 	 * @param InDrawEffects         Optional draw effects to apply
 	 * @param InTint                Color to tint the element
 	 */
@@ -169,7 +166,6 @@ public:
 	 * @param PaintGeometry            DrawSpace position and dimensions; see FPaintGeometry
 	 * @param InGradientStops          List of gradient stops which define the element
 	 * @param InGradientType           The type of gradient (I.E Horizontal, vertical)
-	 * @param InClippingRect           Parts of the element are clipped if it falls outside of this rectangle
 	 * @param InDrawEffects            Optional draw effects to apply
 	 */
 	SLATECORE_API static void MakeGradient( FSlateWindowElementList& ElementList, uint32 InLayer, const FPaintGeometry& PaintGeometry, TArray<FSlateGradientStop> InGradientStops, EOrientation InGradientType, ESlateDrawEffect InDrawEffects = ESlateDrawEffect::None );
@@ -184,7 +180,6 @@ public:
 	 * @param InStartDir            The direction of the spline from the start point
 	 * @param InEnd                 The end point of the spline (local space)
 	 * @param InEndDir              The direction of the spline to the end point
-	 * @param InClippingRect        Parts of the element are clipped if it falls outside of this rectangle
 	 * @param InDrawEffects         Optional draw effects to apply
 	 * @param InTint                Color to tint the element
 	 */
@@ -200,7 +195,6 @@ public:
 	 * @param InStartDir            The direction of the spline from the start point
 	 * @param InEnd                 The end point of the spline (local space)
 	 * @param InEndDir              The direction of the spline to the end point
-	 * @param InClippingRect        Parts of the element are clipped if it falls outside of this rectangle
 	 * @param InDrawEffects         Optional draw effects to apply
 	 * @param InTint                Color to tint the element
 	 */
@@ -223,7 +217,6 @@ public:
 	 * @param InLayer                  The layer to draw the element on
 	 * @param PaintGeometry            DrawSpace position and dimensions; see FPaintGeometry
 	 * @param Points                   Points that make up the lines.  The points are joined together. I.E if Points has A,B,C there the line is A-B-C.  To draw non-joining line segments call MakeLines multiple times
-	 * @param InClippingRect           Parts of the element are clipped if it falls outside of this rectangle
 	 * @param InDrawEffects            Optional draw effects to apply
 	 * @param InTint                   Color to tint the element
 	 * @param bAntialias               Should antialiasing be applied to the line?
@@ -239,7 +232,6 @@ public:
 	 * @param InLayer                  The layer to draw the element on
 	 * @param PaintGeometry            DrawSpace position and dimensions; see FPaintGeometry
 	 * @param Viewport                 Interface for drawing the viewport
-	 * @param InClippingRect           Parts of the element are clipped if it falls outside of this rectangle
 	 * @param InScale                  Draw scale to apply to the entire element
 	 * @param InDrawEffects            Optional draw effects to apply
 	 * @param InTint                   Color to tint the element
