Removed legacy localized fallback font support for 4.24

It had been deprecated in 4.22, and disabled by default in 4.23.

diff --git a/Engine/Source/Runtime/Slate/Private/Framework/Text/SlateTextLayout.cpp b/Engine/Source/Runtime/Slate/Private/Framework/Text/SlateTextLayout.cpp
--- a/Engine/Source/Runtime/Slate/Private/Framework/Text/SlateTextLayout.cpp
+++ b/Engine/Source/Runtime/Slate/Private/Framework/Text/SlateTextLayout.cpp
@@ -23,7 +23,6 @@ FSlateTextLayout::FSlateTextLayout(SWidget* InOwner, FTextBlockStyle InDefaultTe
 	: DefaultTextStyle(MoveTemp(InDefaultTextStyle))
 	, Children(InOwner, false)
 	, bIsPassword(false)
-	, LocalizedFallbackFontRevision(0)
 {
 
 }
@@ -152,24 +151,6 @@ void FSlateTextLayout::EndLayout()
 	AggregateChildren();
 }
 
-void FSlateTextLayout::UpdateIfNeeded()
-{
-	const uint16 CurrentLocalizedFallbackFontRevision = FSlateApplication::Get().GetRenderer()->GetFontCache()->GetLocalizedFallbackFontRevision();
-	if (CurrentLocalizedFallbackFontRevision != LocalizedFallbackFontRevision)
-	{
-		if (LocalizedFallbackFontRevision != 0)
-		{
-			// If the localized fallback font has changed, we need to purge the current layout data as things may need to be re-measured
-			DirtyFlags |= ETextLayoutDirtyState::Layout;
-			DirtyAllLineModels(ELineModelDirtyState::WrappingInformation | ELineModelDirtyState::ShapingCache);
-		}
-
-		LocalizedFallbackFontRevision = CurrentLocalizedFallbackFontRevision;
-	}
-
-	FTextLayout::UpdateIfNeeded();
-}
-
 void FSlateTextLayout::SetDefaultTextStyle(FTextBlockStyle InDefaultTextStyle)
 {
 	DefaultTextStyle = MoveTemp(InDefaultTextStyle);
diff --git a/Engine/Source/Runtime/Slate/Public/Framework/Text/SlateTextLayout.h b/Engine/Source/Runtime/Slate/Public/Framework/Text/SlateTextLayout.h
--- a/Engine/Source/Runtime/Slate/Public/Framework/Text/SlateTextLayout.h
+++ b/Engine/Source/Runtime/Slate/Public/Framework/Text/SlateTextLayout.h
@@ -27,8 +27,6 @@ public:
 
 	virtual void EndLayout() override;
 
-	virtual void UpdateIfNeeded() override;
-
 	void SetDefaultTextStyle(FTextBlockStyle InDefaultTextStyle);
 	const FTextBlockStyle& GetDefaultTextStyle() const;
 
@@ -55,8 +53,5 @@ private:
 	/** This this layout displaying a password? */
 	TAttribute<bool> bIsPassword;
 
-	/** The localized fallback font revision the last time the text layout was updated. Used to force a flush if the font changes. */
-	uint16 LocalizedFallbackFontRevision;
-
 	friend class FSlateTextLayoutFactory;
 };
diff --git a/Engine/Source/Runtime/SlateCore/Private/Fonts/FontCache.cpp b/Engine/Source/Runtime/SlateCore/Private/Fonts/FontCache.cpp
--- a/Engine/Source/Runtime/SlateCore/Private/Fonts/FontCache.cpp
+++ b/Engine/Source/Runtime/SlateCore/Private/Fonts/FontCache.cpp
@@ -1103,11 +1103,6 @@ TArray<FString> FSlateFontCache::GetAvailableFontSubFaces(const FString& InFilen
 	return FFreeTypeFace::GetAvailableSubFaces(FTLibrary.Get(), InFilename);
 }
 
-uint16 FSlateFontCache::GetLocalizedFallbackFontRevision() const
-{
-	return FLegacySlateFontInfoCache::Get().GetLocalizedFallbackFontRevision();
-}
-
 void FSlateFontCache::RequestFlushCache(const FString& FlushReason)
 {
 	if (!bFlushRequested)
diff --git a/Engine/Source/Runtime/SlateCore/Private/Fonts/LegacySlateFontInfoCache.cpp b/Engine/Source/Runtime/SlateCore/Private/Fonts/LegacySlateFontInfoCache.cpp
--- a/Engine/Source/Runtime/SlateCore/Private/Fonts/LegacySlateFontInfoCache.cpp
+++ b/Engine/Source/Runtime/SlateCore/Private/Fonts/LegacySlateFontInfoCache.cpp
@@ -1,7 +1,6 @@
 // Copyright 1998-2019 Epic Games, Inc. All Rights Reserved.
 
 #include "Fonts/LegacySlateFontInfoCache.h"
-#include "HAL/IConsoleManager.h"
 #include "SlateGlobals.h"
 #include "Misc/Paths.h"
 #include "Misc/ScopeLock.h"
@@ -9,9 +8,6 @@
 #include "Fonts/FontProviderInterface.h"
 #include "Fonts/UnicodeBlockRange.h"
 
-static int32 GSlateEnableLegacyLocalizedFallbackFont = 0;
-static FAutoConsoleVariableRef CVarSlateEnableLocalizedFallbackFont(TEXT("Slate.EnableLegacyLocalizedFallbackFont"), GSlateEnableLegacyLocalizedFallbackFont, TEXT("Enable the legacy localized fallback fonts? (0/1)."), ECVF_Default);
-
 FString FLegacySlateFontInfoCache::FFallbackContext::ToString() const
 {
 	FString FontName;
@@ -45,16 +41,11 @@ FLegacySlateFontInfoCache& FLegacySlateFontInfoCache::Get()
 	if (!Singleton.IsValid())
 	{
 		Singleton = MakeShareable(new FLegacySlateFontInfoCache());
-
-		FInternationalization::Get().OnCultureChanged().AddSP(Singleton.Get(), &FLegacySlateFontInfoCache::HandleCultureChanged);
 	}
 	return *Singleton;
 }
 
 FLegacySlateFontInfoCache::FLegacySlateFontInfoCache()
-	: LocalizedFallbackFontRevision(0)
-	, LocalizedFallbackFontDataHistoryVersion(0)
-	, LocalizedFallbackFontFrameCounter(0)
 {
 	LastResortFontPath = FPaths::EngineContentDir() / TEXT("SlateDebug/Fonts/LastResort.ttf");
 	bIsLastResortFontAvailable = !FPlatformProperties::RequiresCookedData() && FPaths::FileExists(LastResortFontPath);
@@ -233,54 +224,6 @@ TSharedPtr<const FCompositeFont> FLegacySlateFontInfoCache::GetSystemFont()
 	return SystemFont;
 }
 
-bool FLegacySlateFontInfoCache::IsLocalizedFallbackFontAvailable() const
-{
-	return !!GSlateEnableLegacyLocalizedFallbackFont;
-}
-
-const FFontData& FLegacySlateFontInfoCache::GetLocalizedFallbackFontData(const FFallbackContext& InContext)
-{
-	// GetLocalizedFallbackFontData is called directly from the font cache, so may be called from multiple threads at once
-	FScopeLock Lock(&LocalizedFallbackFontDataCS);
-
-	// The fallback font can change if the active culture is changed
-	const uint16 CurrentHistoryVersion = FTextLocalizationManager::Get().GetTextRevision();
-	const uint64 CurrentFrameCounter = GFrameCounter;
-
-	// Only allow the fallback font to be updated once per-frame, as a culture change mid-frame could cause it to change unexpectedly and invalidate some assumptions that the font cache makes
-	// By only allowing it to update once per-frame, we ensure that the font cache has been flushed (which happens at the end of the frame) before we return a new font
-	if (!LocalizedFallbackFontData.IsValid() || (LocalizedFallbackFontDataHistoryVersion != CurrentHistoryVersion && LocalizedFallbackFontFrameCounter != CurrentFrameCounter))
-	{
-		LocalizedFallbackFontDataHistoryVersion = CurrentHistoryVersion;
-		LocalizedFallbackFontFrameCounter = CurrentFrameCounter;
-
-		TSharedPtr<const FFontData> PreviousLocalizedFallbackFontData = LocalizedFallbackFontData;
-
-		const FString FallbackFontPath = FPaths::EngineContentDir() / TEXT("Slate/Fonts/") / (NSLOCTEXT("Slate", "LegacyFallbackFont", "DroidSansFallback").ToString() + TEXT(".ttf"));
-		LocalizedFallbackFontData = AllLocalizedFallbackFontData.FindRef(FallbackFontPath);
-
-		if (!LocalizedFallbackFontData.IsValid())
-		{
-			LocalizedFallbackFontData = MakeShared<FFontData>(FallbackFontPath, EFontHinting::Default, EFontLoadingPolicy::LazyLoad);
-			AllLocalizedFallbackFontData.Add(FallbackFontPath, LocalizedFallbackFontData);
-			UE_LOG(LogSlate, Warning, TEXT("Legacy localized fallback font '%s' was requested. %s\nLegacy localized fallback fonts were deprecated in 4.22 and will be removed in a future version!\nPlease update your composite fonts to use localized sub-font families: https://docs.unrealengine.com/en-US/Engine/UMG/UserGuide/Fonts/Overview"), *FallbackFontPath, *InContext.ToString());
-		}
-
-		if (LocalizedFallbackFontData != PreviousLocalizedFallbackFontData)
-		{
-			// Only bump the revision if the font has actually changed
-			while (++LocalizedFallbackFontRevision == 0) {} // Zero is special, don't allow an overflow to stay at zero
-		}
-	}
-
-	return *LocalizedFallbackFontData;
-}
-
-uint16 FLegacySlateFontInfoCache::GetLocalizedFallbackFontRevision() const
-{
-	return LocalizedFallbackFontRevision;
-}
-
 bool FLegacySlateFontInfoCache::IsLastResortFontAvailable() const
 {
 	return bIsLastResortFontAvailable;
@@ -316,11 +259,6 @@ const FFontData& FLegacySlateFontInfoCache::GetLastResortFontData(const FFallbac
 
 void FLegacySlateFontInfoCache::AddReferencedObjects(FReferenceCollector& Collector)
 {
-	for (TPair<FString, TSharedPtr<const FFontData>>& Pair : AllLocalizedFallbackFontData)
-	{
-		const_cast<FFontData&>(*Pair.Value).AddReferencedObjects(Collector);
-	}
-
 	if (LocalizedFallbackFontData.IsValid())
 	{
 		const_cast<FFontData&>(*LocalizedFallbackFontData).AddReferencedObjects(Collector);
@@ -337,8 +275,3 @@ FString FLegacySlateFontInfoCache::GetReferencerName() const
 	return TEXT("FLegacySlateFontInfoCache");
 }
 
-void FLegacySlateFontInfoCache::HandleCultureChanged()
-{
-	// We set this to the current frame count, as this will prevent the fallback font being updated for the remainder of this frame (as the culture change may have affected the fallback font used)
-	LocalizedFallbackFontFrameCounter = GFrameCounter;
-}
diff --git a/Engine/Source/Runtime/SlateCore/Private/Fonts/LegacySlateFontInfoCache.h b/Engine/Source/Runtime/SlateCore/Private/Fonts/LegacySlateFontInfoCache.h
--- a/Engine/Source/Runtime/SlateCore/Private/Fonts/LegacySlateFontInfoCache.h
+++ b/Engine/Source/Runtime/SlateCore/Private/Fonts/LegacySlateFontInfoCache.h
@@ -53,21 +53,6 @@ public:
 	 */
 	TSharedPtr<const FCompositeFont> GetSystemFont();
 
-	/**
-	 * Is the localized fallback font available (can be controlled via a CVar).
-	 */
-	bool IsLocalizedFallbackFontAvailable() const;
-
-	/**
-	 * Get (or create) the culture specific fallback font
-	 */
-	const FFontData& GetLocalizedFallbackFontData(const FFallbackContext& InContext);
-
-	/**
-	 * Get the revision index of the currently active localized fallback font.
-	 */
-	uint16 GetLocalizedFallbackFontRevision() const;
-
 	/**
 	 * Is the last resort fallback font available (not all builds have it).
 	 */
@@ -90,11 +75,6 @@ public:
 private:
 	FLegacySlateFontInfoCache();
 
-	/**
-	 * Called after the active culture has changed
-	 */
-	void HandleCultureChanged();
-
 	struct FLegacyFontKey
 	{
 		FLegacyFontKey()
@@ -139,12 +119,6 @@ private:
 	TSharedPtr<const FFontData> LocalizedFallbackFontData;
 	TSharedPtr<const FFontData> LastResortFontData;
 
-	TMap<FString, TSharedPtr<const FFontData>> AllLocalizedFallbackFontData;
-	uint16 LocalizedFallbackFontRevision;
-	uint16 LocalizedFallbackFontDataHistoryVersion;
-	uint64 LocalizedFallbackFontFrameCounter;
-
-	FCriticalSection LocalizedFallbackFontDataCS;
 	FCriticalSection LastResortFontCS;
 	FCriticalSection LastResortFontDataCS;
 
diff --git a/Engine/Source/Runtime/SlateCore/Private/Fonts/SlateFontRenderer.cpp b/Engine/Source/Runtime/SlateCore/Private/Fonts/SlateFontRenderer.cpp
--- a/Engine/Source/Runtime/SlateCore/Private/Fonts/SlateFontRenderer.cpp
+++ b/Engine/Source/Runtime/SlateCore/Private/Fonts/SlateFontRenderer.cpp
@@ -342,28 +342,6 @@ FFreeTypeFaceGlyphData FSlateFontRenderer::GetFontFaceForCodepoint(const FFontDa
 		}
 	}
 
-	// If the requested glyph doesn't exist, use the localization fallback font
-	if (!ReturnVal.FaceAndMemory.IsValid() || (InCodepoint != 0 && ReturnVal.GlyphIndex == 0))
-	{
-		const bool bCanFallback = bOverrideFallback || MaxFallbackLevel >= EFontFallback::FF_LocalizedFallback;
-
-		if (bCanFallback && FLegacySlateFontInfoCache::Get().IsLocalizedFallbackFontAvailable())
-		{
-			ReturnVal.FaceAndMemory = CompositeFontCache->GetFontFace(FLegacySlateFontInfoCache::Get().GetLocalizedFallbackFontData(FLegacySlateFontInfoCache::FFallbackContext(&InFontData, InCodepoint)));
-
-			if (ReturnVal.FaceAndMemory.IsValid())
-			{	
-				ReturnVal.GlyphIndex = FT_Get_Char_Index(ReturnVal.FaceAndMemory->GetFace(), InCodepoint);
-
-				if (ReturnVal.GlyphIndex != 0)
-				{
-					ReturnVal.CharFallbackLevel = EFontFallback::FF_LocalizedFallback;
-					ReturnVal.GlyphFlags |= FT_LOAD_FORCE_AUTOHINT;
-				}
-			}
-		}
-	}
-
 	// If the requested glyph doesn't exist, use the last resort fallback font
 	if (!ReturnVal.FaceAndMemory.IsValid() || (InCodepoint != 0 && ReturnVal.GlyphIndex == 0))
 	{
diff --git a/Engine/Source/Runtime/SlateCore/Public/Fonts/FontCache.h b/Engine/Source/Runtime/SlateCore/Public/Fonts/FontCache.h
--- a/Engine/Source/Runtime/SlateCore/Public/Fonts/FontCache.h
+++ b/Engine/Source/Runtime/SlateCore/Public/Fonts/FontCache.h
@@ -810,11 +810,6 @@ public:
 	TArray<FString> GetAvailableFontSubFaces(FFontFaceDataConstRef InMemory) const;
 	TArray<FString> GetAvailableFontSubFaces(const FString& InFilename) const;
 
-	/**
-	 * Get the revision index of the currently active localized fallback font.
-	 */
-	uint16 GetLocalizedFallbackFontRevision() const;
-
 	/**
 	 * Issues a request to clear all cached data from the cache
 	 */
diff --git a/Engine/Source/Runtime/SlateCore/Public/Fonts/SlateFontInfo.h b/Engine/Source/Runtime/SlateCore/Public/Fonts/SlateFontInfo.h
--- a/Engine/Source/Runtime/SlateCore/Public/Fonts/SlateFontInfo.h
+++ b/Engine/Source/Runtime/SlateCore/Public/Fonts/SlateFontInfo.h
@@ -10,17 +10,14 @@
 
 /**
  * Sets the maximum font fallback level, for when a character can't be found in the selected font set.
- *
- * UI code that renders strings from a third party (e.g. player chat in a multiplayer game), should restrict font fallback to localized,
- * (or to no fallback, if international font isn't important), to prevent potential performance problems.
+ * UI code that renders strings from a third party (e.g. player chat in a multiplayer game) may want to restrict font fallback to prevent potential performance problems.
  */
-UENUM()
 enum class EFontFallback : uint8
 {
 	/** No fallback font */
 	FF_NoFallback,
 	/** Fallback to localized font set */
-	FF_LocalizedFallback,
+	FF_LocalizedFallback UE_DEPRECATED(4.24, "Legacy localized fallback fonts have been removed. FF_LocalizedFallback no longer has any meaning, so use FF_NoFallback instead."),
 	/** Fallback to last resort font set */
 	FF_LastResortFallback,
 	/** Tries all fallbacks */
