Introduce TWidgetSparseAnnotation in order to reduce memory footprint of accessibility data in SWidget

diff --git a/Engine/Source/Runtime/Slate/Private/Widgets/Input/SButton.cpp b/Engine/Source/Runtime/Slate/Private/Widgets/Input/SButton.cpp
--- a/Engine/Source/Runtime/Slate/Private/Widgets/Input/SButton.cpp
+++ b/Engine/Source/Runtime/Slate/Private/Widgets/Input/SButton.cpp
@@ -14,7 +14,8 @@ static FName SButtonTypeName("SButton");
 SButton::SButton()
 {
 #if WITH_ACCESSIBILITY
-	AccessibleData = FAccessibleWidgetData(EAccessibleBehavior::Summary, EAccessibleBehavior::Auto, false);
+	AccessibleBehavior = EAccessibleBehavior::Summary;
+	bCanChildrenBeAccessible = false;
 #endif
 }
 
diff --git a/Engine/Source/Runtime/Slate/Private/Widgets/Input/SCheckBox.cpp b/Engine/Source/Runtime/Slate/Private/Widgets/Input/SCheckBox.cpp
--- a/Engine/Source/Runtime/Slate/Private/Widgets/Input/SCheckBox.cpp
+++ b/Engine/Source/Runtime/Slate/Private/Widgets/Input/SCheckBox.cpp
@@ -15,7 +15,8 @@
 SCheckBox::SCheckBox()
 {
 #if WITH_ACCESSIBILITY
-	AccessibleData = FAccessibleWidgetData(EAccessibleBehavior::Summary, EAccessibleBehavior::Auto, false);
+	AccessibleBehavior = EAccessibleBehavior::Summary;
+	bCanChildrenBeAccessible = false;
 #endif
 }
 
diff --git a/Engine/Source/Runtime/Slate/Private/Widgets/Input/SEditableText.cpp b/Engine/Source/Runtime/Slate/Private/Widgets/Input/SEditableText.cpp
--- a/Engine/Source/Runtime/Slate/Private/Widgets/Input/SEditableText.cpp
+++ b/Engine/Source/Runtime/Slate/Private/Widgets/Input/SEditableText.cpp
@@ -12,7 +12,8 @@
 SEditableText::SEditableText()
 {
 #if WITH_ACCESSIBILITY
-	AccessibleData = FAccessibleWidgetData(EAccessibleBehavior::Auto, EAccessibleBehavior::Auto, false);
+	AccessibleBehavior = EAccessibleBehavior::Auto;
+	bCanChildrenBeAccessible = false;
 #endif
 }
 
@@ -617,9 +618,8 @@ TSharedRef<FSlateAccessibleWidget> SEditableText::CreateAccessibleWidget()
 	return MakeShareable<FSlateAccessibleWidget>(new FSlateAccessibleEditableText(SharedThis(this)));
 }
 
-void SEditableText::SetDefaultAccessibleText(EAccessibleType AccessibleType)
+TOptional<FText> SEditableText::GetDefaultAccessibleText(EAccessibleType AccessibleType) const
 {
-	TAttribute<FText>& Text = (AccessibleType == EAccessibleType::Main) ? AccessibleData.AccessibleText : AccessibleData.AccessibleSummaryText;
-	Text.Bind(this, &SEditableText::GetHintText);
+	return GetHintText();
 }
 #endif
diff --git a/Engine/Source/Runtime/Slate/Private/Widgets/Input/SEditableTextBox.cpp b/Engine/Source/Runtime/Slate/Private/Widgets/Input/SEditableTextBox.cpp
--- a/Engine/Source/Runtime/Slate/Private/Widgets/Input/SEditableTextBox.cpp
+++ b/Engine/Source/Runtime/Slate/Private/Widgets/Input/SEditableTextBox.cpp
@@ -12,7 +12,8 @@
 SEditableTextBox::SEditableTextBox()
 {
 #if WITH_ACCESSIBILITY
-	AccessibleData = FAccessibleWidgetData(EAccessibleBehavior::Auto, EAccessibleBehavior::Auto, false);
+	AccessibleBehavior = EAccessibleBehavior::Auto;
+	bCanChildrenBeAccessible = false;
 #endif
 }
 
@@ -90,18 +91,6 @@ void SEditableTextBox::Construct( const FArguments& InArgs )
 			ErrorReporting->AsWidget()
 		];
 	}
-
-#if WITH_ACCESSIBILITY
-	// These functions need to be called again because SWidget::Construct is called before SEditableText is created.
-	if (AccessibleData.AccessibleBehavior != EAccessibleBehavior::Custom)
-	{
-		SetDefaultAccessibleText(EAccessibleType::Main);
-	}
-	if (AccessibleData.AccessibleSummaryBehavior != EAccessibleBehavior::Custom)
-	{
-		SetDefaultAccessibleText(EAccessibleType::Summary);
-	}
-#endif
 }
 
 void SEditableTextBox::SetStyle(const FEditableTextBoxStyle* InStyle)
@@ -425,14 +414,14 @@ TSharedRef<FSlateAccessibleWidget> SEditableTextBox::CreateAccessibleWidget()
 	return MakeShareable<FSlateAccessibleWidget>(new FSlateAccessibleEditableTextBox(SharedThis(this)));
 }
 
-void SEditableTextBox::SetDefaultAccessibleText(EAccessibleType AccessibleType)
+TOptional<FText> SEditableTextBox::GetDefaultAccessibleText(EAccessibleType AccessibleType) const
 {
 	// The parent Construct() function will call this before EditableText exists,
 	// so we need a guard here to ignore that function call.
 	if (EditableText.IsValid())
 	{
-		TAttribute<FText>& Text = (AccessibleType == EAccessibleType::Main) ? AccessibleData.AccessibleText : AccessibleData.AccessibleSummaryText;
-		Text.Bind(EditableText.Get(), &SEditableText::GetHintText);
+		return EditableText->GetHintText();
 	}
+	return TOptional<FText>();
 }
 #endif
diff --git a/Engine/Source/Runtime/Slate/Private/Widgets/Input/SSlider.cpp b/Engine/Source/Runtime/Slate/Private/Widgets/Input/SSlider.cpp
--- a/Engine/Source/Runtime/Slate/Private/Widgets/Input/SSlider.cpp
+++ b/Engine/Source/Runtime/Slate/Private/Widgets/Input/SSlider.cpp
@@ -10,7 +10,8 @@
 SSlider::SSlider()
 {
 #if WITH_ACCESSIBILITY
-	AccessibleData = FAccessibleWidgetData(EAccessibleBehavior::Summary, EAccessibleBehavior::Auto, false);
+	AccessibleBehavior = EAccessibleBehavior::Summary;
+	bCanChildrenBeAccessible = false;
 #endif
 }
 
diff --git a/Engine/Source/Runtime/Slate/Private/Widgets/Text/STextBlock.cpp b/Engine/Source/Runtime/Slate/Private/Widgets/Text/STextBlock.cpp
--- a/Engine/Source/Runtime/Slate/Private/Widgets/Text/STextBlock.cpp
+++ b/Engine/Source/Runtime/Slate/Private/Widgets/Text/STextBlock.cpp
@@ -24,7 +24,8 @@ STextBlock::STextBlock()
 	bSimpleTextMode = false;
 
 #if WITH_ACCESSIBILITY
-	AccessibleData = FAccessibleWidgetData(EAccessibleBehavior::Auto, EAccessibleBehavior::Auto, false);
+	AccessibleBehavior = EAccessibleBehavior::Auto;
+	bCanChildrenBeAccessible = false;
 #endif
 }
 
@@ -472,9 +473,8 @@ TSharedRef<FSlateAccessibleWidget> STextBlock::CreateAccessibleWidget()
 	return MakeShareable<FSlateAccessibleWidget>(new FSlateAccessibleTextBlock(SharedThis(this)));
 }
 
-void STextBlock::SetDefaultAccessibleText(EAccessibleType AccessibleType)
+TOptional<FText> STextBlock::GetDefaultAccessibleText(EAccessibleType AccessibleType) const
 {
-	TAttribute<FText>& Text = (AccessibleType == EAccessibleType::Main) ? AccessibleData.AccessibleText : AccessibleData.AccessibleSummaryText;
-	Text.Bind(this, &STextBlock::GetTextCopy);
+	return GetText();
 }
 #endif
diff --git a/Engine/Source/Runtime/Slate/Public/Widgets/Input/SEditableText.h b/Engine/Source/Runtime/Slate/Public/Widgets/Input/SEditableText.h
--- a/Engine/Source/Runtime/Slate/Public/Widgets/Input/SEditableText.h
+++ b/Engine/Source/Runtime/Slate/Public/Widgets/Input/SEditableText.h
@@ -357,7 +357,7 @@ protected:
 	virtual bool ComputeVolatility() const override;
 #if WITH_ACCESSIBILITY
 	virtual TSharedRef<FSlateAccessibleWidget> CreateAccessibleWidget() override;
-	virtual void SetDefaultAccessibleText(EAccessibleType AccessibleType = EAccessibleType::Main) override;
+	virtual TOptional<FText> GetDefaultAccessibleText(EAccessibleType AccessibleType = EAccessibleType::Main) const override;
 #endif
 	//~ End SWidget Interface
 
diff --git a/Engine/Source/Runtime/Slate/Public/Widgets/Input/SEditableTextBox.h b/Engine/Source/Runtime/Slate/Public/Widgets/Input/SEditableTextBox.h
--- a/Engine/Source/Runtime/Slate/Public/Widgets/Input/SEditableTextBox.h
+++ b/Engine/Source/Runtime/Slate/Public/Widgets/Input/SEditableTextBox.h
@@ -351,7 +351,7 @@ public:
 protected:
 #if WITH_ACCESSIBILITY
 	virtual TSharedRef<FSlateAccessibleWidget> CreateAccessibleWidget() override;
-	virtual void SetDefaultAccessibleText(EAccessibleType AccessibleType = EAccessibleType::Main) override;
+	virtual TOptional<FText> GetDefaultAccessibleText(EAccessibleType AccessibleType = EAccessibleType::Main) const override;
 #endif
 
 	const FEditableTextBoxStyle* Style;
diff --git a/Engine/Source/Runtime/Slate/Public/Widgets/Text/STextBlock.h b/Engine/Source/Runtime/Slate/Public/Widgets/Text/STextBlock.h
--- a/Engine/Source/Runtime/Slate/Public/Widgets/Text/STextBlock.h
+++ b/Engine/Source/Runtime/Slate/Public/Widgets/Text/STextBlock.h
@@ -163,11 +163,6 @@ public:
 	{
 		return BoundText.Get();
 	}
-private:
-	FText GetTextCopy() const
-	{
-		return BoundText.Get();
-	}
 	
 public:
 	UE_DEPRECATED(4.24, "SetText taking FString is deprecated. Use the FText version instead")
@@ -244,7 +239,7 @@ public:
 	virtual FVector2D ComputeDesiredSize(float) const override;
 #if WITH_ACCESSIBILITY
 	virtual TSharedRef<FSlateAccessibleWidget> CreateAccessibleWidget() override;
-	virtual void SetDefaultAccessibleText(EAccessibleType AccessibleType = EAccessibleType::Main) override;
+	virtual TOptional<FText> GetDefaultAccessibleText(EAccessibleType AccessibleType = EAccessibleType::Main) const override;
 #endif
 	// End of SWidget interface
 
diff --git a/Engine/Source/Runtime/SlateCore/Private/Widgets/SWidget.cpp b/Engine/Source/Runtime/SlateCore/Private/Widgets/SWidget.cpp
--- a/Engine/Source/Runtime/SlateCore/Private/Widgets/SWidget.cpp
+++ b/Engine/Source/Runtime/SlateCore/Private/Widgets/SWidget.cpp
@@ -17,6 +17,8 @@
 #include "Debugging/SlateDebugging.h"
 #include "Widgets/SWindow.h"
 #include "Types/ReflectionMetadata.h"
+#include "Misc/ScopeRWLock.h"
+#include "HAL/CriticalSection.h"
 
 #if WITH_ACCESSIBILITY
 #include "Widgets/Accessibility/SlateCoreAccessibleWidgets.h"
@@ -33,6 +35,50 @@ DEFINE_STAT(STAT_SlatePrepass);
 DEFINE_STAT(STAT_SlateTotalWidgets);
 DEFINE_STAT(STAT_SlateSWidgetAllocSize);
 
+template <typename AnnotationType>
+class TWidgetSparseAnnotation
+{
+public:
+	const AnnotationType* Find(const SWidget* Widget)
+	{
+		FRWScopeLock Lock(RWLock, SLT_ReadOnly);
+		return AnnotationMap.Find(Widget);
+	}
+
+	AnnotationType& FindOrAdd(const SWidget* Widget)
+	{
+		FRWScopeLock Lock(RWLock, SLT_Write);
+		return AnnotationMap.FindOrAdd(Widget);
+	}
+
+	void Add(const SWidget* Widget, const AnnotationType& Type)
+	{
+		FRWScopeLock Lock(RWLock, SLT_Write);
+		AnnotationMap.Add(Widget, Type);
+	}
+
+	void Remove(const SWidget* Widget)
+	{
+		FRWScopeLock Lock(RWLock, SLT_Write);
+		AnnotationMap.Remove(Widget);
+	}
+private:
+	TMap<const SWidget*, AnnotationType> AnnotationMap;
+	FRWLock RWLock;
+};
+
+#if WITH_ACCESSIBILITY
+TWidgetSparseAnnotation<TAttribute<FText>> AccessibleText;
+TWidgetSparseAnnotation<TAttribute<FText>> AccessibleSummaryText;
+#endif
+
+static void ClearSparseAnnotationsForWidget(const SWidget* Widget)
+{
+#if WITH_ACCESSIBILITY
+	AccessibleText.Remove(Widget);
+	AccessibleSummaryText.Remove(Widget);
+#endif
+}
 
 #if SLATE_CULL_WIDGETS
 
@@ -158,6 +204,11 @@ SWidget::SWidget()
 	, bUpdatingDesiredSize(false)
 	, bHasCustomPrepass(false)
 	, bVolatilityAlwaysInvalidatesPrepass(false)
+#if WITH_ACCESSIBILITY
+	, bCanChildrenBeAccessible(true)
+	, AccessibleBehavior(EAccessibleBehavior::NotAccessible)
+	, AccessibleSummaryBehavior(EAccessibleBehavior::Auto)
+#endif
 	, Clipping(EWidgetClipping::Inherit)
 	, FlowDirectionPreference(EFlowDirectionPreference::Inherit)
 	// Note we are defaulting to tick for backwards compatibility
@@ -211,6 +262,8 @@ SWidget::~SWidget()
 
 	DEC_DWORD_STAT(STAT_SlateTotalWidgets);
 	DEC_MEMORY_STAT_BY(STAT_SlateSWidgetAllocSize, AllocSize);
+
+	ClearSparseAnnotationsForWidget(this);
 }
 
 void SWidget::Construct(
@@ -266,14 +319,6 @@ void SWidget::Construct(
 		SetAccessibleBehavior(InAccessibleData->AccessibleText.IsSet() ? EAccessibleBehavior::Custom : InAccessibleData->AccessibleBehavior, InAccessibleData->AccessibleText, EAccessibleType::Main);
 		SetAccessibleBehavior(InAccessibleData->AccessibleSummaryText.IsSet() ? EAccessibleBehavior::Custom : InAccessibleData->AccessibleSummaryBehavior, InAccessibleData->AccessibleSummaryText, EAccessibleType::Summary);
 	}
-	if (AccessibleData.AccessibleBehavior != EAccessibleBehavior::Custom)
-	{
-		SetDefaultAccessibleText(EAccessibleType::Main);
-	}
-	if (AccessibleData.AccessibleSummaryBehavior != EAccessibleBehavior::Custom)
-	{
-		SetDefaultAccessibleText(EAccessibleType::Summary);
-	}
 #endif
 }
 
@@ -1550,19 +1595,21 @@ TSharedRef<FSlateAccessibleWidget> SWidget::CreateAccessibleWidget()
 
 void SWidget::SetAccessibleBehavior(EAccessibleBehavior InBehavior, const TAttribute<FText>& InText, EAccessibleType AccessibleType)
 {
-	EAccessibleBehavior& Behavior = (AccessibleType == EAccessibleType::Main) ? AccessibleData.AccessibleBehavior : AccessibleData.AccessibleSummaryBehavior;
+	EAccessibleBehavior& Behavior = (AccessibleType == EAccessibleType::Main) ? AccessibleBehavior : AccessibleSummaryBehavior;
+
+	if (InBehavior == EAccessibleBehavior::Custom)
+	{
+		TWidgetSparseAnnotation<TAttribute<FText>>& AccessibleTextAnnotation = (AccessibleType == EAccessibleType::Main) ? AccessibleText : AccessibleSummaryText;
+		AccessibleTextAnnotation.FindOrAdd(this) = InText;
+	}
+	else if (Behavior == EAccessibleBehavior::Custom)
+	{
+		TWidgetSparseAnnotation<TAttribute<FText>>& AccessibleTextAnnotation = (AccessibleType == EAccessibleType::Main) ? AccessibleText : AccessibleSummaryText;
+		AccessibleTextAnnotation.Remove(this);
+	}
+
 	if (Behavior != InBehavior)
 	{
-		// If switching off of custom, revert back to default text
-		if (Behavior == EAccessibleBehavior::Custom)
-		{
-			SetDefaultAccessibleText(AccessibleType);
-		}
-		else if (InBehavior == EAccessibleBehavior::Custom)
-		{
-			TAttribute<FText>& Text = (AccessibleType == EAccessibleType::Main) ? AccessibleData.AccessibleText : AccessibleData.AccessibleSummaryText;
-			Text = InText;
-		}
 		const bool bWasAccessible = Behavior != EAccessibleBehavior::NotAccessible;
 		Behavior = InBehavior;
 		if (AccessibleType == EAccessibleType::Main && bWasAccessible != (Behavior != EAccessibleBehavior::NotAccessible))
@@ -1574,30 +1621,25 @@ void SWidget::SetAccessibleBehavior(EAccessibleBehavior InBehavior, const TAttri
 
 void SWidget::SetCanChildrenBeAccessible(bool InCanChildrenBeAccessible)
 {
-	if (AccessibleData.bCanChildrenBeAccessible != InCanChildrenBeAccessible)
+	if (bCanChildrenBeAccessible != InCanChildrenBeAccessible)
 	{
-		AccessibleData.bCanChildrenBeAccessible = InCanChildrenBeAccessible;
+		bCanChildrenBeAccessible = InCanChildrenBeAccessible;
 		FSlateApplicationBase::Get().GetAccessibleMessageHandler()->MarkDirty();
 	}
 }
 
-void SWidget::SetDefaultAccessibleText(EAccessibleType AccessibleType)
-{
-	TAttribute<FText>& Text = (AccessibleType == EAccessibleType::Main) ? AccessibleData.AccessibleText : AccessibleData.AccessibleSummaryText;
-	Text = TAttribute<FText>();
-}
-
 FText SWidget::GetAccessibleText(EAccessibleType AccessibleType) const
 {
-	const EAccessibleBehavior Behavior = (AccessibleType == EAccessibleType::Main) ? AccessibleData.AccessibleBehavior : AccessibleData.AccessibleSummaryBehavior;
-	const EAccessibleBehavior OtherBehavior = (AccessibleType == EAccessibleType::Main) ? AccessibleData.AccessibleSummaryBehavior : AccessibleData.AccessibleBehavior;
-	const TAttribute<FText>& Text = (AccessibleType == EAccessibleType::Main) ? AccessibleData.AccessibleText : AccessibleData.AccessibleSummaryText;
-	const TAttribute<FText>& OtherText = (AccessibleType == EAccessibleType::Main) ? AccessibleData.AccessibleSummaryText : AccessibleData.AccessibleText;
+	const EAccessibleBehavior Behavior = (AccessibleType == EAccessibleType::Main) ? AccessibleBehavior : AccessibleSummaryBehavior;
+	const EAccessibleBehavior OtherBehavior = (AccessibleType == EAccessibleType::Main) ? AccessibleSummaryBehavior : AccessibleBehavior;
 
 	switch (Behavior)
 	{
 	case EAccessibleBehavior::Custom:
-		return Text.Get(FText::GetEmpty());
+	{
+		const TAttribute<FText>* Text = (AccessibleType == EAccessibleType::Main) ? AccessibleText.Find(this) : AccessibleSummaryText.Find(this);
+		return Text->Get(FText::GetEmpty());
+	}
 	case EAccessibleBehavior::Summary:
 		return GetAccessibleSummary();
 	case EAccessibleBehavior::ToolTip:
@@ -1611,9 +1653,10 @@ FText SWidget::GetAccessibleText(EAccessibleType AccessibleType) const
 		// used instead in that case - however, this will be used for widgets with special default text such as TextBlocks.
 		// If no text is found, then it will attempt to use the other variable's text, so that a developer can do things like
 		// leave Summary on Auto, set Main to Custom, and have Summary automatically use Main's value without having to re-type it.
-		if (Text.IsSet())
+		TOptional<FText> DefaultText = GetDefaultAccessibleText(AccessibleType);
+		if (DefaultText.IsSet())
 		{
-			return Text.Get(FText::GetEmpty());
+			return DefaultText.GetValue();
 		}
 		switch (OtherBehavior)
 		{
@@ -1629,6 +1672,11 @@ FText SWidget::GetAccessibleText(EAccessibleType AccessibleType) const
 	return FText::GetEmpty();
 }
 
+TOptional<FText> SWidget::GetDefaultAccessibleText(EAccessibleType AccessibleType) const
+{
+	return TOptional<FText>();
+}
+
 FText SWidget::GetAccessibleSummary() const
 {
 	FTextBuilder Builder;
@@ -1649,7 +1697,7 @@ FText SWidget::GetAccessibleSummary() const
 
 bool SWidget::IsAccessible() const
 {
-	if (AccessibleData.AccessibleBehavior == EAccessibleBehavior::NotAccessible)
+	if (AccessibleBehavior == EAccessibleBehavior::NotAccessible)
 	{
 		return false;
 	}
@@ -1668,12 +1716,12 @@ bool SWidget::IsAccessible() const
 
 EAccessibleBehavior SWidget::GetAccessibleBehavior(EAccessibleType AccessibleType) const
 {
-	return AccessibleType == EAccessibleType::Main ? AccessibleData.AccessibleBehavior : AccessibleData.AccessibleSummaryBehavior;
+	return AccessibleType == EAccessibleType::Main ? AccessibleBehavior : AccessibleSummaryBehavior;
 }
 
 bool SWidget::CanChildrenBeAccessible() const
 {
-	return AccessibleData.bCanChildrenBeAccessible;
+	return bCanChildrenBeAccessible;
 }
 
 #endif
diff --git a/Engine/Source/Runtime/SlateCore/Private/Widgets/SWindow.cpp b/Engine/Source/Runtime/SlateCore/Private/Widgets/SWindow.cpp
--- a/Engine/Source/Runtime/SlateCore/Private/Widgets/SWindow.cpp
+++ b/Engine/Source/Runtime/SlateCore/Private/Widgets/SWindow.cpp
@@ -1989,7 +1989,7 @@ SWindow::SWindow()
 	SetInvalidationRootHittestGrid(*HittestGrid);
 
 #if WITH_ACCESSIBILITY
-	AccessibleData = FAccessibleWidgetData(EAccessibleBehavior::Auto);
+	AccessibleBehavior = EAccessibleBehavior::Auto;
 #endif
 }
 
@@ -2250,9 +2250,8 @@ TSharedRef<FSlateAccessibleWidget> SWindow::CreateAccessibleWidget()
 	return MakeShareable<FSlateAccessibleWidget>(new FSlateAccessibleWindow(SharedThis(this)));
 }
 
-void SWindow::SetDefaultAccessibleText(EAccessibleType AccessibleType)
+TOptional<FText> SWindow::GetDefaultAccessibleText(EAccessibleType AccessibleType) const
 {
-	TAttribute<FText>& Text = (AccessibleType == EAccessibleType::Main) ? AccessibleData.AccessibleText : AccessibleData.AccessibleSummaryText;
-	Text.Bind(this, &SWindow::GetTitle);
+	return GetTitle();
 }
 #endif
diff --git a/Engine/Source/Runtime/SlateCore/Public/Widgets/SWidget.h b/Engine/Source/Runtime/SlateCore/Public/Widgets/SWidget.h
--- a/Engine/Source/Runtime/SlateCore/Public/Widgets/SWidget.h
+++ b/Engine/Source/Runtime/SlateCore/Public/Widgets/SWidget.h
@@ -984,7 +984,7 @@ public:
 	 *
 	 * @param AccessibleType Whether the widget is being accessed directly or through a summary query.
 	 */
-	virtual void SetDefaultAccessibleText(EAccessibleType AccessibleType = EAccessibleType::Main);
+	virtual TOptional<FText> GetDefaultAccessibleText(EAccessibleType AccessibleType = EAccessibleType::Main) const;
 #endif
 
 	/** Whether or not a widget is volatile and will update every frame without being invalidated */
@@ -1574,6 +1574,13 @@ protected:
 	uint8 bCanTick : 1;
 #endif
 
+#if WITH_ACCESSIBILITY
+	/** All variables surrounding how this widget is exposed to the platform's accessibility API. */
+	uint8 bCanChildrenBeAccessible : 1;
+	EAccessibleBehavior AccessibleBehavior;
+	EAccessibleBehavior AccessibleSummaryBehavior;
+#endif
+
 	/**
 	 * Set to true if all content of the widget should clip to the bounds of this widget.
 	 */
@@ -1642,12 +1649,6 @@ protected:
 	/** Render transform pivot of this widget (in normalized local space) */
 	TAttribute< FVector2D > RenderTransformPivot;
 
-
-#if WITH_ACCESSIBILITY
-	/** All variables surrounding how this widget is exposed to the platform's accessibility API. */
-	FAccessibleWidgetData AccessibleData;
-#endif
-
 	/** Debugging information on the type of widget we're creating for the Widget Reflector. */
 	FName TypeOfWidget;
 
diff --git a/Engine/Source/Runtime/SlateCore/Public/Widgets/SWindow.h b/Engine/Source/Runtime/SlateCore/Public/Widgets/SWindow.h
--- a/Engine/Source/Runtime/SlateCore/Public/Widgets/SWindow.h
+++ b/Engine/Source/Runtime/SlateCore/Public/Widgets/SWindow.h
@@ -803,7 +803,7 @@ public:
 
 #if WITH_ACCESSIBILITY
 	virtual TSharedRef<FSlateAccessibleWidget> CreateAccessibleWidget() override;
-	virtual void SetDefaultAccessibleText(EAccessibleType AccessibleType = EAccessibleType::Main) override;
+	virtual TOptional<FText> GetDefaultAccessibleText(EAccessibleType AccessibleType = EAccessibleType::Main) const override;
 #endif
 private:
 	virtual FReply OnFocusReceived( const FGeometry& MyGeometry, const FFocusEvent& InFocusEvent ) override;
