Move ClearSparseAnnotations inside FSlateApplication::IsInitialized block to work around static SNullWidget destructor order

diff --git a/Engine/Source/Runtime/SlateCore/Private/Widgets/SWidget.cpp b/Engine/Source/Runtime/SlateCore/Private/Widgets/SWidget.cpp
--- a/Engine/Source/Runtime/SlateCore/Private/Widgets/SWidget.cpp
+++ b/Engine/Source/Runtime/SlateCore/Private/Widgets/SWidget.cpp
@@ -258,12 +258,12 @@ SWidget::~SWidget()
 #if WITH_ACCESSIBILITY
 		FSlateApplicationBase::Get().GetAccessibleMessageHandler()->OnWidgetRemoved(this);
 #endif
+		// Only clear if initialized because SNullWidget's destructor may be called after annotations are deleted
+		ClearSparseAnnotationsForWidget(this);
 	}
 
 	DEC_DWORD_STAT(STAT_SlateTotalWidgets);
 	DEC_MEMORY_STAT_BY(STAT_SlateSWidgetAllocSize, AllocSize);
-
-	ClearSparseAnnotationsForWidget(this);
 }
 
 void SWidget::Construct(
