Slate - Adding CSV profiling timers and counters to some key areas in Slate for performance reports.  Also adding an event that fires when something invalidates more than a command line parameter level of parents, currently set to 5.  I think I want it to be smarter though, some way to marking certain widgets as generating events when they're invalidated.

Andrew Ladenberger, Matt Kuhlenschmidt

diff --git a/Engine/Source/Runtime/Slate/Private/Framework/Application/SlateApplication.cpp b/Engine/Source/Runtime/Slate/Private/Framework/Application/SlateApplication.cpp
--- a/Engine/Source/Runtime/Slate/Private/Framework/Application/SlateApplication.cpp
+++ b/Engine/Source/Runtime/Slate/Private/Framework/Application/SlateApplication.cpp
@@ -53,6 +53,8 @@
 
 CSV_DECLARE_CATEGORY_MODULE_EXTERN(CORE_API, Basic);
 
+CSV_DECLARE_CATEGORY_MODULE_EXTERN(SLATECORE_API, Slate);
+
 //////////////////////////////////////////////////////////////////////////
 
 /** 
@@ -1117,6 +1119,7 @@ static void PrepassWindowAndChildren( TSharedRef<SWindow> WindowToPrepass )
 void FSlateApplication::DrawPrepass( TSharedPtr<SWindow> DrawOnlyThisWindow )
 {
 	SCOPED_NAMED_EVENT_TEXT("Slate::Prepass", FColor::Magenta);
+	CSV_SCOPED_TIMING_STAT(Slate, Prepass);
 
 	TSharedPtr<SWindow> ActiveModalWindow = GetActiveModalWindow();
 
@@ -1337,6 +1340,7 @@ void FSlateApplication::Tick(ESlateTickType TickType)
 	FScopeLock SlateTickAccess(&SlateTickCriticalSection);
 
 	SCOPED_NAMED_EVENT_TEXT("Slate::Tick", FColor::Magenta);
+	CSV_SCOPED_TIMING_STAT(Slate, Tick);
 
 	{
 		SCOPE_CYCLE_COUNTER(STAT_SlateTickTime);
diff --git a/Engine/Source/Runtime/SlateCore/Private/Debugging/SlateDebugging.cpp b/Engine/Source/Runtime/SlateCore/Private/Debugging/SlateDebugging.cpp
--- a/Engine/Source/Runtime/SlateCore/Private/Debugging/SlateDebugging.cpp
+++ b/Engine/Source/Runtime/SlateCore/Private/Debugging/SlateDebugging.cpp
@@ -12,6 +12,9 @@
 #include "Styling/CoreStyle.h"
 #include "Widgets/SWidget.h"
 #include "Application/SlateApplicationBase.h"
+#include "ProfilingDebugging/CsvProfiler.h"
+
+CSV_DEFINE_CATEGORY_MODULE(SLATECORE_API, Slate, true);
 
 #define LOCTEXT_NAMESPACE "SlateDebugger"
 
diff --git a/Engine/Source/Runtime/SlateCore/Private/FastUpdate/SlateInvalidationRoot.cpp b/Engine/Source/Runtime/SlateCore/Private/FastUpdate/SlateInvalidationRoot.cpp
--- a/Engine/Source/Runtime/SlateCore/Private/FastUpdate/SlateInvalidationRoot.cpp
+++ b/Engine/Source/Runtime/SlateCore/Private/FastUpdate/SlateInvalidationRoot.cpp
@@ -6,12 +6,14 @@
 #include "Widgets/SWidget.h"
 #include "Layout/Children.h"
 #include "ProfilingDebugging/CsvProfiler.h"
+#include "Types/ReflectionMetadata.h"
+
+CSV_DECLARE_CATEGORY_EXTERN(Slate);
 
 #if WITH_SLATE_DEBUGGING
 bool GDumpUpdateList = false;
 void HandleDumpUpdateList(const TArray<FString>& Args)
 {
-
 	GDumpUpdateList = true;
 }
 
@@ -22,6 +24,13 @@ static FAutoConsoleCommand HandleDumpUpdateListCommand(
 );
 #endif
 
+#if SLATE_CSV_TRACKER
+	static int32 CascadeInvalidationEventAmount = 5;
+	FAutoConsoleVariableRef CVarCascadeInvalidationEventAmount(
+		TEXT("Slate.CSV.CascadeInvalidationEventAmount"),
+		CascadeInvalidationEventAmount,
+		TEXT("The amount of cascaded invalidated parents before we fire a CSV event."));
+#endif
 
 FSlateInvalidationRoot::FSlateInvalidationRoot()
 	: CachedElementData(new FSlateCachedElementData)
@@ -268,6 +277,7 @@ void FSlateInvalidationRoot::OnWidgetDestroyed(const SWidget* Widget)
 bool FSlateInvalidationRoot::PaintFastPath(const FSlateInvalidationContext& Context)
 {
 	SCOPED_NAMED_EVENT(SWidget_FastPathUpdate, FColor::Green);
+	CSV_SCOPED_TIMING_STAT(Slate, FastPathUpdate);
 
 	check(!bNeedsSlowPath);
 
@@ -280,7 +290,7 @@ bool FSlateInvalidationRoot::PaintFastPath(const FSlateInvalidationContext& Cont
 #if WITH_SLATE_DEBUGGING
 		if (GDumpUpdateList)
 		{
-			UE_LOG(LogSlate, Log, TEXT("Dumping update list"));
+			UE_LOG(LogSlate, Log, TEXT("Dumping Update List"));
 
 			// The update list is put in reverse order 
 			for (int32 ListIndex = FinalUpdateList.Num() - 1; ListIndex >= 0; --ListIndex)
@@ -291,18 +301,18 @@ bool FSlateInvalidationRoot::PaintFastPath(const FSlateInvalidationContext& Cont
 
 				if (EnumHasAnyFlags(WidgetProxy.UpdateFlags, EWidgetUpdateFlags::NeedsRepaint | EWidgetUpdateFlags::NeedsVolatilePaint))
 				{
-					UE_LOG(LogSlate, Log, TEXT("Repaint %s"), *WidgetProxy.Widget->GetTypeAsString());
+					UE_LOG(LogSlate, Log, TEXT("Repaint %s"), *FReflectionMetaData::GetWidgetDebugInfo(WidgetProxy.Widget));
 				}
 				else if (!WidgetProxy.bInvisibleDueToParentOrSelfVisibility)
 				{
 					if (EnumHasAnyFlags(WidgetProxy.UpdateFlags, EWidgetUpdateFlags::NeedsActiveTimerUpdate))
 					{
-						UE_LOG(LogSlate, Log, TEXT("ActiveTimer %s"), *WidgetProxy.Widget->GetTypeAsString());
+						UE_LOG(LogSlate, Log, TEXT("ActiveTimer %s"), *FReflectionMetaData::GetWidgetDebugInfo(WidgetProxy.Widget));
 					}
 
 					if (EnumHasAnyFlags(WidgetProxy.UpdateFlags, EWidgetUpdateFlags::NeedsTick))
 					{
-						UE_LOG(LogSlate, Log, TEXT("Tick %s"), *WidgetProxy.Widget->GetTypeAsString());
+						UE_LOG(LogSlate, Log, TEXT("Tick %s"), *FReflectionMetaData::GetWidgetDebugInfo(WidgetProxy.Widget));
 					}
 				}
 			}
@@ -457,6 +467,7 @@ void FSlateInvalidationRoot::BuildFastPathList(SWidget* RootWidget)
 bool FSlateInvalidationRoot::ProcessInvalidation()
 {
 	SCOPED_NAMED_EVENT(Slate_InvalidationProcessing, FColor::Blue);
+	CSV_SCOPED_TIMING_STAT(Slate, InvalidationProcessing);
 
 	FinalUpdateList.Empty();
 
@@ -512,8 +523,21 @@ bool FSlateInvalidationRoot::ProcessInvalidation()
 			bChildOrderInvalidated = false;
 		}
 
+#if SLATE_CSV_TRACKER
+		FCsvProfiler::RecordCustomStat("Invalidate/InitialWidgets", CSV_CATEGORY_INDEX(Slate), WidgetsNeedingUpdate.Num(), ECsvCustomStatOp::Set);
+		int32 Stat_TotalWidgetsInvalidated = 0;
+		int32 Stat_NeedsRepaint = 0;
+		int32 Stat_NeedsVolatilePaint = 0;
+		int32 Stat_NeedsTick = 0;
+		int32 Stat_NeedsActiveTimerUpdate = 0;
+#endif
+
 		while (WidgetsNeedingUpdate.Num() && !bNeedsSlowPath)
 		{
+#if SLATE_CSV_TRACKER
+			Stat_TotalWidgetsInvalidated++;
+#endif
+
 			int32 MyIndex = WidgetsNeedingUpdate.Pop();
 			FinalUpdateList.Add(MyIndex);
 			FWidgetProxy& WidgetProxy = FastWidgetPathList[MyIndex];
@@ -530,9 +554,40 @@ bool FSlateInvalidationRoot::ProcessInvalidation()
 				{
 					WidgetProxy.CurrentInvalidateReason |= EInvalidateWidget::Layout;
 				}
-			
+				
+#if SLATE_CSV_TRACKER
+				const int32 PreviousWidgetsNeedingUpdating = WidgetsNeedingUpdate.Num();
+#endif
+
 				bWidgetsNeedRepaint |= WidgetProxy.ProcessInvalidation(WidgetsNeedingUpdate, FastWidgetPathList, *this);
 
+#if SLATE_CSV_TRACKER
+				const int32 CurrentWidgetsNeedingUpdating = WidgetsNeedingUpdate.Num();
+				const int32 AddedWidgets = CurrentWidgetsNeedingUpdating - PreviousWidgetsNeedingUpdating;
+
+				if (AddedWidgets >= CascadeInvalidationEventAmount)
+				{
+					CSV_EVENT(Slate, TEXT("%s"), *FReflectionMetaData::GetWidgetDebugInfo(WidgetProxy.Widget));
+				}
+
+				if (EnumHasAnyFlags(WidgetProxy.UpdateFlags, EWidgetUpdateFlags::NeedsRepaint))
+				{
+					Stat_NeedsRepaint++;
+				}
+				if (EnumHasAnyFlags(WidgetProxy.UpdateFlags, EWidgetUpdateFlags::NeedsVolatilePaint) && !WidgetProxy.Widget->Advanced_IsInvalidationRoot())
+				{
+					Stat_NeedsVolatilePaint++;
+				}
+				if (EnumHasAnyFlags(WidgetProxy.UpdateFlags, EWidgetUpdateFlags::NeedsTick))
+				{
+					Stat_NeedsTick++;
+				}
+				if (EnumHasAnyFlags(WidgetProxy.UpdateFlags, EWidgetUpdateFlags::NeedsActiveTimerUpdate))
+				{
+					Stat_NeedsActiveTimerUpdate++;
+				}
+#endif
+
 #if WITH_SLATE_DEBUGGING
 				if (GSlateInvalidationDebugging)
 				{
@@ -558,6 +613,14 @@ bool FSlateInvalidationRoot::ProcessInvalidation()
 		}
 
 		WidgetsNeedingUpdate.Reset();
+
+#if SLATE_CSV_TRACKER
+		FCsvProfiler::RecordCustomStat("Invalidate/TotalWidgets", CSV_CATEGORY_INDEX(Slate), Stat_TotalWidgetsInvalidated, ECsvCustomStatOp::Set);
+		FCsvProfiler::RecordCustomStat("Invalidate/NeedsRepaint", CSV_CATEGORY_INDEX(Slate), Stat_NeedsRepaint, ECsvCustomStatOp::Set);
+		FCsvProfiler::RecordCustomStat("Invalidate/NeedsVolatilePaint", CSV_CATEGORY_INDEX(Slate), Stat_NeedsVolatilePaint, ECsvCustomStatOp::Set);
+		FCsvProfiler::RecordCustomStat("Invalidate/NeedsTick", CSV_CATEGORY_INDEX(Slate), Stat_NeedsTick, ECsvCustomStatOp::Set);
+		FCsvProfiler::RecordCustomStat("Invalidate/NeedsActiveTimerUpdate", CSV_CATEGORY_INDEX(Slate), Stat_NeedsActiveTimerUpdate, ECsvCustomStatOp::Set);
+#endif
 	}
 	else
 	{
diff --git a/Engine/Source/Runtime/SlateCore/Private/Rendering/ElementBatcher.cpp b/Engine/Source/Runtime/SlateCore/Private/Rendering/ElementBatcher.cpp
--- a/Engine/Source/Runtime/SlateCore/Private/Rendering/ElementBatcher.cpp
+++ b/Engine/Source/Runtime/SlateCore/Private/Rendering/ElementBatcher.cpp
@@ -8,6 +8,9 @@
 #include "Widgets/SWindow.h"
 #include "HAL/IConsoleManager.h"
 #include "ProfilingDebugging/ScopedTimers.h"
+#include "Debugging/SlateDebugging.h"
+
+CSV_DECLARE_CATEGORY_EXTERN(Slate);
 
 DECLARE_CYCLE_STAT(TEXT("Slate RT: Create Batches"), STAT_SlateRTCreateBatches, STATGROUP_Slate);
 
@@ -549,6 +552,12 @@ void FSlateElementBatcher::AddElementsInternal(const FSlateDrawElementArray& Dra
 
 void FSlateElementBatcher::AddCachedElements(FSlateCachedElementData& CachedElementData, const FVector2D& ViewportSize)
 {
+	CSV_SCOPED_TIMING_STAT(Slate, AddCachedElements);
+
+#if SLATE_CSV_TRACKER
+	FCsvProfiler::RecordCustomStat("Paint/CachedElementLists", CSV_CATEGORY_INDEX(Slate), CachedElementData.CachedElementLists.Num(), ECsvCustomStatOp::Set);
+#endif
+
 #define ALLOW_CACHED_RENDER_BATCHES 1
 	SCOPED_NAMED_EVENT_TEXT("Slate::AddCachedBatches", FColor::Magenta);
 	for ( FSlateCachedElementList& LocalCachedElementList : CachedElementData.CachedElementLists)
@@ -578,6 +587,10 @@ void FSlateElementBatcher::AddCachedElements(FSlateCachedElementData& CachedElem
 #endif
 	}
 	CachedElementData.CleanupUnusedClipStates();
+
+#if SLATE_CSV_TRACKER
+	FCsvProfiler::RecordCustomStat("Paint/CachedElements", CSV_CATEGORY_INDEX(Slate), ElementStat_CachedElements, ECsvCustomStatOp::Set);
+#endif
 }
 
 template<ESlateVertexRounding Rounding>
diff --git a/Engine/Source/Runtime/SlateCore/Public/Debugging/SlateDebugging.h b/Engine/Source/Runtime/SlateCore/Public/Debugging/SlateDebugging.h
--- a/Engine/Source/Runtime/SlateCore/Public/Debugging/SlateDebugging.h
+++ b/Engine/Source/Runtime/SlateCore/Public/Debugging/SlateDebugging.h
@@ -5,6 +5,7 @@
 #include "CoreMinimal.h"
 #include "Delegates/Delegate.h"
 #include "Input/Reply.h"
+#include "ProfilingDebugging/CsvProfiler.h"
 
 #include "SlateDebugging.generated.h"
 
@@ -12,6 +13,10 @@
 	#define WITH_SLATE_DEBUGGING !(UE_BUILD_SHIPPING || UE_BUILD_TEST)
 #endif
 
+#ifndef SLATE_CSV_TRACKER
+	#define SLATE_CSV_TRACKER CSV_PROFILER
+#endif
+
 class SWindow;
 class SWidget;
 struct FGeometry;
