Slate - Fixing where the accumulated stat comes from so it works with stats turned off.

diff --git a/Engine/Source/Runtime/SlateCore/Private/Rendering/ElementBatcher.cpp b/Engine/Source/Runtime/SlateCore/Private/Rendering/ElementBatcher.cpp
--- a/Engine/Source/Runtime/SlateCore/Private/Rendering/ElementBatcher.cpp
+++ b/Engine/Source/Runtime/SlateCore/Private/Rendering/ElementBatcher.cpp
@@ -556,6 +556,7 @@ void FSlateElementBatcher::AddCachedElements(FSlateCachedElementData& CachedElem
 
 #if SLATE_CSV_TRACKER
 	FCsvProfiler::RecordCustomStat("Paint/CachedElementLists", CSV_CATEGORY_INDEX(Slate), CachedElementData.CachedElementLists.Num(), ECsvCustomStatOp::Set);
+	int32 CachedDrawElements = 0;
 #endif
 
 #define ALLOW_CACHED_RENDER_BATCHES 1
@@ -563,6 +564,11 @@ void FSlateElementBatcher::AddCachedElements(FSlateCachedElementData& CachedElem
 	for ( FSlateCachedElementList& LocalCachedElementList : CachedElementData.CachedElementLists)
 	{
 		STAT(ElementStat_CachedElements += LocalCachedElementList.DrawElements.Num());
+
+#if SLATE_CSV_TRACKER
+		CachedDrawElements += LocalCachedElementList.DrawElements.Num();
+#endif
+
 		LocalCachedElementList.bNewData = false;
 #if ALLOW_CACHED_RENDER_BATCHES
 		if (LocalCachedElementList.CachedBatches.Num())
@@ -589,7 +595,7 @@ void FSlateElementBatcher::AddCachedElements(FSlateCachedElementData& CachedElem
 	CachedElementData.CleanupUnusedClipStates();
 
 #if SLATE_CSV_TRACKER
-	FCsvProfiler::RecordCustomStat("Paint/CachedElements", CSV_CATEGORY_INDEX(Slate), ElementStat_CachedElements, ECsvCustomStatOp::Set);
+	FCsvProfiler::RecordCustomStat("Paint/CachedElements", CSV_CATEGORY_INDEX(Slate), CachedDrawElements, ECsvCustomStatOp::Accumulate);
 #endif
 }
 
