Linux Build Fix.  Slate - Refrence the slate CSV define as module extern, even inside its own module.

diff --git a/Engine/Source/Runtime/SlateCore/Private/FastUpdate/SlateInvalidationRoot.cpp b/Engine/Source/Runtime/SlateCore/Private/FastUpdate/SlateInvalidationRoot.cpp
--- a/Engine/Source/Runtime/SlateCore/Private/FastUpdate/SlateInvalidationRoot.cpp
+++ b/Engine/Source/Runtime/SlateCore/Private/FastUpdate/SlateInvalidationRoot.cpp
@@ -8,7 +8,7 @@
 #include "ProfilingDebugging/CsvProfiler.h"
 #include "Types/ReflectionMetadata.h"
 
-CSV_DECLARE_CATEGORY_EXTERN(Slate);
+CSV_DECLARE_CATEGORY_MODULE_EXTERN(SLATECORE_API, Slate);
 
 #if WITH_SLATE_DEBUGGING
 bool GDumpUpdateList = false;
diff --git a/Engine/Source/Runtime/SlateCore/Private/Rendering/ElementBatcher.cpp b/Engine/Source/Runtime/SlateCore/Private/Rendering/ElementBatcher.cpp
--- a/Engine/Source/Runtime/SlateCore/Private/Rendering/ElementBatcher.cpp
+++ b/Engine/Source/Runtime/SlateCore/Private/Rendering/ElementBatcher.cpp
@@ -10,7 +10,7 @@
 #include "ProfilingDebugging/ScopedTimers.h"
 #include "Debugging/SlateDebugging.h"
 
-CSV_DECLARE_CATEGORY_EXTERN(Slate);
+CSV_DECLARE_CATEGORY_MODULE_EXTERN(SLATECORE_API, Slate);
 
 DECLARE_CYCLE_STAT(TEXT("Slate RT: Create Batches"), STAT_SlateRTCreateBatches, STATGROUP_Slate);
 
