Added SCOPE_CYCLE_SWIDGET to mirror SCOPE_CYCLE_UOBJECT and get widget names from slate scope counters.  Added scopes for Invalidation and RecacheElements events

diff --git a/Engine/Source/Runtime/SlateCore/Private/FastUpdate/WidgetProxy.cpp b/Engine/Source/Runtime/SlateCore/Private/FastUpdate/WidgetProxy.cpp
--- a/Engine/Source/Runtime/SlateCore/Private/FastUpdate/WidgetProxy.cpp
+++ b/Engine/Source/Runtime/SlateCore/Private/FastUpdate/WidgetProxy.cpp
@@ -7,6 +7,7 @@
 #include "ProfilingDebugging/CsvProfiler.h"
 #include "Types/ReflectionMetadata.h"
 #include "Input/HittestGrid.h"
+#include "Widgets/SWidgetUtils.h"
 
 const FSlateWidgetPersistentState FSlateWidgetPersistentState::NoState;
 
@@ -64,6 +65,7 @@ bool FWidgetProxy::ProcessInvalidation(FWidgetUpdateList& UpdateList, TArray<FWi
 	bool bWidgetNeedsRepaint = false;
 	if (!bInvisibleDueToParentOrSelfVisibility && ParentIndex != INDEX_NONE && !Widget->PrepassLayoutScaleMultiplier.IsSet())
 	{
+		SCOPE_CYCLE_SWIDGET(WidgetScope, Widget);
 		// If this widget has never been prepassed make sure the parent prepasses it to set the correct multiplier
 		FWidgetProxy& ParentProxy = FastWidgetPathList[ParentIndex];
 		if (ParentProxy.Widget)
@@ -77,6 +79,7 @@ bool FWidgetProxy::ProcessInvalidation(FWidgetUpdateList& UpdateList, TArray<FWi
 	}
 	else if (EnumHasAnyFlags(CurrentInvalidateReason, EInvalidateWidget::RenderTransform | EInvalidateWidget::Layout | EInvalidateWidget::Visibility | EInvalidateWidget::ChildOrder))
 	{
+		SCOPE_CYCLE_SWIDGET(WidgetScope, Widget);
 		// When layout changes compute a new desired size for this widget
 		FVector2D CurrentDesiredSize = Widget->GetDesiredSize();
 		FVector2D NewDesiredSize = FVector2D::ZeroVector;
@@ -131,6 +134,7 @@ bool FWidgetProxy::ProcessInvalidation(FWidgetUpdateList& UpdateList, TArray<FWi
 	}
 	else if (EnumHasAnyFlags(CurrentInvalidateReason, EInvalidateWidget::Paint) && !Widget->IsVolatileIndirectly())
 	{
+		SCOPE_CYCLE_SWIDGET(WidgetScope, Widget);
 		UpdateFlags |= EWidgetUpdateFlags::NeedsRepaint;
 
 		bWidgetNeedsRepaint = true;
diff --git a/Engine/Source/Runtime/SlateCore/Private/Rendering/ElementBatcher.cpp b/Engine/Source/Runtime/SlateCore/Private/Rendering/ElementBatcher.cpp
--- a/Engine/Source/Runtime/SlateCore/Private/Rendering/ElementBatcher.cpp
+++ b/Engine/Source/Runtime/SlateCore/Private/Rendering/ElementBatcher.cpp
@@ -9,6 +9,7 @@
 #include "HAL/IConsoleManager.h"
 #include "ProfilingDebugging/ScopedTimers.h"
 #include "Debugging/SlateDebugging.h"
+#include "Widgets/SWidgetUtils.h"
 
 CSV_DECLARE_CATEGORY_MODULE_EXTERN(SLATECORE_API, Slate);
 
@@ -579,7 +580,7 @@ void FSlateElementBatcher::AddCachedElements(FSlateCachedElementData& CachedElem
 		{
 			CachedElementList = &LocalCachedElementList;
 			{
-				SCOPED_NAMED_EVENT_TEXT("Slate::RecacheElements", FColor::Magenta);
+				SCOPE_CYCLE_SWIDGET(WidgetScope, LocalCachedElementList.Widget);
 				AddElementsInternal(LocalCachedElementList.DrawElements, ViewportSize);
 			}
 
diff --git a/Engine/Source/Runtime/SlateCore/Private/Widgets/SWidget.cpp b/Engine/Source/Runtime/SlateCore/Private/Widgets/SWidget.cpp
--- a/Engine/Source/Runtime/SlateCore/Private/Widgets/SWidget.cpp
+++ b/Engine/Source/Runtime/SlateCore/Private/Widgets/SWidget.cpp
@@ -17,6 +17,8 @@
 #include "Debugging/SlateDebugging.h"
 #include "Widgets/SWindow.h"
 #include "Types/ReflectionMetadata.h"
+#include "Stats/Stats.h"
+#include "Containers/StringConv.h"
 #include "Misc/ScopeRWLock.h"
 #include "HAL/CriticalSection.h"
 
@@ -25,7 +27,6 @@
 #include "Widgets/Accessibility/SlateAccessibleMessageHandler.h"
 #endif
 
-DEFINE_STAT(STAT_SlateVeryVerboseStatGroupTester);
 DEFINE_STAT(STAT_SlateTotalWidgetsPerFrame);
 DEFINE_STAT(STAT_SlateNumPaintedWidgets);
 DEFINE_STAT(STAT_SlateNumTickedWidgets);
@@ -35,6 +36,8 @@ DEFINE_STAT(STAT_SlatePrepass);
 DEFINE_STAT(STAT_SlateTotalWidgets);
 DEFINE_STAT(STAT_SlateSWidgetAllocSize);
 
+DECLARE_CYCLE_STAT(TEXT("SWidget::CreateStatID"), STAT_Slate_CreateStatID, STATGROUP_Slate);
+
 template <typename AnnotationType>
 class TWidgetSparseAnnotation
 {
@@ -100,55 +103,33 @@ static FAutoConsoleVariableRef CVarSlateEnsureAllVisibleWidgetsPaint(TEXT("Slate
 
 #endif
 
-#if STATS
+#if STATS || ENABLE_STATNAMEDEVENTS
 
-struct FScopeCycleCounterSWidget : public FCycleCounter
+void SWidget::CreateStatID() const
 {
-	/**
-	 * Constructor, starts timing
-	 */
-	FORCEINLINE_STATS FScopeCycleCounterSWidget(const SWidget* Widget)
-	{
-		if (Widget)
-		{
-			TStatId WidgetStatId = Widget->GetStatID();
-			if (FThreadStats::IsCollectingData(WidgetStatId))
-			{
-				Start(WidgetStatId);
-			}
-		}
-	}
+	SCOPE_CYCLE_COUNTER(STAT_Slate_CreateStatID);
 
-	/**
-	 * Updates the stat with the time spent
-	 */
-	FORCEINLINE_STATS ~FScopeCycleCounterSWidget()
-	{
-		Stop();
-	}
-};
+	const FString LongName = FReflectionMetaData::GetWidgetDebugInfo(this);
 
-#else
+#if STATS
+	StatID = FDynamicStats::CreateStatId<FStatGroup_STATGROUP_Slate>(LongName);
+#else // ENABLE_STATNAMEDEVENTS
+	const auto& ConversionData = StringCast<PROFILER_CHAR>(*LongName);
+	const int32 NumStorageChars = (ConversionData.Length() + 1);	//length doesn't include null terminator
 
-struct FScopeCycleCounterSWidget
-{
-	FScopeCycleCounterSWidget(const SWidget* Widget)
-	{
-	}
-	~FScopeCycleCounterSWidget()
+	PROFILER_CHAR* StoragePtr = new PROFILER_CHAR[NumStorageChars];
+	FMemory::Memcpy(StoragePtr, ConversionData.Get(), NumStorageChars * sizeof(PROFILER_CHAR));
+
+	if (FPlatformAtomics::InterlockedCompareExchangePointer((void**)&StatIDStringStorage, StoragePtr, nullptr) != nullptr)
 	{
+		delete[] StoragePtr;
 	}
-};
-
+	
+	StatID = TStatId(StatIDStringStorage);
 #endif
+}
 
-
-void SWidget::CreateStatID() const
-{
-#if STATS
-	StatID = FDynamicStats::CreateStatId<FStatGroup_STATGROUP_SlateVeryVerbose>( ToString() );
 #endif
-}
 
 void SWidget::UpdateWidgetProxy(int32 NewLayerId, FSlateCachedElementListNode* CacheNode)
 {
@@ -222,6 +203,9 @@ SWidget::SWidget()
 	, RenderTransformPivot(FVector2D::ZeroVector)
 	, Cursor( TOptional<EMouseCursor::Type>() )
 	, ToolTip()
+#if ENABLE_STATNAMEDEVENTS
+	, StatIDStringStorage(nullptr)
+#endif
 {
 	if (GIsRunning)
 	{
@@ -262,6 +246,11 @@ SWidget::~SWidget()
 		ClearSparseAnnotationsForWidget(this);
 	}
 
+#if ENABLE_STATNAMEDEVENTS
+	delete[] StatIDStringStorage;
+	StatIDStringStorage = nullptr;
+#endif
+
 	DEC_DWORD_STAT(STAT_SlateTotalWidgets);
 	DEC_MEMORY_STAT_BY(STAT_SlateSWidgetAllocSize, AllocSize);
 }
diff --git a/Engine/Source/Runtime/SlateCore/Public/Widgets/SWidget.h b/Engine/Source/Runtime/SlateCore/Public/Widgets/SWidget.h
--- a/Engine/Source/Runtime/SlateCore/Public/Widgets/SWidget.h
+++ b/Engine/Source/Runtime/SlateCore/Public/Widgets/SWidget.h
@@ -42,9 +42,6 @@ class SWidget;
 struct FSlateBrush;
 struct FSlatePaintElementLists;
 
-
-
-DECLARE_DWORD_COUNTER_STAT_EXTERN(TEXT("STAT_SlateVeryVerboseStatGroupTester"), STAT_SlateVeryVerboseStatGroupTester, STATGROUP_SlateVeryVerbose, SLATECORE_API);
 DECLARE_DWORD_COUNTER_STAT_EXTERN(TEXT("Widgets Created (Per Frame)"), STAT_SlateTotalWidgetsPerFrame, STATGROUP_Slate, SLATECORE_API);
 DECLARE_DWORD_COUNTER_STAT_EXTERN(TEXT("SWidget::Paint (Count)"), STAT_SlateNumPaintedWidgets, STATGROUP_Slate, SLATECORE_API);
 DECLARE_DWORD_COUNTER_STAT_EXTERN(TEXT("SWidget::Tick (Count)"), STAT_SlateNumTickedWidgets, STATGROUP_Slate, SLATECORE_API);
@@ -55,7 +52,6 @@ DECLARE_CYCLE_STAT_EXTERN(TEXT("SlatePrepass"), STAT_SlatePrepass, STATGROUP_Sla
 DECLARE_DWORD_ACCUMULATOR_STAT_EXTERN(TEXT("Total Widgets"), STAT_SlateTotalWidgets, STATGROUP_SlateMemory, SLATECORE_API);
 DECLARE_MEMORY_STAT_EXTERN(TEXT("SWidget Total Allocated Size"), STAT_SlateSWidgetAllocSize, STATGROUP_SlateMemory, SLATECORE_API);
 
-
 /** Delegate type for handling mouse events */
 DECLARE_DELEGATE_RetVal_TwoParams(
 	FReply, FPointerEventHandler,
@@ -730,7 +726,9 @@ private:
 		DesiredSize = InDesiredSize;
 	}
 
+#if STATS || ENABLE_STATNAMEDEVENTS
 	void CreateStatID() const;
+#endif
 
 	void AddUpdateFlags(EWidgetUpdateFlags FlagsToAdd)
 	{
@@ -770,7 +768,7 @@ public:
 	{
 #if STATS
 		// this is done to avoid even registering stats for a disabled group (unless we plan on using it later)
-		if (FThreadStats::IsCollectingData(GET_STATID(STAT_SlateVeryVerboseStatGroupTester)))
+		if (FThreadStats::IsCollectingData())
 		{
 			if (!StatID.IsValidStat())
 			{
@@ -778,6 +776,12 @@ public:
 			}
 			return StatID;
 		}
+#elif ENABLE_STATNAMEDEVENTS
+		if (!StatID.IsValidStat() && GCycleStatsShouldEmitNamedEvents)
+		{
+			CreateStatID();
+		}
+		return StatID;
 #endif
 		return TStatId(); // not doing stats at the moment, or ever
 	}
@@ -1673,8 +1677,16 @@ private:
 	FNoReplyPointerEventHandler MouseEnterHandler;
 	FSimpleNoReplyPointerEventHandler MouseLeaveHandler;
 
-	STAT(mutable TStatId StatID;)
 	STAT(size_t AllocSize;)
+
+#if STATS || ENABLE_STATNAMEDEVENTS
+	/** Stat id of this object, 0 if nobody asked for it yet */
+	mutable TStatId				StatID;
+#endif
+
+#if ENABLE_STATNAMEDEVENTS
+	mutable PROFILER_CHAR* StatIDStringStorage;
+#endif
 };
 
 //=================================================================
diff --git a/Engine/Source/Runtime/SlateCore/Public/Widgets/SWidgetUtils.h b/Engine/Source/Runtime/SlateCore/Public/Widgets/SWidgetUtils.h
new file mode 100644
--- /dev/null
+++ b/Engine/Source/Runtime/SlateCore/Public/Widgets/SWidgetUtils.h
@@ -0,0 +1,58 @@
+// Copyright 1998-2019 Epic Games, Inc. All Rights Reserved.
+
+#pragma once
+
+#include "CoreMinimal.h"
+#include "Widgets/SWidget.h"
+
+#if STATS
+
+/** Structure used to track time spent by a SWidget */
+struct FScopeCycleCounterSWidget : public FCycleCounter
+{
+	/**
+	 * Constructor, starts timing
+	 */
+	FORCEINLINE FScopeCycleCounterSWidget(const SWidget* Widget)
+	{
+		if (Widget)
+		{
+			TStatId WidgetStatId = Widget->GetStatID();
+			if (FThreadStats::IsCollectingData(WidgetStatId))
+			{
+				Start(WidgetStatId);
+			}
+		}
+	}
+
+	/**
+	 * Updates the stat with the time spent
+	 */
+	FORCEINLINE ~FScopeCycleCounterSWidget()
+	{
+		Stop();
+	}
+};
+
+#define SCOPE_CYCLE_SWIDGET(Name, Object) \
+	FScopeCycleCounterSWidget ObjCycleCount_##Name(Object);
+
+#elif ENABLE_STATNAMEDEVENTS
+
+struct FScopeCycleCounterSWidget
+{
+	FScopeCycleCounter ScopeCycleCounter;
+	FORCEINLINE FScopeCycleCounterSWidget(const SWidget* Widget)
+		: ScopeCycleCounter(Widget ? Widget->GetStatID().StatString : nullptr)
+	{
+	}
+};
+
+#define SCOPE_CYCLE_SWIDGET(Name, Object) \
+	FScopeCycleCounterSWidget ObjCycleCount_##Name(Object);
+
+#else
+
+#define SCOPE_CYCLE_SWIDGET(Name, Object)
+
+#endif
