Moved redundant scope counters for freetype to a verbose compile time flag

#rnx

diff --git a/Engine/Source/Runtime/SlateCore/Private/Fonts/FontCacheFreeType.cpp b/Engine/Source/Runtime/SlateCore/Private/Fonts/FontCacheFreeType.cpp
--- a/Engine/Source/Runtime/SlateCore/Private/Fonts/FontCacheFreeType.cpp
+++ b/Engine/Source/Runtime/SlateCore/Private/Fonts/FontCacheFreeType.cpp
@@ -170,7 +170,9 @@ void ApplySizeAndScale(FT_Face InFace, const int32 InFontSize, const float InFon
 
 FT_Error LoadGlyph(FT_Face InFace, const uint32 InGlyphIndex, const int32 InLoadFlags, const int32 InFontSize, const float InFontScale)
 {
+#if WITH_VERY_VERBOSE_SLATE_STATS
 	QUICK_SCOPE_CYCLE_COUNTER(STAT_FreetypeLoadGlyph);
+#endif
 	ApplySizeAndScale(InFace, InFontSize, InFontScale);
 	return FT_Load_Glyph(InFace, InGlyphIndex, InLoadFlags);
 }
diff --git a/Engine/Source/Runtime/SlateCore/Private/Fonts/SlateFontRenderer.cpp b/Engine/Source/Runtime/SlateCore/Private/Fonts/SlateFontRenderer.cpp
--- a/Engine/Source/Runtime/SlateCore/Private/Fonts/SlateFontRenderer.cpp
+++ b/Engine/Source/Runtime/SlateCore/Private/Fonts/SlateFontRenderer.cpp
@@ -378,7 +378,10 @@ FFreeTypeFaceGlyphData FSlateFontRenderer::GetFontFaceForCodepoint(const FFontDa
 bool FSlateFontRenderer::GetRenderData(const FShapedGlyphEntry& InShapedGlyph, const FFontOutlineSettings& InOutlineSettings, FCharacterRenderData& OutRenderData) const
 {
 #if WITH_FREETYPE
+
+#if	WITH_VERY_VERBOSE_SLATE_STATS
 	SCOPE_CYCLE_COUNTER(STAT_FreetypeRenderGlyph);
+#endif
 
 	FFreeTypeFaceGlyphData FaceGlyphData;
 	FaceGlyphData.FaceAndMemory = InShapedGlyph.FontFaceData->FontFace.Pin();
