Added drop shadow color and offset to FComboButtonStyle for use in SComboButton.

diff --git a/Engine/Plugins/FX/Niagara/Source/NiagaraEditorWidgets/Private/NiagaraEditorWidgetsStyle.cpp b/Engine/Plugins/FX/Niagara/Source/NiagaraEditorWidgets/Private/NiagaraEditorWidgetsStyle.cpp
--- a/Engine/Plugins/FX/Niagara/Source/NiagaraEditorWidgets/Private/NiagaraEditorWidgetsStyle.cpp
+++ b/Engine/Plugins/FX/Niagara/Source/NiagaraEditorWidgets/Private/NiagaraEditorWidgetsStyle.cpp
@@ -150,7 +150,20 @@ TSharedRef< FSlateStyleSet > FNiagaraEditorWidgetsStyle::Create()
 	Style->Set("NiagaraEditor.Stack.GroupForegroundColor", FLinearColor(FColor(220, 220, 220)));
 	Style->Set("NiagaraEditor.Stack.FlatButtonColor", FLinearColor(FColor(205, 205, 205)));
 
-	Style->Set("NiagaraEditor.Stack.HighlightedButtonBrush", new BOX_CORE_BRUSH("Common/ButtonHoverHint", FMargin(4 / 16.0f), FLinearColor(FColor(218, 130, 9))));
+	Style->Set("NiagaraEditor.Stack.HighlightedButtonBrush", new BOX_CORE_BRUSH("Common/ButtonHoverHint", FMargin(4 / 16.0f), GetDefault<UEditorStyleSettings>()->SelectionColor));
+
+	const FVector2D ViewOptionsShadowOffset = FVector2D(0, 1);
+	Style->Set("NiagaraEditor.Stack.ViewOptionsShadowOffset", ViewOptionsShadowOffset);
+
+	FComboButtonStyle ViewOptionsComboButtonStyle = FCoreStyle::Get().GetWidgetStyle<FComboButtonStyle>("ComboButton");
+
+	const FLinearColor ViewOptionsShadowColor = FLinearColor::Black;
+	Style->Set("NiagaraEditor.Stack.ViewOptionsShadowColor", FLinearColor::Black);
+	Style->Set("NiagaraEditor.Stack.ViewOptionsButton",	ViewOptionsComboButtonStyle
+		.SetButtonStyle(FEditorStyle::Get().GetWidgetStyle<FButtonStyle>("HoverHintOnly"))
+		.SetShadowOffset(ViewOptionsShadowOffset)
+		.SetShadowColorAndOpacity(ViewOptionsShadowColor)
+	);
 
 	Style->Set("NiagaraEditor.Stack.AccentColor.System", FLinearColor(FColor(67, 105, 124)));
 	Style->Set("NiagaraEditor.Stack.AccentColor.Emitter", FLinearColor(FColor(126, 87, 67)));
diff --git a/Engine/Plugins/FX/Niagara/Source/NiagaraEditorWidgets/Private/SNiagaraStack.cpp b/Engine/Plugins/FX/Niagara/Source/NiagaraEditorWidgets/Private/SNiagaraStack.cpp
--- a/Engine/Plugins/FX/Niagara/Source/NiagaraEditorWidgets/Private/SNiagaraStack.cpp
+++ b/Engine/Plugins/FX/Niagara/Source/NiagaraEditorWidgets/Private/SNiagaraStack.cpp
@@ -314,6 +314,8 @@ void SNiagaraStack::SynchronizeTreeExpansion()
 
 TSharedRef<SWidget> SNiagaraStack::ConstructHeaderWidget()
 {
+	const FVector2D ViewOptionsShadowOffset = FNiagaraEditorWidgetsStyle::Get().GetVector("NiagaraEditor.Stack.ViewOptionsShadowOffset");
+
 	return SNew(SVerticalBox)
 		//~ Top level object list view
 		+ SVerticalBox::Slot()
@@ -355,7 +357,7 @@ TSharedRef<SWidget> SNiagaraStack::ConstructHeaderWidget()
 				.BorderImage(this, &SNiagaraStack::GetViewOptionsIconBrush)
 				[
 					SNew(SComboButton)
-					.ButtonStyle(FEditorStyle::Get(), "HoverHintOnly")
+					.ComboButtonStyle(FNiagaraEditorWidgetsStyle::Get(), "NiagaraEditor.Stack.ViewOptionsButton")
 					.ForegroundColor(FSlateColor::UseForeground())
 					.ToolTipText(LOCTEXT("ViewOptionsToolTip", "View Options"))
 					.OnGetMenuContent(this, &SNiagaraStack::GetViewOptionsMenu)
@@ -367,11 +369,11 @@ TSharedRef<SWidget> SNiagaraStack::ConstructHeaderWidget()
 						// drop shadow
 						+ SOverlay::Slot()
 						.VAlign(VAlign_Top)
-						.Padding(0, 1, 0, 0)
+						.Padding(ViewOptionsShadowOffset.X, ViewOptionsShadowOffset.Y, 0, 0)
 						[
 							SNew(SImage)
 							.Image(FEditorStyle::GetBrush("GenericViewButton"))
-							.ColorAndOpacity(FLinearColor::Black)
+							.ColorAndOpacity(FNiagaraEditorWidgetsStyle::Get().GetColor("NiagaraEditor.Stack.ViewOptionsShadowColor"))
 						]
 						+ SOverlay::Slot()
 						.VAlign(VAlign_Top)
diff --git a/Engine/Source/Runtime/Slate/Private/Widgets/Input/SComboButton.cpp b/Engine/Source/Runtime/Slate/Private/Widgets/Input/SComboButton.cpp
--- a/Engine/Source/Runtime/Slate/Private/Widgets/Input/SComboButton.cpp
+++ b/Engine/Source/Runtime/Slate/Private/Widgets/Input/SComboButton.cpp
@@ -21,6 +21,8 @@ void SComboButton::Construct( const FArguments& InArgs )
 	ContentWidgetPtr = InArgs._MenuContent.Widget;
 	bIsFocusable = InArgs._IsFocusable;
 
+	const bool bHasDownArrowShadow = !InArgs._ComboButtonStyle->ShadowOffset.IsZero();
+
 	TSharedPtr<SHorizontalBox> HBox;
 
 	SMenuAnchor::Construct( SMenuAnchor::FArguments()
@@ -62,12 +64,12 @@ void SComboButton::Construct( const FArguments& InArgs )
 					// drop shadow
 					+ SOverlay::Slot()
 					.VAlign(VAlign_Top)
-					.Padding(0, 1, 0, 0)
+					.Padding(FMargin(InArgs._ComboButtonStyle->ShadowOffset.X, InArgs._ComboButtonStyle->ShadowOffset.Y, 0, 0))
 					[
 						SNew(SImage)
-						.Visibility( InArgs._HasDownArrow && InArgs._HasDownArrowDropShadow ? EVisibility::Visible : EVisibility::Collapsed )
+						.Visibility( InArgs._HasDownArrow && bHasDownArrowShadow ? EVisibility::Visible : EVisibility::Collapsed )
 						.Image( &InArgs._ComboButtonStyle->DownArrowImage )
-						.ColorAndOpacity( FLinearColor::Black )
+						.ColorAndOpacity( InArgs._ComboButtonStyle->ShadowColorAndOpacity )
 					]
 					+ SOverlay::Slot()
 					.VAlign(VAlign_Top)
diff --git a/Engine/Source/Runtime/Slate/Public/Widgets/Input/SComboButton.h b/Engine/Source/Runtime/Slate/Public/Widgets/Input/SComboButton.h
--- a/Engine/Source/Runtime/Slate/Public/Widgets/Input/SComboButton.h
+++ b/Engine/Source/Runtime/Slate/Public/Widgets/Input/SComboButton.h
@@ -56,7 +56,6 @@ public:
 		SLATE_EVENT( FOnComboBoxOpened, OnComboBoxOpened )
 		SLATE_ARGUMENT( bool, IsFocusable )
 		SLATE_ARGUMENT( bool, HasDownArrow )
-		SLATE_ARGUMENT( bool, HasDownArrowDropShadow )
 
 		SLATE_ATTRIBUTE( FSlateColor, ForegroundColor )
 		SLATE_ATTRIBUTE( FSlateColor, ButtonColorAndOpacity )
diff --git a/Engine/Source/Runtime/SlateCore/Public/Styling/SlateTypes.h b/Engine/Source/Runtime/SlateCore/Public/Styling/SlateTypes.h
--- a/Engine/Source/Runtime/SlateCore/Public/Styling/SlateTypes.h
+++ b/Engine/Source/Runtime/SlateCore/Public/Styling/SlateTypes.h
@@ -436,28 +436,44 @@ struct SLATECORE_API FComboButtonStyle : public FSlateWidgetStyle
 	static const FComboButtonStyle& GetDefault();
 
 	/**
-	 * The style to use for our SButton
+	 * The style to use for our SButton.
 	 */	
 	UPROPERTY(EditAnywhere, BlueprintReadWrite, Category=Appearance)
 	FButtonStyle ButtonStyle;
 	FComboButtonStyle& SetButtonStyle( const FButtonStyle& InButtonStyle ){ ButtonStyle = InButtonStyle; return *this; }
 
 	/**
-	 * Image to use for the down arrow
+	 * Image to use for the down arrow.
 	 */	
 	UPROPERTY(EditAnywhere, BlueprintReadWrite, Category=Appearance)
 	FSlateBrush DownArrowImage;
 	FComboButtonStyle& SetDownArrowImage( const FSlateBrush& InDownArrowImage ){ DownArrowImage = InDownArrowImage; return *this; }
 
 	/**
-	 * Brush to use to add a "menu border" around the drop-down content
+	  * How much should the shadow be offset for the down arrow? 
+	  * An offset of 0 implies no shadow. 
+	  */
+	UPROPERTY(EditAnywhere, BlueprintReadWrite, Category=Appearance)
+	FVector2D ShadowOffset;
+	FComboButtonStyle& SetShadowOffset(const FVector2D& InShadowOffset) { ShadowOffset = InShadowOffset; return *this; }
+
+	/** 
+	  * The color and opacity of the shadow for the down arrow.
+	  * Only active if ShadowOffset is not (0,0).
+	  */
+	UPROPERTY(EditAnywhere, BlueprintReadWrite, Category=Appearance)
+	FLinearColor ShadowColorAndOpacity;
+	FComboButtonStyle& SetShadowColorAndOpacity(const FLinearColor& InShadowColorAndOpacity) { ShadowColorAndOpacity = InShadowColorAndOpacity; return *this; }
+
+	/**
+	 * Brush to use to add a "menu border" around the drop-down content.
 	 */
 	UPROPERTY(EditAnywhere, BlueprintReadWrite, Category=Appearance)
 	FSlateBrush MenuBorderBrush;
 	FComboButtonStyle& SetMenuBorderBrush( const FSlateBrush& InMenuBorderBrush ){ MenuBorderBrush = InMenuBorderBrush; return *this; }
 
 	/**
-	 * Padding to use to add a "menu border" around the drop-down content
+	 * Padding to use to add a "menu border" around the drop-down content.
 	 */
 	UPROPERTY(EditAnywhere, BlueprintReadWrite, Category=Appearance)
 	FMargin MenuBorderPadding;
