#UE4 Initialize some properties

diff --git a/Engine/Source/Runtime/SlateCore/Private/Styling/SlateTypes.cpp b/Engine/Source/Runtime/SlateCore/Private/Styling/SlateTypes.cpp
--- a/Engine/Source/Runtime/SlateCore/Private/Styling/SlateTypes.cpp
+++ b/Engine/Source/Runtime/SlateCore/Private/Styling/SlateTypes.cpp
@@ -134,7 +134,9 @@ void FButtonStyle::PostSerialize(const FArchive& Ar)
 #endif
 
 FComboButtonStyle::FComboButtonStyle()
-	: MenuBorderPadding(FMargin(0.0f))
+	: ShadowOffset(FVector2D::ZeroVector)
+	, ShadowColorAndOpacity(FLinearColor::Black)
+	, MenuBorderPadding(FMargin(0.0f))
 {
 }
 
