Gracefully handle a glyph failing to load

diff --git a/Engine/Source/Runtime/SlateCore/Private/Fonts/SlateFontRenderer.cpp b/Engine/Source/Runtime/SlateCore/Private/Fonts/SlateFontRenderer.cpp
--- a/Engine/Source/Runtime/SlateCore/Private/Fonts/SlateFontRenderer.cpp
+++ b/Engine/Source/Runtime/SlateCore/Private/Fonts/SlateFontRenderer.cpp
@@ -393,9 +393,10 @@ bool FSlateFontRenderer::GetRenderData(const FShapedGlyphEntry& InShapedGlyph, c
 		check(FaceGlyphData.FaceAndMemory->IsValid());
 
 		FT_Error Error = FreeTypeUtils::LoadGlyph(FaceGlyphData.FaceAndMemory->GetFace(), FaceGlyphData.GlyphIndex, FaceGlyphData.GlyphFlags, InShapedGlyph.FontFaceData->FontSize, InShapedGlyph.FontFaceData->FontScale);
-		check(Error == 0);
-
-		return GetRenderDataInternal(FaceGlyphData, InShapedGlyph.FontFaceData->FontScale, InOutlineSettings, OutRenderData);
+		if (Error == 0)
+		{
+			return GetRenderDataInternal(FaceGlyphData, InShapedGlyph.FontFaceData->FontScale, InOutlineSettings, OutRenderData);
+		}
 	}
 #endif // WITH_FREETYPE
 	return false;
