Add SWidget debug text to profiling tick scope.  Blueprint ticking is already tagged by UObject name, but this will also catch native tick swidgets

#rnx

diff --git a/Engine/Source/Runtime/SlateCore/Private/Widgets/SWidget.cpp b/Engine/Source/Runtime/SlateCore/Private/Widgets/SWidget.cpp
--- a/Engine/Source/Runtime/SlateCore/Private/Widgets/SWidget.cpp
+++ b/Engine/Source/Runtime/SlateCore/Private/Widgets/SWidget.cpp
@@ -21,6 +21,7 @@
 #include "Containers/StringConv.h"
 #include "Misc/ScopeRWLock.h"
 #include "HAL/CriticalSection.h"
+#include "Widgets/SWidgetUtils.h"
 
 #if WITH_ACCESSIBILITY
 #include "Widgets/Accessibility/SlateCoreAccessibleWidgets.h"
@@ -1220,6 +1221,7 @@ int32 SWidget::Paint(const FPaintArgs& Args, const FGeometry& AllottedGeometry,
 		INC_DWORD_STAT(STAT_SlateNumTickedWidgets);
 
 		SCOPE_CYCLE_COUNTER(STAT_SlateTickWidgets);
+		SCOPE_CYCLE_SWIDGET(WidgetTick, this);
 		MutableThis->Tick(DesktopSpaceGeometry, Args.GetCurrentTime(), Args.GetDeltaTime());
 	}
 
