Slate - Faster and safer draw element submitting for invalidation painting. This should fix issues with difficult to reproduce crashes involving cached element lists that have been deleted but are being accessed (suspected).

diff --git a/Engine/Source/Runtime/SlateCore/Private/FastUpdate/SlateInvalidationRoot.cpp b/Engine/Source/Runtime/SlateCore/Private/FastUpdate/SlateInvalidationRoot.cpp
--- a/Engine/Source/Runtime/SlateCore/Private/FastUpdate/SlateInvalidationRoot.cpp
+++ b/Engine/Source/Runtime/SlateCore/Private/FastUpdate/SlateInvalidationRoot.cpp
@@ -131,9 +131,9 @@ int32 RecursiveFindParentWithChildOrderChange(const TArray<FWidgetProxy>& FastWi
 
 void FSlateInvalidationRoot::RemoveWidgetFromFastPath(FWidgetProxy& Proxy)
 {
-	if (Proxy.Widget->PersistentState.CachedElementListNode)
+	if (Proxy.Widget->PersistentState.CachedElementHandle.IsValid())
 	{
-		CachedElementData->RemoveCache(Proxy.Widget->PersistentState.CachedElementListNode);
+		Proxy.Widget->PersistentState.CachedElementHandle.RemoveFromCache();
 	}
 
 	if (Proxy.Index == 0)
@@ -265,13 +265,8 @@ void FSlateInvalidationRoot::OnWidgetDestroyed(const SWidget* Widget)
 		FastWidgetPathList[ProxyIndex].Widget = nullptr;
 	
 	}
-
-	if (Widget->PersistentState.CachedElementListNode)
-	{
-		CachedElementData->RemoveCache(Widget->PersistentState.CachedElementListNode);
-	}
-
-	Widget->PersistentState.CachedElementListNode = nullptr;
+		
+	Widget->PersistentState.CachedElementHandle.RemoveFromCache();
 }
 
 bool FSlateInvalidationRoot::PaintFastPath(const FSlateInvalidationContext& Context)
@@ -636,7 +631,7 @@ void FSlateInvalidationRoot::ClearAllFastPathData(bool bClearResourcesImmediatel
 	{
 		if (Proxy.Widget)
 		{
-			Proxy.Widget->PersistentState.CachedElementListNode = nullptr;
+			Proxy.Widget->PersistentState.CachedElementHandle = FSlateCachedElementsHandle::Invalid;
 			if (bClearResourcesImmediately)
 			{
 				Proxy.Widget->FastPathProxyHandle = FWidgetProxyHandle();
diff --git a/Engine/Source/Runtime/SlateCore/Private/Rendering/DrawElements.cpp b/Engine/Source/Runtime/SlateCore/Private/Rendering/DrawElements.cpp
--- a/Engine/Source/Runtime/SlateCore/Private/Rendering/DrawElements.cpp
+++ b/Engine/Source/Runtime/SlateCore/Private/Rendering/DrawElements.cpp
@@ -736,36 +736,35 @@ void FSlateWindowElementList::EndDeferredGroup()
 	bNeedsDeferredResolve = true;
 }
 
-void FSlateWindowElementList::PushPaintingWidget(const SWidget& CurrentWidget, int32 StartingLayerId, FSlateCachedElementListNode* CurrentCacheNode)
+void FSlateWindowElementList::PushPaintingWidget(const SWidget& CurrentWidget, int32 StartingLayerId, FSlateCachedElementsHandle& CurrentCacheHandle)
 {
 	FSlateCachedElementData* CurrentCachedElementData = GetCurrentCachedElementData();
 	if (CurrentCachedElementData)
 	{
-		const FWidgetDrawElementState& PreviousState = WidgetDrawStack.Num() ? WidgetDrawStack.Top() : FWidgetDrawElementState(nullptr, false, nullptr);
+		const FWidgetDrawElementState& PreviousState = WidgetDrawStack.Num() ? WidgetDrawStack.Top() : FWidgetDrawElementState(FSlateCachedElementsHandle::Invalid, false, nullptr);
 
-		WidgetDrawStack.Emplace(CurrentCacheNode, CurrentWidget.IsVolatileIndirectly() || CurrentWidget.IsVolatile(), &CurrentWidget);
+		WidgetDrawStack.Emplace(CurrentCacheHandle, CurrentWidget.IsVolatileIndirectly() || CurrentWidget.IsVolatile(), &CurrentWidget);
 
 		// When a widget is pushed reset its draw elements.  They are being recached or possibly going away
-		if (CurrentCacheNode != nullptr)
+		if (CurrentCacheHandle.IsValid())
 		{
 #if WITH_SLATE_DEBUGGING
-			check(CurrentCacheNode->GetValue().Widget == &CurrentWidget);
+			check(CurrentCacheNode->GetValue().OwningWidget == &CurrentWidget);
 #endif
-			CurrentCachedElementData->ResetCache(CurrentCacheNode);
+			CurrentCacheHandle.ClearCachedElements();
 		}
 	}
 }
 
-
-FSlateCachedElementListNode* FSlateWindowElementList::PopPaintingWidget()
+FSlateCachedElementsHandle FSlateWindowElementList::PopPaintingWidget()
 {
 	FSlateCachedElementData* CurrentCachedElementData = GetCurrentCachedElementData();
 	if (CurrentCachedElementData)
 	{
-		return WidgetDrawStack.Pop().CacheNode;
+		return WidgetDrawStack.Pop().CacheHandle;
 	}
 
-	return nullptr;
+	return FSlateCachedElementsHandle::Invalid;
 }
 
 /*
@@ -808,12 +807,12 @@ FSlateDrawElement& FSlateWindowElementList::AddCachedElement()
 	FWidgetDrawElementState& CurrentWidgetState = WidgetDrawStack.Top();
 	check(!CurrentWidgetState.bIsVolatile);
 
-	if (CurrentWidgetState.CacheNode == nullptr)
+	if (!CurrentWidgetState.CacheHandle.IsValid())
 	{
-		CurrentWidgetState.CacheNode = CurrentCachedElementData->AddCache(CurrentWidgetState.Widget);
+		CurrentWidgetState.CacheHandle = CurrentCachedElementData->AddCache(CurrentWidgetState.Widget);
 	}
 
-	return CurrentCachedElementData->AddCachedElement(CurrentWidgetState.CacheNode, GetClippingManager(), CurrentWidgetState.Widget);
+	return CurrentCachedElementData->AddCachedElement(CurrentWidgetState.CacheHandle, GetClippingManager(), CurrentWidgetState.Widget);
 }
 
 void FSlateWindowElementList::PushCachedElementData(FSlateCachedElementData& CachedElementData)
@@ -919,53 +918,78 @@ static const FSlateClippingState* GetClipStateFromParent(const FSlateClippingMan
 	}
 }
 
-FSlateCachedElementListNode* FSlateCachedElementData::AddCache(const SWidget* Widget)
+void FSlateCachedElementData::Empty()
+{
+	if (CachedElementLists.Num())
+	{
+		UE_LOG(LogSlate, Log, TEXT("Resetting cached element data.  Num: %d"), CachedElementLists.Num());
+	}
+
+	CachedElementLists.Empty();
+	CachedBatches.Empty();
+	CachedClipStates.Empty();
+	ListsWithNewData.Empty();
+}
+
+FSlateCachedElementsHandle FSlateCachedElementData::AddCache(const SWidget* Widget)
 {
 #if WITH_SLATE_DEBUGGING
-	for (FSlateCachedElementList& CachedElementList : CachedElementLists)
+	for (TSharedPtr<FSlateCachedElementList>& CachedElementList : CachedElementLists)
 	{
-		ensure(CachedElementList.Widget != Widget);
+		ensure(CachedElementList->OwningWidget != Widget);
 	}
 #endif
 
-	FSlateCachedElementListNode* NewNode = new FSlateCachedElementListNode(FSlateCachedElementList(this, Widget));
+	TSharedRef<FSlateCachedElementList> NewList = MakeShared<FSlateCachedElementList>(this, Widget);
+	NewList->Initialize();
 
-	CachedElementLists.AddTail(NewNode);
-	NewNode->GetValue().Initialize();
+	CachedElementLists.Add(NewList);
 
-	return NewNode;
+
+	return FSlateCachedElementsHandle(NewList);
 }
 
-FSlateDrawElement& FSlateCachedElementData::AddCachedElement(FSlateCachedElementListNode* CacheNode, const FSlateClippingManager& ParentClipManager, const SWidget* CurrentWidget)
+FSlateDrawElement& FSlateCachedElementData::AddCachedElement(FSlateCachedElementsHandle& CacheHandle, const FSlateClippingManager& ParentClipManager, const SWidget* CurrentWidget)
 {
+	TSharedPtr<FSlateCachedElementList> List = CacheHandle.Ptr.Pin();
+
 #if WITH_SLATE_DEBUGGING
-	check(CacheNode->GetValue().Widget == CurrentWidget);
+	check(List->OwningWidget == CurrentWidget);
 	check(CurrentWidget->GetParentWidget().IsValid());
 #endif
 
-	FSlateCachedElementList& List = CacheNode->GetValue();
-	FSlateDrawElement& NewElement = List.DrawElements.AddDefaulted_GetRef();
+	FSlateDrawElement& NewElement = List->DrawElements.AddDefaulted_GetRef();
 	NewElement.SetIsCached(true);
 
-	List.bNewData = true;
+	// Check if slow vs checking a flag on the list to see if it contains new data.
+	ListsWithNewData.AddUnique(List.Get());
+
 	const FSlateClippingState* ExistingClipState = GetClipStateFromParent(ParentClipManager);
 
 	if (ExistingClipState)
 	{
 		// We need to cache this clip state for the next time the element draws
 		FSlateCachedClipState& CachedClipState = FindOrAddCachedClipState(ExistingClipState);
-		List.AddCachedClipState(CachedClipState);
+		List->AddCachedClipState(CachedClipState);
 		NewElement.SetCachedClippingState(&CachedClipState.ClippingState.Get());
 	}
 
 	return NewElement;
 }
 
-void FSlateCachedElementData::AddReferencedObjects(FReferenceCollector& Collector)
+FSlateRenderBatch& FSlateCachedElementData::AddCachedRenderBatch(FSlateRenderBatch&& NewBatch, int32& OutIndex)
+{
+	// Check perf against add.  AddAtLowest makes it generally re-add elements at the same index it just removed which is nicer on the cache
+	int32 LowestFreedIndex = 0;
+	OutIndex = CachedBatches.AddAtLowestFreeIndex(NewBatch, LowestFreedIndex);
+	return CachedBatches[OutIndex];
+}
+
+void FSlateCachedElementData::RemoveCachedRenderBatches(const TArray<int32>& CachedRenderBatchIndices)
 {
-	for (FSlateCachedElementList& CachedElementList : CachedElementLists)
+	for (int32 Index : CachedRenderBatchIndices)
 	{
-		CachedElementList.AddReferencedObjects(Collector);
+		CachedBatches.RemoveAt(Index);
 	}
 }
 
@@ -998,21 +1022,24 @@ void FSlateCachedElementData::CleanupUnusedClipStates()
 	}
 }
 
-FSlateCachedElementList::~FSlateCachedElementList()
+
+void FSlateCachedElementData::AddReferencedObjects(FReferenceCollector& Collector)
 {
-	DestroyCachedVertexData();
-	Widget->PersistentState.CachedElementListNode = nullptr;
+	for (TSharedPtr<FSlateCachedElementList>& CachedElementList : CachedElementLists)
+	{
+		CachedElementList->AddReferencedObjects(Collector);
+	}
 }
 
-void FSlateCachedElementList::Reset()
+FSlateCachedElementList::~FSlateCachedElementList()
 {
-	DrawElements.Reset();
-
-	CachedBatches.Reset();
-
+	DestroyCachedData();
+}
 
+void FSlateCachedElementList::ClearCachedElements()
+{
 	// Destroy vertex data in a thread safe way
-	DestroyCachedVertexData();
+	DestroyCachedData();
 
 	CachedRenderingData = new FSlateCachedFastPathRenderingData;
 
@@ -1022,13 +1049,21 @@ void FSlateCachedElementList::Reset()
 		UE_LOG(LogSlate, Log, TEXT("Cleared out data in cached ElementList for Widget: %s before it was batched"), *Widget->GetTag().ToString());
 	}
 #endif
-
-	bNewData = false;
 }
 
 FSlateRenderBatch& FSlateCachedElementList::AddRenderBatch(int32 InLayer, const FShaderParams& InShaderParams, const FSlateShaderResource* InResource, ESlateDrawPrimitive InPrimitiveType, ESlateShader InShaderType, ESlateDrawEffect InDrawEffects, ESlateBatchDrawFlag InDrawFlags, int8 SceneIndex)
 {
-	return CachedBatches.Emplace_GetRef(InLayer, InShaderParams, InResource, InPrimitiveType, InShaderType, InDrawEffects, InDrawFlags, SceneIndex, &CachedRenderingData->Vertices, &CachedRenderingData->Indices, CachedRenderingData->Vertices.Num(), CachedRenderingData->Indices.Num());
+	FSlateRenderBatch NewRenderBatch(InLayer, InShaderParams, InResource, InPrimitiveType, InShaderType, InDrawEffects, InDrawFlags, SceneIndex, &CachedRenderingData->Vertices, &CachedRenderingData->Indices, CachedRenderingData->Vertices.Num(), CachedRenderingData->Indices.Num());
+	int32 RenderBatchIndex = INDEX_NONE;
+	FSlateRenderBatch& AddedBatchRef = ParentData->AddCachedRenderBatch(MoveTemp(NewRenderBatch), RenderBatchIndex);
+	
+	check(RenderBatchIndex != INDEX_NONE);
+
+	CachedRenderBatchIndices.Add(RenderBatchIndex);
+
+	return AddedBatchRef;
+	
+	//return CachedBatches.Emplace_GetRef(InLayer, InShaderParams, InResource, InPrimitiveType, InShaderType, InDrawEffects, InDrawFlags, SceneIndex, &CachedRenderingData->Vertices, &CachedRenderingData->Indices, CachedRenderingData->Vertices.Num(), CachedRenderingData->Indices.Num());
 }
 
 void FSlateCachedElementList::AddCachedClipState(FSlateCachedClipState& ClipStateToCache)
@@ -1044,7 +1079,7 @@ void FSlateCachedElementList::AddReferencedObjects(FReferenceCollector& Collecto
 	}
 }
 
-void FSlateCachedElementList::DestroyCachedVertexData()
+void FSlateCachedElementList::DestroyCachedData()
 {
 	if (CachedRenderingData)
 	{
@@ -1059,4 +1094,43 @@ void FSlateCachedElementList::DestroyCachedVertexData()
 	}
 
 	CachedRenderingData = nullptr;
+
+	if (CachedRenderBatchIndices.Num())
+	{
+		ParentData->RemoveCachedRenderBatches(CachedRenderBatchIndices);
+		CachedRenderBatchIndices.Reset();
+	}
+
+	DrawElements.Reset();
+}
+
+FSlateCachedElementsHandle FSlateCachedElementsHandle::Invalid;
+
+void FSlateCachedElementsHandle::ClearCachedElements()
+{
+	if (TSharedPtr<FSlateCachedElementList> List = Ptr.Pin())
+	{
+		List->ClearCachedElements();
+	}
+}
+
+void FSlateCachedElementsHandle::RemoveFromCache()
+{
+	if (TSharedPtr<FSlateCachedElementList> List = Ptr.Pin())
+	{
+		List->GetOwningData()->RemoveList(*this);
+		ensure(List.IsUnique());
+	}
+
+	check(!Ptr.IsValid());
+}
+
+bool FSlateCachedElementsHandle::IsOwnedByWidget(const SWidget* Widget) const
+{
+	if (const TSharedPtr<FSlateCachedElementList> List = Ptr.Pin())
+	{
+		return List->OwningWidget == Widget;
+	}
+
+	return false;
 }
diff --git a/Engine/Source/Runtime/SlateCore/Private/Rendering/ElementBatcher.cpp b/Engine/Source/Runtime/SlateCore/Private/Rendering/ElementBatcher.cpp
--- a/Engine/Source/Runtime/SlateCore/Private/Rendering/ElementBatcher.cpp
+++ b/Engine/Source/Runtime/SlateCore/Private/Rendering/ElementBatcher.cpp
@@ -89,116 +89,19 @@ bool FSlateBatchData::IsStencilClippingRequired() const
 	return bIsStencilBufferRequired;
 }
 
-void FSlateBatchData::FillVertexAndIndexBuffer(uint8* VertexBuffer, uint8* IndexBuffer, bool bAbsoluteIndices)
-{
-	/*int32 IndexOffset = 0;
-	int32 VertexOffset = 0;
-	int32 BaseVertexIndex = 0;
-
-	const bool bValidBuffers = (nullptr != VertexBuffer) && (nullptr != IndexBuffer);
-
-	for (const FSlateRenderBatch& Batch : RenderBatches)
-	{
-		// Ignore foreign batches that are inserted into our render set.
-		if (RenderDataHandle != Batch.CachedRenderHandle)
-		{
-			continue;
-		}
-
-		if (Batch.VertexArrayIndex != INDEX_NONE && Batch.IndexArrayIndex != INDEX_NONE)
-		{
-			FSlateVertexArray& Vertices = BatchVertexArrays[Batch.VertexArrayIndex];
-			FSlateIndexArray& Indices = BatchIndexArrays[Batch.IndexArrayIndex];
-
-			if (Vertices.Num() && Indices.Num())
-			{
-				if (bValidBuffers)
-				{
-					uint32 RequiredVertexSize = Vertices.Num() * Vertices.GetTypeSize();
-					uint32 RequiredIndexSize = Indices.Num() * Indices.GetTypeSize();
-
-					FMemory::Memcpy(VertexBuffer + VertexOffset, Vertices.GetData(), RequiredVertexSize);
-					if (BaseVertexIndex == 0 || !bAbsoluteIndices)
-					{
-						FMemory::Memcpy(IndexBuffer + IndexOffset, Indices.GetData(), RequiredIndexSize);
-					}
-					else
-					{
-						SlateIndex* TargetIndexBuffer = (SlateIndex*)(IndexBuffer + IndexOffset);
-						for (int32 i = 0; i < Indices.Num(); ++i)
-						{
-							TargetIndexBuffer[i] = Indices[i] + BaseVertexIndex;
-						}
-					}
-
-					BaseVertexIndex += Vertices.Num();
-					IndexOffset += (Indices.Num() * sizeof(SlateIndex));
-					VertexOffset += (Vertices.Num() * sizeof(FSlateVertex));
-				}
-
-				Vertices.Reset();
-				Indices.Reset();
-
-				if (Vertices.GetSlack() > MAX_VERT_ARRAY_RECYCLE)
-				{
-					ResetVertexArray(Vertices);
-				}
-
-				if (Indices.GetSlack() > MAX_INDEX_ARRAY_RECYCLE)
-				{
-					ResetIndexArray(Indices);
-				}
-			}
-		}
-	}*/
-}
-
-/*
-void FSlateBatchData::CreateRenderBatches(FElementBatchMap& LayerToElementBatches)
-{
-	checkSlow(IsInRenderingThread());
-
-	uint32 VertexOffset = 0;
-	uint32 IndexOffset = 0;
-
-	{
-		SCOPED_NAMED_EVENT_TEXT("SlateRT::CreateRenderBatches", FColor::Magenta);
-		Merge(LayerToElementBatches, VertexOffset, IndexOffset);
-	}
-
-	// 
-	if (RenderDataHandle.IsValid())
-	{
-		RenderDataHandle->SetRenderBatches(&RenderBatches);
-	}
-}
-*/
-
 FSlateRenderBatch& FSlateBatchData::AddRenderBatch(int32 InLayer, const FShaderParams& InShaderParams, const FSlateShaderResource* InResource, ESlateDrawPrimitive InPrimitiveType, ESlateShader InShaderType, ESlateDrawEffect InDrawEffects, ESlateBatchDrawFlag InDrawFlags, int8 SceneIndex)
 {
-
-/*
-	NumBatchedVertices += InNumVertices;
-	NumBatchedIndices += InNumIndices;
-
-	if (InElementBatch.ClippingState.IsSet())
-	{
-		bIsStencilBufferRequired |= InElementBatch.ClippingState->GetClippingMethod() == EClippingMethod::Stencil;
-	}
-
-	const int32 Index = RenderBatches.Add(FSlateRenderBatch(InLayer, InElementBatch, RenderDataHandle, InNumVertices, InNumIndices, InVertexOffset, InIndexOffset));
-	RenderBatches[Index].DynamicOffset = FVector2D::ZeroVector;*/
-
 	return RenderBatches.Emplace_GetRef(InLayer, InShaderParams, InResource, InPrimitiveType, InShaderType, InDrawEffects, InDrawFlags, SceneIndex, &UncachedSourceBatchVertices, &UncachedSourceBatchIndices, UncachedSourceBatchVertices.Num(), UncachedSourceBatchIndices.Num());
 }
 
-
-
-void FSlateBatchData::AddCachedBatches(const TArray<FSlateRenderBatch>& InCachedBatches)
+void FSlateBatchData::AddCachedBatches(const TSparseArray<FSlateRenderBatch>& InCachedBatches)
 {
 	RenderBatches.Reserve(RenderBatches.Num() + InCachedBatches.Num());
 
-	RenderBatches.Append(InCachedBatches);
+	for (const FSlateRenderBatch& CachedBatch : InCachedBatches)
+	{
+		RenderBatches.Add(CachedBatch);
+	}
 }
 
 void FSlateBatchData::FillBuffersFromNewBatch(FSlateRenderBatch& Batch, FSlateVertexArray& FinalVertices, FSlateIndexArray& FinalIndices)
@@ -246,26 +149,6 @@ void FSlateBatchData::CombineBatches(FSlateRenderBatch& FirstBatch, FSlateRender
 }
 
 
-/*
-static void AddNewBatch(FSlateRenderBatch& NewBatch, TArray<FSlateRenderBatch>& FinalRenderBatches, TArray<FSlateVertex>& Vertices, TArray<SlateIndex>& Indices)
-{
-	if (NewBatch.GetNumVertices() > 0 && NewBatch.GetNumIndices())
-	{
-		const int32 NewVertOffset = Vertices.Num();
-	
-		NewBatch.CopySourceVertices(Vertices);
-		NewBatch.VertexOffset = NewVertOffset;
-
-		// Starting a new batch does not require updating absolute index position
-		const int32 NewIndexOffset = Indices.Num();
-		NewBatch.CopySourceIndices(Indices);
-		NewBatch.IndexOffset = NewIndexOffset;
-
-		FinalRenderBatches.Add(NewBatch);
-	}
-}
-*/
-
 void FSlateBatchData::MergeRenderBatches()
 {
 	SCOPE_CYCLE_COUNTER(STAT_SlateRTCreateBatches);
@@ -373,7 +256,7 @@ void FSlateBatchData::MergeRenderBatches()
 
 FSlateElementBatcher::FSlateElementBatcher( TSharedRef<FSlateRenderingPolicy> InRenderingPolicy )
 	: BatchData( nullptr )
-	, CachedElementList( nullptr )
+	, CurrentCachedElementList( nullptr )
 	, PrecachedClippingStates(nullptr)
 	, RenderingPolicy( &InRenderingPolicy.Get() )
 	, NumPostProcessPasses(0)
@@ -560,39 +443,27 @@ void FSlateElementBatcher::AddCachedElements(FSlateCachedElementData& CachedElem
 	int32 CachedDrawElements = 0;
 #endif
 
-#define ALLOW_CACHED_RENDER_BATCHES 1
 	SCOPED_NAMED_EVENT_TEXT("Slate::AddCachedBatches", FColor::Magenta);
-	for ( FSlateCachedElementList& LocalCachedElementList : CachedElementData.CachedElementLists)
+
+	for(FSlateCachedElementList* List : CachedElementData.ListsWithNewData)
 	{
-		STAT(ElementStat_CachedElements += LocalCachedElementList.DrawElements.Num());
+		STAT(ElementStat_CachedElements += List->DrawElements.Num());
 
 #if SLATE_CSV_TRACKER
 		CachedDrawElements += LocalCachedElementList.DrawElements.Num();
 #endif
-
-		LocalCachedElementList.bNewData = false;
-#if ALLOW_CACHED_RENDER_BATCHES
-		if (LocalCachedElementList.CachedBatches.Num())
+	
+		CurrentCachedElementList = List;
 		{
-			BatchData->AddCachedBatches(LocalCachedElementList.CachedBatches);
+			SCOPE_CYCLE_SWIDGET(RecacheElements, LocalCachedElementList.Widget);
+			AddElementsInternal(List->DrawElements, ViewportSize);
 		}
-		else
-		{
-			CachedElementList = &LocalCachedElementList;
-			{
-				SCOPE_CYCLE_SWIDGET(WidgetScope, LocalCachedElementList.Widget);
-				AddElementsInternal(LocalCachedElementList.DrawElements, ViewportSize);
-			}
+		CurrentCachedElementList = nullptr;
+	}
+	CachedElementData.ListsWithNewData.Empty();
 
-			// Now apply the cached batches to the batch data
-			BatchData->AddCachedBatches(LocalCachedElementList.CachedBatches);
-		}
+	BatchData->AddCachedBatches(CachedElementData.GetCachedBatches());
 
-		CachedElementList = nullptr;
-#else
-		AddElementsInternal(LocalCachedElementList.DrawElements, ViewportSize);
-#endif
-	}
 	CachedElementData.CleanupUnusedClipStates();
 
 #if SLATE_CSV_TRACKER
@@ -2643,8 +2514,8 @@ FSlateRenderBatch& FSlateElementBatcher::CreateRenderBatch(
 	ESlateBatchDrawFlag DrawFlags,
 	const FSlateDrawElement& DrawElement)
 {
-	FSlateRenderBatch& NewBatch = CachedElementList
-		? CachedElementList->AddRenderBatch(Layer, ShaderParams, InResource, PrimitiveType, ShaderType, DrawEffects, DrawFlags, DrawElement.GetSceneIndex())
+	FSlateRenderBatch& NewBatch = CurrentCachedElementList
+		? CurrentCachedElementList->AddRenderBatch(Layer, ShaderParams, InResource, PrimitiveType, ShaderType, DrawEffects, DrawFlags, DrawElement.GetSceneIndex())
 		: BatchData->AddRenderBatch(Layer, ShaderParams, InResource, PrimitiveType, ShaderType, DrawEffects, DrawFlags, DrawElement.GetSceneIndex());
 
 	NewBatch.ClippingState = ResolveClippingState(DrawElement);
@@ -2659,7 +2530,7 @@ const FSlateClippingState* FSlateElementBatcher::ResolveClippingState(const FSla
 	if (ClipHandle.GetCachedClipState())
 	{
 		// We should be working with cached elements if we have a cached clip state
-		check(CachedElementList);
+		check(CurrentCachedElementList);
 		return ClipHandle.GetCachedClipState();
 	}
 	else if (PrecachedClippingStates->IsValidIndex(ClipHandle.GetPrecachedClipIndex()))
@@ -2671,53 +2542,6 @@ const FSlateClippingState* FSlateElementBatcher::ResolveClippingState(const FSla
 	return nullptr;
 }
 
-/*
-FSlateElementBatch& FSlateElementBatcher::FindBatchForElement(
-	uint32 Layer, 
-	const FShaderParams& ShaderParams, 
-	const FSlateShaderResource* InTexture, 
-	ESlateDrawPrimitive PrimitiveType,
-	ESlateShader ShaderType, 
-	ESlateDrawEffect DrawEffects, 
-	ESlateBatchDrawFlag DrawFlags,
-	const FSlateDrawElement& DrawElement)
-{
-	SCOPE_CYCLE_COUNTER( STAT_SlateFindBatchForElement );
-	FElementBatchMap& LayerToElementBatches = CurrentDrawLayer->GetElementBatchMap();
-
-	// See if the layer already exists.
-	TUniqueObj<FElementBatchArray>* ElementBatches = LayerToElementBatches.Find( Layer );
-	if( !ElementBatches )
-	{
-		// The layer doesn't exist so make it now
-		ElementBatches = &LayerToElementBatches.Add( Layer );
-	}
-
-	checkSlow( ElementBatches );
-
-	// Create a temp batch so we can use it as our key to find if the same batch already exists
-	FSlateElementBatch TempBatch( InTexture, ShaderParams, ShaderType, PrimitiveType, DrawEffects, DrawFlags, DrawElement, 0, 0, nullptr );
-
-	FSlateElementBatch* ElementBatch = (*ElementBatches)->FindByKey( TempBatch );
-	if( !ElementBatch )
-	{
-		// No batch with the specified parameter exists.  Create it from the temp batch.
-		int32 Index = (*ElementBatches)->Add( TempBatch );
-		ElementBatch = &(**ElementBatches)[Index];
-
-		BatchData->AssignVertexArrayToBatch(*ElementBatch);
-		BatchData->AssignIndexArrayToBatch(*ElementBatch);
-		ElementBatch->SaveClippingState(*PrecachedClippingStates);
-
-	}
-	check( ElementBatch );
-
-	// Increment the number of elements in the batch.
-	++ElementBatch->NumElementsInBatch;
-	return *ElementBatch;
-}
-*/
-
 void FSlateElementBatcher::ResetBatches()
 {
 	bRequiresVsync = false;
diff --git a/Engine/Source/Runtime/SlateCore/Private/Widgets/SWidget.cpp b/Engine/Source/Runtime/SlateCore/Private/Widgets/SWidget.cpp
--- a/Engine/Source/Runtime/SlateCore/Private/Widgets/SWidget.cpp
+++ b/Engine/Source/Runtime/SlateCore/Private/Widgets/SWidget.cpp
@@ -132,19 +132,19 @@ void SWidget::CreateStatID() const
 
 #endif
 
-void SWidget::UpdateWidgetProxy(int32 NewLayerId, FSlateCachedElementListNode* CacheNode)
+void SWidget::UpdateWidgetProxy(int32 NewLayerId, FSlateCachedElementsHandle& CacheHandle)
 {
 #if WITH_SLATE_DEBUGGING
-	check(!CacheNode || CacheNode->GetValue().Widget == this);
+	check(!CacheHandle.IsValid() || CacheHandle.IsOwnedByWidget(this));
 #endif
 
-	if (PersistentState.CachedElementListNode != nullptr && PersistentState.CachedElementListNode != CacheNode)
+	// Account for the case when the widget gets a new handle for some reason.  This should really never happen
+	if (PersistentState.CachedElementHandle.IsValid() && PersistentState.CachedElementHandle != CacheHandle)
 	{
-	//	ensure(false);
-		PersistentState.CachedElementListNode->GetValue().GetOwningData()->RemoveCache(PersistentState.CachedElementListNode);
+		ensure(false);
+		PersistentState.CachedElementHandle.RemoveFromCache();
 	}
-
-	PersistentState.CachedElementListNode = CacheNode;
+	PersistentState.CachedElementHandle = CacheHandle;
 
 	if (FastPathProxyHandle.IsValid())
 	{
@@ -235,10 +235,7 @@ SWidget::~SWidget()
 
 		// Note: this would still be valid if a widget was painted and then destroyed in the same frame.  
 		// In that case invalidation hasn't taken place for added widgets so the invalidation panel doesn't know about their cached element data to clean it up
-		if (PersistentState.CachedElementListNode)
-		{
-			PersistentState.CachedElementListNode->GetValue().GetOwningData()->RemoveCache(PersistentState.CachedElementListNode);
-		}
+		PersistentState.CachedElementHandle.RemoveFromCache();
 
 #if WITH_ACCESSIBILITY
 		FSlateApplicationBase::Get().GetAccessibleMessageHandler()->OnWidgetRemoved(this);
@@ -717,12 +714,12 @@ bool SWidget::AssignIndicesToChildren(FSlateInvalidationRoot& Root, int32 Parent
 
 	FastPathProxyHandle = FWidgetProxyHandle(Root, MyProxy.Index);
 
-	if (bInvisibleDueToParentOrSelfVisibility&& PersistentState.CachedElementListNode != nullptr)
+	if (bInvisibleDueToParentOrSelfVisibility&& PersistentState.CachedElementHandle.IsValid())
 	{
 #if WITH_SLATE_DEBUGGING
-		check(PersistentState.CachedElementListNode->GetValue().Widget == this);
+		check(PersistentState.CachedElementHandle.IsOwnedByWidget(this));
 #endif
-		Root.GetCachedElements().ResetCache(PersistentState.CachedElementListNode);
+		PersistentState.CachedElementHandle.RemoveFromCache();
 	}
 
 	FastPathList.Add(MyProxy);
@@ -787,10 +784,8 @@ void SWidget::UpdateFastPathVisibility(bool bParentVisible, bool bWidgetRemoved,
 		HittestGridToRemoveFrom->RemoveWidget(SharedThis(this));
 	}
 
-	if (PersistentState.CachedElementListNode)
-	{
-		PersistentState.CachedElementListNode->GetValue().Reset();
-	}
+
+	PersistentState.CachedElementHandle.ClearCachedElements();
 
 	FChildren* MyChildren = GetAllChildren();
 	const int32 NumChildren = MyChildren->Num();
@@ -1270,13 +1265,14 @@ int32 SWidget::Paint(const FPaintArgs& Args, const FGeometry& AllottedGeometry,
 	ensure(!MutableThis.IsUnique());
 #endif
 
-
-	if (!FastPathProxyHandle.IsValid() && PersistentState.CachedElementListNode != nullptr)
+#if WITH_SLATE_DEBUGGING
+	if (!FastPathProxyHandle.IsValid() && PersistentState.CachedElementHandle.IsValid())
 	{
 		ensure(!bInvisibleDueToParentOrSelfVisibility);
 	}
+#endif
 
-	OutDrawElements.PushPaintingWidget(*this, LayerId, PersistentState.CachedElementListNode);
+	OutDrawElements.PushPaintingWidget(*this, LayerId, PersistentState.CachedElementHandle);
 
 	if (bOutgoingHittestability)
 	{
@@ -1407,13 +1403,13 @@ int32 SWidget::Paint(const FPaintArgs& Args, const FGeometry& AllottedGeometry,
 	}
 #endif
 
-	FSlateCachedElementListNode* NewCacheNode = OutDrawElements.PopPaintingWidget();
+	FSlateCachedElementsHandle NewCacheHandle = OutDrawElements.PopPaintingWidget();
 	if (OutDrawElements.ShouldResolveDeferred())
 	{
 		NewLayerId = OutDrawElements.PaintDeferred(NewLayerId, MyCullingRect);
 	}
 
-	MutableThis->UpdateWidgetProxy(NewLayerId, NewCacheNode);
+	MutableThis->UpdateWidgetProxy(NewLayerId, NewCacheHandle);
 
 	return NewLayerId;
 
diff --git a/Engine/Source/Runtime/SlateCore/Public/FastUpdate/WidgetProxy.h b/Engine/Source/Runtime/SlateCore/Public/FastUpdate/WidgetProxy.h
--- a/Engine/Source/Runtime/SlateCore/Public/FastUpdate/WidgetProxy.h
+++ b/Engine/Source/Runtime/SlateCore/Public/FastUpdate/WidgetProxy.h
@@ -68,7 +68,9 @@ public:
 	EInvalidateWidgetReason CurrentInvalidateReason;
 	/** The widgets own visibility */
 	EVisibility Visibility;
+	/** Used to make sure we don't double process a widget that is invalidated.  (a widget can invalidate itself but an ancestor can end up painting that widget first thus rendering the child's own invalidate unnecessary */
 	uint8 bUpdatedSinceLastInvalidate : 1;
+	/** Is the widget already in a pending update list.  If it already is in an update list we don't bother adding it again */
 	uint8 bInUpdateList : 1;
 	uint8 bInvisibleDueToParentOrSelfVisibility : 1;
 	uint8 bChildOrderInvalid : 1;
@@ -135,7 +137,7 @@ private:
 struct FSlateWidgetPersistentState
 {
 	FSlateWidgetPersistentState()
-		: CachedElementListNode(nullptr)
+		: CachedElementHandle()
 		, LayerId(0)
 		, OutgoingLayerId(0)
 		, IncomingUserIndex(INDEX_NONE)
@@ -150,7 +152,7 @@ struct FSlateWidgetPersistentState
 	FGeometry DesktopGeometry;
 	FSlateRect CullingBounds;
 	FWidgetStyle WidgetStyle;
-	FSlateCachedElementListNode* CachedElementListNode;
+	FSlateCachedElementsHandle CachedElementHandle;
 	/** Starting layer id for drawing children **/
 	int32 LayerId;
 	int32 OutgoingLayerId;
diff --git a/Engine/Source/Runtime/SlateCore/Public/Rendering/DrawElements.h b/Engine/Source/Runtime/SlateCore/Public/Rendering/DrawElements.h
--- a/Engine/Source/Runtime/SlateCore/Public/Rendering/DrawElements.h
+++ b/Engine/Source/Runtime/SlateCore/Public/Rendering/DrawElements.h
@@ -343,11 +343,10 @@ struct FSlateCachedElementData;
 
 struct FSlateCachedElementList
 {
-	FSlateCachedElementList(FSlateCachedElementData* InOwningData, const SWidget* InWidget)
-		: OwningData(InOwningData)
-		, Widget(InWidget)
+	FSlateCachedElementList(FSlateCachedElementData* InParentData, const SWidget* InWidget)
+		: OwningWidget(InWidget)
+		, ParentData(InParentData)
 		, CachedRenderingData(nullptr)
-		, bNewData(false)
 	{}
 
 	void Initialize()
@@ -357,9 +356,9 @@ struct FSlateCachedElementList
 
 	SLATECORE_API ~FSlateCachedElementList();
 
-	void Reset();
+	void ClearCachedElements();
 
-	FSlateCachedElementData* GetOwningData() { return OwningData; }
+	FSlateCachedElementData* GetOwningData() { return ParentData; }
 
 	FSlateRenderBatch& AddRenderBatch(int32 InLayer, const FShaderParams& InShaderParams, const FSlateShaderResource* InResource, ESlateDrawPrimitive InPrimitiveType, ESlateShader InShaderType, ESlateDrawEffect InDrawEffects, ESlateBatchDrawFlag InDrawFlags, int8 SceneIndex);
 
@@ -368,54 +367,82 @@ struct FSlateCachedElementList
 	void AddReferencedObjects(FReferenceCollector& Collector);
 
 private:
-	SLATECORE_API void DestroyCachedVertexData();
+	SLATECORE_API void DestroyCachedData();
 public:
-	FSlateCachedElementData* OwningData;
-	const SWidget* Widget;
-
 	/** List of source draw elements to create batches from */
 	FSlateDrawElementArray DrawElements;
-	/** List of cached batches to submit for drawing */
-	TArray<FSlateRenderBatch> CachedBatches;
+
+	TArray<int32> CachedRenderBatchIndices;
+
+	/** The widget whose draw elements are in this list */
+	const SWidget* OwningWidget;
+
+	FSlateCachedElementData* ParentData;
+
 	FSlateCachedFastPathRenderingData* CachedRenderingData;
-	bool bNewData;
 private:
 };
 
-typedef TDoubleLinkedList<FSlateCachedElementList>::TDoubleLinkedListNode FSlateCachedElementListNode;
+struct FSlateCachedElementsHandle
+{
+	friend struct FSlateCachedElementData;
+
+	static FSlateCachedElementsHandle Invalid;
+	void ClearCachedElements();
+	void RemoveFromCache();
 
+	bool IsOwnedByWidget(const SWidget* Widget) const;
+
+	bool IsValid() const { return Ptr.IsValid(); }
+
+	bool operator!=(FSlateCachedElementsHandle& Other) const { return Ptr != Other.Ptr; }
+
+	FSlateCachedElementsHandle() {}
+private:
+	FSlateCachedElementsHandle(TSharedRef<FSlateCachedElementList>& DataPtr)
+		: Ptr(DataPtr)
+	{}
+private:
+	TWeakPtr<FSlateCachedElementList> Ptr;
+};
 
 struct FSlateCachedElementData
 {
-	void Empty()
-	{
-		CachedElementLists.Empty();
-	}
+	friend class FSlateElementBatcher;
 
-	void ResetCache(FSlateCachedElementListNode* ElementListNode)
-	{
-		ElementListNode->GetValue().Reset();
-	}
+	void Empty();
 
-	void RemoveCache(FSlateCachedElementListNode*& ElementListNode)
+	/** Removes a cache node completely from the cache */
+	void RemoveList(FSlateCachedElementsHandle& CacheHandle)
 	{
-		CachedElementLists.RemoveNode(ElementListNode);
-		ElementListNode = nullptr;
+		TSharedPtr<FSlateCachedElementList> Data = CacheHandle.Ptr.Pin();
+		CachedElementLists.RemoveSingleSwap(Data);
+		Data->ClearCachedElements();
 	}
 
-	FSlateCachedElementListNode* AddCache(const SWidget* Widget);
+	FSlateCachedElementsHandle AddCache(const SWidget* Widget);
 
-	FSlateDrawElement& AddCachedElement(FSlateCachedElementListNode* CacheNode, const FSlateClippingManager& ParentClipManager, const SWidget* Widget);
+	FSlateDrawElement& AddCachedElement(FSlateCachedElementsHandle& CacheHandle, const FSlateClippingManager& ParentClipManager, const SWidget* Widget);
 
-	void AddReferencedObjects(FReferenceCollector& Collector);
+	FSlateRenderBatch& AddCachedRenderBatch(FSlateRenderBatch&& NewBatch, int32& OutIndex);
+	void RemoveCachedRenderBatches(const TArray<int32>& CachedRenderBatchIndices);
 
 	FSlateCachedClipState& FindOrAddCachedClipState(const FSlateClippingState* RefClipState);
 	void CleanupUnusedClipStates();
 
-public:
-	TDoubleLinkedList<FSlateCachedElementList> CachedElementLists;
+	const TSparseArray<FSlateRenderBatch>& GetCachedBatches() const { return CachedBatches; }
+
+	void AddReferencedObjects(FReferenceCollector& Collector);
+
 private:
 
+	/** List of cached batches to submit for drawing */
+	TSparseArray<FSlateRenderBatch> CachedBatches;
+
+	TArray<TSharedPtr<FSlateCachedElementList>> CachedElementLists;
+
+	TArray<FSlateCachedElementList*, TMemStackAllocator<>> ListsWithNewData;
+
 	TArray<FSlateCachedClipState> CachedClipStates;
 };
 
@@ -572,13 +599,13 @@ public:
 	 * This information is used for caching later.
 	 *
 	 */
-	SLATECORE_API void PushPaintingWidget(const SWidget& CurrentWidget, int32 StartingLayerId, FSlateCachedElementListNode* CurrentCacheNode);
+	SLATECORE_API void PushPaintingWidget(const SWidget& CurrentWidget, int32 StartingLayerId, FSlateCachedElementsHandle& CurrentCacheHandle);
 
 	/**
 	 * Pops the current painted widget off the stack
 	 * @return true if an element was added while the widget was pushed
 	 */
-	SLATECORE_API FSlateCachedElementListNode* PopPaintingWidget();
+	SLATECORE_API FSlateCachedElementsHandle PopPaintingWidget();
 
 	/** Pushes cached element data onto the stack.  Any draw elements cached after will use this cached element data until popped */
 	void PushCachedElementData(FSlateCachedElementData& CachedElementData);
@@ -643,14 +670,14 @@ private:
 	/** State of the current widget that is adding draw elements */
 	struct FWidgetDrawElementState
 	{
-		FWidgetDrawElementState(FSlateCachedElementListNode* InCurrentCacheNode,  bool bInIsVolatile, const SWidget* InWidget)
-			: CacheNode(InCurrentCacheNode)
+		FWidgetDrawElementState(FSlateCachedElementsHandle& InCurrentHandle, bool bInIsVolatile, const SWidget* InWidget)
+			: CacheHandle(InCurrentHandle)
 			, Widget(InWidget)
 			, bIsVolatile(bInIsVolatile)
 		{
 		}
 
-		FSlateCachedElementListNode* CacheNode;
+		FSlateCachedElementsHandle CacheHandle;
 		const SWidget* Widget;
 		bool bIsVolatile;
 	
diff --git a/Engine/Source/Runtime/SlateCore/Public/Rendering/ElementBatcher.h b/Engine/Source/Runtime/SlateCore/Public/Rendering/ElementBatcher.h
--- a/Engine/Source/Runtime/SlateCore/Public/Rendering/ElementBatcher.h
+++ b/Engine/Source/Runtime/SlateCore/Public/Rendering/ElementBatcher.h
@@ -206,7 +206,7 @@ public:
 		ESlateBatchDrawFlag InDrawFlags,
 		int8 SceneIndex);
 
-	void AddCachedBatches(const TArray<FSlateRenderBatch>& InCachedBatches);
+	void AddCachedBatches(const TSparseArray<FSlateRenderBatch>& InCachedBatches);
 private:
 	void FillBuffersFromNewBatch(FSlateRenderBatch& Batch, FSlateVertexArray& FinalVertices, FSlateIndexArray& FinalIndices);
 	void CombineBatches(FSlateRenderBatch& FirstBatch, FSlateRenderBatch& SecondBatch, FSlateVertexArray& FinalVertices, FSlateIndexArray& FinalIndices);
@@ -355,7 +355,7 @@ private:
 	FSlateBatchData* BatchData;
 
 	/** Cached batches currently being filled in */
-	FSlateCachedElementList* CachedElementList;
+	FSlateCachedElementList* CurrentCachedElementList;
 
 	/** Clipping states currently being used */
 	const TArray<FSlateClippingState>* PrecachedClippingStates;
diff --git a/Engine/Source/Runtime/SlateCore/Public/Rendering/SlateRenderBatch.h b/Engine/Source/Runtime/SlateCore/Public/Rendering/SlateRenderBatch.h
--- a/Engine/Source/Runtime/SlateCore/Public/Rendering/SlateRenderBatch.h
+++ b/Engine/Source/Runtime/SlateCore/Public/Rendering/SlateRenderBatch.h
@@ -16,7 +16,6 @@ struct FSlateDataPayload;
 class FSlateRenderBatch
 {
 public:
-	//FSlateRenderBatch(uint32 InLayer, const FSlateElementBatch& InBatch, TSharedPtr<FSlateRenderDataHandle, ESPMode::ThreadSafe> InRenderHandle, int32 InNumVertices, int32 InNumIndices, int32 InVertexOffset, int32 InIndexOffset);
 	FSlateRenderBatch(
 		int32 InLayer,
 		const FShaderParams& InShaderParams,
@@ -184,4 +183,7 @@ public:
 
 	uint8 bIsMerged : 1;
 
-};
\ No newline at end of file
+};
+
+static_assert(TIsTriviallyCopyConstructible<FSlateRenderBatch>::Value == true, "FSlateRenderBatch must be mem copyable");
+static_assert(TIsTriviallyDestructible<FSlateRenderBatch>::Value == true, "FSlateRenderBatch must be trivially destructible");
diff --git a/Engine/Source/Runtime/SlateCore/Public/Widgets/SWidget.h b/Engine/Source/Runtime/SlateCore/Public/Widgets/SWidget.h
--- a/Engine/Source/Runtime/SlateCore/Public/Widgets/SWidget.h
+++ b/Engine/Source/Runtime/SlateCore/Public/Widgets/SWidget.h
@@ -755,7 +755,7 @@ private:
 #endif
 	}
 
-	void UpdateWidgetProxy(int32 NewLayerId, FSlateCachedElementListNode* CacheNode);
+	void UpdateWidgetProxy(int32 NewLayerId, FSlateCachedElementsHandle& CacheHandle);
 
 #if WITH_SLATE_DEBUGGING
 	uint32 Debug_GetLastPaintFrame() const { return LastPaintFrame; }
