#UE4 Fix compile error

diff --git a/Engine/Source/Runtime/SlateCore/Private/FastUpdate/WidgetProxy.cpp b/Engine/Source/Runtime/SlateCore/Private/FastUpdate/WidgetProxy.cpp
--- a/Engine/Source/Runtime/SlateCore/Private/FastUpdate/WidgetProxy.cpp
+++ b/Engine/Source/Runtime/SlateCore/Private/FastUpdate/WidgetProxy.cpp
@@ -65,7 +65,7 @@ bool FWidgetProxy::ProcessInvalidation(FWidgetUpdateList& UpdateList, TArray<FWi
 	bool bWidgetNeedsRepaint = false;
 	if (!bInvisibleDueToParentOrSelfVisibility && ParentIndex != INDEX_NONE && !Widget->PrepassLayoutScaleMultiplier.IsSet())
 	{
-		SCOPE_CYCLE_SWIDGET(WidgetScope, Widget);
+		SCOPE_CYCLE_SWIDGET(Widget);
 		// If this widget has never been prepassed make sure the parent prepasses it to set the correct multiplier
 		FWidgetProxy& ParentProxy = FastWidgetPathList[ParentIndex];
 		if (ParentProxy.Widget)
@@ -79,7 +79,7 @@ bool FWidgetProxy::ProcessInvalidation(FWidgetUpdateList& UpdateList, TArray<FWi
 	}
 	else if (EnumHasAnyFlags(CurrentInvalidateReason, EInvalidateWidget::RenderTransform | EInvalidateWidget::Layout | EInvalidateWidget::Visibility | EInvalidateWidget::ChildOrder))
 	{
-		SCOPE_CYCLE_SWIDGET(WidgetScope, Widget);
+		SCOPE_CYCLE_SWIDGET(Widget);
 		// When layout changes compute a new desired size for this widget
 		FVector2D CurrentDesiredSize = Widget->GetDesiredSize();
 		FVector2D NewDesiredSize = FVector2D::ZeroVector;
@@ -134,7 +134,7 @@ bool FWidgetProxy::ProcessInvalidation(FWidgetUpdateList& UpdateList, TArray<FWi
 	}
 	else if (EnumHasAnyFlags(CurrentInvalidateReason, EInvalidateWidget::Paint) && !Widget->IsVolatileIndirectly())
 	{
-		SCOPE_CYCLE_SWIDGET(WidgetScope, Widget);
+		SCOPE_CYCLE_SWIDGET(Widget);
 		UpdateFlags |= EWidgetUpdateFlags::NeedsRepaint;
 
 		bWidgetNeedsRepaint = true;
