Slate - Tweaking some stats for caching elements to better note what it is, and fixing a potential crash due to a struct being marked as a POD type even though it has ctor and dtor needs.

Hoping this fixes the smoke tests.

diff --git a/Engine/Source/Runtime/SlateCore/Private/Rendering/ElementBatcher.cpp b/Engine/Source/Runtime/SlateCore/Private/Rendering/ElementBatcher.cpp
--- a/Engine/Source/Runtime/SlateCore/Private/Rendering/ElementBatcher.cpp
+++ b/Engine/Source/Runtime/SlateCore/Private/Rendering/ElementBatcher.cpp
@@ -24,7 +24,7 @@ DEFINE_STAT(STAT_SlateElements_Text);
 DEFINE_STAT(STAT_SlateElements_ShapedText);
 DEFINE_STAT(STAT_SlateElements_Line);
 DEFINE_STAT(STAT_SlateElements_Other);
-DEFINE_STAT(STAT_SlateCachedElements);
+DEFINE_STAT(STAT_SlateInvalidation_RecachedElements);
 
 int32 GSlateFeathering = 0;
 
@@ -283,7 +283,7 @@ void FSlateElementBatcher::AddElements(FSlateWindowElementList& WindowElementLis
 	ElementStat_Text = 0;
 	ElementStat_ShapedText = 0;
 	ElementStat_Line = 0;
-	ElementStat_CachedElements = 0;
+	ElementStat_RecachedElements = 0;
 #endif
 
 	BatchData = &WindowElementList.GetBatchData();
@@ -329,7 +329,7 @@ void FSlateElementBatcher::AddElements(FSlateWindowElementList& WindowElementLis
 	INC_DWORD_STAT_BY(STAT_SlateElements_ShapedText, ElementStat_ShapedText);
 	INC_DWORD_STAT_BY(STAT_SlateElements_Line, ElementStat_Line);
 	INC_DWORD_STAT_BY(STAT_SlateElements_Other, ElementStat_Other);
-	INC_DWORD_STAT_BY(STAT_SlateCachedElements, ElementStat_CachedElements);
+	INC_DWORD_STAT_BY(STAT_SlateInvalidation_RecachedElements, ElementStat_RecachedElements);
 #endif
 }
 
@@ -437,6 +437,7 @@ void FSlateElementBatcher::AddElementsInternal(const FSlateDrawElementArray& Dra
 void FSlateElementBatcher::AddCachedElements(FSlateCachedElementData& CachedElementData, const FVector2D& ViewportSize)
 {
 	CSV_SCOPED_TIMING_STAT(Slate, AddCachedElements);
+	SCOPED_NAMED_EVENT_TEXT("Slate::AddCachedElements", FColor::Magenta);
 
 #if SLATE_CSV_TRACKER
 	FCsvProfiler::RecordCustomStat("Paint/CacheListsWithNewData", CSV_CATEGORY_INDEX(Slate), CachedElementData.ListsWithNewData.Num(), ECsvCustomStatOp::Set);
@@ -444,13 +445,11 @@ void FSlateElementBatcher::AddCachedElements(FSlateCachedElementData& CachedElem
 	int32 RecachedEmptyDrawLists = 0;
 #endif
 
-	SCOPED_NAMED_EVENT_TEXT("Slate::AddCachedBatches", FColor::Magenta);
-
 	for (FSlateCachedElementList* List : CachedElementData.ListsWithNewData)
 	{
 		if (List->DrawElements.Num() > 0)
 		{
-			STAT(ElementStat_CachedElements += List->DrawElements.Num());
+			STAT(ElementStat_RecachedElements += List->DrawElements.Num());
 
 #if SLATE_CSV_TRACKER
 			RecachedDrawElements += List->DrawElements.Num();
diff --git a/Engine/Source/Runtime/SlateCore/Public/FastUpdate/WidgetProxy.h b/Engine/Source/Runtime/SlateCore/Public/FastUpdate/WidgetProxy.h
--- a/Engine/Source/Runtime/SlateCore/Public/FastUpdate/WidgetProxy.h
+++ b/Engine/Source/Runtime/SlateCore/Public/FastUpdate/WidgetProxy.h
@@ -163,7 +163,6 @@ struct FSlateWidgetPersistentState
 
 	static const FSlateWidgetPersistentState NoState;
 };
-template <> struct TIsPODType<FSlateWidgetPersistentState> { enum { Value = true }; };
 
 struct FWidgetStackData
 {
diff --git a/Engine/Source/Runtime/SlateCore/Public/Rendering/ElementBatcher.h b/Engine/Source/Runtime/SlateCore/Public/Rendering/ElementBatcher.h
--- a/Engine/Source/Runtime/SlateCore/Public/Rendering/ElementBatcher.h
+++ b/Engine/Source/Runtime/SlateCore/Public/Rendering/ElementBatcher.h
@@ -28,7 +28,7 @@ DECLARE_DWORD_COUNTER_STAT_EXTERN(TEXT("Elements (ShapedText)"), STAT_SlateEleme
 DECLARE_DWORD_COUNTER_STAT_EXTERN(TEXT("Elements (Line)"), STAT_SlateElements_Line, STATGROUP_Slate, SLATECORE_API);
 DECLARE_DWORD_COUNTER_STAT_EXTERN(TEXT("Elements (Other)"), STAT_SlateElements_Other, STATGROUP_Slate, SLATECORE_API);
 
-DECLARE_DWORD_COUNTER_STAT_EXTERN(TEXT("Cached Elements"), STAT_SlateCachedElements, STATGROUP_Slate, SLATECORE_API);
+DECLARE_DWORD_COUNTER_STAT_EXTERN(TEXT("Invalidation: Recached Elements"), STAT_SlateInvalidation_RecachedElements, STATGROUP_Slate, SLATECORE_API);
 
 extern int32 GSlateFeathering;
 
@@ -382,7 +382,7 @@ private:
 	/** Track the number of drawn batches from the previous frame to report to stats. */
 	int32 ElementStat_Other;
 
-	int32 ElementStat_CachedElements;
+	int32 ElementStat_RecachedElements;
 #endif
 
 	/** How many post process passes are needed */
