Insert a substitute glyph for narrow nbsp and thin space if the font doesn't support them

This uses 2/3rd the width of a normal space, and also moves the handling of zwsp into a common point rather than duplicate it between the kerning a HarfBuzz shaping paths, as well as giving the composite font first refusal at handling all whitespace as it may have better metrics than our own approximation.

diff --git a/Engine/Source/Runtime/SlateCore/Private/Fonts/SlateTextShaper.cpp b/Engine/Source/Runtime/SlateCore/Private/Fonts/SlateTextShaper.cpp
--- a/Engine/Source/Runtime/SlateCore/Private/Fonts/SlateTextShaper.cpp
+++ b/Engine/Source/Runtime/SlateCore/Private/Fonts/SlateTextShaper.cpp
@@ -12,6 +12,14 @@ DECLARE_CYCLE_STAT(TEXT("Shape Unidirectional Text"), STAT_SlateShapeUnidirectio
 namespace
 {
 
+bool RenderCodepointAsWhitespace(const UTF32CHAR InCodepoint)
+{
+	return FText::IsWhitespace(InCodepoint)
+		|| InCodepoint == TEXT('\u200B')	// Zero Width Space
+		|| InCodepoint == TEXT('\u2009')	// Thin Space
+		|| InCodepoint == TEXT('\u202F');	// Narrow No-Break Space
+}
+
 struct FKerningOnlyTextSequenceEntry
 {
 	int32 TextStartIndex;
@@ -298,24 +306,20 @@ void FSlateTextShaper::PerformKerningOnlyTextShaping(const TCHAR* InText, const
 		const int32 TextEndIndex = InTextStart + InTextLen;
 		for (; RunningTextIndex < TextEndIndex; ++RunningTextIndex)
 		{
-			TCHAR CurrentChar = InText[RunningTextIndex];
-
-			// Substitute whitespace characters with spaces since not all fonts support all kinds of whitespace characters (but we don't care since we don't render them anyway)
-			if (FText::IsWhitespace(CurrentChar))
-			{
-				CurrentChar = TEXT(' ');
-			}
+			const TCHAR CurrentChar = InText[RunningTextIndex];
+			const bool bShouldRenderAsWhitespace = RenderCodepointAsWhitespace(CurrentChar);
 
 			// First try with the actual character
 			float SubFontScalingFactor = 1.0f;
 			const FFontData* FontDataPtr = &CompositeFontCache->GetFontDataForCodepoint(InFontInfo, CurrentChar, SubFontScalingFactor);
-			FFreeTypeFaceGlyphData FaceGlyphData = FontRenderer->GetFontFaceForCodepoint(*FontDataPtr, CurrentChar, InFontInfo.FontFallback);
+			FFreeTypeFaceGlyphData FaceGlyphData = FontRenderer->GetFontFaceForCodepoint(*FontDataPtr, CurrentChar, bShouldRenderAsWhitespace ? EFontFallback::FF_NoFallback : InFontInfo.FontFallback);
 
-			// If none of our fonts can render that character (as the fallback font may be missing), try again with the fallback character
+			// If none of our fonts can render that character (as the fallback font may be missing), try again with the fallback character, or a normal space if this character was supposed to be whitespace (as we don't render whitespace anyway)
 			if (!FaceGlyphData.FaceAndMemory.IsValid())
 			{
-				FontDataPtr = &CompositeFontCache->GetFontDataForCodepoint(InFontInfo, SlateFontRendererUtils::InvalidSubChar, SubFontScalingFactor);
-				FaceGlyphData = FontRenderer->GetFontFaceForCodepoint(*FontDataPtr, SlateFontRendererUtils::InvalidSubChar, InFontInfo.FontFallback);
+				const TCHAR FallbackChar = bShouldRenderAsWhitespace ? TEXT(' ') : SlateFontRendererUtils::InvalidSubChar;
+				FontDataPtr = &CompositeFontCache->GetFontDataForCodepoint(InFontInfo, FallbackChar, SubFontScalingFactor);
+				FaceGlyphData = FontRenderer->GetFontFaceForCodepoint(*FontDataPtr, FallbackChar, InFontInfo.FontFallback);
 			}
 
 			// Only scalable font types can use sub-font scaling
@@ -364,9 +368,6 @@ void FSlateTextShaper::PerformKerningOnlyTextShaping(const TCHAR* InText, const
 
 				if (!InsertSubstituteGlyphs(InText, CurrentCharIndex, ShapedGlyphFaceData, OutGlyphsToRender))
 				{
-					const bool bIsZeroWidthSpace = CurrentChar == TEXT('\u200B');
-					const bool bIsWhitespace = bIsZeroWidthSpace || FText::IsWhitespace(CurrentChar);
-
 					uint32 GlyphIndex = FT_Get_Char_Index(KerningOnlyTextSequenceEntry.FaceAndMemory->GetFace(), CurrentChar);
 
 					// If the given font can't render that character (as the fallback font may be missing), try again with the fallback character
@@ -376,7 +377,6 @@ void FSlateTextShaper::PerformKerningOnlyTextShaping(const TCHAR* InText, const
 					}
 
 					int16 XAdvance = 0;
-					if (!bIsZeroWidthSpace)
 					{
 						FT_Fixed CachedAdvanceData = 0;
 						if (FTAdvanceCache->FindOrCache(KerningOnlyTextSequenceEntry.FaceAndMemory->GetFace(), GlyphIndex, GlyphFlags, InFontInfo.Size, FinalFontScale, CachedAdvanceData))
@@ -390,7 +390,7 @@ void FSlateTextShaper::PerformKerningOnlyTextShaping(const TCHAR* InText, const
 					ShapedGlyphEntry.FontFaceData = ShapedGlyphFaceData;
 					ShapedGlyphEntry.GlyphIndex = GlyphIndex;
 					ShapedGlyphEntry.SourceIndex = CurrentCharIndex;
-					ShapedGlyphEntry.XAdvance = bIsZeroWidthSpace ? 0 : XAdvance;
+					ShapedGlyphEntry.XAdvance = XAdvance;
 					ShapedGlyphEntry.YAdvance = 0;
 					ShapedGlyphEntry.XOffset = 0;
 					ShapedGlyphEntry.YOffset = 0;
@@ -398,7 +398,7 @@ void FSlateTextShaper::PerformKerningOnlyTextShaping(const TCHAR* InText, const
 					ShapedGlyphEntry.NumCharactersInGlyph = 1;
 					ShapedGlyphEntry.NumGraphemeClustersInGlyph = 1;
 					ShapedGlyphEntry.TextDirection = TextBiDi::ETextDirection::LeftToRight;
-					ShapedGlyphEntry.bIsVisible = !bIsWhitespace;
+					ShapedGlyphEntry.bIsVisible = !RenderCodepointAsWhitespace(CurrentChar);
 
 					// Apply the kerning against the previous entry
 					if (CurrentGlyphEntryIndex > 0 && bHasKerning && ShapedGlyphEntry.bIsVisible)
@@ -476,22 +476,19 @@ void FSlateTextShaper::PerformHarfBuzzTextShaping(const TCHAR* InText, const int
 			}
 #endif	// !PLATFORM_TCHAR_IS_4_BYTES
 
-			// Substitute whitespace characters with spaces since not all fonts support all kinds of whitespace characters (but we don't care since we don't render them anyway)
-			if (FText::IsWhitespace(CurrentCodepoint))
-			{
-				CurrentCodepoint = TEXT(' ');
-			}
+			const bool bShouldRenderAsWhitespace = RenderCodepointAsWhitespace(CurrentCodepoint);
 
 			// First try with the actual character
 			float SubFontScalingFactor = 1.0f;
 			const FFontData* FontDataPtr = &CompositeFontCache->GetFontDataForCodepoint(InFontInfo, CurrentCodepoint, SubFontScalingFactor);
-			FFreeTypeFaceGlyphData FaceGlyphData = FontRenderer->GetFontFaceForCodepoint(*FontDataPtr, CurrentCodepoint, InFontInfo.FontFallback);
+			FFreeTypeFaceGlyphData FaceGlyphData = FontRenderer->GetFontFaceForCodepoint(*FontDataPtr, CurrentCodepoint, bShouldRenderAsWhitespace ? EFontFallback::FF_NoFallback : InFontInfo.FontFallback);
 
-			// If none of our fonts can render that character (as the fallback font may be missing), try again with the fallback character
+			// If none of our fonts can render that character (as the fallback font may be missing), try again with the fallback character, or a normal space if this character was supposed to be whitespace (as we don't render whitespace anyway)
 			if (!FaceGlyphData.FaceAndMemory.IsValid())
 			{
-				FontDataPtr = &CompositeFontCache->GetFontDataForCodepoint(InFontInfo, SlateFontRendererUtils::InvalidSubChar, SubFontScalingFactor);
-				FaceGlyphData = FontRenderer->GetFontFaceForCodepoint(*FontDataPtr, SlateFontRendererUtils::InvalidSubChar, InFontInfo.FontFallback);
+				const UTF32CHAR FallbackChar = bShouldRenderAsWhitespace ? TEXT(' ') : SlateFontRendererUtils::InvalidSubChar;
+				FontDataPtr = &CompositeFontCache->GetFontDataForCodepoint(InFontInfo, FallbackChar, SubFontScalingFactor);
+				FaceGlyphData = FontRenderer->GetFontFaceForCodepoint(*FontDataPtr, FallbackChar, InFontInfo.FontFallback);
 			}
 
 			// Only scalable font types can use sub-font scaling
@@ -642,23 +639,20 @@ void FSlateTextShaper::PerformHarfBuzzTextShaping(const TCHAR* InText, const int
 					const TCHAR CurrentChar = InText[CurrentCharIndex];
 					if (!InsertSubstituteGlyphs(InText, CurrentCharIndex, ShapedGlyphFaceData, OutGlyphsToRender))
 					{
-						const bool bIsZeroWidthSpace = CurrentChar == TEXT('\u200B');
-						const bool bIsWhitespace = bIsZeroWidthSpace || FText::IsWhitespace(CurrentChar);
-
 						const int32 CurrentGlyphEntryIndex = OutGlyphsToRender.AddDefaulted();
 						FShapedGlyphEntry& ShapedGlyphEntry = OutGlyphsToRender[CurrentGlyphEntryIndex];
 						ShapedGlyphEntry.FontFaceData = ShapedGlyphFaceData;
 						ShapedGlyphEntry.GlyphIndex = HarfBuzzGlyphInfo.codepoint;
 						ShapedGlyphEntry.SourceIndex = CurrentCharIndex;
-						ShapedGlyphEntry.XAdvance = bIsZeroWidthSpace ? 0 : FreeTypeUtils::Convert26Dot6ToRoundedPixel<int16>(HarfBuzzGlyphPosition.x_advance);
-						ShapedGlyphEntry.YAdvance = bIsZeroWidthSpace ? 0 : -FreeTypeUtils::Convert26Dot6ToRoundedPixel<int16>(HarfBuzzGlyphPosition.y_advance);
-						ShapedGlyphEntry.XOffset = bIsZeroWidthSpace ? 0 : FreeTypeUtils::Convert26Dot6ToRoundedPixel<int16>(HarfBuzzGlyphPosition.x_offset);
-						ShapedGlyphEntry.YOffset = bIsZeroWidthSpace ? 0 : -FreeTypeUtils::Convert26Dot6ToRoundedPixel<int16>(HarfBuzzGlyphPosition.y_offset);
+						ShapedGlyphEntry.XAdvance = FreeTypeUtils::Convert26Dot6ToRoundedPixel<int16>(HarfBuzzGlyphPosition.x_advance);
+						ShapedGlyphEntry.YAdvance = -FreeTypeUtils::Convert26Dot6ToRoundedPixel<int16>(HarfBuzzGlyphPosition.y_advance);
+						ShapedGlyphEntry.XOffset = FreeTypeUtils::Convert26Dot6ToRoundedPixel<int16>(HarfBuzzGlyphPosition.x_offset);
+						ShapedGlyphEntry.YOffset = -FreeTypeUtils::Convert26Dot6ToRoundedPixel<int16>(HarfBuzzGlyphPosition.y_offset);
 						ShapedGlyphEntry.Kerning = 0;
 						ShapedGlyphEntry.NumCharactersInGlyph = 0; // Filled in later once we've processed each cluster
 						ShapedGlyphEntry.NumGraphemeClustersInGlyph = 0; // Filled in later once we have an accurate character count
 						ShapedGlyphEntry.TextDirection = InTextDirection;
-						ShapedGlyphEntry.bIsVisible = !bIsWhitespace;
+						ShapedGlyphEntry.bIsVisible = !RenderCodepointAsWhitespace(CurrentChar);
 
 						// Apply the kerning against the previous entry
 						if (CurrentGlyphEntryIndex > 0 && bHasKerning && ShapedGlyphEntry.bIsVisible)
@@ -796,11 +790,32 @@ void FSlateTextShaper::PerformHarfBuzzTextShaping(const TCHAR* InText, const int
 
 bool FSlateTextShaper::InsertSubstituteGlyphs(const TCHAR* InText, const int32 InCharIndex, const TSharedRef<FShapedGlyphFaceData>& InShapedGlyphFaceData, TArray<FShapedGlyphEntry>& OutGlyphsToRender) const
 {
+	auto GetSpaceGlyphIndexAndAdvance = [this, &InShapedGlyphFaceData](uint32& OutSpaceGlyphIndex, int16& OutSpaceXAdvance)
+	{
+		OutSpaceGlyphIndex = 0;
+		OutSpaceXAdvance = 0;
+#if WITH_FREETYPE
+		{
+			TSharedPtr<FFreeTypeFace> FTFace = InShapedGlyphFaceData->FontFace.Pin();
+			if (FTFace.IsValid())
+			{
+				OutSpaceGlyphIndex = FT_Get_Char_Index(FTFace->GetFace(), TEXT(' '));
+
+				FT_Fixed CachedAdvanceData = 0;
+				if (FTAdvanceCache->FindOrCache(FTFace->GetFace(), OutSpaceGlyphIndex, InShapedGlyphFaceData->GlyphFlags, InShapedGlyphFaceData->FontSize, InShapedGlyphFaceData->FontScale, CachedAdvanceData))
+				{
+					OutSpaceXAdvance = FreeTypeUtils::Convert26Dot6ToRoundedPixel<int16>((CachedAdvanceData + (1<<9)) >> 10);
+				}
+			}
+		}
+#endif // WITH_FREETYPE
+	};
+
 	const TCHAR Char = InText[InCharIndex];
 
-	if (TextBiDi::IsControlCharacter(Char))
+	if (TextBiDi::IsControlCharacter(Char) || Char == TEXT('\u200B'))	// Zero Width Space
 	{
-		// We insert a stub entry for control characters to avoid them being drawn as a visual glyph with size
+		// We insert a stub entry for control characters and zero-width spaces to avoid them being drawn as a visual glyph with size
 		FShapedGlyphEntry& ShapedGlyphEntry = OutGlyphsToRender.AddDefaulted_GetRef();
 		ShapedGlyphEntry.FontFaceData = InShapedGlyphFaceData;
 		ShapedGlyphEntry.GlyphIndex = 0;
@@ -821,21 +836,7 @@ bool FSlateTextShaper::InsertSubstituteGlyphs(const TCHAR* InText, const int32 I
 	{
 		uint32 SpaceGlyphIndex = 0;
 		int16 SpaceXAdvance = 0;
-#if WITH_FREETYPE
-		{
-			TSharedPtr<FFreeTypeFace> FTFace = InShapedGlyphFaceData->FontFace.Pin();
-			if (FTFace.IsValid())
-			{
-				SpaceGlyphIndex = FT_Get_Char_Index(FTFace->GetFace(), TEXT(' '));
-
-				FT_Fixed CachedAdvanceData = 0;
-				if (FTAdvanceCache->FindOrCache(FTFace->GetFace(), SpaceGlyphIndex, InShapedGlyphFaceData->GlyphFlags, InShapedGlyphFaceData->FontSize, InShapedGlyphFaceData->FontScale, CachedAdvanceData))
-				{
-					SpaceXAdvance = FreeTypeUtils::Convert26Dot6ToRoundedPixel<int16>((CachedAdvanceData + (1<<9)) >> 10);
-				}
-			}
-		}
-#endif // WITH_FREETYPE
+		GetSpaceGlyphIndexAndAdvance(SpaceGlyphIndex, SpaceXAdvance);
 
 		// We insert a spacer glyph with (up-to) the width of 4 space glyphs in-place of a tab character
 		const int32 NumSpacesToInsert = 4 - (OutGlyphsToRender.Num() % 4);
@@ -859,5 +860,46 @@ bool FSlateTextShaper::InsertSubstituteGlyphs(const TCHAR* InText, const int32 I
 		return true;
 	}
 
+	if (Char == TEXT('\u2009') ||	// Thin Space
+		Char == TEXT('\u202F')		// Narrow No-Break Space
+		)
+	{
+		// Not all fonts support these characters
+#if WITH_FREETYPE
+		{
+			TSharedPtr<FFreeTypeFace> FTFace = InShapedGlyphFaceData->FontFace.Pin();
+			if (FTFace.IsValid())
+			{
+				const uint32 GlyphIndex = FT_Get_Char_Index(FTFace->GetFace(), Char);
+				if (GlyphIndex != 0)
+				{
+					// If it does, then let it render the character itself as it may have better metrics
+					return false;
+				}
+			}
+		}
+#endif // WITH_FREETYPE
+
+		// If it doesn't, then make these 2/3rd the width of a normal space
+		uint32 SpaceGlyphIndex = 0;
+		int16 SpaceXAdvance = 0;
+		GetSpaceGlyphIndexAndAdvance(SpaceGlyphIndex, SpaceXAdvance);
+
+		FShapedGlyphEntry& ShapedGlyphEntry = OutGlyphsToRender.AddDefaulted_GetRef();
+		ShapedGlyphEntry.FontFaceData = InShapedGlyphFaceData;
+		ShapedGlyphEntry.GlyphIndex = SpaceGlyphIndex;
+		ShapedGlyphEntry.SourceIndex = InCharIndex;
+		ShapedGlyphEntry.XAdvance = (SpaceXAdvance * 2) / 3;
+		ShapedGlyphEntry.YAdvance = 0;
+		ShapedGlyphEntry.XOffset = 0;
+		ShapedGlyphEntry.YOffset = 0;
+		ShapedGlyphEntry.Kerning = 0;
+		ShapedGlyphEntry.NumCharactersInGlyph = 1;
+		ShapedGlyphEntry.NumGraphemeClustersInGlyph = 1;
+		ShapedGlyphEntry.TextDirection = TextBiDi::ETextDirection::LeftToRight;
+		ShapedGlyphEntry.bIsVisible = false;
+		return true;
+	}
+
 	return false;
 }
