Slate - Simplify the SWidget scope to just use anon variables.

diff --git a/Engine/Source/Runtime/SlateCore/Private/FastUpdate/WidgetProxy.cpp b/Engine/Source/Runtime/SlateCore/Private/FastUpdate/WidgetProxy.cpp
--- a/Engine/Source/Runtime/SlateCore/Private/FastUpdate/WidgetProxy.cpp
+++ b/Engine/Source/Runtime/SlateCore/Private/FastUpdate/WidgetProxy.cpp
@@ -65,7 +65,7 @@ bool FWidgetProxy::ProcessInvalidation(FWidgetUpdateList& UpdateList, TArray<FWi
 	bool bWidgetNeedsRepaint = false;
 	if (!bInvisibleDueToParentOrSelfVisibility && ParentIndex != INDEX_NONE && !Widget->PrepassLayoutScaleMultiplier.IsSet())
 	{
-		SCOPE_CYCLE_SWIDGET(WidgetScope, Widget);
+		SCOPE_CYCLE_SWIDGET(Widget);
 		// If this widget has never been prepassed make sure the parent prepasses it to set the correct multiplier
 		FWidgetProxy& ParentProxy = FastWidgetPathList[ParentIndex];
 		if (ParentProxy.Widget)
@@ -79,7 +79,7 @@ bool FWidgetProxy::ProcessInvalidation(FWidgetUpdateList& UpdateList, TArray<FWi
 	}
 	else if (EnumHasAnyFlags(CurrentInvalidateReason, EInvalidateWidget::RenderTransform | EInvalidateWidget::Layout | EInvalidateWidget::Visibility | EInvalidateWidget::ChildOrder))
 	{
-		SCOPE_CYCLE_SWIDGET(WidgetScope, Widget);
+		SCOPE_CYCLE_SWIDGET(Widget);
 		// When layout changes compute a new desired size for this widget
 		FVector2D CurrentDesiredSize = Widget->GetDesiredSize();
 		FVector2D NewDesiredSize = FVector2D::ZeroVector;
@@ -134,7 +134,7 @@ bool FWidgetProxy::ProcessInvalidation(FWidgetUpdateList& UpdateList, TArray<FWi
 	}
 	else if (EnumHasAnyFlags(CurrentInvalidateReason, EInvalidateWidget::Paint) && !Widget->IsVolatileIndirectly())
 	{
-		SCOPE_CYCLE_SWIDGET(WidgetScope, Widget);
+		SCOPE_CYCLE_SWIDGET(Widget);
 		UpdateFlags |= EWidgetUpdateFlags::NeedsRepaint;
 
 		bWidgetNeedsRepaint = true;
diff --git a/Engine/Source/Runtime/SlateCore/Private/Rendering/ElementBatcher.cpp b/Engine/Source/Runtime/SlateCore/Private/Rendering/ElementBatcher.cpp
--- a/Engine/Source/Runtime/SlateCore/Private/Rendering/ElementBatcher.cpp
+++ b/Engine/Source/Runtime/SlateCore/Private/Rendering/ElementBatcher.cpp
@@ -580,7 +580,7 @@ void FSlateElementBatcher::AddCachedElements(FSlateCachedElementData& CachedElem
 		{
 			CachedElementList = &LocalCachedElementList;
 			{
-				SCOPE_CYCLE_SWIDGET(WidgetScope, LocalCachedElementList.Widget);
+				SCOPE_CYCLE_SWIDGET(LocalCachedElementList.Widget);
 				AddElementsInternal(LocalCachedElementList.DrawElements, ViewportSize);
 			}
 
diff --git a/Engine/Source/Runtime/SlateCore/Private/Widgets/SWidget.cpp b/Engine/Source/Runtime/SlateCore/Private/Widgets/SWidget.cpp
--- a/Engine/Source/Runtime/SlateCore/Private/Widgets/SWidget.cpp
+++ b/Engine/Source/Runtime/SlateCore/Private/Widgets/SWidget.cpp
@@ -1221,7 +1221,7 @@ int32 SWidget::Paint(const FPaintArgs& Args, const FGeometry& AllottedGeometry,
 		INC_DWORD_STAT(STAT_SlateNumTickedWidgets);
 
 		SCOPE_CYCLE_COUNTER(STAT_SlateTickWidgets);
-		SCOPE_CYCLE_SWIDGET(WidgetTick, this);
+		SCOPE_CYCLE_SWIDGET(this);
 		MutableThis->Tick(DesktopSpaceGeometry, Args.GetCurrentTime(), Args.GetDeltaTime());
 	}
 
diff --git a/Engine/Source/Runtime/SlateCore/Public/Widgets/SWidgetUtils.h b/Engine/Source/Runtime/SlateCore/Public/Widgets/SWidgetUtils.h
--- a/Engine/Source/Runtime/SlateCore/Public/Widgets/SWidgetUtils.h
+++ b/Engine/Source/Runtime/SlateCore/Public/Widgets/SWidgetUtils.h
@@ -34,8 +34,8 @@ struct FScopeCycleCounterSWidget : public FCycleCounter
 	}
 };
 
-#define SCOPE_CYCLE_SWIDGET(Name, Object) \
-	FScopeCycleCounterSWidget ObjCycleCount_##Name(Object);
+#define SCOPE_CYCLE_SWIDGET(Object) \
+	FScopeCycleCounterSWidget ANONYMOUS_VARIABLE(SlateWidgetCycleCount_) (Object);
 
 #elif ENABLE_STATNAMEDEVENTS
 
@@ -48,11 +48,11 @@ struct FScopeCycleCounterSWidget
 	}
 };
 
-#define SCOPE_CYCLE_SWIDGET(Name, Object) \
-	FScopeCycleCounterSWidget ObjCycleCount_##Name(Object);
+#define SCOPE_CYCLE_SWIDGET(Object) \
+	FScopeCycleCounterSWidget ANONYMOUS_VARIABLE(SlateWidgetCycleCount_) (Object);
 
 #else
 
-#define SCOPE_CYCLE_SWIDGET(Name, Object)
+#define SCOPE_CYCLE_SWIDGET(Object)
 
 #endif
