Fixed fonts that contain both scalable and bitmap font data incorrectly applying extra scale to advance and kerning data

diff --git a/Engine/Source/Runtime/SlateCore/Private/Fonts/FontCacheFreeType.cpp b/Engine/Source/Runtime/SlateCore/Private/Fonts/FontCacheFreeType.cpp
--- a/Engine/Source/Runtime/SlateCore/Private/Fonts/FontCacheFreeType.cpp
+++ b/Engine/Source/Runtime/SlateCore/Private/Fonts/FontCacheFreeType.cpp
@@ -636,7 +636,7 @@ bool FFreeTypeAdvanceCache::FindOrCache(FT_Face InFace, const uint32 InGlyphInde
 	FT_Error Error = FT_Get_Advance(InFace, InGlyphIndex, InLoadFlags, &OutCachedAdvance);
 	if (Error == 0)
 	{
-		if (FT_HAS_FIXED_SIZES(InFace))
+		if (!FT_IS_SCALABLE(InFace) && FT_HAS_FIXED_SIZES(InFace))
 		{
 			// Fixed size fonts don't support scaling, but we calculated the scale to use for the glyph in ApplySizeAndScale
 			OutCachedAdvance = FT_MulFix(OutCachedAdvance, ((InLoadFlags & FT_LOAD_VERTICAL_LAYOUT) ? InFace->size->metrics.y_scale : InFace->size->metrics.x_scale));
@@ -691,7 +691,7 @@ bool FFreeTypeKerningPairCache::FindOrCache(FT_Face InFace, const FKerningPair&
 	{
 		if (InKerningFlags != FT_KERNING_UNSCALED)
 		{
-			if (FT_HAS_FIXED_SIZES(InFace))
+			if (!FT_IS_SCALABLE(InFace) && FT_HAS_FIXED_SIZES(InFace))
 			{
 				// Fixed size fonts don't support scaling, but we calculated the scale to use for the glyph in ApplySizeAndScale
 				OutKerning.x = FT_MulFix(OutKerning.x, InFace->size->metrics.x_scale);
