Changed spammy "Resetting cached element data" from a normal log to Verbose.

diff --git a/Engine/Source/Runtime/SlateCore/Private/Rendering/DrawElements.cpp b/Engine/Source/Runtime/SlateCore/Private/Rendering/DrawElements.cpp
--- a/Engine/Source/Runtime/SlateCore/Private/Rendering/DrawElements.cpp
+++ b/Engine/Source/Runtime/SlateCore/Private/Rendering/DrawElements.cpp
@@ -925,7 +925,7 @@ void FSlateCachedElementData::Empty()
 {
 	if (CachedElementLists.Num())
 	{
-		UE_LOG(LogSlate, Log, TEXT("Resetting cached element data.  Num: %d"), CachedElementLists.Num());
+		UE_LOG(LogSlate, Verbose, TEXT("Resetting cached element data.  Num: %d"), CachedElementLists.Num());
 	}
 
 #if WITH_SLATE_DEBUGGING
