When ordering SWindows, place regular windows after non-regular windows

diff --git a/Engine/Source/Runtime/SlateCore/Private/Application/SlateWindowHelper.cpp b/Engine/Source/Runtime/SlateCore/Private/Application/SlateWindowHelper.cpp
--- a/Engine/Source/Runtime/SlateCore/Private/Application/SlateWindowHelper.cpp
+++ b/Engine/Source/Runtime/SlateCore/Private/Application/SlateWindowHelper.cpp
@@ -24,21 +24,17 @@ void FSlateWindowHelper::ArrangeWindowToFront( TArray< TSharedRef<SWindow> >& Wi
 	{
 		bool PerformedInsert = false;
 
-		for (int WindowIndex = Windows.Num() - 1; WindowIndex >= 0; --WindowIndex)
+		int32 WindowIndex = Windows.Num() - 1;
+		for (; WindowIndex >= 0; --WindowIndex)
 		{
-			if (!Windows[WindowIndex]->IsTopmostWindow())
+			// Topmost windows first, then non-regular windows (popups), then regular windows.
+			if (!Windows[WindowIndex]->IsTopmostWindow() && (!WindowToBringToFront->IsRegularWindow() || Windows[WindowIndex]->IsRegularWindow()))
 			{
-				Windows.Insert(WindowToBringToFront, WindowIndex + 1);
-				PerformedInsert = true;
-
 				break;
 			}
 		}
 
-		if (!PerformedInsert)
-		{
-			Windows.Insert(WindowToBringToFront, 0);
-		}
+		Windows.Insert(WindowToBringToFront, WindowIndex + 1);
 	}
 }
 
