Revert 10883777, clamping negative scales in Slate layouts

diff --git a/Engine/Source/Runtime/SlateCore/Public/Layout/LayoutUtils.h b/Engine/Source/Runtime/SlateCore/Public/Layout/LayoutUtils.h
--- a/Engine/Source/Runtime/SlateCore/Public/Layout/LayoutUtils.h
+++ b/Engine/Source/Runtime/SlateCore/Public/Layout/LayoutUtils.h
@@ -94,10 +94,10 @@ static AlignmentArrangeResult AlignChild(EFlowDirection InLayoutFlow, float Allo
 	switch (Alignment)
 	{
 	case HAlign_Fill:
-		return AlignmentArrangeResult(MarginPre, FMath::Max(0.f, (AllottedSize - TotalMargin) * ContentScale));
+		return AlignmentArrangeResult(MarginPre, (AllottedSize - TotalMargin) * ContentScale);
 	}
 	
-	const float ChildSize = FMath::Max(0.f, bClampToParent ? FMath::Min(ChildDesiredSize, AllottedSize - TotalMargin) : ChildDesiredSize);
+	const float ChildSize = bClampToParent ? FMath::Min(ChildDesiredSize, AllottedSize - TotalMargin) : ChildDesiredSize;
 
 	switch( Alignment )
 	{
@@ -110,7 +110,7 @@ static AlignmentArrangeResult AlignChild(EFlowDirection InLayoutFlow, float Allo
 	}
 
 	// Same as Fill
-	return AlignmentArrangeResult(MarginPre, FMath::Max(0.f, ( AllottedSize - TotalMargin ) * ContentScale));
+	return AlignmentArrangeResult(MarginPre, ( AllottedSize - TotalMargin ) * ContentScale);
 }
 
 template<EOrientation Orientation, typename SlotType>
@@ -132,14 +132,14 @@ static AlignmentArrangeResult AlignChild(EFlowDirection InLayoutFlow, float Allo
 	switch (Alignment)
 	{
 	case HAlign_Fill:
-		return AlignmentArrangeResult(MarginPre, FMath::Max(0.f, (AllottedSize - TotalMargin) * ContentScale));
+		return AlignmentArrangeResult(MarginPre, (AllottedSize - TotalMargin) * ContentScale);
 	}
 
 	const float ChildDesiredSize = ( Orientation == Orient_Horizontal )
 		? ( ChildToArrange.GetWidget()->GetDesiredSize().X * ContentScale )
 		: ( ChildToArrange.GetWidget()->GetDesiredSize().Y * ContentScale );
 
-	const float ChildSize = FMath::Max(0.f, bClampToParent ? FMath::Min(ChildDesiredSize, AllottedSize - TotalMargin) : ChildDesiredSize);
+	const float ChildSize = bClampToParent ? FMath::Min(ChildDesiredSize, AllottedSize - TotalMargin) : ChildDesiredSize;
 
 	switch ( Alignment )
 	{
@@ -152,7 +152,7 @@ static AlignmentArrangeResult AlignChild(EFlowDirection InLayoutFlow, float Allo
 	}
 
 	// Same as Fill
-	return AlignmentArrangeResult(MarginPre, FMath::Max(0.f, (AllottedSize - TotalMargin) * ContentScale));
+	return AlignmentArrangeResult(MarginPre, (AllottedSize - TotalMargin) * ContentScale);
 }
 
 template<EOrientation Orientation, typename SlotType>
