Allow slate events to bind lambdas which take additional parameters.
Add a move function for assigning a delegate to a slate event.

diff --git a/Engine/Source/Runtime/SlateCore/Public/Widgets/DeclarativeSyntaxSupport.h b/Engine/Source/Runtime/SlateCore/Public/Widgets/DeclarativeSyntaxSupport.h
--- a/Engine/Source/Runtime/SlateCore/Public/Widgets/DeclarativeSyntaxSupport.h
+++ b/Engine/Source/Runtime/SlateCore/Public/Widgets/DeclarativeSyntaxSupport.h
@@ -478,6 +478,12 @@ struct NamedSlotProperty
 			return *this; \
 		} \
 		\
+		WidgetArgsType& EventName( DelegateName&& InDelegate ) \
+		{ \
+			_##EventName = MoveTemp(InDelegate); \
+			return *this; \
+		} \
+		\
 		/* Set event delegate to a global function */ \
 		/* NOTE: We use a template here to avoid 'typename' issues when hosting attributes inside templated classes */ \
 		template< typename StaticFuncPtr > \
@@ -513,10 +519,10 @@ struct NamedSlotProperty
 		\
 		/* Set event delegate to a lambda
 		 * technically this works for any functor types, but lambdas are the primary use case */ \
-		template<typename FunctorType> \
-		WidgetArgsType& EventName##_Lambda(FunctorType&& InFunctor) \
+		template<typename FunctorType, typename... VarTypes> \
+		WidgetArgsType& EventName##_Lambda(FunctorType&& InFunctor, VarTypes... Vars) \
 		{ \
-			_##EventName = DelegateName::CreateLambda(Forward<FunctorType>(InFunctor)); \
+			_##EventName = DelegateName::CreateLambda(Forward<FunctorType>(InFunctor), Vars... ); \
 			return this->Me(); \
 		} \
 		\
