Slate - When prepassing children, if a child is collapsed we invalidate its prepass (so that we know we need to do it later), and we give that child its new scale, so that when it does prepass it prepasses when it becomes visible, with the new layout scale it missed because it was collapsed, this corrects issues with layout scale changing and collapsed text fields missing out on the new scale for prepass/CDS purposes - usually revealing themselves in rich text now correctly laying out, large gaps between blocks, or several overlapping blocks.

Slate - Deprecating float GetRelativeLayoutScale(const FSlotBase& Child, float LayoutScaleMultiplier), Long Live float GetRelativeLayoutScale(const int32 ChildIndex, float LayoutScaleMultiplier), you must also set bHasRelativeLayoutScale to true in your constructor for your SWidget if you need this function called.  This removes 2 virtual calls from almost every widget's prepass, and only now requires 1 for every widget that needs to implement GetRelativeLayoutScale, which very few actually do.

Matt.Kuhlenschmidt

diff --git a/Engine/Source/Editor/GraphEditor/Private/SNodePanel.cpp b/Engine/Source/Editor/GraphEditor/Private/SNodePanel.cpp
--- a/Engine/Source/Editor/GraphEditor/Private/SNodePanel.cpp
+++ b/Engine/Source/Editor/GraphEditor/Private/SNodePanel.cpp
@@ -200,9 +200,10 @@ namespace NodePanelDefs
 };
 
 SNodePanel::SNodePanel()
-: Children(this)
-, VisibleChildren(this)
+	: Children(this)
+	, VisibleChildren(this)
 {
+	bHasRelativeLayoutScale = true;
 }
 
 void SNodePanel::OnArrangeChildren(const FGeometry& AllottedGeometry, FArrangedChildren& ArrangedChildren) const
@@ -1063,7 +1064,7 @@ FReply SNodePanel::OnTouchEnded( const FGeometry& MyGeometry, const FPointerEven
 	return FReply::Unhandled();
 }
 
-float SNodePanel::GetRelativeLayoutScale(const FSlotBase& Child, float LayoutScaleMultiplier) const
+float SNodePanel::GetRelativeLayoutScale(int32 ChildIndex, float LayoutScaleMultiplier) const
 {
 	return GetZoomAmount();
 }
diff --git a/Engine/Source/Editor/GraphEditor/Public/SNodePanel.h b/Engine/Source/Editor/GraphEditor/Public/SNodePanel.h
--- a/Engine/Source/Editor/GraphEditor/Public/SNodePanel.h
+++ b/Engine/Source/Editor/GraphEditor/Public/SNodePanel.h
@@ -353,16 +353,16 @@ public:
 			return FVector2D::ZeroVector;
 		}
 
-		virtual float GetRelativeLayoutScale(const FSlotBase& Child, float LayoutScaleMultiplier) const override
+		virtual float GetRelativeLayoutScale(int32 ChildIndex, float LayoutScaleMultiplier) const override
 		{
-			const FNodeSlot& ThisSlot = static_cast<const FNodeSlot&>(Child);
-			if ( !ThisSlot.AllowScale.Get() )
+			const FNodeSlot& ThisSlot = Children[ChildIndex];
+			if (!ThisSlot.AllowScale.Get())
 			{
 				// Child slots that do not allow zooming should scale themselves to negate the node panel's zoom.
 				TSharedPtr<SNodePanel> ParentPanel = GetParentPanel();
 				if (ParentPanel.IsValid())
-				{					
-					return 1.0f/ParentPanel->GetZoomAmount();
+				{
+					return 1.0f / ParentPanel->GetZoomAmount();
 				}
 			}
 
@@ -396,7 +396,7 @@ public:
 						CurChild.GetWidget(),
 						CurChild.Offset.Get(),
 						Size,
-						GetRelativeLayoutScale(CurChild, AllottedGeometry.Scale)
+						GetRelativeLayoutScale(ChildIndex, AllottedGeometry.Scale)
 					);
 					ArrangedChildren.AddWidget( ChildVisibility, ChildGeom );
 				}
@@ -580,6 +580,7 @@ public:
 		, DesiredSizeScale(FVector2D(1,1))
 		, Children(this)
 		{
+			bHasRelativeLayoutScale = true;
 		}
 
 	protected:
@@ -638,7 +639,7 @@ public:
 	virtual void OnFocusLost( const FFocusEvent& InFocusEvent ) override;
 	virtual FReply OnTouchGesture( const FGeometry& MyGeometry, const FPointerEvent& GestureEvent ) override;
 	virtual FReply OnTouchEnded( const FGeometry& MyGeometry, const FPointerEvent& InTouchEvent ) override;
-	virtual float GetRelativeLayoutScale(const FSlotBase& Child, float LayoutScaleMultiplier) const override;
+	virtual float GetRelativeLayoutScale(int32 ChildIndex, float LayoutScaleMultiplier) const override;
 	// End of SWidget interface
 public:
 	/**
diff --git a/Engine/Source/Editor/UMGEditor/Private/Designer/SZoomPan.cpp b/Engine/Source/Editor/UMGEditor/Private/Designer/SZoomPan.cpp
--- a/Engine/Source/Editor/UMGEditor/Private/Designer/SZoomPan.cpp
+++ b/Engine/Source/Editor/UMGEditor/Private/Designer/SZoomPan.cpp
@@ -8,6 +8,8 @@
 
 void SZoomPan::Construct(const FArguments& InArgs)
 {
+	bHasRelativeLayoutScale = true;
+
 	ViewOffset = InArgs._ViewOffset;
 	ZoomAmount = InArgs._ZoomAmount;
 
@@ -20,9 +22,9 @@ void SZoomPan::Construct(const FArguments& InArgs)
 void SZoomPan::SetContent(const TSharedRef< SWidget >& InContent)
 {
 	ChildSlot
-		[
-			InContent
-		];
+	[
+		InContent
+	];
 }
 
 void SZoomPan::OnArrangeChildren(const FGeometry& AllottedGeometry, FArrangedChildren& ArrangedChildren) const
@@ -43,7 +45,7 @@ void SZoomPan::OnArrangeChildren(const FGeometry& AllottedGeometry, FArrangedChi
 	}
 }
 
-float SZoomPan::GetRelativeLayoutScale(const FSlotBase& Child, float LayoutScaleMultiplier) const
+float SZoomPan::GetRelativeLayoutScale(int32 ChildIndex, float LayoutScaleMultiplier) const
 {
 	return ZoomAmount.Get();
 }
diff --git a/Engine/Source/Editor/UMGEditor/Private/Designer/SZoomPan.h b/Engine/Source/Editor/UMGEditor/Private/Designer/SZoomPan.h
--- a/Engine/Source/Editor/UMGEditor/Private/Designer/SZoomPan.h
+++ b/Engine/Source/Editor/UMGEditor/Private/Designer/SZoomPan.h
@@ -35,7 +35,7 @@ public:
 protected:
 	void OnArrangeChildren(const FGeometry& AllottedGeometry, FArrangedChildren& ArrangedChildren) const override;
 
-	virtual float GetRelativeLayoutScale( const FSlotBase& Child, float LayoutScaleMultiplier) const override;
+	virtual float GetRelativeLayoutScale(int32 ChildIndex, float LayoutScaleMultiplier) const override;
 
 	/** The position within the panel at which the user is looking */
 	TAttribute<FVector2D> ViewOffset;
diff --git a/Engine/Source/Runtime/Slate/Private/Widgets/Layout/SScaleBox.cpp b/Engine/Source/Runtime/Slate/Private/Widgets/Layout/SScaleBox.cpp
--- a/Engine/Source/Runtime/Slate/Private/Widgets/Layout/SScaleBox.cpp
+++ b/Engine/Source/Runtime/Slate/Private/Widgets/Layout/SScaleBox.cpp
@@ -14,6 +14,7 @@
 void SScaleBox::Construct(const SScaleBox::FArguments& InArgs)
 {
 	bHasCustomPrepass = true;
+	bHasRelativeLayoutScale = true;
 
 	Stretch = InArgs._Stretch;
 
@@ -400,7 +401,7 @@ FVector2D SScaleBox::ComputeDesiredSize(float InScale) const
 	return SCompoundWidget::ComputeDesiredSize(InScale);
 }
 
-float SScaleBox::GetRelativeLayoutScale(const FSlotBase& Child, float LayoutScaleMultiplier) const
+float SScaleBox::GetRelativeLayoutScale(int32 ChildIndex, float LayoutScaleMultiplier) const
 {
 	return ComputedContentScale.IsSet() ? ComputedContentScale.GetValue() : 1.0f;
 }
diff --git a/Engine/Source/Runtime/Slate/Public/Widgets/Layout/SDPIScaler.h b/Engine/Source/Runtime/Slate/Public/Widgets/Layout/SDPIScaler.h
--- a/Engine/Source/Runtime/Slate/Public/Widgets/Layout/SDPIScaler.h
+++ b/Engine/Source/Runtime/Slate/Public/Widgets/Layout/SDPIScaler.h
@@ -54,7 +54,7 @@ public:
 
 protected:
 
-	virtual float GetRelativeLayoutScale(const FSlotBase& Child, float LayoutScaleMultiplier) const override;
+	virtual float GetRelativeLayoutScale(int32 ChildIndex, float LayoutScaleMultiplier) const override;
 
 	/**
 	 * The scaling factor from 1:1 pixel to slate unit (SU) ratio.
diff --git a/Engine/Source/Runtime/Slate/Public/Widgets/Layout/SScaleBox.h b/Engine/Source/Runtime/Slate/Public/Widgets/Layout/SScaleBox.h
--- a/Engine/Source/Runtime/Slate/Public/Widgets/Layout/SScaleBox.h
+++ b/Engine/Source/Runtime/Slate/Public/Widgets/Layout/SScaleBox.h
@@ -156,7 +156,7 @@ protected:
 	// Begin SWidget overrides.
 	virtual bool CustomPrepass(float LayoutScaleMultiplier) override;
 	virtual FVector2D ComputeDesiredSize(float InScale) const override;
-	virtual float GetRelativeLayoutScale(const FSlotBase& Child, float LayoutScaleMultiplier) const override;
+	virtual float GetRelativeLayoutScale(int32 ChildIndex, float LayoutScaleMultiplier) const override;
 	// End SWidget overrides.
 
 	bool DoesScaleRequireNormalizingPrepassOrLocalGeometry() const;
diff --git a/Engine/Source/Runtime/SlateCore/Private/Widgets/SWidget.cpp b/Engine/Source/Runtime/SlateCore/Private/Widgets/SWidget.cpp
--- a/Engine/Source/Runtime/SlateCore/Private/Widgets/SWidget.cpp
+++ b/Engine/Source/Runtime/SlateCore/Private/Widgets/SWidget.cpp
@@ -185,6 +185,7 @@ SWidget::SWidget()
 	, bNeedsDesiredSize(true)
 	, bUpdatingDesiredSize(false)
 	, bHasCustomPrepass(false)
+	, bHasRelativeLayoutScale(false)
 	, bVolatilityAlwaysInvalidatesPrepass(false)
 #if WITH_ACCESSIBILITY
 	, bCanChildrenBeAccessible(true)
@@ -1415,7 +1416,7 @@ int32 SWidget::Paint(const FPaintArgs& Args, const FGeometry& AllottedGeometry,
 
 }
 
-float SWidget::GetRelativeLayoutScale(const FSlotBase& Child, float LayoutScaleMultiplier) const
+float SWidget::GetRelativeLayoutScale(int32 ChildIndex, float LayoutScaleMultiplier) const
 {
 	return 1.0f;
 }
@@ -1446,13 +1447,24 @@ void SWidget::Prepass_Internal(float InLayoutScaleMultiplier)
 		const int32 NumChildren = MyChildren->Num();
 		for (int32 ChildIndex = 0; ChildIndex < MyChildren->Num(); ++ChildIndex)
 		{
+			const float ChildLayoutScaleMultiplier = bHasRelativeLayoutScale
+				? InLayoutScaleMultiplier * GetRelativeLayoutScale(ChildIndex, InLayoutScaleMultiplier)
+				: InLayoutScaleMultiplier;
+
 			const TSharedRef<SWidget>& Child = MyChildren->GetChildAt(ChildIndex);
 
 			if (Child->Visibility.Get() != EVisibility::Collapsed)
 			{
-				const float ChildLayoutScaleMultiplier = GetRelativeLayoutScale(MyChildren->GetSlotAt(ChildIndex), InLayoutScaleMultiplier);
 				// Recur: Descend down the widget tree.
-				Child->Prepass_Internal(InLayoutScaleMultiplier * ChildLayoutScaleMultiplier);
+				Child->Prepass_Internal(ChildLayoutScaleMultiplier);
+			}
+			else
+			{
+				// If the child widget is collapsed, we need to store the new layout scale it will have when 
+				// it is finally visible and invalidate it's prepass so that it gets that when its visiblity
+				// is finally invalidated.
+				Child->InvalidatePrepass();
+				Child->PrepassLayoutScaleMultiplier = ChildLayoutScaleMultiplier;
 			}
 		}
 		ensure(NumChildren == MyChildren->Num());
diff --git a/Engine/Source/Runtime/SlateCore/Public/Widgets/SWidget.h b/Engine/Source/Runtime/SlateCore/Public/Widgets/SWidget.h
--- a/Engine/Source/Runtime/SlateCore/Public/Widgets/SWidget.h
+++ b/Engine/Source/Runtime/SlateCore/Public/Widgets/SWidget.h
@@ -786,8 +786,11 @@ public:
 		return TStatId(); // not doing stats at the moment, or ever
 	}
 
+	UE_DEPRECATED(4.24, "GetRelativeLayoutScale(int32 ChildIndex, float LayoutScaleMultiplier), your widget will also need to set bHasRelativeLayoutScale in their Construct/ctor.")
+	virtual float GetRelativeLayoutScale(const FSlotBase& Child, float LayoutScaleMultiplier) const { return 1.0f; }
+
 	/** What is the Child's scale relative to this widget. */
-	virtual float GetRelativeLayoutScale(const FSlotBase& Child, float LayoutScaleMultiplier) const;
+	virtual float GetRelativeLayoutScale(const int32 ChildIndex, float LayoutScaleMultiplier) const;
 
 	/**
 	 * Non-virtual entry point for arrange children. ensures common work is executed before calling the virtual
@@ -1560,6 +1563,8 @@ private:
 
 protected:
 	uint8 bHasCustomPrepass : 1;
+
+	uint8 bHasRelativeLayoutScale : 1;
 	
 	/** if this widget should always invalidate the prepass step when volatile */
 	uint8 bVolatilityAlwaysInvalidatesPrepass : 1;
