Wiget::SetUserFocus: If we're able to set focus immediately, cancel any focus request that's waiting to be processed in the LocalPlayer's SlateOperations.

diff --git a/Engine/Source/Runtime/SlateCore/Private/Input/Reply.cpp b/Engine/Source/Runtime/SlateCore/Private/Input/Reply.cpp
--- a/Engine/Source/Runtime/SlateCore/Private/Input/Reply.cpp
+++ b/Engine/Source/Runtime/SlateCore/Private/Input/Reply.cpp
@@ -28,6 +28,15 @@ FReply& FReply::ClearUserFocus(EFocusCause ReasonFocusIsChanging, bool bInAllUse
 	return Me();
 }
 
+FReply& FReply::CancelFocusRequest()
+{
+	this->bSetUserFocus = false;
+	this->FocusRecipient = nullptr;
+	this->bReleaseUserFocus = false;
+
+	return Me();
+}
+
 FString FReply::ToString()
 {
 	FString HandledStr = IsEventHandled() ? TEXT("Handled") : TEXT("Unhandled");
diff --git a/Engine/Source/Runtime/SlateCore/Public/Input/Reply.h b/Engine/Source/Runtime/SlateCore/Public/Input/Reply.h
--- a/Engine/Source/Runtime/SlateCore/Public/Input/Reply.h
+++ b/Engine/Source/Runtime/SlateCore/Public/Input/Reply.h
@@ -43,7 +43,7 @@ public:
 	}
 
 	/**
-	 * An even should return FReply::Handled().SetMousePos to ask Slate to move the mouse cursor to a different location
+	 * An event should return FReply::Handled().SetMousePos to ask Slate to move the mouse cursor to a different location
 	 */
 	FReply& SetMousePos( const FIntPoint& NewMousePos );
 
@@ -59,6 +59,13 @@ public:
 	/** An event should return a FReply::Handled().ClearUserFocus() to ask the system to clear user focus*/
 	FReply& ClearUserFocus(EFocusCause ReasonFocusIsChanging, bool bInAllUsers = false);
 
+	/** 
+	 *	An event should rarely invoke FReply::CancelFocusRequest. This will change the reply
+	 *  so that it no longer contains a request for the system to set or clear focus. 
+	 *  This is most useful when a reply has been cached for use in async operations.
+	 */
+	FReply& CancelFocusRequest();
+
 	/** An event should return FReply::Handled().SetNavigation( NavigationType ) as a means of asking the system to attempt a navigation*/
 	FReply& SetNavigation(EUINavigation InNavigationType, const ENavigationGenesis InNavigationGenesis, const ENavigationSource InNavigationSource = ENavigationSource::FocusedWidget)
 	{
diff --git a/Engine/Source/Runtime/UMG/Private/Components/Widget.cpp b/Engine/Source/Runtime/UMG/Private/Components/Widget.cpp
--- a/Engine/Source/Runtime/UMG/Private/Components/Widget.cpp
+++ b/Engine/Source/Runtime/UMG/Private/Components/Widget.cpp
@@ -589,9 +589,18 @@ void UWidget::SetUserFocus(APlayerController* PlayerController)
 		if ( ULocalPlayer* LocalPlayer = PlayerController->GetLocalPlayer() )
 		{
 			TOptional<int32> UserIndex = FSlateApplication::Get().GetUserIndexForController(LocalPlayer->GetControllerId());
-			if (UserIndex.IsSet() && !FSlateApplication::Get().SetUserFocus(UserIndex.GetValue(), SafeWidget) )
+			if (UserIndex.IsSet())
 			{
-				LocalPlayer->GetSlateOperations().SetUserFocus(SafeWidget.ToSharedRef());
+				FReply& DelayedSlateOperations = LocalPlayer->GetSlateOperations();
+				if (FSlateApplication::Get().SetUserFocus(UserIndex.GetValue(), SafeWidget))
+				{
+					DelayedSlateOperations.CancelFocusRequest();
+				}
+				else
+				{
+					DelayedSlateOperations.SetUserFocus(SafeWidget.ToSharedRef());
+				}
+				
 			}
 		}
 	}
