fixed crash on rebuilding UI invalidation path

diff --git a/Engine/Source/Runtime/SlateCore/Private/FastUpdate/SlateInvalidationRoot.cpp b/Engine/Source/Runtime/SlateCore/Private/FastUpdate/SlateInvalidationRoot.cpp
--- a/Engine/Source/Runtime/SlateCore/Private/FastUpdate/SlateInvalidationRoot.cpp
+++ b/Engine/Source/Runtime/SlateCore/Private/FastUpdate/SlateInvalidationRoot.cpp
@@ -357,8 +357,14 @@ bool FSlateInvalidationRoot::PaintFastPath(const FSlateInvalidationContext& Cont
 	return bWidgetsNeededRepaint;
 }
 
-void FSlateInvalidationRoot::BuildNewFastPathList_Recursive(FSlateInvalidationRoot& Root, FWidgetProxy& Proxy, int32 ParentIndex, int32& NextTreeIndex, TArray<FWidgetProxy>& CurrentFastPathList, TArray<FWidgetProxy, TMemStackAllocator<>>& NewFastPathList)
+bool FSlateInvalidationRoot::BuildNewFastPathList_Recursive(FSlateInvalidationRoot& Root, FWidgetProxy& Proxy, int32 ParentIndex, int32& NextTreeIndex, TArray<FWidgetProxy>& CurrentFastPathList, TArray<FWidgetProxy, TMemStackAllocator<>>& NewFastPathList)
 {
+	if (Proxy.Widget == nullptr)
+	{
+		return false;
+	}
+
+	bool bResult = true;
 	if (Proxy.bChildOrderInvalid)
 	{
 		NextTreeIndex = Proxy.LeafMostChildIndex != INDEX_NONE ? Proxy.LeafMostChildIndex + 1 : NextTreeIndex + 1;
@@ -376,18 +382,25 @@ void FSlateInvalidationRoot::BuildNewFastPathList_Recursive(FSlateInvalidationRo
 			Proxy.Widget->FastPathProxyHandle = FWidgetProxyHandle(Root, Proxy.Index);
 		}
 
-		 NewFastPathList.Add(MoveTemp(Proxy));
+		NewFastPathList.Add(MoveTemp(Proxy));
 
 		for (int32 LocalChildIndex = 0; LocalChildIndex < Proxy.NumChildren; ++LocalChildIndex)
 		{
-			FWidgetProxy& ChildProxy = CurrentFastPathList[NextTreeIndex];
+			if (NextTreeIndex < CurrentFastPathList.Num())
+			{
+				FWidgetProxy& ChildProxy = CurrentFastPathList[NextTreeIndex];
+				if (!ChildProxy.bChildOrderInvalid)
+				{
+					++NextTreeIndex;
+				}
 
-			if (!ChildProxy.bChildOrderInvalid)
+				bResult = bResult && BuildNewFastPathList_Recursive(Root, ChildProxy, Proxy.Index, NextTreeIndex, CurrentFastPathList, NewFastPathList);
+			}
+			else
 			{
-				++NextTreeIndex;
+				bResult = false;
+				break;
 			}
-
-			BuildNewFastPathList_Recursive(Root, ChildProxy, Proxy.Index, NextTreeIndex, CurrentFastPathList, NewFastPathList);
 		}
 
 		{
@@ -397,6 +410,7 @@ void FSlateInvalidationRoot::BuildNewFastPathList_Recursive(FSlateInvalidationRo
 		}
 	}
 	
+	return bResult;
 }
 
 void FSlateInvalidationRoot::AdjustWidgetsDesktopGeometry(FVector2D WindowToDesktopTransform)
@@ -442,13 +456,21 @@ void FSlateInvalidationRoot::BuildFastPathList(SWidget* RootWidget)
 		TArray<FWidgetProxy, TMemStackAllocator<>> TempList;
 		TempList.Reserve(FastWidgetPathList.Num());
 
+		bool bBuiltPath = false;
 		if (FastWidgetPathList.Num() > 0)
 		{
 			int32 NextTreeIndex = 1;
-			BuildNewFastPathList_Recursive(*this, FastWidgetPathList[0], INDEX_NONE, NextTreeIndex, FastWidgetPathList, TempList);
+			bBuiltPath = BuildNewFastPathList_Recursive(*this, FastWidgetPathList[0], INDEX_NONE, NextTreeIndex, FastWidgetPathList, TempList);
+			if (!bBuiltPath)
+			{
+				// invalidate partially built fast path
+				++FastPathGenerationNumber;
+			}
 		}
-		else
+		
+		if (!bBuiltPath)
 		{
+			TempList.Reset();
 			RootWidget->AssignIndicesToChildren(*this, INDEX_NONE, TempList, bParentVisible, bParentVolatile);
 		}
 
diff --git a/Engine/Source/Runtime/SlateCore/Public/FastUpdate/SlateInvalidationRoot.h b/Engine/Source/Runtime/SlateCore/Public/FastUpdate/SlateInvalidationRoot.h
--- a/Engine/Source/Runtime/SlateCore/Public/FastUpdate/SlateInvalidationRoot.h
+++ b/Engine/Source/Runtime/SlateCore/Public/FastUpdate/SlateInvalidationRoot.h
@@ -102,7 +102,7 @@ private:
 	bool PaintFastPath(const FSlateInvalidationContext& Context);
 
 	void BuildFastPathList(SWidget* RootWidget);
-	void BuildNewFastPathList_Recursive(FSlateInvalidationRoot& Root, FWidgetProxy& Proxy, int32 ParentIndex, int32& NextTreeIndex, TArray<FWidgetProxy>& CurrentFastPathList, TArray<FWidgetProxy, TMemStackAllocator<>>& NewFastPathList);
+	bool BuildNewFastPathList_Recursive(FSlateInvalidationRoot& Root, FWidgetProxy& Proxy, int32 ParentIndex, int32& NextTreeIndex, TArray<FWidgetProxy>& CurrentFastPathList, TArray<FWidgetProxy, TMemStackAllocator<>>& NewFastPathList);
 
 	void AdjustWidgetsDesktopGeometry(FVector2D WindowToDesktopTransform);
 private:
