Updated API for creating custom window title bars which will be used for future changes. Deprecated existing one

diff --git a/Engine/Source/Editor/PIEPreviewDeviceProfileSelector/Private/PIEPreviewWindow.cpp b/Engine/Source/Editor/PIEPreviewDeviceProfileSelector/Private/PIEPreviewWindow.cpp
--- a/Engine/Source/Editor/PIEPreviewDeviceProfileSelector/Private/PIEPreviewWindow.cpp
+++ b/Engine/Source/Editor/PIEPreviewDeviceProfileSelector/Private/PIEPreviewWindow.cpp
@@ -34,7 +34,7 @@ SPIEPreviewWindow::~SPIEPreviewWindow()
 
 TSharedRef<SWidget> SPIEPreviewWindow::MakeWindowTitleBar(const TSharedRef<SWindow>& Window, const TSharedPtr<SWidget>& CenterContent, EHorizontalAlignment CenterContentAlignment)
 {
-	/*TSharedRef<SPIEPreviewWindowTitleBar> */WindowTitleBar = SNew(SPIEPreviewWindowTitleBar, Window, CenterContent, EHorizontalAlignment::HAlign_Center)
+	WindowTitleBar = SNew(SPIEPreviewWindowTitleBar, Window, CenterContent, EHorizontalAlignment::HAlign_Center)
 		.Visibility(EVisibility::SelfHitTestInvisible);
 
 	return WindowTitleBar.ToSharedRef();
diff --git a/Engine/Source/Runtime/Slate/Private/Framework/Application/SlateApplication.cpp b/Engine/Source/Runtime/Slate/Private/Framework/Application/SlateApplication.cpp
--- a/Engine/Source/Runtime/Slate/Private/Framework/Application/SlateApplication.cpp
+++ b/Engine/Source/Runtime/Slate/Private/Framework/Application/SlateApplication.cpp
@@ -4017,9 +4017,9 @@ TSharedRef<SImage> FSlateApplication::MakeImage( const TAttribute<const FSlateBr
 }
 
 
-TSharedRef<SWidget> FSlateApplication::MakeWindowTitleBar( const TSharedRef<SWindow>& Window, const TSharedPtr<SWidget>& CenterContent, EHorizontalAlignment CenterContentAlignment, TSharedPtr<IWindowTitleBar>& OutTitleBar ) const
+TSharedRef<SWidget> FSlateApplication::MakeWindowTitleBar(const FWindowTitleBarArgs& InArgs, TSharedPtr<IWindowTitleBar>& OutTitleBar) const
 {
-	TSharedRef<SWindowTitleBar> TitleBar = SNew(SWindowTitleBar, Window, CenterContent, CenterContentAlignment)
+	TSharedRef<SWindowTitleBar> TitleBar = SNew(SWindowTitleBar, InArgs.Window, InArgs.CenterContent, InArgs.CenterContentAlignment)
 		.Visibility(EVisibility::SelfHitTestInvisible);
 
 	OutTitleBar = TitleBar;
diff --git a/Engine/Source/Runtime/Slate/Private/Framework/Docking/SDockingArea.cpp b/Engine/Source/Runtime/Slate/Private/Framework/Docking/SDockingArea.cpp
--- a/Engine/Source/Runtime/Slate/Private/Framework/Docking/SDockingArea.cpp
+++ b/Engine/Source/Runtime/Slate/Private/Framework/Docking/SDockingArea.cpp
@@ -249,7 +249,12 @@ void SDockingArea::SetParentWindow( const TSharedRef<SWindow>& NewParentWindow )
 	// Even thought we don't manage the parent window's lifetime, we are still responsible for making its window chrome.
 	{
 		TSharedPtr<IWindowTitleBar> TitleBar;
-		TSharedRef<SWidget> TitleBarWidget = FSlateApplication::Get().MakeWindowTitleBar(NewParentWindow, SNullWidget::NullWidget, HAlign_Fill, TitleBar);
+
+		FWindowTitleBarArgs Args(NewParentWindow);
+		Args.CenterContent = SNullWidget::NullWidget;
+		Args.CenterContentAlignment = HAlign_Fill;
+
+		TSharedRef<SWidget> TitleBarWidget = FSlateApplication::Get().MakeWindowTitleBar(Args, TitleBar);
 
 		(*WindowControlsArea)[TitleBarWidget];
 		NewParentWindow->SetTitleBar(TitleBar);
diff --git a/Engine/Source/Runtime/Slate/Public/Framework/Application/SlateApplication.h b/Engine/Source/Runtime/Slate/Public/Framework/Application/SlateApplication.h
--- a/Engine/Source/Runtime/Slate/Public/Framework/Application/SlateApplication.h
+++ b/Engine/Source/Runtime/Slate/Public/Framework/Application/SlateApplication.h
@@ -1348,7 +1348,7 @@ public:
 	virtual FWidgetPath LocateWindowUnderMouse( FVector2D ScreenspaceMouseCoordinate, const TArray<TSharedRef<SWindow>>& Windows, bool bIgnoreEnabledStatus = false, int32 UserIndex = INDEX_NONE) override;
 	virtual bool IsWindowHousingInteractiveTooltip(const TSharedRef<const SWindow>& WindowToTest) const override;
 	virtual TSharedRef<SImage> MakeImage( const TAttribute<const FSlateBrush*>& Image, const TAttribute<FSlateColor>& Color, const TAttribute<EVisibility>& Visibility ) const override;
-	virtual TSharedRef<SWidget> MakeWindowTitleBar( const TSharedRef<SWindow>& Window, const TSharedPtr<SWidget>& CenterContent, EHorizontalAlignment CenterContentAlignment, TSharedPtr<IWindowTitleBar>& OutTitleBar ) const override;
+	virtual TSharedRef<SWidget> MakeWindowTitleBar(const FWindowTitleBarArgs& InArgs, TSharedPtr<IWindowTitleBar>& OutTitleBar) const override;
 	virtual TSharedRef<IToolTip> MakeToolTip( const TAttribute<FText>& ToolTipText ) override;
 	virtual TSharedRef<IToolTip> MakeToolTip( const FText& ToolTipText ) override;
 	virtual void RequestDestroyWindow( TSharedRef<SWindow> WindowToDestroy ) override;
diff --git a/Engine/Source/Runtime/SlateCore/Private/Application/SlateApplicationBase.cpp b/Engine/Source/Runtime/SlateCore/Private/Application/SlateApplicationBase.cpp
--- a/Engine/Source/Runtime/SlateCore/Private/Application/SlateApplicationBase.cpp
+++ b/Engine/Source/Runtime/SlateCore/Private/Application/SlateApplicationBase.cpp
@@ -98,6 +98,16 @@ const FHitTesting& FSlateApplicationBase::GetHitTesting() const
 	return HitTesting;
 }
 
+TSharedRef<SWidget> FSlateApplicationBase::MakeWindowTitleBar(const TSharedRef<SWindow>& Window, const TSharedPtr<SWidget>& CenterContent, EHorizontalAlignment CenterContentAlignment, TSharedPtr<IWindowTitleBar>& OutTitleBar) const
+{
+	FWindowTitleBarArgs Args(Window);
+
+	Args.CenterContent = CenterContent;
+	Args.CenterContentAlignment = CenterContentAlignment;
+	
+	return MakeWindowTitleBar(Args, OutTitleBar);
+}
+
 void FSlateApplicationBase::RegisterActiveTimer( const TSharedRef<FActiveTimerHandle>& ActiveTimerHandle )
 {
 	FScopeLock ActiveTimerLock(&ActiveTimerCS);
diff --git a/Engine/Source/Runtime/SlateCore/Private/Widgets/SWindow.cpp b/Engine/Source/Runtime/SlateCore/Private/Widgets/SWindow.cpp
--- a/Engine/Source/Runtime/SlateCore/Private/Widgets/SWindow.cpp
+++ b/Engine/Source/Runtime/SlateCore/Private/Widgets/SWindow.cpp
@@ -500,7 +500,11 @@ FVector2D SWindow::ComputeWindowSizeForContent( FVector2D ContentSize )
 
 TSharedRef<SWidget> SWindow::MakeWindowTitleBar(const TSharedRef<SWindow>& Window, const TSharedPtr<SWidget>& CenterContent, EHorizontalAlignment TitleContentAlignment)
 {
-	return FSlateApplicationBase::Get().MakeWindowTitleBar(Window, nullptr, TitleContentAlignment, TitleBar);
+	FWindowTitleBarArgs Args(Window);
+	Args.CenterContent = CenterContent;
+	Args.CenterContentAlignment = TitleContentAlignment;
+
+	return FSlateApplicationBase::Get().MakeWindowTitleBar(Args, TitleBar);
 }
 
 
@@ -533,7 +537,11 @@ void SWindow::ConstructWindowInternals()
 		SNew( SVerticalBox )
 		.Visibility( EVisibility::SelfHitTestInvisible );
 
-	TSharedRef<SWidget> TitleBarWidget = MakeWindowTitleBar(SharedThis(this), nullptr, GetTitleAlignment());
+	FWindowTitleBarArgs Args(SharedThis(this));
+	Args.CenterContent = nullptr;
+	Args.CenterContentAlignment = GetTitleAlignment();
+
+	TSharedRef<SWidget> TitleBarWidget = FSlateApplicationBase::Get().MakeWindowTitleBar(Args, TitleBar);
 
 	if (bCreateTitleBar)
 	{
diff --git a/Engine/Source/Runtime/SlateCore/Public/Application/SlateApplicationBase.h b/Engine/Source/Runtime/SlateCore/Public/Application/SlateApplicationBase.h
--- a/Engine/Source/Runtime/SlateCore/Public/Application/SlateApplicationBase.h
+++ b/Engine/Source/Runtime/SlateCore/Public/Application/SlateApplicationBase.h
@@ -28,6 +28,23 @@ enum class EInvalidateWidgetReason : uint8;
 
 template< typename ObjectType > class TAttribute;
 
+struct FWindowTitleBarArgs
+{
+	FWindowTitleBarArgs(const TSharedRef<SWindow>& InWindow)
+		: Window(InWindow)
+		, CenterContentAlignment(HAlign_Fill)
+	{}
+
+	/** The window to create the title bar for */
+	TSharedRef<SWindow> Window;
+	
+	/** Optional content for the title bar's center (will override window title) */
+	TSharedPtr<SWidget> CenterContent;
+
+	/** The horizontal alignment of the center content. */
+	EHorizontalAlignment CenterContentAlignment;
+};
+
 /**
  * Interface for window title bars.
  */
@@ -384,7 +401,17 @@ public:
 	 * @param OutTitleBar Will hold a pointer to the title bar's interface.
 	 * @return The new title bar widget.
 	 */
-	virtual TSharedRef<SWidget> MakeWindowTitleBar( const TSharedRef<SWindow>& Window, const TSharedPtr<SWidget>& CenterContent, EHorizontalAlignment CenterContentAlignment, TSharedPtr<IWindowTitleBar>& OutTitleBar ) const = 0;
+	UE_DEPRECATED(4.26, "This version of MakeWindowTitleBar has been deprecated. Use the version that takes in an FWindowTitleBarArgs parameter instead.")
+	virtual TSharedRef<SWidget> MakeWindowTitleBar(const TSharedRef<SWindow>& Window, const TSharedPtr<SWidget>& CenterContent, EHorizontalAlignment CenterContentAlignment, TSharedPtr<IWindowTitleBar>& OutTitleBar) const;
+
+	/**
+	 * Creates a title bar for the specified window.
+	 *
+	 * @param InArgs	The creation arguments for the titlebar
+	 * @param OutTitleBar Will hold a pointer to the title bar's interface.
+	 * @return The new title bar widget.
+	 */
+	virtual TSharedRef<SWidget> MakeWindowTitleBar( const FWindowTitleBarArgs& InArgs, TSharedPtr<IWindowTitleBar>& OutTitleBar ) const = 0;
 
 	/**
 	 * Destroying windows has implications on some OSs (e.g. destroying Win32 HWNDs can cause events to be lost).
