Slate - Adding additional SlateDebugger information for when the cursor changes, and who changes it.

diff --git a/Engine/Source/Runtime/Slate/Private/Framework/Application/SlateUser.cpp b/Engine/Source/Runtime/Slate/Private/Framework/Application/SlateUser.cpp
--- a/Engine/Source/Runtime/Slate/Private/Framework/Application/SlateUser.cpp
+++ b/Engine/Source/Runtime/Slate/Private/Framework/Application/SlateUser.cpp
@@ -752,6 +752,10 @@ void FSlateUser::QueryCursor()
 					CursorReply = ArrangedWidget.Widget->OnCursorQuery(ArrangedWidget.Geometry, CursorEvent);
 					if (CursorReply.IsEventHandled())
 					{
+#if WITH_SLATE_DEBUGGING
+						FSlateDebugging::BroadcastCursorQuery(&ArrangedWidget.Widget.Get(), CursorReply);
+#endif
+
 						if (!CursorReply.GetCursorWidget().IsValid())
 						{
 							for (; WidgetIndex >= 0; --WidgetIndex)
@@ -772,12 +776,20 @@ void FSlateUser::QueryCursor()
 				{
 					// Query was NOT handled, and we are still over a slate window.
 					CursorReply = FCursorReply::Cursor(EMouseCursor::Default);
+					
+#if WITH_SLATE_DEBUGGING
+					FSlateDebugging::BroadcastCursorQuery(nullptr, CursorReply);
+#endif
 				}
 			}
 			else
 			{
 				// Set the default cursor when there isn't an active window under the cursor and the mouse isn't captured
 				CursorReply = FCursorReply::Cursor(EMouseCursor::Default);
+
+#if WITH_SLATE_DEBUGGING
+				FSlateDebugging::BroadcastCursorQuery(nullptr, CursorReply);
+#endif
 			}
 		}
 		ProcessCursorReply(CursorReply);
diff --git a/Engine/Source/Runtime/SlateCore/Private/Debugging/ConsoleSlateDebugger.cpp b/Engine/Source/Runtime/SlateCore/Private/Debugging/ConsoleSlateDebugger.cpp
--- a/Engine/Source/Runtime/SlateCore/Private/Debugging/ConsoleSlateDebugger.cpp
+++ b/Engine/Source/Runtime/SlateCore/Private/Debugging/ConsoleSlateDebugger.cpp
@@ -11,9 +11,11 @@
 #include "Types/ReflectionMetadata.h"
 #include "GenericPlatform/GenericPlatformStackWalk.h"
 #include "HAL/PlatformStackWalk.h"
+#include "Layout/WidgetPath.h"
 #include "UObject/UObjectGlobals.h"
 #include "UObject/Package.h"
-#include "Layout/WidgetPath.h"
+#include "UObject/ReflectedTypeAccessors.h"
+#include "UObject/Class.h"
 
 DEFINE_LOG_CATEGORY_STATIC(LogSlateDebugger, Log, All);
 
@@ -117,6 +119,7 @@ void FConsoleSlateDebugger::RemoveListeners()
 	FSlateDebugging::FocusEvent.RemoveAll(this);
 	FSlateDebugging::AttemptNavigationEvent.RemoveAll(this);
 	FSlateDebugging::MouseCaptureEvent.RemoveAll(this);
+	FSlateDebugging::CursorChangedEvent.RemoveAll(this);
 #endif
 }
 
@@ -131,6 +134,7 @@ void FConsoleSlateDebugger::UpdateListeners()
 	FSlateDebugging::AttemptNavigationEvent.AddRaw(this, &FConsoleSlateDebugger::OnAttemptNavigationEvent);
 	FSlateDebugging::ExecuteNavigationEvent.AddRaw(this, &FConsoleSlateDebugger::OnExecuteNavigationEvent);
 	FSlateDebugging::MouseCaptureEvent.AddRaw(this, &FConsoleSlateDebugger::OnCaptureStateChangeEvent);
+	FSlateDebugging::CursorChangedEvent.AddRaw(this, &FConsoleSlateDebugger::OnCursorChangeEvent);
 #endif
 }
 
@@ -288,6 +292,35 @@ void FConsoleSlateDebugger::OnCaptureStateChangeEvent(const FSlateDebuggingMouse
 	OptionallyDumpCallStack();
 }
 
+void FConsoleSlateDebugger::OnCursorChangeEvent(const FSlateDebuggingCursorQueryEventArgs& EventArgs)
+{
+	const FText ContextWidget = FText::FromString(FReflectionMetaData::GetWidgetDebugInfo(EventArgs.WidgetOverridingCursor));
+
+	FText EventText;
+	if (EventArgs.Reply.GetCursorWidget().IsValid())
+	{
+		static const FText InputEventFormat = LOCTEXT("CursorChangedToWidget", "{0} To Widget: {0} (Widget: {1})");
+		EventText = FText::Format(
+			InputEventFormat,
+			FText::FromString(FReflectionMetaData::GetWidgetDebugInfo(EventArgs.Reply.GetCursorWidget().Get())),
+			ContextWidget
+		);
+	}
+	else
+	{
+		static const FText InputEventFormat = LOCTEXT("CursorChangedToCursor", "Cursor Changed: To Type: {0} (By Widget: {1})");
+		EventText = FText::Format(
+			InputEventFormat,
+			FText::FromString(StaticEnum<EMouseCursor::Type>()->GetNameStringByValue((int64)EventArgs.Reply.GetCursorType())),
+			ContextWidget
+		);
+	}
+
+	UE_LOG(LogSlateDebugger, Log, TEXT("%s"), *EventText.ToString());
+
+	OptionallyDumpCallStack();
+}
+
 void FConsoleSlateDebugger::OptionallyDumpCallStack()
 {
 	if (!bCaptureStack)
diff --git a/Engine/Source/Runtime/SlateCore/Private/Debugging/ConsoleSlateDebugger.h b/Engine/Source/Runtime/SlateCore/Private/Debugging/ConsoleSlateDebugger.h
--- a/Engine/Source/Runtime/SlateCore/Private/Debugging/ConsoleSlateDebugger.h
+++ b/Engine/Source/Runtime/SlateCore/Private/Debugging/ConsoleSlateDebugger.h
@@ -43,6 +43,7 @@ private:
 	void OnAttemptNavigationEvent(const FSlateDebuggingNavigationEventArgs& EventArgs);
 	void OnExecuteNavigationEvent(const FSlateDebuggingExecuteNavigationEventArgs& EventArgs);
 	void OnCaptureStateChangeEvent(const FSlateDebuggingMouseCaptureEventArgs& EventArgs);
+	void OnCursorChangeEvent(const FSlateDebuggingCursorQueryEventArgs& EventArgs);
 
 	void OptionallyDumpCallStack();
 
diff --git a/Engine/Source/Runtime/SlateCore/Private/Debugging/SlateDebugging.cpp b/Engine/Source/Runtime/SlateCore/Private/Debugging/SlateDebugging.cpp
--- a/Engine/Source/Runtime/SlateCore/Private/Debugging/SlateDebugging.cpp
+++ b/Engine/Source/Runtime/SlateCore/Private/Debugging/SlateDebugging.cpp
@@ -168,6 +168,12 @@ FSlateDebuggingMouseCaptureEventArgs::FSlateDebuggingMouseCaptureEventArgs(
 {
 }
 
+FSlateDebuggingCursorQueryEventArgs::FSlateDebuggingCursorQueryEventArgs(const SWidget* InWidgetOverridingCursor, const FCursorReply& InReply)
+	: WidgetOverridingCursor(InWidgetOverridingCursor)
+	, Reply(InReply)
+{
+}
+
 FSlateDebugging::FBeginWindow FSlateDebugging::BeginWindow;
 
 FSlateDebugging::FEndWindow FSlateDebugging::EndWindow;
@@ -190,11 +196,14 @@ FSlateDebugging::FWidgetExecuteNavigationEvent FSlateDebugging::ExecuteNavigatio
 
 FSlateDebugging::FWidgetMouseCaptureEvent FSlateDebugging::MouseCaptureEvent;
 
+FSlateDebugging::FWidgetCursorQuery FSlateDebugging::CursorChangedEvent;
+
 FSlateDebugging::FUICommandRun FSlateDebugging::CommandRun;
 
 DECLARE_MULTICAST_DELEGATE_OneParam(FWidgetMouseCaptureEvent, const FSlateDebuggingMouseCaptureEventArgs& /*EventArgs*/);
 
 TArray<struct FInvalidatedWidgetDrawer> FSlateDebugging::InvalidatedWidgetDrawers;
+FSlateDebugging::FLastCursorQuery FSlateDebugging::LastCursorQuery;
 
 void FSlateDebugging::BroadcastWarning(const FText& WarningText, const TSharedPtr<SWidget>& OptionalContextWidget)
 {
@@ -265,6 +274,20 @@ void FSlateDebugging::BroadcastMouseCaptureLost(uint32 UserIndex, uint32 Pointer
 	MouseCaptureEvent.Broadcast(FSlateDebuggingMouseCaptureEventArgs(false, UserIndex, PointerIndex, InWidgetLostCapture));
 }
 
+void FSlateDebugging::BroadcastCursorQuery(const SWidget* InWidgetOverridingCursor, const FCursorReply& InReply)
+{
+	if (LastCursorQuery.WidgetThatOverrideCursorLast_UnsafeToUseForAnythingButCompare != InWidgetOverridingCursor || 
+		LastCursorQuery.MouseCursor != InReply.GetCursorType() || 
+		LastCursorQuery.CursorWidget.Pin() != InReply.GetCursorWidget())
+	{
+		LastCursorQuery.WidgetThatOverrideCursorLast_UnsafeToUseForAnythingButCompare = InWidgetOverridingCursor;
+		LastCursorQuery.MouseCursor = InReply.GetCursorType();
+		LastCursorQuery.CursorWidget = InReply.GetCursorWidget();
+
+		CursorChangedEvent.Broadcast(FSlateDebuggingCursorQueryEventArgs(InWidgetOverridingCursor, InReply));
+	}
+}
+
 void FSlateDebugging::WidgetInvalidated(FSlateInvalidationRoot& InvalidationRoot, const class FWidgetProxy& WidgetProxy, const FLinearColor* CustomInvalidationColor)
 {
 	if(FSlateApplicationBase::IsInitialized())
diff --git a/Engine/Source/Runtime/SlateCore/Public/Debugging/SlateDebugging.h b/Engine/Source/Runtime/SlateCore/Public/Debugging/SlateDebugging.h
--- a/Engine/Source/Runtime/SlateCore/Public/Debugging/SlateDebugging.h
+++ b/Engine/Source/Runtime/SlateCore/Public/Debugging/SlateDebugging.h
@@ -84,6 +84,15 @@ public:
 	const FString& AdditionalContent;
 };
 
+struct SLATECORE_API FSlateDebuggingCursorQueryEventArgs
+{
+public:
+	FSlateDebuggingCursorQueryEventArgs(const SWidget* InWidgetOverridingCursor, const FCursorReply& InReply);
+
+	const SWidget* WidgetOverridingCursor = nullptr;
+	const FCursorReply& Reply;
+};
+
 UENUM()
 enum class ESlateDebuggingFocusEvent : uint8
 {
@@ -243,6 +252,12 @@ public:
 	static void BroadcastMouseCapture(uint32 UserIndex, uint32 PointerIndex, TSharedPtr<const SWidget> InCapturingWidget);
 	static void BroadcastMouseCaptureLost(uint32 UserIndex, uint32 PointerIndex, TSharedPtr<const SWidget> InWidgetLostCapture);
 
+public:
+	DECLARE_MULTICAST_DELEGATE_OneParam(FWidgetCursorQuery, const FSlateDebuggingCursorQueryEventArgs& /*EventArgs*/);
+	static FWidgetCursorQuery CursorChangedEvent;
+
+	static void BroadcastCursorQuery(const SWidget* InWidgetOverridingCursor, const FCursorReply& InReply);
+
 public:
 	/**  */
 	DECLARE_MULTICAST_DELEGATE_TwoParams(FUICommandRun, const FName& /*CommandName*/, const FText& /*CommandLabel*/);
@@ -261,6 +276,14 @@ private:
 	FSlateDebugging() {}
 
 	static TArray<struct FInvalidatedWidgetDrawer> InvalidatedWidgetDrawers;
+
+	struct FLastCursorQuery
+	{
+		const SWidget* WidgetThatOverrideCursorLast_UnsafeToUseForAnythingButCompare = nullptr;
+		TWeakPtr<SWidget> CursorWidget;
+		EMouseCursor::Type MouseCursor;
+	};
+	static FLastCursorQuery LastCursorQuery;
 };
 
 #endif
\ No newline at end of file
