Check if fullscreen is supported when trying to swap to fullscreen in SlateRenderer (fixed D3D errors which could end up with Device lost crashes)

diff --git a/Engine/Source/Runtime/ApplicationCore/Private/GenericPlatform/GenericWindow.cpp b/Engine/Source/Runtime/ApplicationCore/Private/GenericPlatform/GenericWindow.cpp
--- a/Engine/Source/Runtime/ApplicationCore/Private/GenericPlatform/GenericWindow.cpp
+++ b/Engine/Source/Runtime/ApplicationCore/Private/GenericPlatform/GenericWindow.cpp
@@ -151,6 +151,12 @@ bool FGenericWindow::IsForegroundWindow() const
 	return true;
 }
 
+bool FGenericWindow::IsFullscreenSupported() const
+{
+	// empty default functionality
+	return true;
+}
+
 void FGenericWindow::SetText(const TCHAR* const Text)
 {
 	// empty default functionality
diff --git a/Engine/Source/Runtime/ApplicationCore/Private/Windows/WindowsWindow.cpp b/Engine/Source/Runtime/ApplicationCore/Private/Windows/WindowsWindow.cpp
--- a/Engine/Source/Runtime/ApplicationCore/Private/Windows/WindowsWindow.cpp
+++ b/Engine/Source/Runtime/ApplicationCore/Private/Windows/WindowsWindow.cpp
@@ -959,6 +959,12 @@ bool FWindowsWindow::IsForegroundWindow() const
 	return ::GetForegroundWindow() == HWnd;
 }
 
+bool FWindowsWindow::IsFullscreenSupported() const
+{
+	// fullscreen not supported when using remote desktop
+	return !::GetSystemMetrics(SM_REMOTESESSION);
+}
+
 void FWindowsWindow::SetText( const TCHAR* const Text )
 {
 	SetWindowText(HWnd, Text);
diff --git a/Engine/Source/Runtime/ApplicationCore/Public/GenericPlatform/GenericWindow.h b/Engine/Source/Runtime/ApplicationCore/Public/GenericPlatform/GenericWindow.h
--- a/Engine/Source/Runtime/ApplicationCore/Public/GenericPlatform/GenericWindow.h
+++ b/Engine/Source/Runtime/ApplicationCore/Public/GenericPlatform/GenericWindow.h
@@ -187,6 +187,9 @@ public:
 	/** @return true if the window is in the foreground */
 	virtual bool IsForegroundWindow() const;
 
+	/** @return true if the window is in the foreground */
+	virtual bool IsFullscreenSupported() const;
+
 	/**
 	 * Sets the window text - usually the title but can also be text content for things like controls
 	 *
diff --git a/Engine/Source/Runtime/ApplicationCore/Public/Windows/WindowsWindow.h b/Engine/Source/Runtime/ApplicationCore/Public/Windows/WindowsWindow.h
--- a/Engine/Source/Runtime/ApplicationCore/Public/Windows/WindowsWindow.h
+++ b/Engine/Source/Runtime/ApplicationCore/Public/Windows/WindowsWindow.h
@@ -125,6 +125,7 @@ public:
 	virtual int32 GetWindowTitleBarSize() const override;
 	virtual void* GetOSWindowHandle() const  override { return HWnd; }
 	virtual bool IsForegroundWindow() const override;
+	virtual bool IsFullscreenSupported() const override;
 	virtual void SetText(const TCHAR* const Text) override;
 	virtual void DrawAttention(const FWindowDrawAttentionParameters& Parameters) override;
 
diff --git a/Engine/Source/Runtime/SlateCore/Private/Rendering/SlateRenderer.cpp b/Engine/Source/Runtime/SlateCore/Private/Rendering/SlateRenderer.cpp
--- a/Engine/Source/Runtime/SlateCore/Private/Rendering/SlateRenderer.cpp
+++ b/Engine/Source/Runtime/SlateCore/Private/Rendering/SlateRenderer.cpp
@@ -163,7 +163,8 @@ bool FSlateRenderer::IsViewportFullscreen( const SWindow& Window ) const
 		}
 		else
 		{
-			bFullscreen = Window.GetWindowMode() == EWindowMode::Fullscreen;
+			bFullscreen = Window.GetNativeWindow()->IsFullscreenSupported() && Window.GetWindowMode() == EWindowMode::Fullscreen;
+
 #if PLATFORM_WINDOWS
 			// When we are in fullscreen mode but the user alt-tabs out we need to temporarily drop out of fullscreen while the window has lost focus, otherwise DXGI will eventually
 			// forcibly throw us out of fullscreen mode with device loss and crash as typical result. By returning false here we'll trigger a mode switch to windowed when the user
