Fixed an assertion when shaping oddly formed RTL text

RTL text containing two different script types could sometimes lead to HarfBuzz returning non-monotonic cluster IDs, which could cause an assertion in the code that tries to count the number of graphemes and characters in each glyph.

This assertion can be prevented by moving the logic for performing that calculation to after each sub-section, as then the calculation is always bound to the correct character ranges from the source buffer.

Additionally this change prevents a transient FString being constructed by FSlateTextShaper::PerformHarfBuzzTextShaping each time it called HarfBuzzUtils::AppendStringToBuffer, as the shaper is working with a const TCHAR* rather than an FString (HarfBuzzUtils::AppendStringToBuffer now uses FStringView instead).

diff --git a/Engine/Source/Runtime/SlateCore/Private/Fonts/FontCacheHarfBuzz.cpp b/Engine/Source/Runtime/SlateCore/Private/Fonts/FontCacheHarfBuzz.cpp
--- a/Engine/Source/Runtime/SlateCore/Private/Fonts/FontCacheHarfBuzz.cpp
+++ b/Engine/Source/Runtime/SlateCore/Private/Fonts/FontCacheHarfBuzz.cpp
@@ -54,35 +54,37 @@ namespace Internal
 {
 
 template <bool IsUnicode, size_t TCHARSize>
-void AppendStringToBuffer(const FString& InString, const int32 InStartIndex, const int32 InLength, hb_buffer_t* InHarfBuzzTextBuffer)
+void AppendStringToBuffer(const FStringView InString, const int32 InStartIndex, const int32 InLength, hb_buffer_t* InHarfBuzzTextBuffer)
 {
 	// todo: SHAPING - This is losing the context information that may be required to shape a sub-section of text.
 	//				   In practice this may not be an issue as our platforms should all use the other functions, but to fix it we'd need UTF-8 iteration functions to find the correct points the buffer
-	hb_buffer_add_utf8(InHarfBuzzTextBuffer, TCHAR_TO_UTF8(*InString.Mid(InStartIndex, InLength)), -1, 0, -1);
+	FStringView SubString = InString.Mid(InStartIndex, InLength);
+	FTCHARToUTF8 SubStringUtf8(SubString.GetData(), SubString.Len());
+	hb_buffer_add_utf8(InHarfBuzzTextBuffer, SubStringUtf8.Get(), SubStringUtf8.Length(), 0, SubStringUtf8.Length());
 }
 
 template <>
-void AppendStringToBuffer<true, 2>(const FString& InString, const int32 InStartIndex, const int32 InLength, hb_buffer_t* InHarfBuzzTextBuffer)
+void AppendStringToBuffer<true, 2>(const FStringView InString, const int32 InStartIndex, const int32 InLength, hb_buffer_t* InHarfBuzzTextBuffer)
 {
 	// A unicode encoding with a TCHAR size of 2 bytes is assumed to be UTF-16
-	hb_buffer_add_utf16(InHarfBuzzTextBuffer, reinterpret_cast<const uint16_t*>(InString.GetCharArray().GetData()), InString.Len(), InStartIndex, InLength);
+	hb_buffer_add_utf16(InHarfBuzzTextBuffer, reinterpret_cast<const uint16_t*>(InString.GetData()), InString.Len(), InStartIndex, InLength);
 }
 
 template <>
-void AppendStringToBuffer<true, 4>(const FString& InString, const int32 InStartIndex, const int32 InLength, hb_buffer_t* InHarfBuzzTextBuffer)
+void AppendStringToBuffer<true, 4>(const FStringView InString, const int32 InStartIndex, const int32 InLength, hb_buffer_t* InHarfBuzzTextBuffer)
 {
 	// A unicode encoding with a TCHAR size of 4 bytes is assumed to be UTF-32
-	hb_buffer_add_utf32(InHarfBuzzTextBuffer, reinterpret_cast<const uint32_t*>(InString.GetCharArray().GetData()), InString.Len(), InStartIndex, InLength);
+	hb_buffer_add_utf32(InHarfBuzzTextBuffer, reinterpret_cast<const uint32_t*>(InString.GetData()), InString.Len(), InStartIndex, InLength);
 }
 
 } // namespace Internal
 
-void AppendStringToBuffer(const FString& InString, hb_buffer_t* InHarfBuzzTextBuffer)
+void AppendStringToBuffer(const FStringView InString, hb_buffer_t* InHarfBuzzTextBuffer)
 {
 	return Internal::AppendStringToBuffer<FPlatformString::IsUnicodeEncoded, sizeof(TCHAR)>(InString, 0, InString.Len(), InHarfBuzzTextBuffer);
 }
 
-void AppendStringToBuffer(const FString& InString, const int32 InStartIndex, const int32 InLength, hb_buffer_t* InHarfBuzzTextBuffer)
+void AppendStringToBuffer(const FStringView InString, const int32 InStartIndex, const int32 InLength, hb_buffer_t* InHarfBuzzTextBuffer)
 {
 	return Internal::AppendStringToBuffer<FPlatformString::IsUnicodeEncoded, sizeof(TCHAR)>(InString, InStartIndex, InLength, InHarfBuzzTextBuffer);
 }
diff --git a/Engine/Source/Runtime/SlateCore/Private/Fonts/FontCacheHarfBuzz.h b/Engine/Source/Runtime/SlateCore/Private/Fonts/FontCacheHarfBuzz.h
--- a/Engine/Source/Runtime/SlateCore/Private/Fonts/FontCacheHarfBuzz.h
+++ b/Engine/Source/Runtime/SlateCore/Private/Fonts/FontCacheHarfBuzz.h
@@ -3,6 +3,7 @@
 #pragma once
 
 #include "CoreMinimal.h"
+#include "Containers/StringView.h"
 
 class FFreeTypeAdvanceCache;
 class FFreeTypeFace;
@@ -38,10 +39,10 @@ namespace HarfBuzzUtils
 #if WITH_HARFBUZZ
 
 /** Utility function to append an FString into a hb_buffer_t in the most efficient way based on the string encoding method of the current platform */
-void AppendStringToBuffer(const FString& InString, hb_buffer_t* InHarfBuzzTextBuffer);
+void AppendStringToBuffer(const FStringView InString, hb_buffer_t* InHarfBuzzTextBuffer);
 
 /** Utility function to append an FString into a hb_buffer_t in the most efficient way based on the string encoding method of the current platform */
-void AppendStringToBuffer(const FString& InString, const int32 InStartIndex, const int32 InLength, hb_buffer_t* InHarfBuzzTextBuffer);
+void AppendStringToBuffer(const FStringView InString, const int32 InStartIndex, const int32 InLength, hb_buffer_t* InHarfBuzzTextBuffer);
 
 #endif // #if WITH_HARFBUZZ
 
diff --git a/Engine/Source/Runtime/SlateCore/Private/Fonts/SlateTextShaper.cpp b/Engine/Source/Runtime/SlateCore/Private/Fonts/SlateTextShaper.cpp
--- a/Engine/Source/Runtime/SlateCore/Private/Fonts/SlateTextShaper.cpp
+++ b/Engine/Source/Runtime/SlateCore/Private/Fonts/SlateTextShaper.cpp
@@ -586,11 +586,10 @@ void FSlateTextShaper::PerformHarfBuzzTextShaping(const TCHAR* InText, const int
 		Algo::Reverse(HarfBuzzTextSequence);
 	}
 
-	const int32 InitialNumGlyphsToRender = OutGlyphsToRender.Num();
-
 	// Step 3) Now we use HarfBuzz to shape each font data sequence using its FreeType glyph
 	{
 		hb_buffer_t* HarfBuzzTextBuffer = hb_buffer_create();
+		const FStringView InTextView = InText; // This will do a strlen, so do it once at the start
 
 		for (const FHarfBuzzTextSequenceEntry& HarfBuzzTextSequenceEntry : HarfBuzzTextSequence)
 		{
@@ -618,11 +617,13 @@ void FSlateTextShaper::PerformHarfBuzzTextShaping(const TCHAR* InText, const int
 
 			for (const FHarfBuzzTextSequenceEntry::FSubSequenceEntry& HarfBuzzTextSubSequenceEntry : HarfBuzzTextSequenceEntry.SubSequence)
 			{
+				const int32 InitialNumGlyphsToRender = OutGlyphsToRender.Num();
+
 				hb_buffer_set_cluster_level(HarfBuzzTextBuffer, HB_BUFFER_CLUSTER_LEVEL_MONOTONE_GRAPHEMES);
 				hb_buffer_set_direction(HarfBuzzTextBuffer, (InTextDirection == TextBiDi::ETextDirection::LeftToRight) ? HB_DIRECTION_LTR : HB_DIRECTION_RTL);
 				hb_buffer_set_script(HarfBuzzTextBuffer, HarfBuzzTextSubSequenceEntry.HarfBuzzScript);
 
-				HarfBuzzUtils::AppendStringToBuffer(InText, HarfBuzzTextSubSequenceEntry.StartIndex, HarfBuzzTextSubSequenceEntry.Length, HarfBuzzTextBuffer);
+				HarfBuzzUtils::AppendStringToBuffer(InTextView, HarfBuzzTextSubSequenceEntry.StartIndex, HarfBuzzTextSubSequenceEntry.Length, HarfBuzzTextBuffer);
 				hb_shape(HarfBuzzFont, HarfBuzzTextBuffer, HarfBuzzFeatures, HarfBuzzFeaturesCount);
 
 				uint32 HarfBuzzGlyphCount = 0;
@@ -671,116 +672,116 @@ void FSlateTextShaper::PerformHarfBuzzTextShaping(const TCHAR* InText, const int
 				}
 
 				hb_buffer_clear_contents(HarfBuzzTextBuffer);
-			}
-
-			hb_font_destroy(HarfBuzzFont);
-		}
 
-		hb_buffer_destroy(HarfBuzzTextBuffer);
-	}
-
-	// Step 4) Count the characters and grapheme clusters that belong to each glyph if they haven't already been set
-	{
-		const int32 NumGlyphsRendered = OutGlyphsToRender.Num() - InitialNumGlyphsToRender;
-		if (NumGlyphsRendered > 0)
-		{
-			auto ConditionalUpdateGlyphCountsForRange = [this, InTextStart](FShapedGlyphEntry& ShapedGlyphEntry, const int32 TextStartIndex, const int32 TextEndIndex)
-			{
-				check(TextStartIndex <= TextEndIndex);
-
-				if (ShapedGlyphEntry.NumCharactersInGlyph == 0 && ShapedGlyphEntry.NumGraphemeClustersInGlyph == 0)
+				// Count the characters and grapheme clusters that belong to each glyph in this sub-sequence (if they haven't already been set)
 				{
-					ShapedGlyphEntry.NumCharactersInGlyph = TextEndIndex - TextStartIndex;
-
-					if (ShapedGlyphEntry.NumCharactersInGlyph > 0)
+					const int32 NumGlyphsRendered = OutGlyphsToRender.Num() - InitialNumGlyphsToRender;
+					if (NumGlyphsRendered > 0)
 					{
-						const int32 FirstCharacterIndex = TextStartIndex - InTextStart;
-						const int32 LastCharacterIndex = TextEndIndex - InTextStart;
-
-						// Only count grapheme clusters if this glyph starts on a grapheme boundary
-						int32 PreviousBreak = FirstCharacterIndex;
+						auto ConditionalUpdateGlyphCountsForRange = [this, InTextStart](FShapedGlyphEntry& ShapedGlyphEntry, const int32 TextStartIndex, const int32 TextEndIndex)
 						{
-							GraphemeBreakIterator->MoveToCandidateAfter(FirstCharacterIndex);
-							PreviousBreak = GraphemeBreakIterator->MoveToPrevious();
-						}
+							check(TextStartIndex <= TextEndIndex);
 
-						if (PreviousBreak == FirstCharacterIndex)
-						{
-							int32 CurrentBreak = LastCharacterIndex;
-							for (CurrentBreak = GraphemeBreakIterator->MoveToCandidateAfter(FirstCharacterIndex);
-								CurrentBreak != INDEX_NONE;
-								CurrentBreak = GraphemeBreakIterator->MoveToNext()
-								)
+							if (ShapedGlyphEntry.NumCharactersInGlyph == 0 && ShapedGlyphEntry.NumGraphemeClustersInGlyph == 0)
 							{
-								++ShapedGlyphEntry.NumGraphemeClustersInGlyph;
-								if (CurrentBreak >= LastCharacterIndex)
+								ShapedGlyphEntry.NumCharactersInGlyph = TextEndIndex - TextStartIndex;
+
+								if (ShapedGlyphEntry.NumCharactersInGlyph > 0)
 								{
-									break;
+									const int32 FirstCharacterIndex = TextStartIndex - InTextStart;
+									const int32 LastCharacterIndex = TextEndIndex - InTextStart;
+
+									// Only count grapheme clusters if this glyph starts on a grapheme boundary
+									int32 PreviousBreak = FirstCharacterIndex;
+									{
+										GraphemeBreakIterator->MoveToCandidateAfter(FirstCharacterIndex);
+										PreviousBreak = GraphemeBreakIterator->MoveToPrevious();
+									}
+
+									if (PreviousBreak == FirstCharacterIndex)
+									{
+										int32 CurrentBreak = LastCharacterIndex;
+										for (CurrentBreak = GraphemeBreakIterator->MoveToCandidateAfter(FirstCharacterIndex);
+											CurrentBreak != INDEX_NONE;
+											CurrentBreak = GraphemeBreakIterator->MoveToNext()
+											)
+										{
+											++ShapedGlyphEntry.NumGraphemeClustersInGlyph;
+											if (CurrentBreak >= LastCharacterIndex)
+											{
+												break;
+											}
+										}
+
+										// Only count grapheme clusters if this glyph ends on a grapheme boundary
+										if (CurrentBreak != LastCharacterIndex)
+										{
+											ShapedGlyphEntry.NumGraphemeClustersInGlyph = 0;
+										}
+									}
 								}
 							}
+						};
 
-							// Only count grapheme clusters if this glyph ends on a grapheme boundary
-							if (CurrentBreak != LastCharacterIndex)
+						auto GetNextGlyphToRenderIndex = [&OutGlyphsToRender](const int32 GlyphToRenderIndex) -> int32
+						{
+							FShapedGlyphEntry& ShapedGlyphEntry = OutGlyphsToRender[GlyphToRenderIndex];
+							int32 NextGlyphToRenderIndex = GlyphToRenderIndex + 1;
+
+							// Walk forward to find the first glyph in the next cluster; the number of characters in this glyph is the difference between their two source indices
+							for (; NextGlyphToRenderIndex < OutGlyphsToRender.Num(); ++NextGlyphToRenderIndex)
 							{
-								ShapedGlyphEntry.NumGraphemeClustersInGlyph = 0;
+								const FShapedGlyphEntry& NextShapedGlyphEntry = OutGlyphsToRender[NextGlyphToRenderIndex];
+								if (ShapedGlyphEntry.SourceIndex != NextShapedGlyphEntry.SourceIndex)
+								{
+									break;
+								}
 							}
-						}
-					}
-				}
-			};
 
-			auto GetNextGlyphToRenderIndex = [&OutGlyphsToRender](const int32 GlyphToRenderIndex) -> int32
-			{
-				FShapedGlyphEntry& ShapedGlyphEntry = OutGlyphsToRender[GlyphToRenderIndex];
-				int32 NextGlyphToRenderIndex = GlyphToRenderIndex + 1;
+							return NextGlyphToRenderIndex;
+						};
 
-				// Walk forward to find the first glyph in the next cluster; the number of characters in this glyph is the difference between their two source indices
-				for (; NextGlyphToRenderIndex < OutGlyphsToRender.Num(); ++NextGlyphToRenderIndex)
-				{
-					const FShapedGlyphEntry& NextShapedGlyphEntry = OutGlyphsToRender[NextGlyphToRenderIndex];
-					if (ShapedGlyphEntry.SourceIndex != NextShapedGlyphEntry.SourceIndex)
-					{
-						break;
-					}
-				}
+						// The glyphs in the array are in render order, so LTR and RTL text use different start and end points in the source string
+						if (InTextDirection == TextBiDi::ETextDirection::LeftToRight)
+						{
+							for (int32 GlyphToRenderIndex = InitialNumGlyphsToRender; GlyphToRenderIndex < OutGlyphsToRender.Num();)
+							{
+								FShapedGlyphEntry& ShapedGlyphEntry = OutGlyphsToRender[GlyphToRenderIndex];
 
-				return NextGlyphToRenderIndex;
-			};
+								const int32 NextGlyphToRenderIndex = GetNextGlyphToRenderIndex(GlyphToRenderIndex);
+								if (NextGlyphToRenderIndex < OutGlyphsToRender.Num())
+								{
+									const FShapedGlyphEntry& NextShapedGlyphEntry = OutGlyphsToRender[NextGlyphToRenderIndex];
+									ConditionalUpdateGlyphCountsForRange(ShapedGlyphEntry, ShapedGlyphEntry.SourceIndex, NextShapedGlyphEntry.SourceIndex);
+								}
+								else
+								{
+									ConditionalUpdateGlyphCountsForRange(ShapedGlyphEntry, ShapedGlyphEntry.SourceIndex, HarfBuzzTextSubSequenceEntry.StartIndex + HarfBuzzTextSubSequenceEntry.Length);
+								}
 
-			// The glyphs in the array are in render order, so LTR and RTL text use different start and end points in the source string
-			if (InTextDirection == TextBiDi::ETextDirection::LeftToRight)
-			{
-				for (int32 GlyphToRenderIndex = InitialNumGlyphsToRender; GlyphToRenderIndex < OutGlyphsToRender.Num();)
-				{
-					FShapedGlyphEntry& ShapedGlyphEntry = OutGlyphsToRender[GlyphToRenderIndex];
+								GlyphToRenderIndex = NextGlyphToRenderIndex;
+							}
+						}
+						else
+						{
+							int32 PreviousSourceIndex = HarfBuzzTextSubSequenceEntry.StartIndex + HarfBuzzTextSubSequenceEntry.Length;
+							for (int32 GlyphToRenderIndex = InitialNumGlyphsToRender; GlyphToRenderIndex < OutGlyphsToRender.Num();)
+							{
+								FShapedGlyphEntry& ShapedGlyphEntry = OutGlyphsToRender[GlyphToRenderIndex];
 
-					const int32 NextGlyphToRenderIndex = GetNextGlyphToRenderIndex(GlyphToRenderIndex);
-					if (NextGlyphToRenderIndex < OutGlyphsToRender.Num())
-					{
-						const FShapedGlyphEntry& NextShapedGlyphEntry = OutGlyphsToRender[NextGlyphToRenderIndex];
-						ConditionalUpdateGlyphCountsForRange(ShapedGlyphEntry, ShapedGlyphEntry.SourceIndex, NextShapedGlyphEntry.SourceIndex);
-					}
-					else
-					{
-						ConditionalUpdateGlyphCountsForRange(ShapedGlyphEntry, ShapedGlyphEntry.SourceIndex, InTextStart + InTextLen);
+								ConditionalUpdateGlyphCountsForRange(ShapedGlyphEntry, ShapedGlyphEntry.SourceIndex, PreviousSourceIndex);
+								GlyphToRenderIndex = GetNextGlyphToRenderIndex(GlyphToRenderIndex);
+								PreviousSourceIndex = ShapedGlyphEntry.SourceIndex;
+							}
+						}
 					}
-
-					GlyphToRenderIndex = NextGlyphToRenderIndex;
 				}
 			}
-			else
-			{
-				int32 PreviousSourceIndex = InTextStart + InTextLen;
-				for (int32 GlyphToRenderIndex = InitialNumGlyphsToRender; GlyphToRenderIndex < OutGlyphsToRender.Num();)
-				{
-					FShapedGlyphEntry& ShapedGlyphEntry = OutGlyphsToRender[GlyphToRenderIndex];
 
-					ConditionalUpdateGlyphCountsForRange(ShapedGlyphEntry, ShapedGlyphEntry.SourceIndex, PreviousSourceIndex);
-					GlyphToRenderIndex = GetNextGlyphToRenderIndex(GlyphToRenderIndex);
-					PreviousSourceIndex = ShapedGlyphEntry.SourceIndex;
-				}
-			}
+			hb_font_destroy(HarfBuzzFont);
 		}
+
+		hb_buffer_destroy(HarfBuzzTextBuffer);
 	}
 
 	GraphemeBreakIterator->ClearString();
