Slate - Adding a FVector2D setter for Position on SlateVertex.

diff --git a/Engine/Source/Runtime/SlateCore/Public/Rendering/RenderingCommon.h b/Engine/Source/Runtime/SlateCore/Public/Rendering/RenderingCommon.h
--- a/Engine/Source/Runtime/SlateCore/Public/Rendering/RenderingCommon.h
+++ b/Engine/Source/Runtime/SlateCore/Public/Rendering/RenderingCommon.h
@@ -277,6 +277,11 @@ public:
 		TexCoords[3] = InTexCoords.W;
 	}
 
+	void SetPosition(const FVector2D& InPosition)
+	{
+		Position = InPosition;
+	}
+
 private:
 
 	template<ESlateVertexRounding Rounding>
@@ -353,6 +358,9 @@ struct FShortRect
 	uint16 Bottom;
 };
 
+template<> struct TIsPODType<FShortRect> { enum { Value = true }; };
+static_assert(TIsTriviallyDestructible<FShortRect>::Value == true, "FShortRect should be trivially destructible");
+
 #if STATS
 
 struct FRenderingBufferStatTracker
