Fix to a bug where InvalidationPanel failed to update correctly when used in another panel. This could have lead to widget not updating display when they are properly invalidated.

diff --git a/Engine/Source/Runtime/Slate/Private/Framework/Application/SlateApplication.cpp b/Engine/Source/Runtime/Slate/Private/Framework/Application/SlateApplication.cpp
--- a/Engine/Source/Runtime/Slate/Private/Framework/Application/SlateApplication.cpp
+++ b/Engine/Source/Runtime/Slate/Private/Framework/Application/SlateApplication.cpp
@@ -1463,6 +1463,9 @@ void FSlateApplication::TickAndDrawWidgets(float DeltaTime)
 		Renderer->ReleaseAccessedResources(/* Flush State */ false);
 	}
 
+	// We clear all pending Updates from last frame
+	FSlateInvalidationRoot::ClearAllWidgetUpdatesPending();
+
 	{
 		SCOPE_CYCLE_COUNTER(STAT_SlatePreTickEvent);
 		PreTickEvent.Broadcast(DeltaTime);
diff --git a/Engine/Source/Runtime/SlateCore/Private/FastUpdate/SlateInvalidationRoot.cpp b/Engine/Source/Runtime/SlateCore/Private/FastUpdate/SlateInvalidationRoot.cpp
--- a/Engine/Source/Runtime/SlateCore/Private/FastUpdate/SlateInvalidationRoot.cpp
+++ b/Engine/Source/Runtime/SlateCore/Private/FastUpdate/SlateInvalidationRoot.cpp
@@ -24,6 +24,8 @@ static FAutoConsoleCommand HandleDumpUpdateListCommand(
 );
 #endif
 
+TArray<FSlateInvalidationRoot*> FSlateInvalidationRoot::ClearUpdateList;
+
 #if SLATE_CSV_TRACKER
 	static int32 CascadeInvalidationEventAmount = 5;
 	FAutoConsoleVariableRef CVarCascadeInvalidationEventAmount(
@@ -247,8 +249,6 @@ FSlateInvalidationResult FSlateInvalidationRoot::PaintInvalidationRoot(const FSl
 #endif
 	}
 
-	FinalUpdateList.Empty();
-
 	Result.MaxLayerIdPainted = CachedMaxLayerId;
 	return Result;
 }
@@ -268,6 +268,15 @@ void FSlateInvalidationRoot::OnWidgetDestroyed(const SWidget* Widget)
 	Widget->PersistentState.CachedElementHandle.RemoveFromCache();
 }
 
+void FSlateInvalidationRoot::ClearAllWidgetUpdatesPending()
+{
+	// Once a frame we free the FinalUpdateList
+	for (FSlateInvalidationRoot* Root : ClearUpdateList)
+	{
+		Root->FinalUpdateList.Reset();
+	}
+}
+
 bool FSlateInvalidationRoot::PaintFastPath(const FSlateInvalidationContext& Context)
 {
 	SCOPED_NAMED_EVENT(SWidget_FastPathUpdate, FColor::Green);
@@ -480,6 +489,19 @@ void FSlateInvalidationRoot::BuildFastPathList(SWidget* RootWidget)
 			ensureAlways(Copy[i].Widget == TempList[i].Widget);
 		}
 #endif
+
+		// When it's the first time the FastWidgetPathList has item, we know we will need to clear the UpdateList on the next frame.
+		if (FastWidgetPathList.Num() == 0 && TempList.Num() != 0)
+		{
+			ensure(ClearUpdateList.Find(this) == INDEX_NONE);
+			ClearUpdateList.Push(this);
+		}
+		// When FastWidgetPathList is empty for the first time, we know we don't need to clear the list until we have items readded
+		else if (FastWidgetPathList.Num() != 0 && TempList.Num() == 0)
+		{
+			ClearUpdateList.RemoveSingleSwap(this, false);
+		}
+
 		FastWidgetPathList = MoveTemp(TempList);
 	}
 }
@@ -489,8 +511,6 @@ bool FSlateInvalidationRoot::ProcessInvalidation()
 	SCOPED_NAMED_EVENT(Slate_InvalidationProcessing, FColor::Blue);
 	CSV_SCOPED_TIMING_STAT(Slate, InvalidationProcessing);
 
-	FinalUpdateList.Empty();
-
 	bool bWidgetsNeedRepaint = false;
 
 	if (!bNeedsSlowPath)
@@ -664,6 +684,11 @@ void FSlateInvalidationRoot::ClearAllFastPathData(bool bClearResourcesImmediatel
 		}
 	}
 
+
+	if (FastWidgetPathList.Num() != 0)
+	{
+		ClearUpdateList.RemoveSingleSwap(this, false);
+	}
 	FastWidgetPathList.Empty();
 	WidgetsNeedingUpdate.Empty();
 	CachedElementData->Empty();
diff --git a/Engine/Source/Runtime/SlateCore/Public/FastUpdate/SlateInvalidationRoot.h b/Engine/Source/Runtime/SlateCore/Public/FastUpdate/SlateInvalidationRoot.h
--- a/Engine/Source/Runtime/SlateCore/Public/FastUpdate/SlateInvalidationRoot.h
+++ b/Engine/Source/Runtime/SlateCore/Public/FastUpdate/SlateInvalidationRoot.h
@@ -85,6 +85,9 @@ public:
 	SLATECORE_API FSlateInvalidationResult PaintInvalidationRoot(const FSlateInvalidationContext& Context);
 
 	void OnWidgetDestroyed(const SWidget* Widget);
+
+	SLATECORE_API static void ClearAllWidgetUpdatesPending();
+
 protected:
 	virtual int32 PaintSlowPath(const FSlateInvalidationContext& Context) = 0;
 
@@ -130,4 +133,6 @@ private:
 	bool bChildOrderInvalidated;
 	bool bNeedsSlowPath;
 	bool bNeedScreenPositionShift;
+
+	static TArray<FSlateInvalidationRoot*> ClearUpdateList;
 };
