Slate - Exposing a function to enable the retention off and on for blueprints for retainer panels.  Adding some more invalidation code to retainers that appears to be needed to handle the case with nested retainers.  A little removal of commented out code in SlateInvalidationRoot.

diff --git a/Engine/Source/Runtime/SlateCore/Private/FastUpdate/SlateInvalidationRoot.cpp b/Engine/Source/Runtime/SlateCore/Private/FastUpdate/SlateInvalidationRoot.cpp
--- a/Engine/Source/Runtime/SlateCore/Private/FastUpdate/SlateInvalidationRoot.cpp
+++ b/Engine/Source/Runtime/SlateCore/Private/FastUpdate/SlateInvalidationRoot.cpp
@@ -80,31 +80,28 @@ FString FSlateInvalidationRoot::GetReferencerName() const
 
 void FSlateInvalidationRoot::InvalidateChildOrder()
 {
-	//if (GSlateEnableGlobalInvalidation)
+	if(!bNeedsSlowPath && !bChildOrderInvalidated)
 	{
-		if(!bNeedsSlowPath && !bChildOrderInvalidated)
+		//UE_LOG(LogSlate, Log, TEXT("Child order invalidated, Slow path needed"));
+		//bNeedsSlowPath = true;
+		bChildOrderInvalidated = true;
+		if(!InvalidationRootWidget->Advanced_IsWindow())
 		{
-			//UE_LOG(LogSlate, Log, TEXT("Child order invalidated, Slow path needed"));
-			//bNeedsSlowPath = true;
-			bChildOrderInvalidated = true;
-			if(!InvalidationRootWidget->Advanced_IsWindow())
-			{
-				InvalidationRootWidget->InvalidatePrepass();
-			}
+			InvalidationRootWidget->InvalidatePrepass();
+		}
 
-			if (!GSlateEnableGlobalInvalidation  && !InvalidationRootWidget->Advanced_IsWindow())
-			{
-				//InvalidationRootWidget->Invalidate(EInvalidateWidgetReason::ChildOrder);
-				InvalidationRootWidget->Invalidate(EInvalidateWidgetReason::Layout);
-			}
+		if (!GSlateEnableGlobalInvalidation && !InvalidationRootWidget->Advanced_IsWindow())
+		{
+			//InvalidationRootWidget->Invalidate(EInvalidateWidgetReason::ChildOrder);
+			InvalidationRootWidget->Invalidate(EInvalidateWidgetReason::Layout);
+		}
 
 #if WITH_SLATE_DEBUGGING
-			if (GSlateInvalidationDebugging && FastWidgetPathList.Num())
-			{
-				FSlateDebugging::WidgetInvalidated(*this, FastWidgetPathList[0], &FLinearColor::Red);
-			}
-#endif
+		if (GSlateInvalidationDebugging && FastWidgetPathList.Num())
+		{
+			FSlateDebugging::WidgetInvalidated(*this, FastWidgetPathList[0], &FLinearColor::Red);
 		}
+#endif
 	}
 }
 
diff --git a/Engine/Source/Runtime/UMG/Private/Components/RetainerBox.cpp b/Engine/Source/Runtime/UMG/Private/Components/RetainerBox.cpp
--- a/Engine/Source/Runtime/UMG/Private/Components/RetainerBox.cpp
+++ b/Engine/Source/Runtime/UMG/Private/Components/RetainerBox.cpp
@@ -76,6 +76,16 @@ void URetainerBox::SetTextureParameter(FName InTextureParameter)
 	}
 }
 
+void URetainerBox::SetRetainRendering(bool bInRetainRendering)
+{
+	bRetainRender = bInRetainRendering;
+
+	if (MyRetainerWidget.IsValid())
+	{
+		MyRetainerWidget->SetRetainedRendering(bRetainRender);
+	}
+}
+
 void URetainerBox::ReleaseSlateResources(bool bReleaseChildren)
 {
 	Super::ReleaseSlateResources(bReleaseChildren);
@@ -94,10 +104,8 @@ TSharedRef<SWidget> URetainerBox::RebuildWidget()
 #if STATS
 		.StatId( FName( *FString::Printf(TEXT("%s [%s]"), *GetFName().ToString(), *GetClass()->GetName() ) ) )
 #endif//STATS
-		;
+	;
 
-	MyRetainerWidget->SetRetainedRendering(IsDesignTime() ? false : true);
-	
 	if ( GetChildrenCount() > 0 )
 	{
 		MyRetainerWidget->SetContent(GetContentSlot()->Content ? GetContentSlot()->Content->TakeWidget() : SNullWidget::NullWidget);
@@ -110,6 +118,7 @@ void URetainerBox::SynchronizeProperties()
 {
 	Super::SynchronizeProperties();
 
+	MyRetainerWidget->SetRetainedRendering(IsDesignTime() ? false : bRetainRender);
 	MyRetainerWidget->SetEffectMaterial(EffectMaterial);
 	MyRetainerWidget->SetTextureParameter(TextureParameter);
 	MyRetainerWidget->SetWorld(GetWorld());
diff --git a/Engine/Source/Runtime/UMG/Private/Slate/SRetainerWidget.cpp b/Engine/Source/Runtime/UMG/Private/Slate/SRetainerWidget.cpp
--- a/Engine/Source/Runtime/UMG/Private/Slate/SRetainerWidget.cpp
+++ b/Engine/Source/Runtime/UMG/Private/Slate/SRetainerWidget.cpp
@@ -221,7 +221,13 @@ bool SRetainerWidget::IsAnythingVisibleToRender() const
 
 void SRetainerWidget::OnRetainerModeChanged()
 {
+	if (MyWidget.IsValid())
+	{
+		InvalidateChildRemovedFromTree(*MyWidget.Get());
+	}
+
 	RefreshRenderingMode();
+
 	// Invalidate myself 
 	InvalidateRoot();
 
@@ -250,9 +256,11 @@ void SRetainerWidget::OnRetainerModeCVarChanged( IConsoleVariable* CVar )
 
 void SRetainerWidget::SetRetainedRendering(bool bRetainRendering)
 {
-	bEnableRetainedRenderingDesire = bRetainRendering;
-
-	RefreshRenderingMode();
+	if (bEnableRetainedRenderingDesire != bRetainRendering)
+	{
+		bEnableRetainedRenderingDesire = bRetainRendering;
+		OnRetainerModeChanged();
+	}
 }
 
 void SRetainerWidget::RefreshRenderingMode()
diff --git a/Engine/Source/Runtime/UMG/Public/Components/RetainerBox.h b/Engine/Source/Runtime/UMG/Public/Components/RetainerBox.h
--- a/Engine/Source/Runtime/UMG/Public/Components/RetainerBox.h
+++ b/Engine/Source/Runtime/UMG/Public/Components/RetainerBox.h
@@ -27,6 +27,10 @@ class UMG_API URetainerBox : public UContentWidget
 {
 	GENERATED_UCLASS_BODY()
 
+protected:
+	UPROPERTY(EditAnywhere, Category = RenderRules)
+	bool bRetainRender = true;
+
 public:
 	/**
 	 * Should this widget redraw the contents it has every time it receives an invalidation request
@@ -93,6 +97,11 @@ public:
 	 */
 	UFUNCTION(BlueprintCallable, Category="Retainer|Effect")
 	void SetTextureParameter(FName TextureParameter);
+	/**
+	* Set the flag for if we retain the render or pass-through
+	*/
+	UFUNCTION(BlueprintCallable, Category = "Retainer")
+	void SetRetainRendering(bool bInRetainRendering);
 
 	virtual void ReleaseSlateResources(bool bReleaseChildren) override;
 
