Scene Outliner Styling

diff --git a/Engine/Source/Editor/LevelEditor/Private/SLevelEditor.cpp b/Engine/Source/Editor/LevelEditor/Private/SLevelEditor.cpp
--- a/Engine/Source/Editor/LevelEditor/Private/SLevelEditor.cpp
+++ b/Engine/Source/Editor/LevelEditor/Private/SLevelEditor.cpp
@@ -720,7 +720,7 @@ TSharedRef<SDockTab> SLevelEditor::SpawnLevelEditorTab( const FSpawnTabArgs& Arg
 			.ToolTip( IDocumentation::Get()->CreateToolTip( Label, nullptr, "Shared/LevelEditor", "SceneOutlinerTab" ) )
 			[
 				SNew(SBorder)
-				.Padding(4)
+				.Padding( 0 )
 				.BorderImage( FEditorStyle::GetBrush("ToolPanel.GroupBorder") )
 				.AddMetaData<FTutorialMetaData>(FTutorialMetaData(TEXT("SceneOutliner"), TEXT("LevelEditorSceneOutliner")))
 				[
diff --git a/Engine/Source/Editor/SceneOutliner/Private/SSceneOutliner.cpp b/Engine/Source/Editor/SceneOutliner/Private/SSceneOutliner.cpp
--- a/Engine/Source/Editor/SceneOutliner/Private/SSceneOutliner.cpp
+++ b/Engine/Source/Editor/SceneOutliner/Private/SSceneOutliner.cpp
@@ -284,10 +284,6 @@ namespace SceneOutliner
 		// @todo outliner: Should probably save this in layout!
 		// @todo outliner: Should save spacing for list view in layout
 
-		NoBorder = FEditorStyle::GetBrush( "LevelViewport.NoViewportBorder" );
-		PlayInEditorBorder = FEditorStyle::GetBrush( "LevelViewport.StartingPlayInEditorBorder" );
-		SimulateBorder = FEditorStyle::GetBrush( "LevelViewport.StartingSimulateBorder" );
-
 		//Setup the SearchBox filter
 		{
 			auto Delegate = TreeItemTextFilter::FItemToStringArray::CreateSP( this, &SSceneOutliner::PopulateSearchStrings );
@@ -330,13 +326,7 @@ namespace SceneOutliner
 
 		ChildSlot
 		[
-			SNew( SBorder )
-			.BorderImage( this, &SSceneOutliner::OnGetBorderBrush )
-			.BorderBackgroundColor( this, &SSceneOutliner::OnGetBorderColorAndOpacity )
-			.ShowEffectWhenDisabled( false )
-			[
-				VerticalBox
-			]
+			VerticalBox
 		];
 
 		auto Toolbar = SNew(SHorizontalBox);
@@ -360,19 +350,42 @@ namespace SceneOutliner
 				.Padding(4.f, 0.f, 0.f, 0.f)
 				[
 					SNew(SButton)
-					.ButtonStyle(FEditorStyle::Get(), "HoverHintOnly")
+					.ButtonStyle(FAppStyle::Get(), "SimpleButton")
 					.ToolTipText(LOCTEXT("CreateFolderToolTip", "Create a new folder containing the current actor selection"))
 					.OnClicked(this, &SSceneOutliner::OnCreateFolderClicked)
 					[
 						SNew(SImage)
-						.Image(FEditorStyle::GetBrush("SceneOutliner.NewFolderIcon"))
+	                	.ColorAndOpacity(FSlateColor::UseForeground())
+						.Image(FAppStyle::Get().GetBrush("SceneOutliner.NewFolderIcon"))
 					]
 				];
 		}
 
+
+		if (SharedData->Mode == ESceneOutlinerMode::ActorBrowsing)
+		{
+			// View mode combo button
+			Toolbar->AddSlot()
+			.VAlign(VAlign_Center)
+			.AutoWidth()
+			[
+				SAssignNew( ViewOptionsComboButton, SComboButton )
+				.ComboButtonStyle( FAppStyle::Get(), "SimpleComboButton" ) // Use the tool bar item style for this button
+				.OnGetMenuContent( this, &SSceneOutliner::GetViewButtonContent, false, !InInitOptions.SpecifiedWorldToDisplay)
+				.HasDownArrow(false)
+				.ButtonContent()
+				[
+					SNew(SImage)
+	                .ColorAndOpacity(FSlateColor::UseForeground())
+					.Image( FAppStyle::Get().GetBrush("Icons.cogwheel") )
+				]
+			];
+		}
+
+
 		VerticalBox->AddSlot()
 		.AutoHeight()
-		.Padding( 0.0f, 0.0f, 0.0f, 4.0f )
+		.Padding( 8.0f, 4.0f, 4.0f, 2.0f )
 		[
 			Toolbar
 		];
@@ -390,6 +403,11 @@ namespace SceneOutliner
 				.ColorAndOpacity( FLinearColor( 0.4f, 1.0f, 0.4f ) )
 			]
 
+			+SOverlay::Slot()
+			[
+				SNew(SBorder).BorderImage( FAppStyle::Get().GetBrush("InputBrush") )
+			]
+
 			+SOverlay::Slot()
 			[
 				SAssignNew( OutlinerTreeView, SOutlinerTreeView, StaticCastSharedRef<SSceneOutliner>(AsShared()) )
@@ -433,17 +451,6 @@ namespace SceneOutliner
 			]
 		];
 
-		// Separator
-		if ( SharedData->Mode == ESceneOutlinerMode::ActorBrowsing || !InInitOptions.SpecifiedWorldToDisplay )
-		{
-			VerticalBox->AddSlot()
-			.AutoHeight()
-			.Padding(0, 0, 0, 1)
-			[
-				SNew(SSeparator)
-			];
-		}
-
 		if (SharedData->Mode == ESceneOutlinerMode::ActorBrowsing)
 		{
 			// Bottom panel
@@ -456,43 +463,14 @@ namespace SceneOutliner
 				+SHorizontalBox::Slot()
 				.FillWidth(1.f)
 				.VAlign(VAlign_Center)
-				.Padding(8, 0)
+				.Padding(14, 9)
 				[
 					SNew( STextBlock )
 					.Text( this, &SSceneOutliner::GetFilterStatusText )
 					.ColorAndOpacity( this, &SSceneOutliner::GetFilterStatusTextColor )
 				]
 
-				// View mode combo button
-				+SHorizontalBox::Slot()
-				.AutoWidth()
-				[
-					SAssignNew( ViewOptionsComboButton, SComboButton )
-					.ContentPadding(0)
-					.ForegroundColor( this, &SSceneOutliner::GetViewButtonForegroundColor )
-					.ButtonStyle( FEditorStyle::Get(), "ToggleButton" ) // Use the tool bar item style for this button
-					.OnGetMenuContent( this, &SSceneOutliner::GetViewButtonContent, false, !InInitOptions.SpecifiedWorldToDisplay)
-					.ButtonContent()
-					[
-						SNew(SHorizontalBox)
-
-						+SHorizontalBox::Slot()
-						.AutoWidth()
-						.VAlign(VAlign_Center)
-						[
-							SNew(SImage).Image( FEditorStyle::GetBrush("GenericViewButton") )
-						]
-
-						+SHorizontalBox::Slot()
-						.AutoWidth()
-						.Padding(2, 0, 0, 0)
-						.VAlign(VAlign_Center)
-						[
-							SNew(STextBlock).Text( LOCTEXT("ViewButton", "View Options") )
-						]
-					]
-				]
-			];
+		];
 		} //end if (SharedData->Mode == ESceneOutlinerMode::ActorBrowsing)
 		else if ( !InInitOptions.SpecifiedWorldToDisplay )
 		{
@@ -1096,24 +1074,6 @@ namespace SceneOutliner
 	/** END FILTERS */
 
 
-	const FSlateBrush* SSceneOutliner::OnGetBorderBrush() const
-	{
-		if (SharedData->bRepresentingPlayWorld)
-		{
-			return GEditor->bIsSimulatingInEditor ? SimulateBorder : PlayInEditorBorder;
-		}
-		else
-		{
-			return NoBorder;
-		}
-	}
-
-	FSlateColor SSceneOutliner::OnGetBorderColorAndOpacity() const
-	{
-		return SharedData->bRepresentingPlayWorld	? FLinearColor(1.0f, 1.0f, 1.0f, 0.6f)
-									   				: FLinearColor::Transparent;
-	}
-
 	ESelectionMode::Type SSceneOutliner::GetSelectionMode() const
 	{
 		return SelectionMode;
@@ -3835,18 +3795,15 @@ namespace SceneOutliner
 	{
 		if ( !IsFilterActive() )
 		{
-			// White = no text filter
-			return FLinearColor( 1.0f, 1.0f, 1.0f );
+			 return FAppStyle::Get().GetColor("Colors.Foreground");
 		}
 		else if( FilteredActorCount == 0 )
 		{
-			// Red = no matching actors
-			return FLinearColor( 1.0f, 0.4f, 0.4f );
+			return FAppStyle::Get().GetColor("Colors.AccentRed");
 		}
 		else
 		{
-			// Green = found at least one match!
-			return FLinearColor( 0.4f, 1.0f, 0.4f );
+			return FAppStyle::Get().GetColor("Colors.AccentGreen");
 		}
 	}
 
diff --git a/Engine/Source/Editor/SceneOutliner/Private/SceneOutlinerActorInfoColumn.cpp b/Engine/Source/Editor/SceneOutliner/Private/SceneOutlinerActorInfoColumn.cpp
--- a/Engine/Source/Editor/SceneOutliner/Private/SceneOutlinerActorInfoColumn.cpp
+++ b/Engine/Source/Editor/SceneOutliner/Private/SceneOutlinerActorInfoColumn.cpp
@@ -160,7 +160,7 @@ SHeaderRow::FColumn::FArguments FActorInfoColumn::ConstructHeaderRowColumn()
 	/** Customizable actor data column */
 	return SHeaderRow::Column( GetColumnID() )
 		.FillWidth(2)
-		.HeaderComboVisibility(EHeaderComboVisibility::Ghosted)
+		.HeaderComboVisibility(EHeaderComboVisibility::OnHover)
 		.MenuContent()
 		[
 			SNew(SBorder)
@@ -182,6 +182,8 @@ SHeaderRow::FColumn::FArguments FActorInfoColumn::ConstructHeaderRowColumn()
 			.VAlign(VAlign_Center)
 			[
 				SNew(STextBlock)
+				.TextStyle(FAppStyle::Get(), "ButtonText")
+				.TransformPolicy(ETextTransformPolicy::ToUpper)
 				.Text( this, &FActorInfoColumn::GetSelectedMode )
 			]
 		];
@@ -434,6 +436,7 @@ TSharedRef< ITableRow > FActorInfoColumn::MakeComboButtonItemWidget( TSharedPtr<
 {
 	return
 		SNew( STableRow< TSharedPtr<ECustomColumnMode::Type> >, Owner )
+		.Style(FAppStyle::Get(), "ComboBox.Row")
 		[
 			SNew( STextBlock )
 			.Text( MakeComboText( *Mode ) )
diff --git a/Engine/Source/Editor/SceneOutliner/Private/SceneOutlinerGutter.cpp b/Engine/Source/Editor/SceneOutliner/Private/SceneOutlinerGutter.cpp
--- a/Engine/Source/Editor/SceneOutliner/Private/SceneOutlinerGutter.cpp
+++ b/Engine/Source/Editor/SceneOutliner/Private/SceneOutlinerGutter.cpp
@@ -174,6 +174,7 @@ public:
 
 		SImage::Construct(
 			SImage::FArguments()
+			.ColorAndOpacity(FSlateColor::UseForeground())
 			.Image(this, &SVisibilityWidget::GetBrush)
 		);
 	}
@@ -357,9 +358,15 @@ FName FSceneOutlinerGutter::GetColumnID()
 SHeaderRow::FColumn::FArguments FSceneOutlinerGutter::ConstructHeaderRowColumn()
 {
 	return SHeaderRow::Column(GetColumnID())
-		.FixedWidth(16.f)
+		.FixedWidth(24.f)
+		.HAlignHeader(HAlign_Right)
+		.VAlignHeader(VAlign_Center)
+		.HAlignCell(HAlign_Right)
+		.VAlignCell(VAlign_Center)
 		[
-			SNew(SSpacer)
+			SNew(SImage)
+			.ColorAndOpacity(FSlateColor::UseForeground())
+			.Image(FAppStyle::Get().GetBrush("Level.VisibleIcon16x"))
 		];
 }
 
diff --git a/Engine/Source/Editor/SceneOutliner/Public/SSceneOutliner.h b/Engine/Source/Editor/SceneOutliner/Public/SSceneOutliner.h
--- a/Engine/Source/Editor/SceneOutliner/Public/SSceneOutliner.h
+++ b/Engine/Source/Editor/SceneOutliner/Public/SSceneOutliner.h
@@ -682,12 +682,6 @@ namespace SceneOutliner
 		/**	Returns the current visibility of the Empty label */
 		EVisibility GetEmptyLabelVisibility() const;
 
-		/** @return the border brush */
-		const FSlateBrush* OnGetBorderBrush() const;
-
-		/** @return the the color and opacity of the border brush; green if in PIE/SIE mode */
-		FSlateColor OnGetBorderColorAndOpacity() const;
-
 		/** @return the selection mode; disabled entirely if in PIE/SIE mode */
 		ESelectionMode::Type GetSelectionMode() const;
 
@@ -703,13 +697,6 @@ namespace SceneOutliner
 		/** @return the foreground color for the world picker button */
 		FSlateColor GetWorldPickerForegroundColor() const;
 
-		/** The brush to use when in Editor mode */
-		const FSlateBrush* NoBorder;
-		/** The brush to use when in PIE mode */
-		const FSlateBrush* PlayInEditorBorder;
-		/** The brush to use when in SIE mode */
-		const FSlateBrush* SimulateBorder;
-
 	private:
 
 		/** Open a context menu for this scene outliner */
diff --git a/Engine/Source/Runtime/Slate/Private/Widgets/Views/SHeaderRow.cpp b/Engine/Source/Runtime/Slate/Private/Widgets/Views/SHeaderRow.cpp
--- a/Engine/Source/Runtime/Slate/Private/Widgets/Views/SHeaderRow.cpp
+++ b/Engine/Source/Runtime/Slate/Private/Widgets/Views/SHeaderRow.cpp
@@ -22,7 +22,7 @@ class STableColumnHeader : public SCompoundWidget
 public:
 	 
 	SLATE_BEGIN_ARGS(STableColumnHeader)
-		: _Style( &FCoreStyle::Get().GetWidgetStyle<FTableColumnHeaderStyle>("TableView.Header.Column") )
+		: _Style( &FAppStyle::Get().GetWidgetStyle<FTableColumnHeaderStyle>("TableView.Header.Column") )
 		{}
 		SLATE_STYLE_ARGUMENT( FTableColumnHeaderStyle, Style )
 
@@ -96,6 +96,8 @@ public:
 				.VAlign( VAlign_Center )
 				[
 					SNew(STextBlock)
+					.TextStyle(FAppStyle::Get(), "ButtonText")
+					.TransformPolicy(ETextTransformPolicy::ToUpper)
 					.Text( LabelText )
 					.ToolTipText( TooltipText )
 				];
@@ -104,14 +106,35 @@ public:
 		if ( OnSortModeChanged.IsBound() )
 		{
 			//optional main button with the column's title. Used to toggle sorting modes.
-			PrimaryContent = SNew(SButton)
-			.ButtonStyle( FCoreStyle::Get(), "NoBorder" )
+
+			TSharedRef<SButton> SortButton = SNew(SButton)
+			.ButtonStyle( FAppStyle::Get(), "NoBorder" )
 			.ForegroundColor( FSlateColor::UseForeground() )
-			.ContentPadding( FMargin( 0, 2, 0, 2 ) )
 			.OnClicked(this, &STableColumnHeader::OnTitleClicked)
 			[
-				PrimaryContent
+				SNew(SHorizontalBox)
+				+SHorizontalBox::Slot()
+				.AutoWidth()
+				.Padding( FMargin( 0, 2, 0, 2 ) )
+				[
+					PrimaryContent
+				]
+
+				+SHorizontalBox::Slot()
+				.FillWidth(1.0f)
+				.HAlign(HAlign_Left)
+				.VAlign( VAlign_Center)
+				.Padding( FMargin( 4, 0, 0, 0 ) )
+				[
+					SNew(SImage)
+					.ColorAndOpacity(FSlateColor::UseForeground())
+					.Image( this, &STableColumnHeader::GetSortingBrush )
+					.Visibility( this, &STableColumnHeader::GetSortModeVisibility )
+				]
+
 			];
+
+			PrimaryContent = SortButton;
 		}
 		
 		Box->AddSlot()
@@ -135,7 +158,6 @@ public:
 				]
 
 				+SOverlay::Slot()
-				.Padding(FMargin(0,1,0,1))
 				[
 					SNew( SBorder )
 					.Padding( FMargin( 0, 0, AdjustedDefaultHeaderContentPadding.Right, 0 ) )
@@ -143,7 +165,7 @@ public:
 					[
 						SAssignNew( ComboButton, SComboButton)
 						.HasDownArrow(false)
-						.ButtonStyle( FCoreStyle::Get(), "NoBorder" )
+						.ButtonStyle( FAppStyle::Get(), "NoBorder")
 						.ContentPadding( FMargin(0) )
 						.ButtonContent()
 						[
@@ -158,9 +180,8 @@ public:
 				]
 
 				+SOverlay::Slot()
-				.Padding(FMargin(0,0,0,2))
 				.HAlign( HAlign_Center )
-				.VAlign( VAlign_Bottom )
+				.VAlign( VAlign_Center)
 				[
 					SNew(SImage)
 					.Image( &Style->MenuDropdownImage )
@@ -172,16 +193,6 @@ public:
 			AdjustedDefaultHeaderContentPadding.Right = 0;
 		}
 
-		Overlay->AddSlot( 1 )
-			.HAlign(HAlign_Center)
-			.VAlign( VAlign_Top )
-			.Padding( FMargin( 0, 2, 0, 0 ) )
-			[
-				SNew(SImage)
-				.Image( this, &STableColumnHeader::GetSortingBrush )
-				.Visibility( this, &STableColumnHeader::GetSortModeVisibility )
-			];
-
 		this->ChildSlot
 		[
 			SNew( SBorder )
@@ -665,7 +676,7 @@ void SHeaderRow::RegenerateWidgets()
 					const float GripSize = 5.0f;
 					TSharedRef<SBorder> SizingGrip = SNew(SBorder)
 						.Padding(0.0f)
-						.BorderImage(FCoreStyle::Get().GetBrush("NoBorder"))
+						.BorderImage(FAppStyle::Get().GetBrush("NoBorder"))
 						.Cursor(EMouseCursor::ResizeLeftRight)
 						.Content()
 						[
diff --git a/Engine/Source/Runtime/Slate/Public/Widgets/Docking/SDockTab.h b/Engine/Source/Runtime/Slate/Public/Widgets/Docking/SDockTab.h
--- a/Engine/Source/Runtime/Slate/Public/Widgets/Docking/SDockTab.h
+++ b/Engine/Source/Runtime/Slate/Public/Widgets/Docking/SDockTab.h
@@ -65,8 +65,7 @@ public:
 		: _Content()
 		, _TabWellContentLeft()
 		, _TabWellContentRight()
-		, _TabWellContentBackground()
-		, _ContentPadding( FMargin( 2 ) )
+		, _ContentPadding(0)
 		, _TabRole(ETabRole::PanelTab)
 		, _Label()
 		, _LabelSuffix()
@@ -77,12 +76,12 @@ public:
 		, _OnCanCloseTab()
 		, _OnPersistVisualState()
 		, _TabColorScale(FLinearColor::Transparent)
+		, _ForegroundColor(FSlateColor::UseStyle())
 		{}
 
 		SLATE_DEFAULT_SLOT( FArguments, Content )
 		SLATE_NAMED_SLOT( FArguments, TabWellContentLeft )
 		SLATE_NAMED_SLOT( FArguments, TabWellContentRight )
-		SLATE_NAMED_SLOT(FArguments, TabWellContentBackground)
 		SLATE_ATTRIBUTE( FMargin, ContentPadding )
 		SLATE_ARGUMENT( ETabRole, TabRole )
 		SLATE_ATTRIBUTE( FText, Label )
@@ -94,6 +93,7 @@ public:
 		SLATE_EVENT( FCanCloseTab, OnCanCloseTab )
 		SLATE_EVENT( FOnPersistVisualState, OnPersistVisualState )
 		SLATE_ATTRIBUTE( FLinearColor, TabColorScale )
+		SLATE_ATTRIBUTE( FSlateColor, ForegroundColor )
 	SLATE_END_ARGS()
 
 	/** Construct the widget from the declaration. */
@@ -116,11 +116,11 @@ public:
 	// End of SBorder interface
 
 	/** Content that appears in the TabWell to the left of the tabs */
-	void SetLeftContent( TSharedRef<SWidget> InContent );
+	void SetLeftContent(TSharedRef<SWidget> InContent);
 	/** Content that appears in the TabWell to the right of the tabs */
-	void SetRightContent( TSharedRef<SWidget> InContent );
-	/** Content that appears in the TabWell behind the tabs */
-	void SetBackgroundContent( TSharedRef<SWidget> InContent );
+	void SetRightContent(TSharedRef<SWidget> InContent);
+	/** Content that appears on the right side of the title bar in the window this stack is in */
+	void SetTitleBarRightContent(TSharedRef<SWidget> InContent);
 
 	/** @return True if this tab is currently focused */
 	bool IsActive() const;
@@ -128,6 +128,9 @@ public:
 	/** @return True if this tab appears active; False otherwise */
 	bool IsForeground() const;
 
+	/** @return the Foreground color that this widget sets; unset options if the widget does not set a foreground color */
+	virtual FSlateColor GetForegroundColor() const;
+
 	/** Is this an MajorTab? A tool panel tab? */
 	ETabRole GetTabRole() const;
 
@@ -147,7 +150,7 @@ public:
 	TSharedRef<SWidget> GetContent();
 	TSharedRef<SWidget> GetLeftContent();
 	TSharedRef<SWidget> GetRightContent();
-	TSharedRef<SWidget> GetBackgrounfContent();
+	TSharedRef<SWidget> GetTitleBarRightContent();
 
 	/** Padding around the content when it is presented by the SDockingTabStack */
 	FMargin GetContentPadding() const;
@@ -316,13 +319,22 @@ protected:
 	/** @return if the close button should be visible. */
 	EVisibility HandleIsCloseButtonVisible() const;
 
+	/** @return the size the tab icon should be */
+	TOptional<FVector2D> GetTabIconSize() const;
+
 private:
 	/** Activates the tab in its tab well */
 	EActiveTimerReturnType TriggerActivateTab( double InCurrentTime, float InDeltaTime );
 
-	/** The handle to the active tab activation tick */
-	TWeakPtr<FActiveTimerHandle> ActiveTimerHandle;
+	EActiveTimerReturnType OnHandleUpdateStyle(double InCurrentTime, float InDeltaTime);
+
+	void OnParentSet();
+
+	void UpdateTabStyle();
 
+	/** The handle to the active tab activation tick */
+	TWeakPtr<FActiveTimerHandle> DragDropTimerHandle;
+	TWeakPtr<FActiveTimerHandle> UpdateStyleTimerHandle;
 protected:
 
 	/** The tab manager that created this tab. */
@@ -332,8 +344,7 @@ protected:
 	TSharedRef<SWidget> Content;
 	TSharedRef<SWidget> TabWellContentLeft;
 	TSharedRef<SWidget> TabWellContentRight;
-	TSharedRef<SWidget> TabWellContentBackground;
-
+	TSharedRef<SWidget> TitleBarContentRight;
 	
 	/** The tab's layout identifier */
 	FTabId LayoutIdentifier;
@@ -397,7 +408,8 @@ protected:
 
 	/** Widget used to show the label on the tab */
 	TSharedPtr<STextBlock> LabelWidget;
-	
+	TSharedPtr<STextBlock> LabelSuffix;
+
 	/** Widget used to show the icon on the tab */
 	TSharedPtr<SImage> IconWidget;
 	
diff --git a/Engine/Source/Runtime/SlateCore/Private/Styling/StarshipCoreStyle.cpp b/Engine/Source/Runtime/SlateCore/Private/Styling/StarshipCoreStyle.cpp
--- a/Engine/Source/Runtime/SlateCore/Private/Styling/StarshipCoreStyle.cpp
+++ b/Engine/Source/Runtime/SlateCore/Private/Styling/StarshipCoreStyle.cpp
@@ -39,13 +39,19 @@ const FLinearColor FStyleColors::White				= COLOR("#FFFFFFFF");
 const FLinearColor FStyleColors::White25			= COLOR("#FFFFFF40");
 const FLinearColor FStyleColors::Highlight			= COLOR("#0078D7FF");
 
-const FLinearColor FStyleColors::Foreground			= COLOR("#A6A6A6FF");
-const FLinearColor FStyleColors::ForegroundHover	= COLOR("#FFFFFFFF");
-
 const FLinearColor FStyleColors::Primary			= COLOR("#26BBFFFF");
 const FLinearColor FStyleColors::PrimaryHover		= COLOR("#6FD2FFFF");
 const FLinearColor FStyleColors::PrimaryPress		= COLOR("#1989BCFF");
 
+const FLinearColor FStyleColors::Foreground			= COLOR("#A6A6A6FF");
+const FLinearColor FStyleColors::ForegroundHover	= COLOR("#FFFFFFFF");
+const FLinearColor FStyleColors::ForegroundInverted	= FStyleColors::Input;
+
+const FLinearColor FStyleColors::Select             = FStyleColors::Primary;
+const FLinearColor FStyleColors::SelectInactive     = COLOR("#99B3BFFF");
+const FLinearColor FStyleColors::SelectParent       = COLOR("#2C323AFF"); 
+const FLinearColor FStyleColors::SelectHover        = FStyleColors::Background;
+
 const FLinearColor FStyleColors::AccentBlue			= COLOR("#26BBFFFF");
 const FLinearColor FStyleColors::AccentPurple		= COLOR("#A139BFFF");
 const FLinearColor FStyleColors::AccentPink			= COLOR("#FF729CFF");
@@ -223,7 +229,7 @@ TSharedRef<ISlateStyle> FStarshipCoreStyle::Create()
 	const float InputFocusThickness = 2.0f;
 
 	// These are the Slate colors which reference the dynamic colors in the style;
-	const FSlateColor DefaultForeground(Style->DefaultForeground_LinearRef);
+	const FSlateColor DefaultForeground(FStyleColors::Foreground);
 	const FSlateColor InvertedForeground(Style->InvertedForeground_LinearRef);
 	const FSlateColor SelectorColor(Style->SelectorColor_LinearRef);
 	const FSlateColor SelectionColor(Style->SelectionColor_LinearRef);
@@ -244,9 +250,18 @@ TSharedRef<ISlateStyle> FStarshipCoreStyle::Create()
 	Style->Set("Colors.Hover2", FStyleColors::Hover2);
 	Style->Set("Colors.White", FStyleColors::White);
 	Style->Set("Colors.White25", FStyleColors::White25);
+	Style->Set("Colors.Highlight", FStyleColors::Highlight);
+
 	Style->Set("Colors.Foreground", FStyleColors::Foreground);
 	Style->Set("Colors.ForegroundHover", FStyleColors::ForegroundHover);
-	Style->Set("Colors.Highlight", FStyleColors::Highlight);
+	Style->Set("Colors.ForegroundInverted", FStyleColors::ForegroundInverted);
+
+	Style->Set("Colors.Select", FStyleColors::Select);
+	Style->Set("Colors.SelectInactive", FStyleColors::SelectInactive);
+	Style->Set("Colors.SelectParent", FStyleColors::SelectParent);
+	Style->Set("Colors.SelectHover", FStyleColors::SelectHover);
+	
+
 	Style->Set("Colors.Primary", FStyleColors::Primary);
 	Style->Set("Colors.PrimaryHover", FStyleColors::PrimaryHover);
 	Style->Set("Colors.PrimaryPress", FStyleColors::PrimaryPress);
@@ -284,10 +299,11 @@ TSharedRef<ISlateStyle> FStarshipCoreStyle::Create()
 
 
 	// Common Brushes
-	Style->Set("PerspectiveBox", new IMAGE_BRUSH_SVG("Starship/Common/box-perspective", Icon16x16));
-	Style->Set("Cylinder", new IMAGE_BRUSH_SVG("Starship/Common/cylinder", Icon16x16));
-	Style->Set("Pyramid", new IMAGE_BRUSH_SVG("Starship/Common/pyriamid", Icon16x16));
-	Style->Set("Sphere", new IMAGE_BRUSH_SVG("Starship/Common/sphere", Icon16x16));
+	Style->Set("Icons.box-perspective", new IMAGE_BRUSH_SVG("Starship/Common/box-perspective", Icon16x16));
+	Style->Set("Icons.cylinder", new IMAGE_BRUSH_SVG("Starship/Common/cylinder", Icon16x16));
+	Style->Set("Icons.pyramid", new IMAGE_BRUSH_SVG("Starship/Common/pyriamid", Icon16x16));
+	Style->Set("Icons.sphere", new IMAGE_BRUSH_SVG("Starship/Common/sphere", Icon16x16));
+	Style->Set("Icons.cogwheel", new IMAGE_BRUSH_SVG("Starship/Common/cogwheel", Icon16x16));
 
 
 	// Common Margins
@@ -380,6 +396,8 @@ TSharedRef<ISlateStyle> FStarshipCoreStyle::Create()
 		Style->Set("NoBorder.Pressed", new FSlateNoResource());
 
 		Style->Set("NoBorder", NoBorder);
+
+		Style->Set("InputBrush", new FSlateColorBrush(FStyleColors::Input));
 	}
 
 
@@ -1068,19 +1086,26 @@ TSharedRef<ISlateStyle> FStarshipCoreStyle::Create()
 	// TableView defaults...
 	{
 		const FTableRowStyle DefaultTableRowStyle = FTableRowStyle()
-			.SetEvenRowBackgroundBrush(FSlateNoResource())
-			.SetEvenRowBackgroundHoveredBrush(IMAGE_BRUSH("Common/Selection", Icon8x8, FLinearColor(1.0f, 1.0f, 1.0f, 0.1f)))
-			.SetOddRowBackgroundBrush(FSlateNoResource())
-			.SetOddRowBackgroundHoveredBrush(IMAGE_BRUSH("Common/Selection", Icon8x8, FLinearColor(1.0f, 1.0f, 1.0f, 0.1f)))
+			.SetEvenRowBackgroundBrush(FSlateColorBrush(FStyleColors::Input))
+			.SetEvenRowBackgroundHoveredBrush(FSlateColorBrush(FStyleColors::SelectHover))
+
+			.SetOddRowBackgroundBrush(FSlateColorBrush(FStyleColors::Input))
+			.SetOddRowBackgroundHoveredBrush(FSlateColorBrush(FStyleColors::SelectHover))
+
 			.SetSelectorFocusedBrush(BORDER_BRUSH("Common/Selector", FMargin(4.f / 16.f), SelectorColor))
-			.SetActiveBrush(IMAGE_BRUSH("Common/Selection", Icon8x8, SelectionColor))
-			.SetActiveHoveredBrush(IMAGE_BRUSH("Common/Selection", Icon8x8, SelectionColor))
-			.SetInactiveBrush(IMAGE_BRUSH("Common/Selection", Icon8x8, SelectionColor_Inactive))
-			.SetInactiveHoveredBrush(IMAGE_BRUSH("Common/Selection", Icon8x8, SelectionColor_Inactive))
-			.SetActiveHighlightedBrush(IMAGE_BRUSH("Common/Selection", Icon8x8, HighlightColor))
-			.SetInactiveHighlightedBrush(IMAGE_BRUSH("Common/Selection", Icon8x8, FSlateColor(FLinearColor(.1f, .1f, .1f))))
-			.SetTextColor(DefaultForeground)
-			.SetSelectedTextColor(InvertedForeground)
+
+			.SetActiveBrush(         FSlateColorBrush(FStyleColors::Select))
+			.SetActiveHoveredBrush(  FSlateColorBrush(FStyleColors::Select))
+
+			.SetInactiveBrush(       FSlateColorBrush(FStyleColors::SelectInactive))
+			.SetInactiveHoveredBrush(FSlateColorBrush(FStyleColors::SelectHover))
+
+			.SetActiveHighlightedBrush(  FSlateColorBrush(FStyleColors::SelectParent)) // This is the parent hightlight
+			.SetInactiveHighlightedBrush(FSlateColorBrush(FStyleColors::SelectParent))// This is the parent highlight
+
+			.SetTextColor(FStyleColors::Foreground)
+			.SetSelectedTextColor(FStyleColors::ForegroundInverted)
+
 			.SetDropIndicator_Above(BOX_BRUSH("Common/DropZoneIndicator_Above", FMargin(10.0f / 16.0f, 10.0f / 16.0f, 0, 0), SelectionColor))
 			.SetDropIndicator_Onto(BOX_BRUSH("Common/DropZoneIndicator_Onto", FMargin(4.0f / 16.0f), SelectionColor))
 			.SetDropIndicator_Below(BOX_BRUSH("Common/DropZoneIndicator_Below", FMargin(10.0f / 16.0f, 0, 0, 10.0f / 16.0f), SelectionColor));
@@ -1092,33 +1117,37 @@ TSharedRef<ISlateStyle> FStarshipCoreStyle::Create()
 			.SetOddRowBackgroundBrush(IMAGE_BRUSH("Common/Selection", Icon8x8, FLinearColor(0.0f, 0.0f, 0.0f, 0.1f)));
 		Style->Set("TableView.DarkRow", DarkTableRowStyle);
 
-		Style->Set("TreeArrow_Collapsed", new IMAGE_BRUSH("Common/TreeArrow_Collapsed", Icon10x10, DefaultForeground));
-		Style->Set("TreeArrow_Collapsed_Hovered", new IMAGE_BRUSH("Common/TreeArrow_Collapsed_Hovered", Icon10x10, DefaultForeground));
-		Style->Set("TreeArrow_Expanded", new IMAGE_BRUSH("Common/TreeArrow_Expanded", Icon10x10, DefaultForeground));
-		Style->Set("TreeArrow_Expanded_Hovered", new IMAGE_BRUSH("Common/TreeArrow_Expanded_Hovered", Icon10x10, DefaultForeground));
+		Style->Set("TreeArrow_Collapsed",         new IMAGE_BRUSH_SVG("Starship/Common/chevron-right", Icon16x16, FStyleColors::Foreground));
+		Style->Set("TreeArrow_Collapsed_Hovered", new IMAGE_BRUSH_SVG("Starship/Common/chevron-right", Icon16x16, FStyleColors::ForegroundHover));
+		Style->Set("TreeArrow_Expanded",          new IMAGE_BRUSH_SVG("Starship/Common/chevron-down",  Icon16x16, FStyleColors::Foreground));
+		Style->Set("TreeArrow_Expanded_Hovered",  new IMAGE_BRUSH_SVG("Starship/Common/chevron-down",  Icon16x16, FStyleColors::ForegroundHover));
+
 
 		const FTableColumnHeaderStyle TableColumnHeaderStyle = FTableColumnHeaderStyle()
-			.SetSortPrimaryAscendingImage(IMAGE_BRUSH("Common/SortUpArrow", Icon8x4))
-			.SetSortPrimaryDescendingImage(IMAGE_BRUSH("Common/SortDownArrow", Icon8x4))
-			.SetSortSecondaryAscendingImage(IMAGE_BRUSH("Common/SortUpArrows", Icon16x4))
-			.SetSortSecondaryDescendingImage(IMAGE_BRUSH("Common/SortDownArrows", Icon16x4))
-			.SetNormalBrush(BOX_BRUSH("Common/ColumnHeader", 4.f / 32.f))
-			.SetHoveredBrush(BOX_BRUSH("Common/ColumnHeader_Hovered", 4.f / 32.f))
-			.SetMenuDropdownImage(IMAGE_BRUSH("Common/ColumnHeader_Arrow", Icon8x8))
-			.SetMenuDropdownNormalBorderBrush(BOX_BRUSH("Common/ColumnHeaderMenuButton_Normal", 4.f / 32.f))
-			.SetMenuDropdownHoveredBorderBrush(BOX_BRUSH("Common/ColumnHeaderMenuButton_Hovered", 4.f / 32.f));
+			.SetSortPrimaryAscendingImage(   IMAGE_BRUSH_SVG("Starship/CoreWidgets/TableView/sort-up-arrow",    Icon12x12))
+			.SetSortPrimaryDescendingImage(  IMAGE_BRUSH_SVG("Starship/CoreWidgets/TableView/sort-down-arrow",  Icon12x12))
+			.SetSortSecondaryAscendingImage( IMAGE_BRUSH_SVG("Starship/CoreWidgets/TableView/sort-up-arrows",   Icon12x12))
+			.SetSortSecondaryDescendingImage(IMAGE_BRUSH_SVG("Starship/CoreWidgets/TableView/sort-down-arrows", Icon12x12))
+			.SetNormalBrush( FSlateColorBrush(FStyleColors::Background))
+			.SetHoveredBrush(FSlateColorBrush(FStyleColors::Dropdown))
+
+			.SetMenuDropdownImage(IMAGE_BRUSH_SVG("Starship/Common/ellipsis-vertical-narrow", FVector2D(6.f, 24.f)))
+			.SetMenuDropdownNormalBorderBrush(FSlateNoResource())
+			.SetMenuDropdownHoveredBorderBrush(FSlateNoResource());
+
 		Style->Set("TableView.Header.Column", TableColumnHeaderStyle);
 
 		const FTableColumnHeaderStyle TableLastColumnHeaderStyle = FTableColumnHeaderStyle()
-			.SetSortPrimaryAscendingImage(IMAGE_BRUSH("Common/SortUpArrow", Icon8x4))
-			.SetSortPrimaryDescendingImage(IMAGE_BRUSH("Common/SortDownArrow", Icon8x4))
-			.SetSortSecondaryAscendingImage(IMAGE_BRUSH("Common/SortUpArrows", Icon16x4))
-			.SetSortSecondaryDescendingImage(IMAGE_BRUSH("Common/SortDownArrows", Icon16x4))
+			.SetSortPrimaryAscendingImage(   IMAGE_BRUSH_SVG("Starship/CoreWidgets/TableView/sort-up-arrow",    Icon12x12))
+			.SetSortPrimaryDescendingImage(  IMAGE_BRUSH_SVG("Starship/CoreWidgets/TableView/sort-down-arrow",  Icon12x12))
+			.SetSortSecondaryAscendingImage( IMAGE_BRUSH_SVG("Starship/CoreWidgets/TableView/sort-up-arrows",   Icon12x12))
+			.SetSortSecondaryDescendingImage(IMAGE_BRUSH_SVG("Starship/CoreWidgets/TableView/sort-down-arrows", Icon12x12))
 			.SetNormalBrush(FSlateNoResource())
-			.SetHoveredBrush(BOX_BRUSH("Common/LastColumnHeader_Hovered", 4.f / 32.f))
-			.SetMenuDropdownImage(IMAGE_BRUSH("Common/ColumnHeader_Arrow", Icon8x8))
-			.SetMenuDropdownNormalBorderBrush(BOX_BRUSH("Common/ColumnHeaderMenuButton_Normal", 4.f / 32.f))
-			.SetMenuDropdownHoveredBorderBrush(BOX_BRUSH("Common/ColumnHeaderMenuButton_Hovered", 4.f / 32.f));
+			.SetHoveredBrush(FSlateNoResource())
+
+			.SetMenuDropdownImage(IMAGE_BRUSH_SVG("Starship/Common/ellipsis-vertical-narrow", FVector2D(6.f, 24.f)))
+			.SetMenuDropdownNormalBorderBrush(FSlateNoResource())
+			.SetMenuDropdownHoveredBorderBrush(FSlateNoResource());
 
 		const FSplitterStyle TableHeaderSplitterStyle = FSplitterStyle()
 			.SetHandleNormalBrush(FSlateNoResource())
@@ -1128,7 +1157,7 @@ TSharedRef<ISlateStyle> FStarshipCoreStyle::Create()
 			.SetColumnStyle(TableColumnHeaderStyle)
 			.SetLastColumnStyle(TableLastColumnHeaderStyle)
 			.SetColumnSplitterStyle(TableHeaderSplitterStyle)
-			.SetBackgroundBrush(BOX_BRUSH("Common/TableViewHeader", 4.f / 32.f))
+			.SetBackgroundBrush(FSlateColorBrush(FStyleColors::Background))
 			.SetForegroundColor(DefaultForeground)
 		);
 	}
diff --git a/Engine/Source/Runtime/SlateCore/Public/Styling/StarshipCoreStyle.h b/Engine/Source/Runtime/SlateCore/Public/Styling/StarshipCoreStyle.h
--- a/Engine/Source/Runtime/SlateCore/Public/Styling/StarshipCoreStyle.h
+++ b/Engine/Source/Runtime/SlateCore/Public/Styling/StarshipCoreStyle.h
@@ -90,11 +90,21 @@ struct SLATECORE_API FStyleColors
 	static const FLinearColor White;
 	static const FLinearColor White25;
 	static const FLinearColor Highlight; 
-	static const FLinearColor Foreground;
-	static const FLinearColor ForegroundHover;
+
 	static const FLinearColor Primary;
 	static const FLinearColor PrimaryHover;
 	static const FLinearColor PrimaryPress;
+
+
+	static const FLinearColor Foreground;
+	static const FLinearColor ForegroundHover;
+	static const FLinearColor ForegroundInverted;
+
+	static const FLinearColor Select;
+	static const FLinearColor SelectInactive;
+	static const FLinearColor SelectParent;
+	static const FLinearColor SelectHover;
+
 	static const FLinearColor AccentBlue;
 	static const FLinearColor AccentPurple;
 	static const FLinearColor AccentPink;
