Fix extra dark editable text backgrounds

diff --git a/Engine/Source/Runtime/SlateCore/Private/Styling/StarshipCoreStyle.cpp b/Engine/Source/Runtime/SlateCore/Private/Styling/StarshipCoreStyle.cpp
--- a/Engine/Source/Runtime/SlateCore/Private/Styling/StarshipCoreStyle.cpp
+++ b/Engine/Source/Runtime/SlateCore/Private/Styling/StarshipCoreStyle.cpp
@@ -827,14 +827,14 @@ TSharedRef<ISlateStyle> FStarshipCoreStyle::Create()
 
 	// SEditableTextBox defaults...
 	const FEditableTextBoxStyle NormalEditableTextBoxStyle = FEditableTextBoxStyle()
-		.SetBackgroundImageNormal( FSlateRoundedBoxBrush(FStyleColors::Input, InputFocusRadius))
+		.SetBackgroundImageNormal(FSlateRoundedBoxBrush(FStyleColors::Input, InputFocusRadius))
 		.SetBackgroundImageHovered(FSlateRoundedBoxBrush(FStyleColors::Input, InputFocusRadius, FStyleColors::Hover, InputFocusThickness))
 		.SetBackgroundImageFocused(FSlateRoundedBoxBrush(FStyleColors::Input, InputFocusRadius, FStyleColors::Primary, InputFocusThickness))
 		.SetBackgroundImageReadOnly(FSlateRoundedBoxBrush(FStyleColors::Input, InputFocusRadius))
 		.SetFont(StyleFonts.Normal)
 		.SetPadding(FMargin(12.f, 4.0f, 12.f, 5.0f)) // The padding should be 4 top, 5 bottom
 		.SetForegroundColor(FSlateColor::UseForeground())
-		.SetBackgroundColor(FSlateColor::UseForeground())
+		.SetBackgroundColor(FStyleColors::White)
 		.SetReadOnlyForegroundColor(FSlateColor::UseForeground())
 		.SetScrollBarStyle(ScrollBar);
 	{
