Fix include file casing issue

diff --git a/Engine/Source/Runtime/SlateCore/Public/Styling/ToolBarStyle.h b/Engine/Source/Runtime/SlateCore/Public/Styling/ToolBarStyle.h
--- a/Engine/Source/Runtime/SlateCore/Public/Styling/ToolBarStyle.h
+++ b/Engine/Source/Runtime/SlateCore/Public/Styling/ToolBarStyle.h
@@ -6,7 +6,7 @@
 #include "UObject/ObjectMacros.h"
 #include "Styling/SlateTypes.h"
 
-#include "ToolbarStyle.generated.h"
+#include "ToolBarStyle.generated.h"
 
 /**
  * Represents the appearance of a toolbar 
