Remove rogue ensure

diff --git a/Engine/Source/Runtime/SlateCore/Private/Styling/SlateStyleSet.cpp b/Engine/Source/Runtime/SlateCore/Private/Styling/SlateStyleSet.cpp
--- a/Engine/Source/Runtime/SlateCore/Private/Styling/SlateStyleSet.cpp
+++ b/Engine/Source/Runtime/SlateCore/Private/Styling/SlateStyleSet.cpp
@@ -566,8 +566,6 @@ void FSlateStyleSet::LogMissingResource(EStyleMessageSeverity Severity, const FT
 
 		MissingResources.Add(MissingResource);
 		Log(Severity, StyleAndMessage);
-
-		ensure(false);
 	}
 }
 
