CoreStyle: SScrollBar, SScrollBox, SInlineEditableText

diff --git a/Engine/Source/Runtime/AppFramework/Private/Framework/Testing/SStarshipGallery.cpp b/Engine/Source/Runtime/AppFramework/Private/Framework/Testing/SStarshipGallery.cpp
--- a/Engine/Source/Runtime/AppFramework/Private/Framework/Testing/SStarshipGallery.cpp
+++ b/Engine/Source/Runtime/AppFramework/Private/Framework/Testing/SStarshipGallery.cpp
@@ -41,6 +41,7 @@
 #include "Widgets/Layout/SGridPanel.h"
 #include "Widgets/Layout/SUniformWrapPanel.h"
 #include "Widgets/Input/SMenuAnchor.h"
+#include "Widgets/Text/SInlineEditableTextBlock.h"
 #include "Widgets/Text/SMultiLineEditableText.h"
 #include "Widgets/Input/SMultiLineEditableTextBox.h"
 #include "Widgets/Text/STextBlock.h"
@@ -293,11 +294,6 @@ public:
                     +SUniformWrapPanel::Slot()[ GenerateColorButton( FText::FromString("PRIMARY\nHOVER"),  "Colors.PrimaryHover", true)]
                     +SUniformWrapPanel::Slot()[ GenerateColorButton( FText::FromString("PRIMARY\nPRESS"),  "Colors.PrimaryPress", true)]
                 ]
-
-
-               
-
-
             ]
 
             +SHorizontalBox::Slot()
@@ -356,6 +352,7 @@ public:
                     +SUniformWrapPanel::Slot()[ GenerateColorButton( FText::FromString("YELLOW"),    "Colors.AccentYellow", true)]
                     +SUniformWrapPanel::Slot()[ GenerateColorButton( FText::FromString("GREEN"),     "Colors.AccentGreen", true)]
                     +SUniformWrapPanel::Slot()[ GenerateColorButton( FText::FromString("BROWN"),     "Colors.AccentBrown")]
+                    +SUniformWrapPanel::Slot()[ GenerateColorButton( FText::FromString("FOLDER"),    "Colors.AccentFolder")]
                     +SUniformWrapPanel::Slot()[ GenerateColorButton( FText::FromString("BLACK"),     "Colors.AccentBlack", false, true)]
                     +SUniformWrapPanel::Slot()[ GenerateColorButton( FText::FromString("GRAY"),      "Colors.AccentGray")]
                     +SUniformWrapPanel::Slot()[ GenerateColorButton( FText::FromString("WHITE"),     "Colors.AccentWhite", true)]
@@ -366,7 +363,7 @@ public:
         ];
     };
 
-    TArray< TUniquePtr< FSlateDynamicImageBrush > > DynamicBrushes;
+    TArray< TUniquePtr< FSlateBrush > > DynamicBrushes;
     TSharedRef<SWidget> ConstructIconsGallery()
     {
         // auto GenerateColorButton = [] (FTex      t InTitle, FName StyleColorName, bool Inverted = false) -> TSharedRef<SWidget> 
@@ -410,6 +407,46 @@ public:
             ];
         };
 
+        auto GenerateIconLibrarySVG = [this] (FText InTitle, FString InPath) -> TSharedRef<SWidget>
+        {
+            const FVector2D IconSize(20.f, 20.f);
+            TSharedPtr<SUniformWrapPanel> UniformWrapPanel = SNew(SUniformWrapPanel)
+            .HAlign(HAlign_Left)
+            .SlotPadding( FMargin(12.f, 12.f) );
+
+            TArray<FString> FoundIcons;
+            FString SearchDirectory = FPaths::EngineDir() /  InPath;// TEXT("Editor/Slate/Icons/GeneralTools");
+            // IFileManager::Get().FindFiles(FoundIcons, *SearchDirectory, TEXT(".png"));//, true, true, false);
+            IFileManager::Get().FindFilesRecursive(FoundIcons, *SearchDirectory, TEXT("*.svg"), true, false);
+            for (const FString& Filename : FoundIcons)
+            {
+                // FString IconPath = SearchDirectory / Filename;
+                FString IconPath = Filename;
+
+                DynamicBrushes.Add( TUniquePtr<FSlateVectorImageBrush>(new FSlateVectorImageBrush( IconPath, IconSize )));
+
+                UniformWrapPanel->AddSlot()
+                [
+                    SNew(SImage)
+                    .Image(DynamicBrushes.Last().Get())
+                    .ToolTipText( FText::FromString( IconPath ) )
+                ];
+            }
+
+            return SNew(SVerticalBox)
+            +SVerticalBox::Slot()
+            .AutoHeight()
+            [
+                SNew(STextBlock).Text(InTitle)
+            ]
+
+            +SVerticalBox::Slot(        )
+            [
+                UniformWrapPanel.ToSharedRef()
+            ];
+        };
+
+
 
         return SNew(SBorder)
         .BorderImage( FAppStyle::Get().GetBrush("ToolPanel.GroupBorder") )
@@ -425,7 +462,9 @@ public:
                 +SVerticalBox::Slot().AutoHeight()[ GenerateIconLibrary(NSLOCTEXT("StarshipGallery", "ModelingIconTitle", "Modeling"), "/Plugins/Experimental/ModelingToolsEditorMode/Content/Icons")]
                 +SVerticalBox::Slot().AutoHeight()[ GenerateIconLibrary(NSLOCTEXT("StarshipGallery", "FractureIconTitle", "Fracture"), "/Plugins/Experimental/ChaosEditor/Content")]
                 +SVerticalBox::Slot().AutoHeight()[ GenerateIconLibrary(NSLOCTEXT("StarshipGallery", "CurveEditorIconTitle", "CurveEditor"), "Content/Editor/Slate/GenericCurveEditor/Icons")]
-                +SVerticalBox::Slot().AutoHeight()[ GenerateIconLibrary(NSLOCTEXT("StarshipGallery", "QuixelIconTitle", "Quixel"), "Content/Slate/Starship/Icons/Quixel")]
+                // +SVerticalBox::Slot().AutoHeight()[ GenerateIconLibrary(NSLOCTEXT("StarshipGallery", "QuixelIconTitle", "Quixel"), "Content/Slate/Starship/Icons/Quixel")]
+                +SVerticalBox::Slot().AutoHeight()[ GenerateIconLibrarySVG(NSLOCTEXT("StarshipGallery", "QuixelIcons", "Quixel"), "Content/Slate/Starship/Common")]
+                +SVerticalBox::Slot().AutoHeight()[ GenerateIconLibrarySVG(NSLOCTEXT("StarshipGallery", "SceneOutliner", "SceneOutliner"), "Content/Slate/Starship/SceneOutliner")]
             ]
         ];
     }
@@ -1611,6 +1650,15 @@ public:
         ];
 
 
+        // SInlineEditableTextBlock
+        NextSlot(WidgetGrid, LOCTEXT("SInlineLineEditableTextLabel", "SInlineEditableText"))
+        .FillWidth(1.0)
+        [
+            SNew(SInlineEditableTextBlock)
+            .Text(LOCTEXT("SInlineEditableTextHint", "Inline Editable Text"))
+        ];
+
+
         // SSearchBox
         NextSlot(WidgetGrid, LOCTEXT("SSearchBoxLabel", "SSearchBox"))
         .FillWidth(1.0)
diff --git a/Engine/Source/Runtime/Slate/Public/Widgets/Layout/SScrollBar.h b/Engine/Source/Runtime/Slate/Public/Widgets/Layout/SScrollBar.h
--- a/Engine/Source/Runtime/Slate/Public/Widgets/Layout/SScrollBar.h
+++ b/Engine/Source/Runtime/Slate/Public/Widgets/Layout/SScrollBar.h
@@ -28,7 +28,7 @@ class SLATE_API SScrollBar : public SBorder
 public:
 
 	SLATE_BEGIN_ARGS( SScrollBar )
-		: _Style( &FCoreStyle::Get().GetWidgetStyle<FScrollBarStyle>("Scrollbar") )
+		: _Style( &FAppStyle::Get().GetWidgetStyle<FScrollBarStyle>("Scrollbar") )
 		, _OnUserScrolled()
 		, _AlwaysShowScrollbar(false)
 		, _AlwaysShowScrollbarTrack(true)
@@ -39,7 +39,7 @@ public:
 #endif
 		, _Orientation( Orient_Vertical )
 		, _DragFocusCause( EFocusCause::Mouse )
-		, _Thickness( FVector2D(16.0f, 16.0f) )
+		, _Thickness( FVector2D( _Style->Thickness, _Style->Thickness ) )
 		, _Padding( 2.0f )
 		{}
 
diff --git a/Engine/Source/Runtime/Slate/Public/Widgets/Layout/SScrollBox.h b/Engine/Source/Runtime/Slate/Public/Widgets/Layout/SScrollBox.h
--- a/Engine/Source/Runtime/Slate/Public/Widgets/Layout/SScrollBox.h
+++ b/Engine/Source/Runtime/Slate/Public/Widgets/Layout/SScrollBox.h
@@ -96,14 +96,14 @@ public:
 	};
 
 	SLATE_BEGIN_ARGS(SScrollBox)
-		: _Style( &FCoreStyle::Get().GetWidgetStyle<FScrollBoxStyle>("ScrollBox") )
-		, _ScrollBarStyle( &FCoreStyle::Get().GetWidgetStyle<FScrollBarStyle>("ScrollBar") )
+		: _Style( &FAppStyle::Get().GetWidgetStyle<FScrollBoxStyle>("ScrollBox") )
+		, _ScrollBarStyle( &FAppStyle::Get().GetWidgetStyle<FScrollBarStyle>("ScrollBar") )
 		, _ExternalScrollbar()
 		, _Orientation(Orient_Vertical)
 		, _ScrollBarVisibility(EVisibility::Visible)
 		, _ScrollBarAlwaysVisible(false)
 		, _ScrollBarDragFocusCause(EFocusCause::Mouse)
-		, _ScrollBarThickness(FVector2D(9.0f, 9.0f))
+		, _ScrollBarThickness(FVector2D(_Style->BarThickness, _Style->BarThickness))
 		, _ScrollBarPadding(2.0f)
 		, _AllowOverscroll(EAllowOverscroll::Yes)
 		, _AnimateWheelScrolling(false)
diff --git a/Engine/Source/Runtime/SlateCore/Private/Styling/SlateTypes.cpp b/Engine/Source/Runtime/SlateCore/Private/Styling/SlateTypes.cpp
--- a/Engine/Source/Runtime/SlateCore/Private/Styling/SlateTypes.cpp
+++ b/Engine/Source/Runtime/SlateCore/Private/Styling/SlateTypes.cpp
@@ -26,7 +26,16 @@ FCheckBoxStyle::FCheckBoxStyle()
 , UndeterminedHoveredImage()
 , UndeterminedPressedImage()
 , Padding(FMargin(2,0,0,0))
+, BackgroundImage(FSlateNoResource())
+, BackgroundHoveredImage(FSlateNoResource())
+, BackgroundPressedImage(FSlateNoResource())
 , ForegroundColor(FSlateColor::UseForeground())
+, HoveredForeground(FSlateColor::UseForeground())
+, PressedForeground(FSlateColor::UseForeground())
+, CheckedForeground(FSlateColor::UseForeground())
+, CheckedHoveredForeground(FSlateColor::UseForeground())
+, CheckedPressedForeground(FSlateColor::UseForeground())
+, UndeterminedForeground(FSlateColor::UseForeground())
 , BorderBackgroundColor(FLinearColor::White)
 {
 }
@@ -50,6 +59,9 @@ void FCheckBoxStyle::GetResources( TArray< const FSlateBrush* > & OutBrushes ) c
 	OutBrushes.Add( &UndeterminedImage );
 	OutBrushes.Add( &UndeterminedHoveredImage );
 	OutBrushes.Add( &UndeterminedPressedImage );
+	OutBrushes.Add( &BackgroundImage );
+	OutBrushes.Add( &BackgroundHoveredImage );
+	OutBrushes.Add( &BackgroundPressedImage );
 }
 
 #if WITH_EDITORONLY_DATA
@@ -99,6 +111,10 @@ FButtonStyle::FButtonStyle()
 : Normal()
 , Hovered()
 , Pressed()
+, NormalForeground(FSlateColor::UseForeground()) 
+, HoveredForeground(FSlateColor::UseForeground())
+, PressedForeground(FSlateColor::UseForeground())
+, DisabledForeground(FSlateColor::UseForeground())
 , NormalPadding()
 , PressedPadding()
 {
@@ -137,6 +153,7 @@ FComboButtonStyle::FComboButtonStyle()
 	: ShadowOffset(FVector2D::ZeroVector)
 	, ShadowColorAndOpacity(FLinearColor::Black)
 	, MenuBorderPadding(FMargin(0.0f))
+	, ContentPadding(FMargin(5.0f))
 {
 }
 
@@ -156,6 +173,8 @@ const FComboButtonStyle& FComboButtonStyle::GetDefault()
 }
 
 FComboBoxStyle::FComboBoxStyle()
+	: ContentPadding(FMargin(4.f, 2.f))
+	, MenuRowPadding(FMargin(0.f))
 {
 	ComboButtonStyle.SetMenuBorderPadding(FMargin(1.0));
 }
@@ -313,6 +332,7 @@ FScrollBarStyle::FScrollBarStyle()
 	, NormalThumbImage(FSlateNoResource())
 	, HoveredThumbImage(FSlateNoResource())
 	, DraggedThumbImage(FSlateNoResource())
+	, Thickness(16.0f)
 {
 }
 
@@ -577,7 +597,13 @@ const FHeaderRowStyle& FHeaderRowStyle::GetDefault()
 
 
 FDockTabStyle::FDockTabStyle()
-	: OverlapWidth(0.0f)
+	: TabTextStyle( FTextBlockStyle::GetDefault() )
+	, IconSize(16,16)
+	, OverlapWidth(0.0f)
+	, NormalForegroundColor(FSlateColor::UseForeground())
+	, HoveredForegroundColor(FSlateColor::UseForeground())
+	, ActiveForegroundColor(FSlateColor::UseForeground())
+	, ForegroundForegroundColor(FSlateColor::UseForeground())
 {
 }
 
@@ -604,6 +630,7 @@ const FDockTabStyle& FDockTabStyle::GetDefault()
 
 
 FScrollBoxStyle::FScrollBoxStyle()
+	: BarThickness(9.0f)
 {
 }
 
@@ -644,12 +671,15 @@ const FScrollBorderStyle& FScrollBorderStyle::GetDefault()
 
 
 FWindowStyle::FWindowStyle()
-	: BackgroundColor( FLinearColor::White )
-	, OutlineColor( FLinearColor(0.1f, 0.1f, 0.1f, 1.0f) )
+	: BackgroundColor(FLinearColor::White)
+	, OutlineColor(FLinearColor(0.1f, 0.1f, 0.1f, 1.0f))
+	, BorderColor(FLinearColor::White)
+	, WindowCornerRadius(0)
+	, BorderPadding(FMargin(5, 5, 5, 5))
 {
 }
 
-void FWindowStyle::GetResources( TArray< const FSlateBrush* >& OutBrushes ) const
+void FWindowStyle::GetResources(TArray< const FSlateBrush* >& OutBrushes) const
 {
 	MinimizeButtonStyle.GetResources(OutBrushes);
 	MaximizeButtonStyle.GetResources(OutBrushes);
diff --git a/Engine/Source/Runtime/SlateCore/Private/Styling/StarshipCoreStyle.cpp b/Engine/Source/Runtime/SlateCore/Private/Styling/StarshipCoreStyle.cpp
--- a/Engine/Source/Runtime/SlateCore/Private/Styling/StarshipCoreStyle.cpp
+++ b/Engine/Source/Runtime/SlateCore/Private/Styling/StarshipCoreStyle.cpp
@@ -64,6 +64,7 @@ const FLinearColor FStyleColors::AccentBrown		= COLOR("#804D39FF");
 const FLinearColor FStyleColors::AccentBlack		= COLOR("#242424FF");
 const FLinearColor FStyleColors::AccentGray			= COLOR("#808080FF");
 const FLinearColor FStyleColors::AccentWhite		= COLOR("#FFFFFFFF");
+const FLinearColor FStyleColors::AccentFolder       = COLOR("#B68F55FF");
 
 #define FONT(...) FSlateFontInfo(FLegacySlateFontInfoCache::Get().GetDefaultFont(), __VA_ARGS__)
 
@@ -277,6 +278,7 @@ TSharedRef<ISlateStyle> FStarshipCoreStyle::Create()
 	Style->Set("Colors.AccentBlack", FStyleColors::AccentBlack);
 	Style->Set("Colors.AccentGray", FStyleColors::AccentGray);
 	Style->Set("Colors.AccentWhite", FStyleColors::AccentWhite);
+	Style->Set("Colors.AccentFolder", FStyleColors::AccentFolder);
 
 	Style->Set("AppIcon", new IMAGE_BRUSH_SVG("Starship/Common/unreal", FVector2D(36,36), FStyleColors::Foreground));
 	Style->Set("AppIcon.Small", new IMAGE_BRUSH_SVG("Starship/Common/unreal-small", Icon24x24, FStyleColors::Foreground));
@@ -290,7 +292,7 @@ TSharedRef<ISlateStyle> FStarshipCoreStyle::Create()
 		.SetFont(StyleFonts.Normal)
 		.SetColorAndOpacity(FSlateColor::UseForeground())
 		.SetShadowOffset(FVector2D::ZeroVector)
-		.SetHighlightColor(FStyleColors::Highlight)
+		.SetSelectedBackgroundColor(FStyleColors::Highlight)
 		.SetHighlightShape(FSlateColorBrush(FStyleColors::White));
 
 	Style->Set("NormalFont",     StyleFonts.Normal);
@@ -304,7 +306,8 @@ TSharedRef<ISlateStyle> FStarshipCoreStyle::Create()
 	Style->Set("Icons.cylinder", new IMAGE_BRUSH_SVG("Starship/Common/cylinder", Icon16x16));
 	Style->Set("Icons.pyramid", new IMAGE_BRUSH_SVG("Starship/Common/pyriamid", Icon16x16));
 	Style->Set("Icons.sphere", new IMAGE_BRUSH_SVG("Starship/Common/sphere", Icon16x16));
-	Style->Set("Icons.cogwheel", new IMAGE_BRUSH_SVG("Starship/Common/cogwheel", Icon16x16));
+
+	Style->Set("Icons.Settings", new IMAGE_BRUSH_SVG("Starship/SceneOutliner/Settings", Icon16x16));
 
 
 	// Common Margins
@@ -472,13 +475,11 @@ TSharedRef<ISlateStyle> FStarshipCoreStyle::Create()
 	Style->Set("VerticalBoxDragIndicatorShort", new IMAGE_BRUSH("Common/VerticalBoxDragIndicatorShort", FVector2D(6, 15)));
 	// SScrollBar defaults...
 	const FScrollBarStyle ScrollBar = FScrollBarStyle()
-		.SetVerticalTopSlotImage(IMAGE_BRUSH("Common/Scrollbar_Background_Vertical", FVector2D(8, 8)))
-		.SetVerticalBottomSlotImage(IMAGE_BRUSH("Common/Scrollbar_Background_Vertical", FVector2D(8, 8)))
-		.SetHorizontalTopSlotImage(IMAGE_BRUSH("Common/Scrollbar_Background_Horizontal", FVector2D(8, 8)))
-		.SetHorizontalBottomSlotImage(IMAGE_BRUSH("Common/Scrollbar_Background_Horizontal", FVector2D(8, 8)))
-		.SetNormalThumbImage(BOX_BRUSH("Common/Scrollbar_Thumb", FMargin(4.f / 16.f)))
-		.SetDraggedThumbImage(BOX_BRUSH("Common/Scrollbar_Thumb", FMargin(4.f / 16.f)))
-		.SetHoveredThumbImage(BOX_BRUSH("Common/Scrollbar_Thumb", FMargin(4.f / 16.f)));
+		.SetNormalThumbImage( FSlateRoundedBoxBrush(FStyleColors::Hover, 4.0f))
+		.SetDraggedThumbImage(FSlateRoundedBoxBrush(FStyleColors::Hover2, 4.0f))
+		.SetHoveredThumbImage(FSlateRoundedBoxBrush(FStyleColors::Hover2, 4.0f))
+		.SetThickness(8.0f);
+	;
 	{
 		Style->Set("Scrollbar", ScrollBar);
 	}
@@ -834,7 +835,7 @@ TSharedRef<ISlateStyle> FStarshipCoreStyle::Create()
 		.SetBackgroundImageReadOnly(FSlateRoundedBoxBrush(FStyleColors::Input, InputFocusRadius))
 		.SetFont(StyleFonts.Normal)
 		.SetPadding(FMargin(12.f, 4.0f, 12.f, 5.0f)) // The padding should be 4 top, 5 bottom
-		.SetForegroundColor(FSlateColor::UseForeground())
+		.SetForegroundColor(FStyleColors::White)
 		.SetBackgroundColor(FStyleColors::White)
 		.SetReadOnlyForegroundColor(FSlateColor::UseForeground())
 		.SetScrollBarStyle(ScrollBar);
@@ -865,24 +866,11 @@ TSharedRef<ISlateStyle> FStarshipCoreStyle::Create()
 
 	// SInlineEditableTextBlock
 	{
-		FTextBlockStyle InlineEditableTextBlockReadOnly = FTextBlockStyle(NormalText)
-			.SetColorAndOpacity(FSlateColor::UseForeground())
-			.SetShadowOffset(FVector2D::ZeroVector)
-			.SetShadowColorAndOpacity(FLinearColor::Black);
+		// Normal Editable Text 
+		FTextBlockStyle InlineEditableTextBlockReadOnly = FTextBlockStyle(NormalText);
 
-		FTextBlockStyle InlineEditableTextBlockSmallReadOnly = FTextBlockStyle(InlineEditableTextBlockReadOnly)
-			.SetFont(SmallText.Font);
-
-		FEditableTextBoxStyle InlineEditableTextBlockEditable = FEditableTextBoxStyle()
-			.SetFont(NormalText.Font)
-			.SetBackgroundImageNormal(BOX_BRUSH("Common/TextBox", FMargin(4.0f / 16.0f)))
-			.SetBackgroundImageHovered(BOX_BRUSH("Common/TextBox_Hovered", FMargin(4.0f / 16.0f)))
-			.SetBackgroundImageFocused(BOX_BRUSH("Common/TextBox_Hovered", FMargin(4.0f / 16.0f)))
-			.SetBackgroundImageReadOnly(BOX_BRUSH("Common/TextBox_ReadOnly", FMargin(4.0f / 16.0f)))
-			.SetScrollBarStyle(ScrollBar);
-
-		FEditableTextBoxStyle InlineEditableTextBlockSmallEditable = FEditableTextBoxStyle(InlineEditableTextBlockEditable)
-			.SetFont(SmallText.Font);
+		FEditableTextBoxStyle InlineEditableTextBlockEditable = FEditableTextBoxStyle(NormalEditableTextBoxStyle)
+		.SetPadding(FMargin(6.f, 4.5f, 6.f, 4.5f));
 
 		FInlineEditableTextBlockStyle InlineEditableTextBlockStyle = FInlineEditableTextBlockStyle()
 			.SetTextStyle(InlineEditableTextBlockReadOnly)
@@ -890,6 +878,13 @@ TSharedRef<ISlateStyle> FStarshipCoreStyle::Create()
 
 		Style->Set("InlineEditableTextBlockStyle", InlineEditableTextBlockStyle);
 
+		// Small Editable Text 
+		FTextBlockStyle InlineEditableTextBlockSmallReadOnly = FTextBlockStyle(InlineEditableTextBlockReadOnly)
+			.SetFont(StyleFonts.Small);
+
+		FEditableTextBoxStyle InlineEditableTextBlockSmallEditable = FEditableTextBoxStyle(InlineEditableTextBlockEditable)
+			.SetFont(StyleFonts.Small);
+
 		FInlineEditableTextBlockStyle InlineEditableTextBlockSmallStyle = FInlineEditableTextBlockStyle()
 			.SetTextStyle(InlineEditableTextBlockSmallReadOnly)
 			.SetEditableTextBoxStyle(InlineEditableTextBlockSmallEditable);
@@ -1548,6 +1543,7 @@ TSharedRef<ISlateStyle> FStarshipCoreStyle::Create()
 			.SetBottomShadowBrush(BOX_BRUSH("Common/ScrollBoxShadowBottom", FVector2D(16, 8), FMargin(0.5, 0, 0.5, 1)))
 			.SetLeftShadowBrush(BOX_BRUSH("Common/ScrollBoxShadowLeft", FVector2D(8, 16), FMargin(1, 0.5, 0, 0.5)))
 			.SetRightShadowBrush(BOX_BRUSH("Common/ScrollBoxShadowRight", FVector2D(8, 16), FMargin(0, 0.5, 1, 0.5)))
+			.SetBarThickness(8.0)
 		);
 	}
 
diff --git a/Engine/Source/Runtime/SlateCore/Public/Styling/SlateTypes.h b/Engine/Source/Runtime/SlateCore/Public/Styling/SlateTypes.h
--- a/Engine/Source/Runtime/SlateCore/Public/Styling/SlateTypes.h
+++ b/Engine/Source/Runtime/SlateCore/Public/Styling/SlateTypes.h
@@ -143,11 +143,56 @@ struct SLATECORE_API FCheckBoxStyle : public FSlateWidgetStyle
 	FMargin Padding;
 	FCheckBoxStyle& SetPadding( const FMargin& InPadding ){ Padding = InPadding; return *this; }
 
-	/** The foreground color */
+	/** Background appearance */
+	UPROPERTY( EditAnywhere, BlueprintReadWrite, Category = Appearance )
+	FSlateBrush BackgroundImage;
+	FCheckBoxStyle & SetBackgroundImage( const FSlateBrush& InBackgroundImage ){ BackgroundImage = InBackgroundImage; return *this; }
+
+	/** Background appearance when hovered */
+	UPROPERTY( EditAnywhere, BlueprintReadWrite, Category = Appearance )
+	FSlateBrush BackgroundHoveredImage;
+	FCheckBoxStyle & SetBackgroundHoveredImage( const FSlateBrush& InBackgroundHoveredImage ){ BackgroundHoveredImage = InBackgroundHoveredImage; return *this; }
+
+	/** Background appearance when pressed */
+	UPROPERTY( EditAnywhere, BlueprintReadWrite, Category = Appearance )
+	FSlateBrush BackgroundPressedImage;
+	FCheckBoxStyle & SetBackgroundPressedImage( const FSlateBrush& InBackgroundPressedImage ){ BackgroundPressedImage = InBackgroundPressedImage; return *this; }
+
+	/** The normal unchecked foreground color */
 	UPROPERTY(EditAnywhere, BlueprintReadWrite, Category=Appearance)
 	FSlateColor ForegroundColor;
 	FCheckBoxStyle& SetForegroundColor(const FSlateColor& InForegroundColor) { ForegroundColor = InForegroundColor; return *this; }
 
+	/** Foreground Color when hovered */
+	UPROPERTY(EditAnywhere, BlueprintReadWrite, AdvancedDisplay, Category = Appearance)
+	FSlateColor HoveredForeground;
+	FCheckBoxStyle& SetHoveredForegroundColor(const FSlateColor& InHoveredForeground) { HoveredForeground = InHoveredForeground; return *this; }
+
+	/** Foreground Color when pressed */
+	UPROPERTY(EditAnywhere, BlueprintReadWrite, AdvancedDisplay, Category = Appearance)
+	FSlateColor PressedForeground;
+	FCheckBoxStyle& SetPressedForegroundColor(const FSlateColor& InPressedForeground) { PressedForeground = InPressedForeground; return *this; }
+
+	/** Foreground Color when checked */
+	UPROPERTY(EditAnywhere, BlueprintReadWrite, AdvancedDisplay, Category = Appearance)
+	FSlateColor CheckedForeground;
+	FCheckBoxStyle& SetCheckedForegroundColor(const FSlateColor& InCheckedForeground) { CheckedForeground = InCheckedForeground; return *this; }
+
+	/** Foreground Color when checked and pressed */
+	UPROPERTY(EditAnywhere, BlueprintReadWrite, AdvancedDisplay, Category = Appearance)
+	FSlateColor CheckedHoveredForeground;
+	FCheckBoxStyle& SetCheckedHoveredForegroundColor(const FSlateColor& InCheckedHoveredForeground) { CheckedHoveredForeground = InCheckedHoveredForeground; return *this; }
+
+	/** Foreground Color when checked and pressed */
+	UPROPERTY(EditAnywhere, BlueprintReadWrite, AdvancedDisplay, Category = Appearance)
+	FSlateColor CheckedPressedForeground;
+	FCheckBoxStyle& SetCheckedPressedForegroundColor(const FSlateColor& InCheckedPressedForeground) { CheckedPressedForeground = InCheckedPressedForeground; return *this; }
+
+	/** Foreground Color when the check state is indeterminate */
+	UPROPERTY(EditAnywhere, BlueprintReadWrite, AdvancedDisplay, Category = Appearance)
+	FSlateColor UndeterminedForeground;
+	FCheckBoxStyle& SetUndeterminedForegroundColor(const FSlateColor& InUndeterminedForeground) { UndeterminedForeground = InUndeterminedForeground; return *this; }
+
 	/** BorderBackgroundColor refers to the actual color and opacity of the supplied border image on toggle buttons */
 	UPROPERTY(EditAnywhere, BlueprintReadWrite, Category=Appearance)
 	FSlateColor BorderBackgroundColor;
@@ -348,6 +393,27 @@ struct SLATECORE_API FButtonStyle : public FSlateWidgetStyle
 	FSlateBrush Disabled;
 	FButtonStyle& SetDisabled( const FSlateBrush& InDisabled ){ Disabled = InDisabled; return *this; }
 
+	/** Foreground Color when the button is not hovered or pressed */
+	UPROPERTY(EditAnywhere, BlueprintReadWrite, AdvancedDisplay, Category=Appearance)
+	FSlateColor NormalForeground;
+	FButtonStyle& SetNormalForeground( const FSlateColor& InNormalForeground ){ NormalForeground = InNormalForeground; return *this; }
+
+	/** Foreground Color when hovered */
+	UPROPERTY(EditAnywhere, BlueprintReadWrite, AdvancedDisplay, Category=Appearance)
+	FSlateColor HoveredForeground;
+	FButtonStyle& SetHoveredForeground( const FSlateColor& InHoveredForeground){ HoveredForeground = InHoveredForeground; return *this; }
+
+	/** Foreground Color when pressed */
+	UPROPERTY(EditAnywhere, BlueprintReadWrite, AdvancedDisplay, Category=Appearance)
+	FSlateColor PressedForeground;
+	FButtonStyle& SetPressedForeground( const FSlateColor& InPressedForeground ){ PressedForeground = InPressedForeground; return *this; }
+
+	/** Foreground Color when disabled */
+	UPROPERTY(EditAnywhere, BlueprintReadWrite, AdvancedDisplay, Category=Appearance)
+	FSlateColor DisabledForeground;
+	FButtonStyle& SetDisabledForeground( const FSlateColor& InDisabledForeground ){ DisabledForeground = InDisabledForeground; return *this; }
+
+
 	/**
 	 * Padding that accounts for the border in the button's background image.
 	 * When this is applied, the content of the button should appear flush
@@ -479,6 +545,13 @@ struct SLATECORE_API FComboButtonStyle : public FSlateWidgetStyle
 	FMargin MenuBorderPadding;
 	FComboButtonStyle& SetMenuBorderPadding( const FMargin& InMenuBorderPadding ){ MenuBorderPadding = InMenuBorderPadding; return *this; }
 
+	/*
+	 * Button Content Padding 
+	 */
+	UPROPERTY(EditAnywhere, BlueprintReadWrite, Category=Appearance)
+	FMargin ContentPadding;
+	FComboButtonStyle& SetContentPadding( const FMargin& InContentPadding ) { ContentPadding = InContentPadding; return *this; }
+
 	/**
 	* Unlinks all colors in this style.
 	* @see FSlateColor::Unlink
@@ -532,6 +605,22 @@ struct SLATECORE_API FComboBoxStyle : public FSlateWidgetStyle
 	FSlateSound SelectionChangeSlateSound;
 	FComboBoxStyle& SetSelectionChangeSound( const FSlateSound& InSelectionChangeSound ){ SelectionChangeSlateSound = InSelectionChangeSound; return *this; }
 
+	/*
+	 * Button Content Padding 
+	 */
+	UPROPERTY(EditAnywhere, BlueprintReadWrite, Category=Appearance)
+	FMargin ContentPadding;
+	FComboBoxStyle& SetContentPadding( const FMargin& InContentPadding ) { ContentPadding = InContentPadding; return *this; }
+
+	/*
+	 * Menu Row Padding 
+	 */
+	UPROPERTY(EditAnywhere, BlueprintReadWrite, Category=Appearance)
+	FMargin MenuRowPadding;
+	FComboBoxStyle& SetMenuRowPadding( const FMargin& InMenuRowPadding ) { MenuRowPadding = InMenuRowPadding; return *this; }
+
+
+
 #if WITH_EDITORONLY_DATA
 	UPROPERTY()
 	FName PressedSound_DEPRECATED;
@@ -724,6 +813,10 @@ struct SLATECORE_API FScrollBarStyle : public FSlateWidgetStyle
 	FSlateBrush DraggedThumbImage;
 	FScrollBarStyle& SetDraggedThumbImage( const FSlateBrush& InDraggedThumbImage ){ DraggedThumbImage = InDraggedThumbImage; return *this; }
 
+	UPROPERTY(EditAnywhere, BlueprintReadWrite, Category=Appearance)
+	float Thickness;
+	FScrollBarStyle& SetThickness(float InThickness) { Thickness = InThickness; return *this; }
+
 	/**
 	 * Unlinks all colors in this style.
 	 * @see FSlateColor::Unlink
@@ -1555,7 +1648,6 @@ struct SLATECORE_API FDockTabStyle : public FSlateWidgetStyle
 	FSlateBrush ColorOverlayIconBrush;
 	FDockTabStyle& SetColorOverlayIconBrush(const FSlateBrush& InColorOverlayBrush) { ColorOverlayIconBrush = InColorOverlayBrush; return *this; }
 
-
 	/** Brush used when this tab is in the foreground */
 	UPROPERTY(EditAnywhere, Category=Appearance)
 	FSlateBrush ForegroundBrush;
@@ -1576,11 +1668,21 @@ struct SLATECORE_API FDockTabStyle : public FSlateWidgetStyle
 	FSlateBrush TabWellBrush;
 	FDockTabStyle& SetTabWellBrush( const FSlateBrush& InTabWellBrush ){ TabWellBrush = InTabWellBrush; return *this; }
 
+	/** Tab Text Style */
+	UPROPERTY(EditAnywhere, Category=Appearance)
+	FTextBlockStyle TabTextStyle;
+	FDockTabStyle& SetTabTextStyle( const FTextBlockStyle& InTabTextStyle ){ TabTextStyle = InTabTextStyle; return *this; }
+
 	/** Padding used around this tab */
 	UPROPERTY(EditAnywhere, Category=Appearance)
 	FMargin TabPadding;
 	FDockTabStyle& SetTabPadding( const FMargin& InTabPadding ){ TabPadding = InTabPadding; return *this; }
 
+	/** Icon size for icons in this tab */
+	UPROPERTY(EditAnywhere, Category = Appearance)
+	FVector2D IconSize;
+	FDockTabStyle& SetIconSize(const FVector2D& InIconSize) { IconSize = InIconSize; return *this; }
+
 	/** The width that this tab will overlap with side-by-side tabs */
 	UPROPERTY(EditAnywhere, Category=Appearance)
 	float OverlapWidth;
@@ -1590,6 +1692,27 @@ struct SLATECORE_API FDockTabStyle : public FSlateWidgetStyle
 	UPROPERTY(EditAnywhere, Category=Appearance)
 	FSlateColor FlashColor;
 	FDockTabStyle& SetFlashColor( const FSlateColor& InFlashColor ){ FlashColor = InFlashColor; return *this; }
+
+	/** Foreground Color when the tab is not hovered, pressed, active or in the foreground */
+	UPROPERTY(EditAnywhere, AdvancedDisplay, Category=Appearance)
+	FSlateColor NormalForegroundColor;
+	FDockTabStyle& SetNormalForegroundColor( const FSlateColor& InNormalForegroundColor ){ NormalForegroundColor = InNormalForegroundColor; return *this; }
+
+	/** Foreground Color when hovered */
+	UPROPERTY(EditAnywhere, AdvancedDisplay, Category=Appearance)
+	FSlateColor HoveredForegroundColor;
+	FDockTabStyle& SetHoveredForegroundColor( const FSlateColor& InHoveredForegroundColor){ HoveredForegroundColor = InHoveredForegroundColor; return *this; }
+
+	/** Foreground Color when Active */
+	UPROPERTY(EditAnywhere, AdvancedDisplay, Category=Appearance)
+	FSlateColor ActiveForegroundColor;
+	FDockTabStyle& SetActiveForegroundColor( const FSlateColor& InActiveForegroundColor ){ ActiveForegroundColor = InActiveForegroundColor; return *this; }
+
+	/** Foreground Color when this tab is the Foreground tab */
+	UPROPERTY(EditAnywhere, AdvancedDisplay, Category=Appearance)
+	FSlateColor ForegroundForegroundColor;
+	FDockTabStyle& SetForegroundForegroundColor( const FSlateColor& InForegroundForegroundColor ){ ForegroundForegroundColor = InForegroundForegroundColor; return *this; }
+
 };
 
 
@@ -1612,6 +1735,10 @@ struct SLATECORE_API FScrollBoxStyle : public FSlateWidgetStyle
 
 	static const FScrollBoxStyle& GetDefault();
 
+	UPROPERTY(EditAnywhere, BlueprintReadWrite, Category=Appearance)
+	float BarThickness;
+	FScrollBoxStyle& SetBarThickness(float InBarThickness) { BarThickness = InBarThickness; return *this; }
+
 	/** Brush used to draw the top shadow of a scrollbox */
 	UPROPERTY(EditAnywhere, BlueprintReadWrite, Category=Appearance)
 	FSlateBrush TopShadowBrush;
@@ -1764,6 +1891,11 @@ struct SLATECORE_API FWindowStyle : public FSlateWidgetStyle
 	FSlateBrush BorderBrush;
 	FWindowStyle& SetBorderBrush( const FSlateBrush& InBorderBrush ){ BorderBrush = InBorderBrush; return *this; }
 
+	/** Color used to draw the window border */
+	UPROPERTY(EditAnywhere, BlueprintReadWrite, Category = Appearance)
+	FSlateColor BorderColor;
+	FWindowStyle& SetBorderColor(const FSlateColor& InBorderColor) { BorderColor = InBorderColor; return *this; }
+
 	/** Brush used to draw the window background */
 	UPROPERTY(EditAnywhere, BlueprintReadWrite, Category=Appearance)
 	FSlateBrush BackgroundBrush;
@@ -1773,6 +1905,17 @@ struct SLATECORE_API FWindowStyle : public FSlateWidgetStyle
 	UPROPERTY(EditAnywhere, BlueprintReadWrite, Category=Appearance)
 	FSlateBrush ChildBackgroundBrush;
 	FWindowStyle& SetChildBackgroundBrush( const FSlateBrush& InChildBackgroundBrush ){ ChildBackgroundBrush = InChildBackgroundBrush; return *this; }
+
+	/** Window corner rounding.  If this value is <= 0 no rounding will occur.   Used for regular, non-maximized windows only (not tool-tips or decorators.) */
+	UPROPERTY(EditAnywhere, BlueprintReadWrite, Category=Appearance)
+	float WindowCornerRadius;
+	FWindowStyle& SetCornerRadius(float InCornerRadius) { WindowCornerRadius = InCornerRadius; return *this; }
+
+	/** Window corner rounding.  If this value is <= 0 no rounding will occur.   Used for regular, non-maximized windows only (not tool-tips or decorators.) */
+	UPROPERTY(EditAnywhere, BlueprintReadWrite, Category=Appearance)
+	FMargin BorderPadding;
+	FWindowStyle& SetBorderPadding(FMargin InBorderPadding) { BorderPadding = InBorderPadding; return *this; }
+
 };
 
 
diff --git a/Engine/Source/Runtime/SlateCore/Public/Styling/StarshipCoreStyle.h b/Engine/Source/Runtime/SlateCore/Public/Styling/StarshipCoreStyle.h
--- a/Engine/Source/Runtime/SlateCore/Public/Styling/StarshipCoreStyle.h
+++ b/Engine/Source/Runtime/SlateCore/Public/Styling/StarshipCoreStyle.h
@@ -116,6 +116,7 @@ struct SLATECORE_API FStyleColors
 	static const FLinearColor AccentBlack;
 	static const FLinearColor AccentGray;
 	static const FLinearColor AccentWhite;
+	static const FLinearColor AccentFolder;
 
 
 };
