Fix menu bar not being clickable at the very top of the screen

diff --git a/Engine/Source/Runtime/Slate/Private/Framework/MultiBox/SMenuEntryBlock.cpp b/Engine/Source/Runtime/Slate/Private/Framework/MultiBox/SMenuEntryBlock.cpp
--- a/Engine/Source/Runtime/Slate/Private/Framework/MultiBox/SMenuEntryBlock.cpp
+++ b/Engine/Source/Runtime/Slate/Private/Framework/MultiBox/SMenuEntryBlock.cpp
@@ -220,6 +220,7 @@ TSharedRef< SWidget> SMenuEntryBlock::BuildMenuBarWidget( const FMenuEntryBuildP
 		[
 			SNew(SBorder)
 			.BorderImage(this, &SMenuEntryBlock::GetMenuBarButtonBorder)
+			.Padding(0.0f)
 			[
 				// Create a button
 				SNew(SButton)
diff --git a/Engine/Source/Runtime/SlateCore/Private/Styling/StarshipCoreStyle.cpp b/Engine/Source/Runtime/SlateCore/Private/Styling/StarshipCoreStyle.cpp
--- a/Engine/Source/Runtime/SlateCore/Private/Styling/StarshipCoreStyle.cpp
+++ b/Engine/Source/Runtime/SlateCore/Private/Styling/StarshipCoreStyle.cpp
@@ -1378,7 +1378,7 @@ TSharedRef<ISlateStyle> FStarshipCoreStyle::Create()
 		Style->Set("WindowMenuBar.Expand", new IMAGE_BRUSH("Icons/toolbar_expand_16x", Icon16x16));
 		Style->Set("WindowMenuBar.Button", MenuBarButton);
 		Style->Set("WindowMenuBar.Button.SubMenuOpen", new FSlateColorBrush(FStyleColors::Primary));
-		Style->Set("WindowMenuBar.MenuBar.Padding", FMargin(12,2));
+		Style->Set("WindowMenuBar.MenuBar.Padding", FMargin(12,4));
 
 	}
 
