Fixed slate gradient elements not setting up initial gamma correction properly

diff --git a/Engine/Source/Runtime/SlateCore/Private/Rendering/ElementBatcher.cpp b/Engine/Source/Runtime/SlateCore/Private/Rendering/ElementBatcher.cpp
--- a/Engine/Source/Runtime/SlateCore/Private/Rendering/ElementBatcher.cpp
+++ b/Engine/Source/Runtime/SlateCore/Private/Rendering/ElementBatcher.cpp
@@ -342,6 +342,7 @@ void FSlateElementBatcher::AddElementsInternal(const FSlateDrawElementArray& Dra
 		switch ( DrawElement.GetElementType() )
 		{
 		case EElementType::ET_Box:
+		case EElementType::ET_RoundedBox:	
 		{
 			SCOPED_NAMED_EVENT_TEXT("Slate::AddBoxElement", FColor::Magenta);
 			STAT(ElementStat_Boxes++);
@@ -606,9 +607,19 @@ void FSlateElementBatcher::AddBoxElement(const FSlateDrawElement& DrawElement)
 	ESlateBatchDrawFlag DrawFlags = DrawElement.GetBatchFlags();
 	DrawFlags |= ( ( bTileHorizontal ? ESlateBatchDrawFlag::TileU : ESlateBatchDrawFlag::None ) | ( bTileVertical ? ESlateBatchDrawFlag::TileV : ESlateBatchDrawFlag::None ) );
 
+	// Add Shader Parameters for extra RoundedBox parameters
+	ESlateShader ShaderType = ESlateShader::Default;
+	FShaderParams ShaderParams;
+	if (DrawElement.GetElementType() == EElementType::ET_RoundedBox)
+	{
+		ShaderType = ESlateShader::RoundedBox;
+		const FSlateRoundedBoxPayload& RoundedPayload = DrawElement.GetDataPayload<FSlateRoundedBoxPayload>();
 
-	FSlateRenderBatch& RenderBatch = CreateRenderBatch( Layer, FShaderParams(), Resource, ESlateDrawPrimitive::TriangleList, ESlateShader::Default, InDrawEffects, DrawFlags, DrawElement);
+		ShaderParams.PixelParams = FVector4(RoundedPayload.GetRadius(), RoundedPayload.GetOutlineWeight(), LocalSize.X, LocalSize.Y);//RadiusWeight;
+		ShaderParams.PixelParams2 = FVector4(RoundedPayload.GetOutlineColor());
+	}
 
+	FSlateRenderBatch& RenderBatch = CreateRenderBatch( Layer, ShaderParams, Resource, ESlateDrawPrimitive::TriangleList, ShaderType, InDrawEffects, DrawFlags, DrawElement);
 
 	float HorizontalTiling = bTileHorizontal ? LocalSize.X / TextureWidth : 1.0f;
 	float VerticalTiling = bTileVertical ? LocalSize.Y / TextureHeight : 1.0f;
@@ -1480,14 +1491,14 @@ void FSlateElementBatcher::AddGradientElement( const FSlateDrawElement& DrawElem
 		if( StopIndex == 0 )
 		{
 			// First stop does not have a full quad yet so do not create indices
-			RenderBatch.AddVertex(FSlateVertex::Make<Rounding>(RenderTransform, StartPt, FVector2D::ZeroVector, FVector2D::ZeroVector, CurStop.Color.ToFColor(false) ) );
-			RenderBatch.AddVertex(FSlateVertex::Make<Rounding>(RenderTransform, EndPt, FVector2D::ZeroVector, FVector2D::ZeroVector, CurStop.Color.ToFColor(false) ) );
+			RenderBatch.AddVertex(FSlateVertex::Make<Rounding>(RenderTransform, StartPt, FVector2D::ZeroVector, FVector2D::ZeroVector, PackVertexColor(CurStop.Color)));
+			RenderBatch.AddVertex(FSlateVertex::Make<Rounding>(RenderTransform, EndPt, FVector2D::ZeroVector, FVector2D::ZeroVector, PackVertexColor(CurStop.Color)));
 		}
 		else
 		{
 			// All stops after the first have indices and generate quads
-			RenderBatch.AddVertex(FSlateVertex::Make<Rounding>(RenderTransform, StartPt, FVector2D::ZeroVector, FVector2D::ZeroVector, CurStop.Color.ToFColor(false) ) );
-			RenderBatch.AddVertex(FSlateVertex::Make<Rounding>(RenderTransform, EndPt, FVector2D::ZeroVector, FVector2D::ZeroVector, CurStop.Color.ToFColor(false) ) );
+			RenderBatch.AddVertex(FSlateVertex::Make<Rounding>(RenderTransform, StartPt, FVector2D::ZeroVector, FVector2D::ZeroVector, PackVertexColor(CurStop.Color)));
+			RenderBatch.AddVertex(FSlateVertex::Make<Rounding>(RenderTransform, EndPt, FVector2D::ZeroVector, FVector2D::ZeroVector, PackVertexColor(CurStop.Color)));
 
 			// Connect the indices to the previous vertices
 			RenderBatch.AddIndex(IndexStart - 2);
