Adusted the style of minor tabs to make them fit better in the new style

diff --git a/Engine/Source/Editor/EditorStyle/Private/StarshipStyle.cpp b/Engine/Source/Editor/EditorStyle/Private/StarshipStyle.cpp
--- a/Engine/Source/Editor/EditorStyle/Private/StarshipStyle.cpp
+++ b/Engine/Source/Editor/EditorStyle/Private/StarshipStyle.cpp
@@ -4169,7 +4169,7 @@ void FStarshipEditorStyle::FStyle::SetupLevelEditorStyle()
 		Set( "EditorViewport.Right", new IMAGE_BRUSH("Icons/icon_ViewMode_ViewRight_16px", Icon16x16 ) );
 		Set( "EditorViewport.Back", new IMAGE_BRUSH("Icons/icon_ViewMode_ViewBack_16px", Icon16x16 ) );
 
-		Set("EditorViewport.ActiveBorderColor", FStyleColors::PrimaryHover);
+		Set("EditorViewport.ActiveBorderColor", FStyleColors::Primary);
 
 #endif
 
diff --git a/Engine/Source/Runtime/Slate/Private/Framework/Docking/SDockingTabWell.cpp b/Engine/Source/Runtime/Slate/Private/Framework/Docking/SDockingTabWell.cpp
--- a/Engine/Source/Runtime/Slate/Private/Framework/Docking/SDockingTabWell.cpp
+++ b/Engine/Source/Runtime/Slate/Private/Framework/Docking/SDockingTabWell.cpp
@@ -6,7 +6,7 @@
 #include "Framework/Docking/FDockingDragOperation.h"
 #include "HAL/PlatformApplicationMisc.h"
 
-const FVector2D FDockingConstants::MaxMinorTabSize(150.f, 50.0f);
+const FVector2D FDockingConstants::MaxMinorTabSize(160.f, 50.0f);
 const FVector2D FDockingConstants::MaxMajorTabSize(210.f, 50.f);
 
 const FVector2D FDockingConstants::GetMaxTabSizeFor( ETabRole TabRole )
@@ -28,6 +28,8 @@ void SDockingTabWell::Construct( const FArguments& InArgs )
 	ChildBeingDraggedOffset = 0.0f;
 	TabGrabOffsetFraction = FVector2D::ZeroVector;
 		
+	SeparatorBrush = FAppStyle::Get().GetBrush("Docking.Tab.InactiveTabSeparator");
+
 	// We need a valid parent here. TabPanels must exist in a SDockingNode
 	check( InArgs._ParentStackNode.Get().IsValid() );
 	ParentTabStackPtr = InArgs._ParentStackNode.Get();
@@ -45,6 +47,8 @@ int32 SDockingTabWell::GetNumTabs() const
 
 void SDockingTabWell::AddTab( const TSharedRef<SDockTab>& InTab, int32 AtIndex )
 {
+	InTab->SetParent(SharedThis(this));
+
 	// Add the tab and implicitly activate it.
 	if (AtIndex == INDEX_NONE)
 	{
@@ -58,7 +62,6 @@ void SDockingTabWell::AddTab( const TSharedRef<SDockTab>& InTab, int32 AtIndex )
 		BringTabToFront( AtIndex );
 	}
 
-	InTab->SetParent(SharedThis(this));
 
 	const TSharedPtr<SDockingTabStack> ParentTabStack = ParentTabStackPtr.Pin();
 	if (ParentTabStack.IsValid() && ParentTabStack->GetDockArea().IsValid())
@@ -119,7 +122,6 @@ void SDockingTabWell::OnArrangeChildren( const FGeometry& AllottedGeometry, FArr
 	}
 }
 	
-
 int32 SDockingTabWell::OnPaint( const FPaintArgs& Args, const FGeometry& AllottedGeometry, const FSlateRect& MyCullingRect, FSlateWindowElementList& OutDrawElements, int32 LayerId, const FWidgetStyle& InWidgetStyle, bool bParentEnabled ) const
 {
 	// When we are dragging a tab, it must be painted on top of the other tabs, so we cannot
@@ -148,8 +150,39 @@ int32 SDockingTabWell::OnPaint( const FPaintArgs& Args, const FGeometry& Allotte
 		}
 		else
 		{
+			bool bShouldDrawSeparator = false;
+			if (SeparatorBrush && SeparatorBrush->DrawAs != ESlateBrushDrawType::NoDrawType && ArrangedChildren.IsValidIndex(ChildIndex + 1))
+			{
+				const FArrangedWidget& PrevWidget = ArrangedChildren[ChildIndex + 1];
+				bShouldDrawSeparator = !CurWidget.Widget->IsHovered() && !PrevWidget.Widget->IsHovered() && PrevWidget.Widget != ForegroundTab;
+			}
+
 			FSlateRect ChildClipRect = MyCullingRect.IntersectionWith( CurWidget.Geometry.GetLayoutBoundingRect() );
 			const int32 CurWidgetsMaxLayerId = CurWidget.Widget->Paint( Args.WithNewParent(this), CurWidget.Geometry, ChildClipRect, OutDrawElements, MaxLayerId, InWidgetStyle, ShouldBeEnabled( bParentEnabled ) );
+		
+			if(bShouldDrawSeparator)
+			{
+				const float SeparatorHeight = CurWidget.Geometry.GetLocalSize().Y * .65f;
+
+				// Center the separator
+				float Offset = (CurWidget.Geometry.GetLocalSize().Y - SeparatorHeight) / 2.0f;
+				FPaintGeometry Geometry = CurWidget.Geometry.ToPaintGeometry(FVector2D(CurWidget.Geometry.GetLocalSize().X + 1.0f, Offset), FVector2D(1.0f, SeparatorHeight));
+		
+				// This code rounds the position of the widget so we don't end up on half a pixel and end up with a larger size separator than we want
+				FSlateRenderTransform NewTransform = Geometry.GetAccumulatedRenderTransform();
+				NewTransform.SetTranslation(NewTransform.GetTranslation().RoundToVector());
+				Geometry.SetRenderTransform(NewTransform);
+
+				FSlateDrawElement::MakeBox(
+					OutDrawElements,
+					MaxLayerId,
+					Geometry,
+					SeparatorBrush,
+					ESlateDrawEffect::None,
+					SeparatorBrush->GetTint(InWidgetStyle)
+				);
+			}
+
 			MaxLayerId = FMath::Max( MaxLayerId, CurWidgetsMaxLayerId );
 		}
 	}
@@ -166,7 +199,6 @@ int32 SDockingTabWell::OnPaint( const FPaintArgs& Args, const FGeometry& Allotte
 	return MaxLayerId;
 }
 
-
 FVector2D SDockingTabWell::ComputeDesiredSize( float ) const
 {
 	FVector2D DesiredSizeResult(0,0);
@@ -324,7 +356,7 @@ void SDockingTabWell::OnDragEnter( const FGeometry& MyGeometry, const FDragDropE
 			this->TabGrabOffsetFraction = DragDropOperation->GetTabGrabOffsetFraction();
 			
 			// The user should see the contents of the tab that we're dragging.
-			ParentTabStackPtr.Pin()->SetNodeContent(DragDropOperation->GetTabBeingDragged()->GetContent(), SNullWidget::NullWidget, SNullWidget::NullWidget, SNullWidget::NullWidget);
+			ParentTabStackPtr.Pin()->SetNodeContent(DragDropOperation->GetTabBeingDragged()->GetContent(), FDockingStackOptionalContent());
 		}
 	}
 }
@@ -480,6 +512,7 @@ FReply SDockingTabWell::OnMouseMove( const FGeometry& MyGeometry, const FPointer
 void SDockingTabWell::BringTabToFront( int32 TabIndexToActivate )
 {
 	const bool bActiveIndexChanging = TabIndexToActivate != ForegroundTabIndex;
+	
 	if ( bActiveIndexChanging )
 	{
 		const int32 LastForegroundTabIndex = FMath::Min(ForegroundTabIndex, Tabs.Num()-1);
@@ -517,7 +550,7 @@ void SDockingTabWell::BringTabToFront( int32 TabIndexToActivate )
 		}
 		else
 		{
-			TabManager->UpdateMainMenu(false);
+			TabManager->UpdateMainMenu(Tabs[ForegroundTabIndex], false);
 		}
 	}
 }
@@ -629,11 +662,18 @@ void SDockingTabWell::RefreshParentContent()
 			ParentWindowPtr->SetTitle( ForegroundTab->GetTabLabel() );
 		}
 
-		ParentTabStackPtr.Pin()->SetNodeContent( ForegroundTab->GetContent(), ForegroundTab->GetLeftContent(), ForegroundTab->GetRightContent(), ForegroundTab->GetBackgrounfContent() );
+		TSharedPtr<SDockingTabStack> ParentTabStack = ParentTabStackPtr.Pin();
+
+		FDockingStackOptionalContent OptionalContent;
+		OptionalContent.ContentLeft = ForegroundTab->GetLeftContent();
+		OptionalContent.ContentRight = ForegroundTab->GetRightContent();
+		OptionalContent.TitleBarContentRight = ForegroundTab->GetTitleBarRightContent();
+
+		ParentTabStack->SetNodeContent(ForegroundTab->GetContent(), OptionalContent);
 	}
 	else
 	{
-		ParentTabStackPtr.Pin()->SetNodeContent(SNullWidget::NullWidget, SNullWidget::NullWidget, SNullWidget::NullWidget, SNullWidget::NullWidget);
+		ParentTabStackPtr.Pin()->SetNodeContent(SNullWidget::NullWidget, FDockingStackOptionalContent());
 	}
 }
 
diff --git a/Engine/Source/Runtime/Slate/Private/Widgets/Docking/SDockTab.cpp b/Engine/Source/Runtime/Slate/Private/Widgets/Docking/SDockTab.cpp
--- a/Engine/Source/Runtime/Slate/Private/Widgets/Docking/SDockTab.cpp
+++ b/Engine/Source/Runtime/Slate/Private/Widgets/Docking/SDockTab.cpp
@@ -11,6 +11,8 @@
 #include "Framework/Docking/SDockingTabWell.h"
 #include "Framework/Docking/TabCommands.h"
 #include "Framework/Docking/SDockingArea.h"
+#include "Widgets/Colors/SComplexGradient.h"
+#include "Widgets/Layout/SBox.h"
 
 namespace SDockTabDefs
 {
@@ -124,20 +126,22 @@ FReply SDockTab::OnMouseButtonUp( const FGeometry& MyGeometry, const FPointerEve
 void SDockTab::OnDragEnter( const FGeometry& MyGeometry, const FDragDropEvent& DragDropEvent )
 {
 	// Register to activate the tab after a delay
-	if ( !ActiveTimerHandle.IsValid() )
+	if ( !DragDropTimerHandle.IsValid() )
 	{
-		ActiveTimerHandle = RegisterActiveTimer( SDockTabDefs::DragTimerActivate, FWidgetActiveTimerDelegate::CreateSP( this, &SDockTab::TriggerActivateTab ) );
+		DragDropTimerHandle = RegisterActiveTimer( SDockTabDefs::DragTimerActivate, FWidgetActiveTimerDelegate::CreateSP( this, &SDockTab::TriggerActivateTab ) );
 	}
 
+	UpdateTabStyle();
+
 	SBorder::OnDragEnter( MyGeometry, DragDropEvent );
 }
 
 void SDockTab::OnDragLeave( const FDragDropEvent& DragDropEvent )
 {
 	// Unregister the activation timer if it hasn't fired yet
-	if ( ActiveTimerHandle.IsValid() )
+	if ( DragDropTimerHandle.IsValid() )
 	{
-		UnRegisterActiveTimer( ActiveTimerHandle.Pin().ToSharedRef() );
+		UnRegisterActiveTimer( DragDropTimerHandle.Pin().ToSharedRef() );
 	}
 
 	SBorder::OnDragLeave( DragDropEvent );
@@ -146,9 +150,9 @@ void SDockTab::OnDragLeave( const FDragDropEvent& DragDropEvent )
 FReply SDockTab::OnDrop( const FGeometry& MyGeometry, const FDragDropEvent& DragDropEvent )
 {
 	// Unregister the activation timer if it hasn't fired yet
-	if ( ActiveTimerHandle.IsValid() )
+	if ( DragDropTimerHandle.IsValid() )
 	{
-		UnRegisterActiveTimer( ActiveTimerHandle.Pin().ToSharedRef() );
+		UnRegisterActiveTimer( DragDropTimerHandle.Pin().ToSharedRef() );
 	}
 
 	return SBorder::OnDrop( MyGeometry, DragDropEvent );
@@ -205,16 +209,16 @@ void SDockTab::SetRightContent( TSharedRef<SWidget> InContent )
 	}
 }
 
-void SDockTab::SetBackgroundContent(TSharedRef<SWidget> InContent)
+void SDockTab::SetTitleBarRightContent(TSharedRef<SWidget> InContent)
 {
-	this->TabWellContentBackground = InContent;
+	this->TitleBarContentRight = InContent;
 	if (ParentPtr.IsValid())
 	{
+		// This is critical to do, otherwise the content might remain if currently active even if expected to be destroyed
 		ParentPtr.Pin()->RefreshParentContent();
 	}
 }
 
-
 bool SDockTab::IsActive() const
 {
 	return FGlobalTabmanager::Get()->GetActiveTab() == SharedThis(this);
@@ -225,6 +229,30 @@ bool SDockTab::IsForeground() const
 	return ParentPtr.IsValid() ? (ParentPtr.Pin()->GetForegroundTab() == SharedThis(this)) : true;
 }
 
+FSlateColor SDockTab::GetForegroundColor() const
+{
+	if (ForegroundColor.IsBound() || ForegroundColor.Get() != FSlateColor::UseStyle())
+	{
+		return ForegroundColor.Get();
+	}
+	else if (IsForeground())
+	{
+		return GetCurrentStyle().ForegroundForegroundColor;
+	}
+	else if (IsActive())
+	{
+		return GetCurrentStyle().ActiveForegroundColor;
+	}
+	else if (IsHovered())
+	{
+		return GetCurrentStyle().HoveredForegroundColor;
+	}
+	else
+	{
+		return GetCurrentStyle().NormalForegroundColor;
+	}
+}
+
 ETabRole SDockTab::GetTabRole() const
 {
 	return TabRole;
@@ -285,9 +313,9 @@ TSharedRef<SWidget> SDockTab::GetRightContent()
 	return TabWellContentRight;
 }
 
-TSharedRef<SWidget> SDockTab::GetBackgrounfContent()
+TSharedRef<SWidget> SDockTab::GetTitleBarRightContent()
 {
-	return TabWellContentBackground;
+	return TitleBarContentRight;
 }
 
 FMargin SDockTab::GetContentPadding() const
@@ -308,6 +336,7 @@ const FTabId& SDockTab::GetLayoutIdentifier() const
 void SDockTab::SetParent(TSharedPtr<SDockingTabWell> Parent)
 {
 	ParentPtr = Parent;
+	OnParentSet();
 }
 
 TSharedPtr<SDockingTabWell> SDockTab::GetParent() const
@@ -349,7 +378,7 @@ SDockTab::SDockTab()
 	: Content(SNew(SSpacer))
 	, TabWellContentLeft(SNullWidget::NullWidget)
 	, TabWellContentRight(SNullWidget::NullWidget)
-	, TabWellContentBackground(SNullWidget::NullWidget)
+	, TitleBarContentRight(SNullWidget::NullWidget)
 	, LayoutIdentifier(NAME_None)
 	, TabRole(ETabRole::PanelTab)
 	, ParentPtr()
@@ -487,18 +516,24 @@ void SDockTab::Construct( const FArguments& InArgs )
 	this->bShouldAutosize = InArgs._ShouldAutosize;
 	this->TabColorScale = InArgs._TabColorScale;
 
-	MajorTabStyle = &FCoreStyle::Get().GetWidgetStyle<FDockTabStyle>("Docking.MajorTab");
-	GenericTabStyle = &FCoreStyle::Get().GetWidgetStyle<FDockTabStyle>("Docking.Tab");
+	MajorTabStyle = &FAppStyle::Get().GetWidgetStyle<FDockTabStyle>("Docking.MajorTab");
+	GenericTabStyle = &FAppStyle::Get().GetWidgetStyle<FDockTabStyle>("Docking.Tab");
 
 	ContentAreaPadding = InArgs._ContentPadding;
 
 	const FButtonStyle* const CloseButtonStyle = &GetCurrentStyle().CloseButtonStyle;
+	const FTextBlockStyle* const TabTextStyle = &GetCurrentStyle().TabTextStyle;
+
+	// This makes a gradient that displays whether or not a viewport is active
+	static FLinearColor ActiveBorderColor = FAppStyle::Get().GetColor("Docking.Tab.ActiveTabIndicatorColor");
+	static FLinearColor ActiveBorderColorTransparent = FLinearColor(ActiveBorderColor.R, ActiveBorderColor.G, ActiveBorderColor.B, 0.0f);
+	static TArray<FLinearColor> GradientStops{ ActiveBorderColorTransparent, ActiveBorderColor, ActiveBorderColorTransparent };
 
 	SBorder::Construct( SBorder::FArguments()
 		.BorderImage( FStyleDefaults::GetNoBrush() )
-		.ContentScale( this, &SDockTab::GetAnimatedScale )
 		.VAlign(VAlign_Bottom)
 		.Padding( 0.0f )
+		.ForegroundColor(InArgs._ForegroundColor)
 		[
 			SNew(SOverlay)
 			+ SOverlay::Slot()
@@ -509,12 +544,14 @@ void SDockTab::Construct( const FArguments& InArgs )
 
 			// Overlay for active tab indication.
 			+ SOverlay::Slot()
+			.VAlign(VAlign_Top)
+			.HAlign(HAlign_Fill)
 			[
-				SNew(SBorder)
-				// Don't allow active tab overlay to absorb mouse clicks
-				.Visibility( EVisibility::HitTestInvisible )
-				.Padding( this, &SDockTab::GetTabPadding )
-				.BorderImage( this, &SDockTab::GetActiveTabOverlayImageBrush )
+				SNew(SComplexGradient)
+				.Visibility(this, &SDockTab::GetActiveTabIndicatorVisibility)
+				.DesiredSizeOverride(FVector2D(1.0f, 1.0f))
+				.GradientColors(GradientStops)
+				.Orientation(EOrientation::Orient_Vertical)
 			]
 
 			// Overlay for flashing a tab for attention
@@ -553,7 +590,9 @@ void SDockTab::Construct( const FArguments& InArgs )
 					.BorderBackgroundColor(this, &SDockTab::GetTabColor)
 					[
 						SAssignNew(IconWidget, SImage)
+						.ColorAndOpacity(FSlateColor::UseForeground())
 						.Image(this, &SDockTab::GetTabIcon)
+						.DesiredSizeOverride(this, &SDockTab::GetTabIconSize)
 					]
 				]
 
@@ -565,7 +604,6 @@ void SDockTab::Construct( const FArguments& InArgs )
 				[
 					// Label sub HBox
 					SNew(SHorizontalBox)
-					.Visibility(EVisibility::Visible)
 					.ToolTip(InArgs._ToolTip)
 
 					// Tab Label
@@ -575,7 +613,7 @@ void SDockTab::Construct( const FArguments& InArgs )
 					.VAlign(VAlign_Center)
 					[
 						SAssignNew(LabelWidget, STextBlock)
-						.TextStyle(FCoreStyle::Get(), "Docking.TabFont")
+						.TextStyle(&GetCurrentStyle().TabTextStyle)
 						.Text(this, &SDockTab::GetTabLabel)
 					]
 
@@ -585,8 +623,8 @@ void SDockTab::Construct( const FArguments& InArgs )
 					.Padding(0.0f, 1.0f)
 					.VAlign(VAlign_Center)
 					[
-						SAssignNew(LabelWidget, STextBlock)
-						.TextStyle(FCoreStyle::Get(), "Docking.TabFont")					
+						SAssignNew(LabelSuffix, STextBlock)
+						.TextStyle(&GetCurrentStyle().TabTextStyle)
 						.Text(this, &SDockTab::GetTabLabelSuffix)
 					]
 				]
@@ -620,6 +658,33 @@ EActiveTimerReturnType SDockTab::TriggerActivateTab( double InCurrentTime, float
 	return EActiveTimerReturnType::Stop;
 }
 
+EActiveTimerReturnType SDockTab::OnHandleUpdateStyle(double InCurrentTime, float InDeltaTime)
+{
+	if (GetParent().IsValid() && GetParent()->GetDockArea().IsValid())
+	{
+		UpdateTabStyle();
+		return EActiveTimerReturnType::Stop;
+	}
+
+	return EActiveTimerReturnType::Continue;
+}
+
+void SDockTab::OnParentSet()
+{
+	// Update the style once the parent is set as it could change the way things look for certian tab types.
+	if (!UpdateStyleTimerHandle.IsValid())
+	{
+		UpdateStyleTimerHandle = RegisterActiveTimer(0, FWidgetActiveTimerDelegate::CreateSP(this, &SDockTab::OnHandleUpdateStyle));
+	}
+}
+
+void SDockTab::UpdateTabStyle()
+{
+	const FTextBlockStyle& TabTextStyle = GetCurrentStyle().TabTextStyle;
+	LabelSuffix->SetTextStyle(&TabTextStyle);
+	LabelWidget->SetTextStyle(&TabTextStyle);
+}
+
 const FDockTabStyle& SDockTab::GetCurrentStyle() const
 {
 	if ( GetVisualTabRole() == ETabRole::MajorTab )
@@ -650,11 +715,8 @@ const FSlateBrush* SDockTab::GetImageBrush() const
 FMargin SDockTab::GetTabPadding() const
 {
 	FMargin NewPadding = GetCurrentStyle().TabPadding;
-	if (TabIcon.Get() != FStyleDefaults::GetNoBrush())
-	{
-		NewPadding.Top *= SDockTabDefs::TabVerticalPaddingScaleOverride;
-		NewPadding.Bottom *= SDockTabDefs::TabVerticalPaddingScaleOverride;
-	}
+	NewPadding.Top *= SDockTabDefs::TabVerticalPaddingScaleOverride;
+	NewPadding.Bottom *= SDockTabDefs::TabVerticalPaddingScaleOverride;
 	return NewPadding;
 }
 
@@ -667,17 +729,11 @@ const FSlateBrush* SDockTab::GetColorOverlayImageBrush() const
 	return FStyleDefaults::GetNoBrush();
 }
 
-
-const FSlateBrush* SDockTab::GetActiveTabOverlayImageBrush() const
+EVisibility SDockTab::GetActiveTabIndicatorVisibility() const
 {
-	if (this->IsActive())
-	{
-		return &GetCurrentStyle().ActiveBrush;
-	}
-	return FStyleDefaults::GetNoBrush();
+	return IsActive() ? EVisibility::HitTestInvisible : EVisibility::Collapsed;
 }
 
-
 FSlateColor SDockTab::GetTabColor() const
 {
 	return TabColorScale.Get();
@@ -786,7 +842,12 @@ FText SDockTab::GetCloseButtonToolTipText() const
 
 EVisibility SDockTab::HandleIsCloseButtonVisible() const
 {
-	return MyTabManager.Pin()->IsTabCloseable(SharedThis(this)) ? EVisibility::Visible : EVisibility::Hidden;
+	return ((IsHovered() || IsForeground()) && MyTabManager.Pin()->IsTabCloseable(SharedThis(this))) ? EVisibility::Visible : EVisibility::Hidden;
+}
+
+TOptional<FVector2D> SDockTab::GetTabIconSize() const
+{
+	return GetCurrentStyle().IconSize;
 }
 
 bool SDockTab::CanCloseTab() const
diff --git a/Engine/Source/Runtime/Slate/Public/Widgets/Docking/SDockTab.h b/Engine/Source/Runtime/Slate/Public/Widgets/Docking/SDockTab.h
--- a/Engine/Source/Runtime/Slate/Public/Widgets/Docking/SDockTab.h
+++ b/Engine/Source/Runtime/Slate/Public/Widgets/Docking/SDockTab.h
@@ -295,8 +295,8 @@ protected:
 	/** @return the image brush for the tab's color overlay */
 	const FSlateBrush* GetColorOverlayImageBrush() const;
 
-	/** @return the image brush for the tab's active state overlay */
-	const FSlateBrush* GetActiveTabOverlayImageBrush() const;
+	/** @return The visibility of the active tab indicator */
+	EVisibility GetActiveTabIndicatorVisibility() const;
 
 	/** @return Returns a color to scale the background of this tab by */
 	FSlateColor GetTabColor() const;
diff --git a/Engine/Source/Runtime/SlateCore/Private/Styling/CoreStyle.cpp b/Engine/Source/Runtime/SlateCore/Private/Styling/CoreStyle.cpp
--- a/Engine/Source/Runtime/SlateCore/Private/Styling/CoreStyle.cpp
+++ b/Engine/Source/Runtime/SlateCore/Private/Styling/CoreStyle.cpp
@@ -10,6 +10,7 @@
 #include "Styling/SlateStyle.h"
 #include "Styling/SlateTypes.h"
 #include "Fonts/LegacySlateFontInfoCache.h"
+#include "Misc/CommandLine.h"
 
 
 /* Static initialization
@@ -178,7 +179,11 @@ TSharedRef<ISlateStyle> FCoreStyle::Create( const FName& InStyleSetName )
 	const FSlateColor SelectionColor_Pressed( Style->SelectionColor_Pressed_LinearRef );
 	const FSlateColor HighlightColor( Style->HighlightColor_LinearRef );
 
-	Style->Set("DefaultAppIcon", new IMAGE_BRUSH("Icons/DefaultAppIcon", Icon24x24));
+	Style->Set("AppIcon", new IMAGE_BRUSH("Icons/DefaultAppIcon", Icon24x24));
+	Style->Set("AppIcon.Small", new IMAGE_BRUSH("Icons/DefaultAppIcon", Icon24x24));
+
+	Style->Set("AppIconPadding", FMargin(4, 4, 0, 0));
+	Style->Set("AppIconPadding.Small", FMargin(4, 4, 0, 0));
 
 	Style->Set("NormalFont", DEFAULT_FONT("Regular", RegularTextSize));
 
@@ -354,11 +359,15 @@ TSharedRef<ISlateStyle> FCoreStyle::Create( const FName& InStyleSetName )
 
 	// SButton defaults...
 	const FButtonStyle Button = FButtonStyle()
-		.SetNormal( BOX_BRUSH( "Common/Button", FVector2D(32,32), 8.0f/32.0f ) )
-		.SetHovered( BOX_BRUSH( "Common/Button_Hovered", FVector2D(32,32), 8.0f/32.0f ) )
-		.SetPressed( BOX_BRUSH( "Common/Button_Pressed", FVector2D(32,32), 8.0f/32.0f ) )
-		.SetNormalPadding( FMargin( 2,2,2,2 ) )
-		.SetPressedPadding( FMargin( 2,3,2,1 ) );
+		.SetNormal(BOX_BRUSH("Common/Button", FVector2D(32, 32), 8.0f / 32.0f))
+		.SetHovered(BOX_BRUSH("Common/Button_Hovered", FVector2D(32, 32), 8.0f / 32.0f))
+		.SetPressed(BOX_BRUSH("Common/Button_Pressed", FVector2D(32, 32), 8.0f / 32.0f))
+		.SetNormalPadding(FMargin(2, 2, 2, 2))
+		.SetPressedPadding(FMargin(2, 3, 2, 1))
+		.SetNormalForeground(InvertedForeground)
+		.SetPressedForeground(InvertedForeground)
+		.SetHoveredForeground(InvertedForeground)
+		.SetDisabledForeground(InvertedForeground);
 	{
 		Style->Set( "Button", Button );
 
@@ -772,6 +781,9 @@ TSharedRef<ISlateStyle> FCoreStyle::Create( const FName& InStyleSetName )
 			.SetDropIndicator_Below(BOX_BRUSH("Common/DropZoneIndicator_Below", FMargin(10.0f / 16.0f, 0, 0, 10.0f / 16.0f), SelectionColor));
 		Style->Set("TableView.Row", DefaultTableRowStyle);
 
+		// Make this the default for the ComboBox rows also
+		Style->Set("ComboBox.Row", DefaultTableRowStyle);  
+
 		const FTableRowStyle DarkTableRowStyle = FTableRowStyle(DefaultTableRowStyle)
 			.SetEvenRowBackgroundBrush(IMAGE_BRUSH("Common/Selection", Icon8x8, FLinearColor(0.0f, 0.0f, 0.0f, 0.1f)))
 			.SetOddRowBackgroundBrush(IMAGE_BRUSH("Common/Selection", Icon8x8, FLinearColor(0.0f, 0.0f, 0.0f, 0.1f)));
@@ -913,16 +925,24 @@ TSharedRef<ISlateStyle> FCoreStyle::Create( const FName& InStyleSetName )
 		Style->Set( "ToolBar.Button", FButtonStyle(Button)
 			.SetNormal ( FSlateNoResource() )
 			.SetPressed( BOX_BRUSH( "Common/RoundedSelection_16x", 4.0f/16.0f, SelectionColor_Pressed ) )
-			.SetHovered( BOX_BRUSH( "Common/RoundedSelection_16x", 4.0f/16.0f, SelectionColor ) )
+			.SetHovered( BOX_BRUSH( "Common/RoundedSelection_16x", 4.0f/16.0f, SelectionColor ))
+			.SetNormalForeground(FSlateColor::UseForeground())
+			.SetPressedForeground(FSlateColor::UseForeground())
+			.SetHoveredForeground(FSlateColor::UseForeground())
+			.SetDisabledForeground(FSlateColor::UseForeground())
 		);
 
-		Style->Set( "ToolBar.Button.Normal", new FSlateNoResource() );
+	/*	Style->Set( "ToolBar.Button.Normal", new FSlateNoResource() );
 		Style->Set( "ToolBar.Button.Pressed", new BOX_BRUSH( "Common/RoundedSelection_16x", 4.0f/16.0f, SelectionColor_Pressed ) );
 		Style->Set( "ToolBar.Button.Hovered", new BOX_BRUSH( "Common/RoundedSelection_16x", 4.0f/16.0f, SelectionColor ) );
+*/
 
 		Style->Set( "ToolBar.Button.Checked", new BOX_BRUSH( "Common/RoundedSelection_16x",  4.0f/16.0f, SelectionColor_Pressed ) );
 		Style->Set( "ToolBar.Button.Checked_Hovered", new BOX_BRUSH( "Common/RoundedSelection_16x",  4.0f/16.0f, SelectionColor_Pressed ) );
 		Style->Set( "ToolBar.Button.Checked_Pressed", new BOX_BRUSH( "Common/RoundedSelection_16x",  4.0f/16.0f, SelectionColor ) );
+
+		Style->Set("ToolBar.SimpleComboButton", Style->GetWidgetStyle<FComboButtonStyle>("ComboButton"));
+		Style->Set("ToolBar.IconSize", Icon20x20);
 	}
 
 	// MenuBar
@@ -937,6 +957,7 @@ TSharedRef<ISlateStyle> FCoreStyle::Create( const FName& InStyleSetName )
 		Style->Set( "Menu.SToolBarButtonBlock.CheckBox.Padding", FMargin(0.0f) );
 		Style->Set( "Menu.SToolBarComboButtonBlock.ComboButton.Color", DefaultForeground );
 
+
 		Style->Set( "Menu.Block.IndentedPadding", FMargin( 18.0f, 2.0f, 4.0f, 4.0f ) );
 		Style->Set( "Menu.Block.Padding", FMargin( 2.0f, 2.0f, 4.0f, 4.0f ) );
 
@@ -951,6 +972,9 @@ TSharedRef<ISlateStyle> FCoreStyle::Create( const FName& InStyleSetName )
 			.SetFont( DEFAULT_FONT( "Regular", 8 ) )
 			.SetColorAndOpacity( FLinearColor( 0.4f, 0.4, 0.4f, 1.0f ) ) );
 
+		const FMargin MenuBlockPadding(10.0f, 3.0f, 5.0f, 2.0f);
+		Style->Set("Menu.Heading.Padding", MenuBlockPadding + FMargin(0, 10, 0, 0));
+
 		/* Set images for various SCheckBox states associated with menu check box items... */
 		const FCheckBoxStyle BasicMenuCheckBoxStyle = FCheckBoxStyle()
 			.SetUncheckedImage( IMAGE_BRUSH( "Common/SmallCheckBox", Icon14x14 ) )
@@ -1010,13 +1034,16 @@ TSharedRef<ISlateStyle> FCoreStyle::Create( const FName& InStyleSetName )
 		/* ... and add new style */
 		Style->Set( "Menu.ToggleButton", MenuToggleButtonCheckBoxStyle );
 
-		Style->Set( "Menu.Button", FButtonStyle( NoBorder )
-			.SetNormal ( FSlateNoResource() )
-			.SetPressed( BOX_BRUSH( "Common/RoundedSelection_16x", 4.0f/16.0f, SelectionColor_Pressed ) )
-			.SetHovered( BOX_BRUSH( "Common/RoundedSelection_16x", 4.0f/16.0f, SelectionColor ) )
-			.SetNormalPadding( FMargin(0,1) )
-			.SetPressedPadding( FMargin(0,2,0,0) )
-		);
+
+		FButtonStyle MenuButton =
+			FButtonStyle(NoBorder)
+			.SetNormal(FSlateNoResource())
+			.SetPressed(BOX_BRUSH("Common/RoundedSelection_16x", 4.0f / 16.0f, SelectionColor_Pressed))
+			.SetHovered(BOX_BRUSH("Common/RoundedSelection_16x", 4.0f / 16.0f, SelectionColor))
+			.SetNormalPadding(FMargin(0, 1))
+			.SetPressedPadding(FMargin(0, 2, 0, 0));
+
+		Style->Set("Menu.Button", MenuButton);
 
 		Style->Set( "Menu.Button.Checked", new BOX_BRUSH( "Common/RoundedSelection_16x",  4.0f/16.0f, SelectionColor_Pressed ) );
 		Style->Set( "Menu.Button.Checked_Hovered", new BOX_BRUSH( "Common/RoundedSelection_16x",  4.0f/16.0f, SelectionColor_Pressed ) );
@@ -1024,6 +1051,14 @@ TSharedRef<ISlateStyle> FCoreStyle::Create( const FName& InStyleSetName )
 
 		/* The style of a menu bar button when it has a sub menu open */
 		Style->Set( "Menu.Button.SubMenuOpen", new BORDER_BRUSH( "Common/Selection", FMargin(4.f/16.f), FLinearColor(0.10f, 0.10f, 0.10f) ) );
+
+		Style->Set("WindowMenuBar.Background", new BOX_BRUSH("Old/Menu_Background", FMargin(8.0f / 64.0f)));
+		Style->Set("WindowMenuBar.Label", FTextBlockStyle(NormalText).SetFont(DEFAULT_FONT("Regular", 9)));
+		Style->Set("WindowMenuBar.Expand", new IMAGE_BRUSH("Icons/toolbar_expand_16x", Icon16x16));
+		Style->Set("WindowMenuBar.Button", MenuButton);
+		Style->Set("WindowMenuBar.Button.SubMenuOpen", new BORDER_BRUSH("Common/Selection", FMargin(4.f / 16.f), FLinearColor(0.10f, 0.10f, 0.10f)));
+
+		Style->Set("WindowMenuBar.MenuBar.Padding", FMargin(12, 2));
 	}
 
 	// SExpandableButton defaults...
@@ -1086,12 +1121,11 @@ TSharedRef<ISlateStyle> FCoreStyle::Create( const FName& InStyleSetName )
 		Style->Set( "Docking.Background", new BOX_BRUSH( "Old/Menu_Background", FMargin(8.0f/64.0f) ) );
 		Style->Set( "Docking.Border", new BOX_BRUSH( "Common/GroupBorder", FMargin(4.0f/16.0f) ) );
 
-		Style->Set( "Docking.TabFont", FTextBlockStyle(NormalText)
+		FTextBlockStyle DockingTabFont = FTextBlockStyle(NormalText)
 			.SetFont( DEFAULT_FONT( "Regular", 9 ) )
 			.SetColorAndOpacity( FLinearColor(0.72f, 0.72f, 0.72f, 1.f) )
 			.SetShadowOffset( FVector2D( 1,1 ) )
-			.SetShadowColorAndOpacity( FLinearColor::Black )
-			);
+			.SetShadowColorAndOpacity( FLinearColor::Black ) ;
 
 		Style->Set( "Docking.UnhideTabwellButton", FButtonStyle(Button)
 			.SetNormal ( IMAGE_BRUSH( "/Docking/ShowTabwellButton_Normal", FVector2D(10,10) ) )
@@ -1120,7 +1154,6 @@ TSharedRef<ISlateStyle> FCoreStyle::Create( const FName& InStyleSetName )
 			.SetNormalBrush( BOX_BRUSH( "/Docking/Tab_Shape", 2.f /8.0f, DockColor_Inactive ) )
 			.SetHoveredBrush( BOX_BRUSH( "/Docking/Tab_Shape", 2.f /8.0f, DockColor_Hovered ) )
 			.SetForegroundBrush( BOX_BRUSH( "/Docking/Tab_Shape", 2.f /8.0f, DockColor_Active ) )
-			.SetActiveBrush( BOX_BRUSH( "/Docking/Tab_Active", 4./8.0f) ) // 
 			.SetColorOverlayTabBrush( BOX_BRUSH( "/Docking/Tab_ColorOverlay", 4 / 16.0f))
 			.SetColorOverlayIconBrush( BOX_BRUSH( "/Docking/Tab_ColorOverlayIcon", 4/16.0f ) )
 			.SetContentAreaBrush( FSlateColorBrush( DockColor_Active ) )
@@ -1128,13 +1161,13 @@ TSharedRef<ISlateStyle> FCoreStyle::Create( const FName& InStyleSetName )
 			.SetTabPadding( FMargin(5, 2, 5, 2) )
 			.SetOverlapWidth( -1.0f )
 			.SetFlashColor( TabFlashColor )
+			.SetTabTextStyle(DockingTabFont)
 			);
 
 		// App Tab
 		Style->Set( "Docking.MajorTab", FDockTabStyle()
 			.SetCloseButtonStyle( CloseButton )
 			.SetNormalBrush( BOX_BRUSH( "/Docking/AppTab_Inactive", FMargin(24.0f/64.0f, 4/32.0f) ) )
-			.SetActiveBrush( BOX_BRUSH( "/Docking/AppTab_Active", FMargin(24.0f/64.0f, 4/32.0f) ) )
 			.SetColorOverlayTabBrush(BOX_BRUSH("/Docking/AppTab_ColorOverlay", FMargin(24.0f / 64.0f, 4 / 32.0f)))
 			.SetColorOverlayIconBrush( BOX_BRUSH( "/Docking/AppTab_ColorOverlayIcon", FMargin(24.0f/64.0f, 4/32.0f) ) )
 			.SetForegroundBrush( BOX_BRUSH( "/Docking/AppTab_Foreground", FMargin(24.0f/64.0f, 4/32.0f) ) )
@@ -1144,6 +1177,7 @@ TSharedRef<ISlateStyle> FCoreStyle::Create( const FName& InStyleSetName )
 			.SetTabPadding( FMargin(17, 4, 15, 4) )
 			.SetOverlapWidth( 21.0f )
 			.SetFlashColor( TabFlashColor )
+			.SetTabTextStyle(DockingTabFont)
 			);
 
 		// Dock Cross
@@ -1325,6 +1359,10 @@ const TSharedPtr<FSlateDynamicImageBrush> FCoreStyle::GetDynamicImageBrush( FNam
 	return Instance->GetDynamicImageBrush(BrushTemplate, TextureResource, TextureName);
 }
 
+bool FCoreStyle::IsStarshipStyle()
+{
+	return !FParse::Param(FCommandLine::Get(), TEXT("oldeditorstyle"));
+}
 
 /* FSlateThrottleManager implementation
  *****************************************************************************/
diff --git a/Engine/Source/Runtime/SlateCore/Private/Styling/SlateTypes.cpp b/Engine/Source/Runtime/SlateCore/Private/Styling/SlateTypes.cpp
--- a/Engine/Source/Runtime/SlateCore/Private/Styling/SlateTypes.cpp
+++ b/Engine/Source/Runtime/SlateCore/Private/Styling/SlateTypes.cpp
@@ -611,7 +611,6 @@ void FDockTabStyle::GetResources( TArray< const FSlateBrush* >& OutBrushes ) con
 {
 	CloseButtonStyle.GetResources(OutBrushes);
 	OutBrushes.Add( &NormalBrush );
-	OutBrushes.Add( &ActiveBrush );
 	OutBrushes.Add( &ColorOverlayTabBrush );
 	OutBrushes.Add( &ColorOverlayIconBrush );
 	OutBrushes.Add( &ForegroundBrush );
diff --git a/Engine/Source/Runtime/SlateCore/Private/Styling/StarshipCoreStyle.cpp b/Engine/Source/Runtime/SlateCore/Private/Styling/StarshipCoreStyle.cpp
--- a/Engine/Source/Runtime/SlateCore/Private/Styling/StarshipCoreStyle.cpp
+++ b/Engine/Source/Runtime/SlateCore/Private/Styling/StarshipCoreStyle.cpp
@@ -1453,7 +1453,6 @@ TSharedRef<ISlateStyle> FStarshipCoreStyle::Create()
 		// Flash using the selection color for consistency with the rest of the UI scheme
 		const FSlateColor& TabFlashColor = SelectionColor;
 
-		Style->Set("Quixel.PerspectiveBox", new IMAGE_BRUSH("Starship/Icons/Quixel/3DShapes/box-perspective_32x", Icon16x16));
 		const FButtonStyle CloseButton = FButtonStyle()
 			.SetNormal( IMAGE_BRUSH("Starship/Icons/Quixel/Common/cross_32x", Icon16x16, FStyleColors::Foreground))
 			.SetPressed(IMAGE_BRUSH("Starship/Icons/Quixel/Common/cross_32x", Icon16x16, FStyleColors::Foreground))
@@ -1470,7 +1469,6 @@ TSharedRef<ISlateStyle> FStarshipCoreStyle::Create()
 			.SetNormalBrush(	FSlateNoResource())
 			.SetHoveredBrush(   BOX_BRUSH("/Starship/Docking/DockTab_Hover", 	  4.0f / 20.0f, FStyleColors::Background))
 			.SetForegroundBrush(BOX_BRUSH("/Starship/Docking/DockTab_Foreground", 4.0f / 20.0f, FStyleColors::Background))
-			.SetActiveBrush(    BOX_BRUSH("/Starship/Docking/DockTab_Active",     4.0f / 20.0f, FStyleColors::Primary))
 
 			.SetColorOverlayTabBrush(FSlateNoResource())
 			.SetColorOverlayIconBrush(FSlateNoResource())
@@ -1485,7 +1483,7 @@ TSharedRef<ISlateStyle> FStarshipCoreStyle::Create()
 			.SetActiveForegroundColor(FStyleColors::ForegroundHover)
 			.SetForegroundForegroundColor(FStyleColors::ForegroundHover)
 			.SetHoveredForegroundColor(FStyleColors::ForegroundHover)
-			.SetTabTextStyle(SmallText)
+			.SetTabTextStyle(NormalText)
 		);
 
 		// App Tab
@@ -1494,7 +1492,6 @@ TSharedRef<ISlateStyle> FStarshipCoreStyle::Create()
 			.SetNormalBrush(	FSlateNoResource())
 			.SetHoveredBrush(   BOX_BRUSH("/Starship/Docking/DockTab_Hover", 	  4.0f / 20.0f, FStyleColors::Background))
 			.SetForegroundBrush(BOX_BRUSH("/Starship/Docking/DockTab_Foreground", 4.0f / 20.0f, FStyleColors::Background))
-			.SetActiveBrush(    BOX_BRUSH("/Starship/Docking/DockTab_Active",     4.0f / 20.0f, FStyleColors::Primary))
 
 			.SetColorOverlayTabBrush(FSlateNoResource())
 			.SetColorOverlayIconBrush(FSlateNoResource())
@@ -1516,6 +1513,9 @@ TSharedRef<ISlateStyle> FStarshipCoreStyle::Create()
 
 		Style->Set("Docking.Tab.InactiveTabSeparator", new FSlateColorBrush(FStyleColors::Hover));
 
+		Style->Set("Docking.Tab.ActiveTabIndicatorColor", FStyleColors::Primary);
+
+
 		// Dock Cross
 		Style->Set("Docking.Cross.DockLeft", new IMAGE_BRUSH("/Docking/OuterDockingIndicator", FVector2D(6, 6), FLinearColor(1.0f, 0.35f, 0.0f, 0.25f)));
 		Style->Set("Docking.Cross.DockLeft_Hovered", new IMAGE_BRUSH("/Docking/OuterDockingIndicator", FVector2D(6, 6), FLinearColor(1.0f, 0.35f, 0.0f)));
diff --git a/Engine/Source/Runtime/SlateCore/Public/Styling/SlateTypes.h b/Engine/Source/Runtime/SlateCore/Public/Styling/SlateTypes.h
--- a/Engine/Source/Runtime/SlateCore/Public/Styling/SlateTypes.h
+++ b/Engine/Source/Runtime/SlateCore/Public/Styling/SlateTypes.h
@@ -1633,11 +1633,6 @@ struct SLATECORE_API FDockTabStyle : public FSlateWidgetStyle
 	FSlateBrush NormalBrush;
 	FDockTabStyle& SetNormalBrush( const FSlateBrush& InNormalBrush ){ NormalBrush = InNormalBrush; return *this; }
 
-	/** Brush used when this tab is in its active state */
-	UPROPERTY(EditAnywhere, Category=Appearance)
-	FSlateBrush ActiveBrush;
-	FDockTabStyle& SetActiveBrush( const FSlateBrush& InActiveBrush ){ ActiveBrush = InActiveBrush; return *this; }
-
 	/** Brush used to overlay a given color onto this tab */
 	UPROPERTY(EditAnywhere, Category=Appearance)
 	FSlateBrush ColorOverlayTabBrush;
@@ -1712,7 +1707,6 @@ struct SLATECORE_API FDockTabStyle : public FSlateWidgetStyle
 	UPROPERTY(EditAnywhere, AdvancedDisplay, Category=Appearance)
 	FSlateColor ForegroundForegroundColor;
 	FDockTabStyle& SetForegroundForegroundColor( const FSlateColor& InForegroundForegroundColor ){ ForegroundForegroundColor = InForegroundForegroundColor; return *this; }
-
 };
 
 
