Fix foreground colors of dock tabs

diff --git a/Engine/Source/Runtime/SlateCore/Private/Styling/StarshipCoreStyle.cpp b/Engine/Source/Runtime/SlateCore/Private/Styling/StarshipCoreStyle.cpp
--- a/Engine/Source/Runtime/SlateCore/Private/Styling/StarshipCoreStyle.cpp
+++ b/Engine/Source/Runtime/SlateCore/Private/Styling/StarshipCoreStyle.cpp
@@ -1481,7 +1481,7 @@ TSharedRef<ISlateStyle> FStarshipCoreStyle::Create()
 
 			.SetNormalForegroundColor(FStyleColors::Foreground)
 			.SetActiveForegroundColor(FStyleColors::ForegroundHover)
-			.SetForegroundForegroundColor(FStyleColors::ForegroundHover)
+			.SetForegroundForegroundColor(FStyleColors::Foreground)
 			.SetHoveredForegroundColor(FStyleColors::ForegroundHover)
 			.SetTabTextStyle(NormalText)
 		);
@@ -1504,7 +1504,7 @@ TSharedRef<ISlateStyle> FStarshipCoreStyle::Create()
 
 			.SetNormalForegroundColor(FStyleColors::Foreground)
 			.SetActiveForegroundColor(FStyleColors::ForegroundHover)
-			.SetForegroundForegroundColor(FStyleColors::ForegroundHover)
+			.SetForegroundForegroundColor(FStyleColors::Foreground)
 			.SetHoveredForegroundColor(FStyleColors::ForegroundHover)
 			.SetTabTextStyle(NormalText)
 		);
