Style Adding ForegroundHeader Color

diff --git a/Engine/Source/Runtime/SlateCore/Private/Styling/StarshipCoreStyle.cpp b/Engine/Source/Runtime/SlateCore/Private/Styling/StarshipCoreStyle.cpp
--- a/Engine/Source/Runtime/SlateCore/Private/Styling/StarshipCoreStyle.cpp
+++ b/Engine/Source/Runtime/SlateCore/Private/Styling/StarshipCoreStyle.cpp
@@ -47,6 +47,7 @@ const FLinearColor FStyleColors::PrimaryPress		= COLOR("#1989BCFF");
 const FLinearColor FStyleColors::Foreground			= COLOR("#A6A6A6FF");
 const FLinearColor FStyleColors::ForegroundHover	= COLOR("#FFFFFFFF");
 const FLinearColor FStyleColors::ForegroundInverted	= FStyleColors::Input;
+const FLinearColor FStyleColors::ForegroundHeader   = COLOR("#C8C8C8FF");
 
 const FLinearColor FStyleColors::Select             = FStyleColors::Primary;
 const FLinearColor FStyleColors::SelectInactive     = COLOR("#99B3BFFF");
@@ -242,44 +243,86 @@ TSharedRef<ISlateStyle> FStarshipCoreStyle::Create()
 	const FLinearColor WindowHighlight(FColor::FromHex("#4C4C4CFF"));
 
 
-	Style->Set("Colors.Black", FStyleColors::Black);
-	Style->Set("Colors.Title", FStyleColors::Title);
-	Style->Set("Colors.Foldout", FStyleColors::Foldout);
-	Style->Set("Colors.Input", FStyleColors::Input);
-	Style->Set("Colors.Background", FStyleColors::Background);
-	Style->Set("Colors.Header", FStyleColors::Header);
-	Style->Set("Colors.Dropdown", FStyleColors::Dropdown);
-	Style->Set("Colors.Hover", FStyleColors::Hover);
-	Style->Set("Colors.Hover2", FStyleColors::Hover2);
-	Style->Set("Colors.White", FStyleColors::White);
-	Style->Set("Colors.White25", FStyleColors::White25);
-	Style->Set("Colors.Highlight", FStyleColors::Highlight);
-
-	Style->Set("Colors.Foreground", FStyleColors::Foreground);
-	Style->Set("Colors.ForegroundHover", FStyleColors::ForegroundHover);
+	Style->Set("Colors.Black",              FStyleColors::Black);
+	Style->Set("Colors.Title",              FStyleColors::Title);
+	Style->Set("Colors.Foldout",            FStyleColors::Foldout);
+	Style->Set("Colors.Input",              FStyleColors::Input);
+	Style->Set("Colors.Background",         FStyleColors::Background);
+	Style->Set("Colors.Header",             FStyleColors::Header);
+	Style->Set("Colors.Dropdown",           FStyleColors::Dropdown);
+	Style->Set("Colors.Hover",              FStyleColors::Hover);
+	Style->Set("Colors.Hover2",             FStyleColors::Hover2);
+	Style->Set("Colors.White",              FStyleColors::White);
+	Style->Set("Colors.White25",            FStyleColors::White25);
+	Style->Set("Colors.Highlight",          FStyleColors::Highlight);
+
+	Style->Set("Colors.Foreground",         FStyleColors::Foreground);
+	Style->Set("Colors.ForegroundHover",    FStyleColors::ForegroundHover);
 	Style->Set("Colors.ForegroundInverted", FStyleColors::ForegroundInverted);
+	Style->Set("Colors.ForegroundHeader", 	FStyleColors::ForegroundHeader);
 
-	Style->Set("Colors.Select", FStyleColors::Select);
-	Style->Set("Colors.SelectInactive", FStyleColors::SelectInactive);
-	Style->Set("Colors.SelectParent", FStyleColors::SelectParent);
-	Style->Set("Colors.SelectHover", FStyleColors::SelectHover);
+	Style->Set("Colors.Select",             FStyleColors::Select);
+	Style->Set("Colors.SelectInactive",     FStyleColors::SelectInactive);
+	Style->Set("Colors.SelectParent",       FStyleColors::SelectParent);
+	Style->Set("Colors.SelectHover",        FStyleColors::SelectHover);
 	
 
-	Style->Set("Colors.Primary", FStyleColors::Primary);
-	Style->Set("Colors.PrimaryHover", FStyleColors::PrimaryHover);
-	Style->Set("Colors.PrimaryPress", FStyleColors::PrimaryPress);
-	Style->Set("Colors.AccentBlue", FStyleColors::AccentBlue);
-	Style->Set("Colors.AccentPurple", FStyleColors::AccentPurple);
-	Style->Set("Colors.AccentPink", FStyleColors::AccentPink);
-	Style->Set("Colors.AccentRed", FStyleColors::AccentRed);
-	Style->Set("Colors.AccentOrange", FStyleColors::AccentOrange);
-	Style->Set("Colors.AccentYellow", FStyleColors::AccentYellow);
-	Style->Set("Colors.AccentGreen", FStyleColors::AccentGreen);
-	Style->Set("Colors.AccentBrown", FStyleColors::AccentBrown);
-	Style->Set("Colors.AccentBlack", FStyleColors::AccentBlack);
-	Style->Set("Colors.AccentGray", FStyleColors::AccentGray);
-	Style->Set("Colors.AccentWhite", FStyleColors::AccentWhite);
-	Style->Set("Colors.AccentFolder", FStyleColors::AccentFolder);
+	Style->Set("Colors.Primary",            FStyleColors::Primary);
+	Style->Set("Colors.PrimaryHover",       FStyleColors::PrimaryHover);
+	Style->Set("Colors.PrimaryPress",       FStyleColors::PrimaryPress);
+	Style->Set("Colors.AccentBlue",         FStyleColors::AccentBlue);
+	Style->Set("Colors.AccentPurple",       FStyleColors::AccentPurple);
+	Style->Set("Colors.AccentPink",         FStyleColors::AccentPink);
+	Style->Set("Colors.AccentRed",          FStyleColors::AccentRed);
+	Style->Set("Colors.AccentOrange",       FStyleColors::AccentOrange);
+	Style->Set("Colors.AccentYellow",       FStyleColors::AccentYellow);
+	Style->Set("Colors.AccentGreen",        FStyleColors::AccentGreen);
+	Style->Set("Colors.AccentBrown",        FStyleColors::AccentBrown);
+	Style->Set("Colors.AccentBlack",        FStyleColors::AccentBlack);
+	Style->Set("Colors.AccentGray",         FStyleColors::AccentGray);
+	Style->Set("Colors.AccentWhite",        FStyleColors::AccentWhite);
+	Style->Set("Colors.AccentFolder",       FStyleColors::AccentFolder);
+
+
+	Style->Set("Brushes.Black",              new FSlateColorBrush(FStyleColors::Black));
+	Style->Set("Brushes.Title",              new FSlateColorBrush(FStyleColors::Title));
+	Style->Set("Brushes.Foldout",            new FSlateColorBrush(FStyleColors::Foldout));
+	Style->Set("Brushes.Input",              new FSlateColorBrush(FStyleColors::Input));
+	Style->Set("Brushes.Background",         new FSlateColorBrush(FStyleColors::Background));
+	Style->Set("Brushes.Header",             new FSlateColorBrush(FStyleColors::Header));
+	Style->Set("Brushes.Dropdown",           new FSlateColorBrush(FStyleColors::Dropdown));
+	Style->Set("Brushes.Hover",              new FSlateColorBrush(FStyleColors::Hover));
+	Style->Set("Brushes.Hover2",             new FSlateColorBrush(FStyleColors::Hover2));
+	Style->Set("Brushes.White",              new FSlateColorBrush(FStyleColors::White));
+	Style->Set("Brushes.White25",            new FSlateColorBrush(FStyleColors::White25));
+	Style->Set("Brushes.Highlight",          new FSlateColorBrush(FStyleColors::Highlight));
+
+	Style->Set("Brushes.Foreground",         new FSlateColorBrush(FStyleColors::Foreground));
+	Style->Set("Brushes.ForegroundHover",    new FSlateColorBrush(FStyleColors::ForegroundHover));
+	Style->Set("Brushes.ForegroundInverted", new FSlateColorBrush(FStyleColors::ForegroundInverted));
+	Style->Set("Brushes.ForegroundHeader",   new FSlateColorBrush(FStyleColors::ForegroundHeader));
+
+	Style->Set("Brushes.Select",             new FSlateColorBrush(FStyleColors::Select));
+	Style->Set("Brushes.SelectInactive",     new FSlateColorBrush(FStyleColors::SelectInactive));
+	Style->Set("Brushes.SelectParent",       new FSlateColorBrush(FStyleColors::SelectParent));
+	Style->Set("Brushes.SelectHover",        new FSlateColorBrush(FStyleColors::SelectHover));
+	
+	Style->Set("Brushes.Primary",            new FSlateColorBrush(FStyleColors::Primary));
+	Style->Set("Brushes.PrimaryHover",       new FSlateColorBrush(FStyleColors::PrimaryHover));
+	Style->Set("Brushes.PrimaryPress",       new FSlateColorBrush(FStyleColors::PrimaryPress));
+	Style->Set("Brushes.AccentBlue",         new FSlateColorBrush(FStyleColors::AccentBlue));
+	Style->Set("Brushes.AccentPurple",       new FSlateColorBrush(FStyleColors::AccentPurple));
+	Style->Set("Brushes.AccentPink",         new FSlateColorBrush(FStyleColors::AccentPink));
+	Style->Set("Brushes.AccentRed",          new FSlateColorBrush(FStyleColors::AccentRed));
+	Style->Set("Brushes.AccentOrange",       new FSlateColorBrush(FStyleColors::AccentOrange));
+	Style->Set("Brushes.AccentYellow",       new FSlateColorBrush(FStyleColors::AccentYellow));
+	Style->Set("Brushes.AccentGreen",        new FSlateColorBrush(FStyleColors::AccentGreen));
+	Style->Set("Brushes.AccentBrown",        new FSlateColorBrush(FStyleColors::AccentBrown));
+	Style->Set("Brushes.AccentBlack",        new FSlateColorBrush(FStyleColors::AccentBlack));
+	Style->Set("Brushes.AccentGray",         new FSlateColorBrush(FStyleColors::AccentGray));
+	Style->Set("Brushes.AccentWhite",        new FSlateColorBrush(FStyleColors::AccentWhite));
+	Style->Set("Brushes.AccentFolder",       new FSlateColorBrush(FStyleColors::AccentFolder));
+
 
 	Style->Set("AppIcon", new IMAGE_BRUSH_SVG("Starship/Common/unreal", FVector2D(36,36), FStyleColors::Foreground));
 	Style->Set("AppIcon.Small", new IMAGE_BRUSH_SVG("Starship/Common/unreal-small", Icon24x24, FStyleColors::Foreground));
@@ -308,7 +351,9 @@ TSharedRef<ISlateStyle> FStarshipCoreStyle::Create()
 	Style->Set("Icons.pyramid", new IMAGE_BRUSH_SVG("Starship/Common/pyriamid", Icon16x16));
 	Style->Set("Icons.sphere", new IMAGE_BRUSH_SVG("Starship/Common/sphere", Icon16x16));
 
-	Style->Set("Icons.Settings", new IMAGE_BRUSH_SVG("Starship/SceneOutliner/Settings", Icon16x16));
+	Style->Set("Icons.Settings", new IMAGE_BRUSH_SVG("Starship/Common/settings", Icon16x16));
+	Style->Set("Icons.Blueprints", new IMAGE_BRUSH_SVG("Starship/Common/blueprint", Icon16x16));
+	Style->Set("Icons.Plus", new IMAGE_BRUSH_SVG("Starship/Common/plus", Icon16x16));
 
 
 	// Common Margins
diff --git a/Engine/Source/Runtime/SlateCore/Public/Styling/StarshipCoreStyle.h b/Engine/Source/Runtime/SlateCore/Public/Styling/StarshipCoreStyle.h
--- a/Engine/Source/Runtime/SlateCore/Public/Styling/StarshipCoreStyle.h
+++ b/Engine/Source/Runtime/SlateCore/Public/Styling/StarshipCoreStyle.h
@@ -99,6 +99,7 @@ struct SLATECORE_API FStyleColors
 	static const FLinearColor Foreground;
 	static const FLinearColor ForegroundHover;
 	static const FLinearColor ForegroundInverted;
+	static const FLinearColor ForegroundHeader;
 
 	static const FLinearColor Select;
 	static const FLinearColor SelectInactive;
