Style colors are now FSlateColors

diff --git a/Engine/Source/Editor/EditorStyle/Private/StarshipStyle.cpp b/Engine/Source/Editor/EditorStyle/Private/StarshipStyle.cpp
--- a/Engine/Source/Editor/EditorStyle/Private/StarshipStyle.cpp
+++ b/Engine/Source/Editor/EditorStyle/Private/StarshipStyle.cpp
@@ -2342,7 +2342,7 @@ void FStarshipEditorStyle::FStyle::SetupViewportStyles()
 
 		Set("EditorViewportToolBar.WarningButton", ViewportMenuWarningButton);
 
-		FLinearColor ToolbarBackgroundColor = FStyleColors::Foldout;
+		FLinearColor ToolbarBackgroundColor = FStyleColors::Foldout.GetSpecifiedColor();
 		ToolbarBackgroundColor.A = .75f;
 
 		Set("EditorViewportToolBar.Background", new FSlateColorBrush(ToolbarBackgroundColor));
@@ -5131,7 +5131,7 @@ void FStarshipEditorStyle::FStyle::SetupPersonaStyle()
 
 	// Play in editor / play in world
 	{
-		FLinearColor GreenHSV = FStyleColors::AccentGreen.LinearRGBToHSV();
+		FLinearColor GreenHSV = FStyleColors::AccentGreen.GetSpecifiedColor().LinearRGBToHSV();
 		
 		FLinearColor GreenHover = FLinearColor(GreenHSV.R, GreenHSV.G * .5, GreenHSV.B, GreenHSV.A).HSVToLinearRGB();
 		FLinearColor GreenPress = FLinearColor(GreenHSV.R, GreenHSV.G, GreenHSV.B*.5, GreenHSV.A).HSVToLinearRGB(); 
diff --git a/Engine/Source/Editor/LevelEditor/Private/LevelEditor.cpp b/Engine/Source/Editor/LevelEditor/Private/LevelEditor.cpp
--- a/Engine/Source/Editor/LevelEditor/Private/LevelEditor.cpp
+++ b/Engine/Source/Editor/LevelEditor/Private/LevelEditor.cpp
@@ -93,12 +93,12 @@ public:
 		FString ProjectNameWatermarkPrefix;
 		GConfig->GetString(TEXT("LevelEditor"), TEXT("ProjectNameWatermarkPrefix"), /*out*/ ProjectNameWatermarkPrefix, GEditorPerProjectIni);
 
-		FLinearColor BadgeBackgroundColor = FAppStyle::Get().GetColor("Colors.Foldout");
+		FSlateColor BadgeBackgroundColor = FAppStyle::Get().GetSlateColor("Colors.Foldout");
 
 		FColor ConfigColor;
 		if (GConfig->GetColor(TEXT("LevelEditor"), TEXT("ProjectBadgeBackgroundColor"), /*out*/ ConfigColor, GEditorPerProjectIni))
 		{
-			BadgeBackgroundColor = ConfigColor;
+			BadgeBackgroundColor = FLinearColor(ConfigColor);
 		}
 
 		FColor BadgeTextColor = FColor(128,128,128,255);
diff --git a/Engine/Source/Editor/SceneOutliner/Private/SSceneOutliner.cpp b/Engine/Source/Editor/SceneOutliner/Private/SSceneOutliner.cpp
--- a/Engine/Source/Editor/SceneOutliner/Private/SSceneOutliner.cpp
+++ b/Engine/Source/Editor/SceneOutliner/Private/SSceneOutliner.cpp
@@ -3795,15 +3795,15 @@ namespace SceneOutliner
 	{
 		if ( !IsFilterActive() )
 		{
-			 return FAppStyle::Get().GetColor("Colors.Foreground");
+			 return FAppStyle::Get().GetSlateColor("Colors.Foreground");
 		}
 		else if( FilteredActorCount == 0 )
 		{
-			return FAppStyle::Get().GetColor("Colors.AccentRed");
+			return FAppStyle::Get().GetSlateColor("Colors.AccentRed");
 		}
 		else
 		{
-			return FAppStyle::Get().GetColor("Colors.AccentGreen");
+			return FAppStyle::Get().GetSlateColor("Colors.AccentGreen");
 		}
 	}
 
diff --git a/Engine/Source/Editor/SceneOutliner/Private/SceneOutlinerGutter.cpp b/Engine/Source/Editor/SceneOutliner/Private/SceneOutlinerGutter.cpp
--- a/Engine/Source/Editor/SceneOutliner/Private/SceneOutlinerGutter.cpp
+++ b/Engine/Source/Editor/SceneOutliner/Private/SceneOutlinerGutter.cpp
@@ -316,7 +316,7 @@ private:
 		}
 		else if (IsHovered() && !bIsSelected)
 		{
-			return FAppStyle::Get().GetColor("Colors.ForegroundHover");
+			return FAppStyle::Get().GetSlateColor("Colors.ForegroundHover");
 		}
 
 		return FSlateColor::UseForeground();
diff --git a/Engine/Source/Editor/StatusBar/Private/SStatusBar.cpp b/Engine/Source/Editor/StatusBar/Private/SStatusBar.cpp
--- a/Engine/Source/Editor/StatusBar/Private/SStatusBar.cpp
+++ b/Engine/Source/Editor/StatusBar/Private/SStatusBar.cpp
@@ -85,7 +85,7 @@ class SContentBrowserOverlay : public SCompoundWidget
 
 	virtual int32 OnPaint(const FPaintArgs& Args, const FGeometry& AllottedGeometry, const FSlateRect& MyCullingRect, FSlateWindowElementList& OutDrawElements, int32 LayerId, const FWidgetStyle& InWidgetStyle, bool bParentEnabled) const override
 	{
-		static FLinearColor ShadowColor = FAppStyle::Get().GetColor("Colors.Foldout");
+		static FSlateColor ShadowColor = FAppStyle::Get().GetSlateColor("Colors.Foldout");
 
 		FGeometry RenderTransformedChildGeometry = AllottedGeometry.MakeChild(FSlateRenderTransform(FVector2D(0.0f, MaxContentBrowserHeight - Height)));
 
diff --git a/Engine/Source/Runtime/SlateCore/Private/Styling/StarshipCoreStyle.cpp b/Engine/Source/Runtime/SlateCore/Private/Styling/StarshipCoreStyle.cpp
--- a/Engine/Source/Runtime/SlateCore/Private/Styling/StarshipCoreStyle.cpp
+++ b/Engine/Source/Runtime/SlateCore/Private/Styling/StarshipCoreStyle.cpp
@@ -26,46 +26,46 @@ TSharedPtr< ISlateStyle > FStarshipCoreStyle::Instance = nullptr;
 // HEX Colors from sRGB Color space
 #define COLOR( HexValue ) FLinearColor::FromSRGBColor(FColor::FromHex(HexValue))
 
-const FLinearColor FStyleColors::Black				= COLOR("#000000FF");
-const FLinearColor FStyleColors::Title				= COLOR("#151515FF");
-const FLinearColor FStyleColors::Foldout			= COLOR("0F0F0FFF");
-const FLinearColor FStyleColors::Input				= COLOR("#1A1A1AFF");
-const FLinearColor FStyleColors::Background			= COLOR("#242424FF");
-const FLinearColor FStyleColors::Header				= COLOR("#2F2F2FFF");
-const FLinearColor FStyleColors::Dropdown			= COLOR("#383838FF");
-const FLinearColor FStyleColors::Hover				= COLOR("#575757FF");
-const FLinearColor FStyleColors::Hover2				= COLOR("#808080FF");
-
-const FLinearColor FStyleColors::White				= COLOR("#FFFFFFFF");
-const FLinearColor FStyleColors::White25			= COLOR("#FFFFFF40");
-const FLinearColor FStyleColors::Highlight			= COLOR("#0078D7FF");
-
-const FLinearColor FStyleColors::Primary			= COLOR("#26BBFFFF");
-const FLinearColor FStyleColors::PrimaryHover		= COLOR("#6FD2FFFF");
-const FLinearColor FStyleColors::PrimaryPress		= COLOR("#1989BCFF");
-
-const FLinearColor FStyleColors::Foreground			= COLOR("#A6A6A6FF");
-const FLinearColor FStyleColors::ForegroundHover	= COLOR("#FFFFFFFF");
-const FLinearColor FStyleColors::ForegroundInverted	= FStyleColors::Input;
-const FLinearColor FStyleColors::ForegroundHeader   = COLOR("#C8C8C8FF");
-
-const FLinearColor FStyleColors::Select             = FStyleColors::Primary;
-const FLinearColor FStyleColors::SelectInactive     = COLOR("#99B3BFFF");
-const FLinearColor FStyleColors::SelectParent       = COLOR("#2C323AFF"); 
-const FLinearColor FStyleColors::SelectHover        = FStyleColors::Background;
-
-const FLinearColor FStyleColors::AccentBlue			= COLOR("#26BBFFFF");
-const FLinearColor FStyleColors::AccentPurple		= COLOR("#A139BFFF");
-const FLinearColor FStyleColors::AccentPink			= COLOR("#FF729CFF");
-const FLinearColor FStyleColors::AccentRed			= COLOR("#FF4040FF");
-const FLinearColor FStyleColors::AccentOrange		= COLOR("#FE9B07FF");
-const FLinearColor FStyleColors::AccentYellow		= COLOR("#FFDC1AFF");
-const FLinearColor FStyleColors::AccentGreen		= COLOR("#8BC24AFF");
-const FLinearColor FStyleColors::AccentBrown		= COLOR("#804D39FF");
-const FLinearColor FStyleColors::AccentBlack		= COLOR("#242424FF");
-const FLinearColor FStyleColors::AccentGray			= COLOR("#808080FF");
-const FLinearColor FStyleColors::AccentWhite		= COLOR("#FFFFFFFF");
-const FLinearColor FStyleColors::AccentFolder       = COLOR("#B68F55FF");
+const FSlateColor FStyleColors::Black				= COLOR("#000000FF");
+const FSlateColor FStyleColors::Title				= COLOR("#151515FF");
+const FSlateColor FStyleColors::Foldout				= COLOR("0F0F0FFF");
+const FSlateColor FStyleColors::Input				= COLOR("#1A1A1AFF");
+const FSlateColor FStyleColors::Background			= COLOR("#242424FF");
+const FSlateColor FStyleColors::Header				= COLOR("#2F2F2FFF");
+const FSlateColor FStyleColors::Dropdown			= COLOR("#383838FF");
+const FSlateColor FStyleColors::Hover				= COLOR("#575757FF");
+const FSlateColor FStyleColors::Hover2				= COLOR("#808080FF");
+
+const FSlateColor FStyleColors::White				= COLOR("#FFFFFFFF");
+const FSlateColor FStyleColors::White25				= COLOR("#FFFFFF40");
+const FSlateColor FStyleColors::Highlight			= COLOR("#0078D7FF");
+
+const FSlateColor FStyleColors::Primary				= COLOR("#26BBFFFF");
+const FSlateColor FStyleColors::PrimaryHover		= COLOR("#6FD2FFFF");
+const FSlateColor FStyleColors::PrimaryPress		= COLOR("#1989BCFF");
+
+const FSlateColor FStyleColors::Foreground			= COLOR("#A6A6A6FF");
+const FSlateColor FStyleColors::ForegroundHover		= COLOR("#FFFFFFFF");
+const FSlateColor FStyleColors::ForegroundInverted	= FStyleColors::Input;
+const FSlateColor FStyleColors::ForegroundHeader	= COLOR("#C8C8C8FF");
+
+const FSlateColor FStyleColors::Select				= FStyleColors::Primary;
+const FSlateColor FStyleColors::SelectInactive		= COLOR("#99B3BFFF");
+const FSlateColor FStyleColors::SelectParent		= COLOR("#2C323AFF");
+const FSlateColor FStyleColors::SelectHover			= FStyleColors::Background;
+
+const FSlateColor FStyleColors::AccentBlue			= COLOR("#26BBFFFF");
+const FSlateColor FStyleColors::AccentPurple		= COLOR("#A139BFFF");
+const FSlateColor FStyleColors::AccentPink			= COLOR("#FF729CFF");
+const FSlateColor FStyleColors::AccentRed			= COLOR("#FF4040FF");
+const FSlateColor FStyleColors::AccentOrange		= COLOR("#FE9B07FF");
+const FSlateColor FStyleColors::AccentYellow		= COLOR("#FFDC1AFF");
+const FSlateColor FStyleColors::AccentGreen			= COLOR("#8BC24AFF");
+const FSlateColor FStyleColors::AccentBrown			= COLOR("#804D39FF");
+const FSlateColor FStyleColors::AccentBlack			= COLOR("#242424FF");
+const FSlateColor FStyleColors::AccentGray			= COLOR("#808080FF");
+const FSlateColor FStyleColors::AccentWhite			= COLOR("#FFFFFFFF");
+const FSlateColor FStyleColors::AccentFolder		= COLOR("#B68F55FF");
 
 #define FONT(...) FSlateFontInfo(FLegacySlateFontInfoCache::Get().GetDefaultFont(), __VA_ARGS__)
 
diff --git a/Engine/Source/Runtime/SlateCore/Public/Brushes/SlateImageBrush.h b/Engine/Source/Runtime/SlateCore/Public/Brushes/SlateImageBrush.h
--- a/Engine/Source/Runtime/SlateCore/Public/Brushes/SlateImageBrush.h
+++ b/Engine/Source/Runtime/SlateCore/Public/Brushes/SlateImageBrush.h
@@ -80,3 +80,15 @@ struct SLATECORE_API FSlateImageBrush
 	{ }
 
 };
+
+struct SLATECORE_API FSlateVectorImageBrush
+	: public FSlateImageBrush
+{
+	FSlateVectorImageBrush(const FString& InImageName, const FVector2D& InImageSize, const FSlateColor& InTint, ESlateBrushTileType::Type InTiling = ESlateBrushTileType::NoTile)
+		: FSlateImageBrush(InImageName, InImageSize, InTint, InTiling, ESlateBrushImageType::Vector)
+	{ }
+
+	FSlateVectorImageBrush(const FString& InImageName, const FVector2D& InImageSize, const FLinearColor& InTint = FLinearColor(1, 1, 1, 1), ESlateBrushTileType::Type InTiling = ESlateBrushTileType::NoTile)
+		: FSlateImageBrush(InImageName, InImageSize, InTint, InTiling, ESlateBrushImageType::Vector)
+	{ }
+};
\ No newline at end of file
diff --git a/Engine/Source/Runtime/SlateCore/Public/Styling/StarshipCoreStyle.h b/Engine/Source/Runtime/SlateCore/Public/Styling/StarshipCoreStyle.h
--- a/Engine/Source/Runtime/SlateCore/Public/Styling/StarshipCoreStyle.h
+++ b/Engine/Source/Runtime/SlateCore/Public/Styling/StarshipCoreStyle.h
@@ -78,46 +78,46 @@ private:
 
 struct SLATECORE_API FStyleColors
 {
-	static const FLinearColor Black;
-	static const FLinearColor Title;
-	static const FLinearColor Foldout;
-	static const FLinearColor Input;
-	static const FLinearColor Background;
-	static const FLinearColor Header;
-	static const FLinearColor Dropdown;
-	static const FLinearColor Hover;
-	static const FLinearColor Hover2;
-	static const FLinearColor White;
-	static const FLinearColor White25;
-	static const FLinearColor Highlight; 
-
-	static const FLinearColor Primary;
-	static const FLinearColor PrimaryHover;
-	static const FLinearColor PrimaryPress;
-
-
-	static const FLinearColor Foreground;
-	static const FLinearColor ForegroundHover;
-	static const FLinearColor ForegroundInverted;
-	static const FLinearColor ForegroundHeader;
-
-	static const FLinearColor Select;
-	static const FLinearColor SelectInactive;
-	static const FLinearColor SelectParent;
-	static const FLinearColor SelectHover;
-
-	static const FLinearColor AccentBlue;
-	static const FLinearColor AccentPurple;
-	static const FLinearColor AccentPink;
-	static const FLinearColor AccentRed;
-	static const FLinearColor AccentOrange;
-	static const FLinearColor AccentYellow;
-	static const FLinearColor AccentGreen;
-	static const FLinearColor AccentBrown;
-	static const FLinearColor AccentBlack;
-	static const FLinearColor AccentGray;
-	static const FLinearColor AccentWhite;
-	static const FLinearColor AccentFolder;
+	static const FSlateColor Black;
+	static const FSlateColor Title;
+	static const FSlateColor Foldout;
+	static const FSlateColor Input;
+	static const FSlateColor Background;
+	static const FSlateColor Header;
+	static const FSlateColor Dropdown;
+	static const FSlateColor Hover;
+	static const FSlateColor Hover2;
+	static const FSlateColor White;
+	static const FSlateColor White25;
+	static const FSlateColor Highlight;
+
+	static const FSlateColor Primary;
+	static const FSlateColor PrimaryHover;
+	static const FSlateColor PrimaryPress;
+
+
+	static const FSlateColor Foreground;
+	static const FSlateColor ForegroundHover;
+	static const FSlateColor ForegroundInverted;
+	static const FSlateColor ForegroundHeader;
+
+	static const FSlateColor Select;
+	static const FSlateColor SelectInactive;
+	static const FSlateColor SelectParent;
+	static const FSlateColor SelectHover;
+
+	static const FSlateColor AccentBlue;
+	static const FSlateColor AccentPurple;
+	static const FSlateColor AccentPink;
+	static const FSlateColor AccentRed;
+	static const FSlateColor AccentOrange;
+	static const FSlateColor AccentYellow;
+	static const FSlateColor AccentGreen;
+	static const FSlateColor AccentBrown;
+	static const FSlateColor AccentBlack;
+	static const FSlateColor AccentGray;
+	static const FSlateColor AccentWhite;
+	static const FSlateColor AccentFolder;
 
 
 };
