Fix some standalone applications not starting up due to out of order core/editor style initialization

diff --git a/Engine/Source/Editor/EditorStyle/Private/StarshipStyle.h b/Engine/Source/Editor/EditorStyle/Private/StarshipStyle.h
--- a/Engine/Source/Editor/EditorStyle/Private/StarshipStyle.h
+++ b/Engine/Source/Editor/EditorStyle/Private/StarshipStyle.h
@@ -32,6 +32,9 @@ public:
 	{
 		Settings = NULL;
 
+		// The core style must be initialized before the editor style
+		FSlateApplication::InitializeCoreStyle();
+
 #if WITH_EDITOR
 		Settings = GetMutableDefault<UEditorStyleSettings>();
 		ISettingsModule* SettingsModule = FModuleManager::GetModulePtr<ISettingsModule>("Settings");
diff --git a/Engine/Source/Programs/UnrealFrontend/Private/Commands/UserInterfaceCommand.cpp b/Engine/Source/Programs/UnrealFrontend/Private/Commands/UserInterfaceCommand.cpp
--- a/Engine/Source/Programs/UnrealFrontend/Private/Commands/UserInterfaceCommand.cpp
+++ b/Engine/Source/Programs/UnrealFrontend/Private/Commands/UserInterfaceCommand.cpp
@@ -39,8 +39,6 @@ void FUserInterfaceCommand::Run(  )
 	FConfigCacheIni::InitializeConfigSystem();
 	GetTargetPlatformManager();
 
-	FCoreStyle::ResetToDefault();
-
 	// load required modules
 	FModuleManager::Get().LoadModuleChecked("EditorStyle");
 	FModuleManager::Get().LoadModuleChecked("Messaging");
diff --git a/Engine/Source/Programs/UnrealInsights/Private/UserInterfaceCommand.cpp b/Engine/Source/Programs/UnrealInsights/Private/UserInterfaceCommand.cpp
--- a/Engine/Source/Programs/UnrealInsights/Private/UserInterfaceCommand.cpp
+++ b/Engine/Source/Programs/UnrealInsights/Private/UserInterfaceCommand.cpp
@@ -123,10 +123,12 @@ void FUserInterfaceCommand::Run()
 		}
 	}
 
-	FCoreStyle::ResetToDefault();
+
+	InitializeSlateApplication();
+	//FCoreStyle::ResetToDefault();
 
 	// Load required modules.
-	FModuleManager::Get().LoadModuleChecked("EditorStyle");
+
 	FModuleManager::Get().LoadModuleChecked("TraceInsights");
 
 	// Load plug-ins.
@@ -137,7 +139,6 @@ void FUserInterfaceCommand::Run()
 	// Load optional modules.
 	FModuleManager::Get().LoadModule("SettingsEditor");
 
-	InitializeSlateApplication();
 
 	// Initialize source code access.
 	// Load the source code access module.
@@ -200,6 +201,8 @@ void FUserInterfaceCommand::InitializeSlateApplication()
 	// Crank up a normal Slate application using the platform's standalone renderer.
 	FSlateApplication::InitializeAsStandaloneApplication(GetStandardStandaloneRenderer());
 
+	FModuleManager::Get().LoadModuleChecked("EditorStyle");
+
 	//const FSlateBrush* AppIcon = new FSlateImageBrush(FPaths::EngineContentDir() / "Editor/Slate/Icons/Insights/AppIcon_24x.png", FVector2D(24.0f, 24.0f));
 	//FSlateApplication::Get().SetAppIcon(AppIcon);
 
diff --git a/Engine/Source/Runtime/Slate/Private/Framework/Application/SlateApplication.cpp b/Engine/Source/Runtime/Slate/Private/Framework/Application/SlateApplication.cpp
--- a/Engine/Source/Runtime/Slate/Private/Framework/Application/SlateApplication.cpp
+++ b/Engine/Source/Runtime/Slate/Private/Framework/Application/SlateApplication.cpp
@@ -6560,13 +6560,13 @@ TSharedRef<FSlateApplication> FSlateApplication::InitializeAsStandaloneApplicati
 
 void FSlateApplication::InitializeCoreStyle()
 {
-	if (FCoreStyle::IsStarshipStyle())
+	if (FCoreStyle::IsStarshipStyle() && !FStarshipCoreStyle::IsInitialized())
 	{
 		UStyleColorTable::Get();
 		FStarshipCoreStyle::ResetToDefault();
 		FAppStyle::SetAppStyleSet(FStarshipCoreStyle::GetCoreStyle());
 	}
-	else
+	else if(!FCoreStyle::IsStarshipStyle() && !FCoreStyle::IsInitialized())
 	{
 		FCoreStyle::ResetToDefault();
 		FAppStyle::SetAppStyleSet(FCoreStyle::GetCoreStyle());
diff --git a/Engine/Source/Runtime/SlateCore/Public/Styling/CoreStyle.h b/Engine/Source/Runtime/SlateCore/Public/Styling/CoreStyle.h
--- a/Engine/Source/Runtime/SlateCore/Public/Styling/CoreStyle.h
+++ b/Engine/Source/Runtime/SlateCore/Public/Styling/CoreStyle.h
@@ -4,6 +4,7 @@
 
 #include "CoreMinimal.h"
 #include "Styling/ISlateStyle.h"
+#include "Styling/AppStyle.h"
 
 struct FSlateDynamicImageBrush;
 
@@ -16,8 +17,26 @@ public:
 
 	static TSharedRef<class ISlateStyle> Create( const FName& InStyleSetName = "CoreStyle" );
 
-	/** @return the singleton instance. */
+	/** 
+	* @return the Application Style 
+	*
+	* NOTE: Until the Editor can be fully updated, calling FCoreStyle::Get() will
+	* return the AppStyle instead of the style definied in this class.  
+	*
+	* Using the AppStyle is preferred in most cases as it allows the style to be changed 
+	* and restyled more easily.
+	*
+	* In cases requiring explicit use of the CoreStyle where a Slate Widget should not take on
+	* the appearance of the rest of the application, use FCoreStyle::GetCoreStyle().
+	*
+	*/
 	static const ISlateStyle& Get( )
+	{
+		return FAppStyle::Get();
+	}
+
+	/** @return the singleton instance of the style created in . */
+	static const ISlateStyle& GetCoreStyle()
 	{
 		return *(Instance.Get());
 	}
@@ -45,6 +64,10 @@ public:
 	static const int32 RegularTextSize = 9;
 	static const int32 SmallTextSize = 8;
 
+	static bool IsStarshipStyle();
+
+	static bool IsInitialized() { return Instance.IsValid(); }
+
 private:
 
 	static void SetStyle( const TSharedRef< class ISlateStyle >& NewStyle );
diff --git a/Engine/Source/Runtime/SlateCore/Public/Styling/StarshipCoreStyle.h b/Engine/Source/Runtime/SlateCore/Public/Styling/StarshipCoreStyle.h
--- a/Engine/Source/Runtime/SlateCore/Public/Styling/StarshipCoreStyle.h
+++ b/Engine/Source/Runtime/SlateCore/Public/Styling/StarshipCoreStyle.h
@@ -64,6 +64,7 @@ public:
 	static const int32 RegularTextSize = 10;
 	static const int32 SmallTextSize = 8;
 
+	static bool IsInitialized() { return Instance.IsValid(); }
 private:
 
 	static void SetStyle( const TSharedRef< class ISlateStyle >& NewStyle );
