Fix incremental

diff --git a/Engine/Source/Runtime/SlateCore/Public/Styling/StyleColors.h b/Engine/Source/Runtime/SlateCore/Public/Styling/StyleColors.h
--- a/Engine/Source/Runtime/SlateCore/Public/Styling/StyleColors.h
+++ b/Engine/Source/Runtime/SlateCore/Public/Styling/StyleColors.h
@@ -90,7 +90,7 @@ public:
 
 private:
 	UPROPERTY(EditAnywhere, Config, Category=Colors)
-	FLinearColor Colors[EStyleColor::MAX];
+	FLinearColor Colors[static_cast<int32>(EStyleColor::MAX)];
 };
 
 
