Fixes Inactive Tree Hover Color

diff --git a/Engine/Source/Runtime/SlateCore/Private/Styling/StarshipCoreStyle.cpp b/Engine/Source/Runtime/SlateCore/Private/Styling/StarshipCoreStyle.cpp
--- a/Engine/Source/Runtime/SlateCore/Private/Styling/StarshipCoreStyle.cpp
+++ b/Engine/Source/Runtime/SlateCore/Private/Styling/StarshipCoreStyle.cpp
@@ -1097,7 +1097,7 @@ TSharedRef<ISlateStyle> FStarshipCoreStyle::Create()
 			.SetActiveHoveredBrush(  FSlateColorBrush(FStyleColors::Select))
 
 			.SetInactiveBrush(       FSlateColorBrush(FStyleColors::SelectInactive))
-			.SetInactiveHoveredBrush(FSlateColorBrush(FStyleColors::SelectHover))
+			.SetInactiveHoveredBrush(FSlateColorBrush(FStyleColors::SelectInactive))
 
 			.SetActiveHighlightedBrush(  FSlateColorBrush(FStyleColors::SelectParent)) // This is the parent hightlight
 			.SetInactiveHighlightedBrush(FSlateColorBrush(FStyleColors::SelectParent))// This is the parent highlight
