Fix Slate color unlinking not working and thus causing default UMG widgets to basically be transparent

diff --git a/Engine/Source/Runtime/SlateCore/Public/Styling/SlateColor.h b/Engine/Source/Runtime/SlateCore/Public/Styling/SlateColor.h
--- a/Engine/Source/Runtime/SlateCore/Public/Styling/SlateColor.h
+++ b/Engine/Source/Runtime/SlateCore/Public/Styling/SlateColor.h
@@ -151,6 +151,7 @@ public:
 	{
 		if (ColorUseRule == ESlateColorStylingMode::UseColor_ColorTable)
 		{
+			SpecifiedColor = GetColorFromTable();
 			ColorUseRule = ESlateColorStylingMode::UseColor_Specified;
 		}
 	}
