Updated Widget styles to add outlines to CheckBoxes, ComboBoxes, TextBox widgets

diff --git a/Engine/Source/Runtime/SlateCore/Private/Styling/StarshipCoreStyle.cpp b/Engine/Source/Runtime/SlateCore/Private/Styling/StarshipCoreStyle.cpp
--- a/Engine/Source/Runtime/SlateCore/Private/Styling/StarshipCoreStyle.cpp
+++ b/Engine/Source/Runtime/SlateCore/Private/Styling/StarshipCoreStyle.cpp
@@ -176,6 +176,7 @@ TSharedRef<ISlateStyle> FStarshipCoreStyle::Create()
 	const FVector2D Icon12x16(12.0f, 16.0f);
 	const FVector2D Icon14x14(14.0f, 14.0f);
 	const FVector2D Icon16x16(16.0f, 16.0f);
+	const FVector2D Icon18x18(18.0f, 18.0f);
 	const FVector2D Icon20x20(20.0f, 20.0f);
 	const FVector2D Icon22x22(22.0f, 22.0f);
 	const FVector2D Icon24x24(24.0f, 24.0f);
@@ -187,7 +188,7 @@ TSharedRef<ISlateStyle> FStarshipCoreStyle::Create()
 	const FVector2D Icon128x128(128.0f, 128.0f);
 
 	const float InputFocusRadius = 4.f;
-	const float InputFocusThickness = 2.0f;
+	const float InputFocusThickness = 1.0f;
 
 	// These are the Slate colors which reference the dynamic colors in the style;
 	const FSlateColor DefaultForeground(FStyleColors::Foreground);
@@ -538,7 +539,7 @@ TSharedRef<ISlateStyle> FStarshipCoreStyle::Create()
 	// SComboButton and SComboBox defaults...
 	{
 		const FButtonStyle ComboButtonButton = FButtonStyle()
-		.SetNormal(  FSlateRoundedBoxBrush(FStyleColors::Input, InputFocusRadius))
+		.SetNormal(  FSlateRoundedBoxBrush(FStyleColors::Input, InputFocusRadius, FStyleColors::Dropdown, InputFocusThickness))
 		.SetHovered( FSlateRoundedBoxBrush(FStyleColors::Input, InputFocusRadius, FStyleColors::Hover, InputFocusThickness))
 		.SetPressed( FSlateRoundedBoxBrush(FStyleColors::Input, InputFocusRadius, FStyleColors::Hover, InputFocusThickness))
 		.SetNormalForeground(  FStyleColors::Foreground)
@@ -637,26 +638,47 @@ TSharedRef<ISlateStyle> FStarshipCoreStyle::Create()
 	// SCheckBox defaults...
 	{
 
+		const float CheckboxCornerRadius = 3.f;
+		const float CheckboxOutlineThickness = 1.0f;
+
 		const FCheckBoxStyle BasicCheckBoxStyle = FCheckBoxStyle()
 			.SetCheckBoxType(ESlateCheckBoxType::CheckBox)
+
+			.SetForegroundColor(               FLinearColor::White )
+			.SetHoveredForegroundColor(        FLinearColor::White )
+			.SetPressedForegroundColor(        FLinearColor::White )
+			.SetCheckedForegroundColor(        FLinearColor::White )
+			.SetCheckedHoveredForegroundColor( FLinearColor::White )
+			.SetCheckedPressedForegroundColor( FLinearColor::White )
+			.SetUndeterminedForegroundColor(   FLinearColor::White )
+
 			.SetUncheckedImage(          FSlateNoResource())
 			.SetUncheckedHoveredImage(   FSlateNoResource())
 			.SetUncheckedPressedImage(   FSlateNoResource())
 			.SetCheckedImage(            IMAGE_BRUSH_SVG("Starship/CoreWidgets/CheckBox/check", 		Icon16x16, FStyleColors::Primary))
-			.SetCheckedHoveredImage(     IMAGE_BRUSH_SVG("Starship/CoreWidgets/CheckBox/check", 		Icon16x16, FStyleColors::Primary)) 
+			.SetCheckedHoveredImage(     IMAGE_BRUSH_SVG("Starship/CoreWidgets/CheckBox/check", 		Icon16x16, FStyleColors::PrimaryHover)) 
 			.SetCheckedPressedImage(     IMAGE_BRUSH_SVG("Starship/CoreWidgets/CheckBox/check", 		Icon16x16, FStyleColors::Primary))
 			.SetUndeterminedImage(       IMAGE_BRUSH_SVG("Starship/CoreWidgets/CheckBox/indeterminate", Icon16x16, FStyleColors::Primary))
-			.SetUndeterminedHoveredImage(IMAGE_BRUSH_SVG("Starship/CoreWidgets/CheckBox/indeterminate", Icon16x16, FStyleColors::Primary))
+			.SetUndeterminedHoveredImage(IMAGE_BRUSH_SVG("Starship/CoreWidgets/CheckBox/indeterminate", Icon16x16, FStyleColors::PrimaryHover))
 			.SetUndeterminedPressedImage(IMAGE_BRUSH_SVG("Starship/CoreWidgets/CheckBox/indeterminate", Icon16x16, FStyleColors::Primary))
-			.SetBackgroundImage(         FSlateRoundedBoxBrush(FStyleColors::Input,   InputFocusRadius, FLinearColor::Transparent, InputFocusThickness, Icon20x20))
-			.SetBackgroundHoveredImage(  FSlateRoundedBoxBrush(FStyleColors::Input,   InputFocusRadius, FStyleColors::Hover,       InputFocusThickness, Icon20x20))
-			.SetBackgroundPressedImage(  FSlateRoundedBoxBrush(FStyleColors::Foldout, InputFocusRadius, FStyleColors::Hover,       InputFocusThickness, Icon20x20));
+			.SetBackgroundImage(         FSlateRoundedBoxBrush(FStyleColors::Input,   CheckboxCornerRadius, FStyleColors::Dropdown,    CheckboxOutlineThickness, Icon18x18))
+			.SetBackgroundHoveredImage(  FSlateRoundedBoxBrush(FStyleColors::Input,   CheckboxCornerRadius, FStyleColors::Hover,       CheckboxOutlineThickness, Icon18x18))
+			.SetBackgroundPressedImage(  FSlateRoundedBoxBrush(FStyleColors::Foldout, CheckboxCornerRadius, FStyleColors::Hover,       CheckboxOutlineThickness, Icon18x18));
 
 		Style->Set("Checkbox", BasicCheckBoxStyle);
 
 
 		const FCheckBoxStyle SimplifiedCheckBoxStyle = FCheckBoxStyle()
 			.SetCheckBoxType(ESlateCheckBoxType::CheckBox)
+
+			.SetForegroundColor(               FLinearColor::White )
+			.SetHoveredForegroundColor(        FLinearColor::White )
+			.SetPressedForegroundColor(        FLinearColor::White )
+			.SetCheckedForegroundColor(        FLinearColor::White )
+			.SetCheckedHoveredForegroundColor( FLinearColor::White )
+			.SetCheckedPressedForegroundColor( FLinearColor::White )
+			.SetUndeterminedForegroundColor(   FLinearColor::White )
+
 			.SetUncheckedImage(          FSlateNoResource())
 			.SetUncheckedHoveredImage(   FSlateNoResource())
 			.SetUncheckedPressedImage(   FSlateNoResource())
@@ -666,9 +688,9 @@ TSharedRef<ISlateStyle> FStarshipCoreStyle::Create()
 			.SetUndeterminedImage(       IMAGE_BRUSH_SVG("Starship/CoreWidgets/CheckBox/indeterminate", Icon16x16, FStyleColors::Foreground))
 			.SetUndeterminedHoveredImage(IMAGE_BRUSH_SVG("Starship/CoreWidgets/CheckBox/indeterminate", Icon16x16, FStyleColors::ForegroundHover))
 			.SetUndeterminedPressedImage(IMAGE_BRUSH_SVG("Starship/CoreWidgets/CheckBox/indeterminate", Icon16x16, FStyleColors::Foreground))
-			.SetBackgroundImage(         FSlateRoundedBoxBrush(FStyleColors::Input,   InputFocusRadius, FLinearColor::Transparent, InputFocusThickness, Icon20x20))
-			.SetBackgroundHoveredImage(  FSlateRoundedBoxBrush(FStyleColors::Input,   InputFocusRadius, FStyleColors::Hover,       InputFocusThickness, Icon20x20))
-			.SetBackgroundPressedImage(  FSlateRoundedBoxBrush(FStyleColors::Foldout, InputFocusRadius, FStyleColors::Hover,       InputFocusThickness, Icon20x20));
+			.SetBackgroundImage(         FSlateRoundedBoxBrush(FStyleColors::Input,   CheckboxCornerRadius, FStyleColors::Dropdown,    CheckboxOutlineThickness, Icon18x18))
+			.SetBackgroundHoveredImage(  FSlateRoundedBoxBrush(FStyleColors::Input,   CheckboxCornerRadius, FStyleColors::Hover,       CheckboxOutlineThickness, Icon18x18))
+			.SetBackgroundPressedImage(  FSlateRoundedBoxBrush(FStyleColors::Foldout, CheckboxCornerRadius, FStyleColors::Hover,       CheckboxOutlineThickness, Icon18x18));
 
 		Style->Set("SimplifiedCheckbox", SimplifiedCheckBoxStyle);
 
@@ -799,6 +821,15 @@ TSharedRef<ISlateStyle> FStarshipCoreStyle::Create()
 		/* Set images for various radio button (SCheckBox) states ... */
 		const FCheckBoxStyle BasicRadioButtonStyle = FCheckBoxStyle()
 			.SetCheckBoxType(ESlateCheckBoxType::CheckBox)
+
+			.SetForegroundColor(               FLinearColor::White )
+			.SetHoveredForegroundColor(        FLinearColor::White )
+			.SetPressedForegroundColor(        FLinearColor::White )
+			.SetCheckedForegroundColor(        FLinearColor::White )
+			.SetCheckedHoveredForegroundColor( FLinearColor::White )
+			.SetCheckedPressedForegroundColor( FLinearColor::White )
+			.SetUndeterminedForegroundColor(   FLinearColor::White )
+
 			.SetUncheckedImage(           IMAGE_BRUSH_SVG("/Starship/CoreWidgets/CheckBox/radio-off", Icon16x16, FStyleColors::White25))
 			.SetUncheckedHoveredImage(    IMAGE_BRUSH_SVG("/Starship/CoreWidgets/CheckBox/radio-off", Icon16x16, FStyleColors::ForegroundHover))
 			.SetUncheckedPressedImage(    IMAGE_BRUSH_SVG("/Starship/CoreWidgets/CheckBox/radio-off", Icon16x16, FStyleColors::ForegroundHover))
@@ -832,7 +863,7 @@ TSharedRef<ISlateStyle> FStarshipCoreStyle::Create()
 
 	// SEditableTextBox defaults...
 	const FEditableTextBoxStyle NormalEditableTextBoxStyle = FEditableTextBoxStyle()
-		.SetBackgroundImageNormal(FSlateRoundedBoxBrush(FStyleColors::Input, InputFocusRadius))
+		.SetBackgroundImageNormal(FSlateRoundedBoxBrush(FStyleColors::Input, InputFocusRadius, FStyleColors::Dropdown, InputFocusThickness))
 		.SetBackgroundImageHovered(FSlateRoundedBoxBrush(FStyleColors::Input, InputFocusRadius, FStyleColors::Hover, InputFocusThickness))
 		.SetBackgroundImageFocused(FSlateRoundedBoxBrush(FStyleColors::Input, InputFocusRadius, FStyleColors::Primary, InputFocusThickness))
 		.SetBackgroundImageReadOnly(FSlateRoundedBoxBrush(FStyleColors::Input, InputFocusRadius))
@@ -993,8 +1024,8 @@ TSharedRef<ISlateStyle> FStarshipCoreStyle::Create()
 	// SSpinBox defaults...
 	{
 		Style->Set("SpinBox", FSpinBoxStyle()
-			.SetBackgroundBrush(FSlateRoundedBoxBrush(FStyleColors::Input, InputFocusRadius))
-			.SetHoveredBackgroundBrush(FSlateRoundedBoxBrush(FStyleColors::Input, InputFocusRadius))
+			.SetBackgroundBrush(FSlateRoundedBoxBrush(FStyleColors::Input, InputFocusRadius, FStyleColors::Dropdown, InputFocusThickness))
+			.SetHoveredBackgroundBrush(FSlateRoundedBoxBrush(FStyleColors::Input, InputFocusRadius, FStyleColors::Hover, InputFocusThickness))
 
 			.SetActiveFillBrush(FSlateRoundedBoxBrush(FStyleColors::Hover, InputFocusRadius, FLinearColor::Transparent, InputFocusThickness))
 			.SetInactiveFillBrush(FSlateRoundedBoxBrush(FStyleColors::Dropdown, InputFocusRadius, FLinearColor::Transparent, InputFocusThickness))
@@ -1007,11 +1038,11 @@ TSharedRef<ISlateStyle> FStarshipCoreStyle::Create()
 	// SNumericEntryBox defaults...
 	{
 		Style->Set("NumericEntrySpinBox", FSpinBoxStyle()
-			.SetBackgroundBrush(FSlateNoResource())
-			.SetHoveredBackgroundBrush(FSlateNoResource())
+			.SetBackgroundBrush(FSlateRoundedBoxBrush(FStyleColors::Input, InputFocusRadius, FStyleColors::Dropdown, InputFocusThickness))
+			.SetHoveredBackgroundBrush(FSlateRoundedBoxBrush(FStyleColors::Input, InputFocusRadius, FStyleColors::Hover, InputFocusThickness))
 
-			.SetActiveFillBrush(FSlateRoundedBoxBrush(FStyleColors::Hover, InputFocusRadius, FLinearColor::Transparent, 3.f))
-			.SetInactiveFillBrush(FSlateRoundedBoxBrush(FStyleColors::Dropdown, InputFocusRadius, FLinearColor::Transparent, 2.f))
+			.SetActiveFillBrush(FSlateRoundedBoxBrush(FStyleColors::Hover, InputFocusRadius, FLinearColor::Transparent, InputFocusThickness))
+			.SetInactiveFillBrush(FSlateRoundedBoxBrush(FStyleColors::Dropdown, InputFocusRadius, FLinearColor::Transparent, InputFocusThickness))
 
 			.SetArrowsImage(FSlateNoResource())
 			.SetTextPadding(FMargin(0.f))
@@ -1019,8 +1050,9 @@ TSharedRef<ISlateStyle> FStarshipCoreStyle::Create()
 		);
 
 		Style->Set("NumericEntrySpinBox_Dark", FSpinBoxStyle()
-			.SetBackgroundBrush(FSlateNoResource())
-			.SetHoveredBackgroundBrush(FSlateNoResource())
+			.SetBackgroundBrush(FSlateRoundedBoxBrush(FStyleColors::Input, InputFocusRadius, FStyleColors::Dropdown, InputFocusThickness))
+			.SetHoveredBackgroundBrush(FSlateRoundedBoxBrush(FStyleColors::Input, InputFocusRadius, FStyleColors::Hover, InputFocusThickness))
+
 			.SetActiveFillBrush(FSlateRoundedBoxBrush(FStyleColors::Hover, InputFocusRadius, FLinearColor::Transparent, InputFocusThickness))
 			.SetInactiveFillBrush(FSlateRoundedBoxBrush(FStyleColors::Dropdown, InputFocusRadius, FLinearColor::Transparent, InputFocusThickness))
 
