Allow slate debug quad elements to be tintable

diff --git a/Engine/Source/Runtime/SlateCore/Private/Rendering/DrawElements.cpp b/Engine/Source/Runtime/SlateCore/Private/Rendering/DrawElements.cpp
--- a/Engine/Source/Runtime/SlateCore/Private/Rendering/DrawElements.cpp
+++ b/Engine/Source/Runtime/SlateCore/Private/Rendering/DrawElements.cpp
@@ -186,7 +186,7 @@ void FSlateDrawElement::ApplyPositionOffset(const FVector2D& InOffset)
 	const FSlateLayoutTransform InverseLayoutTransform(Inverse(FSlateLayoutTransform(Scale, Position)));
 }
 
-void FSlateDrawElement::MakeDebugQuad( FSlateWindowElementList& ElementList, uint32 InLayer, const FPaintGeometry& PaintGeometry)
+void FSlateDrawElement::MakeDebugQuad( FSlateWindowElementList& ElementList, uint32 InLayer, const FPaintGeometry& PaintGeometry, FLinearColor Tint)
 {
 	PaintGeometry.CommitTransformsIfUsingLegacyConstructor();
 
@@ -196,8 +196,9 @@ void FSlateDrawElement::MakeDebugQuad( FSlateWindowElementList& ElementList, uin
 	}
 
 	FSlateDrawElement& Element = ElementList.AddUninitialized();
-	ElementList.CreatePayload<FSlateBoxPayload>(Element);
-
+	FSlateBoxPayload& BoxPayload = ElementList.CreatePayload<FSlateBoxPayload>(Element);
+	
+	BoxPayload.SetTint(Tint);
 	Element.Init(ElementList, EElementType::ET_DebugQuad, InLayer, PaintGeometry, ESlateDrawEffect::None);
 }
 
@@ -212,7 +213,7 @@ FSlateDrawElement& FSlateDrawElement::MakeBoxInternal(
 {
 	EElementType ElementType = (InBrush->DrawAs == ESlateBrushDrawType::Border) ? EElementType::ET_Border : (InBrush->DrawAs == ESlateBrushDrawType::RoundedBox) ? EElementType::ET_RoundedBox : EElementType::ET_Box;
 
-	// Cast to Rounded Rect to get the internal paramerters 
+	// Cast to Rounded Rect to get the internal parameters 
 	// New payload type - inherit from BoxPayload 
 
 	FSlateDrawElement& Element = ElementList.AddUninitialized();
diff --git a/Engine/Source/Runtime/SlateCore/Public/Rendering/DrawElements.h b/Engine/Source/Runtime/SlateCore/Public/Rendering/DrawElements.h
--- a/Engine/Source/Runtime/SlateCore/Public/Rendering/DrawElements.h
+++ b/Engine/Source/Runtime/SlateCore/Public/Rendering/DrawElements.h
@@ -42,6 +42,7 @@ enum class EElementType : uint8
 	ET_Custom,
 	ET_CustomVerts,
 	ET_PostProcessPass,
+	ET_RoundedBox,
 	/** Total number of draw commands */
 	ET_Count,
 };
@@ -71,7 +72,7 @@ public:
 	 * @param InLayer				The layer to draw the element on
 	 * @param PaintGeometry         DrawSpace position and dimensions; see FPaintGeometry
 	 */
-	SLATECORE_API static void MakeDebugQuad( FSlateWindowElementList& ElementList, uint32 InLayer, const FPaintGeometry& PaintGeometry);
+	SLATECORE_API static void MakeDebugQuad( FSlateWindowElementList& ElementList, uint32 InLayer, const FPaintGeometry& PaintGeometry, FLinearColor Tint = FLinearColor::White);
 
 	/**
 	 * Creates a box element based on the following diagram.  Allows for this element to be resized while maintain the border of the image
diff --git a/Engine/Source/Runtime/SlateCore/Public/Rendering/ElementBatcher.h b/Engine/Source/Runtime/SlateCore/Public/Rendering/ElementBatcher.h
--- a/Engine/Source/Runtime/SlateCore/Public/Rendering/ElementBatcher.h
+++ b/Engine/Source/Runtime/SlateCore/Public/Rendering/ElementBatcher.h
@@ -283,7 +283,7 @@ private:
 	 * Creates vertices necessary to draw a Quad element 
 	 */
 	template<ESlateVertexRounding Rounding>
-	void AddQuadElement( const FSlateDrawElement& DrawElement, FColor Color = FColor::White);
+	void AddDebugQuadElement( const FSlateDrawElement& DrawElement);
 
 	/** 
 	 * Creates vertices necessary to draw a 3x3 element
