Added missing file

diff --git a/Engine/Source/Runtime/SlateCore/Private/Rendering/ElementBatcher.cpp b/Engine/Source/Runtime/SlateCore/Private/Rendering/ElementBatcher.cpp
--- a/Engine/Source/Runtime/SlateCore/Private/Rendering/ElementBatcher.cpp
+++ b/Engine/Source/Runtime/SlateCore/Private/Rendering/ElementBatcher.cpp
@@ -381,7 +381,7 @@ void FSlateElementBatcher::AddElementsInternal(const FSlateDrawElementArray& Dra
 		{
 			SCOPED_NAMED_EVENT_TEXT("Slate::AddDebugQuadElement", FColor::Magenta);
 			STAT(ElementStat_Other++);
-			DrawElement.IsPixelSnapped() ? AddQuadElement<ESlateVertexRounding::Enabled>(DrawElement) : AddQuadElement<ESlateVertexRounding::Disabled>(DrawElement);
+			DrawElement.IsPixelSnapped() ? AddDebugQuadElement<ESlateVertexRounding::Enabled>(DrawElement) : AddDebugQuadElement<ESlateVertexRounding::Disabled>(DrawElement);
 		}
 			break;
 		case EElementType::ET_Spline:
@@ -484,15 +484,19 @@ void FSlateElementBatcher::AddCachedElements(FSlateCachedElementData& CachedElem
 }
 
 template<ESlateVertexRounding Rounding>
-void FSlateElementBatcher::AddQuadElement( const FSlateDrawElement& DrawElement, FColor Color )
+void FSlateElementBatcher::AddDebugQuadElement(const FSlateDrawElement& DrawElement)
 {
 	const FSlateRenderTransform& RenderTransform = DrawElement.GetRenderTransform();
 	const FVector2D& LocalSize = DrawElement.GetLocalSize();
 	ESlateDrawEffect InDrawEffects = DrawElement.GetDrawEffects();
 	const int32 Layer = DrawElement.GetLayer();
 
-	FSlateRenderBatch& RenderBatch = CreateRenderBatch(Layer, FShaderParams(), nullptr, ESlateDrawPrimitive::TriangleList, ESlateShader::Default, ESlateDrawEffect::None, ESlateBatchDrawFlag::Wireframe | ESlateBatchDrawFlag::NoBlending, DrawElement);
+	const FSlateBoxPayload& DrawElementPayload = DrawElement.GetDataPayload<FSlateBoxPayload>();
+
+	FSlateRenderBatch& RenderBatch = CreateRenderBatch(Layer, FShaderParams(), nullptr, ESlateDrawPrimitive::TriangleList, ESlateShader::Default, ESlateDrawEffect::None, ESlateBatchDrawFlag::None, DrawElement);
 	
+	const FColor Color = PackVertexColor(DrawElementPayload.GetTint());
+
 	// Determine the four corners of the quad
 	FVector2D TopLeft = FVector2D::ZeroVector;
 	FVector2D TopRight = FVector2D(LocalSize.X, 0);
