Edit Menu Icons Updated

diff --git a/Engine/Plugins/Editor/PluginBrowser/Source/PluginBrowser/Private/PluginStyle.cpp b/Engine/Plugins/Editor/PluginBrowser/Source/PluginBrowser/Private/PluginStyle.cpp
--- a/Engine/Plugins/Editor/PluginBrowser/Source/PluginBrowser/Private/PluginStyle.cpp
+++ b/Engine/Plugins/Editor/PluginBrowser/Source/PluginBrowser/Private/PluginStyle.cpp
@@ -8,6 +8,7 @@
 #include "Interfaces/IPluginManager.h"
 
 #define IMAGE_BRUSH( RelativePath, ... ) FSlateImageBrush( FPluginStyle::InContent( RelativePath, ".png" ), __VA_ARGS__ )
+#define IMAGE_BRUSH_SVG( RelativePath, ... ) FSlateVectorImageBrush( FPluginStyle::InContent( RelativePath, ".svg" ), __VA_ARGS__ )
 #define BOX_BRUSH( RelativePath, ... ) FSlateBoxBrush( FPluginStyle::InContent( RelativePath, ".png" ), __VA_ARGS__ )
 #define BORDER_BRUSH( RelativePath, ... ) FSlateBorderBrush( FPluginStyle::InContent( RelativePath, ".png" ), __VA_ARGS__ )
 #define DEFAULT_FONT(...) FCoreStyle::GetDefaultFontStyle(__VA_ARGS__)
@@ -55,7 +56,7 @@ void FPluginStyle::Initialize()
 	{
 		const FTextBlockStyle NormalText = FEditorStyle::GetWidgetStyle<FTextBlockStyle>("NormalText");
 
-		StyleSet->Set( "Plugins.TabIcon", new IMAGE_BRUSH( "icon_tab_Plugins_16x", Icon16x16 ) );
+		StyleSet->Set( "Plugins.TabIcon", new IMAGE_BRUSH_SVG( "Plugins", Icon16x16 ) );
 		StyleSet->Set( "Plugins.BreadcrumbArrow", new IMAGE_BRUSH( "SmallArrowRight", Icon10x10 ) );
 		StyleSet->Set( "Plugins.Warning", new IMAGE_BRUSH( "alert", Icon20x20) );
 
diff --git a/Engine/Source/Editor/EditorStyle/Private/StarshipStyle.cpp b/Engine/Source/Editor/EditorStyle/Private/StarshipStyle.cpp
--- a/Engine/Source/Editor/EditorStyle/Private/StarshipStyle.cpp
+++ b/Engine/Source/Editor/EditorStyle/Private/StarshipStyle.cpp
@@ -1088,7 +1088,7 @@ void FStarshipEditorStyle::FStyle::SetupGeneralStyles()
 
 	// Undo History Window
 	{
-		Set("UndoHistory.TabIcon", new IMAGE_BRUSH( "Icons/icon_tab_UndoHistory_16px", Icon16x16 ) );
+		Set("UndoHistory.TabIcon", new IMAGE_BRUSH_SVG( "Starship/Menus/Edit/UndoHistory", Icon16x16 ) );
 	}
 
 	// InputBinding editor
@@ -1336,12 +1336,12 @@ void FStarshipEditorStyle::FStyle::SetupGeneralStyles()
 
 	// Editor preferences
 	{
-		Set("EditorPreferences.TabIcon", new IMAGE_BRUSH("Icons/Edit/icon_Edit_EditorPreferences_16x", Icon16x16));
+		Set("EditorPreferences.TabIcon", new IMAGE_BRUSH_SVG("Starship/Menus/Edit/Preferences", Icon16x16));
 	}
 
 	// Project settings
 	{
-		Set("ProjectSettings.TabIcon", new IMAGE_BRUSH("Icons/Edit/icon_Edit_ProjectSettings_16x", Icon16x16));
+		Set("ProjectSettings.TabIcon", new IMAGE_BRUSH_SVG("Starship/Menus/Edit/ProjectSettings", Icon16x16));
 	}
 
 	// Main frame
diff --git a/Engine/Source/Runtime/SlateCore/Private/Styling/StarshipCoreStyle.cpp b/Engine/Source/Runtime/SlateCore/Private/Styling/StarshipCoreStyle.cpp
--- a/Engine/Source/Runtime/SlateCore/Private/Styling/StarshipCoreStyle.cpp
+++ b/Engine/Source/Runtime/SlateCore/Private/Styling/StarshipCoreStyle.cpp
@@ -463,14 +463,15 @@ TSharedRef<ISlateStyle> FStarshipCoreStyle::Create()
 
 	// Generic command icons
 	{
-		Style->Set("GenericCommands.Undo", new IMAGE_BRUSH("Icons/icon_undo_16px", Icon16x16));
-		Style->Set("GenericCommands.Redo", new IMAGE_BRUSH("Icons/icon_redo_16px", Icon16x16));
-
-		Style->Set("GenericCommands.Copy", new IMAGE_BRUSH("Icons/Edit/icon_Edit_Copy_16x", Icon16x16));
-		Style->Set("GenericCommands.Cut", new IMAGE_BRUSH("Icons/Edit/icon_Edit_Cut_16x", Icon16x16));
-		Style->Set("GenericCommands.Delete", new IMAGE_BRUSH("Icons/Edit/icon_Edit_Delete_16x", Icon16x16));
-		Style->Set("GenericCommands.Paste", new IMAGE_BRUSH("Icons/Edit/icon_Edit_Paste_16x", Icon16x16));
-		Style->Set("GenericCommands.Duplicate", new IMAGE_BRUSH("Icons/Edit/icon_Edit_Duplicate_16x", Icon16x16));
+		Style->Set("GenericCommands.Undo", new IMAGE_BRUSH_SVG("Starship/Common/Undo", Icon16x16));
+		Style->Set("GenericCommands.Redo", new IMAGE_BRUSH_SVG("Starship/Common/Redo", Icon16x16));
+
+		Style->Set("GenericCommands.Copy", new IMAGE_BRUSH_SVG("Starship/Common/Copy", Icon16x16));
+		Style->Set("GenericCommands.Cut", new IMAGE_BRUSH_SVG("Starship/Common/Cut", Icon16x16));
+		Style->Set("GenericCommands.Delete", new IMAGE_BRUSH_SVG("Starship/Common/Delete", Icon16x16));
+		Style->Set("GenericCommands.Paste", new IMAGE_BRUSH_SVG("Starship/Common/Paste", Icon16x16));
+		Style->Set("GenericCommands.Duplicate", new IMAGE_BRUSH_SVG("Starship/Common/Duplicate", Icon16x16));
+		
 		Style->Set("GenericCommands.Rename", new IMAGE_BRUSH("Icons/Edit/icon_Edit_Rename_16x", Icon16x16));
 	}
 
