Fix static analysis issues with stack size in slate core styles

diff --git a/Engine/Source/Runtime/SlateCore/Private/Styling/CoreStyle.cpp b/Engine/Source/Runtime/SlateCore/Private/Styling/CoreStyle.cpp
--- a/Engine/Source/Runtime/SlateCore/Private/Styling/CoreStyle.cpp
+++ b/Engine/Source/Runtime/SlateCore/Private/Styling/CoreStyle.cpp
@@ -134,6 +134,9 @@ PRAGMA_DISABLE_OPTIMIZATION
 #define BORDER_BRUSH( RelativePath, ... ) FSlateBorderBrush(Style->RootToContentDir(RelativePath, TEXT(".png")), __VA_ARGS__)
 #define DEFAULT_FONT(...) FCoreStyle::GetDefaultFontStyle(__VA_ARGS__)
 
+// We no longer use this style. it is kept around for reference only
+#define ALLOW_LEGACY_CORE_STYLE 0
+
 TSharedRef<ISlateStyle> FCoreStyle::Create( const FName& InStyleSetName )
 {
 	TSharedRef<FSlateCoreStyle> Style = MakeShareable(new FSlateCoreStyle(InStyleSetName));
@@ -149,6 +152,8 @@ TSharedRef<ISlateStyle> FCoreStyle::Create( const FName& InStyleSetName )
 		return Style;
 	}
 
+#if ALLOW_LEGACY_CORE_STYLE
+
 	// Note, these sizes are in Slate Units.
 	// Slate Units do NOT have to map to pixels.
 	const FVector2D Icon5x16(5.0f, 16.0f);
@@ -1331,8 +1336,9 @@ TSharedRef<ISlateStyle> FCoreStyle::Create( const FName& InStyleSetName )
 		Style->Set("SyntaxHighlight.NodeAttribueAssignment", FTextBlockStyle(SmallMonospacedText).SetColorAndOpacity(FLinearColor(FColor(0xffb2b400)))); // yellow
 		Style->Set("SyntaxHighlight.NodeAttributeValue", FTextBlockStyle(SmallMonospacedText).SetColorAndOpacity(FLinearColor(FColor(0xffb46100)))); // orange
 	}
-
+#endif 
 	return Style;
+
 }
 PRAGMA_ENABLE_OPTIMIZATION
 
diff --git a/Engine/Source/Runtime/SlateCore/Private/Styling/StarshipCoreStyle.cpp b/Engine/Source/Runtime/SlateCore/Private/Styling/StarshipCoreStyle.cpp
--- a/Engine/Source/Runtime/SlateCore/Private/Styling/StarshipCoreStyle.cpp
+++ b/Engine/Source/Runtime/SlateCore/Private/Styling/StarshipCoreStyle.cpp
@@ -23,6 +23,42 @@
 
 TSharedPtr< ISlateStyle > FStarshipCoreStyle::Instance = nullptr;
 
+namespace CoreStyleConstants
+{
+	// Note, these sizes are in Slate Units.
+	// Slate Units do NOT have to map to pixels.
+	const FVector2D Icon5x16(5.0f, 16.0f);
+	const FVector2D Icon8x4(8.0f, 4.0f);
+	const FVector2D Icon16x4(16.0f, 4.0f);
+	const FVector2D Icon8x8(8.0f, 8.0f);
+	const FVector2D Icon10x10(10.0f, 10.0f);
+	const FVector2D Icon12x12(12.0f, 12.0f);
+	const FVector2D Icon12x16(12.0f, 16.0f);
+	const FVector2D Icon14x14(14.0f, 14.0f);
+	const FVector2D Icon16x16(16.0f, 16.0f);
+	const FVector2D Icon18x18(18.0f, 18.0f);
+	const FVector2D Icon20x20(20.0f, 20.0f);
+	const FVector2D Icon22x22(22.0f, 22.0f);
+	const FVector2D Icon24x24(24.0f, 24.0f);
+	const FVector2D Icon25x25(25.0f, 25.0f);
+	const FVector2D Icon32x32(32.0f, 32.0f);
+	const FVector2D Icon40x40(40.0f, 40.0f);
+	const FVector2D Icon64x64(64.0f, 64.0f);
+	const FVector2D Icon36x24(36.0f, 24.0f);
+	const FVector2D Icon128x128(128.0f, 128.0f);
+
+	// Common Margins
+	const FMargin DefaultMargins(8.f, 4.f);
+	// Buttons already have a built in (4., 2.) padding - adding to that a little
+	const FMargin ButtonMargins(8.f, 2.f);
+
+	const float InputFocusRadius = 4.f;
+	const float InputFocusThickness = 1.0f;
+
+	const FLinearColor WindowHighlight(FColor::FromHex("#4C4C4CFF"));
+}
+
+using namespace CoreStyleConstants;
 
 #define FONT(...) FSlateFontInfo(FLegacySlateFontInfoCache::Get().GetDefaultFont(), __VA_ARGS__)
 
@@ -144,12 +180,8 @@ void FStarshipCoreStyle::SetFocusBrush(FSlateBrush* NewBrush)
 	FSlateStyleRegistry::RegisterSlateStyle(Style.Get());
 }
 
-
-PRAGMA_DISABLE_OPTIMIZATION
-
 #define DEFAULT_FONT(...) FStarshipCoreStyle::GetDefaultFontStyle(__VA_ARGS__)
 
-
 TSharedRef<ISlateStyle> FStarshipCoreStyle::Create()
 {
 	TSharedRef<FStyle> Style = MakeShareable(new FStyle("CoreStyle"));
@@ -165,30 +197,6 @@ TSharedRef<ISlateStyle> FStarshipCoreStyle::Create()
 		return Style;
 	}
 
-	// Note, these sizes are in Slate Units.
-	// Slate Units do NOT have to map to pixels.
-	const FVector2D Icon5x16(5.0f, 16.0f);
-	const FVector2D Icon8x4(8.0f, 4.0f);
-	const FVector2D Icon16x4(16.0f, 4.0f);
-	const FVector2D Icon8x8(8.0f, 8.0f);
-	const FVector2D Icon10x10(10.0f, 10.0f);
-	const FVector2D Icon12x12(12.0f, 12.0f);
-	const FVector2D Icon12x16(12.0f, 16.0f);
-	const FVector2D Icon14x14(14.0f, 14.0f);
-	const FVector2D Icon16x16(16.0f, 16.0f);
-	const FVector2D Icon18x18(18.0f, 18.0f);
-	const FVector2D Icon20x20(20.0f, 20.0f);
-	const FVector2D Icon22x22(22.0f, 22.0f);
-	const FVector2D Icon24x24(24.0f, 24.0f);
-	const FVector2D Icon25x25(25.0f, 25.0f);
-	const FVector2D Icon32x32(32.0f, 32.0f);
-	const FVector2D Icon40x40(40.0f, 40.0f);
-	const FVector2D Icon64x64(64.0f, 64.0f);
-	const FVector2D Icon36x24(36.0f, 24.0f);
-	const FVector2D Icon128x128(128.0f, 128.0f);
-
-	const float InputFocusRadius = 4.f;
-	const float InputFocusThickness = 1.0f;
 
 	// These are the Slate colors which reference the dynamic colors in the style;
 	const FSlateColor DefaultForeground(FStyleColors::Foreground);
@@ -198,163 +206,32 @@ TSharedRef<ISlateStyle> FStarshipCoreStyle::Create()
 	const FSlateColor SelectionColor_Inactive(Style->SelectionColor_Inactive_LinearRef);
 	const FSlateColor SelectionColor_Pressed(Style->SelectionColor_Pressed_LinearRef);
 	const FSlateColor HighlightColor(FStyleColors::Highlight);
-	const FLinearColor WindowHighlight(FColor::FromHex("#4C4C4CFF"));
-
-
-	Style->Set("Colors.Black",              FStyleColors::Black);
-	Style->Set("Colors.Title",              FStyleColors::Title);
-	Style->Set("Colors.Foldout",            FStyleColors::Foldout);
-	Style->Set("Colors.Input",              FStyleColors::Input);
-	Style->Set("Colors.Background",         FStyleColors::Background);
-	Style->Set("Colors.Header",             FStyleColors::Header);
-	Style->Set("Colors.Dropdown",           FStyleColors::Dropdown);
-	Style->Set("Colors.Hover",              FStyleColors::Hover);
-	Style->Set("Colors.Hover2",             FStyleColors::Hover2);
-	Style->Set("Colors.White",              FStyleColors::White);
-	Style->Set("Colors.White25",            FStyleColors::White25);
-	Style->Set("Colors.Highlight",          FStyleColors::Highlight);
-
-	Style->Set("Colors.Foreground",         FStyleColors::Foreground);
-	Style->Set("Colors.ForegroundHover",    FStyleColors::ForegroundHover);
-	Style->Set("Colors.ForegroundInverted", FStyleColors::ForegroundInverted);
-	Style->Set("Colors.ForegroundHeader", 	FStyleColors::ForegroundHeader);
-
-	Style->Set("Colors.Select",             FStyleColors::Select);
-	Style->Set("Colors.SelectInactive",     FStyleColors::SelectInactive);
-	Style->Set("Colors.SelectParent",       FStyleColors::SelectParent);
-	Style->Set("Colors.SelectHover",        FStyleColors::SelectHover);
-	
-
-	Style->Set("Colors.Primary",            FStyleColors::Primary);
-	Style->Set("Colors.PrimaryHover",       FStyleColors::PrimaryHover);
-	Style->Set("Colors.PrimaryPress",       FStyleColors::PrimaryPress);
-	Style->Set("Colors.AccentBlue",         FStyleColors::AccentBlue);
-	Style->Set("Colors.AccentPurple",       FStyleColors::AccentPurple);
-	Style->Set("Colors.AccentPink",         FStyleColors::AccentPink);
-	Style->Set("Colors.AccentRed",          FStyleColors::AccentRed);
-	Style->Set("Colors.AccentOrange",       FStyleColors::AccentOrange);
-	Style->Set("Colors.AccentYellow",       FStyleColors::AccentYellow);
-	Style->Set("Colors.AccentGreen",        FStyleColors::AccentGreen);
-	Style->Set("Colors.AccentBrown",        FStyleColors::AccentBrown);
-	Style->Set("Colors.AccentBlack",        FStyleColors::AccentBlack);
-	Style->Set("Colors.AccentGray",         FStyleColors::AccentGray);
-	Style->Set("Colors.AccentWhite",        FStyleColors::AccentWhite);
-	Style->Set("Colors.AccentFolder",       FStyleColors::AccentFolder);
-
-
-	Style->Set("Brushes.Black",              new FSlateColorBrush(FStyleColors::Black));
-	Style->Set("Brushes.Title",              new FSlateColorBrush(FStyleColors::Title));
-	Style->Set("Brushes.Foldout",            new FSlateColorBrush(FStyleColors::Foldout));
-	Style->Set("Brushes.Input",              new FSlateColorBrush(FStyleColors::Input));
-	Style->Set("Brushes.Background",         new FSlateColorBrush(FStyleColors::Background));
-	Style->Set("Brushes.Header",             new FSlateColorBrush(FStyleColors::Header));
-	Style->Set("Brushes.Dropdown",           new FSlateColorBrush(FStyleColors::Dropdown));
-	Style->Set("Brushes.Hover",              new FSlateColorBrush(FStyleColors::Hover));
-	Style->Set("Brushes.Hover2",             new FSlateColorBrush(FStyleColors::Hover2));
-	Style->Set("Brushes.White",              new FSlateColorBrush(FStyleColors::White));
-	Style->Set("Brushes.White25",            new FSlateColorBrush(FStyleColors::White25));
-	Style->Set("Brushes.Highlight",          new FSlateColorBrush(FStyleColors::Highlight));
-
-	Style->Set("Brushes.Foreground",         new FSlateColorBrush(FStyleColors::Foreground));
-	Style->Set("Brushes.ForegroundHover",    new FSlateColorBrush(FStyleColors::ForegroundHover));
-	Style->Set("Brushes.ForegroundInverted", new FSlateColorBrush(FStyleColors::ForegroundInverted));
-	Style->Set("Brushes.ForegroundHeader",   new FSlateColorBrush(FStyleColors::ForegroundHeader));
 
-	Style->Set("Brushes.Select",             new FSlateColorBrush(FStyleColors::Select));
-	Style->Set("Brushes.SelectInactive",     new FSlateColorBrush(FStyleColors::SelectInactive));
-	Style->Set("Brushes.SelectParent",       new FSlateColorBrush(FStyleColors::SelectParent));
-	Style->Set("Brushes.SelectHover",        new FSlateColorBrush(FStyleColors::SelectHover));
 	
-	Style->Set("Brushes.Primary",            new FSlateColorBrush(FStyleColors::Primary));
-	Style->Set("Brushes.PrimaryHover",       new FSlateColorBrush(FStyleColors::PrimaryHover));
-	Style->Set("Brushes.PrimaryPress",       new FSlateColorBrush(FStyleColors::PrimaryPress));
-	Style->Set("Brushes.AccentBlue",         new FSlateColorBrush(FStyleColors::AccentBlue));
-	Style->Set("Brushes.AccentPurple",       new FSlateColorBrush(FStyleColors::AccentPurple));
-	Style->Set("Brushes.AccentPink",         new FSlateColorBrush(FStyleColors::AccentPink));
-	Style->Set("Brushes.AccentRed",          new FSlateColorBrush(FStyleColors::AccentRed));
-	Style->Set("Brushes.AccentOrange",       new FSlateColorBrush(FStyleColors::AccentOrange));
-	Style->Set("Brushes.AccentYellow",       new FSlateColorBrush(FStyleColors::AccentYellow));
-	Style->Set("Brushes.AccentGreen",        new FSlateColorBrush(FStyleColors::AccentGreen));
-	Style->Set("Brushes.AccentBrown",        new FSlateColorBrush(FStyleColors::AccentBrown));
-	Style->Set("Brushes.AccentBlack",        new FSlateColorBrush(FStyleColors::AccentBlack));
-	Style->Set("Brushes.AccentGray",         new FSlateColorBrush(FStyleColors::AccentGray));
-	Style->Set("Brushes.AccentWhite",        new FSlateColorBrush(FStyleColors::AccentWhite));
-	Style->Set("Brushes.AccentFolder",       new FSlateColorBrush(FStyleColors::AccentFolder));
-
-
-	Style->Set("AppIcon", new IMAGE_BRUSH_SVG("Starship/Common/unreal", FVector2D(36,36), FStyleColors::Foreground));
-	Style->Set("AppIcon.Small", new IMAGE_BRUSH_SVG("Starship/Common/unreal-small", Icon24x24, FStyleColors::Foreground));
-
-	Style->Set("AppIconPadding", FMargin(11,11,3,5));
-	Style->Set("AppIconPadding.Small", FMargin(4,4,0,0));
-
 	const FStyleFonts& StyleFonts = FStyleFonts::Get();
 
-	const FTextBlockStyle NormalText = FTextBlockStyle()
-		.SetFont(StyleFonts.Normal)
-		.SetColorAndOpacity(FSlateColor::UseForeground())
-		.SetShadowOffset(FVector2D::ZeroVector)
-		.SetSelectedBackgroundColor(FStyleColors::Highlight)
-		.SetHighlightShape(FSlateColorBrush(FStyleColors::White));
-
-	Style->Set("NormalFont",     StyleFonts.Normal);
-	Style->Set("SmallFont",      StyleFonts.Small);
-	Style->Set("NormalFontBold", StyleFonts.NormalBold);
-	Style->Set("SmallFontBold",  StyleFonts.SmallBold);
-
-
-	// Common Brushes
-	Style->Set("Icons.box-perspective", new IMAGE_BRUSH_SVG("Starship/Common/box-perspective", Icon16x16));
-	Style->Set("Icons.cylinder", new IMAGE_BRUSH_SVG("Starship/Common/cylinder", Icon16x16));
-	Style->Set("Icons.pyramid", new IMAGE_BRUSH_SVG("Starship/Common/pyriamid", Icon16x16));
-	Style->Set("Icons.sphere", new IMAGE_BRUSH_SVG("Starship/Common/sphere", Icon16x16));
-
-	Style->Set("Icons.Settings", new IMAGE_BRUSH_SVG("Starship/Common/settings", Icon16x16));
-	Style->Set("Icons.Blueprints", new IMAGE_BRUSH_SVG("Starship/Common/blueprint", Icon16x16));
-	Style->Set("Icons.Plus", new IMAGE_BRUSH_SVG("Starship/Common/plus", Icon16x16));
-
+	Style->Set("InvertedForeground", InvertedForeground);
 
-	// Common Margins
-	FMargin DefaultMargins(8.f, 4.f);	
-
-	// Buttons already have a built in (4., 2.) padding - adding to that a little
-	FMargin ButtonMargins(8.f, 2.f); 
-
-	FSlateBrush* DefaultTextUnderlineBrush = new IMAGE_BRUSH("Old/White", Icon8x8, FLinearColor::White, ESlateBrushTileType::Both);
-	Style->Set("DefaultTextUnderline", DefaultTextUnderlineBrush);
-	const FTextBlockStyle NormalUnderlinedText = FTextBlockStyle(NormalText)
-		.SetUnderlineBrush(*DefaultTextUnderlineBrush);
-
-	// Monospaced Text
-	const FTextBlockStyle MonospacedText = FTextBlockStyle()
-		.SetFont(DEFAULT_FONT("Mono", 10))
-		.SetColorAndOpacity(FSlateColor::UseForeground())
-		.SetShadowOffset(FVector2D::ZeroVector)
-		.SetShadowColorAndOpacity(FLinearColor::Black)
-		.SetHighlightColor(FLinearColor(0.02f, 0.3f, 0.0f))
-		.SetHighlightShape(BOX_BRUSH("Common/TextBlockHighlightShape", FMargin(3.f / 8.f))
-		);
+	SetupColors(Style);
 
-	const FTextBlockStyle MonospacedUnderlinedText = FTextBlockStyle(MonospacedText)
-		.SetUnderlineBrush(*DefaultTextUnderlineBrush);
+	// SScrollBar defaults...
+	const FScrollBarStyle ScrollBar = FScrollBarStyle()
+		.SetNormalThumbImage(FSlateRoundedBoxBrush(FStyleColors::Hover, 4.0f))
+		.SetDraggedThumbImage(FSlateRoundedBoxBrush(FStyleColors::Hover2, 4.0f))
+		.SetHoveredThumbImage(FSlateRoundedBoxBrush(FStyleColors::Hover2, 4.0f))
+		.SetThickness(8.0f);
+	;
+	{
+		Style->Set("Scrollbar", ScrollBar);
+	}
 
-	Style->Set("MonospacedText", MonospacedText);
-	Style->Set("MonospacedUnderlinedText", MonospacedUnderlinedText);
 
-	// Small Text
-	const FTextBlockStyle SmallText = FTextBlockStyle(NormalText)
-		.SetFont(StyleFonts.Small);
+	SetupTextStyles(Style);
 
-	const FTextBlockStyle SmallUnderlinedText = FTextBlockStyle(SmallText)
-		.SetUnderlineBrush(*DefaultTextUnderlineBrush);
+	// Get these from the text style we just created in order to share them
+	const FEditableTextBoxStyle& NormalEditableTextBoxStyle = Style->GetWidgetStyle<FEditableTextBoxStyle>("NormalEditableTextBox");
+	const FTextBlockStyle& NormalText = Style->GetWidgetStyle<FTextBlockStyle>("NormalText");
 
-	// Embossed Text
-	Style->Set("EmbossedText", FTextBlockStyle(NormalText)
-		.SetFont(DEFAULT_FONT("Regular", 24))
-		.SetColorAndOpacity(FLinearColor::Black)
-		.SetShadowOffset(FVector2D(0.0f, 1.0f))
-		.SetShadowColorAndOpacity(FLinearColor(0.8f, 0.8f, 0.8f, 0.5))
-	);
 
 	// Common brushes
 	FSlateBrush* GenericWhiteBox = new IMAGE_BRUSH("Old/White", Icon16x16);
@@ -431,12 +308,15 @@ TSharedRef<ISlateStyle> FStarshipCoreStyle::Create()
 
 	// Common icons
 	{
+
+		Style->Set("AppIcon", new IMAGE_BRUSH_SVG("Starship/Common/unreal", FVector2D(36, 36), FStyleColors::Foreground));
+		Style->Set("AppIcon.Small", new IMAGE_BRUSH_SVG("Starship/Common/unreal-small", Icon24x24, FStyleColors::Foreground));
+
+		Style->Set("AppIconPadding", FMargin(11, 11, 3, 5));
+		Style->Set("AppIconPadding.Small", FMargin(4, 4, 0, 0));
+
 		Style->Set("TrashCan", new IMAGE_BRUSH("Icons/TrashCan", FVector2D(64, 64)));
 		Style->Set("TrashCan_Small", new IMAGE_BRUSH("Icons/TrashCan_Small", FVector2D(18, 18)));
-	}
-
-	// Common icons
-	{
 		Style->Set("Icons.Cross", new IMAGE_BRUSH("Icons/Cross_12x", Icon12x12));
 		Style->Set("Icons.Denied", new IMAGE_BRUSH("Icons/denied_16x", Icon16x16));
 		Style->Set("Icons.Error", new IMAGE_BRUSH("Icons/icon_error_16x", Icon16x16));
@@ -444,6 +324,15 @@ TSharedRef<ISlateStyle> FStarshipCoreStyle::Create()
 		Style->Set("Icons.Info", new IMAGE_BRUSH("Icons/icon_info_16x", Icon16x16));
 		Style->Set("Icons.Warning", new IMAGE_BRUSH("Icons/icon_warning_16x", Icon16x16));
 		Style->Set("Icons.Download", new IMAGE_BRUSH("Icons/icon_Downloads_16x", Icon16x16));
+
+		Style->Set("Icons.box-perspective", new IMAGE_BRUSH_SVG("Starship/Common/box-perspective", Icon16x16));
+		Style->Set("Icons.cylinder", new IMAGE_BRUSH_SVG("Starship/Common/cylinder", Icon16x16));
+		Style->Set("Icons.pyramid", new IMAGE_BRUSH_SVG("Starship/Common/pyriamid", Icon16x16));
+		Style->Set("Icons.sphere", new IMAGE_BRUSH_SVG("Starship/Common/sphere", Icon16x16));
+
+		Style->Set("Icons.Settings", new IMAGE_BRUSH_SVG("Starship/Common/settings", Icon16x16));
+		Style->Set("Icons.Blueprints", new IMAGE_BRUSH_SVG("Starship/Common/blueprint", Icon16x16));
+		Style->Set("Icons.Plus", new IMAGE_BRUSH_SVG("Starship/Common/plus", Icon16x16));
 	}
 
 	// Tool panels
@@ -478,372 +367,563 @@ TSharedRef<ISlateStyle> FStarshipCoreStyle::Create()
 	// SVerticalBox Drag& Drop icon
 	Style->Set("VerticalBoxDragIndicator", new IMAGE_BRUSH("Common/VerticalBoxDragIndicator", FVector2D(6, 45)));
 	Style->Set("VerticalBoxDragIndicatorShort", new IMAGE_BRUSH("Common/VerticalBoxDragIndicatorShort", FVector2D(6, 15)));
-	// SScrollBar defaults...
-	const FScrollBarStyle ScrollBar = FScrollBarStyle()
-		.SetNormalThumbImage( FSlateRoundedBoxBrush(FStyleColors::Hover, 4.0f))
-		.SetDraggedThumbImage(FSlateRoundedBoxBrush(FStyleColors::Hover2, 4.0f))
-		.SetHoveredThumbImage(FSlateRoundedBoxBrush(FStyleColors::Hover2, 4.0f))
-		.SetThickness(8.0f);
-	;
-	{
-		Style->Set("Scrollbar", ScrollBar);
-	}
-
-
-	// SButton defaults
-	const FButtonStyle PrimaryButton = FButtonStyle()
-		.SetNormal(  FSlateRoundedBoxBrush(FStyleColors::Primary, 4.0f))
-		.SetHovered( FSlateRoundedBoxBrush(FStyleColors::PrimaryHover, 4.0f))
-		.SetPressed( FSlateRoundedBoxBrush(FStyleColors::PrimaryPress, 4.0f))
-		.SetNormalForeground(  FStyleColors::Background)
-		.SetHoveredForeground( FStyleColors::Background)
-		.SetPressedForeground( FStyleColors::Background)
-		.SetDisabledForeground(FStyleColors::Background)
-		.SetNormalPadding( ButtonMargins)
-		.SetPressedPadding(ButtonMargins);
-
 
-	const FButtonStyle Button = FButtonStyle()
-		.SetNormal(  FSlateRoundedBoxBrush(FStyleColors::Dropdown, 4.0f))
-		.SetHovered( FSlateRoundedBoxBrush(FStyleColors::Hover,    4.0f))
-		.SetPressed( FSlateRoundedBoxBrush(FStyleColors::Header,   4.0f))
-		.SetNormalForeground(  FStyleColors::ForegroundHover)
-		.SetHoveredForeground( FStyleColors::ForegroundHover)
-		.SetPressedForeground( FStyleColors::ForegroundHover)
-		.SetDisabledForeground(FStyleColors::Foreground)
-		.SetNormalPadding( ButtonMargins)
-		.SetPressedPadding(ButtonMargins);
+	SetupButtonStyles(Style);
+	SetupComboButtonStyles(Style);
+	SetupCheckboxStyles(Style);
 
+	const FButtonStyle& Button = Style->GetWidgetStyle<FButtonStyle>("Button");
+	const FButtonStyle& SimpleButton = Style->GetWidgetStyle<FButtonStyle>("SimpleButton");
 
-	const FButtonStyle SimpleButton = FButtonStyle()
-		.SetNormal(  FSlateNoResource())
-		.SetHovered( FSlateRoundedBoxBrush(FStyleColors::Hover, 4.0f))
-		.SetPressed( FSlateRoundedBoxBrush(FStyleColors::Input, 4.0f))
-		.SetDisabled(FSlateNoResource())
-		.SetNormalForeground(  FStyleColors::Foreground)
-		.SetHoveredForeground( FStyleColors::ForegroundHover)
-		.SetPressedForeground( FStyleColors::ForegroundHover)
-		.SetDisabledForeground(FStyleColors::Foreground)
-		.SetNormalPadding( FMargin(4.0f, 2.0f))
-		.SetPressedPadding(FMargin(4.0f, 2.0f));
+	// SMessageLogListing
 	{
-
-		Style->Set("ButtonText", FTextBlockStyle(NormalText).SetFont(StyleFonts.SmallBold));
-
-		Style->Set("PrimaryButton", PrimaryButton);
-		Style->Set("Button", Button);
-		Style->Set("SimpleButton", SimpleButton);
-
-		Style->Set("InvertedForeground", InvertedForeground);
+		FComboButtonStyle MessageLogListingComboButton = FComboButtonStyle()
+			.SetButtonStyle(NoBorder)
+			.SetDownArrowImage(IMAGE_BRUSH("Common/ComboArrow", Icon8x8))
+			.SetMenuBorderBrush(FSlateNoResource())
+			.SetMenuBorderPadding(FMargin(0.0f));
+		Style->Set("MessageLogListingComboButton", MessageLogListingComboButton);
 	}
 
-	// SComboButton and SComboBox defaults...
+	// SSuggestionTextBox defaults...
 	{
-		const FButtonStyle ComboButtonButton = FButtonStyle()
-		.SetNormal(  FSlateRoundedBoxBrush(FStyleColors::Input, InputFocusRadius, FStyleColors::Dropdown, InputFocusThickness))
-		.SetHovered( FSlateRoundedBoxBrush(FStyleColors::Input, InputFocusRadius, FStyleColors::Hover, InputFocusThickness))
-		.SetPressed( FSlateRoundedBoxBrush(FStyleColors::Input, InputFocusRadius, FStyleColors::Hover, InputFocusThickness))
-		.SetNormalForeground(  FStyleColors::Foreground)
-		.SetHoveredForeground( FStyleColors::ForegroundHover)
-		.SetPressedForeground( FStyleColors::ForegroundHover)
-		.SetDisabledForeground(FStyleColors::White25)
-		.SetNormalPadding( FMargin(8.f, 2.f, 4.f, 2.f))
-		.SetPressedPadding(FMargin(8.f, 2.f, 4.f, 2.f));
-
-
-		// SComboBox 
-		FComboButtonStyle ComboButton = FComboButtonStyle()
-			.SetButtonStyle(ComboButtonButton)
-			.SetContentPadding(0.f)
-			.SetDownArrowImage(IMAGE_BRUSH_SVG("Starship/CoreWidgets/ComboBox/wide-chevron-down", FVector2D(20.f, 16.f)))
-			.SetMenuBorderBrush(FSlateRoundedBoxBrush(FStyleColors::Dropdown, 0.0, WindowHighlight, 1.0))
-			.SetMenuBorderPadding(0.f);
-		Style->Set("ComboButton", ComboButton);
-
-
-		FComboBoxStyle ComboBox = FComboBoxStyle()
-			.SetContentPadding(0.f)
-			.SetMenuRowPadding(FMargin(8.0f, 3.f))
-			.SetComboButtonStyle(ComboButton);
-		Style->Set("ComboBox", ComboBox);
-
-
-
-		// Simple Combo Box  (borderless)
-		FComboButtonStyle SimpleComboButton = FComboButtonStyle()
-			.SetButtonStyle(SimpleButton)
-			.SetContentPadding(0.f)
-			.SetDownArrowImage(IMAGE_BRUSH_SVG("Starship/Common/chevron-down", Icon16x16))
-			.SetMenuBorderBrush(FSlateColorBrush(FStyleColors::Dropdown))
-			.SetMenuBorderPadding(0.0f);
-		Style->Set("SimpleComboButton", SimpleComboButton);
-	
-
-		FComboBoxStyle SimpleComboBox = FComboBoxStyle()
-			.SetContentPadding(0.f)
-			.SetMenuRowPadding(FMargin(8.0f, 2.f))
-			.SetComboButtonStyle(SimpleComboButton);
-		Style->Set("SimpleComboBox", SimpleComboBox);
-
-
-		const FTableRowStyle ComboBoxRow = FTableRowStyle()
-			.SetEvenRowBackgroundBrush( FSlateNoResource())
-			.SetEvenRowBackgroundHoveredBrush( FSlateColorBrush(FStyleColors::Hover) )
-
-			.SetOddRowBackgroundBrush( FSlateNoResource())
-			.SetOddRowBackgroundHoveredBrush( FSlateColorBrush(FStyleColors::Hover) )
-
-			.SetSelectorFocusedBrush(    FSlateNoResource())
-
-			.SetActiveBrush(             FSlateColorBrush(FStyleColors::Primary))
-			.SetActiveHoveredBrush(      FSlateColorBrush(FStyleColors::PrimaryHover))
-			.SetActiveHighlightedBrush(  FSlateColorBrush(FStyleColors::PrimaryHover))
-
-			.SetInactiveBrush(           FSlateNoResource())
-			.SetInactiveHoveredBrush(    FSlateNoResource())
-			.SetInactiveHighlightedBrush(FSlateNoResource())
+		Style->Set("SuggestionTextBox.Background", new BOX_BRUSH("Old/Menu_Background", FMargin(8.0f / 64.0f)));
+		Style->Set("SuggestionTextBox.Text", FTextBlockStyle()
+			.SetFont(StyleFonts.Normal)
+			.SetColorAndOpacity(FLinearColor(FColor(0xffaaaaaa)))
+		);
+	}
 
+	// SToolTip defaults...
+	{
+		Style->Set("ToolTip.Font", StyleFonts.Small);
+		Style->Set("ToolTip.Background", new BOX_BRUSH("Old/ToolTip_Background", FMargin(8.0f / 64.0f)));
 
-			.SetTextColor(FStyleColors::White)
-			.SetSelectedTextColor(FStyleColors::Input)
+		Style->Set("ToolTip.LargerFont", StyleFonts.Normal);
+		Style->Set("ToolTip.BrightBackground", new BOX_BRUSH("Old/ToolTip_BrightBackground", FMargin(8.0f / 64.0f)));
+	}
 
-			.SetDropIndicator_Above(FSlateNoResource())
-			.SetDropIndicator_Onto( FSlateNoResource())
-			.SetDropIndicator_Below(FSlateNoResource());
+	// SBorder defaults...
+	{
+		Style->Set("Border", new FSlateColorBrush(FStyleColors::Background));
+	}
 
-		Style->Set("ComboBox.Row", ComboBoxRow);
+	// SHyperlink defaults...
+	{
+		FButtonStyle HyperlinkButton = FButtonStyle()
+			.SetNormal(BORDER_BRUSH("Old/HyperlinkDotted", FMargin(0, 0, 0, 3 / 16.0f)))
+			.SetPressed(FSlateNoResource())
+			.SetHovered(BORDER_BRUSH("Old/HyperlinkUnderline", FMargin(0, 0, 0, 3 / 16.0f)));
 
-		// SComboBox for Icons Only
-		// FComboButtonStyle IconComboButton = FComboButtonStyle()
-		// Style->Set("IconComboButton")
+		FHyperlinkStyle Hyperlink = FHyperlinkStyle()
+			.SetUnderlineStyle(HyperlinkButton)
+			.SetTextStyle(NormalText)
+			.SetPadding(FMargin(0.0f));
+		Style->Set("Hyperlink", Hyperlink);
 	}
 
-	// SMessageLogListing
+	// SProgressBar defaults...
 	{
-		FComboButtonStyle MessageLogListingComboButton = FComboButtonStyle()
-			.SetButtonStyle(NoBorder)
-			.SetDownArrowImage(IMAGE_BRUSH("Common/ComboArrow", Icon8x8))
-			.SetMenuBorderBrush(FSlateNoResource())
-			.SetMenuBorderPadding(FMargin(0.0f));
-		Style->Set("MessageLogListingComboButton", MessageLogListingComboButton);
+		Style->Set("ProgressBar", FProgressBarStyle()
+			.SetBackgroundImage(BOX_BRUSH("Common/ProgressBar_Background", FMargin(5.f / 12.f)))
+			.SetFillImage(BOX_BRUSH("Common/ProgressBar_Fill", FMargin(5.f / 12.f), FLinearColor(1.0f, 0.22f, 0.0f)))
+			.SetMarqueeImage(IMAGE_BRUSH("Common/ProgressBar_Marquee", FVector2D(20, 12), FLinearColor::White, ESlateBrushTileType::Horizontal))
+		);
 	}
 
-	// SEditableComboBox defaults...
+	// SThrobber, SCircularThrobber defaults...
 	{
-		Style->Set("EditableComboBox.Add", new IMAGE_BRUSH("Icons/PlusSymbol_12x", Icon12x12));
-		Style->Set("EditableComboBox.Delete", new IMAGE_BRUSH("Icons/Cross_12x", Icon12x12));
-		Style->Set("EditableComboBox.Rename", new IMAGE_BRUSH("Icons/ellipsis_12x", Icon12x12));
-		Style->Set("EditableComboBox.Accept", new IMAGE_BRUSH("Common/Check", Icon16x16));
+		Style->Set("Throbber.Chunk", new IMAGE_BRUSH("Common/Throbber_Piece", FVector2D(16, 16)));
+		Style->Set("Throbber.CircleChunk", new IMAGE_BRUSH("Common/Throbber_Piece", FVector2D(8, 8)));
 	}
 
-	// SCheckBox defaults...
+	// SExpandableArea defaults...
 	{
+		Style->Set("ExpandableArea", FExpandableAreaStyle()
+			.SetCollapsedImage(IMAGE_BRUSH("Common/TreeArrow_Collapsed", Icon10x10, DefaultForeground))
+			.SetExpandedImage(IMAGE_BRUSH("Common/TreeArrow_Expanded", Icon10x10, DefaultForeground))
+		);
+		Style->Set("ExpandableArea.TitleFont", StyleFonts.SmallBold);
+		Style->Set("ExpandableArea.Border", new FSlateRoundedBoxBrush(FStyleColors::Background, 4) );
+	}
 
-		const float CheckboxCornerRadius = 3.f;
-		const float CheckboxOutlineThickness = 1.0f;
-
-		const FCheckBoxStyle BasicCheckBoxStyle = FCheckBoxStyle()
-			.SetCheckBoxType(ESlateCheckBoxType::CheckBox)
-
-			.SetForegroundColor(               FLinearColor::White )
-			.SetHoveredForegroundColor(        FLinearColor::White )
-			.SetPressedForegroundColor(        FLinearColor::White )
-			.SetCheckedForegroundColor(        FLinearColor::White )
-			.SetCheckedHoveredForegroundColor( FLinearColor::White )
-			.SetCheckedPressedForegroundColor( FLinearColor::White )
-			.SetUndeterminedForegroundColor(   FLinearColor::White )
-
-			.SetUncheckedImage(          FSlateNoResource())
-			.SetUncheckedHoveredImage(   FSlateNoResource())
-			.SetUncheckedPressedImage(   FSlateNoResource())
-			.SetCheckedImage(            IMAGE_BRUSH_SVG("Starship/CoreWidgets/CheckBox/check", 		Icon16x16, FStyleColors::Primary))
-			.SetCheckedHoveredImage(     IMAGE_BRUSH_SVG("Starship/CoreWidgets/CheckBox/check", 		Icon16x16, FStyleColors::PrimaryHover)) 
-			.SetCheckedPressedImage(     IMAGE_BRUSH_SVG("Starship/CoreWidgets/CheckBox/check", 		Icon16x16, FStyleColors::Primary))
-			.SetUndeterminedImage(       IMAGE_BRUSH_SVG("Starship/CoreWidgets/CheckBox/indeterminate", Icon16x16, FStyleColors::Primary))
-			.SetUndeterminedHoveredImage(IMAGE_BRUSH_SVG("Starship/CoreWidgets/CheckBox/indeterminate", Icon16x16, FStyleColors::PrimaryHover))
-			.SetUndeterminedPressedImage(IMAGE_BRUSH_SVG("Starship/CoreWidgets/CheckBox/indeterminate", Icon16x16, FStyleColors::Primary))
-			.SetBackgroundImage(         FSlateRoundedBoxBrush(FStyleColors::Input,   CheckboxCornerRadius, FStyleColors::Dropdown,    CheckboxOutlineThickness, Icon18x18))
-			.SetBackgroundHoveredImage(  FSlateRoundedBoxBrush(FStyleColors::Input,   CheckboxCornerRadius, FStyleColors::Hover,       CheckboxOutlineThickness, Icon18x18))
-			.SetBackgroundPressedImage(  FSlateRoundedBoxBrush(FStyleColors::Foldout, CheckboxCornerRadius, FStyleColors::Hover,       CheckboxOutlineThickness, Icon18x18));
-
-		Style->Set("Checkbox", BasicCheckBoxStyle);
-
-
-		const FCheckBoxStyle SimplifiedCheckBoxStyle = FCheckBoxStyle()
-			.SetCheckBoxType(ESlateCheckBoxType::CheckBox)
-
-			.SetForegroundColor(               FLinearColor::White )
-			.SetHoveredForegroundColor(        FLinearColor::White )
-			.SetPressedForegroundColor(        FLinearColor::White )
-			.SetCheckedForegroundColor(        FLinearColor::White )
-			.SetCheckedHoveredForegroundColor( FLinearColor::White )
-			.SetCheckedPressedForegroundColor( FLinearColor::White )
-			.SetUndeterminedForegroundColor(   FLinearColor::White )
-
-			.SetUncheckedImage(          FSlateNoResource())
-			.SetUncheckedHoveredImage(   FSlateNoResource())
-			.SetUncheckedPressedImage(   FSlateNoResource())
-			.SetCheckedImage(            IMAGE_BRUSH_SVG("Starship/CoreWidgets/CheckBox/check", 		Icon16x16, FStyleColors::Foreground))
-			.SetCheckedHoveredImage(     IMAGE_BRUSH_SVG("Starship/CoreWidgets/CheckBox/check", 		Icon16x16, FStyleColors::ForegroundHover)) 
-			.SetCheckedPressedImage(     IMAGE_BRUSH_SVG("Starship/CoreWidgets/CheckBox/check", 		Icon16x16, FStyleColors::Foreground))
-			.SetUndeterminedImage(       IMAGE_BRUSH_SVG("Starship/CoreWidgets/CheckBox/indeterminate", Icon16x16, FStyleColors::Foreground))
-			.SetUndeterminedHoveredImage(IMAGE_BRUSH_SVG("Starship/CoreWidgets/CheckBox/indeterminate", Icon16x16, FStyleColors::ForegroundHover))
-			.SetUndeterminedPressedImage(IMAGE_BRUSH_SVG("Starship/CoreWidgets/CheckBox/indeterminate", Icon16x16, FStyleColors::Foreground))
-			.SetBackgroundImage(         FSlateRoundedBoxBrush(FStyleColors::Input,   CheckboxCornerRadius, FStyleColors::Dropdown,    CheckboxOutlineThickness, Icon18x18))
-			.SetBackgroundHoveredImage(  FSlateRoundedBoxBrush(FStyleColors::Input,   CheckboxCornerRadius, FStyleColors::Hover,       CheckboxOutlineThickness, Icon18x18))
-			.SetBackgroundPressedImage(  FSlateRoundedBoxBrush(FStyleColors::Foldout, CheckboxCornerRadius, FStyleColors::Hover,       CheckboxOutlineThickness, Icon18x18));
-
-		Style->Set("SimplifiedCheckbox", SimplifiedCheckBoxStyle);
-
-		/* Set images for various transparent SCheckBox states ... */
-		const FCheckBoxStyle BasicTransparentCheckBoxStyle = FCheckBoxStyle()
-			.SetCheckBoxType(ESlateCheckBoxType::ToggleButton)
-			.SetUncheckedImage(FSlateNoResource())
-			.SetUncheckedHoveredImage(FSlateNoResource())
-			.SetUncheckedPressedImage(FSlateNoResource())
-			.SetCheckedImage(FSlateNoResource())
-			.SetCheckedHoveredImage(FSlateNoResource())
-			.SetCheckedPressedImage(FSlateNoResource())
-			.SetUndeterminedImage(FSlateNoResource())
+	// SSlider and SVolumeControl defaults...
+	{
+		FSliderStyle SliderStyle = FSliderStyle()
+			.SetNormalBarImage(   FSlateRoundedBoxBrush(FStyleColors::Input, 2.0))
+			.SetHoveredBarImage(  FSlateRoundedBoxBrush(FStyleColors::Input, 2.0))
+			.SetNormalThumbImage(  FSlateRoundedBoxBrush(FStyleColors::Hover2, Icon8x8) )
+			.SetHoveredThumbImage( FSlateRoundedBoxBrush(FStyleColors::ForegroundHover, Icon8x8) )
+			.SetBarThickness(4.0f);
+		Style->Set("Slider", SliderStyle);
+
+		Style->Set("VolumeControl", FVolumeControlStyle()
+			.SetSliderStyle(SliderStyle)
+			.SetHighVolumeImage(IMAGE_BRUSH("Common/VolumeControl_High", Icon16x16))
+			.SetMidVolumeImage(IMAGE_BRUSH("Common/VolumeControl_Mid", Icon16x16))
+			.SetLowVolumeImage(IMAGE_BRUSH("Common/VolumeControl_Low", Icon16x16))
+			.SetNoVolumeImage(IMAGE_BRUSH("Common/VolumeControl_Off", Icon16x16))
+			.SetMutedImage(IMAGE_BRUSH("Common/VolumeControl_Muted", Icon16x16))
+		);
+	}
+
+	// SSpinBox defaults...
+	{
+		Style->Set("SpinBox", FSpinBoxStyle()
+			.SetBackgroundBrush(FSlateRoundedBoxBrush(FStyleColors::Input, InputFocusRadius, FStyleColors::Dropdown, InputFocusThickness))
+			.SetHoveredBackgroundBrush(FSlateRoundedBoxBrush(FStyleColors::Input, InputFocusRadius, FStyleColors::Hover, InputFocusThickness))
+
+			.SetActiveFillBrush(FSlateRoundedBoxBrush(FStyleColors::Hover, InputFocusRadius, FLinearColor::Transparent, InputFocusThickness))
+			.SetInactiveFillBrush(FSlateRoundedBoxBrush(FStyleColors::Dropdown, InputFocusRadius, FLinearColor::Transparent, InputFocusThickness))
+			.SetArrowsImage(FSlateNoResource())
+			.SetForegroundColor(FStyleColors::ForegroundHover)
+			.SetTextPadding(FMargin(10.f, 3.5f, 10.f, 4.f))
+		);
+	}
+
+	// SNumericEntryBox defaults...
+	{
+		Style->Set("NumericEntrySpinBox", FSpinBoxStyle()
+			.SetBackgroundBrush(FSlateRoundedBoxBrush(FStyleColors::Input, InputFocusRadius, FStyleColors::Dropdown, InputFocusThickness))
+			.SetHoveredBackgroundBrush(FSlateRoundedBoxBrush(FStyleColors::Input, InputFocusRadius, FStyleColors::Hover, InputFocusThickness))
+
+			.SetActiveFillBrush(FSlateRoundedBoxBrush(FStyleColors::Hover, InputFocusRadius, FLinearColor::Transparent, InputFocusThickness))
+			.SetInactiveFillBrush(FSlateRoundedBoxBrush(FStyleColors::Dropdown, InputFocusRadius, FLinearColor::Transparent, InputFocusThickness))
+
+			.SetArrowsImage(FSlateNoResource())
+			.SetTextPadding(FMargin(0.f))
+			.SetForegroundColor(FStyleColors::ForegroundHover)
+		);
+
+		Style->Set("NumericEntrySpinBox_Dark", FSpinBoxStyle()
+			.SetBackgroundBrush(FSlateRoundedBoxBrush(FStyleColors::Input, InputFocusRadius, FStyleColors::Dropdown, InputFocusThickness))
+			.SetHoveredBackgroundBrush(FSlateRoundedBoxBrush(FStyleColors::Input, InputFocusRadius, FStyleColors::Hover, InputFocusThickness))
+
+			.SetActiveFillBrush(FSlateRoundedBoxBrush(FStyleColors::Hover, InputFocusRadius, FLinearColor::Transparent, InputFocusThickness))
+			.SetInactiveFillBrush(FSlateRoundedBoxBrush(FStyleColors::Dropdown, InputFocusRadius, FLinearColor::Transparent, InputFocusThickness))
+
+			.SetArrowsImage(FSlateNoResource())
+
+			.SetTextPadding(FMargin(0.f))
+			.SetForegroundColor(FStyleColors::ForegroundHover)
+		);
+
+		Style->Set("NumericEntrySpinBox.Decorator", new BOX_BRUSH("Common/TextBoxLabelBorder", FMargin(5.0f / 16.0f)));
+
+		Style->Set("NumericEntrySpinBox.NarrowDecorator", new BOX_BRUSH("Common/TextBoxLabelBorder", FMargin(2.0f / 16.0f, 4.0f / 16.0f, 2.0f / 16.0f, 4.0f / 16.0f)));
+	}
+
+	SetupColorPickerStyles(Style);
+
+	// SSplitter
+	{
+		Style->Set("Splitter", FSplitterStyle()
+			.SetHandleNormalBrush(FSlateNoResource())
+			.SetHandleHighlightBrush(IMAGE_BRUSH("Common/SplitterHandleHighlight", Icon8x8, FLinearColor::White))
+		);
+	}
+
+	// TableView defaults...
+	SetupTableViewStyles(Style);
+
+	SetupMultiboxStyles(Style);
+
+	// SExpandableButton defaults...
+	{
+		Style->Set("ExpandableButton.Background", new BOX_BRUSH("Common/Button", 8.0f / 32.0f));
+
+		// Extra padding on the right and bottom to account for image shadow
+		Style->Set("ExpandableButton.Padding", FMargin(3.f, 3.f, 6.f, 6.f));
+
+		Style->Set("ExpandableButton.CloseButton", new IMAGE_BRUSH("Common/ExpansionButton_CloseOverlay", Icon16x16));
+	}
+
+	// SBreadcrumbTrail defaults...
+	{
+		Style->Set("BreadcrumbTrail.Delimiter", new IMAGE_BRUSH("Common/Delimiter", Icon16x16));
+
+		Style->Set("BreadcrumbButton", FButtonStyle()
+			.SetNormal(FSlateNoResource())
+			.SetPressed(BOX_BRUSH("Common/RoundedSelection_16x", 4.0f / 16.0f, SelectionColor_Pressed))
+			.SetHovered(BOX_BRUSH("Common/RoundedSelection_16x", 4.0f / 16.0f, SelectionColor))
+			.SetNormalPadding(FMargin(0, 0))
+			.SetPressedPadding(FMargin(0, 0))
+		);
+	}
+
+	// SWizard defaults
+	{
+		Style->Set("Wizard.PageTitle", FTextBlockStyle(NormalText)
+			.SetFont(DEFAULT_FONT("BoldCondensed", 28))
+			.SetShadowOffset(FVector2D(1, 1))
+			.SetShadowColorAndOpacity(FLinearColor(0, 0, 0, 0.9f))
+		);
+	}
+
+	// SNotificationList defaults...
+	{
+		Style->Set("NotificationList.FontBold", DEFAULT_FONT("Bold", 16));
+		Style->Set("NotificationList.FontLight", DEFAULT_FONT("Light", 12));
+		Style->Set("NotificationList.ItemBackground", new BOX_BRUSH("Old/Menu_Background", FMargin(8.0f / 64.0f)));
+		Style->Set("NotificationList.ItemBackground_Border", new BOX_BRUSH("Old/Menu_Background_Inverted_Border_Bold", FMargin(8.0f / 64.0f)));
+		Style->Set("NotificationList.ItemBackground_Border_Transparent", new BOX_BRUSH("Old/Notification_Border_Flash", FMargin(8.0f / 64.0f)));
+		Style->Set("NotificationList.SuccessImage", new IMAGE_BRUSH("Icons/notificationlist_success", Icon16x16));
+		Style->Set("NotificationList.FailImage", new IMAGE_BRUSH("Icons/notificationlist_fail", Icon16x16));
+		Style->Set("NotificationList.DefaultMessage", new IMAGE_BRUSH("Common/EventMessage_Default", Icon40x40));
+	}
+
+	// SSeparator defaults...
+	{
+		Style->Set("Separator", new BOX_BRUSH("Common/Separator", 1 / 4.0f, FLinearColor(1, 1, 1, 0.5f)));
+	}
+
+	// SHeader defaults...
+	{
+		Style->Set("Header.Pre", new BOX_BRUSH("Common/Separator", FMargin(1 / 4.0f, 0, 2 / 4.0f, 0), FLinearColor(1, 1, 1, 0.5f)));
+		Style->Set("Header.Post", new BOX_BRUSH("Common/Separator", FMargin(2 / 4.0f, 0, 1 / 4.0f, 0), FLinearColor(1, 1, 1, 0.5f)));
+	}
+
+	SetupDockingStyles(Style);
+
+	// SScrollBox defaults...
+	{
+		Style->Set("ScrollBox", FScrollBoxStyle()
+			.SetTopShadowBrush(BOX_BRUSH("Common/ScrollBoxShadowTop", FVector2D(16, 8), FMargin(0.5, 1, 0.5, 0)))
+			.SetBottomShadowBrush(BOX_BRUSH("Common/ScrollBoxShadowBottom", FVector2D(16, 8), FMargin(0.5, 0, 0.5, 1)))
+			.SetLeftShadowBrush(BOX_BRUSH("Common/ScrollBoxShadowLeft", FVector2D(8, 16), FMargin(1, 0.5, 0, 0.5)))
+			.SetRightShadowBrush(BOX_BRUSH("Common/ScrollBoxShadowRight", FVector2D(8, 16), FMargin(0, 0.5, 1, 0.5)))
+			.SetBarThickness(8.0)
+		);
+	}
+
+	// SScrollBorder defaults...
+	{
+		Style->Set("ScrollBorder", FScrollBorderStyle()
+			.SetTopShadowBrush(BOX_BRUSH("Common/ScrollBorderShadowTop", FVector2D(16, 8), FMargin(0.5, 1, 0.5, 0)))
+			.SetBottomShadowBrush(BOX_BRUSH("Common/ScrollBorderShadowBottom", FVector2D(16, 8), FMargin(0.5, 0, 0.5, 1)))
+		);
+	}
+
+
+
+	// SWindow defaults...
+	{
+#if !PLATFORM_MAC
+		const FButtonStyle MinimizeButtonStyle = FButtonStyle(NoBorder)
+			.SetNormal( IMAGE_BRUSH_SVG("Starship/CoreWidgets/Window/minimize", FVector2D(42.0f, 34.0f), FStyleColors::Foreground))
+			.SetHovered(IMAGE_BRUSH_SVG("Starship/CoreWidgets/Window/minimize", FVector2D(42.0f, 34.0f), FStyleColors::ForegroundHover))
+			.SetPressed(IMAGE_BRUSH_SVG("Starship/CoreWidgets/Window/minimize", FVector2D(42.0f, 34.0f), FStyleColors::Foreground));
+
+		const FButtonStyle MaximizeButtonStyle = FButtonStyle(NoBorder)
+			.SetNormal( IMAGE_BRUSH_SVG("Starship/CoreWidgets/Window/maximize", FVector2D(42.0f, 34.0f), FStyleColors::Foreground))
+			.SetHovered(IMAGE_BRUSH_SVG("Starship/CoreWidgets/Window/maximize", FVector2D(42.0f, 34.0f), FStyleColors::ForegroundHover))
+			.SetPressed(IMAGE_BRUSH_SVG("Starship/CoreWidgets/Window/maximize", FVector2D(42.0f, 34.0f), FStyleColors::Foreground));
+
+		const FButtonStyle RestoreButtonStyle = FButtonStyle(NoBorder)
+			.SetNormal( IMAGE_BRUSH_SVG("Starship/CoreWidgets/Window/restore", FVector2D(42.0f, 34.0f), FStyleColors::Foreground))
+			.SetHovered(IMAGE_BRUSH_SVG("Starship/CoreWidgets/Window/restore", FVector2D(42.0f, 34.0f), FStyleColors::ForegroundHover))
+			.SetPressed(IMAGE_BRUSH_SVG("Starship/CoreWidgets/Window/restore", FVector2D(42.0f, 34.0f), FStyleColors::Foreground));
+
+		const FButtonStyle CloseButtonStyle = FButtonStyle(NoBorder)
+			.SetNormal( IMAGE_BRUSH_SVG("Starship/CoreWidgets/Window/close", FVector2D(42.0f, 34.0f), FStyleColors::Foreground))
+			.SetHovered(IMAGE_BRUSH_SVG("Starship/CoreWidgets/Window/close", FVector2D(42.0f, 34.0f), FStyleColors::ForegroundHover))
+			.SetPressed(IMAGE_BRUSH_SVG("Starship/CoreWidgets/Window/close", FVector2D(42.0f, 34.0f), FStyleColors::Foreground));
+#endif
+
+		const FTextBlockStyle TitleTextStyle = FTextBlockStyle(NormalText)
+			.SetFont(StyleFonts.Normal)
+			.SetColorAndOpacity(FLinearColor::White)
+			.SetShadowOffset(FVector2D(1.0f, 1.0f))
+			.SetShadowColorAndOpacity(FLinearColor::Black);
+
+		Style->Set("Window", FWindowStyle()
+
+#if !PLATFORM_MAC
+			.SetMinimizeButtonStyle(MinimizeButtonStyle)
+			.SetMaximizeButtonStyle(MaximizeButtonStyle)
+			.SetRestoreButtonStyle(RestoreButtonStyle)
+			.SetCloseButtonStyle(CloseButtonStyle)
+#endif
+			.SetTitleTextStyle(TitleTextStyle)
+			.SetActiveTitleBrush(FSlateNoResource())
+			.SetInactiveTitleBrush(FSlateNoResource())
+			.SetFlashTitleBrush(IMAGE_BRUSH("Common/Window/WindowTitle_Flashing", Icon24x24, FLinearColor(1.0f, 1.0f, 1.0f, 1.0f), ESlateBrushTileType::Horizontal))	
+			.SetBackgroundBrush(FSlateColorBrush(FStyleColors::Input))
+			.SetBorderBrush(FSlateNoResource())
+			.SetOutlineBrush(FSlateRoundedBoxBrush(FLinearColor::Transparent, 0.0, FLinearColor::FromSRGBColor(FColor::FromHex("#2C2C2CFF")), 1.0))
+			.SetChildBackgroundBrush(FSlateColorBrush(FStyleColors::Input))
+			.SetOutlineColor(WindowHighlight)
+			.SetCornerRadius(0)
+			.SetBorderPadding(FMargin(2.0, 1.0, 2.0, 2.0))
+		);
+
+		Style->Set("ChildWindow.Background", new FSlateColorBrush(FStyleColors::Input));
+	}
+
+
+	// Standard Dialog Settings
+	{
+		Style->Set("StandardDialog.ContentPadding", FMargin(12.0f, 2.0f));
+		Style->Set("StandardDialog.SlotPadding", FMargin(6.0f, 0.0f, 0.0f, 0.0f));
+		Style->Set("StandardDialog.MinDesiredSlotWidth", 80.0f);
+		Style->Set("StandardDialog.MinDesiredSlotHeight", 0.0f);
+		Style->Set("StandardDialog.SmallFont", StyleFonts.Small);
+		Style->Set("StandardDialog.LargeFont", StyleFonts.Normal);
+	}
+
+	// Widget Reflector Window
+	{
+		Style->Set("WidgetReflector.TabIcon", new IMAGE_BRUSH("Icons/icon_tab_WidgetReflector_16x", Icon16x16));
+		Style->Set("WidgetReflector.Icon", new IMAGE_BRUSH("Icons/icon_tab_WidgetReflector_40x", Icon40x40));
+		Style->Set("WidgetReflector.Icon.Small", new IMAGE_BRUSH("Icons/icon_tab_WidgetReflector_40x", Icon20x20));
+		Style->Set("WidgetReflector.FocusableCheck", FCheckBoxStyle()
+			.SetUncheckedImage(IMAGE_BRUSH("Icons/Empty_14x", Icon14x14))
+			.SetUncheckedHoveredImage(IMAGE_BRUSH("Icons/Empty_14x", Icon14x14))
+			.SetUncheckedPressedImage(IMAGE_BRUSH("Common/SmallCheckBox_Hovered", Icon14x14))
+			.SetCheckedImage(IMAGE_BRUSH("Common/SmallCheck", Icon14x14))
+			.SetCheckedHoveredImage(IMAGE_BRUSH("Common/SmallCheck", Icon14x14))
+			.SetCheckedPressedImage(IMAGE_BRUSH("Common/SmallCheck", Icon14x14))
+			.SetUndeterminedImage(IMAGE_BRUSH("Icons/Empty_14x", Icon14x14))
 			.SetUndeterminedHoveredImage(FSlateNoResource())
-			.SetUndeterminedPressedImage(FSlateNoResource());
-		Style->Set("TransparentCheckBox", BasicTransparentCheckBoxStyle);
+			.SetUndeterminedPressedImage(FSlateNoResource())
+		);
+	}
 
-		/* Default Style for a toggleable button */
-		const FCheckBoxStyle ToggleButtonStyle = FCheckBoxStyle()
-			.SetCheckBoxType(ESlateCheckBoxType::ToggleButton)
-			.SetCheckedImage(FSlateNoResource())
-			.SetCheckedHoveredImage(FSlateRoundedBoxBrush(FStyleColors::Hover, InputFocusRadius))
-			.SetCheckedPressedImage(FSlateRoundedBoxBrush(FStyleColors::Input, InputFocusRadius))
-			.SetUncheckedImage(FSlateNoResource())
-			.SetUncheckedHoveredImage(FSlateRoundedBoxBrush(FStyleColors::Hover, InputFocusRadius))
-			.SetUncheckedPressedImage(FSlateRoundedBoxBrush(FStyleColors::Input, InputFocusRadius))
-			.SetForegroundColor(FStyleColors::Foreground)
-			.SetHoveredForegroundColor(FStyleColors::ForegroundHover)
-			.SetPressedForegroundColor(FStyleColors::ForegroundHover)
-			.SetCheckedForegroundColor(FStyleColors::Primary)
-			.SetCheckedHoveredForegroundColor(FStyleColors::PrimaryHover)
-			.SetPadding(DefaultMargins)
-		;
+	// Message Log
+	{
+		Style->Set("MessageLog", FTextBlockStyle(NormalText)
+			.SetFont(StyleFonts.Small)
+			.SetShadowOffset(FVector2D::ZeroVector)
+		);
+		Style->Set("MessageLog.Error", new IMAGE_BRUSH("MessageLog/Log_Error", Icon16x16));
+		Style->Set("MessageLog.Warning", new IMAGE_BRUSH("MessageLog/Log_Warning", Icon16x16));
+		Style->Set("MessageLog.Note", new IMAGE_BRUSH("MessageLog/Log_Note", Icon16x16));
+	}
 
-		Style->Set("ToggleButtonCheckbox", ToggleButtonStyle);
+	// Wizard icons
+	{
+		Style->Set("Wizard.BackIcon", new IMAGE_BRUSH("Icons/BackIcon", Icon8x8));
+		Style->Set("Wizard.NextIcon", new IMAGE_BRUSH("Icons/NextIcon", Icon8x8));
+	}
 
-		/* Default Style for a toggleable button */
-		const FCheckBoxStyle ToggleButtonAltStyle = FCheckBoxStyle()
-			.SetCheckBoxType(ESlateCheckBoxType::ToggleButton)
-			.SetCheckedImage(       FSlateRoundedBoxBrush(FStyleColors::Input, InputFocusRadius))
-			.SetCheckedHoveredImage(FSlateRoundedBoxBrush(FStyleColors::Hover, InputFocusRadius))
-			.SetCheckedPressedImage(FSlateRoundedBoxBrush(FStyleColors::Input, InputFocusRadius))
-			.SetUncheckedImage(FSlateNoResource())
-			.SetUncheckedHoveredImage(FSlateRoundedBoxBrush(FStyleColors::Hover, InputFocusRadius))
-			.SetUncheckedPressedImage(FSlateRoundedBoxBrush(FStyleColors::Input, InputFocusRadius))
-			.SetForegroundColor(FStyleColors::Foreground)
-			.SetHoveredForegroundColor(FStyleColors::ForegroundHover)
-			.SetPressedForegroundColor(FStyleColors::ForegroundHover)
-			.SetCheckedForegroundColor(FStyleColors::Primary)
-			.SetCheckedHoveredForegroundColor(FStyleColors::PrimaryHover)
-			.SetPadding(DefaultMargins)
-		;
+	// Syntax highlighting
+	{
+		const FTextBlockStyle SmallMonospacedText = FTextBlockStyle(Style->GetWidgetStyle<FTextBlockStyle>("MonospacedText"))
+			.SetFont(DEFAULT_FONT("Mono", 9));
 
-		Style->Set("ToggleButtonCheckboxAlt", ToggleButtonAltStyle);
+		Style->Set("SyntaxHighlight.Normal", SmallMonospacedText);
+		Style->Set("SyntaxHighlight.Node", FTextBlockStyle(SmallMonospacedText).SetColorAndOpacity(FLinearColor(FColor(0xff006ab4)))); // blue
+		Style->Set("SyntaxHighlight.NodeAttributeKey", FTextBlockStyle(SmallMonospacedText).SetColorAndOpacity(FLinearColor(FColor(0xffb40000)))); // red
+		Style->Set("SyntaxHighlight.NodeAttribueAssignment", FTextBlockStyle(SmallMonospacedText).SetColorAndOpacity(FLinearColor(FColor(0xffb2b400)))); // yellow
+		Style->Set("SyntaxHighlight.NodeAttributeValue", FTextBlockStyle(SmallMonospacedText).SetColorAndOpacity(FLinearColor(FColor(0xffb46100)))); // orange
+	}
 
-		/* Style for a segmented box */
-		const FCheckBoxStyle SegmentedBoxLeft = FCheckBoxStyle()
-			.SetCheckBoxType(ESlateCheckBoxType::ToggleButton)
+	return Style;
+}
 
-			.SetUncheckedImage(           BOX_BRUSH("/Starship/CoreWidgets/SegmentedBox/left", FVector2D(16.f, 16.f), FMargin(4.0f/16.0f), FStyleColors::Dropdown))
-			.SetUncheckedHoveredImage(    BOX_BRUSH("/Starship/CoreWidgets/SegmentedBox/left", FVector2D(16.f, 16.f), FMargin(4.0f/16.0f), FStyleColors::Hover))
-			.SetUncheckedPressedImage(    BOX_BRUSH("/Starship/CoreWidgets/SegmentedBox/left", FVector2D(16.f, 16.f), FMargin(4.0f/16.0f), FStyleColors::Dropdown))
-			.SetCheckedImage(             BOX_BRUSH("/Starship/CoreWidgets/SegmentedBox/left", FVector2D(16.f, 16.f), FMargin(4.0f/16.0f), FStyleColors::Input))
-			.SetCheckedHoveredImage(      BOX_BRUSH("/Starship/CoreWidgets/SegmentedBox/left", FVector2D(16.f, 16.f), FMargin(4.0f/16.0f), FStyleColors::Input))
-			.SetCheckedPressedImage(      BOX_BRUSH("/Starship/CoreWidgets/SegmentedBox/left", FVector2D(16.f, 16.f), FMargin(4.0f/16.0f), FStyleColors::Input))
-			.SetForegroundColor(FStyleColors::Foreground)
-			.SetHoveredForegroundColor(FStyleColors::ForegroundHover)
-			.SetPressedForegroundColor(FStyleColors::ForegroundHover)
-			.SetCheckedForegroundColor(FStyleColors::Primary)
-			.SetCheckedHoveredForegroundColor(FStyleColors::Primary)
-			.SetCheckedPressedForegroundColor(FStyleColors::Primary)
-			.SetPadding(DefaultMargins);
-		
 
-		const FCheckBoxStyle SegmentedBoxCenter = FCheckBoxStyle()
-			.SetCheckBoxType(ESlateCheckBoxType::ToggleButton)
-			.SetUncheckedImage(           FSlateColorBrush(FStyleColors::Dropdown))
-			.SetUncheckedHoveredImage(    FSlateColorBrush(FStyleColors::Hover))
-			.SetUncheckedPressedImage(    FSlateColorBrush(FStyleColors::Dropdown))
-			.SetCheckedImage(             FSlateColorBrush(FStyleColors::Input))
-			.SetCheckedHoveredImage(      FSlateColorBrush(FStyleColors::Input))
-			.SetCheckedPressedImage(      FSlateColorBrush(FStyleColors::Input))
-			.SetForegroundColor(FStyleColors::Foreground)
-			.SetHoveredForegroundColor(FStyleColors::ForegroundHover)
-			.SetPressedForegroundColor(FStyleColors::ForegroundHover)
-			.SetCheckedForegroundColor(FStyleColors::Primary)
-			.SetCheckedHoveredForegroundColor(FStyleColors::Primary)
-			.SetCheckedPressedForegroundColor(FStyleColors::Primary)
-			.SetPadding(DefaultMargins);
-		
+const TSharedPtr<FSlateDynamicImageBrush> FStarshipCoreStyle::GetDynamicImageBrush(FName BrushTemplate, FName TextureName, const ANSICHAR* Specifier)
+{
+	return Instance->GetDynamicImageBrush(BrushTemplate, TextureName, Specifier);
+}
 
-		const FCheckBoxStyle SegmentedBoxRight = FCheckBoxStyle()
-			.SetCheckBoxType(ESlateCheckBoxType::ToggleButton)
-			.SetUncheckedImage(           BOX_BRUSH("/Starship/CoreWidgets/SegmentedBox/right", FVector2D(16.f, 16.f), FMargin(4.0f/16.0f), FStyleColors::Dropdown))
-			.SetUncheckedHoveredImage(    BOX_BRUSH("/Starship/CoreWidgets/SegmentedBox/right", FVector2D(16.f, 16.f), FMargin(4.0f/16.0f), FStyleColors::Hover))
-			.SetUncheckedPressedImage(    BOX_BRUSH("/Starship/CoreWidgets/SegmentedBox/right", FVector2D(16.f, 16.f), FMargin(4.0f/16.0f), FStyleColors::Dropdown))
-			.SetCheckedImage(             BOX_BRUSH("/Starship/CoreWidgets/SegmentedBox/right", FVector2D(16.f, 16.f), FMargin(4.0f/16.0f), FStyleColors::Input))
-			.SetCheckedHoveredImage(      BOX_BRUSH("/Starship/CoreWidgets/SegmentedBox/right", FVector2D(16.f, 16.f), FMargin(4.0f/16.0f), FStyleColors::Input))
-			.SetCheckedPressedImage(      BOX_BRUSH("/Starship/CoreWidgets/SegmentedBox/right", FVector2D(16.f, 16.f), FMargin(4.0f/16.0f), FStyleColors::Input))
-			.SetForegroundColor(FStyleColors::Foreground)
-			.SetHoveredForegroundColor(FStyleColors::ForegroundHover)
-			.SetPressedForegroundColor(FStyleColors::ForegroundHover)
-			.SetCheckedForegroundColor(FStyleColors::Primary)
-			.SetCheckedHoveredForegroundColor(FStyleColors::Primary)
-			.SetCheckedPressedForegroundColor(FStyleColors::Primary)
-			.SetPadding(DefaultMargins);
 
-		Style->Set("SegmentedBoxLeft", SegmentedBoxLeft);
-		Style->Set("SegmentedBoxCenter", SegmentedBoxCenter);
-		Style->Set("SegmentedBoxRight", SegmentedBoxRight);
+const TSharedPtr<FSlateDynamicImageBrush> FStarshipCoreStyle::GetDynamicImageBrush(FName BrushTemplate, const ANSICHAR* Specifier, class UTexture2D* TextureResource, FName TextureName)
+{
+	return Instance->GetDynamicImageBrush(BrushTemplate, Specifier, TextureResource, TextureName);
+}
+
+
+const TSharedPtr<FSlateDynamicImageBrush> FStarshipCoreStyle::GetDynamicImageBrush(FName BrushTemplate, class UTexture2D* TextureResource, FName TextureName)
+{
+	return Instance->GetDynamicImageBrush(BrushTemplate, TextureResource, TextureName);
+}
+
+
+/* FSlateThrottleManager implementation
+ *****************************************************************************/
+
+void FStarshipCoreStyle::SetStyle(const TSharedRef< ISlateStyle >& NewStyle)
+{
+	if (Instance.IsValid())
+	{
+		FSlateStyleRegistry::UnRegisterSlateStyle(*Instance.Get());
+	}
+
+	Instance = NewStyle;
+
+	if (Instance.IsValid())
+	{
+		FSlateStyleRegistry::RegisterSlateStyle(*Instance.Get());
+	}
+	else
+	{
+		ResetToDefault();
+	}
+}
+
+void FStarshipCoreStyle::SetupColors(TSharedRef<FStyle>& Style)
+{
+	Style->Set("Colors.Black", FStyleColors::Black);
+	Style->Set("Colors.Title", FStyleColors::Title);
+	Style->Set("Colors.Foldout", FStyleColors::Foldout);
+	Style->Set("Colors.Input", FStyleColors::Input);
+	Style->Set("Colors.Background", FStyleColors::Background);
+	Style->Set("Colors.Header", FStyleColors::Header);
+	Style->Set("Colors.Dropdown", FStyleColors::Dropdown);
+	Style->Set("Colors.Hover", FStyleColors::Hover);
+	Style->Set("Colors.Hover2", FStyleColors::Hover2);
+	Style->Set("Colors.White", FStyleColors::White);
+	Style->Set("Colors.White25", FStyleColors::White25);
+	Style->Set("Colors.Highlight", FStyleColors::Highlight);
+
+	Style->Set("Colors.Foreground", FStyleColors::Foreground);
+	Style->Set("Colors.ForegroundHover", FStyleColors::ForegroundHover);
+	Style->Set("Colors.ForegroundInverted", FStyleColors::ForegroundInverted);
+	Style->Set("Colors.ForegroundHeader", FStyleColors::ForegroundHeader);
+
+	Style->Set("Colors.Select", FStyleColors::Select);
+	Style->Set("Colors.SelectInactive", FStyleColors::SelectInactive);
+	Style->Set("Colors.SelectParent", FStyleColors::SelectParent);
+	Style->Set("Colors.SelectHover", FStyleColors::SelectHover);
+
+
+	Style->Set("Colors.Primary", FStyleColors::Primary);
+	Style->Set("Colors.PrimaryHover", FStyleColors::PrimaryHover);
+	Style->Set("Colors.PrimaryPress", FStyleColors::PrimaryPress);
+	Style->Set("Colors.AccentBlue", FStyleColors::AccentBlue);
+	Style->Set("Colors.AccentPurple", FStyleColors::AccentPurple);
+	Style->Set("Colors.AccentPink", FStyleColors::AccentPink);
+	Style->Set("Colors.AccentRed", FStyleColors::AccentRed);
+	Style->Set("Colors.AccentOrange", FStyleColors::AccentOrange);
+	Style->Set("Colors.AccentYellow", FStyleColors::AccentYellow);
+	Style->Set("Colors.AccentGreen", FStyleColors::AccentGreen);
+	Style->Set("Colors.AccentBrown", FStyleColors::AccentBrown);
+	Style->Set("Colors.AccentBlack", FStyleColors::AccentBlack);
+	Style->Set("Colors.AccentGray", FStyleColors::AccentGray);
+	Style->Set("Colors.AccentWhite", FStyleColors::AccentWhite);
+	Style->Set("Colors.AccentFolder", FStyleColors::AccentFolder);
+
+
+	Style->Set("Brushes.Black", new FSlateColorBrush(FStyleColors::Black));
+	Style->Set("Brushes.Title", new FSlateColorBrush(FStyleColors::Title));
+	Style->Set("Brushes.Foldout", new FSlateColorBrush(FStyleColors::Foldout));
+	Style->Set("Brushes.Input", new FSlateColorBrush(FStyleColors::Input));
+	Style->Set("Brushes.Background", new FSlateColorBrush(FStyleColors::Background));
+	Style->Set("Brushes.Header", new FSlateColorBrush(FStyleColors::Header));
+	Style->Set("Brushes.Dropdown", new FSlateColorBrush(FStyleColors::Dropdown));
+	Style->Set("Brushes.Hover", new FSlateColorBrush(FStyleColors::Hover));
+	Style->Set("Brushes.Hover2", new FSlateColorBrush(FStyleColors::Hover2));
+	Style->Set("Brushes.White", new FSlateColorBrush(FStyleColors::White));
+	Style->Set("Brushes.White25", new FSlateColorBrush(FStyleColors::White25));
+	Style->Set("Brushes.Highlight", new FSlateColorBrush(FStyleColors::Highlight));
+
+	Style->Set("Brushes.Foreground", new FSlateColorBrush(FStyleColors::Foreground));
+	Style->Set("Brushes.ForegroundHover", new FSlateColorBrush(FStyleColors::ForegroundHover));
+	Style->Set("Brushes.ForegroundInverted", new FSlateColorBrush(FStyleColors::ForegroundInverted));
+	Style->Set("Brushes.ForegroundHeader", new FSlateColorBrush(FStyleColors::ForegroundHeader));
+
+	Style->Set("Brushes.Select", new FSlateColorBrush(FStyleColors::Select));
+	Style->Set("Brushes.SelectInactive", new FSlateColorBrush(FStyleColors::SelectInactive));
+	Style->Set("Brushes.SelectParent", new FSlateColorBrush(FStyleColors::SelectParent));
+	Style->Set("Brushes.SelectHover", new FSlateColorBrush(FStyleColors::SelectHover));
+
+	Style->Set("Brushes.Primary", new FSlateColorBrush(FStyleColors::Primary));
+	Style->Set("Brushes.PrimaryHover", new FSlateColorBrush(FStyleColors::PrimaryHover));
+	Style->Set("Brushes.PrimaryPress", new FSlateColorBrush(FStyleColors::PrimaryPress));
+	Style->Set("Brushes.AccentBlue", new FSlateColorBrush(FStyleColors::AccentBlue));
+	Style->Set("Brushes.AccentPurple", new FSlateColorBrush(FStyleColors::AccentPurple));
+	Style->Set("Brushes.AccentPink", new FSlateColorBrush(FStyleColors::AccentPink));
+	Style->Set("Brushes.AccentRed", new FSlateColorBrush(FStyleColors::AccentRed));
+	Style->Set("Brushes.AccentOrange", new FSlateColorBrush(FStyleColors::AccentOrange));
+	Style->Set("Brushes.AccentYellow", new FSlateColorBrush(FStyleColors::AccentYellow));
+	Style->Set("Brushes.AccentGreen", new FSlateColorBrush(FStyleColors::AccentGreen));
+	Style->Set("Brushes.AccentBrown", new FSlateColorBrush(FStyleColors::AccentBrown));
+	Style->Set("Brushes.AccentBlack", new FSlateColorBrush(FStyleColors::AccentBlack));
+	Style->Set("Brushes.AccentGray", new FSlateColorBrush(FStyleColors::AccentGray));
+	Style->Set("Brushes.AccentWhite", new FSlateColorBrush(FStyleColors::AccentWhite));
+	Style->Set("Brushes.AccentFolder", new FSlateColorBrush(FStyleColors::AccentFolder));
+
+}
+
+void FStarshipCoreStyle::SetupTextStyles(TSharedRef<FStyle>& Style)
+{
+	const FScrollBarStyle& ScrollBar = Style->GetWidgetStyle<FScrollBarStyle>("Scrollbar");
+
+	const FStyleFonts& StyleFonts = FStyleFonts::Get();
+
+	const FTextBlockStyle NormalText = FTextBlockStyle()
+		.SetFont(StyleFonts.Normal)
+		.SetColorAndOpacity(FSlateColor::UseForeground())
+		.SetShadowOffset(FVector2D::ZeroVector)
+		.SetSelectedBackgroundColor(FStyleColors::Highlight)
+		.SetHighlightShape(FSlateColorBrush(FStyleColors::White));
+
+	Style->Set("NormalFont", StyleFonts.Normal);
+	Style->Set("SmallFont", StyleFonts.Small);
+	Style->Set("NormalFontBold", StyleFonts.NormalBold);
+	Style->Set("SmallFontBold", StyleFonts.SmallBold);
+
+	FSlateBrush* DefaultTextUnderlineBrush = new IMAGE_BRUSH("Old/White", Icon8x8, FLinearColor::White, ESlateBrushTileType::Both);
+
+	Style->Set("DefaultTextUnderline", DefaultTextUnderlineBrush);
+
+	const FTextBlockStyle NormalUnderlinedText = FTextBlockStyle(NormalText)
+		.SetUnderlineBrush(*DefaultTextUnderlineBrush);
+
+	// Monospaced Text
+	const FTextBlockStyle MonospacedText = FTextBlockStyle()
+		.SetFont(DEFAULT_FONT("Mono", 10))
+		.SetColorAndOpacity(FSlateColor::UseForeground())
+		.SetShadowOffset(FVector2D::ZeroVector)
+		.SetShadowColorAndOpacity(FLinearColor::Black)
+		.SetHighlightColor(FLinearColor(0.02f, 0.3f, 0.0f))
+		.SetHighlightShape(BOX_BRUSH("Common/TextBlockHighlightShape", FMargin(3.f / 8.f))
+		);
+
+	const FTextBlockStyle MonospacedUnderlinedText = FTextBlockStyle(MonospacedText)
+		.SetUnderlineBrush(*DefaultTextUnderlineBrush);
+
+	Style->Set("MonospacedText", MonospacedText);
+	Style->Set("MonospacedUnderlinedText", MonospacedUnderlinedText);
+
+	// Small Text
+	const FTextBlockStyle SmallText = FTextBlockStyle(NormalText)
+		.SetFont(StyleFonts.Small);
+
+	const FTextBlockStyle SmallUnderlinedText = FTextBlockStyle(SmallText)
+		.SetUnderlineBrush(*DefaultTextUnderlineBrush);
+
+	// Embossed Text
+	Style->Set("EmbossedText", FTextBlockStyle(NormalText)
+		.SetFont(DEFAULT_FONT("Regular", 24))
+		.SetColorAndOpacity(FLinearColor::Black)
+		.SetShadowOffset(FVector2D(0.0f, 1.0f))
+		.SetShadowColorAndOpacity(FLinearColor(0.8f, 0.8f, 0.8f, 0.5))
+	);
+
+	const FEditableTextBoxStyle DarkEditableTextBoxStyle = FEditableTextBoxStyle()
+		.SetBackgroundImageNormal(BOX_BRUSH("Common/TextBox_Dark", FMargin(4.0f / 16.0f)))
+		.SetBackgroundImageHovered(BOX_BRUSH("Common/TextBox_Hovered_Dark", FMargin(4.0f / 16.0f)))
+		.SetBackgroundImageFocused(BOX_BRUSH("Common/TextBox_Hovered_Dark", FMargin(4.0f / 16.0f)))
+		.SetBackgroundImageReadOnly(BOX_BRUSH("Common/TextBox_ReadOnly", FMargin(4.0f / 16.0f)))
+		.SetScrollBarStyle(ScrollBar);
+	{
+		Style->Set("DarkEditableTextBox", DarkEditableTextBoxStyle);
+		// "NormalFont".
+	}
 
 
+	// STextBlock defaults...
+	{
+		Style->Set("NormalText", NormalText);
+		Style->Set("NormalUnderlinedText", NormalUnderlinedText);
 
-		/* Style for a toggleable button that mimics the coloring and look of a Table Row */
-		/*
-		const FCheckBoxStyle ToggleButtonRowStyle = FCheckBoxStyle()
-			.SetCheckBoxType(ESlateCheckBoxType::ToggleButton)
-			.SetUncheckedImage(FSlateNoResource())
-			.SetUncheckedHoveredImage(IMAGE_BRUSH("Common/Selection", Icon8x8, SelectionColor_Inactive))
-			.SetUncheckedPressedImage(IMAGE_BRUSH("Common/Selection", Icon8x8, SelectionColor_Inactive))
-			.SetCheckedImage(IMAGE_BRUSH("Common/Selection", Icon8x8, SelectionColor))
-			.SetCheckedHoveredImage(IMAGE_BRUSH("Common/Selection", Icon8x8, SelectionColor))
-			.SetCheckedPressedImage(BOX_BRUSH("Common/Selector", 4.0f / 16.0f, SelectorColor));
-		Style->Set("ToggleButtonRowStyle", ToggleButtonRowStyle);
-		*/
-
-		/* A radio button is actually just a SCheckBox box with different images */
-		/* Set images for various radio button (SCheckBox) states ... */
-		const FCheckBoxStyle BasicRadioButtonStyle = FCheckBoxStyle()
-			.SetCheckBoxType(ESlateCheckBoxType::CheckBox)
-
-			.SetForegroundColor(               FLinearColor::White )
-			.SetHoveredForegroundColor(        FLinearColor::White )
-			.SetPressedForegroundColor(        FLinearColor::White )
-			.SetCheckedForegroundColor(        FLinearColor::White )
-			.SetCheckedHoveredForegroundColor( FLinearColor::White )
-			.SetCheckedPressedForegroundColor( FLinearColor::White )
-			.SetUndeterminedForegroundColor(   FLinearColor::White )
-
-			.SetUncheckedImage(           IMAGE_BRUSH_SVG("/Starship/CoreWidgets/CheckBox/radio-off", Icon16x16, FStyleColors::White25))
-			.SetUncheckedHoveredImage(    IMAGE_BRUSH_SVG("/Starship/CoreWidgets/CheckBox/radio-off", Icon16x16, FStyleColors::ForegroundHover))
-			.SetUncheckedPressedImage(    IMAGE_BRUSH_SVG("/Starship/CoreWidgets/CheckBox/radio-off", Icon16x16, FStyleColors::ForegroundHover))
-			.SetCheckedImage(             IMAGE_BRUSH_SVG("/Starship/CoreWidgets/CheckBox/radio-on",  Icon16x16, FStyleColors::Primary))
-			.SetCheckedHoveredImage(      IMAGE_BRUSH_SVG("/Starship/CoreWidgets/CheckBox/radio-on",  Icon16x16, FStyleColors::Primary))
-			.SetCheckedPressedImage(      IMAGE_BRUSH_SVG("/Starship/CoreWidgets/CheckBox/radio-off", Icon16x16, FStyleColors::Primary))
-			.SetUndeterminedImage(        IMAGE_BRUSH_SVG("/Starship/CoreWidgets/CheckBox/radio-off", Icon16x16, FStyleColors::White25))
-			.SetUndeterminedHoveredImage( IMAGE_BRUSH_SVG("/Starship/CoreWidgets/CheckBox/radio-off", Icon16x16, FStyleColors::ForegroundHover))
-			.SetUndeterminedPressedImage( IMAGE_BRUSH_SVG("/Starship/CoreWidgets/CheckBox/radio-off", Icon16x16, FStyleColors::ForegroundHover))
-			.SetPadding(FMargin(4.0));
-		Style->Set("RadioButton", BasicRadioButtonStyle);
+		Style->Set("SmallText", SmallText);
+		Style->Set("SmallUnderlinedText", SmallUnderlinedText);
 	}
 
+
 	// SEditableText defaults...
 	{
 		FSlateBrush* SelectionBackground = new FSlateColorBrush(FStyleColors::Highlight);
@@ -854,7 +934,8 @@ TSharedRef<ISlateStyle> FStarshipCoreStyle::Create()
 		const FEditableTextStyle NormalEditableTextStyle = FEditableTextStyle()
 			.SetBackgroundImageSelected(*SelectionBackground)
 			.SetBackgroundImageComposing(*CompositionBackground)
-			.SetCaretImage(*GenericWhiteBox);
+			.SetCaretImage(FSlateColorBrush(FStyleColors::White));
+
 		Style->Set("NormalEditableText", NormalEditableTextStyle);
 
 		Style->Set("EditableText.SelectionBackground", SelectionBackground);
@@ -876,27 +957,19 @@ TSharedRef<ISlateStyle> FStarshipCoreStyle::Create()
 		.SetScrollBarStyle(ScrollBar);
 	{
 		Style->Set("NormalEditableTextBox", NormalEditableTextBoxStyle);
-		// "NormalFont".
 	}
 
-	const FEditableTextBoxStyle DarkEditableTextBoxStyle = FEditableTextBoxStyle()
-		.SetBackgroundImageNormal(BOX_BRUSH("Common/TextBox_Dark", FMargin(4.0f / 16.0f)))
-		.SetBackgroundImageHovered(BOX_BRUSH("Common/TextBox_Hovered_Dark", FMargin(4.0f / 16.0f)))
-		.SetBackgroundImageFocused(BOX_BRUSH("Common/TextBox_Hovered_Dark", FMargin(4.0f / 16.0f)))
-		.SetBackgroundImageReadOnly(BOX_BRUSH("Common/TextBox_ReadOnly", FMargin(4.0f / 16.0f)))
-		.SetScrollBarStyle(ScrollBar);
-	{
-		Style->Set("DarkEditableTextBox", DarkEditableTextBoxStyle);
-		// "NormalFont".
-	}
-
-	// STextBlock defaults...
+	// SSearchBox defaults...
 	{
-		Style->Set("NormalText", NormalText);
-		Style->Set("NormalUnderlinedText", NormalUnderlinedText);
-
-		Style->Set("SmallText", SmallText);
-		Style->Set("SmallUnderlinedText", SmallUnderlinedText);
+		Style->Set("SearchBox", FSearchBoxStyle()
+			.SetTextBoxStyle(NormalEditableTextBoxStyle)
+			.SetUpArrowImage(IMAGE_BRUSH_SVG("Starship/Common/arrow-north", Icon8x8, FStyleColors::Foreground))
+			.SetDownArrowImage(IMAGE_BRUSH_SVG("Starship/Common/arrow-south", Icon8x8, FStyleColors::Foreground))
+			.SetGlassImage(IMAGE_BRUSH_SVG("Starship/Common/search", Icon16x16, FStyleColors::Foreground))
+			.SetClearImage(IMAGE_BRUSH_SVG("Starship/Common/close", Icon16x16, FStyleColors::Foreground))
+			.SetImagePadding(FMargin(3.f, 2.f, -2.f, 2.0))
+			.SetLeftAlignButtons(true)
+		);
 	}
 
 	// SInlineEditableTextBlock
@@ -905,7 +978,7 @@ TSharedRef<ISlateStyle> FStarshipCoreStyle::Create()
 		FTextBlockStyle InlineEditableTextBlockReadOnly = FTextBlockStyle(NormalText);
 
 		FEditableTextBoxStyle InlineEditableTextBlockEditable = FEditableTextBoxStyle(NormalEditableTextBoxStyle)
-		.SetPadding(FMargin(6.f, 4.5f, 6.f, 4.5f));
+			.SetPadding(FMargin(6.f, 4.5f, 6.f, 4.5f));
 
 		FInlineEditableTextBlockStyle InlineEditableTextBlockStyle = FInlineEditableTextBlockStyle()
 			.SetTextStyle(InlineEditableTextBlockReadOnly)
@@ -927,153 +1000,471 @@ TSharedRef<ISlateStyle> FStarshipCoreStyle::Create()
 		Style->Set("InlineEditableTextBlockSmallStyle", InlineEditableTextBlockSmallStyle);
 	}
 
-	// SSuggestionTextBox defaults...
-	{
-		Style->Set("SuggestionTextBox.Background", new BOX_BRUSH("Old/Menu_Background", FMargin(8.0f / 64.0f)));
-		Style->Set("SuggestionTextBox.Text", FTextBlockStyle()
-			.SetFont(StyleFonts.Normal)
-			.SetColorAndOpacity(FLinearColor(FColor(0xffaaaaaa)))
-		);
-	}
 
-	// SToolTip defaults...
-	{
-		Style->Set("ToolTip.Font", StyleFonts.Small);
-		Style->Set("ToolTip.Background", new BOX_BRUSH("Old/ToolTip_Background", FMargin(8.0f / 64.0f)));
+}
 
-		Style->Set("ToolTip.LargerFont", StyleFonts.Normal);
-		Style->Set("ToolTip.BrightBackground", new BOX_BRUSH("Old/ToolTip_BrightBackground", FMargin(8.0f / 64.0f)));
-	}
+void FStarshipCoreStyle::SetupButtonStyles(TSharedRef<FStyle>& Style)
+{
+	// SButton defaults
+	const FButtonStyle PrimaryButton = FButtonStyle()
+		.SetNormal(FSlateRoundedBoxBrush(FStyleColors::Primary, 4.0f))
+		.SetHovered(FSlateRoundedBoxBrush(FStyleColors::PrimaryHover, 4.0f))
+		.SetPressed(FSlateRoundedBoxBrush(FStyleColors::PrimaryPress, 4.0f))
+		.SetNormalForeground(FStyleColors::Background)
+		.SetHoveredForeground(FStyleColors::Background)
+		.SetPressedForeground(FStyleColors::Background)
+		.SetDisabledForeground(FStyleColors::Background)
+		.SetNormalPadding(ButtonMargins)
+		.SetPressedPadding(ButtonMargins);
 
-	// SBorder defaults...
-	{
-		Style->Set("Border", new FSlateColorBrush(FStyleColors::Background));
-	}
 
-	// SHyperlink defaults...
-	{
-		FButtonStyle HyperlinkButton = FButtonStyle()
-			.SetNormal(BORDER_BRUSH("Old/HyperlinkDotted", FMargin(0, 0, 0, 3 / 16.0f)))
-			.SetPressed(FSlateNoResource())
-			.SetHovered(BORDER_BRUSH("Old/HyperlinkUnderline", FMargin(0, 0, 0, 3 / 16.0f)));
+	const FButtonStyle Button = FButtonStyle()
+		.SetNormal(FSlateRoundedBoxBrush(FStyleColors::Dropdown, 4.0f))
+		.SetHovered(FSlateRoundedBoxBrush(FStyleColors::Hover, 4.0f))
+		.SetPressed(FSlateRoundedBoxBrush(FStyleColors::Header, 4.0f))
+		.SetNormalForeground(FStyleColors::ForegroundHover)
+		.SetHoveredForeground(FStyleColors::ForegroundHover)
+		.SetPressedForeground(FStyleColors::ForegroundHover)
+		.SetDisabledForeground(FStyleColors::Foreground)
+		.SetNormalPadding(ButtonMargins)
+		.SetPressedPadding(ButtonMargins);
 
-		FHyperlinkStyle Hyperlink = FHyperlinkStyle()
-			.SetUnderlineStyle(HyperlinkButton)
-			.SetTextStyle(NormalText)
-			.SetPadding(FMargin(0.0f));
-		Style->Set("Hyperlink", Hyperlink);
-	}
 
-	// SProgressBar defaults...
+	const FButtonStyle SimpleButton = FButtonStyle()
+		.SetNormal(FSlateNoResource())
+		.SetHovered(FSlateRoundedBoxBrush(FStyleColors::Hover, 4.0f))
+		.SetPressed(FSlateRoundedBoxBrush(FStyleColors::Input, 4.0f))
+		.SetDisabled(FSlateNoResource())
+		.SetNormalForeground(FStyleColors::Foreground)
+		.SetHoveredForeground(FStyleColors::ForegroundHover)
+		.SetPressedForeground(FStyleColors::ForegroundHover)
+		.SetDisabledForeground(FStyleColors::Foreground)
+		.SetNormalPadding(FMargin(4.0f, 2.0f))
+		.SetPressedPadding(FMargin(4.0f, 2.0f));
 	{
-		Style->Set("ProgressBar", FProgressBarStyle()
-			.SetBackgroundImage(BOX_BRUSH("Common/ProgressBar_Background", FMargin(5.f / 12.f)))
-			.SetFillImage(BOX_BRUSH("Common/ProgressBar_Fill", FMargin(5.f / 12.f), FLinearColor(1.0f, 0.22f, 0.0f)))
-			.SetMarqueeImage(IMAGE_BRUSH("Common/ProgressBar_Marquee", FVector2D(20, 12), FLinearColor::White, ESlateBrushTileType::Horizontal))
-		);
-	}
 
-	// SThrobber, SCircularThrobber defaults...
-	{
-		Style->Set("Throbber.Chunk", new IMAGE_BRUSH("Common/Throbber_Piece", FVector2D(16, 16)));
-		Style->Set("Throbber.CircleChunk", new IMAGE_BRUSH("Common/Throbber_Piece", FVector2D(8, 8)));
-	}
+		const FTextBlockStyle& NormalText = Style->GetWidgetStyle<FTextBlockStyle>("NormalText");
 
-	// SExpandableArea defaults...
-	{
-		Style->Set("ExpandableArea", FExpandableAreaStyle()
-			.SetCollapsedImage(IMAGE_BRUSH("Common/TreeArrow_Collapsed", Icon10x10, DefaultForeground))
-			.SetExpandedImage(IMAGE_BRUSH("Common/TreeArrow_Expanded", Icon10x10, DefaultForeground))
-		);
-		Style->Set("ExpandableArea.TitleFont", StyleFonts.SmallBold);
-		Style->Set("ExpandableArea.Border", new FSlateRoundedBoxBrush(FStyleColors::Background, 4) );
+		Style->Set("ButtonText", FTextBlockStyle(NormalText).SetFont(FStyleFonts::Get().SmallBold));
+
+		Style->Set("PrimaryButton", PrimaryButton);
+		Style->Set("Button", Button);
+		Style->Set("SimpleButton", SimpleButton);
 	}
+}
 
-	// SSearchBox defaults...
+void FStarshipCoreStyle::SetupComboButtonStyles(TSharedRef<FStyle>& Style)
+{
+	// SComboButton and SComboBox defaults...
+	const FButtonStyle ComboButtonButton = FButtonStyle()
+		.SetNormal(FSlateRoundedBoxBrush(FStyleColors::Input, InputFocusRadius, FStyleColors::Dropdown, InputFocusThickness))
+		.SetHovered(FSlateRoundedBoxBrush(FStyleColors::Input, InputFocusRadius, FStyleColors::Hover, InputFocusThickness))
+		.SetPressed(FSlateRoundedBoxBrush(FStyleColors::Input, InputFocusRadius, FStyleColors::Hover, InputFocusThickness))
+		.SetNormalForeground(FStyleColors::Foreground)
+		.SetHoveredForeground(FStyleColors::ForegroundHover)
+		.SetPressedForeground(FStyleColors::ForegroundHover)
+		.SetDisabledForeground(FStyleColors::White25)
+		.SetNormalPadding(FMargin(8.f, 2.f, 4.f, 2.f))
+		.SetPressedPadding(FMargin(8.f, 2.f, 4.f, 2.f));
+
+
+	// SComboBox 
+	FComboButtonStyle ComboButton = FComboButtonStyle()
+		.SetButtonStyle(ComboButtonButton)
+		.SetContentPadding(0.f)
+		.SetDownArrowImage(IMAGE_BRUSH_SVG("Starship/CoreWidgets/ComboBox/wide-chevron-down", FVector2D(20.f, 16.f)))
+		.SetMenuBorderBrush(FSlateRoundedBoxBrush(FStyleColors::Dropdown, 0.0, WindowHighlight, 1.0))
+		.SetMenuBorderPadding(0.f);
+	Style->Set("ComboButton", ComboButton);
+
+
+	FComboBoxStyle ComboBox = FComboBoxStyle()
+		.SetContentPadding(0.f)
+		.SetMenuRowPadding(FMargin(8.0f, 3.f))
+		.SetComboButtonStyle(ComboButton);
+	Style->Set("ComboBox", ComboBox);
+
+	const FButtonStyle& SimpleButton = Style->GetWidgetStyle<FButtonStyle>("SimpleButton");
+
+	// Simple Combo Box  (borderless)
+	FComboButtonStyle SimpleComboButton = FComboButtonStyle()
+		.SetButtonStyle(SimpleButton)
+		.SetContentPadding(0.f)
+		.SetDownArrowImage(IMAGE_BRUSH_SVG("Starship/Common/chevron-down", Icon16x16))
+		.SetMenuBorderBrush(FSlateColorBrush(FStyleColors::Dropdown))
+		.SetMenuBorderPadding(0.0f);
+	Style->Set("SimpleComboButton", SimpleComboButton);
+
+
+	FComboBoxStyle SimpleComboBox = FComboBoxStyle()
+		.SetContentPadding(0.f)
+		.SetMenuRowPadding(FMargin(8.0f, 2.f))
+		.SetComboButtonStyle(SimpleComboButton);
+	Style->Set("SimpleComboBox", SimpleComboBox);
+
+
+	const FTableRowStyle ComboBoxRow = FTableRowStyle()
+		.SetEvenRowBackgroundBrush(FSlateNoResource())
+		.SetEvenRowBackgroundHoveredBrush(FSlateColorBrush(FStyleColors::Hover))
+
+		.SetOddRowBackgroundBrush(FSlateNoResource())
+		.SetOddRowBackgroundHoveredBrush(FSlateColorBrush(FStyleColors::Hover))
+
+		.SetSelectorFocusedBrush(FSlateNoResource())
+
+		.SetActiveBrush(FSlateColorBrush(FStyleColors::Primary))
+		.SetActiveHoveredBrush(FSlateColorBrush(FStyleColors::PrimaryHover))
+		.SetActiveHighlightedBrush(FSlateColorBrush(FStyleColors::PrimaryHover))
+
+		.SetInactiveBrush(FSlateNoResource())
+		.SetInactiveHoveredBrush(FSlateNoResource())
+		.SetInactiveHighlightedBrush(FSlateNoResource())
+
+
+		.SetTextColor(FStyleColors::White)
+		.SetSelectedTextColor(FStyleColors::Input)
+
+		.SetDropIndicator_Above(FSlateNoResource())
+		.SetDropIndicator_Onto(FSlateNoResource())
+		.SetDropIndicator_Below(FSlateNoResource());
+
+	Style->Set("ComboBox.Row", ComboBoxRow);
+
+	// SEditableComboBox defaults...
 	{
-		Style->Set("SearchBox", FSearchBoxStyle()
-			.SetTextBoxStyle(NormalEditableTextBoxStyle)
-			.SetUpArrowImage(  IMAGE_BRUSH_SVG("Starship/Common/arrow-north", Icon8x8, FStyleColors::Foreground))
-			.SetDownArrowImage(IMAGE_BRUSH_SVG("Starship/Common/arrow-south", Icon8x8, FStyleColors::Foreground))
-			.SetGlassImage(    IMAGE_BRUSH_SVG("Starship/Common/search", Icon16x16, FStyleColors::Foreground))
-			.SetClearImage(    IMAGE_BRUSH_SVG("Starship/Common/close", Icon16x16, FStyleColors::Foreground))
-			.SetImagePadding(FMargin(3.f, 2.f, -2.f, 2.0))
-			.SetLeftAlignButtons(true)
-		);
+		Style->Set("EditableComboBox.Add", new IMAGE_BRUSH("Icons/PlusSymbol_12x", Icon12x12));
+		Style->Set("EditableComboBox.Delete", new IMAGE_BRUSH("Icons/Cross_12x", Icon12x12));
+		Style->Set("EditableComboBox.Rename", new IMAGE_BRUSH("Icons/ellipsis_12x", Icon12x12));
+		Style->Set("EditableComboBox.Accept", new IMAGE_BRUSH("Common/Check", Icon16x16));
 	}
 
-	// SSlider and SVolumeControl defaults...
-	{
-		FSliderStyle SliderStyle = FSliderStyle()
-			.SetNormalBarImage(   FSlateRoundedBoxBrush(FStyleColors::Input, 2.0))
-			.SetHoveredBarImage(  FSlateRoundedBoxBrush(FStyleColors::Input, 2.0))
-			.SetNormalThumbImage(  FSlateRoundedBoxBrush(FStyleColors::Hover2, Icon8x8) )
-			.SetHoveredThumbImage( FSlateRoundedBoxBrush(FStyleColors::ForegroundHover, Icon8x8) )
-			.SetBarThickness(4.0f);
-		Style->Set("Slider", SliderStyle);
+	// SComboBox for Icons Only
+	// FComboButtonStyle IconComboButton = FComboButtonStyle()
+	// Style->Set("IconComboButton")
+}
 
-		Style->Set("VolumeControl", FVolumeControlStyle()
-			.SetSliderStyle(SliderStyle)
-			.SetHighVolumeImage(IMAGE_BRUSH("Common/VolumeControl_High", Icon16x16))
-			.SetMidVolumeImage(IMAGE_BRUSH("Common/VolumeControl_Mid", Icon16x16))
-			.SetLowVolumeImage(IMAGE_BRUSH("Common/VolumeControl_Low", Icon16x16))
-			.SetNoVolumeImage(IMAGE_BRUSH("Common/VolumeControl_Off", Icon16x16))
-			.SetMutedImage(IMAGE_BRUSH("Common/VolumeControl_Muted", Icon16x16))
-		);
-	}
+void FStarshipCoreStyle::SetupCheckboxStyles(TSharedRef<FStyle>& Style)
+{
+	// SCheckBox defaults...
+	const float CheckboxCornerRadius = 3.f;
+	const float CheckboxOutlineThickness = 1.0f;
+
+	const FCheckBoxStyle BasicCheckBoxStyle = FCheckBoxStyle()
+		.SetCheckBoxType(ESlateCheckBoxType::CheckBox)
+
+		.SetForegroundColor(FLinearColor::White)
+		.SetHoveredForegroundColor(FLinearColor::White)
+		.SetPressedForegroundColor(FLinearColor::White)
+		.SetCheckedForegroundColor(FLinearColor::White)
+		.SetCheckedHoveredForegroundColor(FLinearColor::White)
+		.SetCheckedPressedForegroundColor(FLinearColor::White)
+		.SetUndeterminedForegroundColor(FLinearColor::White)
+
+		.SetUncheckedImage(FSlateNoResource())
+		.SetUncheckedHoveredImage(FSlateNoResource())
+		.SetUncheckedPressedImage(FSlateNoResource())
+		.SetCheckedImage(IMAGE_BRUSH_SVG("Starship/CoreWidgets/CheckBox/check", Icon16x16, FStyleColors::Primary))
+		.SetCheckedHoveredImage(IMAGE_BRUSH_SVG("Starship/CoreWidgets/CheckBox/check", Icon16x16, FStyleColors::PrimaryHover))
+		.SetCheckedPressedImage(IMAGE_BRUSH_SVG("Starship/CoreWidgets/CheckBox/check", Icon16x16, FStyleColors::Primary))
+		.SetUndeterminedImage(IMAGE_BRUSH_SVG("Starship/CoreWidgets/CheckBox/indeterminate", Icon16x16, FStyleColors::Primary))
+		.SetUndeterminedHoveredImage(IMAGE_BRUSH_SVG("Starship/CoreWidgets/CheckBox/indeterminate", Icon16x16, FStyleColors::PrimaryHover))
+		.SetUndeterminedPressedImage(IMAGE_BRUSH_SVG("Starship/CoreWidgets/CheckBox/indeterminate", Icon16x16, FStyleColors::Primary))
+		.SetBackgroundImage(FSlateRoundedBoxBrush(FStyleColors::Input, CheckboxCornerRadius, FStyleColors::Dropdown, CheckboxOutlineThickness, Icon18x18))
+		.SetBackgroundHoveredImage(FSlateRoundedBoxBrush(FStyleColors::Input, CheckboxCornerRadius, FStyleColors::Hover, CheckboxOutlineThickness, Icon18x18))
+		.SetBackgroundPressedImage(FSlateRoundedBoxBrush(FStyleColors::Foldout, CheckboxCornerRadius, FStyleColors::Hover, CheckboxOutlineThickness, Icon18x18));
+
+	Style->Set("Checkbox", BasicCheckBoxStyle);
+
+	const FCheckBoxStyle SimplifiedCheckBoxStyle = FCheckBoxStyle()
+		.SetCheckBoxType(ESlateCheckBoxType::CheckBox)
+
+		.SetForegroundColor(FLinearColor::White)
+		.SetHoveredForegroundColor(FLinearColor::White)
+		.SetPressedForegroundColor(FLinearColor::White)
+		.SetCheckedForegroundColor(FLinearColor::White)
+		.SetCheckedHoveredForegroundColor(FLinearColor::White)
+		.SetCheckedPressedForegroundColor(FLinearColor::White)
+		.SetUndeterminedForegroundColor(FLinearColor::White)
+
+		.SetUncheckedImage(FSlateNoResource())
+		.SetUncheckedHoveredImage(FSlateNoResource())
+		.SetUncheckedPressedImage(FSlateNoResource())
+		.SetCheckedImage(IMAGE_BRUSH_SVG("Starship/CoreWidgets/CheckBox/check", Icon16x16, FStyleColors::Foreground))
+		.SetCheckedHoveredImage(IMAGE_BRUSH_SVG("Starship/CoreWidgets/CheckBox/check", Icon16x16, FStyleColors::ForegroundHover))
+		.SetCheckedPressedImage(IMAGE_BRUSH_SVG("Starship/CoreWidgets/CheckBox/check", Icon16x16, FStyleColors::Foreground))
+		.SetUndeterminedImage(IMAGE_BRUSH_SVG("Starship/CoreWidgets/CheckBox/indeterminate", Icon16x16, FStyleColors::Foreground))
+		.SetUndeterminedHoveredImage(IMAGE_BRUSH_SVG("Starship/CoreWidgets/CheckBox/indeterminate", Icon16x16, FStyleColors::ForegroundHover))
+		.SetUndeterminedPressedImage(IMAGE_BRUSH_SVG("Starship/CoreWidgets/CheckBox/indeterminate", Icon16x16, FStyleColors::Foreground))
+		.SetBackgroundImage(FSlateRoundedBoxBrush(FStyleColors::Input, CheckboxCornerRadius, FStyleColors::Dropdown, CheckboxOutlineThickness, Icon18x18))
+		.SetBackgroundHoveredImage(FSlateRoundedBoxBrush(FStyleColors::Input, CheckboxCornerRadius, FStyleColors::Hover, CheckboxOutlineThickness, Icon18x18))
+		.SetBackgroundPressedImage(FSlateRoundedBoxBrush(FStyleColors::Foldout, CheckboxCornerRadius, FStyleColors::Hover, CheckboxOutlineThickness, Icon18x18));
+
+	Style->Set("SimplifiedCheckbox", SimplifiedCheckBoxStyle);
+
+	/* Set images for various transparent SCheckBox states ... */
+	const FCheckBoxStyle BasicTransparentCheckBoxStyle = FCheckBoxStyle()
+		.SetCheckBoxType(ESlateCheckBoxType::ToggleButton)
+		.SetUncheckedImage(FSlateNoResource())
+		.SetUncheckedHoveredImage(FSlateNoResource())
+		.SetUncheckedPressedImage(FSlateNoResource())
+		.SetCheckedImage(FSlateNoResource())
+		.SetCheckedHoveredImage(FSlateNoResource())
+		.SetCheckedPressedImage(FSlateNoResource())
+		.SetUndeterminedImage(FSlateNoResource())
+		.SetUndeterminedHoveredImage(FSlateNoResource())
+		.SetUndeterminedPressedImage(FSlateNoResource());
+	Style->Set("TransparentCheckBox", BasicTransparentCheckBoxStyle);
+
+	/* Default Style for a toggleable button */
+	const FCheckBoxStyle ToggleButtonStyle = FCheckBoxStyle()
+		.SetCheckBoxType(ESlateCheckBoxType::ToggleButton)
+		.SetCheckedImage(FSlateNoResource())
+		.SetCheckedHoveredImage(FSlateRoundedBoxBrush(FStyleColors::Hover, InputFocusRadius))
+		.SetCheckedPressedImage(FSlateRoundedBoxBrush(FStyleColors::Input, InputFocusRadius))
+		.SetUncheckedImage(FSlateNoResource())
+		.SetUncheckedHoveredImage(FSlateRoundedBoxBrush(FStyleColors::Hover, InputFocusRadius))
+		.SetUncheckedPressedImage(FSlateRoundedBoxBrush(FStyleColors::Input, InputFocusRadius))
+		.SetForegroundColor(FStyleColors::Foreground)
+		.SetHoveredForegroundColor(FStyleColors::ForegroundHover)
+		.SetPressedForegroundColor(FStyleColors::ForegroundHover)
+		.SetCheckedForegroundColor(FStyleColors::Primary)
+		.SetCheckedHoveredForegroundColor(FStyleColors::PrimaryHover)
+		.SetPadding(DefaultMargins)
+		;
 
-	// SSpinBox defaults...
-	{
-		Style->Set("SpinBox", FSpinBoxStyle()
-			.SetBackgroundBrush(FSlateRoundedBoxBrush(FStyleColors::Input, InputFocusRadius, FStyleColors::Dropdown, InputFocusThickness))
-			.SetHoveredBackgroundBrush(FSlateRoundedBoxBrush(FStyleColors::Input, InputFocusRadius, FStyleColors::Hover, InputFocusThickness))
+	Style->Set("ToggleButtonCheckbox", ToggleButtonStyle);
+
+	/* Default Style for a toggleable button */
+	const FCheckBoxStyle ToggleButtonAltStyle = FCheckBoxStyle()
+		.SetCheckBoxType(ESlateCheckBoxType::ToggleButton)
+		.SetCheckedImage(FSlateRoundedBoxBrush(FStyleColors::Input, InputFocusRadius))
+		.SetCheckedHoveredImage(FSlateRoundedBoxBrush(FStyleColors::Hover, InputFocusRadius))
+		.SetCheckedPressedImage(FSlateRoundedBoxBrush(FStyleColors::Input, InputFocusRadius))
+		.SetUncheckedImage(FSlateNoResource())
+		.SetUncheckedHoveredImage(FSlateRoundedBoxBrush(FStyleColors::Hover, InputFocusRadius))
+		.SetUncheckedPressedImage(FSlateRoundedBoxBrush(FStyleColors::Input, InputFocusRadius))
+		.SetForegroundColor(FStyleColors::Foreground)
+		.SetHoveredForegroundColor(FStyleColors::ForegroundHover)
+		.SetPressedForegroundColor(FStyleColors::ForegroundHover)
+		.SetCheckedForegroundColor(FStyleColors::Primary)
+		.SetCheckedHoveredForegroundColor(FStyleColors::PrimaryHover)
+		.SetPadding(DefaultMargins)
+		;
 
-			.SetActiveFillBrush(FSlateRoundedBoxBrush(FStyleColors::Hover, InputFocusRadius, FLinearColor::Transparent, InputFocusThickness))
-			.SetInactiveFillBrush(FSlateRoundedBoxBrush(FStyleColors::Dropdown, InputFocusRadius, FLinearColor::Transparent, InputFocusThickness))
-			.SetArrowsImage(FSlateNoResource())
-			.SetForegroundColor(FStyleColors::ForegroundHover)
-			.SetTextPadding(FMargin(10.f, 3.5f, 10.f, 4.f))
-		);
-	}
+	Style->Set("ToggleButtonCheckboxAlt", ToggleButtonAltStyle);
+
+	/* Style for a segmented box */
+	const FCheckBoxStyle SegmentedBoxLeft = FCheckBoxStyle()
+		.SetCheckBoxType(ESlateCheckBoxType::ToggleButton)
+
+		.SetUncheckedImage(BOX_BRUSH("/Starship/CoreWidgets/SegmentedBox/left", FVector2D(16.f, 16.f), FMargin(4.0f / 16.0f), FStyleColors::Dropdown))
+		.SetUncheckedHoveredImage(BOX_BRUSH("/Starship/CoreWidgets/SegmentedBox/left", FVector2D(16.f, 16.f), FMargin(4.0f / 16.0f), FStyleColors::Hover))
+		.SetUncheckedPressedImage(BOX_BRUSH("/Starship/CoreWidgets/SegmentedBox/left", FVector2D(16.f, 16.f), FMargin(4.0f / 16.0f), FStyleColors::Dropdown))
+		.SetCheckedImage(BOX_BRUSH("/Starship/CoreWidgets/SegmentedBox/left", FVector2D(16.f, 16.f), FMargin(4.0f / 16.0f), FStyleColors::Input))
+		.SetCheckedHoveredImage(BOX_BRUSH("/Starship/CoreWidgets/SegmentedBox/left", FVector2D(16.f, 16.f), FMargin(4.0f / 16.0f), FStyleColors::Input))
+		.SetCheckedPressedImage(BOX_BRUSH("/Starship/CoreWidgets/SegmentedBox/left", FVector2D(16.f, 16.f), FMargin(4.0f / 16.0f), FStyleColors::Input))
+		.SetForegroundColor(FStyleColors::Foreground)
+		.SetHoveredForegroundColor(FStyleColors::ForegroundHover)
+		.SetPressedForegroundColor(FStyleColors::ForegroundHover)
+		.SetCheckedForegroundColor(FStyleColors::Primary)
+		.SetCheckedHoveredForegroundColor(FStyleColors::Primary)
+		.SetCheckedPressedForegroundColor(FStyleColors::Primary)
+		.SetPadding(DefaultMargins);
+
+
+	const FCheckBoxStyle SegmentedBoxCenter = FCheckBoxStyle()
+		.SetCheckBoxType(ESlateCheckBoxType::ToggleButton)
+		.SetUncheckedImage(FSlateColorBrush(FStyleColors::Dropdown))
+		.SetUncheckedHoveredImage(FSlateColorBrush(FStyleColors::Hover))
+		.SetUncheckedPressedImage(FSlateColorBrush(FStyleColors::Dropdown))
+		.SetCheckedImage(FSlateColorBrush(FStyleColors::Input))
+		.SetCheckedHoveredImage(FSlateColorBrush(FStyleColors::Input))
+		.SetCheckedPressedImage(FSlateColorBrush(FStyleColors::Input))
+		.SetForegroundColor(FStyleColors::Foreground)
+		.SetHoveredForegroundColor(FStyleColors::ForegroundHover)
+		.SetPressedForegroundColor(FStyleColors::ForegroundHover)
+		.SetCheckedForegroundColor(FStyleColors::Primary)
+		.SetCheckedHoveredForegroundColor(FStyleColors::Primary)
+		.SetCheckedPressedForegroundColor(FStyleColors::Primary)
+		.SetPadding(DefaultMargins);
+
+
+	const FCheckBoxStyle SegmentedBoxRight = FCheckBoxStyle()
+		.SetCheckBoxType(ESlateCheckBoxType::ToggleButton)
+		.SetUncheckedImage(BOX_BRUSH("/Starship/CoreWidgets/SegmentedBox/right", FVector2D(16.f, 16.f), FMargin(4.0f / 16.0f), FStyleColors::Dropdown))
+		.SetUncheckedHoveredImage(BOX_BRUSH("/Starship/CoreWidgets/SegmentedBox/right", FVector2D(16.f, 16.f), FMargin(4.0f / 16.0f), FStyleColors::Hover))
+		.SetUncheckedPressedImage(BOX_BRUSH("/Starship/CoreWidgets/SegmentedBox/right", FVector2D(16.f, 16.f), FMargin(4.0f / 16.0f), FStyleColors::Dropdown))
+		.SetCheckedImage(BOX_BRUSH("/Starship/CoreWidgets/SegmentedBox/right", FVector2D(16.f, 16.f), FMargin(4.0f / 16.0f), FStyleColors::Input))
+		.SetCheckedHoveredImage(BOX_BRUSH("/Starship/CoreWidgets/SegmentedBox/right", FVector2D(16.f, 16.f), FMargin(4.0f / 16.0f), FStyleColors::Input))
+		.SetCheckedPressedImage(BOX_BRUSH("/Starship/CoreWidgets/SegmentedBox/right", FVector2D(16.f, 16.f), FMargin(4.0f / 16.0f), FStyleColors::Input))
+		.SetForegroundColor(FStyleColors::Foreground)
+		.SetHoveredForegroundColor(FStyleColors::ForegroundHover)
+		.SetPressedForegroundColor(FStyleColors::ForegroundHover)
+		.SetCheckedForegroundColor(FStyleColors::Primary)
+		.SetCheckedHoveredForegroundColor(FStyleColors::Primary)
+		.SetCheckedPressedForegroundColor(FStyleColors::Primary)
+		.SetPadding(DefaultMargins);
+
+	Style->Set("SegmentedBoxLeft", SegmentedBoxLeft);
+	Style->Set("SegmentedBoxCenter", SegmentedBoxCenter);
+	Style->Set("SegmentedBoxRight", SegmentedBoxRight);
+
+
+
+	/* Style for a toggleable button that mimics the coloring and look of a Table Row */
+	/*
+	const FCheckBoxStyle ToggleButtonRowStyle = FCheckBoxStyle()
+		.SetCheckBoxType(ESlateCheckBoxType::ToggleButton)
+		.SetUncheckedImage(FSlateNoResource())
+		.SetUncheckedHoveredImage(IMAGE_BRUSH("Common/Selection", Icon8x8, SelectionColor_Inactive))
+		.SetUncheckedPressedImage(IMAGE_BRUSH("Common/Selection", Icon8x8, SelectionColor_Inactive))
+		.SetCheckedImage(IMAGE_BRUSH("Common/Selection", Icon8x8, SelectionColor))
+		.SetCheckedHoveredImage(IMAGE_BRUSH("Common/Selection", Icon8x8, SelectionColor))
+		.SetCheckedPressedImage(BOX_BRUSH("Common/Selector", 4.0f / 16.0f, SelectorColor));
+	Style->Set("ToggleButtonRowStyle", ToggleButtonRowStyle);
+	*/
+
+	/* A radio button is actually just a SCheckBox box with different images */
+	/* Set images for various radio button (SCheckBox) states ... */
+	const FCheckBoxStyle BasicRadioButtonStyle = FCheckBoxStyle()
+		.SetCheckBoxType(ESlateCheckBoxType::CheckBox)
+
+		.SetForegroundColor(FLinearColor::White)
+		.SetHoveredForegroundColor(FLinearColor::White)
+		.SetPressedForegroundColor(FLinearColor::White)
+		.SetCheckedForegroundColor(FLinearColor::White)
+		.SetCheckedHoveredForegroundColor(FLinearColor::White)
+		.SetCheckedPressedForegroundColor(FLinearColor::White)
+		.SetUndeterminedForegroundColor(FLinearColor::White)
+
+		.SetUncheckedImage(IMAGE_BRUSH_SVG("/Starship/CoreWidgets/CheckBox/radio-off", Icon16x16, FStyleColors::White25))
+		.SetUncheckedHoveredImage(IMAGE_BRUSH_SVG("/Starship/CoreWidgets/CheckBox/radio-off", Icon16x16, FStyleColors::ForegroundHover))
+		.SetUncheckedPressedImage(IMAGE_BRUSH_SVG("/Starship/CoreWidgets/CheckBox/radio-off", Icon16x16, FStyleColors::ForegroundHover))
+		.SetCheckedImage(IMAGE_BRUSH_SVG("/Starship/CoreWidgets/CheckBox/radio-on", Icon16x16, FStyleColors::Primary))
+		.SetCheckedHoveredImage(IMAGE_BRUSH_SVG("/Starship/CoreWidgets/CheckBox/radio-on", Icon16x16, FStyleColors::Primary))
+		.SetCheckedPressedImage(IMAGE_BRUSH_SVG("/Starship/CoreWidgets/CheckBox/radio-off", Icon16x16, FStyleColors::Primary))
+		.SetUndeterminedImage(IMAGE_BRUSH_SVG("/Starship/CoreWidgets/CheckBox/radio-off", Icon16x16, FStyleColors::White25))
+		.SetUndeterminedHoveredImage(IMAGE_BRUSH_SVG("/Starship/CoreWidgets/CheckBox/radio-off", Icon16x16, FStyleColors::ForegroundHover))
+		.SetUndeterminedPressedImage(IMAGE_BRUSH_SVG("/Starship/CoreWidgets/CheckBox/radio-off", Icon16x16, FStyleColors::ForegroundHover))
+		.SetPadding(FMargin(4.0));
+	Style->Set("RadioButton", BasicRadioButtonStyle);
+}
 
-	// SNumericEntryBox defaults...
-	{
-		Style->Set("NumericEntrySpinBox", FSpinBoxStyle()
-			.SetBackgroundBrush(FSlateRoundedBoxBrush(FStyleColors::Input, InputFocusRadius, FStyleColors::Dropdown, InputFocusThickness))
-			.SetHoveredBackgroundBrush(FSlateRoundedBoxBrush(FStyleColors::Input, InputFocusRadius, FStyleColors::Hover, InputFocusThickness))
+void FStarshipCoreStyle::SetupDockingStyles(TSharedRef<FStyle>& Style)
+{
+	const FButtonStyle& Button = Style->GetWidgetStyle<FButtonStyle>("Button");
+	const FTextBlockStyle& NormalText = Style->GetWidgetStyle<FTextBlockStyle>("NormalText");
+	const FSlateColor SelectionColor = Style->GetSlateColor("SelectionColor");
 
-			.SetActiveFillBrush(FSlateRoundedBoxBrush(FStyleColors::Hover, InputFocusRadius, FLinearColor::Transparent, InputFocusThickness))
-			.SetInactiveFillBrush(FSlateRoundedBoxBrush(FStyleColors::Dropdown, InputFocusRadius, FLinearColor::Transparent, InputFocusThickness))
+	// SDockTab, SDockingTarget, SDockingTabStack defaults...
+	Style->Set("Docking.Background", new BOX_BRUSH("Old/Menu_Background", FMargin(8.0f / 64.0f)));
+	Style->Set("Docking.Border", new FSlateRoundedBoxBrush(FStyleColors::Background, 4));
+
+	Style->Set("Docking.UnhideTabwellButton", FButtonStyle(Button)
+		.SetNormal(IMAGE_BRUSH_SVG("Starship/Docking/show-tab-well", Icon8x8, FStyleColors::Primary))
+		.SetPressed(IMAGE_BRUSH_SVG("Starship/Docking/show-tab-well", Icon8x8, FStyleColors::PrimaryPress))
+		.SetHovered(IMAGE_BRUSH_SVG("Starship/Docking/show-tab-well", Icon8x8, FStyleColors::PrimaryHover))
+		.SetNormalPadding(0)
+		.SetPressedPadding(0)
+	);
 
-			.SetArrowsImage(FSlateNoResource())
-			.SetTextPadding(FMargin(0.f))
-			.SetForegroundColor(FStyleColors::ForegroundHover)
-		);
+	// Flash using the selection color for consistency with the rest of the UI scheme
+	const FSlateColor& TabFlashColor = SelectionColor;
 
-		Style->Set("NumericEntrySpinBox_Dark", FSpinBoxStyle()
-			.SetBackgroundBrush(FSlateRoundedBoxBrush(FStyleColors::Input, InputFocusRadius, FStyleColors::Dropdown, InputFocusThickness))
-			.SetHoveredBackgroundBrush(FSlateRoundedBoxBrush(FStyleColors::Input, InputFocusRadius, FStyleColors::Hover, InputFocusThickness))
+	const FButtonStyle CloseButton = FButtonStyle()
+		.SetNormal(IMAGE_BRUSH("Starship/Icons/Quixel/Common/cross_32x", Icon16x16, FStyleColors::Foreground))
+		.SetPressed(IMAGE_BRUSH("Starship/Icons/Quixel/Common/cross_32x", Icon16x16, FStyleColors::Foreground))
+		.SetHovered(IMAGE_BRUSH("Starship/Icons/Quixel/Common/cross_32x", Icon16x16, FStyleColors::ForegroundHover));
 
-			.SetActiveFillBrush(FSlateRoundedBoxBrush(FStyleColors::Hover, InputFocusRadius, FLinearColor::Transparent, InputFocusThickness))
-			.SetInactiveFillBrush(FSlateRoundedBoxBrush(FStyleColors::Dropdown, InputFocusRadius, FLinearColor::Transparent, InputFocusThickness))
 
-			.SetArrowsImage(FSlateNoResource())
+	FLinearColor DockColor_Inactive(FColor(45, 45, 45));
+	FLinearColor DockColor_Hovered(FColor(54, 54, 54));
+	FLinearColor DockColor_Active(FColor(62, 62, 62));
 
-			.SetTextPadding(FMargin(0.f))
-			.SetForegroundColor(FStyleColors::ForegroundHover)
-		);
+	// Panel Tab
+	Style->Set("Docking.Tab", FDockTabStyle()
+		.SetCloseButtonStyle(CloseButton)
+		.SetNormalBrush(FSlateNoResource())
+		.SetHoveredBrush(BOX_BRUSH("/Starship/Docking/DockTab_Hover", 4.0f / 20.0f, FStyleColors::Background))
+		.SetForegroundBrush(BOX_BRUSH("/Starship/Docking/DockTab_Foreground", 4.0f / 20.0f, FStyleColors::Background))
 
-		Style->Set("NumericEntrySpinBox.Decorator", new BOX_BRUSH("Common/TextBoxLabelBorder", FMargin(5.0f / 16.0f)));
+		.SetColorOverlayTabBrush(FSlateNoResource())
+		.SetColorOverlayIconBrush(FSlateNoResource())
+		.SetContentAreaBrush(FSlateNoResource())
+		.SetTabWellBrush(FSlateNoResource())
+		.SetFlashColor(TabFlashColor)
 
-		Style->Set("NumericEntrySpinBox.NarrowDecorator", new BOX_BRUSH("Common/TextBoxLabelBorder", FMargin(2.0f / 16.0f, 4.0f / 16.0f, 2.0f / 16.0f, 4.0f / 16.0f)));
-	}
+		.SetTabPadding(FMargin(10, 3, 10, 4))
+		.SetOverlapWidth(0.0f)
+
+		.SetNormalForegroundColor(FStyleColors::Foreground)
+		.SetActiveForegroundColor(FStyleColors::ForegroundHover)
+		.SetForegroundForegroundColor(FStyleColors::Foreground)
+		.SetHoveredForegroundColor(FStyleColors::ForegroundHover)
+		.SetTabTextStyle(NormalText)
+	);
+
+	// App Tab
+	Style->Set("Docking.MajorTab", FDockTabStyle()
+		.SetCloseButtonStyle(CloseButton)
+		.SetNormalBrush(FSlateNoResource())
+		.SetHoveredBrush(BOX_BRUSH("/Starship/Docking/DockTab_Hover", 4.0f / 20.0f, FStyleColors::Background))
+		.SetForegroundBrush(BOX_BRUSH("/Starship/Docking/DockTab_Foreground", 4.0f / 20.0f, FStyleColors::Background))
+
+		.SetColorOverlayTabBrush(FSlateNoResource())
+		.SetColorOverlayIconBrush(FSlateNoResource())
+		.SetContentAreaBrush(FSlateNoResource())
+		.SetTabWellBrush(FSlateNoResource())
+
+		.SetTabPadding(FMargin(10, 7, 10, 8))
+		.SetOverlapWidth(0.f)
+		.SetFlashColor(TabFlashColor)
+
+		.SetNormalForegroundColor(FStyleColors::Foreground)
+		.SetActiveForegroundColor(FStyleColors::ForegroundHover)
+		.SetForegroundForegroundColor(FStyleColors::Foreground)
+		.SetHoveredForegroundColor(FStyleColors::ForegroundHover)
+		.SetTabTextStyle(NormalText)
+	);
+
+	Style->Set("Docking.Tab.ContentAreaBrush", new FSlateNoResource());
+
+	Style->Set("Docking.Tab.InactiveTabSeparator", new FSlateColorBrush(FStyleColors::Hover));
+
+	Style->Set("Docking.Tab.ActiveTabIndicatorColor", FStyleColors::Primary);
+
+
+	// Dock Cross
+	Style->Set("Docking.Cross.DockLeft", new IMAGE_BRUSH("/Docking/OuterDockingIndicator", FVector2D(6, 6), FLinearColor(1.0f, 0.35f, 0.0f, 0.25f)));
+	Style->Set("Docking.Cross.DockLeft_Hovered", new IMAGE_BRUSH("/Docking/OuterDockingIndicator", FVector2D(6, 6), FLinearColor(1.0f, 0.35f, 0.0f)));
+	Style->Set("Docking.Cross.DockTop", new IMAGE_BRUSH("/Docking/OuterDockingIndicator", FVector2D(6, 6), FLinearColor(1.0f, 0.35f, 0.0f, 0.25f)));
+	Style->Set("Docking.Cross.DockTop_Hovered", new IMAGE_BRUSH("/Docking/OuterDockingIndicator", FVector2D(6, 6), FLinearColor(1.0f, 0.35f, 0.0f)));
+	Style->Set("Docking.Cross.DockRight", new IMAGE_BRUSH("/Docking/OuterDockingIndicator", FVector2D(6, 6), FLinearColor(1.0f, 0.35f, 0.0f, 0.25f)));
+	Style->Set("Docking.Cross.DockRight_Hovered", new IMAGE_BRUSH("/Docking/OuterDockingIndicator", FVector2D(6, 6), FLinearColor(1.0f, 0.35f, 0.0f)));
+	Style->Set("Docking.Cross.DockBottom", new IMAGE_BRUSH("/Docking/OuterDockingIndicator", FVector2D(6, 6), FLinearColor(1.0f, 0.35f, 0.0f, 0.25f)));
+	Style->Set("Docking.Cross.DockBottom_Hovered", new IMAGE_BRUSH("/Docking/OuterDockingIndicator", FVector2D(6, 6), FLinearColor(1.0f, 0.35f, 0.0f)));
+	Style->Set("Docking.Cross.DockCenter", new IMAGE_BRUSH("/Docking/DockingIndicator_Center", Icon64x64, FLinearColor(1.0f, 0.35f, 0.0f, 0.25f)));
+	Style->Set("Docking.Cross.DockCenter_Hovered", new IMAGE_BRUSH("/Docking/DockingIndicator_Center", Icon64x64, FLinearColor(1.0f, 0.35f, 0.0f)));
 
+	Style->Set("Docking.Cross.BorderLeft", new FSlateNoResource());
+	Style->Set("Docking.Cross.BorderTop", new FSlateNoResource());
+	Style->Set("Docking.Cross.BorderRight", new FSlateNoResource());
+	Style->Set("Docking.Cross.BorderBottom", new FSlateNoResource());
+	Style->Set("Docking.Cross.BorderCenter", new FSlateNoResource());
+
+	Style->Set("Docking.Cross.PreviewWindowTint", FLinearColor(1.0f, 0.75f, 0.5f));
+	Style->Set("Docking.Cross.Tint", FLinearColor::White);
+	Style->Set("Docking.Cross.HoveredTint", FLinearColor::White);
+}
+
+void FStarshipCoreStyle::SetupColorPickerStyles(TSharedRef<FStyle>& Style)
+{
 	// SColorPicker defaults...
 	{
-		Style->Set("ColorPicker.Border", new FSlateRoundedBoxBrush(FStyleColors::Background, 4) );
+		Style->Set("ColorPicker.Border", new FSlateRoundedBoxBrush(FStyleColors::Background, 4));
 		Style->Set("ColorPicker.AlphaBackground", new IMAGE_BRUSH("Common/Checker", Icon16x16, FLinearColor::White, ESlateBrushTileType::Both));
 		Style->Set("ColorPicker.EyeDropper", new IMAGE_BRUSH("Icons/eyedropper_16px", Icon16x16));
-		Style->Set("ColorPicker.Font", StyleFonts.Normal);
+		Style->Set("ColorPicker.Font", FStyleFonts::Get().Normal);
 		Style->Set("ColorPicker.Mode", new IMAGE_BRUSH("Common/ColorPicker_Mode_16x", Icon16x16));
 		Style->Set("ColorPicker.Separator", new IMAGE_BRUSH("Common/ColorPicker_Separator", FVector2D(2.0f, 2.0f)));
 		Style->Set("ColorPicker.Selector", new IMAGE_BRUSH("Common/Circle", FVector2D(8, 8)));
@@ -1106,93 +1497,98 @@ TSharedRef<ISlateStyle> FStarshipCoreStyle::Create()
 		Style->Set("ColorGradingWheel.HueValueCircle", new IMAGE_BRUSH("Common/ColorGradingWheel", FVector2D(192, 192)));
 		Style->Set("ColorGradingWheel.Selector", new IMAGE_BRUSH("Common/Circle", FVector2D(8, 8)));
 	}
+}
 
-	// SSplitter
-	{
-		Style->Set("Splitter", FSplitterStyle()
-			.SetHandleNormalBrush(FSlateNoResource())
-			.SetHandleHighlightBrush(IMAGE_BRUSH("Common/SplitterHandleHighlight", Icon8x8, FLinearColor::White))
-		);
-	}
+void FStarshipCoreStyle::SetupTableViewStyles(TSharedRef<FStyle>& Style)
+{
+	const FSlateColor SelectionColor = Style->GetSlateColor("SelectionColor");
+	const FSlateColor SelectorColor = Style->GetSlateColor("SelectorColor");
 
-	// TableView defaults...
-	{
-		const FTableRowStyle DefaultTableRowStyle = FTableRowStyle()
-			.SetEvenRowBackgroundBrush(FSlateColorBrush(FStyleColors::Input))
-			.SetEvenRowBackgroundHoveredBrush(FSlateColorBrush(FStyleColors::SelectHover))
+	const FTableRowStyle DefaultTableRowStyle = FTableRowStyle()
+		.SetEvenRowBackgroundBrush(FSlateColorBrush(FStyleColors::Input))
+		.SetEvenRowBackgroundHoveredBrush(FSlateColorBrush(FStyleColors::SelectHover))
 
-			.SetOddRowBackgroundBrush(FSlateColorBrush(FStyleColors::Input))
-			.SetOddRowBackgroundHoveredBrush(FSlateColorBrush(FStyleColors::SelectHover))
+		.SetOddRowBackgroundBrush(FSlateColorBrush(FStyleColors::Input))
+		.SetOddRowBackgroundHoveredBrush(FSlateColorBrush(FStyleColors::SelectHover))
 
-			.SetSelectorFocusedBrush(BORDER_BRUSH("Common/Selector", FMargin(4.f / 16.f), SelectorColor))
+		.SetSelectorFocusedBrush(BORDER_BRUSH("Common/Selector", FMargin(4.f / 16.f), SelectorColor))
 
-			.SetActiveBrush(         FSlateColorBrush(FStyleColors::Select))
-			.SetActiveHoveredBrush(  FSlateColorBrush(FStyleColors::Select))
+		.SetActiveBrush(FSlateColorBrush(FStyleColors::Select))
+		.SetActiveHoveredBrush(FSlateColorBrush(FStyleColors::Select))
 
-			.SetInactiveBrush(       FSlateColorBrush(FStyleColors::SelectInactive))
-			.SetInactiveHoveredBrush(FSlateColorBrush(FStyleColors::SelectInactive))
+		.SetInactiveBrush(FSlateColorBrush(FStyleColors::SelectInactive))
+		.SetInactiveHoveredBrush(FSlateColorBrush(FStyleColors::SelectInactive))
 
-			.SetActiveHighlightedBrush(  FSlateColorBrush(FStyleColors::SelectParent)) // This is the parent hightlight
-			.SetInactiveHighlightedBrush(FSlateColorBrush(FStyleColors::SelectParent))// This is the parent highlight
+		.SetActiveHighlightedBrush(FSlateColorBrush(FStyleColors::SelectParent)) // This is the parent hightlight
+		.SetInactiveHighlightedBrush(FSlateColorBrush(FStyleColors::SelectParent))// This is the parent highlight
 
-			.SetTextColor(FStyleColors::Foreground)
-			.SetSelectedTextColor(FStyleColors::ForegroundInverted)
+		.SetTextColor(FStyleColors::Foreground)
+		.SetSelectedTextColor(FStyleColors::ForegroundInverted)
 
-			.SetDropIndicator_Above(BOX_BRUSH("Common/DropZoneIndicator_Above", FMargin(10.0f / 16.0f, 10.0f / 16.0f, 0, 0), SelectionColor))
-			.SetDropIndicator_Onto(BOX_BRUSH("Common/DropZoneIndicator_Onto", FMargin(4.0f / 16.0f), SelectionColor))
-			.SetDropIndicator_Below(BOX_BRUSH("Common/DropZoneIndicator_Below", FMargin(10.0f / 16.0f, 0, 0, 10.0f / 16.0f), SelectionColor));
+		.SetDropIndicator_Above(BOX_BRUSH("Common/DropZoneIndicator_Above", FMargin(10.0f / 16.0f, 10.0f / 16.0f, 0, 0), SelectionColor))
+		.SetDropIndicator_Onto(BOX_BRUSH("Common/DropZoneIndicator_Onto", FMargin(4.0f / 16.0f), SelectionColor))
+		.SetDropIndicator_Below(BOX_BRUSH("Common/DropZoneIndicator_Below", FMargin(10.0f / 16.0f, 0, 0, 10.0f / 16.0f), SelectionColor));
 
-		Style->Set("TableView.Row", DefaultTableRowStyle);
+	Style->Set("TableView.Row", DefaultTableRowStyle);
 
-		const FTableRowStyle DarkTableRowStyle = FTableRowStyle(DefaultTableRowStyle)
-			.SetEvenRowBackgroundBrush(IMAGE_BRUSH("Common/Selection", Icon8x8, FLinearColor(0.0f, 0.0f, 0.0f, 0.1f)))
-			.SetOddRowBackgroundBrush(IMAGE_BRUSH("Common/Selection", Icon8x8, FLinearColor(0.0f, 0.0f, 0.0f, 0.1f)));
-		Style->Set("TableView.DarkRow", DarkTableRowStyle);
+	const FTableRowStyle DarkTableRowStyle = FTableRowStyle(DefaultTableRowStyle)
+		.SetEvenRowBackgroundBrush(IMAGE_BRUSH("Common/Selection", Icon8x8, FLinearColor(0.0f, 0.0f, 0.0f, 0.1f)))
+		.SetOddRowBackgroundBrush(IMAGE_BRUSH("Common/Selection", Icon8x8, FLinearColor(0.0f, 0.0f, 0.0f, 0.1f)));
+	Style->Set("TableView.DarkRow", DarkTableRowStyle);
 
-		Style->Set("TreeArrow_Collapsed",         new IMAGE_BRUSH_SVG("Starship/Common/chevron-right", Icon16x16, FStyleColors::Foreground));
-		Style->Set("TreeArrow_Collapsed_Hovered", new IMAGE_BRUSH_SVG("Starship/Common/chevron-right", Icon16x16, FStyleColors::ForegroundHover));
-		Style->Set("TreeArrow_Expanded",          new IMAGE_BRUSH_SVG("Starship/Common/chevron-down",  Icon16x16, FStyleColors::Foreground));
-		Style->Set("TreeArrow_Expanded_Hovered",  new IMAGE_BRUSH_SVG("Starship/Common/chevron-down",  Icon16x16, FStyleColors::ForegroundHover));
+	Style->Set("TreeArrow_Collapsed", new IMAGE_BRUSH_SVG("Starship/Common/chevron-right", Icon16x16, FStyleColors::Foreground));
+	Style->Set("TreeArrow_Collapsed_Hovered", new IMAGE_BRUSH_SVG("Starship/Common/chevron-right", Icon16x16, FStyleColors::ForegroundHover));
+	Style->Set("TreeArrow_Expanded", new IMAGE_BRUSH_SVG("Starship/Common/chevron-down", Icon16x16, FStyleColors::Foreground));
+	Style->Set("TreeArrow_Expanded_Hovered", new IMAGE_BRUSH_SVG("Starship/Common/chevron-down", Icon16x16, FStyleColors::ForegroundHover));
 
 
-		const FTableColumnHeaderStyle TableColumnHeaderStyle = FTableColumnHeaderStyle()
-			.SetSortPrimaryAscendingImage(   IMAGE_BRUSH_SVG("Starship/CoreWidgets/TableView/sort-up-arrow",    Icon12x12))
-			.SetSortPrimaryDescendingImage(  IMAGE_BRUSH_SVG("Starship/CoreWidgets/TableView/sort-down-arrow",  Icon12x12))
-			.SetSortSecondaryAscendingImage( IMAGE_BRUSH_SVG("Starship/CoreWidgets/TableView/sort-up-arrows",   Icon12x12))
-			.SetSortSecondaryDescendingImage(IMAGE_BRUSH_SVG("Starship/CoreWidgets/TableView/sort-down-arrows", Icon12x12))
-			.SetNormalBrush( FSlateColorBrush(FStyleColors::Background))
-			.SetHoveredBrush(FSlateColorBrush(FStyleColors::Dropdown))
+	const FTableColumnHeaderStyle TableColumnHeaderStyle = FTableColumnHeaderStyle()
+		.SetSortPrimaryAscendingImage(IMAGE_BRUSH_SVG("Starship/CoreWidgets/TableView/sort-up-arrow", Icon12x12))
+		.SetSortPrimaryDescendingImage(IMAGE_BRUSH_SVG("Starship/CoreWidgets/TableView/sort-down-arrow", Icon12x12))
+		.SetSortSecondaryAscendingImage(IMAGE_BRUSH_SVG("Starship/CoreWidgets/TableView/sort-up-arrows", Icon12x12))
+		.SetSortSecondaryDescendingImage(IMAGE_BRUSH_SVG("Starship/CoreWidgets/TableView/sort-down-arrows", Icon12x12))
+		.SetNormalBrush(FSlateColorBrush(FStyleColors::Background))
+		.SetHoveredBrush(FSlateColorBrush(FStyleColors::Dropdown))
 
-			.SetMenuDropdownImage(IMAGE_BRUSH_SVG("Starship/Common/ellipsis-vertical-narrow", FVector2D(6.f, 24.f)))
-			.SetMenuDropdownNormalBorderBrush(FSlateNoResource())
-			.SetMenuDropdownHoveredBorderBrush(FSlateNoResource());
+		.SetMenuDropdownImage(IMAGE_BRUSH_SVG("Starship/Common/ellipsis-vertical-narrow", FVector2D(6.f, 24.f)))
+		.SetMenuDropdownNormalBorderBrush(FSlateNoResource())
+		.SetMenuDropdownHoveredBorderBrush(FSlateNoResource());
 
-		Style->Set("TableView.Header.Column", TableColumnHeaderStyle);
+	Style->Set("TableView.Header.Column", TableColumnHeaderStyle);
 
-		const FTableColumnHeaderStyle TableLastColumnHeaderStyle = FTableColumnHeaderStyle()
-			.SetSortPrimaryAscendingImage(   IMAGE_BRUSH_SVG("Starship/CoreWidgets/TableView/sort-up-arrow",    Icon12x12))
-			.SetSortPrimaryDescendingImage(  IMAGE_BRUSH_SVG("Starship/CoreWidgets/TableView/sort-down-arrow",  Icon12x12))
-			.SetSortSecondaryAscendingImage( IMAGE_BRUSH_SVG("Starship/CoreWidgets/TableView/sort-up-arrows",   Icon12x12))
-			.SetSortSecondaryDescendingImage(IMAGE_BRUSH_SVG("Starship/CoreWidgets/TableView/sort-down-arrows", Icon12x12))
-			.SetNormalBrush(FSlateNoResource())
-			.SetHoveredBrush(FSlateNoResource())
+	const FTableColumnHeaderStyle TableLastColumnHeaderStyle = FTableColumnHeaderStyle()
+		.SetSortPrimaryAscendingImage(IMAGE_BRUSH_SVG("Starship/CoreWidgets/TableView/sort-up-arrow", Icon12x12))
+		.SetSortPrimaryDescendingImage(IMAGE_BRUSH_SVG("Starship/CoreWidgets/TableView/sort-down-arrow", Icon12x12))
+		.SetSortSecondaryAscendingImage(IMAGE_BRUSH_SVG("Starship/CoreWidgets/TableView/sort-up-arrows", Icon12x12))
+		.SetSortSecondaryDescendingImage(IMAGE_BRUSH_SVG("Starship/CoreWidgets/TableView/sort-down-arrows", Icon12x12))
+		.SetNormalBrush(FSlateNoResource())
+		.SetHoveredBrush(FSlateNoResource())
 
-			.SetMenuDropdownImage(IMAGE_BRUSH_SVG("Starship/Common/ellipsis-vertical-narrow", FVector2D(6.f, 24.f)))
-			.SetMenuDropdownNormalBorderBrush(FSlateNoResource())
-			.SetMenuDropdownHoveredBorderBrush(FSlateNoResource());
+		.SetMenuDropdownImage(IMAGE_BRUSH_SVG("Starship/Common/ellipsis-vertical-narrow", FVector2D(6.f, 24.f)))
+		.SetMenuDropdownNormalBorderBrush(FSlateNoResource())
+		.SetMenuDropdownHoveredBorderBrush(FSlateNoResource());
 
-		const FSplitterStyle TableHeaderSplitterStyle = FSplitterStyle()
-			.SetHandleNormalBrush(FSlateNoResource())
-			.SetHandleHighlightBrush(IMAGE_BRUSH("Common/HeaderSplitterGrip", Icon8x8));
-
-		Style->Set("TableView.Header", FHeaderRowStyle()
-			.SetColumnStyle(TableColumnHeaderStyle)
-			.SetLastColumnStyle(TableLastColumnHeaderStyle)
-			.SetColumnSplitterStyle(TableHeaderSplitterStyle)
-			.SetBackgroundBrush(FSlateColorBrush(FStyleColors::Background))
-			.SetForegroundColor(DefaultForeground)
-		);
-	}
+	const FSplitterStyle TableHeaderSplitterStyle = FSplitterStyle()
+		.SetHandleNormalBrush(FSlateNoResource())
+		.SetHandleHighlightBrush(IMAGE_BRUSH("Common/HeaderSplitterGrip", Icon8x8));
+
+	Style->Set("TableView.Header", FHeaderRowStyle()
+		.SetColumnStyle(TableColumnHeaderStyle)
+		.SetLastColumnStyle(TableLastColumnHeaderStyle)
+		.SetColumnSplitterStyle(TableHeaderSplitterStyle)
+		.SetBackgroundBrush(FSlateColorBrush(FStyleColors::Background))
+		.SetForegroundColor(FStyleColors::Foreground)
+	);
+}
+
+void FStarshipCoreStyle::SetupMultiboxStyles(TSharedRef<FStyle>& Style)
+{
+	const FEditableTextBoxStyle& NormalEditableTextBoxStyle = Style->GetWidgetStyle<FEditableTextBoxStyle>("NormalEditableTextBox");
+	const FTextBlockStyle& NormalText = Style->GetWidgetStyle<FTextBlockStyle>("NormalText");
+
+	const FSlateColor SelectionColor = Style->GetSlateColor("SelectionColor");
+	const FSlateColor SelectionColor_Pressed = Style->GetSlateColor("SelectionColor_Pressed");
+	const FSlateColor DefaultForeground = Style->GetSlateColor("SelectionColor_Pressed");
 
 	// MultiBox
 	{
@@ -1214,8 +1610,8 @@ TSharedRef<ISlateStyle> FStarshipCoreStyle::Create()
 			.SetCheckBoxPadding(FMargin(4.0f, 0.f))
 			.SetSeparatorBrush(FSlateColorBrush(FStyleColors::Input))
 			.SetSeparatorPadding(FMargin(1.f, 0.f, 1.f, 0.f))
-			.SetLabelStyle(FTextBlockStyle(NormalText).SetFont(StyleFonts.Normal))
-			.SetEditableTextStyle(FEditableTextBoxStyle(NormalEditableTextBoxStyle).SetFont(StyleFonts.Normal))
+			.SetLabelStyle(FTextBlockStyle(NormalText).SetFont(FStyleFonts::Get().Normal))
+			.SetEditableTextStyle(FEditableTextBoxStyle(NormalEditableTextBoxStyle).SetFont(FStyleFonts::Get().Normal))
 			.SetComboButtonStyle(Style->GetWidgetStyle<FComboButtonStyle>("ComboButton"))
 			.SetBlockPadding(FMargin(2.0f, 2.0f, 4.0f, 4.0f))
 			.SetIndentedBlockPadding(FMargin(18.0f, 2.0f, 4.0f, 4.0f));
@@ -1297,7 +1693,7 @@ TSharedRef<ISlateStyle> FStarshipCoreStyle::Create()
 		FComboButtonStyle SlimToolBarComboButton = FComboButtonStyle(Style->GetWidgetStyle<FComboButtonStyle>("ComboButton"))
 			.SetContentPadding(0)
 			.SetButtonStyle(SlimToolBarButton)
-			.SetDownArrowImage(IMAGE_BRUSH_SVG("Starship/Common/ellipsis-vertical-narrow", FVector2D(6,24)));
+			.SetDownArrowImage(IMAGE_BRUSH_SVG("Starship/Common/ellipsis-vertical-narrow", FVector2D(6, 24)));
 
 		SlimToolbarStyle.SetSettingsComboButtonStyle(SlimToolBarComboButton);
 		SlimToolbarStyle.SetIconSize(Icon20x20);
@@ -1320,7 +1716,7 @@ TSharedRef<ISlateStyle> FStarshipCoreStyle::Create()
 		Style->Set("Menu.MenuIconSize", 16.f);
 
 		const FMargin MenuBlockPadding(10.0f, 3.0f, 5.0f, 2.0f);
-		Style->Set("Menu.Block.IndentedPadding", MenuBlockPadding + FMargin(18.0f,0,0,0));
+		Style->Set("Menu.Block.IndentedPadding", MenuBlockPadding + FMargin(18.0f, 0, 0, 0));
 		Style->Set("Menu.Block.Padding", MenuBlockPadding);
 
 		Style->Set("Menu.Separator", new FSlateColorBrush(FStyleColors::Header));
@@ -1328,14 +1724,14 @@ TSharedRef<ISlateStyle> FStarshipCoreStyle::Create()
 
 		Style->Set("Menu.Label", NormalText);
 
-		Style->Set("Menu.EditableText", FEditableTextBoxStyle(NormalEditableTextBoxStyle).SetFont(StyleFonts.Normal));
-		Style->Set("Menu.Keybinding", FTextBlockStyle(NormalText).SetFont(StyleFonts.Small));
+		Style->Set("Menu.EditableText", FEditableTextBoxStyle(NormalEditableTextBoxStyle).SetFont(FStyleFonts::Get().Normal));
+		Style->Set("Menu.Keybinding", FTextBlockStyle(NormalText).SetFont(FStyleFonts::Get().Small));
 
-		Style->Set("Menu.Heading", 
+		Style->Set("Menu.Heading",
 			FTextBlockStyle(NormalText)
-			.SetFont(StyleFonts.SmallBold)
+			.SetFont(FStyleFonts::Get().SmallBold)
 			.SetColorAndOpacity(FStyleColors::Foreground));
-		Style->Set("Menu.Heading.Padding", MenuBlockPadding+FMargin(0,10,0,0));
+		Style->Set("Menu.Heading.Padding", MenuBlockPadding + FMargin(0, 10, 0, 0));
 
 		/* Set images for various SCheckBox states associated with menu check box items... */
 		const FCheckBoxStyle BasicMenuCheckBoxStyle = FCheckBoxStyle()
@@ -1383,7 +1779,7 @@ TSharedRef<ISlateStyle> FStarshipCoreStyle::Create()
 		Style->Set("Menu.ToggleButton", MenuToggleButtonCheckBoxStyle);
 
 		FButtonStyle MenuButton =
-			FButtonStyle(NoBorder)
+			FButtonStyle(Style->GetWidgetStyle<FButtonStyle>("NoBorder"))
 			.SetNormal(FSlateNoResource())
 			.SetPressed(FSlateColorBrush(FStyleColors::Primary))
 			.SetHovered(FSlateColorBrush(FStyleColors::Primary))
@@ -1405,354 +1801,21 @@ TSharedRef<ISlateStyle> FStarshipCoreStyle::Create()
 			.SetHoveredForeground(FStyleColors::ForegroundHover)
 			.SetPressedForeground(FStyleColors::Black)
 			.SetNormalForeground(FStyleColors::Foreground);
-		
+
 		// For menu bars we need to ignore the button style
 
 		Style->Set("WindowMenuBar.Background", new FSlateNoResource());
-		Style->Set("WindowMenuBar.Label", FTextBlockStyle(NormalText).SetFont(StyleFonts.Normal));
+		Style->Set("WindowMenuBar.Label", FTextBlockStyle(NormalText).SetFont(FStyleFonts::Get().Normal));
 		Style->Set("WindowMenuBar.Expand", new IMAGE_BRUSH("Icons/toolbar_expand_16x", Icon16x16));
 		Style->Set("WindowMenuBar.Button", MenuBarButton);
 		Style->Set("WindowMenuBar.Button.SubMenuOpen", new FSlateColorBrush(FStyleColors::Primary));
-		Style->Set("WindowMenuBar.MenuBar.Padding", FMargin(12,4));
-
-	}
-
-	// SExpandableButton defaults...
-	{
-		Style->Set("ExpandableButton.Background", new BOX_BRUSH("Common/Button", 8.0f / 32.0f));
-
-		// Extra padding on the right and bottom to account for image shadow
-		Style->Set("ExpandableButton.Padding", FMargin(3.f, 3.f, 6.f, 6.f));
-
-		Style->Set("ExpandableButton.CloseButton", new IMAGE_BRUSH("Common/ExpansionButton_CloseOverlay", Icon16x16));
-	}
-
-	// SBreadcrumbTrail defaults...
-	{
-		Style->Set("BreadcrumbTrail.Delimiter", new IMAGE_BRUSH("Common/Delimiter", Icon16x16));
-
-		Style->Set("BreadcrumbButton", FButtonStyle()
-			.SetNormal(FSlateNoResource())
-			.SetPressed(BOX_BRUSH("Common/RoundedSelection_16x", 4.0f / 16.0f, SelectionColor_Pressed))
-			.SetHovered(BOX_BRUSH("Common/RoundedSelection_16x", 4.0f / 16.0f, SelectionColor))
-			.SetNormalPadding(FMargin(0, 0))
-			.SetPressedPadding(FMargin(0, 0))
-		);
-	}
-
-	// SWizard defaults
-	{
-		Style->Set("Wizard.PageTitle", FTextBlockStyle(NormalText)
-			.SetFont(DEFAULT_FONT("BoldCondensed", 28))
-			.SetShadowOffset(FVector2D(1, 1))
-			.SetShadowColorAndOpacity(FLinearColor(0, 0, 0, 0.9f))
-		);
-	}
-
-	// SNotificationList defaults...
-	{
-		Style->Set("NotificationList.FontBold", DEFAULT_FONT("Bold", 16));
-		Style->Set("NotificationList.FontLight", DEFAULT_FONT("Light", 12));
-		Style->Set("NotificationList.ItemBackground", new BOX_BRUSH("Old/Menu_Background", FMargin(8.0f / 64.0f)));
-		Style->Set("NotificationList.ItemBackground_Border", new BOX_BRUSH("Old/Menu_Background_Inverted_Border_Bold", FMargin(8.0f / 64.0f)));
-		Style->Set("NotificationList.ItemBackground_Border_Transparent", new BOX_BRUSH("Old/Notification_Border_Flash", FMargin(8.0f / 64.0f)));
-		Style->Set("NotificationList.SuccessImage", new IMAGE_BRUSH("Icons/notificationlist_success", Icon16x16));
-		Style->Set("NotificationList.FailImage", new IMAGE_BRUSH("Icons/notificationlist_fail", Icon16x16));
-		Style->Set("NotificationList.DefaultMessage", new IMAGE_BRUSH("Common/EventMessage_Default", Icon40x40));
-	}
-
-	// SSeparator defaults...
-	{
-		Style->Set("Separator", new BOX_BRUSH("Common/Separator", 1 / 4.0f, FLinearColor(1, 1, 1, 0.5f)));
-	}
-
-	// SHeader defaults...
-	{
-		Style->Set("Header.Pre", new BOX_BRUSH("Common/Separator", FMargin(1 / 4.0f, 0, 2 / 4.0f, 0), FLinearColor(1, 1, 1, 0.5f)));
-		Style->Set("Header.Post", new BOX_BRUSH("Common/Separator", FMargin(2 / 4.0f, 0, 1 / 4.0f, 0), FLinearColor(1, 1, 1, 0.5f)));
-	}
-
-	// SDockTab, SDockingTarget, SDockingTabStack defaults...
-	{
-		Style->Set("Docking.Background", new BOX_BRUSH("Old/Menu_Background", FMargin(8.0f / 64.0f)));
-		Style->Set("Docking.Border", new FSlateRoundedBoxBrush(FStyleColors::Background, 4) );
-
-		Style->Set("Docking.UnhideTabwellButton", FButtonStyle(Button)
-			.SetNormal(IMAGE_BRUSH_SVG("Starship/Docking/show-tab-well", Icon8x8, FStyleColors::Primary))
-			.SetPressed(IMAGE_BRUSH_SVG("Starship/Docking/show-tab-well", Icon8x8, FStyleColors::PrimaryPress))
-			.SetHovered(IMAGE_BRUSH_SVG("Starship/Docking/show-tab-well", Icon8x8, FStyleColors::PrimaryHover))
-			.SetNormalPadding(0)
-			.SetPressedPadding(0)
-		);
-
-		// Flash using the selection color for consistency with the rest of the UI scheme
-		const FSlateColor& TabFlashColor = SelectionColor;
-
-		const FButtonStyle CloseButton = FButtonStyle()
-			.SetNormal( IMAGE_BRUSH("Starship/Icons/Quixel/Common/cross_32x", Icon16x16, FStyleColors::Foreground))
-			.SetPressed(IMAGE_BRUSH("Starship/Icons/Quixel/Common/cross_32x", Icon16x16, FStyleColors::Foreground))
-			.SetHovered(IMAGE_BRUSH("Starship/Icons/Quixel/Common/cross_32x", Icon16x16, FStyleColors::ForegroundHover));
-
-
-		FLinearColor DockColor_Inactive(FColor(45, 45, 45));
-		FLinearColor DockColor_Hovered(FColor(54, 54, 54));
-		FLinearColor DockColor_Active(FColor(62, 62, 62));
-
-		// Panel Tab
-		Style->Set("Docking.Tab", FDockTabStyle()
-			.SetCloseButtonStyle(CloseButton)
-			.SetNormalBrush(	FSlateNoResource())
-			.SetHoveredBrush(   BOX_BRUSH("/Starship/Docking/DockTab_Hover", 	  4.0f / 20.0f, FStyleColors::Background))
-			.SetForegroundBrush(BOX_BRUSH("/Starship/Docking/DockTab_Foreground", 4.0f / 20.0f, FStyleColors::Background))
-
-			.SetColorOverlayTabBrush(FSlateNoResource())
-			.SetColorOverlayIconBrush(FSlateNoResource())
-			.SetContentAreaBrush(FSlateNoResource())
-			.SetTabWellBrush(FSlateNoResource())
-			.SetFlashColor(TabFlashColor)
-
-			.SetTabPadding(FMargin(10, 3, 10, 4))
-			.SetOverlapWidth(0.0f)
-
-			.SetNormalForegroundColor(FStyleColors::Foreground)
-			.SetActiveForegroundColor(FStyleColors::ForegroundHover)
-			.SetForegroundForegroundColor(FStyleColors::Foreground)
-			.SetHoveredForegroundColor(FStyleColors::ForegroundHover)
-			.SetTabTextStyle(NormalText)
-		);
-
-		// App Tab
-		Style->Set("Docking.MajorTab", FDockTabStyle()
-			.SetCloseButtonStyle(CloseButton)
-			.SetNormalBrush(	FSlateNoResource())
-			.SetHoveredBrush(   BOX_BRUSH("/Starship/Docking/DockTab_Hover", 	  4.0f / 20.0f, FStyleColors::Background))
-			.SetForegroundBrush(BOX_BRUSH("/Starship/Docking/DockTab_Foreground", 4.0f / 20.0f, FStyleColors::Background))
-
-			.SetColorOverlayTabBrush(FSlateNoResource())
-			.SetColorOverlayIconBrush(FSlateNoResource())
-			.SetContentAreaBrush(FSlateNoResource())
-			.SetTabWellBrush(FSlateNoResource())
-
-			.SetTabPadding(FMargin(10, 7, 10, 8))
-			.SetOverlapWidth(0.f)
-			.SetFlashColor(TabFlashColor)
-
-			.SetNormalForegroundColor(FStyleColors::Foreground)
-			.SetActiveForegroundColor(FStyleColors::ForegroundHover)
-			.SetForegroundForegroundColor(FStyleColors::Foreground)
-			.SetHoveredForegroundColor(FStyleColors::ForegroundHover)
-			.SetTabTextStyle(NormalText)
-		);
-
-		Style->Set("Docking.Tab.ContentAreaBrush", new FSlateNoResource());
-
-		Style->Set("Docking.Tab.InactiveTabSeparator", new FSlateColorBrush(FStyleColors::Hover));
-
-		Style->Set("Docking.Tab.ActiveTabIndicatorColor", FStyleColors::Primary);
-
-
-		// Dock Cross
-		Style->Set("Docking.Cross.DockLeft", new IMAGE_BRUSH("/Docking/OuterDockingIndicator", FVector2D(6, 6), FLinearColor(1.0f, 0.35f, 0.0f, 0.25f)));
-		Style->Set("Docking.Cross.DockLeft_Hovered", new IMAGE_BRUSH("/Docking/OuterDockingIndicator", FVector2D(6, 6), FLinearColor(1.0f, 0.35f, 0.0f)));
-		Style->Set("Docking.Cross.DockTop", new IMAGE_BRUSH("/Docking/OuterDockingIndicator", FVector2D(6, 6), FLinearColor(1.0f, 0.35f, 0.0f, 0.25f)));
-		Style->Set("Docking.Cross.DockTop_Hovered", new IMAGE_BRUSH("/Docking/OuterDockingIndicator", FVector2D(6, 6), FLinearColor(1.0f, 0.35f, 0.0f)));
-		Style->Set("Docking.Cross.DockRight", new IMAGE_BRUSH("/Docking/OuterDockingIndicator", FVector2D(6, 6), FLinearColor(1.0f, 0.35f, 0.0f, 0.25f)));
-		Style->Set("Docking.Cross.DockRight_Hovered", new IMAGE_BRUSH("/Docking/OuterDockingIndicator", FVector2D(6, 6), FLinearColor(1.0f, 0.35f, 0.0f)));
-		Style->Set("Docking.Cross.DockBottom", new IMAGE_BRUSH("/Docking/OuterDockingIndicator", FVector2D(6, 6), FLinearColor(1.0f, 0.35f, 0.0f, 0.25f)));
-		Style->Set("Docking.Cross.DockBottom_Hovered", new IMAGE_BRUSH("/Docking/OuterDockingIndicator", FVector2D(6, 6), FLinearColor(1.0f, 0.35f, 0.0f)));
-		Style->Set("Docking.Cross.DockCenter", new IMAGE_BRUSH("/Docking/DockingIndicator_Center", Icon64x64, FLinearColor(1.0f, 0.35f, 0.0f, 0.25f)));
-		Style->Set("Docking.Cross.DockCenter_Hovered", new IMAGE_BRUSH("/Docking/DockingIndicator_Center", Icon64x64, FLinearColor(1.0f, 0.35f, 0.0f)));
-
-		Style->Set("Docking.Cross.BorderLeft", new FSlateNoResource());
-		Style->Set("Docking.Cross.BorderTop", new FSlateNoResource());
-		Style->Set("Docking.Cross.BorderRight", new FSlateNoResource());
-		Style->Set("Docking.Cross.BorderBottom", new FSlateNoResource());
-		Style->Set("Docking.Cross.BorderCenter", new FSlateNoResource());
-
-		Style->Set("Docking.Cross.PreviewWindowTint", FLinearColor(1.0f, 0.75f, 0.5f));
-		Style->Set("Docking.Cross.Tint", FLinearColor::White);
-		Style->Set("Docking.Cross.HoveredTint", FLinearColor::White);
-	}
-
-	// SScrollBox defaults...
-	{
-		Style->Set("ScrollBox", FScrollBoxStyle()
-			.SetTopShadowBrush(BOX_BRUSH("Common/ScrollBoxShadowTop", FVector2D(16, 8), FMargin(0.5, 1, 0.5, 0)))
-			.SetBottomShadowBrush(BOX_BRUSH("Common/ScrollBoxShadowBottom", FVector2D(16, 8), FMargin(0.5, 0, 0.5, 1)))
-			.SetLeftShadowBrush(BOX_BRUSH("Common/ScrollBoxShadowLeft", FVector2D(8, 16), FMargin(1, 0.5, 0, 0.5)))
-			.SetRightShadowBrush(BOX_BRUSH("Common/ScrollBoxShadowRight", FVector2D(8, 16), FMargin(0, 0.5, 1, 0.5)))
-			.SetBarThickness(8.0)
-		);
-	}
-
-	// SScrollBorder defaults...
-	{
-		Style->Set("ScrollBorder", FScrollBorderStyle()
-			.SetTopShadowBrush(BOX_BRUSH("Common/ScrollBorderShadowTop", FVector2D(16, 8), FMargin(0.5, 1, 0.5, 0)))
-			.SetBottomShadowBrush(BOX_BRUSH("Common/ScrollBorderShadowBottom", FVector2D(16, 8), FMargin(0.5, 0, 0.5, 1)))
-		);
-	}
-
-
-
-	// SWindow defaults...
-	{
-#if !PLATFORM_MAC
-		const FButtonStyle MinimizeButtonStyle = FButtonStyle(NoBorder)
-			.SetNormal( IMAGE_BRUSH_SVG("Starship/CoreWidgets/Window/minimize", FVector2D(42.0f, 34.0f), FStyleColors::Foreground))
-			.SetHovered(IMAGE_BRUSH_SVG("Starship/CoreWidgets/Window/minimize", FVector2D(42.0f, 34.0f), FStyleColors::ForegroundHover))
-			.SetPressed(IMAGE_BRUSH_SVG("Starship/CoreWidgets/Window/minimize", FVector2D(42.0f, 34.0f), FStyleColors::Foreground));
-
-		const FButtonStyle MaximizeButtonStyle = FButtonStyle(NoBorder)
-			.SetNormal( IMAGE_BRUSH_SVG("Starship/CoreWidgets/Window/maximize", FVector2D(42.0f, 34.0f), FStyleColors::Foreground))
-			.SetHovered(IMAGE_BRUSH_SVG("Starship/CoreWidgets/Window/maximize", FVector2D(42.0f, 34.0f), FStyleColors::ForegroundHover))
-			.SetPressed(IMAGE_BRUSH_SVG("Starship/CoreWidgets/Window/maximize", FVector2D(42.0f, 34.0f), FStyleColors::Foreground));
-
-		const FButtonStyle RestoreButtonStyle = FButtonStyle(NoBorder)
-			.SetNormal( IMAGE_BRUSH_SVG("Starship/CoreWidgets/Window/restore", FVector2D(42.0f, 34.0f), FStyleColors::Foreground))
-			.SetHovered(IMAGE_BRUSH_SVG("Starship/CoreWidgets/Window/restore", FVector2D(42.0f, 34.0f), FStyleColors::ForegroundHover))
-			.SetPressed(IMAGE_BRUSH_SVG("Starship/CoreWidgets/Window/restore", FVector2D(42.0f, 34.0f), FStyleColors::Foreground));
-
-		const FButtonStyle CloseButtonStyle = FButtonStyle(NoBorder)
-			.SetNormal( IMAGE_BRUSH_SVG("Starship/CoreWidgets/Window/close", FVector2D(42.0f, 34.0f), FStyleColors::Foreground))
-			.SetHovered(IMAGE_BRUSH_SVG("Starship/CoreWidgets/Window/close", FVector2D(42.0f, 34.0f), FStyleColors::ForegroundHover))
-			.SetPressed(IMAGE_BRUSH_SVG("Starship/CoreWidgets/Window/close", FVector2D(42.0f, 34.0f), FStyleColors::Foreground));
-#endif
-
-		const FTextBlockStyle TitleTextStyle = FTextBlockStyle(NormalText)
-			.SetFont(StyleFonts.Normal)
-			.SetColorAndOpacity(FLinearColor::White)
-			.SetShadowOffset(FVector2D(1.0f, 1.0f))
-			.SetShadowColorAndOpacity(FLinearColor::Black);
-
-		Style->Set("Window", FWindowStyle()
-
-#if !PLATFORM_MAC
-			.SetMinimizeButtonStyle(MinimizeButtonStyle)
-			.SetMaximizeButtonStyle(MaximizeButtonStyle)
-			.SetRestoreButtonStyle(RestoreButtonStyle)
-			.SetCloseButtonStyle(CloseButtonStyle)
-#endif
-			.SetTitleTextStyle(TitleTextStyle)
-			.SetActiveTitleBrush(FSlateNoResource())
-			.SetInactiveTitleBrush(FSlateNoResource())
-			.SetFlashTitleBrush(IMAGE_BRUSH("Common/Window/WindowTitle_Flashing", Icon24x24, FLinearColor(1.0f, 1.0f, 1.0f, 1.0f), ESlateBrushTileType::Horizontal))	
-			.SetBackgroundBrush(FSlateColorBrush(FStyleColors::Input))
-			.SetBorderBrush(FSlateNoResource())
-			.SetOutlineBrush(FSlateRoundedBoxBrush(FLinearColor::Transparent, 0.0, FLinearColor::FromSRGBColor(FColor::FromHex("#2C2C2CFF")), 1.0))
-			.SetChildBackgroundBrush(FSlateColorBrush(FStyleColors::Input))
-			.SetOutlineColor(WindowHighlight)
-			.SetCornerRadius(0)
-			.SetBorderPadding(FMargin(2.0, 1.0, 2.0, 2.0))
-		);
-
-		Style->Set("ChildWindow.Background", new FSlateColorBrush(FStyleColors::Input));
-	}
-
-
-	// Standard Dialog Settings
-	{
-		Style->Set("StandardDialog.ContentPadding", FMargin(12.0f, 2.0f));
-		Style->Set("StandardDialog.SlotPadding", FMargin(6.0f, 0.0f, 0.0f, 0.0f));
-		Style->Set("StandardDialog.MinDesiredSlotWidth", 80.0f);
-		Style->Set("StandardDialog.MinDesiredSlotHeight", 0.0f);
-		Style->Set("StandardDialog.SmallFont", StyleFonts.Small);
-		Style->Set("StandardDialog.LargeFont", StyleFonts.Normal);
-	}
-
-	// Widget Reflector Window
-	{
-		Style->Set("WidgetReflector.TabIcon", new IMAGE_BRUSH("Icons/icon_tab_WidgetReflector_16x", Icon16x16));
-		Style->Set("WidgetReflector.Icon", new IMAGE_BRUSH("Icons/icon_tab_WidgetReflector_40x", Icon40x40));
-		Style->Set("WidgetReflector.Icon.Small", new IMAGE_BRUSH("Icons/icon_tab_WidgetReflector_40x", Icon20x20));
-		Style->Set("WidgetReflector.FocusableCheck", FCheckBoxStyle()
-			.SetUncheckedImage(IMAGE_BRUSH("Icons/Empty_14x", Icon14x14))
-			.SetUncheckedHoveredImage(IMAGE_BRUSH("Icons/Empty_14x", Icon14x14))
-			.SetUncheckedPressedImage(IMAGE_BRUSH("Common/SmallCheckBox_Hovered", Icon14x14))
-			.SetCheckedImage(IMAGE_BRUSH("Common/SmallCheck", Icon14x14))
-			.SetCheckedHoveredImage(IMAGE_BRUSH("Common/SmallCheck", Icon14x14))
-			.SetCheckedPressedImage(IMAGE_BRUSH("Common/SmallCheck", Icon14x14))
-			.SetUndeterminedImage(IMAGE_BRUSH("Icons/Empty_14x", Icon14x14))
-			.SetUndeterminedHoveredImage(FSlateNoResource())
-			.SetUndeterminedPressedImage(FSlateNoResource())
-		);
-	}
-
-	// Message Log
-	{
-		Style->Set("MessageLog", FTextBlockStyle(NormalText)
-			.SetFont(StyleFonts.Small)
-			.SetShadowOffset(FVector2D::ZeroVector)
-		);
-		Style->Set("MessageLog.Error", new IMAGE_BRUSH("MessageLog/Log_Error", Icon16x16));
-		Style->Set("MessageLog.Warning", new IMAGE_BRUSH("MessageLog/Log_Warning", Icon16x16));
-		Style->Set("MessageLog.Note", new IMAGE_BRUSH("MessageLog/Log_Note", Icon16x16));
-	}
-
-	// Wizard icons
-	{
-		Style->Set("Wizard.BackIcon", new IMAGE_BRUSH("Icons/BackIcon", Icon8x8));
-		Style->Set("Wizard.NextIcon", new IMAGE_BRUSH("Icons/NextIcon", Icon8x8));
-	}
-
-	// Syntax highlighting
-	{
-		const FTextBlockStyle SmallMonospacedText = FTextBlockStyle(MonospacedText)
-			.SetFont(DEFAULT_FONT("Mono", 9));
+		Style->Set("WindowMenuBar.MenuBar.Padding", FMargin(12, 4));
 
-		Style->Set("SyntaxHighlight.Normal", SmallMonospacedText);
-		Style->Set("SyntaxHighlight.Node", FTextBlockStyle(SmallMonospacedText).SetColorAndOpacity(FLinearColor(FColor(0xff006ab4)))); // blue
-		Style->Set("SyntaxHighlight.NodeAttributeKey", FTextBlockStyle(SmallMonospacedText).SetColorAndOpacity(FLinearColor(FColor(0xffb40000)))); // red
-		Style->Set("SyntaxHighlight.NodeAttribueAssignment", FTextBlockStyle(SmallMonospacedText).SetColorAndOpacity(FLinearColor(FColor(0xffb2b400)))); // yellow
-		Style->Set("SyntaxHighlight.NodeAttributeValue", FTextBlockStyle(SmallMonospacedText).SetColorAndOpacity(FLinearColor(FColor(0xffb46100)))); // orange
 	}
 
-	return Style;
 }
-PRAGMA_ENABLE_OPTIMIZATION
 
 #undef IMAGE_BRUSH
 #undef BOX_BRUSH
 #undef BORDER_BRUSH
 #undef DEFAULT_FONT
-
-
-const TSharedPtr<FSlateDynamicImageBrush> FStarshipCoreStyle::GetDynamicImageBrush(FName BrushTemplate, FName TextureName, const ANSICHAR* Specifier)
-{
-	return Instance->GetDynamicImageBrush(BrushTemplate, TextureName, Specifier);
-}
-
-
-const TSharedPtr<FSlateDynamicImageBrush> FStarshipCoreStyle::GetDynamicImageBrush(FName BrushTemplate, const ANSICHAR* Specifier, class UTexture2D* TextureResource, FName TextureName)
-{
-	return Instance->GetDynamicImageBrush(BrushTemplate, Specifier, TextureResource, TextureName);
-}
-
-
-const TSharedPtr<FSlateDynamicImageBrush> FStarshipCoreStyle::GetDynamicImageBrush(FName BrushTemplate, class UTexture2D* TextureResource, FName TextureName)
-{
-	return Instance->GetDynamicImageBrush(BrushTemplate, TextureResource, TextureName);
-}
-
-
-/* FSlateThrottleManager implementation
- *****************************************************************************/
-
-void FStarshipCoreStyle::SetStyle(const TSharedRef< ISlateStyle >& NewStyle)
-{
-	if (Instance.IsValid())
-	{
-		FSlateStyleRegistry::UnRegisterSlateStyle(*Instance.Get());
-	}
-
-	Instance = NewStyle;
-
-	if (Instance.IsValid())
-	{
-		FSlateStyleRegistry::RegisterSlateStyle(*Instance.Get());
-	}
-	else
-	{
-		ResetToDefault();
-	}
-}
diff --git a/Engine/Source/Runtime/SlateCore/Public/Styling/StarshipCoreStyle.h b/Engine/Source/Runtime/SlateCore/Public/Styling/StarshipCoreStyle.h
--- a/Engine/Source/Runtime/SlateCore/Public/Styling/StarshipCoreStyle.h
+++ b/Engine/Source/Runtime/SlateCore/Public/Styling/StarshipCoreStyle.h
@@ -8,6 +8,8 @@
 
 struct FSlateDynamicImageBrush;
 
+class FStyle;
+
 /**
  * Core slate style
  */
@@ -67,7 +69,17 @@ public:
 	static bool IsInitialized() { return Instance.IsValid(); }
 private:
 
-	static void SetStyle( const TSharedRef< class ISlateStyle >& NewStyle );
+	static void SetStyle(const TSharedRef<class ISlateStyle>& NewStyle);
+	
+	static void SetupColors(TSharedRef<FStyle>& Style);
+	static void SetupTextStyles(TSharedRef<FStyle>& Style);
+	static void SetupButtonStyles(TSharedRef<FStyle>& Style);
+	static void SetupComboButtonStyles(TSharedRef<FStyle>& Style);
+	static void SetupCheckboxStyles(TSharedRef<FStyle>& Style);
+	static void SetupDockingStyles(TSharedRef<FStyle>& Style);
+	static void SetupColorPickerStyles(TSharedRef<FStyle>& Style);
+	static void SetupTableViewStyles(TSharedRef<FStyle>& Style);
+	static void SetupMultiboxStyles(TSharedRef<FStyle>& Style);
 
 private:
 
