Updated Menu Radio Check Style

diff --git a/Engine/Source/Runtime/SlateCore/Private/Styling/StarshipCoreStyle.cpp b/Engine/Source/Runtime/SlateCore/Private/Styling/StarshipCoreStyle.cpp
--- a/Engine/Source/Runtime/SlateCore/Private/Styling/StarshipCoreStyle.cpp
+++ b/Engine/Source/Runtime/SlateCore/Private/Styling/StarshipCoreStyle.cpp
@@ -1752,15 +1752,15 @@ void FStarshipCoreStyle::SetupMultiboxStyles(TSharedRef<FStyle>& Style)
 		/* This radio button is actually just a check box with different images */
 		/* Set images for various Menu radio button (SCheckBox) states... */
 		const FCheckBoxStyle BasicMenuRadioButtonStyle = FCheckBoxStyle()
-			.SetUncheckedImage(IMAGE_BRUSH("Common/RadioButton_Unselected_16x", Icon16x16))
-			.SetUncheckedHoveredImage(IMAGE_BRUSH("Common/RadioButton_Unselected_16x", Icon16x16))
-			.SetUncheckedPressedImage(IMAGE_BRUSH("Common/RadioButton_Unselected_16x", Icon16x16))
-			.SetCheckedImage(IMAGE_BRUSH("Common/RadioButton_Selected_16x", Icon16x16))
-			.SetCheckedHoveredImage(IMAGE_BRUSH("Common/RadioButton_Selected_16x", Icon16x16, SelectionColor))
-			.SetCheckedPressedImage(IMAGE_BRUSH("Common/RadioButton_Unselected_16x", Icon16x16, SelectionColor_Pressed))
-			.SetUndeterminedImage(IMAGE_BRUSH("Common/RadioButton_Unselected_16x", Icon16x16))
-			.SetUndeterminedHoveredImage(IMAGE_BRUSH("Common/RadioButton_Unselected_16x", Icon16x16, SelectionColor))
-			.SetUndeterminedPressedImage(IMAGE_BRUSH("Common/RadioButton_Unselected_16x", Icon16x16, SelectionColor_Pressed));
+			.SetUncheckedImage(           FSlateRoundedBoxBrush(FStyleColors::Header,  Icon8x8))
+			.SetUncheckedHoveredImage(    FSlateRoundedBoxBrush(FStyleColors::Hover2, Icon8x8))
+			.SetUncheckedPressedImage(    FSlateRoundedBoxBrush(FStyleColors::White,   Icon8x8))
+			.SetCheckedImage(             FSlateRoundedBoxBrush(FStyleColors::White,   Icon8x8))
+			.SetCheckedHoveredImage(      FSlateRoundedBoxBrush(FStyleColors::White,   Icon8x8))
+			.SetCheckedPressedImage(      FSlateRoundedBoxBrush(FStyleColors::White,   Icon8x8))
+			.SetUndeterminedImage(        FSlateRoundedBoxBrush(FStyleColors::Header,  Icon8x8))
+			.SetUndeterminedHoveredImage( FSlateRoundedBoxBrush(FStyleColors::Hover2, Icon8x8))
+			.SetUndeterminedPressedImage( FSlateRoundedBoxBrush(FStyleColors::White,   Icon8x8));
 
 		/* ...and set new style */
 		Style->Set("Menu.RadioButton", BasicMenuRadioButtonStyle);
