Made the svg raster time logging verbose

diff --git a/Engine/Source/Runtime/SlateCore/Private/Rendering/SlateVectorGraphicsCache.cpp b/Engine/Source/Runtime/SlateCore/Private/Rendering/SlateVectorGraphicsCache.cpp
--- a/Engine/Source/Runtime/SlateCore/Private/Rendering/SlateVectorGraphicsCache.cpp
+++ b/Engine/Source/Runtime/SlateCore/Private/Rendering/SlateVectorGraphicsCache.cpp
@@ -130,7 +130,7 @@ void FSlateVectorGraphicsCache::UpdateCache()
 
 	if (PendingRequests.Num() > 0)
 	{
-		UE_LOG(LogSlate, Log, TEXT("SVG raster took: %fms"), Time * 1000.f);
+		UE_LOG(LogSlate, Verbose, TEXT("SVG raster took: %fms"), Time * 1000.f);
 	}
 	PendingRequests.Empty();
 
