Remove padding from SWindow when maximized

diff --git a/Engine/Source/Runtime/SlateCore/Private/Widgets/SWindow.cpp b/Engine/Source/Runtime/SlateCore/Private/Widgets/SWindow.cpp
--- a/Engine/Source/Runtime/SlateCore/Private/Widgets/SWindow.cpp
+++ b/Engine/Source/Runtime/SlateCore/Private/Widgets/SWindow.cpp
@@ -881,7 +881,8 @@ FMargin SWindow::GetWindowBorderSize( bool bIncTitleBar ) const
 	if (NativeWindow.IsValid() && NativeWindow->IsMaximized())
 	{
 		const float DesktopPixelsToSlateUnits = 1.0f / (FSlateApplicationBase::Get().GetApplicationScale() * GetDPIScaleFactor());
-		FMargin BorderSize(NativeWindow->GetWindowBorderSize() * DesktopPixelsToSlateUnits);
+		const float PaddingAmount = NativeWindow->GetWindowBorderSize() * DesktopPixelsToSlateUnits;
+		FMargin BorderSize(0.0f, PaddingAmount, 0.0f, PaddingAmount);
 		if(bIncTitleBar)
 		{
 			// Add title bar size (whether it's visible or not)
