Back out padding change to SWindow.  Causes too many issues with UI elements being on the edge of the screen

diff --git a/Engine/Source/Runtime/SlateCore/Private/Widgets/SWindow.cpp b/Engine/Source/Runtime/SlateCore/Private/Widgets/SWindow.cpp
--- a/Engine/Source/Runtime/SlateCore/Private/Widgets/SWindow.cpp
+++ b/Engine/Source/Runtime/SlateCore/Private/Widgets/SWindow.cpp
@@ -881,8 +881,7 @@ FMargin SWindow::GetWindowBorderSize( bool bIncTitleBar ) const
 	if (NativeWindow.IsValid() && NativeWindow->IsMaximized())
 	{
 		const float DesktopPixelsToSlateUnits = 1.0f / (FSlateApplicationBase::Get().GetApplicationScale() * GetDPIScaleFactor());
-		const float PaddingAmount = NativeWindow->GetWindowBorderSize() * DesktopPixelsToSlateUnits;
-		FMargin BorderSize(0.0f, PaddingAmount, 0.0f, PaddingAmount);
+		FMargin BorderSize(NativeWindow->GetWindowBorderSize() * DesktopPixelsToSlateUnits);
 		if(bIncTitleBar)
 		{
 			// Add title bar size (whether it's visible or not)
