Give splitters a default color to distinguish them from tab backgrounds

diff --git a/Engine/Source/Runtime/SlateCore/Private/Styling/StarshipCoreStyle.cpp b/Engine/Source/Runtime/SlateCore/Private/Styling/StarshipCoreStyle.cpp
--- a/Engine/Source/Runtime/SlateCore/Private/Styling/StarshipCoreStyle.cpp
+++ b/Engine/Source/Runtime/SlateCore/Private/Styling/StarshipCoreStyle.cpp
@@ -519,7 +519,7 @@ TSharedRef<ISlateStyle> FStarshipCoreStyle::Create()
 	// SSplitter
 	{
 		Style->Set("Splitter", FSplitterStyle()
-			.SetHandleNormalBrush(FSlateNoResource())
+			.SetHandleNormalBrush(FSlateColorBrush(FStyleColors::Title))
 			.SetHandleHighlightBrush(IMAGE_BRUSH("Common/SplitterHandleHighlight", Icon8x8, FLinearColor::White))
 		);
 	}
