Window border darkening

diff --git a/Engine/Source/Runtime/SlateCore/Private/Styling/StarshipCoreStyle.cpp b/Engine/Source/Runtime/SlateCore/Private/Styling/StarshipCoreStyle.cpp
--- a/Engine/Source/Runtime/SlateCore/Private/Styling/StarshipCoreStyle.cpp
+++ b/Engine/Source/Runtime/SlateCore/Private/Styling/StarshipCoreStyle.cpp
@@ -651,11 +651,11 @@ TSharedRef<ISlateStyle> FStarshipCoreStyle::Create()
 			.SetFlashTitleBrush(IMAGE_BRUSH("Common/Window/WindowTitle_Flashing", Icon24x24, FLinearColor(1.0f, 1.0f, 1.0f, 1.0f), ESlateBrushTileType::Horizontal))	
 			.SetBackgroundBrush(FSlateColorBrush(FStyleColors::Input))
 			.SetBorderBrush(FSlateNoResource())
-			.SetOutlineBrush(FSlateRoundedBoxBrush(FLinearColor::Transparent, 0.0, FLinearColor::FromSRGBColor(FColor::FromHex("#2C2C2CFF")), 1.0))
+			.SetOutlineBrush(FSlateRoundedBoxBrush(FLinearColor::Transparent, 0.0, FStyleColors::Foldout, 2.0))
 			.SetChildBackgroundBrush(FSlateColorBrush(FStyleColors::Input))
 			.SetOutlineColor(WindowHighlight)
 			.SetCornerRadius(0)
-			.SetBorderPadding(FMargin(2.0, 1.0, 2.0, 2.0))
+			.SetBorderPadding(FMargin(2.0, 2.0, 2.0, 2.0))
 		);
 
 		Style->Set("ChildWindow.Background", new FSlateColorBrush(FStyleColors::Input));
