New common icons, style changes related to using common icons

diff --git a/Engine/Source/Runtime/SlateCore/Private/Styling/StarshipCoreStyle.cpp b/Engine/Source/Runtime/SlateCore/Private/Styling/StarshipCoreStyle.cpp
--- a/Engine/Source/Runtime/SlateCore/Private/Styling/StarshipCoreStyle.cpp
+++ b/Engine/Source/Runtime/SlateCore/Private/Styling/StarshipCoreStyle.cpp
@@ -334,6 +334,20 @@ TSharedRef<ISlateStyle> FStarshipCoreStyle::Create()
 		Style->Set("Icons.Settings", new IMAGE_BRUSH_SVG("Starship/Common/settings", Icon16x16));
 		Style->Set("Icons.Blueprints", new IMAGE_BRUSH_SVG("Starship/Common/blueprint", Icon16x16));
 		Style->Set("Icons.Plus", new IMAGE_BRUSH_SVG("Starship/Common/plus", Icon16x16));
+
+		Style->Set("Icons.Save", new IMAGE_BRUSH_SVG("Starship/Common/save", Icon16x16));
+
+		Style->Set("Icons.Import", new IMAGE_BRUSH_SVG("Starship/Common/import", Icon16x16));
+		Style->Set("Icons.Filter", new IMAGE_BRUSH_SVG("Starship/Common/filter", Icon16x16));
+
+		Style->Set("Icons.Lock", new IMAGE_BRUSH_SVG("Starship/Common/lock", Icon16x16));
+		Style->Set("Icons.Unlock", new IMAGE_BRUSH_SVG("Starship/Common/lock-unlocked", Icon16x16));
+
+		Style->Set("Icons.CircleArrowLeft", new IMAGE_BRUSH_SVG("Starship/Common/circle-arrow-left", Icon16x16));
+		Style->Set("Icons.CircleArrowRight", new IMAGE_BRUSH_SVG("Starship/Common/circle-arrow-right", Icon16x16));
+
+		Style->Set("Icons.CircleArrowUp", new IMAGE_BRUSH_SVG("Starship/Common/circle-arrow-up", Icon16x16));
+		Style->Set("Icons.CircleArrowDown", new IMAGE_BRUSH_SVG("Starship/Common/circle-arrow-down", Icon16x16));
 	}
 
 	// Tool panels
@@ -1030,18 +1044,17 @@ void FStarshipCoreStyle::SetupButtonStyles(TSharedRef<FStyle>& Style)
 		.SetNormalPadding(ButtonMargins)
 		.SetPressedPadding(ButtonMargins);
 
-
 	const FButtonStyle SimpleButton = FButtonStyle()
 		.SetNormal(FSlateNoResource())
-		.SetHovered(FSlateRoundedBoxBrush(FStyleColors::Hover, 4.0f))
-		.SetPressed(FSlateRoundedBoxBrush(FStyleColors::Input, 4.0f))
+		.SetHovered(FSlateNoResource())
+		.SetPressed(FSlateNoResource())
 		.SetDisabled(FSlateNoResource())
 		.SetNormalForeground(FStyleColors::Foreground)
 		.SetHoveredForeground(FStyleColors::ForegroundHover)
 		.SetPressedForeground(FStyleColors::ForegroundHover)
 		.SetDisabledForeground(FStyleColors::Foreground)
-		.SetNormalPadding(FMargin(4.0f, 2.0f))
-		.SetPressedPadding(FMargin(4.0f, 2.0f));
+		.SetNormalPadding(FMargin(2, 2, 2, 2))
+		.SetPressedPadding(FMargin(2, 3, 2, 1));
 	{
 
 		const FTextBlockStyle& NormalText = Style->GetWidgetStyle<FTextBlockStyle>("NormalText");
@@ -1375,9 +1388,9 @@ void FStarshipCoreStyle::SetupDockingStyles(TSharedRef<FStyle>& Style)
 	const FSlateColor& TabFlashColor = SelectionColor;
 
 	const FButtonStyle CloseButton = FButtonStyle()
-		.SetNormal(IMAGE_BRUSH("Starship/Icons/Quixel/Common/cross_32x", Icon16x16, FStyleColors::Foreground))
-		.SetPressed(IMAGE_BRUSH("Starship/Icons/Quixel/Common/cross_32x", Icon16x16, FStyleColors::Foreground))
-		.SetHovered(IMAGE_BRUSH("Starship/Icons/Quixel/Common/cross_32x", Icon16x16, FStyleColors::ForegroundHover));
+		.SetNormal(IMAGE_BRUSH_SVG("Starship/Common/close", Icon16x16, FStyleColors::Foreground))
+		.SetPressed(IMAGE_BRUSH_SVG("Starship/Common/close", Icon16x16, FStyleColors::Foreground))
+		.SetHovered(IMAGE_BRUSH_SVG("Starship/Common/close", Icon16x16, FStyleColors::ForegroundHover));
 
 
 	FLinearColor DockColor_Inactive(FColor(45, 45, 45));
@@ -1738,7 +1751,7 @@ void FStarshipCoreStyle::SetupMultiboxStyles(TSharedRef<FStyle>& Style)
 		Style->Set("Menu.Outline", new BORDER_BRUSH("Common/Window/WindowOutline", FMargin(1.0f / 32.0f), WindowHighlight));
 		Style->Set("Menu.Icon", new IMAGE_BRUSH("Icons/icon_tab_toolbar_16px", Icon16x16));
 		Style->Set("Menu.Expand", new IMAGE_BRUSH("Icons/toolbar_expand_16x", Icon16x16));
-		Style->Set("Menu.SubMenuIndicator", new IMAGE_BRUSH("Starship/Icons/Quixel/Arrows/chevron-right_32x", Icon16x16, FStyleColors::Foreground));
+		Style->Set("Menu.SubMenuIndicator", new IMAGE_BRUSH_SVG("Starship/Common/chevron-right", Icon16x16, FStyleColors::Foreground));
 		Style->Set("Menu.SToolBarComboButtonBlock.Padding", FMargin(4.0f));
 		Style->Set("Menu.SToolBarButtonBlock.Padding", FMargin(4.0f));
 		Style->Set("Menu.SToolBarCheckComboButtonBlock.Padding", FMargin(4.0f));
diff --git a/Engine/Source/Runtime/SlateCore/Public/Styling/StyleColors.h b/Engine/Source/Runtime/SlateCore/Public/Styling/StyleColors.h
--- a/Engine/Source/Runtime/SlateCore/Public/Styling/StyleColors.h
+++ b/Engine/Source/Runtime/SlateCore/Public/Styling/StyleColors.h
@@ -137,7 +137,10 @@ private:
 	FStyleColorList ColorList;
 };
 
-
+/**
+ * Common/themeable colors used by all styles
+ * Please avoid adding new generic colors to this list without discussion first
+ */
 struct SLATECORE_API FStyleColors
 {
 	static const FSlateColor Black;
