Fix inconsistent scaling between rounded box radius and outline size

diff --git a/Engine/Source/Runtime/SlateCore/Private/Rendering/DrawElements.cpp b/Engine/Source/Runtime/SlateCore/Private/Rendering/DrawElements.cpp
--- a/Engine/Source/Runtime/SlateCore/Private/Rendering/DrawElements.cpp
+++ b/Engine/Source/Runtime/SlateCore/Private/Rendering/DrawElements.cpp
@@ -230,7 +230,7 @@ FSlateDrawElement& FSlateDrawElement::MakeBoxInternal(
 			Radius = PaintGeometry.GetLocalSize().Y/2.0;
 		}
 		RBoxPayload->SetRadius(Radius);
-		RBoxPayload->SetOutline(InBrush->OutlineSettings.Color.GetSpecifiedColor(), InBrush->OutlineSettings.Width * PaintGeometry.DrawScale);
+		RBoxPayload->SetOutline(InBrush->OutlineSettings.Color.GetSpecifiedColor(), InBrush->OutlineSettings.Width);
 		BoxPayload = RBoxPayload;
 	}
 	else 
