Remove a few outdated style colors

diff --git a/Engine/Source/Runtime/SlateCore/Private/Styling/StarshipCoreStyle.cpp b/Engine/Source/Runtime/SlateCore/Private/Styling/StarshipCoreStyle.cpp
--- a/Engine/Source/Runtime/SlateCore/Private/Styling/StarshipCoreStyle.cpp
+++ b/Engine/Source/Runtime/SlateCore/Private/Styling/StarshipCoreStyle.cpp
@@ -85,8 +85,6 @@ public:
 		: FSlateStyleSet(InStyleSetName)
 
 		// These are the colors that are updated by the user style customizations
-		, DefaultForeground_LinearRef(MakeShared<FLinearColor>(0.72f, 0.72f, 0.72f, 1.f))
-		, InvertedForeground_LinearRef(MakeShared<FLinearColor>(0.0f, 0.0f, 0.0f))
 		, SelectorColor_LinearRef(MakeShared<FLinearColor>(0.701f, 0.225f, 0.003f))
 		, SelectionColor_LinearRef(MakeShared<FLinearColor>(COLOR("18A0FBFF")))
 		, SelectionColor_Inactive_LinearRef(MakeShared<FLinearColor>(0.25f, 0.25f, 0.25f))
@@ -104,8 +102,6 @@ public:
 	}
 
 	// These are the colors that are updated by the user style customizations
-	const TSharedRef<FLinearColor> DefaultForeground_LinearRef;
-	const TSharedRef<FLinearColor> InvertedForeground_LinearRef;
 	const TSharedRef<FLinearColor> SelectorColor_LinearRef;
 	const TSharedRef<FLinearColor> SelectionColor_LinearRef;
 	const TSharedRef<FLinearColor> SelectionColor_Inactive_LinearRef;
@@ -201,7 +197,7 @@ TSharedRef<ISlateStyle> FStarshipCoreStyle::Create()
 
 	// These are the Slate colors which reference the dynamic colors in the style;
 	const FSlateColor DefaultForeground(FStyleColors::Foreground);
-	const FSlateColor InvertedForeground(Style->InvertedForeground_LinearRef);
+	const FSlateColor InvertedForeground(FStyleColors::ForegroundInverted);
 	const FSlateColor SelectorColor(Style->SelectorColor_LinearRef);
 	const FSlateColor SelectionColor(Style->SelectionColor_LinearRef);
 	const FSlateColor SelectionColor_Inactive(Style->SelectionColor_Inactive_LinearRef);
