Updated Style Colors and Color Usage

diff --git a/Engine/Source/Editor/ContentBrowser/Private/SContentBrowser.cpp b/Engine/Source/Editor/ContentBrowser/Private/SContentBrowser.cpp
--- a/Engine/Source/Editor/ContentBrowser/Private/SContentBrowser.cpp
+++ b/Engine/Source/Editor/ContentBrowser/Private/SContentBrowser.cpp
@@ -337,7 +337,7 @@ void SContentBrowser::Construct( const FArguments& InArgs, const FName& InInstan
 				[
 					SNew(SBorder)
 					.Padding(FMargin(0.0f, 2.0f, 0.0f, 3.0f))
-					.BorderImage(FEditorStyle::GetBrush("Brushes.Input"))
+					.BorderImage(FEditorStyle::GetBrush("Brushes.Recessed"))
 					[
 						// Note: If adding more widgets here, fix ContentBrowserSourcesWidgetSwitcherIndex and the code that uses it!
 						SAssignNew(SourcesWidgetSwitcher, SWidgetSwitcher)
@@ -737,7 +737,7 @@ TSharedRef<SWidget> SContentBrowser::CreateFavoritesView(const FContentBrowserCo
 	return
 		SAssignNew(FavoritesArea, SExpandableArea)
 		.BorderImage(FAppStyle::Get().GetBrush("Brushes.Background"))
-		.BodyBorderImage(FAppStyle::Get().GetBrush("Brushes.Input"))
+		.BodyBorderImage(FAppStyle::Get().GetBrush("Brushes.Recessed"))
 		.HeaderPadding(FMargin(5.0f, 3.0f))
 		.Visibility(this, &SContentBrowser::GetFavoriteFolderVisibility)
 		.AllowAnimatedTransition(false)
@@ -768,7 +768,7 @@ TSharedRef<SWidget> SContentBrowser::CreatePathView(const FContentBrowserConfig*
 	return
 		SAssignNew(PathArea, SExpandableArea)
 		.BorderImage(FAppStyle::Get().GetBrush("Brushes.Background"))
-		.BodyBorderImage(FAppStyle::Get().GetBrush("Brushes.Input"))
+		.BodyBorderImage(FAppStyle::Get().GetBrush("Brushes.Recessed"))
 		.HeaderPadding(FMargin(5.0f, 3.0f))
 		.AllowAnimatedTransition(false)
 		.HeaderContent()
@@ -811,7 +811,7 @@ TSharedRef<SWidget> SContentBrowser::CreateDockedCollectionsView(const FContentB
 	return
 		SAssignNew(CollectionArea, SExpandableArea)
 		.BorderImage(FAppStyle::Get().GetBrush("Brushes.Background"))
-		.BodyBorderImage(FAppStyle::Get().GetBrush("Brushes.Input"))
+		.BodyBorderImage(FAppStyle::Get().GetBrush("Brushes.Recessed"))
 		.HeaderPadding(FMargin(5.0f, 3.0f))
 		.Visibility(this, &SContentBrowser::GetDockedCollectionsVisibility)
 		.AllowAnimatedTransition(false)
diff --git a/Engine/Source/Editor/EditorStyle/Private/StarshipStyle.cpp b/Engine/Source/Editor/EditorStyle/Private/StarshipStyle.cpp
--- a/Engine/Source/Editor/EditorStyle/Private/StarshipStyle.cpp
+++ b/Engine/Source/Editor/EditorStyle/Private/StarshipStyle.cpp
@@ -1273,7 +1273,7 @@ void FStarshipEditorStyle::FStyle::SetupGeneralStyles()
 		
 	// Asset Thumbnail
 	{
-		Set( "AssetThumbnail.AssetBackground", new FSlateColorBrush(FStyleColors::Input));
+		Set( "AssetThumbnail.AssetBackground", new FSlateColorBrush(FStyleColors::Recessed));
 		Set( "AssetThumbnail.ClassBackground", new IMAGE_BRUSH( "Common/ClassBackground_64x", FVector2D(64.f, 64.f), FLinearColor(0.75f, 0.75f, 0.75f, 1.0f) ) );
 		Set( "AssetThumbnail.Font", DEFAULT_FONT( "Regular", 10 ) );
 		Set( "AssetThumbnail.FontSmall", DEFAULT_FONT( "Regular", 7 ) );
@@ -1773,7 +1773,7 @@ void FStarshipEditorStyle::FStyle::SetupGeneralStyles()
 
 		PaletteToolBarStyle.SetLabelStyle( GetParentStyle()->GetWidgetStyle<FTextBlockStyle>("SmallText") );
 		
-		PaletteToolBarStyle.SetBackground(FSlateColorBrush(FStyleColors::Input));
+		PaletteToolBarStyle.SetBackground(FSlateColorBrush(FStyleColors::Recessed));
 
 		PaletteToolBarStyle.SetLabelPadding(FMargin(0.0f, 4.0f, 0.0f, 0.0f));
 
@@ -1785,12 +1785,12 @@ void FStarshipEditorStyle::FStyle::SetupGeneralStyles()
 
 		Set( "PaletteToolBar.Tab",  FCheckBoxStyle()
 			.SetCheckBoxType(            ESlateCheckBoxType::ToggleButton)
-			.SetCheckedImage(            FSlateRoundedBoxBrush(FStyleColors::Input, 2.0))
-			.SetCheckedHoveredImage(     FSlateRoundedBoxBrush(FStyleColors::Input, 2.0))
-			.SetCheckedPressedImage(     FSlateRoundedBoxBrush(FStyleColors::Input, 2.0))
-			.SetUncheckedImage(          FSlateRoundedBoxBrush(FStyleColors::Dropdown, 2.0))
-			.SetUncheckedHoveredImage(   FSlateRoundedBoxBrush(FStyleColors::Dropdown, 2.0))
-			.SetUncheckedPressedImage(   FSlateRoundedBoxBrush(FStyleColors::Dropdown, 2.0))
+			.SetCheckedImage(            FSlateRoundedBoxBrush(FStyleColors::Recessed, 2.0))
+			.SetCheckedHoveredImage(     FSlateRoundedBoxBrush(FStyleColors::Recessed, 2.0))
+			.SetCheckedPressedImage(     FSlateRoundedBoxBrush(FStyleColors::Recessed, 2.0))
+			.SetUncheckedImage(          FSlateRoundedBoxBrush(FStyleColors::Secondary, 2.0))
+			.SetUncheckedHoveredImage(   FSlateRoundedBoxBrush(FStyleColors::Secondary, 2.0))
+			.SetUncheckedPressedImage(   FSlateRoundedBoxBrush(FStyleColors::Secondary, 2.0))
 			.SetForegroundColor(         FStyleColors::Foreground)
 			.SetHoveredForegroundColor(  FStyleColors::ForegroundHover)
 			.SetPressedForegroundColor(  FStyleColors::ForegroundHover)
@@ -1810,9 +1810,9 @@ void FStarshipEditorStyle::FStyle::SetupGeneralStyles()
 				.SetUncheckedImage(       FSlateRoundedBoxBrush(FStyleColors::Background, 2.0))
 				.SetUncheckedPressedImage(FSlateRoundedBoxBrush(FStyleColors::Background, 2.0))
 				.SetUncheckedHoveredImage(FSlateRoundedBoxBrush(FStyleColors::Background, 2.0))
-				.SetCheckedImage(         FSlateRoundedBoxBrush(FStyleColors::Input, 2.0))
-				.SetCheckedHoveredImage(  FSlateRoundedBoxBrush(FStyleColors::Input, 2.0))
-				.SetCheckedPressedImage(  FSlateRoundedBoxBrush(FStyleColors::Input, 2.0))
+				.SetCheckedImage(         FSlateRoundedBoxBrush(FStyleColors::Recessed, 2.0))
+				.SetCheckedHoveredImage(  FSlateRoundedBoxBrush(FStyleColors::Recessed, 2.0))
+				.SetCheckedPressedImage(  FSlateRoundedBoxBrush(FStyleColors::Recessed, 2.0))
 				.SetForegroundColor(       FStyleColors::Foreground)
 				.SetHoveredForegroundColor(FStyleColors::ForegroundHover)
 				.SetPressedForegroundColor(FStyleColors::ForegroundHover)
@@ -3173,7 +3173,7 @@ void FStarshipEditorStyle::FStyle::SetupPropertyEditorStyles()
 		Set( "DetailsView.CategoryTop",               new FSlateColorBrush(FStyleColors::Background));
 		Set( "DetailsView.CollapsedCategory_Hovered", new FSlateColorBrush(FStyleColors::Header));
 		Set( "DetailsView.CategoryTop_Hovered",       new FSlateColorBrush(FStyleColors::Header));
-		Set( "DetailsView.CategoryBottom",            new FSlateColorBrush(FStyleColors::Input));
+		Set( "DetailsView.CategoryBottom",            new FSlateColorBrush(FStyleColors::Recessed));
 
 		Set( "DetailsView.CategoryMiddle",            new FSlateColorBrush(FStyleColors::Background));
 		Set( "DetailsView.CategoryMiddle_Hovered",    new FSlateColorBrush(FStyleColors::Header));
@@ -3186,7 +3186,7 @@ void FStarshipEditorStyle::FStyle::SetupPropertyEditorStyles()
 
 		Set( "DetailsView.Splitter", FSplitterStyle()
 			.SetHandleNormalBrush(   FSlateNoResource())                   
-			.SetHandleHighlightBrush(FSlateColorBrush(FStyleColors::Input))
+			.SetHandleHighlightBrush(FSlateColorBrush(FStyleColors::Recessed))
 		);
 
 		Set( "DetailsView.AdvancedDropdownBorder",      new FSlateColorBrush(FStyleColors::Background));
@@ -5926,7 +5926,7 @@ void FStarshipEditorStyle::FStyle::SetupContentBrowserStyle()
 		Set("ContentBrowser.DefaultFolderColor", FStyleColors::AccentFolder);
 
 		Set( "ContentBrowser.Splitter", FSplitterStyle(FStarshipCoreStyle::GetCoreStyle().GetWidgetStyle<FSplitterStyle>("Splitter"))
-			.SetHandleNormalBrush(FSlateColorBrush(FStyleColors::Input))
+			.SetHandleNormalBrush(FSlateColorBrush(FStyleColors::Recessed))
 			);
 
 		// Asset list view
@@ -6006,8 +6006,8 @@ void FStarshipEditorStyle::FStyle::SetupContentBrowserStyle()
 			.SetPressedForegroundColor(FStyleColors::ForegroundHover);
 
 		Set("ContentBrowser.FilterImage", new CORE_IMAGE_BRUSH_SVG("Starship/Common/filled-circle", Icon8x8));
-		Set("ContentBrowser.FilterBackgroundUnchecked", new FSlateRoundedBoxBrush(FStyleColors::Dropdown, 4.0f));
-		Set("ContentBrowser.FilterBackgroundChecked", new FSlateRoundedBoxBrush(FStyleColors::Input, 4.0f));
+		Set("ContentBrowser.FilterBackgroundUnchecked", new FSlateRoundedBoxBrush(FStyleColors::Secondary, 4.0f));
+		Set("ContentBrowser.FilterBackgroundChecked", new FSlateRoundedBoxBrush(FStyleColors::Recessed, 4.0f));
 		/* ... and add the new style */
 		Set("ContentBrowser.FilterButton", ContentBrowserFilterButtonCheckBoxStyle );
 
@@ -6127,9 +6127,9 @@ void FStarshipEditorStyle::FStyle::SetupContentBrowserStyle()
 		Set( "ContentBrowser.ListViewFolderIcon", new IMAGE_BRUSH_SVG( "Starship/ContentBrowser/folder", FVector2D(64, 64) ) );
 		Set( "ContentBrowser.ListViewDeveloperFolderIcon", new IMAGE_BRUSH( "Icons/Folders/FolderDev_Base_256x", FVector2D(256, 256) ) );
 
-		Set("ContentBrowser.AssetTileItem.FolderAreaHoveredBackground", new FSlateRoundedBoxBrush(FStyleColors::Dropdown, 4.0f));
-		Set("ContentBrowser.AssetTileItem.ThumbnailAreaBackground", new FSlateRoundedBoxBrush(FStyleColors::Input, 4.0f));
-		Set("ContentBrowser.AssetTileItem.NameAreaBackground", new FSlateRoundedBoxBrush(FStyleColors::Dropdown, 4.0f));
+		Set("ContentBrowser.AssetTileItem.FolderAreaHoveredBackground", new FSlateRoundedBoxBrush(FStyleColors::Secondary, 4.0f));
+		Set("ContentBrowser.AssetTileItem.ThumbnailAreaBackground", new FSlateRoundedBoxBrush(FStyleColors::Recessed, 4.0f));
+		Set("ContentBrowser.AssetTileItem.NameAreaBackground", new FSlateRoundedBoxBrush(FStyleColors::Secondary, 4.0f));
 		Set("ContentBrowser.AssetTileItem.SelectedBorder", new FSlateRoundedBoxBrush(FStyleColors::Transparent, 4.0f, FStyleColors::Primary, 1.0f));
 		Set("ContentBrowser.AssetTileItem.SelectedHoverBorder", new FSlateRoundedBoxBrush(FStyleColors::Transparent, 4.0f, FStyleColors::PrimaryHover, 1.0f));
 		Set("ContentBrowser.AssetTileItem.HoverBorder", new FSlateRoundedBoxBrush(FStyleColors::Transparent, 4.0f, FStyleColors::Hover, 1.0f));
diff --git a/Engine/Source/Editor/SceneOutliner/Private/SSceneOutliner.cpp b/Engine/Source/Editor/SceneOutliner/Private/SSceneOutliner.cpp
--- a/Engine/Source/Editor/SceneOutliner/Private/SSceneOutliner.cpp
+++ b/Engine/Source/Editor/SceneOutliner/Private/SSceneOutliner.cpp
@@ -178,7 +178,7 @@ void SSceneOutliner::Construct(const FArguments& InArgs, const FSceneOutlinerIni
 
 		+SOverlay::Slot()
 		[
-			SNew(SBorder).BorderImage( FAppStyle::Get().GetBrush("InputBrush") )
+			SNew(SBorder).BorderImage( FAppStyle::Get().GetBrush("Brushes.Recessed") )
 		]
 
 		+SOverlay::Slot()
diff --git a/Engine/Source/Runtime/AppFramework/Private/Framework/Testing/SStarshipGallery.cpp b/Engine/Source/Runtime/AppFramework/Private/Framework/Testing/SStarshipGallery.cpp
--- a/Engine/Source/Runtime/AppFramework/Private/Framework/Testing/SStarshipGallery.cpp
+++ b/Engine/Source/Runtime/AppFramework/Private/Framework/Testing/SStarshipGallery.cpp
@@ -207,9 +207,9 @@ public:
             // TSharedRef<SButton> ColorButton;
             // = SNew(SButton)
             return SNew(SBox)
-            .Padding(12)
+            .Padding(8)
             .WidthOverride(120)
-            .HeightOverride(120)
+            .HeightOverride(100)
             [
                     SNew(SOverlay)
 
@@ -261,88 +261,78 @@ public:
             SNew(SVerticalBox)
 
             +SVerticalBox::Slot()
+            .AutoHeight()
+            .Padding(8.f, 24.f, 8.f, 8.f)
             [
-                SNew(SHorizontalBox)
-
-                +SHorizontalBox::Slot()
-                .FillWidth(1.0)
-                [
-                    SNew(SVerticalBox)
-
-                    +SVerticalBox::Slot()
-                    .AutoHeight()
-                    .Padding(8.f, 0.f, 8.f, 12.f)
-                    [
-                        SNew(STextBlock).ColorAndOpacity(LabelColor).Text(NSLOCTEXT("StarshipGallery", "BaseColors", "BASE COLORS"))
-                    ]
-
-                    +SVerticalBox::Slot()
-                    .AutoHeight()
-                    [
-                        SNew(SUniformWrapPanel)
-                        +SUniformWrapPanel::Slot()[ GenerateColorButton( FText::FromString("BLACK"),           "Colors.Black")]
-                        +SUniformWrapPanel::Slot()[ GenerateColorButton( FText::FromString("FOLDOUT"),         "Colors.Foldout")]
-                        +SUniformWrapPanel::Slot()[ GenerateColorButton( FText::FromString("INPUT"),           "Colors.Input")]
-                        +SUniformWrapPanel::Slot()[ GenerateColorButton( FText::FromString("BACKGROUND"),      "Colors.Background", false, true)]
-                        +SUniformWrapPanel::Slot()[ GenerateColorButton( FText::FromString("HEADER"),          "Colors.Header")]
-                        +SUniformWrapPanel::Slot()[ GenerateColorButton( FText::FromString("DROPDOWN"),        "Colors.Dropdown")]
-                        +SUniformWrapPanel::Slot()[ GenerateColorButton( FText::FromString("HOVER"),           "Colors.Hover")]
-                        +SUniformWrapPanel::Slot()[ GenerateColorButton( FText::FromString("HOVER2"),          "Colors.Hover2")]
-                        +SUniformWrapPanel::Slot()[ GenerateColorButton( FText::FromString("WHITE"),           "Colors.White", true)]
-                        +SUniformWrapPanel::Slot()[ GenerateColorButton( FText::FromString("WHITE25"),         "Colors.White25")]
+                SNew(STextBlock).ColorAndOpacity(LabelColor).Text(NSLOCTEXT("StarshipGallery", "BaseColors", "BASE COLORS"))
+            ]
 
-                    ]
+            +SVerticalBox::Slot()
+            .AutoHeight()
+            [
+                SNew(SUniformWrapPanel)
+                +SUniformWrapPanel::Slot()[ GenerateColorButton( FText::FromString("BLACK"),           "Colors.Black")]
+                +SUniformWrapPanel::Slot()[ GenerateColorButton( FText::FromString("TITLE"),           "Colors.Title")]
+                +SUniformWrapPanel::Slot()[ GenerateColorButton( FText::FromString("WINDOW\nBORDER"),  "Colors.WindowBorder")]
+                +SUniformWrapPanel::Slot()[ GenerateColorButton( FText::FromString("FOLDOUT"),         "Colors.Foldout")]
+                +SUniformWrapPanel::Slot()[ GenerateColorButton( FText::FromString("INPUT"),           "Colors.Input")]
+                +SUniformWrapPanel::Slot()[ GenerateColorButton( FText::FromString("RECESSED"),        "Colors.Recessed")]
+                +SUniformWrapPanel::Slot()[ GenerateColorButton( FText::FromString("BACKGROUND"),      "Colors.Background", false, true)]
+                +SUniformWrapPanel::Slot()[ GenerateColorButton( FText::FromString("HEADER"),          "Colors.Header")]
+                +SUniformWrapPanel::Slot()[ GenerateColorButton( FText::FromString("DROPDOWN"),        "Colors.Dropdown")]
+                +SUniformWrapPanel::Slot()[ GenerateColorButton( FText::FromString("HOVER"),           "Colors.Hover")]
+                +SUniformWrapPanel::Slot()[ GenerateColorButton( FText::FromString("HOVER2"),          "Colors.Hover2")]
+                +SUniformWrapPanel::Slot()[ GenerateColorButton( FText::FromString("WHITE"),           "Colors.White", true)]
+                +SUniformWrapPanel::Slot()[ GenerateColorButton( FText::FromString("WHITE25"),         "Colors.White25")]
 
-                    +SVerticalBox::Slot()
-                    .AutoHeight()
-                    [
-                        SNew(SUniformWrapPanel)
-                        +SUniformWrapPanel::Slot()[ GenerateColorButton( FText::FromString("PRIMARY"),         "Colors.Primary", true)]
-                        +SUniformWrapPanel::Slot()[ GenerateColorButton( FText::FromString("PRIMARY\nHOVER"),  "Colors.PrimaryHover", true)]
-                        +SUniformWrapPanel::Slot()[ GenerateColorButton( FText::FromString("PRIMARY\nPRESS"),  "Colors.PrimaryPress", true)]
-                    ]
-                ]
+                /*
+            ]
 
-                +SHorizontalBox::Slot()
-                .FillWidth(1.0)
-                .Padding(48.0, 0.0)
-                [
-                    SNew(SVerticalBox)
+            +SVerticalBox::Slot()
+            .AutoHeight()
+            [
+                SNew(SUniformWrapPanel)
+*/
+                +SUniformWrapPanel::Slot()[ GenerateColorButton( FText::FromString("PRIMARY"),         "Colors.Primary", true)]
+                +SUniformWrapPanel::Slot()[ GenerateColorButton( FText::FromString("PRIMARY\nHOVER"),  "Colors.PrimaryHover", true)]
+                +SUniformWrapPanel::Slot()[ GenerateColorButton( FText::FromString("PRIMARY\nPRESS"),  "Colors.PrimaryPress", true)]
+                +SUniformWrapPanel::Slot()[ GenerateColorButton( FText::FromString("SECONDARY"),       "Colors.Secondary", true)]
+            ]
 
-                    +SVerticalBox::Slot()
-                    .Padding(8.f, 0.f, 8.f, 12.f)
-                    .AutoHeight()
-                    [
-                        SNew(STextBlock).ColorAndOpacity(LabelColor).Text(NSLOCTEXT("StarshipGallery", "TextIconColors", "FOREGROUND COLORS"))
+            +SVerticalBox::Slot()
+            .Padding(8.f, 24.f, 8.f, 8.f)
+            .AutoHeight()
+            [
+                SNew(STextBlock).ColorAndOpacity(LabelColor).Text(NSLOCTEXT("StarshipGallery", "TextIconColors", "FOREGROUND COLORS"))
 
-                    ]
+            ]
 
-                    +SVerticalBox::Slot()
-                    .AutoHeight()
-                    [
-                        SNew(SUniformWrapPanel)
+            +SVerticalBox::Slot()
+            .AutoHeight()
+            [
+                SNew(SUniformWrapPanel)
 
-                        +SUniformWrapPanel::Slot()[ GenerateColorButton( FText::FromString("FOREGROUND"),        "Colors.Foreground")]
-                        +SUniformWrapPanel::Slot()[ GenerateColorButton( FText::FromString("FOREGROUND\nHOVER"), "Colors.ForegroundHover", true)]
-                        +SUniformWrapPanel::Slot()[ GenerateColorButton( FText::FromString("FOREGROUND\nINVERTED"), "Colors.ForegroundInverted")]
-                        +SUniformWrapPanel::Slot()[ GenerateColorButton( FText::FromString("HIGHLIGHT"),         "Colors.Highlight")]
-                        +SUniformWrapPanel::Slot()[ GenerateColorButton( FText::FromString("FOREGROUND\nHEADER"), "Colors.ForegroundHeader", true)]
-                    ]
+                +SUniformWrapPanel::Slot()[ GenerateColorButton( FText::FromString("FOREGROUND"),        "Colors.Foreground")]
+                +SUniformWrapPanel::Slot()[ GenerateColorButton( FText::FromString("FOREGROUND\nHOVER"), "Colors.ForegroundHover", true)]
+                +SUniformWrapPanel::Slot()[ GenerateColorButton( FText::FromString("FOREGROUND\nINVERTED"), "Colors.ForegroundInverted")]
+                +SUniformWrapPanel::Slot()[ GenerateColorButton( FText::FromString("HIGHLIGHT"),         "Colors.Highlight")]
+                +SUniformWrapPanel::Slot()[ GenerateColorButton( FText::FromString("FOREGROUND\nHEADER"), "Colors.ForegroundHeader", true)]
+                /*
+            ]
 
-                    +SVerticalBox::Slot()
-                    .AutoHeight()
-                    [
-                        SNew(SUniformWrapPanel)
-                        +SUniformWrapPanel::Slot()[ GenerateColorButton( FText::FromString("SELECT"),          "Colors.Select", true)]
-                        +SUniformWrapPanel::Slot()[ GenerateColorButton( FText::FromString("SELECT\nHOVER"),   "Colors.SelectHover")]
-                        +SUniformWrapPanel::Slot()[ GenerateColorButton( FText::FromString("SELECT\nPARENT"),  "Colors.SelectParent")]
-                        +SUniformWrapPanel::Slot()[ GenerateColorButton( FText::FromString("SELECT\nINACTIVE"),"Colors.SelectInactive")]
-                    ]
-                ]
+            +SVerticalBox::Slot()
+            .AutoHeight()
+            [
+                SNew(SUniformWrapPanel)
+            */
+                +SUniformWrapPanel::Slot()[ GenerateColorButton( FText::FromString("SELECT"),          "Colors.Select", true)]
+                +SUniformWrapPanel::Slot()[ GenerateColorButton( FText::FromString("SELECT\nHOVER"),   "Colors.SelectHover")]
+                +SUniformWrapPanel::Slot()[ GenerateColorButton( FText::FromString("SELECT\nPARENT"),  "Colors.SelectParent")]
+                +SUniformWrapPanel::Slot()[ GenerateColorButton( FText::FromString("SELECT\nINACTIVE"),"Colors.SelectInactive")]
             ]
 
             +SVerticalBox::Slot()
-            .Padding(48.f)
+            .Padding(8.f, 24.f, 8.f, 8.f)
             .AutoHeight()
             [
                 SNew(STextBlock).ColorAndOpacity(LabelColor).Text(NSLOCTEXT("StarshipGallery", "AccentColors", "ACCENT COLORS"))
diff --git a/Engine/Source/Runtime/SlateCore/Private/Styling/StarshipCoreStyle.cpp b/Engine/Source/Runtime/SlateCore/Private/Styling/StarshipCoreStyle.cpp
--- a/Engine/Source/Runtime/SlateCore/Private/Styling/StarshipCoreStyle.cpp
+++ b/Engine/Source/Runtime/SlateCore/Private/Styling/StarshipCoreStyle.cpp
@@ -279,7 +279,6 @@ TSharedRef<ISlateStyle> FStarshipCoreStyle::Create()
 
 		Style->Set("NoBorder", NoBorder);
 
-		Style->Set("InputBrush", new FSlateColorBrush(FStyleColors::Input));
 	}
 
 
@@ -502,11 +501,11 @@ TSharedRef<ISlateStyle> FStarshipCoreStyle::Create()
 	// SSpinBox defaults...
 	{
 		Style->Set("SpinBox", FSpinBoxStyle()
-			.SetBackgroundBrush(FSlateRoundedBoxBrush(FStyleColors::Input, InputFocusRadius, FStyleColors::Dropdown, InputFocusThickness))
+			.SetBackgroundBrush(FSlateRoundedBoxBrush(FStyleColors::Input, InputFocusRadius, FStyleColors::Secondary, InputFocusThickness))
 			.SetHoveredBackgroundBrush(FSlateRoundedBoxBrush(FStyleColors::Input, InputFocusRadius, FStyleColors::Hover, InputFocusThickness))
 
 			.SetActiveFillBrush(FSlateRoundedBoxBrush(FStyleColors::Hover, InputFocusRadius, FLinearColor::Transparent, InputFocusThickness))
-			.SetInactiveFillBrush(FSlateRoundedBoxBrush(FStyleColors::Dropdown, InputFocusRadius, FLinearColor::Transparent, InputFocusThickness))
+			.SetInactiveFillBrush(FSlateRoundedBoxBrush(FStyleColors::Secondary, InputFocusRadius, FLinearColor::Transparent, InputFocusThickness))
 			.SetArrowsImage(FSlateNoResource())
 			.SetForegroundColor(FStyleColors::ForegroundHover)
 			.SetTextPadding(FMargin(10.f, 3.5f, 10.f, 4.f))
@@ -516,11 +515,11 @@ TSharedRef<ISlateStyle> FStarshipCoreStyle::Create()
 	// SNumericEntryBox defaults...
 	{
 		Style->Set("NumericEntrySpinBox", FSpinBoxStyle()
-			.SetBackgroundBrush(FSlateRoundedBoxBrush(FStyleColors::Input, InputFocusRadius, FStyleColors::Dropdown, InputFocusThickness))
+			.SetBackgroundBrush(FSlateRoundedBoxBrush(FStyleColors::Input, InputFocusRadius, FStyleColors::Secondary, InputFocusThickness))
 			.SetHoveredBackgroundBrush(FSlateRoundedBoxBrush(FStyleColors::Input, InputFocusRadius, FStyleColors::Hover, InputFocusThickness))
 
 			.SetActiveFillBrush(FSlateRoundedBoxBrush(FStyleColors::Hover, InputFocusRadius, FLinearColor::Transparent, InputFocusThickness))
-			.SetInactiveFillBrush(FSlateRoundedBoxBrush(FStyleColors::Dropdown, InputFocusRadius, FLinearColor::Transparent, InputFocusThickness))
+			.SetInactiveFillBrush(FSlateRoundedBoxBrush(FStyleColors::Secondary, InputFocusRadius, FLinearColor::Transparent, InputFocusThickness))
 
 			.SetArrowsImage(FSlateNoResource())
 			.SetTextPadding(FMargin(0.f))
@@ -528,11 +527,11 @@ TSharedRef<ISlateStyle> FStarshipCoreStyle::Create()
 		);
 
 		Style->Set("NumericEntrySpinBox_Dark", FSpinBoxStyle()
-			.SetBackgroundBrush(FSlateRoundedBoxBrush(FStyleColors::Input, InputFocusRadius, FStyleColors::Dropdown, InputFocusThickness))
+			.SetBackgroundBrush(FSlateRoundedBoxBrush(FStyleColors::Input, InputFocusRadius, FStyleColors::Secondary, InputFocusThickness))
 			.SetHoveredBackgroundBrush(FSlateRoundedBoxBrush(FStyleColors::Input, InputFocusRadius, FStyleColors::Hover, InputFocusThickness))
 
 			.SetActiveFillBrush(FSlateRoundedBoxBrush(FStyleColors::Hover, InputFocusRadius, FLinearColor::Transparent, InputFocusThickness))
-			.SetInactiveFillBrush(FSlateRoundedBoxBrush(FStyleColors::Dropdown, InputFocusRadius, FLinearColor::Transparent, InputFocusThickness))
+			.SetInactiveFillBrush(FSlateRoundedBoxBrush(FStyleColors::Secondary, InputFocusRadius, FLinearColor::Transparent, InputFocusThickness))
 
 			.SetArrowsImage(FSlateNoResource())
 
@@ -680,16 +679,16 @@ TSharedRef<ISlateStyle> FStarshipCoreStyle::Create()
 			.SetActiveTitleBrush(FSlateNoResource())
 			.SetInactiveTitleBrush(FSlateNoResource())
 			.SetFlashTitleBrush(IMAGE_BRUSH("Common/Window/WindowTitle_Flashing", Icon24x24, FLinearColor(1.0f, 1.0f, 1.0f, 1.0f), ESlateBrushTileType::Horizontal))	
-			.SetBackgroundBrush(FSlateColorBrush(FStyleColors::Input))
+			.SetBackgroundBrush(FSlateColorBrush(FStyleColors::Recessed))
 			.SetBorderBrush(FSlateNoResource())
-			.SetOutlineBrush(FSlateRoundedBoxBrush(FLinearColor::Transparent, 0.0, FStyleColors::Foldout, 2.0))
-			.SetChildBackgroundBrush(FSlateColorBrush(FStyleColors::Input))
+			.SetOutlineBrush(FSlateRoundedBoxBrush(FLinearColor::Transparent, 0.0, FStyleColors::WindowBorder, 2.0))
+			.SetChildBackgroundBrush(FSlateColorBrush(FStyleColors::Recessed))
 			.SetOutlineColor(WindowHighlight)
 			.SetCornerRadius(0)
 			.SetBorderPadding(FMargin(2.0, 2.0, 2.0, 2.0))
 		);
 
-		Style->Set("ChildWindow.Background", new FSlateColorBrush(FStyleColors::Input));
+		Style->Set("ChildWindow.Background", new FSlateColorBrush(FStyleColors::Recessed));
 	}
 
 
@@ -798,8 +797,10 @@ void FStarshipCoreStyle::SetupColors(TSharedRef<FStyle>& Style)
 {
 	Style->Set("Colors.Black", FStyleColors::Black);
 	Style->Set("Colors.Title", FStyleColors::Title);
+	Style->Set("Colors.WindowBorder", FStyleColors::WindowBorder);
 	Style->Set("Colors.Foldout", FStyleColors::Foldout);
 	Style->Set("Colors.Input", FStyleColors::Input);
+	Style->Set("Colors.Recessed", FStyleColors::Recessed);
 	Style->Set("Colors.Background", FStyleColors::Background);
 	Style->Set("Colors.Header", FStyleColors::Header);
 	Style->Set("Colors.Dropdown", FStyleColors::Dropdown);
@@ -819,10 +820,11 @@ void FStarshipCoreStyle::SetupColors(TSharedRef<FStyle>& Style)
 	Style->Set("Colors.SelectParent", FStyleColors::SelectParent);
 	Style->Set("Colors.SelectHover", FStyleColors::SelectHover);
 
-
 	Style->Set("Colors.Primary", FStyleColors::Primary);
 	Style->Set("Colors.PrimaryHover", FStyleColors::PrimaryHover);
 	Style->Set("Colors.PrimaryPress", FStyleColors::PrimaryPress);
+	Style->Set("Colors.Secondary", FStyleColors::Secondary);
+
 	Style->Set("Colors.AccentBlue", FStyleColors::AccentBlue);
 	Style->Set("Colors.AccentPurple", FStyleColors::AccentPurple);
 	Style->Set("Colors.AccentPink", FStyleColors::AccentPink);
@@ -841,6 +843,7 @@ void FStarshipCoreStyle::SetupColors(TSharedRef<FStyle>& Style)
 	Style->Set("Brushes.Title", new FSlateColorBrush(FStyleColors::Title));
 	Style->Set("Brushes.Foldout", new FSlateColorBrush(FStyleColors::Foldout));
 	Style->Set("Brushes.Input", new FSlateColorBrush(FStyleColors::Input));
+	Style->Set("Brushes.Recessed", new FSlateColorBrush(FStyleColors::Recessed));
 	Style->Set("Brushes.Background", new FSlateColorBrush(FStyleColors::Background));
 	Style->Set("Brushes.Header", new FSlateColorBrush(FStyleColors::Header));
 	Style->Set("Brushes.Dropdown", new FSlateColorBrush(FStyleColors::Dropdown));
@@ -863,6 +866,7 @@ void FStarshipCoreStyle::SetupColors(TSharedRef<FStyle>& Style)
 	Style->Set("Brushes.Primary", new FSlateColorBrush(FStyleColors::Primary));
 	Style->Set("Brushes.PrimaryHover", new FSlateColorBrush(FStyleColors::PrimaryHover));
 	Style->Set("Brushes.PrimaryPress", new FSlateColorBrush(FStyleColors::PrimaryPress));
+	Style->Set("Brushes.Secondary", new FSlateColorBrush(FStyleColors::Secondary));
 	Style->Set("Brushes.AccentBlue", new FSlateColorBrush(FStyleColors::AccentBlue));
 	Style->Set("Brushes.AccentPurple", new FSlateColorBrush(FStyleColors::AccentPurple));
 	Style->Set("Brushes.AccentPink", new FSlateColorBrush(FStyleColors::AccentPink));
@@ -978,7 +982,7 @@ void FStarshipCoreStyle::SetupTextStyles(TSharedRef<FStyle>& Style)
 
 	// SEditableTextBox defaults...
 	const FEditableTextBoxStyle NormalEditableTextBoxStyle = FEditableTextBoxStyle()
-		.SetBackgroundImageNormal(FSlateRoundedBoxBrush(FStyleColors::Input, InputFocusRadius, FStyleColors::Dropdown, InputFocusThickness))
+		.SetBackgroundImageNormal(FSlateRoundedBoxBrush(FStyleColors::Input, InputFocusRadius, FStyleColors::Secondary, InputFocusThickness))
 		.SetBackgroundImageHovered(FSlateRoundedBoxBrush(FStyleColors::Input, InputFocusRadius, FStyleColors::Hover, InputFocusThickness))
 		.SetBackgroundImageFocused(FSlateRoundedBoxBrush(FStyleColors::Input, InputFocusRadius, FStyleColors::Primary, InputFocusThickness))
 		.SetBackgroundImageReadOnly(FSlateRoundedBoxBrush(FStyleColors::Input, InputFocusRadius))
@@ -1052,7 +1056,7 @@ void FStarshipCoreStyle::SetupButtonStyles(TSharedRef<FStyle>& Style)
 
 
 	const FButtonStyle Button = FButtonStyle()
-		.SetNormal(FSlateRoundedBoxBrush(FStyleColors::Dropdown, 4.0f))
+		.SetNormal(FSlateRoundedBoxBrush(FStyleColors::Secondary, 4.0f))
 		.SetHovered(FSlateRoundedBoxBrush(FStyleColors::Hover, 4.0f))
 		.SetPressed(FSlateRoundedBoxBrush(FStyleColors::Header, 4.0f))
 		.SetNormalForeground(FStyleColors::ForegroundHover)
@@ -1091,7 +1095,7 @@ void FStarshipCoreStyle::SetupComboButtonStyles(TSharedRef<FStyle>& Style)
 {
 	// SComboButton and SComboBox defaults...
 	const FButtonStyle ComboButtonButton = FButtonStyle()
-		.SetNormal(FSlateRoundedBoxBrush(FStyleColors::Input, InputFocusRadius, FStyleColors::Dropdown, InputFocusThickness))
+		.SetNormal(FSlateRoundedBoxBrush(FStyleColors::Input, InputFocusRadius, FStyleColors::Secondary, InputFocusThickness))
 		.SetHovered(FSlateRoundedBoxBrush(FStyleColors::Input, InputFocusRadius, FStyleColors::Hover, InputFocusThickness))
 		.SetPressed(FSlateRoundedBoxBrush(FStyleColors::Input, InputFocusRadius, FStyleColors::Hover, InputFocusThickness))
 		.SetNormalForeground(FStyleColors::Foreground)
@@ -1203,7 +1207,7 @@ void FStarshipCoreStyle::SetupCheckboxStyles(TSharedRef<FStyle>& Style)
 		.SetUndeterminedImage(IMAGE_BRUSH_SVG("Starship/CoreWidgets/CheckBox/indeterminate", Icon16x16, FStyleColors::Primary))
 		.SetUndeterminedHoveredImage(IMAGE_BRUSH_SVG("Starship/CoreWidgets/CheckBox/indeterminate", Icon16x16, FStyleColors::PrimaryHover))
 		.SetUndeterminedPressedImage(IMAGE_BRUSH_SVG("Starship/CoreWidgets/CheckBox/indeterminate", Icon16x16, FStyleColors::Primary))
-		.SetBackgroundImage(FSlateRoundedBoxBrush(FStyleColors::Input, CheckboxCornerRadius, FStyleColors::Dropdown, CheckboxOutlineThickness, Icon18x18))
+		.SetBackgroundImage(FSlateRoundedBoxBrush(FStyleColors::Input, CheckboxCornerRadius, FStyleColors::Secondary, CheckboxOutlineThickness, Icon18x18))
 		.SetBackgroundHoveredImage(FSlateRoundedBoxBrush(FStyleColors::Input, CheckboxCornerRadius, FStyleColors::Hover, CheckboxOutlineThickness, Icon18x18))
 		.SetBackgroundPressedImage(FSlateRoundedBoxBrush(FStyleColors::Foldout, CheckboxCornerRadius, FStyleColors::Hover, CheckboxOutlineThickness, Icon18x18));
 
@@ -1229,7 +1233,7 @@ void FStarshipCoreStyle::SetupCheckboxStyles(TSharedRef<FStyle>& Style)
 		.SetUndeterminedImage(IMAGE_BRUSH_SVG("Starship/CoreWidgets/CheckBox/indeterminate", Icon16x16, FStyleColors::Foreground))
 		.SetUndeterminedHoveredImage(IMAGE_BRUSH_SVG("Starship/CoreWidgets/CheckBox/indeterminate", Icon16x16, FStyleColors::ForegroundHover))
 		.SetUndeterminedPressedImage(IMAGE_BRUSH_SVG("Starship/CoreWidgets/CheckBox/indeterminate", Icon16x16, FStyleColors::Foreground))
-		.SetBackgroundImage(FSlateRoundedBoxBrush(FStyleColors::Input, CheckboxCornerRadius, FStyleColors::Dropdown, CheckboxOutlineThickness, Icon18x18))
+		.SetBackgroundImage(FSlateRoundedBoxBrush(FStyleColors::Input, CheckboxCornerRadius, FStyleColors::Secondary, CheckboxOutlineThickness, Icon18x18))
 		.SetBackgroundHoveredImage(FSlateRoundedBoxBrush(FStyleColors::Input, CheckboxCornerRadius, FStyleColors::Hover, CheckboxOutlineThickness, Icon18x18))
 		.SetBackgroundPressedImage(FSlateRoundedBoxBrush(FStyleColors::Foldout, CheckboxCornerRadius, FStyleColors::Hover, CheckboxOutlineThickness, Icon18x18));
 
@@ -1291,9 +1295,9 @@ void FStarshipCoreStyle::SetupCheckboxStyles(TSharedRef<FStyle>& Style)
 	const FCheckBoxStyle SegmentedBoxLeft = FCheckBoxStyle()
 		.SetCheckBoxType(ESlateCheckBoxType::ToggleButton)
 
-		.SetUncheckedImage(BOX_BRUSH("/Starship/CoreWidgets/SegmentedBox/left", FVector2D(16.f, 16.f), FMargin(4.0f / 16.0f), FStyleColors::Dropdown))
+		.SetUncheckedImage(BOX_BRUSH("/Starship/CoreWidgets/SegmentedBox/left", FVector2D(16.f, 16.f), FMargin(4.0f / 16.0f), FStyleColors::Secondary))
 		.SetUncheckedHoveredImage(BOX_BRUSH("/Starship/CoreWidgets/SegmentedBox/left", FVector2D(16.f, 16.f), FMargin(4.0f / 16.0f), FStyleColors::Hover))
-		.SetUncheckedPressedImage(BOX_BRUSH("/Starship/CoreWidgets/SegmentedBox/left", FVector2D(16.f, 16.f), FMargin(4.0f / 16.0f), FStyleColors::Dropdown))
+		.SetUncheckedPressedImage(BOX_BRUSH("/Starship/CoreWidgets/SegmentedBox/left", FVector2D(16.f, 16.f), FMargin(4.0f / 16.0f), FStyleColors::Secondary))
 		.SetCheckedImage(BOX_BRUSH("/Starship/CoreWidgets/SegmentedBox/left", FVector2D(16.f, 16.f), FMargin(4.0f / 16.0f), FStyleColors::Input))
 		.SetCheckedHoveredImage(BOX_BRUSH("/Starship/CoreWidgets/SegmentedBox/left", FVector2D(16.f, 16.f), FMargin(4.0f / 16.0f), FStyleColors::Input))
 		.SetCheckedPressedImage(BOX_BRUSH("/Starship/CoreWidgets/SegmentedBox/left", FVector2D(16.f, 16.f), FMargin(4.0f / 16.0f), FStyleColors::Input))
@@ -1308,9 +1312,9 @@ void FStarshipCoreStyle::SetupCheckboxStyles(TSharedRef<FStyle>& Style)
 
 	const FCheckBoxStyle SegmentedBoxCenter = FCheckBoxStyle()
 		.SetCheckBoxType(ESlateCheckBoxType::ToggleButton)
-		.SetUncheckedImage(FSlateColorBrush(FStyleColors::Dropdown))
+		.SetUncheckedImage(FSlateColorBrush(FStyleColors::Secondary))
 		.SetUncheckedHoveredImage(FSlateColorBrush(FStyleColors::Hover))
-		.SetUncheckedPressedImage(FSlateColorBrush(FStyleColors::Dropdown))
+		.SetUncheckedPressedImage(FSlateColorBrush(FStyleColors::Secondary))
 		.SetCheckedImage(FSlateColorBrush(FStyleColors::Input))
 		.SetCheckedHoveredImage(FSlateColorBrush(FStyleColors::Input))
 		.SetCheckedPressedImage(FSlateColorBrush(FStyleColors::Input))
@@ -1325,9 +1329,9 @@ void FStarshipCoreStyle::SetupCheckboxStyles(TSharedRef<FStyle>& Style)
 
 	const FCheckBoxStyle SegmentedBoxRight = FCheckBoxStyle()
 		.SetCheckBoxType(ESlateCheckBoxType::ToggleButton)
-		.SetUncheckedImage(BOX_BRUSH("/Starship/CoreWidgets/SegmentedBox/right", FVector2D(16.f, 16.f), FMargin(4.0f / 16.0f), FStyleColors::Dropdown))
+		.SetUncheckedImage(BOX_BRUSH("/Starship/CoreWidgets/SegmentedBox/right", FVector2D(16.f, 16.f), FMargin(4.0f / 16.0f), FStyleColors::Secondary))
 		.SetUncheckedHoveredImage(BOX_BRUSH("/Starship/CoreWidgets/SegmentedBox/right", FVector2D(16.f, 16.f), FMargin(4.0f / 16.0f), FStyleColors::Hover))
-		.SetUncheckedPressedImage(BOX_BRUSH("/Starship/CoreWidgets/SegmentedBox/right", FVector2D(16.f, 16.f), FMargin(4.0f / 16.0f), FStyleColors::Dropdown))
+		.SetUncheckedPressedImage(BOX_BRUSH("/Starship/CoreWidgets/SegmentedBox/right", FVector2D(16.f, 16.f), FMargin(4.0f / 16.0f), FStyleColors::Secondary))
 		.SetCheckedImage(BOX_BRUSH("/Starship/CoreWidgets/SegmentedBox/right", FVector2D(16.f, 16.f), FMargin(4.0f / 16.0f), FStyleColors::Input))
 		.SetCheckedHoveredImage(BOX_BRUSH("/Starship/CoreWidgets/SegmentedBox/right", FVector2D(16.f, 16.f), FMargin(4.0f / 16.0f), FStyleColors::Input))
 		.SetCheckedPressedImage(BOX_BRUSH("/Starship/CoreWidgets/SegmentedBox/right", FVector2D(16.f, 16.f), FMargin(4.0f / 16.0f), FStyleColors::Input))
@@ -1495,7 +1499,7 @@ void FStarshipCoreStyle::SetupDockingStyles(TSharedRef<FStyle>& Style)
 
 	Style->Set("Docking.Sidebar.DrawerShadow", new BOX_BRUSH("/Starship/Docking/drawer-shadow", FMargin(8/64.f), FLinearColor(0, 0, 0, 1)));
 	Style->Set("Docking.Sidebar.DrawerBackground", new FSlateColorBrush(FStyleColors::Background));
-	Style->Set("Docking.Sidebar.Background", new FSlateColorBrush(FStyleColors::Input));
+	Style->Set("Docking.Sidebar.Background", new FSlateColorBrush(FStyleColors::Recessed));
 	Style->Set("Docking.Sidebar.Border", new FSlateRoundedBoxBrush(FSlateColor(FLinearColor::Transparent), 5.0f, FStyleColors::Hover, 1.0f) );
 
 	// Dock Cross
@@ -1569,10 +1573,10 @@ void FStarshipCoreStyle::SetupTableViewStyles(TSharedRef<FStyle>& Style)
 	const FSlateColor SelectorColor = Style->GetSlateColor("SelectorColor");
 
 	const FTableRowStyle DefaultTableRowStyle = FTableRowStyle()
-		.SetEvenRowBackgroundBrush(FSlateColorBrush(FStyleColors::Input))
+		.SetEvenRowBackgroundBrush(FSlateColorBrush(FStyleColors::Recessed))
 		.SetEvenRowBackgroundHoveredBrush(FSlateColorBrush(FStyleColors::SelectHover))
 
-		.SetOddRowBackgroundBrush(FSlateColorBrush(FStyleColors::Input))
+		.SetOddRowBackgroundBrush(FSlateColorBrush(FStyleColors::Recessed))
 		.SetOddRowBackgroundHoveredBrush(FSlateColorBrush(FStyleColors::SelectHover))
 
 		.SetSelectorFocusedBrush(BORDER_BRUSH("Common/Selector", FMargin(4.f / 16.f), SelectorColor))
@@ -1612,7 +1616,7 @@ void FStarshipCoreStyle::SetupTableViewStyles(TSharedRef<FStyle>& Style)
 		.SetSortSecondaryAscendingImage(IMAGE_BRUSH_SVG("Starship/CoreWidgets/TableView/sort-up-arrows", Icon12x12))
 		.SetSortSecondaryDescendingImage(IMAGE_BRUSH_SVG("Starship/CoreWidgets/TableView/sort-down-arrows", Icon12x12))
 		.SetNormalBrush(FSlateColorBrush(FStyleColors::Background))
-		.SetHoveredBrush(FSlateColorBrush(FStyleColors::Dropdown))
+		.SetHoveredBrush(FSlateColorBrush(FStyleColors::Secondary))
 
 		.SetMenuDropdownImage(IMAGE_BRUSH_SVG("Starship/Common/ellipsis-vertical-narrow", FVector2D(6.f, 24.f)))
 		.SetMenuDropdownNormalBorderBrush(FSlateNoResource())
@@ -1672,7 +1676,7 @@ void FStarshipCoreStyle::SetupMultiboxStyles(TSharedRef<FStyle>& Style)
 			.SetComboButtonPadding(FMargin(4.0f, 0.0f))
 			.SetButtonPadding(FMargin(2.0f, 0.f))
 			.SetCheckBoxPadding(FMargin(4.0f, 0.f))
-			.SetSeparatorBrush(FSlateColorBrush(FStyleColors::Input))
+			.SetSeparatorBrush(FSlateColorBrush(FStyleColors::Recessed))
 			.SetSeparatorPadding(FMargin(1.f, 0.f, 1.f, 0.f))
 			.SetLabelStyle(FTextBlockStyle(NormalText).SetFont(FStyleFonts::Get().Normal))
 			.SetEditableTextStyle(FEditableTextBoxStyle(NormalEditableTextBoxStyle).SetFont(FStyleFonts::Get().Normal))
@@ -1729,7 +1733,7 @@ void FStarshipCoreStyle::SetupMultiboxStyles(TSharedRef<FStyle>& Style)
 			.SetComboButtonPadding(FMargin(6.0f, 0.0f))
 			.SetButtonPadding(FMargin(4.0f, 0.0f))
 			.SetCheckBoxPadding(FMargin(10.0f, 0.0f))
-			.SetSeparatorBrush(FSlateColorBrush(FStyleColors::Input))
+			.SetSeparatorBrush(FSlateColorBrush(FStyleColors::Recessed))
 			.SetSeparatorPadding(FMargin(8.f, 0))
 			.SetLabelStyle(FTextBlockStyle(NormalText))
 			.SetComboButtonStyle(Style->GetWidgetStyle<FComboButtonStyle>("ComboButton"))
diff --git a/Engine/Source/Runtime/SlateCore/Private/Styling/StyleColors.cpp b/Engine/Source/Runtime/SlateCore/Private/Styling/StyleColors.cpp
--- a/Engine/Source/Runtime/SlateCore/Private/Styling/StyleColors.cpp
+++ b/Engine/Source/Runtime/SlateCore/Private/Styling/StyleColors.cpp
@@ -6,8 +6,10 @@ const FSlateColor FStyleColors::Transparent = FSlateColor(FLinearColor::Transpar
 
 const FSlateColor FStyleColors::Black = EStyleColor::Black;
 const FSlateColor FStyleColors::Title = EStyleColor::Title;
+const FSlateColor FStyleColors::WindowBorder = EStyleColor::WindowBorder;
 const FSlateColor FStyleColors::Foldout = EStyleColor::Foldout;
 const FSlateColor FStyleColors::Input = EStyleColor::Input;
+const FSlateColor FStyleColors::Recessed = EStyleColor::Recessed;
 const FSlateColor FStyleColors::Background = EStyleColor::Background;
 const FSlateColor FStyleColors::Header = EStyleColor::Header;
 const FSlateColor FStyleColors::Dropdown = EStyleColor::Dropdown;
@@ -21,6 +23,7 @@ const FSlateColor FStyleColors::Highlight = EStyleColor::Highlight;
 const FSlateColor FStyleColors::Primary = EStyleColor::Primary;
 const FSlateColor FStyleColors::PrimaryHover = EStyleColor::PrimaryHover;
 const FSlateColor FStyleColors::PrimaryPress = EStyleColor::PrimaryPress;
+const FSlateColor FStyleColors::Secondary = EStyleColor::Secondary;
 
 const FSlateColor FStyleColors::Foreground = EStyleColor::Foreground;
 const FSlateColor FStyleColors::ForegroundHover = EStyleColor::ForegroundHover;
@@ -56,8 +59,10 @@ void UStyleColorTable::InitalizeDefaults()
 {
 	SetColor(EStyleColor::Black, COLOR("#000000FF"));
 	SetColor(EStyleColor::Title, COLOR("#151515FF"));
+	SetColor(EStyleColor::WindowBorder, COLOR("0F0F0FFF"));
 	SetColor(EStyleColor::Foldout, COLOR("0F0F0FFF"));
-	SetColor(EStyleColor::Input, COLOR("#1A1A1AFF"));
+	SetColor(EStyleColor::Input, COLOR("0F0F0FFF"));
+	SetColor(EStyleColor::Recessed, COLOR("#1A1A1AFF"));
 	SetColor(EStyleColor::Background, COLOR("#242424FF"));
 	SetColor(EStyleColor::Header, COLOR("#2F2F2FFF"));
 	SetColor(EStyleColor::Dropdown, COLOR("#383838FF"));
@@ -71,6 +76,7 @@ void UStyleColorTable::InitalizeDefaults()
 	SetColor(EStyleColor::Primary, COLOR("#26BBFFFF"));
 	SetColor(EStyleColor::PrimaryHover, COLOR("#6FD2FFFF"));
 	SetColor(EStyleColor::PrimaryPress, COLOR("#1989BCFF"));
+	SetColor(EStyleColor::Secondary, COLOR("#383838FF"));
 
 	SetColor(EStyleColor::Foreground, COLOR("#A6A6A6FF"));
 	SetColor(EStyleColor::ForegroundHover, COLOR("#FFFFFFFF"));
diff --git a/Engine/Source/Runtime/SlateCore/Public/Styling/StyleColors.h b/Engine/Source/Runtime/SlateCore/Public/Styling/StyleColors.h
--- a/Engine/Source/Runtime/SlateCore/Public/Styling/StyleColors.h
+++ b/Engine/Source/Runtime/SlateCore/Public/Styling/StyleColors.h
@@ -15,8 +15,10 @@ enum class EStyleColor : uint8
 {
 	Black,
 	Title,
+	WindowBorder,
 	Foldout,
 	Input,
+	Recessed,
 	Background,
 	Header,
 	Dropdown,
@@ -28,6 +30,7 @@ enum class EStyleColor : uint8
 	Primary,
 	PrimaryHover,
 	PrimaryPress,
+	Secondary, 
 	Foreground,
 	ForegroundHover,
 	ForegroundInverted,
@@ -146,8 +149,10 @@ struct SLATECORE_API FStyleColors
 	static const FSlateColor Transparent;
 	static const FSlateColor Black;
 	static const FSlateColor Title;
+	static const FSlateColor WindowBorder;
 	static const FSlateColor Foldout;
 	static const FSlateColor Input;
+	static const FSlateColor Recessed;
 	static const FSlateColor Background;
 	static const FSlateColor Header;
 	static const FSlateColor Dropdown;
@@ -160,7 +165,7 @@ struct SLATECORE_API FStyleColors
 	static const FSlateColor Primary;
 	static const FSlateColor PrimaryHover;
 	static const FSlateColor PrimaryPress;
-
+	static const FSlateColor Secondary;
 
 	static const FSlateColor Foreground;
 	static const FSlateColor ForegroundHover;
