Fix menus popping due to larger size when they first get a check mark next to a menu item

diff --git a/Engine/Source/Editor/ContentBrowser/Private/SFilterList.cpp b/Engine/Source/Editor/ContentBrowser/Private/SFilterList.cpp
--- a/Engine/Source/Editor/ContentBrowser/Private/SFilterList.cpp
+++ b/Engine/Source/Editor/ContentBrowser/Private/SFilterList.cpp
@@ -1083,7 +1083,7 @@ void SFilterList::AddFilter(const TSharedRef<SFilter>& FilterToAdd)
 	Filters.Add(FilterToAdd);
 
 	FilterBox->AddSlot()
-	.Padding(3, 3)
+	.Padding(3, 0)
 	[
 		FilterToAdd
 	];
@@ -1447,7 +1447,6 @@ TSharedRef<SWidget> SFilterList::MakeAddFilterMenu(EAssetTypeCategories::Type Me
 
 	return 
 		SNew(SVerticalBox)
-
 		+SVerticalBox::Slot()
 		.MaxHeight(DisplaySize.Y * 0.9)
 		[
diff --git a/Engine/Source/Runtime/SlateCore/Private/Styling/StarshipCoreStyle.cpp b/Engine/Source/Runtime/SlateCore/Private/Styling/StarshipCoreStyle.cpp
--- a/Engine/Source/Runtime/SlateCore/Private/Styling/StarshipCoreStyle.cpp
+++ b/Engine/Source/Runtime/SlateCore/Private/Styling/StarshipCoreStyle.cpp
@@ -1803,9 +1803,9 @@ void FStarshipCoreStyle::SetupMultiboxStyles(TSharedRef<FStyle>& Style)
 
 		/* Set images for various SCheckBox states associated with menu check box items... */
 		const FCheckBoxStyle BasicMenuCheckBoxStyle = FCheckBoxStyle()
-			.SetUncheckedImage(FSlateNoResource(Icon14x14))
-			.SetUncheckedHoveredImage(FSlateNoResource(Icon14x14))
-			.SetUncheckedPressedImage(FSlateNoResource(Icon14x14))
+			.SetUncheckedImage(FSlateNoResource(Icon16x16))
+			.SetUncheckedHoveredImage(FSlateNoResource(Icon16x16))
+			.SetUncheckedPressedImage(FSlateNoResource(Icon16x16))
 			.SetCheckedImage(	    IMAGE_BRUSH_SVG("Starship/Common/check", Icon16x16))
 			.SetCheckedHoveredImage(IMAGE_BRUSH_SVG("Starship/Common/check", Icon16x16))
 			.SetCheckedPressedImage(IMAGE_BRUSH_SVG("Starship/Common/check", Icon16x16))
