Reset starship style colors to default using the old rename the variable trick until we have a better system.  This fixes folks colors from being offset.

diff --git a/Engine/Source/Runtime/SlateCore/Public/Styling/StyleColors.h b/Engine/Source/Runtime/SlateCore/Public/Styling/StyleColors.h
--- a/Engine/Source/Runtime/SlateCore/Public/Styling/StyleColors.h
+++ b/Engine/Source/Runtime/SlateCore/Public/Styling/StyleColors.h
@@ -101,22 +101,22 @@ public:
 
 	const FLinearColor& GetColor(EStyleColor Color)
 	{
-		return ColorList.StyleColors[static_cast<int32>(Color)];
+		return Colors.StyleColors[static_cast<int32>(Color)];
 	}
 
 	void SetColor(EStyleColor InColorId, FLinearColor InColor)
 	{
-		ColorList.StyleColors[static_cast<int32>(InColorId)] = InColor;
+		Colors.StyleColors[static_cast<int32>(InColorId)] = InColor;
 	}
 
 	void SetColorDisplayName(EStyleColor InColorId, FText DisplayName)
 	{
-		ColorList.DisplayNames[static_cast<int32>(InColorId)] = DisplayName;
+		Colors.DisplayNames[static_cast<int32>(InColorId)] = DisplayName;
 	}
 
 	FText GetColorDisplayName(EStyleColor InColorId) const
 	{
-		return ColorList.DisplayNames[static_cast<int32>(InColorId)];
+		return Colors.DisplayNames[static_cast<int32>(InColorId)];
 	}
 
 	UStyleColorTable();
@@ -137,7 +137,7 @@ public:
 
 private:
 	UPROPERTY(EditAnywhere, Config, Category = Colors)
-	FStyleColorList ColorList;
+	FStyleColorList Colors;
 };
 
 /**
