Fix having no theme causing the editor to be black

diff --git a/Engine/Source/Runtime/SlateCore/Private/Styling/StyleColors.cpp b/Engine/Source/Runtime/SlateCore/Private/Styling/StyleColors.cpp
--- a/Engine/Source/Runtime/SlateCore/Private/Styling/StyleColors.cpp
+++ b/Engine/Source/Runtime/SlateCore/Private/Styling/StyleColors.cpp
@@ -360,13 +360,13 @@ void USlateThemeManager::EnsureValidCurrentTheme()
 void USlateThemeManager::LoadThemeColors(FStyleTheme& Theme)
 {
 	FString ThemeData;
-	if (FFileHelper::LoadFileToString(ThemeData, *Theme.Filename))
-	{
-		Theme.LoadedDefaultColors.Empty();
 
-		// Start with the hard coded default colors. They are a fallback if the theme is incomplete 
-		Theme.LoadedDefaultColors = MakeArrayView<FLinearColor>(DefaultColors, (int32)EStyleColor::MAX);
+	// Start with the hard coded default colors. They are a fallback if the theme is incomplete 
+	Theme.LoadedDefaultColors = MakeArrayView<FLinearColor>(DefaultColors, (int32)EStyleColor::MAX);
 
+	if (FFileHelper::LoadFileToString(ThemeData, *Theme.Filename))
+	{
+		//Theme.LoadedDefaultColors.Empty();
 		TSharedRef<TJsonReader<>> ReaderRef = TJsonReaderFactory<>::Create(ThemeData);
 		TJsonReader<>& Reader = ReaderRef.Get();
 
