Brightened up foregreound text color

diff --git a/Engine/Source/Runtime/SlateCore/Private/Styling/StyleColors.cpp b/Engine/Source/Runtime/SlateCore/Private/Styling/StyleColors.cpp
--- a/Engine/Source/Runtime/SlateCore/Private/Styling/StyleColors.cpp
+++ b/Engine/Source/Runtime/SlateCore/Private/Styling/StyleColors.cpp
@@ -87,7 +87,7 @@ void USlateThemeManager::InitalizeDefaults()
 	SetDefaultColor(EStyleColor::PrimaryPress, COLOR("#1989BCFF"));
 	SetDefaultColor(EStyleColor::Secondary, COLOR("#383838FF"));
 
-	SetDefaultColor(EStyleColor::Foreground, COLOR("#A6A6A6FF"));
+	SetDefaultColor(EStyleColor::Foreground, COLOR("#C0C0C0FF"));
 	SetDefaultColor(EStyleColor::ForegroundHover, COLOR("#FFFFFFFF"));
 	SetDefaultColor(EStyleColor::ForegroundInverted, GetDefaultColor(EStyleColor::Input));
 	SetDefaultColor(EStyleColor::ForegroundHeader, COLOR("#C8C8C8FF"));
