More fixes to splitters and separators after the style flattening

diff --git a/Engine/Source/Developer/ProjectLauncher/Private/Widgets/SProjectLauncher.cpp b/Engine/Source/Developer/ProjectLauncher/Private/Widgets/SProjectLauncher.cpp
--- a/Engine/Source/Developer/ProjectLauncher/Private/Widgets/SProjectLauncher.cpp
+++ b/Engine/Source/Developer/ProjectLauncher/Private/Widgets/SProjectLauncher.cpp
@@ -94,7 +94,6 @@ void SProjectLauncher::Construct(const FArguments& InArgs, const TSharedRef<SDoc
 		+ SWidgetSwitcher::Slot()
 		[
 			SNew(SSplitter)
-			.Style(FEditorStyle::Get(), "ContentBrowser.Splitter")
 			.Orientation(Orient_Vertical)
 
 			// Simple SProjectLauncher
diff --git a/Engine/Source/Editor/ContentBrowser/Private/SAssetView.cpp b/Engine/Source/Editor/ContentBrowser/Private/SAssetView.cpp
--- a/Engine/Source/Editor/ContentBrowser/Private/SAssetView.cpp
+++ b/Engine/Source/Editor/ContentBrowser/Private/SAssetView.cpp
@@ -327,7 +327,7 @@ void SAssetView::Construct( const FArguments& InArgs )
 	TSharedRef<SVerticalBox> VerticalBox = SNew(SVerticalBox);
 
 	ChildSlot
-	.Padding(0.0f, 1.0f, 0.0f, 0.0f)
+	.Padding(0.0f)
 	[
 		SNew(SBorder)
 		.Padding(0)
diff --git a/Engine/Source/Editor/ContentBrowser/Private/SContentBrowser.cpp b/Engine/Source/Editor/ContentBrowser/Private/SContentBrowser.cpp
--- a/Engine/Source/Editor/ContentBrowser/Private/SContentBrowser.cpp
+++ b/Engine/Source/Editor/ContentBrowser/Private/SContentBrowser.cpp
@@ -316,22 +316,13 @@ void SContentBrowser::Construct( const FArguments& InArgs, const FName& InInstan
 				]
 			]
 		]
-		+ SVerticalBox::Slot()
-		.AutoHeight()
-		[
-			SNew(SSeparator)
-			.SeparatorImage(FAppStyle::Get().GetBrush("Brushes.Recessed"))
-			.Thickness(2.0f)
-		]
-
 		// Assets/tree
 		+ SVerticalBox::Slot()
 		.FillHeight(1.0f)
-		.Padding(0)
+		.Padding(0.0f,2.0f, 0.0f, 0.0f)
 		[
 			// The tree/assets splitter
 			SAssignNew(PathAssetSplitterPtr, SSplitter)
-			.Style(FEditorStyle::Get(), "ContentBrowser.Splitter")
 			.PhysicalSplitterHandleSize(2.0f)
 
 			// Sources View
@@ -350,7 +341,6 @@ void SContentBrowser::Construct( const FArguments& InArgs, const FName& InInstan
 					+SWidgetSwitcher::Slot()
 					[
 						SAssignNew(PathFavoriteSplitterPtr, SSplitter)
-						.Style(FEditorStyle::Get(), "ContentBrowser.Splitter")
 						.PhysicalSplitterHandleSize(3.0f)
 						.HitDetectionSplitterHandleSize(3.0f)
 						.Orientation(EOrientation::Orient_Vertical)
@@ -622,7 +612,7 @@ TSharedRef<SWidget> SContentBrowser::CreateAssetView(const FContentBrowserConfig
 	return
 		SNew(SVerticalBox)
 		+ SVerticalBox::Slot()
-		.Padding(0.0f, 2.0f, 0.0f, 0.0f)
+		.Padding(0.0f)
 		.AutoHeight()
 		[
 			SNew(SBorder)
@@ -704,7 +694,7 @@ TSharedRef<SWidget> SContentBrowser::CreateAssetView(const FContentBrowserConfig
 		// Assets
 		+ SVerticalBox::Slot()
 		.FillHeight( 1.0f )
-		.Padding( 0 )
+		.Padding( 0, 0 )
 		[
 			AssetViewPtr.ToSharedRef()
 		];
diff --git a/Engine/Source/Editor/EditorStyle/Private/StarshipStyle.cpp b/Engine/Source/Editor/EditorStyle/Private/StarshipStyle.cpp
--- a/Engine/Source/Editor/EditorStyle/Private/StarshipStyle.cpp
+++ b/Engine/Source/Editor/EditorStyle/Private/StarshipStyle.cpp
@@ -4977,8 +4977,9 @@ void FStarshipEditorStyle::FStyle::SetupPersonaStyle()
 		Set("AnimTimeline.Outliner.DefaultBorder", new FSlateColorBrush(FLinearColor::White));
 
 		const FSplitterStyle OutlinerSplitterStyle = FSplitterStyle()
-			.SetHandleNormalBrush( FSlateNoResource() )
-			.SetHandleHighlightBrush( FSlateNoResource() );
+			.SetHandleNormalBrush(FSlateColorBrush(FStyleColors::Recessed))
+			.SetHandleHighlightBrush(FSlateColorBrush(FStyleColors::Secondary));
+
 		Set("AnimTimeline.Outliner.Splitter", OutlinerSplitterStyle);
 
 		Set("AnimTimeline.Outliner.Label", FTextBlockStyle(NormalText)
@@ -5946,10 +5947,6 @@ void FStarshipEditorStyle::FStyle::SetupContentBrowserStyle()
 
 		Set("ContentBrowser.DefaultFolderColor", FStyleColors::AccentFolder);
 
-		Set( "ContentBrowser.Splitter", FSplitterStyle(FStarshipCoreStyle::GetCoreStyle().GetWidgetStyle<FSplitterStyle>("Splitter"))
-			.SetHandleNormalBrush(FSlateColorBrush(FStyleColors::Recessed))
-			);
-
 		// Asset list view
 		Set( "ContentBrowser.AssetListViewNameFont", DEFAULT_FONT( "Regular", 12 ) );
 		Set( "ContentBrowser.AssetListViewNameFontDirty", DEFAULT_FONT( "Bold", 12 ) );
diff --git a/Engine/Source/Editor/HierarchicalLODOutliner/Private/HLODOutliner.cpp b/Engine/Source/Editor/HierarchicalLODOutliner/Private/HLODOutliner.cpp
--- a/Engine/Source/Editor/HierarchicalLODOutliner/Private/HLODOutliner.cpp
+++ b/Engine/Source/Editor/HierarchicalLODOutliner/Private/HLODOutliner.cpp
@@ -237,7 +237,6 @@ namespace HLODOutliner
 				[
 					SNew(SSplitter)
 					.Orientation(Orient_Horizontal)
-					.Style(FEditorStyle::Get(), "ContentBrowser.Splitter")
 					+ SSplitter::Slot()
 					.Value(0.5)
 					[
@@ -253,7 +252,6 @@ namespace HLODOutliner
 				[
 					SNew(SSplitter)
 					.Orientation(Orient_Vertical)
-					.Style(FEditorStyle::Get(), "ContentBrowser.Splitter")
 					+ SSplitter::Slot()
 					.Value(0.5)
 					[
diff --git a/Engine/Source/Editor/Kismet/Private/UserDefinedStructureEditor.cpp b/Engine/Source/Editor/Kismet/Private/UserDefinedStructureEditor.cpp
--- a/Engine/Source/Editor/Kismet/Private/UserDefinedStructureEditor.cpp
+++ b/Engine/Source/Editor/Kismet/Private/UserDefinedStructureEditor.cpp
@@ -385,8 +385,7 @@ TSharedRef<SDockTab> FUserDefinedStructureEditor::SpawnStructureTab(const FSpawn
 		.Orientation(Orient_Vertical)
 		.PhysicalSplitterHandleSize(2.0f)
 		.HitDetectionSplitterHandleSize(8.f)
-		.ResizeMode(ESplitterResizeMode::FixedPosition)
-		.Style(FEditorStyle::Get(), "ContentBrowser.Splitter");
+		.ResizeMode(ESplitterResizeMode::FixedPosition);
 	{
 		// Create a property view
 		FPropertyEditorModule& EditModule = FModuleManager::Get().GetModuleChecked<FPropertyEditorModule>("PropertyEditor");
diff --git a/Engine/Source/Runtime/SlateCore/Private/Styling/StarshipCoreStyle.cpp b/Engine/Source/Runtime/SlateCore/Private/Styling/StarshipCoreStyle.cpp
--- a/Engine/Source/Runtime/SlateCore/Private/Styling/StarshipCoreStyle.cpp
+++ b/Engine/Source/Runtime/SlateCore/Private/Styling/StarshipCoreStyle.cpp
@@ -553,8 +553,8 @@ TSharedRef<ISlateStyle> FStarshipCoreStyle::Create()
 	// SSplitter
 	{
 		Style->Set("Splitter", FSplitterStyle()
-			.SetHandleNormalBrush(FSlateColorBrush(FStyleColors::Title))
-			.SetHandleHighlightBrush(IMAGE_BRUSH("Common/SplitterHandleHighlight", Icon8x8, FLinearColor::White))
+			.SetHandleNormalBrush(FSlateColorBrush(FStyleColors::Recessed))
+			.SetHandleHighlightBrush(FSlateColorBrush(FStyleColors::Secondary))
 		);
 	}
 
@@ -609,7 +609,7 @@ TSharedRef<ISlateStyle> FStarshipCoreStyle::Create()
 
 	// SSeparator defaults...
 	{
-		Style->Set("Separator", new BOX_BRUSH("Common/Separator", 1 / 4.0f, FLinearColor(1, 1, 1, 0.5f)));
+		Style->Set("Separator", new FSlateColorBrush(FStyleColors::Recessed));
 	}
 
 	// SHeader defaults...
@@ -1428,7 +1428,7 @@ void FStarshipCoreStyle::SetupDockingStyles(TSharedRef<FStyle>& Style)
 
 		.SetColorOverlayTabBrush(FSlateNoResource())
 		.SetColorOverlayIconBrush(FSlateNoResource())
-		.SetContentAreaBrush(FSlateColorBrush(FStyleColors::Background))
+		.SetContentAreaBrush(FSlateColorBrush(FStyleColors::Recessed))
 		.SetTabWellBrush(FSlateNoResource())
 		.SetFlashColor(TabFlashColor)
 
@@ -1453,7 +1453,7 @@ void FStarshipCoreStyle::SetupDockingStyles(TSharedRef<FStyle>& Style)
 
 		.SetColorOverlayTabBrush(FSlateNoResource())
 		.SetColorOverlayIconBrush(FSlateNoResource())
-		.SetContentAreaBrush(FSlateColorBrush(FStyleColors::Background))
+		.SetContentAreaBrush(FSlateColorBrush(FStyleColors::Recessed))
 		.SetTabWellBrush(FSlateNoResource())
 
 		.SetTabPadding(FMargin(10, 7, 10, 8))
@@ -1635,7 +1635,7 @@ void FStarshipCoreStyle::SetupTableViewStyles(TSharedRef<FStyle>& Style)
 		.SetMenuDropdownHoveredBorderBrush(FSlateNoResource());
 
 	const FSplitterStyle TableHeaderSplitterStyle = FSplitterStyle()
-		.SetHandleNormalBrush(FSlateNoResource())
+		.SetHandleNormalBrush(FSlateColorBrush(FStyleColors::Recessed))
 		.SetHandleHighlightBrush(IMAGE_BRUSH("Common/HeaderSplitterGrip", Icon8x8));
 
 	Style->Set("TableView.Header", FHeaderRowStyle()
