Stylable Input Outline Color

diff --git a/Engine/Source/Runtime/SlateCore/Private/Styling/StarshipCoreStyle.cpp b/Engine/Source/Runtime/SlateCore/Private/Styling/StarshipCoreStyle.cpp
--- a/Engine/Source/Runtime/SlateCore/Private/Styling/StarshipCoreStyle.cpp
+++ b/Engine/Source/Runtime/SlateCore/Private/Styling/StarshipCoreStyle.cpp
@@ -485,8 +485,8 @@ TSharedRef<ISlateStyle> FStarshipCoreStyle::Create()
 	// SSlider and SVolumeControl defaults...
 	{
 		FSliderStyle SliderStyle = FSliderStyle()
-			.SetNormalBarImage(   FSlateRoundedBoxBrush(FStyleColors::Input, 2.0))
-			.SetHoveredBarImage(  FSlateRoundedBoxBrush(FStyleColors::Input, 2.0))
+			.SetNormalBarImage(   FSlateRoundedBoxBrush(FStyleColors::Input, 2.0, FStyleColors::Input, 1.0))
+			.SetHoveredBarImage(  FSlateRoundedBoxBrush(FStyleColors::Input, 2.0, FStyleColors::Input, 1.0))
 			.SetNormalThumbImage(  FSlateRoundedBoxBrush(FStyleColors::Hover2, Icon8x8) )
 			.SetHoveredThumbImage( FSlateRoundedBoxBrush(FStyleColors::ForegroundHover, Icon8x8) )
 			.SetBarThickness(4.0f);
@@ -505,7 +505,7 @@ TSharedRef<ISlateStyle> FStarshipCoreStyle::Create()
 	// SSpinBox defaults...
 	{
 		Style->Set("SpinBox", FSpinBoxStyle()
-			.SetBackgroundBrush(FSlateRoundedBoxBrush(FStyleColors::Input, InputFocusRadius, FStyleColors::Secondary, InputFocusThickness))
+			.SetBackgroundBrush(FSlateRoundedBoxBrush(FStyleColors::Input, InputFocusRadius, FStyleColors::InputOutline, InputFocusThickness))
 			.SetHoveredBackgroundBrush(FSlateRoundedBoxBrush(FStyleColors::Input, InputFocusRadius, FStyleColors::Hover, InputFocusThickness))
 
 			.SetActiveFillBrush(FSlateRoundedBoxBrush(FStyleColors::Hover, InputFocusRadius, FLinearColor::Transparent, InputFocusThickness))
@@ -519,7 +519,7 @@ TSharedRef<ISlateStyle> FStarshipCoreStyle::Create()
 	// SNumericEntryBox defaults...
 	{
 		Style->Set("NumericEntrySpinBox", FSpinBoxStyle()
-			.SetBackgroundBrush(FSlateRoundedBoxBrush(FStyleColors::Input, InputFocusRadius, FStyleColors::Secondary, InputFocusThickness))
+			.SetBackgroundBrush(FSlateRoundedBoxBrush(FStyleColors::Input, InputFocusRadius, FStyleColors::InputOutline, InputFocusThickness))
 			.SetHoveredBackgroundBrush(FSlateRoundedBoxBrush(FStyleColors::Input, InputFocusRadius, FStyleColors::Hover, InputFocusThickness))
 
 			.SetActiveFillBrush(FSlateRoundedBoxBrush(FStyleColors::Hover, InputFocusRadius, FLinearColor::Transparent, InputFocusThickness))
@@ -531,7 +531,7 @@ TSharedRef<ISlateStyle> FStarshipCoreStyle::Create()
 		);
 
 		Style->Set("NumericEntrySpinBox_Dark", FSpinBoxStyle()
-			.SetBackgroundBrush(FSlateRoundedBoxBrush(FStyleColors::Input, InputFocusRadius, FStyleColors::Secondary, InputFocusThickness))
+			.SetBackgroundBrush(FSlateRoundedBoxBrush(FStyleColors::Input, InputFocusRadius, FStyleColors::InputOutline, InputFocusThickness))
 			.SetHoveredBackgroundBrush(FSlateRoundedBoxBrush(FStyleColors::Input, InputFocusRadius, FStyleColors::Hover, InputFocusThickness))
 
 			.SetActiveFillBrush(FSlateRoundedBoxBrush(FStyleColors::Hover, InputFocusRadius, FLinearColor::Transparent, InputFocusThickness))
@@ -798,6 +798,7 @@ void FStarshipCoreStyle::SetupColors(TSharedRef<FStyle>& Style)
 	Style->Set("Colors.WindowBorder", FStyleColors::WindowBorder);
 	Style->Set("Colors.Foldout", FStyleColors::Foldout);
 	Style->Set("Colors.Input", FStyleColors::Input);
+	Style->Set("Colors.InputOutline", FStyleColors::InputOutline);
 	Style->Set("Colors.Recessed", FStyleColors::Recessed);
 	Style->Set("Colors.Background", FStyleColors::Background);
 	Style->Set("Colors.Header", FStyleColors::Header);
@@ -841,6 +842,7 @@ void FStarshipCoreStyle::SetupColors(TSharedRef<FStyle>& Style)
 	Style->Set("Brushes.Title", new FSlateColorBrush(FStyleColors::Title));
 	Style->Set("Brushes.Foldout", new FSlateColorBrush(FStyleColors::Foldout));
 	Style->Set("Brushes.Input", new FSlateColorBrush(FStyleColors::Input));
+	Style->Set("Brushes.InputOutline", new FSlateColorBrush(FStyleColors::InputOutline));
 	Style->Set("Brushes.Recessed", new FSlateColorBrush(FStyleColors::Recessed));
 	Style->Set("Brushes.Background", new FSlateColorBrush(FStyleColors::Background));
 	Style->Set("Brushes.Header", new FSlateColorBrush(FStyleColors::Header));
@@ -980,7 +982,7 @@ void FStarshipCoreStyle::SetupTextStyles(TSharedRef<FStyle>& Style)
 
 	// SEditableTextBox defaults...
 	const FEditableTextBoxStyle NormalEditableTextBoxStyle = FEditableTextBoxStyle()
-		.SetBackgroundImageNormal(FSlateRoundedBoxBrush(FStyleColors::Input, InputFocusRadius, FStyleColors::Secondary, InputFocusThickness))
+		.SetBackgroundImageNormal(FSlateRoundedBoxBrush(FStyleColors::Input, InputFocusRadius, FStyleColors::InputOutline, InputFocusThickness))
 		.SetBackgroundImageHovered(FSlateRoundedBoxBrush(FStyleColors::Input, InputFocusRadius, FStyleColors::Hover, InputFocusThickness))
 		.SetBackgroundImageFocused(FSlateRoundedBoxBrush(FStyleColors::Input, InputFocusRadius, FStyleColors::Primary, InputFocusThickness))
 		.SetBackgroundImageReadOnly(FSlateRoundedBoxBrush(FStyleColors::Input, InputFocusRadius))
@@ -1042,9 +1044,9 @@ void FStarshipCoreStyle::SetupButtonStyles(TSharedRef<FStyle>& Style)
 {
 	// SButton defaults
 	const FButtonStyle PrimaryButton = FButtonStyle()
-		.SetNormal(FSlateRoundedBoxBrush(FStyleColors::Primary, 4.0f))
-		.SetHovered(FSlateRoundedBoxBrush(FStyleColors::PrimaryHover, 4.0f))
-		.SetPressed(FSlateRoundedBoxBrush(FStyleColors::PrimaryPress, 4.0f))
+		.SetNormal(FSlateRoundedBoxBrush(FStyleColors::Primary, 4.0f, FStyleColors::InputOutline, InputFocusThickness))
+		.SetHovered(FSlateRoundedBoxBrush(FStyleColors::PrimaryHover, 4.0f, FStyleColors::Hover, InputFocusThickness))
+		.SetPressed(FSlateRoundedBoxBrush(FStyleColors::PrimaryPress, 4.0f, FStyleColors::Hover, InputFocusThickness))
 		.SetNormalForeground(FStyleColors::Background)
 		.SetHoveredForeground(FStyleColors::Background)
 		.SetPressedForeground(FStyleColors::Background)
@@ -1054,9 +1056,9 @@ void FStarshipCoreStyle::SetupButtonStyles(TSharedRef<FStyle>& Style)
 
 
 	const FButtonStyle Button = FButtonStyle()
-		.SetNormal(FSlateRoundedBoxBrush(FStyleColors::Secondary, 4.0f))
-		.SetHovered(FSlateRoundedBoxBrush(FStyleColors::Hover, 4.0f))
-		.SetPressed(FSlateRoundedBoxBrush(FStyleColors::Header, 4.0f))
+		.SetNormal(FSlateRoundedBoxBrush(FStyleColors::Secondary, 4.0f, FStyleColors::InputOutline, InputFocusThickness))
+		.SetHovered(FSlateRoundedBoxBrush(FStyleColors::Hover, 4.0f, FStyleColors::Hover, InputFocusThickness))
+		.SetPressed(FSlateRoundedBoxBrush(FStyleColors::Header, 4.0f, FStyleColors::Hover, InputFocusThickness))
 		.SetNormalForeground(FStyleColors::ForegroundHover)
 		.SetHoveredForeground(FStyleColors::ForegroundHover)
 		.SetPressedForeground(FStyleColors::ForegroundHover)
@@ -1093,7 +1095,7 @@ void FStarshipCoreStyle::SetupComboButtonStyles(TSharedRef<FStyle>& Style)
 {
 	// SComboButton and SComboBox defaults...
 	const FButtonStyle ComboButtonButton = FButtonStyle()
-		.SetNormal(FSlateRoundedBoxBrush(FStyleColors::Input, InputFocusRadius, FStyleColors::Secondary, InputFocusThickness))
+		.SetNormal(FSlateRoundedBoxBrush(FStyleColors::Input, InputFocusRadius, FStyleColors::InputOutline, InputFocusThickness))
 		.SetHovered(FSlateRoundedBoxBrush(FStyleColors::Input, InputFocusRadius, FStyleColors::Hover, InputFocusThickness))
 		.SetPressed(FSlateRoundedBoxBrush(FStyleColors::Input, InputFocusRadius, FStyleColors::Hover, InputFocusThickness))
 		.SetNormalForeground(FStyleColors::Foreground)
@@ -1205,7 +1207,7 @@ void FStarshipCoreStyle::SetupCheckboxStyles(TSharedRef<FStyle>& Style)
 		.SetUndeterminedImage(IMAGE_BRUSH_SVG("Starship/CoreWidgets/CheckBox/indeterminate", Icon16x16, FStyleColors::Primary))
 		.SetUndeterminedHoveredImage(IMAGE_BRUSH_SVG("Starship/CoreWidgets/CheckBox/indeterminate", Icon16x16, FStyleColors::PrimaryHover))
 		.SetUndeterminedPressedImage(IMAGE_BRUSH_SVG("Starship/CoreWidgets/CheckBox/indeterminate", Icon16x16, FStyleColors::Primary))
-		.SetBackgroundImage(FSlateRoundedBoxBrush(FStyleColors::Input, CheckboxCornerRadius, FStyleColors::Secondary, CheckboxOutlineThickness, Icon18x18))
+		.SetBackgroundImage(FSlateRoundedBoxBrush(FStyleColors::Input, CheckboxCornerRadius, FStyleColors::InputOutline, CheckboxOutlineThickness, Icon18x18))
 		.SetBackgroundHoveredImage(FSlateRoundedBoxBrush(FStyleColors::Input, CheckboxCornerRadius, FStyleColors::Hover, CheckboxOutlineThickness, Icon18x18))
 		.SetBackgroundPressedImage(FSlateRoundedBoxBrush(FStyleColors::Foldout, CheckboxCornerRadius, FStyleColors::Hover, CheckboxOutlineThickness, Icon18x18));
 
@@ -1231,7 +1233,7 @@ void FStarshipCoreStyle::SetupCheckboxStyles(TSharedRef<FStyle>& Style)
 		.SetUndeterminedImage(IMAGE_BRUSH_SVG("Starship/CoreWidgets/CheckBox/indeterminate", Icon16x16, FStyleColors::Foreground))
 		.SetUndeterminedHoveredImage(IMAGE_BRUSH_SVG("Starship/CoreWidgets/CheckBox/indeterminate", Icon16x16, FStyleColors::ForegroundHover))
 		.SetUndeterminedPressedImage(IMAGE_BRUSH_SVG("Starship/CoreWidgets/CheckBox/indeterminate", Icon16x16, FStyleColors::Foreground))
-		.SetBackgroundImage(FSlateRoundedBoxBrush(FStyleColors::Input, CheckboxCornerRadius, FStyleColors::Secondary, CheckboxOutlineThickness, Icon18x18))
+		.SetBackgroundImage(FSlateRoundedBoxBrush(FStyleColors::Input, CheckboxCornerRadius, FStyleColors::InputOutline, CheckboxOutlineThickness, Icon18x18))
 		.SetBackgroundHoveredImage(FSlateRoundedBoxBrush(FStyleColors::Input, CheckboxCornerRadius, FStyleColors::Hover, CheckboxOutlineThickness, Icon18x18))
 		.SetBackgroundPressedImage(FSlateRoundedBoxBrush(FStyleColors::Foldout, CheckboxCornerRadius, FStyleColors::Hover, CheckboxOutlineThickness, Icon18x18));
 
diff --git a/Engine/Source/Runtime/SlateCore/Private/Styling/StyleColors.cpp b/Engine/Source/Runtime/SlateCore/Private/Styling/StyleColors.cpp
--- a/Engine/Source/Runtime/SlateCore/Private/Styling/StyleColors.cpp
+++ b/Engine/Source/Runtime/SlateCore/Private/Styling/StyleColors.cpp
@@ -16,6 +16,7 @@ const FSlateColor FStyleColors::Title = EStyleColor::Title;
 const FSlateColor FStyleColors::WindowBorder = EStyleColor::WindowBorder;
 const FSlateColor FStyleColors::Foldout = EStyleColor::Foldout;
 const FSlateColor FStyleColors::Input = EStyleColor::Input;
+const FSlateColor FStyleColors::InputOutline = EStyleColor::InputOutline;
 const FSlateColor FStyleColors::Recessed = EStyleColor::Recessed;
 const FSlateColor FStyleColors::Background = EStyleColor::Background;
 const FSlateColor FStyleColors::Header = EStyleColor::Header;
@@ -71,6 +72,7 @@ void USlateThemeManager::InitalizeDefaults()
 	SetDefaultColor(EStyleColor::WindowBorder, COLOR("0F0F0FFF"));
 	SetDefaultColor(EStyleColor::Foldout, COLOR("0F0F0FFF"));
 	SetDefaultColor(EStyleColor::Input, COLOR("0F0F0FFF"));
+	SetDefaultColor(EStyleColor::InputOutline, COLOR("0F0F0FFF"));
 	SetDefaultColor(EStyleColor::Recessed, COLOR("#1A1A1AFF"));
 	SetDefaultColor(EStyleColor::Background, COLOR("#242424FF"));
 	SetDefaultColor(EStyleColor::Header, COLOR("#2F2F2FFF"));
diff --git a/Engine/Source/Runtime/SlateCore/Public/Styling/StyleColors.h b/Engine/Source/Runtime/SlateCore/Public/Styling/StyleColors.h
--- a/Engine/Source/Runtime/SlateCore/Public/Styling/StyleColors.h
+++ b/Engine/Source/Runtime/SlateCore/Public/Styling/StyleColors.h
@@ -23,6 +23,7 @@ enum class EStyleColor : uint8
 	WindowBorder,
 	Foldout,
 	Input,
+	InputOutline,
 	Recessed,
 	Background,
 	Header,
@@ -282,6 +283,7 @@ struct SLATECORE_API FStyleColors
 	static const FSlateColor WindowBorder;
 	static const FSlateColor Foldout;
 	static const FSlateColor Input;
+	static const FSlateColor InputOutline;
 	static const FSlateColor Recessed;
 	static const FSlateColor Background;
 	static const FSlateColor Header;
