Search boxes are now rounded to make them stand out

diff --git a/Engine/Source/Runtime/SlateCore/Private/Styling/StarshipCoreStyle.cpp b/Engine/Source/Runtime/SlateCore/Private/Styling/StarshipCoreStyle.cpp
--- a/Engine/Source/Runtime/SlateCore/Private/Styling/StarshipCoreStyle.cpp
+++ b/Engine/Source/Runtime/SlateCore/Private/Styling/StarshipCoreStyle.cpp
@@ -999,8 +999,16 @@ void FStarshipCoreStyle::SetupTextStyles(TSharedRef<FStyle>& Style)
 
 	// SSearchBox defaults...
 	{
+		const FEditableTextBoxStyle SearchBoxEditStyle =
+			FEditableTextBoxStyle(NormalEditableTextBoxStyle)
+			.SetBackgroundImageNormal(FSlateRoundedBoxBrush(FStyleColors::Input, FStyleColors::Secondary, InputFocusThickness))
+			.SetBackgroundImageHovered(FSlateRoundedBoxBrush(FStyleColors::Input, FStyleColors::Hover, InputFocusThickness))
+			.SetBackgroundImageFocused(FSlateRoundedBoxBrush(FStyleColors::Input, FStyleColors::Primary, InputFocusThickness))
+			.SetBackgroundImageReadOnly(FSlateRoundedBoxBrush(FStyleColors::Input)
+		);
+
 		Style->Set("SearchBox", FSearchBoxStyle()
-			.SetTextBoxStyle(NormalEditableTextBoxStyle)
+			.SetTextBoxStyle(SearchBoxEditStyle)
 			.SetUpArrowImage(IMAGE_BRUSH_SVG("Starship/Common/arrow-north", Icon8x8, FStyleColors::Foreground))
 			.SetDownArrowImage(IMAGE_BRUSH_SVG("Starship/Common/arrow-south", Icon8x8, FStyleColors::Foreground))
 			.SetGlassImage(IMAGE_BRUSH_SVG("Starship/Common/search", Icon16x16))
