Widget Style Change - Updated Outline Color

diff --git a/Engine/Source/Runtime/SlateCore/Private/Styling/StyleColors.cpp b/Engine/Source/Runtime/SlateCore/Private/Styling/StyleColors.cpp
--- a/Engine/Source/Runtime/SlateCore/Private/Styling/StyleColors.cpp
+++ b/Engine/Source/Runtime/SlateCore/Private/Styling/StyleColors.cpp
@@ -72,7 +72,7 @@ void USlateThemeManager::InitalizeDefaults()
 	SetDefaultColor(EStyleColor::WindowBorder, COLOR("0F0F0FFF"));
 	SetDefaultColor(EStyleColor::Foldout, COLOR("0F0F0FFF"));
 	SetDefaultColor(EStyleColor::Input, COLOR("0F0F0FFF"));
-	SetDefaultColor(EStyleColor::InputOutline, COLOR("0F0F0FFF"));
+	SetDefaultColor(EStyleColor::InputOutline, COLOR("454545FF"));
 	SetDefaultColor(EStyleColor::Recessed, COLOR("#1A1A1AFF"));
 	SetDefaultColor(EStyleColor::Background, COLOR("#242424FF"));
 	SetDefaultColor(EStyleColor::Header, COLOR("#2F2F2FFF"));
