Updated Dock Tab Style close icon

diff --git a/Engine/Source/Runtime/Slate/Private/Widgets/Docking/SDockTab.cpp b/Engine/Source/Runtime/Slate/Private/Widgets/Docking/SDockTab.cpp
--- a/Engine/Source/Runtime/Slate/Private/Widgets/Docking/SDockTab.cpp
+++ b/Engine/Source/Runtime/Slate/Private/Widgets/Docking/SDockTab.cpp
@@ -646,7 +646,7 @@ void SDockTab::Construct( const FArguments& InArgs )
 					SNew(SButton)
 					.ButtonStyle( CloseButtonStyle )
 					.OnClicked( this, &SDockTab::OnCloseButtonClicked )
-					.ContentPadding( 0 )
+					.ContentPadding(FMargin(0.0, 1.5, 0.0, 0.0))
 					.ToolTipText(this, &SDockTab::GetCloseButtonToolTipText)
 					.Visibility(this, &SDockTab::HandleIsCloseButtonVisible)
 					[
diff --git a/Engine/Source/Runtime/SlateCore/Private/Styling/StarshipCoreStyle.cpp b/Engine/Source/Runtime/SlateCore/Private/Styling/StarshipCoreStyle.cpp
--- a/Engine/Source/Runtime/SlateCore/Private/Styling/StarshipCoreStyle.cpp
+++ b/Engine/Source/Runtime/SlateCore/Private/Styling/StarshipCoreStyle.cpp
@@ -1421,9 +1421,9 @@ void FStarshipCoreStyle::SetupDockingStyles(TSharedRef<FStyle>& Style)
 	const FSlateColor& TabFlashColor = SelectionColor;
 
 	const FButtonStyle CloseButton = FButtonStyle()
-		.SetNormal(IMAGE_BRUSH_SVG("Starship/Common/close", Icon16x16, FStyleColors::Foreground))
-		.SetPressed(IMAGE_BRUSH_SVG("Starship/Common/close", Icon16x16, FStyleColors::Foreground))
-		.SetHovered(IMAGE_BRUSH_SVG("Starship/Common/close", Icon16x16, FStyleColors::ForegroundHover));
+		.SetNormal(IMAGE_BRUSH_SVG("Starship/Common/close-small", Icon16x16, FStyleColors::Foreground))
+		.SetPressed(IMAGE_BRUSH_SVG("Starship/Common/close-small", Icon16x16, FStyleColors::Foreground))
+		.SetHovered(IMAGE_BRUSH_SVG("Starship/Common/close-small", Icon16x16, FStyleColors::ForegroundHover));
 
 
 	FLinearColor DockColor_Inactive(FColor(45, 45, 45));
