ComboButton dropdown arrow padding is now configurable in a style

diff --git a/Engine/Source/Runtime/Slate/Private/Widgets/Input/SComboButton.cpp b/Engine/Source/Runtime/Slate/Private/Widgets/Input/SComboButton.cpp
--- a/Engine/Source/Runtime/Slate/Private/Widgets/Input/SComboButton.cpp
+++ b/Engine/Source/Runtime/Slate/Private/Widgets/Input/SComboButton.cpp
@@ -61,7 +61,7 @@ void SComboButton::Construct( const FArguments& InArgs )
 				.AutoWidth()
 				.HAlign( HAlign_Center )
 				.VAlign( VAlign_Center )
-				.Padding( InArgs._HasDownArrow ? 2.f : 0.f )
+				.Padding(InArgs._HasDownArrow ? InArgs._ComboButtonStyle->DownArrowPadding : FMargin(0))
 				[
 					SNew(SOverlay)
 					// drop shadow
diff --git a/Engine/Source/Runtime/SlateCore/Private/Styling/SlateTypes.cpp b/Engine/Source/Runtime/SlateCore/Private/Styling/SlateTypes.cpp
--- a/Engine/Source/Runtime/SlateCore/Private/Styling/SlateTypes.cpp
+++ b/Engine/Source/Runtime/SlateCore/Private/Styling/SlateTypes.cpp
@@ -153,8 +153,9 @@ void FButtonStyle::PostSerialize(const FArchive& Ar)
 FComboButtonStyle::FComboButtonStyle()
 	: ShadowOffset(FVector2D::ZeroVector)
 	, ShadowColorAndOpacity(FLinearColor::Black)
-	, MenuBorderPadding(FMargin(0.0f))
-	, ContentPadding(FMargin(5.0f))
+	, MenuBorderPadding(0.0f)
+	, ContentPadding(5.0f)
+	, DownArrowPadding(2.0f)
 {
 }
 
diff --git a/Engine/Source/Runtime/SlateCore/Public/Styling/SlateTypes.h b/Engine/Source/Runtime/SlateCore/Public/Styling/SlateTypes.h
--- a/Engine/Source/Runtime/SlateCore/Public/Styling/SlateTypes.h
+++ b/Engine/Source/Runtime/SlateCore/Public/Styling/SlateTypes.h
@@ -573,6 +573,13 @@ struct SLATECORE_API FComboButtonStyle : public FSlateWidgetStyle
 	FMargin ContentPadding;
 	FComboButtonStyle& SetContentPadding( const FMargin& InContentPadding ) { ContentPadding = InContentPadding; return *this; }
 
+	/*
+	 * Dropdown arrow padding (if a dropdown arrow exists)
+	 */
+	UPROPERTY(EditAnywhere, BlueprintReadWrite, Category = Appearance)
+	FMargin DownArrowPadding;
+	FComboButtonStyle& SetDownArrowPadding(const FMargin& InDownArrowPadding) { DownArrowPadding = InDownArrowPadding; return *this; }
+
 	/**
 	* Unlinks all colors in this style.
 	* @see FSlateColor::Unlink
