Updated the slate keyboard focus rectangle

diff --git a/Engine/Source/Editor/EditorStyle/Private/StarshipStyle.cpp b/Engine/Source/Editor/EditorStyle/Private/StarshipStyle.cpp
--- a/Engine/Source/Editor/EditorStyle/Private/StarshipStyle.cpp
+++ b/Engine/Source/Editor/EditorStyle/Private/StarshipStyle.cpp
@@ -653,8 +653,6 @@ void FStarshipEditorStyle::FStyle::SetupGeneralStyles()
 
 	Set( "Editor.AppIcon", new IMAGE_BRUSH( "Icons/EditorAppIcon", Icon24x24) );
 
-	Set( "FocusRectangle", new BORDER_BRUSH( "Old/DashedBorder", FMargin(6.0f/32.0f) ) );
-
 	Set( "MarqueeSelection", new BORDER_BRUSH( "Old/DashedBorder", FMargin(6.0f/32.0f) ) );
 
 	Set( "GenericLock", new IMAGE_BRUSH( "Icons/padlock_locked_16x", Icon16x16 ) );
diff --git a/Engine/Source/Runtime/SlateCore/Private/Styling/StarshipCoreStyle.cpp b/Engine/Source/Runtime/SlateCore/Private/Styling/StarshipCoreStyle.cpp
--- a/Engine/Source/Runtime/SlateCore/Private/Styling/StarshipCoreStyle.cpp
+++ b/Engine/Source/Runtime/SlateCore/Private/Styling/StarshipCoreStyle.cpp
@@ -242,7 +242,7 @@ TSharedRef<ISlateStyle> FStarshipCoreStyle::Create()
 
 		Style->Set("BoxShadow", new BOX_BRUSH("Common/BoxShadow", FMargin(5.0f / 64.0f)));
 
-		Style->Set("FocusRectangle", new BORDER_BRUSH("Old/DashedBorder", FMargin(6.0f / 32.0f)));
+		Style->Set("FocusRectangle", new FSlateRoundedBoxBrush(FStyleColors::Transparent, InputFocusRadius, FStyleColors::Primary, InputFocusThickness));
 	}
 
 	// Important colors
diff --git a/Engine/Source/Runtime/SlateCore/Private/Widgets/SWidget.cpp b/Engine/Source/Runtime/SlateCore/Private/Widgets/SWidget.cpp
--- a/Engine/Source/Runtime/SlateCore/Private/Widgets/SWidget.cpp
+++ b/Engine/Source/Runtime/SlateCore/Private/Widgets/SWidget.cpp
@@ -926,7 +926,7 @@ bool SWidget::HasAnyUserFocusOrFocusedDescendants() const
 
 const FSlateBrush* SWidget::GetFocusBrush() const
 {
-	return FCoreStyle::Get().GetBrush("FocusRectangle");
+	return FAppStyle::Get().GetBrush("FocusRectangle");
 }
 
 bool SWidget::HasMouseCapture() const
