Slate rounded boxes now support texture fills.
Also cleaned up rounded box constructor overloads

diff --git a/Engine/Shaders/Private/SlateElementPixelShader.usf b/Engine/Shaders/Private/SlateElementPixelShader.usf
--- a/Engine/Shaders/Private/SlateElementPixelShader.usf
+++ b/Engine/Shaders/Private/SlateElementPixelShader.usf
@@ -262,20 +262,20 @@ half4 GetRoundedBoxElementColor( VertexToPixelInterpolants InVertex )
 
 	// Compute the border intensity and fill intensity with a smooth transition
 	float spread = 0.5;
-    float bi = smoothstep(spread, -spread, dext);
-    float fi = smoothstep(spread, -spread, din);
+	float bi = smoothstep(spread, -spread, dext);
+	float fi = smoothstep(spread, -spread, din);
 
-	half4 fill = InVertex.Color;
-    half4 border = ShaderParams2;
+	half4 fill = GetColor(InVertex, GetUV(InVertex, 0) * GetUV(InVertex, 1)); //InVertex.Color;
+	half4 border = ShaderParams2;
 
-    // alpha blend the external color 
-    half4 OutColor = lerp(border, fill, half(thickness > radius));
-    OutColor.a = 0.0;
+	// alpha blend the external color 
+	half4 OutColor = lerp(border, fill, half(thickness > radius));
+	OutColor.a = 0.0;
 
-    // blend in the border and fill colors
-    OutColor = lerp(OutColor, border, bi);
-    OutColor = lerp(OutColor, fill, fi);
-    return OutColor;
+	// blend in the border and fill colors
+	OutColor = lerp(OutColor, border, bi);
+	OutColor = lerp(OutColor, fill, fi);
+	return OutColor;
 }
 #endif
 
diff --git a/Engine/Shaders/StandaloneRenderer/D3D/SlateElementPixelShader.hlsl b/Engine/Shaders/StandaloneRenderer/D3D/SlateElementPixelShader.hlsl
--- a/Engine/Shaders/StandaloneRenderer/D3D/SlateElementPixelShader.hlsl
+++ b/Engine/Shaders/StandaloneRenderer/D3D/SlateElementPixelShader.hlsl
@@ -158,19 +158,19 @@ float4 GetRoundedBoxElementColor( VertexOut InVertex )
 
 	// Compute the border intensity and fill intensity with a smooth transition
 	float spread = 0.5;
-    float bi = smoothstep(spread, -spread, dext);
-    float fi = smoothstep(spread, -spread, din);
-
-    // alpha blend the external color 
-	float4 fill = InVertex.Color;
-    float4 border = ShaderParams2;
-    float4 OutColor = lerp(border, fill, float(thickness > radius));
-    OutColor.a = 0.0;
-
-    // blend in the border and fill colors
-    OutColor = lerp(OutColor, border, bi);
-    OutColor = lerp(OutColor, fill, fi);
-    return OutColor;
+	float bi = smoothstep(spread, -spread, dext);
+	float fi = smoothstep(spread, -spread, din);
+
+	// alpha blend the external color 
+	float4 fill = GetColor(InVertex, InVertex.TextureCoordinates.xy * InVertex.TextureCoordinates.zw);
+	float4 border = ShaderParams2;
+	float4 OutColor = lerp(border, fill, float(thickness > radius));
+	OutColor.a = 0.0;
+
+	// blend in the border and fill colors
+	OutColor = lerp(OutColor, border, bi);
+	OutColor = lerp(OutColor, fill, fi);
+	return OutColor;
 }
 
 float4 GetSplineElementColor( VertexOut InVertex )
diff --git a/Engine/Shaders/StandaloneRenderer/OpenGL/SlateElementPixelShader.glsl b/Engine/Shaders/StandaloneRenderer/OpenGL/SlateElementPixelShader.glsl
--- a/Engine/Shaders/StandaloneRenderer/OpenGL/SlateElementPixelShader.glsl
+++ b/Engine/Shaders/StandaloneRenderer/OpenGL/SlateElementPixelShader.glsl
@@ -243,19 +243,19 @@ vec4 GetRoundedBoxElementColor()
 
 	// Compute the border intensity and fill intensity with a smooth transition
 	float spread = 0.5;
-    float bi = smoothstep(spread, -spread, dext);
-    float fi = smoothstep(spread, -spread, din);
-
-    // alpha blend the external color 
-	vec4 fill = Color;
-    vec4 border = ShaderParams2.bgra;
-    vec4 OutColor = mix(border, fill, float(thickness > radius));
-    OutColor.a = 0.0;
-
-    // blend in the border and fill colors
-    OutColor = mix(OutColor, border, bi);
-    OutColor = mix(OutColor, fill, fi);
-    return OutColor;
+	float bi = smoothstep(spread, -spread, dext);
+	float fi = smoothstep(spread, -spread, din);
+
+	// alpha blend the external color 
+	vec4 fill = GetDefaultElementColor();
+	vec4 border = ShaderParams2.bgra;
+	vec4 OutColor = mix(border, fill, float(thickness > radius));
+	OutColor.a = 0.0;
+
+	// blend in the border and fill colors
+	OutColor = mix(OutColor, border, bi);
+	OutColor = mix(OutColor, fill, fi);
+	return OutColor;
 }
 
 vec4 GetSplineElementColor()
diff --git a/Engine/Source/Editor/EditorStyle/Private/SlateEditorStyle.cpp b/Engine/Source/Editor/EditorStyle/Private/SlateEditorStyle.cpp
--- a/Engine/Source/Editor/EditorStyle/Private/SlateEditorStyle.cpp
+++ b/Engine/Source/Editor/EditorStyle/Private/SlateEditorStyle.cpp
@@ -4858,19 +4858,6 @@ void FSlateEditorStyle::FStyle::SetupLevelEditorStyle()
 		Set( "LevelEditor.EditorModes.Menu", new IMAGE_BRUSH( "Icons/icon_Editor_Modes_16x", Icon16x16 ) );
 
 
-		Set( "ToolPalette.DockingTab", FCheckBoxStyle()
-			.SetCheckBoxType(ESlateCheckBoxType::ToggleButton)
-			.SetPadding( FMargin(16.0f, 2.0f, 16.0f, 2.0f ) )
-			.SetCheckedImage(             FSlateBoxBrush( RootToCoreContentDir("Docking/Tab_Shape.png"),  2.f/8.0f, FColor(62, 62, 62) ) )
-			.SetCheckedHoveredImage(      FSlateBoxBrush( RootToCoreContentDir("Docking/Tab_Shape.png"),  2.f/8.0f, FColor(62, 62, 62) ) )
-			.SetCheckedPressedImage(      FSlateBoxBrush( RootToCoreContentDir("Docking/Tab_Shape.png"),  2.f/8.0f, FColor(62, 62, 62) ) )
-			.SetUncheckedImage(           FSlateBoxBrush( RootToCoreContentDir("Docking/Tab_Shape.png"),  2.f/8.0f, FColor(45, 45, 45) ) )
-			.SetUncheckedHoveredImage(    FSlateBoxBrush( RootToCoreContentDir("Docking/Tab_Shape.png"),  2.f/8.0f, FColor(54, 54, 54) ) )    			
-			.SetUncheckedPressedImage(    FSlateBoxBrush( RootToCoreContentDir("Docking/Tab_Shape.png"),  2.f/8.0f, FColor(54, 54, 54) ) )    			   
-			.SetUndeterminedImage(        FSlateNoResource() )
-			.SetUndeterminedHoveredImage( FSlateNoResource() )
-			.SetUndeterminedPressedImage( FSlateNoResource() )
-		);	
 		Set( "ToolPalette.DockingWell", new FSlateColorBrush(FLinearColor(FColor(34, 34, 34, 255))));
 
 		Set("LevelEditor.SelectMode", new IMAGE_BRUSH("Icons/icon_DefaultEdMode_40x", Icon40x40));
diff --git a/Engine/Source/Editor/EditorStyle/Private/StarshipStyle.cpp b/Engine/Source/Editor/EditorStyle/Private/StarshipStyle.cpp
--- a/Engine/Source/Editor/EditorStyle/Private/StarshipStyle.cpp
+++ b/Engine/Source/Editor/EditorStyle/Private/StarshipStyle.cpp
@@ -4395,12 +4395,12 @@ void FStarshipEditorStyle::FStyle::SetupLevelEditorStyle()
 		Set( "ToolPalette.DockingTab", FCheckBoxStyle()
 			.SetCheckBoxType(ESlateCheckBoxType::ToggleButton)
 			.SetPadding( FMargin(16.0f, 2.0f, 16.0f, 2.0f ) )
-			.SetCheckedImage(             FSlateBoxBrush( RootToCoreContentDir("Docking/Tab_Shape.png"),  2.f/8.0f, FColor(62, 62, 62) ) )
-			.SetCheckedHoveredImage(      FSlateBoxBrush( RootToCoreContentDir("Docking/Tab_Shape.png"),  2.f/8.0f, FColor(62, 62, 62) ) )
-			.SetCheckedPressedImage(      FSlateBoxBrush( RootToCoreContentDir("Docking/Tab_Shape.png"),  2.f/8.0f, FColor(62, 62, 62) ) )
-			.SetUncheckedImage(           FSlateBoxBrush( RootToCoreContentDir("Docking/Tab_Shape.png"),  2.f/8.0f, FColor(45, 45, 45) ) )
-			.SetUncheckedHoveredImage(    FSlateBoxBrush( RootToCoreContentDir("Docking/Tab_Shape.png"),  2.f/8.0f, FColor(54, 54, 54) ) )    			
-			.SetUncheckedPressedImage(    FSlateBoxBrush( RootToCoreContentDir("Docking/Tab_Shape.png"),  2.f/8.0f, FColor(54, 54, 54) ) )    			   
+			.SetCheckedImage(       BOX_BRUSH("Docking/Tab_Shape",  2.f/8.0f, FLinearColor(FColor(62, 62, 62)) ) )
+			.SetCheckedHoveredImage(BOX_BRUSH("Docking/Tab_Shape",  2.f/8.0f, FLinearColor(FColor(62, 62, 62)) ) )
+			.SetCheckedPressedImage(BOX_BRUSH("Docking/Tab_Shape",  2.f/8.0f, FLinearColor(FColor(62, 62, 62)) ) )
+			.SetUncheckedImage(     BOX_BRUSH("Docking/Tab_Shape",  2.f/8.0f, FLinearColor(FColor(45, 45, 45)) ) )
+			.SetUncheckedHoveredImage(BOX_BRUSH("Docking/Tab_Shape",2.f/8.0f, FLinearColor(FColor(54, 54, 54)) ) )
+			.SetUncheckedPressedImage(BOX_BRUSH("Docking/Tab_Shape",2.f/8.0f, FLinearColor(FColor(54, 54, 54)) ) )
 			.SetUndeterminedImage(        FSlateNoResource() )
 			.SetUndeterminedHoveredImage( FSlateNoResource() )
 			.SetUndeterminedPressedImage( FSlateNoResource() )
diff --git a/Engine/Source/Runtime/SlateCore/Public/Brushes/SlateRoundedBoxBrush.h b/Engine/Source/Runtime/SlateCore/Public/Brushes/SlateRoundedBoxBrush.h
--- a/Engine/Source/Runtime/SlateCore/Public/Brushes/SlateRoundedBoxBrush.h
+++ b/Engine/Source/Runtime/SlateCore/Public/Brushes/SlateRoundedBoxBrush.h
@@ -15,7 +15,9 @@ struct SLATECORE_API FSlateRoundedBoxBrush
 	 * @param InColor 		Linear Fill Color 
 	 * @param InRadius      Corner Radius
 	 */
-	FORCENOINLINE FSlateRoundedBoxBrush( const FLinearColor& InColor,  const FVector2D& InImageSize = FVector2D::ZeroVector)
+
+	template<typename FillColorType>
+	FORCENOINLINE FSlateRoundedBoxBrush(const FillColorType& InColor,  const FVector2D& InImageSize = FVector2D::ZeroVector)
 		: FSlateBrush(ESlateBrushDrawType::RoundedBox, 
 					  NAME_None, 
 					  FMargin(0.0f), 
@@ -29,6 +31,21 @@ struct SLATECORE_API FSlateRoundedBoxBrush
 	{ 
 	}
 
+	template<typename FillColorType>
+	FORCENOINLINE FSlateRoundedBoxBrush(const FName InFillResourceName, const FillColorType& FillColorTint, const FVector2D& InImageSize = FVector2D::ZeroVector, ESlateBrushTileType::Type InTileType = ESlateBrushTileType::NoTile)
+		: FSlateBrush(ESlateBrushDrawType::RoundedBox,
+			InFillResourceName,
+			FMargin(0.0f),
+			InTileType,
+			ESlateBrushImageType::FullColor,
+			InImageSize,
+			FillColorTint,
+			nullptr,
+			false
+		)
+	{
+	}
+
 
 	/** 
 	 * Creates and initializes a new instance with the specified color and corner radius
@@ -36,7 +53,8 @@ struct SLATECORE_API FSlateRoundedBoxBrush
 	 * @param InColor 		Linear Fill Color 
 	 * @param InRadius      Corner Radius
 	 */
-	FORCENOINLINE FSlateRoundedBoxBrush( const FLinearColor& InColor, float InRadius , const FVector2D& InImageSize = FVector2D::ZeroVector)
+	template<typename FillColorType>
+	FORCENOINLINE FSlateRoundedBoxBrush(const FillColorType& InColor, float InRadius, const FVector2D& InImageSize = FVector2D::ZeroVector)
 		: FSlateBrush(ESlateBrushDrawType::RoundedBox, 
 					  NAME_None, 
 					  FMargin(0.0f), 
@@ -51,39 +69,33 @@ struct SLATECORE_API FSlateRoundedBoxBrush
 		OutlineSettings = FSlateBrushOutlineSettings(InRadius);
 	}
 
-	/** 
-	 * Creates and initializes a new instance with the specified color and rounds based on height
-	 *
-	 * @param InColor 		 Linear Fill Color 
-	 * @param InOutlineColor Outline Color 
-	 * @param InOutlineWidth Outline Width or Thickness
-	 */
-	FORCENOINLINE FSlateRoundedBoxBrush( const FLinearColor& InColor, const FLinearColor& InOutlineColor, float InOutlineWidth, const FVector2D& InImageSize = FVector2D::ZeroVector)
-		: FSlateBrush(ESlateBrushDrawType::RoundedBox, 
-					  NAME_None, 
-					  FMargin(0.0f), 
-					  ESlateBrushTileType::NoTile, 
-					  ESlateBrushImageType::NoImage, 
-					  InImageSize,
-					  InColor,
-					  nullptr,
-					  false
-					  )
-	{ 
-		OutlineSettings = FSlateBrushOutlineSettings(InOutlineColor, InOutlineWidth);
-	}
 
+	template<typename FillColorType>
+	FORCENOINLINE FSlateRoundedBoxBrush(const FName InFillResourceName, const FillColorType& FillColorTint, float InRadius, const FVector2D& InImageSize = FVector2D::ZeroVector, ESlateBrushTileType::Type InTileType = ESlateBrushTileType::NoTile)
+		: FSlateBrush(ESlateBrushDrawType::RoundedBox,
+			InFillResourceName,
+			FMargin(0.0f),
+			InTileType,
+			ESlateBrushImageType::FullColor,
+			InImageSize,
+			FillColorTint,
+			nullptr,
+			false
+		)
+	{
+		OutlineSettings = FSlateBrushOutlineSettings(InRadius);
+	}
 
 
 	/** 
-	 * Creates and initializes a new instance with the specified color and corner radius
+	 * Creates and initializes a new instance with the specified color and rounds based on height
 	 *
 	 * @param InColor 		 Linear Fill Color 
-	 * @param InRadius       Corner Radius
 	 * @param InOutlineColor Outline Color 
 	 * @param InOutlineWidth Outline Width or Thickness
 	 */
-	FORCENOINLINE FSlateRoundedBoxBrush( const FLinearColor& InColor, float InRadius, const FLinearColor& InOutlineColor, float InOutlineWidth, const FVector2D& InImageSize = FVector2D::ZeroVector)
+	template<typename FillColorType, typename OutlineColorType>
+	FORCENOINLINE FSlateRoundedBoxBrush(const FillColorType& InColor, const OutlineColorType& InOutlineColor, float InOutlineWidth, const FVector2D& InImageSize = FVector2D::ZeroVector)
 		: FSlateBrush(ESlateBrushDrawType::RoundedBox, 
 					  NAME_None, 
 					  FMargin(0.0f), 
@@ -95,70 +107,30 @@ struct SLATECORE_API FSlateRoundedBoxBrush
 					  false
 					  )
 	{ 
-		OutlineSettings = FSlateBrushOutlineSettings(InRadius, InOutlineColor, InOutlineWidth);
-	}
-
-
-	/** 
-	 * Creates and initializes a new instance with the specified color and rounds based on height
-	 *
-	 * @param InColor 		Linear Fill Color 
-	 */
-	FORCENOINLINE FSlateRoundedBoxBrush( const FColor& InColor, const FVector2D& InImageSize = FVector2D::ZeroVector)
-		: FSlateBrush(ESlateBrushDrawType::RoundedBox, 
-					  NAME_None, 
-					  FMargin(0.0f), 
-					  ESlateBrushTileType::NoTile, 
-					  ESlateBrushImageType::NoImage, 
-					  InImageSize,
-					  InColor.ReinterpretAsLinear(),
-					  nullptr,
-					  false
-					 )
-	{ 
-	}
-
-	/** 
-	 * Creates and initializes a new instance with the specified color and corner radius
-	 *
-	 * @param InColor 		Linear Fill Color 
-	 * @param InRadius      Corner Radius
-	 */
-	FORCENOINLINE FSlateRoundedBoxBrush( const FColor& InColor, float InRadius , const FVector2D& InImageSize = FVector2D::ZeroVector)
-		: FSlateBrush(ESlateBrushDrawType::RoundedBox, 
-					  NAME_None, 
-					  FMargin(0.0f), 
-					  ESlateBrushTileType::NoTile, 
-					  ESlateBrushImageType::NoImage, 
-					  InImageSize,
-					  InColor.ReinterpretAsLinear(),
-					  nullptr,
-					  false
-					 )
-	{ 
-		OutlineSettings = FSlateBrushOutlineSettings(InRadius);
+		OutlineSettings = FSlateBrushOutlineSettings(InOutlineColor, InOutlineWidth);
 	}
 
-	/** 
+	/**
 	 * Creates and initializes a new instance with the specified color and rounds based on height
 	 *
-	 * @param InColor 		 Linear Fill Color 
-	 * @param InOutlineColor Outline Color 
+	 * @param InColor 		 Linear Fill Color
+	 * @param InOutlineColor Outline Color
 	 * @param InOutlineWidth Outline Width or Thickness
 	 */
-	FORCENOINLINE FSlateRoundedBoxBrush( const FColor& InColor, const FColor& InOutlineColor, float InOutlineWidth, const FVector2D& InImageSize = FVector2D::ZeroVector)
-		: FSlateBrush(ESlateBrushDrawType::RoundedBox, 
-					  NAME_None, 
-					  FMargin(0.0f), 
-					  ESlateBrushTileType::NoTile, 
-					  ESlateBrushImageType::NoImage, 
-					  InImageSize,
-					  InColor.ReinterpretAsLinear(),
-					  nullptr,
-					  false
-					  )
-	{ 
-		OutlineSettings = FSlateBrushOutlineSettings(InOutlineColor.ReinterpretAsLinear(), InOutlineWidth);
+	template<typename FillColorType, typename OutlineColorType>
+	FORCENOINLINE FSlateRoundedBoxBrush(const FName InFillResourceName, const FillColorType& FillColorTint, const OutlineColorType& InOutlineColor, float InOutlineWidth, const FVector2D& InImageSize = FVector2D::ZeroVector, ESlateBrushTileType::Type InTileType = ESlateBrushTileType::NoTile)
+		: FSlateBrush(ESlateBrushDrawType::RoundedBox,
+			InFillResourceName,
+			FMargin(0.0f),
+			InTileType,
+			ESlateBrushImageType::FullColor,
+			InImageSize,
+			FillColorTint,
+			nullptr,
+			false
+		)
+	{
+		OutlineSettings = FSlateBrushOutlineSettings(InOutlineColor, InOutlineWidth);
 	}
 
 
@@ -171,70 +143,8 @@ struct SLATECORE_API FSlateRoundedBoxBrush
 	 * @param InOutlineColor Outline Color 
 	 * @param InOutlineWidth Outline Width or Thickness
 	 */
-	FORCENOINLINE FSlateRoundedBoxBrush( const FColor& InColor, float InRadius, const FColor& InOutlineColor, float InOutlineWidth, const FVector2D& InImageSize = FVector2D::ZeroVector)
-		: FSlateBrush(ESlateBrushDrawType::RoundedBox, 
-					  NAME_None, 
-					  FMargin(0.0f), 
-					  ESlateBrushTileType::NoTile, 
-					  ESlateBrushImageType::NoImage, 
-					  InImageSize,
-					  InColor.ReinterpretAsLinear(),
-					  nullptr,
-					  false
-					  )
-	{ 
-		OutlineSettings = FSlateBrushOutlineSettings(InRadius, InOutlineColor.ReinterpretAsLinear(), InOutlineWidth);
-	}
-
-	/** 
-	 * Creates and initializes a new instance with the specified color and rounds based on height
-	 *
-	 * @param InColor 		Linear Fill Color 
-	 */
-	FORCENOINLINE FSlateRoundedBoxBrush( const FSlateColor& InColor , const FVector2D& InImageSize = FVector2D::ZeroVector)
-		: FSlateBrush(ESlateBrushDrawType::RoundedBox, 
-					  NAME_None, 
-					  FMargin(0.0f), 
-					  ESlateBrushTileType::NoTile, 
-					  ESlateBrushImageType::NoImage, 
-					  InImageSize,
-					  InColor,
-					  nullptr,
-					  false
-					 )
-	{ 
-	}
-
-
-	/** 
-	 * Creates and initializes a new instance with the specified color and corner radius
-	 *
-	 * @param InColor 		Linear Fill Color 
-	 * @param InRadius      Corner Radius
-	 */
-	FORCENOINLINE FSlateRoundedBoxBrush( const FSlateColor& InColor, float InRadius , const FVector2D& InImageSize = FVector2D::ZeroVector)
-		: FSlateBrush(ESlateBrushDrawType::RoundedBox, 
-					  NAME_None, 
-					  FMargin(0.0f), 
-					  ESlateBrushTileType::NoTile, 
-					  ESlateBrushImageType::NoImage, 
-					  InImageSize,
-					  InColor,
-					  nullptr,
-					  false
-					 )
-	{ 
-		OutlineSettings = FSlateBrushOutlineSettings(InRadius);
-	}
-
-	/** 
-	 * Creates and initializes a new instance with the specified color and rounds based on height
-	 *
-	 * @param InColor 		 Linear Fill Color 
-	 * @param InOutlineColor Outline Color 
-	 * @param InOutlineWidth Outline Width or Thickness
-	 */
-	FORCENOINLINE FSlateRoundedBoxBrush( const FSlateColor& InColor, const FSlateColor& InOutlineColor, float InOutlineWidth, const FVector2D& InImageSize = FVector2D::ZeroVector)
+	template<typename FillColorType, typename OutlineColorType>
+	FORCENOINLINE FSlateRoundedBoxBrush(const FillColorType& InColor, float InRadius, const OutlineColorType& InOutlineColor, float InOutlineWidth, const FVector2D& InImageSize = FVector2D::ZeroVector)
 		: FSlateBrush(ESlateBrushDrawType::RoundedBox, 
 					  NAME_None, 
 					  FMargin(0.0f), 
@@ -246,29 +156,30 @@ struct SLATECORE_API FSlateRoundedBoxBrush
 					  false
 					  )
 	{ 
-		OutlineSettings = FSlateBrushOutlineSettings(InOutlineColor, InOutlineWidth);
+		OutlineSettings = FSlateBrushOutlineSettings(InRadius, InOutlineColor, InOutlineWidth);
 	}
 
-	/** 
-	 * Creates and initializes a new instance with the specified color and corner radius
-	 *
-	 * @param InColor 		 Linear Fill Color 
-	 * @param InRadius       Corner Radius
-	 * @param InOutlineColor Outline Color 
-	 * @param InOutlineWidth Outline Width or Thickness
-	 */
-	FORCENOINLINE FSlateRoundedBoxBrush( const FSlateColor& InColor, float InRadius, const FSlateColor& InOutlineColor, float InOutlineWidth, const FVector2D& InImageSize = FVector2D::ZeroVector)
-		: FSlateBrush(ESlateBrushDrawType::RoundedBox, 
-					  NAME_None, 
-					  FMargin(0.0f), 
-					  ESlateBrushTileType::NoTile, 
-					  ESlateBrushImageType::NoImage, 
+	/**
+ * Creates and initializes a new instance with the specified color and corner radius
+ *
+ * @param InColor 		 Linear Fill Color
+ * @param InRadius       Corner Radius
+ * @param InOutlineColor Outline Color
+ * @param InOutlineWidth Outline Width or Thickness
+ */
+	template<typename FillColorType, typename OutlineColorType>
+	FORCENOINLINE FSlateRoundedBoxBrush(const FName InFillResourceName, const FillColorType& FillColorTint, float InRadius, const OutlineColorType& InOutlineColor, float InOutlineWidth, const FVector2D& InImageSize = FVector2D::ZeroVector, ESlateBrushTileType::Type InTileType = ESlateBrushTileType::NoTile)
+		: FSlateBrush(ESlateBrushDrawType::RoundedBox,
+					  InFillResourceName,
+					  FMargin(0.0f),
+					  InTileType,
+					  ESlateBrushImageType::FullColor,
 					  InImageSize,
-					  InColor,
+					  FillColorTint,
 					  nullptr,
 					  false
-					  )
-	{ 
+		)
+	{
 		OutlineSettings = FSlateBrushOutlineSettings(InRadius, InOutlineColor, InOutlineWidth);
 	}
 
diff --git a/Engine/Source/Runtime/SlateCore/Public/Styling/SlateColor.h b/Engine/Source/Runtime/SlateCore/Public/Styling/SlateColor.h
--- a/Engine/Source/Runtime/SlateCore/Public/Styling/SlateColor.h
+++ b/Engine/Source/Runtime/SlateCore/Public/Styling/SlateColor.h
@@ -59,17 +59,29 @@ public:
 	 *
 	 * @param InColor The color value to assign.
 	 */
-	FSlateColor( const FLinearColor& InColor )
+	FSlateColor(const FLinearColor& InColor)
 		: SpecifiedColor(InColor)
 		, ColorUseRule(ESlateColorStylingMode::UseColor_Specified)
 	{ }
 
+
 	/**
+	 * Creates a new Slate color with the specified values.
+	 *
+	 * @param InColor The color value to assign.
+	 */
+	FSlateColor(const FColor InColor)
+		: SpecifiedColor(InColor.ReinterpretAsLinear())
+		, ColorUseRule(ESlateColorStylingMode::UseColor_Specified)
+	{ }
+
+
+	/*
 	 * Creates a new Slate color that is linked to the given values.
 	 *
 	 * @param InColor The color value to assign.
 	 */
-	FSlateColor( const TSharedRef< FLinearColor >& InColor )
+	FSlateColor(const TSharedRef< FLinearColor >& InColor)
 		: SpecifiedColor(*InColor)
 		, ColorUseRule(ESlateColorStylingMode::UseColor_Specified)
 	{ }
