Changes to window styles for a more visible border between parent and child windows.

diff --git a/Engine/Source/Editor/EditorStyle/Private/StarshipStyle.cpp b/Engine/Source/Editor/EditorStyle/Private/StarshipStyle.cpp
--- a/Engine/Source/Editor/EditorStyle/Private/StarshipStyle.cpp
+++ b/Engine/Source/Editor/EditorStyle/Private/StarshipStyle.cpp
@@ -212,9 +212,6 @@ void FStarshipEditorStyle::FStyle::SyncSettings()
 		{
 			SetColor(WindowHighlightColor_LinearRef, FLinearColor(0, 0, 0, 0));
 		}
-
-		WindowStyle.SetBorderBrush(EditorWindowHighlightBrush);
-		WindowStyle.SetBorderColor(WindowHighlightColor);
 	}
 }
 
@@ -7813,9 +7810,6 @@ void FStarshipEditorStyle::FStyle::SetupStatusBarStyle()
 		.SetColorAndOpacity(FStyleColors::Primary));
 }
 
-#undef IMAGE_BRUSH
-#undef BOX_BRUSH
-#undef BORDER_BRUSH
 #undef DEFAULT_FONT
 #undef ICON_FONT
 #undef LOCTEXT_NAMESPACE
diff --git a/Engine/Source/Runtime/ApplicationCore/Private/Windows/WindowsWindow.cpp b/Engine/Source/Runtime/ApplicationCore/Private/Windows/WindowsWindow.cpp
--- a/Engine/Source/Runtime/ApplicationCore/Private/Windows/WindowsWindow.cpp
+++ b/Engine/Source/Runtime/ApplicationCore/Private/Windows/WindowsWindow.cpp
@@ -468,7 +468,7 @@ void FWindowsWindow::ReshapeWindow( int32 NewX, int32 NewY, int32 NewWidth, int3
 	// prevents us being able to change to a higher resolution while in fullscreen mode
 	::SetWindowPos( HWnd, nullptr, WindowX, WindowY, NewWidth, NewHeight, SWP_NOZORDER | SWP_NOACTIVATE | ((WindowMode == EWindowMode::Fullscreen) ? SWP_NOSENDCHANGING : 0) );
 
-	if( Definition->SizeWillChangeOften && bVirtualSizeChanged )
+	if( (Definition->SizeWillChangeOften && bVirtualSizeChanged) || Definition->CornerRadius > 0 )
 	{
 		AdjustWindowRegion( VirtualWidth, VirtualHeight );
 	}
diff --git a/Engine/Source/Runtime/SlateCore/Private/Styling/SlateTypes.cpp b/Engine/Source/Runtime/SlateCore/Private/Styling/SlateTypes.cpp
--- a/Engine/Source/Runtime/SlateCore/Private/Styling/SlateTypes.cpp
+++ b/Engine/Source/Runtime/SlateCore/Private/Styling/SlateTypes.cpp
@@ -677,7 +677,7 @@ const FScrollBorderStyle& FScrollBorderStyle::GetDefault()
 
 FWindowStyle::FWindowStyle()
 	: BackgroundColor(FLinearColor::White)
-	, OutlineColor(FLinearColor(0.1f, 0.1f, 0.1f, 1.0f))
+	, OutlineColor(FLinearColor::White)
 	, BorderColor(FLinearColor::White)
 	, WindowCornerRadius(0)
 	, BorderPadding(FMargin(5, 5, 5, 5))
diff --git a/Engine/Source/Runtime/SlateCore/Private/Styling/StarshipCoreStyle.cpp b/Engine/Source/Runtime/SlateCore/Private/Styling/StarshipCoreStyle.cpp
--- a/Engine/Source/Runtime/SlateCore/Private/Styling/StarshipCoreStyle.cpp
+++ b/Engine/Source/Runtime/SlateCore/Private/Styling/StarshipCoreStyle.cpp
@@ -671,8 +671,8 @@ TSharedRef<ISlateStyle> FStarshipCoreStyle::Create()
 			.SetPressed(IMAGE_BRUSH_SVG("Starship/CoreWidgets/Window/close", FVector2D(42.0f, 34.0f), FStyleColors::Foreground));
 #endif
 
-		Style->Set("Window", FWindowStyle()
-
+		FWindowStyle Window =
+			FWindowStyle()
 #if !PLATFORM_MAC
 			.SetMinimizeButtonStyle(MinimizeButtonStyle)
 			.SetMaximizeButtonStyle(MaximizeButtonStyle)
@@ -682,16 +682,17 @@ TSharedRef<ISlateStyle> FStarshipCoreStyle::Create()
 			.SetTitleTextStyle(NormalText)
 			.SetActiveTitleBrush(FSlateNoResource())
 			.SetInactiveTitleBrush(FSlateNoResource())
-			.SetFlashTitleBrush(IMAGE_BRUSH("Common/Window/WindowTitle_Flashing", Icon24x24, FLinearColor(1.0f, 1.0f, 1.0f, 1.0f), ESlateBrushTileType::Horizontal))	
+			.SetFlashTitleBrush(IMAGE_BRUSH("Common/Window/WindowTitle_Flashing", Icon24x24, FLinearColor(1.0f, 1.0f, 1.0f, 1.0f), ESlateBrushTileType::Horizontal))
 			.SetBackgroundBrush(FSlateColorBrush(FStyleColors::Recessed))
-			.SetBorderBrush(FSlateNoResource())
-			.SetOutlineBrush(FSlateRoundedBoxBrush(FLinearColor::Transparent, 0.0, FStyleColors::WindowBorder, 2.0))
+			.SetBorderBrush(FSlateRoundedBoxBrush(FStyleColors::Recessed, 2.0, FStyleColors::WindowBorder, 2.0))
+			.SetOutlineBrush(FSlateRoundedBoxBrush(FStyleColors::Recessed, 2.0, FStyleColors::InputOutline, 1.0))
 			.SetChildBackgroundBrush(FSlateColorBrush(FStyleColors::Recessed))
-			.SetOutlineColor(WindowHighlight)
-			.SetCornerRadius(0)
-			.SetBorderPadding(FMargin(2.0, 2.0, 2.0, 2.0))
-		);
+			.SetCornerRadius(2)
+			.SetBorderPadding(FMargin(3.0, 3.0, 3.0, 3.0));
 
+
+		Style->Set("Window", Window);
+		
 		Style->Set("ChildWindow.Background", new FSlateColorBrush(FStyleColors::Recessed));
 	}
 
diff --git a/Engine/Source/Runtime/SlateCore/Private/Widgets/SWindow.cpp b/Engine/Source/Runtime/SlateCore/Private/Widgets/SWindow.cpp
--- a/Engine/Source/Runtime/SlateCore/Private/Widgets/SWindow.cpp
+++ b/Engine/Source/Runtime/SlateCore/Private/Widgets/SWindow.cpp
@@ -530,7 +530,7 @@ EHorizontalAlignment SWindow::GetTitleAlignment()
 
 void SWindow::ConstructWindowInternals()
 {
-	ForegroundColor = FCoreStyle::Get().GetSlateColor("DefaultForeground");
+	ForegroundColor = FAppStyle::Get().GetSlateColor("DefaultForeground");
 
 	// Setup widget that represents the main area of the window.  That is, everything inside the window's border.
 	TSharedRef< SVerticalBox > MainWindowArea =
@@ -561,17 +561,18 @@ void SWindow::ConstructWindowInternals()
 
 	UpdateWindowContentVisibility();
 
-	// create window content slot
-	MainWindowArea->AddSlot()
-		.FillHeight(1.0f)
-		.Expose(ContentSlot)
-		[
-			SNullWidget::NullWidget
-		];
-
 	// create window
-	if (Type != EWindowType::ToolTip && Type != EWindowType::CursorDecorator && !bIsPopupWindow && !bHasOSWindowBorder)
+	if ((Type == EWindowType::Normal || Type == EWindowType::GameWindow) && !bHasOSWindowBorder)
 	{
+		// create window content slot
+		MainWindowArea->AddSlot()
+			.FillHeight(1.0f)
+			.Expose(ContentSlot)
+			[
+				SNullWidget::NullWidget
+			];
+
+
 		WindowBackgroundImage =
 			FSlateApplicationBase::Get().MakeImage(
 				WindowBackground,
@@ -609,18 +610,18 @@ void SWindow::ConstructWindowInternals()
 
 			// window outline
 			+ SOverlay::Slot()
+			.Padding(2.0f)
 			[
 				WindowOutline.ToSharedRef()
 			]
 
 			// main area
 			+ SOverlay::Slot()
+			.Padding(TAttribute<FMargin>::Create(TAttribute<FMargin>::FGetter::CreateSP(this, &SWindow::GetWindowBorderSize, false)))
 			[
 				SAssignNew(ContentAreaVBox, SVerticalBox)
 				.Visibility(WindowContentVisibility)
-
-				+ SVerticalBox::Slot()
-				.Padding(TAttribute<FMargin>::Create(TAttribute<FMargin>::FGetter::CreateSP(this, &SWindow::GetWindowBorderSize, false)))
+				+ SVerticalBox::Slot()	
 				[
 					MainWindowArea
 				]
@@ -635,6 +636,15 @@ void SWindow::ConstructWindowInternals()
 	}
 	else if ( bHasOSWindowBorder || bVirtualWindow )
 	{
+		// create window content slot
+		MainWindowArea->AddSlot()
+			.FillHeight(1.0f)
+			.Expose(ContentSlot)
+			[
+				SNullWidget::NullWidget
+			];
+
+
 		this->ChildSlot
 		[
 			SAssignNew(WindowOverlay, SOverlay)
@@ -1247,26 +1257,27 @@ void SWindow::SetNativeWindow( TSharedRef<FGenericWindow> InNativeWindow )
 
 void SWindow::SetContent( TSharedRef<SWidget> InContent )
 {
-	if ( bIsPopupWindow || Type == EWindowType::CursorDecorator )
+	if (ContentSlot)
 	{
-		this->ChildSlot.operator[]( InContent );
+		ContentSlot->operator[](InContent);
 	}
 	else
 	{
-		this->ContentSlot->operator[]( InContent );
+		ChildSlot.operator[](InContent);
 	}
+	
 	Invalidate(EInvalidateWidgetReason::ChildOrder);
 }
 
 TSharedRef<const SWidget> SWindow::GetContent() const
 {
-	if ( bIsPopupWindow || Type == EWindowType::CursorDecorator )
+	if (ContentSlot)
 	{
-		return this->ChildSlot.GetChildAt(0);
+		return ContentSlot->GetWidget();
 	}
 	else
 	{
-		return this->ContentSlot->GetWidget();
+		return ChildSlot.GetChildAt(0);
 	}
 }
 
@@ -1694,7 +1705,6 @@ bool SWindow::OnIsActiveChanged( const FWindowActivateEvent& ActivateEvent )
 	return true;
 }
 
-
 void SWindow::Maximize()
 {
 	if (NativeWindow.IsValid())
@@ -1721,7 +1731,7 @@ void SWindow::Minimize()
 
 int32 SWindow::GetCornerRadius()
 {
-	return IsRegularWindow() ? Style->WindowCornerRadius : 0;
+	return Style->WindowCornerRadius;
 }
 
 bool SWindow::SupportsKeyboardFocus() const
