Adding Missing Style Brush for Floating Borders

diff --git a/Engine/Source/Runtime/SlateCore/Private/Styling/StarshipCoreStyle.cpp b/Engine/Source/Runtime/SlateCore/Private/Styling/StarshipCoreStyle.cpp
--- a/Engine/Source/Runtime/SlateCore/Private/Styling/StarshipCoreStyle.cpp
+++ b/Engine/Source/Runtime/SlateCore/Private/Styling/StarshipCoreStyle.cpp
@@ -448,6 +448,11 @@ TSharedRef<ISlateStyle> FStarshipCoreStyle::Create()
 	// SBorder defaults...
 	{
 		Style->Set("Border", new FSlateColorBrush(FStyleColors::Background));
+
+		FLinearColor TransBackground = FStyleColors::Background.GetSpecifiedColor();
+		TransBackground.A = .5;
+
+		Style->Set("FloatingBorder", new FSlateRoundedBoxBrush(TransBackground, 8.f));
 	}
 
 	// SHyperlink defaults...
