Checkbox style on Menu Hover

diff --git a/Engine/Source/Runtime/Slate/Private/Framework/MultiBox/SMenuEntryBlock.cpp b/Engine/Source/Runtime/Slate/Private/Framework/MultiBox/SMenuEntryBlock.cpp
--- a/Engine/Source/Runtime/Slate/Private/Framework/MultiBox/SMenuEntryBlock.cpp
+++ b/Engine/Source/Runtime/Slate/Private/Framework/MultiBox/SMenuEntryBlock.cpp
@@ -1576,10 +1576,22 @@ FSlateColor SMenuEntryBlock::InvertOnHover() const
  */
 const FSlateBrush* SMenuEntryBlock::GetCheckBoxImageBrushFromStyle(const FCheckBoxStyle* Style) const
 {
-	switch (IsChecked())
+	if (IsHovered()) 
 	{
-	case ECheckBoxState::Checked: return &Style->CheckedImage;
-	case ECheckBoxState::Unchecked: return &Style->UncheckedImage;
-	default: return &Style->UndeterminedImage;
+		switch (IsChecked())
+		{
+			case ECheckBoxState::Checked: return &Style->CheckedHoveredImage;
+			case ECheckBoxState::Unchecked: return &Style->UncheckedHoveredImage;
+			default: return &Style->UndeterminedHoveredImage;
+		}
+	}
+	else
+	{
+		switch (IsChecked())
+		{
+			case ECheckBoxState::Checked: return &Style->CheckedImage;
+			case ECheckBoxState::Unchecked: return &Style->UncheckedImage;
+			default: return &Style->UndeterminedImage;
+		}
 	}
 }
diff --git a/Engine/Source/Runtime/SlateCore/Private/Styling/StarshipCoreStyle.cpp b/Engine/Source/Runtime/SlateCore/Private/Styling/StarshipCoreStyle.cpp
--- a/Engine/Source/Runtime/SlateCore/Private/Styling/StarshipCoreStyle.cpp
+++ b/Engine/Source/Runtime/SlateCore/Private/Styling/StarshipCoreStyle.cpp
@@ -1831,16 +1831,19 @@ void FStarshipCoreStyle::SetupMultiboxStyles(TSharedRef<FStyle>& Style)
 		Style->Set("Menu.Heading.Padding", MenuBlockPadding + FMargin(0, 10, 0, 0));
 
 		/* Set images for various SCheckBox states associated with menu check box items... */
+		FLinearColor Transparent20 = FLinearColor(1.0, 1.0, 1.0, 0.2);
 		const FCheckBoxStyle BasicMenuCheckBoxStyle = FCheckBoxStyle()
-			.SetUncheckedImage(FSlateNoResource(Icon16x16))
-			.SetUncheckedHoveredImage(FSlateNoResource(Icon16x16))
-			.SetUncheckedPressedImage(FSlateNoResource(Icon16x16))
-			.SetCheckedImage(	    IMAGE_BRUSH_SVG("Starship/Common/check", Icon16x16))
-			.SetCheckedHoveredImage(IMAGE_BRUSH_SVG("Starship/Common/check", Icon16x16))
-			.SetCheckedPressedImage(IMAGE_BRUSH_SVG("Starship/Common/check", Icon16x16))
-			.SetUndeterminedImage(FSlateNoResource())
-			.SetUndeterminedHoveredImage(FSlateNoResource())
-			.SetUndeterminedPressedImage(FSlateNoResource());
+			.SetUncheckedImage(           FSlateNoResource())
+			.SetUncheckedHoveredImage(    IMAGE_BRUSH_SVG("Starship/Common/check", Icon16x16, Transparent20))
+			.SetUncheckedPressedImage(    IMAGE_BRUSH_SVG("Starship/Common/check", Icon16x16, Transparent20))
+
+			.SetCheckedImage(             IMAGE_BRUSH_SVG("Starship/Common/check", Icon16x16))
+			.SetCheckedHoveredImage(      IMAGE_BRUSH_SVG("Starship/Common/check", Icon16x16))
+			.SetCheckedPressedImage(      IMAGE_BRUSH_SVG("Starship/Common/check", Icon16x16))
+
+			.SetUndeterminedImage(        FSlateNoResource())
+			.SetUndeterminedHoveredImage( IMAGE_BRUSH_SVG("Starship/Common/check", Icon16x16, Transparent20))
+			.SetUndeterminedPressedImage( IMAGE_BRUSH_SVG("Starship/Common/check", Icon16x16, Transparent20));
 
 		/* ...and add the new style */
 		Style->Set("Menu.CheckBox", BasicMenuCheckBoxStyle);
