Temp Fix to menu flickering due to checkboxes

diff --git a/Engine/Source/Runtime/SlateCore/Private/Styling/StarshipCoreStyle.cpp b/Engine/Source/Runtime/SlateCore/Private/Styling/StarshipCoreStyle.cpp
--- a/Engine/Source/Runtime/SlateCore/Private/Styling/StarshipCoreStyle.cpp
+++ b/Engine/Source/Runtime/SlateCore/Private/Styling/StarshipCoreStyle.cpp
@@ -1845,8 +1845,9 @@ void FStarshipCoreStyle::SetupMultiboxStyles(TSharedRef<FStyle>& Style)
 
 		/* Set images for various SCheckBox states associated with menu check box items... */
 		FLinearColor Transparent20 = FLinearColor(1.0, 1.0, 1.0, 0.2);
+		FLinearColor Transparent01 = FLinearColor(1.0, 1.0, 1.0, 0.01);
 		const FCheckBoxStyle BasicMenuCheckBoxStyle = FCheckBoxStyle()
-			.SetUncheckedImage(           FSlateNoResource())
+			.SetUncheckedImage(           IMAGE_BRUSH_SVG("Starship/Common/check", Icon16x16, Transparent01))
 			.SetUncheckedHoveredImage(    IMAGE_BRUSH_SVG("Starship/Common/check", Icon16x16, Transparent20))
 			.SetUncheckedPressedImage(    IMAGE_BRUSH_SVG("Starship/Common/check", Icon16x16, Transparent20))
 
@@ -1854,7 +1855,7 @@ void FStarshipCoreStyle::SetupMultiboxStyles(TSharedRef<FStyle>& Style)
 			.SetCheckedHoveredImage(      IMAGE_BRUSH_SVG("Starship/Common/check", Icon16x16))
 			.SetCheckedPressedImage(      IMAGE_BRUSH_SVG("Starship/Common/check", Icon16x16))
 
-			.SetUndeterminedImage(        FSlateNoResource())
+			.SetUndeterminedImage(        IMAGE_BRUSH_SVG("Starship/Common/check", Icon16x16, Transparent01))
 			.SetUndeterminedHoveredImage( IMAGE_BRUSH_SVG("Starship/Common/check", Icon16x16, Transparent20))
 			.SetUndeterminedPressedImage( IMAGE_BRUSH_SVG("Starship/Common/check", Icon16x16, Transparent20));
 
