Avoid allocations inside trace log scope.

diff --git a/Engine/Source/Runtime/SlateCore/Private/Trace/SlateTrace.cpp b/Engine/Source/Runtime/SlateCore/Private/Trace/SlateTrace.cpp
--- a/Engine/Source/Runtime/SlateCore/Private/Trace/SlateTrace.cpp
+++ b/Engine/Source/Runtime/SlateCore/Private/Trace/SlateTrace.cpp
@@ -248,11 +248,13 @@ void FSlateTrace::UpdateWidgetInfo(const SWidget* Widget)
 	if (UE_TRACE_CHANNELEXPR_IS_ENABLED(SlateChannel))
 	{
 		const uint64 WidgetId = SlateTraceDetail::GetWidgetId(Widget);
+		const FString Path = FReflectionMetaData::GetWidgetPath(Widget);
+		const FString DebugInfo = FReflectionMetaData::GetWidgetDebugInfo(Widget);
 
 		UE_TRACE_LOG(SlateTrace, WidgetInfo, SlateChannel)
 			<< WidgetInfo.WidgetId(WidgetId)
-			<< WidgetInfo.Path(*FReflectionMetaData::GetWidgetPath(Widget))
-			<< WidgetInfo.DebugInfo(*FReflectionMetaData::GetWidgetDebugInfo(Widget));
+			<< WidgetInfo.Path(*Path)
+			<< WidgetInfo.DebugInfo(*DebugInfo);
 	}
 }
 
