Fix some inplicit conversation in Slate code.

diff --git a/Engine/Source/Runtime/Slate/Private/Widgets/Input/SVirtualJoystick.cpp b/Engine/Source/Runtime/Slate/Private/Widgets/Input/SVirtualJoystick.cpp
--- a/Engine/Source/Runtime/Slate/Private/Widgets/Input/SVirtualJoystick.cpp
+++ b/Engine/Source/Runtime/Slate/Private/Widgets/Input/SVirtualJoystick.cpp
@@ -82,7 +82,7 @@ bool SVirtualJoystick::ShouldDisplayTouchInterface()
 	return FPlatformMisc::GetUseVirtualJoysticks() || bAlwaysShowTouchInterface || ( FSlateApplication::Get().IsFakingTouchEvents() && FPlatformMisc::ShouldDisplayTouchInterfaceOnFakingTouchEvents());
 }
 
-static int32 ResolveRelativePosition(float Position, float RelativeTo, float ScaleFactor)
+static float ResolveRelativePosition(float Position, float RelativeTo, float ScaleFactor)
 {
 	// absolute from edge
 	if (Position < -1.0f)
diff --git a/Engine/Source/Runtime/SlateCore/Private/Debugging/ConsoleSlateDebugger.cpp b/Engine/Source/Runtime/SlateCore/Private/Debugging/ConsoleSlateDebugger.cpp
--- a/Engine/Source/Runtime/SlateCore/Private/Debugging/ConsoleSlateDebugger.cpp
+++ b/Engine/Source/Runtime/SlateCore/Private/Debugging/ConsoleSlateDebugger.cpp
@@ -198,9 +198,9 @@ namespace SlateDebuggerUtils
 		{
 			for (int32 Index = 0; Index < Params.Num(); Index += 2)
 			{
-				const int64 InputEventEnumValue = EventEnum->GetValueByNameString(Params[Index]);
+				const int32 InputEventEnumValue = (int32)EventEnum->GetValueByNameString(Params[Index]);
 				// We are casting ESlateDebuggingInputEvent::MAX to uint8. Let's make sure it has a valid value. Prefer to crash now than doing bad memory access.
-				check(EventEnum->GetMaxEnumValue() < (int64)TNumericLimits<uint8>::Max());
+				check(EventEnum->GetMaxEnumValue() < TNumericLimits<uint8>::Max());
 				if (InputEventEnumValue == INDEX_NONE)
 				{
 					return;
diff --git a/Engine/Source/Runtime/SlateCore/Private/Debugging/ConsoleSlateDebuggerInvalidate.cpp b/Engine/Source/Runtime/SlateCore/Private/Debugging/ConsoleSlateDebuggerInvalidate.cpp
--- a/Engine/Source/Runtime/SlateCore/Private/Debugging/ConsoleSlateDebuggerInvalidate.cpp
+++ b/Engine/Source/Runtime/SlateCore/Private/Debugging/ConsoleSlateDebuggerInvalidate.cpp
@@ -36,7 +36,7 @@ FConsoleSlateDebuggerInvalidate::FConsoleSlateDebuggerInvalidate()
 	, DrawWidgetRenderTransformColor(FColorList::Black)
 	, DrawWidgetVisibilityColor(FColorList::White)
 	, MaxNumberOfWidgetInList(20)
-	, CacheDuration(2.0)
+	, CacheDuration(2.0f)
 	, StartCommand(
 		TEXT("SlateDebugger.Invalidate.Start"),
 		TEXT("Start the Invalidation widget debug tool. It shows when widgets are invalidated."),
@@ -682,7 +682,7 @@ void FConsoleSlateDebuggerInvalidate::HandlePaintDebugInfo(const FPaintArgs& InA
 			AlreadyProcessedInvalidatedId.Add(InvalidationInfo.WidgetInvalidatorId);
 		}
 
-		const double LerpValue = FMath::Clamp((SlateApplicationCurrentTime - InvalidationInfo.InvalidationTime) / CacheDuration, 0.0, 1.0);
+		const float LerpValue = FMath::Clamp((float)(SlateApplicationCurrentTime - InvalidationInfo.InvalidationTime) / CacheDuration, 0.0f, 1.0f);
 		const FLinearColor ColorWithOpacity = InvalidationInfo.DisplayColor.CopyWithNewOpacity(FMath::InterpExpoOut(1.0f, 0.2f, LerpValue));
 
 		{
diff --git a/Engine/Source/Runtime/SlateCore/Private/Debugging/ConsoleSlateDebuggerInvalidationRoot.cpp b/Engine/Source/Runtime/SlateCore/Private/Debugging/ConsoleSlateDebuggerInvalidationRoot.cpp
--- a/Engine/Source/Runtime/SlateCore/Private/Debugging/ConsoleSlateDebuggerInvalidationRoot.cpp
+++ b/Engine/Source/Runtime/SlateCore/Private/Debugging/ConsoleSlateDebuggerInvalidationRoot.cpp
@@ -30,7 +30,7 @@ FConsoleSlateDebuggerInvalidationRoot::FConsoleSlateDebuggerInvalidationRoot()
 	, DrawFastPathColor(FColorList::Green)
 	, DrawNoneColor(FColorList::Blue)
 	, MaxNumberOfWidgetInList(20)
-	, CacheDuration(2.0)
+	, CacheDuration(2.0f)
 	, StartCommand(
 		TEXT("SlateDebugger.InvalidationRoot.Start"),
 		TEXT("Start the Invalidation Root widget debug tool. It shows when Invalidation Root are using the slow or the fast path."),
@@ -242,7 +242,7 @@ void FConsoleSlateDebuggerInvalidationRoot::HandlePaintDebugInfo(const FPaintArg
 		if (WindowId == PaintWindow)
 		{
 			FLinearColor DrawColor = GetColor(LastPaintType);
-			double LerpValue = 1.0;
+			float LerpValue = 1.0f;
 
 			FInvalidatedInfo& FoundInvalidationInfo = InvaliadatedRoots[Itt.Key];
 			// If we went from fast to slow or to none, flash it on screen
@@ -259,12 +259,12 @@ void FConsoleSlateDebuggerInvalidationRoot::HandlePaintDebugInfo(const FPaintArg
 			}
 			FoundInvalidationInfo.PaintType = LastPaintType;
 
-			const double DeltaTime = SlateApplicationCurrentTime - FoundInvalidationInfo.FlashingSeconds;
+			const float DeltaTime = (float)(SlateApplicationCurrentTime - FoundInvalidationInfo.FlashingSeconds);
 			const bool bLearp = DeltaTime <= CacheDuration;
 			if (bLearp)
 			{
 				DrawColor = FoundInvalidationInfo.FlashingColor;
-				LerpValue = FMath::Clamp(DeltaTime / CacheDuration, 0.1, 1.0);
+				LerpValue = FMath::Clamp(DeltaTime / CacheDuration, 0.1f, 1.0f);
 			}
 
 			const SWidget* Widget = Itt.Value->GetInvalidationRootWidget();
diff --git a/Engine/Source/Runtime/SlateCore/Private/Debugging/ConsoleSlateDebuggerPaint.cpp b/Engine/Source/Runtime/SlateCore/Private/Debugging/ConsoleSlateDebuggerPaint.cpp
--- a/Engine/Source/Runtime/SlateCore/Private/Debugging/ConsoleSlateDebuggerPaint.cpp
+++ b/Engine/Source/Runtime/SlateCore/Private/Debugging/ConsoleSlateDebuggerPaint.cpp
@@ -30,7 +30,7 @@ FConsoleSlateDebuggerPaint::FConsoleSlateDebuggerPaint()
 	, DrawQuadColor(1.0f, 1.0f, 1.0f, 1.0f)
 	, DrawWidgetNameColor(FColorList::SpicyPink)
 	, MaxNumberOfWidgetInList(20)
-	, CacheDuration(2.0)
+	, CacheDuration(2.0f)
 	, ShowPaintWidgetCommand(
 		TEXT("SlateDebugger.Paint.Start")
 		, TEXT("Start the painted widget debug tool. It shows when widgets are painted.")
@@ -238,7 +238,7 @@ void FConsoleSlateDebuggerPaint::HandlePaintDebugInfo(const FPaintArgs& InArgs,
 	{
 		if (Itt.Value.Window == PaintWindow)
 		{
-			const double LerpValue = FMath::Clamp((SlateApplicationCurrentTime - Itt.Value.LastPaint) / CacheDuration, 0.0, 1.0);
+			const float LerpValue = FMath::Clamp((float)(SlateApplicationCurrentTime - Itt.Value.LastPaint) / CacheDuration, 0.0f, 1.0f);
 			const FGeometry Geometry = FGeometry::MakeRoot(Itt.Value.PaintSize, FSlateLayoutTransform(1.f, Itt.Value.PaintLocation));
 			const FPaintGeometry PaintGeometry = Geometry.ToPaintGeometry();
 
diff --git a/Engine/Source/Runtime/SlateCore/Private/Debugging/ConsoleSlateDebuggerUpdate.cpp b/Engine/Source/Runtime/SlateCore/Private/Debugging/ConsoleSlateDebuggerUpdate.cpp
--- a/Engine/Source/Runtime/SlateCore/Private/Debugging/ConsoleSlateDebuggerUpdate.cpp
+++ b/Engine/Source/Runtime/SlateCore/Private/Debugging/ConsoleSlateDebuggerUpdate.cpp
@@ -33,7 +33,7 @@ FConsoleSlateDebuggerUpdate::FConsoleSlateDebuggerUpdate()
 	, DrawWidgetNameColor(FColorList::Red)
 	, MaxNumberOfWidgetInList(20)
 	, InvalidationRootIdFilter(-1)
-	, CacheDuration(2.0)
+	, CacheDuration(2.0f)
 	, StartCommand(
 		TEXT("SlateDebugger.Update.Start"),
 		TEXT("Start the update widget debug tool. It shows when widgets are updated."),
@@ -389,7 +389,7 @@ void FConsoleSlateDebuggerUpdate::HandlePaintDebugInfo(const FPaintArgs& InArgs,
 	{
 		if (Itt.Value.WindowId == PaintWindow)
 		{
-			const double LerpValue = FMath::Clamp((SlateApplicationCurrentTime - Itt.Value.LastInvalidationTime) / CacheDuration, 0.0, 1.0);
+			const float LerpValue = FMath::Clamp((float)(SlateApplicationCurrentTime - Itt.Value.LastInvalidationTime) / CacheDuration, 0.f, 1.f);
 			const FGeometry Geometry = FGeometry::MakeRoot(Itt.Value.PaintSize, FSlateLayoutTransform(1.f, Itt.Value.PaintLocation));
 			const FPaintGeometry PaintGeometry = Geometry.ToPaintGeometry();
 			FLinearColor Color = DrawVolatilePaintColor;
