Slate: Fix some implicit conversion in SlateCore code.

diff --git a/Engine/Source/Runtime/SlateCore/Private/Animation/CurveSequence.cpp b/Engine/Source/Runtime/SlateCore/Private/Animation/CurveSequence.cpp
--- a/Engine/Source/Runtime/SlateCore/Private/Animation/CurveSequence.cpp
+++ b/Engine/Source/Runtime/SlateCore/Private/Animation/CurveSequence.cpp
@@ -6,21 +6,22 @@
 #include "Widgets/SWidget.h"
 #include "Stats/Stats.h"
 
-
 /* FCurveSequence constructors
  *****************************************************************************/
 
 FCurveSequence::FCurveSequence( )
-	: StartTime(0)
-	, TotalDuration(0)
+	: StartTime(0.0)
+	, PauseTime(0.0)
+	, TotalDuration(0.0)
 	, bInReverse(true)
 	, bIsLooping(false)
 	, bIsPaused(false)
 { }
 
 FCurveSequence::FCurveSequence( const float InStartTimeSeconds, const float InDurationSeconds, const ECurveEaseFunction InEaseFunction )
-	: StartTime(0)
-	, TotalDuration(0)
+	: StartTime(0.0)
+	, PauseTime(0.0)
+	, TotalDuration(0.0)
 	, bInReverse(true)
 	, bIsLooping( false )
 	, bIsPaused( false )
@@ -48,7 +49,7 @@ FCurveSequence::~FCurveSequence()
 FCurveHandle FCurveSequence::AddCurve( const float InStartTimeSeconds, const float InDurationSeconds, const ECurveEaseFunction InEaseFunction )
 {
 	// Keep track of how long this sequence is
-	TotalDuration = FMath::Max(TotalDuration, InStartTimeSeconds + InDurationSeconds);
+	TotalDuration = FMath::Max(TotalDuration, (double)(InStartTimeSeconds + InDurationSeconds));
 	// The initial state is to be at the end of the animation.
 	StartTime = TotalDuration;
 
@@ -60,7 +61,7 @@ FCurveHandle FCurveSequence::AddCurve( const float InStartTimeSeconds, const flo
 
 FCurveHandle FCurveSequence::AddCurveRelative( const float InOffset, const float InDurationSecond, const ECurveEaseFunction InEaseFunction )
 {
-	const float CurveStartTime = TotalDuration + InOffset;
+	const float CurveStartTime = (float)TotalDuration + InOffset;
 	return AddCurve(CurveStartTime, InDurationSecond, InEaseFunction);
 }
 
@@ -99,7 +100,7 @@ void FCurveSequence::Play(const FTickerDelegate& InDelegate, bool bPlayLooped, c
 void FCurveSequence::Reverse( )
 {	
 	// We've played this far into the animation.
-	const float FractionCompleted = FMath::Clamp(GetSequenceTime() / TotalDuration, 0.0f, 1.0f);
+	const double FractionCompleted = FMath::Clamp((double)GetSequenceTime() / TotalDuration, 0.0, 1.0);
 
 	// We're going the other way now.
 	bInReverse = !bInReverse;
@@ -170,9 +171,9 @@ void FCurveSequence::SetStartTime( double InStartTime )
 float FCurveSequence::GetSequenceTime( ) const
 {
 	const double CurrentTime = bIsPaused ? PauseTime : FSlateApplicationBase::Get().GetCurrentTime();
-	float SequenceTime = IsInReverse() ? TotalDuration - ( CurrentTime - StartTime ) : CurrentTime - StartTime;
+	const float SequenceTime = IsInReverse() ? (float)(TotalDuration - (CurrentTime - StartTime)) : (float)(CurrentTime - StartTime);
 
-	return bIsLooping ? FMath::Fmod( SequenceTime, TotalDuration ) : SequenceTime;
+	return bIsLooping ? FMath::Fmod( SequenceTime, (float)TotalDuration ) : SequenceTime;
 }
 
 
diff --git a/Engine/Source/Runtime/SlateCore/Public/Animation/CurveSequence.h b/Engine/Source/Runtime/SlateCore/Public/Animation/CurveSequence.h
--- a/Engine/Source/Runtime/SlateCore/Public/Animation/CurveSequence.h
+++ b/Engine/Source/Runtime/SlateCore/Public/Animation/CurveSequence.h
@@ -210,7 +210,7 @@ private:
 	double PauseTime;
 
 	/** How long the entire sequence lasts. */
-	float TotalDuration;
+	double TotalDuration;
 
 	/** Are we playing the animation in reverse */
 	uint8 bInReverse : 1;
