Fix border margins not rendering correctly with negative sizes.

diff --git a/Engine/Source/Runtime/SlateCore/Private/Rendering/ElementBatcher.cpp b/Engine/Source/Runtime/SlateCore/Private/Rendering/ElementBatcher.cpp
--- a/Engine/Source/Runtime/SlateCore/Private/Rendering/ElementBatcher.cpp
+++ b/Engine/Source/Runtime/SlateCore/Private/Rendering/ElementBatcher.cpp
@@ -2262,11 +2262,20 @@ void FSlateElementBatcher::AddBorderElement( const FSlateDrawElement& DrawElemen
 	const FVector2D StartUV = HalfTexel;
 	const FVector2D EndUV = FVector2D( 1.0f, 1.0f ) + HalfTexel;
 
-	const FMargin& Margin = DrawElementPayload.GetBrushMargin();
+	FMargin Margin = DrawElementPayload.GetBrushMargin();
 
 	// Do pixel snapping
 	FVector2D TopLeft(0,0);
 	FVector2D BotRight(LocalSize);
+
+	// Account for negative sizes
+	bool bIsFlippedX = TopLeft.X > BotRight.X;
+	bool bIsFlippedY = TopLeft.Y > BotRight.Y;
+	Margin.Left = FMath::Clamp(bIsFlippedX ? -Margin.Left : Margin.Left, -1.0f, 1.0f);
+	Margin.Top = FMath::Clamp(bIsFlippedY ? -Margin.Top : Margin.Top, -1.0f, 1.0f);
+	Margin.Right = FMath::Clamp(bIsFlippedX ? -Margin.Right : Margin.Right, -1.0f, 1.0f);
+	Margin.Bottom = FMath::Clamp(bIsFlippedY ? -Margin.Bottom : Margin.Bottom, -1.0f, 1.0f);
+
 	// Determine the margins for each quad
 	FVector2D TopLeftMargin(TextureSizeLocalSpace * FVector2D(Margin.Left, Margin.Top));
 	FVector2D BotRightMargin(LocalSize - TextureSizeLocalSpace * FVector2D(Margin.Right, Margin.Bottom));
@@ -2278,23 +2287,23 @@ void FSlateElementBatcher::AddBorderElement( const FSlateDrawElement& DrawElemen
 
 	// If the margins are overlapping the margins are too big or the button is too small
 	// so clamp margins to half of the box size
-	if( RightMarginX < LeftMarginX )
+	if (FMath::Abs(RightMarginX) < FMath::Abs(LeftMarginX))
 	{
 		LeftMarginX = LocalSize.X / 2;
 		RightMarginX = LeftMarginX;
 	}
 
-	if( BottomMarginY < TopMarginY )
+	if (FMath::Abs(BottomMarginY) < FMath::Abs(TopMarginY))
 	{
 		TopMarginY = LocalSize.Y / 2;
 		BottomMarginY = TopMarginY;
 	}
 
 	// Determine the texture coordinates for each quad
-	float LeftMarginU = (Margin.Left > 0.0f) ? Margin.Left : 0.0f;
-	float TopMarginV = (Margin.Top > 0.0f) ? Margin.Top : 0.0f;
-	float RightMarginU = (Margin.Right > 0.0f) ? 1.0f - Margin.Right : 1.0f;
-	float BottomMarginV = (Margin.Bottom > 0.0f) ? 1.0f - Margin.Bottom : 1.0f;
+	float LeftMarginU = FMath::Abs(Margin.Left);
+	float TopMarginV = FMath::Abs(Margin.Top);
+	float RightMarginU = 1.0f - FMath::Abs(Margin.Right);
+	float BottomMarginV = 1.0f - FMath::Abs(Margin.Bottom);
 
 	LeftMarginU += HalfTexel.X;
 	TopMarginV += HalfTexel.Y;
