Prevent virtual windows from creating any window borders

diff --git a/Engine/Source/Runtime/SlateCore/Private/Widgets/SWindow.cpp b/Engine/Source/Runtime/SlateCore/Private/Widgets/SWindow.cpp
--- a/Engine/Source/Runtime/SlateCore/Private/Widgets/SWindow.cpp
+++ b/Engine/Source/Runtime/SlateCore/Private/Widgets/SWindow.cpp
@@ -562,7 +562,7 @@ void SWindow::ConstructWindowInternals()
 	UpdateWindowContentVisibility();
 
 	// create window
-	if ((Type == EWindowType::Normal || Type == EWindowType::GameWindow) && !bHasOSWindowBorder)
+	if ((Type == EWindowType::Normal || Type == EWindowType::GameWindow) && !bHasOSWindowBorder && !bVirtualWindow)
 	{
 		// create window content slot
 		MainWindowArea->AddSlot()
