Slate: Fix some implicit type conversations. Cull the DrawElement is it's pixel size would be negative or 0.

diff --git a/Engine/Source/Runtime/SlateCore/Private/Rendering/DrawElements.cpp b/Engine/Source/Runtime/SlateCore/Private/Rendering/DrawElements.cpp
--- a/Engine/Source/Runtime/SlateCore/Private/Rendering/DrawElements.cpp
+++ b/Engine/Source/Runtime/SlateCore/Private/Rendering/DrawElements.cpp
@@ -11,6 +11,10 @@
 #include "Debugging/SlateDebugging.h"
 #include "Application/SlateApplicationBase.h"
 
+#if UE_SLATE_VERIFY_PIXELSIZE
+#include <limits>
+#endif
+
 DECLARE_CYCLE_STAT(TEXT("FSlateDrawElement::Make Time"), STAT_SlateDrawElementMakeTime, STATGROUP_SlateVerbose);
 DECLARE_CYCLE_STAT(TEXT("FSlateDrawElement::MakeCustomVerts Time"), STAT_SlateDrawElementMakeCustomVertsTime, STATGROUP_Slate);
 DECLARE_CYCLE_STAT(TEXT("FSlateDrawElement::Prebatch Time"), STAT_SlateDrawElementPrebatchTime, STATGROUP_Slate);
@@ -45,7 +49,10 @@ static bool ShouldCull(const FSlateWindowElementList& ElementList)
 static bool ShouldCull(const FSlateWindowElementList& ElementList, const FPaintGeometry& PaintGeometry)
 {
 	const FVector2D& LocalSize = PaintGeometry.GetLocalSize();
-	if (LocalSize.X == 0 || LocalSize.Y == 0)
+	const float DrawScale = PaintGeometry.DrawScale;
+
+	const FVector2D PixelSize = (LocalSize * DrawScale);
+	if (PixelSize.X <= 0.f || PixelSize.Y <= 0.f)
 	{
 		return true;
 	}
@@ -149,6 +156,12 @@ void FSlateDrawElement::Init(FSlateWindowElementList& ElementList, EElementType
 	LocalSize = PaintGeometry.GetLocalSize();
 	ClipStateHandle.SetPreCachedClipIndex(ElementList.GetClippingIndex());
 
+#if UE_SLATE_VERIFY_PIXELSIZE
+	const FVector2D PixelSize = (LocalSize * Scale);
+	ensureMsgf(PixelSize.X >= 0.f && PixelSize.X <= (float)std::numeric_limits<uint16>::max(), TEXT("The size X '%f' is too small or big to fit in the SlateVertex buffer."), PixelSize.X);
+	ensureMsgf(PixelSize.Y >= 0.f && PixelSize.Y <= (float)std::numeric_limits<uint16>::max(), TEXT("The size Y '%f' is too small or big to fit in the SlateVertex buffer."), PixelSize.Y);
+#endif
+
 	LayerId = InLayer;
 
 	ElementType = InElementType;
@@ -225,7 +238,7 @@ FSlateDrawElement& FSlateDrawElement::MakeBoxInternal(
 		float Radius = InBrush->OutlineSettings.Radius;
 		if (InBrush->OutlineSettings.RoundingType == ESlateBrushRoundingType::HalfHeightRadius)
 		{
-			Radius = PaintGeometry.GetLocalSize().Y/2.0;
+			Radius = PaintGeometry.GetLocalSize().Y/2.0f;
 		}
 		RBoxPayload->SetRadius(Radius);
 		RBoxPayload->SetOutline(InBrush->OutlineSettings.Color.GetSpecifiedColor(), InBrush->OutlineSettings.Width);
diff --git a/Engine/Source/Runtime/SlateCore/Public/Layout/LayoutUtils.h b/Engine/Source/Runtime/SlateCore/Public/Layout/LayoutUtils.h
--- a/Engine/Source/Runtime/SlateCore/Public/Layout/LayoutUtils.h
+++ b/Engine/Source/Runtime/SlateCore/Public/Layout/LayoutUtils.h
@@ -94,10 +94,10 @@ static AlignmentArrangeResult AlignChild(EFlowDirection InLayoutFlow, float Allo
 	switch (Alignment)
 	{
 	case HAlign_Fill:
-		return AlignmentArrangeResult(MarginPre, (AllottedSize - TotalMargin) * ContentScale);
+		return AlignmentArrangeResult(MarginPre, FMath::Max((AllottedSize - TotalMargin) * ContentScale, 0.f));
 	}
 	
-	const float ChildSize = bClampToParent ? FMath::Min(ChildDesiredSize, AllottedSize - TotalMargin) : ChildDesiredSize;
+	const float ChildSize = FMath::Max((bClampToParent ? FMath::Min(ChildDesiredSize, AllottedSize - TotalMargin) : ChildDesiredSize), 0.f);
 
 	switch( Alignment )
 	{
@@ -110,7 +110,7 @@ static AlignmentArrangeResult AlignChild(EFlowDirection InLayoutFlow, float Allo
 	}
 
 	// Same as Fill
-	return AlignmentArrangeResult(MarginPre, ( AllottedSize - TotalMargin ) * ContentScale);
+	return AlignmentArrangeResult(MarginPre, FMath::Max(( AllottedSize - TotalMargin ) * ContentScale, 0.f));
 }
 
 template<EOrientation Orientation, typename SlotType>
@@ -132,14 +132,14 @@ static AlignmentArrangeResult AlignChild(EFlowDirection InLayoutFlow, float Allo
 	switch (Alignment)
 	{
 	case HAlign_Fill:
-		return AlignmentArrangeResult(MarginPre, (AllottedSize - TotalMargin) * ContentScale);
+		return AlignmentArrangeResult(MarginPre, FMath::Max((AllottedSize - TotalMargin) * ContentScale, 0.f));
 	}
 
 	const float ChildDesiredSize = ( Orientation == Orient_Horizontal )
 		? ( ChildToArrange.GetWidget()->GetDesiredSize().X * ContentScale )
 		: ( ChildToArrange.GetWidget()->GetDesiredSize().Y * ContentScale );
 
-	const float ChildSize = bClampToParent ? FMath::Min(ChildDesiredSize, AllottedSize - TotalMargin) : ChildDesiredSize;
+	const float ChildSize = FMath::Max((bClampToParent ? FMath::Min(ChildDesiredSize, AllottedSize - TotalMargin) : ChildDesiredSize), 0.f);
 
 	switch ( Alignment )
 	{
@@ -152,7 +152,7 @@ static AlignmentArrangeResult AlignChild(EFlowDirection InLayoutFlow, float Allo
 	}
 
 	// Same as Fill
-	return AlignmentArrangeResult(MarginPre, (AllottedSize - TotalMargin) * ContentScale);
+	return AlignmentArrangeResult(MarginPre, FMath::Max((AllottedSize - TotalMargin) * ContentScale, 0.f));
 }
 
 template<EOrientation Orientation, typename SlotType>
diff --git a/Engine/Source/Runtime/SlateCore/Public/Rendering/RenderingCommon.h b/Engine/Source/Runtime/SlateCore/Public/Rendering/RenderingCommon.h
--- a/Engine/Source/Runtime/SlateCore/Public/Rendering/RenderingCommon.h
+++ b/Engine/Source/Runtime/SlateCore/Public/Rendering/RenderingCommon.h
@@ -26,6 +26,8 @@ DECLARE_DWORD_COUNTER_STAT_EXTERN(TEXT("Num Cached Elements"), STAT_SlateNumCach
 
 DECLARE_CYCLE_STAT_EXTERN(TEXT("PreFill Buffers RT"), STAT_SlatePreFullBufferRTTime, STATGROUP_Slate, SLATECORE_API);
 
+#define UE_SLATE_VERIFY_PIXELSIZE UE_BUILD_DEBUG
+
 #define SLATE_USE_32BIT_INDICES !PLATFORM_USES_GLES
 
 #if SLATE_USE_32BIT_INDICES
@@ -264,8 +266,18 @@ public:
 		Vertex.TexCoords[3] = InTexCoords.W;
 		Vertex.MaterialTexCoords = FVector2D(InLocalPosition.X / InLocalSize.X, InLocalPosition.Y / InLocalSize.Y);
 		Vertex.InitCommon<Rounding>(RenderTransform, InLocalPosition, InColor);
-		Vertex.PixelSize[0] = FMath::RoundToInt(InLocalSize.X * Scale);
-		Vertex.PixelSize[1] = FMath::RoundToInt(InLocalSize.Y * Scale);
+
+		const int32 PixelSizeX = FMath::RoundToInt(InLocalSize.X * Scale);
+		const int32 PixelSizeY = FMath::RoundToInt(InLocalSize.Y * Scale);
+		Vertex.PixelSize[0] = (uint16)PixelSizeX;
+		Vertex.PixelSize[1] = (uint16)PixelSizeY;
+
+#if UE_SLATE_VERIFY_PIXELSIZE
+		ensureMsgf((int32)Vertex.PixelSize[0] == PixelSizeX, TEXT("Conversion of PixelSizeX is bigger than 16. Cast:%d, int16:%d, int32:%d")
+			, (int32)Vertex.PixelSize[0], Vertex.PixelSize[0], PixelSizeX);
+		ensureMsgf((int32)Vertex.PixelSize[1] == PixelSizeY, TEXT("Conversion of PixelSizeY is bigger than 16. Cast:%d, int16:%d, int32:%d")
+			, (int32)Vertex.PixelSize[1], Vertex.PixelSize[1], PixelSizeY);
+#endif
 		Vertex.MaterialTexCoords = FVector2D(InLocalPosition.X / InLocalSize.X, InLocalPosition.Y / InLocalSize.Y);
 
 		return Vertex;
@@ -293,8 +305,8 @@ private:
 
 		if ( Rounding == ESlateVertexRounding::Enabled )
 		{
-			Position.X = FMath::RoundToInt(Position.X);
-			Position.Y = FMath::RoundToInt(Position.Y);
+			Position.X = FMath::RoundToFloat(Position.X);
+			Position.Y = FMath::RoundToFloat(Position.Y);
 		}
 
 		Color = InColor;
