Slate: Fix some implicit conversion type.

diff --git a/Engine/Source/Runtime/SlateCore/Private/Brushes/SlateDynamicImageBrush.cpp b/Engine/Source/Runtime/SlateCore/Private/Brushes/SlateDynamicImageBrush.cpp
--- a/Engine/Source/Runtime/SlateCore/Private/Brushes/SlateDynamicImageBrush.cpp
+++ b/Engine/Source/Runtime/SlateCore/Private/Brushes/SlateDynamicImageBrush.cpp
@@ -18,7 +18,8 @@ TSharedPtr<FSlateDynamicImageBrush> FSlateDynamicImageBrush::CreateWithImageData
 {
 	TSharedPtr<FSlateDynamicImageBrush> Brush;
 	if (FSlateApplicationBase::IsInitialized() &&
-		FSlateApplicationBase::Get().GetRenderer()->GenerateDynamicImageResource(InTextureName, InImageSize.X, InImageSize.Y, InImageData))
+		InImageSize.X > 0.f && InImageSize.Y > 0.f &&
+		FSlateApplicationBase::Get().GetRenderer()->GenerateDynamicImageResource(InTextureName, (uint32)InImageSize.X, (uint32)InImageSize.Y, InImageData))
 	{
 		Brush = MakeShareable(new FSlateDynamicImageBrush(
 			InTextureName,
diff --git a/Engine/Source/Runtime/SlateCore/Private/Textures/TextureAtlas.cpp b/Engine/Source/Runtime/SlateCore/Private/Textures/TextureAtlas.cpp
--- a/Engine/Source/Runtime/SlateCore/Private/Textures/TextureAtlas.cpp
+++ b/Engine/Source/Runtime/SlateCore/Private/Textures/TextureAtlas.cpp
@@ -7,10 +7,36 @@
 #include "HAL/LowLevelMemTracker.h"
 #include "Misc/MemStack.h"
 
+#include <limits>
+
 DEFINE_STAT(STAT_SlateTextureGPUMemory);
 DEFINE_STAT(STAT_SlateTextureDataMemory);
 DECLARE_MEMORY_STAT(TEXT("Texture Atlas Memory (CPU)"), STAT_SlateTextureAtlasMemory, STATGROUP_SlateMemory);
 
+static int32 GInitialMaxAtlasPagesBeforeFlushRequest = 1;
+FAutoConsoleVariableRef CVarMaxAtlasPagesBeforeFlush(
+	TEXT("Slate.MaxAtlasPagesBeforeFlush"),
+	GInitialMaxAtlasPagesBeforeFlushRequest,
+	TEXT("The number of atlas textures created and used before we flush the cache if a texture atlas is full"));
+
+static int32 GInitialMaxNonAtlasedTexturesBeforeFlushRequest = 1;
+FAutoConsoleVariableRef CVarMaxFontNonAtlasPagesBeforeFlush(
+	TEXT("Slate.MaxNonAtlasTexturesBeforeFlush"),
+	GInitialMaxNonAtlasedTexturesBeforeFlushRequest,
+	TEXT("The number of large textures initially."));
+
+static int32 GGrowAtlasFrameWindow = 1;
+FAutoConsoleVariableRef CVarGrowFontAtlasFrameWindow(
+	TEXT("Slate.GrowAtlasFrameWindow"),
+	GGrowAtlasFrameWindow,
+	TEXT("The number of frames within the atlas will resize rather than flush."));
+
+static int32 GGrowNonAtlasFrameWindow = 1;
+FAutoConsoleVariableRef CVarGrowFontNonAtlasFrameWindow(
+	TEXT("Slate.GrowNonAtlasFrameWindow"),
+	GGrowNonAtlasFrameWindow,
+	TEXT("The number of frames within the large pool will resize rather than flush."));
+
 ESlateTextureAtlasThreadId GetCurrentSlateTextureAtlasThreadId()
 {
 	// Note: For Game thread ownership, there is a point at which multiple worker threads operate on text simultaneously while the game thread is blocked
@@ -70,12 +96,12 @@ void FSlateTextureAtlas::EmptyAtlasData()
 	DeleteSlots.Empty();
 
 
-	STAT(uint32 MemoryBefore = AtlasData.GetAllocatedSize());
+	STAT(SIZE_T MemoryBefore = AtlasData.GetAllocatedSize());
 
 	// Clear all raw data
 	AtlasData.Empty();
 
-	STAT(uint32 MemoryAfter = AtlasData.GetAllocatedSize());
+	STAT(SIZE_T MemoryAfter = AtlasData.GetAllocatedSize());
 	DEC_MEMORY_STAT_BY(STAT_SlateTextureAtlasMemory, MemoryBefore-MemoryAfter);
 }
 
@@ -253,13 +279,15 @@ const FAtlasedTextureSlot* FSlateTextureAtlas::GetSlotAtPosition(FIntPoint InPos
 	{
 		FAtlasedTextureSlot& CurSlot = *SlotIt;
 
-		FSlateRect CurSlotRect(FVector2D(CurSlot.X, CurSlot.Y), FVector2D(CurSlot.X + CurSlot.Width, CurSlot.Y + CurSlot.Height));
-
-		if (CurSlotRect.ContainsPoint(InPosition))
+		checkSlow(CurSlot.X <= (uint32)std::numeric_limits<int32>()::max());
+		checkSlow(CurSlot.Y <= (uint32)std::numeric_limits<int32>()::max());
+		checkSlow(CurSlot.X + CurSlot.Width <= (uint32)std::numeric_limits<int32>()::max());
+		checkSlow(CurSlot.Y + CurSlot.Height <= (uint32)std::numeric_limits<int32>()::max());
+		FIntRect CurSlotRect(FIntPoint((int32)CurSlot.X, (int32)CurSlot.Y), FIntPoint((int32)(CurSlot.X + CurSlot.Width), (int32)(CurSlot.Y + CurSlot.Height)));
+		if (CurSlotRect.Contains(InPosition))
 		{
 			return &CurSlot;
 		}
-
 	}
 
 	return nullptr;
@@ -271,7 +299,7 @@ const FAtlasedTextureSlot* FSlateTextureAtlas::FindSlotForTexture(uint32 InWidth
 	FAtlasedTextureSlot* ReturnVal = NULL;
 
 	// Account for padding on both sides
-	const uint32 Padding = GetPaddingAmount();
+	const uint8 Padding = GetPaddingAmount();
 	const uint32 TotalPadding = Padding * 2;
 	const uint32 PaddedWidth = InWidth + TotalPadding;
 	const uint32 PaddedHeight = InHeight + TotalPadding;
@@ -360,3 +388,75 @@ const FAtlasedTextureSlot* FSlateTextureAtlas::FindSlotForTexture(uint32 InWidth
 	return ReturnVal;
 }
 
+
+
+FSlateFlushableAtlasCache::FSlateFlushableAtlasCache()
+	: InitialMaxAtlasPagesBeforeFlushRequest(GInitialMaxAtlasPagesBeforeFlushRequest)
+	, InitialMaxNonAtlasPagesBeforeFlushRequest(GInitialMaxNonAtlasedTexturesBeforeFlushRequest)
+	, CurrentMaxGrayscaleAtlasPagesBeforeFlushRequest(GInitialMaxAtlasPagesBeforeFlushRequest)
+	, CurrentMaxColorAtlasPagesBeforeFlushRequest(GInitialMaxAtlasPagesBeforeFlushRequest)
+	, CurrentMaxNonAtlasedTexturesBeforeFlushRequest(GInitialMaxNonAtlasedTexturesBeforeFlushRequest)
+{
+}
+
+FSlateFlushableAtlasCache::FSlateFlushableAtlasCache(int32 InInitialMaxAtlasPagesBeforeFlushRequest, int32 InInitialMaxNonAtlasedTexturesBeforeFlushRequest)
+	: InitialMaxAtlasPagesBeforeFlushRequest(InInitialMaxAtlasPagesBeforeFlushRequest)
+	, InitialMaxNonAtlasPagesBeforeFlushRequest(InInitialMaxNonAtlasedTexturesBeforeFlushRequest)
+	, CurrentMaxGrayscaleAtlasPagesBeforeFlushRequest(InInitialMaxAtlasPagesBeforeFlushRequest)
+	, CurrentMaxColorAtlasPagesBeforeFlushRequest(InInitialMaxAtlasPagesBeforeFlushRequest)
+	, CurrentMaxNonAtlasedTexturesBeforeFlushRequest(InInitialMaxNonAtlasedTexturesBeforeFlushRequest)
+{
+}
+
+void FSlateFlushableAtlasCache::ResetFlushCounters()
+{
+	CurrentMaxGrayscaleAtlasPagesBeforeFlushRequest = InitialMaxAtlasPagesBeforeFlushRequest;
+	CurrentMaxColorAtlasPagesBeforeFlushRequest = InitialMaxAtlasPagesBeforeFlushRequest;
+	CurrentMaxNonAtlasedTexturesBeforeFlushRequest = InitialMaxNonAtlasPagesBeforeFlushRequest;
+	FrameCounterLastFlushRequest = GFrameCounter;
+}
+
+void FSlateFlushableAtlasCache::UpdateFlushCounters(int32 NumGrayscale, int32 NumColor, int32 NumNonAtlased)
+{
+	bool bFlushRequested = false;
+	bFlushRequested |= UpdateInternal(NumGrayscale, CurrentMaxGrayscaleAtlasPagesBeforeFlushRequest, InitialMaxAtlasPagesBeforeFlushRequest, GGrowAtlasFrameWindow);
+	bFlushRequested |= UpdateInternal(NumColor, CurrentMaxColorAtlasPagesBeforeFlushRequest, InitialMaxAtlasPagesBeforeFlushRequest, GGrowAtlasFrameWindow);
+	bFlushRequested |= UpdateInternal(NumNonAtlased, CurrentMaxNonAtlasedTexturesBeforeFlushRequest, InitialMaxNonAtlasPagesBeforeFlushRequest, GGrowNonAtlasFrameWindow);
+
+	if (bFlushRequested)
+	{
+		ResetFlushCounters();
+	}
+}
+
+bool FSlateFlushableAtlasCache::UpdateInternal(int32 CurrentNum, int32& MaxNum, int32 InitialMax, int32 FrameWindowNum)
+{
+	const uint64 LastFlushRequestFrameDelta = GFrameCounter - FrameCounterLastFlushRequest;
+
+	if (CurrentNum > MaxNum)
+	{
+		// The InitialMaxNonAtlasedTexturesBeforeFlushRequest may have changed since last flush,
+		// so double check that we're below the current max or initial, if we're under it,
+		// update the current to the initial.
+		if (CurrentNum <= InitialMax)
+		{
+			MaxNum = InitialMax;
+		}
+		// If we grew back up to this number of non-atlased textures within the same or next frame of the previous flush request, then we likely legitimately have 
+		// a lot of data cached. We should update CurrentMaxNonAtlasedTexturesBeforeFlushRequest to give us a bit more flexibility before the next flush request
+		else if (LastFlushRequestFrameDelta <= FrameWindowNum)
+		{
+			MaxNum = CurrentNum;
+			UE_LOG(LogSlate, Warning, TEXT("Setting the threshold to trigger a flush to %d non-atlased textures as there is a lot of font data being cached."), CurrentNum);
+		}
+		else
+		{
+			// We've grown beyond our current stable limit - try and request a flush
+			RequestFlushCache(FString::Printf(TEXT("Large atlases out of space; %d/%d Textures; frames since last flush: %llu"), CurrentNum, MaxNum, LastFlushRequestFrameDelta));
+			
+			return true;
+		}
+	}
+
+	return false;
+}
diff --git a/Engine/Source/Runtime/SlateCore/Private/Widgets/Accessibility/SlateCoreAccessibleWidgets.cpp b/Engine/Source/Runtime/SlateCore/Private/Widgets/Accessibility/SlateCoreAccessibleWidgets.cpp
--- a/Engine/Source/Runtime/SlateCore/Private/Widgets/Accessibility/SlateCoreAccessibleWidgets.cpp
+++ b/Engine/Source/Runtime/SlateCore/Private/Widgets/Accessibility/SlateCoreAccessibleWidgets.cpp
@@ -270,7 +270,7 @@ TSharedPtr<IAccessibleWidget> FSlateAccessibleWindow::GetChildAtPosition(int32 X
 		if (UseHitTestGrid)
 		{
 			TSharedPtr<SWindow> SlateWindow = StaticCastSharedPtr<SWindow>(Widget.Pin());
-			TArray<FWidgetAndPointer> Hits = SlateWindow->GetHittestGrid().GetBubblePath(FVector2D(X, Y), 0.0f, false, INDEX_NONE);
+			TArray<FWidgetAndPointer> Hits = SlateWindow->GetHittestGrid().GetBubblePath(FVector2D((float)X, (float)Y), 0.0f, false, INDEX_NONE);
 			TSharedPtr<SWidget> LastAccessibleWidget = nullptr;
 			for (int32 i = 0; i < Hits.Num(); ++i)
 			{
@@ -299,7 +299,7 @@ TSharedPtr<IAccessibleWidget> FSlateAccessibleWindow::GetChildAtPosition(int32 X
 				const TSharedPtr<IAccessibleWidget> Current = ToProcess.Pop(false);
 				// Because children are weak pointers, Current could be invalid in the case where the SWidget and all
 				// shared pointers were deleted while in the middle of FSlateAccessibleMessageHandler refreshing the data.
-				if (Current.IsValid() && !Current->IsHidden() && Current->GetBounds().IsInside(FVector2D(X, Y)))
+				if (Current.IsValid() && !Current->IsHidden() && Current->GetBounds().IsInside(FVector2D((float)X, (float)Y)))
 				{
 					// The widgets are being traversed in reverse render order, so usually if a widget is rendered
 					// on top of another this will return the rendered one. But it's not 100% guarantee, and opacity
diff --git a/Engine/Source/Runtime/SlateCore/Private/Widgets/SWidget.cpp b/Engine/Source/Runtime/SlateCore/Private/Widgets/SWidget.cpp
--- a/Engine/Source/Runtime/SlateCore/Private/Widgets/SWidget.cpp
+++ b/Engine/Source/Runtime/SlateCore/Private/Widgets/SWidget.cpp
@@ -25,6 +25,8 @@
 #include "HAL/CriticalSection.h"
 #include "Widgets/SWidgetUtils.h"
 
+#include <limits>
+
 #if WITH_ACCESSIBILITY
 #include "Widgets/Accessibility/SlateCoreAccessibleWidgets.h"
 #include "Widgets/Accessibility/SlateAccessibleMessageHandler.h"
diff --git a/Engine/Source/Runtime/SlateCore/Private/Widgets/SWindow.cpp b/Engine/Source/Runtime/SlateCore/Private/Widgets/SWindow.cpp
--- a/Engine/Source/Runtime/SlateCore/Private/Widgets/SWindow.cpp
+++ b/Engine/Source/Runtime/SlateCore/Private/Widgets/SWindow.cpp
@@ -320,11 +320,11 @@ void SWindow::Construct(const FArguments& InArgs)
 	{
 		if ( !SizeLimits.GetMaxWidth().IsSet() )
 		{
-			SizeLimits.SetMaxWidth(PrimaryDisplayRect.Right - PrimaryDisplayRect.Left);
+			SizeLimits.SetMaxWidth(static_cast<float>(PrimaryDisplayRect.Right - PrimaryDisplayRect.Left));
 		}
 		if ( !SizeLimits.GetMaxHeight().IsSet() )
 		{
-			SizeLimits.SetMaxHeight(PrimaryDisplayRect.Bottom - PrimaryDisplayRect.Top);
+			SizeLimits.SetMaxHeight(static_cast<float>(PrimaryDisplayRect.Bottom - PrimaryDisplayRect.Top));
 		}
 	}
 
@@ -383,13 +383,13 @@ void SWindow::Construct(const FArguments& InArgs)
 	{
 		if (InArgs._SaneWindowPlacement)
 		{
-			float PrimaryWidthPadding = DisplayMetrics.PrimaryDisplayWidth -
-				(PrimaryDisplayRect.Right - PrimaryDisplayRect.Left);
-			float PrimaryHeightPadding = DisplayMetrics.PrimaryDisplayHeight -
-				(PrimaryDisplayRect.Bottom - PrimaryDisplayRect.Top);
+			float PrimaryWidthPadding = static_cast<float>(DisplayMetrics.PrimaryDisplayWidth -
+				(PrimaryDisplayRect.Right - PrimaryDisplayRect.Left));
+			float PrimaryHeightPadding = static_cast<float>(DisplayMetrics.PrimaryDisplayHeight -
+				(PrimaryDisplayRect.Bottom - PrimaryDisplayRect.Top));
 
-			float VirtualWidth = (VirtualDisplayRect.Right - VirtualDisplayRect.Left);
-			float VirtualHeight = (VirtualDisplayRect.Bottom - VirtualDisplayRect.Top);
+			float VirtualWidth = static_cast<float>(VirtualDisplayRect.Right - VirtualDisplayRect.Left);
+			float VirtualHeight = static_cast<float>(VirtualDisplayRect.Bottom - VirtualDisplayRect.Top);
 
 			// Make sure that the window size is no larger than the virtual display area.
 			WindowSize.X = FMath::Clamp(WindowSize.X, 0.0f, VirtualWidth - PrimaryWidthPadding);
@@ -829,14 +829,14 @@ FVector2D SWindow::GetSizeInScreen() const
 
 FSlateRect SWindow::GetNonMaximizedRectInScreen() const
 {
-	int X = 0;
-	int Y = 0;
-	int Width = 0;
-	int Height = 0;
+	int32 X = 0;
+	int32 Y = 0;
+	int32 Width = 0;
+	int32 Height = 0;
 
 	if ( NativeWindow.IsValid() && NativeWindow->GetRestoredDimensions(X, Y, Width, Height) )
 	{
-		return FSlateRect( X, Y, X+Width, Y+Height );
+		return FSlateRect( (float)X, (float)Y, static_cast<float>(X+Width), static_cast<float>(Y+Height) );
 	}
 	else
 	{
@@ -926,7 +926,7 @@ void SWindow::MoveWindowTo( FVector2D NewPosition )
 		// This expectation is generally invalid (see UE-1308) as there may be a delay before the OS reports it back.
 		// This hack sets the position speculatively, keeping Slate happy while also giving the OS chance to report it
 		// correctly after or even during the actual call.
-		FVector2D SpeculativeScreenPosition(FMath::TruncToInt(NewPosition.X), FMath::TruncToInt(NewPosition.Y));
+		FVector2D SpeculativeScreenPosition(FMath::TruncToFloat(NewPosition.X), FMath::TruncToFloat(NewPosition.Y));
 		SetCachedScreenPosition(SpeculativeScreenPosition);
 #endif // PLATFORM_LINUX
 
@@ -1017,7 +1017,7 @@ FSlateRect SWindow::GetFullScreenInfo() const
 
 		if ( NativeWindow->GetFullScreenInfo( X, Y, Width, Height ) )
 		{
-			return FSlateRect( X, Y, X + Width, Y + Height );
+			return FSlateRect( (float)X, (float)Y, (float)(X + Width), (float)(Y + Height) );
 		}
 	}
 
@@ -2196,7 +2196,7 @@ void SWindow::SetWindowMode( EWindowMode::Type NewWindowMode )
 		NativeWindow->SetWindowMode( NewWindowMode );
 
 		const FVector2D vp = IsMirrorWindow() ? GetSizeInScreen() : GetViewportSize();
-		FSlateApplicationBase::Get().GetRenderer()->UpdateFullscreenState(SharedThis(this), vp.X, vp.Y);
+		FSlateApplicationBase::Get().GetRenderer()->UpdateFullscreenState(SharedThis(this), (uint32)vp.X, (uint32)vp.Y);
 
 		if( TitleArea.IsValid() )
 		{
diff --git a/Engine/Source/Runtime/SlateCore/Public/Styling/SlateTypes.h b/Engine/Source/Runtime/SlateCore/Public/Styling/SlateTypes.h
--- a/Engine/Source/Runtime/SlateCore/Public/Styling/SlateTypes.h
+++ b/Engine/Source/Runtime/SlateCore/Public/Styling/SlateTypes.h
@@ -1947,8 +1947,8 @@ struct SLATECORE_API FWindowStyle : public FSlateWidgetStyle
 
 	/** Window corner rounding.  If this value is <= 0 no rounding will occur.   Used for regular, non-maximized windows only (not tool-tips or decorators.) */
 	UPROPERTY(EditAnywhere, BlueprintReadWrite, Category=Appearance)
-	float WindowCornerRadius;
-	FWindowStyle& SetCornerRadius(float InCornerRadius) { WindowCornerRadius = InCornerRadius; return *this; }
+	int32 WindowCornerRadius;
+	FWindowStyle& SetCornerRadius(int32 InCornerRadius) { WindowCornerRadius = InCornerRadius; return *this; }
 
 	/** Window corner rounding.  If this value is <= 0 no rounding will occur.   Used for regular, non-maximized windows only (not tool-tips or decorators.) */
 	UPROPERTY(EditAnywhere, BlueprintReadWrite, Category=Appearance)
diff --git a/Engine/Source/Runtime/SlateCore/Public/Textures/TextureAtlas.h b/Engine/Source/Runtime/SlateCore/Public/Textures/TextureAtlas.h
--- a/Engine/Source/Runtime/SlateCore/Public/Textures/TextureAtlas.h
+++ b/Engine/Source/Runtime/SlateCore/Public/Textures/TextureAtlas.h
@@ -8,6 +8,8 @@
 
 #define WITH_ATLAS_DEBUGGING (WITH_EDITOR || IS_PROGRAM) && !UE_BUILD_SHIPPING
 
+class FSlateShaderResource;
+
 /** 
  * Specifies how to handle texture atlas padding (when specified for the atlas). 
  * We only support one pixel of padding because we don't support mips or aniso filtering on atlas textures right now.
@@ -123,7 +125,14 @@ public:
 	 * Updates the texture used for rendering if needed
 	 */
 	virtual void ConditionalUpdateTexture() = 0;
-	
+
+	/**
+	 * Releases rendering resources of this texture
+	 */
+	virtual void ReleaseResources() = 0;
+
+	virtual FSlateShaderResource* GetAtlasTexture() const = 0;
+
 #if WITH_ATLAS_DEBUGGING
 	const FAtlasedTextureSlot* GetSlotAtPosition(FIntPoint InPosition) const;
 #endif
@@ -185,7 +194,7 @@ protected:
 
 private:
 	/** Returns the amount of padding needed for the current padding style */
-	FORCEINLINE int32 GetPaddingAmount() const
+	FORCEINLINE uint8 GetPaddingAmount() const
 	{
 		return (PaddingStyle == ESlateTextureAtlasPaddingStyle::NoPadding) ? 0 : 1;
 	}
@@ -234,6 +243,61 @@ struct FAtlasSlotInfo
 	FName TextureName;
 };
 
+
+/** A factory capable of generating a texture atlas or shader resource for textures too big to be in an atlas */
+class ISlateTextureAtlasFactory
+{
+public:
+	virtual ~ISlateTextureAtlasFactory() {}
+	virtual TUniquePtr<FSlateTextureAtlas> CreateTextureAtlas(int32 AtlasSize, int32 AtlasStride, ESlateTextureAtlasPaddingStyle PaddingStyle, bool bUpdatesAfterInitialization) const = 0;
+
+	virtual TUniquePtr<FSlateShaderResource> CreateNonAtlasedTexture(const uint32 InWidth, const uint32 InHeight, const TArray<uint8>& InRawData) const = 0;
+
+	virtual void ReleaseTextureAtlases(const TArray<TUniquePtr<FSlateTextureAtlas>>& InTextureAtlases, const TArray<TUniquePtr<FSlateShaderResource>>& InNonAtlasedTextures, const bool bWaitForRelease) const = 0;
+};
+
+
+/** Base class for any atlas cache which has flushing logic to keep the number of in use pages small */
+class FSlateFlushableAtlasCache
+{
+public:
+	FSlateFlushableAtlasCache();
+	FSlateFlushableAtlasCache(int32 InInitialMaxAtlasPagesBeforeFlushRequest, int32 InInitialMaxNonAtlasedTexturesBeforeFlushRequest);
+
+	virtual ~FSlateFlushableAtlasCache() {}
+
+	/** 
+	 * Called when this cache must be flushed 
+	 * 
+ 	 * @param Reason A string explaining the reason the cache was flushed (generally for debugging or logging
+	 */
+	virtual void RequestFlushCache(const FString& Reason) = 0;
+
+	/** Resets all counters to their initial state to start over flushing logic */
+	void ResetFlushCounters();
+
+	/** Increments counters that determine if a flush is needed.  If a flush is needed RequestFlushCache will be called from here */
+	void UpdateFlushCounters(int32 NumGrayscale, int32 NumColor, int32 NumNonAtlased);
+
+private:
+	bool UpdateInternal(int32 CurrentNum, int32& MaxNum, int32 InitialMax, int32 FrameWindowNum);
+private:
+	int32 InitialMaxAtlasPagesBeforeFlushRequest;
+	int32 InitialMaxNonAtlasPagesBeforeFlushRequest;
+
+	/** Number of grayscale atlas pages we can have before we request that the cache be flushed */
+	int32 CurrentMaxGrayscaleAtlasPagesBeforeFlushRequest;
+
+	/** Number of color atlas pages we can have before we request that the cache be flushed */
+	int32 CurrentMaxColorAtlasPagesBeforeFlushRequest;
+
+	/** Number of non-atlased textures we can have before we request that the cache be flushed */
+	int32 CurrentMaxNonAtlasedTexturesBeforeFlushRequest;
+
+	/** The frame counter the last time the font cache was asked to be flushed */
+	uint64 FrameCounterLastFlushRequest;
+};
+
 /**
  * Interface to allow the Slate atlas visualizer to query atlas page information for an atlas provider
  */
@@ -257,3 +321,5 @@ public:
 	virtual FAtlasSlotInfo GetAtlasSlotInfoAtPosition(FIntPoint InPosition, int32 AtlasIndex) const = 0;
 #endif
 };
+
+
