Slate: Fix CIS in linux build.
#rnx

diff --git a/Engine/Source/Runtime/SlateCore/Private/Textures/TextureAtlas.cpp b/Engine/Source/Runtime/SlateCore/Private/Textures/TextureAtlas.cpp
--- a/Engine/Source/Runtime/SlateCore/Private/Textures/TextureAtlas.cpp
+++ b/Engine/Source/Runtime/SlateCore/Private/Textures/TextureAtlas.cpp
@@ -279,10 +279,10 @@ const FAtlasedTextureSlot* FSlateTextureAtlas::GetSlotAtPosition(FIntPoint InPos
 	{
 		FAtlasedTextureSlot& CurSlot = *SlotIt;
 
-		checkSlow(CurSlot.X <= (uint32)std::numeric_limits<int32>()::max());
-		checkSlow(CurSlot.Y <= (uint32)std::numeric_limits<int32>()::max());
-		checkSlow(CurSlot.X + CurSlot.Width <= (uint32)std::numeric_limits<int32>()::max());
-		checkSlow(CurSlot.Y + CurSlot.Height <= (uint32)std::numeric_limits<int32>()::max());
+		checkSlow(CurSlot.X <= (uint32)std::numeric_limits<int32>::max());
+		checkSlow(CurSlot.Y <= (uint32)std::numeric_limits<int32>::max());
+		checkSlow(CurSlot.X + CurSlot.Width <= (uint32)std::numeric_limits<int32>::max());
+		checkSlow(CurSlot.Y + CurSlot.Height <= (uint32)std::numeric_limits<int32>::max());
 		FIntRect CurSlotRect(FIntPoint((int32)CurSlot.X, (int32)CurSlot.Y), FIntPoint((int32)(CurSlot.X + CurSlot.Width), (int32)(CurSlot.Y + CurSlot.Height)));
 		if (CurSlotRect.Contains(InPosition))
 		{
