Slate: Fix bad test when adding a hittest grid.

diff --git a/Engine/Source/Runtime/SlateCore/Private/Input/HittestGrid.cpp b/Engine/Source/Runtime/SlateCore/Private/Input/HittestGrid.cpp
--- a/Engine/Source/Runtime/SlateCore/Private/Input/HittestGrid.cpp
+++ b/Engine/Source/Runtime/SlateCore/Private/Input/HittestGrid.cpp
@@ -587,8 +587,9 @@ bool FHittestGrid::IsValidCellCoord(const int32 XCoord, const int32 YCoord) cons
 void FHittestGrid::AddGrid(const TSharedRef<const FHittestGrid>& OtherGrid)
 {
 	const bool bIsContains = AppendedGridArray.ContainsByPredicate([OtherGrid](const FAppendedGridData& Other) { return Other.Grid == OtherGrid; });
-	if (!bIsContains && ensure(CanBeAppended(&OtherGrid.Get())))
+	if (ensure(CanBeAppended(&OtherGrid.Get())))
 	{
+		if (!bIsContains)
 		{
 			// Check for recursion
 			FCollapsedHittestGridArray AllHittestGrid;
@@ -599,9 +600,9 @@ void FHittestGrid::AddGrid(const TSharedRef<const FHittestGrid>& OtherGrid)
 			{
 				return;
 			}
-		  }
 
-		AppendedGridArray.Emplace(OtherGrid->Owner, OtherGrid);
+			AppendedGridArray.Emplace(OtherGrid->Owner, OtherGrid);
+		}
 	}
 	else
 	{
