Creates Secondary Button Style

diff --git a/Engine/Source/Runtime/SlateCore/Private/Styling/StarshipCoreStyle.cpp b/Engine/Source/Runtime/SlateCore/Private/Styling/StarshipCoreStyle.cpp
--- a/Engine/Source/Runtime/SlateCore/Private/Styling/StarshipCoreStyle.cpp
+++ b/Engine/Source/Runtime/SlateCore/Private/Styling/StarshipCoreStyle.cpp
@@ -1067,6 +1067,19 @@ void FStarshipCoreStyle::SetupButtonStyles(TSharedRef<FStyle>& Style)
 		.SetDisabledForeground(FStyleColors::Foreground)
 		.SetNormalPadding(FMargin(2, 2, 2, 2))
 		.SetPressedPadding(FMargin(2, 3, 2, 1));
+
+
+	const FButtonStyle SecondaryButton = FButtonStyle()
+		.SetNormal(FSlateRoundedBoxBrush(FStyleColors::Secondary, 4.0f))
+		.SetHovered(FSlateRoundedBoxBrush(FStyleColors::Hover, 4.0f))
+		.SetPressed(FSlateRoundedBoxBrush(FStyleColors::Header, 4.0f))
+		.SetDisabled(FSlateNoResource())
+		.SetNormalForeground(FStyleColors::ForegroundHover)
+		.SetHoveredForeground(FStyleColors::ForegroundHover)
+		.SetPressedForeground(FStyleColors::ForegroundHover)
+		.SetDisabledForeground(FStyleColors::Foreground)
+		.SetNormalPadding(FMargin(8, 4.5, 8, 3.5))
+		.SetPressedPadding(FMargin(8, 5, 6, 3));
 	{
 
 		const FTextBlockStyle& NormalText = Style->GetWidgetStyle<FTextBlockStyle>("NormalText");
@@ -1076,8 +1089,15 @@ void FStarshipCoreStyle::SetupButtonStyles(TSharedRef<FStyle>& Style)
 		Style->Set("PrimaryButton", PrimaryButton);
 		Style->Set("Button", Button);
 		Style->Set("SimpleButton", SimpleButton);
+		Style->Set("SecondaryButton", SecondaryButton);
 
 		Style->Set("DialogButtonText", FTextBlockStyle(NormalText).SetFont(FStyleFonts::Get().NormalBold).SetTransformPolicy(ETextTransformPolicy::ToUpper));
+
+
+		FSlateFontInfo AddFont = FStyleFonts::Get().SmallBold;
+		AddFont.LetterSpacing = 200;
+
+		Style->Set("SmallButtonText", FTextBlockStyle(NormalText).SetFont(AddFont).SetTransformPolicy(ETextTransformPolicy::ToUpper));
 	}
 }
 
