Slate: Add SlateInvalidationWidgetList a container, similar to a paged array, that can update the fast path list faster. It uses 2 indexes instead of one. It only rebuilds the elements that are affected by the update. Use FSlateInvalidationWidgetIndex and FSlateInvalidationWidgetSortIndex. FWidgetProxy is currently used but will be phase out soon to remove duplication and improve performance.

diff --git a/Engine/Plugins/Experimental/DigitalMultiplex/DMXPixelMapping/Source/DMXPixelMappingEditor/Private/Views/SDMXPixelMappingDesignerView.cpp b/Engine/Plugins/Experimental/DigitalMultiplex/DMXPixelMapping/Source/DMXPixelMappingEditor/Private/Views/SDMXPixelMappingDesignerView.cpp
--- a/Engine/Plugins/Experimental/DigitalMultiplex/DMXPixelMapping/Source/DMXPixelMappingEditor/Private/Views/SDMXPixelMappingDesignerView.cpp
+++ b/Engine/Plugins/Experimental/DigitalMultiplex/DMXPixelMapping/Source/DMXPixelMappingEditor/Private/Views/SDMXPixelMappingDesignerView.cpp
@@ -550,7 +550,7 @@ void SDMXPixelMappingDesignerView::PopulateWidgetGeometryCache_Loop(FArrangedWid
 
 	if (bIncludeInHitTestGrid)
 	{
-		HittestGrid->AddWidget(CurrentWidget.Widget, 0, 0, 0);
+		HittestGrid->AddWidget(CurrentWidget.Widget, 0, 0, FSlateInvalidationWidgetSortOrder());
 	}
 
 	FArrangedChildren ArrangedChildren(EVisibility::All);
diff --git a/Engine/Source/Editor/UMGEditor/Private/Designer/SDesignerView.cpp b/Engine/Source/Editor/UMGEditor/Private/Designer/SDesignerView.cpp
--- a/Engine/Source/Editor/UMGEditor/Private/Designer/SDesignerView.cpp
+++ b/Engine/Source/Editor/UMGEditor/Private/Designer/SDesignerView.cpp
@@ -5,6 +5,7 @@
 #include "Components/PanelWidget.h"
 #include "Misc/ConfigCacheIni.h"
 #include "WidgetBlueprint.h"
+#include "FastUpdate/SlateInvalidationWidgetSortOrder.h"
 #include "Framework/Application/MenuStack.h"
 #include "Framework/Application/SlateApplication.h"
 #include "Widgets/Layout/SBorder.h"
@@ -2144,7 +2145,7 @@ void SDesignerView::PopulateWidgetGeometryCache_Loop(FArrangedWidget& CurrentWid
 
 	if (bIncludeInHitTestGrid)
 	{
-		DesignerHittestGrid->AddWidget(CurrentWidget.Widget, 0, 0, 0);
+		DesignerHittestGrid->AddWidget(CurrentWidget.Widget, 0, 0, FSlateInvalidationWidgetSortOrder());
 	}
 
 	FArrangedChildren ArrangedChildren(EVisibility::All);
diff --git a/Engine/Source/Runtime/Slate/Private/Widgets/SInvalidationPanel.cpp b/Engine/Source/Runtime/Slate/Private/Widgets/SInvalidationPanel.cpp
--- a/Engine/Source/Runtime/Slate/Private/Widgets/SInvalidationPanel.cpp
+++ b/Engine/Source/Runtime/Slate/Private/Widgets/SInvalidationPanel.cpp
@@ -93,9 +93,9 @@ bool SInvalidationPanel::GetCanCache() const
 	// In global invalidation this code makes no sense so we don't bother running it because everything is in an "invalidation panel" at the window level
 #if WITH_SLATE_DEBUGGING
 	// Disable invalidation panels if global invalidation is turned on
-	return bCanCache && !GSlateEnableGlobalInvalidation && !GetProxyHandle().IsValid() && InvalidationPanelsEnabled;
+	return bCanCache && !GSlateEnableGlobalInvalidation && !GetProxyHandle().IsValid(this) && InvalidationPanelsEnabled;
 #else
-	return bCanCache && !GSlateEnableGlobalInvalidation && !GetProxyHandle().IsValid();
+	return bCanCache && !GSlateEnableGlobalInvalidation && !GetProxyHandle().IsValid(this);
 #endif
 }
 
diff --git a/Engine/Source/Runtime/SlateCore/Private/Debugging/SlateDebugging.cpp b/Engine/Source/Runtime/SlateCore/Private/Debugging/SlateDebugging.cpp
--- a/Engine/Source/Runtime/SlateCore/Private/Debugging/SlateDebugging.cpp
+++ b/Engine/Source/Runtime/SlateCore/Private/Debugging/SlateDebugging.cpp
@@ -24,12 +24,6 @@ CSV_DEFINE_CATEGORY_MODULE(SLATECORE_API, Slate, true);
 
 #define LOCTEXT_NAMESPACE "SlateDebugger"
 
-// Scalar goes from 0 to 1.  0 being all yellow, 1 being all red
-FLinearColor YellowToRedFromScalar(float Scalar)
-{
-	return FLinearColor(1.0f, 1.0f * (1.0f - Scalar), 0.0f);
-}
-
 FSlateDebuggingInputEventArgs::FSlateDebuggingInputEventArgs(ESlateDebuggingInputEvent InInputEventType, const FReply& InReply, const TSharedPtr<SWidget>& InHandlerWidget, const FString& InAdditionalContent)
 	: InputEventType(InInputEventType)
 	, Reply(InReply)
diff --git a/Engine/Source/Runtime/SlateCore/Private/FastUpdate/SlateInvalidationRoot.cpp b/Engine/Source/Runtime/SlateCore/Private/FastUpdate/SlateInvalidationRoot.cpp
--- a/Engine/Source/Runtime/SlateCore/Private/FastUpdate/SlateInvalidationRoot.cpp
+++ b/Engine/Source/Runtime/SlateCore/Private/FastUpdate/SlateInvalidationRoot.cpp
@@ -3,6 +3,9 @@
 #include "FastUpdate/SlateInvalidationRoot.h"
 #include "FastUpdate/SlateInvalidationRootHandle.h"
 #include "FastUpdate/SlateInvalidationRootList.h"
+#include "FastUpdate/SlateInvalidationWidgetHeap.h"
+#include "FastUpdate/SlateInvalidationWidgetList.h"
+#include "FastUpdate/SlateInvalidationWidgetSortOrder.h"
 #include "Async/TaskGraphInterfaces.h"
 #include "Application/SlateApplicationBase.h"
 #include "Widgets/SWidget.h"
@@ -26,6 +29,36 @@ static FAutoConsoleCommand HandleDumpUpdateListCommand(
 	TEXT(""),
 	FConsoleCommandWithArgsDelegate::CreateStatic(&HandleDumpUpdateList)
 );
+
+bool GSlateInvalidationRootVerifyWidgetList = false;
+static FAutoConsoleVariableRef CVarSlateInvalidationRootVerifyWidgetList(
+	TEXT("Slate.InvalidationRoot.VerifyWidgetList"),
+	GSlateInvalidationRootVerifyWidgetList,
+	TEXT("Each frame, verify that the updated list doesn't match a newly created list.")
+);
+void VerifyWidgetList(SWidget* InvalidationRootWidget, FSlateInvalidationRootHandle InvalidationRootHandle, FSlateInvalidationWidgetList& WidgetList);
+
+bool GSlateInvalidationRootVerifyWidgetsIndex = false;
+static FAutoConsoleVariableRef CVarSlateInvalidationRootVerifyWidgetsIndex(
+	TEXT("Slate.InvalidationRoot.VerifyWidgetsIndex"),
+	GSlateInvalidationRootVerifyWidgetsIndex,
+	TEXT("Each frame, verify that every widgets has the correct index.")
+);
+
+bool GSlateInvalidationRootVerifyValidWidgets = false;
+static FAutoConsoleVariableRef CVarSlateInvalidationRootVerifyValidWidgets(
+	TEXT("Slate.InvalidationRoot.VerifyValidWidges"),
+	GSlateInvalidationRootVerifyValidWidgets,
+	TEXT("Each frame, verify that every WidgetProxy has a valid SWidget.")
+);
+
+bool GSlateInvalidationRootVerifyHittestGrid = false;
+static FAutoConsoleVariableRef CVarSlateInvalidationRootVerifyHittestGrid(
+	TEXT("Slate.InvalidationRoot.VerifyHittestGrid"),
+	GSlateInvalidationRootVerifyHittestGrid,
+	TEXT("Each frame, verify the hittest grid.")
+);
+void VerifyHittest(SWidget* InvalidationRootWidget, FSlateInvalidationWidgetList& WidgetList, FHittestGrid* HittestGrid);
 #endif //WITH_SLATE_DEBUGGING
 
 #if SLATE_CSV_TRACKER
@@ -36,6 +69,17 @@ FAutoConsoleVariableRef CVarCascadeInvalidationEventAmount(
 	TEXT("The amount of cascaded invalidated parents before we fire a CSV event."));
 #endif //SLATE_CSV_TRACKER
 
+int32 GSlateInvalidationWidgetListMaxArrayElements = 64;
+FAutoConsoleVariableRef CVarSlateInvalidationWidgetListMaxArrayElements(
+	TEXT("Slate.InvalidationList.MaxArrayElements"),
+	GSlateInvalidationWidgetListMaxArrayElements,
+	TEXT("With Global Invalidation, the preferred size of the elements array."));
+
+int32 GSlateInvalidationWidgetListNumberElementLeftBeforeSplitting = 40;
+FAutoConsoleVariableRef CVarSlateInvalidationWidgetListNumElementLeftBeforeSplitting(
+	TEXT("Slate.InvalidationList.NumberElementLeftBeforeSplitting"),
+	GSlateInvalidationWidgetListNumberElementLeftBeforeSplitting,
+	TEXT("With Global Invalidation, when splitting, only split the array when the number of element left is under X."));
 
 /**
  *
@@ -52,10 +96,15 @@ FSlateInvalidationRoot::FSlateInvalidationRoot()
 	, bChildOrderInvalidated(false)
 	, bNeedsSlowPath(true)
 	, bNeedScreenPositionShift(false)
+	, bProcessingChildOrderUpdate(false)
 {
 	InvalidationRootHandle = FSlateInvalidationRootHandle(GSlateInvalidationRootListInstance.AddInvalidationRoot(this));
 	FSlateApplicationBase::Get().OnInvalidateAllWidgets().AddRaw(this, &FSlateInvalidationRoot::HandleInvalidateAllWidgets);
 
+	const FSlateInvalidationWidgetList::FArguments Arg = { GSlateInvalidationWidgetListMaxArrayElements, GSlateInvalidationWidgetListNumberElementLeftBeforeSplitting };
+	FastWidgetPathList = new FSlateInvalidationWidgetList(InvalidationRootHandle, Arg);
+	WidgetsNeedingUpdate = new FSlateInvalidationWidgetHeap(*FastWidgetPathList);
+
 #if WITH_SLATE_DEBUGGING
 	SetLastPaintType(ESlateInvalidationPaintType::None);
 #endif
@@ -80,7 +129,8 @@ FSlateInvalidationRoot::~FSlateInvalidationRoot()
 		delete CachedElementData;
 	}
 
-	CachedElementData = nullptr;
+	delete FastWidgetPathList;
+	delete WidgetsNeedingUpdate;
 
 	GSlateInvalidationRootListInstance.RemoveInvalidationRoot(InvalidationRootHandle.GetUniqueId());
 }
@@ -95,28 +145,46 @@ FString FSlateInvalidationRoot::GetReferencerName() const
 	return TEXT("FSlateInvalidationRoot");
 }
 
-void FSlateInvalidationRoot::InvalidateChildOrder(const SWidget* Investigator)
+void FSlateInvalidationRoot::InvalidateRoot(const SWidget* Investigator)
+{
+	// Update the generation number. This will effectively invalidate all proxy handles
+	++FastPathGenerationNumber;
+
+	InvalidationRootWidget->InvalidatePrepass();
+
+	bNeedsSlowPath = true;
+
+#if WITH_SLATE_DEBUGGING
+	FSlateDebugging::BroadcastInvalidationRootInvalidate(InvalidationRootWidget, Investigator, ESlateDebuggingInvalidateRootReason::Root);
+#endif
+	UE_TRACE_SLATE_ROOT_INVALIDATED(InvalidationRootWidget, Investigator);
+}
+
+void FSlateInvalidationRoot::InvalidateWidgetChildOrder(TSharedRef<SWidget> Widget)
 {
-	if(!bNeedsSlowPath && !bChildOrderInvalidated)
+	if (!bNeedsSlowPath)
 	{
-		//UE_LOG(LogSlate, Log, TEXT("Child order invalidated, Slow path needed"));
-		//bNeedsSlowPath = true;
-		bChildOrderInvalidated = true;
-		if(!InvalidationRootWidget->Advanced_IsWindow())
-		{
-			InvalidationRootWidget->InvalidatePrepass();
-		}
+		ensureAlways(bProcessingChildOrderUpdate == false);
+		WidgetsNeedingChildOrderUpdate.Add(Widget);
 
-		if (!GSlateEnableGlobalInvalidation && !InvalidationRootWidget->Advanced_IsWindow())
+		if (!bChildOrderInvalidated)
 		{
-			//InvalidationRootWidget->Invalidate(EInvalidateWidgetReason::ChildOrder);
-			InvalidationRootWidget->Invalidate(EInvalidateWidgetReason::Layout);
+			bChildOrderInvalidated = true;
+			if (!InvalidationRootWidget->Advanced_IsWindow())
+			{
+				InvalidationRootWidget->InvalidatePrepass();
+			}
+
+			if (!GSlateEnableGlobalInvalidation && !InvalidationRootWidget->Advanced_IsWindow())
+			{
+				InvalidationRootWidget->Invalidate(EInvalidateWidgetReason::Layout);
+			}
 		}
 
 #if WITH_SLATE_DEBUGGING
-		FSlateDebugging::BroadcastInvalidationRootInvalidate(InvalidationRootWidget, Investigator, ESlateDebuggingInvalidateRootReason::ChildOrder);
+		FSlateDebugging::BroadcastInvalidationRootInvalidate(InvalidationRootWidget, &Widget.Get(), ESlateDebuggingInvalidateRootReason::ChildOrder);
 #endif
-		UE_TRACE_SLATE_ROOT_CHILDORDER_INVALIDATED(InvalidationRootWidget, Investigator);
+		UE_TRACE_SLATE_ROOT_CHILDORDER_INVALIDATED(InvalidationRootWidget, &Widget.Get());
 	}
 }
 
@@ -134,52 +202,6 @@ void FSlateInvalidationRoot::InvalidateScreenPosition(const SWidget* Investigato
 #endif
 }
 
-int32 RecursiveFindParentWithChildOrderChange(const TArray<FWidgetProxy>& FastWidgetPathList, const FWidgetProxy& Proxy)
-{
-	if (Proxy.bChildOrderInvalid)
-	{
-		return Proxy.Index;
-	}
-	else if (Proxy.ParentIndex == INDEX_NONE)
-	{
-		return INDEX_NONE;
-	}
-	else
-	{
-		return RecursiveFindParentWithChildOrderChange(FastWidgetPathList, FastWidgetPathList[Proxy.ParentIndex]);
-	}
-}
-
-void FSlateInvalidationRoot::RemoveWidgetFromFastPath(FWidgetProxy& Proxy)
-{
-	if (Proxy.Index == 0)
-	{
-		InvalidateRoot(Proxy.Widget);
-	}
-	else
-	{
-		InvalidateChildOrder(Proxy.Widget);
-	}
-
-	Proxy.Widget->FastPathProxyHandle = FWidgetProxyHandle();
-	Proxy.Widget = nullptr;
-}
-
-void FSlateInvalidationRoot::InvalidateRoot(const SWidget* Investigator)
-{
-	// Update the generation number.  This will effectively invalidate all proxy handles
-	++FastPathGenerationNumber;
-
-	InvalidationRootWidget->InvalidatePrepass();
-
-	bNeedsSlowPath = true;
-
-#if WITH_SLATE_DEBUGGING
-	FSlateDebugging::BroadcastInvalidationRootInvalidate(InvalidationRootWidget, Investigator, ESlateDebuggingInvalidateRootReason::Root);
-#endif
-	UE_TRACE_SLATE_ROOT_INVALIDATED(InvalidationRootWidget, Investigator);
-}
-
 FSlateInvalidationResult FSlateInvalidationRoot::PaintInvalidationRoot(const FSlateInvalidationContext& Context)
 {
 	const int32 LayerId = 0;
@@ -198,7 +220,7 @@ FSlateInvalidationResult FSlateInvalidationRoot::PaintInvalidationRoot(const FSl
 		Context.WindowElementList->PushCachedElementData(*CachedElementData);
 	}
 
-	SWidget* RootWidget = InvalidationRootWidget->Advanced_IsWindow() ? InvalidationRootWidget : &(*InvalidationRootWidget->GetAllChildren()->GetChildAt(0));
+	TSharedRef<SWidget> RootWidget = InvalidationRootWidget->Advanced_IsWindow() ? InvalidationRootWidget->AsShared() : InvalidationRootWidget->GetAllChildren()->GetChildAt(0);
 
 	if (bNeedScreenPositionShift)
 	{
@@ -229,7 +251,7 @@ FSlateInvalidationResult FSlateInvalidationRoot::PaintInvalidationRoot(const FSl
 			{
 				TGuardValue<bool> InSlowPathGuard(GSlateIsInInvalidationSlowPath, true);
 
-				BuildFastPathList(RootWidget);
+				BuildFastPathWidgetList(RootWidget);
 
 				if(GSlateEnableGlobalInvalidation)
 				{
@@ -246,10 +268,10 @@ FSlateInvalidationResult FSlateInvalidationRoot::PaintInvalidationRoot(const FSl
 		Result.bRepaintedWidgets = true;
 
 	}
-	else if (FastWidgetPathList.Num())
+	else if (!FastWidgetPathList->IsEmpty())
 	{
 		// We should not have been supplied a different root than the one we generated a path to
-		check(RootWidget == FastWidgetPathList[0].Widget);
+		check(RootWidget == FastWidgetPathList->GetRoot().Pin());
 
 		Result.bRepaintedWidgets = PaintFastPath(Context);
 	}
@@ -261,20 +283,34 @@ FSlateInvalidationResult FSlateInvalidationRoot::PaintInvalidationRoot(const FSl
 
 	FinalUpdateList.Reset();
 
+#if WITH_SLATE_DEBUGGING
+	if (GSlateInvalidationRootVerifyHittestGrid)
+	{
+		VerifyHittest(InvalidationRootWidget, GetFastPathWidgetList(), GetHittestGrid());
+	}
+#endif
+
+	FinalUpdateList.Reset();
+
 	Result.MaxLayerIdPainted = CachedMaxLayerId;
 	return Result;
 }
 
 void FSlateInvalidationRoot::OnWidgetDestroyed(const SWidget* Widget)
 {
-	InvalidateChildOrder(Widget);
-
 	// We need the index even if we've invalidated this root.  We need to clear out its proxy regardless
-	const bool bEvenIfInvalid = true;
-	const int32 ProxyIndex = Widget->FastPathProxyHandle.GetIndex(bEvenIfInvalid);
-	if (FastWidgetPathList.IsValidIndex(ProxyIndex) && FastWidgetPathList[ProxyIndex].Widget == Widget)
+	const FSlateInvalidationWidgetIndex ProxyIndex = Widget->GetProxyHandle().GetWidgetIndex();
+	if (FastWidgetPathList->IsValidIndex(ProxyIndex))
 	{
-		FastWidgetPathList[ProxyIndex].Widget = nullptr;
+		FSlateInvalidationWidgetList::InvalidationWidgetType& Proxy = (*FastWidgetPathList)[ProxyIndex];
+#if UE_SLATE_WITH_WIDGETPROXY_WEAKPTR
+		if (Proxy.GetWidget(false) == Widget)
+#else
+		if (Proxy.GetWidget() == Widget)
+#endif
+		{
+			Proxy.ResetWidget();
+		}
 	}
 }
 
@@ -286,12 +322,12 @@ void FSlateInvalidationRoot::ClearAllWidgetUpdatesPending()
 	{
 		if (int32 NumUpdatePending = Root->FinalUpdateList.Num())
 		{
-			for (int32 index : Root->FinalUpdateList)
+			for (FSlateInvalidationWidgetIndex index : Root->FinalUpdateList)
 			{
-				FWidgetProxy& Proxy = Root->FastWidgetPathList[index];
+				FWidgetProxy& Proxy = (*Root->FastWidgetPathList)[index];
 				if (EnumHasAnyFlags(Proxy.UpdateFlags, EWidgetUpdateFlags::AnyUpdate))
 				{
-					Root->WidgetsNeedingUpdate.Push(Proxy);
+					Root->WidgetsNeedingUpdate->PushUnique(Proxy);
 				}
 			}
 		}
@@ -320,28 +356,31 @@ bool FSlateInvalidationRoot::PaintFastPath(const FSlateInvalidationContext& Cont
 			// The update list is put in reverse order 
 			for (int32 ListIndex = FinalUpdateList.Num() - 1; ListIndex >= 0; --ListIndex)
 			{
-				const int32 MyIndex = FinalUpdateList[ListIndex];
-
-				FWidgetProxy& WidgetProxy = FastWidgetPathList[MyIndex];
+				const FSlateInvalidationWidgetIndex MyIndex = FinalUpdateList[ListIndex];
 
-				if (EnumHasAnyFlags(WidgetProxy.UpdateFlags, EWidgetUpdateFlags::NeedsVolatilePaint))
+				FWidgetProxy& WidgetProxy = (*FastWidgetPathList)[MyIndex];
+				SWidget* WidgetPtr = WidgetProxy.GetWidget();
+				if (ensure(WidgetPtr))
 				{
-					UE_LOG(LogSlate, Log, TEXT("Volatile Repaint %s"), *FReflectionMetaData::GetWidgetDebugInfo(WidgetProxy.Widget));
-				}
-				else if (EnumHasAnyFlags(WidgetProxy.UpdateFlags, EWidgetUpdateFlags::NeedsRepaint))
-				{
-					UE_LOG(LogSlate, Log, TEXT("Repaint %s"), *FReflectionMetaData::GetWidgetDebugInfo(WidgetProxy.Widget));
-				}
-				else if (!WidgetProxy.bInvisibleDueToParentOrSelfVisibility)
-				{
-					if (EnumHasAnyFlags(WidgetProxy.UpdateFlags, EWidgetUpdateFlags::NeedsActiveTimerUpdate))
+					if (EnumHasAnyFlags(WidgetProxy.UpdateFlags, EWidgetUpdateFlags::NeedsVolatilePaint))
 					{
-						UE_LOG(LogSlate, Log, TEXT("ActiveTimer %s"), *FReflectionMetaData::GetWidgetDebugInfo(WidgetProxy.Widget));
+						UE_LOG(LogSlate, Log, TEXT("Volatile Repaint %s"), *FReflectionMetaData::GetWidgetDebugInfo(WidgetPtr));
 					}
-
-					if (EnumHasAnyFlags(WidgetProxy.UpdateFlags, EWidgetUpdateFlags::NeedsTick))
+					else if (EnumHasAnyFlags(WidgetProxy.UpdateFlags, EWidgetUpdateFlags::NeedsRepaint))
+					{
+						UE_LOG(LogSlate, Log, TEXT("Repaint %s"), *FReflectionMetaData::GetWidgetDebugInfo(WidgetPtr));
+					}
+					else if (!WidgetProxy.bInvisibleDueToParentOrSelfVisibility)
 					{
-						UE_LOG(LogSlate, Log, TEXT("Tick %s"), *FReflectionMetaData::GetWidgetDebugInfo(WidgetProxy.Widget));
+						if (EnumHasAnyFlags(WidgetProxy.UpdateFlags, EWidgetUpdateFlags::NeedsActiveTimerUpdate))
+						{
+							UE_LOG(LogSlate, Log, TEXT("ActiveTimer %s"), *FReflectionMetaData::GetWidgetDebugInfo(WidgetPtr));
+						}
+
+						if (EnumHasAnyFlags(WidgetProxy.UpdateFlags, EWidgetUpdateFlags::NeedsTick))
+						{
+							UE_LOG(LogSlate, Log, TEXT("Tick %s"), *FReflectionMetaData::GetWidgetDebugInfo(WidgetPtr));
+						}
 					}
 				}
 			}
@@ -354,18 +393,18 @@ bool FSlateInvalidationRoot::PaintFastPath(const FSlateInvalidationContext& Cont
 			// The update list is put in reverse order by ProcessInvalidation
 			for (int32 ListIndex = FinalUpdateList.Num() - 1; ListIndex >= 0; --ListIndex)
 			{
-				const int32 MyIndex = FinalUpdateList[ListIndex];
-				FWidgetProxy& WidgetProxy = FastWidgetPathList[MyIndex];
+				const FSlateInvalidationWidgetIndex MyIndex = FinalUpdateList[ListIndex];
+				FWidgetProxy& WidgetProxy = (*FastWidgetPathList)[MyIndex];
 
 				// Check visibility, it may have been in the update list but a parent who was also in the update list already updated it.
-				if (!WidgetProxy.bInvisibleDueToParentOrSelfVisibility && !WidgetProxy.bUpdatedSinceLastInvalidate && ensure(WidgetProxy.Widget))
+				if (!WidgetProxy.bInvisibleDueToParentOrSelfVisibility && !WidgetProxy.bUpdatedSinceLastInvalidate && ensure(WidgetProxy.GetWidget()))
 				{
 					bWidgetsNeededRepaint = bWidgetsNeededRepaint || EnumHasAnyFlags(WidgetProxy.UpdateFlags, EWidgetUpdateFlags::NeedsRepaint | EWidgetUpdateFlags::NeedsVolatilePaint);
 
-					const int32 NewLayerId = WidgetProxy.Update(*Context.PaintArgs, MyIndex, *Context.WindowElementList);
+					const int32 NewLayerId = WidgetProxy.Update(*Context.PaintArgs, *Context.WindowElementList);
 					CachedMaxLayerId = FMath::Max(NewLayerId, CachedMaxLayerId);
 
-					WidgetProxy.MarkProxyUpdatedThisFrame(WidgetsNeedingUpdate);
+					WidgetProxy.MarkProxyUpdatedThisFrame(*WidgetsNeedingUpdate);
 
 					if (bNeedsSlowPath)
 					{
@@ -390,185 +429,54 @@ bool FSlateInvalidationRoot::PaintFastPath(const FSlateInvalidationContext& Cont
 	return bWidgetsNeededRepaint;
 }
 
-bool FSlateInvalidationRoot::BuildNewFastPathList_Recursive(FSlateInvalidationRoot& Root, FWidgetProxy& Proxy, int32 ParentIndex, int32& NextTreeIndex, TArray<FWidgetProxy>& CurrentFastPathList, TArray<FWidgetProxy, TMemStackAllocator<>>& NewFastPathList)
-{
-	if (Proxy.Widget == nullptr)
-	{
-		return false;
-	}
-
-	bool bResult = true;
-	if (Proxy.bChildOrderInvalid)
-	{
-		NextTreeIndex = Proxy.LeafMostChildIndex != INDEX_NONE ? Proxy.LeafMostChildIndex + 1 : NextTreeIndex + 1;
-		Proxy.Widget->AssignIndicesToChildren(*this, ParentIndex, NewFastPathList, !Proxy.bInvisibleDueToParentOrSelfVisibility, Proxy.Widget->IsVolatileIndirectly());
-	}
-	else
-	{ 
-		const int32 PrevIndex = Proxy.Index;
-		const int32 PrevParentIndex = Proxy.ParentIndex;
-		Proxy.Index = NewFastPathList.Num();
-		Proxy.ParentIndex = ParentIndex;
-		//if (PrevIndex != Proxy.Index)
-		{
-			// Update the proxy handle
-			Proxy.Widget->FastPathProxyHandle = FWidgetProxyHandle(Root, Proxy.Index);
-		}
-
-		NewFastPathList.Add(MoveTemp(Proxy));
-
-		for (int32 LocalChildIndex = 0; LocalChildIndex < Proxy.NumChildren; ++LocalChildIndex)
-		{
-			if (NextTreeIndex < CurrentFastPathList.Num())
-			{
-				FWidgetProxy& ChildProxy = CurrentFastPathList[NextTreeIndex];
-				if (!ChildProxy.bChildOrderInvalid)
-				{
-					++NextTreeIndex;
-				}
-
-				bResult = bResult && BuildNewFastPathList_Recursive(Root, ChildProxy, Proxy.Index, NextTreeIndex, CurrentFastPathList, NewFastPathList);
-			}
-			else
-			{
-				bResult = false;
-				break;
-			}
-		}
-
-		{
-			FWidgetProxy& MyProxyRef = NewFastPathList[Proxy.Index];
-			int32 LastIndex = NewFastPathList.Num() - 1;
-			MyProxyRef.LeafMostChildIndex = LastIndex != Proxy.Index ? LastIndex : INDEX_NONE;
-		}
-	}
-	
-	return bResult;
-}
-
 void FSlateInvalidationRoot::AdjustWidgetsDesktopGeometry(FVector2D WindowToDesktopTransform)
 {
 	FSlateLayoutTransform WindowToDesktop(WindowToDesktopTransform);
 
-	for (FWidgetProxy& Proxy : FastWidgetPathList)
-	{
-		if (SWidget* Widget = Proxy.Widget)
+	FastWidgetPathList->ForEachWidget([WindowToDesktopTransform, &WindowToDesktop](SWidget* Widget)
 		{
 			Widget->PersistentState.DesktopGeometry = Widget->PersistentState.AllottedGeometry;
 			Widget->PersistentState.DesktopGeometry.AppendTransform(WindowToDesktop);
-		}
-	}
+		});
 }
 
-#define UE_VERIFY_CHILD_ORDER 0
-#define UE_VERIFY_UNIQUENESS 0
-#define UE_VERIFY_FAST_PATH_PROXY_HANDLE 0
-void FSlateInvalidationRoot::BuildFastPathList(SWidget* RootWidget)
-{
-	SCOPED_NAMED_EVENT_TEXT("AssignFastPathIndices", FColor::Magenta);
-
-	TSharedPtr<SWidget> Parent = RootWidget->GetParentWidget();
-	// If the widget has no parent it is likely a window
-	const bool bParentVisible = Parent.IsValid() ? Parent->GetVisibility().IsVisible() : true;
-	const bool bParentVolatile = false;
 
-	FMemMark Mark(FMemStack::Get());
+void FSlateInvalidationRoot::ProcessChildOrderUpdate()
+{
+	TGuardValue<bool> Tmp(bProcessingChildOrderUpdate, true);
 	{
-		// Update the generation number.  This will effectively invalidate all proxy handles
-		++FastPathGenerationNumber;
-
-		WidgetsNeedingUpdate.Empty();
-
-
-#if UE_VERIFY_CHILD_ORDER
-		TArray<FWidgetProxy, TMemStackAllocator<>> Copy;
-		Copy.Reserve(FastWidgetPathList.Num());
-		RootWidget->AssignIndicesToChildren(*this, INDEX_NONE, Copy, bParentVisible, bParentVolatile);
-
-		// Update the generation number.  This will effectively invalidate all proxy handles
-		//++FastPathGenerationNumber;
-#endif
-		TArray<FWidgetProxy, TMemStackAllocator<>> TempList;
-		TempList.Reserve(FastWidgetPathList.Num());
-
-		bool bBuiltPath = false;
-		if (FastWidgetPathList.Num() > 0)
-		{
-			SCOPED_NAMED_EVENT_FSTRING(FString::Printf(TEXT("BuildFastPathList_BuildNewFastPathList_Recursive: %s"), *FReflectionMetaData::GetWidgetDebugInfo(FastWidgetPathList[0].Widget)), FColor::Magenta);
-
-			int32 NextTreeIndex = 1;
-			bBuiltPath = BuildNewFastPathList_Recursive(*this, FastWidgetPathList[0], INDEX_NONE, NextTreeIndex, FastWidgetPathList, TempList);
-			if (!bBuiltPath)
-			{
-				// invalidate partially built fast path
-				++FastPathGenerationNumber;
-			}
-		}
-		
-		if (!bBuiltPath)
-		{
-			SCOPED_NAMED_EVENT_FSTRING(FString::Printf(TEXT("BuildFastPathList_AssignIndicesToChildren: %s"), *FReflectionMetaData::GetWidgetDebugInfo(RootWidget)), FColor::Magenta);
-			TempList.Reset();
-			RootWidget->AssignIndicesToChildren(*this, INDEX_NONE, TempList, bParentVisible, bParentVolatile);
-		}
-
-#if UE_VERIFY_CHILD_ORDER
-		for (int32 i = 0; i < Copy.Num(); ++i)
+		TSharedRef<SWidget> RootWidget = InvalidationRootWidget->Advanced_IsWindow() ? InvalidationRootWidget->AsShared() : InvalidationRootWidget->GetAllChildren()->GetChildAt(0);
+		if (FastWidgetPathList->GetRoot().Pin() != RootWidget)
 		{
-			ensureAlways(Copy[i].Widget == TempList[i].Widget);
+			FastWidgetPathList->BuildWidgetList(RootWidget);
 		}
-#endif
-
-		// When it's the first time the FastWidgetPathList has item, we know we will need to clear the UpdateList on the next frame.
-		if (FastWidgetPathList.Num() == 0 && TempList.Num() != 0)
-		{
-			ensure(ClearUpdateList.Find(this) == INDEX_NONE);
-			ClearUpdateList.Push(this);
-		}
-		// When FastWidgetPathList is empty for the first time, we know we don't need to clear the list until we have items readded
-		else if (FastWidgetPathList.Num() != 0 && TempList.Num() == 0)
+		else
 		{
-			ClearUpdateList.RemoveSingleSwap(this, false);
+			FastWidgetPathList->ProcessChildOrderInvalidation(WidgetsNeedingChildOrderUpdate);
 		}
 
-#if UE_VERIFY_UNIQUENESS
-		for (int32 Index = 0; Index < TempList.Num(); ++Index)
-		{
-			const SWidget* Widget = TempList[Index].Widget;
-			ensureAlways(TempList.FindLastByPredicate([Widget](const FWidgetProxy& Proxy){ return Proxy.Widget == Widget; }) == Index);
-		}
-#endif
-
-#if UE_VERIFY_FAST_PATH_PROXY_HANDLE
-		for (const FWidgetProxy& Proxy : FastWidgetPathList)
-		{
-			const SWidget* Widget = Proxy.Widget;
-			if (Widget)
-			{
-				const bool bContains = TempList.ContainsByPredicate([Widget](const FWidgetProxy& Proxy) { return Proxy.Widget == Widget; });
-				if (!bContains)
-				{
-					ensureAlways(Widget->FastPathProxyHandle.GetIndex(true) == INDEX_NONE);
-				}
-			}
-
-		}
-#endif
+		WidgetsNeedingChildOrderUpdate.Reset();
+		bChildOrderInvalidated = false;
+	}
 
-#if UE_SLATE_DEBUGGING_CLEAR_ALL_FAST_PATH_DATA
-		for (const FWidgetProxy& Proxy : TempList)
-		{
-			FastWidgetPathToClearedBecauseOfDelay.RemoveSingleSwap(Proxy.Widget);
-		}
-		ensureAlways(FastWidgetPathToClearedBecauseOfDelay.Num() == 0);
-#endif
+#if WITH_SLATE_DEBUGGING
+	if (GSlateInvalidationRootVerifyWidgetList)
+	{
+		VerifyWidgetList(InvalidationRootWidget, InvalidationRootHandle, GetFastPathWidgetList());
+	}
 
-		FastWidgetPathList = TempList;
+	if (GSlateInvalidationRootVerifyWidgetsIndex)
+	{
+		ensureMsgf(FastWidgetPathList->VerifyWidgetsIndex(), TEXT("We failed to verify that every widgets has the correct index."));
 	}
+#endif //WITH_SLATE_DEBUGGING
+}
+
+void FSlateInvalidationRoot::BuildFastPathWidgetList(TSharedRef<SWidget> RootWidget)
+{
+	FastWidgetPathList->BuildWidgetList(RootWidget);
+	WidgetsNeedingChildOrderUpdate.Reset();
 }
-#undef UE_VERIFY_CHILD_ORDER
-#undef UE_VERIFY_UNIQUENESS
-#undef UE_VERIFY_FAST_PATH_PROXY_HANDLE
 
 bool FSlateInvalidationRoot::ProcessInvalidation()
 {
@@ -579,75 +487,73 @@ bool FSlateInvalidationRoot::ProcessInvalidation()
 
 	if (!bNeedsSlowPath)
 	{
-		if (bChildOrderInvalidated)
+		ensure(bChildOrderInvalidated == (WidgetsNeedingChildOrderUpdate.Num() != 0));
+		if (WidgetsNeedingChildOrderUpdate.Num())
 		{
-			SCOPED_NAMED_EVENT(Slate_InvalidationProcessing_SortChildren, FColor::Orange);
+			SCOPED_NAMED_EVENT(Slate_InvalidationProcessing_ChildOrder, FColor::Blue);
 
-			struct FWidgetNeedingUpdate
-			{
-				FWidgetNeedingUpdate(SWidget* InWidget, EInvalidateWidgetReason InCurrentInvalidateReason, EWidgetUpdateFlags InUpdateFlags)
-					: Widget(InWidget)
-					, CurrentInvalidateReason(InCurrentInvalidateReason)
-					, UpdateFlags(InUpdateFlags)
-				{}
-
-				SWidget* Widget;
-				EInvalidateWidgetReason CurrentInvalidateReason;
-				EWidgetUpdateFlags UpdateFlags;
-			};
-
-			// We need to store off all widgets needing update as the child order change may invalidate them all
-			TArray<FWidgetNeedingUpdate, TMemStackAllocator<>> WidgetsNeedingUpdateCache;
-			WidgetsNeedingUpdateCache.Reserve(FinalUpdateList.Num() + WidgetsNeedingUpdate.Num());
-
-			for (int32 WidgetIndex : FinalUpdateList)
+			FMemMark Mark(FMemStack::Get());
+			TArray<SWidget*, TMemStackAllocator<>> FastPathWidgetsNeedingUpdateCache;
+			FastPathWidgetsNeedingUpdateCache.Reserve(WidgetsNeedingUpdate->Num());
+			TArray<SWidget*, TMemStackAllocator<>> FastPathUpdateListCache;
+			FastPathUpdateListCache.Reserve(FinalUpdateList.Num());
+
+			for (const FSlateInvalidationWidgetHeap::FElement& Element : WidgetsNeedingUpdate->GetRaw())
 			{
-				FWidgetProxy& WidgetProxy = FastWidgetPathList[WidgetIndex];
-				if (WidgetProxy.Widget && WidgetProxy.Widget->FastPathProxyHandle.GetInvalidationRoot() == this) // If the Widget is no longer in the same Invalidation Root, don't bother adding it.
+				FSlateInvalidationWidgetList::InvalidationWidgetType& InvalidationWidget = (*FastWidgetPathList)[Element.Get<0>()];
+				if (SWidget* Widget = InvalidationWidget.GetWidget())
 				{
-					WidgetsNeedingUpdateCache.Emplace(WidgetProxy.Widget, WidgetProxy.CurrentInvalidateReason, WidgetProxy.UpdateFlags);
+					FastPathWidgetsNeedingUpdateCache.Add(Widget);
 				}
 			}
 
-			for(int32 WidgetIndex : WidgetsNeedingUpdate.GetRawData())
+			for (FSlateInvalidationWidgetIndex WidgetIndex : FinalUpdateList)
 			{
-				FWidgetProxy& WidgetProxy = FastWidgetPathList[WidgetIndex];
-				if (WidgetProxy.Widget)
+				FSlateInvalidationWidgetList::InvalidationWidgetType& InvalidationWidget = (*FastWidgetPathList)[WidgetIndex];
+				if (SWidget* Widget = InvalidationWidget.GetWidget())
 				{
-					check(WidgetProxy.Widget->FastPathProxyHandle.GetInvalidationRoot() == this);
-					WidgetsNeedingUpdateCache.Emplace(WidgetProxy.Widget, WidgetProxy.CurrentInvalidateReason, WidgetProxy.UpdateFlags);
+					FastPathUpdateListCache.Add(Widget);
 				}
 			}
 
-			SWidget* RootWidget = InvalidationRootWidget->Advanced_IsWindow() ? InvalidationRootWidget : &(*InvalidationRootWidget->GetAllChildren()->GetChildAt(0));
-			BuildFastPathList(RootWidget);
+			ProcessChildOrderUpdate();
+			WidgetsNeedingUpdate->Reset(false);
 
-			for (FWidgetNeedingUpdate& WidgetNeedingUpdate : WidgetsNeedingUpdateCache)
+			for (SWidget* Widget : FastPathWidgetsNeedingUpdateCache)
 			{
-				const int32 NewIndex = WidgetNeedingUpdate.Widget->FastPathProxyHandle.GetIndex();
-				FWidgetProxy& WidgetProxy = FastWidgetPathList[NewIndex];
-				check(WidgetProxy.Widget == WidgetNeedingUpdate.Widget);
-				WidgetProxy.CurrentInvalidateReason = WidgetNeedingUpdate.CurrentInvalidateReason;
-				WidgetProxy.UpdateFlags = WidgetNeedingUpdate.UpdateFlags;
-				WidgetProxy.bInUpdateList = true;
-				WidgetsNeedingUpdate.Push(WidgetProxy);
+				if (Widget->GetProxyHandle().IsValid(Widget) && Widget->GetProxyHandle().GetInvalidationRoot() == this)
+				{
+					WidgetsNeedingUpdate->ForcePush(Widget->GetProxyHandle().GetWidgetIndex());
+				}
 			}
 
-			bChildOrderInvalidated = false;
+			for (SWidget* Widget : FastPathUpdateListCache)
+			{
+				if (Widget->GetProxyHandle().IsValid(Widget) && Widget->GetProxyHandle().GetInvalidationRoot() == this)
+				{
+					WidgetsNeedingUpdate->PushUnique(Widget->GetProxyHandle().GetWidgetIndex());
+				}
+			}
 		}
-		else if(FinalUpdateList.Num() != 0)
+		else if (FinalUpdateList.Num() != 0)
 		{
 			// Put Widget waiting for update back in WidgetsNeedingUpdate to ensure index order and just in case, Prepass need to be reexecuted.
-			for (int32 WidgetIndex : FinalUpdateList)
+			for (FSlateInvalidationWidgetIndex WidgetIndex : FinalUpdateList)
 			{
-				FWidgetProxy& WidgetProxy = FastWidgetPathList[WidgetIndex];
-				WidgetsNeedingUpdate.Push(WidgetProxy);
+				WidgetsNeedingUpdate->PushUnique(WidgetIndex);
 			}
 		}
-		FinalUpdateList.Reset(WidgetsNeedingUpdate.Num());
+		FinalUpdateList.Reset(WidgetsNeedingUpdate->Num());
+
+#if WITH_SLATE_DEBUGGING
+		if (GSlateInvalidationRootVerifyValidWidgets)
+		{
+			ensureMsgf(FastWidgetPathList->VerifyProxiesWidget(), TEXT("We failed to verify that every WidgetProxy has a valid SWidget"));
+		}
+#endif
 
 #if SLATE_CSV_TRACKER
-		FCsvProfiler::RecordCustomStat("Invalidate/InitialWidgets", CSV_CATEGORY_INDEX(Slate), WidgetsNeedingUpdate.Num(), ECsvCustomStatOp::Set);
+		FCsvProfiler::RecordCustomStat("Invalidate/InitialWidgets", CSV_CATEGORY_INDEX(Slate), WidgetsNeedingUpdate->Num(), ECsvCustomStatOp::Set);
 		int32 Stat_TotalWidgetsInvalidated = 0;
 		int32 Stat_NeedsRepaint = 0;
 		int32 Stat_NeedsVolatilePaint = 0;
@@ -655,15 +561,15 @@ bool FSlateInvalidationRoot::ProcessInvalidation()
 		int32 Stat_NeedsActiveTimerUpdate = 0;
 #endif
 
-		while (WidgetsNeedingUpdate.Num() && !bNeedsSlowPath)
+		while (WidgetsNeedingUpdate->Num() && !bNeedsSlowPath)
 		{
 #if SLATE_CSV_TRACKER
 			Stat_TotalWidgetsInvalidated++;
 #endif
 
-			int32 MyIndex = WidgetsNeedingUpdate.Pop();
+			FSlateInvalidationWidgetIndex MyIndex = WidgetsNeedingUpdate->Pop();
 			FinalUpdateList.Add(MyIndex);
-			FWidgetProxy& WidgetProxy = FastWidgetPathList[MyIndex];
+			FWidgetProxy& WidgetProxy = (*FastWidgetPathList)[MyIndex];
 
 			// Reset each widgets paint state
 			// Must be done before actual painting because children can repaint 
@@ -671,37 +577,37 @@ bool FSlateInvalidationRoot::ProcessInvalidation()
 			WidgetProxy.bInUpdateList = false;
 
 			// Widget could be null if it was removed and we are on the slow path
-			if (WidgetProxy.Widget)
+			if (SWidget* WidgetPtr = WidgetProxy.GetWidget())
 			{
-				if (!GSlateEnableGlobalInvalidation && !InvalidationRootWidget->NeedsPrepass() && WidgetProxy.Widget->Advanced_IsInvalidationRoot())
+				if (!GSlateEnableGlobalInvalidation && !InvalidationRootWidget->NeedsPrepass() && WidgetPtr->Advanced_IsInvalidationRoot())
 				{
 					WidgetProxy.CurrentInvalidateReason |= EInvalidateWidgetReason::Layout;
 #if WITH_SLATE_DEBUGGING
-					FSlateDebugging::BroadcastWidgetInvalidate(WidgetProxy.Widget, nullptr, EInvalidateWidgetReason::Layout);
+					FSlateDebugging::BroadcastWidgetInvalidate(WidgetPtr, nullptr, EInvalidateWidgetReason::Layout);
 #endif
-					UE_TRACE_SLATE_WIDGET_INVALIDATED(WidgetProxy.Widget, nullptr, EInvalidateWidgetReason::Layout);
+					UE_TRACE_SLATE_WIDGET_INVALIDATED(WidgetPtr, nullptr, EInvalidateWidgetReason::Layout);
 				}
 
 #if SLATE_CSV_TRACKER
-				const int32 PreviousWidgetsNeedingUpdating = WidgetsNeedingUpdate.Num();
+				const int32 PreviousWidgetsNeedingUpdating = WidgetsNeedingUpdate->Num();
 #endif
 
-				bWidgetsNeedRepaint |= WidgetProxy.ProcessInvalidation(WidgetsNeedingUpdate, FastWidgetPathList, *this);
+				bWidgetsNeedRepaint |= WidgetProxy.ProcessInvalidation(*WidgetsNeedingUpdate, *FastWidgetPathList, *this);
 
 #if SLATE_CSV_TRACKER
-				const int32 CurrentWidgetsNeedingUpdating = WidgetsNeedingUpdate.Num();
+				const int32 CurrentWidgetsNeedingUpdating = WidgetsNeedingUpdate->Num();
 				const int32 AddedWidgets = CurrentWidgetsNeedingUpdating - PreviousWidgetsNeedingUpdating;
 
 				if (AddedWidgets >= CascadeInvalidationEventAmount)
 				{
-					CSV_EVENT(Slate, TEXT("Invalidated %s"), *FReflectionMetaData::GetWidgetDebugInfo(WidgetProxy.Widget));
+					CSV_EVENT(Slate, TEXT("Invalidated %s"), *FReflectionMetaData::GetWidgetDebugInfo(WidgetPtr));
 				}
 
 				if (EnumHasAnyFlags(WidgetProxy.UpdateFlags, EWidgetUpdateFlags::NeedsRepaint))
 				{
 					Stat_NeedsRepaint++;
 				}
-				if (EnumHasAnyFlags(WidgetProxy.UpdateFlags, EWidgetUpdateFlags::NeedsVolatilePaint) && !WidgetProxy.Widget->Advanced_IsInvalidationRoot())
+				if (EnumHasAnyFlags(WidgetProxy.UpdateFlags, EWidgetUpdateFlags::NeedsVolatilePaint) && !WidgetPtr->Advanced_IsInvalidationRoot())
 				{
 					Stat_NeedsVolatilePaint++;
 				}
@@ -717,7 +623,7 @@ bool FSlateInvalidationRoot::ProcessInvalidation()
 			}
 		}
 		
-		WidgetsNeedingUpdate.Reset();
+		WidgetsNeedingUpdate->Reset(true);
 
 #if SLATE_CSV_TRACKER
 		FCsvProfiler::RecordCustomStat("Invalidate/TotalWidgets", CSV_CATEGORY_INDEX(Slate), Stat_TotalWidgetsInvalidated, ECsvCustomStatOp::Set);
@@ -737,24 +643,23 @@ bool FSlateInvalidationRoot::ProcessInvalidation()
 
 void FSlateInvalidationRoot::ClearAllFastPathData(bool bClearResourcesImmediately)
 {
-	for (const FWidgetProxy& Proxy : FastWidgetPathList)
-	{
-		if (SWidget* Widget = Proxy.Widget)
+
+	FastWidgetPathList->Reset();
+	FastWidgetPathList->ForEachWidget([bClearResourcesImmediately](SWidget* Widget)
 		{
 			Widget->PersistentState.CachedElementHandle = FSlateCachedElementsHandle::Invalid;
 			if (bClearResourcesImmediately)
 			{
 				Widget->FastPathProxyHandle = FWidgetProxyHandle();
 			}
-		}
-	}
+		});
 
 #if UE_SLATE_DEBUGGING_CLEAR_ALL_FAST_PATH_DATA
 	if (!bClearResourcesImmediately)
 	{
 		for (const FWidgetProxy& Proxy : FastWidgetPathList)
 		{
-			if (SWidget* Widget = Proxy.Widget)
+			if (SWidget* Widget = Proxy.GetWidget())
 			{
 				if (Widget->FastPathProxyHandle.IsValid())
 				{
@@ -767,18 +672,18 @@ void FSlateInvalidationRoot::ClearAllFastPathData(bool bClearResourcesImmediatel
 	{
 		for (const FWidgetProxy& Proxy : FastWidgetPathList)
 		{
-			FastWidgetPathToClearedBecauseOfDelay.RemoveSingleSwap(Proxy.Widget);
+			FastWidgetPathToClearedBecauseOfDelay.RemoveSingleSwap(Proxy.GetWidget());
 		}
 	}
 #endif
 
 
-	if (FastWidgetPathList.Num() != 0)
+	if (!FastWidgetPathList->IsEmpty())
 	{
 		ClearUpdateList.RemoveSingleSwap(this, false);
 	}
-	FastWidgetPathList.Empty();
-	WidgetsNeedingUpdate.Empty();
+	WidgetsNeedingUpdate->Reset(false);
+	FastWidgetPathList->Empty();
 	CachedElementData->Empty();
 	FinalUpdateList.Empty();
 }
@@ -791,7 +696,7 @@ void FSlateInvalidationRoot::HandleInvalidateAllWidgets(bool bClearResourcesImme
 
 void FSlateInvalidationRoot::Advanced_ResetInvalidation(bool bClearResourcesImmediately)
 {
-	InvalidateChildOrder();
+	InvalidateRoot();
 
 	InvalidationRootWidget->InvalidatePrepass();
 
@@ -803,6 +708,102 @@ void FSlateInvalidationRoot::Advanced_ResetInvalidation(bool bClearResourcesImme
 	bNeedsSlowPath = true;
 }
 
+#if WITH_SLATE_DEBUGGING
+void VerifyWidgetList(SWidget* InvalidationRootWidget, FSlateInvalidationRootHandle InvalidationRootHandle, FSlateInvalidationWidgetList& WidgetList)
+{
+	TSharedRef<SWidget> RootWidget = InvalidationRootWidget->Advanced_IsWindow() ? InvalidationRootWidget->AsShared() : InvalidationRootWidget->GetAllChildren()->GetChildAt(0);
+
+	FSlateInvalidationWidgetList List(InvalidationRootHandle, FSlateInvalidationWidgetList::FArguments{ 128, 128, 1000, false });
+	List.BuildWidgetList(RootWidget);
+	bool bIsIdentical = (List.DeapCompare(WidgetList));
+	if (!bIsIdentical)
+	{
+		UE_LOG(LogSlate, Log, TEXT("**-- New Build List --**"));
+		List.LogWidgetsList();
+		UE_LOG(LogSlate, Log, TEXT("**-- Invaliation Root List --**"));
+		WidgetList.LogWidgetsList();
+
+		ensureMsgf(false, TEXT("The updated list doesn't match a newly created list."));
+	}
+}
+
+void VerifyHittest(SWidget* InvalidationRootWidget, FSlateInvalidationWidgetList& WidgetList, FHittestGrid* HittestGrid)
+{
+	check(InvalidationRootWidget);
+	check(HittestGrid);
+
+	ensureAlwaysMsgf(WidgetList.VerifySortOrder()
+		, TEXT("The array's sort order for InvalidationRoot '%d' is not respected.")
+		, *FReflectionMetaData::GetWidgetPath(InvalidationRootWidget));
+
+	TArray<FHittestGrid::FWidgetSortData> WeakHittestGridSortDatas = HittestGrid->GetAllWidgetSortDatas();
+
+	struct FHittestWidgetSortData
+	{
+		const SWidget* Widget;
+		int64 PrimarySort;
+		FSlateInvalidationWidgetSortOrder SecondarySort;
+	};
+
+	FMemMark Mark(FMemStack::Get());
+	TArray<FHittestWidgetSortData, TMemStackAllocator<>> HittestGridSortDatas;
+	HittestGridSortDatas.Reserve(WeakHittestGridSortDatas.Num());
+
+	// Widgets need to be valid in the hittestgrid
+	for (const FHittestGrid::FWidgetSortData& Data : WeakHittestGridSortDatas)
+	{
+		TSharedPtr<SWidget> Widget = Data.WeakWidget.Pin();
+		if (ensureAlwaysMsgf(Widget, TEXT("A widget is invalid in the HittestGrid")))
+		{
+			FHittestWidgetSortData SortData = { Widget.Get(), Data.PrimarySort, Data.SecondarySort };
+			HittestGridSortDatas.Add(MoveTemp(SortData));
+		}
+	}
+
+	// The order in the WidgetList is sorted. It's not the case of the HittestGrid.
+
+	FSlateInvalidationWidgetSortOrder PreviousSecondarySort;
+	const SWidget* LastWidget = nullptr;
+	WidgetList.ForEachWidget([&HittestGridSortDatas, &PreviousSecondarySort, &LastWidget](const SWidget* Widget)
+		{
+			if (Widget->GetVisibility().IsHitTestVisible())
+			{
+				// Is the widget in the hittestgrid
+				const int32 FoundHittestIndex = HittestGridSortDatas.IndexOfByPredicate([Widget](const FHittestWidgetSortData& HittestGrid)
+					{
+						return HittestGrid.Widget == Widget;
+					});
+				const bool bHasFoundWidget = HittestGridSortDatas.IsValidIndex(FoundHittestIndex);
+				if (!bHasFoundWidget)
+				{
+					return;
+				}
+
+				ensureAlwaysMsgf(Widget->GetProxyHandle().GetWidgetSortOrder() == HittestGridSortDatas[FoundHittestIndex].SecondarySort
+					, TEXT("The SecondarySort of widget '%s' doesn't match the SecondarySort inside the hittestgrid.")
+					, *FReflectionMetaData::GetWidgetPath(Widget));
+
+				LastWidget = Widget;
+				PreviousSecondarySort = HittestGridSortDatas[FoundHittestIndex].SecondarySort;
+
+				HittestGridSortDatas.RemoveAtSwap(FoundHittestIndex);
+			}
+		});
+
+	const int32 FoundHittestIndex = HittestGridSortDatas.IndexOfByPredicate([InvalidationRootWidget](const FHittestWidgetSortData& HittestGrid)
+		{
+			return HittestGrid.Widget == InvalidationRootWidget;
+		});
+	if (HittestGridSortDatas.IsValidIndex(FoundHittestIndex))
+	{
+		HittestGridSortDatas.RemoveAtSwap(FoundHittestIndex);
+	}
+
+	ensureAlwaysMsgf(HittestGridSortDatas.Num() == 0, TEXT("The hittest grid of Root '%s' has widget that are not inside the InvalidationRoot's widget list")
+		, *FReflectionMetaData::GetWidgetPath(InvalidationRootWidget));
+}
+#endif //WITH_SLATE_DEBUGGING
+
 /**
  * 
  */
@@ -823,3 +824,4 @@ FSlateInvalidationRoot* FSlateInvalidationRootHandle::GetInvalidationRoot() cons
 {
 	return GSlateInvalidationRootListInstance.GetInvalidationRoot(UniqueId);
 }
+
diff --git a/Engine/Source/Runtime/SlateCore/Private/FastUpdate/SlateInvalidationWidgetHeap.h b/Engine/Source/Runtime/SlateCore/Private/FastUpdate/SlateInvalidationWidgetHeap.h
new file mode 100644
--- /dev/null
+++ b/Engine/Source/Runtime/SlateCore/Private/FastUpdate/SlateInvalidationWidgetHeap.h
@@ -0,0 +1,113 @@
+// Copyright Epic Games, Inc. All Rights Reserved.
+
+#pragma once
+
+#include "CoreMinimal.h"
+
+#include "FastUpdate/SlateInvalidationWidgetIndex.h"
+#include "FastUpdate/SlateInvalidationWidgetList.h"
+#include "FastUpdate/SlateInvalidationWidgetSortOrder.h"
+
+/**
+ * Ordered list of WidgetIndex. The order is based on the WidgetSortIndex.
+ */
+class FSlateInvalidationWidgetHeap
+{
+public:
+	using FElement = TPair<FSlateInvalidationWidgetIndex, FSlateInvalidationWidgetSortOrder>;
+	struct TWidgetOrderGreater
+	{
+		FORCEINLINE bool operator()(const FElement& A, const FElement& B) const
+		{
+			return B.Get<1>() < A.Get<1>();
+		}
+	};
+
+public:
+	FSlateInvalidationWidgetHeap(FSlateInvalidationWidgetList& InOwnerList)
+		: OwnerList(InOwnerList)
+	{ }
+
+	/** Insert into the list at the proper order (see binary heap) only if it's not already contains by the list. */
+	void PushUnique(const FSlateInvalidationWidgetIndex WidgetIndex)
+	{
+		FSlateInvalidationWidgetList::InvalidationWidgetType& InvalidationWidget = OwnerList[WidgetIndex];
+		check(Contains(WidgetIndex) == InvalidationWidget.bContainedByWidgetHeap);
+		if (!InvalidationWidget.bContainedByWidgetHeap)
+		{
+			InvalidationWidget.bContainedByWidgetHeap = true;
+			InvalidationWidget.bInUpdateList = true;
+			Heap.HeapPush(FElement{ WidgetIndex, FSlateInvalidationWidgetSortOrder{OwnerList, WidgetIndex} }, TWidgetOrderGreater());
+		}
+	}
+
+	/** Insert into the list at the proper order (see binary heap) only if it's not already contains by the list. */
+	void PushUnique(FSlateInvalidationWidgetList::InvalidationWidgetType& InvalidationWidget)
+	{
+		check(Contains(InvalidationWidget.Index) == InvalidationWidget.bContainedByWidgetHeap);
+
+		if (!InvalidationWidget.bContainedByWidgetHeap)
+		{
+			InvalidationWidget.bContainedByWidgetHeap = true;
+			InvalidationWidget.bInUpdateList = true;
+			Heap.HeapPush(FElement{ InvalidationWidget.Index, FSlateInvalidationWidgetSortOrder{OwnerList, InvalidationWidget.Index} }, TWidgetOrderGreater());
+		}
+	}
+
+	/** Insert into the list at the proper order (see binary heap). */
+	void ForcePush(const FSlateInvalidationWidgetIndex WidgetIndex)
+	{
+		FSlateInvalidationWidgetList::InvalidationWidgetType& InvalidationWidget = OwnerList[WidgetIndex];
+
+		InvalidationWidget.bContainedByWidgetHeap = true;
+		InvalidationWidget.bInUpdateList = true;
+		Heap.HeapPush(FElement{ WidgetIndex, FSlateInvalidationWidgetSortOrder{OwnerList, WidgetIndex} }, TWidgetOrderGreater());
+	}
+
+	/** Returns and removes the biggest WidgetIndex from the list. */
+	FSlateInvalidationWidgetIndex Pop()
+	{
+		FSlateInvalidationWidgetIndex Result = Heap.HeapTop().Get<0>();
+		Heap.HeapPopDiscard(TWidgetOrderGreater(), false);
+		OwnerList[Result].bContainedByWidgetHeap = false;
+		return Result;
+	}
+
+	/** Empties the list, but doesn't change memory allocations. */
+	inline void Reset(bool bResetContained)
+	{
+		if (bResetContained)
+		{
+			for (const FElement& Element : Heap)
+			{
+				OwnerList[Element.Get<0>()].bContainedByWidgetHeap = false;
+			}
+		}
+		Heap.Reset();
+	}
+	
+	/** Returns the number of elements in the list. */
+	inline int32 Num() const
+	{
+		return Heap.Num();
+	}
+
+	/** Returns the raw list. */
+	const TArray<FElement, TInlineAllocator<100>>& GetRaw() const
+	{
+		return Heap;
+	}
+
+private:
+	bool Contains(const FSlateInvalidationWidgetIndex WidgetIndex) const
+	{
+		return Heap.ContainsByPredicate([WidgetIndex](const FElement& Element)
+			{
+				return Element.Get<0>() == WidgetIndex;
+			});
+	}
+
+private:
+	TArray<FElement, TInlineAllocator<100>> Heap;
+	FSlateInvalidationWidgetList& OwnerList;
+};
\ No newline at end of file
diff --git a/Engine/Source/Runtime/SlateCore/Private/FastUpdate/SlateInvalidationWidgetIndex.cpp b/Engine/Source/Runtime/SlateCore/Private/FastUpdate/SlateInvalidationWidgetIndex.cpp
new file mode 100644
--- /dev/null
+++ b/Engine/Source/Runtime/SlateCore/Private/FastUpdate/SlateInvalidationWidgetIndex.cpp
@@ -0,0 +1,8 @@
+// Copyright Epic Games, Inc. All Rights Reserved.
+
+#include "FastUpdate/SlateInvalidationWidgetIndex.h"
+
+#include <limits>
+
+const FSlateInvalidationWidgetIndex FSlateInvalidationWidgetIndex::Invalid = { std::numeric_limits<FSlateInvalidationWidgetIndex::IndexType>::max(), std::numeric_limits<FSlateInvalidationWidgetIndex::IndexType>::max() };
+
diff --git a/Engine/Source/Runtime/SlateCore/Private/FastUpdate/SlateInvalidationWidgetList.cpp b/Engine/Source/Runtime/SlateCore/Private/FastUpdate/SlateInvalidationWidgetList.cpp
new file mode 100644
--- /dev/null
+++ b/Engine/Source/Runtime/SlateCore/Private/FastUpdate/SlateInvalidationWidgetList.cpp
@@ -0,0 +1,1367 @@
+// Copyright Epic Games, Inc. All Rights Reserved.
+
+#pragma once
+
+#include "FastUpdate/SlateInvalidationWidgetList.h"
+
+#include "FastUpdate/SlateInvalidationRoot.h"
+#include "FastUpdate/SlateInvalidationRootHandle.h"
+#include "FastUpdate/WidgetProxy.h"
+
+#include "Algo/Unique.h"
+#include "Layout/Children.h"
+#include "Misc/StringBuilder.h"
+#include "Rendering/SlateLayoutTransform.h"
+#include "Templates/ChooseClass.h"
+#include "Templates/IsConst.h"
+#include "Templates/Tuple.h"
+#include "Types/ReflectionMetadata.h"
+
+#include <limits>
+
+
+DECLARE_CYCLE_STAT(TEXT("WidgetList ProcessInvalidation"), STAT_WidgetList_ProcessInvalidation, STATGROUP_Slate);
+
+#define UE_SLATE_WITH_WIDGETLIST_UPDATEONLYWHATISNEEDED 0
+
+#define UE_SLATE_VERIFY_REBUILDWIDGETDATA_ORDER 0
+#define UE_SLATE_VERIFY_INVALID_INVALIDATIONHANDLE 0
+#define UE_SLATE_VERIFY_REMOVED_WIDGET_ARE_NOT_INVALIDATED 0
+
+#if UE_SLATE_VERIFY_REMOVED_WIDGET_ARE_NOT_INVALIDATED
+uint16 GSlateInvalidationWidgetIndex_RemovedIndex = 0xffee;
+#endif
+
+// See FSlateInvalidationWidgetSortOrder::FSlateInvalidationWidgetSortOrder
+int32 FSlateInvalidationWidgetList::FArguments::MaxPreferedElementsNum = (1 << 10) - 1;
+int32 FSlateInvalidationWidgetList::FArguments::MaxSortOrderPaddingBetweenArray = (1 << 22) - 1;
+
+
+FSlateInvalidationWidgetList::FSlateInvalidationWidgetList(FSlateInvalidationRootHandle InOwner, const FArguments& Args)
+	: Owner(InOwner)
+	, WidgetListConfig(Args)
+{
+	if (WidgetListConfig.PreferedElementsNum <= 1
+		|| WidgetListConfig.PreferedElementsNum > FArguments::MaxPreferedElementsNum
+		 || WidgetListConfig.SortOrderPaddingBetweenArray <= WidgetListConfig.PreferedElementsNum
+		 || WidgetListConfig.SortOrderPaddingBetweenArray > FArguments::MaxSortOrderPaddingBetweenArray)
+	{
+		ensureMsgf(false, TEXT("The PreferedElementsNum or SortOrderPaddingBetweenArray have incorrect values. '%d,%d'. Reset to default value.")
+			, WidgetListConfig.PreferedElementsNum
+			, WidgetListConfig.SortOrderPaddingBetweenArray);
+		const_cast<FArguments&>(WidgetListConfig).PreferedElementsNum = FArguments().PreferedElementsNum;
+		const_cast<FArguments&>(WidgetListConfig).SortOrderPaddingBetweenArray = FArguments().SortOrderPaddingBetweenArray;
+	}
+}
+
+
+FSlateInvalidationWidgetIndex FSlateInvalidationWidgetList::_BuildWidgetList_Recursive(TSharedRef<SWidget>& Widget, FSlateInvalidationWidgetIndex ParentIndex, IndexType& LastestIndex, bool bParentVisible, bool bParentVolatile)
+{
+	const bool bDoRecursion = ShouldDoRecursion(Widget);
+	IndexType TemporaryIndex = LastestIndex;
+	const FSlateInvalidationWidgetIndex NewIndex = EmplaceInsertAfter(LastestIndex, Widget);
+	LastestIndex = NewIndex.ArrayIndex;
+
+	FSlateInvalidationWidgetIndex LeafMostChildIndex = NewIndex;
+	const EVisibility Visibility = Widget->GetVisibility();
+	const bool bParentAndSelfVisible = bParentVisible && Visibility.IsVisible();
+
+	{
+		InvalidationWidgetType& WidgetProxy = (*this)[NewIndex];
+		WidgetProxy.Index = NewIndex;
+		WidgetProxy.ParentIndex = ParentIndex;
+		WidgetProxy.LeafMostChildIndex = LeafMostChildIndex;
+		WidgetProxy.Visibility = Visibility;
+		WidgetProxy.bInvisibleDueToParentOrSelfVisibility = !bParentAndSelfVisible;
+	}
+
+#if WITH_SLATE_DEBUGGING
+	if (WidgetListConfig.bAssignedWidgetIndex)
+#endif
+	{
+		const FSlateInvalidationWidgetSortOrder SortIndex = { *this, NewIndex };
+		Widget->SetFastPathProxyHandle(FWidgetProxyHandle{ Owner, NewIndex, SortIndex }, !bParentAndSelfVisible, bParentVolatile);
+	}
+
+	const bool bParentOrSelfVolatile = bParentVolatile || Widget->IsVolatile();
+
+	if (bDoRecursion)
+	{
+		FChildren* Children = Widget->GetAllChildren();
+		int32 NumChildren = Children->Num();
+		for (int32 Index = 0; Index < NumChildren; ++Index)
+		{
+			TSharedRef<SWidget> NewWidget = Children->GetChildAt(Index);
+
+			const bool bShouldAdd = ShouldBeAdded(NewWidget);
+			if (bShouldAdd)
+			{
+				check(NewWidget->GetParentWidget() == Widget);
+				LeafMostChildIndex = _BuildWidgetList_Recursive(NewWidget, NewIndex, LastestIndex, bParentAndSelfVisible, bParentOrSelfVolatile);
+			}
+		}
+
+		InvalidationWidgetType& WidgetProxy = (*this)[NewIndex];
+		WidgetProxy.LeafMostChildIndex = LeafMostChildIndex;
+	}
+
+	return LeafMostChildIndex;
+}
+
+
+void FSlateInvalidationWidgetList::BuildWidgetList(TSharedRef<SWidget> InRoot)
+{
+	Reset();
+	Root = InRoot;
+	const bool bShouldAdd = ShouldBeAdded(InRoot);
+	if (bShouldAdd)
+	{
+		const TSharedPtr<SWidget> Parent = InRoot->GetParentWidget();
+		const bool bParentVisible = Parent.IsValid() ? Parent->GetVisibility().IsVisible() : true;
+		const bool bParentVolatile = false;
+		IndexType LatestArrayIndex = FSlateInvalidationWidgetIndex::Invalid.ArrayIndex;
+		_BuildWidgetList_Recursive(InRoot, FSlateInvalidationWidgetIndex::Invalid, LatestArrayIndex, bParentVisible, bParentVolatile);
+	}
+}
+
+
+void FSlateInvalidationWidgetList::_RebuildWidgetListTree(TSharedRef<SWidget> Widget, int32 ChildAtIndex)
+{
+	const bool bShouldAddWidget = ShouldBeAdded(Widget);
+	FChildren* ParentChildren = Widget->GetAllChildren();
+	const FSlateInvalidationWidgetIndex WidgetIndex = Widget->GetProxyHandle().GetWidgetIndex();
+	if (bShouldAddWidget && ChildAtIndex < ParentChildren->Num() && WidgetIndex != FSlateInvalidationWidgetIndex::Invalid)
+	{
+		// Since we are going to add item, the array may get invalidated. Do not use after _BuildWidgetList_Recursive
+		ensure((*this)[WidgetIndex].GetWidget() == &Widget.Get());
+		FSlateInvalidationWidgetIndex PreviousLeafIndex = (*this)[WidgetIndex].LeafMostChildIndex;
+		FSlateInvalidationWidgetIndex NewLeafIndex = PreviousLeafIndex;
+		IndexType LastestArrayIndex = PreviousLeafIndex.ArrayIndex;
+		const bool bDoRecursion = ShouldDoRecursion(Widget);
+		if (bDoRecursion)
+		{
+			const bool bParentVisible = Widget->GetVisibility().IsVisible();
+			const bool bParentVolatile = Widget->IsVolatile();
+			int32 NumChildren = ParentChildren->Num();
+			for (int32 Index = ChildAtIndex; Index < NumChildren; ++Index)
+			{
+				TSharedRef<SWidget> ChildWidget = ParentChildren->GetChildAt(Index);
+				const bool bShouldAddChild = ShouldBeAdded(ChildWidget);
+				if (bShouldAddChild)
+				{
+					check(ChildWidget->GetParentWidget() == Widget);
+					NewLeafIndex = _BuildWidgetList_Recursive(ChildWidget, WidgetIndex, LastestArrayIndex, bParentVisible, bParentVolatile);
+				}
+			}
+		}
+
+		if (NewLeafIndex != PreviousLeafIndex)
+		{
+			InvalidationWidgetType& InvalidationWidget = (*this)[WidgetIndex];
+			InvalidationWidget.LeafMostChildIndex = NewLeafIndex;
+			UpdateParentLeafIndex(InvalidationWidget, PreviousLeafIndex, NewLeafIndex);
+		}
+	}
+}
+
+
+//test if it's FSlateInvalidationWidgetIndex is valid and this == Widget.
+//when we remove a widget we may not invalidate it's index and it may not point to the correct widget anymore
+//That behavior is normal, if it occured during ProcessChildOrderInvalidation.
+void FSlateInvalidationWidgetList::ProcessChildOrderInvalidation(const TArray<TWeakPtr<SWidget>>& InvalidatedWidgets)
+{
+	SCOPE_CYCLE_COUNTER(STAT_WidgetList_ProcessInvalidation);
+
+	if (InvalidatedWidgets.Num() > 0)
+	{
+		if (FirstIndex() == FSlateInvalidationWidgetIndex::Invalid)
+		{
+			ensureMsgf(false, TEXT("No tree were built and widgets were invalidated."));
+			return;
+		}
+
+		FMemMark Mark(FMemStack::Get());
+		using TInvalidateWidgetIndexType = TTuple<TSharedRef<SWidget>, FSlateInvalidationWidgetIndex, FSlateInvalidationWidgetSortOrder>;
+		TArray<TInvalidateWidgetIndexType> InvalidatedWidgetIndexes;
+		InvalidatedWidgetIndexes.Reserve(InvalidatedWidgets.Num());
+
+#if UE_SLATE_VERIFY_INVALID_INVALIDATIONHANDLE
+		TArray<TWeakPtr<SWidget>> VerifyWidgetInvalidationHandle;
+		InvalidatedWidgetIndexes.Reserve(InvalidatedWidgets.Num());
+#endif
+
+		// Build Invalidation Indexes
+		{
+			for (const TWeakPtr<SWidget>& InvalidatedWidget : InvalidatedWidgets)
+			{
+				if (const TSharedPtr<SWidget> Widget = InvalidatedWidget.Pin())
+				{
+					const FSlateInvalidationWidgetIndex WidgetIndex = Widget->GetProxyHandle().GetWidgetIndex();
+					if (IsValidIndex(WidgetIndex) && (*this)[WidgetIndex].GetWidget() == Widget.Get())
+					{
+						InvalidatedWidgetIndexes.Emplace(Widget.ToSharedRef(), WidgetIndex, FSlateInvalidationWidgetSortOrder{ *this, WidgetIndex });
+					}
+#if UE_SLATE_VERIFY_INVALID_INVALIDATIONHANDLE
+					else
+					{
+						// This widget has requested an invalidation ChildOrder but it's index is not valid.
+						//Confirm that it will have a valid index at the end of the algo.
+						VerifyWidgetInvalidationHandle.Add(InvalidatedWidget);
+						Widget->SetFastPathProxyHandle(FWidgetProxyHandle{ Owner, FSlateInvalidationWidgetIndex::Invalid, FSlateInvalidationWidgetSortOrder{} });
+					}
+#endif
+				}
+			}
+
+			// Sort to invalidate them in the same order they would be if we were in slow path.
+			InvalidatedWidgetIndexes.Sort([](const TInvalidateWidgetIndexType& A, const TInvalidateWidgetIndexType& B) {
+				const FSlateInvalidationWidgetSortOrder OrderA = A.Get<2>();
+				const FSlateInvalidationWidgetSortOrder OrderB = B.Get<2>();
+				return OrderA < OrderB;
+				});
+
+			const int32 InvalidatedNum = Algo::Unique(InvalidatedWidgetIndexes);
+			if (InvalidatedNum != InvalidatedWidgetIndexes.Num())
+			{
+				InvalidatedWidgetIndexes.RemoveAt(InvalidatedNum, InvalidatedWidgetIndexes.Num() - InvalidatedNum, false);
+			}
+		}
+
+		struct FChildOrderInvalidationData
+		{
+			TSharedPtr<SWidget> Widget;
+			FIndexRange Range;
+			FSlateInvalidationWidgetIndex WhereToCut;
+			int32 ChildAtToStartWith;
+			bool bRemove;
+			bool operator==(const FChildOrderInvalidationData& Other) const
+			{
+				return Other.Widget == Widget
+					&& Other.bRemove == bRemove
+					&& (bRemove ? Other.Range == Range : Other.WhereToCut == WhereToCut)
+					&& Other.ChildAtToStartWith == ChildAtToStartWith;
+			}
+			FChildOrderInvalidationData(TSharedPtr<SWidget> InWidget, FIndexRange InRange, int32 InChildAtToStartWith)
+				: Widget(InWidget)
+				, Range(InRange)
+				, WhereToCut(FSlateInvalidationWidgetIndex::Invalid)
+				, ChildAtToStartWith(InChildAtToStartWith)
+				, bRemove(true) {}
+			FChildOrderInvalidationData(TSharedPtr<SWidget> InWidget, FSlateInvalidationWidgetIndex InWhereToCut, int32 InChildAtToStartWith)
+				: Widget(InWidget)
+				, Range()
+				, WhereToCut(InWhereToCut)
+				, ChildAtToStartWith(InChildAtToStartWith)
+				, bRemove(false) {}
+		};
+
+		TArray<FChildOrderInvalidationData, TMemStackAllocator<>> RebuildWigetData;
+		RebuildWigetData.Reserve(InvalidatedWidgetIndexes.Num());
+
+
+		// Compute the InvalidationDatas
+		{
+			SCOPED_NAMED_EVENT(Slate_InvalidationList_ProcessCompute, FColorList::Blue);
+			TArray<FIndexRange, TMemStackAllocator<>> RebuildWidgetRange;
+			RebuildWidgetRange.Reserve(InvalidatedWidgetIndexes.Num());
+
+			for (int32 InvalidatedIndex = 0; InvalidatedIndex < InvalidatedWidgetIndexes.Num(); ++InvalidatedIndex)
+			{
+				FSlateInvalidationWidgetIndex WidgetIndex = InvalidatedWidgetIndexes[InvalidatedIndex].Get<1>();
+
+				// Is it already going to be invalidated by something else
+				{
+					const FSlateInvalidationWidgetSortOrder WidgetOrder = { *this, WidgetIndex };
+					if (RebuildWidgetRange.ContainsByPredicate([WidgetOrder](const FIndexRange& Range) {
+						return Range.Include(WidgetOrder);
+						}))
+					{
+						continue;
+					}
+				}
+
+				const InvalidationWidgetType& InvalidationWidget = (*this)[WidgetIndex];
+
+				// If the parent is invalid, it's means that we are the root. Rebuild the full tree.
+				if (InvalidationWidget.ParentIndex == FSlateInvalidationWidgetIndex::Invalid)
+				{
+					SWidget* WidgetPtr = InvalidationWidget.GetWidget();
+					if (ensure(WidgetPtr))
+					{
+						ensure(GetRoot().Pin().Get() == WidgetPtr);
+						UE_LOG(LogSlate, Warning, TEXT("Performance: A BuildTree() was requested by a ChildOrder invalidation."));
+						BuildWidgetList(WidgetPtr->AsShared());
+					}
+					return;
+				}
+
+				TSharedRef<SWidget> WidgetRef = InvalidatedWidgetIndexes[InvalidatedIndex].Get<0>();
+				ensureMsgf(WidgetRef->GetProxyHandle().GetWidgetIndex() == WidgetIndex, TEXT("The widget index doesn't match the index in the InvalidationWidgetList"));
+
+				// Find all the InvalidationWidget children
+				{
+					// Was added, but it should not be there anymore
+					if (!ShouldBeAdded(WidgetRef))
+					{
+						const FIndexRange Range = { *this, WidgetIndex, InvalidationWidget.LeafMostChildIndex };
+						RebuildWidgetRange.Add(Range);
+						RebuildWigetData.Emplace(TSharedPtr<SWidget>(), Range, INDEX_NONE);
+					}
+					// If it is not supposed to had child, but had some, them remove them
+					else if (!ShouldDoRecursion(WidgetRef))
+					{
+						if (WidgetIndex != InvalidationWidget.LeafMostChildIndex)
+						{
+							const FIndexRange Range = { *this, IncrementIndex(WidgetIndex), InvalidationWidget.LeafMostChildIndex };
+							RebuildWidgetRange.Add(Range);
+							RebuildWigetData.Emplace(TSharedPtr<SWidget>(), Range, INDEX_NONE);
+						}
+					}
+					// Find the difference between list and the reality (if it used to have at least 1 child)
+					else if (WidgetIndex != InvalidationWidget.LeafMostChildIndex)
+					{
+#if UE_SLATE_WITH_WIDGETLIST_UPDATEONLYWHATISNEEDED
+						// Find all it's previous children
+						FMemMark Mark2(FMemStack::Get());
+						TArray<SWidget*, TMemStackAllocator<>> PreviousChildrenWidget;
+						_FindChildren(WidgetIndex, PreviousChildrenWidget);
+
+						FChildren* InvalidatedChildren = WidgetRef->GetAllChildren();
+						check(InvalidatedChildren);
+
+						/**
+						 ****** ****** ****** ****** ******
+						 * todo @pat improve with a loop
+						 * As a fist step, invalidate everything from this point. Could be a loop and invalidate what need to be invalidated before move the other.
+						 * We have to see if that case occurred in data before doing the code for it.
+						 ****** ****** ****** ****** ******
+						*/
+
+						// Find where it starts to get different
+						FSlateInvalidationWidgetIndex IndexWhereToStart = IncrementIndex(WidgetIndex);
+						FSlateInvalidationWidgetIndex IndexWhereToEnd = InvalidationWidget.LeafMostChildIndex;
+						int32 InvalidatedChildrenIndex = 0;
+
+						const int32 InvalidatedChildrenNum = InvalidatedChildren->Num();
+						const int32 PreviousChildrenNum = PreviousChildrenWidget.Num();
+						int32 PreviousIndex = 0;
+						for (; InvalidatedChildrenIndex < InvalidatedChildrenNum && PreviousIndex < PreviousChildrenNum; ++InvalidatedChildrenIndex)
+						{
+							TSharedRef<SWidget> NewWidget = InvalidatedChildren->GetChildAt(InvalidatedChildrenIndex);
+							if (ShouldBeAdded(NewWidget))
+							{
+								if (&NewWidget.Get() != PreviousChildrenWidget[PreviousIndex])
+								{
+									break;
+								}
+								IndexWhereToStart = IncrementIndex((*this)[IndexWhereToStart].LeafMostChildIndex);
+								check(IsValidIndex(IndexWhereToStart));
+
+								++PreviousIndex;
+							}
+						}
+
+						if (InvalidatedChildrenIndex >= InvalidatedChildrenNum && PreviousIndex >= PreviousChildrenNum)
+						{
+							// The widget was invalidated but nothing changed. This could be normal if a widget was removed, then re-added.
+						}
+						else if (PreviousIndex >= PreviousChildrenNum)
+						{
+							// Nothing to remove, but need to add some. We want to break the array.
+							RebuildWigetData.Emplace(WidgetRef, InvalidationWidget.LeafMostChildIndex, InvalidatedChildrenIndex);
+						}
+						else
+						{
+							const FIndexRange Range = { *this, IndexWhereToStart, IndexWhereToEnd };
+							check(Range.OrderMin <= Range.OrderMax);
+							RebuildWidgetRange.Add(Range);
+							RebuildWigetData.Emplace(WidgetRef, Range, InvalidatedChildrenIndex);
+						}
+#else
+						// Remove every children and rebuld the widget.
+						const FSlateInvalidationWidgetIndex IndexWhereToStart = IncrementIndex(WidgetIndex);
+						const FSlateInvalidationWidgetIndex IndexWhereToEnd = InvalidationWidget.LeafMostChildIndex;
+						const FIndexRange Range = { *this, IndexWhereToStart, IndexWhereToEnd };
+						RebuildWidgetRange.Add(Range);
+						RebuildWigetData.Emplace(WidgetRef, Range, 0);
+#endif
+					}
+					// There was not child, but maybe it has some now
+					else
+					{
+						FChildren* InvalidatedChildren = WidgetRef->GetAllChildren();
+						check(InvalidatedChildren);
+						if (InvalidatedChildren->Num() > 0)
+						{
+							// Nothing to remove, but need to add new. We want to break the array.
+							RebuildWigetData.Emplace(WidgetRef, InvalidationWidget.LeafMostChildIndex, 0);
+						}
+					}
+				}
+			}
+		}
+
+#if UE_SLATE_VERIFY_REBUILDWIDGETDATA_ORDER
+		// The array was sorted, but we may have invalidated a parent to only add at the end and the sub tree before is also invalidated.
+		//A (B(1,2), C(5,6)) => A (B(1,2,3), C(5, 6) D(8))
+		//A and B are invalidated. They were sorted (A, B). Now we know that D and B process. B need to be process after.
+		//This should already work (because we are doing the reverse operation). Lets confirm.
+		{
+			FMemMark Mark2(FMemStack::Get());
+			TArray<FChildOrderInvalidationData, TMemStackAllocator<>> SortedRebuildWigetData = RebuildWigetData;
+			FSlateInvalidationWidgetList* Self = this;
+			SortedRebuildWigetData.Sort([Self](const FChildOrderInvalidationData& A, const FChildOrderInvalidationData& B) -> bool
+				{
+					const FSlateInvalidationWidgetIndex AIndex = A.bRemove ? A.Range.InclusiveMin : A.WhereToCut;
+					const FSlateInvalidationWidgetIndex BIndex = B.bRemove ? B.Range.InclusiveMin : B.WhereToCut;
+					return FSlateInvalidationWidgetSortOrder(*Self, AIndex) < FSlateInvalidationWidgetSortOrder(*Self, BIndex);
+				});
+
+			// Is the sorted list in the same order
+			ensureAlways(SortedRebuildWigetData == RebuildWigetData);
+		}
+#endif //UE_SLATE_VERIFY_REBUILDWIDGETDATA_ORDER
+
+		{
+			SCOPED_NAMED_EVENT(Slate_InvalidationList_ProcessRemove, FColorList::Blue);
+			for (int32 Index = RebuildWigetData.Num() - 1; Index >= 0; --Index)
+			{
+				const FChildOrderInvalidationData& RebuildData = RebuildWigetData[Index];
+				if (RebuildData.bRemove)
+				{
+					_RemoveRangeFromSameParent(RebuildData.Range);
+				}
+				else
+				{
+					CutArray(RebuildData.WhereToCut);
+				}
+			}
+		}
+
+		{
+			SCOPED_NAMED_EVENT(Slate_InvalidationList_ProcessRebuild, FColorList::Blue);
+			for (const FChildOrderInvalidationData& RebuildData : RebuildWigetData)
+			{
+				if (RebuildData.Widget)
+				{
+#if UE_SLATE_VERIFY_REMOVED_WIDGET_ARE_NOT_INVALIDATED
+					FSlateInvalidationWidgetIndex SlateInvalidationWidgetIndexRemoved = { GSlateInvalidationWidgetIndex_RemovedIndex, GSlateInvalidationWidgetIndex_RemovedIndex };
+					ensure(RebuildData.Widget->GetProxyHandle().GetWidgetIndex() != SlateInvalidationWidgetIndexRemoved);
+#endif
+					_RebuildWidgetListTree(RebuildData.Widget.ToSharedRef(), RebuildData.ChildAtToStartWith);
+				}
+			}
+		}
+
+
+#if UE_SLATE_VERIFY_INVALID_INVALIDATIONHANDLE
+		{
+			for (TWeakPtr<SWidget>& WidgetInvalidated : VerifyWidgetInvalidationHandle)
+			{
+				if (TSharedPtr<SWidget> Widget = WidgetInvalidated.Pin())
+				{
+					const FSlateInvalidationWidgetIndex WidgetIndex = Widget->GetProxyHandle().GetWidgetIndex();
+					ensureMsgf(WidgetIndex != FSlateInvalidationWidgetIndex::Invalid, TEXT("The widget '%s' requested a ChildOrder but didn't have a valid index and was not rebuilt by something else.")
+						, *FReflectionMetaData::GetWidgetDebugInfo(Widget.Get()));
+				}
+			}
+		}
+#endif
+	}
+}
+
+
+namespace Private
+{
+	template<typename TSlateInvalidationWidgetList, typename Predicate>
+	void ForEachChildren(TSlateInvalidationWidgetList& Self, const typename TSlateInvalidationWidgetList::InvalidationWidgetType& InvalidationWidget, FSlateInvalidationWidgetIndex WidgetIndex, Predicate InPredicate)
+	{
+		using SlateInvalidationWidgetType = typename TChooseClass<TIsConst<TSlateInvalidationWidgetList>::Value, const typename TSlateInvalidationWidgetList::InvalidationWidgetType, typename TSlateInvalidationWidgetList::InvalidationWidgetType>::Result;
+		if (InvalidationWidget.LeafMostChildIndex != WidgetIndex)
+		{
+			FSlateInvalidationWidgetIndex CurrentWidgetIndex = Self.IncrementIndex(WidgetIndex);
+			while (true)
+			{
+				SlateInvalidationWidgetType& CurrentInvalidationWidget = Self[CurrentWidgetIndex];
+				InPredicate(CurrentInvalidationWidget);
+				CurrentWidgetIndex = CurrentInvalidationWidget.LeafMostChildIndex;
+				if (InvalidationWidget.LeafMostChildIndex == CurrentWidgetIndex)
+				{
+					break;
+				}
+				CurrentWidgetIndex = Self.IncrementIndex(CurrentWidgetIndex);
+				if (InvalidationWidget.LeafMostChildIndex == CurrentWidgetIndex)
+				{
+					InPredicate(Self[CurrentWidgetIndex]);
+					break;
+				}
+			}
+		}
+	}
+}
+
+
+void FSlateInvalidationWidgetList::_FindChildren(FSlateInvalidationWidgetIndex WidgetIndex, TArray<SWidget*, TMemStackAllocator<>>& Widgets) const
+{
+	Widgets.Reserve(16);
+	const InvalidationWidgetType& InvalidationWidget = (*this)[WidgetIndex];
+	Private::ForEachChildren(*this, InvalidationWidget, WidgetIndex, [&Widgets](const InvalidationWidgetType& ChildWidget)
+		{
+			Widgets.Add(ChildWidget.GetWidget());
+		});
+}
+
+
+//FSlateInvalidationWidgetIndex FSlateInvalidationWidgetList::_FindPreviousSibling(FSlateInvalidationWidgetIndex WidgetIndex) const
+//{
+//	FSlateInvalidationWidgetIndex Result = FSlateInvalidationWidgetIndex::Invalid;
+//	const InvalidationWidgetType& InvalidationWidget = (*this)[WidgetIndex];
+//	if (InvalidationWidget.ParentIndex != FSlateInvalidationWidgetIndex::Invalid)
+//	{
+//		const FSlateInvalidationWidgetIndex PreviousWidgetIndex = DecrementIndex(WidgetIndex);
+//		const InvalidationWidgetType& PreviousInvalidationWidget = (*this)[WidgetIndex];
+//		if (PreviousInvalidationWidget.ParentIndex == InvalidationWidget.ParentIndex)
+//		{
+//			Result = PreviousWidgetIndex;
+//		}
+//		else
+//		{
+//			Result = PreviousInvalidationWidget.ParentIndex;
+//		}
+//	}
+//	return  Result;
+//}
+
+
+FSlateInvalidationWidgetIndex FSlateInvalidationWidgetList::IncrementIndex(FSlateInvalidationWidgetIndex Index) const
+{
+	check(Data.IsValidIndex(Index.ArrayIndex));
+	++Index.ElementIndex;
+	if (Index.ElementIndex >= Data[Index.ArrayIndex].ElementList.Num())
+	{
+		if (Data[Index.ArrayIndex].NextArrayIndex == INDEX_NONE)
+		{
+			return FSlateInvalidationWidgetIndex::Invalid;
+		}
+		check(Data[Index.ArrayIndex].NextArrayIndex < FSlateInvalidationWidgetIndex::Invalid.ArrayIndex);
+		Index.ArrayIndex = (IndexType)Data[Index.ArrayIndex].NextArrayIndex;
+		Index.ElementIndex = Data[Index.ArrayIndex].StartIndex;
+	}
+	return Index;
+}
+
+
+FSlateInvalidationWidgetIndex FSlateInvalidationWidgetList::DecrementIndex(FSlateInvalidationWidgetIndex Index) const
+{
+	check(Data.IsValidIndex(Index.ArrayIndex));
+	if (Index.ElementIndex == Data[Index.ArrayIndex].StartIndex)
+	{
+		if (Data[Index.ArrayIndex].PreviousArrayIndex == INDEX_NONE)
+		{
+			return FSlateInvalidationWidgetIndex::Invalid;
+		}
+		check(Data[Index.ArrayIndex].PreviousArrayIndex < FSlateInvalidationWidgetIndex::Invalid.ArrayIndex);
+		Index.ArrayIndex = (IndexType)Data[Index.ArrayIndex].PreviousArrayIndex;
+		check(Data[Index.ArrayIndex].ElementList.Num() > 0);
+		Index.ElementIndex = Data[Index.ArrayIndex].ElementList.Num() - 1;
+	}
+	else
+	{
+		--Index.ElementIndex;
+	}
+	return Index;
+}
+
+
+void FSlateInvalidationWidgetList::Empty()
+{
+	Data.Empty();
+	Root.Reset();
+	FirstArrayIndex = INDEX_NONE;
+	LastArrayIndex = INDEX_NONE;
+}
+
+
+void FSlateInvalidationWidgetList::Reset()
+{
+	Data.Reset();
+	Root.Reset();
+	FirstArrayIndex = INDEX_NONE;
+	LastArrayIndex = INDEX_NONE;
+}
+
+
+FSlateInvalidationWidgetList::IndexType FSlateInvalidationWidgetList::AddArrayNodeIfNeeded(bool bReserveElementList)
+{
+	if (LastArrayIndex == INDEX_NONE || Data[LastArrayIndex].ElementList.Num() + 1 > WidgetListConfig.PreferedElementsNum)
+	{
+		if (Data.Num() + 1 == FSlateInvalidationWidgetIndex::Invalid.ArrayIndex)
+		{
+			ensure(false);
+			return LastArrayIndex;
+		}
+		const int32 Index = Data.Add(FArrayNode());
+		check(Index < std::numeric_limits<IndexType>::max());
+		if (bReserveElementList)
+		{
+			Data[Index].ElementList.Reserve(WidgetListConfig.PreferedElementsNum);
+		}
+
+		if (LastArrayIndex != INDEX_NONE)
+		{
+			Data[LastArrayIndex].NextArrayIndex = Index;
+			Data[Index].SortOrder = Data[LastArrayIndex].SortOrder + WidgetListConfig.SortOrderPaddingBetweenArray;
+		}
+		Data[Index].PreviousArrayIndex = LastArrayIndex;
+
+		LastArrayIndex = Index;
+		if (FirstArrayIndex == INDEX_NONE)
+		{
+			FirstArrayIndex = LastArrayIndex;
+		}
+	}
+	return LastArrayIndex;
+}
+
+
+FSlateInvalidationWidgetList::IndexType FSlateInvalidationWidgetList::InsertArrayNodeIfNeeded(IndexType AfterArrayIndex, bool bReserveElementList)
+{
+	if (AfterArrayIndex == FSlateInvalidationWidgetIndex::Invalid.ArrayIndex)
+	{
+		return AddArrayNodeIfNeeded(bReserveElementList);
+	}
+	else if (Data[AfterArrayIndex].ElementList.Num() + 1 > WidgetListConfig.PreferedElementsNum)
+	{
+		return InsertDataNodeAfter(AfterArrayIndex, bReserveElementList);
+	}
+	return AfterArrayIndex;
+}
+
+
+void FSlateInvalidationWidgetList::RebuildOrderIndex(IndexType StartFrom)
+{
+	check(StartFrom != INDEX_NONE);
+	check(WidgetListConfig.PreferedElementsNum < WidgetListConfig.MaxPreferedElementsNum);
+
+	FSlateInvalidationWidgetList* Self = this;
+	auto SetValue = [Self](IndexType ArrayIndex)
+		{
+			ElementListType& ElementList = Self->Data[ArrayIndex].ElementList;
+			for (int32 ElementIndex = Self->Data[ArrayIndex].StartIndex; ElementIndex < ElementList.Num(); ++ElementIndex)
+			{
+				FSlateInvalidationWidgetIndex WidgetIndex = { ArrayIndex, (IndexType)ElementIndex };
+				FSlateInvalidationWidgetSortOrder SortOrder = { *Self, WidgetIndex };
+				ElementList[ElementIndex].GetWidget()->SetFastPathSortOrder(SortOrder);
+			}
+		};
+
+	for (int32 CurrentIndex = StartFrom; CurrentIndex != INDEX_NONE; CurrentIndex = Data[CurrentIndex].NextArrayIndex)
+	{
+		const int32 PreviousIndex = Data[CurrentIndex].PreviousArrayIndex;
+		const int32 NextIndex = Data[CurrentIndex].NextArrayIndex;
+
+		if (PreviousIndex == INDEX_NONE)
+		{
+			Data[CurrentIndex].SortOrder = 0;
+			SetValue((IndexType)CurrentIndex);
+		}
+		else if (NextIndex == INDEX_NONE)
+		{
+			Data[CurrentIndex].SortOrder = Data[PreviousIndex].SortOrder + WidgetListConfig.SortOrderPaddingBetweenArray;
+			SetValue((IndexType)CurrentIndex);
+			break;
+		}
+		else
+		{
+			const int32 PreviousMinSortOrder = Data[PreviousIndex].SortOrder;
+			const int32 PreviousMaxSortOrder = Data[PreviousIndex].SortOrder + WidgetListConfig.PreferedElementsNum;
+			const int32 NextMinSortOrder = Data[NextIndex].SortOrder;
+			const int32 CurrentSortOrder = Data[CurrentIndex].SortOrder;
+
+			// Is everything already good
+			if (PreviousMaxSortOrder < CurrentSortOrder && CurrentSortOrder + WidgetListConfig.PreferedElementsNum < NextMinSortOrder)
+			{
+				break;
+			}
+			// Would the normal padding be valid
+			else if (PreviousMinSortOrder + WidgetListConfig.SortOrderPaddingBetweenArray < NextMinSortOrder)
+			{
+				Data[CurrentIndex].SortOrder = PreviousMaxSortOrder + WidgetListConfig.SortOrderPaddingBetweenArray;
+				SetValue((IndexType)CurrentIndex);
+			}
+			// Would padding by half would be valid
+			else if (NextMinSortOrder > PreviousMaxSortOrder && NextMinSortOrder - PreviousMaxSortOrder >= WidgetListConfig.PreferedElementsNum)
+			{
+				// We prefer to keep space of the exact amount in PreferedElementsNum in front Previous sort order and in the back Next sort order.
+				//That way we potently need less reorder in the future.
+				const int32 NumSpacesAvailable = (NextMinSortOrder - PreviousMaxSortOrder) / WidgetListConfig.PreferedElementsNum;
+				const int32 NewCurrentOrderIndex = PreviousMaxSortOrder + (WidgetListConfig.PreferedElementsNum * (NumSpacesAvailable/2));
+				check(PreviousMaxSortOrder <= NewCurrentOrderIndex && NewCurrentOrderIndex + WidgetListConfig.PreferedElementsNum <= NextMinSortOrder);
+				Data[CurrentIndex].SortOrder = NewCurrentOrderIndex;
+				SetValue((IndexType)CurrentIndex);
+			}
+			// Worst case, need to also rebuild the next array
+			else
+			{
+				Data[CurrentIndex].SortOrder = PreviousMaxSortOrder + FMath::Min(WidgetListConfig.PreferedElementsNum * 2, WidgetListConfig.SortOrderPaddingBetweenArray);
+				SetValue((IndexType)CurrentIndex);
+			}
+		}
+		ensureMsgf(Data[CurrentIndex].SortOrder <= WidgetListConfig.MaxSortOrderPaddingBetweenArray
+			, TEXT("The order index '%d' is too big to be contained inside the WidgetSortIndex. The Widget order will not be valid.")
+			, Data[CurrentIndex].SortOrder); // See FSlateInvalidationWidgetSortOrder
+	}
+}
+
+
+FSlateInvalidationWidgetList::IndexType FSlateInvalidationWidgetList::InsertDataNodeAfter(IndexType AfterIndex, bool bReserveElementList)
+{
+	if (FirstArrayIndex == INDEX_NONE)
+	{
+		check(AfterIndex == INDEX_NONE);
+		check(LastArrayIndex == INDEX_NONE);
+		AddArrayNodeIfNeeded(bReserveElementList);
+		return (IndexType)LastArrayIndex;
+	}
+	else
+	{
+		check(AfterIndex != INDEX_NONE);
+
+		if (Data.Num() + 1 == FSlateInvalidationWidgetIndex::Invalid.ArrayIndex)
+		{
+			ensure(false);
+			return LastArrayIndex;
+		}
+
+		const int32 NewIndex = Data.Add(FArrayNode());
+		check(NewIndex < std::numeric_limits<IndexType>::max());
+		if (bReserveElementList)
+		{
+			Data[NewIndex].ElementList.Reserve(WidgetListConfig.PreferedElementsNum);
+		}
+
+		FArrayNode& AfterArrayNode = Data[AfterIndex];
+		if (AfterArrayNode.NextArrayIndex != INDEX_NONE)
+		{
+			Data[AfterArrayNode.NextArrayIndex].PreviousArrayIndex = NewIndex;
+			Data[NewIndex].NextArrayIndex = AfterArrayNode.NextArrayIndex;
+
+			AfterArrayNode.NextArrayIndex = NewIndex;
+			Data[NewIndex].PreviousArrayIndex = AfterIndex;
+
+			if (LastArrayIndex == AfterIndex)
+			{
+				LastArrayIndex = NewIndex;
+			}
+
+			RebuildOrderIndex(NewIndex);
+		}
+		else
+		{
+			check(LastArrayIndex == AfterIndex);
+			LastArrayIndex = NewIndex;
+			Data[NewIndex].PreviousArrayIndex = AfterIndex;
+			Data[AfterIndex].NextArrayIndex = NewIndex;
+			Data[NewIndex].SortOrder = Data[AfterIndex].SortOrder + WidgetListConfig.SortOrderPaddingBetweenArray;
+		}
+
+		return (IndexType)NewIndex;
+	}
+}
+
+
+void FSlateInvalidationWidgetList::RemoveDataNode(IndexType Index)
+{
+	check(Index != INDEX_NONE && Index != std::numeric_limits<IndexType>::max());
+	FArrayNode& ArrayNode = Data[Index];
+	if (ArrayNode.PreviousArrayIndex != INDEX_NONE)
+	{
+		Data[ArrayNode.PreviousArrayIndex].NextArrayIndex = ArrayNode.NextArrayIndex;
+	}
+	else
+	{
+		FirstArrayIndex = ArrayNode.NextArrayIndex;
+	}
+
+	if (ArrayNode.NextArrayIndex != INDEX_NONE)
+	{
+		Data[ArrayNode.NextArrayIndex].PreviousArrayIndex = ArrayNode.PreviousArrayIndex;
+	}
+	else
+	{
+		LastArrayIndex = ArrayNode.PreviousArrayIndex;
+	}
+	ArrayNode.ElementList.Empty();
+	Data.RemoveAt(Index);
+
+	// No need to rebuild the order when we remove.
+	//OrderIndex is incremental and only use to sort.
+	//RebuildOrderIndex();
+
+	check(FirstArrayIndex != Index);
+	check(LastArrayIndex != Index);
+	if (Data.Num() == 0)
+	{
+		check(LastArrayIndex == INDEX_NONE && FirstArrayIndex == INDEX_NONE);
+	}
+	else
+	{
+		check(FirstArrayIndex != INDEX_NONE && LastArrayIndex != INDEX_NONE);
+	}
+}
+
+
+void FSlateInvalidationWidgetList::UpdateParentLeafIndex(const InvalidationWidgetType& NewInvalidationWidget, FSlateInvalidationWidgetIndex OldWidgetIndex, FSlateInvalidationWidgetIndex NewWidgetIndex)
+{
+	if (NewInvalidationWidget.ParentIndex != FSlateInvalidationWidgetIndex::Invalid)
+	{
+		InvalidationWidgetType* ParentInvalidationWidget = &(*this)[NewInvalidationWidget.ParentIndex];
+		while (ParentInvalidationWidget->LeafMostChildIndex == OldWidgetIndex)
+		{
+			ParentInvalidationWidget->LeafMostChildIndex = NewWidgetIndex;
+			if (ParentInvalidationWidget->ParentIndex != FSlateInvalidationWidgetIndex::Invalid)
+			{
+				ParentInvalidationWidget = &(*this)[ParentInvalidationWidget->ParentIndex];
+			}
+		}
+	}
+}
+
+
+/**
+ * To remove child from the same parent or all it's children.
+ * A ( B (C,D), E (F,G) )
+ * Can use to remove range (B,D) or (C,C) or (C,D) or (E, G) or (F,G) or (B,G).
+ * Cannot use (B,E) or (B,C) or (B,F)
+ */
+void FSlateInvalidationWidgetList::_RemoveRangeFromSameParent(const FIndexRange Range)
+{
+	// Fix up Parent's LeafIndex if they are in Range
+	{
+		//N.B. The algo doesn't support cross family removal. We do not need to worry about fixing the ParentIndex.
+		//(i)	There is no other child. Set Parent's LeafIndex to itself (recursive) [remove (F,G), leaf of E=E and leaf of A=E]
+		//(ii)	The Parent's LeafIndex is already set properly [remove (B,C), leaf of A is already G]
+		//(iii)	Parent's LeafIndex should be set to the previous sibling' leaf (recursive) [remove (G,G), leaf of E=F, A=F]
+		const InvalidationWidgetType& InvalidationWidgetStart = (*this)[Range.InclusiveMin];
+		const InvalidationWidgetType& InvalidationWidgetEnd = (*this)[Range.InclusiveMax];
+
+		// Parent index could be invalid if there is only one widget and we are removing it
+		if (InvalidationWidgetStart.ParentIndex != FSlateInvalidationWidgetIndex::Invalid)
+		{
+			// Is the parent's leaf being removed
+			const InvalidationWidgetType& InvalidationWidgetParentStart = (*this)[InvalidationWidgetStart.ParentIndex];
+			if (Range.Include(FSlateInvalidationWidgetSortOrder{ *this, InvalidationWidgetParentStart.LeafMostChildIndex }))
+			{
+				// _RemoveRangeFromSameParent doesn't support cross family removal
+				check(Range.Include(FSlateInvalidationWidgetSortOrder{ *this, (*this)[InvalidationWidgetEnd.ParentIndex].LeafMostChildIndex }));
+				const FSlateInvalidationWidgetIndex PreviousWidget = DecrementIndex(Range.InclusiveMin);
+				UpdateParentLeafIndex(InvalidationWidgetStart, InvalidationWidgetParentStart.LeafMostChildIndex, PreviousWidget);
+			}
+		}
+	}
+
+	//N.B. Theres is no parent/left relation in the array.
+	//ie.		1234 5678 90ab	(they have a size of 4, we cut at 2)
+	//(i)		123x xxxx x0ab => 123 x0ab			=> no cut, remove 4, remove 5-8, set StartIndex to 0
+	//(ii)		123x xxxx xxxb => 123 b				=> cut, remove 4, remove 5-8, remove 9-b (b was moved)
+	//(iii)		12xx 5478 90ab => 12 5478 90ab		=> cut, remove 3-4
+	//(iv)		1234 x678 90ab => 1234 x678 90ab	=> no cut, set StartIndex to 6
+	//(v)		1234 xxx8 90ab => 1234 8 90ab		=> cut, remove 5-8, (8 was moved)
+	//(vi)		1234 5xx8 90ab => 1234 5 8 90ab		=> cut, remove 6-8, (8 was moved)
+
+	const int32 NumberElementLeft = Data[Range.InclusiveMax.ArrayIndex].ElementList.Num()
+		//- Data[Range.InclusiveMax.ArrayIndex].StartIndex
+		- Range.InclusiveMax.ElementIndex
+		- 1;
+	const bool bRangeIsInSameElementArray = Range.InclusiveMin.ArrayIndex == Range.InclusiveMax.ArrayIndex;
+	const bool bShouldCutArray = NumberElementLeft < WidgetListConfig.NumberElementsLeftBeforeSplitting
+		|| (bRangeIsInSameElementArray && Data[Range.InclusiveMin.ArrayIndex].StartIndex != Range.InclusiveMin.ElementIndex);
+	if (bShouldCutArray)
+	{
+		_CutArray(Range.InclusiveMax);
+	}
+
+	// Destroy/Remove the data that is not needed anymore
+	{
+		FSlateInvalidationWidgetList* Self = this;
+		auto SetFakeInvalidatWidgetHandle = [Self](IndexType ArrayIndex, int32 StartIndex, int32 Num)
+		{
+#if UE_SLATE_VERIFY_REMOVED_WIDGET_ARE_NOT_INVALIDATED
+			const FSlateInvalidationWidgetIndex SlateInvalidationWidgetIndexRemoved = { GSlateInvalidationWidgetIndex_RemovedIndex, GSlateInvalidationWidgetIndex_RemovedIndex };
+			for (int32 ElementIndex = StartIndex; ElementIndex < Num; ++ElementIndex)
+			{
+				InvalidationWidgetType& InvalidationWidget = Self->Data[ArrayIndex].ElementList[ElementIndex];
+				if (SWidget* Widget = InvalidationWidget.GetWidget())
+				{
+					Widget->SetFastPathProxyHandle(FWidgetProxyHandle{ Self->Owner, SlateInvalidationWidgetIndexRemoved, FSlateInvalidationWidgetSortOrder() });
+				}
+			}
+#endif
+		};
+		auto ResetInvalidationWidget = [Self](IndexType ArrayIndex, int32 StartIndex, int32 Num)
+		{
+			ElementListType& ResetElementList = Self->Data[ArrayIndex].ElementList;
+			for (int32 ElementIndex = StartIndex; ElementIndex < Num; ++ElementIndex)
+			{
+				ResetElementList[ElementIndex].ResetWidget();
+			}
+		};
+
+		auto RemoveDataNodeIfNeeded = [Self](int32 ArrayIndex)
+		{
+			if (Self->Data[ArrayIndex].StartIndex >= Self->Data[ArrayIndex].ElementList.Num())
+			{
+				Self->RemoveDataNode(ArrayIndex);
+			}
+		};
+
+		// Remove the other arrays that are between min and max	(ie. i, ii)
+		//if (Range.InclusiveMax.ArrayIndex - Range.InclusiveMin.ArrayIndex > 1)
+		if (!bRangeIsInSameElementArray)
+		{
+			const IndexType BeginArrayIndex = (IndexType)Data[Range.InclusiveMin.ArrayIndex].NextArrayIndex;
+			const IndexType EndArrayIndex = (IndexType)Data[Range.InclusiveMax.ArrayIndex].PreviousArrayIndex;
+
+			if (BeginArrayIndex != Range.InclusiveMax.ArrayIndex)
+			{
+				IndexType NextToRemove = BeginArrayIndex;
+				IndexType CurrentArrayIndex = NextToRemove;
+				do
+				{
+					SetFakeInvalidatWidgetHandle(CurrentArrayIndex, Data[CurrentArrayIndex].StartIndex, Data[CurrentArrayIndex].ElementList.Num());
+					CurrentArrayIndex = NextToRemove;
+					NextToRemove = (IndexType)Data[CurrentArrayIndex].NextArrayIndex;
+					RemoveDataNode(CurrentArrayIndex);
+				} while (CurrentArrayIndex != EndArrayIndex);
+			}
+		}
+
+		// Remove the start of the Max array
+		if (bShouldCutArray && !bRangeIsInSameElementArray)
+		{
+			// The valid data in the array was moved (ie. ii)
+			SetFakeInvalidatWidgetHandle(Range.InclusiveMax.ArrayIndex, Data[Range.InclusiveMax.ArrayIndex].StartIndex, Range.InclusiveMax.ElementIndex + 1);
+			RemoveDataNode(Range.InclusiveMax.ArrayIndex);
+		}
+		else if (!bRangeIsInSameElementArray)
+		{
+			// Set StartIndex (ie. i)
+			check(Range.InclusiveMin.ArrayIndex != Range.InclusiveMax.ArrayIndex);
+
+			SetFakeInvalidatWidgetHandle(Range.InclusiveMax.ArrayIndex, Data[Range.InclusiveMax.ArrayIndex].StartIndex, Range.InclusiveMax.ElementIndex + 1);
+			ResetInvalidationWidget(Range.InclusiveMax.ArrayIndex, Data[Range.InclusiveMax.ArrayIndex].StartIndex, Range.InclusiveMax.ElementIndex + 1);
+			Data[Range.InclusiveMax.ArrayIndex].StartIndex = Range.InclusiveMax.ElementIndex + 1;
+			RemoveDataNodeIfNeeded(Range.InclusiveMax.ArrayIndex);
+		}
+
+		// Remove what is left of the Min array
+		if (bShouldCutArray || !bRangeIsInSameElementArray)
+		{
+			// RemoveAt Min to Num of the array. (ie. i, ii, iii, v, vi)
+			ElementListType& RemoveElementList = Data[Range.InclusiveMin.ArrayIndex].ElementList;
+			if (bRangeIsInSameElementArray)
+			{
+				SetFakeInvalidatWidgetHandle(Range.InclusiveMin.ArrayIndex, Range.InclusiveMin.ElementIndex, Range.InclusiveMax.ElementIndex);
+			}
+			else
+			{
+				SetFakeInvalidatWidgetHandle(Range.InclusiveMin.ArrayIndex, Range.InclusiveMin.ElementIndex, RemoveElementList.Num());
+			}
+
+			const int32 RemoveArrayAt = Range.InclusiveMin.ElementIndex;
+			RemoveElementList.RemoveAt(RemoveArrayAt, RemoveElementList.Num() - RemoveArrayAt, true);
+			RemoveDataNodeIfNeeded(Range.InclusiveMin.ArrayIndex);
+		}
+		else
+		{
+			// Set StartIndex (ie. iv)
+			check(Range.InclusiveMin.ArrayIndex == Range.InclusiveMax.ArrayIndex);
+			check(Range.InclusiveMin.ElementIndex == Data[Range.InclusiveMin.ArrayIndex].StartIndex);
+
+			SetFakeInvalidatWidgetHandle(Range.InclusiveMin.ArrayIndex, Range.InclusiveMin.ElementIndex, Range.InclusiveMax.ElementIndex + 1);
+			ResetInvalidationWidget(Range.InclusiveMin.ArrayIndex, Range.InclusiveMin.ElementIndex, Range.InclusiveMax.ElementIndex + 1);
+			Data[Range.InclusiveMin.ArrayIndex].StartIndex = Range.InclusiveMax.ElementIndex + 1;
+			RemoveDataNodeIfNeeded(Range.InclusiveMin.ArrayIndex);
+		}
+	}
+}
+
+
+void FSlateInvalidationWidgetList::CutArray(const FSlateInvalidationWidgetIndex WhereToCut)
+{
+	int32 OldElementIndexStart = _CutArray(WhereToCut);
+
+	// Remove the old data that is now moved to the new array
+	if (OldElementIndexStart != INDEX_NONE)
+	{
+		ElementListType& RemoveElementList = Data[WhereToCut.ArrayIndex].ElementList;
+		RemoveElementList.RemoveAt(OldElementIndexStart, RemoveElementList.Num() - OldElementIndexStart, true);
+		if (RemoveElementList.Num() == 0)
+		{
+			RemoveDataNode(WhereToCut.ArrayIndex);
+		}
+	}
+}
+
+
+int32 FSlateInvalidationWidgetList::_CutArray(const FSlateInvalidationWidgetIndex WhereToCut)
+{
+	SCOPED_NAMED_EVENT(Slate_InvalidationList_CutArray, FColor::Blue);
+
+	// Should we cut the array and move the data to another array
+	// N.B. We can cut/move anywhere. Cross family may occur. Fix up everything.
+	if (WhereToCut.ElementIndex < Data[WhereToCut.ArrayIndex].ElementList.Num() - 1)
+	{
+		//From where to where we are moving the item
+		const IndexType OldArrayIndex = WhereToCut.ArrayIndex;
+		const IndexType OldElementIndexStart = WhereToCut.ElementIndex + 1;
+		const IndexType OldElementIndexEnd = Data[WhereToCut.ArrayIndex].ElementList.Num();
+		const IndexType NewArrayIndex = InsertDataNodeAfter(OldArrayIndex, false);
+		const int32 NewExpectedElementArraySize = OldElementIndexEnd - OldElementIndexStart;
+		Data[NewArrayIndex].ElementList.Reserve(NewExpectedElementArraySize);
+
+		const FIndexRange OldRange = { *this, FSlateInvalidationWidgetIndex{ OldArrayIndex, OldElementIndexStart },  FSlateInvalidationWidgetIndex{ OldArrayIndex, (IndexType)(OldElementIndexEnd - 1) } };
+
+		auto OldToNewIndex = [NewArrayIndex, OldElementIndexStart](FSlateInvalidationWidgetIndex OldIndex) -> FSlateInvalidationWidgetIndex
+		{
+			const IndexType NewElementIndex = OldIndex.ElementIndex - OldElementIndexStart;
+			const FSlateInvalidationWidgetIndex NewWidgetIndex = { NewArrayIndex, NewElementIndex };
+			return NewWidgetIndex;
+		};
+
+		// Copy and assign the new index to the widget
+//Todo @pat fix this with proper iterator so that we do not have to use operator[]
+		for (IndexType OldElementIndex = OldElementIndexStart; OldElementIndex < OldElementIndexEnd; ++OldElementIndex)
+		{
+			const FSlateInvalidationWidgetIndex MoveWidgetIndex = { OldArrayIndex, OldElementIndex };
+			const IndexType NewElementIndex = (IndexType)Data[NewArrayIndex].ElementList.Add(MoveTemp((*this)[MoveWidgetIndex]));
+			(*this)[MoveWidgetIndex].ResetWidget();
+			const FSlateInvalidationWidgetIndex NewWidgetIndex = { NewArrayIndex, NewElementIndex };
+			check(OldToNewIndex(MoveWidgetIndex) == NewWidgetIndex);
+			InvalidationWidgetType& NewInvalidationWidget = (*this)[NewWidgetIndex];
+
+			// Fix up Index
+			NewInvalidationWidget.Index = NewWidgetIndex;
+
+			// Fix up the parent index
+			{
+				if (OldRange.Include(FSlateInvalidationWidgetSortOrder{ *this, NewInvalidationWidget.ParentIndex }))
+				{
+					NewInvalidationWidget.ParentIndex = OldToNewIndex(NewInvalidationWidget.ParentIndex);
+				}
+			}
+
+			// Fix up the leaf index
+			{
+				if (OldRange.Include(FSlateInvalidationWidgetSortOrder{ *this, NewInvalidationWidget.LeafMostChildIndex }))
+				{
+					NewInvalidationWidget.LeafMostChildIndex = OldToNewIndex(NewInvalidationWidget.LeafMostChildIndex);
+				}
+				check(NewInvalidationWidget.LeafMostChildIndex != FSlateInvalidationWidgetIndex::Invalid);
+			}
+
+			// Anyone in the hierarchy can point to a invalid LeafIndex.
+			{
+				InvalidationWidgetType* ParentInvalidationWidget = &(*this)[NewInvalidationWidget.ParentIndex];
+				// If the Leaf is in the range, then recursive up.
+				while (OldRange.Include(FSlateInvalidationWidgetSortOrder{ *this, ParentInvalidationWidget->LeafMostChildIndex }))
+				{
+					ParentInvalidationWidget->LeafMostChildIndex = OldToNewIndex(ParentInvalidationWidget->LeafMostChildIndex);
+					if (ParentInvalidationWidget->ParentIndex == FSlateInvalidationWidgetIndex::Invalid)
+					{
+						break;
+					}
+					ParentInvalidationWidget = &(*this)[ParentInvalidationWidget->ParentIndex];
+				}
+			}
+
+			// Set new index
+			if (SWidget* Widget = NewInvalidationWidget.GetWidget())
+			{
+				FSlateInvalidationWidgetSortOrder SortIndex = { *this, NewWidgetIndex };
+				Widget->SetFastPathProxyHandle(FWidgetProxyHandle{ Owner, NewWidgetIndex, SortIndex });
+			}
+		}
+
+		// Random children may still point to the old ParentIndex
+		check(Data[NewArrayIndex].ElementList.Num() == NewExpectedElementArraySize);
+		check(Data[NewArrayIndex].StartIndex == 0);
+		for (IndexType NewElementIndex = 0; NewElementIndex < NewExpectedElementArraySize; ++NewElementIndex)
+		{
+			const FSlateInvalidationWidgetIndex NewWidgetIndex = { NewArrayIndex, NewElementIndex };
+			InvalidationWidgetType& NewInvalidationWidget = (*this)[NewWidgetIndex];
+
+			// The parent is only invalid when it's the root. We can't remove it, but we shouldn't move it.
+			check(NewInvalidationWidget.ParentIndex != FSlateInvalidationWidgetIndex::Invalid);
+			Private::ForEachChildren(*this, NewInvalidationWidget, NewWidgetIndex, [NewWidgetIndex](InvalidationWidgetType& NewChildInvalidationWidget)
+				{
+					NewChildInvalidationWidget.ParentIndex = NewWidgetIndex;
+				});
+		}
+
+		return OldElementIndexStart;
+	}
+	return INDEX_NONE;
+}
+
+#if WITH_SLATE_DEBUGGING
+FSlateInvalidationWidgetIndex FSlateInvalidationWidgetList::FindWidget(const TSharedRef<SWidget> WidgetToFind) const
+{
+	SWidget* WidgetToFindPtr = &WidgetToFind.Get();
+	for (FSlateInvalidationWidgetIndex Index = FirstIndex(); Index != FSlateInvalidationWidgetIndex::Invalid; Index = IncrementIndex(Index))
+	{
+		if (WidgetToFindPtr == (*this)[Index].GetWidget())
+		{
+			return Index;
+		}
+	}
+	return FSlateInvalidationWidgetIndex::Invalid;
+}
+
+
+void FSlateInvalidationWidgetList::RemoveWidget(const FSlateInvalidationWidgetIndex WidgetIndex)
+{
+	if (WidgetIndex != FSlateInvalidationWidgetIndex::Invalid && IsValidIndex(WidgetIndex))
+	{
+		const InvalidationWidgetType& InvalidationWidget = (*this)[WidgetIndex];
+
+		const FIndexRange Range = { *this, WidgetIndex, InvalidationWidget.LeafMostChildIndex };
+		_RemoveRangeFromSameParent(Range);
+	}
+}
+
+
+void FSlateInvalidationWidgetList::RemoveWidget(const TSharedRef<SWidget> WidgetToRemove)
+{
+	if (ensure(WidgetToRemove->GetProxyHandle().GetInvalidationRootHandle().GetUniqueId() == Owner.GetUniqueId()))
+	{
+		FSlateInvalidationWidgetIndex WidgetIndex = WidgetToRemove->GetProxyHandle().GetWidgetIndex();
+		if (WidgetIndex != FSlateInvalidationWidgetIndex::Invalid)
+		{
+			const InvalidationWidgetType& InvalidationWidget = (*this)[WidgetIndex];
+			const FIndexRange Range = { *this, WidgetIndex, InvalidationWidget.LeafMostChildIndex };
+			_RemoveRangeFromSameParent(Range);
+		}
+	}
+
+}
+
+
+TArray<TSharedPtr<SWidget>> FSlateInvalidationWidgetList::FindChildren(const TSharedRef<SWidget> Widget) const
+{
+	TArray<TSharedPtr<SWidget>> Result;
+	if (ensure(Widget->GetProxyHandle().GetInvalidationRootHandle().GetUniqueId() == Owner.GetUniqueId()))
+	{
+		FSlateInvalidationWidgetIndex WidgetIndex = Widget->GetProxyHandle().GetWidgetIndex();
+		if (WidgetIndex == FSlateInvalidationWidgetIndex::Invalid)
+		{
+			return Result;
+		}
+
+		FMemMark Mark(FMemStack::Get());
+		TArray<SWidget*, TMemStackAllocator<>> PreviousChildrenWidget;
+		_FindChildren(WidgetIndex, PreviousChildrenWidget);
+
+		Result.Reserve(PreviousChildrenWidget.Num());
+		for (SWidget* WidgetPtr : PreviousChildrenWidget)
+		{
+			Result.Add(WidgetPtr ? WidgetPtr->AsShared() : TSharedPtr<SWidget>());
+		}
+	}
+	return Result;
+}
+
+
+bool FSlateInvalidationWidgetList::DeapCompare(const FSlateInvalidationWidgetList& Other) const
+{
+	if (Root.Pin() != Other.Root.Pin())
+	{
+		return false;
+	}
+
+	FSlateInvalidationWidgetIndex IndexA = FirstIndex();
+	FSlateInvalidationWidgetIndex IndexB = Other.FirstIndex();
+	for (; IndexA != FSlateInvalidationWidgetIndex::Invalid && IndexB != FSlateInvalidationWidgetIndex::Invalid; IndexA = IncrementIndex(IndexA), IndexB = Other.IncrementIndex(IndexB))
+	{
+		const InvalidationWidgetType& InvalidationWidgetA = (*this)[IndexA];
+		const InvalidationWidgetType& InvalidationWidgetB = Other[IndexB];
+		if (InvalidationWidgetA.GetWidget() != InvalidationWidgetB.GetWidget())
+		{
+			return false;
+		}
+		if (InvalidationWidgetA.ParentIndex == FSlateInvalidationWidgetIndex::Invalid)
+		{
+			if (InvalidationWidgetA.ParentIndex != InvalidationWidgetB.ParentIndex)
+			{
+				return false;
+			}
+		}
+		else
+		{
+			if ((*this)[InvalidationWidgetA.ParentIndex].GetWidget() != Other[InvalidationWidgetB.ParentIndex].GetWidget())
+			{
+				return false;
+			}
+		}
+		check(InvalidationWidgetA.LeafMostChildIndex != FSlateInvalidationWidgetIndex::Invalid);
+		check(InvalidationWidgetB.LeafMostChildIndex != FSlateInvalidationWidgetIndex::Invalid);
+		if ((*this)[InvalidationWidgetA.LeafMostChildIndex].GetWidget() != Other[InvalidationWidgetB.LeafMostChildIndex].GetWidget())
+		{
+			return false;
+		}
+	}
+
+	if (IndexA != FSlateInvalidationWidgetIndex::Invalid || IndexB != FSlateInvalidationWidgetIndex::Invalid)
+	{
+		return false;
+	}
+
+	return true;
+}
+
+
+void FSlateInvalidationWidgetList::LogWidgetsList()
+{
+	TStringBuilder<256> Builder;
+	for (FSlateInvalidationWidgetIndex Index = FirstIndex(); Index != FSlateInvalidationWidgetIndex::Invalid; Index = IncrementIndex(Index))
+	{
+		Builder.Reset();
+		const InvalidationWidgetType& InvalidateWidget = (*this)[Index];
+		if (SWidget* Widget = InvalidateWidget.GetWidget())
+		{
+			Builder << Widget->GetTag();
+		}
+		else
+		{
+			Builder << TEXT("[None]");
+		}
+		Builder << TEXT("\t");
+		if (InvalidateWidget.ParentIndex != FSlateInvalidationWidgetIndex::Invalid)
+		{
+			if (SWidget* Widget = (*this)[InvalidateWidget.ParentIndex].GetWidget())
+			{
+				Builder << Widget->GetTag();
+			}
+			else
+			{
+				Builder << TEXT("[None]");
+			}
+		}
+		else
+		{
+			Builder << TEXT("[---]");
+		}
+		Builder << TEXT("\t");
+		if (InvalidateWidget.LeafMostChildIndex != FSlateInvalidationWidgetIndex::Invalid)
+		{
+			if (SWidget* Widget = (*this)[InvalidateWidget.LeafMostChildIndex].GetWidget())
+			{
+				Builder << Widget->GetTag();
+			}
+			else
+			{
+				Builder << TEXT("[None]");
+			}
+		}
+		else
+		{
+			Builder << TEXT("[---]");
+		}
+		Builder << TEXT("\t");
+
+		UE_LOG(LogSlate, Log, TEXT("%s"), Builder.ToString());
+	}
+}
+
+
+bool FSlateInvalidationWidgetList::VerifyWidgetsIndex() const
+{
+	bool bResult = true;
+	for (FSlateInvalidationWidgetIndex Index = FirstIndex(); Index != FSlateInvalidationWidgetIndex::Invalid; Index = IncrementIndex(Index))
+	{
+		const InvalidationWidgetType& InvalidateWidget = (*this)[Index];
+		if (SWidget* Widget = InvalidateWidget.GetWidget())
+		{
+			const FSlateInvalidationWidgetIndex WidgetIndex = Widget->GetProxyHandle().GetWidgetIndex();
+			if (Index != WidgetIndex)
+			{
+				UE_LOG(LogSlate, Warning, TEXT("Widget '%s' at index [%d,%d] is set to [%d,%d].")
+					, *FReflectionMetaData::GetWidgetDebugInfo(Widget)
+					, Index.ArrayIndex, Index.ElementIndex
+					, WidgetIndex.ArrayIndex, WidgetIndex.ElementIndex);
+				bResult = false;
+			}
+			else if (InvalidateWidget.Index != Index)
+			{
+				UE_LOG(LogSlate, Warning, TEXT("Widget '%s' at index [%d,%d] is set to the correct proxy index [%d,%d].")
+					, *FReflectionMetaData::GetWidgetDebugInfo(Widget)
+					, Index.ArrayIndex, Index.ElementIndex
+					, WidgetIndex.ArrayIndex, WidgetIndex.ElementIndex);
+				bResult = false;
+			}
+		}
+		else
+		{
+			UE_LOG(LogSlate, Warning, TEXT("Widget at index [%d,%d] is [null]"), Index.ArrayIndex, Index.ElementIndex);
+			bResult = false;
+		}
+	}
+	return bResult;
+}
+
+bool FSlateInvalidationWidgetList::VerifyProxiesWidget() const
+{
+	bool bResult = true;
+	for (const FArrayNode& Node : Data)
+	{
+		// Before StartIndex, pointer need to be empty.
+		for (int32 ElementIndex = 0; ElementIndex < Node.StartIndex; ++ElementIndex)
+		{
+			if (SWidget* Widget = Node.ElementList[ElementIndex].GetWidget())
+			{
+				UE_LOG(LogSlate, Warning, TEXT("Element '%d' in the array of sort value '%d' as a valid widget '%s' when it should be set to none.")
+					, ElementIndex, Node.SortOrder
+					, *FReflectionMetaData::GetWidgetDebugInfo(Widget));
+				bResult = false;
+			}
+		}
+
+		// Every other element need to point to a valid widget.
+		for (int32 ElementIndex = Node.StartIndex; ElementIndex < Node.ElementList.Num(); ++ElementIndex)
+		{
+			SWidget* Widget = Node.ElementList[ElementIndex].GetWidget();
+			if (!Widget)
+			{
+				UE_LOG(LogSlate, Warning, TEXT("Element '%d' in the array of sort value '%d' does not have a valid widget.")
+					, ElementIndex, Node.SortOrder);
+				bResult = false;
+			}
+		}
+	}
+	return bResult;
+}
+
+bool FSlateInvalidationWidgetList::VerifySortOrder() const
+{
+	bool bResult = true;
+	if (FirstArrayIndex != INDEX_NONE)
+	{
+		int32 PreviousSortOrder = Data[FirstArrayIndex].SortOrder;
+ 		for (int32 ArrayIndex = Data[FirstArrayIndex].NextArrayIndex; ArrayIndex != INDEX_NONE; ArrayIndex = Data[ArrayIndex].NextArrayIndex)
+		{
+			if (PreviousSortOrder >= Data[ArrayIndex].SortOrder)
+			{
+				UE_LOG(LogSlate, Warning, TEXT("Array '%d' has a bigger sort order than previous array node '%d'.")
+					, ArrayIndex, Data[ArrayIndex].PreviousArrayIndex);
+				bResult = false;
+				break;
+			}
+		}
+	}
+	return bResult;
+}
+#endif //WITH_SLATE_DEBUGGING
+
+#undef UE_SLATE_VERIFY_INVALID_INVALIDATIONHANDLE
+#undef UE_SLATE_VERIFY_REBUILDWIDGETDATA_ORDER
+#undef UE_SLATE_VERIFY_REMOVED_WIDGET_ARE_NOT_INVALIDATED
diff --git a/Engine/Source/Runtime/SlateCore/Private/FastUpdate/SlateInvalidationWidgetList.h b/Engine/Source/Runtime/SlateCore/Private/FastUpdate/SlateInvalidationWidgetList.h
new file mode 100644
--- /dev/null
+++ b/Engine/Source/Runtime/SlateCore/Private/FastUpdate/SlateInvalidationWidgetList.h
@@ -0,0 +1,265 @@
+// Copyright Epic Games, Inc. All Rights Reserved.
+
+#pragma once
+
+#include "CoreMinimal.h"
+
+#include "Debugging/SlateDebugging.h"
+#include "FastUpdate/SlateInvalidationRootHandle.h"
+#include "FastUpdate/SlateInvalidationWidgetIndex.h"
+#include "FastUpdate/SlateInvalidationWidgetSortOrder.h"
+#include "Widgets/SNullWidget.h"
+#include "Widgets/SWidget.h"
+
+class FSlateInvalidationWidgetList
+{
+public:
+	using InvalidationWidgetType = FWidgetProxy;
+private:
+	using IndexType = FSlateInvalidationWidgetIndex::IndexType;
+	using ElementListType = TArray<InvalidationWidgetType>;
+
+public:
+	struct FArguments
+	{
+		static int32 MaxPreferedElementsNum;
+		static int32 MaxSortOrderPaddingBetweenArray;
+		/**
+		 * Prefered size of the elements array.
+		 * The value should be between 2 and MaxPreferedElementsNum.
+		 */
+		int32 PreferedElementsNum = 64;
+
+		/**
+		 * When splitting, the elements will be copied to another array when the number of element is bellow this.
+		 * The value should be between 1 and PreferedElementsNum.
+		 */
+		int32 NumberElementsLeftBeforeSplitting = 40;
+
+		/**
+		 * The sort order is used by the HittestGrid and the LayerId.
+		 * The value should be between PreferedElementsNum and MaxSortOrderPaddingBetweenArray.
+		 */
+		int32 SortOrderPaddingBetweenArray = 1000;
+#if WITH_SLATE_DEBUGGING
+		/** Change the Widget Index when building the array. Use when building temporary list. */
+		bool bAssignedWidgetIndex = true;
+#endif
+	};
+	FSlateInvalidationWidgetList(FSlateInvalidationRootHandle Owner, const FArguments& Args);
+
+	/** Build the widget list from the root widget. */
+	void BuildWidgetList(TSharedRef<SWidget> Root);
+
+	/** Get the root the widget list was built with. */
+	TWeakPtr<SWidget> GetRoot() { return Root; };
+
+	/** Process widget that have a ChildOrder invalidation. */
+	void ProcessChildOrderInvalidation(const TArray<TWeakPtr<SWidget>>& Invalidated);
+
+	/** Performs an operation on all SWidget in the list. */
+	template<typename Predicate>
+	void ForEachWidget(Predicate Pred)
+	{
+		int32 ArrayIndex = FirstArrayIndex;
+		while (ArrayIndex != INDEX_NONE)
+		{
+			ElementListType& ElementList = Data[ArrayIndex].ElementList;
+			const int32 ElementNum = ElementList.Num();
+			for (int32 ElementIndex = Data[ArrayIndex].StartIndex; ElementIndex < ElementNum; ++ElementIndex)
+			{
+				SWidget* Widget = ElementList[ElementIndex].GetWidget();
+				if (Widget)
+				{
+					Pred(Widget);
+				}
+			}
+
+			ArrayIndex = Data[ArrayIndex].NextArrayIndex;
+		}
+	}
+
+public:
+	/** Returns reference to element at give index. */
+	InvalidationWidgetType& operator[](const FSlateInvalidationWidgetIndex Index)
+	{
+		check(IsValidIndex(Index));
+		return Data[Index.ArrayIndex].ElementList[Index.ElementIndex];
+	}
+
+	/** Returns reference to element at give index. */
+	const InvalidationWidgetType& operator[](const FSlateInvalidationWidgetIndex Index) const
+	{
+		check(IsValidIndex(Index));
+		return Data[Index.ArrayIndex].ElementList[Index.ElementIndex];
+	}
+
+	/** Tests if index is in the WidgetList range. */
+	bool IsValidIndex(const FSlateInvalidationWidgetIndex Index) const
+	{
+		if (Data.IsValidIndex(Index.ArrayIndex))
+		{
+			return Index.ElementIndex >= Data[Index.ArrayIndex].StartIndex && Index.ElementIndex < Data[Index.ArrayIndex].ElementList.Num();
+		}
+		return false;
+	}
+
+	/** Returns true if there is not element in the WidgetList. */
+	bool IsEmpty() const
+	{
+		return FirstArrayIndex == INDEX_NONE || Data[FirstArrayIndex].ElementList.Num() == 0;
+	}
+
+	/** Returns the first index from the WidgetList. */
+	FSlateInvalidationWidgetIndex FirstIndex() const
+	{
+		return FirstArrayIndex == INDEX_NONE
+			? FSlateInvalidationWidgetIndex::Invalid
+			: FSlateInvalidationWidgetIndex{ (IndexType)FirstArrayIndex, Data[FirstArrayIndex].StartIndex };
+	}
+
+	/** Returns the last index from the WidgetList. */
+	FSlateInvalidationWidgetIndex LastIndex() const
+	{
+		return LastArrayIndex == INDEX_NONE ? FSlateInvalidationWidgetIndex::Invalid : FSlateInvalidationWidgetIndex{ (IndexType)LastArrayIndex, (IndexType)(Data[LastArrayIndex].ElementList.Num() - 1) };
+	}
+
+	/** Increment a widget index to the next entry in the WidgetList. */
+	FSlateInvalidationWidgetIndex IncrementIndex(FSlateInvalidationWidgetIndex Index) const;
+
+	/** Decrement a widget index to the next entry in the WidgetList. */
+	FSlateInvalidationWidgetIndex DecrementIndex(FSlateInvalidationWidgetIndex Index) const;
+
+	/** Empties the WidgetList. */
+	void Empty();
+
+	/** Empties the WidgetList, but doesn't change the memory allocations. */
+	void Reset();
+
+public:
+#if WITH_SLATE_DEBUGGING
+	/** For testing purposes. Return the InvalidationWidgetIndex of the Widget within the InvalidationWidgetList. */
+	FSlateInvalidationWidgetIndex FindWidget(const TSharedRef<SWidget> Widget) const;
+
+	/** For testing purposes. Use ProcessChildOrderInvalidation. */
+	void RemoveWidget(const FSlateInvalidationWidgetIndex ToRemove);
+	/** For testing purposes. Use ProcessChildOrderInvalidation. */
+	void RemoveWidget(const TSharedRef<SWidget> WidgetToRemove);
+
+	/** For testing purposes. Use to test ProcessChildOrderInvalidation */
+	TArray<TSharedPtr<SWidget>> FindChildren(const TSharedRef<SWidget> Widget) const;
+
+	/**
+	 * For testing purposes.
+	 * The list may not be the same, but the InvalidationWidgetType must:
+	 * (1) be in the same order
+	 * (2) point the same SWidget (itself, parent, leaf)
+	 */
+	bool DeapCompare(const FSlateInvalidationWidgetList& Other) const;
+
+	/** For testing purposes. Log the tree. */
+	void LogWidgetsList();
+
+	/** For testing purposes. Verify that every widgets has the correct index. */
+	bool VerifyWidgetsIndex() const;
+
+	/** For testing purposes. Verify that every WidgetProxy has a valid SWidget. */
+	bool VerifyProxiesWidget() const;
+
+	/** For testing purposes. Verify that every the sorting order is increasing between widgets. */
+	bool VerifySortOrder() const;
+#endif
+
+private:
+	friend struct FSlateInvalidationWidgetSortOrder;
+
+	struct FIndexRange
+	{
+		FSlateInvalidationWidgetIndex InclusiveMin = FSlateInvalidationWidgetIndex::Invalid;
+		FSlateInvalidationWidgetIndex InclusiveMax = FSlateInvalidationWidgetIndex::Invalid;
+		FSlateInvalidationWidgetSortOrder OrderMin;
+		FSlateInvalidationWidgetSortOrder OrderMax;
+		FIndexRange() = default;
+		FIndexRange(const FSlateInvalidationWidgetList& Self, FSlateInvalidationWidgetIndex InFrom, FSlateInvalidationWidgetIndex InEnd)
+			: InclusiveMin(InFrom), InclusiveMax(InEnd)
+			, OrderMin(Self, InFrom), OrderMax(Self, InEnd)
+		{
+			check(OrderMin <= OrderMax);
+		}
+		bool Include(FSlateInvalidationWidgetSortOrder Other) const
+		{
+			return OrderMin <= Other && Other <= OrderMax;
+		}
+		bool operator==(const FIndexRange& Other) const { return Other.InclusiveMin == InclusiveMin && Other.InclusiveMax == InclusiveMax; }
+	};
+
+private:
+	bool ShouldDoRecursion(const SWidget* Widget) const
+	{
+		return !Widget->Advanced_IsInvalidationRoot() || IsEmpty();
+	}
+	bool ShouldDoRecursion(const TSharedRef<SWidget>& Widget) const
+	{
+		return !Widget->Advanced_IsInvalidationRoot() || IsEmpty();
+	}
+	bool ShouldBeAdded(const SWidget* Widget) const
+	{
+		return Widget != &(SNullWidget::NullWidget.Get());
+	}
+	bool ShouldBeAdded(const TSharedRef<SWidget>& Widget) const
+	{
+		return Widget != SNullWidget::NullWidget;
+	}
+
+private:
+	template <typename... ArgsType>
+	FSlateInvalidationWidgetIndex Emplace(ArgsType&&... Args)
+	{
+		const IndexType ArrayIndex = AddArrayNodeIfNeeded(true);
+		const IndexType ElementIndex = (IndexType)Data[ArrayIndex].ElementList.Emplace(Forward<ArgsType>(Args)...);
+		return FSlateInvalidationWidgetIndex{ ArrayIndex, ElementIndex };
+	}
+	template <typename... ArgsType>
+	FSlateInvalidationWidgetIndex EmplaceInsertAfter(IndexType AfterArrayIndex, ArgsType&&... Args)
+	{
+		const IndexType ArrayIndex = InsertArrayNodeIfNeeded(AfterArrayIndex, true);
+		const IndexType ElementIndex = (IndexType)Data[ArrayIndex].ElementList.Emplace(Forward<ArgsType>(Args)...);
+		return FSlateInvalidationWidgetIndex{ ArrayIndex, ElementIndex };
+	}
+
+private:
+	IndexType AddArrayNodeIfNeeded(bool bReserveElementList);
+	IndexType InsertArrayNodeIfNeeded(IndexType AfterArrayIndex, bool bReserveElementList);
+	IndexType InsertDataNodeAfter(IndexType AfterIndex, bool bReserveElementList);
+	void RemoveDataNode(IndexType Index);
+	void RebuildOrderIndex(IndexType StartFrom);
+	void UpdateParentLeafIndex(const InvalidationWidgetType& InvalidationWidget, FSlateInvalidationWidgetIndex OldIndex, FSlateInvalidationWidgetIndex NewIndex);
+	void CutArray(const FSlateInvalidationWidgetIndex WhereToCut);
+
+private:
+	FSlateInvalidationWidgetIndex _BuildWidgetList_Recursive(TSharedRef<SWidget>& Widget, FSlateInvalidationWidgetIndex ParentIndex, IndexType& LastestIndex, bool bParentVisible, bool bParentVolatile);
+	void _RebuildWidgetListTree(TSharedRef<SWidget> Widget, int32 ChildAtIndex);
+	void _FindChildren(FSlateInvalidationWidgetIndex WidgetIndex, TArray<SWidget*, TMemStackAllocator<>>& Widgets) const;
+	void _RemoveRangeFromSameParent(const FIndexRange Range);
+	int32 _CutArray(const FSlateInvalidationWidgetIndex WhereToCut);
+
+private:
+	struct FArrayNode
+	{
+		FArrayNode() = default;
+		int32 PreviousArrayIndex = INDEX_NONE;
+		int32 NextArrayIndex = INDEX_NONE;
+		int32 SortOrder = 0;
+		IndexType StartIndex = 0; // The array may start further in the array as a result of a split.
+		ElementListType ElementList;
+	};
+	using ArrayNodeType = TSparseArray<FArrayNode>;
+
+	FSlateInvalidationRootHandle Owner;
+	ArrayNodeType Data;
+	TWeakPtr<SWidget> Root;
+	int32 FirstArrayIndex = INDEX_NONE;
+	int32 LastArrayIndex = INDEX_NONE;
+
+	const FArguments WidgetListConfig;
+};
\ No newline at end of file
diff --git a/Engine/Source/Runtime/SlateCore/Private/FastUpdate/SlateInvalidationWidgetSortOrder.cpp b/Engine/Source/Runtime/SlateCore/Private/FastUpdate/SlateInvalidationWidgetSortOrder.cpp
new file mode 100644
--- /dev/null
+++ b/Engine/Source/Runtime/SlateCore/Private/FastUpdate/SlateInvalidationWidgetSortOrder.cpp
@@ -0,0 +1,13 @@
+// Copyright Epic Games, Inc. All Rights Reserved.
+
+#include "FastUpdate/SlateInvalidationWidgetSortOrder.h"
+
+#include "FastUpdate/SlateInvalidationWidgetList.h"
+
+FSlateInvalidationWidgetSortOrder::FSlateInvalidationWidgetSortOrder(const FSlateInvalidationWidgetList& List, FSlateInvalidationWidgetIndex Index)
+	: Order()
+{
+	// 22 bits for the order (array index)
+	// 10 bits for the element index
+	Order = (List.Data[Index.ArrayIndex].SortOrder << 10) + Index.ElementIndex;
+}
\ No newline at end of file
diff --git a/Engine/Source/Runtime/SlateCore/Private/FastUpdate/WidgetProxy.cpp b/Engine/Source/Runtime/SlateCore/Private/FastUpdate/WidgetProxy.cpp
--- a/Engine/Source/Runtime/SlateCore/Private/FastUpdate/WidgetProxy.cpp
+++ b/Engine/Source/Runtime/SlateCore/Private/FastUpdate/WidgetProxy.cpp
@@ -4,6 +4,8 @@
 #include "Widgets/SWidget.h"
 #include "Rendering/DrawElements.h"
 #include "Widgets/SWindow.h"
+#include "FastUpdate/SlateInvalidationWidgetHeap.h"
+#include "FastUpdate/SlateInvalidationWidgetList.h"
 #include "ProfilingDebugging/CsvProfiler.h"
 #include "Types/ReflectionMetadata.h"
 #include "Input/HittestGrid.h"
@@ -12,12 +14,15 @@
 
 const FSlateWidgetPersistentState FSlateWidgetPersistentState::NoState;
 
-FWidgetProxy::FWidgetProxy(SWidget& InWidget)
-	: Widget(&InWidget)
-	, Index(INDEX_NONE)
-	, ParentIndex(INDEX_NONE)
-	, NumChildren(0)
-	, LeafMostChildIndex(INDEX_NONE)
+FWidgetProxy::FWidgetProxy(TSharedRef<SWidget>& InWidget)
+#if UE_SLATE_WITH_WIDGETPROXY_WEAKPTR
+	: Widget(InWidget)
+#else
+	: Widget(&InWidget.Get())
+#endif
+	, Index(FSlateInvalidationWidgetIndex::Invalid)
+	, ParentIndex(FSlateInvalidationWidgetIndex::Invalid)
+	, LeafMostChildIndex(FSlateInvalidationWidgetIndex::Invalid)
 	, UpdateFlags(EWidgetUpdateFlags::None)
 	, CurrentInvalidateReason(EInvalidateWidgetReason::None)
 	// Potentially unsafe to update visibility from the widget due to attribute bindings.  This is updated later when the widgets are sorted in ProcessInvalidation
@@ -26,10 +31,11 @@ FWidgetProxy::FWidgetProxy(SWidget& InWidget)
 	, bInUpdateList(false)
 	, bInvisibleDueToParentOrSelfVisibility(false)
 	, bChildOrderInvalid(false)
+	, bContainedByWidgetHeap(false)
 {
 }
 
-int32 FWidgetProxy::Update(const FPaintArgs& PaintArgs, int32 MyIndex, FSlateWindowElementList& OutDrawElements)
+int32 FWidgetProxy::Update(const FPaintArgs& PaintArgs, FSlateWindowElementList& OutDrawElements)
 {
 // Commenting this since it could be triggered in specific cases where Widgte->UpdateFlags in reset but the Widget Proxy is still in the  update list
 //#if WITH_SLATE_DEBUGGING
@@ -41,12 +47,16 @@ int32 FWidgetProxy::Update(const FPaintArgs& PaintArgs, int32 MyIndex, FSlateWin
 	if (EnumHasAnyFlags(UpdateFlags, EWidgetUpdateFlags::NeedsRepaint|EWidgetUpdateFlags::NeedsVolatilePaint))
 	{
 		ensure(!bInvisibleDueToParentOrSelfVisibility);
-		OutgoingLayerId = Repaint(PaintArgs, MyIndex, OutDrawElements);
+		OutgoingLayerId = Repaint(PaintArgs, OutDrawElements);
 	}
 	else if(!bInvisibleDueToParentOrSelfVisibility)
 	{
 		EWidgetUpdateFlags PreviousUpdateFlag = UpdateFlags;
+#if UE_SLATE_WITH_WIDGETPROXY_WEAKPTR
+		TSharedPtr<SWidget> CurrentWidget = Widget.Pin();
+#else
 		TSharedPtr<SWidget> CurrentWidget = Widget->AsShared();
+#endif
 
 		if (EnumHasAnyFlags(UpdateFlags, EWidgetUpdateFlags::NeedsActiveTimerUpdate))
 		{
@@ -73,92 +83,94 @@ int32 FWidgetProxy::Update(const FPaintArgs& PaintArgs, int32 MyIndex, FSlateWin
 	return OutgoingLayerId;
 }
 
-bool FWidgetProxy::ProcessInvalidation(FWidgetUpdateList& UpdateList, TArray<FWidgetProxy>& FastWidgetPathList, FSlateInvalidationRoot& Root)
+bool FWidgetProxy::ProcessInvalidation(FSlateInvalidationWidgetHeap& UpdateList, FSlateInvalidationWidgetList& FastWidgetPathList, FSlateInvalidationRoot& Root)
 {
 	bool bWidgetNeedsRepaint = false;
-	if (!bInvisibleDueToParentOrSelfVisibility && ParentIndex != INDEX_NONE && !Widget->PrepassLayoutScaleMultiplier.IsSet())
+	SWidget* WidgetPtr = GetWidget();
+
+	if (!bInvisibleDueToParentOrSelfVisibility && ParentIndex != FSlateInvalidationWidgetIndex::Invalid && !WidgetPtr->PrepassLayoutScaleMultiplier.IsSet())
 	{
-		SCOPE_CYCLE_SWIDGET(Widget);
+		SCOPE_CYCLE_SWIDGET(WidgetPtr);
 		// If this widget has never been prepassed make sure the parent prepasses it to set the correct multiplier
 		FWidgetProxy& ParentProxy = FastWidgetPathList[ParentIndex];
-		if (ParentProxy.Widget)
+		if (SWidget* ParentWidgetPtr = ParentProxy.GetWidget())
 		{
-			ParentProxy.Widget->InvalidatePrepass();
+			ParentWidgetPtr->InvalidatePrepass();
 			ParentProxy.CurrentInvalidateReason |= EInvalidateWidgetReason::Layout;
 #if WITH_SLATE_DEBUGGING
-			FSlateDebugging::BroadcastWidgetInvalidate(ParentProxy.Widget, Widget, EInvalidateWidgetReason::Layout);
+			FSlateDebugging::BroadcastWidgetInvalidate(ParentWidgetPtr, WidgetPtr, EInvalidateWidgetReason::Layout);
 #endif
-			UE_TRACE_SLATE_WIDGET_INVALIDATED(ParentProxy.Widget, Widget, EInvalidateWidgetReason::Layout);
-			UpdateList.Push(ParentProxy);
+			UE_TRACE_SLATE_WIDGET_INVALIDATED(ParentWidgetPtr, WidgetPtr, EInvalidateWidgetReason::Layout);
+			UpdateList.PushUnique(ParentProxy);
 		}
 		bWidgetNeedsRepaint = true;
 	}
 	else if (EnumHasAnyFlags(CurrentInvalidateReason, EInvalidateWidgetReason::RenderTransform | EInvalidateWidgetReason::Layout | EInvalidateWidgetReason::Visibility | EInvalidateWidgetReason::ChildOrder))
 	{
-		SCOPE_CYCLE_SWIDGET(Widget);
+		SCOPE_CYCLE_SWIDGET(WidgetPtr);
 		// When layout changes compute a new desired size for this widget
-		FVector2D CurrentDesiredSize = Widget->GetDesiredSize();
+		FVector2D CurrentDesiredSize = WidgetPtr->GetDesiredSize();
 		FVector2D NewDesiredSize = FVector2D::ZeroVector;
 		if (Visibility != EVisibility::Collapsed)
 		{
-			if (Widget->NeedsPrepass() || (!GSlateEnableGlobalInvalidation && Widget->Advanced_IsInvalidationRoot()))
+			if (WidgetPtr->NeedsPrepass() || (!GSlateEnableGlobalInvalidation && WidgetPtr->Advanced_IsInvalidationRoot()))
 			{
-				Widget->SlatePrepass(Widget->PrepassLayoutScaleMultiplier.Get(1.0f));
+				WidgetPtr->SlatePrepass(WidgetPtr->PrepassLayoutScaleMultiplier.Get(1.0f));
 			}
 			else
 			{
-				Widget->CacheDesiredSize(Widget->PrepassLayoutScaleMultiplier.Get(1.0f));
+				WidgetPtr->CacheDesiredSize(WidgetPtr->PrepassLayoutScaleMultiplier.Get(1.0f));
 			}
 
-			NewDesiredSize = Widget->GetDesiredSize();
+			NewDesiredSize = WidgetPtr->GetDesiredSize();
 		}
 
 		// Note even if volatile we need to recompute desired size. We do not need to invalidate parents though if they are volatile since they will naturally redraw this widget
-		if (!Widget->IsVolatileIndirectly() && Visibility.IsVisible())
+		if (!WidgetPtr->IsVolatileIndirectly() && Visibility.IsVisible())
 		{
 			// Set the value directly instead of calling AddUpdateFlags as an optimization
-			Widget->UpdateFlags |= EWidgetUpdateFlags::NeedsRepaint;
+			WidgetPtr->UpdateFlags |= EWidgetUpdateFlags::NeedsRepaint;
 			UpdateFlags |= EWidgetUpdateFlags::NeedsRepaint;
 		}
 
 		// If the desired size changed, invalidate the parent if it is visible
 		if (NewDesiredSize != CurrentDesiredSize || EnumHasAnyFlags(CurrentInvalidateReason, EInvalidateWidgetReason::Visibility|EInvalidateWidgetReason::RenderTransform))
 		{
-			if (ParentIndex != INDEX_NONE)
+			if (ParentIndex != FSlateInvalidationWidgetIndex::Invalid)
 			{
 				FWidgetProxy& ParentProxy = FastWidgetPathList[ParentIndex];
-				if (ParentIndex == 0)
+				if (ParentIndex == FastWidgetPathList.FirstIndex())
 				{
 					// root of the invalidation panel just invalidate the whole thing
-					Root.InvalidateRoot(Widget);
+					Root.InvalidateRoot(WidgetPtr);
 				}
 				else if (ParentProxy.Visibility.IsVisible())
 				{
 					ParentProxy.CurrentInvalidateReason |= EInvalidateWidgetReason::Layout;
 #if WITH_SLATE_DEBUGGING
-					FSlateDebugging::BroadcastWidgetInvalidate(ParentProxy.Widget, Widget, EInvalidateWidgetReason::Layout);
+					FSlateDebugging::BroadcastWidgetInvalidate(ParentProxy.GetWidget(), WidgetPtr, EInvalidateWidgetReason::Layout);
 #endif
-					UE_TRACE_SLATE_WIDGET_INVALIDATED(ParentProxy.Widget, Widget, EInvalidateWidgetReason::Layout);
-					UpdateList.Push(ParentProxy);
+					UE_TRACE_SLATE_WIDGET_INVALIDATED(ParentProxy.GetWidget(), WidgetPtr, EInvalidateWidgetReason::Layout);
+					UpdateList.PushUnique(ParentProxy);
 				}
 			}
-			else if (!GSlateEnableGlobalInvalidation && Widget->IsParentValid())
+			else if (!GSlateEnableGlobalInvalidation && WidgetPtr->IsParentValid())
 			{
-				TSharedPtr<SWidget> ParentWidget = Widget->GetParentWidget();
+				TSharedPtr<SWidget> ParentWidget = WidgetPtr->GetParentWidget();
 				if (ParentWidget->Advanced_IsInvalidationRoot())
 				{
-					Root.InvalidateRoot(Widget);
+					Root.InvalidateRoot(WidgetPtr);
 				}
 			}
 		}
 
 		bWidgetNeedsRepaint = true;
 	}
-	else if (EnumHasAnyFlags(CurrentInvalidateReason, EInvalidateWidgetReason::Paint) && !Widget->IsVolatileIndirectly())
+	else if (EnumHasAnyFlags(CurrentInvalidateReason, EInvalidateWidgetReason::Paint) && !WidgetPtr->IsVolatileIndirectly())
 	{
-		SCOPE_CYCLE_SWIDGET(Widget);
+		SCOPE_CYCLE_SWIDGET(WidgetPtr);
 		// Set the value directly instead of calling AddUpdateFlags as an optimization
-		Widget->UpdateFlags |= EWidgetUpdateFlags::NeedsRepaint;
+		WidgetPtr->UpdateFlags |= EWidgetUpdateFlags::NeedsRepaint;
 		UpdateFlags |= EWidgetUpdateFlags::NeedsRepaint;
 
 		bWidgetNeedsRepaint = true;
@@ -169,7 +181,7 @@ bool FWidgetProxy::ProcessInvalidation(FWidgetUpdateList& UpdateList, TArray<FWi
 	return bWidgetNeedsRepaint;
 }
 
-void FWidgetProxy::MarkProxyUpdatedThisFrame(FWidgetUpdateList& UpdateList)
+void FWidgetProxy::MarkProxyUpdatedThisFrame(FSlateInvalidationWidgetHeap& UpdateList)
 {
 	bUpdatedSinceLastInvalidate = true;
 
@@ -178,7 +190,7 @@ void FWidgetProxy::MarkProxyUpdatedThisFrame(FWidgetUpdateList& UpdateList)
 		if (!bInUpdateList && !bInvisibleDueToParentOrSelfVisibility)
 		{
 			// If there are any updates still needed add them to the next update list
-			UpdateList.Push(*this);
+			UpdateList.PushUnique(*this);
 		}
 	}
 	else
@@ -187,13 +199,16 @@ void FWidgetProxy::MarkProxyUpdatedThisFrame(FWidgetUpdateList& UpdateList)
 	}
 }
 
-int32 FWidgetProxy::Repaint(const FPaintArgs& PaintArgs, int32 MyIndex, FSlateWindowElementList& OutDrawElements) const
+int32 FWidgetProxy::Repaint(const FPaintArgs& PaintArgs, FSlateWindowElementList& OutDrawElements) const
 {
+	SWidget* WidgetPtr = GetWidget();
+	check(WidgetPtr);
+
 #if WITH_SLATE_DEBUGGING
-	SCOPED_NAMED_EVENT_FSTRING(FReflectionMetaData::GetWidgetDebugInfo(Widget), EnumHasAnyFlags(UpdateFlags, EWidgetUpdateFlags::NeedsRepaint) ? FColor::Orange : FColor::Red);
+	SCOPED_NAMED_EVENT_FSTRING(FReflectionMetaData::GetWidgetDebugInfo(WidgetPtr), EnumHasAnyFlags(UpdateFlags, EWidgetUpdateFlags::NeedsRepaint) ? FColor::Orange : FColor::Red);
 #endif
 
-	const FSlateWidgetPersistentState& MyState = Widget->GetPersistentState();
+	const FSlateWidgetPersistentState& MyState = WidgetPtr->GetPersistentState();
 
 	const int32 StartingClipIndex = OutDrawElements.GetClippingIndex();
 
@@ -218,14 +233,14 @@ int32 FWidgetProxy::Repaint(const FPaintArgs& PaintArgs, int32 MyIndex, FSlateWi
 	{
 		if (EnumHasAnyFlags(UpdateFlags, EWidgetUpdateFlags::NeedsVolatilePaint))
 		{
-			if (Widget->ShouldInvalidatePrepassDueToVolatility())
+			if (WidgetPtr->ShouldInvalidatePrepassDueToVolatility())
 			{
-				Widget->InvalidatePrepass();
+				WidgetPtr->InvalidatePrepass();
 			}
-			Widget->SlatePrepass(Widget->GetPrepassLayoutScaleMultiplier());
+			WidgetPtr->SlatePrepass(WidgetPtr->GetPrepassLayoutScaleMultiplier());
 		}
 	}
-	const int32 NewLayerId = Widget->Paint(UpdatedArgs, MyState.AllottedGeometry, MyState.CullingBounds, OutDrawElements, MyState.LayerId, MyState.WidgetStyle, MyState.bParentEnabled);
+	const int32 NewLayerId = WidgetPtr->Paint(UpdatedArgs, MyState.AllottedGeometry, MyState.CullingBounds, OutDrawElements, MyState.LayerId, MyState.WidgetStyle, MyState.bParentEnabled);
 
 	PaintArgs.GetHittestGrid().SetUserIndex(PrevUserIndex);
 
@@ -239,57 +254,72 @@ int32 FWidgetProxy::Repaint(const FPaintArgs& PaintArgs, int32 MyIndex, FSlateWi
 	return NewLayerId;
 }
 
-FWidgetProxyHandle::FWidgetProxyHandle(FSlateInvalidationRoot& InInvalidationRoot, int32 InIndex)
-	: InvalidationRootHandle(InInvalidationRoot.GetInvalidationRootHandle())
-	, MyIndex(InIndex)
-	, GenerationNumber(InInvalidationRoot.GetFastPathGenerationNumber())
+FWidgetProxyHandle::FWidgetProxyHandle(const FSlateInvalidationRootHandle& InInvalidationRoot, FSlateInvalidationWidgetIndex InIndex, FSlateInvalidationWidgetSortOrder InSortIndex)
+	: InvalidationRootHandle(InInvalidationRoot)
+	, WidgetIndex(InIndex)
+	, WidgetSortOrder(InSortIndex)
+{
+	if (InInvalidationRoot.Advanced_GetInvalidationRootNoCheck())
+	{
+		GenerationNumber = InInvalidationRoot.Advanced_GetInvalidationRootNoCheck()->GetFastPathGenerationNumber();
+	}
+}
+
+FWidgetProxyHandle::FWidgetProxyHandle(const FSlateInvalidationRootHandle& InInvalidationRoot, FSlateInvalidationWidgetIndex InIndex, FSlateInvalidationWidgetSortOrder InSortIndex, int32 InGenerationNumber)
+	: InvalidationRootHandle(InInvalidationRoot)
+	, WidgetIndex(InIndex)
+	, WidgetSortOrder(InSortIndex)
+	, GenerationNumber(InGenerationNumber)
+{
+}
+
+FWidgetProxyHandle::FWidgetProxyHandle(FSlateInvalidationWidgetIndex InIndex)
+	: InvalidationRootHandle()
+	, WidgetIndex(InIndex)
+	, WidgetSortOrder()
+	, GenerationNumber(INDEX_NONE)
 {
 
 }
 
-bool FWidgetProxyHandle::IsValid() const
+bool FWidgetProxyHandle::IsValid(const SWidget* Widget) const
 {
 	FSlateInvalidationRoot* InvalidationRoot = InvalidationRootHandle.GetInvalidationRoot();
-	return InvalidationRoot && InvalidationRoot->GetFastPathGenerationNumber() == GenerationNumber && MyIndex != INDEX_NONE;
+	return InvalidationRoot
+		&& InvalidationRoot->GetFastPathGenerationNumber() == GenerationNumber
+		&& InvalidationRoot->GetFastPathWidgetList().IsValidIndex(WidgetIndex)
+		&& InvalidationRoot->GetFastPathWidgetList()[WidgetIndex].GetWidget() == Widget;
 }
 
 FWidgetProxy& FWidgetProxyHandle::GetProxy()
 {
-	check(IsValid());
-	return GetInvalidationRoot()->FastWidgetPathList[MyIndex];
+	return GetInvalidationRoot()->GetFastPathWidgetList()[WidgetIndex];
 }
 
 const FWidgetProxy& FWidgetProxyHandle::GetProxy() const
 {
-	check(IsValid());
-	return GetInvalidationRoot()->FastWidgetPathList[MyIndex];
+	return GetInvalidationRoot()->GetFastPathWidgetList()[WidgetIndex];
 }
 
 void FWidgetProxyHandle::MarkWidgetUpdatedThisFrame()
 {
-	check(IsValid());
-	GetInvalidationRoot()->FastWidgetPathList[MyIndex].MarkProxyUpdatedThisFrame(GetInvalidationRoot()->WidgetsNeedingUpdate);
+	GetInvalidationRoot()->GetFastPathWidgetList()[WidgetIndex].MarkProxyUpdatedThisFrame(*GetInvalidationRoot()->WidgetsNeedingUpdate);
 }
 
 void FWidgetProxyHandle::MarkWidgetDirty(EInvalidateWidgetReason InvalidateReason)
 {
-	check(IsValid());
-	FWidgetProxy& Proxy = GetInvalidationRoot()->FastWidgetPathList[MyIndex];
+	FWidgetProxy& Proxy = GetInvalidationRoot()->GetFastPathWidgetList()[WidgetIndex];
+	SWidget* WidgetPtr = Proxy.GetWidget();
 
 	if (EnumHasAnyFlags(InvalidateReason, EInvalidateWidgetReason::ChildOrder))
 	{
-		/*
-				CSV_EVENT_GLOBAL(TEXT("Slow Path Needed"));
-		#if WITH_SLATE_DEBUGGING
-				UE_LOG(LogSlate, Log, TEXT("Slow Widget Path Needed: %s %s"), *Proxy.Widget->ToString(), *Proxy.Widget->GetTag().ToString());
-		#endif*/
 		Proxy.bChildOrderInvalid = true;
-		GetInvalidationRoot()->InvalidateChildOrder(Proxy.Widget);
+		GetInvalidationRoot()->InvalidateWidgetChildOrder(WidgetPtr->AsShared());
 	}
 
 	if (Proxy.CurrentInvalidateReason == EInvalidateWidgetReason::None)
 	{
-		GetInvalidationRoot()->WidgetsNeedingUpdate.Push(Proxy);
+		GetInvalidationRoot()->WidgetsNeedingUpdate->PushUnique(Proxy);
 	}
 #if 0
 	else
@@ -299,24 +329,26 @@ void FWidgetProxyHandle::MarkWidgetDirty(EInvalidateWidgetReason InvalidateReaso
 #endif
 	Proxy.CurrentInvalidateReason |= InvalidateReason;
 #if WITH_SLATE_DEBUGGING
-	FSlateDebugging::BroadcastWidgetInvalidate(Proxy.Widget, nullptr, InvalidateReason);
+	FSlateDebugging::BroadcastWidgetInvalidate(WidgetPtr, nullptr, InvalidateReason);
 #endif
-	UE_TRACE_SLATE_WIDGET_INVALIDATED(Proxy.Widget, nullptr, InvalidateReason);
+	UE_TRACE_SLATE_WIDGET_INVALIDATED(WidgetPtr, nullptr, InvalidateReason);
 }
 
-void FWidgetProxyHandle::UpdateWidgetFlags(EWidgetUpdateFlags NewFlags)
+void FWidgetProxyHandle::UpdateWidgetFlags(const SWidget* Widget, EWidgetUpdateFlags NewFlags)
 {
-	check(IsValid());
-	FWidgetProxy& Proxy = GetInvalidationRoot()->FastWidgetPathList[MyIndex];
-
-	if (!Proxy.bInvisibleDueToParentOrSelfVisibility)
+	if (IsValid(Widget))
 	{
-		Proxy.UpdateFlags = NewFlags;
+		FWidgetProxy& Proxy = GetInvalidationRoot()->GetFastPathWidgetList()[WidgetIndex];
 
-		// Add to update list if the widget is now tickable or has an active timer.  Disregard if dirty, it was already in the list
-		if (!Proxy.bInUpdateList && EnumHasAnyFlags(NewFlags, EWidgetUpdateFlags::AnyUpdate))
+		if (!Proxy.bInvisibleDueToParentOrSelfVisibility)
 		{
-			GetInvalidationRoot()->WidgetsNeedingUpdate.Push(Proxy);
+			Proxy.UpdateFlags = NewFlags;
+
+			// Add to update list if the widget is now tickable or has an active timer.  Disregard if dirty, it was already in the list
+			if (!Proxy.bInUpdateList && EnumHasAnyFlags(NewFlags, EWidgetUpdateFlags::AnyUpdate))
+			{
+				GetInvalidationRoot()->WidgetsNeedingUpdate->PushUnique(Proxy);
+			}
 		}
 	}
 }
diff --git a/Engine/Source/Runtime/SlateCore/Private/Input/HittestGrid.cpp b/Engine/Source/Runtime/SlateCore/Private/Input/HittestGrid.cpp
--- a/Engine/Source/Runtime/SlateCore/Private/Input/HittestGrid.cpp
+++ b/Engine/Source/Runtime/SlateCore/Private/Input/HittestGrid.cpp
@@ -653,7 +653,7 @@ bool FHittestGrid::SameSize(const FHittestGrid* OtherGrid) const
 	return GridOrigin == OtherGrid->GridOrigin && GridWindowOrigin == OtherGrid->GridWindowOrigin && GridSize == OtherGrid->GridSize;
 }
 
-void FHittestGrid::AddWidget(const TSharedRef<SWidget>& InWidget, int32 InBatchPriorityGroup, int32 InLayerId, int32 InSecondarySort)
+void FHittestGrid::AddWidget(const TSharedRef<SWidget>& InWidget, int32 InBatchPriorityGroup, int32 InLayerId, FSlateInvalidationWidgetSortOrder InSecondarySort)
 {
 	if (!InWidget->GetVisibility().IsHitTestVisible())
 	{
@@ -744,6 +744,14 @@ void FHittestGrid::RemoveWidget(const SWidget* InWidget)
 	RemoveGrid(InWidget);
 }
 
+void FHittestGrid::UpdateWidget(const TSharedRef<SWidget>& InWidget, FSlateInvalidationWidgetSortOrder InSecondarySort)
+{
+	if (int32* FoundWidgetIndex = WidgetMap.Find(&*InWidget))
+	{
+		WidgetArray[*FoundWidgetIndex].SecondarySort = InSecondarySort;
+	}
+}
+
 void FHittestGrid::InsertCustomHitTestPath(const TSharedRef<SWidget> InWidget, TSharedRef<ICustomHitTestPath> CustomHitTestPath)
 {
 	int32 WidgetIndex = WidgetMap.FindChecked(&*InWidget);
@@ -1040,6 +1048,18 @@ void FHittestGrid::DisplayGrid(int32 InLayer, const FGeometry& AllottedGeometry,
 		}
 	}
 }
+
+TArray<FHittestGrid::FWidgetSortData> FHittestGrid::GetAllWidgetSortDatas() const
+{
+	TArray<FWidgetSortData> Result;
+	Result.Reserve(WidgetArray.Num());
+	for (const FWidgetData& WidgetData : WidgetArray)
+	{
+		FWidgetSortData WidgetSortData = { WidgetData.WeakWidget, WidgetData.PrimarySort, WidgetData.SecondarySort };
+		Result.Add(MoveTemp(WidgetSortData));
+	}
+	return Result;
+}
 #endif // WITH_SLATE_DEBUGGING
 
 #undef UE_SLATE_HITTESTGRID_ARRAYSIZEMAX
diff --git a/Engine/Source/Runtime/SlateCore/Private/Test/SlateInvalidationWidgetListTest.cpp b/Engine/Source/Runtime/SlateCore/Private/Test/SlateInvalidationWidgetListTest.cpp
new file mode 100644
--- /dev/null
+++ b/Engine/Source/Runtime/SlateCore/Private/Test/SlateInvalidationWidgetListTest.cpp
@@ -0,0 +1,333 @@
+// Copyright Epic Games, Inc. All Rights Reserved.
+
+#include "CoreMinimal.h"
+#include "Misc/AutomationTest.h"
+
+#include "FastUpdate/SlateInvalidationWidgetList.h"
+#include "Widgets/SBoxPanel.h"
+#include "Widgets/SLeafWidget.h"
+#include "Widgets/SNullWidget.h"
+#include "Widgets/SWindow.h"
+
+
+#if WITH_AUTOMATION_WORKER & WITH_SLATE_DEBUGGING
+
+#define LOCTEXT_NAMESPACE "Slate.FastPath.InvalidationWidgetList"
+
+//IMPLEMENT_SIMPLE_AUTOMATION_TEST(FSlateInvalidationWidgetListTest, "Slate.FastPath.InvalidationWidgetList.AddBuildRemove", EAutomationTestFlags::ApplicationContextMask | EAutomationTestFlags::SmokeFilter)
+IMPLEMENT_SIMPLE_AUTOMATION_TEST(FSlateInvalidationWidgetListTest, "Slate.FastPath.InvalidationWidgetList.AddBuildRemove", EAutomationTestFlags::ApplicationContextMask | EAutomationTestFlags::EngineFilter)
+
+namespace UE
+{
+namespace Slate
+{
+namespace Private
+{
+
+class SEmptyLeftWidget : public SLeafWidget
+{
+public:
+	SLATE_BEGIN_ARGS(SEmptyLeftWidget) {}
+	SLATE_END_ARGS()
+
+	void Construct(const FArguments& InArgs){}
+	virtual FVector2D ComputeDesiredSize(float) const override { return FVector2D{ 100, 100 }; }
+	virtual int32 OnPaint(const FPaintArgs& Args, const FGeometry& AllottedGeometry, const FSlateRect& MyCullingRect, FSlateWindowElementList& OutDrawElements, int32 LayerId, const FWidgetStyle& InWidgetStyle, bool bParentEnabled) const override
+	{
+		return LayerId;
+	}
+};
+
+TSharedRef<SVerticalBox> AddVerticalBox(TSharedPtr<SVerticalBox> VerticalBox, TCHAR Letter)
+{
+	FName NewName = *FString::Printf(TEXT("TagVerticalBox-%c"), Letter);
+	TSharedRef<SVerticalBox> Result = SNew(SVerticalBox).Tag(NewName);
+	VerticalBox->AddSlot()[Result];
+	return Result;
+}
+
+TSharedRef<SWidget> AddEmptyWidget(TSharedPtr<SVerticalBox> VerticalBox, int32 Number)
+{
+	static FName NewName = TEXT("TagEmptyLeftWidget");
+	NewName.SetNumber(Number);
+	TSharedRef<SWidget> Result = SNew(SEmptyLeftWidget).Tag(NewName);
+	VerticalBox->AddSlot()[Result];
+	return Result;
+}
+
+TSharedRef<SVerticalBox> BuildTestUI_ChildOrder(TSharedPtr<SVerticalBox>& WidgetC, TSharedPtr<SVerticalBox>& WidgetF)
+{
+	// A
+	//  B (1, 2, 3)
+	//  C (4, 5, 6, 7)
+	//  Null
+	//  D
+	//  E (8, 9, 10)
+	//  F
+	//   G (11, Null, 12)
+	//   H (13)
+	//   I
+	//  J (14)
+
+	TSharedRef<SVerticalBox> Root = SNew(SVerticalBox);
+	{
+		TSharedRef<SVerticalBox> Sub = AddVerticalBox(Root, TEXT('B'));
+		AddEmptyWidget(Sub, 1);
+		AddEmptyWidget(Sub, 2);
+		AddEmptyWidget(Sub, 3);
+	}
+	{
+		WidgetC = AddVerticalBox(Root, TEXT('C'));
+		AddEmptyWidget(WidgetC, 4);
+		AddEmptyWidget(WidgetC, 5);
+		AddEmptyWidget(WidgetC, 6);
+		AddEmptyWidget(WidgetC, 7);
+	}
+	Root->AddSlot()[SNullWidget::NullWidget];
+	AddVerticalBox(Root, TEXT('D'));
+	{
+		TSharedRef<SVerticalBox> Sub = AddVerticalBox(Root, TEXT('E'));
+		AddEmptyWidget(Sub, 8);
+		AddEmptyWidget(Sub, 9);
+		AddEmptyWidget(Sub, 10);
+	}
+	{
+		WidgetF = AddVerticalBox(Root, TEXT('F'));
+		{
+			TSharedRef<SVerticalBox> SubSub = AddVerticalBox(WidgetF, TEXT('G'));
+			AddEmptyWidget(SubSub, 11);
+			SubSub->AddSlot()[SNullWidget::NullWidget];
+			AddEmptyWidget(SubSub, 12);
+		}
+		{
+			TSharedRef<SVerticalBox> SubSub = AddVerticalBox(WidgetF, TEXT('H'));
+			AddEmptyWidget(SubSub, 13);
+		}
+		{
+			AddVerticalBox(WidgetF, TEXT('I'));
+		}
+	}
+	{
+		TSharedRef<SVerticalBox> SubSub = AddVerticalBox(Root, TEXT('J'));
+		AddEmptyWidget(SubSub, 14);
+	}
+
+	return Root;
+}
+
+TSharedRef<SVerticalBox> BuildTestUI_Child(TSharedPtr<SVerticalBox>& WidgetA, TArray<TSharedPtr<SWidget>>& ChildOfWidgetA
+	, TSharedPtr<SVerticalBox>& WidgetB, TArray<TSharedPtr<SWidget>>& ChildOfWidgetB
+	, TSharedPtr<SVerticalBox>& WidgetC, TArray<TSharedPtr<SWidget>>& ChildOfWidgetC
+	, TSharedPtr<SVerticalBox>& WidgetD, TArray<TSharedPtr<SWidget>>& ChildOfWidgetD
+	, TSharedPtr<SVerticalBox>& WidgetH, TArray<TSharedPtr<SWidget>>& ChildOfWidgetH)
+{
+	// A
+	//  B (1, 2, 3)
+	//  C
+	//  D
+	//   E (4, 5)
+	//   F (6)
+	//   G (7)
+	//  H (8)
+	//  I
+
+	WidgetA = SNew(SVerticalBox);
+	{
+		WidgetB = AddVerticalBox(WidgetA, TEXT('B'));
+		ChildOfWidgetA.Add(WidgetB);
+		ChildOfWidgetB.Add(AddEmptyWidget(WidgetB, 1));
+		ChildOfWidgetB.Add(AddEmptyWidget(WidgetB, 2));
+		ChildOfWidgetB.Add(AddEmptyWidget(WidgetB, 3));
+	}
+	{
+		WidgetC = AddVerticalBox(WidgetA, TEXT('C'));
+		ChildOfWidgetA.Add(WidgetC);
+	}
+	{
+		WidgetD = AddVerticalBox(WidgetA, TEXT('D'));
+		ChildOfWidgetA.Add(WidgetD);
+		{
+			TSharedRef<SVerticalBox> SubSub = AddVerticalBox(WidgetD, TEXT('E'));
+			ChildOfWidgetD.Add(SubSub);
+			AddEmptyWidget(SubSub, 4);
+			AddEmptyWidget(SubSub, 5);
+		}
+		{
+			TSharedRef<SVerticalBox> SubSub = AddVerticalBox(WidgetD, TEXT('F'));
+			ChildOfWidgetD.Add(SubSub);
+			AddEmptyWidget(SubSub, 6);
+		}
+		{
+			TSharedRef<SVerticalBox> SubSub = AddVerticalBox(WidgetD, TEXT('G'));
+			ChildOfWidgetD.Add(SubSub);
+			AddEmptyWidget(SubSub, 7);
+		}
+	}
+	{
+		WidgetH = AddVerticalBox(WidgetA, TEXT('H'));
+		ChildOfWidgetA.Add(WidgetH);
+		ChildOfWidgetH.Add(AddEmptyWidget(WidgetH, 8));
+	}
+	ChildOfWidgetA.Add(AddVerticalBox(WidgetA, TEXT('I')));
+	return WidgetA.ToSharedRef();
+}
+}
+}
+}
+
+
+bool FSlateInvalidationWidgetListTest::RunTest(const FString& Parameters)
+{
+	FSlateInvalidationWidgetList::FArguments ArgsToTest[] = {{4,2}, {4,3}, {5,1}, {6, 1}};
+	for (FSlateInvalidationWidgetList::FArguments Args : ArgsToTest)
+	{
+		{
+			TSharedPtr<SVerticalBox> WidgetA, WidgetB, WidgetC, WidgetD, WidgetH;
+			TArray<TSharedPtr<SWidget>> ChildOfWidgetA, ChildOfWidgetB, ChildOfWidgetC, ChildOfWidgetD, ChildOfWidgetH;
+			TSharedRef<SVerticalBox> RootChild = UE::Slate::Private::BuildTestUI_Child(WidgetA, ChildOfWidgetA, WidgetB, ChildOfWidgetB, WidgetC, ChildOfWidgetC, WidgetD, ChildOfWidgetD, WidgetH, ChildOfWidgetH);
+
+			FSlateInvalidationWidgetList List = { FSlateInvalidationRootHandle(), Args };
+			List.BuildWidgetList(RootChild);
+			AddErrorIfFalse(List.VerifyWidgetsIndex(), TEXT("The widget list integrity has failed."));
+
+			{
+				TArray<TSharedPtr<SWidget>> FoundChildren = List.FindChildren(WidgetA.ToSharedRef());
+				if (ChildOfWidgetA != FoundChildren)
+				{
+					AddError(TEXT("Was not able to find the child of VerticalBox A."));
+				}
+			}
+			{
+				TArray<TSharedPtr<SWidget>> FoundChildren = List.FindChildren(WidgetB.ToSharedRef());
+				if (ChildOfWidgetB != FoundChildren)
+				{
+					AddError(TEXT("Was not able to find the child of VerticalBox B."));
+				}
+			}
+			{
+				TArray<TSharedPtr<SWidget>> FoundChildren = List.FindChildren(WidgetC.ToSharedRef());
+				if (ChildOfWidgetC != FoundChildren)
+				{
+					AddError(TEXT("Was not able to find the child of VerticalBox C."));
+				}
+			}
+			{
+				TArray<TSharedPtr<SWidget>> FoundChildren = List.FindChildren(WidgetD.ToSharedRef());
+				if (ChildOfWidgetD != FoundChildren)
+				{
+					AddError(TEXT("Was not able to find the child of VerticalBox D."));
+				}
+			}
+			{
+				TArray<TSharedPtr<SWidget>> FoundChildren = List.FindChildren(WidgetH.ToSharedRef());
+				if (ChildOfWidgetH != FoundChildren)
+				{
+					AddError(TEXT("Was not able to find the child of VerticalBox H."));
+				}
+			}
+		}
+
+		{
+			TSharedPtr<SVerticalBox> WidgetC, WidgetF;
+			TSharedRef<SVerticalBox> RootChildOrder = UE::Slate::Private::BuildTestUI_ChildOrder(WidgetC, WidgetF);
+
+			FSlateInvalidationWidgetList List = { FSlateInvalidationRootHandle(), Args };
+			List.BuildWidgetList(RootChildOrder);
+			AddErrorIfFalse(List.VerifyWidgetsIndex(), TEXT("The widget list integrity has failed."));
+
+			// Remove Second child of F
+			{
+				List.RemoveWidget(WidgetF->GetAllChildren()->GetChildAt(1));
+				WidgetF->RemoveSlot(WidgetF->GetAllChildren()->GetChildAt(1));
+				{
+					FSlateInvalidationWidgetList TempList = { FSlateInvalidationRootHandle(), Args };
+					TempList.BuildWidgetList(RootChildOrder);
+					if (!TempList.DeapCompare(List))
+					{
+						AddError(TEXT("Was not able to remove a child of F."));
+					}
+				}
+			}
+
+			// Remove C and F
+			{
+				{
+					const FSlateInvalidationWidgetIndex WidgetIndexF = List.FindWidget(WidgetF.ToSharedRef());
+					if (!List.IsValidIndex(WidgetIndexF) || List[WidgetIndexF].GetWidget() != WidgetF.Get())
+					{
+						AddError(TEXT("The index of F is not valid anymore."));
+					}
+					List.RemoveWidget(WidgetIndexF);
+				}
+				{
+					const FSlateInvalidationWidgetIndex WidgetIndexC = List.FindWidget(WidgetC.ToSharedRef());
+					if (!List.IsValidIndex(WidgetIndexC) || List[WidgetIndexC].GetWidget() != WidgetC.Get())
+					{
+						AddError(TEXT("The index of C is not valid anymore."));
+					}
+					List.RemoveWidget(WidgetIndexC);
+				}
+				RootChildOrder->RemoveSlot(WidgetF.ToSharedRef());
+				RootChildOrder->RemoveSlot(WidgetC.ToSharedRef());
+				{
+					FSlateInvalidationWidgetList TempList = { FSlateInvalidationRootHandle(), Args };
+					TempList.BuildWidgetList(RootChildOrder);
+					if (!TempList.DeapCompare(List))
+					{
+						AddError(TEXT("Was not able to remove F and C."));
+					}
+				}
+			}
+			// remove last item
+			{
+				int32 ToRemoveIndex = RootChildOrder->GetAllChildren()->Num() - 1;
+				TSharedRef<SWidget> RemovedWidget = RootChildOrder->GetAllChildren()->GetChildAt(ToRemoveIndex);
+				RootChildOrder->RemoveSlot(RemovedWidget);
+				List.RemoveWidget(List.FindWidget(RemovedWidget));
+
+				{
+					FSlateInvalidationWidgetList TempList = { FSlateInvalidationRootHandle(), Args };
+					TempList.BuildWidgetList(RootChildOrder);
+					if (!TempList.DeapCompare(List))
+					{
+						AddError(TEXT("Was not able to remove the last item of A."));
+					}
+				}
+			}
+		}
+		{
+			TSharedPtr<SVerticalBox> WidgetC, WidgetF;
+			TSharedRef<SVerticalBox> RootChildOrder = UE::Slate::Private::BuildTestUI_ChildOrder(WidgetC, WidgetF);
+
+			// Child invalidation
+			FSlateInvalidationWidgetList List = { FSlateInvalidationRootHandle(), Args };
+			List.BuildWidgetList(RootChildOrder);
+			AddErrorIfFalse(List.VerifyWidgetsIndex(), TEXT("The widget list integrity has failed."));
+
+			WidgetF->RemoveSlot(WidgetF->GetAllChildren()->GetChildAt(0));
+			WidgetC->RemoveSlot(WidgetC->GetAllChildren()->GetChildAt(1));
+
+			TArray<TWeakPtr<SWidget>> InvalidationWidgetIndexes;
+			InvalidationWidgetIndexes.Add(WidgetF->AsShared());
+			InvalidationWidgetIndexes.Add(WidgetC->AsShared());
+			InvalidationWidgetIndexes.Add(WidgetC->AsShared());
+			List.ProcessChildOrderInvalidation(InvalidationWidgetIndexes);
+			AddErrorIfFalse(List.VerifyWidgetsIndex(), TEXT("The widget list integrity has failed."));
+
+			{
+				FSlateInvalidationWidgetList TempList = { FSlateInvalidationRootHandle(), Args };
+				TempList.BuildWidgetList(RootChildOrder);
+				if (!TempList.DeapCompare(List))
+				{
+					AddError(TEXT("Was not able to process invalidation C and F."));
+				}
+			}
+		}
+	}
+
+	return true;
+}
+
+#undef LOCTEXT_NAMESPACE 
+
+#endif //WITH_AUTOMATION_WORKER & WITH_SLATE_DEBUGGING
diff --git a/Engine/Source/Runtime/SlateCore/Private/Widgets/SWidget.cpp b/Engine/Source/Runtime/SlateCore/Private/Widgets/SWidget.cpp
--- a/Engine/Source/Runtime/SlateCore/Private/Widgets/SWidget.cpp
+++ b/Engine/Source/Runtime/SlateCore/Private/Widgets/SWidget.cpp
@@ -164,7 +164,7 @@ void SWidget::UpdateWidgetProxy(int32 NewLayerId, FSlateCachedElementsHandle& Ca
 	}
 	PersistentState.CachedElementHandle = CacheHandle;
 
-	if (FastPathProxyHandle.IsValid())
+	if (FastPathProxyHandle.IsValid(this))
 	{
 		FWidgetProxy& MyProxy = FastPathProxyHandle.GetProxy();
 
@@ -265,11 +265,7 @@ SWidget::~SWidget()
 		}
 
 		// Warn the invalidation root
-		FSlateInvalidationRootHandle SlateInvalidationRootHandle = FastPathProxyHandle.GetInvalidationRootHandle();
-#if WITH_SLATE_DEBUGGING
-		ensure(!SlateInvalidationRootHandle.IsStale());
-#endif
-		if (FSlateInvalidationRoot* InvalidationRoot = SlateInvalidationRootHandle.GetInvalidationRoot())
+		if (FSlateInvalidationRoot* InvalidationRoot = FastPathProxyHandle.GetInvalidationRootHandle().GetInvalidationRoot())
 		{
 			InvalidationRoot->OnWidgetDestroyed(this);
 		}
@@ -728,75 +724,39 @@ bool SWidget::ConditionallyDetatchParentWidget(SWidget* InExpectedParent)
 	return false;
 }
 
-bool SWidget::AssignIndicesToChildren(FSlateInvalidationRoot& Root, int32 ParentIndex, TArray<FWidgetProxy, TMemStackAllocator<>>& FastPathList, bool bParentVisible, bool bParentVolatile)
+void SWidget::SetFastPathSortOrder(const FSlateInvalidationWidgetSortOrder InSortOrder)
 {
-	// Because null widgets are a shared static widget they are not valid for the fast path and are treated as non-existent
-	if (this == &SNullWidget::NullWidget.Get())
+	if (InSortOrder != FastPathProxyHandle.GetWidgetSortOrder())
 	{
-		return false;
-	}
+		FastPathProxyHandle.WidgetSortOrder = InSortOrder;
+		if (FSlateInvalidationRoot* Root = FastPathProxyHandle.GetInvalidationRootHandle().GetInvalidationRoot())
+		{
+			if (FHittestGrid* HittestGrid = Root->GetHittestGrid())
+			{
+				HittestGrid->UpdateWidget(AsShared(), InSortOrder);
+			}
+		}
 
-	if (FastPathProxyHandle.IsValid())
-	{
-		ensureAlwaysMsgf(!FastPathProxyHandle.IsValid(), TEXT("Widget %s was already assigned a proxy handle. If this is being hit this widget is in the active slate tree more than once.  This is illegal and at best will result in UI corruption."), *FReflectionMetaData::GetWidgetDebugInfo(*this));
-		return false;
+		//TODO, update Cached LayerId
 	}
+}
 
-	FWidgetProxy MyProxy(*this);
-	MyProxy.Index = FastPathList.Num();
-	MyProxy.ParentIndex = ParentIndex;
-	MyProxy.Visibility = GetVisibility();
+void SWidget::SetFastPathProxyHandle(const FWidgetProxyHandle& Handle, bool bInInvisibleDueToParentOrSelfVisibility, bool bParentVolatile)
+{
+	check(this != &SNullWidget::NullWidget.Get());
 
-	check(ParentIndex != MyProxy.Index);
+	FastPathProxyHandle = Handle;
 
-	// If this method is being called, child order changed.  Initial visibility and volatility needs to be propagated
-	// Update visibility
-	const bool bParentAndSelfVisible = bParentVisible && MyProxy.Visibility.IsVisible();
-	const bool bWasInvisible = bInvisibleDueToParentOrSelfVisibility;
-	bInvisibleDueToParentOrSelfVisibility = !bParentAndSelfVisible;
-	MyProxy.bInvisibleDueToParentOrSelfVisibility = bInvisibleDueToParentOrSelfVisibility;
-
-	// Update volatility
+	bInvisibleDueToParentOrSelfVisibility = bInInvisibleDueToParentOrSelfVisibility;
 	bInheritedVolatility = bParentVolatile;
 
-	FastPathProxyHandle = FWidgetProxyHandle(Root, MyProxy.Index);
-
-	if (bInvisibleDueToParentOrSelfVisibility&& PersistentState.CachedElementHandle.IsValid())
+	if (bInvisibleDueToParentOrSelfVisibility && PersistentState.CachedElementHandle.IsValid())
 	{
 #if WITH_SLATE_DEBUGGING
 		check(PersistentState.CachedElementHandle.IsOwnedByWidget(this));
 #endif
 		PersistentState.CachedElementHandle.RemoveFromCache();
 	}
-
-	FastPathList.Add(MyProxy);
-
-	// Don't recur into children if we are at a different invalidation root(nested invalidation panels) than where we started and not at the root of the tree. Those children should belong to that roots tree.
-	if (!Advanced_IsInvalidationRoot() || ParentIndex == INDEX_NONE)
-	{
-		FChildren* MyChildren = GetAllChildren();
-		const int32 NumChildren = MyChildren->Num();
-
-		int32 NumChildrenValidForFastPath = 0;
-		for (int32 ChildIndex = 0; ChildIndex < NumChildren; ++ChildIndex)
-		{
-			// Because null widgets are a shared static widget they are not valid for the fast path and are treated as non-existent
-			const TSharedRef<SWidget>& Child = MyChildren->GetChildAt(ChildIndex);
-			if (Child->AssignIndicesToChildren(Root, MyProxy.Index, FastPathList, bParentAndSelfVisible, bParentVolatile || IsVolatile()))
-			{
-				NumChildrenValidForFastPath++;
-			}
-		}
-
-		{
-			FWidgetProxy& MyProxyRef = FastPathList[MyProxy.Index];
-			MyProxyRef.NumChildren = NumChildrenValidForFastPath;
-			int32 LastIndex = FastPathList.Num() - 1;
-			MyProxyRef.LeafMostChildIndex = LastIndex != MyProxy.Index ? LastIndex : INDEX_NONE;
-		}
-	}
-
-	return true;
 }
 
 void SWidget::UpdateFastPathVisibility(bool bParentVisible, bool bWidgetRemoved, FHittestGrid* ParentHittestGrid)
@@ -808,18 +768,13 @@ void SWidget::UpdateFastPathVisibility(bool bParentVisible, bool bWidgetRemoved,
 	const bool bVisibilityChanged = bWasInvisible != bInvisibleDueToParentOrSelfVisibility;
 
 	FHittestGrid* HittestGridToRemoveFrom = ParentHittestGrid;
-	if (FastPathProxyHandle.IsValid())
+	if (FastPathProxyHandle.IsValid(this))
 	{	
 		// Try and remove this from the current handles hit test grid.  If we are in a nested invalidation situation the hittest grid may have changed
 		HittestGridToRemoveFrom = FastPathProxyHandle.GetInvalidationRoot()->GetHittestGrid();
 		FWidgetProxy& Proxy = FastPathProxyHandle.GetProxy();
 		Proxy.Visibility = CurrentVisibility;
 		Proxy.bInvisibleDueToParentOrSelfVisibility = bInvisibleDueToParentOrSelfVisibility;
-
-		if (bWidgetRemoved)
-		{
-			FastPathProxyHandle.GetInvalidationRoot()->RemoveWidgetFromFastPath(Proxy);
-		}
 	}
 	else if (bWidgetRemoved)
 	{
@@ -1147,7 +1102,7 @@ void SWidget::Invalidate(EInvalidateWidgetReason InvalidateReason)
 		InvalidatePrepass();
 	}
 
-	if(FastPathProxyHandle.IsValid())
+	if(FastPathProxyHandle.IsValid(this))
 	{
 		// Current thinking is that visibility and volatility should be updated right away, not during fast path invalidation processing next frame
 		if (EnumHasAnyFlags(InvalidateReason, EInvalidateWidgetReason::Visibility))
@@ -1344,7 +1299,7 @@ int32 SWidget::Paint(const FPaintArgs& Args, const FGeometry& AllottedGeometry,
 #endif
 
 #if WITH_SLATE_DEBUGGING
-	if (!FastPathProxyHandle.IsValid() && PersistentState.CachedElementHandle.IsValid())
+	if (!FastPathProxyHandle.IsValid(this) && PersistentState.CachedElementHandle.IsValid())
 	{
 		ensure(!bInvisibleDueToParentOrSelfVisibility);
 	}
@@ -1354,7 +1309,7 @@ int32 SWidget::Paint(const FPaintArgs& Args, const FGeometry& AllottedGeometry,
 
 	if (bOutgoingHittestability)
 	{
-		Args.GetHittestGrid().AddWidget(MutableThis, 0, LayerId, FastPathProxyHandle.GetIndex());
+		Args.GetHittestGrid().AddWidget(MutableThis, 0, LayerId, FastPathProxyHandle.GetWidgetSortOrder());
 	}
 
 	if (bClipToBounds)
diff --git a/Engine/Source/Runtime/SlateCore/Private/Widgets/SWindow.cpp b/Engine/Source/Runtime/SlateCore/Private/Widgets/SWindow.cpp
--- a/Engine/Source/Runtime/SlateCore/Private/Widgets/SWindow.cpp
+++ b/Engine/Source/Runtime/SlateCore/Private/Widgets/SWindow.cpp
@@ -713,7 +713,7 @@ void SWindow::SetAllowFastUpdate(bool bInAllowFastUpdate)
 		bAllowFastUpdate = bInAllowFastUpdate;
 		if (bAllowFastUpdate)
 		{
-			InvalidateChildOrder();
+			InvalidateRoot();
 		}
 	}
 }
diff --git a/Engine/Source/Runtime/SlateCore/Public/FastUpdate/SlateInvalidationRoot.h b/Engine/Source/Runtime/SlateCore/Public/FastUpdate/SlateInvalidationRoot.h
--- a/Engine/Source/Runtime/SlateCore/Public/FastUpdate/SlateInvalidationRoot.h
+++ b/Engine/Source/Runtime/SlateCore/Public/FastUpdate/SlateInvalidationRoot.h
@@ -6,9 +6,12 @@
 #include "UObject/GCObject.h"
 #include "WidgetProxy.h"
 #include "FastUpdate/SlateInvalidationRootHandle.h"
+#include "FastUpdate/SlateInvalidationWidgetIndex.h"
 #include "Rendering/DrawElements.h"
 
 struct FSlateCachedElementData;
+class FSlateInvalidationWidgetHeap;
+class FSlateInvalidationWidgetList;
 class FSlateWindowElementList;
 class FWidgetStyle;
 
@@ -75,21 +78,18 @@ public:
 
 	SLATECORE_API void InvalidateRoot(const SWidget* Investigator = nullptr);
 
-	/** Call to notify that the ordering of children somewhere in the hierarchy below this root has changed and the fast path is no longer valid */
-	SLATECORE_API void InvalidateChildOrder(const SWidget* Investigator = nullptr);
-
 	SLATECORE_API void InvalidateScreenPosition(const SWidget* Investigator = nullptr);
 
 	bool NeedsSlowPath() const { return bNeedsSlowPath; }
 
-	void RemoveWidgetFromFastPath(FWidgetProxy& Proxy);
-
 	/** @return the cached draw elements for this window and its widget hierarchy*/
 	FSlateCachedElementData& GetCachedElements() { return *CachedElementData; }
 	const FSlateCachedElementData& GetCachedElements() const { return *CachedElementData; }
 	const SWidget* GetInvalidationRootWidget() const;
 	int32 GetFastPathGenerationNumber() const { return FastPathGenerationNumber; }
 	FSlateInvalidationRootHandle GetInvalidationRootHandle() const { return InvalidationRootHandle; }
+	FSlateInvalidationWidgetList& GetFastPathWidgetList() { return *FastWidgetPathList; }
+	const FSlateInvalidationWidgetList& GetFastPathWidgetList() const { return *FastWidgetPathList; }
 
 	SLATECORE_API FSlateInvalidationResult PaintInvalidationRoot(const FSlateInvalidationContext& Context);
 
@@ -123,17 +123,22 @@ protected:
 private:
 	bool PaintFastPath(const FSlateInvalidationContext& Context);
 
-	void BuildFastPathList(SWidget* RootWidget);
-	bool BuildNewFastPathList_Recursive(FSlateInvalidationRoot& Root, FWidgetProxy& Proxy, int32 ParentIndex, int32& NextTreeIndex, TArray<FWidgetProxy>& CurrentFastPathList, TArray<FWidgetProxy, TMemStackAllocator<>>& NewFastPathList);
+	/** Call to notify that the ordering of children below this Widget has changed and the fast path is no longer valid. */
+	void InvalidateWidgetChildOrder(TSharedRef<SWidget> Widget);
+	void ProcessChildOrderUpdate();
+	void BuildFastPathWidgetList(TSharedRef<SWidget> RootWidget);
 
 	void AdjustWidgetsDesktopGeometry(FVector2D WindowToDesktopTransform);
 
 private:
-	TArray<FWidgetProxy> FastWidgetPathList;
+	/** List of all the Widget included by this SlateInvalidationRoot. */
+	FSlateInvalidationWidgetList* FastWidgetPathList;
 	/** Index to widgets which are dirty, volatile, or need some sort of per frame update (such as a tick or timer) */
-	FWidgetUpdateList WidgetsNeedingUpdate;
-
-	TArray<int32, TInlineAllocator<100>> FinalUpdateList;
+	FSlateInvalidationWidgetHeap* WidgetsNeedingUpdate;
+	/** Index to widgets that will be updated. */
+	TArray<FSlateInvalidationWidgetIndex> FinalUpdateList;
+	/** Widget that has ChildOrder invalidation. */
+	TArray<TWeakPtr<SWidget>> WidgetsNeedingChildOrderUpdate;
 
 	FSlateCachedElementData* CachedElementData;
 
@@ -155,6 +160,7 @@ private:
 	bool bChildOrderInvalidated;
 	bool bNeedsSlowPath;
 	bool bNeedScreenPositionShift;
+	bool bProcessingChildOrderUpdate;
 
 #if WITH_SLATE_DEBUGGING
 	ESlateInvalidationPaintType LastPaintType;
diff --git a/Engine/Source/Runtime/SlateCore/Public/FastUpdate/SlateInvalidationWidgetIndex.h b/Engine/Source/Runtime/SlateCore/Public/FastUpdate/SlateInvalidationWidgetIndex.h
new file mode 100644
--- /dev/null
+++ b/Engine/Source/Runtime/SlateCore/Public/FastUpdate/SlateInvalidationWidgetIndex.h
@@ -0,0 +1,29 @@
+// Copyright Epic Games, Inc. All Rights Reserved.
+
+#pragma once
+
+#include "CoreMinimal.h"
+
+class SWidget;
+
+/** Index of the special container to order widget in InvalidateRoot. See SlateInvalidationWidgetlist. */
+struct FSlateInvalidationWidgetIndex
+{
+	friend class FSlateInvalidationWidgetList;
+	friend struct FSlateInvalidationWidgetSortOrder;
+	
+public:
+	static const FSlateInvalidationWidgetIndex Invalid;
+	bool operator== (FSlateInvalidationWidgetIndex Other) const { return ArrayIndex == Other.ArrayIndex && ElementIndex == Other.ElementIndex; }
+	bool operator!= (FSlateInvalidationWidgetIndex Other) const { return ArrayIndex != Other.ArrayIndex || ElementIndex != Other.ElementIndex; }
+	
+private:
+	using IndexType = uint16;
+
+	FSlateInvalidationWidgetIndex(IndexType InArray, IndexType InElement) : ArrayIndex(InArray), ElementIndex(InElement) {}
+
+	/** Index of the array node in FSlateInvalidationWidgetList */
+	IndexType ArrayIndex;
+	/** Index of the element node in FSlateInvalidationWidgetList */
+	IndexType ElementIndex;
+};
diff --git a/Engine/Source/Runtime/SlateCore/Public/FastUpdate/SlateInvalidationWidgetSortOrder.h b/Engine/Source/Runtime/SlateCore/Public/FastUpdate/SlateInvalidationWidgetSortOrder.h
new file mode 100644
--- /dev/null
+++ b/Engine/Source/Runtime/SlateCore/Public/FastUpdate/SlateInvalidationWidgetSortOrder.h
@@ -0,0 +1,31 @@
+// Copyright Epic Games, Inc. All Rights Reserved.
+
+#pragma once
+
+#include "CoreMinimal.h"
+#include "FastUpdate/SlateInvalidationWidgetIndex.h"
+
+class FSlateInvalidationWidgetList;
+
+/**
+ * SlateInvalidationWidgetIndex cannot be used to sort the widget since the ArrayIndex may not be in the expected order. (See the array as a double linked list).
+ * SlateInvalidationWidgetSortOrder builds a unique number that represents the order of the widget.
+ * The number is padded in a way to keep the order but not necessarily sequential.
+ * It is valid until the next SlateInvalidationRoot::ProcessInvalidation()
+ */
+struct FSlateInvalidationWidgetSortOrder
+{
+private:
+	uint32 Order = 0;
+
+public:
+	FSlateInvalidationWidgetSortOrder() = default;
+	FSlateInvalidationWidgetSortOrder(const FSlateInvalidationWidgetList& List, FSlateInvalidationWidgetIndex Index);
+
+	bool operator< (const FSlateInvalidationWidgetSortOrder Other) const { return Order < Other.Order; }
+	bool operator<= (const FSlateInvalidationWidgetSortOrder Other) const { return Order <= Other.Order; }
+	bool operator> (const FSlateInvalidationWidgetSortOrder Other) const { return Order > Other.Order; }
+	bool operator>= (const FSlateInvalidationWidgetSortOrder Other) const { return Order >= Other.Order; }
+	bool operator== (const FSlateInvalidationWidgetSortOrder Other) const { return Order == Other.Order; }
+	bool operator!= (const FSlateInvalidationWidgetSortOrder Other) const { return Order != Other.Order; }
+};
\ No newline at end of file
diff --git a/Engine/Source/Runtime/SlateCore/Public/FastUpdate/WidgetProxy.h b/Engine/Source/Runtime/SlateCore/Public/FastUpdate/WidgetProxy.h
--- a/Engine/Source/Runtime/SlateCore/Public/FastUpdate/WidgetProxy.h
+++ b/Engine/Source/Runtime/SlateCore/Public/FastUpdate/WidgetProxy.h
@@ -7,6 +7,8 @@
 #include "Styling/WidgetStyle.h"
 #include "Misc/MemStack.h"
 #include "FastUpdate/SlateInvalidationRootHandle.h"
+#include "FastUpdate/SlateInvalidationWidgetIndex.h"
+#include "FastUpdate/SlateInvalidationWidgetSortOrder.h"
 #include "FastUpdate/WidgetUpdateFlags.h"
 #include "Layout/Clipping.h"
 #include "Layout/FlowDirection.h"
@@ -15,35 +17,55 @@
 class SWidget;
 class FPaintArgs;
 struct FFastPathPerFrameData;
+class FSlateInvalidationWidgetHeap;
+class FSlateInvalidationWidgetList;
 struct FSlateWidgetPersistentState;
-struct FWidgetUpdateList;
 class FSlateInvalidationRoot;
 
 enum class EInvalidateWidgetReason : uint8;
 
+#define UE_SLATE_WITH_WIDGETPROXY_WEAKPTR 0
+#define UE_SLATE_VERIFY_WIDGETPROXY_WEAKPTR_STALE 0
+
 class FWidgetProxy
 {
 public:
-	FWidgetProxy(SWidget& InWidget);
+	FWidgetProxy(TSharedRef<SWidget>& InWidget);
+
+	int32 Update(const FPaintArgs& PaintArgs, FSlateWindowElementList& OutDrawElements);
 
-	int32 Update(
-		const FPaintArgs& PaintArgs,
-		int32 MyIndex,
-		FSlateWindowElementList& OutDrawElements);
+	bool ProcessInvalidation(FSlateInvalidationWidgetHeap& UpdateList, FSlateInvalidationWidgetList& FastPathWidgetList, FSlateInvalidationRoot& Root);
 
-	bool ProcessInvalidation(FWidgetUpdateList& UpdateList, TArray<FWidgetProxy>& FastPathWidgetList, FSlateInvalidationRoot& Root);
+	void MarkProxyUpdatedThisFrame(FSlateInvalidationWidgetHeap& UpdateList);
 
-	void MarkProxyUpdatedThisFrame(FWidgetUpdateList& UpdateList);
+#if UE_SLATE_WITH_WIDGETPROXY_WEAKPTR
+	SWidget* GetWidget(bool bTestForStale = true) const
+	{
+#if UE_SLATE_VERIFY_WIDGETPROXY_WEAKPTR_STALE
+		ensureAlways(!bTestForStale || !Widget.IsStale()); // Widget.Object != nullptr && !Widget.WeakRerenceCount.IsValid()
+#endif
+		return Widget.Pin().Get();
+	}
+	void ResetWidget() { Widget.Reset(); }
+#else
+	SWidget* GetWidget() const { return Widget; }
+	void ResetWidget() { Widget = nullptr; }
+#endif
 
 private:
-	int32 Repaint(const FPaintArgs& PaintArgs, int32 MyIndex, FSlateWindowElementList& OutDrawElements) const;
+	int32 Repaint(const FPaintArgs& PaintArgs, FSlateWindowElementList& OutDrawElements) const;
 
-public:
+private:
+#if UE_SLATE_WITH_WIDGETPROXY_WEAKPTR
+	TWeakPtr<SWidget> Widget;
+#else
 	SWidget* Widget;
-	int32 Index;
-	int32 ParentIndex;
-	int32 NumChildren;
-	int32 LeafMostChildIndex;
+#endif
+
+public:
+	FSlateInvalidationWidgetIndex Index;
+	FSlateInvalidationWidgetIndex ParentIndex;
+	FSlateInvalidationWidgetIndex LeafMostChildIndex;
 	EWidgetUpdateFlags UpdateFlags;
 	EInvalidateWidgetReason CurrentInvalidateReason;
 	/** The widgets own visibility */
@@ -54,61 +76,15 @@ public:
 	uint8 bInUpdateList : 1;
 	uint8 bInvisibleDueToParentOrSelfVisibility : 1;
 	uint8 bChildOrderInvalid : 1;
+	uint8 bContainedByWidgetHeap : 1;
 };
 
 static_assert(sizeof(FWidgetProxy) <= 32, "FWidgetProxy should be 32 bytes");
 
+#if !UE_SLATE_WITH_WIDGETPROXY_WEAKPTR
 static_assert(TIsTriviallyDestructible<FWidgetProxy>::Value == true, "FWidgetProxy must be trivially destructible");
-
 template <> struct TIsPODType<FWidgetProxy> { enum { Value = true }; };
-
-struct FWidgetUpdateList
-{
-public:
-	inline void Push(FWidgetProxy& Proxy)
-	{
-		if (!Heap.Contains(Proxy.Index))
-		{
-			Proxy.bInUpdateList = true;
-			Heap.HeapPush(Proxy.Index, TGreater<>());
-		}
-	}
-
-	inline int32 Pop()
-	{
-		int32 Index;
-		Heap.HeapPop(Index, TGreater<>(), false);
-		return Index;
-	}
-
-	inline void Empty()
-	{
-		Heap.Empty();
-	}
-
-	inline void Reset()
-	{
-		Heap.Reset();
-	}
-
-	inline int32 Num() const
-	{
-		return Heap.Num();
-	}
-
-	bool Contains(FWidgetProxy& Proxy) const
-	{
-		return Heap.Contains(Proxy.Index);
-	}
-
-	const TArray<int32, TInlineAllocator<100>>& GetRawData() const
-	{
-		return Heap;
-	}
-private:
-	TArray<int32, TInlineAllocator<100>> Heap;
-};
-
+#endif
 
 /**
  * Represents the state of a widget from when it last had SWidget::Paint called on it. 
@@ -144,51 +120,31 @@ struct FSlateWidgetPersistentState
 	static const FSlateWidgetPersistentState NoState;
 };
 
-struct FWidgetStackData
-{
-	FWidgetStackData()
-		: ClipStackIndex(INDEX_NONE)
-		, IncomingLayerId(0)
-		, CurrentMaxLayerId(0)
-	{}
-
-	/** Index into the current stack where the widgets clip state is stored */
-	int32 ClipStackIndex;
-
-	/** The incoming layer id for children and draw elements*/
-	int32 IncomingLayerId;
-
-	/**
-	* Current Max layer id being computed.
-	* This is the layer id that would be returned from the widget during a traditional recurive
-	* OnPaint call
-	*/
-	int32 CurrentMaxLayerId;
-};
-
-template <> struct TIsPODType<FWidgetStackData> { enum { Value = true }; };
-
 class FWidgetProxyHandle
 {
 	friend class SWidget;
-	friend class FSlateDebugging;
 	friend class FSlateInvalidationRoot;
+	friend class FSlateInvalidationWidgetList;
+
 public:
 	FWidgetProxyHandle()
-		: MyIndex(INDEX_NONE)
+		: WidgetIndex(FSlateInvalidationWidgetIndex::Invalid)
 		, GenerationNumber(INDEX_NONE)
-
 	{}
 
-	SLATECORE_API bool IsValid() const;
+	SLATECORE_API bool IsValid(const SWidget* Widget) const;
 
 	FSlateInvalidationRootHandle GetInvalidationRootHandle() const { return InvalidationRootHandle; }
-	FSlateInvalidationRoot* GetInvalidationRoot() const { return InvalidationRootHandle.Advanced_GetInvalidationRootNoCheck(); }
+
+	FSlateInvalidationWidgetIndex GetWidgetIndex() const { return WidgetIndex; }
+	FSlateInvalidationWidgetSortOrder GetWidgetSortOrder() const { return WidgetSortOrder; }
 
 	FWidgetProxy& GetProxy();
 	const FWidgetProxy& GetProxy() const;
 
-	int32 GetIndex(bool bEvenIfInvalid = false) const { return bEvenIfInvalid || IsValid() ? MyIndex : INDEX_NONE; }
+private:
+	bool HasValidIndexAndInvalidationRootHandle() const;
+	FSlateInvalidationRoot* GetInvalidationRoot() const { return InvalidationRootHandle.Advanced_GetInvalidationRootNoCheck(); }
 
 	/**
 	 * Marks the widget as updated this frame
@@ -198,16 +154,20 @@ public:
 	void MarkWidgetUpdatedThisFrame();
 	
 	void MarkWidgetDirty(EInvalidateWidgetReason InvalidateReason);
-	SLATECORE_API void UpdateWidgetFlags(EWidgetUpdateFlags NewFlags);
+	SLATECORE_API void UpdateWidgetFlags(const SWidget* Widget, EWidgetUpdateFlags NewFlags);
 
 private:
-	FWidgetProxyHandle(FSlateInvalidationRoot& InInvalidationRoot, int32 InIndex);
+	FWidgetProxyHandle(const FSlateInvalidationRootHandle& InInvalidationRoot, FSlateInvalidationWidgetIndex InIndex, FSlateInvalidationWidgetSortOrder InSortIndex, int32 InGenerationNumber);
+	FWidgetProxyHandle(const FSlateInvalidationRootHandle& InInvalidationRoot, FSlateInvalidationWidgetIndex InIndex, FSlateInvalidationWidgetSortOrder InSortIndex);
+	FWidgetProxyHandle(FSlateInvalidationWidgetIndex InIndex);
 
 private:
-	/** The root of invalidation tree this proxy belongs to */
+	/** The root of invalidation tree this proxy belongs to. */
 	FSlateInvalidationRootHandle InvalidationRootHandle;
-	/** Index to myself in the fast path list */
-	int32 MyIndex;
-	/** This serves as an efficient way to test for validity which does not require invalidating all handles directly*/
+	/** Index to myself in the fast path list. */
+	FSlateInvalidationWidgetIndex WidgetIndex;
+	/** Order of the widget in the fast path list. */
+	FSlateInvalidationWidgetSortOrder WidgetSortOrder;
+	/** This serves as an efficient way to test for validity which does not require invalidating all handles directly. */
 	int32 GenerationNumber;
-};
\ No newline at end of file
+};
diff --git a/Engine/Source/Runtime/SlateCore/Public/Input/HittestGrid.h b/Engine/Source/Runtime/SlateCore/Public/Input/HittestGrid.h
--- a/Engine/Source/Runtime/SlateCore/Public/Input/HittestGrid.h
+++ b/Engine/Source/Runtime/SlateCore/Public/Input/HittestGrid.h
@@ -3,6 +3,7 @@
 #pragma once
 
 #include "CoreMinimal.h"
+#include "FastUpdate/SlateInvalidationWidgetSortOrder.h"
 #include "Layout/SlateRect.h"
 #include "Layout/ArrangedWidget.h"
 #include "Layout/Clipping.h"
@@ -78,7 +79,7 @@ public:
 	void Clear();
 
 	/** Add SWidget from the HitTest Grid */
-	void AddWidget(const TSharedRef<SWidget>& InWidget, int32 InBatchPriorityGroup, int32 InLayerId, int32 InSecondarySort);
+	void AddWidget(const TSharedRef<SWidget>& InWidget, int32 InBatchPriorityGroup, int32 InLayerId, FSlateInvalidationWidgetSortOrder InSecondarySort);
 
 	/** Remove SWidget from the HitTest Grid */
 	void RemoveWidget(const TSharedRef<SWidget>& InWidget);
@@ -86,6 +87,9 @@ public:
 	/** Remove SWidget from the HitTest Grid */
 	void RemoveWidget(const SWidget* InWidget);
 
+	/** Update the widget SecondarySort without removing it and readding it again. */
+	void UpdateWidget(const TSharedRef<SWidget>& InWidget, FSlateInvalidationWidgetSortOrder InSecondarySort);
+
 	/** Append an already existing grid that occupy the same space. */
 	UE_DEPRECATED(4.26, "Deprecated. Use the FHittestGrid::AddGrid method instead")
 	void AppendGrid(FHittestGrid& OtherGrid) {}
@@ -134,6 +138,14 @@ public:
 		UseFocusBrush = 1 << 2,
 	};
 	void DisplayGrid(int32 InLayer, const FGeometry& AllottedGeometry, FSlateWindowElementList& WindowElementList, EDisplayGridFlags DisplayFlags = EDisplayGridFlags::UseFocusBrush) const;
+
+	struct FWidgetSortData
+	{
+		const TWeakPtr<SWidget> WeakWidget;
+		int64 PrimarySort;
+		FSlateInvalidationWidgetSortOrder SecondarySort;
+	};
+	TArray<FWidgetSortData> GetAllWidgetSortDatas() const;
 #endif
 
 private:
@@ -142,7 +154,7 @@ private:
 	 */
 	struct FWidgetData
 	{
-		FWidgetData(TSharedRef<SWidget> InWidget, const FIntPoint& InUpperLeftCell, const FIntPoint& InLowerRightCell, int64 InPrimarySort, int32 InSecondarySort, int32 InUserIndex)
+		FWidgetData(TSharedRef<SWidget> InWidget, const FIntPoint& InUpperLeftCell, const FIntPoint& InLowerRightCell, int64 InPrimarySort, FSlateInvalidationWidgetSortOrder InSecondarySort, int32 InUserIndex)
 			: WeakWidget(InWidget)
 			, UpperLeftCell(InUpperLeftCell)
 			, LowerRightCell(InLowerRightCell)
@@ -155,7 +167,7 @@ private:
 		FIntPoint UpperLeftCell;
 		FIntPoint LowerRightCell;
 		int64 PrimarySort;
-		int32 SecondarySort;
+		FSlateInvalidationWidgetSortOrder SecondarySort;
 		int32 UserIndex;
 
 		TSharedPtr<SWidget> GetWidget() const { return WeakWidget.Pin(); }
diff --git a/Engine/Source/Runtime/SlateCore/Public/Widgets/SWidget.h b/Engine/Source/Runtime/SlateCore/Public/Widgets/SWidget.h
--- a/Engine/Source/Runtime/SlateCore/Public/Widgets/SWidget.h
+++ b/Engine/Source/Runtime/SlateCore/Public/Widgets/SWidget.h
@@ -212,6 +212,7 @@ class SLATECORE_API SWidget
 	friend struct FCurveSequence;
 	friend class FWidgetProxy;
 	friend class FSlateInvalidationRoot;
+	friend class FSlateInvalidationWidgetList;
 	friend class FSlateWindowElementList;
 	friend class SWindow;
 	friend struct FSlateCachedElementList;
@@ -689,7 +690,6 @@ public:
 
 protected:
 	virtual bool CustomPrepass(float LayoutScaleMultiplier) { return false; }
-	bool AssignIndicesToChildren(FSlateInvalidationRoot& Root, int32 ParentIndex, TArray<FWidgetProxy, TMemStackAllocator<>>& FastPathList, bool bParentVisible, bool bParentVolatile);
 
 	/**
 	 * The system calls this method. It performs a breadth-first traversal of every visible widget and asks
@@ -718,6 +718,9 @@ protected:
 	}
 
 private:
+	void SetFastPathProxyHandle(const FWidgetProxyHandle& Handle) { FastPathProxyHandle = Handle; }
+	void SetFastPathProxyHandle(const FWidgetProxyHandle& Handle, bool bInvisibleDueToParentOrSelfVisibility, bool bParentVolatile);
+	void SetFastPathSortOrder(const FSlateInvalidationWidgetSortOrder SortOrder);
 
 	void UpdateFastPathVisibility(bool bParentVisible, bool bWidgetRemoved, FHittestGrid* ParentHittestGrid);
 
@@ -739,19 +742,13 @@ private:
 	void AddUpdateFlags(EWidgetUpdateFlags FlagsToAdd)
 	{
 		UpdateFlags |= FlagsToAdd;
-		if (FastPathProxyHandle.IsValid())
-		{
-			FastPathProxyHandle.UpdateWidgetFlags(UpdateFlags);
-		}
+		FastPathProxyHandle.UpdateWidgetFlags(this, UpdateFlags);
 	}
 
 	void RemoveUpdateFlags(EWidgetUpdateFlags FlagsToRemove)
 	{
 		UpdateFlags &= (~FlagsToRemove);
-		if (FastPathProxyHandle.IsValid())
-		{
-			FastPathProxyHandle.UpdateWidgetFlags(UpdateFlags);
-		}
+		FastPathProxyHandle.UpdateWidgetFlags(this, UpdateFlags);
 
 #if WITH_SLATE_DEBUGGING
 		if (EnumHasAnyFlags(FlagsToRemove, EWidgetUpdateFlags::NeedsRepaint))
diff --git a/Engine/Source/Runtime/UMG/Private/Slate/SRetainerWidget.cpp b/Engine/Source/Runtime/UMG/Private/Slate/SRetainerWidget.cpp
--- a/Engine/Source/Runtime/UMG/Private/Slate/SRetainerWidget.cpp
+++ b/Engine/Source/Runtime/UMG/Private/Slate/SRetainerWidget.cpp
@@ -271,7 +271,7 @@ void SRetainerWidget::RefreshRenderingMode()
 	if ( bEnableRetainedRendering != bShouldBeRenderingOffscreen )
 	{
 		bEnableRetainedRendering = bShouldBeRenderingOffscreen;
-		InvalidateChildOrder();
+		InvalidateRoot();
 	}
 }
 
