Reduce size of SWidget by 14% by moving data that is not used often to MetaData. Deprecate SWidget::Cursor, you should now use SetCursor or GetCursor.

diff --git a/Engine/Plugins/Enterprise/DataprepEditor/Source/DataprepEditor/Private/Widgets/DataprepGraph/SDataprepGraphActionNode.cpp b/Engine/Plugins/Enterprise/DataprepEditor/Source/DataprepEditor/Private/Widgets/DataprepGraph/SDataprepGraphActionNode.cpp
--- a/Engine/Plugins/Enterprise/DataprepEditor/Source/DataprepEditor/Private/Widgets/DataprepGraph/SDataprepGraphActionNode.cpp
+++ b/Engine/Plugins/Enterprise/DataprepEditor/Source/DataprepEditor/Private/Widgets/DataprepGraph/SDataprepGraphActionNode.cpp
@@ -644,7 +644,7 @@ FReply SDataprepGraphActionNode::OnMouseButtonUp(const FGeometry & MyGeometry, c
 
 FCursorReply SDataprepGraphActionNode::OnCursorQuery( const FGeometry& MyGeometry, const FPointerEvent& CursorEvent ) const
 {
-	TOptional<EMouseCursor::Type> TheCursor = Cursor.Get();
+	TOptional<EMouseCursor::Type> TheCursor = GetCursor();
 	return ( TheCursor.IsSet() )
 		? FCursorReply::Cursor( TheCursor.GetValue() )
 		: FCursorReply::Unhandled();
diff --git a/Engine/Plugins/Experimental/VirtualCamera/Source/VirtualCamera/Private/Widgets/Input/SRadialSlider.cpp b/Engine/Plugins/Experimental/VirtualCamera/Source/VirtualCamera/Private/Widgets/Input/SRadialSlider.cpp
--- a/Engine/Plugins/Experimental/VirtualCamera/Source/VirtualCamera/Private/Widgets/Input/SRadialSlider.cpp
+++ b/Engine/Plugins/Experimental/VirtualCamera/Source/VirtualCamera/Private/Widgets/Input/SRadialSlider.cpp
@@ -371,7 +371,7 @@ FReply SRadialSlider::OnMouseButtonDown( const FGeometry& MyGeometry, const FPoi
 {
 	if ((MouseEvent.GetEffectingButton() == EKeys::LeftMouseButton) && !IsLocked())
 	{
-		CachedCursor = Cursor.Get().Get(EMouseCursor::Default);
+		CachedCursor = GetCursor().Get(EMouseCursor::Default);
 		OnInputStarted(MyGeometry, MouseEvent.GetLastScreenSpacePosition());
 
 		OnMouseCaptureBegin.ExecuteIfBound();
@@ -449,7 +449,7 @@ FReply SRadialSlider::OnTouchMoved(const FGeometry& MyGeometry, const FPointerEv
 	{
 		if (FSlateApplication::Get().HasTraveledFarEnoughToTriggerDrag(InTouchEvent, PressedScreenSpaceTouchDownPosition, EOrientation::Orient_Horizontal))
 		{
-			CachedCursor = Cursor.Get().Get(EMouseCursor::Default);
+			CachedCursor = GetCursor().Get(EMouseCursor::Default);
 			OnMouseCaptureBegin.ExecuteIfBound();
 
 			CommitValue(PositionToValue(MyGeometry, InTouchEvent.GetScreenSpacePosition()));
diff --git a/Engine/Plugins/Runtime/nDisplay/Source/DisplayClusterConfigurator/Private/Views/OutputMapping/GraphNodes/SDisplayClusterConfiguratorBaseNode.cpp b/Engine/Plugins/Runtime/nDisplay/Source/DisplayClusterConfigurator/Private/Views/OutputMapping/GraphNodes/SDisplayClusterConfiguratorBaseNode.cpp
--- a/Engine/Plugins/Runtime/nDisplay/Source/DisplayClusterConfigurator/Private/Views/OutputMapping/GraphNodes/SDisplayClusterConfiguratorBaseNode.cpp
+++ b/Engine/Plugins/Runtime/nDisplay/Source/DisplayClusterConfigurator/Private/Views/OutputMapping/GraphNodes/SDisplayClusterConfiguratorBaseNode.cpp
@@ -134,10 +134,7 @@ FReply SDisplayClusterConfiguratorBaseNode::OnDrop(const FGeometry& MyGeometry,
 
 FCursorReply SDisplayClusterConfiguratorBaseNode::OnCursorQuery(const FGeometry& MyGeometry, const FPointerEvent& CursorEvent) const
 {
-	TOptional<EMouseCursor::Type> TheCursor = Cursor.Get();
-	return (TheCursor.IsSet())
-		? FCursorReply::Cursor(TheCursor.GetValue())
-		: FCursorReply::Unhandled();
+	return SWidget::OnCursorQuery(MyGeometry, CursorEvent);
 }
 
 const FSlateBrush* SDisplayClusterConfiguratorBaseNode::GetShadowBrush(bool bSelected) const
diff --git a/Engine/Plugins/Tests/AutomationDriverTests/Source/AutomationDriverTests/Private/SAutomationDriverSpecSuite.cpp b/Engine/Plugins/Tests/AutomationDriverTests/Source/AutomationDriverTests/Private/SAutomationDriverSpecSuite.cpp
--- a/Engine/Plugins/Tests/AutomationDriverTests/Source/AutomationDriverTests/Private/SAutomationDriverSpecSuite.cpp
+++ b/Engine/Plugins/Tests/AutomationDriverTests/Source/AutomationDriverTests/Private/SAutomationDriverSpecSuite.cpp
@@ -24,7 +24,7 @@ public:
 	virtual void Construct(const FArguments& InArgs, const TSharedRef<IAutomationDriverSpecSuiteViewModel>& InViewModel) override
 	{
 		ViewModel = InViewModel;
-		MetaData.Add(FDriverMetaData::Id("Suite"));
+		AddMetadata(FDriverMetaData::Id("Suite"));
 
 		SAssignNew(WindowContents, SVerticalBox)
 		+ SVerticalBox::Slot()
diff --git a/Engine/Source/Editor/LevelEditor/Private/SCaptureRegionWidget.cpp b/Engine/Source/Editor/LevelEditor/Private/SCaptureRegionWidget.cpp
--- a/Engine/Source/Editor/LevelEditor/Private/SCaptureRegionWidget.cpp
+++ b/Engine/Source/Editor/LevelEditor/Private/SCaptureRegionWidget.cpp
@@ -263,58 +263,58 @@ FReply SCaptureRegionWidget::OnMouseMove( const FGeometry& MyGeometry, const FPo
 					// Check if we're over the corners
 					if ((ViewportPosition - FVector2D((float)CurrentCaptureRegion.Min.X, CurrentCaptureRegion.Min.Y)).Size() < EdgeDistance)
 					{
-						this->Cursor = EMouseCursor::ResizeSouthEast;
+						SetCursor(EMouseCursor::ResizeSouthEast);
 						PotentialInteraction = PI_ResizeTL;
 					}
 					else if ((ViewportPosition - FVector2D((float)CurrentCaptureRegion.Min.X, CurrentCaptureRegion.Max.Y)).Size() < EdgeDistance)
 					{
-						this->Cursor = EMouseCursor::ResizeSouthWest;
+						SetCursor(EMouseCursor::ResizeSouthWest);
 						PotentialInteraction = PI_ResizeBL;
 					}
 					else if ((ViewportPosition - FVector2D((float)CurrentCaptureRegion.Max.X, CurrentCaptureRegion.Min.Y)).Size() < EdgeDistance)
 					{
-						this->Cursor = EMouseCursor::ResizeSouthWest;
+						SetCursor(EMouseCursor::ResizeSouthWest);
 						PotentialInteraction = PI_ResizeTR;
 					}
 					else if ((ViewportPosition - FVector2D((float)CurrentCaptureRegion.Max.X, CurrentCaptureRegion.Max.Y)).Size() < EdgeDistance)
 					{
-						this->Cursor = EMouseCursor::ResizeSouthEast;
+						SetCursor(EMouseCursor::ResizeSouthEast);
 						PotentialInteraction = PI_ResizeBR;
 					}
 					else if (FMath::Abs((float)CurrentCaptureRegion.Min.X - ViewportPosition.X) < EdgeDistance && bWithinYRangeOfExistingRegion)
 					{
-						this->Cursor = EMouseCursor::ResizeLeftRight;
+						SetCursor(EMouseCursor::ResizeLeftRight);
 						PotentialInteraction = PI_ResizeLeft;
 					}
 					else if (FMath::Abs((float)CurrentCaptureRegion.Max.X - ViewportPosition.X) < EdgeDistance && bWithinYRangeOfExistingRegion)
 					{
-						this->Cursor = EMouseCursor::ResizeLeftRight;
+						SetCursor(EMouseCursor::ResizeLeftRight);
 						PotentialInteraction = PI_ResizeRight;
 					}
 					else if (FMath::Abs((float)CurrentCaptureRegion.Min.Y - ViewportPosition.Y) < EdgeDistance && bWithinXRangeOfExistingRegion)
 					{
-						this->Cursor = EMouseCursor::ResizeUpDown;
+						SetCursor(EMouseCursor::ResizeUpDown);
 						PotentialInteraction = PI_ResizeTop;
 					}
 					else if (FMath::Abs((float)CurrentCaptureRegion.Max.Y - ViewportPosition.Y) < EdgeDistance && bWithinXRangeOfExistingRegion)
 					{
-						this->Cursor = EMouseCursor::ResizeUpDown;
+						SetCursor(EMouseCursor::ResizeUpDown);
 						PotentialInteraction = PI_ResizeBottom;
 					}
 					else if (CurrentCaptureRegion.Contains(FIntPoint((int32)ViewportPosition.X, (int32)ViewportPosition.Y)))
 					{
-						this->Cursor = EMouseCursor::CardinalCross;
+						SetCursor(EMouseCursor::CardinalCross);
 						PotentialInteraction = PI_MoveExistingRegion;
 					}
 					else
 					{
-						this->Cursor = EMouseCursor::Crosshairs;
+						SetCursor(EMouseCursor::Crosshairs);
 						PotentialInteraction = PI_DrawNewCaptureRegion;
 					}
 				}
 				else
 				{
-					this->Cursor = EMouseCursor::Crosshairs;
+					SetCursor(EMouseCursor::Crosshairs);
 					PotentialInteraction = PI_DrawNewCaptureRegion;
 				}
 
diff --git a/Engine/Source/Runtime/Slate/Private/Widgets/Input/SSlider.cpp b/Engine/Source/Runtime/Slate/Private/Widgets/Input/SSlider.cpp
--- a/Engine/Source/Runtime/Slate/Private/Widgets/Input/SSlider.cpp
+++ b/Engine/Source/Runtime/Slate/Private/Widgets/Input/SSlider.cpp
@@ -265,7 +265,7 @@ FReply SSlider::OnMouseButtonDown( const FGeometry& MyGeometry, const FPointerEv
 {
 	if ((MouseEvent.GetEffectingButton() == EKeys::LeftMouseButton) && !IsLocked())
 	{
-		CachedCursor = Cursor.Get().Get(EMouseCursor::Default);
+		CachedCursor = GetCursor().Get(EMouseCursor::Default);
 		OnMouseCaptureBegin.ExecuteIfBound();
 		CommitValue(PositionToValue(MyGeometry, MouseEvent.GetLastScreenSpacePosition()));
 		
@@ -340,7 +340,7 @@ FReply SSlider::OnTouchMoved(const FGeometry& MyGeometry, const FPointerEvent& I
 	{
 		if (FSlateApplication::Get().HasTraveledFarEnoughToTriggerDrag(InTouchEvent, PressedScreenSpaceTouchDownPosition, Orientation))
 		{
-			CachedCursor = Cursor.Get().Get(EMouseCursor::Default);
+			CachedCursor = GetCursor().Get(EMouseCursor::Default);
 			OnMouseCaptureBegin.ExecuteIfBound();
 
 			CommitValue(PositionToValue(MyGeometry, InTouchEvent.GetScreenSpacePosition()));
diff --git a/Engine/Source/Runtime/Slate/Private/Widgets/Layout/SExpandableArea.cpp b/Engine/Source/Runtime/Slate/Private/Widgets/Layout/SExpandableArea.cpp
--- a/Engine/Source/Runtime/Slate/Private/Widgets/Layout/SExpandableArea.cpp
+++ b/Engine/Source/Runtime/Slate/Private/Widgets/Layout/SExpandableArea.cpp
@@ -68,7 +68,7 @@ void SExpandableArea::Construct( const FArguments& InArgs )
 				.Padding(0.0f)
 				[
 					SNew( SButton )
-					.Cursor(InArgs._HeaderCursor.IsSet() ? InArgs._HeaderCursor : Cursor)
+					.Cursor(InArgs._HeaderCursor.IsSet() ? InArgs._HeaderCursor : GetCursor())
 					.ButtonStyle(FCoreStyle::Get(), "NoBorder")
 					.ContentPadding(InArgs._HeaderPadding)
 					.ForegroundColor(FSlateColor::UseForeground())
diff --git a/Engine/Source/Runtime/SlateCore/Private/Debugging/WidgetList.cpp b/Engine/Source/Runtime/SlateCore/Private/Debugging/WidgetList.cpp
--- a/Engine/Source/Runtime/SlateCore/Private/Debugging/WidgetList.cpp
+++ b/Engine/Source/Runtime/SlateCore/Private/Debugging/WidgetList.cpp
@@ -8,11 +8,15 @@
 #include "CoreGlobals.h"
 #include "FastUpdate/WidgetProxy.h"
 #include "FastUpdate/SlateInvalidationRoot.h"
+#include "GenericPlatform/ICursor.h"
 #include "Layout/Children.h"
 #include "Misc/OutputDeviceFile.h"
 #include "Misc/Paths.h"
 #include "Misc/StringBuilder.h"
+#include "Types/CursorMetaData.h"
+#include "Types/MouseEventsMetaData.h"
 #include "Types/ReflectionMetadata.h"
+#include "Widgets/IToolTip.h"
 #include "Widgets/SWidget.h"
 
 namespace UE
@@ -26,31 +30,114 @@ TArray<const SWidget*> FWidgetList::AllWidgets;
 
 // ------------------------------------------------
 
-void LogAllWidgetsDebugInfoImpl(FOutputDevice& Ar, bool bDebug, bool bPaint, bool bProxy, bool bChildren, bool bParent)
+struct FLogAllWidgetsDebugInfoFlags
+{
+	bool bDebug;
+	bool bPaint;
+	bool bProxy;
+	bool bChildren;
+	bool bParent;
+	bool bToolTip;
+	bool bCursor;
+	bool bMouseEventsHandler;
+
+	void Parse(const FString& Arg)
+	{
+		int32 FoundIndex = Arg.Find(TEXT("Debug="));
+		if (FoundIndex != INDEX_NONE)
+		{
+			LexFromString(bDebug, *Arg + FoundIndex + 6);
+			return;
+		}
+
+		FoundIndex = Arg.Find(TEXT("Paint="));
+		if (FoundIndex != INDEX_NONE)
+		{
+			LexFromString(bPaint, *Arg + FoundIndex + 6);
+			return;
+		}
+
+		FoundIndex = Arg.Find(TEXT("Proxy="));
+		if (FoundIndex != INDEX_NONE)
+		{
+			LexFromString(bProxy, *Arg + FoundIndex + 6);
+			return;
+		}
+
+		FoundIndex = Arg.Find(TEXT("Children="));
+		if (FoundIndex != INDEX_NONE)
+		{
+			LexFromString(bChildren, *Arg + FoundIndex + 9);
+			return;
+		}
+
+		FoundIndex = Arg.Find(TEXT("Parent="));
+		if (FoundIndex != INDEX_NONE)
+		{
+			LexFromString(bParent, *Arg + FoundIndex + 7);
+			return;
+		}
+
+		FoundIndex = Arg.Find(TEXT("ToolTip="));
+		if (FoundIndex != INDEX_NONE)
+		{
+			LexFromString(bToolTip, *Arg + FoundIndex + 8);
+			return;
+		}
+
+		FoundIndex = Arg.Find(TEXT("Cursor="));
+		if (FoundIndex != INDEX_NONE)
+		{
+			LexFromString(bCursor, *Arg + FoundIndex + 7);
+			return;
+		}
+
+		FoundIndex = Arg.Find(TEXT("MouseEvents="));
+		if (FoundIndex != INDEX_NONE)
+		{
+			LexFromString(bMouseEventsHandler, *Arg + FoundIndex + 12);
+			return;
+		}
+	}
+};
+
+void LogAllWidgetsDebugInfoImpl(FOutputDevice& Ar, const FLogAllWidgetsDebugInfoFlags& DebugInfoFlags)
 {
 	TStringBuilder<1024> MessageBuilder;
 
 	MessageBuilder << TEXT("Pointer;DebugInfo");
-	if (bDebug)
+	if (DebugInfoFlags.bDebug)
 	{
 		MessageBuilder << TEXT(";Type;WidgetPath;ReadableLocation");
 	}
-	if (bPaint)
+	if (DebugInfoFlags.bPaint)
 	{
 		MessageBuilder << TEXT(";LastPaintFrame;LayerId;AllottedGeometryAbsoluteSizeX;AllottedGeometryAbsoluteSizeY");
 	}
-	if (bProxy)
+	if (DebugInfoFlags.bProxy)
 	{
 		MessageBuilder << TEXT(";InvalidationRootPointer;InvalidationRootDebugInfo;ProxyIndex");
 	}
-	if (bChildren)
+	if (DebugInfoFlags.bChildren)
 	{
 		MessageBuilder << TEXT(";NumAllChildren;NumChildren");
 	}
-	if (bParent)
+	if (DebugInfoFlags.bParent)
 	{
 		MessageBuilder << TEXT(";ParentPointer;ParentDebugInfo");
 	}
+	if (DebugInfoFlags.bToolTip)
+	{
+		MessageBuilder << TEXT(";ToolTipIsSet;ToolTipIsEmpty");
+	}
+	if (DebugInfoFlags.bCursor)
+	{
+		MessageBuilder << TEXT(";CursorIsSet;CursorValue");
+	}
+	if (DebugInfoFlags.bMouseEventsHandler)
+	{
+		MessageBuilder << TEXT(";MouseButtonDown;MouseButtonUp;MouseMove;MouseDblClick;MouseEnter;MouseLeave");
+	}
 	Ar.Log(MessageBuilder.ToString());
 
 	const TArray<const SWidget*>& WidgetList = FWidgetList::GetAllWidgets();
@@ -62,7 +149,7 @@ void LogAllWidgetsDebugInfoImpl(FOutputDevice& Ar, bool bDebug, bool bPaint, boo
 		MessageBuilder << TEXT(";");
 		MessageBuilder << FReflectionMetaData::GetWidgetDebugInfo(Widget);
 
-		if (bDebug)
+		if (DebugInfoFlags.bDebug)
 		{
 			MessageBuilder << TEXT(";");
 			MessageBuilder << Widget->GetTypeAsString();
@@ -72,7 +159,7 @@ void LogAllWidgetsDebugInfoImpl(FOutputDevice& Ar, bool bDebug, bool bPaint, boo
 			MessageBuilder << Widget->GetReadableLocation();
 		}
 
-		if (bPaint)
+		if (DebugInfoFlags.bPaint)
 		{
 			FVector2D AbsoluteSize = Widget->GetPersistentState().AllottedGeometry.GetAbsoluteSize();
 			MessageBuilder << TEXT(";");
@@ -85,7 +172,7 @@ void LogAllWidgetsDebugInfoImpl(FOutputDevice& Ar, bool bDebug, bool bPaint, boo
 			MessageBuilder.Appendf(TEXT("%f"), AbsoluteSize.Y);
 		}
 
-		if (bProxy)
+		if (DebugInfoFlags.bProxy)
 		{
 			const FWidgetProxyHandle ProxyHandle = Widget->GetProxyHandle();
 			if (ProxyHandle.IsValid())
@@ -94,6 +181,8 @@ void LogAllWidgetsDebugInfoImpl(FOutputDevice& Ar, bool bDebug, bool bPaint, boo
 				MessageBuilder.Appendf(TEXT("%p"), (void*)ProxyHandle.GetInvalidationRoot());
 				MessageBuilder << TEXT(";");
 				MessageBuilder << FReflectionMetaData::GetWidgetDebugInfo(ProxyHandle.GetInvalidationRoot()->GetInvalidationRootWidget());
+				MessageBuilder << TEXT(";");
+				MessageBuilder << ProxyHandle.GetIndex(true);
 			}
 			else
 			{
@@ -101,7 +190,7 @@ void LogAllWidgetsDebugInfoImpl(FOutputDevice& Ar, bool bDebug, bool bPaint, boo
 			}
 		}
 
-		if (bChildren)
+		if (DebugInfoFlags.bChildren)
 		{
 			MessageBuilder << TEXT(";");
 			MessageBuilder << const_cast<SWidget*>(Widget)->GetAllChildren()->Num();
@@ -109,7 +198,7 @@ void LogAllWidgetsDebugInfoImpl(FOutputDevice& Ar, bool bDebug, bool bPaint, boo
 			MessageBuilder << const_cast<SWidget*>(Widget)->GetChildren()->Num();
 		}
 
-		if (bParent)
+		if (DebugInfoFlags.bParent)
 		{
 			const SWidget* ParentWidget = Widget->GetParentWidget().Get();
 			MessageBuilder << TEXT(";");
@@ -118,17 +207,78 @@ void LogAllWidgetsDebugInfoImpl(FOutputDevice& Ar, bool bDebug, bool bPaint, boo
 			MessageBuilder << FReflectionMetaData::GetWidgetDebugInfo(ParentWidget);
 		}
 
+		if (DebugInfoFlags.bToolTip)
+		{
+			if (TSharedPtr<IToolTip> ToolTip = const_cast<SWidget*>(Widget)->GetToolTip())
+			{
+				MessageBuilder << TEXT(";true");
+				if (ToolTip->IsEmpty())
+				{
+					MessageBuilder << TEXT(";true");
+				}
+				else
+				{
+					MessageBuilder << TEXT(";false");
+				}
+			}
+			else
+			{
+				MessageBuilder << TEXT(";false;false");
+			}
+		}
+		
+		if (DebugInfoFlags.bCursor)
+		{
+			if (TSharedPtr<FCursorMetaData> Data = Widget->GetMetaData<FCursorMetaData>())
+			{
+				if (Data->Cursor.IsSet())
+				{
+					TOptional<EMouseCursor::Type> Cursor = Data->Cursor.Get();
+					if (Cursor.IsSet())
+					{
+						MessageBuilder << TEXT(";Set;");
+						MessageBuilder << static_cast<int32>(Data->Cursor.Get().GetValue());
+					}
+					else
+					{
+						MessageBuilder << TEXT(";Optional;0");
+					}
+				}
+				else
+				{
+					MessageBuilder << TEXT(";MetaData;0");
+				}
+			}
+			else
+			{
+				MessageBuilder << TEXT(";None;0");
+			}
+		}
+
+		if (DebugInfoFlags.bMouseEventsHandler)
+		{
+			if (TSharedPtr<FMouseEventsMetaData> Data = Widget->GetMetaData<FMouseEventsMetaData>())
+			{
+				MessageBuilder << (Data->MouseButtonDownHandle.IsBound() ? TEXT(";bound") : TEXT(";"));
+				MessageBuilder << (Data->MouseButtonUpHandle.IsBound() ? TEXT(";bound") : TEXT(";"));
+				MessageBuilder << (Data->MouseMoveHandle.IsBound() ? TEXT(";bound") : TEXT(";"));
+				MessageBuilder << (Data->MouseDoubleClickHandle.IsBound() ? TEXT(";bound") : TEXT(";"));
+				MessageBuilder << (Data->MouseEnterHandler.IsBound() ? TEXT(";bound") : TEXT(";"));
+				MessageBuilder << (Data->MouseLeaveHandler.IsBound() ? TEXT(";bound") : TEXT(";"));
+			}
+			else
+			{
+				MessageBuilder << TEXT(";;;;;;");
+			}
+		}
+
 		Ar.Log(MessageBuilder.ToString());
 	}
 }
 
 void LogAllWidgetsDebugInfo(const TArray< FString >& Args, UWorld*, FOutputDevice& Ar)
 {
-	bool bDebug = true;
-	bool bPaint = false;
-	bool bProxy = false;
-	bool bChildren = false;
-	bool bParent = false;
+	FLogAllWidgetsDebugInfoFlags DebugInfoFlags;
 	FString OutputFilename;
 	for (const FString& Arg : Args)
 	{
@@ -136,57 +286,25 @@ void LogAllWidgetsDebugInfo(const TArray< FString >& Args, UWorld*, FOutputDevic
 		{
 			continue;
 		}
-		int32 FoundIndex = Arg.Find(TEXT("Debug="));
-		if (FoundIndex != INDEX_NONE)
-		{
-			LexFromString(bDebug, *Arg + FoundIndex + 6);
-			continue;
-		}
 
-		FoundIndex = Arg.Find(TEXT("Paint="));
-		if (FoundIndex != INDEX_NONE)
-		{
-			LexFromString(bPaint, *Arg + FoundIndex + 6);
-			continue;
-		}
-
-		FoundIndex = Arg.Find(TEXT("Proxy="));
-		if (FoundIndex != INDEX_NONE)
-		{
-			LexFromString(bProxy, *Arg + FoundIndex + 6);
-			continue;
-		}
-
-		FoundIndex = Arg.Find(TEXT("Children="));
-		if (FoundIndex != INDEX_NONE)
-		{
-			LexFromString(bChildren, *Arg + FoundIndex + 9);
-			continue;
-		}
-		
-		FoundIndex = Arg.Find(TEXT("Parent="));
-		if (FoundIndex != INDEX_NONE)
-		{
-			LexFromString(bParent, *Arg + FoundIndex + 7);
-			continue;
-		}
+		DebugInfoFlags.Parse(Arg);
 	}
 
 	if (OutputFilename.Len() > 0)
 	{
 		FOutputDeviceFile OutputDeviceFile {*FPaths::Combine(FPaths::ProjectSavedDir(), OutputFilename), true};
 		OutputDeviceFile.SetSuppressEventTag(true);
-		LogAllWidgetsDebugInfoImpl(OutputDeviceFile, bDebug, bPaint, bProxy, bChildren, bParent);
+		LogAllWidgetsDebugInfoImpl(OutputDeviceFile, DebugInfoFlags);
 	}
 	else
 	{
-		LogAllWidgetsDebugInfoImpl(Ar, bDebug, bPaint, bProxy, bChildren, bParent);
+		LogAllWidgetsDebugInfoImpl(Ar, DebugInfoFlags);
 	}
 }
 
 FAutoConsoleCommandWithWorldArgsAndOutputDevice ConsoleCommandLogAllWidgets(
 	TEXT("Slate.Debug.LogAllWidgets"),
-	TEXT("Prints all the SWidgets type, debug info, path, painted, ...\n")
+	TEXT("Prints all the SWidgets type, debug info, path or painted.\n")
 	TEXT("If a file name is not provided, it will output to the log console.\n")
 	TEXT("Slate.Debug.LogAllWidgets [File=MyFile.csv] [Debug=true] [Paint=false] [Proxy=false] [Children=false] [Parent=false]"),
 	FConsoleCommandWithWorldArgsAndOutputDeviceDelegate::CreateStatic(&LogAllWidgetsDebugInfo)
diff --git a/Engine/Source/Runtime/SlateCore/Private/Types/SlateCursorMetaData.h b/Engine/Source/Runtime/SlateCore/Private/Types/SlateCursorMetaData.h
new file mode 100644
--- /dev/null
+++ b/Engine/Source/Runtime/SlateCore/Private/Types/SlateCursorMetaData.h
@@ -0,0 +1,20 @@
+// Copyright Epic Games, Inc. All Rights Reserved.
+
+#pragma once
+
+#include "CoreMinimal.h"
+#include "GenericPlatform/ICursor.h"
+#include "Misc/Attribute.h"
+#include "Types/ISlateMetaData.h"
+#include "Widgets/IToolTip.h"
+
+
+/** Metadata that holds the cursor content for a widget. */
+class FSlateCursorMetaData : public ISlateMetaData
+{
+public:
+	SLATE_METADATA_TYPE(FSlateCursorMetaData, ISlateMetaData)
+	
+	/** The cursor to show when the mouse is hovering over this widget. */
+	TAttribute<TOptional<EMouseCursor::Type>> Cursor;
+};
diff --git a/Engine/Source/Runtime/SlateCore/Private/Types/SlateMouseEventsMetaData.h b/Engine/Source/Runtime/SlateCore/Private/Types/SlateMouseEventsMetaData.h
new file mode 100644
--- /dev/null
+++ b/Engine/Source/Runtime/SlateCore/Private/Types/SlateMouseEventsMetaData.h
@@ -0,0 +1,21 @@
+// Copyright Epic Games, Inc. All Rights Reserved.
+
+#pragma once
+
+#include "CoreMinimal.h"
+#include "Types/ISlateMetaData.h"
+#include "Types/WidgetMouseEventsDelegate.h"
+
+/** Metadata to hold mouse event set on the SWidget */
+class FSlateMouseEventsMetaData : public ISlateMetaData
+{
+public:
+	SLATE_METADATA_TYPE(FSlateMouseEventsMetaData, ISlateMetaData)
+	
+	FPointerEventHandler MouseButtonDownHandle;
+	FPointerEventHandler MouseButtonUpHandle;
+	FPointerEventHandler MouseMoveHandle;
+	FPointerEventHandler MouseDoubleClickHandle;
+	FNoReplyPointerEventHandler MouseEnterHandler;
+	FSimpleNoReplyPointerEventHandler MouseLeaveHandler;
+};
diff --git a/Engine/Source/Runtime/SlateCore/Private/Types/SlateToolTipMetaData.h b/Engine/Source/Runtime/SlateCore/Private/Types/SlateToolTipMetaData.h
new file mode 100644
--- /dev/null
+++ b/Engine/Source/Runtime/SlateCore/Private/Types/SlateToolTipMetaData.h
@@ -0,0 +1,17 @@
+// Copyright Epic Games, Inc. All Rights Reserved.
+
+#pragma once
+
+#include "CoreMinimal.h"
+#include "Types/ISlateMetaData.h"
+#include "Widgets/IToolTip.h"
+
+
+/** Metadata that holds the tool tip content for a widget. */
+class FSlateToolTipMetaData : public ISlateMetaData
+{
+public:
+	SLATE_METADATA_TYPE(FSlateToolTipMetaData, ISlateMetaData)
+	
+	TSharedPtr<IToolTip> ToolTip;
+};
diff --git a/Engine/Source/Runtime/SlateCore/Private/Widgets/SWidget.cpp b/Engine/Source/Runtime/SlateCore/Private/Widgets/SWidget.cpp
--- a/Engine/Source/Runtime/SlateCore/Private/Widgets/SWidget.cpp
+++ b/Engine/Source/Runtime/SlateCore/Private/Widgets/SWidget.cpp
@@ -18,7 +18,10 @@
 #include "Debugging/WidgetList.h"
 #include "Widgets/SWindow.h"
 #include "Trace/SlateTrace.h"
+#include "Types/SlateCursorMetaData.h"
+#include "Types/SlateMouseEventsMetaData.h"
 #include "Types/ReflectionMetadata.h"
+#include "Types/SlateToolTipMetaData.h"
 #include "Stats/Stats.h"
 #include "Containers/StringConv.h"
 #include "Misc/ScopeRWLock.h"
@@ -54,6 +57,7 @@ DEFINE_STAT(STAT_SlateTickWidgets);
 DEFINE_STAT(STAT_SlatePrepass);
 DEFINE_STAT(STAT_SlateTotalWidgets);
 DEFINE_STAT(STAT_SlateSWidgetAllocSize);
+DEFINE_STAT(STAT_SlateGetMetaData);
 
 DECLARE_CYCLE_STAT(TEXT("SWidget::CreateStatID"), STAT_Slate_CreateStatID, STATGROUP_Slate);
 
@@ -191,11 +195,6 @@ namespace SlateTraceMetaData
 }
 #endif
 
-FName NAME_MouseButtonDown(TEXT("MouseButtonDown"));
-FName NAME_MouseButtonUp(TEXT("MouseButtonUp"));
-FName NAME_MouseMove(TEXT("MouseMove"));
-FName NAME_MouseDoubleClick(TEXT("MouseDoubleClick"));
-
 SWidget::SWidget()
 	: bIsHovered(false)
 	, bCanSupportFocus(true)
@@ -228,8 +227,6 @@ SWidget::SWidget()
 	, RenderOpacity(1.0f)
 	, RenderTransform()
 	, RenderTransformPivot(FVector2D::ZeroVector)
-	, Cursor( TOptional<EMouseCursor::Type>() )
-	, ToolTip()
 #if UE_SLATE_WITH_WIDGET_UNIQUE_IDENTIFIER
 	, UniqueIdentifier(++SlateTraceMetaData::UniqueIdGenerator)
 #endif
@@ -312,23 +309,6 @@ void SWidget::Construct(
 	const TArray<TSharedRef<ISlateMetaData>>& InMetaData
 )
 {
-	if ( InToolTip.IsValid() )
-	{
-		// If someone specified a fancy widget tooltip, use it.
-		ToolTip = InToolTip;
-	}
-	else if ( InToolTipText.IsSet() )
-	{
-		// If someone specified a text binding, make a tooltip out of it
-		ToolTip = FSlateApplicationBase::Get().MakeToolTip(InToolTipText);
-	}
-	else if( !ToolTip.IsValid() || (ToolTip.IsValid() && ToolTip->IsEmpty()) )
-	{	
-		// We don't have a tooltip.
-		ToolTip.Reset();
-	}
-
-	Cursor = InCursor;
 	EnabledState = InEnabledState;
 	Visibility = InVisibility;
 	RenderOpacity = InRenderOpacity;
@@ -338,7 +318,19 @@ void SWidget::Construct(
 	bForceVolatile = InForceVolatile;
 	Clipping = InClipping;
 	FlowDirectionPreference = InFlowPreference;
-	MetaData = InMetaData;
+	MetaData.Append(InMetaData);
+
+	if (InToolTip.IsValid())
+	{
+		// If someone specified a fancy widget tooltip, use it.
+		SetToolTip(InToolTip);
+	}
+	else if (InToolTipText.IsSet())
+	{
+		// If someone specified a text binding, make a tooltip out of it
+		SetToolTipText(InToolTipText);
+	}
+	SetCursor(InCursor);
 
 #if WITH_ACCESSIBILITY
 	if (InAccessibleData.IsSet())
@@ -430,11 +422,11 @@ FReply SWidget::OnPreviewMouseButtonDown( const FGeometry& MyGeometry, const FPo
 
 FReply SWidget::OnMouseButtonDown(const FGeometry& MyGeometry, const FPointerEvent& MouseEvent)
 {
-	if (const FPointerEventHandler* Event = GetPointerEvent(NAME_MouseButtonDown))
+	if (TSharedPtr<FSlateMouseEventsMetaData> Data = GetMetaData<FSlateMouseEventsMetaData>())
 	{
-		if ( Event->IsBound() )
+		if (Data->MouseButtonDownHandle.IsBound() )
 		{
-			return Event->Execute(MyGeometry, MouseEvent);
+			return Data->MouseButtonDownHandle.Execute(MyGeometry, MouseEvent);
 		}
 	}
 	return FReply::Unhandled();
@@ -442,11 +434,11 @@ FReply SWidget::OnMouseButtonDown(const FGeometry& MyGeometry, const FPointerEve
 
 FReply SWidget::OnMouseButtonUp(const FGeometry& MyGeometry, const FPointerEvent& MouseEvent)
 {
-	if (const FPointerEventHandler* Event = GetPointerEvent(NAME_MouseButtonUp) )
+	if (TSharedPtr<FSlateMouseEventsMetaData> Data = GetMetaData<FSlateMouseEventsMetaData>())
 	{
-		if ( Event->IsBound() )
+		if (Data->MouseButtonUpHandle.IsBound())
 		{
-			return Event->Execute(MyGeometry, MouseEvent);
+			return Data->MouseButtonUpHandle.Execute(MyGeometry, MouseEvent);
 		}
 	}
 	return FReply::Unhandled();
@@ -454,11 +446,11 @@ FReply SWidget::OnMouseButtonUp(const FGeometry& MyGeometry, const FPointerEvent
 
 FReply SWidget::OnMouseMove(const FGeometry& MyGeometry, const FPointerEvent& MouseEvent)
 {
-	if (const FPointerEventHandler* Event = GetPointerEvent(NAME_MouseMove) )
+	if (TSharedPtr<FSlateMouseEventsMetaData> Data = GetMetaData<FSlateMouseEventsMetaData>())
 	{
-		if ( Event->IsBound() )
+		if (Data->MouseMoveHandle.IsBound())
 		{
-			return Event->Execute(MyGeometry, MouseEvent);
+			return Data->MouseMoveHandle.Execute(MyGeometry, MouseEvent);
 		}
 	}
 	return FReply::Unhandled();
@@ -466,11 +458,11 @@ FReply SWidget::OnMouseMove(const FGeometry& MyGeometry, const FPointerEvent& Mo
 
 FReply SWidget::OnMouseButtonDoubleClick(const FGeometry& MyGeometry, const FPointerEvent& MouseEvent)
 {
-	if ( const FPointerEventHandler* Event = GetPointerEvent(NAME_MouseDoubleClick) )
+	if (TSharedPtr<FSlateMouseEventsMetaData> Data = GetMetaData<FSlateMouseEventsMetaData>())
 	{
-		if ( Event->IsBound() )
+		if (Data->MouseDoubleClickHandle.IsBound())
 		{
-			return Event->Execute(MyGeometry, MouseEvent);
+			return Data->MouseDoubleClickHandle.Execute(MyGeometry, MouseEvent);
 		}
 	}
 	return FReply::Unhandled();
@@ -480,10 +472,13 @@ void SWidget::OnMouseEnter( const FGeometry& MyGeometry, const FPointerEvent& Mo
 {
 	bIsHovered = true;
 
-	if (MouseEnterHandler.IsBound())
+	if (TSharedPtr<FSlateMouseEventsMetaData> Data = GetMetaData<FSlateMouseEventsMetaData>())
 	{
-		// A valid handler is assigned; let it handle the event.
-		MouseEnterHandler.Execute(MyGeometry, MouseEvent);
+		if (Data->MouseEnterHandler.IsBound())
+		{
+			// A valid handler is assigned; let it handle the event.
+			Data->MouseEnterHandler.Execute(MyGeometry, MouseEvent);
+		}
 	}
 }
 
@@ -491,10 +486,13 @@ void SWidget::OnMouseLeave( const FPointerEvent& MouseEvent )
 {
 	bIsHovered = false;
 
-	if (MouseLeaveHandler.IsBound())
+	if (TSharedPtr<FSlateMouseEventsMetaData> Data = GetMetaData<FSlateMouseEventsMetaData>())
 	{
-		// A valid handler is assigned; let it handle the event.
-		MouseLeaveHandler.Execute(MouseEvent);
+		if (Data->MouseLeaveHandler.IsBound())
+		{
+			// A valid handler is assigned; let it handle the event.
+			Data->MouseLeaveHandler.Execute(MouseEvent);
+		}
 	}
 }
 
@@ -505,7 +503,7 @@ FReply SWidget::OnMouseWheel( const FGeometry& MyGeometry, const FPointerEvent&
 
 FCursorReply SWidget::OnCursorQuery( const FGeometry& MyGeometry, const FPointerEvent& CursorEvent ) const
 {
-	TOptional<EMouseCursor::Type> TheCursor = Cursor.Get();
+	TOptional<EMouseCursor::Type> TheCursor = GetCursor();
 	return ( TheCursor.IsSet() )
 		? FCursorReply::Cursor( TheCursor.GetValue() )
 		: FCursorReply::Unhandled();
@@ -1043,24 +1041,63 @@ const FGeometry& SWidget::GetPaintSpaceGeometry() const
 	return PersistentState.AllottedGeometry;
 }
 
+namespace Private
+{
+	TSharedPtr<FSlateToolTipMetaData> FindOrAddToolTipMetaData(SWidget* Widget)
+	{
+		TSharedPtr<FSlateToolTipMetaData> Data = Widget->GetMetaData<FSlateToolTipMetaData>();
+		if (!Data)
+		{
+			Data = MakeShared<FSlateToolTipMetaData>();
+			Widget->AddMetadata(Data.ToSharedRef());
+		}
+		return Data;
+	}
+}
+
 void SWidget::SetToolTipText(const TAttribute<FText>& ToolTipText)
 {
-	ToolTip = FSlateApplicationBase::Get().MakeToolTip(ToolTipText);
+	if (ToolTipText.IsSet())
+	{
+		Private::FindOrAddToolTipMetaData(this)->ToolTip = FSlateApplicationBase::Get().MakeToolTip(ToolTipText);
+	}
+	else
+	{
+		RemoveMetaData<FSlateToolTipMetaData>();
+	}
 }
 
 void SWidget::SetToolTipText( const FText& ToolTipText )
 {
-	ToolTip = FSlateApplicationBase::Get().MakeToolTip(ToolTipText);
+	if (!ToolTipText.IsEmptyOrWhitespace())
+	{
+		Private::FindOrAddToolTipMetaData(this)->ToolTip = FSlateApplicationBase::Get().MakeToolTip(ToolTipText);
+	}
+	else
+	{
+		RemoveMetaData<FSlateToolTipMetaData>();
+	}
 }
 
 void SWidget::SetToolTip( const TSharedPtr<IToolTip> & InToolTip )
 {
-	ToolTip = InToolTip;
+	if (InToolTip)
+	{
+		Private::FindOrAddToolTipMetaData(this)->ToolTip = InToolTip;
+	}
+	else
+	{
+		RemoveMetaData<FSlateToolTipMetaData>();
+	}
 }
 
 TSharedPtr<IToolTip> SWidget::GetToolTip()
 {
-	return ToolTip;
+	if (TSharedPtr<FSlateToolTipMetaData> Data = GetMetaData<FSlateToolTipMetaData>())
+	{
+		return Data->ToolTip;
+	}
+	return TSharedPtr<IToolTip>();
 }
 
 void SWidget::OnToolTipClosing()
@@ -1136,7 +1173,30 @@ void SWidget::Invalidate(EInvalidateWidgetReason InvalidateReason)
 
 void SWidget::SetCursor( const TAttribute< TOptional<EMouseCursor::Type> >& InCursor )
 {
-	Cursor = InCursor;
+	// If bounded or has a valid optional value
+	if (InCursor.IsBound() || InCursor.Get().IsSet())
+	{
+		TSharedPtr<FSlateCursorMetaData> Data = GetMetaData<FSlateCursorMetaData>();
+		if (!Data)
+		{
+			Data = MakeShared<FSlateCursorMetaData>();
+			AddMetadata(Data.ToSharedRef());
+		}
+		Data->Cursor = InCursor;
+	}
+	else
+	{
+		RemoveMetaData<FSlateCursorMetaData>();
+	}
+}
+
+TOptional<EMouseCursor::Type> SWidget::GetCursor() const
+{
+	if (TSharedPtr<FSlateCursorMetaData> Data = GetMetaData<FSlateCursorMetaData>())
+	{
+		return Data->Cursor.Get();
+	}
+	return TOptional<EMouseCursor::Type>();
 }
 
 void SWidget::SetDebugInfo( const ANSICHAR* InType, const ANSICHAR* InFile, int32 OnLine, size_t InAllocSize )
@@ -1568,59 +1628,48 @@ void SWidget::ExecuteActiveTimers(double CurrentTime, float DeltaTime)
 	}
 }
 
-const FPointerEventHandler* SWidget::GetPointerEvent(const FName EventName) const
+namespace Private
 {
-	auto* FoundPair = PointerEvents.FindByPredicate([&EventName](const auto& TestPair) {return TestPair.Key == EventName; });
-	if (FoundPair)
-	{
-		return &FoundPair->Value;
-	}
-	return nullptr;
-}
-
-void SWidget::SetPointerEvent(const FName EventName, FPointerEventHandler& InEvent)
-{
-	// Find the event name and if found, replace the delegate
-	auto* FoundPair = PointerEvents.FindByPredicate([&EventName](const auto& TestPair) {return TestPair.Key == EventName; });
-	if (FoundPair)
-	{
-		FoundPair->Value = InEvent;
-	}
-	else
+	TSharedPtr<FSlateMouseEventsMetaData> FindOrAddMouseEventsMetaData(SWidget* Widget)
 	{
-		PointerEvents.Emplace(EventName, InEvent);
+		TSharedPtr<FSlateMouseEventsMetaData> Data = Widget->GetMetaData<FSlateMouseEventsMetaData>();
+		if (!Data)
+		{
+			Data = MakeShared<FSlateMouseEventsMetaData>();
+			Widget->AddMetadata(Data.ToSharedRef());
+		}
+		return Data;
 	}
-
 }
 
 void SWidget::SetOnMouseButtonDown(FPointerEventHandler EventHandler)
 {
-	SetPointerEvent(NAME_MouseButtonDown, EventHandler);
+	Private::FindOrAddMouseEventsMetaData(this)->MouseButtonDownHandle = EventHandler;
 }
 
 void SWidget::SetOnMouseButtonUp(FPointerEventHandler EventHandler)
 {
-	SetPointerEvent(NAME_MouseButtonUp, EventHandler);
+	Private::FindOrAddMouseEventsMetaData(this)->MouseButtonUpHandle = EventHandler;
 }
 
 void SWidget::SetOnMouseMove(FPointerEventHandler EventHandler)
 {
-	SetPointerEvent(NAME_MouseMove, EventHandler);
+	Private::FindOrAddMouseEventsMetaData(this)->MouseMoveHandle = EventHandler;
 }
 
 void SWidget::SetOnMouseDoubleClick(FPointerEventHandler EventHandler)
 {
-	SetPointerEvent(NAME_MouseDoubleClick, EventHandler);
+	Private::FindOrAddMouseEventsMetaData(this)->MouseDoubleClickHandle = EventHandler;
 }
 
 void SWidget::SetOnMouseEnter(FNoReplyPointerEventHandler EventHandler)
 {
-	MouseEnterHandler = EventHandler;
+	Private::FindOrAddMouseEventsMetaData(this)->MouseEnterHandler = EventHandler;
 }
 
 void SWidget::SetOnMouseLeave(FSimpleNoReplyPointerEventHandler EventHandler)
 {
-	MouseLeaveHandler = EventHandler;
+	Private::FindOrAddMouseEventsMetaData(this)->MouseLeaveHandler = EventHandler;
 }
 
 void SWidget::AddMetadataInternal(const TSharedRef<ISlateMetaData>& AddMe)
@@ -1700,11 +1749,17 @@ FText SWidget::GetAccessibleText(EAccessibleType AccessibleType) const
 	case EAccessibleBehavior::Summary:
 		return GetAccessibleSummary();
 	case EAccessibleBehavior::ToolTip:
-		if (ToolTip.IsValid() && !ToolTip->IsEmpty())
+	{
+		//TODO should use GetToolTip
+		if (TSharedPtr<FSlateToolTipMetaData> Data = GetMetaData<FSlateToolTipMetaData>())
 		{
-			return ToolTip->GetContentWidget()->GetAccessibleText(EAccessibleType::Main);
+			if (Data->ToolTip && !Data->ToolTip->IsEmpty())
+			{
+				return Data->ToolTip->GetContentWidget()->GetAccessibleText(EAccessibleType::Main);
+			}
 		}
 		break;
+	}
 	case EAccessibleBehavior::Auto:
 		// Auto first checks if custom text was set. This should never happen with user-defined values as custom should be
 		// used instead in that case - however, this will be used for widgets with special default text such as TextBlocks.
diff --git a/Engine/Source/Runtime/SlateCore/Public/Types/WidgetMouseEventsDelegate.h b/Engine/Source/Runtime/SlateCore/Public/Types/WidgetMouseEventsDelegate.h
new file mode 100644
--- /dev/null
+++ b/Engine/Source/Runtime/SlateCore/Public/Types/WidgetMouseEventsDelegate.h
@@ -0,0 +1,27 @@
+// Copyright Epic Games, Inc. All Rights Reserved.
+
+#pragma once
+
+#include "CoreMinimal.h"
+#include "Input/Events.h"
+#include "Layout/Geometry.h"
+
+/** Delegate type for handling mouse events */
+DECLARE_DELEGATE_RetVal_TwoParams(
+	FReply, FPointerEventHandler,
+	/** The geometry of the widget*/
+	const FGeometry&,
+	/** The Mouse Event that we are processing */
+	const FPointerEvent&)
+
+DECLARE_DELEGATE_TwoParams(
+	FNoReplyPointerEventHandler,
+	/** The geometry of the widget*/
+	const FGeometry&,
+	/** The Mouse Event that we are processing */
+	const FPointerEvent&)
+
+DECLARE_DELEGATE_OneParam(
+	FSimpleNoReplyPointerEventHandler,
+	/** The Mouse Event that we are processing */
+	const FPointerEvent&)
diff --git a/Engine/Source/Runtime/SlateCore/Public/Widgets/SWidget.h b/Engine/Source/Runtime/SlateCore/Public/Widgets/SWidget.h
--- a/Engine/Source/Runtime/SlateCore/Public/Widgets/SWidget.h
+++ b/Engine/Source/Runtime/SlateCore/Public/Widgets/SWidget.h
@@ -22,6 +22,7 @@
 #include "Input/PopupMethodReply.h"
 #include "Types/ISlateMetaData.h"
 #include "Types/WidgetActiveTimerDelegate.h"
+#include "Types/WidgetMouseEventsDelegate.h"
 #include "Textures/SlateShaderResource.h"
 #include "SlateGlobals.h"
 #include "Types/PaintArgs.h"
@@ -48,29 +49,11 @@ DECLARE_DWORD_COUNTER_STAT_EXTERN(TEXT("SWidget::Tick (Count)"), STAT_SlateNumTi
 DECLARE_CYCLE_STAT_EXTERN(TEXT("Execute Active Timers"), STAT_SlateExecuteActiveTimers, STATGROUP_Slate, SLATECORE_API);
 DECLARE_CYCLE_STAT_EXTERN(TEXT("Tick Widgets"), STAT_SlateTickWidgets, STATGROUP_Slate, SLATECORE_API);
 DECLARE_CYCLE_STAT_EXTERN(TEXT("SlatePrepass"), STAT_SlatePrepass, STATGROUP_Slate, SLATECORE_API);
+DECLARE_CYCLE_STAT_EXTERN(TEXT("SWidget MetaData"), STAT_SlateGetMetaData, STATGROUP_Slate, SLATECORE_API);
 
 DECLARE_DWORD_ACCUMULATOR_STAT_EXTERN(TEXT("Total Widgets"), STAT_SlateTotalWidgets, STATGROUP_SlateMemory, SLATECORE_API);
 DECLARE_MEMORY_STAT_EXTERN(TEXT("SWidget Total Allocated Size"), STAT_SlateSWidgetAllocSize, STATGROUP_SlateMemory, SLATECORE_API);
 
-/** Delegate type for handling mouse events */
-DECLARE_DELEGATE_RetVal_TwoParams(
-	FReply, FPointerEventHandler,
-	/** The geometry of the widget*/
-	const FGeometry&,
-	/** The Mouse Event that we are processing */
-	const FPointerEvent&)
-
-DECLARE_DELEGATE_TwoParams(
-	FNoReplyPointerEventHandler,
-	/** The geometry of the widget*/
-	const FGeometry&,
-	/** The Mouse Event that we are processing */
-	const FPointerEvent&)
-
-DECLARE_DELEGATE_OneParam(
-	FSimpleNoReplyPointerEventHandler,
-	/** The Mouse Event that we are processing */
-	const FPointerEvent&)
 
 enum class EPopupMethod : uint8;
 
@@ -1242,6 +1225,13 @@ public:
 	 */
 	void SetDebugInfo( const ANSICHAR* InType, const ANSICHAR* InFile, int32 OnLine, size_t InAllocSize );
 	
+protected:
+
+	/** The cursor to show when the mouse is hovering over this widget. */
+	TOptional<EMouseCursor::Type> GetCursor() const;
+
+public:
+
 	/**
 	 * Get the metadata of the type provided.
 	 * @return the first metadata of the type supplied that we encounter
@@ -1249,6 +1239,7 @@ public:
 	template<typename MetaDataType>
 	TSharedPtr<MetaDataType> GetMetaData() const
 	{
+		SCOPE_CYCLE_COUNTER(STAT_SlateGetMetaData);
 		for (const auto& MetaDataEntry : MetaData)
 		{
 			if (MetaDataEntry->IsOfType<MetaDataType>())
@@ -1266,6 +1257,7 @@ public:
 	template<typename MetaDataType>
 	TArray<TSharedRef<MetaDataType>> GetAllMetaData() const
 	{
+		SCOPE_CYCLE_COUNTER(STAT_SlateGetMetaData);
 		TArray<TSharedRef<MetaDataType>> FoundMetaData;
 		for (const auto& MetaDataEntry : MetaData)
 		{
@@ -1287,10 +1279,50 @@ public:
 		AddMetadataInternal(AddMe);
 	}
 
+	/**
+	 * Remove metadata to this widget.
+	 * @returns Number of removed elements.
+	 */
+	template<typename MetaDataType>
+	int32 RemoveMetaData(const TSharedRef<MetaDataType>& RemoveMe)
+	{
+		return MetaData.RemoveSingleSwap(RemoveMe);
+	}
+
 private:
 
 	void AddMetadataInternal(const TSharedRef<ISlateMetaData>& AddMe);
 
+	template<typename MetaDataType>
+	int32 RemoveAllMetaData()
+	{
+		int32 NumBefore = MetaData.Num();
+		for (int32 Index = NumBefore - 1; Index >= 0; --Index)
+		{
+			const auto& MetaDataEntry = MetaData[Index];
+			if (MetaDataEntry->IsOfType<MetaDataType>())
+			{
+				MetaData.RemoveAtSwap(Index);
+			}
+		}
+		return NumBefore - MetaData.Num();
+	}
+
+	template<typename MetaDataType>
+	bool RemoveMetaData()
+	{
+		for (int32 Index = MetaData.Num() - 1; Index >= 0; --Index)
+		{
+			const auto& MetaDataEntry = MetaData[Index];
+			if (MetaDataEntry->IsOfType<MetaDataType>())
+			{
+				MetaData.RemoveAtSwap(Index);
+				return true;
+			}
+		}
+		return false;
+	}
+
 public:
 
 	/** See OnMouseButtonDown event */
@@ -1528,9 +1560,6 @@ private:
 	/** Iterates over the active timer handles on the widget and executes them if their interval has elapsed. */
 	void ExecuteActiveTimers(double CurrentTime, float DeltaTime);
 
-	const FPointerEventHandler* GetPointerEvent(const FName EventName) const;
-	void SetPointerEvent(const FName EventName, FPointerEventHandler& InEvent);
-
 protected:
 	/**
 	 * Performs the attribute assignment and invalidates the widget minimally based on what actually changed.  So if the boundness of the attribute didn't change
@@ -1547,9 +1576,9 @@ protected:
 	virtual ~SWidget();
 
 private:
-
 	/** Handle to the proxy when on the fast path */
 	mutable FWidgetProxyHandle FastPathProxyHandle;
+
 protected:
 	/** Is this widget hovered? */
 	uint8 bIsHovered : 1;
@@ -1570,6 +1599,7 @@ protected:
 	  * really do the clipping.
 	  */
 	uint8 bClippingProxy : 1;
+
 private:
 	/**
 	 * Whether this widget is a "tool tip force field".  That is, tool-tips should never spawn over the area
@@ -1605,11 +1635,6 @@ protected:
 	/** if this widget should always invalidate the prepass step when volatile */
 	uint8 bVolatilityAlwaysInvalidatesPrepass : 1;
 
-#if !(UE_BUILD_SHIPPING || UE_BUILD_TEST)
-	UE_DEPRECATED(4.21, "Setting bCanTick on a widget directly is deprecated and will not function.  Call SetCanTick instead")
-	uint8 bCanTick : 1;
-#endif
-
 #if WITH_ACCESSIBILITY
 	/** All variables surrounding how this widget is exposed to the platform's accessibility API. */
 	uint8 bCanChildrenBeAccessible : 1;
@@ -1676,7 +1701,7 @@ protected:
 	/** Is this widget visible, hidden or collapsed */
 	TAttribute< EVisibility > Visibility;
 
-	/** The opacity of the widget.  Automatically applied during rendering. */
+	/** The opacity of the widget. Automatically applied during rendering. */
 	float RenderOpacity;
 
 	/** Render transform of this widget. TOptional<> to allow code to skip expensive overhead if there is no render transform applied. */
@@ -1685,39 +1710,32 @@ protected:
 	/** Render transform pivot of this widget (in normalized local space) */
 	TAttribute< FVector2D > RenderTransformPivot;
 
+private:
+	/** Metadata associated with this widget. */
+	TArray<TSharedRef<ISlateMetaData>> MetaData;
+
+	/** Pointer to this widgets parent widget.  If it is null this is a root widget or it is not in the widget tree */
+	TWeakPtr<SWidget> ParentWidgetPtr;
+
 	/** Debugging information on the type of widget we're creating for the Widget Reflector. */
 	FName TypeOfWidget;
 
 #if !UE_BUILD_SHIPPING
-
 	/** Full file path (and line) in which this widget was created */
 	FName CreatedInLocation;
-
 #endif
 
 	/** Tag for this widget */
 	FName Tag;
 
-	/** Metadata associated with this widget */
-	TArray<TSharedRef<ISlateMetaData>> MetaData;
-
+#if !(UE_BUILD_SHIPPING || UE_BUILD_TEST)
+	UE_DEPRECATED(4.27, "Access to SWidget::Cursor is deprecated and will not function. Call SetCursor/GetCursor instead")
 	/** The cursor to show when the mouse is hovering over this widget. */
-	TAttribute< TOptional<EMouseCursor::Type> > Cursor;
+	TAttribute<TOptional<EMouseCursor::Type>> Cursor;
+#endif
 
 private:
 
-	/** Tool tip content for this widget */
-	TSharedPtr<IToolTip> ToolTip;
-
-	/** Pointer to this widgets parent widget.  If it is null this is a root widget or it is not in the widget tree */
-	TWeakPtr<SWidget> ParentWidgetPtr;
-
-	// Events
-	TArray<TPair<FName, FPointerEventHandler>> PointerEvents;
-
-	FNoReplyPointerEventHandler MouseEnterHandler;
-	FSimpleNoReplyPointerEventHandler MouseLeaveHandler;
-
 #if UE_SLATE_WITH_WIDGET_UNIQUE_IDENTIFIER
 	/** The widget's id */
 	uint64 UniqueIdentifier;
