Slate: Fix CIS, add deprecation warning.
#rnx

diff --git a/Engine/Source/Runtime/SlateCore/Private/Widgets/SWidget.cpp b/Engine/Source/Runtime/SlateCore/Private/Widgets/SWidget.cpp
--- a/Engine/Source/Runtime/SlateCore/Private/Widgets/SWidget.cpp
+++ b/Engine/Source/Runtime/SlateCore/Private/Widgets/SWidget.cpp
@@ -244,6 +244,7 @@ SWidget::SWidget()
 	UE_TRACE_SLATE_WIDGET_ADDED(this);
 }
 
+PRAGMA_DISABLE_DEPRECATION_WARNINGS
 SWidget::~SWidget()
 {
 #if UE_WITH_SLATE_DEBUG_FIND_WIDGET_REFLECTION_METADATA
@@ -291,6 +292,7 @@ SWidget::~SWidget()
 	DEC_DWORD_STAT(STAT_SlateTotalWidgets);
 	DEC_MEMORY_STAT_BY(STAT_SlateSWidgetAllocSize, AllocSize);
 }
+PRAGMA_ENABLE_DEPRECATION_WARNINGS
 
 void SWidget::Construct(
 	const TAttribute<FText>& InToolTipText,
