Optimize slate text kerning by removing string hashing from critical path in CompositeFont

 - 864ms to 395ms for Slate_PaintSlowPath after logging 10 000 lines to the output log in a single frame

diff --git a/Engine/Source/Runtime/SlateCore/Private/Fonts/CompositeFont.cpp b/Engine/Source/Runtime/SlateCore/Private/Fonts/CompositeFont.cpp
--- a/Engine/Source/Runtime/SlateCore/Private/Fonts/CompositeFont.cpp
+++ b/Engine/Source/Runtime/SlateCore/Private/Fonts/CompositeFont.cpp
@@ -31,6 +31,7 @@ FFontData::FFontData()
 	, FontData_DEPRECATED()
 #endif // WITH_EDITORONLY_DATA
 {
+	FontFilenameHash = GetTypeHash(FontFilename);
 }
 
 FFontData::FFontData(const UObject* const InFontFaceAsset, const int32 InSubFaceIndex)
@@ -44,6 +45,7 @@ FFontData::FFontData(const UObject* const InFontFaceAsset, const int32 InSubFace
 	, FontData_DEPRECATED()
 #endif // WITH_EDITORONLY_DATA
 {
+	FontFilenameHash = GetTypeHash(FontFilename);
 	if (FontFaceAsset)
 	{
 		CastChecked<const IFontFaceInterface>(FontFaceAsset);
@@ -62,6 +64,7 @@ FFontData::FFontData(FString InFontFilename, const EFontHinting InHinting, const
 #endif // WITH_EDITORONLY_DATA
 {
 	check(InLoadingPolicy != EFontLoadingPolicy::Inline);
+	FontFilenameHash = GetTypeHash(FontFilename);
 }
 
 bool FFontData::HasFont() const
@@ -187,10 +190,11 @@ bool FFontData::operator==(const FFontData& Other) const
 	}
 
 	// Compare inline properties
-	return FontFilename == Other.FontFilename
+	return FontFilenameHash == Other.FontFilenameHash 
 		&& Hinting == Other.Hinting
 		&& LoadingPolicy == Other.LoadingPolicy
-		&& SubFaceIndex == Other.SubFaceIndex;
+		&& SubFaceIndex == Other.SubFaceIndex
+		&& FontFilename == Other.FontFilename;
 }
 
 bool FFontData::operator!=(const FFontData& Other) const
@@ -242,6 +246,7 @@ bool FFontData::Serialize(FArchive& Ar)
 		}
 	}
 
+	FontFilenameHash = GetTypeHash(FontFilename);
 	return true;
 }
 
diff --git a/Engine/Source/Runtime/SlateCore/Public/Fonts/CompositeFont.h b/Engine/Source/Runtime/SlateCore/Public/Fonts/CompositeFont.h
--- a/Engine/Source/Runtime/SlateCore/Public/Fonts/CompositeFont.h
+++ b/Engine/Source/Runtime/SlateCore/Public/Fonts/CompositeFont.h
@@ -201,7 +201,7 @@ struct SLATECORE_API FFontData
 		}
 		else
 		{
-			KeyHash = HashCombine(KeyHash, GetTypeHash(Key.FontFilename));
+			KeyHash = HashCombine(KeyHash, Key.FontFilenameHash);
 			KeyHash = HashCombine(KeyHash, GetTypeHash(Key.Hinting));
 			KeyHash = HashCombine(KeyHash, GetTypeHash(Key.LoadingPolicy));
 		}
@@ -229,6 +229,12 @@ private:
 	UPROPERTY()
 	FString FontFilename;
 
+	/**
+	 * Cached hash value of FontFilename.
+	 * Must be updated everytime FontFilename changes.
+	 */
+	uint32  FontFilenameHash;
+
 	/**
 	 * The hinting algorithm to use with the font.
 	 * This variable is ignored if we have a font face asset, and is synchronized with the font face asset on load in a cooked build.
