Updated TableView/TreeView Header Style

diff --git a/Engine/Source/Editor/SceneOutliner/Private/SceneOutlinerActorInfoColumn.cpp b/Engine/Source/Editor/SceneOutliner/Private/SceneOutlinerActorInfoColumn.cpp
--- a/Engine/Source/Editor/SceneOutliner/Private/SceneOutlinerActorInfoColumn.cpp
+++ b/Engine/Source/Editor/SceneOutliner/Private/SceneOutlinerActorInfoColumn.cpp
@@ -208,8 +208,6 @@ SHeaderRow::FColumn::FArguments FActorInfoColumn::ConstructHeaderRowColumn()
 			.VAlign(VAlign_Center)
 			[
 				SNew(STextBlock)
-				.TextStyle(FAppStyle::Get(), "ButtonText")
-				.TransformPolicy(ETextTransformPolicy::ToUpper)
 				.Text( this, &FActorInfoColumn::GetSelectedMode )
 			]
 		];
diff --git a/Engine/Source/Editor/SceneOutliner/Private/SceneOutlinerGutter.cpp b/Engine/Source/Editor/SceneOutliner/Private/SceneOutlinerGutter.cpp
--- a/Engine/Source/Editor/SceneOutliner/Private/SceneOutlinerGutter.cpp
+++ b/Engine/Source/Editor/SceneOutliner/Private/SceneOutlinerGutter.cpp
@@ -347,10 +347,11 @@ SHeaderRow::FColumn::FArguments FSceneOutlinerGutter::ConstructHeaderRowColumn()
 {
 	return SHeaderRow::Column(GetColumnID())
 		.FixedWidth(24.f)
-		.HAlignHeader(HAlign_Right)
+		.HAlignHeader(HAlign_Center)
 		.VAlignHeader(VAlign_Center)
-		.HAlignCell(HAlign_Right)
+		.HAlignCell(HAlign_Center)
 		.VAlignCell(VAlign_Center)
+		.HeaderContentPadding(FMargin(0.0))
 		[
 			SNew(SImage)
 			.ColorAndOpacity(FSlateColor::UseForeground())
diff --git a/Engine/Source/Runtime/Slate/Private/Widgets/Views/SHeaderRow.cpp b/Engine/Source/Runtime/Slate/Private/Widgets/Views/SHeaderRow.cpp
--- a/Engine/Source/Runtime/Slate/Private/Widgets/Views/SHeaderRow.cpp
+++ b/Engine/Source/Runtime/Slate/Private/Widgets/Views/SHeaderRow.cpp
@@ -99,8 +99,7 @@ public:
 				.VAlign( VAlign_Center )
 				[
 					SNew(STextBlock)
-					.TextStyle(FAppStyle::Get(), "ButtonText")
-					.TransformPolicy(ETextTransformPolicy::ToUpper)
+					.TextStyle(FAppStyle::Get(), "NormalText")
 					.Text( LabelText )
 					.ToolTipText( TooltipText )
 				];
diff --git a/Engine/Source/Runtime/SlateCore/Private/Styling/StarshipCoreStyle.cpp b/Engine/Source/Runtime/SlateCore/Private/Styling/StarshipCoreStyle.cpp
--- a/Engine/Source/Runtime/SlateCore/Private/Styling/StarshipCoreStyle.cpp
+++ b/Engine/Source/Runtime/SlateCore/Private/Styling/StarshipCoreStyle.cpp
@@ -1628,8 +1628,8 @@ void FStarshipCoreStyle::SetupTableViewStyles(TSharedRef<FStyle>& Style)
 		.SetSortPrimaryDescendingImage(IMAGE_BRUSH_SVG("Starship/CoreWidgets/TableView/sort-down-arrow", Icon12x12))
 		.SetSortSecondaryAscendingImage(IMAGE_BRUSH_SVG("Starship/CoreWidgets/TableView/sort-up-arrows", Icon12x12))
 		.SetSortSecondaryDescendingImage(IMAGE_BRUSH_SVG("Starship/CoreWidgets/TableView/sort-down-arrows", Icon12x12))
-		.SetNormalBrush(FSlateColorBrush(FStyleColors::Background))
-		.SetHoveredBrush(FSlateColorBrush(FStyleColors::Secondary))
+		.SetNormalBrush(FSlateColorBrush(FStyleColors::Header))
+		.SetHoveredBrush(FSlateColorBrush(FStyleColors::Dropdown))
 
 		.SetMenuDropdownImage(IMAGE_BRUSH_SVG("Starship/Common/ellipsis-vertical-narrow", FVector2D(6.f, 24.f)))
 		.SetMenuDropdownNormalBorderBrush(FSlateNoResource())
@@ -1642,8 +1642,8 @@ void FStarshipCoreStyle::SetupTableViewStyles(TSharedRef<FStyle>& Style)
 		.SetSortPrimaryDescendingImage(IMAGE_BRUSH_SVG("Starship/CoreWidgets/TableView/sort-down-arrow", Icon12x12))
 		.SetSortSecondaryAscendingImage(IMAGE_BRUSH_SVG("Starship/CoreWidgets/TableView/sort-up-arrows", Icon12x12))
 		.SetSortSecondaryDescendingImage(IMAGE_BRUSH_SVG("Starship/CoreWidgets/TableView/sort-down-arrows", Icon12x12))
-		.SetNormalBrush(FSlateNoResource())
-		.SetHoveredBrush(FSlateNoResource())
+		.SetNormalBrush(FSlateColorBrush(FStyleColors::Header))
+		.SetHoveredBrush(FSlateColorBrush(FStyleColors::Dropdown))
 
 		.SetMenuDropdownImage(IMAGE_BRUSH_SVG("Starship/Common/ellipsis-vertical-narrow", FVector2D(6.f, 24.f)))
 		.SetMenuDropdownNormalBorderBrush(FSlateNoResource())
@@ -1651,12 +1651,13 @@ void FStarshipCoreStyle::SetupTableViewStyles(TSharedRef<FStyle>& Style)
 
 	const FSplitterStyle TableHeaderSplitterStyle = FSplitterStyle()
 		.SetHandleNormalBrush(FSlateColorBrush(FStyleColors::Recessed))
-		.SetHandleHighlightBrush(IMAGE_BRUSH("Common/HeaderSplitterGrip", Icon8x8));
+		.SetHandleHighlightBrush(FSlateColorBrush(FStyleColors::Recessed));
 
 	Style->Set("TableView.Header", FHeaderRowStyle()
 		.SetColumnStyle(TableColumnHeaderStyle)
 		.SetLastColumnStyle(TableLastColumnHeaderStyle)
 		.SetColumnSplitterStyle(TableHeaderSplitterStyle)
+		.SetSplitterHandleSize(1.0)
 		.SetBackgroundBrush(FSlateColorBrush(FStyleColors::Background))
 		.SetForegroundColor(FStyleColors::Foreground)
 		.SetHorizontalSeparatorBrush(FSlateColorBrush(FStyleColors::Recessed))
