Slate: Clean DeclarativeSyntaxSupport. Tooltip is now an attribute. Pass the FArguments to SWidget.

diff --git a/Engine/Plugins/Media/MediaIOFramework/Source/MediaIOEditor/Public/Widgets/SMediaPermutationsSelector.inl b/Engine/Plugins/Media/MediaIOFramework/Source/MediaIOEditor/Public/Widgets/SMediaPermutationsSelector.inl
--- a/Engine/Plugins/Media/MediaIOFramework/Source/MediaIOEditor/Public/Widgets/SMediaPermutationsSelector.inl
+++ b/Engine/Plugins/Media/MediaIOFramework/Source/MediaIOEditor/Public/Widgets/SMediaPermutationsSelector.inl
@@ -22,9 +22,6 @@ ItemType SMediaPermutationsSelector<ItemType, ItemBuilder>::GetSelectedItem() co
 template <typename ItemType, typename ItemBuilder>
 void SMediaPermutationsSelector<ItemType, ItemBuilder>::Construct(const FArguments& InArgs)
 {
-	SWidget::Construct(InArgs._ToolTipText, InArgs._ToolTip, InArgs._Cursor, InArgs._IsEnabled, InArgs._Visibility, InArgs._RenderOpacity, InArgs._RenderTransform, InArgs._RenderTransformPivot, InArgs._Tag, InArgs._ForceVolatile, InArgs._Clipping,
-					   InArgs._FlowDirectionPreference, InArgs._AccessibleParams, InArgs.MetaData);
-
 	PermutationsSource = InArgs._PermutationsSource;
 	SelectedPermutationIndex = INDEX_NONE;
 	OnSelectionChanged = InArgs._OnSelectionChanged;
diff --git a/Engine/Source/Editor/GraphEditor/Private/KismetNodes/SGraphNodeK2Base.cpp b/Engine/Source/Editor/GraphEditor/Private/KismetNodes/SGraphNodeK2Base.cpp
--- a/Engine/Source/Editor/GraphEditor/Private/KismetNodes/SGraphNodeK2Base.cpp
+++ b/Engine/Source/Editor/GraphEditor/Private/KismetNodes/SGraphNodeK2Base.cpp
@@ -38,7 +38,7 @@ void SGraphNodeK2Base::UpdateStandardNode()
 {
 	SGraphNode::UpdateGraphNode();
 	// clear the default tooltip, to make room for our custom "complex" tooltip
-	SetToolTip(NULL);
+	SetToolTip(nullptr);
 }
 
 void SGraphNodeK2Base::UpdateCompactNode()
diff --git a/Engine/Source/Editor/GraphEditor/Private/SGraphNode.cpp b/Engine/Source/Editor/GraphEditor/Private/SGraphNode.cpp
--- a/Engine/Source/Editor/GraphEditor/Private/SGraphNode.cpp
+++ b/Engine/Source/Editor/GraphEditor/Private/SGraphNode.cpp
@@ -442,7 +442,7 @@ void SGraphNode::OnToolTipClosing()
 {
 	if (bProvidedComplexTooltip)
 	{
-		SetToolTip(NULL);
+		SetToolTip(nullptr);
 		bProvidedComplexTooltip = false;
 	}
 }
diff --git a/Engine/Source/Runtime/AppFramework/Public/Widgets/Workflow/SWizard.h b/Engine/Source/Runtime/AppFramework/Public/Widgets/Workflow/SWizard.h
--- a/Engine/Source/Runtime/AppFramework/Public/Widgets/Workflow/SWizard.h
+++ b/Engine/Source/Runtime/AppFramework/Public/Widgets/Workflow/SWizard.h
@@ -41,7 +41,7 @@ public:
 	{
 	public:
 
-		HACK_SLATE_SLOT_ARGS(FWizardPage)
+		SLATE_BEGIN_ARGS(FWizardPage)
 			: _ButtonContent()
 			, _CanShow(true)
 			, _PageContent()
diff --git a/Engine/Source/Runtime/Engine/Private/UnrealEngine.cpp b/Engine/Source/Runtime/Engine/Private/UnrealEngine.cpp
--- a/Engine/Source/Runtime/Engine/Private/UnrealEngine.cpp
+++ b/Engine/Source/Runtime/Engine/Private/UnrealEngine.cpp
@@ -26,6 +26,7 @@ UnrealEngine.cpp: Implements the UEngine class and helpers.
 #include "Misc/CommandLine.h"
 #include "Misc/FileHelper.h"
 #include "Misc/Paths.h"
+#include "Misc/FrameValue.h"
 #include "HAL/Runnable.h"
 #include "Misc/OutputDeviceArchiveWrapper.h"
 #include "Stats/StatsMisc.h"
diff --git a/Engine/Source/Runtime/Slate/Private/Widgets/Views/SHeaderRow.cpp b/Engine/Source/Runtime/Slate/Private/Widgets/Views/SHeaderRow.cpp
--- a/Engine/Source/Runtime/Slate/Private/Widgets/Views/SHeaderRow.cpp
+++ b/Engine/Source/Runtime/Slate/Private/Widgets/Views/SHeaderRow.cpp
@@ -51,9 +51,6 @@ public:
 	{
 		check(InArgs._Style);
 
-		SWidget::Construct( InArgs._ToolTipText, InArgs._ToolTip, InArgs._Cursor, InArgs._IsEnabled, InArgs._Visibility, InArgs._RenderOpacity, InArgs._RenderTransform, InArgs._RenderTransformPivot, InArgs._Tag, InArgs._ForceVolatile, InArgs._Clipping,
-							InArgs._FlowDirectionPreference, InArgs._AccessibleParams, InArgs.MetaData );
-
 		Style = InArgs._Style;
 		ColumnId = Column.ColumnId;
 		SortMode = Column.SortMode;
@@ -442,9 +439,6 @@ void SHeaderRow::Construct( const FArguments& InArgs )
 {
 	check(InArgs._Style);
 
-	SWidget::Construct( InArgs._ToolTipText, InArgs._ToolTip, InArgs._Cursor, InArgs._IsEnabled, InArgs._Visibility, InArgs._RenderOpacity, InArgs._RenderTransform, InArgs._RenderTransformPivot, InArgs._Tag, InArgs._ForceVolatile, InArgs._Clipping,
-						InArgs._FlowDirectionPreference, InArgs._AccessibleParams, InArgs.MetaData );
-
 	ScrollBarThickness = FVector2D::ZeroVector;
 	ScrollBarVisibility = EVisibility::Collapsed;
 	Style = InArgs._Style;
diff --git a/Engine/Source/Runtime/Slate/Public/Widgets/Views/SHeaderRow.h b/Engine/Source/Runtime/Slate/Public/Widgets/Views/SHeaderRow.h
--- a/Engine/Source/Runtime/Slate/Public/Widgets/Views/SHeaderRow.h
+++ b/Engine/Source/Runtime/Slate/Public/Widgets/Views/SHeaderRow.h
@@ -96,7 +96,7 @@ public:
 	{
 	public:
 
-		HACK_SLATE_SLOT_ARGS(FColumn)
+		SLATE_BEGIN_ARGS(FColumn)
 			: _ColumnId()
 			, _DefaultLabel()
 			, _DefaultTooltip()
diff --git a/Engine/Source/Runtime/SlateCore/Private/Types/SlateToolTipMetaData.h b/Engine/Source/Runtime/SlateCore/Private/Types/SlateToolTipMetaData.h
--- a/Engine/Source/Runtime/SlateCore/Private/Types/SlateToolTipMetaData.h
+++ b/Engine/Source/Runtime/SlateCore/Private/Types/SlateToolTipMetaData.h
@@ -13,5 +13,5 @@ class FSlateToolTipMetaData : public ISlateMetaData
 public:
 	SLATE_METADATA_TYPE(FSlateToolTipMetaData, ISlateMetaData)
 	
-	TSharedPtr<IToolTip> ToolTip;
+	TAttribute<TSharedPtr<IToolTip>> ToolTip;
 };
diff --git a/Engine/Source/Runtime/SlateCore/Private/Widgets/SWidget.cpp b/Engine/Source/Runtime/SlateCore/Private/Widgets/SWidget.cpp
--- a/Engine/Source/Runtime/SlateCore/Private/Widgets/SWidget.cpp
+++ b/Engine/Source/Runtime/SlateCore/Private/Widgets/SWidget.cpp
@@ -294,6 +294,7 @@ SWidget::~SWidget()
 }
 PRAGMA_ENABLE_DEPRECATION_WARNINGS
 
+PRAGMA_DISABLE_DEPRECATION_WARNINGS
 void SWidget::Construct(
 	const TAttribute<FText>& InToolTipText,
 	const TSharedPtr<IToolTip>& InToolTip,
@@ -311,49 +312,83 @@ void SWidget::Construct(
 	const TArray<TSharedRef<ISlateMetaData>>& InMetaData
 )
 {
-	EnabledState = InEnabledState;
-	Visibility = InVisibility;
-	RenderOpacity = InRenderOpacity;
-	RenderTransform = InTransform;
-	RenderTransformPivot = InTransformPivot;
-	Tag = InTag;
-	bForceVolatile = InForceVolatile;
-	Clipping = InClipping;
-	FlowDirectionPreference = InFlowPreference;
-	MetaData.Append(InMetaData);
+	FSlateBaseNamedArgs Args;
+	Args._ToolTipText = InToolTipText;
+	Args._ToolTip = InToolTip;
+	Args._Cursor = InCursor;
+	Args._IsEnabled = InEnabledState;
+	Args._Visibility = InVisibility;
+	Args._RenderOpacity = InRenderOpacity;
+	Args._ForceVolatile = InForceVolatile;
+	Args._Clipping = InClipping;
+	Args._FlowDirectionPreference = InFlowPreference;
+	Args._RenderTransform = InTransform;
+	Args._RenderTransformPivot = InTransformPivot;
+	Args._Tag = InTag;
+	Args._AccessibleParams = InAccessibleData;
+	Args.MetaData = InMetaData;
+	SWidgetConstruct(Args);
+}
+PRAGMA_ENABLE_DEPRECATION_WARNINGS
+
+PRAGMA_DISABLE_DEPRECATION_WARNINGS
+void SWidget::SWidgetConstruct(const TAttribute<FText>& InToolTipText, const TSharedPtr<IToolTip>& InToolTip, const TAttribute< TOptional<EMouseCursor::Type> >& InCursor, const TAttribute<bool>& InEnabledState,
+							   const TAttribute<EVisibility>& InVisibility, const float InRenderOpacity, const TAttribute<TOptional<FSlateRenderTransform>>& InTransform, const TAttribute<FVector2D>& InTransformPivot,
+							   const FName& InTag, const bool InForceVolatile, const EWidgetClipping InClipping, const EFlowDirectionPreference InFlowPreference, const TOptional<FAccessibleWidgetData>& InAccessibleData,
+							   const TArray<TSharedRef<ISlateMetaData>>& InMetaData)
+{
+	Construct(InToolTipText, InToolTip, InCursor, InEnabledState, InVisibility, InRenderOpacity, InTransform, InTransformPivot, InTag, InForceVolatile, InClipping, InFlowPreference, InAccessibleData, InMetaData);
+}
+PRAGMA_ENABLE_DEPRECATION_WARNINGS
+
+void SWidget::SWidgetConstruct(const FSlateBaseNamedArgs& Args)
+{
+	EnabledState = Args._IsEnabled;
+	Visibility = Args._Visibility;
+	RenderOpacity = Args._RenderOpacity;
+	RenderTransform = Args._RenderTransform;
+	RenderTransformPivot = Args._RenderTransformPivot;
+	Tag = Args._Tag;
+	bForceVolatile = Args._ForceVolatile;
+	Clipping = Args._Clipping;
+	FlowDirectionPreference = Args._FlowDirectionPreference;
+	MetaData.Append(Args.MetaData);
 
-	if (InToolTip.IsValid())
+	if (Args._ToolTip.IsSet())
 	{
 		// If someone specified a fancy widget tooltip, use it.
-		SetToolTip(InToolTip);
+		SetToolTip(Args._ToolTip);
 	}
-	else if (InToolTipText.IsSet())
+	else if (Args._ToolTipText.IsSet())
 	{
 		// If someone specified a text binding, make a tooltip out of it
-		SetToolTipText(InToolTipText);
+		SetToolTipText(Args._ToolTipText);
 	}
 
-	SetCursor(InCursor);
+	SetCursor(Args._Cursor);
 
 #if WITH_ACCESSIBILITY
-	if (InAccessibleData.IsSet())
+	// If custom text is provided, force behavior to custom. Otherwise, use the passed-in behavior and set their default text.
+	if (Args._AccessibleText.IsSet() || Args._AccessibleParams.IsSet())
 	{
-		SetCanChildrenBeAccessible(InAccessibleData->bCanChildrenBeAccessible);
-		// If custom text is provided, force behavior to custom. Otherwise, use the passed-in behavior and set their default text.
-		SetAccessibleBehavior(InAccessibleData->AccessibleText.IsSet() ? EAccessibleBehavior::Custom : InAccessibleData->AccessibleBehavior, InAccessibleData->AccessibleText, EAccessibleType::Main);
-		SetAccessibleBehavior(InAccessibleData->AccessibleSummaryText.IsSet() ? EAccessibleBehavior::Custom : InAccessibleData->AccessibleSummaryBehavior, InAccessibleData->AccessibleSummaryText, EAccessibleType::Summary);
+		auto SetAccessibleWidgetData = [this](const FAccessibleWidgetData& AccessibleParams)
+		{
+			SetCanChildrenBeAccessible(AccessibleParams.bCanChildrenBeAccessible);
+			SetAccessibleBehavior(AccessibleParams.AccessibleText.IsSet() ? EAccessibleBehavior::Custom : AccessibleParams.AccessibleBehavior, AccessibleParams.AccessibleText, EAccessibleType::Main);
+		SetAccessibleBehavior(AccessibleParams.AccessibleSummaryText.IsSet() ? EAccessibleBehavior::Custom : AccessibleParams.AccessibleSummaryBehavior, AccessibleParams.AccessibleSummaryText, EAccessibleType::Summary);
+		};
+		if (Args._AccessibleText.IsSet())
+		{
+			SetAccessibleWidgetData(FAccessibleWidgetData{ Args._AccessibleText });
+		}
+		else
+		{
+			SetAccessibleWidgetData(Args._AccessibleParams.GetValue());
+		}
 	}
 #endif
 }
 
-void SWidget::SWidgetConstruct(const TAttribute<FText>& InToolTipText, const TSharedPtr<IToolTip>& InToolTip, const TAttribute< TOptional<EMouseCursor::Type> >& InCursor, const TAttribute<bool>& InEnabledState,
-							   const TAttribute<EVisibility>& InVisibility, const float InRenderOpacity, const TAttribute<TOptional<FSlateRenderTransform>>& InTransform, const TAttribute<FVector2D>& InTransformPivot,
-							   const FName& InTag, const bool InForceVolatile, const EWidgetClipping InClipping, const EFlowDirectionPreference InFlowPreference, const TOptional<FAccessibleWidgetData>& InAccessibleData,
-							   const TArray<TSharedRef<ISlateMetaData>>& InMetaData)
-{
-	Construct(InToolTipText, InToolTip, InCursor, InEnabledState, InVisibility, InRenderOpacity, InTransform, InTransformPivot, InTag, InForceVolatile, InClipping, InFlowPreference, InAccessibleData, InMetaData);
-}
-
 FReply SWidget::OnFocusReceived(const FGeometry& MyGeometry, const FFocusEvent& InFocusEvent)
 {
 	return FReply::Unhandled();
@@ -1082,9 +1117,9 @@ void SWidget::SetToolTipText( const FText& ToolTipText )
 	}
 }
 
-void SWidget::SetToolTip( const TSharedPtr<IToolTip> & InToolTip )
+void SWidget::SetToolTip(const TAttribute<TSharedPtr<IToolTip>>& InToolTip)
 {
-	if (InToolTip)
+	if (InToolTip.IsSet())
 	{
 		Private::FindOrAddToolTipMetaData(this)->ToolTip = InToolTip;
 	}
@@ -1098,7 +1133,7 @@ TSharedPtr<IToolTip> SWidget::GetToolTip()
 {
 	if (TSharedPtr<FSlateToolTipMetaData> Data = GetMetaData<FSlateToolTipMetaData>())
 	{
-		return Data->ToolTip;
+		return Data->ToolTip.Get();
 	}
 	return TSharedPtr<IToolTip>();
 }
@@ -1756,9 +1791,12 @@ FText SWidget::GetAccessibleText(EAccessibleType AccessibleType) const
 		//TODO should use GetToolTip
 		if (TSharedPtr<FSlateToolTipMetaData> Data = GetMetaData<FSlateToolTipMetaData>())
 		{
-			if (Data->ToolTip && !Data->ToolTip->IsEmpty())
+			if (TSharedPtr<IToolTip> ToolTip = Data->ToolTip.Get())
 			{
-				return Data->ToolTip->GetContentWidget()->GetAccessibleText(EAccessibleType::Main);
+				if (ToolTip && !ToolTip->IsEmpty())
+				{
+					return ToolTip->GetContentWidget()->GetAccessibleText(EAccessibleType::Main);
+				}
 			}
 		}
 		break;
diff --git a/Engine/Source/Runtime/SlateCore/Public/Widgets/DeclarativeSyntaxSupport.h b/Engine/Source/Runtime/SlateCore/Public/Widgets/DeclarativeSyntaxSupport.h
--- a/Engine/Source/Runtime/SlateCore/Public/Widgets/DeclarativeSyntaxSupport.h
+++ b/Engine/Source/Runtime/SlateCore/Public/Widgets/DeclarativeSyntaxSupport.h
@@ -105,27 +105,14 @@ template<typename WidgetType> struct TSlateBaseNamedArgs;
 		typedef FArguments WidgetArgsType; \
 		FORCENOINLINE FArguments()
 
-
-// @todo UMG: Probably remove this?
-#define HACK_SLATE_SLOT_ARGS( WidgetType ) \
-	public: \
-	struct FArguments : public TSlateBaseNamedArgs<WidgetType> \
-	{ \
-		typedef FArguments WidgetArgsType; \
-		FORCENOINLINE FArguments()
-
-
-
 #define SLATE_END_ARGS() \
 	};
 
 
-/**
- * Use this macro to add a attribute to the declaration of your widget.
- * An attribute can be a value or a function.
- */
-#define SLATE_ATTRIBUTE( AttrType, AttrName ) \
-		TAttribute< AttrType > _##AttrName; \
+#define SLATE_PRIVATE_ATTRIBUTE_VARIABLE( AttrType, AttrName ) \
+		TAttribute< AttrType > _##AttrName
+
+#define SLATE_PRIVATE_ATTRIBUTE_FUNCTION( AttrType, AttrName ) \
 		WidgetArgsType& AttrName( const TAttribute< AttrType >& InAttribute ) \
 		{ \
 			_##AttrName = InAttribute; \
@@ -301,31 +288,41 @@ template<typename WidgetType> struct TSlateBaseNamedArgs;
 			return this->Me(); \
 		}
 
+#define SLATE_PRIVATE_ARGUMENT_VARIABLE( ArgType, ArgName ) \
+		ArgType _##ArgName
+		
 
-/**
- * Use this macro to declare a slate argument.
- * Arguments differ from attributes in that they can only be values
- */
-#define SLATE_ARGUMENT( ArgType, ArgName ) \
-		ArgType _##ArgName; \
+#define SLATE_PRIVATE_ARGUMENT_FUNCTION( ArgType, ArgName ) \
 		WidgetArgsType& ArgName( ArgType InArg ) \
 		{ \
 			_##ArgName = InArg; \
 			return this->Me(); \
 		}
 
+/**
+ * Use this macro to add a attribute to the declaration of your widget.
+ * An attribute can be a value or a function.
+ */
+#define SLATE_ATTRIBUTE( AttrType, AttrName ) \
+		SLATE_PRIVATE_ATTRIBUTE_VARIABLE( AttrType, AttrName ); \
+		SLATE_PRIVATE_ATTRIBUTE_FUNCTION( AttrType, AttrName )
+
+/**
+ * Use this macro to declare a slate argument.
+ * Arguments differ from attributes in that they can only be values
+ */
+#define SLATE_ARGUMENT( ArgType, ArgName ) \
+		SLATE_PRIVATE_ARGUMENT_VARIABLE( ArgType, ArgName ); \
+		SLATE_PRIVATE_ARGUMENT_FUNCTION ( ArgType, ArgName )
+
 /**
  * Like SLATE_ARGUMENT, but support a default value. e.g.
  * 
  * SLATE_ARGUMENT_DEFAULT(float, WheelScrollMultiplier) { 1.0f };
  */
 #define SLATE_ARGUMENT_DEFAULT( ArgType, ArgName ) \
-		WidgetArgsType& ArgName( ArgType InArg ) \
-		{ \
-			_##ArgName = InArg; \
-			return this->Me(); \
-		}\
-		ArgType _##ArgName
+	SLATE_PRIVATE_ARGUMENT_FUNCTION(ArgType, ArgName) \
+	SLATE_PRIVATE_ARGUMENT_VARIABLE(ArgType, ArgName)
 
 /**
  * Use this macro to declare a slate argument.
@@ -777,28 +774,49 @@ struct NamedSlotProperty
 		DelegateName _##EventName; \
 
 
+/** Base class for named arguments. Provides settings necessary for all widgets. */
+struct FSlateBaseNamedArgs
+{
+	FSlateBaseNamedArgs() = default;
+
+	SLATE_PRIVATE_ATTRIBUTE_VARIABLE(FText, ToolTipText);
+	SLATE_PRIVATE_ATTRIBUTE_VARIABLE(TSharedPtr<IToolTip>, ToolTip);
+	SLATE_PRIVATE_ATTRIBUTE_VARIABLE(TOptional<EMouseCursor::Type>, Cursor);
+	SLATE_PRIVATE_ATTRIBUTE_VARIABLE(bool, IsEnabled) = true;
+	SLATE_PRIVATE_ATTRIBUTE_VARIABLE(EVisibility, Visibility) = EVisibility::Visible;
+	SLATE_PRIVATE_ARGUMENT_VARIABLE(float, RenderOpacity) = 1.f;
+	SLATE_PRIVATE_ARGUMENT_VARIABLE(bool, ForceVolatile) = false;
+	SLATE_PRIVATE_ARGUMENT_VARIABLE(EWidgetClipping, Clipping) = EWidgetClipping::Inherit;
+	SLATE_PRIVATE_ARGUMENT_VARIABLE(EFlowDirectionPreference, FlowDirectionPreference) = EFlowDirectionPreference::Inherit;
+	SLATE_PRIVATE_ATTRIBUTE_VARIABLE(TOptional<FSlateRenderTransform>, RenderTransform);
+	SLATE_PRIVATE_ATTRIBUTE_VARIABLE(FVector2D, RenderTransformPivot) = FVector2D::ZeroVector;
+	SLATE_PRIVATE_ARGUMENT_VARIABLE(FName, Tag);
+	SLATE_PRIVATE_ARGUMENT_VARIABLE(TOptional<FAccessibleWidgetData>, AccessibleParams);
+	SLATE_PRIVATE_ATTRIBUTE_VARIABLE(FText, AccessibleText);
+	TArray<TSharedRef<ISlateMetaData>> MetaData;
+};
+
+
 /** Base class for named arguments. Provides settings necessary for all widgets. */
 template<typename WidgetType>
-struct TSlateBaseNamedArgs
+struct TSlateBaseNamedArgs : public FSlateBaseNamedArgs
 {
 	typedef typename WidgetType::FArguments WidgetArgsType;
-	
-	TSlateBaseNamedArgs()
-	: _ToolTipText()
-	, _ToolTip()
-	, _Cursor( TOptional<EMouseCursor::Type>() )
-	, _IsEnabled( true )
-	, _Visibility( EVisibility::Visible )
-	, _RenderOpacity(1.0f)
-	, _ForceVolatile( false )
-	, _Clipping( EWidgetClipping::Inherit )
-	, _FlowDirectionPreference( EFlowDirectionPreference::Inherit )
-	, _RenderTransform( )
-	, _RenderTransformPivot( FVector2D::ZeroVector )
-	, _AccessibleParams()
-	, _AccessibleText()
-	{
-	}
+
+	SLATE_PRIVATE_ATTRIBUTE_FUNCTION(FText, ToolTipText)
+	SLATE_PRIVATE_ATTRIBUTE_FUNCTION(TSharedPtr<IToolTip>, ToolTip)
+	SLATE_PRIVATE_ATTRIBUTE_FUNCTION(TOptional<EMouseCursor::Type>, Cursor)
+	SLATE_PRIVATE_ATTRIBUTE_FUNCTION(bool, IsEnabled)
+	SLATE_PRIVATE_ATTRIBUTE_FUNCTION(EVisibility, Visibility)
+	SLATE_PRIVATE_ARGUMENT_FUNCTION(float, RenderOpacity)
+	SLATE_PRIVATE_ARGUMENT_FUNCTION(bool, ForceVolatile)
+	SLATE_PRIVATE_ARGUMENT_FUNCTION(EWidgetClipping, Clipping)
+	SLATE_PRIVATE_ARGUMENT_FUNCTION(EFlowDirectionPreference, FlowDirectionPreference)
+	SLATE_PRIVATE_ATTRIBUTE_FUNCTION(TOptional<FSlateRenderTransform>, RenderTransform)
+	SLATE_PRIVATE_ATTRIBUTE_FUNCTION(FVector2D, RenderTransformPivot)
+	SLATE_PRIVATE_ARGUMENT_FUNCTION(FName, Tag)
+	SLATE_PRIVATE_ARGUMENT_FUNCTION(TOptional<FAccessibleWidgetData>, AccessibleParams)
+	SLATE_PRIVATE_ATTRIBUTE_FUNCTION(FText, AccessibleText)
 
 	/** Used by the named argument pattern as a safe way to 'return *this' for call-chaining purposes. */
 	WidgetArgsType& Me()
@@ -828,23 +846,6 @@ struct TSlateBaseNamedArgs
 		MetaData.Add(MakeShared<MetaDataType>(InArg0, InArg1));
 		return Me();
 	}
-
-	SLATE_ATTRIBUTE( FText, ToolTipText )
-	SLATE_ARGUMENT( TSharedPtr<IToolTip>, ToolTip )
-	SLATE_ATTRIBUTE( TOptional<EMouseCursor::Type>, Cursor )
-	SLATE_ATTRIBUTE( bool, IsEnabled )
-	SLATE_ATTRIBUTE( EVisibility, Visibility )
-	SLATE_ARGUMENT( float, RenderOpacity )
-	SLATE_ARGUMENT( bool, ForceVolatile )
-	SLATE_ARGUMENT( EWidgetClipping, Clipping )
-	SLATE_ARGUMENT( EFlowDirectionPreference, FlowDirectionPreference)
-	SLATE_ATTRIBUTE( TOptional<FSlateRenderTransform>, RenderTransform )
-	SLATE_ATTRIBUTE( FVector2D, RenderTransformPivot )
-	SLATE_ARGUMENT( FName, Tag )
-	SLATE_ARGUMENT(TOptional<FAccessibleWidgetData>, AccessibleParams)
-	SLATE_ATTRIBUTE(FText, AccessibleText)
-
-	TArray<TSharedRef<ISlateMetaData>> MetaData;
 };
 
 namespace RequiredArgs
@@ -860,7 +861,6 @@ namespace RequiredArgs
 		{
 			// YOUR WIDGET MUST IMPLEMENT void Construct(const FArguments& InArgs)
 			OnWidget->Construct(WithNamedArgs);
-			OnWidget->CacheVolatility();
 		}
 	};
 
@@ -877,7 +877,6 @@ namespace RequiredArgs
 		{
 			// YOUR WIDGET MUST IMPLEMENT void Construct(const FArguments& InArgs)
 			OnWidget->Construct(WithNamedArgs, Forward<Arg0Type>(Arg0));
-			OnWidget->CacheVolatility();
 		}
 
 		Arg0Type& Arg0;
@@ -897,7 +896,6 @@ namespace RequiredArgs
 		{
 			// YOUR WIDGET MUST IMPLEMENT Construct(const FArguments& InArgs)
 			OnWidget->Construct(WithNamedArgs, Forward<Arg0Type>(Arg0), Forward<Arg1Type>(Arg1));
-			OnWidget->CacheVolatility();
 		}
 
 		Arg0Type& Arg0;
@@ -919,7 +917,6 @@ namespace RequiredArgs
 		{
 			// YOUR WIDGET MUST IMPLEMENT Construct(const FArguments& InArgs)
 			OnWidget->Construct(WithNamedArgs, Forward<Arg0Type>(Arg0), Forward<Arg1Type>(Arg1), Forward<Arg2Type>(Arg2));
-			OnWidget->CacheVolatility();
 		}
 
 		Arg0Type& Arg0;
@@ -943,7 +940,6 @@ namespace RequiredArgs
 		{
 			// YOUR WIDGET MUST IMPLEMENT Construct(const FArguments& InArgs)
 			OnWidget->Construct(WithNamedArgs, Forward<Arg0Type>(Arg0), Forward<Arg1Type>(Arg1), Forward<Arg2Type>(Arg2), Forward<Arg3Type>(Arg3));
-			OnWidget->CacheVolatility();
 		}
 
 		Arg0Type& Arg0;
@@ -969,7 +965,6 @@ namespace RequiredArgs
 		{
 			// YOUR WIDGET MUST IMPLEMENT Construct(const FArguments& InArgs)
 			OnWidget->Construct(WithNamedArgs, Forward<Arg0Type>(Arg0), Forward<Arg1Type>(Arg1), Forward<Arg2Type>(Arg2), Forward<Arg3Type>(Arg3), Forward<Arg4Type>(Arg4));
-			OnWidget->CacheVolatility();
 		}
 
 		Arg0Type& Arg0;
@@ -1015,7 +1010,8 @@ namespace RequiredArgs
 	}
 }
 
-/** Normal widgets are allocated directly by the TDecl. */
+
+/** Normal widgets are allocated directly by the TSlateDecl. */
 template<typename WidgetType, bool IsDerived>
 struct TWidgetAllocator
 {
@@ -1039,7 +1035,6 @@ struct TWidgetAllocator< WidgetType, true >
 	}
 };
 
-class SUserWidget;
 
 /**
  * Utility class used during widget instantiation.
@@ -1051,9 +1046,9 @@ class SUserWidget;
  * @see SAssignNew
  */
 template<class WidgetType, typename RequiredArgsPayloadType>
-struct TDecl
+struct TSlateDecl
 {
-	TDecl( const ANSICHAR* InType, const ANSICHAR* InFile, int32 OnLine, RequiredArgsPayloadType&& InRequiredArgs )
+	TSlateDecl( const ANSICHAR* InType, const ANSICHAR* InFile, int32 OnLine, RequiredArgsPayloadType&& InRequiredArgs )
 		: _Widget( TWidgetAllocator<WidgetType, TIsDerivedFrom<WidgetType, SUserWidget>::IsDerived >::PrivateAllocateWidget() )
 		, _RequiredArgs(InRequiredArgs)
 	{
@@ -1065,7 +1060,7 @@ struct TDecl
 	 * @see SAssignNew
 	 */
 	template<class ExposeAsWidgetType>
-	TDecl& Expose( TSharedPtr<ExposeAsWidgetType>& OutVarToInit )
+	TSlateDecl& Expose( TSharedPtr<ExposeAsWidgetType>& OutVarToInit )
 	{
 		OutVarToInit = _Widget;
 		return *this;
@@ -1076,7 +1071,7 @@ struct TDecl
 	 * @see SAssignNew
 	 */
 	template<class ExposeAsWidgetType>
-	TDecl& Expose( TSharedRef<ExposeAsWidgetType>& OutVarToInit )
+	TSlateDecl& Expose( TSharedRef<ExposeAsWidgetType>& OutVarToInit )
 	{
 		OutVarToInit = _Widget;
 		return *this;
@@ -1087,7 +1082,7 @@ struct TDecl
 	 * @see SAssignNew
 	 */
 	template<class ExposeAsWidgetType>
-	TDecl& Expose( TWeakPtr<ExposeAsWidgetType>& OutVarToInit )
+	TSlateDecl& Expose( TWeakPtr<ExposeAsWidgetType>& OutVarToInit )
 	{
 		OutVarToInit = _Widget;
 		return *this;
@@ -1102,24 +1097,9 @@ struct TDecl
 	 */
 	TSharedRef<WidgetType> operator<<=( const typename WidgetType::FArguments& InArgs ) const
 	{
-		//@todo UMG: This should be removed in favor of all widgets calling their superclass construct.
-		_Widget->SWidgetConstruct(
-			InArgs._ToolTipText,
-			InArgs._ToolTip,
-			InArgs._Cursor,
-			InArgs._IsEnabled,
-			InArgs._Visibility,
-			InArgs._RenderOpacity,
-			InArgs._RenderTransform,
-			InArgs._RenderTransformPivot,
-			InArgs._Tag,
-			InArgs._ForceVolatile,
-			InArgs._Clipping,
-			InArgs._FlowDirectionPreference,
-			InArgs._AccessibleText.IsSet() ? FAccessibleWidgetData(InArgs._AccessibleText) : InArgs._AccessibleParams,
-			InArgs.MetaData );
-
+		_Widget->SWidgetConstruct(InArgs);
 		_RequiredArgs.CallConstruct(_Widget, InArgs);
+		_Widget->CacheVolatility();
 
 		return _Widget;
 	}
@@ -1130,7 +1110,7 @@ struct TDecl
 
 
 template<typename WidgetType, typename RequiredArgsPayloadType>
-TDecl<WidgetType, RequiredArgsPayloadType> MakeTDecl( const ANSICHAR* InType, const ANSICHAR* InFile, int32 OnLine, RequiredArgsPayloadType&& InRequiredArgs )
+TSlateDecl<WidgetType, RequiredArgsPayloadType> MakeTDecl( const ANSICHAR* InType, const ANSICHAR* InFile, int32 OnLine, RequiredArgsPayloadType&& InRequiredArgs )
 {
-	return TDecl<WidgetType, RequiredArgsPayloadType>(InType, InFile, OnLine, Forward<RequiredArgsPayloadType>(InRequiredArgs));
+	return TSlateDecl<WidgetType, RequiredArgsPayloadType>(InType, InFile, OnLine, Forward<RequiredArgsPayloadType>(InRequiredArgs));
 }
diff --git a/Engine/Source/Runtime/SlateCore/Public/Widgets/SWidget.h b/Engine/Source/Runtime/SlateCore/Public/Widgets/SWidget.h
--- a/Engine/Source/Runtime/SlateCore/Public/Widgets/SWidget.h
+++ b/Engine/Source/Runtime/SlateCore/Public/Widgets/SWidget.h
@@ -40,6 +40,7 @@ class FWeakWidgetPath;
 class FWidgetPath;
 class IToolTip;
 class SWidget;
+struct FSlateBaseNamedArgs;
 struct FSlateBrush;
 struct FSlatePaintElementLists;
 
@@ -55,8 +56,6 @@ DECLARE_DWORD_ACCUMULATOR_STAT_EXTERN(TEXT("Total Widgets"), STAT_SlateTotalWidg
 DECLARE_MEMORY_STAT_EXTERN(TEXT("SWidget Total Allocated Size"), STAT_SlateSWidgetAllocSize, STATGROUP_SlateMemory, SLATECORE_API);
 
 
-enum class EPopupMethod : uint8;
-
 namespace SharedPointerInternals
 {
 	template <typename ObjectType>
@@ -199,14 +198,12 @@ class SLATECORE_API SWidget
 	friend class FSlateWindowElementList;
 	friend class SWindow;
 	friend struct FSlateCachedElementList;
-#if WITH_SLATE_DEBUGGING
-	friend struct FInvalidatedWidgetDrawer;
-#endif
+	template<class WidgetType, typename RequiredArgsPayloadType>
+	friend struct TSlateDecl;
 public:
 
-	/**
-	 * Construct a SWidget based on initial parameters.
-	 */
+	/** Construct a SWidget based on initial parameters. */
+	 UE_DEPRECATED(4.27, "SWidget::Construct should not be called directly. Use SNew or SAssignNew to create a SWidget")
 	void Construct(
 		const TAttribute<FText>& InToolTipText,
 		const TSharedPtr<IToolTip>& InToolTip,
@@ -223,6 +220,7 @@ public:
 		const TOptional<FAccessibleWidgetData>& InAccessibleData,
 		const TArray<TSharedRef<ISlateMetaData>>& InMetaData);
 
+	UE_DEPRECATED(4.27, "SWidget::SWidgetConstruct should not be called directly. Use SNew or SAssignNew to create a SWidget")
 	void SWidgetConstruct(const TAttribute<FText>& InToolTipText,
 		const TSharedPtr<IToolTip>& InToolTip,
 		const TAttribute< TOptional<EMouseCursor::Type> >& InCursor,
@@ -1199,34 +1197,23 @@ public:
 	/** Gets the desired flow direction for the layout. */
 	EFlowDirectionPreference GetFlowDirectionPreference() const { return FlowDirectionPreference; }
 
-	/**
-	 * Set the tool tip that should appear when this widget is hovered.
-	 *
-	 * @param InToolTipText  the text that should appear in the tool tip
-	 */
+	/** Set the tool tip that should appear when this widget is hovered. */
 	void SetToolTipText(const TAttribute<FText>& ToolTipText);
 
+	/** Set the tool tip that should appear when this widget is hovered. */
 	void SetToolTipText( const FText& InToolTipText );
 
-	/**
-	 * Set the tool tip that should appear when this widget is hovered.
-	 *
-	 * @param InToolTip  the widget that should appear in the tool tip
-	 */
-	void SetToolTip( const TSharedPtr<IToolTip>& InToolTip );
+	/** Set the tool tip that should appear when this widget is hovered. */
+	void SetToolTip(const TAttribute<TSharedPtr<IToolTip>>& InToolTip);
 
-	/**
-	 * Set the cursor that should appear when this widget is hovered
-	 */
+	/** Set the cursor that should appear when this widget is hovered  */
 	void SetCursor( const TAttribute< TOptional<EMouseCursor::Type> >& InCursor );
 
-	/**
-	 * Used by Slate to set the runtime debug info about this widget.
-	 */
-	void SetDebugInfo( const ANSICHAR* InType, const ANSICHAR* InFile, int32 OnLine, size_t InAllocSize );
-	
 protected:
 
+	/** Used by Slate to set the runtime debug info about this widget. */
+	void SetDebugInfo( const ANSICHAR* InType, const ANSICHAR* InFile, int32 OnLine, size_t InAllocSize );
+
 	/** The cursor to show when the mouse is hovering over this widget. */
 	TOptional<EMouseCursor::Type> GetCursor() const;
 
@@ -1403,6 +1390,7 @@ public:
 	virtual bool Advanced_IsWindow() const { return false; }
 	virtual bool Advanced_IsInvalidationRoot() const { return false; }
 	virtual const FSlateInvalidationRoot* Advanced_AsInvalidationRoot() const { return nullptr; }
+
 protected:
 
 	/**
@@ -1414,6 +1402,9 @@ protected:
 	 */
 	SWidget();
 
+	/** Construct a SWidget based on initial parameters. */
+	void SWidgetConstruct(const FSlateBaseNamedArgs& Args);
+
 	/** 
 	 * Find the geometry of a descendant widget. This method assumes that WidgetsToFind are a descendants of this widget.
 	 * Note that not all widgets are guaranteed to be found; OutResult will contain null entries for missing widgets.
