Reduce the icon size and padding for menu entry rows

diff --git a/Engine/Source/Runtime/SlateCore/Private/Styling/StarshipCoreStyle.cpp b/Engine/Source/Runtime/SlateCore/Private/Styling/StarshipCoreStyle.cpp
--- a/Engine/Source/Runtime/SlateCore/Private/Styling/StarshipCoreStyle.cpp
+++ b/Engine/Source/Runtime/SlateCore/Private/Styling/StarshipCoreStyle.cpp
@@ -437,7 +437,7 @@ TSharedRef<ISlateStyle> FStarshipCoreStyle::Create()
 		Style->Set("ProgressBar", FProgressBarStyle()
 			.SetBackgroundImage(FSlateColorBrush(FStyleColors::Foldout))
 			.SetFillImage(IMAGE_BRUSH("Starship/CoreWidgets/ProgressBar/ProgressMarquee", FVector2D(20, 12), FStyleColors::Primary, ESlateBrushTileType::Horizontal))
-			.SetMarqueeImage(IMAGE_BRUSH("Starship/CoreWidgets/ProgressBar/ProgressMarquee", FVector2D(20, 12), FLinearColor::White, ESlateBrushTileType::Horizontal))
+			.SetMarqueeImage(IMAGE_BRUSH("Starship/CoreWidgets/ProgressBar/ProgressMarquee", FVector2D(20, 12), FStyleColors::Primary, ESlateBrushTileType::Horizontal))
 			.SetEnableFillAnimation(true)
 		);
 	}
@@ -1805,14 +1805,14 @@ void FStarshipCoreStyle::SetupMultiboxStyles(TSharedRef<FStyle>& Style)
 		Style->Set("Menu.Icon", new IMAGE_BRUSH("Icons/icon_tab_toolbar_16px", Icon16x16));
 		Style->Set("Menu.Expand", new IMAGE_BRUSH("Icons/toolbar_expand_16x", Icon16x16));
 		Style->Set("Menu.SubMenuIndicator", new IMAGE_BRUSH_SVG("Starship/Common/chevron-right", Icon16x16, FStyleColors::Foreground));
-		Style->Set("Menu.SToolBarComboButtonBlock.Padding", FMargin(4.0f));
-		Style->Set("Menu.SToolBarButtonBlock.Padding", FMargin(4.0f));
-		Style->Set("Menu.SToolBarCheckComboButtonBlock.Padding", FMargin(4.0f));
+		Style->Set("Menu.SToolBarComboButtonBlock.Padding", FMargin(2.0f));
+		Style->Set("Menu.SToolBarButtonBlock.Padding", FMargin(2.0f));
+		Style->Set("Menu.SToolBarCheckComboButtonBlock.Padding", FMargin(2.0f));
 		Style->Set("Menu.SToolBarButtonBlock.CheckBox.Padding", FMargin(0.0f));
 		Style->Set("Menu.SToolBarComboButtonBlock.ComboButton.Color", DefaultForeground);
-		Style->Set("Menu.MenuIconSize", 16.f);
+		Style->Set("Menu.MenuIconSize", 14.f);
 
-		const FMargin MenuBlockPadding(10.0f, 3.0f, 5.0f, 2.0f);
+		const FMargin MenuBlockPadding(10.0f, 1.0f, 5.0f, 1.0f);
 		Style->Set("Menu.Block.IndentedPadding", MenuBlockPadding + FMargin(18.0f, 0, 0, 0));
 		Style->Set("Menu.Block.Padding", MenuBlockPadding);
 
