Menu Styling

diff --git a/Engine/Source/Runtime/Slate/Private/Framework/MultiBox/MultiBox.cpp b/Engine/Source/Runtime/Slate/Private/Framework/MultiBox/MultiBox.cpp
--- a/Engine/Source/Runtime/Slate/Private/Framework/MultiBox/MultiBox.cpp
+++ b/Engine/Source/Runtime/Slate/Private/Framework/MultiBox/MultiBox.cpp
@@ -1021,7 +1021,7 @@ void SMultiBoxWidget::BuildMultiBoxWidget()
 	{
 	case EMultiBoxType::Menu:
 		{
-			BorderPadding = FMargin(0, 0, 0, 7);
+			BorderPadding = FMargin(0, 8, 0, 8);
 		}
 		break;
 	default:
diff --git a/Engine/Source/Runtime/Slate/Private/Framework/MultiBox/MultiBoxBuilder.cpp b/Engine/Source/Runtime/Slate/Private/Framework/MultiBox/MultiBoxBuilder.cpp
--- a/Engine/Source/Runtime/Slate/Private/Framework/MultiBox/MultiBoxBuilder.cpp
+++ b/Engine/Source/Runtime/Slate/Private/Framework/MultiBox/MultiBoxBuilder.cpp
@@ -331,13 +331,6 @@ void FMenuBuilder::ApplySectionBeginning()
 {
 	if (bSectionNeedsToBeApplied)
 	{
-		const int32 NumBlocksBeforeSeparator = MultiBox->bHasSearchWidget ? 2 : 1;
-
-		// Do not count search block, which starts as invisible
-		if( MultiBox->GetBlocks().Num() >= NumBlocksBeforeSeparator || FMultiBoxSettings::DisplayMultiboxHooks.Get() )
-		{
-			MultiBox->AddMultiBlock( MakeShareable( new FMenuSeparatorBlock(CurrentSectionExtensionHook, /* bInIsPartOfHeading=*/ true) ) );
-		}
 		if (!CurrentSectionHeadingText.IsEmpty())
 		{
 			MultiBox->AddMultiBlock( MakeShareable( new FHeadingBlock(CurrentSectionExtensionHook, CurrentSectionHeadingText) ) );
diff --git a/Engine/Source/Runtime/Slate/Private/Framework/MultiBox/SHeadingBlock.cpp b/Engine/Source/Runtime/Slate/Private/Framework/MultiBox/SHeadingBlock.cpp
--- a/Engine/Source/Runtime/Slate/Private/Framework/MultiBox/SHeadingBlock.cpp
+++ b/Engine/Source/Runtime/Slate/Private/Framework/MultiBox/SHeadingBlock.cpp
@@ -1,6 +1,8 @@
 // Copyright Epic Games, Inc. All Rights Reserved.
 
 #include "Framework/MultiBox/SHeadingBlock.h"
+#include "Widgets/SBoxPanel.h"
+#include "Widgets/Layout/SSeparator.h"
 
 
 /**
@@ -51,8 +53,25 @@ void SHeadingBlock::BuildMultiBlockWidget(const ISlateStyle* StyleSet, const FNa
 	ChildSlot
 	.Padding(StyleSet->GetMargin(StyleName, ".Heading.Padding"))
 	[
-		SNew( STextBlock )
+		SNew(SHorizontalBox)
+		+ SHorizontalBox::Slot()
+		.VAlign(VAlign_Center)
+		.AutoWidth()
+		[
+			SNew( STextBlock )
 			.Text( HeadingBlock->HeadingText.ToUpper() )
 			.TextStyle( StyleSet, ISlateStyle::Join( StyleName, ".Heading" ) )
+		]
+
+		+ SHorizontalBox::Slot()
+		.Padding(FMargin(14.f, 0.f, 0.f, 0.f))
+		.VAlign(VAlign_Center)
+		.HAlign(HAlign_Fill)
+		[
+			SNew(SSeparator)
+			.Orientation(Orient_Horizontal)
+			.Thickness(1.0f)
+			.SeparatorImage(StyleSet->GetBrush(StyleName, ".Separator") )
+		]
 	];
 }
diff --git a/Engine/Source/Runtime/Slate/Private/Framework/MultiBox/SMenuSeparatorBlock.cpp b/Engine/Source/Runtime/Slate/Private/Framework/MultiBox/SMenuSeparatorBlock.cpp
--- a/Engine/Source/Runtime/Slate/Private/Framework/MultiBox/SMenuSeparatorBlock.cpp
+++ b/Engine/Source/Runtime/Slate/Private/Framework/MultiBox/SMenuSeparatorBlock.cpp
@@ -46,16 +46,15 @@ void SMenuSeparatorBlock::BuildMultiBlockWidget(const ISlateStyle* StyleSet, con
 	ChildSlot
 	[
 		SNew( SVerticalBox )
-			+SVerticalBox::Slot()
-				.AutoHeight()
-
-				// Add some empty space before the line, and a tiny bit after it
-				.Padding( 0.0f, 4.0f, 0.0f, 2.0f )
-				[
-					SNew( SSeparator )
-					.SeparatorImage(StyleSet->GetBrush(StyleName, ".Separator"))
-					.Thickness(2.0f)
-				]
+
+		+SVerticalBox::Slot()
+		.AutoHeight()
+		.Padding(StyleSet->GetMargin(StyleName, ".Separator.Padding"))
+		[
+			SNew( SSeparator )
+			.SeparatorImage(StyleSet->GetBrush(StyleName, ".Separator"))
+			.Thickness(1.0f)
+		]
 	];
 
 	// Add this widget to the search list of the multibox and hide it
diff --git a/Engine/Source/Runtime/SlateCore/Private/Styling/StarshipCoreStyle.cpp b/Engine/Source/Runtime/SlateCore/Private/Styling/StarshipCoreStyle.cpp
--- a/Engine/Source/Runtime/SlateCore/Private/Styling/StarshipCoreStyle.cpp
+++ b/Engine/Source/Runtime/SlateCore/Private/Styling/StarshipCoreStyle.cpp
@@ -1669,6 +1669,7 @@ void FStarshipCoreStyle::SetupMultiboxStyles(TSharedRef<FStyle>& Style)
 {
 	const FEditableTextBoxStyle& NormalEditableTextBoxStyle = Style->GetWidgetStyle<FEditableTextBoxStyle>("NormalEditableTextBox");
 	const FTextBlockStyle& NormalText = Style->GetWidgetStyle<FTextBlockStyle>("NormalText");
+	const FTextBlockStyle& SmallButtonText = Style->GetWidgetStyle<FTextBlockStyle>("SmallButtonText");
 
 	const FSlateColor SelectionColor = Style->GetSlateColor("SelectionColor");
 	const FSlateColor SelectionColor_Pressed = Style->GetSlateColor("SelectionColor_Pressed");
@@ -1812,12 +1813,12 @@ void FStarshipCoreStyle::SetupMultiboxStyles(TSharedRef<FStyle>& Style)
 		Style->Set("Menu.SToolBarComboButtonBlock.ComboButton.Color", DefaultForeground);
 		Style->Set("Menu.MenuIconSize", 14.f);
 
-		const FMargin MenuBlockPadding(10.0f, 1.0f, 5.0f, 1.0f);
+		const FMargin MenuBlockPadding(12.0f, 1.0f, 5.0f, 1.0f);
 		Style->Set("Menu.Block.IndentedPadding", MenuBlockPadding + FMargin(18.0f, 0, 0, 0));
 		Style->Set("Menu.Block.Padding", MenuBlockPadding);
 
-		Style->Set("Menu.Separator", new FSlateColorBrush(FStyleColors::Header));
-		Style->Set("Menu.Separator.Padding", FMargin(0.5f));
+		Style->Set("Menu.Separator", new FSlateColorBrush(FStyleColors::White25));
+		Style->Set("Menu.Separator.Padding", FMargin(0.0f, 6.f, 0.0f, 6.f));
 
 		Style->Set("Menu.Label", NormalText);
 
@@ -1825,10 +1826,9 @@ void FStarshipCoreStyle::SetupMultiboxStyles(TSharedRef<FStyle>& Style)
 		Style->Set("Menu.Keybinding", FTextBlockStyle(NormalText).SetFont(FStyleFonts::Get().Small));
 
 		Style->Set("Menu.Heading",
-			FTextBlockStyle(NormalText)
-			.SetFont(FStyleFonts::Get().SmallBold)
-			.SetColorAndOpacity(FStyleColors::Foreground));
-		Style->Set("Menu.Heading.Padding", MenuBlockPadding + FMargin(0, 10, 0, 0));
+			FTextBlockStyle(SmallButtonText)
+			.SetColorAndOpacity(FStyleColors::White25));
+		Style->Set("Menu.Heading.Padding", FMargin(12.0f, 6.f, 0.f, 6.f));
 
 		/* Set images for various SCheckBox states associated with menu check box items... */
 		FLinearColor Transparent20 = FLinearColor(1.0, 1.0, 1.0, 0.2);
