SWidget CacheDesiredSize stat is now only enabled with very verbose slate stats as the stat creates a huge number of events that breaks profilers

diff --git a/Engine/Source/Runtime/SlateCore/Private/Widgets/SWidget.cpp b/Engine/Source/Runtime/SlateCore/Private/Widgets/SWidget.cpp
--- a/Engine/Source/Runtime/SlateCore/Private/Widgets/SWidget.cpp
+++ b/Engine/Source/Runtime/SlateCore/Private/Widgets/SWidget.cpp
@@ -873,7 +873,7 @@ void SWidget::UpdateFastPathVolatility(bool bParentVolatile)
 
 void SWidget::CacheDesiredSize(float InLayoutScaleMultiplier)
 {
-#if SLATE_VERBOSE_NAMED_EVENTS
+#if WITH_VERY_VERBOSE_SLATE_STATS
 	SCOPED_NAMED_EVENT(SWidget_CacheDesiredSize, FColor::Red);
 #endif
 
