Updated Menu Header Font

diff --git a/Engine/Source/Runtime/SlateCore/Private/Styling/StarshipCoreStyle.cpp b/Engine/Source/Runtime/SlateCore/Private/Styling/StarshipCoreStyle.cpp
--- a/Engine/Source/Runtime/SlateCore/Private/Styling/StarshipCoreStyle.cpp
+++ b/Engine/Source/Runtime/SlateCore/Private/Styling/StarshipCoreStyle.cpp
@@ -1825,8 +1825,13 @@ void FStarshipCoreStyle::SetupMultiboxStyles(TSharedRef<FStyle>& Style)
 		Style->Set("Menu.EditableText", FEditableTextBoxStyle(NormalEditableTextBoxStyle).SetFont(FStyleFonts::Get().Normal));
 		Style->Set("Menu.Keybinding", FTextBlockStyle(NormalText).SetFont(FStyleFonts::Get().Small));
 
+
+		FSlateFontInfo XSFont(FONT(7, "Bold"));
+		XSFont.LetterSpacing =  250;
+
 		Style->Set("Menu.Heading",
 			FTextBlockStyle(SmallButtonText)
+			.SetFont(XSFont)
 			.SetColorAndOpacity(FStyleColors::White25));
 		Style->Set("Menu.Heading.Padding", FMargin(12.0f, 6.f, 0.f, 6.f));
 
