Menu - Slight separator tweak

diff --git a/Engine/Source/Runtime/SlateCore/Private/Styling/StarshipCoreStyle.cpp b/Engine/Source/Runtime/SlateCore/Private/Styling/StarshipCoreStyle.cpp
--- a/Engine/Source/Runtime/SlateCore/Private/Styling/StarshipCoreStyle.cpp
+++ b/Engine/Source/Runtime/SlateCore/Private/Styling/StarshipCoreStyle.cpp
@@ -1818,7 +1818,7 @@ void FStarshipCoreStyle::SetupMultiboxStyles(TSharedRef<FStyle>& Style)
 		Style->Set("Menu.Block.Padding", MenuBlockPadding);
 
 		Style->Set("Menu.Separator", new FSlateColorBrush(FStyleColors::White25));
-		Style->Set("Menu.Separator.Padding", FMargin(0.0f, 6.f, 0.0f, 6.f));
+		Style->Set("Menu.Separator.Padding", FMargin(12.0f, 6.f, 12.0f, 6.f));
 
 		Style->Set("Menu.Label", NormalText);
 
@@ -1833,7 +1833,7 @@ void FStarshipCoreStyle::SetupMultiboxStyles(TSharedRef<FStyle>& Style)
 			FTextBlockStyle(SmallButtonText)
 			.SetFont(XSFont)
 			.SetColorAndOpacity(FStyleColors::White25));
-		Style->Set("Menu.Heading.Padding", FMargin(12.0f, 6.f, 0.f, 6.f));
+		Style->Set("Menu.Heading.Padding", FMargin(12.0f, 6.f, 12.f, 6.f));
 
 		/* Set images for various SCheckBox states associated with menu check box items... */
 		FLinearColor Transparent20 = FLinearColor(1.0, 1.0, 1.0, 0.2);
