Styling on read-only (note: not IsEditable) SEditableTextBoxes updated.

diff --git a/Engine/Source/Runtime/SlateCore/Private/Styling/SlateTypes.cpp b/Engine/Source/Runtime/SlateCore/Private/Styling/SlateTypes.cpp
--- a/Engine/Source/Runtime/SlateCore/Private/Styling/SlateTypes.cpp
+++ b/Engine/Source/Runtime/SlateCore/Private/Styling/SlateTypes.cpp
@@ -6,14 +6,6 @@
 #include "Widgets/InvalidateWidgetReason.h"
 #include "Widgets/SWidget.h"
 
-namespace SlateTypeDefs
-{
-	// Colors
-	const static FLinearColor DefaultForeground = FLinearColor(0.72f, 0.72f, 0.72f, 1.f);
-	const static FLinearColor InvertedForeground = FLinearColor(0,0,0);
-}
-
-
 FCheckBoxStyle::FCheckBoxStyle()
 : CheckBoxType(ESlateCheckBoxType::CheckBox)
 , UncheckedImage()
@@ -256,9 +248,9 @@ FEditableTextBoxStyle::FEditableTextBoxStyle()
 	, BackgroundImageReadOnly()
 	, Padding(FMargin(4.0f, 2.0f))
 	, Font(FStyleDefaults::GetFontInfo(9))
-	, ForegroundColor(SlateTypeDefs::InvertedForeground)
+	, ForegroundColor(FSlateColor::UseForeground())
 	, BackgroundColor(FLinearColor::White)
-	, ReadOnlyForegroundColor(SlateTypeDefs::DefaultForeground)
+	, ReadOnlyForegroundColor(FSlateColor::UseForeground())
 {
 }
 
diff --git a/Engine/Source/Runtime/SlateCore/Private/Styling/StarshipCoreStyle.cpp b/Engine/Source/Runtime/SlateCore/Private/Styling/StarshipCoreStyle.cpp
--- a/Engine/Source/Runtime/SlateCore/Private/Styling/StarshipCoreStyle.cpp
+++ b/Engine/Source/Runtime/SlateCore/Private/Styling/StarshipCoreStyle.cpp
@@ -969,7 +969,7 @@ void FStarshipCoreStyle::SetupTextStyles(TSharedRef<FStyle>& Style)
 		.SetBackgroundImageNormal(FSlateRoundedBoxBrush(FStyleColors::Input, InputFocusRadius, FStyleColors::InputOutline, InputFocusThickness))
 		.SetBackgroundImageHovered(FSlateRoundedBoxBrush(FStyleColors::Input, InputFocusRadius, FStyleColors::Hover, InputFocusThickness))
 		.SetBackgroundImageFocused(FSlateRoundedBoxBrush(FStyleColors::Input, InputFocusRadius, FStyleColors::Primary, InputFocusThickness))
-		.SetBackgroundImageReadOnly(FSlateRoundedBoxBrush(FStyleColors::Input, InputFocusRadius))
+		.SetBackgroundImageReadOnly(FSlateRoundedBoxBrush(FStyleColors::Header, InputFocusRadius, FStyleColors::InputOutline, InputFocusThickness))
 		.SetFont(StyleFonts.Normal)
 		.SetPadding(FMargin(12.f, 4.0f, 12.f, 5.0f)) // The padding should be 4 top, 5 bottom
 		.SetForegroundColor(FStyleColors::White)
@@ -987,7 +987,7 @@ void FStarshipCoreStyle::SetupTextStyles(TSharedRef<FStyle>& Style)
 			.SetBackgroundImageNormal(FSlateRoundedBoxBrush(FStyleColors::Input, FStyleColors::Secondary, InputFocusThickness))
 			.SetBackgroundImageHovered(FSlateRoundedBoxBrush(FStyleColors::Input, FStyleColors::Hover, InputFocusThickness))
 			.SetBackgroundImageFocused(FSlateRoundedBoxBrush(FStyleColors::Input, FStyleColors::Primary, InputFocusThickness))
-			.SetBackgroundImageReadOnly(FSlateRoundedBoxBrush(FStyleColors::Input)
+			.SetBackgroundImageReadOnly(FSlateRoundedBoxBrush(FStyleColors::Header, FStyleColors::InputOutline, InputFocusThickness)
 		);
 
 		Style->Set("SearchBox", FSearchBoxStyle()
@@ -1282,8 +1282,7 @@ void FStarshipCoreStyle::SetupCheckboxStyles(TSharedRef<FStyle>& Style)
 		.SetPressedForegroundColor(FStyleColors::ForegroundHover)
 		.SetCheckedForegroundColor(FStyleColors::Primary)
 		.SetCheckedHoveredForegroundColor(FStyleColors::PrimaryHover)
-		.SetPadding(DefaultMargins)
-		;
+		.SetPadding(DefaultMargins);
 
 	Style->Set("ToggleButtonCheckbox", ToggleButtonStyle);
 
@@ -1301,8 +1300,7 @@ void FStarshipCoreStyle::SetupCheckboxStyles(TSharedRef<FStyle>& Style)
 		.SetPressedForegroundColor(FStyleColors::ForegroundHover)
 		.SetCheckedForegroundColor(FStyleColors::Primary)
 		.SetCheckedHoveredForegroundColor(FStyleColors::PrimaryHover)
-		.SetPadding(DefaultMargins)
-		;
+		.SetPadding(DefaultMargins);
 
 	Style->Set("ToggleButtonCheckboxAlt", ToggleButtonAltStyle);
 
