Slate - Rounded brushes can now have a radius specified per corner instead of just uniformly

diff --git a/Engine/Shaders/Private/SlateElementPixelShader.usf b/Engine/Shaders/Private/SlateElementPixelShader.usf
--- a/Engine/Shaders/Private/SlateElementPixelShader.usf
+++ b/Engine/Shaders/Private/SlateElementPixelShader.usf
@@ -250,10 +250,20 @@ float GetRoundedBoxDistance(float2 pos, float2 center, float radius, float inset
 half4 GetRoundedBoxElementColor( VertexToPixelInterpolants InVertex )
 {
 	const float2 size = ShaderParams.zw;
-	float2 pos = size * InVertex.TextureCoordinates[0].xy;
+	float2 UV = InVertex.TextureCoordinates[0].xy;
+	float2 pos = size * UV;
 	float2 center = size / 2.0;
 
-	float radius = ShaderParams.x;	
+	//X = Top Left, Y = Top Right, Z = Bottom Right, W = Bottom Left */
+	float4 cornerRadii = ShaderParams2;
+
+	// figure out which radius to use based on which quadrant we're in
+	float2 quadrant = step(UV, float2(.5,.5));
+
+	float left = lerp(cornerRadii.y, cornerRadii.x, quadrant.x);
+	float right = lerp(cornerRadii.z, cornerRadii.w, quadrant.x);
+	float radius = lerp(right, left, quadrant.y);
+
 	float thickness = ShaderParams.y; 
 
 	// Compute the distances internal and external to the border outline
@@ -266,7 +276,7 @@ half4 GetRoundedBoxElementColor( VertexToPixelInterpolants InVertex )
 	float fi = smoothstep(spread, -spread, din);
 
 	half4 fill = GetColor(InVertex, GetUV(InVertex, 0) * GetUV(InVertex, 1)); //InVertex.Color;
-	half4 border = ShaderParams2;
+	half4 border = InVertex.SecondaryColor;
 
 	// alpha blend the external color 
 	half4 OutColor = lerp(border, fill, half(thickness > radius));
diff --git a/Engine/Shaders/Private/SlateShaderCommon.ush b/Engine/Shaders/Private/SlateShaderCommon.ush
--- a/Engine/Shaders/Private/SlateShaderCommon.ush
+++ b/Engine/Shaders/Private/SlateShaderCommon.ush
@@ -26,11 +26,14 @@
 #define ST_ColorFont		3
 #define ST_LineSegment		4
 #define ST_Custom			5
+#define ST_PostProcess		6
+#define ST_RoundedBox		7
 
 struct VertexToPixelInterpolants
 {
 	float4 Position : SV_POSITION;
 	half4 Color : COLOR0;
+	half4 SecondaryColor : COLOR1;
 	float4 OriginalPosition : ORIGINAL_POSITION;
 	float2 MaterialTexCoords : TEXCOORD0;
 	float4 TextureCoordinates[NUM_SLATE_TEXCOORDS] : TEXCOORD1;
diff --git a/Engine/Shaders/Private/SlateVertexShader.usf b/Engine/Shaders/Private/SlateVertexShader.usf
--- a/Engine/Shaders/Private/SlateVertexShader.usf
+++ b/Engine/Shaders/Private/SlateVertexShader.usf
@@ -13,10 +13,11 @@ VertexToPixelInterpolants Main(
 	in float2 InMaterialTextureCoordinates : ATTRIBUTE1,
 	in float2 InPosition : ATTRIBUTE2,
 	in half4 InColor : ATTRIBUTE3,
+	in half4 InSecondaryColor : ATTRIBUTE4,
 #if USE_SLATE_INSTANCING
-	in float4 InstanceParam : ATTRIBUTE4
+	in float4 InstanceParam : ATTRIBUTE5
 #else
-    in uint2 PixelSize : ATTRIBUTE4
+    in uint2 PixelSize : ATTRIBUTE5
 #endif
 	)
 {
@@ -29,12 +30,12 @@ VertexToPixelInterpolants Main(
 
 #if SOURCE_IN_LINEAR_SPACE
 	InColor.rgb = sRGBToLinear(InColor.rgb);
+	InSecondaryColor.rgb = sRGBToLinear(InSecondaryColor.rgb);
 #endif
 
-	half4 FinalVertexColor = InColor FCOLOR_COMPONENT_SWIZZLE;
-
 	VOut.MaterialTexCoords = InMaterialTextureCoordinates;
-	VOut.Color = FinalVertexColor;
+	VOut.Color = InColor FCOLOR_COMPONENT_SWIZZLE;;
+	VOut.SecondaryColor = InSecondaryColor FCOLOR_COMPONENT_SWIZZLE;
 	VOut.TextureCoordinates[0] = InTextureCoordinates;
 
 #if USE_MATERIALS
diff --git a/Engine/Shaders/StandaloneRenderer/D3D/SlateDefaultVertexShader.hlsl b/Engine/Shaders/StandaloneRenderer/D3D/SlateDefaultVertexShader.hlsl
--- a/Engine/Shaders/StandaloneRenderer/D3D/SlateDefaultVertexShader.hlsl
+++ b/Engine/Shaders/StandaloneRenderer/D3D/SlateDefaultVertexShader.hlsl
@@ -11,6 +11,7 @@ struct VertexOut
 {
 	float4 Position : SV_POSITION;
 	float4 Color : COLOR0;
+	float4 SecondaryColor : COLOR1;
 	float4 TextureCoordinates : TEXCOORD0;
 };
 
@@ -18,7 +19,8 @@ VertexOut Main(
 	in float2 InPosition : POSITION,
 	in float4 InTextureCoordinates : TEXCOORD0,
 	in float2 MaterialTexCoords : TEXCOORD1,
-	in float4 InColor : COLOR0
+	in float4 InColor : COLOR0,
+	in float4 InSecondaryColor : COLOR1
 	)
 {
 	VertexOut Out;
@@ -27,8 +29,11 @@ VertexOut Main(
 
 	Out.TextureCoordinates = InTextureCoordinates;
 	
-	Out.Color.rgb = sRGBToLinear(InColor.rgb);
-	Out.Color.a = InColor.a;
+	InColor.rgb = sRGBToLinear(InColor.rgb);
+	InSecondaryColor.rgb = sRGBToLinear(InSecondaryColor.rgb);
+
+	Out.Color = InColor;
+	Out.SecondaryColor = InSecondaryColor;
 
 	return Out;
 }
diff --git a/Engine/Shaders/StandaloneRenderer/D3D/SlateElementPixelShader.hlsl b/Engine/Shaders/StandaloneRenderer/D3D/SlateElementPixelShader.hlsl
--- a/Engine/Shaders/StandaloneRenderer/D3D/SlateElementPixelShader.hlsl
+++ b/Engine/Shaders/StandaloneRenderer/D3D/SlateElementPixelShader.hlsl
@@ -35,6 +35,7 @@ struct VertexOut
 {
 	float4 Position : SV_POSITION;
 	float4 Color : COLOR0;
+	float4 SecondaryColor : COLOR1;
 	float4 TextureCoordinates : TEXCOORD0;
 };
 
@@ -149,7 +150,16 @@ float4 GetRoundedBoxElementColor( VertexOut InVertex )
 	float2 pos = size * InVertex.TextureCoordinates.xy;
 	float2 center = size / 2.0;
 
-	float radius = ShaderParams.x;	
+	//X = Top Left, Y = Top Right, Z = Bottom Right, W = Bottom Left */
+	float4 cornerRadii = ShaderParams2;
+
+	// figure out which radius to use based on which quadrant we're in
+	float2 quadrant = step(InVertex.TextureCoordinates.xy, float2(.5,.5));
+
+	float left = lerp(cornerRadii.y, cornerRadii.x, quadrant.x);
+	float right = lerp(cornerRadii.z, cornerRadii.w, quadrant.x);
+	float radius = lerp(right, left, quadrant.y);
+
 	float thickness = ShaderParams.y; 
 
 	// Compute the distances internal and external to the border outline
@@ -163,7 +173,7 @@ float4 GetRoundedBoxElementColor( VertexOut InVertex )
 
 	// alpha blend the external color 
 	float4 fill = GetColor(InVertex, InVertex.TextureCoordinates.xy * InVertex.TextureCoordinates.zw);
-	float4 border = ShaderParams2;
+	float4 border = InVertex.SecondaryColor;
 	float4 OutColor = lerp(border, fill, float(thickness > radius));
 	OutColor.a = 0.0;
 
diff --git a/Engine/Shaders/StandaloneRenderer/OpenGL/SlateElementPixelShader.glsl b/Engine/Shaders/StandaloneRenderer/OpenGL/SlateElementPixelShader.glsl
--- a/Engine/Shaders/StandaloneRenderer/OpenGL/SlateElementPixelShader.glsl
+++ b/Engine/Shaders/StandaloneRenderer/OpenGL/SlateElementPixelShader.glsl
@@ -37,6 +37,7 @@ uniform sampler2D ElementTexture;
 varying vec4 Position;
 varying vec4 TexCoords;
 varying vec4 Color;
+varying vec4 SecondaryColor;
 
 vec3 maxWithScalar(float test, vec3 values)
 {
@@ -234,7 +235,16 @@ vec4 GetRoundedBoxElementColor()
 	vec2 pos = size * TexCoords.xy;
 	vec2 center = size / 2.0;
 
-	float radius = ShaderParams.x;	
+	//X = Top Left, Y = Top Right, Z = Bottom Right, W = Bottom Left */
+	vec4 cornerRadii = ShaderParams2;
+
+	// figure out which radius to use based on which quadrant we're in
+	vec2 quadrant = step(TexCoords.xy, vec2(.5,.5));
+
+	float left = mix(cornerRadii.y, cornerRadii.x, quadrant.x);
+	float right = mix(cornerRadii.z, cornerRadii.w, quadrant.x);
+	float radius = mix(right, left, quadrant.y);
+
 	float thickness = ShaderParams.y; 
 
 	// Compute the distances internal and external to the border outline
@@ -248,7 +258,7 @@ vec4 GetRoundedBoxElementColor()
 
 	// alpha blend the external color 
 	vec4 fill = GetDefaultElementColor();
-	vec4 border = ShaderParams2.bgra;
+	vec4 border = SecondaryColor;
 	vec4 OutColor = mix(border, fill, float(thickness > radius));
 	OutColor.a = 0.0;
 
diff --git a/Engine/Shaders/StandaloneRenderer/OpenGL/SlateVertexShader.glsl b/Engine/Shaders/StandaloneRenderer/OpenGL/SlateVertexShader.glsl
--- a/Engine/Shaders/StandaloneRenderer/OpenGL/SlateVertexShader.glsl
+++ b/Engine/Shaders/StandaloneRenderer/OpenGL/SlateVertexShader.glsl
@@ -8,11 +8,13 @@ uniform mat4 ViewProjectionMatrix;
 attribute vec4 InTexCoords;
 attribute vec2 InPosition;
 attribute vec4 InColor;
+attribute vec4 InSecondaryColor;
 
 // Between vertex and pixel shader
 varying vec4 Position;
 varying vec4 TexCoords;
 varying vec4 Color;
+varying vec4 SecondaryColor;
 
 vec3 powScalar(vec3 values, float power)
 {
@@ -39,6 +41,9 @@ void main()
 	Color.rgb = sRGBToLinear(InColor.rgb);
 	Color.a = InColor.a;
 
+	SecondaryColor.rgb = sRGBToLinear(InSecondaryColor.rgb);
+	SecondaryColor.a = InSecondaryColor.a;
+
     Position = vec4( InPosition, 0, 1 );
     
 	gl_Position = ViewProjectionMatrix * vec4( InPosition, 0, 1 );
diff --git a/Engine/Source/Developer/StandaloneRenderer/Private/OpenGL/SlateOpenGLRenderingPolicy.cpp b/Engine/Source/Developer/StandaloneRenderer/Private/OpenGL/SlateOpenGLRenderingPolicy.cpp
--- a/Engine/Source/Developer/StandaloneRenderer/Private/OpenGL/SlateOpenGLRenderingPolicy.cpp
+++ b/Engine/Source/Developer/StandaloneRenderer/Private/OpenGL/SlateOpenGLRenderingPolicy.cpp
@@ -268,18 +268,20 @@ void FSlateOpenGLRenderingPolicy::DrawElements( const FMatrix& ViewProjectionMat
 
 		// Set up offsets into the vertex buffer for each vertex
 		glEnableVertexAttribArray(0);
-		Offset = STRUCT_OFFSET( FSlateVertex, TexCoords );
-		glVertexAttribPointer( 0, 4, GL_FLOAT, GL_FALSE, Stride, BUFFER_OFFSET(Stride*BaseVertexIndex+Offset) );
+		Offset = STRUCT_OFFSET(FSlateVertex, TexCoords);
+		glVertexAttribPointer(0, 4, GL_FLOAT, GL_FALSE, Stride, BUFFER_OFFSET(Stride*BaseVertexIndex+Offset));
 
 		glEnableVertexAttribArray(1);
-		Offset = STRUCT_OFFSET( FSlateVertex, Position );
-		glVertexAttribPointer( 1, 2, GL_FLOAT, GL_FALSE, Stride, BUFFER_OFFSET(Stride*BaseVertexIndex+Offset) );
-
-		glEnableVertexAttribArray(4);
-		Offset = STRUCT_OFFSET( FSlateVertex, Color );
-		glVertexAttribPointer( 4, 4, GL_UNSIGNED_BYTE, GL_TRUE, Stride, BUFFER_OFFSET(Stride*BaseVertexIndex+Offset) );
+		Offset = STRUCT_OFFSET(FSlateVertex, Position);
+		glVertexAttribPointer(1, 2, GL_FLOAT, GL_FALSE, Stride, BUFFER_OFFSET(Stride*BaseVertexIndex+Offset));
 
+		glEnableVertexAttribArray(2);
+		Offset = STRUCT_OFFSET(FSlateVertex, Color);
+		glVertexAttribPointer(2, 4, GL_UNSIGNED_BYTE, GL_TRUE, Stride, BUFFER_OFFSET(Stride*BaseVertexIndex+Offset));
 
+		glEnableVertexAttribArray(3);
+		Offset = STRUCT_OFFSET(FSlateVertex, SecondaryColor);
+		glVertexAttribPointer(3, 4, GL_UNSIGNED_BYTE, GL_TRUE, Stride, BUFFER_OFFSET(Stride*BaseVertexIndex+Offset));
 		
 		// Bind the index buffer so glDrawRangeElements knows which one to use
 		IndexBuffer.Bind();
diff --git a/Engine/Source/Developer/StandaloneRenderer/Private/OpenGL/SlateOpenGLShaders.cpp b/Engine/Source/Developer/StandaloneRenderer/Private/OpenGL/SlateOpenGLShaders.cpp
--- a/Engine/Source/Developer/StandaloneRenderer/Private/OpenGL/SlateOpenGLShaders.cpp
+++ b/Engine/Source/Developer/StandaloneRenderer/Private/OpenGL/SlateOpenGLShaders.cpp
@@ -243,7 +243,8 @@ void FSlateOpenGLShaderProgram::LinkShaders( const FSlateOpenGLVS& VertexShader,
 	// Set up attribute locations for per vertex data
 	glBindAttribLocation(ProgramID, 0, "InTexCoords");
 	glBindAttribLocation(ProgramID, 1, "InPosition");
-	glBindAttribLocation(ProgramID, 4, "InColor");
+	glBindAttribLocation(ProgramID, 2, "InColor");
+	glBindAttribLocation(ProgramID, 3, "InSecondaryColor");
 
 	// Link the shaders
 	glLinkProgram( ProgramID );
@@ -275,7 +276,8 @@ void FSlateOpenGLElementProgram::CreateProgram( const FSlateOpenGLVS& VertexShad
 	EffectsDisabledParam = glGetUniformLocation( ProgramID, "EffectsDisabled" );
 	IgnoreTextureAlphaParam = glGetUniformLocation( ProgramID, "IgnoreTextureAlpha" );
 	ShaderTypeParam = glGetUniformLocation( ProgramID, "ShaderType" );
-	MarginUVsParam = glGetUniformLocation( ProgramID, "MarginUVs" );
+	ShaderParamsParam = glGetUniformLocation( ProgramID, "ShaderParams" );
+	ShaderParams2Param = glGetUniformLocation(ProgramID, "ShaderParams2");
 	GammaValuesParam = glGetUniformLocation(ProgramID, "GammaValues");
 
 	CHECK_GL_ERRORS;
@@ -324,10 +326,14 @@ void FSlateOpenGLElementProgram::SetShaderType( uint32 InShaderType )
 	CHECK_GL_ERRORS;
 }
 
-void FSlateOpenGLElementProgram::SetMarginUVs( const FVector4& InMarginUVs )
+void FSlateOpenGLElementProgram::SetShaderParams(const FShaderParams& InShaderParams)
 {
-	const GLfloat* Params = (GLfloat*)&InMarginUVs;
-	glUniform4fv( MarginUVsParam, 1, Params );
+	const GLfloat* Params = (GLfloat*)&InShaderParams.PixelParams;
+	glUniform4fv(ShaderParamsParam, 1, Params);
+
+	const GLfloat* Params2 = (GLfloat*)&InShaderParams.PixelParams2;
+	glUniform4fv(ShaderParams2Param, 1, Params2);
+
 	CHECK_GL_ERRORS;
 }
 
diff --git a/Engine/Source/Developer/StandaloneRenderer/Private/Windows/D3D/SlateD3DShaders.cpp b/Engine/Source/Developer/StandaloneRenderer/Private/Windows/D3D/SlateD3DShaders.cpp
--- a/Engine/Source/Developer/StandaloneRenderer/Private/Windows/D3D/SlateD3DShaders.cpp
+++ b/Engine/Source/Developer/StandaloneRenderer/Private/Windows/D3D/SlateD3DShaders.cpp
@@ -20,8 +20,7 @@ typedef HRESULT(WINAPI* pD3DReflect)
 	__out void** ppReflector);
 
 
-// @return pointer to the D3DCompile function
-static pD3DCompile GetD3DCompileFunc()
+static HMODULE GetCompilerModule()
 {
 	// Override default compiler path to newer dll
 	FString CompilerPath = FPaths::EngineDir();
@@ -38,39 +37,31 @@ static pD3DCompile GetD3DCompileFunc()
 		CompilerDLL = LoadLibrary(*CompilerPath);
 	}
 
+	return CompilerDLL;
+}
+
+// @return pointer to the D3DCompile function
+static pD3DCompile GetD3DCompileFunc()
+{
+	static HMODULE CompilerDLL = GetCompilerModule();
+
 	if (CompilerDLL)
 	{
 		return (pD3DCompile)(void*)GetProcAddress(CompilerDLL, "D3DCompile");
 	}
 
-	// D3D SDK we compiled with (usually D3DCompiler_43.dll from windows folder)
 	return &D3DCompile;
 }
 
 // @return pointer to the D3DCompile function
 static pD3DReflect GetD3DReflectFunc()
 {
-	// Override default compiler path to newer dll
-	FString CompilerPath = FPaths::EngineDir();
-#if !PLATFORM_64BITS
-	CompilerPath.Append(TEXT("Binaries/ThirdParty/Windows/DirectX/x86/d3dcompiler_47.dll"));
-#else
-	CompilerPath.Append(TEXT("Binaries/ThirdParty/Windows/DirectX/x64/d3dcompiler_47.dll"));
-#endif
-	static bool bHasCompiler = false;
-	static HMODULE CompilerDLL = 0;
-
-	if (bHasCompiler == false)
-	{
-		CompilerDLL = LoadLibrary(*CompilerPath);
-	}
-
+	static HMODULE CompilerDLL = GetCompilerModule();
 	if (CompilerDLL)
 	{
 		return (pD3DReflect)(void*)GetProcAddress(CompilerDLL, "D3DReflect");
 	}
 
-	// D3D SDK we compiled with (usually D3DCompiler_43.dll from windows folder)
 	return &D3DReflect;
 }
 
@@ -274,61 +265,6 @@ void FSlateD3DVS::BindParameters()
 		delete[] ConstantBuffers;
 	}	
 }
-void FSlateD3DGeometryShader::Create( const FString& Filename, const FString& EntryPoint, const FString& ShaderModel )
-{
-	TRefCountPtr<ID3DBlob> Blob;
-	if(CompileShader( Filename, EntryPoint, ShaderModel, Blob))
-	{
-		HRESULT Hr = GD3DDevice->CreateGeometryShader(Blob->GetBufferPointer(), Blob->GetBufferSize(), NULL, GeometryShader.GetInitReference());
-		if (FAILED(Hr))
-		{
-			LogSlateD3DRendererFailure(TEXT("FSlateD3DGeometryShader::Create() - ID3D11Device::CreateGeometryShader"), Hr);
-			GEncounteredCriticalD3DDeviceError = true;
-			return;
-		}
-
-
-		TRefCountPtr<ID3D11ShaderReflection> Reflector;
-		Hr = D3DReflect(Blob->GetBufferPointer(), Blob->GetBufferSize(), IID_ID3D11ShaderReflection, (void**)Reflector.GetInitReference());
-		if (FAILED(Hr))
-		{
-			LogSlateD3DRendererFailure(TEXT("FSlateD3DGeometryShader::Create() - D3DReflect"), Hr);
-			GEncounteredCriticalD3DDeviceError = true;
-			return;
-		}
-
-		GetShaderBindings(Reflector, ShaderBindings);
-	}
-	else
-	{
-		GEncounteredCriticalD3DDeviceError = true;
-	}
-}
-
-
-void FSlateD3DGeometryShader::BindShader()
-{
-	GD3DDeviceContext->GSSetShader( GeometryShader, NULL, 0 );
-}
-
-void FSlateD3DGeometryShader::BindParameters()
-{
-	UpdateParameters();
-
-	if( ShaderBindings.ConstantBuffers.Num() )
-	{
-		const uint32 BufferCount = ShaderBindings.ConstantBuffers.Num();
-		ID3D11Buffer** const ConstantBuffers = new ID3D11Buffer*[ BufferCount ];
-		for( uint32 I = 0; I < BufferCount; ++I )
-		{
-			ConstantBuffers[I] = ShaderBindings.ConstantBuffers[I]->GetParameter().GetReference();
-		}
-
-		GD3DDeviceContext->GSSetConstantBuffers(0, ShaderBindings.ConstantBuffers.Num(), ConstantBuffers);
-
-		delete[] ConstantBuffers;
-	}	
-}
 
 void FSlateD3DPS::Create( const FString& Filename, const FString& EntryPoint, const FString& ShaderModel )
 {
@@ -430,6 +366,7 @@ FSlateDefaultVS::FSlateDefaultVS()
 		{ "TEXCOORD",	1, DXGI_FORMAT_R32G32_FLOAT,		0, D3D11_APPEND_ALIGNED_ELEMENT,	D3D11_INPUT_PER_VERTEX_DATA,	0 },
 		{ "POSITION",	0, DXGI_FORMAT_R32G32_FLOAT,		0, D3D11_APPEND_ALIGNED_ELEMENT,	D3D11_INPUT_PER_VERTEX_DATA,	0 },
 		{ "COLOR",		0, DXGI_FORMAT_B8G8R8A8_UNORM,		0, D3D11_APPEND_ALIGNED_ELEMENT,	D3D11_INPUT_PER_VERTEX_DATA,	0 },
+		{ "COLOR",		1, DXGI_FORMAT_B8G8R8A8_UNORM,		0, D3D11_APPEND_ALIGNED_ELEMENT,	D3D11_INPUT_PER_VERTEX_DATA,	0 },
 	};
 
 	Create( FString::Printf( TEXT("%s/StandaloneRenderer/D3D/SlateDefaultVertexShader.hlsl"), FPlatformProcess::ShaderDir() ), TEXT("Main"), TEXT("vs_4_0"), Layout, UE_ARRAY_COUNT(Layout) );
diff --git a/Engine/Source/Developer/StandaloneRenderer/Private/Windows/D3D/SlateD3DShaders.h b/Engine/Source/Developer/StandaloneRenderer/Private/Windows/D3D/SlateD3DShaders.h
--- a/Engine/Source/Developer/StandaloneRenderer/Private/Windows/D3D/SlateD3DShaders.h
+++ b/Engine/Source/Developer/StandaloneRenderer/Private/Windows/D3D/SlateD3DShaders.h
@@ -114,19 +114,6 @@ private:
 	TRefCountPtr<ID3D11InputLayout> InputLayout;
 };
 
-class FSlateD3DGeometryShader
-{
-public:
-	virtual ~FSlateD3DGeometryShader() {}
-	void Create( const FString& Filename, const FString& EntryPoint, const FString& ShaderModel );
-	void BindShader();
-	void BindParameters();
-protected:
-	virtual void UpdateParameters() {}
-private:
-	FSlateD3DShaderBindings ShaderBindings;
-	TRefCountPtr<ID3D11GeometryShader> GeometryShader;
-};
 
 class FSlateD3DPS
 {
@@ -173,19 +160,20 @@ public:
 		SamplerState = InSamplerState;
 	}
 
-	void SetShaderType( uint32 InShaderType );
-	void SetDrawEffects( ESlateDrawEffect InDrawEffects );
-	void SetShaderParams( const FVector4& InShaderParams );
+	void SetShaderType(uint32 InShaderType);
+	void SetDrawEffects(ESlateDrawEffect InDrawEffects);
+	void SetShaderParams(const FShaderParams& InShaderParams);
 	void SetGammaValues(const FVector2D& InGammaValues);
 protected:
 	void UpdateParameters();
 private:
 	MS_ALIGN(16) struct FPerElementConstants
 	{
-		uint32 ShaderType;			//  4 bytes
 		FVector4 ShaderParams;		// 16 bytes
-		uint32 IgnoreTextureAlpha;	//	4 byte
-		uint32 DisableEffect;		//  4 byte
+		FVector4 ShaderParams2;		// 16 bytes
+		uint32 ShaderType;			//  4 bytes
+		uint32 IgnoreTextureAlpha;	//	4 bytes
+		uint32 DisableEffect;		//  4 bytes
 		uint32 UNUSED[1];			//  4 bytes
 	};
 
diff --git a/Engine/Source/Editor/DetailCustomizations/Private/ColorStructCustomization.cpp b/Engine/Source/Editor/DetailCustomizations/Private/ColorStructCustomization.cpp
--- a/Engine/Source/Editor/DetailCustomizations/Private/ColorStructCustomization.cpp
+++ b/Engine/Source/Editor/DetailCustomizations/Private/ColorStructCustomization.cpp
@@ -106,7 +106,7 @@ TSharedRef<SWidget> FColorStructCustomization::CreateColorWidget(TWeakPtr<IPrope
 					.AlphaDisplayMode(bIgnoreAlpha ? EColorBlockAlphaDisplayMode::Ignore : EColorBlockAlphaDisplayMode::Separate)
 					.OnMouseButtonDown(this, &FColorStructCustomization::OnMouseButtonDownColorBlock)
 					.Size(FVector2D(70.0f, 22.0f))
-					.CornerRadius(4.0f)
+					.CornerRadius(FVector4(4.0f,4.0f,4.0f,4.0f))
 					.IsEnabled(this, &FColorStructCustomization::IsValueEnabled, StructWeakHandlePtr)
 				]
 				+ SOverlay::Slot()
diff --git a/Engine/Source/Runtime/Slate/Private/Widgets/Colors/SColorBlock.cpp b/Engine/Source/Runtime/Slate/Private/Widgets/Colors/SColorBlock.cpp
--- a/Engine/Source/Runtime/Slate/Private/Widgets/Colors/SColorBlock.cpp
+++ b/Engine/Source/Runtime/Slate/Private/Widgets/Colors/SColorBlock.cpp
@@ -64,7 +64,7 @@ int32 SColorBlock::OnPaint(const FPaintArgs& Args, const FGeometry& AllottedGeom
 		MoveTemp(GradientStops),
 		(AllottedGeometry.GetLocalSize().X > AllottedGeometry.GetLocalSize().Y) ? Orient_Vertical : Orient_Horizontal,
 		DrawEffects,
-		GradientCornerRadius.Get(0.0f)
+		GradientCornerRadius.Get(FVector4(0.0f))
 	);
 
 
diff --git a/Engine/Source/Runtime/Slate/Public/Widgets/Colors/SColorBlock.h b/Engine/Source/Runtime/Slate/Public/Widgets/Colors/SColorBlock.h
--- a/Engine/Source/Runtime/Slate/Public/Widgets/Colors/SColorBlock.h
+++ b/Engine/Source/Runtime/Slate/Public/Widgets/Colors/SColorBlock.h
@@ -45,7 +45,7 @@ public:
 		SLATE_ATTRIBUTE(const FSlateBrush*, AlphaBackgroundBrush)
 
 		/** Rounding to apply to the corners of the block */
-		SLATE_ATTRIBUTE(float, CornerRadius)
+		SLATE_ATTRIBUTE(FVector4, CornerRadius)
 
 		/** Whether the color displayed is HSV or not */
 		SLATE_ATTRIBUTE(bool, ColorIsHSV)
@@ -100,7 +100,7 @@ private:
 
 	TAttribute<const FSlateBrush*> SolidBackgroundBrush;
 
-	TAttribute<float> GradientCornerRadius;
+	TAttribute<FVector4> GradientCornerRadius;
 
 	/** Whether the color displayed is HSV or not */
 	TAttribute<bool> ColorIsHSV;
diff --git a/Engine/Source/Runtime/SlateCore/Private/Rendering/DrawElements.cpp b/Engine/Source/Runtime/SlateCore/Private/Rendering/DrawElements.cpp
--- a/Engine/Source/Runtime/SlateCore/Private/Rendering/DrawElements.cpp
+++ b/Engine/Source/Runtime/SlateCore/Private/Rendering/DrawElements.cpp
@@ -235,12 +235,13 @@ FSlateDrawElement& FSlateDrawElement::MakeBoxInternal(
 	if ( ElementType == EElementType::ET_RoundedBox )
 	{
 		FSlateRoundedBoxPayload* RBoxPayload = &ElementList.CreatePayload<FSlateRoundedBoxPayload>(Element);
-		float Radius = InBrush->OutlineSettings.Radius;
+		FVector4 CornerRadii = InBrush->OutlineSettings.CornerRadii;
 		if (InBrush->OutlineSettings.RoundingType == ESlateBrushRoundingType::HalfHeightRadius)
 		{
-			Radius = PaintGeometry.GetLocalSize().Y/2.0f;
+			const float UniformRadius = PaintGeometry.GetLocalSize().Y / 2.0f;
+			CornerRadii = FVector4(UniformRadius, UniformRadius, UniformRadius, UniformRadius);
 		}
-		RBoxPayload->SetRadius(Radius);
+		RBoxPayload->SetRadius(CornerRadii);
 		RBoxPayload->SetOutline(InBrush->OutlineSettings.Color.GetSpecifiedColor(), InBrush->OutlineSettings.Width);
 		BoxPayload = RBoxPayload;
 	}
@@ -439,7 +440,7 @@ void FSlateDrawElement::MakeShapedText(FSlateWindowElementList& ElementList, uin
 	Element.Init(ElementList, EElementType::ET_ShapedText, InLayer, PaintGeometry, InDrawEffects);
 }
 
-void FSlateDrawElement::MakeGradient( FSlateWindowElementList& ElementList, uint32 InLayer, const FPaintGeometry& PaintGeometry, TArray<FSlateGradientStop> InGradientStops, EOrientation InGradientType, ESlateDrawEffect InDrawEffects, float CornerRadius)
+void FSlateDrawElement::MakeGradient( FSlateWindowElementList& ElementList, uint32 InLayer, const FPaintGeometry& PaintGeometry, TArray<FSlateGradientStop> InGradientStops, EOrientation InGradientType, ESlateDrawEffect InDrawEffects, FVector4 CornerRadius)
 {
 	PaintGeometry.CommitTransformsIfUsingLegacyConstructor();
 
diff --git a/Engine/Source/Runtime/SlateCore/Private/Rendering/ElementBatcher.cpp b/Engine/Source/Runtime/SlateCore/Private/Rendering/ElementBatcher.cpp
--- a/Engine/Source/Runtime/SlateCore/Private/Rendering/ElementBatcher.cpp
+++ b/Engine/Source/Runtime/SlateCore/Private/Rendering/ElementBatcher.cpp
@@ -615,17 +615,21 @@ void FSlateElementBatcher::AddBoxElement(const FSlateDrawElement& DrawElement)
 	// Add Shader Parameters for extra RoundedBox parameters
 	ESlateShader ShaderType = ESlateShader::Default;
 	FShaderParams ShaderParams;
+	FColor SecondaryColor;
 	if (DrawElement.GetElementType() == EElementType::ET_RoundedBox)
 	{
+	
 		ShaderType = ESlateShader::RoundedBox;
 		const FSlateRoundedBoxPayload& RoundedPayload = DrawElement.GetDataPayload<FSlateRoundedBoxPayload>();
 
-		ShaderParams.PixelParams = FVector4(RoundedPayload.GetRadius(), RoundedPayload.GetOutlineWeight(), LocalSize.X, LocalSize.Y);//RadiusWeight;
-		ShaderParams.PixelParams2 = FVector4(RoundedPayload.GetOutlineColor());
-	}
+		ShaderParams.PixelParams = FVector4(0, RoundedPayload.GetOutlineWeight(), LocalSize.X, LocalSize.Y);//RadiusWeight;
+		ShaderParams.PixelParams2 = RoundedPayload.GetRadius();
 
-	FSlateRenderBatch& RenderBatch = CreateRenderBatch( Layer, ShaderParams, Resource, ESlateDrawPrimitive::TriangleList, ShaderType, InDrawEffects, DrawFlags, DrawElement);
+		SecondaryColor = PackVertexColor(RoundedPayload.OutlineColor);
+	}
 
+	FSlateRenderBatch& RenderBatch = CreateRenderBatch(Layer, ShaderParams, Resource, ESlateDrawPrimitive::TriangleList, ShaderType, InDrawEffects, DrawFlags, DrawElement);
+	 
 	float HorizontalTiling = bTileHorizontal ? LocalSize.X / TextureWidth : 1.0f;
 	float VerticalTiling = bTileVertical ? LocalSize.Y / TextureHeight : 1.0f;
 
@@ -718,23 +722,23 @@ void FSlateElementBatcher::AddBoxElement(const FSlateDrawElement& DrawElement)
 		FVector2D Position = TopLeft;
 		FVector2D EndPos = BotRight;
 
-		RenderBatch.AddVertex( FSlateVertex::Make<Rounding>( RenderTransform, FVector2D( Position.X, Position.Y ),		LocalSize, DrawScale, FVector4(StartUV,										Tiling),	Tint ) ); //0
-		RenderBatch.AddVertex( FSlateVertex::Make<Rounding>( RenderTransform, FVector2D( Position.X, TopMarginY ),		LocalSize, DrawScale, FVector4(FVector2D( StartUV.X, TopMarginV ),			Tiling),	Tint ) ); //1
-		RenderBatch.AddVertex( FSlateVertex::Make<Rounding>( RenderTransform, FVector2D( LeftMarginX, Position.Y ),		LocalSize, DrawScale, FVector4(FVector2D( LeftMarginU, StartUV.Y ),			Tiling),	Tint ) ); //2
-		RenderBatch.AddVertex( FSlateVertex::Make<Rounding>( RenderTransform, FVector2D( LeftMarginX, TopMarginY ),		LocalSize, DrawScale, FVector4(FVector2D( LeftMarginU, TopMarginV ),		Tiling),	Tint ) ); //3
-		RenderBatch.AddVertex( FSlateVertex::Make<Rounding>( RenderTransform, FVector2D( RightMarginX, Position.Y ),	LocalSize, DrawScale, FVector4(FVector2D( RightMarginU, StartUV.Y ),		Tiling),	Tint ) ); //4
-		RenderBatch.AddVertex( FSlateVertex::Make<Rounding>( RenderTransform, FVector2D( RightMarginX, TopMarginY ),	LocalSize, DrawScale, FVector4(FVector2D( RightMarginU,TopMarginV),			Tiling),	Tint ) ); //5
-		RenderBatch.AddVertex( FSlateVertex::Make<Rounding>( RenderTransform, FVector2D( EndPos.X, Position.Y ),		LocalSize, DrawScale, FVector4(FVector2D( EndUV.X, StartUV.Y ),				Tiling),	Tint ) ); //6
-		RenderBatch.AddVertex( FSlateVertex::Make<Rounding>( RenderTransform, FVector2D( EndPos.X, TopMarginY ),		LocalSize, DrawScale, FVector4(FVector2D( EndUV.X, TopMarginV),				Tiling),	Tint ) ); //7
-
-		RenderBatch.AddVertex( FSlateVertex::Make<Rounding>( RenderTransform, FVector2D( Position.X, BottomMarginY ),	LocalSize, DrawScale, FVector4(FVector2D( StartUV.X, BottomMarginV ),		Tiling),	Tint ) ); //8
-		RenderBatch.AddVertex( FSlateVertex::Make<Rounding>( RenderTransform, FVector2D( LeftMarginX, BottomMarginY ),	LocalSize, DrawScale, FVector4(FVector2D( LeftMarginU, BottomMarginV ),		Tiling),	Tint ) ); //9
-		RenderBatch.AddVertex( FSlateVertex::Make<Rounding>( RenderTransform, FVector2D( RightMarginX, BottomMarginY ),	LocalSize, DrawScale, FVector4(FVector2D( RightMarginU, BottomMarginV ),	Tiling),	Tint ) ); //10
-		RenderBatch.AddVertex( FSlateVertex::Make<Rounding>( RenderTransform, FVector2D( EndPos.X, BottomMarginY ),		LocalSize, DrawScale, FVector4(FVector2D( EndUV.X, BottomMarginV ),			Tiling),	Tint ) ); //11
-		RenderBatch.AddVertex( FSlateVertex::Make<Rounding>( RenderTransform, FVector2D( Position.X, EndPos.Y ),		LocalSize, DrawScale, FVector4(FVector2D( StartUV.X, EndUV.Y ),				Tiling),	Tint ) ); //12
-		RenderBatch.AddVertex( FSlateVertex::Make<Rounding>( RenderTransform, FVector2D( LeftMarginX, EndPos.Y ),		LocalSize, DrawScale, FVector4(FVector2D( LeftMarginU, EndUV.Y ),			Tiling),	Tint ) ); //13
-		RenderBatch.AddVertex( FSlateVertex::Make<Rounding>( RenderTransform, FVector2D( RightMarginX, EndPos.Y ),		LocalSize, DrawScale, FVector4(FVector2D( RightMarginU, EndUV.Y ),			Tiling),	Tint ) ); //14
-		RenderBatch.AddVertex( FSlateVertex::Make<Rounding>( RenderTransform, FVector2D( EndPos.X, EndPos.Y ),			LocalSize, DrawScale, FVector4(EndUV,										Tiling),	Tint ) ); //15
+		RenderBatch.AddVertex( FSlateVertex::Make<Rounding>( RenderTransform, FVector2D( Position.X, Position.Y ),		LocalSize, DrawScale, FVector4(StartUV,										Tiling),	Tint, SecondaryColor) ); //0
+		RenderBatch.AddVertex( FSlateVertex::Make<Rounding>( RenderTransform, FVector2D( Position.X, TopMarginY ),		LocalSize, DrawScale, FVector4(FVector2D( StartUV.X, TopMarginV ),			Tiling),	Tint, SecondaryColor) ); //1
+		RenderBatch.AddVertex( FSlateVertex::Make<Rounding>( RenderTransform, FVector2D( LeftMarginX, Position.Y ),		LocalSize, DrawScale, FVector4(FVector2D( LeftMarginU, StartUV.Y ),			Tiling),	Tint, SecondaryColor) ); //2
+		RenderBatch.AddVertex( FSlateVertex::Make<Rounding>( RenderTransform, FVector2D( LeftMarginX, TopMarginY ),		LocalSize, DrawScale, FVector4(FVector2D( LeftMarginU, TopMarginV ),		Tiling),	Tint, SecondaryColor) ); //3
+		RenderBatch.AddVertex( FSlateVertex::Make<Rounding>( RenderTransform, FVector2D( RightMarginX, Position.Y ),	LocalSize, DrawScale, FVector4(FVector2D( RightMarginU, StartUV.Y ),		Tiling),	Tint, SecondaryColor) ); //4
+		RenderBatch.AddVertex( FSlateVertex::Make<Rounding>( RenderTransform, FVector2D( RightMarginX, TopMarginY ),	LocalSize, DrawScale, FVector4(FVector2D( RightMarginU,TopMarginV),			Tiling),	Tint, SecondaryColor) ); //5
+		RenderBatch.AddVertex( FSlateVertex::Make<Rounding>( RenderTransform, FVector2D( EndPos.X, Position.Y ),		LocalSize, DrawScale, FVector4(FVector2D( EndUV.X, StartUV.Y ),				Tiling),	Tint, SecondaryColor) ); //6
+		RenderBatch.AddVertex( FSlateVertex::Make<Rounding>( RenderTransform, FVector2D( EndPos.X, TopMarginY ),		LocalSize, DrawScale, FVector4(FVector2D( EndUV.X, TopMarginV),				Tiling),	Tint, SecondaryColor) ); //7
+
+		RenderBatch.AddVertex( FSlateVertex::Make<Rounding>( RenderTransform, FVector2D( Position.X, BottomMarginY ),	LocalSize, DrawScale, FVector4(FVector2D( StartUV.X, BottomMarginV ),		Tiling),	Tint, SecondaryColor) ); //8
+		RenderBatch.AddVertex( FSlateVertex::Make<Rounding>( RenderTransform, FVector2D( LeftMarginX, BottomMarginY ),	LocalSize, DrawScale, FVector4(FVector2D( LeftMarginU, BottomMarginV ),		Tiling),	Tint, SecondaryColor) ); //9
+		RenderBatch.AddVertex( FSlateVertex::Make<Rounding>( RenderTransform, FVector2D( RightMarginX, BottomMarginY ),	LocalSize, DrawScale, FVector4(FVector2D( RightMarginU, BottomMarginV ),	Tiling),	Tint, SecondaryColor) ); //10
+		RenderBatch.AddVertex( FSlateVertex::Make<Rounding>( RenderTransform, FVector2D( EndPos.X, BottomMarginY ),		LocalSize, DrawScale, FVector4(FVector2D( EndUV.X, BottomMarginV ),			Tiling),	Tint, SecondaryColor) ); //11
+		RenderBatch.AddVertex( FSlateVertex::Make<Rounding>( RenderTransform, FVector2D( Position.X, EndPos.Y ),		LocalSize, DrawScale, FVector4(FVector2D( StartUV.X, EndUV.Y ),				Tiling),	Tint, SecondaryColor) ); //12
+		RenderBatch.AddVertex( FSlateVertex::Make<Rounding>( RenderTransform, FVector2D( LeftMarginX, EndPos.Y ),		LocalSize, DrawScale, FVector4(FVector2D( LeftMarginU, EndUV.Y ),			Tiling),	Tint, SecondaryColor) ); //13
+		RenderBatch.AddVertex( FSlateVertex::Make<Rounding>( RenderTransform, FVector2D( RightMarginX, EndPos.Y ),		LocalSize, DrawScale, FVector4(FVector2D( RightMarginU, EndUV.Y ),			Tiling),	Tint, SecondaryColor) ); //14
+		RenderBatch.AddVertex( FSlateVertex::Make<Rounding>( RenderTransform, FVector2D( EndPos.X, EndPos.Y ),			LocalSize, DrawScale, FVector4(EndUV,										Tiling),	Tint, SecondaryColor) ); //15
 
 		// Top
 		RenderBatch.AddIndex( IndexStart + 0 );
@@ -807,24 +811,24 @@ void FSlateElementBatcher::AddBoxElement(const FSlateDrawElement& DrawElement)
 			const int32 FeatherStart = RenderBatch.GetNumVertices();
 
 			// Top
-			RenderBatch.AddVertex(FSlateVertex::Make<Rounding>(RenderTransform, FVector2D(Position.X, Position.Y) + FVector2D(-1, -1) / DrawScale, LocalSize, DrawScale, FVector4(StartUV, Tiling), FeatherColor)); //0
-			RenderBatch.AddVertex(FSlateVertex::Make<Rounding>(RenderTransform, FVector2D(LeftMarginX, Position.Y) + FVector2D(0, -1) / DrawScale, LocalSize, DrawScale, FVector4(FVector2D(LeftMarginU, StartUV.Y), Tiling), FeatherColor)); //1
-			RenderBatch.AddVertex(FSlateVertex::Make<Rounding>(RenderTransform, FVector2D(RightMarginX, Position.Y) + FVector2D(0, -1) / DrawScale, LocalSize, DrawScale, FVector4(FVector2D(RightMarginU, StartUV.Y), Tiling), FeatherColor)); //2
-			RenderBatch.AddVertex(FSlateVertex::Make<Rounding>(RenderTransform, FVector2D(EndPos.X, Position.Y) + FVector2D(1, -1) / DrawScale, LocalSize, DrawScale, FVector4(FVector2D(EndUV.X, StartUV.Y), Tiling), FeatherColor)); //3
+			RenderBatch.AddVertex(FSlateVertex::Make<Rounding>(RenderTransform, FVector2D(Position.X, Position.Y) + FVector2D(-1, -1) / DrawScale, LocalSize, DrawScale, FVector4(StartUV, Tiling), FeatherColor, SecondaryColor)); //0
+			RenderBatch.AddVertex(FSlateVertex::Make<Rounding>(RenderTransform, FVector2D(LeftMarginX, Position.Y) + FVector2D(0, -1) / DrawScale, LocalSize, DrawScale, FVector4(FVector2D(LeftMarginU, StartUV.Y), Tiling), FeatherColor, SecondaryColor)); //1
+			RenderBatch.AddVertex(FSlateVertex::Make<Rounding>(RenderTransform, FVector2D(RightMarginX, Position.Y) + FVector2D(0, -1) / DrawScale, LocalSize, DrawScale, FVector4(FVector2D(RightMarginU, StartUV.Y), Tiling), FeatherColor, SecondaryColor)); //2
+			RenderBatch.AddVertex(FSlateVertex::Make<Rounding>(RenderTransform, FVector2D(EndPos.X, Position.Y) + FVector2D(1, -1) / DrawScale, LocalSize, DrawScale, FVector4(FVector2D(EndUV.X, StartUV.Y), Tiling), FeatherColor, SecondaryColor)); //3
 
 			// Left
-			RenderBatch.AddVertex(FSlateVertex::Make<Rounding>(RenderTransform, FVector2D(Position.X, TopMarginY) + FVector2D(-1, 0) / DrawScale, LocalSize, DrawScale, FVector4(FVector2D(StartUV.X, TopMarginV), Tiling), FeatherColor)); //4
-			RenderBatch.AddVertex(FSlateVertex::Make<Rounding>(RenderTransform, FVector2D(Position.X, BottomMarginY) + FVector2D(-1, 0) / DrawScale, LocalSize, DrawScale, FVector4(FVector2D(StartUV.X, BottomMarginV), Tiling), FeatherColor)); //5
+			RenderBatch.AddVertex(FSlateVertex::Make<Rounding>(RenderTransform, FVector2D(Position.X, TopMarginY) + FVector2D(-1, 0) / DrawScale, LocalSize, DrawScale, FVector4(FVector2D(StartUV.X, TopMarginV), Tiling), FeatherColor, SecondaryColor)); //4
+			RenderBatch.AddVertex(FSlateVertex::Make<Rounding>(RenderTransform, FVector2D(Position.X, BottomMarginY) + FVector2D(-1, 0) / DrawScale, LocalSize, DrawScale, FVector4(FVector2D(StartUV.X, BottomMarginV), Tiling), FeatherColor, SecondaryColor)); //5
 
 			// Right
-			RenderBatch.AddVertex(FSlateVertex::Make<Rounding>(RenderTransform, FVector2D(EndPos.X, TopMarginY) + FVector2D(1, 0) / DrawScale, LocalSize, DrawScale, FVector4(FVector2D(EndUV.X, TopMarginV), Tiling), FeatherColor)); //6
-			RenderBatch.AddVertex(FSlateVertex::Make<Rounding>(RenderTransform, FVector2D(EndPos.X, BottomMarginY) + FVector2D(1, 0) / DrawScale, LocalSize, DrawScale, FVector4(FVector2D(EndUV.X, BottomMarginV), Tiling), FeatherColor)); //7
+			RenderBatch.AddVertex(FSlateVertex::Make<Rounding>(RenderTransform, FVector2D(EndPos.X, TopMarginY) + FVector2D(1, 0) / DrawScale, LocalSize, DrawScale, FVector4(FVector2D(EndUV.X, TopMarginV), Tiling), FeatherColor, SecondaryColor)); //6
+			RenderBatch.AddVertex(FSlateVertex::Make<Rounding>(RenderTransform, FVector2D(EndPos.X, BottomMarginY) + FVector2D(1, 0) / DrawScale, LocalSize, DrawScale, FVector4(FVector2D(EndUV.X, BottomMarginV), Tiling), FeatherColor, SecondaryColor)); //7
 
 			// Bottom
-			RenderBatch.AddVertex(FSlateVertex::Make<Rounding>(RenderTransform, FVector2D(Position.X, EndPos.Y) + FVector2D(-1, 1) / DrawScale, LocalSize, DrawScale, FVector4(FVector2D(StartUV.X, EndUV.Y), Tiling), FeatherColor)); //8
-			RenderBatch.AddVertex(FSlateVertex::Make<Rounding>(RenderTransform, FVector2D(LeftMarginX, EndPos.Y) + FVector2D(0, 1) / DrawScale, LocalSize, DrawScale, FVector4(FVector2D(LeftMarginU, EndUV.Y), Tiling), FeatherColor)); //9
-			RenderBatch.AddVertex(FSlateVertex::Make<Rounding>(RenderTransform, FVector2D(RightMarginX, EndPos.Y) + FVector2D(0, 1) / DrawScale, LocalSize, DrawScale, FVector4(FVector2D(RightMarginU, EndUV.Y), Tiling), FeatherColor)); //10
-			RenderBatch.AddVertex(FSlateVertex::Make<Rounding>(RenderTransform, FVector2D(EndPos.X, EndPos.Y) + FVector2D(1, 1) / DrawScale, LocalSize, DrawScale, FVector4(EndUV, Tiling), FeatherColor)); //11
+			RenderBatch.AddVertex(FSlateVertex::Make<Rounding>(RenderTransform, FVector2D(Position.X, EndPos.Y) + FVector2D(-1, 1) / DrawScale, LocalSize, DrawScale, FVector4(FVector2D(StartUV.X, EndUV.Y), Tiling), FeatherColor, SecondaryColor)); //8
+			RenderBatch.AddVertex(FSlateVertex::Make<Rounding>(RenderTransform, FVector2D(LeftMarginX, EndPos.Y) + FVector2D(0, 1) / DrawScale, LocalSize, DrawScale, FVector4(FVector2D(LeftMarginU, EndUV.Y), Tiling), FeatherColor, SecondaryColor)); //9
+			RenderBatch.AddVertex(FSlateVertex::Make<Rounding>(RenderTransform, FVector2D(RightMarginX, EndPos.Y) + FVector2D(0, 1) / DrawScale, LocalSize, DrawScale, FVector4(FVector2D(RightMarginU, EndUV.Y), Tiling), FeatherColor, SecondaryColor)); //10
+			RenderBatch.AddVertex(FSlateVertex::Make<Rounding>(RenderTransform, FVector2D(EndPos.X, EndPos.Y) + FVector2D(1, 1) / DrawScale, LocalSize, DrawScale, FVector4(EndUV, Tiling), FeatherColor, SecondaryColor)); //11
 
 			// Top Left
 			IndexQuad(RenderBatch, FeatherStart + 0, FeatherStart + 1, IndexStart + 2, IndexStart + 0);
@@ -881,10 +885,10 @@ void FSlateElementBatcher::AddBoxElement(const FSlateDrawElement& DrawElement)
 		}
 
 		// Add four vertices to the list of verts to be added to the vertex buffer
-		RenderBatch.AddVertex(FSlateVertex::Make<Rounding>(RenderTransform, TopLeft, LocalSize, DrawScale, FVector4(StartUV, Tiling), Tint));
-		RenderBatch.AddVertex(FSlateVertex::Make<Rounding>(RenderTransform, TopRight, LocalSize, DrawScale, FVector4(FVector2D(EndUV.X, StartUV.Y), Tiling), Tint));
-		RenderBatch.AddVertex(FSlateVertex::Make<Rounding>(RenderTransform, BotLeft, LocalSize, DrawScale, FVector4(FVector2D(StartUV.X, EndUV.Y), Tiling), Tint));
-		RenderBatch.AddVertex(FSlateVertex::Make<Rounding>(RenderTransform, BotRight, LocalSize, DrawScale, FVector4(EndUV, Tiling), Tint));
+		RenderBatch.AddVertex(FSlateVertex::Make<Rounding>(RenderTransform, TopLeft, LocalSize, DrawScale, FVector4(StartUV, Tiling), Tint, SecondaryColor));
+		RenderBatch.AddVertex(FSlateVertex::Make<Rounding>(RenderTransform, TopRight, LocalSize, DrawScale, FVector4(FVector2D(EndUV.X, StartUV.Y), Tiling), Tint, SecondaryColor));
+		RenderBatch.AddVertex(FSlateVertex::Make<Rounding>(RenderTransform, BotLeft, LocalSize, DrawScale, FVector4(FVector2D(StartUV.X, EndUV.Y), Tiling), Tint, SecondaryColor));
+		RenderBatch.AddVertex(FSlateVertex::Make<Rounding>(RenderTransform, BotRight, LocalSize, DrawScale, FVector4(EndUV, Tiling), Tint, SecondaryColor));
 
 		RenderBatch.AddIndex( IndexStart + 0 );
 		RenderBatch.AddIndex( IndexStart + 1 );
@@ -903,10 +907,10 @@ void FSlateElementBatcher::AddBoxElement(const FSlateDrawElement& DrawElement)
 		{
 			const int32 FeatherStart = RenderBatch.GetNumVertices();
 
-			RenderBatch.AddVertex(FSlateVertex::Make<Rounding>(RenderTransform, TopLeft + FVector2D(-1, -1) / DrawScale, LocalSize, DrawScale, FVector4(StartUV, Tiling), FeatherColor));
-			RenderBatch.AddVertex(FSlateVertex::Make<Rounding>(RenderTransform, TopRight + FVector2D(1, -1) / DrawScale, LocalSize, DrawScale, FVector4(FVector2D(EndUV.X, StartUV.Y), Tiling), FeatherColor));
-			RenderBatch.AddVertex(FSlateVertex::Make<Rounding>(RenderTransform, BotLeft + FVector2D(-1, 1) / DrawScale, LocalSize, DrawScale, FVector4(FVector2D(StartUV.X, EndUV.Y), Tiling), FeatherColor));
-			RenderBatch.AddVertex(FSlateVertex::Make<Rounding>(RenderTransform, BotRight + FVector2D(1, 1) / DrawScale, LocalSize, DrawScale, FVector4(EndUV, Tiling), FeatherColor));
+			RenderBatch.AddVertex(FSlateVertex::Make<Rounding>(RenderTransform, TopLeft + FVector2D(-1, -1) / DrawScale, LocalSize, DrawScale, FVector4(StartUV, Tiling), FeatherColor, SecondaryColor));
+			RenderBatch.AddVertex(FSlateVertex::Make<Rounding>(RenderTransform, TopRight + FVector2D(1, -1) / DrawScale, LocalSize, DrawScale, FVector4(FVector2D(EndUV.X, StartUV.Y), Tiling), FeatherColor, SecondaryColor));
+			RenderBatch.AddVertex(FSlateVertex::Make<Rounding>(RenderTransform, BotLeft + FVector2D(-1, 1) / DrawScale, LocalSize, DrawScale, FVector4(FVector2D(StartUV.X, EndUV.Y), Tiling), FeatherColor, SecondaryColor));
+			RenderBatch.AddVertex(FSlateVertex::Make<Rounding>(RenderTransform, BotRight + FVector2D(1, 1) / DrawScale, LocalSize, DrawScale, FVector4(EndUV, Tiling), FeatherColor, SecondaryColor));
 
 			// Top-Top
 			RenderBatch.AddIndex(FeatherStart + 0);
@@ -1441,11 +1445,11 @@ void FSlateElementBatcher::AddGradientElement( const FSlateDrawElement& DrawElem
 	FShaderParams ShaderParams;
 
 	ESlateShader ShaderType = ESlateShader::Default;
-	if (InPayload.CornerRadius > 0)
+	if (InPayload.CornerRadius != FVector4(0))
 	{
 		ShaderType = ESlateShader::RoundedBox;
-		ShaderParams.PixelParams = FVector4(InPayload.CornerRadius, 0.0f, LocalSize.X, LocalSize.Y);
-		ShaderParams.PixelParams2 = FVector4(FLinearColor::Transparent);
+		ShaderParams.PixelParams = FVector4(0.0f, 0.0f, LocalSize.X, LocalSize.Y);
+		ShaderParams.PixelParams2 = InPayload.CornerRadius;
 	}
 
 	FSlateRenderBatch& RenderBatch = 
@@ -1553,14 +1557,14 @@ void FSlateElementBatcher::AddGradientElement( const FSlateDrawElement& DrawElem
 		if( StopIndex == 0 )
 		{
 			// First stop does not have a full quad yet so do not create indices
-			RenderBatch.AddVertex(FSlateVertex::Make<Rounding>(RenderTransform, StartPt, LocalSize, DrawScale, FVector4(StartUV.X, StartUV.Y, 0,0), PackVertexColor(CurStop.Color)));
-			RenderBatch.AddVertex(FSlateVertex::Make<Rounding>(RenderTransform, EndPt, LocalSize, DrawScale, FVector4(EndUV.X, EndUV.Y, 0, 0), PackVertexColor(CurStop.Color)));
+			RenderBatch.AddVertex(FSlateVertex::Make<Rounding>(RenderTransform, StartPt, LocalSize, DrawScale, FVector4(StartUV.X, StartUV.Y, 0,0), PackVertexColor(CurStop.Color), FColor::Transparent));
+			RenderBatch.AddVertex(FSlateVertex::Make<Rounding>(RenderTransform, EndPt, LocalSize, DrawScale, FVector4(EndUV.X, EndUV.Y, 0, 0), PackVertexColor(CurStop.Color), FColor::Transparent));
 		}
 		else
 		{
 			// All stops after the first have indices and generate quads
-			RenderBatch.AddVertex(FSlateVertex::Make<Rounding>(RenderTransform, StartPt, LocalSize, DrawScale, FVector4(StartUV.X, StartUV.Y, 0, 0), PackVertexColor(CurStop.Color)));
-			RenderBatch.AddVertex(FSlateVertex::Make<Rounding>(RenderTransform, EndPt, LocalSize, DrawScale, FVector4(EndUV.X, EndUV.Y, 0, 0), PackVertexColor(CurStop.Color)));
+			RenderBatch.AddVertex(FSlateVertex::Make<Rounding>(RenderTransform, StartPt, LocalSize, DrawScale, FVector4(StartUV.X, StartUV.Y, 0, 0), PackVertexColor(CurStop.Color), FColor::Transparent));
+			RenderBatch.AddVertex(FSlateVertex::Make<Rounding>(RenderTransform, EndPt, LocalSize, DrawScale, FVector4(EndUV.X, EndUV.Y, 0, 0), PackVertexColor(CurStop.Color), FColor::Transparent));
 
 			// Connect the indices to the previous vertices
 			RenderBatch.AddIndex(IndexStart - 2);
diff --git a/Engine/Source/Runtime/SlateCore/Public/Brushes/SlateRoundedBoxBrush.h b/Engine/Source/Runtime/SlateCore/Public/Brushes/SlateRoundedBoxBrush.h
--- a/Engine/Source/Runtime/SlateCore/Public/Brushes/SlateRoundedBoxBrush.h
+++ b/Engine/Source/Runtime/SlateCore/Public/Brushes/SlateRoundedBoxBrush.h
@@ -13,11 +13,9 @@ struct SLATECORE_API FSlateRoundedBoxBrush
 	 * Creates and initializes a new instance with the specified color and rounds based on height
 	 *
 	 * @param InColor 		Linear Fill Color 
-	 * @param InRadius      Corner Radius
 	 */
-
 	template<typename FillColorType>
-	FORCENOINLINE FSlateRoundedBoxBrush(const FillColorType& InColor,  const FVector2D& InImageSize = FVector2D::ZeroVector)
+	FORCENOINLINE FSlateRoundedBoxBrush(const FillColorType& InColor, const FVector2D& InImageSize = FVector2D::ZeroVector)
 		: FSlateBrush(ESlateBrushDrawType::RoundedBox, 
 					  NAME_None, 
 					  FMargin(0.0f), 
@@ -51,10 +49,10 @@ struct SLATECORE_API FSlateRoundedBoxBrush
 	 * Creates and initializes a new instance with the specified color and corner radius
 	 *
 	 * @param InColor 		Linear Fill Color 
-	 * @param InRadius      Corner Radius
+	 * @param InRadius      Corner Radius in Slate Units applied to the outline at each corner. X = Top Left, Y = Top Right, Z = Bottom Right, W = Bottom Left
 	 */
-	template<typename FillColorType>
-	FORCENOINLINE FSlateRoundedBoxBrush(const FillColorType& InColor, float InRadius, const FVector2D& InImageSize = FVector2D::ZeroVector)
+	template<typename FillColorType, typename RadiusType>
+	FORCENOINLINE FSlateRoundedBoxBrush(const FillColorType& InColor, RadiusType InRadius, const FVector2D& InImageSize = FVector2D::ZeroVector)
 		: FSlateBrush(ESlateBrushDrawType::RoundedBox, 
 					  NAME_None, 
 					  FMargin(0.0f), 
@@ -70,8 +68,8 @@ struct SLATECORE_API FSlateRoundedBoxBrush
 	}
 
 
-	template<typename FillColorType>
-	FORCENOINLINE FSlateRoundedBoxBrush(const FName InFillResourceName, const FillColorType& FillColorTint, float InRadius, const FVector2D& InImageSize = FVector2D::ZeroVector, ESlateBrushTileType::Type InTileType = ESlateBrushTileType::NoTile)
+	template<typename FillColorType, typename RadiusType>
+	FORCENOINLINE FSlateRoundedBoxBrush(const FName InFillResourceName, const FillColorType& FillColorTint, RadiusType InRadius, const FVector2D& InImageSize = FVector2D::ZeroVector, ESlateBrushTileType::Type InTileType = ESlateBrushTileType::NoTile)
 		: FSlateBrush(ESlateBrushDrawType::RoundedBox,
 			InFillResourceName,
 			FMargin(0.0f),
@@ -139,12 +137,12 @@ struct SLATECORE_API FSlateRoundedBoxBrush
 	 * Creates and initializes a new instance with the specified color and corner radius
 	 *
 	 * @param InColor 		 Linear Fill Color 
-	 * @param InRadius       Corner Radius
+	 * @param InRadius       Corner Radius in Slate Units applied to the outline at each corner. X = Top Left, Y = Top Right, Z = Bottom Right, W = Bottom Left
 	 * @param InOutlineColor Outline Color 
 	 * @param InOutlineWidth Outline Width or Thickness
 	 */
-	template<typename FillColorType, typename OutlineColorType>
-	FORCENOINLINE FSlateRoundedBoxBrush(const FillColorType& InColor, float InRadius, const OutlineColorType& InOutlineColor, float InOutlineWidth, const FVector2D& InImageSize = FVector2D::ZeroVector)
+	template<typename FillColorType, typename OutlineColorType, typename RadiusType>
+	FORCENOINLINE FSlateRoundedBoxBrush(const FillColorType& InColor, RadiusType InRadius, const OutlineColorType& InOutlineColor, float InOutlineWidth, const FVector2D& InImageSize = FVector2D::ZeroVector)
 		: FSlateBrush(ESlateBrushDrawType::RoundedBox, 
 					  NAME_None, 
 					  FMargin(0.0f), 
@@ -160,15 +158,15 @@ struct SLATECORE_API FSlateRoundedBoxBrush
 	}
 
 	/**
- * Creates and initializes a new instance with the specified color and corner radius
- *
- * @param InColor 		 Linear Fill Color
- * @param InRadius       Corner Radius
- * @param InOutlineColor Outline Color
- * @param InOutlineWidth Outline Width or Thickness
- */
-	template<typename FillColorType, typename OutlineColorType>
-	FORCENOINLINE FSlateRoundedBoxBrush(const FName InFillResourceName, const FillColorType& FillColorTint, float InRadius, const OutlineColorType& InOutlineColor, float InOutlineWidth, const FVector2D& InImageSize = FVector2D::ZeroVector, ESlateBrushTileType::Type InTileType = ESlateBrushTileType::NoTile)
+	 * Creates and initializes a new instance with the specified color and corner radius
+	 *
+	 * @param InColor 		 Linear Fill Color
+	 * @param InRadius      Corner Radius in Slate Units applied to the outline at each corner. X = Top Left, Y = Top Right, Z = Bottom Right, W = Bottom Left
+	 * @param InOutlineColor Outline Color
+	 * @param InOutlineWidth Outline Width or Thickness
+	 */
+	template<typename FillColorType, typename OutlineColorType, typename RadiusType>
+	FORCENOINLINE FSlateRoundedBoxBrush(const FName InFillResourceName, const FillColorType& FillColorTint, RadiusType InRadius, const OutlineColorType& InOutlineColor, float InOutlineWidth, const FVector2D& InImageSize = FVector2D::ZeroVector, ESlateBrushTileType::Type InTileType = ESlateBrushTileType::NoTile)
 		: FSlateBrush(ESlateBrushDrawType::RoundedBox,
 					  InFillResourceName,
 					  FMargin(0.0f),
diff --git a/Engine/Source/Runtime/SlateCore/Public/Rendering/DrawElementPayloads.h b/Engine/Source/Runtime/SlateCore/Public/Rendering/DrawElementPayloads.h
--- a/Engine/Source/Runtime/SlateCore/Public/Rendering/DrawElementPayloads.h
+++ b/Engine/Source/Runtime/SlateCore/Public/Rendering/DrawElementPayloads.h
@@ -85,11 +85,11 @@ struct FSlateBoxPayload : public FSlateDataPayload, public FSlateTintableElement
 struct FSlateRoundedBoxPayload : public FSlateBoxPayload
 {
 	FLinearColor OutlineColor;
-	float Radius;
+	FVector4 Radius;
 	float OutlineWeight;
 
-	FORCEINLINE void SetRadius(float InRadius) { Radius = InRadius; }
-	FORCEINLINE float GetRadius() const { return Radius; }
+	FORCEINLINE void SetRadius(FVector4 InRadius) { Radius = InRadius; }
+	FORCEINLINE FVector4 GetRadius() const { return Radius; }
 
 	FORCEINLINE void SetOutline(const FLinearColor& InOutlineColor, float InOutlineWeight) { OutlineColor = InOutlineColor; OutlineWeight = InOutlineWeight; }
 	FORCEINLINE FLinearColor GetOutlineColor() const { return OutlineColor; }
@@ -162,9 +162,9 @@ struct FSlateGradientPayload : public FSlateDataPayload
 {
 	TArray<FSlateGradientStop> GradientStops;
 	EOrientation GradientType;
-	float CornerRadius;
+	FVector4 CornerRadius;
 
-	void SetGradient(const TArray<FSlateGradientStop>& InGradientStops, EOrientation InGradientType, float InCornerRadius)
+	void SetGradient(const TArray<FSlateGradientStop>& InGradientStops, EOrientation InGradientType, FVector4 InCornerRadius)
 	{
 		GradientStops = InGradientStops;
 		GradientType = InGradientType;
diff --git a/Engine/Source/Runtime/SlateCore/Public/Rendering/DrawElements.h b/Engine/Source/Runtime/SlateCore/Public/Rendering/DrawElements.h
--- a/Engine/Source/Runtime/SlateCore/Public/Rendering/DrawElements.h
+++ b/Engine/Source/Runtime/SlateCore/Public/Rendering/DrawElements.h
@@ -170,7 +170,7 @@ public:
 	 * @param InDrawEffects            Optional draw effects to apply
 	 * @param CornerRadius			   Rounds the corners of the box created by the gradient by the specified radius
 	 */
-	SLATECORE_API static void MakeGradient( FSlateWindowElementList& ElementList, uint32 InLayer, const FPaintGeometry& PaintGeometry, TArray<FSlateGradientStop> InGradientStops, EOrientation InGradientType, ESlateDrawEffect InDrawEffects = ESlateDrawEffect::None, float CornerRadius = 0.0f );
+	SLATECORE_API static void MakeGradient( FSlateWindowElementList& ElementList, uint32 InLayer, const FPaintGeometry& PaintGeometry, TArray<FSlateGradientStop> InGradientStops, EOrientation InGradientType, ESlateDrawEffect InDrawEffects = ESlateDrawEffect::None, FVector4 CornerRadius = FVector4(0.0f) );
 
 	/**
 	 * Creates a Hermite Spline element
diff --git a/Engine/Source/Runtime/SlateCore/Public/Rendering/RenderingCommon.h b/Engine/Source/Runtime/SlateCore/Public/Rendering/RenderingCommon.h
--- a/Engine/Source/Runtime/SlateCore/Public/Rendering/RenderingCommon.h
+++ b/Engine/Source/Runtime/SlateCore/Public/Rendering/RenderingCommon.h
@@ -209,6 +209,9 @@ struct SLATECORE_API FSlateVertex
 	/** Vertex color */
 	FColor Color;
 	
+	/** Secondary vertex color. Generally used for outlines */
+	FColor SecondaryColor;
+
 	/** Local size of the element */
 	uint16 PixelSize[2];
 
@@ -217,33 +220,33 @@ struct SLATECORE_API FSlateVertex
 public:
 
 	template<ESlateVertexRounding Rounding>
-	static FSlateVertex Make(const FSlateRenderTransform& RenderTransform, const FVector2D& InLocalPosition, const FVector2D& InTexCoord, const FVector2D& InTexCoord2, const FColor& InColor)
+	static FSlateVertex Make(const FSlateRenderTransform& RenderTransform, const FVector2D InLocalPosition, const FVector2D InTexCoord, const FVector2D InTexCoord2, const FColor InColor, const FColor SecondaryColor = FColor())
 	{
 		FSlateVertex Vertex;
 		Vertex.TexCoords[0] = InTexCoord.X;
 		Vertex.TexCoords[1] = InTexCoord.Y;
 		Vertex.TexCoords[2] = InTexCoord2.X;
 		Vertex.TexCoords[3] = InTexCoord2.Y;
-		Vertex.InitCommon<Rounding>(RenderTransform, InLocalPosition, InColor);
+		Vertex.InitCommon<Rounding>(RenderTransform, InLocalPosition, InColor, SecondaryColor);
 
 		return Vertex;
 	}
 
 	template<ESlateVertexRounding Rounding>
-	static FSlateVertex Make(const FSlateRenderTransform& RenderTransform, const FVector2D& InLocalPosition, const FVector2D& InTexCoord, const FColor& InColor)
+	static FSlateVertex Make(const FSlateRenderTransform& RenderTransform, const FVector2D InLocalPosition, const FVector2D InTexCoord, const FColor& InColor, const FColor SecondaryColor = FColor())
 	{
 		FSlateVertex Vertex;
 		Vertex.TexCoords[0] = InTexCoord.X;
 		Vertex.TexCoords[1] = InTexCoord.Y;
 		Vertex.TexCoords[2] = 1.0f;
 		Vertex.TexCoords[3] = 1.0f;
-		Vertex.InitCommon<Rounding>(RenderTransform, InLocalPosition, InColor);
+		Vertex.InitCommon<Rounding>(RenderTransform, InLocalPosition, InColor, SecondaryColor);
 
 		return Vertex;
 	}
 
 	template<ESlateVertexRounding Rounding>
-	static FSlateVertex Make(const FSlateRenderTransform& RenderTransform, const FVector2D& InLocalPosition, const FVector4& InTexCoords, const FVector2D& InMaterialTexCoords, const FColor& InColor)
+	static FSlateVertex Make(const FSlateRenderTransform& RenderTransform, const FVector2D InLocalPosition, const FVector4 InTexCoords, const FVector2D InMaterialTexCoords, const FColor InColor, const FColor SecondaryColor = FColor())
 	{
 		FSlateVertex Vertex;
 		Vertex.TexCoords[0] = InTexCoords.X;
@@ -251,13 +254,13 @@ public:
 		Vertex.TexCoords[2] = InTexCoords.Z;
 		Vertex.TexCoords[3] = InTexCoords.W;
 		Vertex.MaterialTexCoords = InMaterialTexCoords;
-		Vertex.InitCommon<Rounding>(RenderTransform, InLocalPosition, InColor);
+		Vertex.InitCommon<Rounding>(RenderTransform, InLocalPosition, InColor, SecondaryColor);
 
 		return Vertex;
 	}
 
 	template<ESlateVertexRounding Rounding>
-	static FSlateVertex Make(const FSlateRenderTransform& RenderTransform, const FVector2D& InLocalPosition, const FVector2D& InLocalSize, float Scale, const FVector4& InTexCoords, const FColor& InColor)
+	static FSlateVertex Make(const FSlateRenderTransform& RenderTransform, const FVector2D InLocalPosition, const FVector2D InLocalSize, float Scale, const FVector4 InTexCoords, const FColor InColor, const FColor SecondaryColor = FColor())
 	{
 		FSlateVertex Vertex;
 		Vertex.TexCoords[0] = InTexCoords.X;
@@ -265,7 +268,7 @@ public:
 		Vertex.TexCoords[2] = InTexCoords.Z;
 		Vertex.TexCoords[3] = InTexCoords.W;
 		Vertex.MaterialTexCoords = FVector2D(InLocalPosition.X / InLocalSize.X, InLocalPosition.Y / InLocalSize.Y);
-		Vertex.InitCommon<Rounding>(RenderTransform, InLocalPosition, InColor);
+		Vertex.InitCommon<Rounding>(RenderTransform, InLocalPosition, InColor, SecondaryColor);
 
 		const int32 PixelSizeX = FMath::RoundToInt(InLocalSize.X * Scale);
 		const int32 PixelSizeY = FMath::RoundToInt(InLocalSize.Y * Scale);
@@ -283,7 +286,7 @@ public:
 		return Vertex;
 	}
 
-	void SetTexCoords(const FVector4& InTexCoords)
+	void SetTexCoords(const FVector4 InTexCoords)
 	{
 		TexCoords[0] = InTexCoords.X;
 		TexCoords[1] = InTexCoords.Y;
@@ -291,7 +294,7 @@ public:
 		TexCoords[3] = InTexCoords.W;
 	}
 
-	void SetPosition(const FVector2D& InPosition)
+	void SetPosition(const FVector2D InPosition)
 	{
 		Position = InPosition;
 	}
@@ -299,7 +302,7 @@ public:
 private:
 
 	template<ESlateVertexRounding Rounding>
-	FORCEINLINE void InitCommon(const FSlateRenderTransform& RenderTransform, const FVector2D& InLocalPosition, const FColor& InColor)
+	FORCEINLINE void InitCommon(const FSlateRenderTransform& RenderTransform, const FVector2D InLocalPosition, const FColor InColor, const FColor InSecondaryColor)
 	{
 		Position = TransformPoint(RenderTransform, InLocalPosition);
 
@@ -310,6 +313,7 @@ private:
 		}
 
 		Color = InColor;
+		SecondaryColor = InSecondaryColor;
 	}
 };
 
diff --git a/Engine/Source/Runtime/SlateCore/Public/Styling/SlateBrush.h b/Engine/Source/Runtime/SlateCore/Public/Styling/SlateBrush.h
--- a/Engine/Source/Runtime/SlateCore/Public/Styling/SlateBrush.h
+++ b/Engine/Source/Runtime/SlateCore/Public/Styling/SlateBrush.h
@@ -130,36 +130,50 @@ struct SLATECORE_API FSlateBrushOutlineSettings
 	GENERATED_USTRUCT_BODY()
 
 	FSlateBrushOutlineSettings()
-		: Radius(0.0)
+		: CornerRadii(0.0)
 		, Color(FLinearColor::Transparent)
 		, Width(0.0)
 		, RoundingType(ESlateBrushRoundingType::HalfHeightRadius)
 	{}
 
-	FSlateBrushOutlineSettings(float InRadius)
-		: Radius(InRadius)
+	FSlateBrushOutlineSettings(float InUniformRadius)
+		: CornerRadii(FVector4(InUniformRadius, InUniformRadius, InUniformRadius, InUniformRadius))
+		, Color(FLinearColor::Transparent)
+		, Width(0.0)
+		, RoundingType(ESlateBrushRoundingType::FixedRadius)
+	{}
+
+	FSlateBrushOutlineSettings(FVector4 InRadius)
+		: CornerRadii(InRadius)
 		, Color(FLinearColor::Transparent)
 		, Width(0.0)
 		, RoundingType(ESlateBrushRoundingType::FixedRadius)
 	{}
 
 	FSlateBrushOutlineSettings(const FSlateColor& InColor, float InWidth)
-		: Radius(0.0)
+		: CornerRadii(0.0)
 		, Color(InColor)
 		, Width(InWidth)
 		, RoundingType(ESlateBrushRoundingType::HalfHeightRadius)
 	{}
 
-	FSlateBrushOutlineSettings(float InRadius, const FSlateColor& InColor, float InWidth)
-		: Radius(InRadius)
+	FSlateBrushOutlineSettings(float InUniformRadius, const FSlateColor& InColor, float InWidth)
+		: CornerRadii(FVector4(InUniformRadius, InUniformRadius, InUniformRadius, InUniformRadius))
+		, Color(InColor)
+		, Width(InWidth)
+		, RoundingType(ESlateBrushRoundingType::FixedRadius)
+	{}
+
+	FSlateBrushOutlineSettings(FVector4 InRadius, const FSlateColor& InColor, float InWidth)
+		: CornerRadii(InRadius)
 		, Color(InColor)
 		, Width(InWidth)
 		, RoundingType(ESlateBrushRoundingType::FixedRadius)
 	{}
 
-	/** Radius in Slate Units applied to the outline. */
+	/** Radius in Slate Units applied to the outline at each corner. X = Top Left, Y = Top Right, Z = Bottom Right, W = Bottom Left */
 	UPROPERTY(EditAnywhere, BlueprintReadWrite, Category=Brush)
-	float Radius;
+	FVector4 CornerRadii;
 
 	/** Tinting applied to the border outline. */
 	UPROPERTY(EditAnywhere, BlueprintReadWrite, Category=Brush, meta=(DisplayName="Outline", sRGB="true"))
diff --git a/Engine/Source/Runtime/SlateRHIRenderer/Private/SlateShaders.cpp b/Engine/Source/Runtime/SlateRHIRenderer/Private/SlateShaders.cpp
--- a/Engine/Source/Runtime/SlateRHIRenderer/Private/SlateShaders.cpp
+++ b/Engine/Source/Runtime/SlateRHIRenderer/Private/SlateShaders.cpp
@@ -53,6 +53,9 @@ IMPLEMENT_SLATE_PIXELSHADER_TYPE(ColorFont, true, true);
 IMPLEMENT_SLATE_PIXELSHADER_TYPE(LineSegment, false, true);
 IMPLEMENT_SLATE_PIXELSHADER_TYPE(LineSegment, true, true);
 
+IMPLEMENT_SLATE_PIXELSHADER_TYPE(RoundedBox, false, true);
+IMPLEMENT_SLATE_PIXELSHADER_TYPE(RoundedBox, true, true);
+
 /** The Slate vertex declaration. */
 TGlobalResource<FSlateVertexDeclaration> GSlateVertexDeclaration;
 TGlobalResource<FSlateInstancedVertexDeclaration> GSlateInstancedVertexDeclaration;
@@ -70,7 +73,8 @@ void FSlateVertexDeclaration::InitRHI()
 	Elements.Add(FVertexElement(0, STRUCT_OFFSET(FSlateVertex, MaterialTexCoords), VET_Float2, 1, Stride));
 	Elements.Add(FVertexElement(0, STRUCT_OFFSET(FSlateVertex, Position), VET_Float2, 2, Stride));
 	Elements.Add(FVertexElement(0, STRUCT_OFFSET(FSlateVertex, Color), VET_Color, 3, Stride));
-	Elements.Add(FVertexElement(0, STRUCT_OFFSET(FSlateVertex, PixelSize), VET_UShort2, 4, Stride));
+	Elements.Add(FVertexElement(0, STRUCT_OFFSET(FSlateVertex, SecondaryColor), VET_Color, 4, Stride));
+	Elements.Add(FVertexElement(0, STRUCT_OFFSET(FSlateVertex, PixelSize), VET_UShort2, 5, Stride));
 
 	VertexDeclarationRHI = PipelineStateCache::GetOrCreateVertexDeclaration(Elements);
 }
@@ -92,7 +96,8 @@ void FSlateInstancedVertexDeclaration::InitRHI()
 	Elements.Add(FVertexElement(0, STRUCT_OFFSET(FSlateVertex, MaterialTexCoords), VET_Float2, 1, Stride));
 	Elements.Add(FVertexElement(0, STRUCT_OFFSET(FSlateVertex, Position), VET_Float2, 2, Stride));
 	Elements.Add(FVertexElement(0, STRUCT_OFFSET(FSlateVertex, Color), VET_Color, 3, Stride));
-	Elements.Add(FVertexElement(1, 0, VET_Float4, 4, sizeof(FVector4), true));
+	Elements.Add(FVertexElement(0, STRUCT_OFFSET(FSlateVertex, SecondaryColor), VET_Color, 4, Stride));
+	Elements.Add(FVertexElement(1, 0, VET_Float4, 5, sizeof(FVector4), true));
 	
 	VertexDeclarationRHI = PipelineStateCache::GetOrCreateVertexDeclaration(Elements);
 }
