Updated UE Logo

diff --git a/Engine/Source/Runtime/SlateCore/Private/Styling/StarshipCoreStyle.cpp b/Engine/Source/Runtime/SlateCore/Private/Styling/StarshipCoreStyle.cpp
--- a/Engine/Source/Runtime/SlateCore/Private/Styling/StarshipCoreStyle.cpp
+++ b/Engine/Source/Runtime/SlateCore/Private/Styling/StarshipCoreStyle.cpp
@@ -273,10 +273,16 @@ TSharedRef<ISlateStyle> FStarshipCoreStyle::Create()
 	// Common icons
 	{
 
+		// chrome icon or not (we are testing this)
+#if 0
 		Style->Set("AppIcon", new IMAGE_BRUSH_SVG("Starship/Common/unreal", FVector2D(36, 36), FStyleColors::Foreground));
-		Style->Set("AppIcon.Small", new IMAGE_BRUSH_SVG("Starship/Common/unreal-small", Icon24x24, FStyleColors::Foreground));
-
 		Style->Set("AppIconPadding", FMargin(11, 11, 3, 5));
+#else
+		Style->Set("AppIcon", new IMAGE_BRUSH("Starship/Common/uelogo_chrome_1", FVector2D(53, 53), FStyleColors::White));
+		Style->Set("AppIconPadding", FMargin(5, 5, 5, 5));
+#endif
+
+		Style->Set("AppIcon.Small", new IMAGE_BRUSH_SVG("Starship/Common/unreal-small", Icon24x24, FStyleColors::Foreground));
 		Style->Set("AppIconPadding.Small", FMargin(4, 4, 0, 0));
 
 		Style->Set("Checker", new IMAGE_BRUSH("Starship/Common/Checker", Icon16x16, FLinearColor::White, ESlateBrushTileType::Both));
