Slate: Update WidgetList with FParse::Bool.
#rnx

diff --git a/Engine/Source/Runtime/SlateCore/Private/Debugging/WidgetList.cpp b/Engine/Source/Runtime/SlateCore/Private/Debugging/WidgetList.cpp
--- a/Engine/Source/Runtime/SlateCore/Private/Debugging/WidgetList.cpp
+++ b/Engine/Source/Runtime/SlateCore/Private/Debugging/WidgetList.cpp
@@ -13,8 +13,8 @@
 #include "Misc/OutputDeviceFile.h"
 #include "Misc/Paths.h"
 #include "Misc/StringBuilder.h"
-#include "Types/CursorMetaData.h"
-#include "Types/MouseEventsMetaData.h"
+#include "Types/SlateCursorMetaData.h"
+#include "Types/SlateMouseEventsMetaData.h"
 #include "Types/ReflectionMetadata.h"
 #include "Widgets/IToolTip.h"
 #include "Widgets/SWidget.h"
@@ -43,59 +43,43 @@ struct FLogAllWidgetsDebugInfoFlags
 
 	void Parse(const FString& Arg)
 	{
-		int32 FoundIndex = Arg.Find(TEXT("Debug="));
-		if (FoundIndex != INDEX_NONE)
+		if (FParse::Bool(*Arg, TEXT("Debug="), bDebug))
 		{
-			LexFromString(bDebug, *Arg + FoundIndex + 6);
 			return;
 		}
 
-		FoundIndex = Arg.Find(TEXT("Paint="));
-		if (FoundIndex != INDEX_NONE)
+		if (FParse::Bool(*Arg, TEXT("Paint="), bPaint))
 		{
-			LexFromString(bPaint, *Arg + FoundIndex + 6);
 			return;
 		}
 
-		FoundIndex = Arg.Find(TEXT("Proxy="));
-		if (FoundIndex != INDEX_NONE)
+		if (FParse::Bool(*Arg, TEXT("Proxy="), bProxy))
 		{
-			LexFromString(bProxy, *Arg + FoundIndex + 6);
 			return;
 		}
 
-		FoundIndex = Arg.Find(TEXT("Children="));
-		if (FoundIndex != INDEX_NONE)
+		if (FParse::Bool(*Arg, TEXT("Children="), bChildren))
 		{
-			LexFromString(bChildren, *Arg + FoundIndex + 9);
 			return;
 		}
 
-		FoundIndex = Arg.Find(TEXT("Parent="));
-		if (FoundIndex != INDEX_NONE)
+		if (FParse::Bool(*Arg, TEXT("Parent="), bParent))
 		{
-			LexFromString(bParent, *Arg + FoundIndex + 7);
 			return;
 		}
 
-		FoundIndex = Arg.Find(TEXT("ToolTip="));
-		if (FoundIndex != INDEX_NONE)
+		if (FParse::Bool(*Arg, TEXT("ToolTip="), bToolTip))
 		{
-			LexFromString(bToolTip, *Arg + FoundIndex + 8);
 			return;
 		}
 
-		FoundIndex = Arg.Find(TEXT("Cursor="));
-		if (FoundIndex != INDEX_NONE)
+		if (FParse::Bool(*Arg, TEXT("Cursor="), bCursor))
 		{
-			LexFromString(bCursor, *Arg + FoundIndex + 7);
 			return;
 		}
 
-		FoundIndex = Arg.Find(TEXT("MouseEvents="));
-		if (FoundIndex != INDEX_NONE)
+		if (FParse::Bool(*Arg, TEXT("MouseEvents="), bMouseEventsHandler))
 		{
-			LexFromString(bMouseEventsHandler, *Arg + FoundIndex + 12);
 			return;
 		}
 	}
@@ -174,15 +158,15 @@ void LogAllWidgetsDebugInfoImpl(FOutputDevice& Ar, const FLogAllWidgetsDebugInfo
 
 		if (DebugInfoFlags.bProxy)
 		{
-			const FWidgetProxyHandle ProxyHandle = Widget->GetProxyHandle();
-			if (ProxyHandle.IsValid())
+			const FWidgetProxyHandle& ProxyHandle = Widget->GetProxyHandle();
+			if (ProxyHandle.IsValid(Widget))
 			{
 				MessageBuilder << TEXT(";");
-				MessageBuilder.Appendf(TEXT("%p"), (void*)ProxyHandle.GetInvalidationRoot());
+				MessageBuilder.Appendf(TEXT("%p"), (void*)ProxyHandle.GetInvalidationRootHandle().Advanced_GetInvalidationRootNoCheck());
 				MessageBuilder << TEXT(";");
-				MessageBuilder << FReflectionMetaData::GetWidgetDebugInfo(ProxyHandle.GetInvalidationRoot()->GetInvalidationRootWidget());
+				MessageBuilder << FReflectionMetaData::GetWidgetDebugInfo(ProxyHandle.GetInvalidationRootHandle().Advanced_GetInvalidationRootNoCheck()->GetInvalidationRootWidget());
 				MessageBuilder << TEXT(";");
-				MessageBuilder << ProxyHandle.GetIndex(true);
+				MessageBuilder << ProxyHandle.GetWidgetIndex().ToString();
 			}
 			else
 			{
@@ -229,7 +213,7 @@ void LogAllWidgetsDebugInfoImpl(FOutputDevice& Ar, const FLogAllWidgetsDebugInfo
 		
 		if (DebugInfoFlags.bCursor)
 		{
-			if (TSharedPtr<FCursorMetaData> Data = Widget->GetMetaData<FCursorMetaData>())
+			if (TSharedPtr<FSlateCursorMetaData> Data = Widget->GetMetaData<FSlateCursorMetaData>())
 			{
 				if (Data->Cursor.IsSet())
 				{
@@ -257,7 +241,7 @@ void LogAllWidgetsDebugInfoImpl(FOutputDevice& Ar, const FLogAllWidgetsDebugInfo
 
 		if (DebugInfoFlags.bMouseEventsHandler)
 		{
-			if (TSharedPtr<FMouseEventsMetaData> Data = Widget->GetMetaData<FMouseEventsMetaData>())
+			if (TSharedPtr<FSlateMouseEventsMetaData> Data = Widget->GetMetaData<FSlateMouseEventsMetaData>())
 			{
 				MessageBuilder << (Data->MouseButtonDownHandle.IsBound() ? TEXT(";bound") : TEXT(";"));
 				MessageBuilder << (Data->MouseButtonUpHandle.IsBound() ? TEXT(";bound") : TEXT(";"));
@@ -304,9 +288,9 @@ void LogAllWidgetsDebugInfo(const TArray< FString >& Args, UWorld*, FOutputDevic
 
 FAutoConsoleCommandWithWorldArgsAndOutputDevice ConsoleCommandLogAllWidgets(
 	TEXT("Slate.Debug.LogAllWidgets"),
-	TEXT("Prints all the SWidgets type, debug info, path or painted.\n")
+	TEXT("Prints all the SWidgets type, debug info, path, painted, ...\n")
 	TEXT("If a file name is not provided, it will output to the log console.\n")
-	TEXT("Slate.Debug.LogAllWidgets [File=MyFile.csv] [Debug=true] [Paint=false] [Proxy=false] [Children=false] [Parent=false]"),
+	TEXT("Slate.Debug.LogAllWidgets [File=MyFile.csv] [Debug=true] [Paint=false] [Proxy=false] [Children=false] [Parent=false] [ToolTip=false] [Cursor=false] [MouseEvents=false]"),
 	FConsoleCommandWithWorldArgsAndOutputDeviceDelegate::CreateStatic(&LogAllWidgetsDebugInfo)
 );
 
diff --git a/Engine/Source/Runtime/SlateCore/Private/Debugging/WidgetList.h b/Engine/Source/Runtime/SlateCore/Private/Debugging/WidgetList.h
--- a/Engine/Source/Runtime/SlateCore/Private/Debugging/WidgetList.h
+++ b/Engine/Source/Runtime/SlateCore/Private/Debugging/WidgetList.h
@@ -7,7 +7,7 @@
 
 // Enabled to build a list of all the SWidget currently constructed
 #ifndef UE_WITH_SLATE_DEBUG_WIDGETLIST
-	#define UE_WITH_SLATE_DEBUG_WIDGETLIST 0
+	#define UE_WITH_SLATE_DEBUG_WIDGETLIST 1
 #endif
 
 #if !WITH_SLATE_DEBUGGING && UE_WITH_SLATE_DEBUG_WIDGETLIST
diff --git a/Engine/Source/Runtime/SlateCore/Public/FastUpdate/SlateInvalidationWidgetIndex.h b/Engine/Source/Runtime/SlateCore/Public/FastUpdate/SlateInvalidationWidgetIndex.h
--- a/Engine/Source/Runtime/SlateCore/Public/FastUpdate/SlateInvalidationWidgetIndex.h
+++ b/Engine/Source/Runtime/SlateCore/Public/FastUpdate/SlateInvalidationWidgetIndex.h
@@ -16,6 +16,7 @@ public:
 	static const FSlateInvalidationWidgetIndex Invalid;
 	bool operator== (FSlateInvalidationWidgetIndex Other) const { return ArrayIndex == Other.ArrayIndex && ElementIndex == Other.ElementIndex; }
 	bool operator!= (FSlateInvalidationWidgetIndex Other) const { return ArrayIndex != Other.ArrayIndex || ElementIndex != Other.ElementIndex; }
+	FString ToString() const { return FString::Printf(TEXT("{%d, %d}"), ArrayIndex, ElementIndex); }
 	
 private:
 	using IndexType = uint16;
