Slate: Fix build issue with WigetList
#rnx

diff --git a/Engine/Source/Runtime/SlateCore/Private/Debugging/WidgetList.h b/Engine/Source/Runtime/SlateCore/Private/Debugging/WidgetList.h
--- a/Engine/Source/Runtime/SlateCore/Private/Debugging/WidgetList.h
+++ b/Engine/Source/Runtime/SlateCore/Private/Debugging/WidgetList.h
@@ -7,7 +7,7 @@
 
 // Enabled to build a list of all the SWidget currently constructed
 #ifndef UE_WITH_SLATE_DEBUG_WIDGETLIST
-	#define UE_WITH_SLATE_DEBUG_WIDGETLIST 1
+	#define UE_WITH_SLATE_DEBUG_WIDGETLIST 0
 #endif
 
 #if !WITH_SLATE_DEBUGGING && UE_WITH_SLATE_DEBUG_WIDGETLIST
