Corner Dropdown Combo Box Style

diff --git a/Engine/Source/Runtime/Slate/Private/Widgets/Input/SComboButton.cpp b/Engine/Source/Runtime/Slate/Private/Widgets/Input/SComboButton.cpp
--- a/Engine/Source/Runtime/Slate/Private/Widgets/Input/SComboButton.cpp
+++ b/Engine/Source/Runtime/Slate/Private/Widgets/Input/SComboButton.cpp
@@ -60,7 +60,7 @@ void SComboButton::Construct( const FArguments& InArgs )
 				+SHorizontalBox::Slot()
 				.AutoWidth()
 				.HAlign( HAlign_Center )
-				.VAlign( VAlign_Center )
+				.VAlign(InArgs._HasDownArrow ? (EVerticalAlignment) InArgs._ComboButtonStyle->DownArrowAlign : VAlign_Center)
 				.Padding(InArgs._HasDownArrow ? InArgs._ComboButtonStyle->DownArrowPadding : FMargin(0))
 				[
 					SNew(SOverlay)
diff --git a/Engine/Source/Runtime/SlateCore/Private/Styling/SlateTypes.cpp b/Engine/Source/Runtime/SlateCore/Private/Styling/SlateTypes.cpp
--- a/Engine/Source/Runtime/SlateCore/Private/Styling/SlateTypes.cpp
+++ b/Engine/Source/Runtime/SlateCore/Private/Styling/SlateTypes.cpp
@@ -148,6 +148,7 @@ FComboButtonStyle::FComboButtonStyle()
 	, MenuBorderPadding(0.0f)
 	, ContentPadding(5.0f)
 	, DownArrowPadding(2.0f)
+	, DownArrowAlign(EVerticalAlignment::VAlign_Center)
 {
 }
 
diff --git a/Engine/Source/Runtime/SlateCore/Private/Styling/StarshipCoreStyle.cpp b/Engine/Source/Runtime/SlateCore/Private/Styling/StarshipCoreStyle.cpp
--- a/Engine/Source/Runtime/SlateCore/Private/Styling/StarshipCoreStyle.cpp
+++ b/Engine/Source/Runtime/SlateCore/Private/Styling/StarshipCoreStyle.cpp
@@ -1146,7 +1146,10 @@ void FStarshipCoreStyle::SetupComboButtonStyles(TSharedRef<FStyle>& Style)
 	FComboButtonStyle SimpleComboButton = FComboButtonStyle()
 		.SetButtonStyle(SimpleButton)
 		.SetContentPadding(0.f)
-		.SetDownArrowImage(IMAGE_BRUSH_SVG("Starship/Common/chevron-down", Icon16x16))
+		.SetShadowOffset(FVector2D(0.0, 0.0))
+		.SetDownArrowImage(IMAGE_BRUSH_SVG("Starship/CoreWidgets/ComboBox/corner-dropdown", FVector2D(7.0f, 7.0f)))
+		.SetDownArrowPadding(FMargin(0.0f))
+		.SetDownArrowAlignment(EVerticalAlignment::VAlign_Bottom)
 		.SetMenuBorderBrush(FSlateColorBrush(FStyleColors::Dropdown))
 		.SetMenuBorderPadding(0.0f);
 	Style->Set("SimpleComboButton", SimpleComboButton);
@@ -1790,6 +1793,11 @@ void FStarshipCoreStyle::SetupMultiboxStyles(TSharedRef<FStyle>& Style)
 			.SetHoveredForeground(FStyleColors::ForegroundHover)
 			.SetDisabledForeground(FStyleColors::Foreground);
 
+
+		FComboButtonStyle SlimToolBarComboButton = FComboButtonStyle(Style->GetWidgetStyle<FComboButtonStyle>("SimpleComboButton"))
+			.SetButtonStyle(SlimToolBarButton)
+			.SetDownArrowPadding(FMargin(1.f, 0.f, 2.f, 5.f));
+
 		FToolBarStyle SlimToolbarStyle =
 			FToolBarStyle()
 			.SetBackground(FSlateColorBrush(FStyleColors::Background))
@@ -1801,7 +1809,7 @@ void FStarshipCoreStyle::SetupMultiboxStyles(TSharedRef<FStyle>& Style)
 			.SetSeparatorBrush(FSlateColorBrush(FStyleColors::Recessed))
 			.SetSeparatorPadding(FMargin(8.f, 0))
 			.SetLabelStyle(FTextBlockStyle(NormalText))
-			.SetComboButtonStyle(Style->GetWidgetStyle<FComboButtonStyle>("ComboButton"))
+			.SetComboButtonStyle(SlimToolBarComboButton)
 			.SetLabelPadding(FMargin(5, 9, 0, 9))
 			.SetEditableTextStyle(FEditableTextBoxStyle(NormalEditableTextBoxStyle));
 
@@ -1823,11 +1831,6 @@ void FStarshipCoreStyle::SetupMultiboxStyles(TSharedRef<FStyle>& Style)
 		SlimToolbarStyle.SetToggleButtonStyle(SlimToolBarToggleButtonCheckBoxStyle);
 		SlimToolbarStyle.SetButtonStyle(SlimToolBarButton);
 
-		FComboButtonStyle SlimToolBarComboButton = FComboButtonStyle(Style->GetWidgetStyle<FComboButtonStyle>("ComboButton"))
-			.SetContentPadding(0)
-			.SetButtonStyle(SlimToolBarButton)
-			.SetDownArrowImage(IMAGE_BRUSH_SVG("Starship/Common/ellipsis-vertical-narrow", FVector2D(6, 24)));
-
 		SlimToolbarStyle.SetSettingsComboButtonStyle(SlimToolBarComboButton);
 		SlimToolbarStyle.SetIconSize(Icon20x20);
 
diff --git a/Engine/Source/Runtime/SlateCore/Public/Styling/SlateTypes.h b/Engine/Source/Runtime/SlateCore/Public/Styling/SlateTypes.h
--- a/Engine/Source/Runtime/SlateCore/Public/Styling/SlateTypes.h
+++ b/Engine/Source/Runtime/SlateCore/Public/Styling/SlateTypes.h
@@ -580,6 +580,13 @@ struct SLATECORE_API FComboButtonStyle : public FSlateWidgetStyle
 	FMargin DownArrowPadding;
 	FComboButtonStyle& SetDownArrowPadding(const FMargin& InDownArrowPadding) { DownArrowPadding = InDownArrowPadding; return *this; }
 
+	/*
+	 * Dropdown arrow vertical alignment
+	 */
+	UPROPERTY(EditAnywhere, BlueprintReadWrite, Category = Appearance)
+	TEnumAsByte<EVerticalAlignment> DownArrowAlign;
+	FComboButtonStyle& SetDownArrowAlignment(const EVerticalAlignment& InVAlign) { DownArrowAlign = InVAlign; return *this; }
+
 	/**
 	* Unlinks all colors in this style.
 	* @see FSlateColor::Unlink
