Updated Style Colors - replaced "Background" Color with "Panel" Color and updated "Background" Color to be dark base

diff --git a/Engine/Source/Developer/OutputLog/Private/SDebugConsole.cpp b/Engine/Source/Developer/OutputLog/Private/SDebugConsole.cpp
--- a/Engine/Source/Developer/OutputLog/Private/SDebugConsole.cpp
+++ b/Engine/Source/Developer/OutputLog/Private/SDebugConsole.cpp
@@ -63,7 +63,7 @@ void SDebugConsole::Construct( const FArguments& InArgs, const EDebugConsoleStyl
 				SNew( SBorder )
 				.VAlign( VAlign_Center )
 				.Padding( FMargin(3.0f, 3.0f) )
-				.BorderImage(FAppStyle::Get().GetBrush("Brushes.Background"))
+				.BorderImage(FAppStyle::Get().GetBrush("Brushes.Panel"))
 				.ColorAndOpacity( this, &SDebugConsole::GetAnimatedColorAndOpacity )
 				.BorderBackgroundColor( this, &SDebugConsole::GetAnimatedSlateColor )
 				[
diff --git a/Engine/Source/Editor/ContentBrowser/Private/SAssetView.cpp b/Engine/Source/Editor/ContentBrowser/Private/SAssetView.cpp
--- a/Engine/Source/Editor/ContentBrowser/Private/SAssetView.cpp
+++ b/Engine/Source/Editor/ContentBrowser/Private/SAssetView.cpp
@@ -333,7 +333,7 @@ void SAssetView::Construct( const FArguments& InArgs )
 	[
 		SNew(SBorder)
 		.Padding(0)
-		.BorderImage(FAppStyle::Get().GetBrush("Brushes.Background"))
+		.BorderImage(FAppStyle::Get().GetBrush("Brushes.Panel"))
 		[
 			VerticalBox
 		]
@@ -419,7 +419,7 @@ void SAssetView::Construct( const FArguments& InArgs )
 	[
 		SNew(SBorder)
 		.Visibility( this, &SAssetView::GetEditModeLabelVisibility )
-		.BorderImage(FAppStyle::Get().GetBrush("Brushes.Background"))
+		.BorderImage(FAppStyle::Get().GetBrush("Brushes.Panel"))
 		.Content()
 		[
 			SNew(SHorizontalBox)
diff --git a/Engine/Source/Editor/ContentBrowser/Private/SContentBrowser.cpp b/Engine/Source/Editor/ContentBrowser/Private/SContentBrowser.cpp
--- a/Engine/Source/Editor/ContentBrowser/Private/SContentBrowser.cpp
+++ b/Engine/Source/Editor/ContentBrowser/Private/SContentBrowser.cpp
@@ -221,7 +221,7 @@ void SContentBrowser::Construct( const FArguments& InArgs, const FName& InInstan
 		[
 			SNew( SBorder )
 			.Padding( FMargin( 3 ) )
-			.BorderImage(FAppStyle::Get().GetBrush("Brushes.Background"))
+			.BorderImage(FAppStyle::Get().GetBrush("Brushes.Panel"))
 			[
 				SNew(SHorizontalBox)
 				+ SHorizontalBox::Slot()
@@ -641,7 +641,7 @@ TSharedRef<SWidget> SContentBrowser::CreateAssetView(const FContentBrowserConfig
 		.AutoHeight()
 		[
 			SNew(SBorder)
-			.BorderImage(FAppStyle::Get().GetBrush("Brushes.Background"))
+			.BorderImage(FAppStyle::Get().GetBrush("Brushes.Panel"))
 			.Padding(FMargin(0.0f, 5.0f))
 			[
 				SNew(SHorizontalBox)
diff --git a/Engine/Source/Editor/CurveEditor/Private/SCurveEditorViewContainer.cpp b/Engine/Source/Editor/CurveEditor/Private/SCurveEditorViewContainer.cpp
--- a/Engine/Source/Editor/CurveEditor/Private/SCurveEditorViewContainer.cpp
+++ b/Engine/Source/Editor/CurveEditor/Private/SCurveEditorViewContainer.cpp
@@ -85,7 +85,7 @@ int32 SCurveEditorViewContainer::OnPaint(const FPaintArgs& Args, const FGeometry
 {
 	const ESlateDrawEffect DrawEffects = ShouldBeEnabled(bParentEnabled) ? ESlateDrawEffect::None : ESlateDrawEffect::DisabledEffect;
 	
-	static const FName BackgroundBrushName("Brushes.Background");
+	static const FName BackgroundBrushName("Brushes.Panel");
 	const FSlateBrush* Background = FEditorStyle::GetBrush(BackgroundBrushName);
 	FSlateDrawElement::MakeBox(OutDrawElements, LayerId, AllottedGeometry.ToPaintGeometry(), Background, DrawEffects, Background->GetTint(InWidgetStyle));
 
diff --git a/Engine/Source/Editor/EditorStyle/Private/EditorStyleSettingsCustomization.cpp b/Engine/Source/Editor/EditorStyle/Private/EditorStyleSettingsCustomization.cpp
--- a/Engine/Source/Editor/EditorStyle/Private/EditorStyleSettingsCustomization.cpp
+++ b/Engine/Source/Editor/EditorStyle/Private/EditorStyleSettingsCustomization.cpp
@@ -58,7 +58,7 @@ public:
 		ChildSlot
 		[
 			SNew(SBorder)
-			.BorderImage(FAppStyle::Get().GetBrush("Brushes.Background"))
+			.BorderImage(FAppStyle::Get().GetBrush("Brushes.Panel"))
 			[
 				SNew(SVerticalBox)
 				+ SVerticalBox::Slot()
diff --git a/Engine/Source/Editor/EditorStyle/Private/StarshipStyle.cpp b/Engine/Source/Editor/EditorStyle/Private/StarshipStyle.cpp
--- a/Engine/Source/Editor/EditorStyle/Private/StarshipStyle.cpp
+++ b/Engine/Source/Editor/EditorStyle/Private/StarshipStyle.cpp
@@ -6136,7 +6136,7 @@ void FStarshipEditorStyle::FStyle::SetupContentBrowserStyle()
 
 		Set( "ContentBrowser.TileViewTooltip.ToolTipBorder", new FSlateColorBrush( FLinearColor::Black ) );
 		Set( "ContentBrowser.TileViewTooltip.NonContentBorder", new BOX_BRUSH( "/Docking/TabContentArea", FMargin(4/16.0f) ) );
-		Set( "ContentBrowser.TileViewTooltip.ContentBorder", new FSlateColorBrush( FStyleColors::Background ));
+		Set( "ContentBrowser.TileViewTooltip.ContentBorder", new FSlateColorBrush( FStyleColors::Panel));
 		Set( "ContentBrowser.TileViewTooltip.NameFont", DEFAULT_FONT( "Regular", 12 ) );
 		Set( "ContentBrowser.TileViewTooltip.AssetUserDescriptionFont", DEFAULT_FONT("Regular", 12 ) );
 
@@ -7898,7 +7898,7 @@ void FStarshipEditorStyle::FStyle::SetupStatusBarStyle()
 	Set("StatusBar.ContentBrowserShadow", new BOX_BRUSH("Starship/StatusBar/drawer-shadow-bottom", FMargin(10/64.0f, 20/64.f, 10/64.f, 0), FLinearColor(0,0,0,1)));
 
 	Set("StatusBar.Height", 32.0f);
-	Set("StatusBar.Background", new FSlateColorBrush(FStyleColors::Background));
+	Set("StatusBar.Background", new FSlateColorBrush(FStyleColors::Panel));
 	Set("StatusBar.HelpIcon", new CORE_IMAGE_BRUSH_SVG("Starship/Common/help", Icon16x16, FStyleColors::Foreground));
 
 	FToolBarStyle StatusBarToolBarStyle = FStarshipCoreStyle::GetCoreStyle().GetWidgetStyle<FToolBarStyle>("SlimToolBar");
diff --git a/Engine/Source/Editor/LevelEditor/Private/SActorDetails.cpp b/Engine/Source/Editor/LevelEditor/Private/SActorDetails.cpp
--- a/Engine/Source/Editor/LevelEditor/Private/SActorDetails.cpp
+++ b/Engine/Source/Editor/LevelEditor/Private/SActorDetails.cpp
@@ -55,7 +55,7 @@ public:
 		ChildSlot
 		[
 			SNew(SBorder)
-			.BorderImage(FAppStyle::Get().GetBrush("Brushes.Background"))
+			.BorderImage(FAppStyle::Get().GetBrush("Brushes.Panel"))
 			[
 				SNew(SHorizontalBox)
 				+ SHorizontalBox::Slot()
diff --git a/Engine/Source/Editor/LevelEditor/Private/SLevelEditor.cpp b/Engine/Source/Editor/LevelEditor/Private/SLevelEditor.cpp
--- a/Engine/Source/Editor/LevelEditor/Private/SLevelEditor.cpp
+++ b/Engine/Source/Editor/LevelEditor/Private/SLevelEditor.cpp
@@ -659,7 +659,7 @@ TSharedRef<SDockTab> SLevelEditor::SummonDetailsPanel( FName TabIdentifier )
 			})
 		[
 			SNew( SBorder)
-			.BorderImage(FAppStyle::Get().GetBrush("Brushes.Background"))
+			.BorderImage(FAppStyle::Get().GetBrush("Brushes.Panel"))
 			.AddMetaData<FTutorialMetaData>(FTutorialMetaData(TEXT("ActorDetails"), TEXT("LevelEditorSelectionDetails")))
 			[
 				ActorDetails
diff --git a/Engine/Source/Editor/MainFrame/Private/Menus/SaveLayoutDialog.cpp b/Engine/Source/Editor/MainFrame/Private/Menus/SaveLayoutDialog.cpp
--- a/Engine/Source/Editor/MainFrame/Private/Menus/SaveLayoutDialog.cpp
+++ b/Engine/Source/Editor/MainFrame/Private/Menus/SaveLayoutDialog.cpp
@@ -191,7 +191,7 @@ void SSaveLayoutDialog::Construct(const FArguments& InArgs, const TSharedRef<FSa
 	ChildSlot
 	[
 		SNew(SBorder)
-		.BorderImage(FAppStyle::Get().GetBrush("Brushes.Background"))
+		.BorderImage(FAppStyle::Get().GetBrush("Brushes.Panel"))
 		[
 			SNew(SBox)
 			.WidthOverride(600.f)
diff --git a/Engine/Source/Editor/PropertyEditor/Private/DetailCategoryGroupNode.cpp b/Engine/Source/Editor/PropertyEditor/Private/DetailCategoryGroupNode.cpp
--- a/Engine/Source/Editor/PropertyEditor/Private/DetailCategoryGroupNode.cpp
+++ b/Engine/Source/Editor/PropertyEditor/Private/DetailCategoryGroupNode.cpp
@@ -151,7 +151,7 @@ FSlateColor SDetailCategoryTableRow::GetOuterBackgroundColor() const
 		return FAppStyle::Get().GetSlateColor("Colors.Header");
 	}
 
-	return FAppStyle::Get().GetSlateColor("Colors.Background");
+	return FAppStyle::Get().GetSlateColor("Colors.Panel");
 }
 
 FReply SDetailCategoryTableRow::OnMouseButtonDown( const FGeometry& MyGeometry, const FPointerEvent& MouseEvent )
diff --git a/Engine/Source/Editor/PropertyEditor/Private/SDetailSingleItemRow.cpp b/Engine/Source/Editor/PropertyEditor/Private/SDetailSingleItemRow.cpp
--- a/Engine/Source/Editor/PropertyEditor/Private/SDetailSingleItemRow.cpp
+++ b/Engine/Source/Editor/PropertyEditor/Private/SDetailSingleItemRow.cpp
@@ -629,17 +629,17 @@ FSlateColor SDetailSingleItemRow::GetOuterBackgroundColor() const
 {
 	if (IsHighlighted())
 	{
-		return FAppStyle::Get().GetSlateColor("Colors.Background");
+		return FAppStyle::Get().GetSlateColor("Colors.Panel");
 	}
 	else if (bIsDragDropObject)
 	{
-		return FAppStyle::Get().GetSlateColor("Colors.Background");
+		return FAppStyle::Get().GetSlateColor("Colors.Panel");
 	}
 	else if (IsHovered())
 	{
 		if (bIsHoveredDragTarget)
 		{
-			return FAppStyle::Get().GetSlateColor("Colors.Background");
+			return FAppStyle::Get().GetSlateColor("Colors.Panel");
 		}
 		else
 		{
@@ -647,7 +647,7 @@ FSlateColor SDetailSingleItemRow::GetOuterBackgroundColor() const
 		}
 	}
 	
-	return FAppStyle::Get().GetSlateColor("Colors.Background");
+	return FAppStyle::Get().GetSlateColor("Colors.Panel");
 }
 
 /** Get the background color of the inner part of the row, which contains the name and value widgets. */
diff --git a/Engine/Source/Editor/PropertyEditor/Private/SDetailsView.cpp b/Engine/Source/Editor/PropertyEditor/Private/SDetailsView.cpp
--- a/Engine/Source/Editor/PropertyEditor/Private/SDetailsView.cpp
+++ b/Engine/Source/Editor/PropertyEditor/Private/SDetailsView.cpp
@@ -278,7 +278,7 @@ void SDetailsView::Construct(const FArguments& InArgs, const FDetailsViewArgs& I
 		.AutoHeight()
 		[
 			SNew(SBorder)
-			.BorderImage(FAppStyle::Get().GetBrush("Brushes.Background"))
+			.BorderImage(FAppStyle::Get().GetBrush("Brushes.Panel"))
 			[
 				FilterRow.ToSharedRef()
 			]
diff --git a/Engine/Source/Editor/PropertyEditor/Private/UserInterface/PropertyEditor/PropertyEditorConstants.cpp b/Engine/Source/Editor/PropertyEditor/Private/UserInterface/PropertyEditor/PropertyEditorConstants.cpp
--- a/Engine/Source/Editor/PropertyEditor/Private/UserInterface/PropertyEditor/PropertyEditorConstants.cpp
+++ b/Engine/Source/Editor/PropertyEditor/Private/UserInterface/PropertyEditor/PropertyEditorConstants.cpp
@@ -20,7 +20,7 @@ FSlateColor PropertyEditorConstants::GetRowBackgroundColor(int32 IndentLevel)
 {
 	if (IndentLevel == 0)
 	{
-		return FAppStyle::Get().GetSlateColor("Colors.Background");
+		return FAppStyle::Get().GetSlateColor("Colors.Panel");
 	}
 
 	int32 ColorIndex = 0;
@@ -41,7 +41,7 @@ FSlateColor PropertyEditorConstants::GetRowBackgroundColor(int32 IndentLevel)
 		2, (2 + 4), (6 + 6), (12 + 8)
 	};
 
-	FColor BaseColor = FAppStyle::Get().GetSlateColor("Colors.Background").GetSpecifiedColor().ToFColor(true);
+	FColor BaseColor = FAppStyle::Get().GetSlateColor("Colors.Panel").GetSpecifiedColor().ToFColor(true);
 
 	FColor ColorWithOffset(
 		BaseColor.R + ColorOffsets[ColorIndex], 
diff --git a/Engine/Source/Editor/UnrealEd/Private/SCSVImportOptions.cpp b/Engine/Source/Editor/UnrealEd/Private/SCSVImportOptions.cpp
--- a/Engine/Source/Editor/UnrealEd/Private/SCSVImportOptions.cpp
+++ b/Engine/Source/Editor/UnrealEd/Private/SCSVImportOptions.cpp
@@ -59,7 +59,7 @@ void SCSVImportOptions::Construct(const FArguments& InArgs)
 	this->ChildSlot
 	[
 		SNew(SBorder)
-		.BorderImage(FAppStyle::Get().GetBrush(TEXT("Brushes.Background")))
+		.BorderImage(FAppStyle::Get().GetBrush(TEXT("Brushes.Panel")))
 		.Padding(10)
 		[
 			SNew(SVerticalBox)
diff --git a/Engine/Source/Editor/UnrealEd/Private/SClassPickerDialog.cpp b/Engine/Source/Editor/UnrealEd/Private/SClassPickerDialog.cpp
--- a/Engine/Source/Editor/UnrealEd/Private/SClassPickerDialog.cpp
+++ b/Engine/Source/Editor/UnrealEd/Private/SClassPickerDialog.cpp
@@ -83,7 +83,7 @@ void SClassPickerDialog::Construct(const FArguments& InArgs)
 	ChildSlot
 	[
 		SNew(SBorder)
-		.BorderImage(FEditorStyle::GetBrush("Brushes.Background"))
+		.BorderImage(FEditorStyle::GetBrush("Brushes.Panel"))
 		[
 			SNew(SBox)
 			.WidthOverride(610.0f)
diff --git a/Engine/Source/Editor/UnrealEd/Private/SSocketManager.cpp b/Engine/Source/Editor/UnrealEd/Private/SSocketManager.cpp
--- a/Engine/Source/Editor/UnrealEd/Private/SSocketManager.cpp
+++ b/Engine/Source/Editor/UnrealEd/Private/SSocketManager.cpp
@@ -280,7 +280,7 @@ void SSocketManager::Construct(const FArguments& InArgs)
 						+SHorizontalBox::Slot()
 						[
 							SNew(SBorder)
-							.BorderImage(FAppStyle::Get().GetBrush("Brushes.Background"))
+							.BorderImage(FAppStyle::Get().GetBrush("Brushes.Panel"))
 							.HAlign(HAlign_Left)
 							.Padding(FMargin(12.f, 6.f))
 							[
diff --git a/Engine/Source/Editor/UnrealEd/Private/Toolkits/AssetEditorToolkit.cpp b/Engine/Source/Editor/UnrealEd/Private/Toolkits/AssetEditorToolkit.cpp
--- a/Engine/Source/Editor/UnrealEd/Private/Toolkits/AssetEditorToolkit.cpp
+++ b/Engine/Source/Editor/UnrealEd/Private/Toolkits/AssetEditorToolkit.cpp
@@ -1117,7 +1117,7 @@ void FAssetEditorToolkit::GenerateToolbar()
 		[
 			SNew(SBorder)
 			.VAlign(VAlign_Center)
-			.BorderImage(FAppStyle::Get().GetBrush("Brushes.Background"))
+			.BorderImage(FAppStyle::Get().GetBrush("Brushes.Panel"))
 			.Padding(FMargin(0.0f))
 			[
 				MiscWidgets
diff --git a/Engine/Source/Runtime/SlateCore/Private/Styling/StarshipCoreStyle.cpp b/Engine/Source/Runtime/SlateCore/Private/Styling/StarshipCoreStyle.cpp
--- a/Engine/Source/Runtime/SlateCore/Private/Styling/StarshipCoreStyle.cpp
+++ b/Engine/Source/Runtime/SlateCore/Private/Styling/StarshipCoreStyle.cpp
@@ -351,7 +351,7 @@ TSharedRef<ISlateStyle> FStarshipCoreStyle::Create()
 	// Tool panels
 	{
 
-		Style->Set("ToolPanel.GroupBorder", new FSlateColorBrush(FStyleColors::Background) );
+		Style->Set("ToolPanel.GroupBorder", new FSlateColorBrush(FStyleColors::Panel) );
 		Style->Set("ToolPanel.DarkGroupBorder", new BOX_BRUSH("Common/DarkGroupBorder", FMargin(4.0f / 16.0f)));
 		Style->Set("ToolPanel.LightGroupBorder", new BOX_BRUSH("Common/LightGroupBorder", FMargin(4.0f / 16.0f)));
 
@@ -418,9 +418,9 @@ TSharedRef<ISlateStyle> FStarshipCoreStyle::Create()
 
 	// SBorder defaults...
 	{
-		Style->Set("Border", new FSlateColorBrush(FStyleColors::Background));
+		Style->Set("Border", new FSlateColorBrush(FStyleColors::Panel));
 
-		FLinearColor TransBackground = FStyleColors::Background.GetSpecifiedColor();
+		FLinearColor TransBackground = FStyleColors::Panel.GetSpecifiedColor();
 		TransBackground.A = .5;
 
 		Style->Set("FloatingBorder", new FSlateRoundedBoxBrush(TransBackground, 8.f));
@@ -464,7 +464,7 @@ TSharedRef<ISlateStyle> FStarshipCoreStyle::Create()
 			.SetExpandedImage(IMAGE_BRUSH_SVG("Starship/Common/chevron-down", Icon16x16, DefaultForeground))
 		);
 		Style->Set("ExpandableArea.TitleFont", StyleFonts.SmallBold);
-		Style->Set("ExpandableArea.Border", new FSlateRoundedBoxBrush(FStyleColors::Background, 4) );
+		Style->Set("ExpandableArea.Border", new FSlateRoundedBoxBrush(FStyleColors::Panel, 4) );
 	}
 
 	// SSlider and SVolumeControl defaults...
@@ -787,6 +787,7 @@ void FStarshipCoreStyle::SetupColors(TSharedRef<FStyle>& Style)
 	Style->Set("Colors.InputOutline", FStyleColors::InputOutline);
 	Style->Set("Colors.Recessed", FStyleColors::Recessed);
 	Style->Set("Colors.Background", FStyleColors::Background);
+	Style->Set("Colors.Panel", FStyleColors::Panel);
 	Style->Set("Colors.Header", FStyleColors::Header);
 	Style->Set("Colors.Dropdown", FStyleColors::Dropdown);
 	Style->Set("Colors.Hover", FStyleColors::Hover);
@@ -831,6 +832,7 @@ void FStarshipCoreStyle::SetupColors(TSharedRef<FStyle>& Style)
 	Style->Set("Brushes.InputOutline", new FSlateColorBrush(FStyleColors::InputOutline));
 	Style->Set("Brushes.Recessed", new FSlateColorBrush(FStyleColors::Recessed));
 	Style->Set("Brushes.Background", new FSlateColorBrush(FStyleColors::Background));
+	Style->Set("Brushes.Panel", new FSlateColorBrush(FStyleColors::Panel));
 	Style->Set("Brushes.Header", new FSlateColorBrush(FStyleColors::Header));
 	Style->Set("Brushes.Dropdown", new FSlateColorBrush(FStyleColors::Dropdown));
 	Style->Set("Brushes.Hover", new FSlateColorBrush(FStyleColors::Hover));
@@ -1045,10 +1047,10 @@ void FStarshipCoreStyle::SetupButtonStyles(TSharedRef<FStyle>& Style)
 		.SetNormal(FSlateRoundedBoxBrush(FStyleColors::Primary, 4.0f, FStyleColors::InputOutline, InputFocusThickness))
 		.SetHovered(FSlateRoundedBoxBrush(FStyleColors::PrimaryHover, 4.0f, FStyleColors::Hover, InputFocusThickness))
 		.SetPressed(FSlateRoundedBoxBrush(FStyleColors::PrimaryPress, 4.0f, FStyleColors::Hover, InputFocusThickness))
-		.SetNormalForeground(FStyleColors::Background)
-		.SetHoveredForeground(FStyleColors::Background)
-		.SetPressedForeground(FStyleColors::Background)
-		.SetDisabledForeground(FStyleColors::Background)
+		.SetNormalForeground(FStyleColors::Panel)
+		.SetHoveredForeground(FStyleColors::Panel)
+		.SetPressedForeground(FStyleColors::Panel)
+		.SetDisabledForeground(FStyleColors::Panel)
 		.SetNormalPadding(ButtonMargins)
 		.SetPressedPadding(ButtonMargins);
 
@@ -1459,7 +1461,7 @@ void FStarshipCoreStyle::SetupDockingStyles(TSharedRef<FStyle>& Style)
 
 	// SDockTab, SDockingTarget, SDockingTabStack defaults...
 	Style->Set("Docking.Background", new BOX_BRUSH("Old/Menu_Background", FMargin(8.0f / 64.0f)));
-	Style->Set("Docking.Border", new FSlateRoundedBoxBrush(FStyleColors::Background, 4));
+	Style->Set("Docking.Border", new FSlateRoundedBoxBrush(FStyleColors::Panel, 4));
 
 	Style->Set("Docking.UnhideTabwellButton", FButtonStyle(Button)
 		.SetNormal(IMAGE_BRUSH_SVG("Starship/Docking/show-tab-well", Icon8x8, FStyleColors::Primary))
@@ -1486,12 +1488,12 @@ void FStarshipCoreStyle::SetupDockingStyles(TSharedRef<FStyle>& Style)
 		FDockTabStyle()
 		.SetCloseButtonStyle(CloseButton)
 		.SetNormalBrush(FSlateNoResource())
-		.SetHoveredBrush(BOX_BRUSH("/Starship/Docking/DockTab_Hover", 4.0f / 20.0f, FStyleColors::Background))
-		.SetForegroundBrush(BOX_BRUSH("/Starship/Docking/DockTab_Foreground", 4.0f / 20.0f, FStyleColors::Background))
+		.SetHoveredBrush(BOX_BRUSH("/Starship/Docking/DockTab_Hover", 4.0f / 20.0f, FStyleColors::Panel))
+		.SetForegroundBrush(BOX_BRUSH("/Starship/Docking/DockTab_Foreground", 4.0f / 20.0f, FStyleColors::Panel))
 
 		.SetColorOverlayTabBrush(FSlateNoResource())
 		.SetColorOverlayIconBrush(FSlateNoResource())
-		.SetContentAreaBrush(FSlateColorBrush(FStyleColors::Background))
+		.SetContentAreaBrush(FSlateColorBrush(FStyleColors::Panel))
 		.SetTabWellBrush(FSlateNoResource())
 		.SetFlashColor(TabFlashColor)
 
@@ -1511,8 +1513,8 @@ void FStarshipCoreStyle::SetupDockingStyles(TSharedRef<FStyle>& Style)
 	Style->Set("Docking.MajorTab", FDockTabStyle()
 		.SetCloseButtonStyle(CloseButton)
 		.SetNormalBrush(FSlateNoResource())
-		.SetHoveredBrush(BOX_BRUSH("/Starship/Docking/DockTab_Hover", 4.0f / 20.0f, FStyleColors::Background))
-		.SetForegroundBrush(BOX_BRUSH("/Starship/Docking/DockTab_Foreground", 4.0f / 20.0f, FStyleColors::Background))
+		.SetHoveredBrush(BOX_BRUSH("/Starship/Docking/DockTab_Hover", 4.0f / 20.0f, FStyleColors::Panel))
+		.SetForegroundBrush(BOX_BRUSH("/Starship/Docking/DockTab_Foreground", 4.0f / 20.0f, FStyleColors::Panel))
 
 		.SetColorOverlayTabBrush(FSlateNoResource())
 		.SetColorOverlayIconBrush(FSlateNoResource())
@@ -1559,7 +1561,7 @@ void FStarshipCoreStyle::SetupDockingStyles(TSharedRef<FStyle>& Style)
 	Style->Set("Docking.SidebarButton.Opened", SidebarTabButtonOpened);
 
 	Style->Set("Docking.Sidebar.DrawerShadow", new BOX_BRUSH("/Starship/Docking/drawer-shadow", FMargin(8/64.f), FLinearColor(0, 0, 0, 1)));
-	Style->Set("Docking.Sidebar.DrawerBackground", new FSlateColorBrush(FStyleColors::Background));
+	Style->Set("Docking.Sidebar.DrawerBackground", new FSlateColorBrush(FStyleColors::Panel));
 	Style->Set("Docking.Sidebar.Background", new FSlateColorBrush(FStyleColors::Recessed));
 	Style->Set("Docking.Sidebar.Border", new FSlateRoundedBoxBrush(FSlateColor(FLinearColor::Transparent), 5.0f, FStyleColors::Hover, 1.0f) );
 
@@ -1738,7 +1740,7 @@ void FStarshipCoreStyle::SetupMultiboxStyles(TSharedRef<FStyle>& Style)
 	{
 		FToolBarStyle NormalToolbarStyle =
 			FToolBarStyle()
-			.SetBackground(FSlateColorBrush(FStyleColors::Background))
+			.SetBackground(FSlateColorBrush(FStyleColors::Panel))
 			.SetBackgroundPadding(FMargin(0, 4, 0, 4))
 			.SetExpandBrush(IMAGE_BRUSH("Icons/toolbar_expand_16x", Icon16x16))
 			.SetComboButtonPadding(FMargin(4.0f, 0.0f))
@@ -1798,7 +1800,7 @@ void FStarshipCoreStyle::SetupMultiboxStyles(TSharedRef<FStyle>& Style)
 
 		FToolBarStyle SlimToolbarStyle =
 			FToolBarStyle()
-			.SetBackground(FSlateColorBrush(FStyleColors::Background))
+			.SetBackground(FSlateColorBrush(FStyleColors::Panel))
 			.SetBackgroundPadding(FMargin(0,4,0,4))
 			.SetExpandBrush(IMAGE_BRUSH("Icons/toolbar_expand_16x", Icon16x16))
 			.SetComboButtonPadding(FMargin(6.0f, 0.0f))
diff --git a/Engine/Source/Runtime/SlateCore/Private/Styling/StyleColors.cpp b/Engine/Source/Runtime/SlateCore/Private/Styling/StyleColors.cpp
--- a/Engine/Source/Runtime/SlateCore/Private/Styling/StyleColors.cpp
+++ b/Engine/Source/Runtime/SlateCore/Private/Styling/StyleColors.cpp
@@ -19,6 +19,7 @@ const FSlateColor FStyleColors::Input = EStyleColor::Input;
 const FSlateColor FStyleColors::InputOutline = EStyleColor::InputOutline;
 const FSlateColor FStyleColors::Recessed = EStyleColor::Recessed;
 const FSlateColor FStyleColors::Background = EStyleColor::Background;
+const FSlateColor FStyleColors::Panel = EStyleColor::Panel;
 const FSlateColor FStyleColors::Header = EStyleColor::Header;
 const FSlateColor FStyleColors::Dropdown = EStyleColor::Dropdown;
 const FSlateColor FStyleColors::Hover = EStyleColor::Hover;
@@ -71,12 +72,13 @@ void USlateThemeManager::InitalizeDefaults()
 
 	SetDefaultColor(EStyleColor::Black, COLOR("#000000FF"));
 	SetDefaultColor(EStyleColor::Title, COLOR("#151515FF"));
+	SetDefaultColor(EStyleColor::Background, COLOR("#151515FF"));
 	SetDefaultColor(EStyleColor::WindowBorder, COLOR("0F0F0FFF"));
 	SetDefaultColor(EStyleColor::Foldout, COLOR("0F0F0FFF"));
 	SetDefaultColor(EStyleColor::Input, COLOR("0F0F0FFF"));
 	SetDefaultColor(EStyleColor::InputOutline, COLOR("454545FF"));
 	SetDefaultColor(EStyleColor::Recessed, COLOR("#1A1A1AFF"));
-	SetDefaultColor(EStyleColor::Background, COLOR("#242424FF"));
+	SetDefaultColor(EStyleColor::Panel, COLOR("#242424FF"));
 	SetDefaultColor(EStyleColor::Header, COLOR("#2F2F2FFF"));
 	SetDefaultColor(EStyleColor::Dropdown, COLOR("#383838FF"));
 	SetDefaultColor(EStyleColor::Hover, COLOR("#575757FF"));
@@ -99,7 +101,7 @@ void USlateThemeManager::InitalizeDefaults()
 	SetDefaultColor(EStyleColor::Select, GetDefaultColor(EStyleColor::Primary));
 	SetDefaultColor(EStyleColor::SelectInactive, COLOR("#99B3BFFF"));
 	SetDefaultColor(EStyleColor::SelectParent, COLOR("#2C323AFF"));
-	SetDefaultColor(EStyleColor::SelectHover, GetDefaultColor(EStyleColor::Background));
+	SetDefaultColor(EStyleColor::SelectHover, GetDefaultColor(EStyleColor::Panel));
 	// if select ==  primary shouldnt we have a select pressed which is the same as primary press?
 
 	SetDefaultColor(EStyleColor::AccentBlue, COLOR("#26BBFFFF"));
diff --git a/Engine/Source/Runtime/SlateCore/Public/Styling/StyleColors.h b/Engine/Source/Runtime/SlateCore/Public/Styling/StyleColors.h
--- a/Engine/Source/Runtime/SlateCore/Public/Styling/StyleColors.h
+++ b/Engine/Source/Runtime/SlateCore/Public/Styling/StyleColors.h
@@ -19,13 +19,14 @@ UENUM()
 enum class EStyleColor : uint8
 {
 	Black,
+	Background,
 	Title,
 	WindowBorder,
 	Foldout,
 	Input,
 	InputOutline,
 	Recessed,
-	Background,
+	Panel,
 	Header,
 	Dropdown,
 	Hover,
@@ -288,6 +289,7 @@ struct SLATECORE_API FStyleColors
 	static const FSlateColor InputOutline;
 	static const FSlateColor Recessed;
 	static const FSlateColor Background;
+	static const FSlateColor Panel;
 	static const FSlateColor Header;
 	static const FSlateColor Dropdown;
 	static const FSlateColor Hover;
