Background Color and Border Style updates

diff --git a/Engine/Source/Editor/LevelEditor/Private/SLevelEditor.cpp b/Engine/Source/Editor/LevelEditor/Private/SLevelEditor.cpp
--- a/Engine/Source/Editor/LevelEditor/Private/SLevelEditor.cpp
+++ b/Engine/Source/Editor/LevelEditor/Private/SLevelEditor.cpp
@@ -272,7 +272,7 @@ void SLevelEditor::Initialize( const TSharedRef<SDockTab>& OwnerTab, const TShar
 			FLevelEditorToolBar::MakeLevelEditorToolBar(LevelEditorCommands.ToSharedRef(), SharedThis(this))
 		]
 		+SVerticalBox::Slot()
-		.Padding(0.0f, 1.0f, 0.0f, 0.0f)
+		.Padding(4.0f, 2.f, 4.f, 2.f)
 		.FillHeight( 1.0f )
 		[
 			ContentArea
diff --git a/Engine/Source/Editor/UnrealEd/Private/Toolkits/SStandaloneAssetEditorToolkitHost.cpp b/Engine/Source/Editor/UnrealEd/Private/Toolkits/SStandaloneAssetEditorToolkitHost.cpp
--- a/Engine/Source/Editor/UnrealEd/Private/Toolkits/SStandaloneAssetEditorToolkitHost.cpp
+++ b/Engine/Source/Editor/UnrealEd/Private/Toolkits/SStandaloneAssetEditorToolkitHost.cpp
@@ -175,13 +175,13 @@ void SStandaloneAssetEditorToolkitHost::RestoreFromLayout( const TSharedRef<FTab
 		.AutoHeight()
 		.Expose(ToolbarSlot)
 		+ SVerticalBox::Slot()
-		.Padding(0.0f, 1.0f, 0.0f, 0.0f)
+		.Padding(4.f, 2.f, 4.f, 2.f)
 		.FillHeight(1.0f)
 		[
 			RestoredUI.ToSharedRef()
 		]
 		+ SVerticalBox::Slot()
-		.Padding(0.0f, 1.0f, 0.0f, 0.0f)
+		.Padding(0.0f, 2.0f, 0.0f, 0.0f)
 		.AutoHeight()
 		[
 			GEditor->GetEditorSubsystem<UStatusBarSubsystem>()->MakeStatusBarWidget(GetMenuName(), HostTab)
diff --git a/Engine/Source/Runtime/Slate/Private/Framework/Docking/SDockingTabStack.cpp b/Engine/Source/Runtime/Slate/Private/Framework/Docking/SDockingTabStack.cpp
--- a/Engine/Source/Runtime/Slate/Private/Framework/Docking/SDockingTabStack.cpp
+++ b/Engine/Source/Runtime/Slate/Private/Framework/Docking/SDockingTabStack.cpp
@@ -106,7 +106,7 @@ void SDockingTabStack::Construct( const FArguments& InArgs, const TSharedRef<FTa
 		+ SHorizontalBox::Slot() 
 		.FillWidth(1.0f)
 		.VAlign(VAlign_Bottom)
-		.Padding(4.0f, 0.0f, 0.0f, 0.0f)
+		.Padding(0.0f, 0.0f, 0.0f, 0.0f)
 		[
 			SNew(SVerticalBox)
 			.Visibility(EVisibility::SelfHitTestInvisible)
@@ -1130,7 +1130,7 @@ FSlateColor SDockingTabStack::GetUnhideTabWellButtonOpacity() const
 const FSlateBrush* SDockingTabStack::GetTabStackBorderImage() const
 {
 	static const FSlateBrush* MajorTabBackgroundBrush = FAppStyle::Get().GetBrush("Brushes.Title");
-	static const FSlateBrush* MinorTabBackgroundBrush = FAppStyle::Get().GetBrush("Brushes.Recessed");
+	static const FSlateBrush* MinorTabBackgroundBrush = FAppStyle::Get().GetBrush("Brushes.Background");
 	return bShowingTitleBarArea ? MajorTabBackgroundBrush : MinorTabBackgroundBrush;
 }
 
diff --git a/Engine/Source/Runtime/SlateCore/Private/Styling/StarshipCoreStyle.cpp b/Engine/Source/Runtime/SlateCore/Private/Styling/StarshipCoreStyle.cpp
--- a/Engine/Source/Runtime/SlateCore/Private/Styling/StarshipCoreStyle.cpp
+++ b/Engine/Source/Runtime/SlateCore/Private/Styling/StarshipCoreStyle.cpp
@@ -543,7 +543,7 @@ TSharedRef<ISlateStyle> FStarshipCoreStyle::Create()
 	// SSplitter
 	{
 		Style->Set("Splitter", FSplitterStyle()
-			.SetHandleNormalBrush(FSlateColorBrush(FStyleColors::Recessed))
+			.SetHandleNormalBrush(FSlateColorBrush(FStyleColors::Background))
 			.SetHandleHighlightBrush(FSlateColorBrush(FStyleColors::Secondary))
 		);
 	}
@@ -1499,7 +1499,7 @@ void FStarshipCoreStyle::SetupDockingStyles(TSharedRef<FStyle>& Style)
 		.SetColorOverlayTabBrush(FSlateNoResource())
 		.SetColorOverlayIconBrush(FSlateNoResource())
 		.SetContentAreaBrush(FSlateColorBrush(FStyleColors::Panel))
-		.SetTabWellBrush(FSlateNoResource())
+		.SetTabWellBrush(FSlateColorBrush(FStyleColors::Background))
 		.SetFlashColor(TabFlashColor)
 
 		.SetTabPadding(FMargin(10, 3, 10, 4))
@@ -1523,8 +1523,8 @@ void FStarshipCoreStyle::SetupDockingStyles(TSharedRef<FStyle>& Style)
 
 		.SetColorOverlayTabBrush(FSlateNoResource())
 		.SetColorOverlayIconBrush(FSlateNoResource())
-		.SetContentAreaBrush(FSlateColorBrush(FStyleColors::Recessed))
-		.SetTabWellBrush(FSlateNoResource())
+		.SetContentAreaBrush(FSlateColorBrush(FStyleColors::Background))
+		.SetTabWellBrush(FSlateColorBrush(FStyleColors::Background))
 
 		.SetTabPadding(FMargin(10, 7, 10, 8))
 		.SetOverlapWidth(0.f)
@@ -1757,7 +1757,7 @@ void FStarshipCoreStyle::SetupMultiboxStyles(TSharedRef<FStyle>& Style)
 			.SetComboButtonPadding(FMargin(4.0f, 0.0f))
 			.SetButtonPadding(FMargin(2.0f, 0.f))
 			.SetCheckBoxPadding(FMargin(4.0f, 0.f))
-			.SetSeparatorBrush(FSlateColorBrush(FStyleColors::Recessed))
+			.SetSeparatorBrush(FSlateColorBrush(FStyleColors::Background))
 			.SetSeparatorPadding(FMargin(1.f, 0.f, 1.f, 0.f))
 			.SetLabelStyle(FTextBlockStyle(NormalText).SetFont(FStyleFonts::Get().Normal))
 			.SetEditableTextStyle(FEditableTextBoxStyle(NormalEditableTextBoxStyle).SetFont(FStyleFonts::Get().Normal))
