Slate: First implementation of SlateAttribute. SlateAttribute makes the attribute system viable for invalidation and more performance-friendly while keeping the benefits of attributes intact. They are updated once per frame before the Prepass. If the cached value of the SlateAttribute changes, then it will invalidate the widget. The feature works in slow path but not quite in fast (child order invalidation may cause infinit loop, deactivated for now).

diff --git a/Engine/Source/Editor/AnimationBlueprintEditor/Private/AnimationStateNodes/SGraphNodeAnimState.cpp b/Engine/Source/Editor/AnimationBlueprintEditor/Private/AnimationStateNodes/SGraphNodeAnimState.cpp
--- a/Engine/Source/Editor/AnimationBlueprintEditor/Private/AnimationStateNodes/SGraphNodeAnimState.cpp
+++ b/Engine/Source/Editor/AnimationBlueprintEditor/Private/AnimationStateNodes/SGraphNodeAnimState.cpp
@@ -41,8 +41,6 @@ void SStateMachineOutputPin::Construct(const FArguments& InArgs, UEdGraphPin* In
 {
 	this->SetCursor( EMouseCursor::Default );
 
-	typedef SStateMachineOutputPin ThisClass;
-
 	bShowLabel = true;
 
 	GraphPinObj = InPin;
@@ -54,9 +52,9 @@ void SStateMachineOutputPin::Construct(const FArguments& InArgs, UEdGraphPin* In
 	// Set up a hover for pins that is tinted the color of the pin.
 	SBorder::Construct( SBorder::FArguments()
 		.BorderImage( this, &SStateMachineOutputPin::GetPinBorder )
-		.BorderBackgroundColor( this, &ThisClass::GetPinColor )
-		.OnMouseButtonDown( this, &ThisClass::OnPinMouseDown )
-		.Cursor( this, &ThisClass::GetPinCursor )
+		.BorderBackgroundColor( this, &SStateMachineOutputPin::GetPinColor )
+		.OnMouseButtonDown( this, &SStateMachineOutputPin::OnPinMouseDown )
+		.Cursor( this, &SStateMachineOutputPin::GetPinCursor )
 	);
 }
 
diff --git a/Engine/Source/Runtime/SlateCore/Private/SlateCoreClasses.cpp b/Engine/Source/Runtime/SlateCore/Private/SlateCoreClasses.cpp
--- a/Engine/Source/Runtime/SlateCore/Private/SlateCoreClasses.cpp
+++ b/Engine/Source/Runtime/SlateCore/Private/SlateCoreClasses.cpp
@@ -51,6 +51,7 @@ static FAutoConsoleVariableRef CVarSlateNewUpdateMethod(
 );
 
 bool GSlateIsOnFastUpdatePath = false;
+bool GSlateIsOnFastProcessInvalidation = false;
 bool GSlateIsInInvalidationSlowPath = false;
 
 #if SLATE_CHECK_UOBJECT_RENDER_RESOURCES
diff --git a/Engine/Source/Runtime/SlateCore/Private/Types/SlateAttribute.cpp b/Engine/Source/Runtime/SlateCore/Private/Types/SlateAttribute.cpp
new file mode 100644
--- /dev/null
+++ b/Engine/Source/Runtime/SlateCore/Private/Types/SlateAttribute.cpp
@@ -0,0 +1,94 @@
+// Copyright Epic Games, Inc. All Rights Reserved.
+
+#include "Types/SlateAttribute.h"
+
+#include "Types/ReflectionMetadata.h"
+#include "Types/SlateAttributeMetaData.h"
+#include "Widgets/SWidget.h"
+#include "Debugging/WidgetList.h"
+
+namespace SlateAttributePrivate
+{
+	void TestAttributeAddress(const SWidget& OwningWidget, const FSlateAttributeImpl& Attribute, ESlateAttributeType AttributeType)
+	{
+#if STATS && DO_CHECK
+		if (AttributeType == ESlateAttributeType::Member)
+		{
+			UPTRINT SlateAttributePtr = (UPTRINT)&Attribute;
+			UPTRINT WidgetPtr = (UPTRINT)&OwningWidget;
+			checkf(SlateAttributePtr >= WidgetPtr && SlateAttributePtr <= WidgetPtr + OwningWidget.GetAllocSize(),
+				TEXT("You can only register Attribute that are defined in a SWidget. ")
+				TEXT("Use TAttribute or TSlateExternalAttribute instead. See SWidget: '%s'. See SlateAttribute.h for more info.")
+				, *FReflectionMetaData::GetWidgetPath(OwningWidget));
+		}
+#endif
+	}
+
+	bool FSlateAttributeImpl::ProtectedIsWidgetInDestructionPath(SWidget* Widget) const
+	{
+#if UE_WITH_SLATE_DEBUG_WIDGETLIST
+		int32 FoundIndex = UE::Slate::FWidgetList::GetAllWidgets().Find(Widget);
+		if (FoundIndex == INDEX_NONE)
+		{
+			return false; // the widget is already destroyed
+		}
+		// if the shared instance exist, then the widget is not being destroyed
+		return !UE::Slate::FWidgetList::GetAllWidgets()[FoundIndex]->DoesSharedInstanceExist();
+#endif
+	return true;
+	}
+
+	void FSlateAttributeImpl::ProtectedUnregisterAttribute(SWidget& OwningWidget, ESlateAttributeType AttributeType) const
+	{
+		TestAttributeAddress(OwningWidget, *this, AttributeType);
+		
+		FSlateAttributeMetaData::UnregisterAttribute(OwningWidget, *this);
+	}
+
+
+	void FSlateAttributeImpl::ProtectedRegisterAttribute(SWidget& OwningWidget, ESlateAttributeType AttributeType, TUniquePtr<ISlateAttributeGetter>&& Wrapper)
+	{
+		TestAttributeAddress(OwningWidget, *this, AttributeType);
+
+		FSlateAttributeMetaData::RegisterAttribute(OwningWidget, *this, AttributeType, MoveTemp(Wrapper));
+	}
+
+
+	void FSlateAttributeImpl::ProtectedInvalidateWidget(SWidget& OwningWidget, ESlateAttributeType AttributeType, EInvalidateWidgetReason InvalidationReason) const
+	{
+		TestAttributeAddress(OwningWidget, *this, AttributeType);
+
+		FSlateAttributeMetaData::InvalidateWidget(OwningWidget, *this, InvalidationReason);
+	}
+
+
+	bool FSlateAttributeImpl::ProtectedIsBound(const SWidget& OwningWidget, ESlateAttributeType AttributeType) const
+	{
+		TestAttributeAddress(OwningWidget, *this, AttributeType);
+
+		return FSlateAttributeMetaData::IsAttributeBound(OwningWidget, *this);
+	}
+
+
+	FDelegateHandle FSlateAttributeImpl::ProtectedFindGetterHandle(const SWidget& OwningWidget, ESlateAttributeType AttributeType) const
+	{
+		TestAttributeAddress(OwningWidget, *this, AttributeType);
+
+		return FSlateAttributeMetaData::GetAttributeGetterHandle(OwningWidget, *this);
+	}
+	
+
+	void FSlateAttributeImpl::ProtectedUpdateNow(SWidget& OwningWidget, ESlateAttributeType AttributeType)
+	{
+		TestAttributeAddress(OwningWidget, *this, AttributeType);
+
+		FSlateAttributeMetaData::UpdateAttribute(OwningWidget, *this);
+	}
+
+	void FSlateAttributeImpl::ProtectedMoveAttribute(SWidget& OwningWidget, ESlateAttributeType AttributeType, const FSlateAttributeBase* Other)
+	{
+		checkf(AttributeType != ESlateAttributeType::Member, TEXT("Member Attribute cannot be moved."));
+
+		FSlateAttributeMetaData::MoveAttribute(OwningWidget, *this, AttributeType, Other);
+	}
+}
diff --git a/Engine/Source/Runtime/SlateCore/Private/Types/SlateAttributeDescriptor.cpp b/Engine/Source/Runtime/SlateCore/Private/Types/SlateAttributeDescriptor.cpp
new file mode 100644
--- /dev/null
+++ b/Engine/Source/Runtime/SlateCore/Private/Types/SlateAttributeDescriptor.cpp
@@ -0,0 +1,234 @@
+// Copyright Epic Games, Inc. All Rights Reserved.
+
+#pragma once
+
+#include "Types/SlateAttributeDescriptor.h"
+
+
+FSlateAttributeDescriptor::FInitializer::FAttributeEntry::FAttributeEntry(FSlateAttributeDescriptor& InDescriptor, int32 InAttributeIndex)
+	: Descriptor(InDescriptor)
+	, AttributeIndex(InAttributeIndex)
+{}
+			
+FSlateAttributeDescriptor::FInitializer::FAttributeEntry& FSlateAttributeDescriptor::FInitializer::FAttributeEntry::SetPrerequisite(FName Prerequisite)
+{
+	if (Descriptor.Attributes.IsValidIndex(AttributeIndex))
+	{
+		Descriptor.SetPrerequisite(Descriptor.Attributes[AttributeIndex], Prerequisite);
+	}
+	return *this;
+}
+
+FSlateAttributeDescriptor::FInitializer::FAttributeEntry& FSlateAttributeDescriptor::FInitializer::FAttributeEntry::UpdateEveryFrame()
+{
+	if (Descriptor.Attributes.IsValidIndex(AttributeIndex))
+	{
+		Descriptor.Attributes[AttributeIndex].Dependency = FName();
+	}
+	return *this;
+}
+
+FSlateAttributeDescriptor::FInitializer::FAttributeEntry& FSlateAttributeDescriptor::FInitializer::FAttributeEntry::UpdateDependency(FName Dependency)
+{
+	if (Descriptor.Attributes.IsValidIndex(AttributeIndex))
+	{
+		Descriptor.SetDependency(Descriptor.Attributes[AttributeIndex], Dependency);
+	}
+	return *this;
+}
+
+
+FSlateAttributeDescriptor::FInitializer::FInitializer(FSlateAttributeDescriptor& InDescriptor)
+	: Descriptor(InDescriptor)
+{
+}
+
+FSlateAttributeDescriptor::FInitializer::FInitializer(FSlateAttributeDescriptor& InDescriptor, const FSlateAttributeDescriptor& ParentDescriptor)
+	: Descriptor(InDescriptor)
+{
+	InDescriptor.Attributes = ParentDescriptor.Attributes;
+}
+
+FSlateAttributeDescriptor::FInitializer::~FInitializer()
+{
+	// because adding the attribute is not required at the moment
+	//try to not change the order in which they were added
+	//for (int32 Index = 0; Index < Descriptor.Attributes.Num(); ++Index)
+	//{
+	//	if (!Descriptor.Attributes[Index].Prerequisite.IsNone())
+	//	{
+	//		// Find the Prerequisite index
+	//		const FName Prerequisite = Descriptor.Attributes[Index].Prerequisite;
+	//		const int32 FoundIndex = Descriptor.Attributes.IndexOfByPredicate([Prerequisite](const FAttribute& Other) { return Other.Name == Prerequisite; });
+	//		if (ensureAlwaysMsgf(Descriptor.Attributes.IsValidIndex(FoundIndex), TEXT("The Prerequisite '%s' doesn't exist"), *Prerequisite.ToString()))
+	//		{
+	//			Descriptor.Attributes[Index].SortOrder = Descriptor.Attributes[FoundIndex].SortOrder + 1;
+	//		}
+	//	}
+	//}
+}
+
+
+FSlateAttributeDescriptor::FInitializer::FAttributeEntry FSlateAttributeDescriptor::FInitializer::AddMemberAttribute(FName AttributeName, OffsetType Offset)
+{
+	return Descriptor.AddMemberAttribute(AttributeName, Offset, TAttribute<EInvalidateWidgetReason>());
+}
+
+
+FSlateAttributeDescriptor::FInitializer::FAttributeEntry FSlateAttributeDescriptor::FInitializer::AddMemberAttribute(FName AttributeName, OffsetType Offset, TAttribute<EInvalidateWidgetReason> Reason)
+{
+	check(Reason.IsSet());
+	return Descriptor.AddMemberAttribute(AttributeName, Offset, MoveTemp(Reason));
+}
+
+
+void FSlateAttributeDescriptor::FInitializer::OverrideInvalidationReason(FName AttributeName, TAttribute<EInvalidateWidgetReason> Reason)
+{
+	check(Reason.IsSet());
+	Descriptor.OverrideInvalidationReason(AttributeName, MoveTemp(Reason));
+}
+
+
+void FSlateAttributeDescriptor::FInitializer::SetPrerequisite(FName AttributeName, FName Prerequisite)
+{
+	check(!AttributeName.IsNone());
+
+	FAttribute* FoundAttribute = Descriptor.FindAttribute(AttributeName);
+	if (ensureAlwaysMsgf(FoundAttribute == nullptr, TEXT("The attribute '%s' doesn't exist."), *AttributeName.ToString()))
+	{
+		Descriptor.SetPrerequisite(*FoundAttribute, Prerequisite);
+	}
+}
+
+
+FSlateAttributeDescriptor::FAttribute const* FSlateAttributeDescriptor::FindAttribute(FName AttributeName) const
+{
+	return Attributes.FindByPredicate([AttributeName](const FAttribute& Other) { return Other.Name == AttributeName; });
+}
+
+
+FSlateAttributeDescriptor::FAttribute const* FSlateAttributeDescriptor::FindMemberAttribute(OffsetType AttributeOffset) const
+{
+	FSlateAttributeDescriptor::FAttribute const* Result = Attributes.FindByPredicate([AttributeOffset](const FAttribute& Other) { return Other.Offset == AttributeOffset; });
+	check(Result == nullptr || Result->bIsMemberAttribute);
+	return Result;
+}
+
+
+FSlateAttributeDescriptor::FAttribute* FSlateAttributeDescriptor::FindAttribute(FName AttributeName)
+{
+	return Attributes.FindByPredicate([AttributeName](const FAttribute& Other) { return Other.Name == AttributeName; });
+}
+
+
+FSlateAttributeDescriptor::FInitializer::FAttributeEntry FSlateAttributeDescriptor::AddMemberAttribute(FName AttributeName, OffsetType Offset, TAttribute<EInvalidateWidgetReason>&& Reason)
+{
+	check(!AttributeName.IsNone());
+
+	int32 NewIndex = INDEX_NONE;
+	FAttribute const* FoundAttribute = FindAttribute(AttributeName);
+	if (ensureAlwaysMsgf(FoundAttribute == nullptr, TEXT("The attribute '%s' already exist."), *AttributeName.ToString()))
+	{
+		NewIndex = Attributes.AddZeroed();
+		FAttribute& NewAttribute = Attributes[NewIndex];
+		NewAttribute.Name = AttributeName;
+		NewAttribute.Offset = Offset;
+		NewAttribute.InvalidationReason = MoveTemp(Reason);
+		NewAttribute.SortOrder = DefaultSortOrder(Offset);
+		NewAttribute.bIsMemberAttribute = true;
+	}
+	return FInitializer::FAttributeEntry(*this, NewIndex);
+}
+
+
+void FSlateAttributeDescriptor::OverrideInvalidationReason(FName AttributeName, TAttribute<EInvalidateWidgetReason>&& Reason)
+{
+	check(!AttributeName.IsNone());
+
+	FAttribute* FoundAttribute = FindAttribute(AttributeName);
+	if (ensureAlwaysMsgf(FoundAttribute != nullptr, TEXT("The attribute 's' doesn't exist."), *AttributeName.ToString()))
+	{
+		FoundAttribute->InvalidationReason = MoveTemp(Reason);
+	}
+}
+
+void FSlateAttributeDescriptor::SetDependency(FSlateAttributeDescriptor::FAttribute& Attribute, FName Dependency)
+{
+	if (Dependency.IsNone())
+	{
+		Attribute.Dependency = FName();
+	}
+	else
+	{
+		FAttribute const* FoundDependency = FindAttribute(Dependency);
+		if (ensureAlwaysMsgf(FoundDependency, TEXT("The Dependency '%s' doesn't exist for attribute '%s'"), *Dependency.ToString(), *Attribute.Name.ToString()))
+		{
+			Attribute.Dependency = Dependency;
+
+			// Verify recursion
+			{
+				FMemMark Mark(FMemStack::Get());
+				TArray<FName, TMemStackAllocator<>> Recursion;
+				Recursion.Reserve(Attributes.Num());
+				FAttribute const* CurrentAttribute = &Attribute;
+				while (!CurrentAttribute->Dependency.IsNone())
+				{
+					if (Recursion.Contains(CurrentAttribute->Name))
+					{
+						ensureAlwaysMsgf(false, TEXT("The Dependency '%s' would introduce an infinit loop with attribute '%s'."), *Dependency.ToString(), *Attribute.Name.ToString());
+						Attribute.Dependency = FName();
+						break;
+					}
+					Recursion.Add(CurrentAttribute->Name);
+					CurrentAttribute = FindAttribute(CurrentAttribute->Dependency);
+					check(CurrentAttribute);
+				}
+			}
+		}
+		else
+		{
+			Attribute.Dependency = FName();
+		}
+	}
+}
+
+void FSlateAttributeDescriptor::SetPrerequisite(FSlateAttributeDescriptor::FAttribute& Attribute, FName Prerequisite)
+{
+	if (Prerequisite.IsNone())
+	{
+		Attribute.Prerequisite = FName();
+	}
+	else
+	{
+		FAttribute const* FoundPrerequisite = FindAttribute(Prerequisite);
+		if (ensureAlwaysMsgf(FoundPrerequisite, TEXT("The prerequisite '%s' doesn't exist for attribute '%s'"), *Prerequisite.ToString(), *Attribute.Name.ToString()))
+		{
+			Attribute.Prerequisite = Prerequisite;
+			Attribute.SortOrder = FoundPrerequisite->SortOrder + 1;
+
+			// Verify recursion
+			{
+				FMemMark Mark(FMemStack::Get());
+				TArray<FName, TMemStackAllocator<>> Recursion;
+				Recursion.Reserve(Attributes.Num());
+				FAttribute const* CurrentAttribute = &Attribute;
+				while (!CurrentAttribute->Prerequisite.IsNone())
+				{
+					if (Recursion.Contains(CurrentAttribute->Name))
+					{
+						ensureAlwaysMsgf(false, TEXT("The prerequsite '%s' would introduce an infinit loop with attribute '%s'."), *Prerequisite.ToString(), *Attribute.Name.ToString());
+						Attribute.Prerequisite = FName();
+						break;
+					}
+					Recursion.Add(CurrentAttribute->Name);
+					CurrentAttribute = FindAttribute(CurrentAttribute->Prerequisite);
+					check(CurrentAttribute);
+				}
+			}
+		}
+		else
+		{
+			Attribute.Prerequisite = FName();
+		}
+	}
+}
diff --git a/Engine/Source/Runtime/SlateCore/Private/Types/SlateAttributeMetaData.cpp b/Engine/Source/Runtime/SlateCore/Private/Types/SlateAttributeMetaData.cpp
new file mode 100644
--- /dev/null
+++ b/Engine/Source/Runtime/SlateCore/Private/Types/SlateAttributeMetaData.cpp
@@ -0,0 +1,229 @@
+// Copyright Epic Games, Inc. All Rights Reserved.
+
+#include "Types/SlateAttributeMetaData.h"
+
+#include "Types/ReflectionMetadata.h"
+#include "Widgets/SWidget.h"
+
+#include <limits>
+
+
+namespace Private
+{
+	FSlateAttributeDescriptor::OffsetType FindOffet(const SWidget& OwningWidget, FSlateAttributeBase& Attribute)
+	{
+		return (FSlateAttributeDescriptor::OffsetType)((UPTRINT)(&Attribute) - (UPTRINT)(&OwningWidget));
+	}
+}
+
+
+FSlateAttributeMetaData* FSlateAttributeMetaData::FindMetaData(const SWidget& OwningWidget)
+{
+	if (OwningWidget.HasRegisteredSlateAttribute())
+	{
+		check(OwningWidget.MetaData.Num() > 0);
+		TSharedRef<ISlateMetaData> SlateMetaData = OwningWidget.MetaData[0];
+		check(SlateMetaData->IsOfType<FSlateAttributeMetaData>());
+		return &(StaticCastSharedRef<FSlateAttributeMetaData>(SlateMetaData).Get());
+	}
+	else if (OwningWidget.MetaData.Num() > 0)
+	{
+		TSharedRef<ISlateMetaData> SlateMetaData = OwningWidget.MetaData[0];
+		if (SlateMetaData->IsOfType<FSlateAttributeMetaData>())
+		{
+			return &(StaticCastSharedRef<FSlateAttributeMetaData>(SlateMetaData).Get());
+		}
+	}
+	return nullptr;
+}
+
+
+void FSlateAttributeMetaData::RegisterAttribute(SWidget& OwningWidget, FSlateAttributeBase& Attribute, ESlateAttributeType AttributeType, TUniquePtr<ISlateAttributeGetter>&& Wrapper)
+{
+	if (FSlateAttributeMetaData* AttributeMetaData = FSlateAttributeMetaData::FindMetaData(OwningWidget))
+	{
+		AttributeMetaData->RegisterAttributeImpl(OwningWidget, Attribute, AttributeType, MoveTemp(Wrapper));
+	}
+	else
+	{
+		TSharedRef<FSlateAttributeMetaData> NewAttributeMetaData = MakeShared<FSlateAttributeMetaData>();
+		NewAttributeMetaData->RegisterAttributeImpl(OwningWidget, Attribute, AttributeType, MoveTemp(Wrapper));
+		OwningWidget.bHasRegisteredSlateAttribute = true;
+		if (OwningWidget.GetProxyHandle().IsValid(&OwningWidget))
+		{
+			//OwningWidget.FastPathProxyHandle.MarkWidgetRegisteredSlateAttributeDirty(&OwningWidget);
+		}
+		OwningWidget.MetaData.Insert(NewAttributeMetaData, 0);
+	}
+}
+
+
+void FSlateAttributeMetaData::RegisterAttributeImpl(SWidget& OwningWidget, FSlateAttributeBase& Attribute, ESlateAttributeType AttributeType, TUniquePtr<ISlateAttributeGetter>&& Getter)
+{
+	const int32 FoundIndex = IndexOfAttribute(Attribute);
+	if (FoundIndex != INDEX_NONE)
+	{
+		Attributes[FoundIndex].Getter = MoveTemp(Getter);
+	}
+	else
+	{
+		if (AttributeType == ESlateAttributeType::Member)
+		{
+			FSlateAttributeDescriptor::OffsetType Offset = Private::FindOffet(OwningWidget, Attribute);
+			uint32 SortOrder = FSlateAttributeDescriptor::DefaultSortOrder(Offset);
+
+			// MemberAttribute are optional for now but will be needed in the future
+			const FSlateAttributeDescriptor& Descriptor = OwningWidget.GetWidgetClass().GetAttributeDescriptor();
+			if (FSlateAttributeDescriptor::FAttribute const* FoundAttribute = Descriptor.FindMemberAttribute(Offset))
+			{
+				Attributes.Emplace(&Attribute, FoundAttribute->SortOrder, MoveTemp(Getter), FoundAttribute->InvalidationReason);
+			}
+			else
+			{
+				Attributes.Emplace(&Attribute, SortOrder, MoveTemp(Getter));
+			}
+		}
+		else if (AttributeType == ESlateAttributeType::Managed)
+		{
+			const uint32 ManagedSortOrder = std::numeric_limits<uint32>::max();
+			Attributes.Emplace(&Attribute, ManagedSortOrder, MoveTemp(Getter));
+		}
+		else
+		{
+			ensureMsgf(false, TEXT("The SlateAttributeType is not supported"));
+		}
+
+		Attributes.Sort();
+	}
+}
+
+
+bool FSlateAttributeMetaData::UnregisterAttribute(SWidget& OwningWidget, const FSlateAttributeBase& Attribute)
+{
+	if (FSlateAttributeMetaData* AttributeMetaData = FSlateAttributeMetaData::FindMetaData(OwningWidget))
+	{
+		const bool bResult = AttributeMetaData->UnregisterAttributeImpl(Attribute);
+		if (AttributeMetaData->Attributes.Num() == 0)
+		{
+			check(bResult); // if the num is 0 then we should have remove an item.
+			OwningWidget.bHasRegisteredSlateAttribute = false;
+			if (OwningWidget.GetProxyHandle().IsValid(&OwningWidget))
+			{
+				//OwningWidget.FastPathProxyHandle.MarkWidgetRegisteredSlateAttributeDirty(&OwningWidget);
+			}
+			OwningWidget.MetaData.RemoveAtSwap(0);
+		}
+		return bResult;
+	}
+	return false;
+}
+
+
+bool FSlateAttributeMetaData::UnregisterAttributeImpl(const FSlateAttributeBase& Attribute)
+{
+	const int32 FoundIndex = IndexOfAttribute(Attribute);
+	if (FoundIndex != INDEX_NONE)
+	{
+		Attributes.RemoveAt(FoundIndex);
+		return true;
+	}
+	return false;
+}
+
+
+void FSlateAttributeMetaData::InvalidateWidget(SWidget& OwningWidget, const FSlateAttributeBase& Attribute, EInvalidateWidgetReason Reason)
+{
+	if (FSlateAttributeMetaData* AttributeMetaData = FSlateAttributeMetaData::FindMetaData(OwningWidget))
+	{
+		const int32 FoundIndex = AttributeMetaData->IndexOfAttribute(Attribute);
+		if (FoundIndex != INDEX_NONE)
+		{
+			OwningWidget.Invalidate(AttributeMetaData->Attributes[FoundIndex].OverrideInvalidateReason.Get(Reason));
+		}
+		else
+		{
+			OwningWidget.Invalidate(Reason);
+		}
+	}
+	else
+	{
+		OwningWidget.Invalidate(Reason);
+	}
+}
+
+
+void FSlateAttributeMetaData::UpdateAttributes(SWidget& OwningWidget)
+{
+	if (FSlateAttributeMetaData* AttributeMetaData = FSlateAttributeMetaData::FindMetaData(OwningWidget))
+	{
+		EInvalidateWidgetReason InvalidationReason = EInvalidateWidgetReason::None;
+		for (FGetterItem& GetterItem: AttributeMetaData->Attributes)
+		{
+			ISlateAttributeGetter::FUpdateAttributeResult Result = GetterItem.Getter->UpdateAttribute(OwningWidget);
+			if (Result.bInvalidationRequested)
+			{
+				InvalidationReason |= GetterItem.OverrideInvalidateReason.Get(Result.InvalidationReason);
+			}
+		}
+		OwningWidget.Invalidate(InvalidationReason);
+	}
+}
+
+
+void FSlateAttributeMetaData::UpdateAttribute(SWidget& OwningWidget, FSlateAttributeBase& Attribute)
+{
+	if (FSlateAttributeMetaData* AttributeMetaData = FSlateAttributeMetaData::FindMetaData(OwningWidget))
+	{
+		const int32 FoundIndex = AttributeMetaData->IndexOfAttribute(Attribute);
+		if (FoundIndex != INDEX_NONE)
+		{
+			FGetterItem& GetterItem = AttributeMetaData->Attributes[FoundIndex];
+			check(GetterItem.Getter.Get());
+			ISlateAttributeGetter::FUpdateAttributeResult Result = GetterItem.Getter->UpdateAttribute(OwningWidget);
+			if (Result.bInvalidationRequested)
+			{
+				OwningWidget.Invalidate(GetterItem.OverrideInvalidateReason.Get(Result.InvalidationReason));
+			}
+		}
+	}
+}
+
+
+bool FSlateAttributeMetaData::IsAttributeBound(const SWidget& OwningWidget, const FSlateAttributeBase& Attribute)
+{
+	if (FSlateAttributeMetaData* AttributeMetaData = FSlateAttributeMetaData::FindMetaData(OwningWidget))
+	{
+		return AttributeMetaData->IndexOfAttribute(Attribute) != INDEX_NONE;
+	}
+	return false;
+}
+
+
+FDelegateHandle FSlateAttributeMetaData::GetAttributeGetterHandle(const SWidget& OwningWidget, const FSlateAttributeBase& Attribute)
+{
+	if (FSlateAttributeMetaData* AttributeMetaData = FSlateAttributeMetaData::FindMetaData(OwningWidget))
+	{
+		const int32 FoundIndex = AttributeMetaData->IndexOfAttribute(Attribute);
+		if (FoundIndex != INDEX_NONE)
+		{
+			return AttributeMetaData->Attributes[FoundIndex].Getter->GetDelegateHandle();
+		}
+	}
+	return FDelegateHandle();
+}
+
+
+void FSlateAttributeMetaData::MoveAttribute(const SWidget& OwningWidget, FSlateAttributeBase& NewAttribute, ESlateAttributeType AttributeType, const FSlateAttributeBase* PreviousAttribute)
+{
+	checkf(AttributeType == ESlateAttributeType::Managed, TEXT("TSlateAttribute cannot be moved. This should be already prevented in SlateAttribute.h"));
+	if (FSlateAttributeMetaData* AttributeMetaData = FSlateAttributeMetaData::FindMetaData(OwningWidget))
+	{
+		const int32 FoundIndex = AttributeMetaData->Attributes.IndexOfByPredicate([PreviousAttribute](const FGetterItem& Item) { return Item.Attribute == PreviousAttribute; });
+		if (FoundIndex != INDEX_NONE)
+		{
+			AttributeMetaData->Attributes[FoundIndex].Attribute = &NewAttribute;
+			AttributeMetaData->Attributes[FoundIndex].Getter->SetAttribute(NewAttribute);
+			//Attributes.Sort(); // Managed are always at the end and there order is not realiable.
+		}
+	}
+}
diff --git a/Engine/Source/Runtime/SlateCore/Private/Widgets/SWidget.cpp b/Engine/Source/Runtime/SlateCore/Private/Widgets/SWidget.cpp
--- a/Engine/Source/Runtime/SlateCore/Private/Widgets/SWidget.cpp
+++ b/Engine/Source/Runtime/SlateCore/Private/Widgets/SWidget.cpp
@@ -18,6 +18,7 @@
 #include "Debugging/WidgetList.h"
 #include "Widgets/SWindow.h"
 #include "Trace/SlateTrace.h"
+#include "Types/SlateAttributeMetaData.h"
 #include "Types/SlateCursorMetaData.h"
 #include "Types/SlateMouseEventsMetaData.h"
 #include "Types/ReflectionMetadata.h"
@@ -195,6 +196,14 @@ namespace SlateTraceMetaData
 }
 #endif
 
+SLATE_IMPLEMENT_WIDGET(SWidget)
+void SWidget::PrivateRegisterAttributes(FSlateAttributeInitializer& AttributeInitializer)
+{
+	AttributeInitializer.AddMemberAttribute(GET_MEMBER_NAME_CHECKED(PrivateThisType, EnabledState), STRUCT_OFFSET(PrivateThisType, EnabledState));
+	//SLATE_ADD_ATTRIBUTE_DEFINITION(AttributeInitializer, EnabledState);
+	//SLATE_ADD_ATTRIBUTE_DEFINITION(AttributeInitializer, Visibility, EInvalidateWidgetReason::All);
+}
+
 SWidget::SWidget()
 	: bIsHovered(false)
 	, bCanSupportFocus(true)
@@ -207,6 +216,7 @@ SWidget::SWidget()
 	, bInvisibleDueToParentOrSelfVisibility(false)
 	, bNeedsPrepass(true)
 	, bUpdatingDesiredSize(false)
+	, bHasRegisteredSlateAttribute(false)
 	, bHasCustomPrepass(false)
 	, bHasRelativeLayoutScale(false)
 	, bVolatilityAlwaysInvalidatesPrepass(false)
@@ -1170,6 +1180,11 @@ void SWidget::Invalidate(EInvalidateWidgetReason InvalidateReason)
 {
 	SLATE_CROSS_THREAD_CHECK();
 
+	if (InvalidateReason == EInvalidateWidgetReason::None)
+	{
+		return;
+	}
+
 	SCOPED_NAMED_EVENT_TEXT("SWidget::Invalidate", FColor::Orange);
 	const bool bWasVolatile = IsVolatileIndirectly() || IsVolatile();
 
@@ -1250,7 +1265,9 @@ void SWidget::SetDebugInfo( const ANSICHAR* InType, const ANSICHAR* InFile, int3
 {
 	TypeOfWidget = InType;
 
-	STAT(AllocSize = InAllocSize);
+#if STATS
+	AllocSize = InAllocSize;
+#endif
 	INC_MEMORY_STAT_BY(STAT_SlateSWidgetAllocSize, AllocSize);
 
 #if !UE_BUILD_SHIPPING
@@ -1576,6 +1593,11 @@ void SWidget::Prepass_Internal(float InLayoutScaleMultiplier)
 {
 	PrepassLayoutScaleMultiplier = InLayoutScaleMultiplier;
 
+	if (HasRegisteredSlateAttribute() && !GSlateIsOnFastProcessInvalidation)
+	{
+		FSlateAttributeMetaData::UpdateAttributes(*this);
+	}
+
 	bool bShouldPrepassChildren = true;
 	if (bHasCustomPrepass)
 	{
@@ -1727,7 +1749,8 @@ void SWidget::SetOnMouseLeave(FSimpleNoReplyPointerEventHandler EventHandler)
 
 void SWidget::AddMetadataInternal(const TSharedRef<ISlateMetaData>& AddMe)
 {
-	MetaData.Add(AddMe);
+	int32 Index = MetaData.Add(AddMe);
+	checkf(Index != 0 || !HasRegisteredSlateAttribute(), TEXT("The first slot is reserved for SlateAttribute"));
 
 #if UE_WITH_SLATE_DEBUG_FIND_WIDGET_REFLECTION_METADATA || UE_SLATE_TRACE_ENABLED
 	if (AddMe->IsOfType<FReflectionMetaData>())
diff --git a/Engine/Source/Runtime/SlateCore/Public/SlateGlobals.h b/Engine/Source/Runtime/SlateCore/Public/SlateGlobals.h
--- a/Engine/Source/Runtime/SlateCore/Public/SlateGlobals.h
+++ b/Engine/Source/Runtime/SlateCore/Public/SlateGlobals.h
@@ -59,10 +59,16 @@ DECLARE_STATS_GROUP_MAYBE_COMPILED_OUT(TEXT("SlateVeryVerbose"), STATGROUP_Slate
 /** Whether or not we've enabled fast widget pathing which validates paths to widgets without arranging children. */
 extern SLATECORE_API bool GSlateFastWidgetPath;
 
+/** Whether or not the SWindow can be an Invalidation Panel (use the fast path update). Normal Invalidation Panel will be deactivated. */
 extern SLATECORE_API bool GSlateEnableGlobalInvalidation;
 
+/** Whether or not we currently Painting/Updating the widget from the FastUpdate path (global invalidation). */
 extern SLATECORE_API bool GSlateIsOnFastUpdatePath;
 
+/** Whether or not we currently processing the widget invalidation from the InvalidationRoot (global invalidation). */
+extern SLATECORE_API bool GSlateIsOnFastProcessInvalidation;
+
+/** Whether or not we are currently running building the list of widget in slow path (global invalidation). */
 extern SLATECORE_API bool GSlateIsInInvalidationSlowPath;
 
 extern SLATECORE_API int32 GSlateLayoutGeneration;
diff --git a/Engine/Source/Runtime/SlateCore/Public/Types/SlateAttribute.h b/Engine/Source/Runtime/SlateCore/Public/Types/SlateAttribute.h
new file mode 100644
--- /dev/null
+++ b/Engine/Source/Runtime/SlateCore/Public/Types/SlateAttribute.h
@@ -0,0 +1,990 @@
+// Copyright Epic Games, Inc. All Rights Reserved.
+
+#pragma once
+
+#include "CoreMinimal.h"
+#include "Misc/Attribute.h"
+#include "Templates/EqualTo.h"
+#include "Widgets/InvalidateWidgetReason.h"
+
+#include <type_traits>
+
+#ifndef UE_SLATE_WITH_MEMBER_ATTRIBUTE_DEBUGGING
+	#define UE_SLATE_WITH_MEMBER_ATTRIBUTE_DEBUGGING UE_BUILD_DEBUG
+#endif
+
+/**
+ * Use TSlateAttribute when it's a SWidget member.
+ * Use TSlateManagedAttribute when it's a member inside an array or other moving structure, and the array is a SWidget member. They can only be moved (they can't be copied).
+ * For everything else, use TAttribute.
+ *
+ *
+ *
+ * In Slate TAttributes are optimized for developer efficiency.
+ * They enable widgets to poll for data instead of requiring the user to manually set the state on widgets.
+ * Attributes generally work well when performance is not a concern but break down when performance is critical (like a game UI).
+ *
+ * The invalidation system allows only widgets that have changed to perform expensive Slate layout.
+ * Bound TAttributes are incompatible with invalidation because we do not know when the data changes.
+ * Additionally TAttributes for common functionality such as visibility are called multiple times per frame and the delegate overhead for that alone is very high.
+ * TAttributes have a high memory overhead and are not cache-friendly.
+ *
+ * TSlateAttribute makes the attribute system viable for invalidation and more performance-friendly while keeping the benefits of attributes intact.
+ * TSlateAttributes are updated once per frame before the Prepass. If the cached value of the TSlateAttribute changes, then it will invalidate the widget.
+ * The TSlateAttributes are updated in the order they are defined in the SWidget (by default). TSlateManagedAttribute are updated in a random order (after TSlateAttribute).
+ * The order can be defined/override by setting a Prerequisite (see bellow for an example).
+ * The invalidation reason can be a predicate (see bellow for an example).
+ * The invalidation reason can be override per SWidget. Use override with precaution since it can break the invalidation of child widget.
+ *
+ *
+ * TSlateAttribute is not copyable and can only live inside a SWidget.
+ * For performance reasons, we do not save the extra information that would be needed to be memory save in all contexts.
+ * If you create a TSlateAttribute that can be moved, you need to use TSlateManagedAttribute.
+ * The managed version is as fast but use more memory and is less cache-friendly.
+ * Note, if you use TAttribute to change the state of a SWidget, you need to override bool ComputeVolatility() const.
+ *
+ * TSlateAttribute request a SWidget pointer. The "this" pointer should ALWAYS be used.
+ * (The SlateAttribute pointer is saved inside the SlateAttributeMetaData. The widget needs to be aware when the pointer changes.)
+ *
+ *
+ *	.h
+ *  // The new way of using attribute in your SWidget
+ *	class SMyNewWidget : public SLeafWidget
+ *	{
+ *		SLATE_DECLARE_WIDGET(SMyNewWidget, SLeafWidget)	// I is optional.
+														// If added, you need to implement void PrivateRegisterAttributes(FSlateAttributeInitializer& AttributeInitializer).
+ *
+ *		SLATE_BEGIN_ARGS(SMyNewWidget)
+ *			SLATE_ATTRIBUTE(FLinearColor, Color)
+ *		SLATE_END_ARGS()
+ *
+ *		SMyNewWidget()
+ *			: NewWay(*this, FLinearColor::White) // Always use "this" pointer. No exception. If you can't, use TAttribute.
+ *		{}
+ *
+ *		virtual int32 OnPaint( const FPaintArgs& Args, const FGeometry& AllottedGeometry, const FSlateRect& MyCullingRect, FSlateWindowElementList& OutDrawElements, int32 LayerId, const FWidgetStyle& InWidgetStyle, bool bParentEnabled ) const override
+ *		{
+ *			++LayerId;
+ *			FSlateDrawElement::MakeDebugQuad(
+ *				OutDrawElements,
+ *				LayerId,
+ *				AllottedGeometry.ToPaintGeometry(),
+ *				NewWay.Get()							// The value is already cached.
+ *			);
+ *			return LayerId;
+ *		}
+ *
+ *		// bEnabled do not invalidate Layout. We need to override the default invalidation.
+ *		virtual FVector2D ComputeDesiredSize(float) const override { return IsEnabled() ? FVector2D(10.f, 10.f) : FVector2D(20.f, 20.f); }
+ *
+ *		void Construct(const FArguments& InArgs)
+ *		{
+ *			// Assigned the attribute if needed.
+			// If the TAttribute is bound, bound the TSlateAttribute
+			// It the TAttribute is set, set the TSlateAttribute, test if the value changed and invalidate the widget.
+ *			// Else, nothing (use the previous value).
+ *			NewWays.Bind(*this, InArgs._Color); // Always use the "this" pointer.
+ *		}
+ *
+ *		// If NewWay is bounded and the values changed from the previous frame, the widget will be invalidated.
+ *		//In this case, it will only paint the widget (no layout required).
+ *		TSlateAttribute<FLinearColor, EInvalidationReason::Paint> NewWay;
+ *	};
+ *
+ *  .cpp
+ *	// This is optional. It is still a good practice to implement it since it will allow user to override the default behavior and control the update order.
+ *	// The WidgetReflector use that information.
+ *	SLATE_IMPLEMENT_WIDGET(STextBlock)
+ *	void STextBlock::PrivateRegisterAttributes(FSlateAttributeInitializer& AttributeInitializer)
+ *	{
+ *		SLATE_ADD_MEMBER_ATTRIBUTE_DEFINITION(AttributeInitializer, BoundText)
+			.SetPrerequisite("bEnabled"); // SetPrerequisite is not needed here, this is just an example to show how you could do it if needed.
+
+		//bEnabled invalidate paint, we need it to invalidate the Layout.
+		AttributeInitializer.OverrideInvalidationReason("bEnabled", EWidgetInvalidationReason::Layout | EWidgetInvalidationReason::Paint);
+ *	}
+ *
+ *
+ * // We used to do this. Keep using this method when you are not able to provide the "this" pointer to TSlateAttribute.
+ * class SMyOldWidget : public SLeafWidget
+ * {
+ *		SLATE_BEGIN_ARGS(SMyOldWidget)
+ *			SLATE_ATTRIBUTE(FLinearColor, Color)
+ *		SLATE_END_ARGS()
+ *
+ *		// Widget will be updated every frame if a TAttribute is bound. (Even if the value didn't change).
+ *		virtual bool ComputeVolatility() const override { return SLeafWidget::ComputeVolatility() || OldWay.IsBound(); }
+ *
+ *		virtual int32 OnPaint( const FPaintArgs& Args, const FGeometry& AllottedGeometry, const FSlateRect& MyCullingRect, FSlateWindowElementList& OutDrawElements, int32 LayerId, const FWidgetStyle& InWidgetStyle, bool bParentEnabled ) const override
+ *		{
+ *			++LayerId;
+ *			FSlateDrawElement::MakeDebugQuad(
+ *				OutDrawElements,
+ *				LayerId,
+ *				AllottedGeometry.ToPaintGeometry(),
+ *				OldWay.Get(FLinearColor::White)		// Fetch the value of OldWays and use White, if the value is not set.
+ *			);
+ *			return LayerId;
+ *		}
+ *
+ *		virtual FVector2D ComputeDesiredSize(float) const override { return IsEnabled() ? FVector2D(10.f, 10.f) : FVector2D(20.f, 20.f); }
+ *
+ *		void Construct(const FArguments& InArgs)
+ *		{
+ *			NewWays = InArgs._Color;
+ *		}
+ *
+ *	private:
+ *		TAttribute<FLinearColor> OldWay;
+ * };
+ */
+
+class SWidget;
+
+/** */
+struct TSlateAttributeFTextComparePredicate
+{
+	bool operator()(const FText& Lhs, const FText& Rhs) const
+	{
+		return Lhs.IdenticalTo(Rhs, ETextIdenticalModeFlags::DeepCompare | ETextIdenticalModeFlags::LexicalCompareInvariants);
+	}
+};
+
+/** */
+struct FSlateAttributeBase
+{
+	
+};
+
+/** */
+template<EInvalidateWidgetReason InvalidationReason>
+struct TSlateAttributeInvalidationReason
+{
+	static constexpr EInvalidateWidgetReason GetInvalidationReason(const SWidget&) { return InvalidationReason; }
+};
+
+/** */
+namespace SlateAttributePrivate
+{
+	/** */
+	enum class ESlateAttributeType : uint8
+	{
+		Member = 0,
+		Managed = 1,
+	};
+
+	class ISlateAttributeGetter;
+	template<typename ObjectType, typename InvalidationReasonPredicate, typename FComparePredicate, ESlateAttributeType AttributeType>
+	struct TSlateAttributeBase;
+
+
+	/** */
+	class ISlateAttributeGetter
+	{
+	public:
+		struct FUpdateAttributeResult
+		{
+			FUpdateAttributeResult(EInvalidateWidgetReason InInvalidationReason)
+				: InvalidationReason(InInvalidationReason)
+				, bInvalidationRequested(true)
+			{}
+			FUpdateAttributeResult()
+				: InvalidationReason(EInvalidateWidgetReason::None)
+				, bInvalidationRequested(false)
+			{}
+			EInvalidateWidgetReason InvalidationReason;
+			bool bInvalidationRequested;
+		};
+
+
+		virtual FUpdateAttributeResult UpdateAttribute(const SWidget& Widget) = 0;
+		virtual const FSlateAttributeBase& GetAttribute() const = 0;
+		virtual void SetAttribute(FSlateAttributeBase&) = 0;
+		virtual FDelegateHandle GetDelegateHandle() const = 0;
+		virtual ~ISlateAttributeGetter() = default;
+	};
+
+
+	/** */
+	struct SLATECORE_API FSlateAttributeImpl : public FSlateAttributeBase
+	{
+	protected:
+		bool ProtectedIsWidgetInDestructionPath(SWidget* Widget) const;
+		void ProtectedUnregisterAttribute(SWidget& Widget, ESlateAttributeType AttributeType) const;
+		void ProtectedRegisterAttribute(SWidget& Widget, ESlateAttributeType AttributeType, TUniquePtr<ISlateAttributeGetter>&& Wrapper);
+		void ProtectedInvalidateWidget(SWidget& Widget, ESlateAttributeType AttributeType, EInvalidateWidgetReason InvalidationReason) const;
+		bool ProtectedIsBound(const SWidget& Widget, ESlateAttributeType AttributeType) const;
+		FDelegateHandle ProtectedFindGetterHandle(const SWidget& Widget, ESlateAttributeType AttributeType) const;
+		bool ProtectedIsIdenticalTo(const SWidget& Widget, ESlateAttributeType AttributeType, const FSlateAttributeBase& Other, const bool bHasSameValue) const;
+		bool ProtectedIsIdenticalToAttribute(const SWidget& Widget, ESlateAttributeType AttributeType, const void* Other, const bool bHasSameValue) const;
+		void ProtectedUpdateNow(SWidget& Widget, ESlateAttributeType AttributeType);
+		void ProtectedMoveAttribute(SWidget& Widget, ESlateAttributeType AttributeType, const FSlateAttributeBase* Other);
+	};
+
+
+	/**
+	 * Attribute object
+	 * InObjectType - Type of the value to store
+	 * InInvalidationReasonPredicate - Predicate that returns the type of invalidation to do when the value changes (e.g layout or paint)
+	 *								The invalidation can be overridden per widget. (This use memory allocation. See FSlateAttributeMetadata.)
+	 * InComparePredicateType - Predicate to compare the cached value with the Getter.
+	 * bInIsExternal - The attribute life is not controlled by the SWidget.
+	 */
+	template<typename InObjectType, typename InInvalidationReasonPredicate, typename InComparePredicateType, ESlateAttributeType InAttributeType>
+	struct TSlateAttributeBase : public FSlateAttributeImpl
+	{
+	public:
+		using ObjectType = InObjectType;
+		using FInvalidationReasonPredicate = InInvalidationReasonPredicate;
+		using FGetter = typename TAttribute<ObjectType>::FGetter;
+		using FComparePredicate = InComparePredicateType;
+
+		static EInvalidateWidgetReason GetInvalidationReason(const SWidget& Widget) { return FInvalidationReasonPredicate::GetInvalidationReason(Widget); }
+		
+	private:
+		void VerifyOwningWidget(const SWidget& Widget) const
+		{
+#if UE_SLATE_WITH_MEMBER_ATTRIBUTE_DEBUGGING
+			checkf(&Widget == Debug_OwningWidget,
+				TEXT("The Owning Widget is not the same as used at construction. This will cause bad memory access."));
+#endif
+		}
+
+	public:
+		TSlateAttributeBase() = default;
+
+		TSlateAttributeBase(SWidget& Widget)
+			: Value()
+#if UE_SLATE_WITH_MEMBER_ATTRIBUTE_DEBUGGING
+			, Debug_OwningWidget(&Widget)
+#endif
+		{
+		}
+
+		TSlateAttributeBase(SWidget& Widget, const ObjectType& InValue)
+			: Value(InValue)
+#if UE_SLATE_WITH_MEMBER_ATTRIBUTE_DEBUGGING
+			, Debug_OwningWidget(&Widget)
+#endif
+		{
+		}
+
+		TSlateAttributeBase(SWidget& Widget, ObjectType&& InValue)
+			: Value(MoveTemp(InValue))
+#if UE_SLATE_WITH_MEMBER_ATTRIBUTE_DEBUGGING
+			, Debug_OwningWidget(&Widget)
+#endif
+		{
+		}
+
+		TSlateAttributeBase(SWidget& Widget, const ObjectType& InitialValue, const FGetter& Getter)
+			: Value(InitialValue)
+#if UE_SLATE_WITH_MEMBER_ATTRIBUTE_DEBUGGING
+			, Debug_OwningWidget(&Widget)
+#endif
+		{
+			if (Getter.IsBound())
+			{
+				TUniquePtr<ISlateAttributeGetter> Wrapper = MakeUniqueGetter(*this, Getter);
+				ProtectedRegisterAttribute(Widget, InAttributeType, MoveTemp(Wrapper));
+			}
+		}
+
+		TSlateAttributeBase(SWidget& Widget, const ObjectType& InitialValue, FGetter&& Getter)
+			: Value(InitialValue)
+#if UE_SLATE_WITH_MEMBER_ATTRIBUTE_DEBUGGING
+			, Debug_OwningWidget(&Widget)
+#endif
+		{
+			if (Getter.IsBound())
+			{
+				TUniquePtr<ISlateAttributeGetter> Wrapper = MakeUniqueGetter(*this, MoveTemp(Getter));
+				ProtectedRegisterAttribute(Widget, InAttributeType, MoveTemp(Wrapper));
+			}
+		}
+
+		TSlateAttributeBase(SWidget& Widget, const ObjectType& InitialValue, const TAttribute<ObjectType>& Attribute)
+			: Value((Attribute.IsSet() && !Attribute.IsBound()) ? Attribute.Get() : InitialValue)
+#if UE_SLATE_WITH_MEMBER_ATTRIBUTE_DEBUGGING
+			, Debug_OwningWidget(&Widget)
+#endif
+		{
+			if (Attribute.IsBound())
+			{
+				TUniquePtr<ISlateAttributeGetter> Wrapper = MakeUniqueGetter(*this, Attribute.GetBinding());
+				ProtectedRegisterAttribute(Widget, InAttributeType, MoveTemp(Wrapper));
+			}
+		}
+
+		TSlateAttributeBase(SWidget& Widget, ObjectType&& InitialValue, const FGetter& Getter)
+			: Value(MoveTemp(InitialValue))
+#if UE_SLATE_WITH_MEMBER_ATTRIBUTE_DEBUGGING
+			, Debug_OwningWidget(&Widget)
+#endif
+		{
+			if (Getter.IsBound())
+			{
+				TUniquePtr<ISlateAttributeGetter> Wrapper = MakeUniqueGetter(*this, Getter);
+				ProtectedRegisterAttribute(Widget, InAttributeType, MoveTemp(Wrapper));
+			}
+		}
+
+		TSlateAttributeBase(SWidget& Widget, ObjectType&& InitialValue, FGetter&& Getter)
+			: Value(MoveTemp(InitialValue))
+#if UE_SLATE_WITH_MEMBER_ATTRIBUTE_DEBUGGING
+			, Debug_OwningWidget(&Widget)
+#endif
+		{
+			if (Getter.IsBound())
+			{
+				TUniquePtr<ISlateAttributeGetter> Wrapper = MakeUniqueGetter(*this, MoveTemp(Getter));
+				ProtectedRegisterAttribute(Widget, InAttributeType, MoveTemp(Wrapper));
+			}
+		}
+
+		TSlateAttributeBase(SWidget& Widget, ObjectType&& InitialValue, const TAttribute<ObjectType>& Attribute)
+			: Value((Attribute.IsSet() && !Attribute.IsBound()) ? Attribute.Get() : MoveTemp(InitialValue))
+#if UE_SLATE_WITH_MEMBER_ATTRIBUTE_DEBUGGING
+			, Debug_OwningWidget(&Widget)
+#endif
+		{
+			if (Attribute.IsBound())
+			{
+				TUniquePtr<ISlateAttributeGetter> Wrapper = MakeUniqueGetter(*this, Attribute.GetBinding());
+				ProtectedRegisterAttribute(Widget, InAttributeType, MoveTemp(Wrapper));
+			}
+		}
+
+	public:
+		/** @return the SlateAttribute value. If the SlateAttribute is bound, value is cached only at the end of the frame. */ 
+		const ObjectType& Get() const
+		{
+			return Value;
+		}
+		
+		/** Update the cached value and invalidate the widget if needed. */
+		void UpdateNow(SWidget& Widget)
+		{
+			VerifyOwningWidget(Widget);
+			ProtectedUpdateNow(Widget, InAttributeType);
+		}
+
+	public:
+		/** Unbind the SlateAttribute and set its value. It may invalidate the Widget if the value is different. */
+		void Set(SWidget& Widget, const ObjectType& NewValue)
+		{
+			VerifyOwningWidget(Widget);
+			ProtectedUnregisterAttribute(Widget, InAttributeType);
+
+			const bool bIsIdentical = FComparePredicate{}.operator()(Value, NewValue);
+			if (!bIsIdentical)
+			{
+				Value = NewValue;
+				ProtectedInvalidateWidget(Widget, InAttributeType, GetInvalidationReason(Widget));
+			}
+		}
+
+		/** Unbind the SlateAttribute and set its value. It may invalidate the Widget if the value is different. */
+		void Set(SWidget& Widget, ObjectType&& NewValue)
+		{
+			VerifyOwningWidget(Widget);
+			ProtectedUnregisterAttribute(Widget, InAttributeType);
+
+			const bool bIsIdentical = FComparePredicate{}.operator()(Value, NewValue);
+			if (!bIsIdentical)
+			{
+				Value = MoveTemp(NewValue);
+				ProtectedInvalidateWidget(Widget, InAttributeType, GetInvalidationReason(Widget));
+			}
+		}
+
+		/** Bind the SlateAttribute to the Getter function. The SlateAttribute will be updated every frame from the Getter. */
+		void Bind(SWidget& Widget, const FGetter& Getter)
+		{
+			VerifyOwningWidget(Widget);
+			if (Getter.IsBound())
+			{
+				const FDelegateHandle PreviousGetterHandle = ProtectedFindGetterHandle(Widget, InAttributeType);
+				if (PreviousGetterHandle != Getter.GetHandle())
+				{
+					TUniquePtr<ISlateAttributeGetter> Wrapper = MakeUniqueGetter(*this, Getter);
+					ProtectedRegisterAttribute(Widget, InAttributeType, MoveTemp(Wrapper));
+				}
+			}
+			else
+			{
+				ProtectedUnregisterAttribute(Widget, InAttributeType);
+			}
+		}
+
+		/** Bind the SlateAttribute to the Getter function. The SlateAttribute will be updated every frame from the Getter. */
+		void Bind(SWidget& Widget, FGetter&& Getter)
+		{
+			VerifyOwningWidget(Widget);
+			if (Getter.IsBound())
+			{
+				const FDelegateHandle PreviousGetterHandle = ProtectedFindGetterHandle(Widget, InAttributeType);
+				if (PreviousGetterHandle != Getter.GetHandle())
+				{
+					TUniquePtr<ISlateAttributeGetter> Wrapper = MakeUniqueGetter(*this, MoveTemp(Getter));
+					ProtectedRegisterAttribute(Widget, InAttributeType, MoveTemp(Wrapper));
+				}
+			}
+			else
+			{
+				ProtectedUnregisterAttribute(Widget, InAttributeType);
+			}
+		}
+
+		/** Bind the SlateAttribute to the Getter function. The SlateAttribute will be updated every frame from the Getter. */
+		template<typename WidgetType, typename U = typename std::enable_if<std::is_base_of<SWidget, WidgetType>::value>::type>
+		void Bind(WidgetType& Widget, typename FGetter::template TSPMethodDelegate_Const<WidgetType>::FMethodPtr MethodPtr)
+		{
+			Bind(Widget, FGetter::CreateSP(&Widget, MethodPtr));
+		}
+
+		/**
+		 * Bind the SlateAttribute to the Attribute Getter function (if it exist). The SlateAttribute will be updated every frame from the Getter.
+		 * Or set the value if the Attribute is bound.
+		 * This will Unbind any previously bound getter.
+		 * It may Invalidate the Widget. @see Set
+		 */
+		void Assign(SWidget& Widget, const TAttribute<ObjectType>& OtherAttribute)
+		{
+			VerifyOwningWidget(Widget);
+			if (OtherAttribute.IsBound())
+			{
+				const FDelegateHandle PreviousGetterHandle = ProtectedFindGetterHandle(Widget, InAttributeType);
+				if (PreviousGetterHandle != OtherAttribute.GetBinding().GetHandle())
+				{
+					TUniquePtr<ISlateAttributeGetter> Wrapper = MakeUniqueGetter(*this, OtherAttribute.GetBinding());
+					ProtectedRegisterAttribute(Widget, InAttributeType, MoveTemp(Wrapper));
+				}
+			}
+			else if (OtherAttribute.IsSet())
+			{
+				Set(Widget, OtherAttribute.Get());
+			}
+			else
+			{
+				ProtectedUnregisterAttribute(Widget, InAttributeType);
+			}
+		}
+
+		/**
+		 * Bind the SlateAttribute to the Attribute Getter function (if it exist). The SlateAttribute will be updated every frame from the Getter.
+		 * Or set the value if the Attribute is bound. If the Attribute is not set the DefaultValue will be used.
+		 * This will Unbind any previously bound getter.
+		 * It may Invalidate the Widget. @see Set
+		 */
+		void Assign(SWidget& Widget, const TAttribute<ObjectType>& OtherAttribute, const ObjectType& DefaultValue)
+		{
+			VerifyOwningWidget(Widget);
+			if (OtherAttribute.IsBound())
+			{
+				const FDelegateHandle PreviousGetterHandle = ProtectedFindGetterHandle(Widget, InAttributeType);
+				if (PreviousGetterHandle != OtherAttribute.GetBinding().GetHandle())
+				{
+					TUniquePtr<ISlateAttributeGetter> Wrapper = MakeUniqueGetter(*this, OtherAttribute.GetBinding());
+					ProtectedRegisterAttribute(Widget, InAttributeType, MoveTemp(Wrapper));
+				}
+			}
+			else if (OtherAttribute.IsSet())
+			{
+				Set(Widget, OtherAttribute.Get());
+			}
+			else
+			{
+				Set(Widget, DefaultValue);
+			}
+		}
+
+		/**
+		 * Bind the SlateAttribute to the Attribute Getter function (if it exist). The SlateAttribute will be updated every frame from the Getter.
+		 * Or set the value if the Attribute is bound. If the Attribute is not set the DefaultValue will be used.
+		 * This will Unbind any previously bound getter.
+		 * It may Invalidate the Widget. @see Set
+		 */
+		void Assign(SWidget& Widget, const TAttribute<ObjectType>& OtherAttribute, ObjectType&& DefaultValue)
+		{
+			VerifyOwningWidget(Widget);
+			if (OtherAttribute.IsBound())
+			{
+				const FDelegateHandle PreviousGetterHandle = ProtectedFindGetterHandle(Widget, InAttributeType);
+				if (PreviousGetterHandle != OtherAttribute.GetBinding().GetHandle())
+				{
+					TUniquePtr<ISlateAttributeGetter> Wrapper = MakeUniqueGetter(*this, OtherAttribute.GetBinding());
+					ProtectedRegisterAttribute(Widget, InAttributeType, MoveTemp(Wrapper));
+				}
+			}
+			else if (OtherAttribute.IsSet())
+			{
+				Set(Widget, OtherAttribute.Get());
+			}
+			else
+			{
+				Set(Widget, MoveTemp(DefaultValue));
+			}
+		}
+
+		/**
+		 * Remove the Getter function.
+		 * The Slate Attribute will not be updated anymore and will keep its current cached value.
+		 */
+		void Unbind(SWidget& Widget)
+		{
+			VerifyOwningWidget(Widget);
+			ProtectedUnregisterAttribute(Widget, InAttributeType);
+		}
+
+public:
+		/** @return True if the SlateAttribute is bound to a getter function. */
+		bool IsBound(const SWidget& Widget) const
+		{
+			VerifyOwningWidget(Widget);
+			return ProtectedIsBound(Widget, InAttributeType);
+		}
+
+		/** @return True if they have the same Getter or the same value. */
+		bool IsIdenticalTo(const SWidget& Widget, const TSlateAttributeBase& Other) const
+		{
+			VerifyOwningWidget(Widget);
+			FDelegateHandle ThisDelegateHandle = ProtectedFindGetterHandle(Widget, InAttributeType);
+			FDelegateHandle OthersDelegateHandle = ProtectedFindGetterHandle(Widget, InAttributeType);
+			if (ThisDelegateHandle == OthersDelegateHandle)
+			{
+				if (ThisDelegateHandle.IsValid())
+				{
+					return true;
+				}
+				return Get() == Other.Get();
+			}
+			return false;
+		}
+
+		/**  if they have the same Getter or, if the Attribute is set, the same value. */
+		bool IsIdenticalTo(const SWidget& Widget, const TAttribute<ObjectType>& Other) const
+		{
+			VerifyOwningWidget(Widget);
+			FDelegateHandle ThisDelegateHandle = ProtectedFindGetterHandle(Widget, InAttributeType);
+			if (Other.IsBound())
+			{
+				return Other.GetBinding().GetHandle() == ThisDelegateHandle;
+			}
+			return !ThisDelegateHandle.IsValid() && Other.IsSet() && Get() == Other.Get();
+		}
+
+	private:
+		template<typename SlateAttributeType>
+		class FSlateAttributeGetterWrapper : public ISlateAttributeGetter
+		{
+		public:
+			using ObjectType = typename SlateAttributeType::ObjectType;
+			using FGetter = typename TAttribute<ObjectType>::FGetter;
+			using FComparePredicate = typename SlateAttributeType::FComparePredicate;
+
+			FSlateAttributeGetterWrapper() = delete;
+			FSlateAttributeGetterWrapper(const FSlateAttributeGetterWrapper&) = delete;
+			FSlateAttributeGetterWrapper& operator= (const FSlateAttributeGetterWrapper&) = delete;
+			virtual ~FSlateAttributeGetterWrapper() = default;
+
+		public:
+			FSlateAttributeGetterWrapper(SlateAttributeType& InOwningAttribute, const FGetter& InGetterDelegate)
+				: Getter(InGetterDelegate)
+				, Attribute(&InOwningAttribute)
+			{
+			}
+			FSlateAttributeGetterWrapper(SlateAttributeType& InOwningAttribute, FGetter&& InGetterDelegate)
+				: Getter(MoveTemp(InGetterDelegate))
+				, Attribute(&InOwningAttribute)
+			{
+			}
+
+			virtual FUpdateAttributeResult UpdateAttribute(const SWidget& Widget) override
+			{
+				ObjectType NewValue = Getter.Execute();
+
+				if (!(FComparePredicate{}.operator()(Attribute->Value, NewValue)))
+				{
+					// Set the value on the widget
+					Attribute->Value = MoveTemp(NewValue);
+					return Attribute->GetInvalidationReason(Widget);
+				}
+				return FUpdateAttributeResult();
+			}
+
+			virtual const FSlateAttributeBase& GetAttribute() const override
+			{
+				return *Attribute;
+			}
+
+			virtual void SetAttribute(FSlateAttributeBase& InBase) override
+			{
+				Attribute = &static_cast<SlateAttributeType&>(InBase);
+			}
+
+			virtual FDelegateHandle GetDelegateHandle() const override
+			{
+				return Getter.GetHandle();
+			}
+
+		private:
+			/** Getter function to fetch the new value of the SlateAttribute. */
+			FGetter Getter;
+			/** The SlateAttribute of the SWidget owning the value. */
+			SlateAttributeType* Attribute;
+		};
+
+		static TUniquePtr<ISlateAttributeGetter> MakeUniqueGetter(TSlateAttributeBase& Attribute, const FGetter& Getter)
+		{
+			return MakeUnique<FSlateAttributeGetterWrapper<TSlateAttributeBase>>(Attribute, Getter);
+		}
+
+		static TUniquePtr<ISlateAttributeGetter> MakeUniqueGetter(TSlateAttributeBase& Attribute, FGetter&& Getter)
+		{
+			return MakeUnique<FSlateAttributeGetterWrapper<TSlateAttributeBase>>(Attribute, MoveTemp(Getter));
+		}
+
+	private:
+		ObjectType Value;
+
+	protected:
+#if UE_SLATE_WITH_MEMBER_ATTRIBUTE_DEBUGGING
+		SWidget* Debug_OwningWidget;
+#endif
+	};
+
+
+	/**
+	 *
+	 */
+	template<typename InObjectType, typename InInvalidationReasonPredicate, typename InComparePredicate = TEqualTo<>>
+	struct TSlateMemberAttribute : public TSlateAttributeBase<InObjectType, InInvalidationReasonPredicate, InComparePredicate, ESlateAttributeType::Member>
+	{
+	private:
+		using Super = TSlateAttributeBase<InObjectType, InInvalidationReasonPredicate, InComparePredicate, ESlateAttributeType::Member>;
+
+		//using TSlateAttributeBase<InObjectType, InInvalidationReasonPredicate, InComparePredicate, ESlateAttributeType::Member>::TSlateAttributeBase;
+
+		template<typename WidgetType, typename U = typename std::enable_if<std::is_base_of<SWidget, WidgetType>::value>::type>
+		static void VerifyAttributeAddress(WidgetType& Widget, TSlateMemberAttribute* Self)
+		{
+			checkf((UPTRINT)Self >= (UPTRINT)&Widget && (UPTRINT)Self < (UPTRINT)&Widget + sizeof(WidgetType),
+				TEXT("Use TAttribute or TSlateManagedAttribute instead. See SlateAttribute.h for more info."));
+		}
+
+	public:
+		using FGetter = typename Super::FGetter;
+		using ObjectType = typename Super::ObjectType;
+
+	public:
+		//~ You can only register Attribute that are defined in a SWidget.
+		//~ Use the constructor with Widget pointer.
+		//~ See documentation on top of this file for more info.
+		TSlateMemberAttribute() = delete;
+		TSlateMemberAttribute(const TSlateMemberAttribute&) = delete;
+		TSlateMemberAttribute(TSlateMemberAttribute&&) = delete;
+		TSlateMemberAttribute& operator=(const TSlateMemberAttribute&) = delete;
+		TSlateMemberAttribute& operator=(TSlateMemberAttribute&&) = delete;
+
+#if UE_SLATE_WITH_MEMBER_ATTRIBUTE_DEBUGGING
+		~TSlateMemberAttribute()
+		{
+			/**
+			 * The parent should now be destroyed and the shared pointer should invalidate.
+			 * If you hit the check, that means the TSlateAttribute is not a variable member of SWidget.
+			 * It will introduced bad memory access.
+			 * See documentation above.
+			*/
+			checkf(Super::ProtectedIsWidgetInDestructionPath(Super::Debug_OwningWidget), TEXT("The Owning widget should be invalid."));
+		}
+#endif
+
+		template<typename WidgetType, typename U = typename std::enable_if<std::is_base_of<SWidget, WidgetType>::value>::type>
+		TSlateMemberAttribute(WidgetType& Widget)
+			: Super(Widget)
+		{
+			VerifyAttributeAddress(Widget, this);
+		}
+
+		template<typename WidgetType, typename U = typename std::enable_if<std::is_base_of<SWidget, WidgetType>::value>::type>
+		TSlateMemberAttribute(WidgetType& Widget, const ObjectType& InValue)
+			: Super(Widget, InValue)
+		{
+			VerifyAttributeAddress(Widget, this);
+		}
+
+		template<typename WidgetType, typename U = typename std::enable_if<std::is_base_of<SWidget, WidgetType>::value>::type>
+		TSlateMemberAttribute(WidgetType& Widget, ObjectType&& InValue)
+			: Super(Widget, MoveTemp(InValue))
+		{
+			VerifyAttributeAddress(Widget, this);
+		}
+
+		template<typename WidgetType, typename U = typename std::enable_if<std::is_base_of<SWidget, WidgetType>::value>::type>
+		TSlateMemberAttribute(WidgetType& Widget, const ObjectType& InitialValue, const FGetter& Getter)
+			: Super(Widget, InitialValue, Getter)
+		{
+			VerifyAttributeAddress(Widget, this);
+		}
+
+		template<typename WidgetType, typename U = typename std::enable_if<std::is_base_of<SWidget, WidgetType>::value>::type>
+		TSlateMemberAttribute(WidgetType& Widget, const ObjectType& InitialValue, FGetter&& Getter)
+			: Super(Widget, InitialValue, MoveTemp(Getter))
+		{
+			VerifyAttributeAddress(Widget, this);
+		}
+
+		template<typename WidgetType, typename U = typename std::enable_if<std::is_base_of<SWidget, WidgetType>::value>::type>
+		TSlateMemberAttribute(WidgetType& Widget, const ObjectType& InitialValue, const TAttribute<ObjectType>& Attribute)
+			: Super(Widget, InitialValue, Attribute)
+		{
+			VerifyAttributeAddress(Widget, this);
+		}
+
+		template<typename WidgetType, typename U = typename std::enable_if<std::is_base_of<SWidget, WidgetType>::value>::type>
+		TSlateMemberAttribute(WidgetType& Widget, ObjectType&& InitialValue, const FGetter& Getter)
+			: Super(Widget, MoveTemp(InitialValue), Getter)
+		{
+			VerifyAttributeAddress(Widget, this);
+		}
+
+		template<typename WidgetType, typename U = typename std::enable_if<std::is_base_of<SWidget, WidgetType>::value>::type>
+		TSlateMemberAttribute(WidgetType& Widget, ObjectType&& InitialValue, FGetter&& Getter)
+			: Super(Widget, MoveTemp(InitialValue), MoveTemp(Getter))
+		{
+			VerifyAttributeAddress(Widget, this);
+		}
+
+		template<typename WidgetType, typename U = typename std::enable_if<std::is_base_of<SWidget, WidgetType>::value>::type>
+		TSlateMemberAttribute(WidgetType& Widget, ObjectType&& InitialValue, const TAttribute<ObjectType>& Attribute)
+			: Super(Widget, MoveTemp(InitialValue), Attribute)
+		{
+			VerifyAttributeAddress(Widget, this);
+		}
+	};
+
+
+	/**
+	 *
+	 */
+	template<typename InObjectType, typename InInvalidationReasonPredicate, typename InComparePredicate = TEqualTo<>>
+	struct TSlateManagedAttribute : protected TSlateAttributeBase<InObjectType, InInvalidationReasonPredicate, InComparePredicate, ESlateAttributeType::Managed>
+	{
+	private:
+		using Super = TSlateAttributeBase<InObjectType, InInvalidationReasonPredicate, InComparePredicate, ESlateAttributeType::Managed>;
+
+	public:
+		using ObjectType = typename Super::ObjectType;
+		using FInvalidationReasonPredicate = typename Super::FInvalidationReasonPredicate;
+		using FGetter = typename Super::FGetter;
+		using FComparePredicate = typename Super::FComparePredicate;
+
+		static EInvalidateWidgetReason GetInvalidationReason(const SWidget& Widget) { return Super::GetInvalidationReason(Widget); }
+
+	public:
+		TSlateManagedAttribute() = delete;
+		TSlateManagedAttribute(const TSlateManagedAttribute&) = delete;
+		TSlateManagedAttribute& operator=(const TSlateManagedAttribute&) = delete;
+
+		~TSlateManagedAttribute()
+		{
+			Unbind();
+		}
+
+		TSlateManagedAttribute(TSlateManagedAttribute&& Other)
+			: Super(MoveTemp(Other))
+			, ManagedWidget(MoveTemp(Other.ManagedWidget))
+		{
+			if (TSharedPtr<SWidget> Pin = ManagedWidget.Pin())
+			{
+				FSlateAttributeImpl::ProtectedMoveAttribute(*Pin.Get(), ESlateAttributeType::Managed, &Other);
+			}
+		}
+
+		TSlateManagedAttribute& operator=(TSlateManagedAttribute&& Other)
+		{
+			Super::operator=(MoveTemp(Other));
+			ManagedWidget = MoveTemp(Other.ManagedWidget);
+
+			if (TSharedPtr<SWidget> Pin = ManagedWidget.Pin())
+			{
+				FSlateAttributeImpl::ProtectedMoveAttribute(*Pin.Get(), ESlateAttributeType::Managed, &Other);
+			}
+			return *this;
+		}
+
+		TSlateManagedAttribute(TSharedRef<SWidget> Widget)
+			: Super(Widget.Get())
+			, ManagedWidget(Widget)
+		{ }
+
+		TSlateManagedAttribute(TSharedRef<SWidget> Widget, const ObjectType& InValue)
+			: Super(Widget.Get(), InValue)
+			, ManagedWidget(Widget)
+		{
+		}
+
+		TSlateManagedAttribute(TSharedRef<SWidget> Widget, ObjectType&& InValue)
+			: Super(Widget.Get(), MoveTemp(InValue))
+			, ManagedWidget(Widget)
+		{
+		}
+
+		TSlateManagedAttribute(TSharedRef<SWidget> Widget, const ObjectType& InitialValue, const FGetter& Getter)
+			: Super(Widget.Get(), InitialValue, Getter)
+			, ManagedWidget(Widget)
+		{
+		}
+
+		TSlateManagedAttribute(TSharedRef<SWidget> Widget, const ObjectType& InitialValue, FGetter&& Getter)
+			: Super(Widget.Get(), InitialValue, MoveTemp(Getter))
+			, ManagedWidget(Widget)
+		{
+		}
+
+		TSlateManagedAttribute(TSharedRef<SWidget> Widget, const ObjectType& InitialValue, const TAttribute<ObjectType>& Attribute)
+			: Super(Widget.Get(), InitialValue, Attribute)
+			, ManagedWidget(Widget)
+		{
+		}
+
+		TSlateManagedAttribute(TSharedRef<SWidget> Widget, ObjectType&& InitialValue, const FGetter& Getter)
+			: Super(Widget.Get(), MoveTemp(InitialValue), Getter)
+			, ManagedWidget(Widget)
+		{
+		}
+
+		TSlateManagedAttribute(TSharedRef<SWidget> Widget, ObjectType&& InitialValue, FGetter&& Getter)
+			: Super(Widget.Get(), MoveTemp(InitialValue), MoveTemp(Getter))
+			, ManagedWidget(Widget)
+		{
+		}
+
+		TSlateManagedAttribute(TSharedRef<SWidget> Widget, ObjectType&& InitialValue, const TAttribute<ObjectType>& Attribute)
+			: Super(Widget.Get(), MoveTemp(InitialValue), Attribute)
+			, ManagedWidget(Widget)
+		{
+		}
+
+	public:
+		const ObjectType& Get() const
+		{
+			return Super::Get();
+		}
+
+		void UpdateNow()
+		{
+			if (TSharedPtr<SWidget> Pin = ManagedWidget.Pin())
+			{
+				Super::UpdateNow(*Pin.Get());
+			}
+		}
+
+	public:
+		void Set(const ObjectType& NewValue)
+		{
+			if (TSharedPtr<SWidget> Pin = ManagedWidget.Pin())
+			{
+				Super::Set(*Pin.Get(), NewValue);
+			}
+		}
+
+		void Set(ObjectType&& NewValue)
+		{
+			if (TSharedPtr<SWidget> Pin = ManagedWidget.Pin())
+			{
+				Super::Set(*Pin.Get(), MoveTemp(NewValue));
+			}
+		}
+
+		void Bind(const FGetter& Getter)
+		{
+			if (TSharedPtr<SWidget> Pin = ManagedWidget.Pin())
+			{
+				Super::Bind(*Pin.Get(), Getter);
+			}
+		}
+
+		void Bind(FGetter&& Getter)
+		{
+			if (TSharedPtr<SWidget> Pin = ManagedWidget.Pin())
+			{
+				Super::Bind(*Pin.Get(), MoveTemp(Getter));
+			}
+		}
+
+		template<typename WidgetType, typename U = typename std::enable_if<std::is_base_of<SWidget, WidgetType>::value>::type>
+		void Bind(typename FGetter::template TSPMethodDelegate_Const<WidgetType>::FMethodPtr MethodPtr)
+		{
+			if (TSharedPtr<SWidget> Pin = ManagedWidget.Pin())
+			{
+				Super::Bind(*Pin.Get(), FGetter::CreateSP(Pin, MethodPtr));
+			}
+		}
+
+		void Assign(const TAttribute<ObjectType>& OtherAttribute)
+		{
+			if (TSharedPtr<SWidget> Pin = ManagedWidget.Pin())
+			{
+				Super::Assign(*Pin.Get(), OtherAttribute);
+			}
+		}
+
+		void Assign(const TAttribute<ObjectType>& OtherAttribute, const ObjectType& DefaultValue)
+		{
+			if (TSharedPtr<SWidget> Pin = ManagedWidget.Pin())
+			{
+				Super::Assign(*Pin.Get(), OtherAttribute, DefaultValue);
+			}
+		}
+
+		void Assign(const TAttribute<ObjectType>& OtherAttribute, ObjectType&& DefaultValue)
+		{
+			if (TSharedPtr<SWidget> Pin = ManagedWidget.Pin())
+			{
+				Super::Assign(*Pin.Get(), OtherAttribute, MoveTemp(DefaultValue));
+			}
+		}
+
+		void Unbind()
+		{
+			if (TSharedPtr<SWidget> Pin = ManagedWidget.Pin())
+			{
+				Super::Unbind(*Pin.Get());
+			}
+		}
+
+	public:
+		bool IsBound() const
+		{
+			if (TSharedPtr<SWidget> Pin = ManagedWidget.Pin())
+			{
+				return Super::IsBound(*Pin.Get());
+			}
+			return false;
+		}
+
+		bool IsIdenticalTo(const TSlateManagedAttribute& Other) const
+		{
+			if (TSharedPtr<SWidget> Pin = ManagedWidget.Pin())
+			{
+				return Super::IsIdenticalTo(*Pin.Get(), Other);
+			}
+			return false;
+		}
+
+		bool IsIdenticalTo(const TAttribute<ObjectType>& Other) const
+		{
+			if (TSharedPtr<SWidget> Pin = ManagedWidget.Pin())
+			{
+				return Super::IsIdenticalTo(*Pin.Get(), Other);
+			}
+			return false;
+		}
+
+	private:
+		TWeakPtr<SWidget> ManagedWidget;
+	};
+
+} // SlateAttributePrivate
diff --git a/Engine/Source/Runtime/SlateCore/Public/Types/SlateAttributeDescriptor.h b/Engine/Source/Runtime/SlateCore/Public/Types/SlateAttributeDescriptor.h
new file mode 100644
--- /dev/null
+++ b/Engine/Source/Runtime/SlateCore/Public/Types/SlateAttributeDescriptor.h
@@ -0,0 +1,91 @@
+// Copyright Epic Games, Inc. All Rights Reserved.
+
+#pragma once
+
+#include "CoreMinimal.h"
+#include "Misc/Attribute.h"
+#include "Widgets/InvalidateWidgetReason.h"
+
+#include <type_traits>
+
+class FSlateWidgetClassData;
+
+
+class SLATECORE_API FSlateAttributeDescriptor
+{
+public:
+	using OffsetType = uint32; //UPTRINT
+
+	static uint32 DefaultSortOrder(OffsetType Offset) { return Offset * 10; }
+
+	struct FAttribute
+	{
+		FName Name;
+		OffsetType Offset;
+		FName Prerequisite;
+		FName Dependency;
+		uint32 SortOrder;
+		TAttribute<EInvalidateWidgetReason> InvalidationReason;
+		bool bIsMemberAttribute;
+	};
+
+	struct SLATECORE_API FInitializer
+	{
+	private:
+		friend FSlateWidgetClassData;
+		FInitializer(FSlateAttributeDescriptor& InDescriptor);
+		FInitializer(FSlateAttributeDescriptor& InDescriptor, const FSlateAttributeDescriptor& ParentDescriptor);
+		FInitializer(const FInitializer&) = delete;
+		FInitializer& operator=(const FInitializer&) = delete;
+
+	public:
+		~FInitializer();
+
+		struct SLATECORE_API FAttributeEntry
+		{
+			FAttributeEntry(FSlateAttributeDescriptor& Descriptor, int32 InAttributeIndex);
+
+			/** Assign an order in which the attributes should be updated. */
+			FAttributeEntry& SetPrerequisite(FName Prerequisite);
+			/** Update the property every. */
+			FAttributeEntry& UpdateEveryFrame();
+			/** The property only need to be poll when the dependency changes. */
+			FAttributeEntry& UpdateDependency(FName Dependency);
+
+		private:
+			FSlateAttributeDescriptor& Descriptor;
+			int32 AttributeIndex;
+		};
+
+		FAttributeEntry AddMemberAttribute(FName AttributeName, OffsetType Offset);
+
+		FAttributeEntry AddMemberAttribute(FName AttributeName, OffsetType Offset, TAttribute<EInvalidateWidgetReason> ReasonGetter);
+
+		void OverrideInvalidationReason(FName AttributeName, TAttribute<EInvalidateWidgetReason> Reason);
+
+		void SetPrerequisite(FName AttributeName, FName Prerequisite);
+
+	private:
+		FSlateAttributeDescriptor& Descriptor;
+	};
+
+	FAttribute const* FindAttribute(FName AttributeName) const;
+
+	FAttribute const* FindMemberAttribute(OffsetType AttributeOffset) const;
+
+private:
+	FAttribute* FindAttribute(FName AttributeName);
+
+	FInitializer::FAttributeEntry AddMemberAttribute(FName AttributeName, OffsetType Offset, TAttribute<EInvalidateWidgetReason>&& ReasonGetter);
+	void OverrideInvalidationReason(FName AttributeName, TAttribute<EInvalidateWidgetReason>&& ReasonGetter);
+	void SetDependency(FAttribute& Attribute, FName Dependency);
+	void SetPrerequisite(FAttribute& Attribute, FName Prerequisite);
+
+private:
+	TArray<FAttribute> Attributes;
+};
+
+#define SLATE_ADD_MEMBER_ATTRIBUTE_DEFINITION(_Initializer, _Property) {_Initializer.AddMemberAttribute(GET_MEMBER_NAME_CHECKED(PrivateThisType, _Property), STRUCT_OFFSET(PrivateThisType, _Property));}
+#define SLATE_ADD_MEMBER_ATTRIBUTE_DEFINITION_REASON(_Initializer, _Property, _Reason) {_Initializer.AddMemberAttribute(GET_MEMBER_NAME_CHECKED(PrivateThisType, _Property), STRUCT_OFFSET(PrivateThisType, _Property), _Reason));}
+
+using FSlateAttributeInitializer = FSlateAttributeDescriptor::FInitializer;
diff --git a/Engine/Source/Runtime/SlateCore/Public/Types/SlateAttributeMetaData.h b/Engine/Source/Runtime/SlateCore/Public/Types/SlateAttributeMetaData.h
new file mode 100644
--- /dev/null
+++ b/Engine/Source/Runtime/SlateCore/Public/Types/SlateAttributeMetaData.h
@@ -0,0 +1,89 @@
+ //Copyright Epic Games, Inc. All Rights Reserved.
+
+#pragma once
+
+#include "CoreMinimal.h"
+#include "Algo/BinarySearch.h"
+#include "Types/ISlateMetaData.h"
+#include "Types/SlateAttribute.h"
+#include "Widgets/InvalidateWidgetReason.h"
+
+
+class SWidget;
+
+/** */
+class SLATECORE_API FSlateAttributeMetaData : public ISlateMetaData
+{
+	friend SlateAttributePrivate::FSlateAttributeImpl;
+
+public:
+	SLATE_METADATA_TYPE(FSlateAttributeMetaData, ISlateMetaData);
+
+	FSlateAttributeMetaData() = default;
+	FSlateAttributeMetaData(const FSlateAttributeMetaData&) = delete;
+	FSlateAttributeMetaData& operator=(const FSlateAttributeMetaData&) = delete;
+
+public:
+	static FSlateAttributeMetaData* FindMetaData(const SWidget& OwningWidget);
+	static void UpdateAttributes(SWidget& OwningWidget);
+
+public:
+	bool IsBound(const FSlateAttributeBase& Attribute) const
+	{
+		return IndexOfAttribute(Attribute) != INDEX_NONE;
+	}
+
+	int32 RegisteredNum() const { return Attributes.Num(); }
+
+private:
+	using ESlateAttributeType = SlateAttributePrivate::ESlateAttributeType;
+	using ISlateAttributeGetter = SlateAttributePrivate::ISlateAttributeGetter;
+	static void RegisterAttribute(SWidget& OwningWidget, FSlateAttributeBase& Attribute, ESlateAttributeType AttributeType, TUniquePtr<ISlateAttributeGetter>&& Wrapper);
+	static bool UnregisterAttribute(SWidget& OwningWidget, const FSlateAttributeBase& Attribute);
+	static void InvalidateWidget(SWidget& OwningWidget, const FSlateAttributeBase& Attribute, EInvalidateWidgetReason Reason);
+	static void UpdateAttribute(SWidget& OwningWidget, FSlateAttributeBase& Attribute);
+	static bool IsAttributeBound(const SWidget& OwningWidget, const FSlateAttributeBase& Attribute);
+	static FDelegateHandle GetAttributeGetterHandle(const SWidget& OwningWidget, const FSlateAttributeBase& Attribute);
+	static void MoveAttribute(const SWidget& OwningWidget, FSlateAttributeBase& NewAttribute, ESlateAttributeType AttributeType, const FSlateAttributeBase* PreviousAttribute);
+
+private:
+	void RegisterAttributeImpl(SWidget& OwningWidget, FSlateAttributeBase& Attribute, ESlateAttributeType AttributeType, TUniquePtr<ISlateAttributeGetter>&& Getter);
+	bool UnregisterAttributeImpl(const FSlateAttributeBase& Attribute);
+
+private:
+	int32 IndexOfAttribute(const FSlateAttributeBase& Attribute) const
+	{
+		const FSlateAttributeBase* AttributePtr = &Attribute;
+		return Attributes.IndexOfByPredicate([AttributePtr](const FGetterItem& Item) { return Item.Attribute == AttributePtr; });
+	}
+
+	struct FGetterItem
+	{
+		FGetterItem() = default;
+		FGetterItem(const FGetterItem&) = delete;
+		FGetterItem& operator=(const FGetterItem&) = delete;
+		FGetterItem(FSlateAttributeBase* InAttribute, uint32 InSortOrder, TUniquePtr<SlateAttributePrivate::ISlateAttributeGetter>&& InGetter)
+			: Attribute(InAttribute)
+			, Getter(MoveTemp(InGetter))
+			, SortOrder(InSortOrder)
+			, OverrideInvalidateReason()
+		{ }
+		FGetterItem(FSlateAttributeBase* InAttribute, uint32 InSortOrder, TUniquePtr<SlateAttributePrivate::ISlateAttributeGetter> && InGetter, TAttribute<EInvalidateWidgetReason> InInvalidaitonReason)
+			: Attribute(InAttribute)
+			, Getter(MoveTemp(InGetter))
+			, SortOrder(InSortOrder)
+			, OverrideInvalidateReason(MoveTemp(InInvalidaitonReason))
+		{ }
+		FSlateAttributeBase* Attribute;
+		TUniquePtr<SlateAttributePrivate::ISlateAttributeGetter> Getter;
+		uint32 SortOrder;
+		TAttribute<EInvalidateWidgetReason> OverrideInvalidateReason;
+
+		bool operator<(const FGetterItem& Other) const
+		{
+			return SortOrder < Other.SortOrder;
+		}
+	};
+
+	TArray<FGetterItem, TInlineAllocator<4>> Attributes;
+};
diff --git a/Engine/Source/Runtime/SlateCore/Public/Widgets/SWidget.h b/Engine/Source/Runtime/SlateCore/Public/Widgets/SWidget.h
--- a/Engine/Source/Runtime/SlateCore/Public/Widgets/SWidget.h
+++ b/Engine/Source/Runtime/SlateCore/Public/Widgets/SWidget.h
@@ -26,8 +26,10 @@
 #include "Textures/SlateShaderResource.h"
 #include "SlateGlobals.h"
 #include "Types/PaintArgs.h"
+#include "Types/SlateAttribute.h"
 #include "FastUpdate/WidgetProxy.h"
 #include "InvalidateWidgetReason.h"
+#include "Widgets/SlateControlledConstruction.h"
 #include "Widgets/Accessibility/SlateWidgetAccessibleTypes.h"
 
 class FActiveTimerHandle;
@@ -56,50 +58,6 @@ DECLARE_DWORD_ACCUMULATOR_STAT_EXTERN(TEXT("Total Widgets"), STAT_SlateTotalWidg
 DECLARE_MEMORY_STAT_EXTERN(TEXT("SWidget Total Allocated Size"), STAT_SlateSWidgetAllocSize, STATGROUP_SlateMemory, SLATECORE_API);
 
 
-namespace SharedPointerInternals
-{
-	template <typename ObjectType>
-	class TIntrusiveReferenceController;
-}
-
-class SLATECORE_API FSlateControlledConstruction
-{
-public:
-	FSlateControlledConstruction(){}
-	virtual ~FSlateControlledConstruction(){}
-	
-private:
-	/** UI objects cannot be copy-constructed */
-	FSlateControlledConstruction(const FSlateControlledConstruction& Other) = delete;
-	
-	/** UI objects cannot be copied. */
-	void operator= (const FSlateControlledConstruction& Other) = delete;
-
-	/** Widgets should only ever be constructed via SNew or SAssignNew */
-	void* operator new ( const size_t InSize )
-	{
-		return FMemory::Malloc(InSize);
-	}
-
-	/** Widgets should only ever be constructed via SNew or SAssignNew */
-	void* operator new ( const size_t InSize, void* Addr )
-	{
-		return Addr;
-	}
-
-	template<class WidgetType, bool bIsUserWidget>
-	friend struct TWidgetAllocator;
-
-	template <typename ObjectType>
-	friend class SharedPointerInternals::TIntrusiveReferenceController;
-
-public:
-	void operator delete(void* mem)
-	{
-		FMemory::Free(mem);
-	}
-};
-
 enum class EAccessibleType : uint8
 {
 	Main,
@@ -191,8 +149,10 @@ class SLATECORE_API SWidget
 	: public FSlateControlledConstruction,
 	public TSharedFromThis<SWidget>		// Enables 'this->AsShared()'
 {
-	friend struct FCurveSequence;
+	SLATE_DECLARE_WIDGET(SWidget, FSlateControlledConstruction)
+
 	friend class FWidgetProxy;
+	friend class FSlateAttributeMetaData;
 	friend class FSlateInvalidationRoot;
 	friend class FSlateInvalidationWidgetList;
 	friend class FSlateWindowElementList;
@@ -200,6 +160,44 @@ class SLATECORE_API SWidget
 	friend struct FSlateCachedElementList;
 	template<class WidgetType, typename RequiredArgsPayloadType>
 	friend struct TSlateDecl;
+
+protected:
+	template<typename InObjectType, EInvalidateWidgetReason InInvalidationReasonValue, typename InComparePredicate = TEqualTo<>>
+	struct TSlateAttribute : public ::SlateAttributePrivate::TSlateMemberAttribute<InObjectType, TSlateAttributeInvalidationReason<InInvalidationReasonValue>, InComparePredicate>
+	{
+		using ::SlateAttributePrivate::TSlateMemberAttribute<InObjectType, TSlateAttributeInvalidationReason<InInvalidationReasonValue>, InComparePredicate>::TSlateMemberAttribute;
+	};
+
+	template<EInvalidateWidgetReason InInvalidationReasonValue>
+	struct TSlateAttribute<FText, InInvalidationReasonValue> : public ::SlateAttributePrivate::TSlateMemberAttribute<FText, TSlateAttributeInvalidationReason<InInvalidationReasonValue>, TSlateAttributeFTextComparePredicate>
+	{
+		using ::SlateAttributePrivate::TSlateMemberAttribute<FText, TSlateAttributeInvalidationReason<InInvalidationReasonValue>, TSlateAttributeFTextComparePredicate>::TSlateMemberAttribute;
+	};
+
+	template<typename InObjectType, typename InInvalidationReasonPredicate, typename InComparePredicate = TEqualTo<>>
+	struct TSlateAttributeWithPredicate : public ::SlateAttributePrivate::TSlateMemberAttribute<InObjectType, InInvalidationReasonPredicate, InComparePredicate>
+	{
+		using ::SlateAttributePrivate::TSlateMemberAttribute<InObjectType, InInvalidationReasonPredicate, InComparePredicate>::TSlateMemberAttribute;
+	};
+
+	template<typename InInvalidationReasonPredicate>
+	struct TSlateAttributeWithPredicate<FText, InInvalidationReasonPredicate> : public ::SlateAttributePrivate::TSlateMemberAttribute<FText, InInvalidationReasonPredicate, TSlateAttributeFTextComparePredicate>
+	{
+		using ::SlateAttributePrivate::TSlateMemberAttribute<FText, InInvalidationReasonPredicate, TSlateAttributeFTextComparePredicate>::TSlateMemberAttribute;
+	};
+
+	template<typename InObjectType, EInvalidateWidgetReason InInvalidationReasonValue, typename InComparePredicate = TEqualTo<>>
+	struct TSlateManagedAttribute : public ::SlateAttributePrivate::TSlateManagedAttribute<InObjectType, TSlateAttributeInvalidationReason<InInvalidationReasonValue>, InComparePredicate>
+	{
+		using ::SlateAttributePrivate::TSlateManagedAttribute<InObjectType, TSlateAttributeInvalidationReason<InInvalidationReasonValue>, InComparePredicate>::TSlateManagedAttribute;
+	};
+
+	template<typename InObjectType, typename InInvalidationReasonPredicate, typename InComparePredicate = TEqualTo<>>
+	struct TSlateManagedAttributeWithPredicate : public ::SlateAttributePrivate::TSlateManagedAttribute<InObjectType, InInvalidationReasonPredicate, InComparePredicate>
+	{
+		using ::SlateAttributePrivate::TSlateManagedAttribute<InObjectType, InInvalidationReasonPredicate, InComparePredicate>::TSlateManagedAttribute;
+	};
+
 public:
 
 	/** Construct a SWidget based on initial parameters. */
@@ -646,6 +644,8 @@ public:
 	void SetCanTick(bool bInCanTick) { bInCanTick ? AddUpdateFlags(EWidgetUpdateFlags::NeedsTick) : RemoveUpdateFlags(EWidgetUpdateFlags::NeedsTick); }
 	bool GetCanTick() const { return HasAnyUpdateFlags(EWidgetUpdateFlags::NeedsTick); }
 
+	bool HasRegisteredSlateAttribute() const { return bHasRegisteredSlateAttribute; }
+
 	const FSlateWidgetPersistentState& GetPersistentState() const { return PersistentState; }
 	const FWidgetProxyHandle GetProxyHandle() const { return FastPathProxyHandle; }
 
@@ -1282,7 +1282,15 @@ public:
 	template<typename MetaDataType>
 	int32 RemoveMetaData(const TSharedRef<MetaDataType>& RemoveMe)
 	{
-		return MetaData.RemoveSingleSwap(RemoveMe);
+		int32 Index = MetaData.Find(RemoveMe);
+		if (Index == INDEX_NONE)
+		{
+			return 0;
+		}
+
+		checkf(Index != 0 || !HasRegisteredSlateAttribute(), TEXT("The first slot is reserved for SlateAttribute"));
+		MetaData.RemoveAtSwap(Index, 1);
+		return 1;
 	}
 
 private:
@@ -1298,6 +1306,7 @@ private:
 			const auto& MetaDataEntry = MetaData[Index];
 			if (MetaDataEntry->IsOfType<MetaDataType>())
 			{
+				checkf(Index != 0 || !HasRegisteredSlateAttribute(), TEXT("The first slot is reserved for SlateAttribute"));
 				MetaData.RemoveAtSwap(Index);
 			}
 		}
@@ -1312,6 +1321,7 @@ private:
 			const auto& MetaDataEntry = MetaData[Index];
 			if (MetaDataEntry->IsOfType<MetaDataType>())
 			{
+				checkf(Index != 0 || !HasRegisteredSlateAttribute(), TEXT("The first slot is reserved for SlateAttribute"));
 				MetaData.RemoveAtSwap(Index);
 				return true;
 			}
@@ -1361,6 +1371,10 @@ public:
 	/** @return The name this widget was tagged with */
 	virtual FName GetTag() const;
 
+#if STATS
+	size_t GetAllocSize() const { return AllocSize; }
+#endif
+
 #if UE_SLATE_WITH_WIDGET_UNIQUE_IDENTIFIER
 	/** @return The widget's id */
 	uint64 GetId() const { return UniqueIdentifier; }
@@ -1628,6 +1642,9 @@ private:
 	/** Are we currently updating the desired size? */
 	mutable uint8 bUpdatingDesiredSize : 1;
 
+	/** Is there at least one SlateAttribute currently registered. */
+	uint8 bHasRegisteredSlateAttribute : 1;
+
 protected:
 	uint8 bHasCustomPrepass : 1;
 
@@ -1742,7 +1759,9 @@ private:
 	uint64 UniqueIdentifier;
 #endif
 
-	STAT(size_t AllocSize;)
+#if STATS
+	size_t AllocSize;
+#endif
 
 #if STATS || ENABLE_STATNAMEDEVENTS
 	/** Stat id of this object, 0 if nobody asked for it yet */
diff --git a/Engine/Source/Runtime/SlateCore/Public/Widgets/SlateControlledConstruction.h b/Engine/Source/Runtime/SlateCore/Public/Widgets/SlateControlledConstruction.h
new file mode 100644
--- /dev/null
+++ b/Engine/Source/Runtime/SlateCore/Public/Widgets/SlateControlledConstruction.h
@@ -0,0 +1,120 @@
+// Copyright Epic Games, Inc. All Rights Reserved.
+
+#pragma once
+
+#include "CoreMinimal.h"
+#include "SlateGlobals.h"
+#include "Types/SlateAttributeDescriptor.h"
+#include "Templates/Identity.h"
+
+class FSlateControlledConstruction;
+namespace SharedPointerInternals
+{
+	template <typename ObjectType>
+	class TIntrusiveReferenceController;
+}
+
+/** */
+#define SLATE_DECLARE_WIDGET(WidgetType, ParentType) \
+	private: \
+		using Super = ParentType; \
+		using ThisClass = WidgetType; \
+		using PrivateThisType = WidgetType; \
+		using PrivateParentType = ParentType; \
+		static const FSlateWidgetClassData& GetPrivateWidgetClass() \
+		{ \
+			static FSlateWidgetClassData WidgetClassDataInstance = FSlateWidgetClassData(TIdentity<ParentType>(), #WidgetType, &WidgetType::PrivateRegisterAttributes); \
+			return WidgetClassDataInstance; \
+		} \
+		static void PrivateRegisterAttributes(FSlateAttributeInitializer&); \
+	public: \
+		static const FSlateWidgetClassData& StaticWidgetClass() { return GetPrivateWidgetClass(); } \
+		virtual const FSlateWidgetClassData& GetWidgetClass() const override { return GetPrivateWidgetClass(); } \
+	private:
+
+
+#define SLATE_IMPLEMENT_WIDGET(WidgetType) \
+	FSlateWidgetClassRegistration ClassRegistration__##WidgetType = WidgetType::StaticWidgetClass();
+
+/** */
+class SLATECORE_API FSlateWidgetClassData
+{
+private:
+	friend FSlateControlledConstruction;
+
+	FSlateWidgetClassData(FName InWidgetTypeName)
+		: WidgetType(InWidgetTypeName)
+	{}
+
+public:
+	template<typename InWidgetParentType>
+	FSlateWidgetClassData(TIdentity<InWidgetParentType>, FName InWidgetTypeName, void(*AttributeInitializer)(FSlateAttributeInitializer&))
+		: WidgetType(InWidgetTypeName)
+	{
+		// Initialize the parent class if it's not already done
+		const FSlateWidgetClassData& ParentWidgetClassData = InWidgetParentType::StaticWidgetClass();
+		// Initialize the attribute descriptor
+		FSlateAttributeInitializer Initializer = {Descriptor, ParentWidgetClassData.GetAttributeDescriptor()};
+		(*AttributeInitializer)(Initializer);
+	}
+
+	const FSlateAttributeDescriptor& GetAttributeDescriptor() const { return Descriptor; };
+	FName GetWidgetType() const { return WidgetType; }
+
+private:
+	FSlateAttributeDescriptor Descriptor;
+	FName WidgetType;
+};
+
+/** */
+struct FSlateWidgetClassRegistration
+{
+	FSlateWidgetClassRegistration(const FSlateWidgetClassData&) {}
+};
+
+/** */
+class SLATECORE_API FSlateControlledConstruction
+{
+public:
+	FSlateControlledConstruction() = default;
+	virtual ~FSlateControlledConstruction() = default;
+
+public:
+	static const FSlateWidgetClassData& StaticWidgetClass()
+	{
+		static FSlateWidgetClassData Instance = FSlateWidgetClassData(TEXT("FSlateControlledConstruction"));
+		return Instance;
+	}
+	virtual const FSlateWidgetClassData& GetWidgetClass() const = 0;
+	
+private:
+	/** UI objects cannot be copy-constructed */
+	FSlateControlledConstruction(const FSlateControlledConstruction& Other) = delete;
+	
+	/** UI objects cannot be copied. */
+	void operator= (const FSlateControlledConstruction& Other) = delete;
+
+	/** Widgets should only ever be constructed via SNew or SAssignNew */
+	void* operator new ( const size_t InSize )
+	{
+		return FMemory::Malloc(InSize);
+	}
+
+	/** Widgets should only ever be constructed via SNew or SAssignNew */
+	void* operator new ( const size_t InSize, void* Addr )
+	{
+		return Addr;
+	}
+
+	template<class WidgetType, bool bIsUserWidget>
+	friend struct TWidgetAllocator;
+
+	template <typename ObjectType>
+	friend class SharedPointerInternals::TIntrusiveReferenceController;
+
+public:
+	void operator delete(void* mem)
+	{
+		FMemory::Free(mem);
+	}
+};
