Slate: Fix Static Analysis fro SlateAttribute.cpp
#rn none
#rnx

diff --git a/Engine/Source/Runtime/SlateCore/Private/Types/SlateAttribute.cpp b/Engine/Source/Runtime/SlateCore/Private/Types/SlateAttribute.cpp
--- a/Engine/Source/Runtime/SlateCore/Private/Types/SlateAttribute.cpp
+++ b/Engine/Source/Runtime/SlateCore/Private/Types/SlateAttribute.cpp
@@ -89,6 +89,10 @@ namespace SlateAttributePrivate
 	{
 		checkf(AttributeType != ESlateAttributeType::Member, TEXT("Member Attribute cannot be moved."));
 
+		// Other can't be pass as ref because it's already moved.
+		//We only need it's pointer to make sure it's not also used in FSlateAttributeMetaData (as a void*).
+		check(Other);
+
 		FSlateAttributeMetaData::MoveAttribute(OwningWidget, *this, AttributeType, Other);
 	}
 }
