Slate: Fix missing include
#rnx

diff --git a/Engine/Source/Runtime/SlateCore/Private/Types/SlateAttributeDescriptor.cpp b/Engine/Source/Runtime/SlateCore/Private/Types/SlateAttributeDescriptor.cpp
--- a/Engine/Source/Runtime/SlateCore/Private/Types/SlateAttributeDescriptor.cpp
+++ b/Engine/Source/Runtime/SlateCore/Private/Types/SlateAttributeDescriptor.cpp
@@ -1,8 +1,8 @@
 // Copyright Epic Games, Inc. All Rights Reserved.
 
-#pragma once
-
 #include "Types/SlateAttributeDescriptor.h"
+#include "Types/SlateAttribute.h"
+#include "Misc/MemStack.h"
 
 
 FSlateAttributeDescriptor::FInitializer::FAttributeEntry::FAttributeEntry(FSlateAttributeDescriptor& InDescriptor, int32 InAttributeIndex)
