Slate: Fix bad dereferencing pointer in AttributeDescriptor
#rnx

diff --git a/Engine/Source/Runtime/SlateCore/Private/Types/SlateAttributeDescriptor.cpp b/Engine/Source/Runtime/SlateCore/Private/Types/SlateAttributeDescriptor.cpp
--- a/Engine/Source/Runtime/SlateCore/Private/Types/SlateAttributeDescriptor.cpp
+++ b/Engine/Source/Runtime/SlateCore/Private/Types/SlateAttributeDescriptor.cpp
@@ -94,7 +94,7 @@ void FSlateAttributeDescriptor::FInitializer::SetPrerequisite(FName AttributeNam
 	check(!AttributeName.IsNone());
 
 	FAttribute* FoundAttribute = Descriptor.FindAttribute(AttributeName);
-	if (ensureAlwaysMsgf(FoundAttribute == nullptr, TEXT("The attribute '%s' doesn't exist."), *AttributeName.ToString()))
+	if (ensureAlwaysMsgf(FoundAttribute != nullptr, TEXT("The attribute '%s' doesn't exist."), *AttributeName.ToString()))
 	{
 		Descriptor.SetPrerequisite(*FoundAttribute, Prerequisite);
 	}
