CVar Documentation corrections
#rnx

diff --git a/Engine/Source/Runtime/SlateCore/Private/Debugging/ConsoleSlateDebuggerInvalidate.cpp b/Engine/Source/Runtime/SlateCore/Private/Debugging/ConsoleSlateDebuggerInvalidate.cpp
--- a/Engine/Source/Runtime/SlateCore/Private/Debugging/ConsoleSlateDebuggerInvalidate.cpp
+++ b/Engine/Source/Runtime/SlateCore/Private/Debugging/ConsoleSlateDebuggerInvalidate.cpp
@@ -56,11 +56,11 @@ FConsoleSlateDebuggerInvalidate::FConsoleSlateDebuggerInvalidate()
 		FConsoleCommandDelegate::CreateRaw(this, &FConsoleSlateDebuggerInvalidate::ToggleLegend))
 	, ToogleWidgetsNameListCommand(
 		TEXT("SlateDebugger.Invalidate.ToggleWidgetNameList"),
-		TEXT("Option to display the name of the invalidated widget."),
+		TEXT("Option to display the names of invalidated widgets."),
 		FConsoleCommandDelegate::CreateRaw(this, &FConsoleSlateDebuggerInvalidate::ToggleWidgetNameList))
 	, ToogleLogInvalidatedWidgetCommand(
 		TEXT("SlateDebugger.Invalidate.ToggleLogInvalidatedWidget"),
-		TEXT("Option to log to the console the invalidated widget."),
+		TEXT("Option to log the invalidated widget to the console."),
 		FConsoleCommandDelegate::CreateRaw(this, &FConsoleSlateDebuggerInvalidate::ToggleLogInvalidatedWidget))
 	, SetInvalidateWidgetReasonFilterCommand(
 		TEXT("SlateDebugger.Invalidate.SetInvalidateWidgetReasonFilter"),
diff --git a/Engine/Source/Runtime/SlateCore/Private/Debugging/ConsoleSlateDebuggerInvalidationRoot.cpp b/Engine/Source/Runtime/SlateCore/Private/Debugging/ConsoleSlateDebuggerInvalidationRoot.cpp
--- a/Engine/Source/Runtime/SlateCore/Private/Debugging/ConsoleSlateDebuggerInvalidationRoot.cpp
+++ b/Engine/Source/Runtime/SlateCore/Private/Debugging/ConsoleSlateDebuggerInvalidationRoot.cpp
@@ -33,7 +33,7 @@ FConsoleSlateDebuggerInvalidationRoot::FConsoleSlateDebuggerInvalidationRoot()
 	, CacheDuration(2.0f)
 	, StartCommand(
 		TEXT("SlateDebugger.InvalidationRoot.Start"),
-		TEXT("Start the Invalidation Root widget debug tool. It shows when Invalidation Root are using the slow or the fast path."),
+		TEXT("Start the Invalidation Root widget debug tool. It shows when Invalidation Roots are using the slow or the fast path."),
 		FConsoleCommandDelegate::CreateRaw(this, &FConsoleSlateDebuggerInvalidationRoot::StartDebugging))
 	, StopCommand(
 		TEXT("SlateDebugger.InvalidationRoot.Stop"),
@@ -42,7 +42,7 @@ FConsoleSlateDebuggerInvalidationRoot::FConsoleSlateDebuggerInvalidationRoot()
 	, EnabledRefCVar(
 		TEXT("SlateDebugger.InvalidationRoot.Enable")
 		, bEnabledCVarValue
-		, TEXT("Start/Stop the Invalidation Root widget debug tool. It shows when Invalidation Root are using the slow or the fast path.")
+		, TEXT("Start/Stop the Invalidation Root widget debug tool. It shows when Invalidation Roots are using the slow or the fast path.")
 		, FConsoleVariableDelegate::CreateRaw(this, &FConsoleSlateDebuggerInvalidationRoot::HandleEnabled))
 	, ToggleLegendCommand(
 		TEXT("SlateDebugger.InvalidationRoot.ToggleLegend"),
diff --git a/Engine/Source/Runtime/SlateCore/Private/Debugging/ConsoleSlateDebuggerPaint.cpp b/Engine/Source/Runtime/SlateCore/Private/Debugging/ConsoleSlateDebuggerPaint.cpp
--- a/Engine/Source/Runtime/SlateCore/Private/Debugging/ConsoleSlateDebuggerPaint.cpp
+++ b/Engine/Source/Runtime/SlateCore/Private/Debugging/ConsoleSlateDebuggerPaint.cpp
@@ -33,7 +33,7 @@ FConsoleSlateDebuggerPaint::FConsoleSlateDebuggerPaint()
 	, CacheDuration(2.0f)
 	, ShowPaintWidgetCommand(
 		TEXT("SlateDebugger.Paint.Start")
-		, TEXT("Start the painted widget debug tool. It shows when widgets are painted.")
+		, TEXT("Start the painted widget debug tool. Use to show widget that have been painted this frame.")
 		, FConsoleCommandDelegate::CreateRaw(this, &FConsoleSlateDebuggerPaint::StartDebugging))
 	, HidePaintWidgetCommand(
 		TEXT("SlateDebugger.Paint.Stop")
@@ -46,7 +46,7 @@ FConsoleSlateDebuggerPaint::FConsoleSlateDebuggerPaint()
 		, FConsoleVariableDelegate::CreateRaw(this, &FConsoleSlateDebuggerPaint::HandleEnabled))
 	, LogPaintedWidgetOnceCommand(
 		TEXT("SlateDebugger.Paint.LogOnce")
-		, TEXT("Log the widgets that has been painted during the last update once")
+		, TEXT("Log the names of all widgets that were painted during the last update.")
 		, FConsoleCommandDelegate::CreateRaw(this, &FConsoleSlateDebuggerPaint::HandleLogOnce))
 	, DisplayWidgetsNameListCommand(
 		TEXT("SlateDebugger.Paint.ToggleWidgetNameList"),
@@ -55,11 +55,11 @@ FConsoleSlateDebuggerPaint::FConsoleSlateDebuggerPaint()
 	, MaxNumberOfWidgetInListtRefCVar(
 		TEXT("SlateDebugger.Paint.MaxNumberOfWidgetDisplayedInList")
 		, MaxNumberOfWidgetInList
-		, TEXT("The max number of widget that will be displayed when DisplayWidgetNameList is active."))
+		, TEXT("The max number of widgets that will be displayed when DisplayWidgetNameList is active."))
 	, LogWarningIfWidgetIsPaintedMoreThanOnceRefCVar(
 		TEXT("SlateDebugger.Paint.LogWarningIfWidgetIsPaintedMoreThanOnce")
 		, bLogWarningIfWidgetIsPaintedMoreThanOnce
-		, TEXT("Option to log a warning if a widget is painted more than once in the same frame."))
+		, TEXT("Option to log a warning if a widget is painted more than once in a single frame."))
 {
 	GConfig->GetBool(TEXT("SlateDebugger.Paint"), TEXT("bDisplayWidgetsNameList"), bDisplayWidgetsNameList, *GEditorPerProjectIni);
 	GConfig->GetBool(TEXT("SlateDebugger.Paint"), TEXT("bUseWidgetPathAsName"), bUseWidgetPathAsName, *GEditorPerProjectIni);
