Updated warning and error core style colors

diff --git a/Engine/Source/Runtime/SlateCore/Private/Styling/StyleColors.cpp b/Engine/Source/Runtime/SlateCore/Private/Styling/StyleColors.cpp
--- a/Engine/Source/Runtime/SlateCore/Private/Styling/StyleColors.cpp
+++ b/Engine/Source/Runtime/SlateCore/Private/Styling/StyleColors.cpp
@@ -117,8 +117,8 @@ void USlateThemeManager::InitalizeDefaults()
 	SetDefaultColor(EStyleColor::AccentWhite, COLOR("#FFFFFFFF"));
 	SetDefaultColor(EStyleColor::AccentFolder, COLOR("#B68F55FF"));
 
-	SetDefaultColor(EStyleColor::Warning, COLOR("#D78102FF"));
-	SetDefaultColor(EStyleColor::Error, COLOR("#CF1212FF"));
+	SetDefaultColor(EStyleColor::Warning, COLOR("#FFB800FF"));
+	SetDefaultColor(EStyleColor::Error, COLOR("#EF3535FF"));
 }
 
 #if ALLOW_THEMES
