Slate: Update child order invalidation in order similar to other invalidation. It create a path for other invalidations to be processed in that same order. It reduce the number of Layout invalidation that are already processed by the child order invalidation. Add Slate Attribute InvalidationWidgetReason. For FastUpdate,  add list of attributes that need to be updated each frame.

diff --git a/Engine/Source/Runtime/Slate/Private/Framework/Application/SlateApplication.cpp b/Engine/Source/Runtime/Slate/Private/Framework/Application/SlateApplication.cpp
--- a/Engine/Source/Runtime/Slate/Private/Framework/Application/SlateApplication.cpp
+++ b/Engine/Source/Runtime/Slate/Private/Framework/Application/SlateApplication.cpp
@@ -1511,9 +1511,6 @@ void FSlateApplication::TickAndDrawWidgets(float DeltaTime)
 		Renderer->ReleaseAccessedResources(/* Flush State */ false);
 	}
 
-	// We clear all pending Updates from last frame
-	FSlateInvalidationRoot::ClearAllWidgetUpdatesPending();
-
 	{
 		SCOPE_CYCLE_COUNTER(STAT_SlatePreTickEvent);
 		PreTickEvent.Broadcast(DeltaTime);
diff --git a/Engine/Source/Runtime/SlateCore/Private/FastUpdate/SlateElementSortedArray.h b/Engine/Source/Runtime/SlateCore/Private/FastUpdate/SlateElementSortedArray.h
new file mode 100644
--- /dev/null
+++ b/Engine/Source/Runtime/SlateCore/Private/FastUpdate/SlateElementSortedArray.h
@@ -0,0 +1,165 @@
+// Copyright Epic Games, Inc. All Rights Reserved.
+
+#pragma once
+
+#include "Algo/BinarySearch.h"
+#include "Containers/Array.h"
+#include "Templates/UnrealTypeTraits.h"
+
+/**
+ * An Array of elements that are always sorted.
+ */
+template <typename InElementType, typename InSortPredicate = TLess<>>
+class TSlateElementSortedArray
+{
+	using DataType = TArray<InElementType>;
+	DataType Data;
+
+public:
+	using SizeType = typename DataType::SizeType;
+	using ElementType = InElementType;
+	using ElementParamType = typename TCallTraits<ElementType>::ParamType;
+	using SortPredicate = TLess<>;
+
+public:
+	/** @returns Reference to indexed element. */
+	FORCEINLINE ElementParamType operator[](SizeType Index) const
+	{
+		return Data.GetData()[Index];
+	}
+
+public:
+	/** Removes all elements from the array, potentially leaving space allocated for an expected number of elements about to be added. */
+	FORCEINLINE void Empty(SizeType ExpectedNumElements = 0)
+	{
+		Data.Empty(ExpectedNumElements);
+	}
+
+	/** @returns True if index is valid. False otherwise. */
+	FORCEINLINE bool IsValidIndex(SizeType Index) const
+	{
+		return Data.IsValidIndex(Index);
+	}
+
+	/** @return The number of elements in the array. */
+	FORCEINLINE SizeType Num() const
+	{
+		return Data.Num();
+	}
+
+public:
+	/**
+	 * Binary searches the element within the array.
+	 *
+	 * @param Item Item to look for.
+	 * @returns Index of the found element. INDEX_NONE otherwise.
+	 */
+	SizeType Find(ElementParamType Item) const
+	{
+		return Algo::BinarySearch(Data, Item, SortPredicate());
+	}
+
+	/**
+	 * Checks if this array contains the element.
+	 *
+	 * @param Item Item to look for.
+	 * @returns	True if found. False otherwise.
+	 */
+	bool Contains(ElementParamType Item) const
+	{
+		return Find(Item) != INDEX_NONE;
+	}
+
+	/**
+	 * @return the index of the first element >= value, or INDEX_NONE if no such element is found.
+	 */
+	SizeType FindLowerBound(ElementParamType Value) const
+	{
+		const SizeType FoundIndex = Algo::LowerBound(Data, Value, SortPredicate());
+		return Data.IsValidIndex(FoundIndex) ? FoundIndex : INDEX_NONE;
+	}
+
+	/**
+	 * @returns the index of the first element > value, or INDEX_NONE if no such element if found.
+	 */
+	SizeType FindUpperBound(ElementParamType Value) const
+	{
+		const SizeType FoundIndex = Algo::UpperBound(Data, Value, SortPredicate());
+		return Data.IsValidIndex(FoundIndex) ? FoundIndex : INDEX_NONE;
+	}
+
+public:
+	/**
+	 * Adds a new item to the array at the sorted location if it doesn't exist, possibly reallocating the whole array to fit.
+	 *
+	 * @param Item The item to add
+	 * @return Index to the new item
+	 */
+	SizeType InsertUnique(ElementParamType Item)
+	{
+		SizeType LowerBoundIndex = Algo::LowerBound(Data, Item, SortPredicate());
+		if (LowerBoundIndex >= Num() || SortPredicate()(Item, Data.GetData()[LowerBoundIndex]))
+		{
+			return Data.Insert(Item, LowerBoundIndex);
+		}
+		return 0;
+	}
+
+	/**
+	 * Adds a new item to the end of the array. The array needs to be sorted manually after.
+	 * Use when multiple insertion is needed and it would be faster to sort the array once.
+	 *
+	 * @param Item The item to add
+	 * @return Index to the new item
+	*/
+	FORCEINLINE SizeType AddUnsorted(ElementType Item)
+	{
+		return Data.Add(Item);
+	}
+
+	/**
+	 * Removes the first occurrence of the specified item in the array, maintaining order but not indices.
+	 *
+	 * @param Item The item to remove.
+	 * @returns The number of items removed. For RemoveSingle, this is always either 0 or 1.
+	 */
+	SizeType RemoveSingle(ElementParamType Item)
+	{
+		SizeType Index = Find(Item);
+		if (Index != INDEX_NONE)
+		{
+			Data.RemoveAt(Index, 1, false);
+			return 1;
+		}
+		return 0;
+	}
+
+	/**
+	 * Removes an element (or elements) at given location optionally shrinking
+	 * the array.
+	 *
+	 * @param Index Location in array of the element to remove.
+	 * @param Count (Optional) Number of elements to remove. Default is 1.
+	 * @param bAllowShrinking (Optional) Tells if this call can shrink array if suitable after remove. Default is true.
+	 */
+	template <typename CountType>
+	FORCEINLINE void RemoveAt(SizeType Index, CountType Count)
+	{
+		Data.RemoveAt(Index, Count, false);
+	}
+
+	/** Sorts the array using the SortPredicate. */
+	void Sort()
+	{
+		Algo::Sort(Data, SortPredicate());
+	}
+
+public:
+	/**
+	 * DO NOT USE DIRECTLY
+	 * STL-like iterators to enable range-based for loop support.
+	 */
+	using RangedForConstIteratorType = typename DataType::RangedForConstIteratorType;
+	FORCEINLINE RangedForConstIteratorType	begin() const { return Data.begin(); }
+	FORCEINLINE RangedForConstIteratorType	end() const { return Data.end(); }
+};
diff --git a/Engine/Source/Runtime/SlateCore/Private/FastUpdate/SlateInvalidationRoot.cpp b/Engine/Source/Runtime/SlateCore/Private/FastUpdate/SlateInvalidationRoot.cpp
--- a/Engine/Source/Runtime/SlateCore/Private/FastUpdate/SlateInvalidationRoot.cpp
+++ b/Engine/Source/Runtime/SlateCore/Private/FastUpdate/SlateInvalidationRoot.cpp
@@ -14,22 +14,51 @@
 #include "ProfilingDebugging/CsvProfiler.h"
 #include "Trace/SlateTrace.h"
 #include "Types/ReflectionMetadata.h"
+#include "Types/SlateAttributeMetaData.h"
 
 CSV_DECLARE_CATEGORY_MODULE_EXTERN(SLATECORE_API, Slate);
 
+
 #if WITH_SLATE_DEBUGGING
-bool GDumpUpdateList = false;
-void HandleDumpUpdateList(const TArray<FString>& Args)
-{
-	GDumpUpdateList = true;
-}
+bool GSlateInvalidationRootDumpUpdateList = false;
+static FAutoConsoleVariableRef CVarSlateInvalidationRootDumpUpdateList(
+	TEXT("Slate.InvalidationRoot.DumpUpdateList"),
+	GSlateInvalidationRootDumpUpdateList,
+	TEXT("Each frame, log the widgets that will be updated.")
+);
+void DumpUpdateList(const FSlateInvalidationWidgetList& FastWidgetPathList, const TArray<FSlateInvalidationWidgetIndex>&);
+
+bool GSlateInvalidationRootDumpUpdateListOnce = false;
+static FAutoConsoleVariableRef CVarSlateInvalidationRootDumpUpdateListOnce(
+	TEXT("Slate.InvalidationRoot.DumpUpdateListOnce"),
+	GSlateInvalidationRootDumpUpdateListOnce,
+	TEXT("Log the widgets that will be updated this frame.")
+);
 
-static FAutoConsoleCommand HandleDumpUpdateListCommand(
+static FAutoConsoleCommand CVarHandleDumpUpdateListCommand_Deprecated(
 	TEXT("Slate.DumpUpdateList"),
-	TEXT(""),
-	FConsoleCommandWithArgsDelegate::CreateStatic(&HandleDumpUpdateList)
+	TEXT("(Deprecated) use Slate.InvalidationRoot.DumpUpdateListOnce"),
+	FConsoleCommandDelegate::CreateStatic([](){ GSlateInvalidationRootDumpUpdateListOnce = true; })
 );
 
+bool GSlateInvalidationRootDumpPreInvalidationList = false;
+static FAutoConsoleVariableRef CVarSlateInvalidationRootDumpPreInvalidationList(
+	TEXT("Slate.InvalidationRoot.DumpPreInvalidationList"),
+	GSlateInvalidationRootDumpPreInvalidationList,
+	TEXT("Each frame, log the widgets that are processed in the pre update phase.")
+);
+void LogPreInvalidationItem(const FSlateInvalidationWidgetList& FastWidgetPathList, FSlateInvalidationWidgetIndex WidgetIndex);
+
+bool GSlateInvalidationRootDumpPostInvalidationList = false;
+static FAutoConsoleVariableRef CVarSlateInvalidationRootDumpPostInvalidationList(
+	TEXT("Slate.InvalidationRoot.DumpPostInvalidationList"),
+	GSlateInvalidationRootDumpPostInvalidationList,
+	TEXT("Each frame, log the widgets that are processed in the post update phase.")
+);
+void LogPostInvalidationItem(const FSlateInvalidationWidgetList& FastWidgetPathList, FSlateInvalidationWidgetIndex WidgetIndex);
+#endif
+
+#if UE_SLATE_WITH_INVALIDATIONWIDGETLIST_DEBUGGING
 bool GSlateInvalidationRootVerifyWidgetList = false;
 static FAutoConsoleVariableRef CVarSlateInvalidationRootVerifyWidgetList(
 	TEXT("Slate.InvalidationRoot.VerifyWidgetList"),
@@ -47,7 +76,7 @@ static FAutoConsoleVariableRef CVarSlateInvalidationRootVerifyWidgetsIndex(
 
 bool GSlateInvalidationRootVerifyValidWidgets = false;
 static FAutoConsoleVariableRef CVarSlateInvalidationRootVerifyValidWidgets(
-	TEXT("Slate.InvalidationRoot.VerifyValidWidges"),
+	TEXT("Slate.InvalidationRoot.VerifyValidWidgets"),
 	GSlateInvalidationRootVerifyValidWidgets,
 	TEXT("Each frame, verify that every WidgetProxy has a valid SWidget.")
 );
@@ -75,15 +104,28 @@ static FAutoConsoleVariableRef CVarSlateInvalidationRootVerifyWidgetVolatile(
 	TEXT("Each frame, verify that volatile widgets are mark properly and are in the correct list.")
 );
 void VerifyWidgetVolatile(FSlateInvalidationWidgetList& WidgetList, TArray<FSlateInvalidationWidgetIndex>& FinalUpdateList);
-#endif //WITH_SLATE_DEBUGGING
 
-#if SLATE_CSV_TRACKER
-static int32 CascadeInvalidationEventAmount = 5;
-FAutoConsoleVariableRef CVarCascadeInvalidationEventAmount(
-	TEXT("Slate.CSV.CascadeInvalidationEventAmount"),
-	CascadeInvalidationEventAmount,
-	TEXT("The amount of cascaded invalidated parents before we fire a CSV event."));
-#endif //SLATE_CSV_TRACKER
+bool GSlateInvalidationRootVerifyWidgetsUpdateList = false;
+static FAutoConsoleVariableRef CVarSlateInvalidationRootVerifyWidgetsUpdateList(
+	TEXT("Slate.InvalidationRoot.VerifyWidgetUpdateList"),
+	GSlateInvalidationRootVerifyWidgetsUpdateList,
+	TEXT("Each frame, verify that pre and post update list contains the correct information and they are sorted.")
+);
+void VerifyWidgetsUpdateList_BeforeProcessPreUpdate(const TSharedRef<SWidget>&, FSlateInvalidationWidgetList*, FSlateInvalidationWidgetPreHeap*, FSlateInvalidationWidgetPostHeap*, TArray<FSlateInvalidationWidgetIndex>&);
+void VerifyWidgetsUpdateList_BeoreProcessPostUpdate(const TSharedRef<SWidget>&, FSlateInvalidationWidgetList*, FSlateInvalidationWidgetPreHeap*, FSlateInvalidationWidgetPostHeap*, TArray<FSlateInvalidationWidgetIndex>&);
+
+bool GSlateInvalidationRootVerifySlateAttribute = false;
+static FAutoConsoleVariableRef CVarSlateInvalidationRootVerifySlateAttributes(
+	TEXT("Slate.InvalidationRoot.VerifySlateAttribute"),
+	GSlateInvalidationRootVerifySlateAttribute,
+	TEXT("Each frame, verify that the widgets that have registered attribute are correctly updated once and the list contains all the widgets.")
+);
+void VerifySlateAttribute_Pre(FSlateInvalidationWidgetList& FastWidgetPathList);
+void VerifySlateAttribute_Post(FSlateInvalidationWidgetList const& FastWidgetPathList);
+
+#endif //UE_SLATE_WITH_INVALIDATIONWIDGETLIST_DEBUGGING
+
+
 
 int32 GSlateInvalidationWidgetListMaxArrayElements = 64;
 FAutoConsoleVariableRef CVarSlateInvalidationWidgetListMaxArrayElements(
@@ -101,7 +143,6 @@ FAutoConsoleVariableRef CVarSlateInvalidationWidgetListNumElementLeftBeforeSplit
  *
  */
 FSlateInvalidationRootList GSlateInvalidationRootListInstance;
-TArray<FSlateInvalidationRoot*> FSlateInvalidationRoot::ClearUpdateList;
 
 FSlateInvalidationRoot::FSlateInvalidationRoot()
 	: CachedElementData(new FSlateCachedElementData)
@@ -112,18 +153,21 @@ FSlateInvalidationRoot::FSlateInvalidationRoot()
 	, bChildOrderInvalidated(false)
 	, bNeedsSlowPath(true)
 	, bNeedScreenPositionShift(false)
-	, bProcessingChildOrderUpdate(false)
+	, bProcessingPreUpdate(false)
+	, bProcessingPostUpdate(false)
+	, bBuildingWidgetList(false)
+	, bProcessingChildOrderInvalidation(false)
 #if WITH_SLATE_DEBUGGING
 	, LastPaintType(ESlateInvalidationPaintType::None)
-	, ProcessInvalidationFrameNumber(0)
 #endif
 {
 	InvalidationRootHandle = FSlateInvalidationRootHandle(GSlateInvalidationRootListInstance.AddInvalidationRoot(this));
 	FSlateApplicationBase::Get().OnInvalidateAllWidgets().AddRaw(this, &FSlateInvalidationRoot::HandleInvalidateAllWidgets);
 
 	const FSlateInvalidationWidgetList::FArguments Arg = { GSlateInvalidationWidgetListMaxArrayElements, GSlateInvalidationWidgetListNumberElementLeftBeforeSplitting };
-	FastWidgetPathList = new FSlateInvalidationWidgetList(InvalidationRootHandle, Arg);
-	WidgetsNeedingUpdate = new FSlateInvalidationWidgetHeap(*FastWidgetPathList);
+	FastWidgetPathList = MakeUnique<FSlateInvalidationWidgetList>(InvalidationRootHandle, Arg);
+	WidgetsNeedingPreUpdate = MakeUnique<FSlateInvalidationWidgetPreHeap>(*FastWidgetPathList);
+	WidgetsNeedingPostUpdate = MakeUnique<FSlateInvalidationWidgetPostHeap>(*FastWidgetPathList);
 
 #if WITH_SLATE_DEBUGGING
 	SetLastPaintType(ESlateInvalidationPaintType::None);
@@ -149,9 +193,6 @@ FSlateInvalidationRoot::~FSlateInvalidationRoot()
 		delete CachedElementData;
 	}
 
-	delete FastWidgetPathList;
-	delete WidgetsNeedingUpdate;
-
 	GSlateInvalidationRootListInstance.RemoveInvalidationRoot(InvalidationRootHandle.GetUniqueId());
 }
 
@@ -174,7 +215,17 @@ void FSlateInvalidationRoot::InvalidateRootChildOrder(const SWidget* Investigato
 {
 	// Update the generation number. This will effectively invalidate all proxy handles
 	++FastPathGenerationNumber;
-	InvalidateRootLayout(Investigator);
+
+	// Invalidate all proxy handles
+	FastWidgetPathList->Reset();
+	InvalidationRootWidget->InvalidatePrepass();
+	InvalidationRootWidget->Invalidate(EInvalidateWidgetReason::Layout);
+	bNeedsSlowPath = true;
+
+#if WITH_SLATE_DEBUGGING
+	FSlateDebugging::BroadcastInvalidationRootInvalidate(InvalidationRootWidget, Investigator, ESlateDebuggingInvalidateRootReason::ChildOrder);
+#endif
+	UE_TRACE_SLATE_ROOT_CHILDORDER_INVALIDATED(InvalidationRootWidget, Investigator);
 }
 
 void FSlateInvalidationRoot::InvalidateRootLayout(const SWidget* Investigator)
@@ -189,31 +240,49 @@ void FSlateInvalidationRoot::InvalidateRootLayout(const SWidget* Investigator)
 	UE_TRACE_SLATE_ROOT_INVALIDATED(InvalidationRootWidget, Investigator);
 }
 
-void FSlateInvalidationRoot::InvalidateWidgetChildOrder(TSharedRef<SWidget> Widget)
+void FSlateInvalidationRoot::InvalidateWidget(FWidgetProxy& Proxy, EInvalidateWidgetReason InvalidateReason)
 {
+	ensureMsgf(bProcessingChildOrderInvalidation == false, TEXT("A widget got invalidated while building the childorder."));
+
 	if (!bNeedsSlowPath)
 	{
-		ensureAlways(bProcessingChildOrderUpdate == false);
-		WidgetsNeedingChildOrderUpdate.Add(Widget);
-
-		if (!bChildOrderInvalidated)
+		Proxy.CurrentInvalidateReason |= InvalidateReason;
+		if (EnumHasAnyFlags(InvalidateReason, EInvalidateWidgetReason::ChildOrder))
 		{
-			bChildOrderInvalidated = true;
-			if (!InvalidationRootWidget->Advanced_IsWindow())
-			{
-				InvalidationRootWidget->InvalidatePrepass();
-			}
+			WidgetsNeedingPreUpdate->HeapPushUnique(Proxy);
 
-			if (!GSlateEnableGlobalInvalidation && !InvalidationRootWidget->Advanced_IsWindow())
+			if (!bChildOrderInvalidated)
 			{
-				InvalidationRootWidget->Invalidate(EInvalidateWidgetReason::Layout);
+				bChildOrderInvalidated = true;
+				if (!InvalidationRootWidget->Advanced_IsWindow())
+				{
+					InvalidationRootWidget->InvalidatePrepass();
+				}
+
+				if (!GSlateEnableGlobalInvalidation && !InvalidationRootWidget->Advanced_IsWindow())
+				{
+					InvalidationRootWidget->Invalidate(EInvalidateWidgetReason::Layout);
+				}
 			}
 		}
+		else if (EnumHasAnyFlags(InvalidateReason, EInvalidateWidgetReason::AttributeRegistration))
+		{
+			WidgetsNeedingPreUpdate->HeapPushUnique(Proxy);
+		}
+
+		const uint8 AnyPostUpdate = (0xFF & ~(uint8)EInvalidateWidgetReason::AttributeRegistration);
+		if (((uint8)InvalidateReason & AnyPostUpdate) != 0)
+		{
+			WidgetsNeedingPostUpdate->PushBackOrHeapUnique(Proxy);
+		}
 
+		{
+			SWidget* WidgetPtr = Proxy.GetWidget();
 #if WITH_SLATE_DEBUGGING
-		FSlateDebugging::BroadcastInvalidationRootInvalidate(InvalidationRootWidget, &Widget.Get(), ESlateDebuggingInvalidateRootReason::ChildOrder);
+			FSlateDebugging::BroadcastWidgetInvalidate(WidgetPtr, nullptr, InvalidateReason);
 #endif
-		UE_TRACE_SLATE_ROOT_CHILDORDER_INVALIDATED(InvalidationRootWidget, &Widget.Get());
+			UE_TRACE_SLATE_WIDGET_INVALIDATED(WidgetPtr, nullptr, InvalidateReason);
+		}
 	}
 }
 
@@ -268,7 +337,6 @@ FSlateInvalidationResult FSlateInvalidationRoot::PaintInvalidationRoot(const FSl
 	{
 		SCOPED_NAMED_EVENT(Slate_PaintSlowPath, FColor::Red);
 
-		//CSV_EVENT(Basic, "Slate Slow Path update");
 		ClearAllFastPathData(!Context.bAllowFastPathUpdate);
 
 		GSlateIsOnFastUpdatePath = false;
@@ -307,7 +375,7 @@ FSlateInvalidationResult FSlateInvalidationRoot::PaintInvalidationRoot(const FSl
 
 	FinalUpdateList.Reset();
 
-#if WITH_SLATE_DEBUGGING
+#if UE_SLATE_WITH_INVALIDATIONWIDGETLIST_DEBUGGING
 	if (GSlateInvalidationRootVerifyHittestGrid && Context.bAllowFastPathUpdate)
 	{
 		VerifyHittest(InvalidationRootWidget, GetFastPathWidgetList(), GetHittestGrid());
@@ -332,27 +400,6 @@ void FSlateInvalidationRoot::OnWidgetDestroyed(const SWidget* Widget)
 	}
 }
 
-void FSlateInvalidationRoot::ClearAllWidgetUpdatesPending()
-{
-	// Once a frame we free the FinalUpdateList, any widget still in that list are
-	// Volatile widgets or widgets that need constant Update. So we put them back in the WidgetsNeedingUpdate list
-	for (FSlateInvalidationRoot* Root : ClearUpdateList)
-	{
-		if (int32 NumUpdatePending = Root->FinalUpdateList.Num())
-		{
-			for (FSlateInvalidationWidgetIndex index : Root->FinalUpdateList)
-			{
-				FWidgetProxy& Proxy = (*Root->FastWidgetPathList)[index];
-				if (EnumHasAnyFlags(Proxy.UpdateFlags, EWidgetUpdateFlags::AnyUpdate))
-				{
-					Root->WidgetsNeedingUpdate->PushUnique(Proxy);
-				}
-			}
-		}
-		Root->FinalUpdateList.Empty();
-	}
-}
-
 bool FSlateInvalidationRoot::PaintFastPath(const FSlateInvalidationContext& Context)
 {
 	SCOPED_NAMED_EVENT(SWidget_FastPathUpdate, FColor::Green);
@@ -362,73 +409,35 @@ bool FSlateInvalidationRoot::PaintFastPath(const FSlateInvalidationContext& Cont
 
 	bool bWidgetsNeededRepaint = false;
 	{
-		TGuardValue<bool> OnFastPathGuard(GSlateIsOnFastUpdatePath, true);
-
-		int32 LastParentIndex = 0;
-
 #if WITH_SLATE_DEBUGGING
-		if (GDumpUpdateList)
+		if (GSlateInvalidationRootDumpUpdateList || GSlateInvalidationRootDumpUpdateListOnce)
 		{
-			UE_LOG(LogSlate, Log, TEXT("Dumping Update List"));
-
-			// The update list is put in reverse order 
-			for (int32 ListIndex = FinalUpdateList.Num() - 1; ListIndex >= 0; --ListIndex)
-			{
-				const FSlateInvalidationWidgetIndex MyIndex = FinalUpdateList[ListIndex];
-
-				FWidgetProxy& WidgetProxy = (*FastWidgetPathList)[MyIndex];
-				SWidget* WidgetPtr = WidgetProxy.GetWidget();
-				if (WidgetPtr)
-				{
-					if (EnumHasAnyFlags(WidgetProxy.UpdateFlags, EWidgetUpdateFlags::NeedsVolatilePaint))
-					{
-						UE_LOG(LogSlate, Log, TEXT("Volatile Repaint %s"), *FReflectionMetaData::GetWidgetDebugInfo(WidgetPtr));
-					}
-					else if (EnumHasAnyFlags(WidgetProxy.UpdateFlags, EWidgetUpdateFlags::NeedsRepaint))
-					{
-						UE_LOG(LogSlate, Log, TEXT("Repaint %s"), *FReflectionMetaData::GetWidgetDebugInfo(WidgetPtr));
-					}
-					else if (WidgetPtr->IsFastPathVisible())
-					{
-						if (EnumHasAnyFlags(WidgetProxy.UpdateFlags, EWidgetUpdateFlags::NeedsActiveTimerUpdate))
-						{
-							UE_LOG(LogSlate, Log, TEXT("ActiveTimer %s"), *FReflectionMetaData::GetWidgetDebugInfo(WidgetPtr));
-						}
-
-						if (EnumHasAnyFlags(WidgetProxy.UpdateFlags, EWidgetUpdateFlags::NeedsTick))
-						{
-							UE_LOG(LogSlate, Log, TEXT("Tick %s"), *FReflectionMetaData::GetWidgetDebugInfo(WidgetPtr));
-						}
-					}
-				}
-			}
-
-			GDumpUpdateList = false;
+			DumpUpdateList(*FastWidgetPathList, FinalUpdateList);
 		}
 #endif
 
+		TGuardValue<bool> OnFastPathGuard(GSlateIsOnFastUpdatePath, true);
+
+		// The update list is put in reverse order by ProcessInvalidation
+		for (int32 ListIndex = FinalUpdateList.Num() - 1; ListIndex >= 0; --ListIndex)
 		{
-			// The update list is put in reverse order by ProcessInvalidation
-			for (int32 ListIndex = FinalUpdateList.Num() - 1; ListIndex >= 0; --ListIndex)
-			{
-				const FSlateInvalidationWidgetIndex MyIndex = FinalUpdateList[ListIndex];
-				FWidgetProxy& WidgetProxy = (*FastWidgetPathList)[MyIndex];
+			const FSlateInvalidationWidgetIndex MyIndex = FinalUpdateList[ListIndex];
+			FWidgetProxy& WidgetProxy = (*FastWidgetPathList)[MyIndex];
 
-				// Check visibility, it may have been in the update list but a parent who was also in the update list already updated it.
-				SWidget* WidgetPtr = WidgetProxy.GetWidget();
-				if (!WidgetProxy.bUpdatedSinceLastInvalidate && WidgetPtr && WidgetPtr->IsFastPathVisible())
-				{
-					bWidgetsNeededRepaint = bWidgetsNeededRepaint || EnumHasAnyFlags(WidgetProxy.UpdateFlags, EWidgetUpdateFlags::NeedsRepaint | EWidgetUpdateFlags::NeedsVolatilePaint);
+			// Check visibility, it may have been in the update list but a parent who was also in the update list already updated it.
+			SWidget* WidgetPtr = WidgetProxy.GetWidget();
+			if (!WidgetProxy.bUpdatedSinceLastInvalidate && WidgetPtr && WidgetPtr->IsFastPathVisible())
+			{
+				bWidgetsNeededRepaint = bWidgetsNeededRepaint || EnumHasAnyFlags(WidgetProxy.UpdateFlags, EWidgetUpdateFlags::NeedsRepaint | EWidgetUpdateFlags::NeedsVolatilePaint);
 
-					const int32 NewLayerId = WidgetProxy.Update(*Context.PaintArgs, *Context.WindowElementList);
-					CachedMaxLayerId = FMath::Max(NewLayerId, CachedMaxLayerId);
+				const int32 NewLayerId = WidgetProxy.Update(*Context.PaintArgs, *Context.WindowElementList);
+				CachedMaxLayerId = FMath::Max(NewLayerId, CachedMaxLayerId);
 
-					WidgetProxy.MarkProxyUpdatedThisFrame(*WidgetsNeedingUpdate);
+				WidgetProxy.MarkProxyUpdatedThisFrame(*WidgetsNeedingPostUpdate);
 
-					if (bNeedsSlowPath)
-					{
-						break;
-					}
+				if (bNeedsSlowPath)
+				{
+					break;
 				}
 			}
 		}
@@ -458,209 +467,331 @@ void FSlateInvalidationRoot::AdjustWidgetsDesktopGeometry(FVector2D WindowToDesk
 		});
 }
 
-
-void FSlateInvalidationRoot::ProcessChildOrderUpdate()
+void FSlateInvalidationRoot::BuildFastPathWidgetList(TSharedRef<SWidget> RootWidget)
 {
-	TSharedRef<SWidget> RootWidget = GetRootWidget();
+	TGuardValue<bool> Tmp(bBuildingWidgetList, true);
 
-	{
-		TGuardValue<bool> Tmp(bProcessingChildOrderUpdate, true);
-		if (FastWidgetPathList->GetRoot().Pin() != RootWidget)
-		{
-			FastWidgetPathList->BuildWidgetList(RootWidget);
-		}
-		else
-		{
-			FastWidgetPathList->ProcessChildOrderInvalidation(WidgetsNeedingChildOrderUpdate);
-		}
+	// We do not care if update are requested. We need to redo all the data.
+	WidgetsNeedingPreUpdate->Reset(false);
+	WidgetsNeedingPostUpdate->Reset(false);
+	FinalUpdateList.Reset();
 
-		WidgetsNeedingChildOrderUpdate.Reset();
-		bChildOrderInvalidated = false;
-	}
+	// Rebuild the list and update SlateAttribute
+	FastWidgetPathList->BuildWidgetList(RootWidget);
+}
 
-#if WITH_SLATE_DEBUGGING
-	if (GSlateInvalidationRootVerifyWidgetList)
+void FSlateInvalidationRoot::ProcessPreUpdate()
+{
+#if UE_SLATE_WITH_INVALIDATIONWIDGETLIST_DEBUGGING
+	if (GSlateInvalidationRootVerifyWidgetsUpdateList)
 	{
-		VerifyWidgetList(RootWidget, InvalidationRootHandle, GetFastPathWidgetList());
+		VerifyWidgetsUpdateList_BeforeProcessPreUpdate(GetRootWidget(), FastWidgetPathList.Get(), WidgetsNeedingPreUpdate.Get(), WidgetsNeedingPostUpdate.Get(), FinalUpdateList);
 	}
-
-	if (GSlateInvalidationRootVerifyWidgetsIndex)
+	if (GSlateInvalidationRootVerifySlateAttribute)
 	{
-		ensureMsgf(FastWidgetPathList->VerifyWidgetsIndex(), TEXT("We failed to verify that every widgets has the correct index."));
+		VerifySlateAttribute_Pre(*FastWidgetPathList);
 	}
-#endif //WITH_SLATE_DEBUGGING
-}
-
-void FSlateInvalidationRoot::BuildFastPathWidgetList(TSharedRef<SWidget> RootWidget)
-{
-	FastWidgetPathList->BuildWidgetList(RootWidget);
-	WidgetsNeedingChildOrderUpdate.Reset();
-}
+#endif
 
-bool FSlateInvalidationRoot::ProcessInvalidation()
-{
-	SCOPED_NAMED_EVENT(Slate_InvalidationProcessing, FColor::Blue);
-	CSV_SCOPED_TIMING_STAT(Slate, InvalidationProcessing);
+	TGuardValue<bool> Tmp(bProcessingPreUpdate, true);
 
-	bool bWidgetsNeedRepaint = false;
+	TSharedRef<SWidget> RootWidget = GetRootWidget();
+	if (FastWidgetPathList->GetRoot().Pin() != RootWidget)
+	{
+		BuildFastPathWidgetList(RootWidget);
 
-	if (!bNeedsSlowPath)
+		// Add the root to the update list (to prepass and paint it)
+		check(RootWidget->GetProxyHandle().IsValid(&RootWidget.Get()));
+		RootWidget->InvalidatePrepass();
+		InvalidateWidget(RootWidget->FastPathProxyHandle.GetProxy(), EInvalidateWidgetReason::Layout);
+	}
+	else
 	{
-		ensure(bChildOrderInvalidated == (WidgetsNeedingChildOrderUpdate.Num() != 0));
-		if (WidgetsNeedingChildOrderUpdate.Num())
+		// Put Widget waiting for update back in WidgetsNeedingPostUpdate to ensure index order
+		//and just in case Prepass need to be reexecuted.
 		{
-			SCOPED_NAMED_EVENT(Slate_InvalidationProcessing_ChildOrder, FColor::Blue);
-
-			FMemMark Mark(FMemStack::Get());
-			TArray<SWidget*, TMemStackAllocator<>> FastPathWidgetsNeedingUpdateCache;
-			FastPathWidgetsNeedingUpdateCache.Reserve(WidgetsNeedingUpdate->Num());
-			TArray<SWidget*, TMemStackAllocator<>> FastPathUpdateListCache;
-			FastPathUpdateListCache.Reserve(FinalUpdateList.Num());
-
-			for (const FSlateInvalidationWidgetHeap::FElement& Element : WidgetsNeedingUpdate->GetRaw())
+			for (FSlateInvalidationWidgetIndex WidgetIndex : FinalUpdateList)
 			{
-				FSlateInvalidationWidgetList::InvalidationWidgetType& InvalidationWidget = (*FastWidgetPathList)[Element.Get<0>()];
-				if (SWidget* Widget = InvalidationWidget.GetWidget())
-				{
-					FastPathWidgetsNeedingUpdateCache.Add(Widget);
-					// we will remove it soon, clear it right now to remove another loop
-					InvalidationWidget.bContainedByWidgetHeap = false;
-				}
+				FSlateInvalidationWidgetList::InvalidationWidgetType& InvalidationWidget = (*FastWidgetPathList)[WidgetIndex];
+				WidgetsNeedingPostUpdate->PushBackUnique(InvalidationWidget);
 			}
-			const bool bSetContainedByWidgetHeap = false;
-			WidgetsNeedingUpdate->Reset(bSetContainedByWidgetHeap);
+			FinalUpdateList.Reset(WidgetsNeedingPostUpdate->Num());
+		}
 
-			for (FSlateInvalidationWidgetIndex WidgetIndex : FinalUpdateList)
+		{
+#if WITH_SLATE_DEBUGGING
+			if (GSlateInvalidationRootDumpPreInvalidationList)
 			{
-				FSlateInvalidationWidgetList::InvalidationWidgetType& InvalidationWidget = (*FastWidgetPathList)[WidgetIndex];
-				if (SWidget* Widget = InvalidationWidget.GetWidget())
-				{
-					FastPathUpdateListCache.Add(Widget);
-				}
+				UE_LOG(LogSlate, Log, TEXT("Dumping Pre Invalidation List"));
+				UE_LOG(LogSlate, Log, TEXT("-------------------"));
 			}
+#endif
 
-			ProcessChildOrderUpdate();
+			FSlateInvalidationWidgetList::FWidgetAttributeIterator AttributeItt = FastWidgetPathList->CreateWidgetAttributeIterator();
 
-			for (SWidget* Widget : FastPathWidgetsNeedingUpdateCache)
+			/** */
+			struct FChildOriderInvalidationCallbackImpl : FSlateInvalidationWidgetList::IProcessChildOrderInvalidationCallback
 			{
-				if (Widget->GetProxyHandle().IsValid(Widget) && Widget->GetProxyHandle().GetInvalidationRoot() == this)
+				FChildOriderInvalidationCallbackImpl(
+					const FSlateInvalidationWidgetList& InWidgetList
+					, FSlateInvalidationWidgetPreHeap& InPreUpdate
+					, FSlateInvalidationWidgetPostHeap& InPostUpdate
+					, FSlateInvalidationWidgetList::FWidgetAttributeIterator& InAttributeItt)
+					: WidgetList(InWidgetList)
+					, PreUpdate(InPreUpdate)
+					, PostUpdate(InPostUpdate)
+					, AttributeItt(InAttributeItt)
+				{}
+				virtual ~FChildOriderInvalidationCallbackImpl() = default;
+				const FSlateInvalidationWidgetList& WidgetList;
+				FSlateInvalidationWidgetPreHeap& PreUpdate;
+				FSlateInvalidationWidgetPostHeap& PostUpdate;
+				FSlateInvalidationWidgetList::FWidgetAttributeIterator& AttributeItt;
+				TArray<FSlateInvalidationWidgetPreHeap::FElement*> WidgetToResort;
+
+				virtual void PreChildRemove(const FSlateInvalidationWidgetList::FIndexRange& Range) override
 				{
-					WidgetsNeedingUpdate->PushUnique(Widget->GetProxyHandle().GetWidgetIndex());
+					// The widgets got removed from the list. There is no need to update them anymore.
+					//Also, their index will not be valid after this function.
+					PreUpdate.RemoveRange(Range);
+					PostUpdate.RemoveRange(Range);
+					AttributeItt.PreChildRemove(Range);
 				}
-			}
-
-			for (SWidget* Widget : FastPathUpdateListCache)
-			{
-				if (Widget->GetProxyHandle().IsValid(Widget) && Widget->GetProxyHandle().GetInvalidationRoot() == this)
+				using FReIndexOperation = FSlateInvalidationWidgetList::IProcessChildOrderInvalidationCallback::FReIndexOperation;
+				virtual void ProxiesReIndexed(const FReIndexOperation& Operation) override
+				{
+					// Re-index in Pre and Post list (modify the index and the sort value)
+					FChildOriderInvalidationCallbackImpl const* Self = this;
+					auto ReIndexIfNeeded = [&Operation, Self](FSlateInvalidationWidgetPreHeap::FElement& Element)
+					{
+						if (Operation.GetRange().Include(Element.GetWidgetSortOrder()))
+						{
+							Element.GetWidgetIndexRef() = Operation.ReIndex(Element.GetWidgetIndex());
+							Element.GetWidgetSortOrderRef() = FSlateInvalidationWidgetSortOrder(Self->WidgetList, Element.GetWidgetIndex());
+						}
+					};
+					PreUpdate.ForEachIndexes(ReIndexIfNeeded);
+					PostUpdate.ForEachIndexes(ReIndexIfNeeded);
+					AttributeItt.ReIndexed(Operation);
+				}
+				using FReSortOperation = FSlateInvalidationWidgetList::IProcessChildOrderInvalidationCallback::FReSortOperation;
+				virtual void ProxiesPreResort(const FReSortOperation& Operation) override
 				{
-					WidgetsNeedingUpdate->PushUnique(Widget->GetProxyHandle().GetWidgetIndex());
+					// The sort order value will change but the order (operator<) is still valid.
+					FChildOriderInvalidationCallbackImpl* Self = this;
+					auto ReSortIfNeeded = [&Operation, Self](FSlateInvalidationWidgetPreHeap::FElement& Element)
+					{
+						if (Operation.GetRange().Include(Element.GetWidgetSortOrder()))
+						{
+							Self->WidgetToResort.Add(&Element);
+						}
+					};
+					PreUpdate.ForEachIndexes(ReSortIfNeeded);
+					PostUpdate.ForEachIndexes(ReSortIfNeeded);
 				}
-			}
-		}
-		else if (FinalUpdateList.Num() != 0)
-		{
-			// Put Widget waiting for update back in WidgetsNeedingUpdate to ensure index order and just in case, Prepass need to be reexecuted.
-			for (FSlateInvalidationWidgetIndex WidgetIndex : FinalUpdateList)
+				virtual void ProxiesPostResort()
+				{
+					for (FSlateInvalidationWidgetPreHeap::FElement* Element : WidgetToResort)
+					{
+						Element->GetWidgetSortOrderRef() = FSlateInvalidationWidgetSortOrder{ WidgetList, Element->GetWidgetIndex() };
+					}
+					WidgetToResort.Reset();
+					AttributeItt.PostResort();
+				}
+				virtual void ProxiesBuilt(const FSlateInvalidationWidgetList::FIndexRange& Range) override
+				{
+					AttributeItt.ProxiesBuilt(Range);
+				}
+
+			}  ChildOriderInvalidationCallback{ *FastWidgetPathList, *WidgetsNeedingPreUpdate, *WidgetsNeedingPostUpdate, AttributeItt };
+
+			while((AttributeItt.IsValid() || WidgetsNeedingPreUpdate->Num() > 0) && !bNeedsSlowPath)
 			{
-				WidgetsNeedingUpdate->PushUnique(WidgetIndex);
-			}
-		}
-		FinalUpdateList.Reset(WidgetsNeedingUpdate->Num());
+				FSlateInvalidationWidgetSortOrder AttributeSortOrder = (AttributeItt.IsValid()) ? AttributeItt.GetCurrentSortOrder() : FSlateInvalidationWidgetSortOrder::LimitMax();
+				FSlateInvalidationWidgetSortOrder NeedUpdateSortOrder = (WidgetsNeedingPreUpdate->Num() > 0) ? WidgetsNeedingPreUpdate->HeapPeekElement().GetWidgetSortOrder() : FSlateInvalidationWidgetSortOrder::LimitMax();
 
+				if (NeedUpdateSortOrder < AttributeSortOrder)
+				{
+					const FSlateInvalidationWidgetIndex WidgetIndex = WidgetsNeedingPreUpdate->HeapPop();
+					FSlateInvalidationWidgetList::InvalidationWidgetType& InvalidationWidget = (*FastWidgetPathList)[WidgetIndex];
+					// It could have been destroyed since then
+					if (InvalidationWidget.GetWidget())
+					{
 #if WITH_SLATE_DEBUGGING
-		if (GSlateInvalidationRootVerifyValidWidgets)
-		{
-			ensureMsgf(FastWidgetPathList->VerifyProxiesWidget(), TEXT("We failed to verify that every WidgetProxy has a valid SWidget"));
-		}
+						if (GSlateInvalidationRootDumpPreInvalidationList)
+						{
+							LogPreInvalidationItem(*FastWidgetPathList, WidgetIndex);
+						}
 #endif
 
-#if SLATE_CSV_TRACKER
-		FCsvProfiler::RecordCustomStat("Invalidate/InitialWidgets", CSV_CATEGORY_INDEX(Slate), WidgetsNeedingUpdate->Num(), ECsvCustomStatOp::Set);
-		int32 Stat_TotalWidgetsInvalidated = 0;
-		int32 Stat_NeedsRepaint = 0;
-		int32 Stat_NeedsVolatilePaint = 0;
-		int32 Stat_NeedsTick = 0;
-		int32 Stat_NeedsActiveTimerUpdate = 0;
+						bool bIsInvalidationWidgetValid = true;
+						if (EnumHasAnyFlags(InvalidationWidget.CurrentInvalidateReason, EInvalidateWidgetReason::ChildOrder))
+						{
+#if 0
+						FMemMark Mark(FMemStack::Get());
+						TArray<TTuple<FSlateInvalidationWidgetIndex, FSlateInvalidationWidgetSortOrder, TWeakPtr<SWidget>>, TMemStackAllocator<>> PreviousPreUpdate;
+						TArray<TTuple<FSlateInvalidationWidgetIndex, FSlateInvalidationWidgetSortOrder, TWeakPtr<SWidget>>, TMemStackAllocator<>> PreviousPostUpdate;
+						PreviousPreUpdate.Reserve(WidgetsNeedingPreUpdate->Num());
+						PreviousPostUpdate.Reserve(WidgetsNeedingPostUpdate->Num());
+						for (const auto& Element : WidgetsNeedingPreUpdate->GetRaw())
+						{
+							ensureMsgf(FastWidgetPathList->IsValidIndex(Element.GetWidgetIndex()), TEXT("F"));
+							PreviousPreUpdate.Emplace(Element.GetWidgetIndex(), Element.GetWidgetSortOrder(), (*FastWidgetPathList)[Element.GetWidgetIndex()].GetWidgetAsShared());
+						}
+						for (const auto& Element : WidgetsNeedingPostUpdate->GetRaw())
+						{
+							ensureMsgf(FastWidgetPathList->IsValidIndex(Element.GetWidgetIndex()), TEXT("gF"));
+							PreviousPostUpdate.Emplace(Element.GetWidgetIndex(), Element.GetWidgetSortOrder(), (*FastWidgetPathList)[Element.GetWidgetIndex()].GetWidgetAsShared());
+						}
 #endif
 
-		while (WidgetsNeedingUpdate->Num() && !bNeedsSlowPath)
-		{
-#if SLATE_CSV_TRACKER
-			Stat_TotalWidgetsInvalidated++;
-#endif
+							TGuardValue<bool> Tasdfdmp(bProcessingChildOrderInvalidation, true);
+							bIsInvalidationWidgetValid = FastWidgetPathList->ProcessChildOrderInvalidation(InvalidationWidget, ChildOriderInvalidationCallback);
 
-			FSlateInvalidationWidgetIndex MyIndex = WidgetsNeedingUpdate->Pop();
-			FinalUpdateList.Add(MyIndex);
-			FWidgetProxy& WidgetProxy = (*FastWidgetPathList)[MyIndex];
+							// This widget may not be valid anymore (got removed because it doesn't fulfill the requirement anymore ie. NullWidget).
+						}
 
-			// Reset each widgets paint state
-			// Must be done before actual painting because children can repaint 
-			WidgetProxy.bUpdatedSinceLastInvalidate = false;
+						if (EnumHasAnyFlags(InvalidationWidget.CurrentInvalidateReason, EInvalidateWidgetReason::AttributeRegistration) && bIsInvalidationWidgetValid)
+						{
+							FastWidgetPathList->ProcessAttributeRegistrationInvalidation(InvalidationWidget);
+							AttributeItt.Seek(WidgetIndex);
+						}
 
-			// Widget could be null if it was removed and we are on the slow path
-			if (SWidget* WidgetPtr = WidgetProxy.GetWidget())
-			{
-				const bool bIsInvalidationRoot = WidgetPtr->Advanced_IsInvalidationRoot();
-				if (bIsInvalidationRoot && WidgetPtr != InvalidationRootWidget)
+						AttributeItt.FixCurrentWidgetIndex();
+					}
+				}
+				else
 				{
-					FSlateInvalidationRoot* InvalidationRoot = const_cast<FSlateInvalidationRoot*>(WidgetPtr->Advanced_AsInvalidationRoot());
-					check(InvalidationRoot);
-					// Prevent reentering call
-					FSlateInvalidationWidgetHeap::FScopeWidgetCannotBeAdded Guard{ *WidgetsNeedingUpdate, WidgetProxy };
-					InvalidationRoot->ProcessInvalidation();
+					FSlateInvalidationWidgetList::InvalidationWidgetType& InvalidationWidget = (*FastWidgetPathList)[AttributeItt.GetCurrentIndex()];
+					SWidget* WidgetPtr = InvalidationWidget.GetWidget();
+					if (ensureMsgf(WidgetPtr, TEXT("Child order invalidation should have been called before we process this widget.")))
+					{
+						if (WidgetPtr->IsFastPathVisible())
+						{
+#if UE_SLATE_WITH_INVALIDATIONWIDGETLIST_DEBUGGING
+							if (GSlateInvalidationRootVerifySlateAttribute)
+							{
+								ensureMsgf(InvalidationWidget.bDebug_AttributeUpdated == false, TEXT("Attribute should only be updated once per frame."));
+							}
+#endif
+							FSlateAttributeMetaData::UpdateAttributes(*WidgetPtr);
+							InvalidationWidget.bDebug_AttributeUpdated = true;
+						}
+					}
+					AttributeItt.Advance();
 				}
+			}
+		}
+	}
+
 
-#if SLATE_CSV_TRACKER
-				const int32 PreviousWidgetsNeedingUpdating = WidgetsNeedingUpdate->Num();
+#if UE_SLATE_WITH_INVALIDATIONWIDGETLIST_DEBUGGING
+	if (GSlateInvalidationRootVerifyWidgetList)
+	{
+		VerifyWidgetList(RootWidget, InvalidationRootHandle, GetFastPathWidgetList());
+	}
+
+	if (GSlateInvalidationRootVerifyWidgetsIndex)
+	{
+		ensureMsgf(FastWidgetPathList->VerifyWidgetsIndex(), TEXT("We failed to verify that every widgets has the correct index."));
+	}
+	if (GSlateInvalidationRootVerifySlateAttribute)
+	{
+		VerifySlateAttribute_Post(*FastWidgetPathList);
+	}
+#endif //UE_SLATE_WITH_INVALIDATIONWIDGETLIST_DEBUGGING
+}
+
+bool FSlateInvalidationRoot::ProcessPostUpdate()
+{
+	WidgetsNeedingPostUpdate->Heapify();
+
+#if UE_SLATE_WITH_INVALIDATIONWIDGETLIST_DEBUGGING
+	if (GSlateInvalidationRootVerifyWidgetsUpdateList)
+	{
+		VerifyWidgetsUpdateList_BeoreProcessPostUpdate(GetRootWidget(), FastWidgetPathList.Get(), WidgetsNeedingPreUpdate.Get(), WidgetsNeedingPostUpdate.Get(), FinalUpdateList);
+	}
 #endif
 
-				bWidgetsNeedRepaint |= WidgetProxy.ProcessInvalidation(*WidgetsNeedingUpdate, *FastWidgetPathList, *this);
+	TGuardValue<bool> Tmp(bProcessingPostUpdate, true);
+	bool bWidgetsNeedRepaint = false;
 
-#if SLATE_CSV_TRACKER
-				const int32 CurrentWidgetsNeedingUpdating = WidgetsNeedingUpdate->Num();
-				const int32 AddedWidgets = CurrentWidgetsNeedingUpdating - PreviousWidgetsNeedingUpdating;
+#if WITH_SLATE_DEBUGGING
+	if (GSlateInvalidationRootDumpPostInvalidationList)
+	{
+		UE_LOG(LogSlate, Log, TEXT("Dumping Post Invalidation List"));
+		UE_LOG(LogSlate, Log, TEXT("-------------------"));
+	}
+#endif
 
-				if (AddedWidgets >= CascadeInvalidationEventAmount)
-				{
-					CSV_EVENT(Slate, TEXT("Invalidated %s"), *FReflectionMetaData::GetWidgetDebugInfo(WidgetPtr));
-				}
+	// It update backward (biggest index to smallest)
+	while (WidgetsNeedingPostUpdate->Num() && !bNeedsSlowPath)
+	{
+		const FSlateInvalidationWidgetIndex WidgetIndex = WidgetsNeedingPostUpdate->HeapPop();
+		FinalUpdateList.Add(WidgetIndex);
+		FWidgetProxy& WidgetProxy = (*FastWidgetPathList)[WidgetIndex];
 
-				if (EnumHasAnyFlags(WidgetProxy.UpdateFlags, EWidgetUpdateFlags::NeedsRepaint))
-				{
-					Stat_NeedsRepaint++;
-				}
-				if (EnumHasAnyFlags(WidgetProxy.UpdateFlags, EWidgetUpdateFlags::NeedsVolatilePaint) && !WidgetPtr->Advanced_IsInvalidationRoot())
-				{
-					Stat_NeedsVolatilePaint++;
-				}
-				if (EnumHasAnyFlags(WidgetProxy.UpdateFlags, EWidgetUpdateFlags::NeedsTick))
-				{
-					Stat_NeedsTick++;
-				}
-				if (EnumHasAnyFlags(WidgetProxy.UpdateFlags, EWidgetUpdateFlags::NeedsActiveTimerUpdate))
-				{
-					Stat_NeedsActiveTimerUpdate++;
-				}
+		// Reset each widgets paint state
+		// Must be done before actual painting because children can repaint 
+		WidgetProxy.bUpdatedSinceLastInvalidate = false;
+
+		// Widget could be null if it was removed and we are on the slow path
+		if (SWidget* WidgetPtr = WidgetProxy.GetWidget())
+		{
+#if WITH_SLATE_DEBUGGING
+			if (GSlateInvalidationRootDumpPostInvalidationList)
+			{
+				LogPostInvalidationItem(*FastWidgetPathList, WidgetIndex);
+			}
 #endif
+
+			const bool bIsInvalidationRoot = WidgetPtr->Advanced_IsInvalidationRoot();
+			if (bIsInvalidationRoot && WidgetPtr != InvalidationRootWidget)
+			{
+				FSlateInvalidationRoot* InvalidationRoot = const_cast<FSlateInvalidationRoot*>(WidgetPtr->Advanced_AsInvalidationRoot());
+				check(InvalidationRoot);
+				// Prevent reentering call
+				FSlateInvalidationWidgetPostHeap::FScopeWidgetCannotBeAdded Guard{ *WidgetsNeedingPostUpdate, WidgetProxy };
+				InvalidationRoot->ProcessInvalidation();
 			}
+
+			bWidgetsNeedRepaint |= WidgetProxy.ProcessPostInvalidation(*WidgetsNeedingPostUpdate, *FastWidgetPathList, *this);
 		}
-		
-		WidgetsNeedingUpdate->Reset(true);
-
-#if SLATE_CSV_TRACKER
-		FCsvProfiler::RecordCustomStat("Invalidate/TotalWidgets", CSV_CATEGORY_INDEX(Slate), Stat_TotalWidgetsInvalidated, ECsvCustomStatOp::Set);
-		FCsvProfiler::RecordCustomStat("Invalidate/NeedsRepaint", CSV_CATEGORY_INDEX(Slate), Stat_NeedsRepaint, ECsvCustomStatOp::Set);
-		FCsvProfiler::RecordCustomStat("Invalidate/NeedsVolatilePaint", CSV_CATEGORY_INDEX(Slate), Stat_NeedsVolatilePaint, ECsvCustomStatOp::Set);
-		FCsvProfiler::RecordCustomStat("Invalidate/NeedsTick", CSV_CATEGORY_INDEX(Slate), Stat_NeedsTick, ECsvCustomStatOp::Set);
-		FCsvProfiler::RecordCustomStat("Invalidate/NeedsActiveTimerUpdate", CSV_CATEGORY_INDEX(Slate), Stat_NeedsActiveTimerUpdate, ECsvCustomStatOp::Set);
-#endif
+	}
+	WidgetsNeedingPostUpdate->Reset(true);
+
+	return bWidgetsNeedRepaint;
+}
+
+bool FSlateInvalidationRoot::ProcessInvalidation()
+{
+	SCOPED_NAMED_EVENT(Slate_InvalidationProcessing, FColor::Blue);
+	CSV_SCOPED_TIMING_STAT(Slate, InvalidationProcessing);
+
+	bool bWidgetsNeedRepaint = false;
+
+	if (!bNeedsSlowPath)
+	{
+		check(WidgetsNeedingPreUpdate);
+		check(WidgetsNeedingPostUpdate);
+
+		ProcessPreUpdate();
+
+#if UE_SLATE_WITH_INVALIDATIONWIDGETLIST_DEBUGGING
+		if (GSlateInvalidationRootVerifyValidWidgets)
+		{
+			ensureMsgf(FastWidgetPathList->VerifyProxiesWidget(), TEXT("We failed to verify that every WidgetProxy has a valid SWidget"));
+		}
+#endif //UE_SLATE_WITH_INVALIDATIONWIDGETLIST_DEBUGGING
+
+		bWidgetsNeedRepaint = ProcessPostUpdate();
 	}
 	else
 	{
 		bWidgetsNeedRepaint = true;
 	}
 
-#if WITH_SLATE_DEBUGGING
+#if UE_SLATE_WITH_INVALIDATIONWIDGETLIST_DEBUGGING
 	if (GSlateInvalidationRootVerifyWidgetVisibility && !bNeedsSlowPath)
 	{
 		VerifyWidgetVisibility(GetFastPathWidgetList());
@@ -669,15 +800,13 @@ bool FSlateInvalidationRoot::ProcessInvalidation()
 	{
 		VerifyWidgetVolatile(GetFastPathWidgetList(), FinalUpdateList);
 	}
-#endif
+#endif //UE_SLATE_WITH_INVALIDATIONWIDGETLIST_DEBUGGING
 
 	return bWidgetsNeedRepaint;
 }
 
 void FSlateInvalidationRoot::ClearAllFastPathData(bool bClearResourcesImmediately)
 {
-
-	FastWidgetPathList->Reset();
 	FastWidgetPathList->ForEachWidget([bClearResourcesImmediately](SWidget* Widget)
 		{
 			Widget->PersistentState.CachedElementHandle = FSlateCachedElementsHandle::Invalid;
@@ -686,6 +815,7 @@ void FSlateInvalidationRoot::ClearAllFastPathData(bool bClearResourcesImmediatel
 				Widget->FastPathProxyHandle = FWidgetProxyHandle();
 			}
 		});
+	FastWidgetPathList->Reset();
 
 #if UE_SLATE_DEBUGGING_CLEAR_ALL_FAST_PATH_DATA
 	if (!bClearResourcesImmediately)
@@ -710,12 +840,8 @@ void FSlateInvalidationRoot::ClearAllFastPathData(bool bClearResourcesImmediatel
 	}
 #endif
 
-
-	if (!FastWidgetPathList->IsEmpty())
-	{
-		ClearUpdateList.RemoveSingleSwap(this, false);
-	}
-	WidgetsNeedingUpdate->Reset(false);
+	WidgetsNeedingPreUpdate->Reset(false);
+	WidgetsNeedingPostUpdate->Reset(false);
 	FastWidgetPathList->Empty();
 	CachedElementData->Empty();
 	FinalUpdateList.Empty();
@@ -742,6 +868,93 @@ void FSlateInvalidationRoot::Advanced_ResetInvalidation(bool bClearResourcesImme
 }
 
 #if WITH_SLATE_DEBUGGING
+void DumpUpdateList(const FSlateInvalidationWidgetList& FastWidgetPathList, const TArray<FSlateInvalidationWidgetIndex>& FinalUpdateList)
+{
+	UE_LOG(LogSlate, Log, TEXT("Dumping Update List"));
+	UE_LOG(LogSlate, Log, TEXT("-------------------"));
+	// The update list is put in reverse order 
+	for (int32 ListIndex = FinalUpdateList.Num() - 1; ListIndex >= 0; --ListIndex)
+	{
+		const FSlateInvalidationWidgetIndex MyIndex = FinalUpdateList[ListIndex];
+
+		const FWidgetProxy& WidgetProxy = FastWidgetPathList[MyIndex];
+		const SWidget* WidgetPtr = WidgetProxy.GetWidget();
+		if (WidgetPtr)
+		{
+			if (EnumHasAnyFlags(WidgetProxy.UpdateFlags, EWidgetUpdateFlags::NeedsVolatilePaint))
+			{
+				UE_LOG(LogSlate, Log, TEXT("Volatile Repaint %s"), *FReflectionMetaData::GetWidgetDebugInfo(WidgetPtr));
+			}
+			else if (EnumHasAnyFlags(WidgetProxy.UpdateFlags, EWidgetUpdateFlags::NeedsRepaint))
+			{
+				UE_LOG(LogSlate, Log, TEXT("Repaint %s"), *FReflectionMetaData::GetWidgetDebugInfo(WidgetPtr));
+			}
+			else if (WidgetPtr->IsFastPathVisible())
+			{
+				if (EnumHasAnyFlags(WidgetProxy.UpdateFlags, EWidgetUpdateFlags::NeedsActiveTimerUpdate))
+				{
+					UE_LOG(LogSlate, Log, TEXT("ActiveTimer %s"), *FReflectionMetaData::GetWidgetDebugInfo(WidgetPtr));
+				}
+
+				if (EnumHasAnyFlags(WidgetProxy.UpdateFlags, EWidgetUpdateFlags::NeedsTick))
+				{
+					UE_LOG(LogSlate, Log, TEXT("Tick %s"), *FReflectionMetaData::GetWidgetDebugInfo(WidgetPtr));
+				}
+			}
+		}
+	}
+	UE_LOG(LogSlate, Log, TEXT("-------------------"));
+
+	GSlateInvalidationRootDumpUpdateListOnce = false;
+}
+
+
+void LogPreInvalidationItem(const FSlateInvalidationWidgetList& FastWidgetPathList, FSlateInvalidationWidgetIndex WidgetIndex)
+{
+	const FWidgetProxy& Proxy = FastWidgetPathList[WidgetIndex];
+
+	if (EnumHasAnyFlags(Proxy.CurrentInvalidateReason, EInvalidateWidgetReason::ChildOrder))
+	{
+		UE_LOG(LogSlate, Log, TEXT("  Child Order %s"), *FReflectionMetaData::GetWidgetDebugInfo(Proxy.GetWidget()));
+	}
+	else
+	{
+		UE_LOG(LogSlate, Log, TEXT("  [?] %s"), *FReflectionMetaData::GetWidgetDebugInfo(Proxy.GetWidget()));
+	}
+}
+
+void LogPostInvalidationItem(const FSlateInvalidationWidgetList& FastWidgetPathList, FSlateInvalidationWidgetIndex WidgetIndex)
+{
+	const FWidgetProxy& Proxy = FastWidgetPathList[WidgetIndex];
+
+	if (EnumHasAnyFlags(Proxy.CurrentInvalidateReason, EInvalidateWidgetReason::Layout))
+	{
+		UE_LOG(LogSlate, Log, TEXT("  Layout %s"), *FReflectionMetaData::GetWidgetDebugInfo(Proxy.GetWidget()));
+	}
+	else if (EnumHasAnyFlags(Proxy.CurrentInvalidateReason, EInvalidateWidgetReason::Visibility))
+	{
+		UE_LOG(LogSlate, Log, TEXT("  Visibility %s"), *FReflectionMetaData::GetWidgetDebugInfo(Proxy.GetWidget()));
+	}
+	else if (EnumHasAnyFlags(Proxy.CurrentInvalidateReason, EInvalidateWidgetReason::Volatility))
+	{
+		UE_LOG(LogSlate, Log, TEXT("  Volatility %s"), *FReflectionMetaData::GetWidgetDebugInfo(Proxy.GetWidget()));
+	}
+	else if (EnumHasAnyFlags(Proxy.CurrentInvalidateReason, EInvalidateWidgetReason::RenderTransform))
+	{
+		UE_LOG(LogSlate, Log, TEXT("  RenderTransform %s"), *FReflectionMetaData::GetWidgetDebugInfo(Proxy.GetWidget()));
+	}
+	else if (EnumHasAnyFlags(Proxy.CurrentInvalidateReason, EInvalidateWidgetReason::Paint))
+	{
+		UE_LOG(LogSlate, Log, TEXT("  Paint %s"), *FReflectionMetaData::GetWidgetDebugInfo(Proxy.GetWidget()));
+	}
+	else if (!EnumHasAnyFlags(Proxy.UpdateFlags, EWidgetUpdateFlags::AnyUpdate))
+	{
+		UE_LOG(LogSlate, Log, TEXT("  [?] %s"), *FReflectionMetaData::GetWidgetDebugInfo(Proxy.GetWidget()));
+	}
+}
+#endif
+
+#if UE_SLATE_WITH_INVALIDATIONWIDGETLIST_DEBUGGING
 void VerifyWidgetList(TSharedRef<SWidget> RootWidget, FSlateInvalidationRootHandle InvalidationRootHandle, FSlateInvalidationWidgetList& WidgetList)
 {
 	FSlateInvalidationWidgetList List(InvalidationRootHandle, FSlateInvalidationWidgetList::FArguments{ 128, 128, 1000, false });
@@ -784,7 +997,7 @@ void VerifyHittest(SWidget* InvalidationRootWidget, FSlateInvalidationWidgetList
 	for (const FHittestGrid::FWidgetSortData& Data : WeakHittestGridSortDatas)
 	{
 		TSharedPtr<SWidget> Widget = Data.WeakWidget.Pin();
-		if (ensureAlwaysMsgf(Widget, TEXT("A widget is invalid in the HittestGrid")))
+		if (ensureMsgf(Widget, TEXT("A widget is invalid in the HittestGrid")))
 		{
 			FHittestWidgetSortData SortData = { Widget.Get(), Data.PrimarySort, Data.SecondarySort };
 			HittestGridSortDatas.Add(MoveTemp(SortData));
@@ -797,7 +1010,7 @@ void VerifyHittest(SWidget* InvalidationRootWidget, FSlateInvalidationWidgetList
 	const SWidget* LastWidget = nullptr;
 	WidgetList.ForEachWidget([&HittestGridSortDatas, &PreviousSecondarySort, &LastWidget](const SWidget* Widget)
 		{
-			if (Widget->GetVisibility().IsHitTestVisible())
+			if (Widget && Widget->GetVisibility().IsHitTestVisible())
 			{
 				// Is the widget in the hittestgrid
 				const int32 FoundHittestIndex = HittestGridSortDatas.IndexOfByPredicate([Widget](const FHittestWidgetSortData& HittestGrid)
@@ -810,7 +1023,7 @@ void VerifyHittest(SWidget* InvalidationRootWidget, FSlateInvalidationWidgetList
 					return;
 				}
 
-				ensureAlwaysMsgf(Widget->GetProxyHandle().GetWidgetSortOrder() == HittestGridSortDatas[FoundHittestIndex].SecondarySort
+				ensureMsgf(Widget->GetProxyHandle().GetWidgetSortOrder() == HittestGridSortDatas[FoundHittestIndex].SecondarySort
 					, TEXT("The SecondarySort of widget '%s' doesn't match the SecondarySort inside the hittestgrid.")
 					, *FReflectionMetaData::GetWidgetPath(Widget));
 
@@ -830,7 +1043,7 @@ void VerifyHittest(SWidget* InvalidationRootWidget, FSlateInvalidationWidgetList
 		HittestGridSortDatas.RemoveAtSwap(FoundHittestIndex);
 	}
 
-	ensureAlwaysMsgf(HittestGridSortDatas.Num() == 0, TEXT("The hittest grid of Root '%s' has widgets that are not inside the InvalidationRoot's widget list")
+	ensureMsgf(HittestGridSortDatas.Num() == 0, TEXT("The hittest grid of Root '%s' has widgets that are not inside the InvalidationRoot's widget list")
 		, *FReflectionMetaData::GetWidgetPath(InvalidationRootWidget));
 }
 
@@ -911,7 +1124,78 @@ void VerifyWidgetVolatile(FSlateInvalidationWidgetList& WidgetList, TArray<FSlat
 			}
 		});
 }
-#endif //WITH_SLATE_DEBUGGING
+
+void VerifyWidgetsUpdateList_BeforeProcessPreUpdate(const TSharedRef<SWidget>& RootWidget,
+	FSlateInvalidationWidgetList* FastWidgetPathList,
+	FSlateInvalidationWidgetPreHeap* WidgetsNeedingPreUpdate,
+	FSlateInvalidationWidgetPostHeap* WidgetsNeedingPostUpdate,
+	TArray<FSlateInvalidationWidgetIndex>& FinalUpdateList)
+{
+	if (FastWidgetPathList->GetRoot().Pin() != RootWidget)
+	{
+		return;
+	}
+
+	for (FSlateInvalidationWidgetIndex WidgetIndex : FinalUpdateList)
+	{
+		ensureMsgf(FastWidgetPathList->IsValidIndex(WidgetIndex), TEXT("An element is not valid."));
+		const FSlateInvalidationWidgetList::InvalidationWidgetType& InvalidationWidget = (*FastWidgetPathList)[WidgetIndex];
+		ensureMsgf(InvalidationWidget.CurrentInvalidateReason != EInvalidateWidgetReason::None || InvalidationWidget.UpdateFlags != EWidgetUpdateFlags::None
+			, TEXT("A widget is in the update list but doesn't have an update flag set."));
+	}
+
+	ensureMsgf(WidgetsNeedingPreUpdate->IsValidHeap_Debug(), TEXT("The PreUpdate list need to stay a valid heap"));
+
+	auto CheckIfValidInvalidation = [FastWidgetPathList](const FSlateInvalidationWidgetPreHeap::FElement& Element)
+	{
+		ensureMsgf(FastWidgetPathList->IsValidIndex(Element.GetWidgetIndex()), TEXT("An element is not valid."));
+		const FSlateInvalidationWidgetList::InvalidationWidgetType& InvalidationWidget = (*FastWidgetPathList)[Element.GetWidgetIndex()];
+		ensureMsgf(InvalidationWidget.CurrentInvalidateReason != EInvalidateWidgetReason::None || InvalidationWidget.UpdateFlags != EWidgetUpdateFlags::None
+			, TEXT("A widget is in the update list but doesn't have an update flag set."));
+		if (SWidget* Widget = (*FastWidgetPathList)[Element.GetWidgetIndex()].GetWidget())
+		{
+			ensureMsgf(Widget->GetProxyHandle().GetWidgetSortOrder() == Element.GetWidgetSortOrder(), TEXT("The sort order of the widget do not matches what is in the heap."));
+			ensureMsgf(Widget->GetProxyHandle().GetWidgetIndex() == Element.GetWidgetIndex(), TEXT("The widget index of the widget do not matches what is in the heap."));
+		}
+	};
+	WidgetsNeedingPreUpdate->ForEachIndexes(CheckIfValidInvalidation);
+	WidgetsNeedingPostUpdate->ForEachIndexes(CheckIfValidInvalidation);
+
+	FastWidgetPathList->ForEachInvalidationWidget([WidgetsNeedingPreUpdate, WidgetsNeedingPostUpdate](FSlateInvalidationWidgetList::InvalidationWidgetType& InvalidationWidget)
+		{
+			ensureMsgf(WidgetsNeedingPreUpdate->Contains_Debug(InvalidationWidget.Index) == InvalidationWidget.bContainedByWidgetPreHeap
+				, TEXT("A widget is or is not in the PreUpdate but the flag say otherwise."));
+			ensureMsgf(WidgetsNeedingPostUpdate->Contains_Debug(InvalidationWidget.Index) == InvalidationWidget.bContainedByWidgetPostHeap
+				, TEXT("A widget is or is not in the PostUpdate but the flag say otherwise."));
+		});
+}
+
+void VerifyWidgetsUpdateList_BeoreProcessPostUpdate(const TSharedRef<SWidget>& RootWidget,
+	FSlateInvalidationWidgetList* FastWidgetPathList,
+	FSlateInvalidationWidgetPreHeap* WidgetsNeedingPreUpdate,
+	FSlateInvalidationWidgetPostHeap* WidgetsNeedingPostUpdate,
+	TArray<FSlateInvalidationWidgetIndex>& FinalUpdateList)
+{
+	ensureMsgf(WidgetsNeedingPostUpdate->IsValidHeap_Debug(), TEXT("The PostUpdate list need to stay a valid heap"));
+
+	VerifyWidgetsUpdateList_BeforeProcessPreUpdate(RootWidget, FastWidgetPathList, WidgetsNeedingPreUpdate, WidgetsNeedingPostUpdate, FinalUpdateList);
+}
+
+void VerifySlateAttribute_Pre(FSlateInvalidationWidgetList& FastWidgetPathList)
+{
+	FastWidgetPathList.ForEachInvalidationWidget([](FSlateInvalidationWidgetList::InvalidationWidgetType& InvalidationWidget)
+		{
+			InvalidationWidget.bDebug_AttributeUpdated = false;
+		});
+}
+
+void VerifySlateAttribute_Post(FSlateInvalidationWidgetList const& FastWidgetPathList)
+{
+	FastWidgetPathList.VerifyElementIndexList();
+}
+
+
+#endif //UE_SLATE_WITH_INVALIDATIONWIDGETLIST_DEBUGGING
 
 /**
  * 
diff --git a/Engine/Source/Runtime/SlateCore/Private/FastUpdate/SlateInvalidationWidgetHeap.h b/Engine/Source/Runtime/SlateCore/Private/FastUpdate/SlateInvalidationWidgetHeap.h
--- a/Engine/Source/Runtime/SlateCore/Private/FastUpdate/SlateInvalidationWidgetHeap.h
+++ b/Engine/Source/Runtime/SlateCore/Private/FastUpdate/SlateInvalidationWidgetHeap.h
@@ -8,128 +8,357 @@
 #include "FastUpdate/SlateInvalidationWidgetList.h"
 #include "FastUpdate/SlateInvalidationWidgetSortOrder.h"
 
+
+namespace UE
+{
+namespace Slate
+{
+namespace Private
+{
+	/**
+	* Ordered list of WidgetIndex. The order is based on the WidgetSortIndex.
+	*/
+	struct FSlateInvalidationWidgetHeapElement
+	{
+	public:
+		FSlateInvalidationWidgetHeapElement(FSlateInvalidationWidgetIndex InIndex, FSlateInvalidationWidgetSortOrder InSortOrder)
+			: WidgetIndex(InIndex), WidgetSortOrder(InSortOrder)
+		{}
+		inline FSlateInvalidationWidgetIndex GetWidgetIndex() const { return WidgetIndex; }
+		inline FSlateInvalidationWidgetIndex& GetWidgetIndexRef() { return WidgetIndex; }
+		inline FSlateInvalidationWidgetSortOrder GetWidgetSortOrder() const { return WidgetSortOrder; }
+		inline FSlateInvalidationWidgetSortOrder& GetWidgetSortOrderRef() { return WidgetSortOrder; }
+
+	private:
+		FSlateInvalidationWidgetIndex WidgetIndex;
+		FSlateInvalidationWidgetSortOrder WidgetSortOrder;
+	};
+} // Private
+} // Slate
+} // UE
+
 /**
- * Ordered list of WidgetIndex. The order is based on the WidgetSortIndex.
+ * Heap of widget that is ordered by increasing sort order. The list need to always stay ordered.
  */
-class FSlateInvalidationWidgetHeap
+class FSlateInvalidationWidgetPreHeap
 {
 public:
-	using FElement = TPair<FSlateInvalidationWidgetIndex, FSlateInvalidationWidgetSortOrder>;
-	struct TWidgetOrderGreater
+	using FElement = UE::Slate::Private::FSlateInvalidationWidgetHeapElement;
+	struct FWidgetOrderLess
 	{
 		FORCEINLINE bool operator()(const FElement& A, const FElement& B) const
 		{
-			return B.Get<1>() < A.Get<1>();
+			return A.GetWidgetSortOrder() < B.GetWidgetSortOrder();
 		}
 	};
+	using SortPredicate = FWidgetOrderLess;
+	
+	static constexpr int32 NumberOfPreAllocatedElement = 32;
+	using FElementContainer = TArray<FElement, TInlineAllocator<NumberOfPreAllocatedElement>>;
 
-public:
-	FSlateInvalidationWidgetHeap(FSlateInvalidationWidgetList& InOwnerList)
+	FSlateInvalidationWidgetPreHeap(FSlateInvalidationWidgetList& InOwnerList)
 		: OwnerList(InOwnerList)
-		, WidgetCannotBeAdded(FSlateInvalidationWidgetIndex::Invalid)
 	{ }
 
+public:
 	/** Insert into the list at the proper order (see binary heap) only if it's not already contains by the list. */
-	void PushUnique(const FSlateInvalidationWidgetIndex WidgetIndex)
+	void HeapPushUnique(FSlateInvalidationWidgetList::InvalidationWidgetType& InvalidationWidget)
 	{
-		check(WidgetIndex != FSlateInvalidationWidgetIndex::Invalid);
+		check(InvalidationWidget.Index != FSlateInvalidationWidgetIndex::Invalid);
+		check(Contains_Debug(InvalidationWidget.Index) == InvalidationWidget.bContainedByWidgetPreHeap);
 
-		FSlateInvalidationWidgetList::InvalidationWidgetType& InvalidationWidget = OwnerList[WidgetIndex];
-		check(Contains(WidgetIndex) == InvalidationWidget.bContainedByWidgetHeap
-			|| WidgetCannotBeAdded == WidgetIndex);
+		if (!InvalidationWidget.bContainedByWidgetPreHeap)
+		{
+			InvalidationWidget.bContainedByWidgetPreHeap = true;
+			Heap.HeapPush(FElement{ InvalidationWidget.Index, FSlateInvalidationWidgetSortOrder{OwnerList, InvalidationWidget.Index} }, SortPredicate());
+		}
+	}
 
-		if (!InvalidationWidget.bContainedByWidgetHeap)
+	/** Returns and removes the biggest WidgetIndex from the list. */
+	FSlateInvalidationWidgetIndex HeapPop()
+	{
+		FSlateInvalidationWidgetIndex Result = Heap.HeapTop().GetWidgetIndex();
+		Heap.HeapPopDiscard(SortPredicate(), false);
+		OwnerList[Result].bContainedByWidgetPreHeap = false;
+		return Result;
+	}
+	
+	/** Returns the biggest WidgetIndex from the list. */
+	inline FSlateInvalidationWidgetIndex HeapPeek() const
+	{
+		return Heap.HeapTop().GetWidgetIndex();
+	}
+	
+	/** Returns the biggest WidgetIndex from the list. */
+	inline const FElement& HeapPeekElement() const
+	{
+		return Heap.HeapTop();
+	}
+
+	/** Remove range */
+	int32 RemoveRange(const FSlateInvalidationWidgetList::FIndexRange& Range)
+	{
+		int32 RemoveCount = 0;
+		for (int32 Index = 0; Index < Heap.Num(); ++Index)
 		{
-			InvalidationWidget.bContainedByWidgetHeap = true;
-			Heap.HeapPush(FElement{ WidgetIndex, FSlateInvalidationWidgetSortOrder{OwnerList, WidgetIndex} }, TWidgetOrderGreater());
+			if (Range.Include(Heap[Index].GetWidgetSortOrder()))
+			{
+				const FSlateInvalidationWidgetIndex WidgetIndex = Heap[Index].GetWidgetIndex();
+				OwnerList[WidgetIndex].bContainedByWidgetPreHeap = false;
+				Heap.HeapRemoveAt(Index, SortPredicate());
+				Index = INDEX_NONE; // start again
+			}
 		}
+		return RemoveCount;
 	}
 
-	/** Insert into the list at the proper order (see binary heap) only if it's not already contains by the list. */
-	void PushUnique(FSlateInvalidationWidgetList::InvalidationWidgetType& InvalidationWidget)
+	/** Empties the list, but doesn't change memory allocations. */
+	void Reset(bool bResetContained)
+	{
+		if (bResetContained)
+		{
+			for (const FElement& Element : Heap)
+			{
+				OwnerList[Element.GetWidgetIndex()].bContainedByWidgetPreHeap = false;
+			}
+		}
+
+		// Destroy the second allocator, if it exists.
+		Heap.Empty(NumberOfPreAllocatedElement);
+	}
+
+	/** Returns the number of elements in the list. */
+	inline int32 Num() const
+	{
+		return Heap.Num();
+	}
+
+	/** Does it contains the widget index. */
+	bool Contains_Debug(const FSlateInvalidationWidgetIndex WidgetIndex) const
+	{
+		return Heap.ContainsByPredicate([WidgetIndex](const FElement& Element)
+			{
+				return Element.GetWidgetIndex() == WidgetIndex;
+			});
+	}
+
+	/** @retuns true if the heap is heapified. */
+	inline bool IsValidHeap_Debug()
+	{
+		return Algo::IsHeap(Heap, SortPredicate());
+	}
+
+	/** Returns the raw list. */
+	inline const FElementContainer& GetRaw() const
+	{
+		return Heap;
+	}
+
+	/** Iterate over each element in the list in any order. */
+	template<typename Predicate>
+	void ForEachIndexes(Predicate Pred)
+	{
+		for (FElement& Element : Heap)
+		{
+			Pred(Element);
+		}
+	}
+
+private:
+	FElementContainer Heap;
+	FSlateInvalidationWidgetList& OwnerList;
+};
+
+
+/** */
+class FSlateInvalidationWidgetPostHeap
+{
+public:
+	using FElement = UE::Slate::Private::FSlateInvalidationWidgetHeapElement;
+	struct FWidgetOrderGreater
+	{
+		FORCEINLINE bool operator()(const FElement& A, const FElement& B) const
+		{
+			return B.GetWidgetSortOrder() < A.GetWidgetSortOrder();
+		}
+	};
+	using SortPredicate = FWidgetOrderGreater;
+	
+	static constexpr int32 NumberOfPreAllocatedElement = 100;
+	using FElementContainer = TArray<FElement, TInlineAllocator<NumberOfPreAllocatedElement>>;
+
+	FSlateInvalidationWidgetPostHeap(FSlateInvalidationWidgetList& InOwnerList)
+		: OwnerList(InOwnerList)
+		, WidgetCannotBeAdded(FSlateInvalidationWidgetIndex::Invalid)
+		, bIsHeap(false)
+	{ }
+
+public:
+	/** Insert into the list at the proper order (see binary heap) but only if it's not already contains by the list. */
+	void HeapPushUnique(FSlateInvalidationWidgetList::InvalidationWidgetType& InvalidationWidget)
 	{
+		check(bIsHeap);
 		check(InvalidationWidget.Index != FSlateInvalidationWidgetIndex::Invalid);
+		check(Contains_Debug(InvalidationWidget.Index) == InvalidationWidget.bContainedByWidgetPostHeap
+			|| WidgetCannotBeAdded == InvalidationWidget.Index);
 
-		check(Contains(InvalidationWidget.Index) == InvalidationWidget.bContainedByWidgetHeap
+		if (!InvalidationWidget.bContainedByWidgetPostHeap)
+		{
+			InvalidationWidget.bContainedByWidgetPostHeap = true;
+			Heap.HeapPush(FElement{ InvalidationWidget.Index, FSlateInvalidationWidgetSortOrder{OwnerList, InvalidationWidget.Index} }, SortPredicate());
+		}
+	}
+
+	/** Insert at the end of the list but only if it's not already contains by the list. */
+	void PushBackUnique(FSlateInvalidationWidgetList::InvalidationWidgetType& InvalidationWidget)
+	{
+		check(bIsHeap == false);
+		check(InvalidationWidget.Index != FSlateInvalidationWidgetIndex::Invalid);
+		check(Contains_Debug(InvalidationWidget.Index) == InvalidationWidget.bContainedByWidgetPostHeap
 			|| WidgetCannotBeAdded == InvalidationWidget.Index);
 
-		if (!InvalidationWidget.bContainedByWidgetHeap)
+		if (!InvalidationWidget.bContainedByWidgetPostHeap)
 		{
-			InvalidationWidget.bContainedByWidgetHeap = true;
-			Heap.HeapPush(FElement{ InvalidationWidget.Index, FSlateInvalidationWidgetSortOrder{OwnerList, InvalidationWidget.Index} }, TWidgetOrderGreater());
+			InvalidationWidget.bContainedByWidgetPostHeap = true;
+			Heap.Emplace(InvalidationWidget.Index, FSlateInvalidationWidgetSortOrder{OwnerList, InvalidationWidget.Index});
+		}
+	}
+
+	/** PushBackUnique or PushHeapUnique depending if the list is Heapified. */
+	void PushBackOrHeapUnique(FSlateInvalidationWidgetList::InvalidationWidgetType& InvalidationWidget)
+	{
+		if (bIsHeap)
+		{
+			HeapPushUnique(InvalidationWidget);
+		}
+		else
+		{
+			PushBackUnique(InvalidationWidget);
 		}
 	}
 
 	/** Returns and removes the biggest WidgetIndex from the list. */
-	FSlateInvalidationWidgetIndex Pop()
+	FSlateInvalidationWidgetIndex HeapPop()
 	{
-		FSlateInvalidationWidgetIndex Result = Heap.HeapTop().Get<0>();
-		Heap.HeapPopDiscard(TWidgetOrderGreater(), false);
-		OwnerList[Result].bContainedByWidgetHeap = false;
+		check(bIsHeap == true);
+		FSlateInvalidationWidgetIndex Result = Heap.HeapTop().GetWidgetIndex();
+		Heap.HeapPopDiscard(SortPredicate(), false);
+		OwnerList[Result].bContainedByWidgetPostHeap = false;
 		return Result;
 	}
 
+	/** Remove range */
+	int32 RemoveRange(const FSlateInvalidationWidgetList::FIndexRange& Range)
+	{
+		check(bIsHeap == false);
+		int32 RemoveCount = 0;
+		for (int32 Index = Heap.Num() -1; Index >= 0 ; --Index)
+		{
+			if (Range.Include(Heap[Index].GetWidgetSortOrder()))
+			{
+				const FSlateInvalidationWidgetIndex WidgetIndex = Heap[Index].GetWidgetIndex();
+				OwnerList[WidgetIndex].bContainedByWidgetPostHeap = false;
+				Heap.RemoveAtSwap(Index);
+				++RemoveCount;
+			}
+		}
+
+		return RemoveCount;
+	}
+
 	/** Empties the list, but doesn't change memory allocations. */
-	inline void Reset(bool bResetContained)
+	void Reset(bool bResetContained)
 	{
 		if (bResetContained)
 		{
 			for (const FElement& Element : Heap)
 			{
-				OwnerList[Element.Get<0>()].bContainedByWidgetHeap = false;
+				OwnerList[Element.GetWidgetIndex()].bContainedByWidgetPostHeap = false;
 			}
 		}
-		Heap.Reset();
+
+		// Destroy the second allocator, if it exists.
+		Heap.Empty(NumberOfPreAllocatedElement);
+		bIsHeap = false;
 	}
-	
+
+	void Heapify()
+	{
+		check(!bIsHeap);
+		Heap.Heapify(SortPredicate());
+		bIsHeap = true;
+	}
+
 	/** Returns the number of elements in the list. */
 	inline int32 Num() const
 	{
 		return Heap.Num();
 	}
 
+	/** Returns the number of elements in the list. */
+	inline bool IsHeap() const
+	{
+		return bIsHeap;
+	}
+
+	/** Does it contains the widget index. */
+	bool Contains_Debug(const FSlateInvalidationWidgetIndex WidgetIndex) const
+	{
+		return Heap.ContainsByPredicate([WidgetIndex](const FElement& Element)
+			{
+				return Element.GetWidgetIndex() == WidgetIndex;
+			});
+	}
+
+	/** @retuns true if the heap is heapified. */
+	inline bool IsValidHeap_Debug()
+	{
+		return Algo::IsHeap(Heap, SortPredicate());
+	}
+
 	/** Returns the raw list. */
-	const TArray<FElement, TInlineAllocator<100>>& GetRaw() const
+	inline const FElementContainer& GetRaw() const
 	{
 		return Heap;
 	}
 
+	template<typename Predicate>
+	void ForEachIndexes(Predicate Pred)
+	{
+		for (FElement& Element : Heap)
+		{
+			Pred(Element);
+		}
+	}
+	
 public:
 	struct FScopeWidgetCannotBeAdded
 	{
-		FScopeWidgetCannotBeAdded(FSlateInvalidationWidgetHeap& InHeap, FSlateInvalidationWidgetList::InvalidationWidgetType& InInvalidationWidget)
+		FScopeWidgetCannotBeAdded(FSlateInvalidationWidgetPostHeap& InHeap, FSlateInvalidationWidgetList::InvalidationWidgetType& InInvalidationWidget)
 			: Heap(InHeap)
 			, InvalidationWidget(InInvalidationWidget)
 			, WidgetIndex(InvalidationWidget.Index)
 		{
-			check(!InvalidationWidget.bContainedByWidgetHeap
+			check(!InvalidationWidget.bContainedByWidgetPostHeap
 				&& Heap.WidgetCannotBeAdded == FSlateInvalidationWidgetIndex::Invalid);
 			Heap.WidgetCannotBeAdded = InvalidationWidget.Index;
-			InvalidationWidget.bContainedByWidgetHeap = true;
+			InvalidationWidget.bContainedByWidgetPostHeap = true;
 		}
 		~FScopeWidgetCannotBeAdded()
 		{
 			Heap.WidgetCannotBeAdded = FSlateInvalidationWidgetIndex::Invalid;
 			check(Heap.OwnerList.IsValidIndex(WidgetIndex));
 			check(&Heap.OwnerList[WidgetIndex] == &InvalidationWidget);
-			InvalidationWidget.bContainedByWidgetHeap = false;
+			InvalidationWidget.bContainedByWidgetPostHeap = false;
 		}
-		FSlateInvalidationWidgetHeap& Heap;
+		FSlateInvalidationWidgetPostHeap& Heap;
 		FSlateInvalidationWidgetList::InvalidationWidgetType& InvalidationWidget;
 		FSlateInvalidationWidgetIndex WidgetIndex;
 	};
 
 private:
-	bool Contains(const FSlateInvalidationWidgetIndex WidgetIndex) const
-	{
-		return Heap.ContainsByPredicate([WidgetIndex](const FElement& Element)
-			{
-				return Element.Get<0>() == WidgetIndex;
-			});
-	}
-
-private:
-	TArray<FElement, TInlineAllocator<100>> Heap;
+	FElementContainer Heap;
 	FSlateInvalidationWidgetList& OwnerList;
 	FSlateInvalidationWidgetIndex WidgetCannotBeAdded;
+	bool bIsHeap;
 };
diff --git a/Engine/Source/Runtime/SlateCore/Private/FastUpdate/SlateInvalidationWidgetList.cpp b/Engine/Source/Runtime/SlateCore/Private/FastUpdate/SlateInvalidationWidgetList.cpp
--- a/Engine/Source/Runtime/SlateCore/Private/FastUpdate/SlateInvalidationWidgetList.cpp
+++ b/Engine/Source/Runtime/SlateCore/Private/FastUpdate/SlateInvalidationWidgetList.cpp
@@ -14,19 +14,17 @@
 #include "Templates/IsConst.h"
 #include "Templates/Tuple.h"
 #include "Types/ReflectionMetadata.h"
+#include "Types/SlateAttributeMetaData.h"
 
 #include <limits>
 
 
-DECLARE_CYCLE_STAT(TEXT("WidgetList ProcessInvalidation"), STAT_WidgetList_ProcessInvalidation, STATGROUP_Slate);
+DECLARE_CYCLE_STAT(TEXT("WidgetList ChildOrder Invalidation"), STAT_WidgetList_ProcessChildOrderInvalidation, STATGROUP_Slate);
 
 #define UE_SLATE_WITH_WIDGETLIST_UPDATEONLYWHATISNEEDED 0
+#define UE_SLATE_WITH_WIDGETLIST_ASSIGNINVALIDPROXYWHENREMOVED 0
 
-#define UE_SLATE_VERIFY_REBUILDWIDGETDATA_ORDER 0
-#define UE_SLATE_VERIFY_INVALID_INVALIDATIONHANDLE 0
-#define UE_SLATE_VERIFY_REMOVED_WIDGET_ARE_NOT_INVALIDATED 0
-
-#if UE_SLATE_VERIFY_REMOVED_WIDGET_ARE_NOT_INVALIDATED
+#if UE_SLATE_WITH_WIDGETLIST_ASSIGNINVALIDPROXYWHENREMOVED
 uint16 GSlateInvalidationWidgetIndex_RemovedIndex = 0xffee;
 #endif
 
@@ -34,7 +32,165 @@ uint16 GSlateInvalidationWidgetIndex_RemovedIndex = 0xffee;
 int32 FSlateInvalidationWidgetList::FArguments::MaxPreferedElementsNum = (1 << 10) - 1;
 int32 FSlateInvalidationWidgetList::FArguments::MaxSortOrderPaddingBetweenArray = (1 << 22) - 1;
 
+/** */
+FSlateInvalidationWidgetIndex FSlateInvalidationWidgetList::IProcessChildOrderInvalidationCallback::FReIndexOperation::ReIndex(FSlateInvalidationWidgetIndex Index) const
+{
+	check(Range.GetInclusiveMinWidgetIndex().ArrayIndex == Range.GetInclusiveMaxWidgetIndex().ArrayIndex);
+	check(Index.ArrayIndex == Range.GetInclusiveMaxWidgetIndex().ArrayIndex);
+	const IndexType NewElementIndex = Index.ElementIndex - Range.GetInclusiveMinWidgetIndex().ElementIndex + ReIndexTarget.ElementIndex;
+	return { ReIndexTarget.ArrayIndex, NewElementIndex };
+}
+
+
+/** */
+FSlateInvalidationWidgetList::FWidgetAttributeIterator::FWidgetAttributeIterator(FSlateInvalidationWidgetList const& InWidgetList)
+	: WidgetList(InWidgetList)
+	, CurrentWidgetIndex(FSlateInvalidationWidgetIndex::Invalid)
+	, CurrentWidgetSortOrder(FSlateInvalidationWidgetSortOrder::LimitMax())
+	, AttributeIndex(0)
+	, MoveToWidgetIndexOnNextAdvance(FSlateInvalidationWidgetIndex::Invalid)
+{
+	int32 ArrayIndex = WidgetList.FirstArrayIndex;
+	while(ArrayIndex != INDEX_NONE)
+	{
+		const FArrayNode& ArrayNode = WidgetList.Data[ArrayIndex];
+		if (ArrayNode.ElementIndexList_WidgetWithRegisteredSlateAttribute.Num() > 0)
+		{
+			CurrentWidgetIndex = FSlateInvalidationWidgetIndex{(IndexType)ArrayIndex, ArrayNode.ElementIndexList_WidgetWithRegisteredSlateAttribute[0]};
+			CurrentWidgetSortOrder = FSlateInvalidationWidgetSortOrder{InWidgetList, CurrentWidgetIndex};
+			break;
+		}
+		ArrayIndex = ArrayNode.NextArrayIndex;
+	}
+}
+
+
+void FSlateInvalidationWidgetList::FWidgetAttributeIterator::PreChildRemove(const FIndexRange& Range)
+{
+	if (CurrentWidgetSortOrder < Range.GetInclusiveMinWidgetSortOrder())
+	{
+		MoveToWidgetIndexOnNextAdvance = WidgetList.IncrementIndex(Range.GetInclusiveMaxWidgetIndex());
+	}
+}
+
+
+void FSlateInvalidationWidgetList::FWidgetAttributeIterator::ReIndexed(IProcessChildOrderInvalidationCallback::FReIndexOperation const& Operation)
+{
+	if (MoveToWidgetIndexOnNextAdvance != FSlateInvalidationWidgetIndex::Invalid)
+	{
+		if (Operation.GetRange().Include(FSlateInvalidationWidgetSortOrder{ WidgetList, MoveToWidgetIndexOnNextAdvance }))
+		{
+			MoveToWidgetIndexOnNextAdvance = Operation.ReIndex(MoveToWidgetIndexOnNextAdvance);
+		}
+	}
+	else
+	{
+		CurrentWidgetIndex = Operation.ReIndex(CurrentWidgetIndex);
+		CurrentWidgetSortOrder = FSlateInvalidationWidgetSortOrder{WidgetList, CurrentWidgetIndex };
+	}
+}
+
+
+void FSlateInvalidationWidgetList::FWidgetAttributeIterator::PostResort()
+{
+	if (MoveToWidgetIndexOnNextAdvance != FSlateInvalidationWidgetIndex::Invalid)
+	{
+		CurrentWidgetSortOrder = FSlateInvalidationWidgetSortOrder{ WidgetList, CurrentWidgetIndex };
+	}
+}
+
+
+void FSlateInvalidationWidgetList::FWidgetAttributeIterator::ProxiesBuilt(FIndexRange const& Range)
+{
+	MoveToWidgetIndexOnNextAdvance = Range.GetInclusiveMinWidgetIndex();
+}
+
+
+void FSlateInvalidationWidgetList::FWidgetAttributeIterator::FixCurrentWidgetIndex()
+{
+	//MoveToWidgetIndexOnNextAdvance is to be consider as a valid Attribute
+	if (MoveToWidgetIndexOnNextAdvance != FSlateInvalidationWidgetIndex::Invalid)
+	{
+		const int32 ArrayIndex = MoveToWidgetIndexOnNextAdvance.ArrayIndex;
+		check(WidgetList.Data.IsValidIndex(ArrayIndex));
+
+		const FArrayNode& ArrayNode = WidgetList.Data[ArrayIndex];
+		AttributeIndex = ArrayNode.ElementIndexList_WidgetWithRegisteredSlateAttribute.FindLowerBound(MoveToWidgetIndexOnNextAdvance.ElementIndex);
+		if (AttributeIndex == INDEX_NONE)
+		{
+			AdvanceArrayIndex(ArrayIndex);
+		}
+		else
+		{
+			CurrentWidgetIndex = FSlateInvalidationWidgetIndex{ (IndexType)ArrayIndex, ArrayNode.ElementIndexList_WidgetWithRegisteredSlateAttribute[AttributeIndex] };
+			CurrentWidgetSortOrder = FSlateInvalidationWidgetSortOrder{ WidgetList, CurrentWidgetIndex };
+		}
+		MoveToWidgetIndexOnNextAdvance = FSlateInvalidationWidgetIndex::Invalid;
+	}
+}
+
+
+void FSlateInvalidationWidgetList::FWidgetAttributeIterator::Seek(FSlateInvalidationWidgetIndex SeekTo)
+{
+	MoveToWidgetIndexOnNextAdvance = SeekTo;
+}
+
+
+void FSlateInvalidationWidgetList::FWidgetAttributeIterator::Advance()
+{
+	check(MoveToWidgetIndexOnNextAdvance == FSlateInvalidationWidgetIndex::Invalid);
+
+	++AttributeIndex;
+
+	const FArrayNode& ArrayNode = WidgetList.Data[CurrentWidgetIndex.ArrayIndex];
+	if (AttributeIndex >= ArrayNode.ElementIndexList_WidgetWithRegisteredSlateAttribute.Num())
+	{
+		int32 ArrayIndex = ArrayNode.NextArrayIndex;
+		AdvanceArrayIndex(ArrayIndex);
+	}
+	else
+	{
+		int32 ArrayIndex = CurrentWidgetIndex.ArrayIndex;
+		CurrentWidgetIndex = FSlateInvalidationWidgetIndex{ (IndexType)ArrayIndex, ArrayNode.ElementIndexList_WidgetWithRegisteredSlateAttribute[AttributeIndex] };
+		CurrentWidgetSortOrder = (CurrentWidgetIndex != FSlateInvalidationWidgetIndex::Invalid)
+			? FSlateInvalidationWidgetSortOrder{ WidgetList, CurrentWidgetIndex }
+			: FSlateInvalidationWidgetSortOrder::LimitMax();
+	}
+}
+
+
+void FSlateInvalidationWidgetList::FWidgetAttributeIterator::AdvanceArrayIndex(int32 ArrayIndex)
+{
+	AttributeIndex = 0;
+	CurrentWidgetIndex = FSlateInvalidationWidgetIndex::Invalid;
+	CurrentWidgetSortOrder = FSlateInvalidationWidgetSortOrder::LimitMax();
+
+	while (ArrayIndex != INDEX_NONE)
+	{
+		const FArrayNode& NewArrayNode = WidgetList.Data[ArrayIndex];
+		if (NewArrayNode.ElementIndexList_WidgetWithRegisteredSlateAttribute.Num() > 0)
+		{
+			CurrentWidgetIndex = FSlateInvalidationWidgetIndex{ (IndexType)ArrayIndex, NewArrayNode.ElementIndexList_WidgetWithRegisteredSlateAttribute[0] };
+			CurrentWidgetSortOrder = FSlateInvalidationWidgetSortOrder{ WidgetList, CurrentWidgetIndex };
+			break;
+		}
+		ArrayIndex = NewArrayNode.NextArrayIndex;
+	}
+}
 
+
+/** */
+void FSlateInvalidationWidgetList::FArrayNode::RemoveElementIndexBiggerOrEqualThan(IndexType ElementIndex)
+{
+	const int32 FoundIndex = ElementIndexList_WidgetWithRegisteredSlateAttribute.FindLowerBound(ElementIndex);
+	if (ElementIndexList_WidgetWithRegisteredSlateAttribute.IsValidIndex(FoundIndex))
+	{
+		ElementIndexList_WidgetWithRegisteredSlateAttribute.RemoveAt(FoundIndex, ElementIndexList_WidgetWithRegisteredSlateAttribute.Num() - FoundIndex);
+	}
+}
+
+
+/** */
 FSlateInvalidationWidgetList::FSlateInvalidationWidgetList(FSlateInvalidationRootHandle InOwner, const FArguments& Args)
 	: Owner(InOwner)
 	, WidgetListConfig(Args)
@@ -56,7 +212,6 @@ FSlateInvalidationWidgetList::FSlateInvalidationWidgetList(FSlateInvalidationRoo
 FSlateInvalidationWidgetIndex FSlateInvalidationWidgetList::_BuildWidgetList_Recursive(TSharedRef<SWidget>& Widget, FSlateInvalidationWidgetIndex ParentIndex, IndexType& LastestIndex, bool bParentVisible, bool bParentVolatile)
 {
 	const bool bIsEmpty = IsEmpty();
-	IndexType TemporaryIndex = LastestIndex;
 	const FSlateInvalidationWidgetIndex NewIndex = EmplaceInsertAfter(LastestIndex, Widget);
 	LastestIndex = NewIndex.ArrayIndex;
 
@@ -72,7 +227,13 @@ FSlateInvalidationWidgetIndex FSlateInvalidationWidgetList::_BuildWidgetList_Rec
 		WidgetProxy.Visibility = Visibility;
 	}
 
-#if WITH_SLATE_DEBUGGING
+	if (Widget->HasRegisteredSlateAttribute())
+	{
+		// The list is already sorted at this point. Add to the end.
+		Data[LastestIndex].ElementIndexList_WidgetWithRegisteredSlateAttribute.AddUnsorted(NewIndex.ElementIndex);
+	}
+
+#if UE_SLATE_WITH_INVALIDATIONWIDGETLIST_DEBUGGING
 	if (WidgetListConfig.bAssignedWidgetIndex)
 #endif
 	{
@@ -111,6 +272,8 @@ FSlateInvalidationWidgetIndex FSlateInvalidationWidgetList::_BuildWidgetList_Rec
 
 void FSlateInvalidationWidgetList::BuildWidgetList(TSharedRef<SWidget> InRoot)
 {
+	SCOPED_NAMED_EVENT(Slate_InvalidationList_ProcessBuild, FColorList::Blue);
+
 	Reset();
 	Root = InRoot;
 	FSlateInvalidationRoot*  InvalidationRoot = Owner.GetInvalidationRoot();
@@ -171,323 +334,186 @@ void FSlateInvalidationWidgetList::_RebuildWidgetListTree(TSharedRef<SWidget> Wi
 }
 
 
-//test if it's FSlateInvalidationWidgetIndex is valid and this == Widget.
-//when we remove a widget we may not invalidate it's index and it may not point to the correct widget anymore
-//That behavior is normal, if it occurred during ProcessChildOrderInvalidation.
-void FSlateInvalidationWidgetList::ProcessChildOrderInvalidation(const TArray<TWeakPtr<SWidget>>& InvalidatedWidgets)
+bool FSlateInvalidationWidgetList::ProcessChildOrderInvalidation(InvalidationWidgetType const& InvalidationWidget, IProcessChildOrderInvalidationCallback& Callback)
 {
-	SCOPE_CYCLE_COUNTER(STAT_WidgetList_ProcessInvalidation);
+	SCOPE_CYCLE_COUNTER(STAT_WidgetList_ProcessChildOrderInvalidation);
+
+	bool bIsInvalidationWidgetStillValid = true;
+	TGuardValue<IProcessChildOrderInvalidationCallback*> TmpGuard(CurrentInvalidationCallback, &Callback);
 
-	FSlateInvalidationRoot* InvalidationRoot = Owner.GetInvalidationRoot();
-	if (InvalidationRoot && InvalidationRoot->GetFastPathGenerationNumber() != GenerationNumber)
+	SWidget* WidgetPtr = InvalidationWidget.GetWidget();
+	check(WidgetPtr);
 	{
-		TSharedPtr<SWidget> RootPinned = GetRoot().Pin();
-		if (ensure(RootPinned))
+		const bool bWidgetIndexMatches = WidgetPtr->GetProxyHandle().GetWidgetIndex() == InvalidationWidget.Index;
+		if (!bWidgetIndexMatches)
 		{
-			BuildWidgetList(RootPinned.ToSharedRef());
+			ensureMsgf(false, TEXT("The widget index doesn't match the index in the InvalidationWidgetList"));
+			return false;
 		}
-		return;
 	}
 
-	if (InvalidatedWidgets.Num() > 0)
+	int32 StartChildIndex = INDEX_NONE;
+	enum class EOperation : uint8 { Remove, Cut, None } CurrentOperation = EOperation::None;
+
+	// Find all the InvalidationWidget children
 	{
-		if (FirstIndex() == FSlateInvalidationWidgetIndex::Invalid)
+		FIndexRange ChildToRemoveRange;
+
+		// Was added, but it should not be there anymore
+		if (!ShouldBeAdded(WidgetPtr))
 		{
-			ensureMsgf(false, TEXT("No tree were built and widgets were invalidated."));
-			return;
+			// This proxy will not be valid after this call
+			bIsInvalidationWidgetStillValid = false;
+			// Remove it and all it's child
+			ChildToRemoveRange = { *this, InvalidationWidget.Index, InvalidationWidget.LeafMostChildIndex };
+			CurrentOperation = EOperation::Remove;
+			// Do not add new child
+			//Result.ChildBuilt
 		}
-
-		FMemMark Mark(FMemStack::Get());
-		using TInvalidateWidgetIndexType = TTuple<TSharedRef<SWidget>, FSlateInvalidationWidgetIndex, FSlateInvalidationWidgetSortOrder>;
-		TArray<TInvalidateWidgetIndexType, TMemStackAllocator<>> InvalidatedWidgetIndexes;
-		InvalidatedWidgetIndexes.Reserve(InvalidatedWidgets.Num());
-
-#if UE_SLATE_VERIFY_INVALID_INVALIDATIONHANDLE
-		TArray<TWeakPtr<SWidget>> VerifyWidgetInvalidationHandle;
-		InvalidatedWidgetIndexes.Reserve(InvalidatedWidgets.Num());
-#endif
-
-		// Build Invalidation Indexes
+		// If it is not supposed to had child, but had some, them remove them
+		else if (!ShouldDoRecursion(WidgetPtr))
 		{
-			for (const TWeakPtr<SWidget>& InvalidatedWidget : InvalidatedWidgets)
-			{
-				if (const TSharedPtr<SWidget> Widget = InvalidatedWidget.Pin())
-				{
-					const FSlateInvalidationWidgetIndex WidgetIndex = Widget->GetProxyHandle().GetWidgetIndex();
-					if (IsValidIndex(WidgetIndex) && (*this)[WidgetIndex].GetWidget() == Widget.Get())
-					{
-						InvalidatedWidgetIndexes.Emplace(Widget.ToSharedRef(), WidgetIndex, FSlateInvalidationWidgetSortOrder{ *this, WidgetIndex });
-					}
-#if UE_SLATE_VERIFY_INVALID_INVALIDATIONHANDLE
-					else
-					{
-						// This widget has requested an invalidation ChildOrder but it's index is not valid.
-						//Confirm that it will have a valid index at the end of the algo.
-						VerifyWidgetInvalidationHandle.Add(InvalidatedWidget);
-						Widget->SetFastPathProxyHandle(FWidgetProxyHandle{ Owner, FSlateInvalidationWidgetIndex::Invalid, FSlateInvalidationWidgetSortOrder{}, GenerationNumber });
-					}
-#endif
-				}
-			}
-
-			// Sort to invalidate them in the same order they would be if we were in slow path.
-			InvalidatedWidgetIndexes.Sort([](const TInvalidateWidgetIndexType& A, const TInvalidateWidgetIndexType& B) {
-				const FSlateInvalidationWidgetSortOrder OrderA = A.Get<2>();
-				const FSlateInvalidationWidgetSortOrder OrderB = B.Get<2>();
-				return OrderA < OrderB;
-				});
-
-			const int32 InvalidatedNum = Algo::Unique(InvalidatedWidgetIndexes);
-			if (InvalidatedNum != InvalidatedWidgetIndexes.Num())
+			if (InvalidationWidget.Index != InvalidationWidget.LeafMostChildIndex)
 			{
-				InvalidatedWidgetIndexes.RemoveAt(InvalidatedNum, InvalidatedWidgetIndexes.Num() - InvalidatedNum, false);
+				// All it's child (but not him)
+				ChildToRemoveRange = { *this, IncrementIndex(InvalidationWidget.Index), InvalidationWidget.LeafMostChildIndex };
+				CurrentOperation = EOperation::Remove;
+				// Do not add new child
+				//Result.ChildBuilt
 			}
 		}
-
-		struct FChildOrderInvalidationData
+		// Find the difference between list and the reality (if it used to have at least 1 child)
+		else if (InvalidationWidget.Index != InvalidationWidget.LeafMostChildIndex)
 		{
-			TSharedPtr<SWidget> Widget;
-			FIndexRange Range;
-			FSlateInvalidationWidgetIndex WhereToCut;
-			int32 ChildAtToStartWith;
-			bool bRemove;
-			bool operator==(const FChildOrderInvalidationData& Other) const
-			{
-				return Other.Widget == Widget
-					&& Other.bRemove == bRemove
-					&& (bRemove ? Other.Range == Range : Other.WhereToCut == WhereToCut)
-					&& Other.ChildAtToStartWith == ChildAtToStartWith;
-			}
-			FChildOrderInvalidationData(TSharedPtr<SWidget> InWidget, FIndexRange InRange, int32 InChildAtToStartWith)
-				: Widget(InWidget)
-				, Range(InRange)
-				, WhereToCut(FSlateInvalidationWidgetIndex::Invalid)
-				, ChildAtToStartWith(InChildAtToStartWith)
-				, bRemove(true) {}
-			FChildOrderInvalidationData(TSharedPtr<SWidget> InWidget, FSlateInvalidationWidgetIndex InWhereToCut, int32 InChildAtToStartWith)
-				: Widget(InWidget)
-				, Range()
-				, WhereToCut(InWhereToCut)
-				, ChildAtToStartWith(InChildAtToStartWith)
-				, bRemove(false) {}
-		};
-
-		TArray<FChildOrderInvalidationData, TMemStackAllocator<>> RebuildWigetData;
-		RebuildWigetData.Reserve(InvalidatedWidgetIndexes.Num());
+#if UE_SLATE_WITH_WIDGETLIST_UPDATEONLYWHATISNEEDED
+			// Find all it's previous children
+			FMemMark Mark2(FMemStack::Get());
+			TArray<FFindChildrenElement, TMemStackAllocator<>> PreviousChildrenWidget;
+			_FindChildren(InvalidationWidget.Index, PreviousChildrenWidget);
 
+			FChildren* InvalidatedChildren = WidgetPtr->GetAllChildren();
+			check(InvalidatedChildren);
 
-		// Compute the InvalidationDatas
-		{
-			SCOPED_NAMED_EVENT(Slate_InvalidationList_ProcessCompute, FColorList::Blue);
-			TArray<FIndexRange, TMemStackAllocator<>> RebuildWidgetRange;
-			RebuildWidgetRange.Reserve(InvalidatedWidgetIndexes.Num());
+			// Find where it starts to get different
+			int32 InvalidatedChildrenIndex = 0;
 
-			for (int32 InvalidatedIndex = 0; InvalidatedIndex < InvalidatedWidgetIndexes.Num(); ++InvalidatedIndex)
+			const int32 InvalidatedChildrenNum = InvalidatedChildren->Num();
+			const int32 PreviousChildrenNum = PreviousChildrenWidget.Num();
+			int32 PreviousIndex = 0;
+			for (; InvalidatedChildrenIndex < InvalidatedChildrenNum && PreviousIndex < PreviousChildrenNum; ++InvalidatedChildrenIndex)
 			{
-				FSlateInvalidationWidgetIndex WidgetIndex = InvalidatedWidgetIndexes[InvalidatedIndex].Get<1>();
-
-				// Is it already going to be invalidated by something else
+				TSharedRef<SWidget> NewWidget = InvalidatedChildren->GetChildAt(InvalidatedChildrenIndex);
+				if (ShouldBeAdded(NewWidget))
 				{
-					const FSlateInvalidationWidgetSortOrder WidgetOrder = { *this, WidgetIndex };
-					if (RebuildWidgetRange.ContainsByPredicate([WidgetOrder](const FIndexRange& Range) {
-						return Range.Include(WidgetOrder);
-						}))
+					if (&NewWidget.Get() != PreviousChildrenWidget[PreviousIndex].Get<0>())
 					{
-						continue;
+						break;
 					}
-				}
 
-				const InvalidationWidgetType& InvalidationWidget = (*this)[WidgetIndex];
-
-				// If the parent is invalid, it's means that we are the root. Rebuild the full tree.
-				if (InvalidationWidget.ParentIndex == FSlateInvalidationWidgetIndex::Invalid)
-				{
-					SWidget* WidgetPtr = InvalidationWidget.GetWidget();
-					if (ensure(WidgetPtr))
-					{
-						ensure(GetRoot().Pin().Get() == WidgetPtr);
-						UE_LOG(LogSlate, Warning, TEXT("Performance: A BuildTree() was requested by a ChildOrder invalidation."));
-						BuildWidgetList(WidgetPtr->AsShared());
-					}
-					return;
+					++PreviousIndex;
 				}
+			}
 
-				TSharedRef<SWidget> WidgetRef = InvalidatedWidgetIndexes[InvalidatedIndex].Get<0>();
-				ensureMsgf(WidgetRef->GetProxyHandle().GetWidgetIndex() == WidgetIndex, TEXT("The widget index doesn't match the index in the InvalidationWidgetList"));
-
-				// Find all the InvalidationWidget children
-				{
-					// Was added, but it should not be there anymore
-					if (!ShouldBeAdded(WidgetRef))
-					{
-						const FIndexRange Range = { *this, WidgetIndex, InvalidationWidget.LeafMostChildIndex };
-						RebuildWidgetRange.Add(Range);
-						RebuildWigetData.Emplace(TSharedPtr<SWidget>(), Range, INDEX_NONE);
-					}
-					// If it is not supposed to had child, but had some, them remove them
-					else if (!ShouldDoRecursion(WidgetRef))
-					{
-						if (WidgetIndex != InvalidationWidget.LeafMostChildIndex)
-						{
-							const FIndexRange Range = { *this, IncrementIndex(WidgetIndex), InvalidationWidget.LeafMostChildIndex };
-							RebuildWidgetRange.Add(Range);
-							RebuildWigetData.Emplace(TSharedPtr<SWidget>(), Range, INDEX_NONE);
-						}
-					}
-					// Find the difference between list and the reality (if it used to have at least 1 child)
-					else if (WidgetIndex != InvalidationWidget.LeafMostChildIndex)
-					{
-#if UE_SLATE_WITH_WIDGETLIST_UPDATEONLYWHATISNEEDED
-						// Find all it's previous children
-						FMemMark Mark2(FMemStack::Get());
-						TArray<SWidget*, TMemStackAllocator<>> PreviousChildrenWidget;
-						_FindChildren(WidgetIndex, PreviousChildrenWidget);
-
-						FChildren* InvalidatedChildren = WidgetRef->GetAllChildren();
-						check(InvalidatedChildren);
-
-						/**
-						 ****** ****** ****** ****** ******
-						 * todo @pat improve with a loop
-						 * As a fist step, invalidate everything from this point. Could be a loop and invalidate what need to be invalidated before move the other.
-						 * We have to see if that case occurred in data before doing the code for it.
-						 ****** ****** ****** ****** ******
-						*/
-
-						// Find where it starts to get different
-						FSlateInvalidationWidgetIndex IndexWhereToStart = IncrementIndex(WidgetIndex);
-						FSlateInvalidationWidgetIndex IndexWhereToEnd = InvalidationWidget.LeafMostChildIndex;
-						int32 InvalidatedChildrenIndex = 0;
-
-						const int32 InvalidatedChildrenNum = InvalidatedChildren->Num();
-						const int32 PreviousChildrenNum = PreviousChildrenWidget.Num();
-						int32 PreviousIndex = 0;
-						for (; InvalidatedChildrenIndex < InvalidatedChildrenNum && PreviousIndex < PreviousChildrenNum; ++InvalidatedChildrenIndex)
-						{
-							TSharedRef<SWidget> NewWidget = InvalidatedChildren->GetChildAt(InvalidatedChildrenIndex);
-							if (ShouldBeAdded(NewWidget))
-							{
-								if (&NewWidget.Get() != PreviousChildrenWidget[PreviousIndex])
-								{
-									break;
-								}
-								IndexWhereToStart = IncrementIndex((*this)[IndexWhereToStart].LeafMostChildIndex);
-								check(IsValidIndex(IndexWhereToStart));
-
-								++PreviousIndex;
-							}
-						}
-
-						if (InvalidatedChildrenIndex >= InvalidatedChildrenNum && PreviousIndex >= PreviousChildrenNum)
-						{
-							// The widget was invalidated but nothing changed. This could be normal if a widget was removed, then re-added.
-						}
-						else if (PreviousIndex >= PreviousChildrenNum)
-						{
-							// Nothing to remove, but need to add some. We want to break the array.
-							RebuildWigetData.Emplace(WidgetRef, InvalidationWidget.LeafMostChildIndex, InvalidatedChildrenIndex);
-						}
-						else
-						{
-							const FIndexRange Range = { *this, IndexWhereToStart, IndexWhereToEnd };
-							check(Range.OrderMin <= Range.OrderMax);
-							RebuildWidgetRange.Add(Range);
-							RebuildWigetData.Emplace(WidgetRef, Range, InvalidatedChildrenIndex);
-						}
+			if (InvalidatedChildrenIndex >= InvalidatedChildrenNum && PreviousIndex >= PreviousChildrenNum)
+			{
+				// The widget was invalidated but nothing changed. This could be normal if a widget was removed, then re-added.
+				//InvalidationWidgetReason::Layout still need to be process
+			}
+			else if (PreviousIndex >= PreviousChildrenNum)
+			{
+				// Nothing to remove.
+				//Result.ChildRemoved
+				// Need to add widgets. We want to break the array first.
+				StartChildIndex = InvalidatedChildrenIndex;
+				CurrentOperation = EOperation::Cut;
+			}
+			else
+			{
+				check(PreviousIndex < PreviousChildrenNum && PreviousChildrenNum != 0);
+				const FIndexRange Range = { *this, PreviousChildrenWidget[PreviousIndex].Get<1>(), PreviousChildrenWidget.Last().Get<1>() };
+				ChildToRemoveRange = Range;
+				StartChildIndex = InvalidatedChildrenIndex;
+				CurrentOperation = EOperation::Remove;
+			}
 #else
-						// Remove every children and rebuld the widget.
-						const FSlateInvalidationWidgetIndex IndexWhereToStart = IncrementIndex(WidgetIndex);
-						const FSlateInvalidationWidgetIndex IndexWhereToEnd = InvalidationWidget.LeafMostChildIndex;
-						const FIndexRange Range = { *this, IndexWhereToStart, IndexWhereToEnd };
-						RebuildWidgetRange.Add(Range);
-						RebuildWigetData.Emplace(WidgetRef, Range, 0);
+			// All it's child (but not him)
+			ChildToRemoveRange = { *this, IncrementIndex(InvalidationWidget.Index), InvalidationWidget.LeafMostChildIndex };
+			StartChildIndex = 0;
+			CurrentOperation = EOperation::Remove;
 #endif
-					}
-					// There was not child, but maybe it has some now
-					else
-					{
-						FChildren* InvalidatedChildren = WidgetRef->GetAllChildren();
-						check(InvalidatedChildren);
-						if (InvalidatedChildren->Num() > 0)
-						{
-							// Nothing to remove, but need to add new. We want to break the array.
-							RebuildWigetData.Emplace(WidgetRef, InvalidationWidget.LeafMostChildIndex, 0);
-						}
-					}
-				}
-			}
 		}
-
-#if UE_SLATE_VERIFY_REBUILDWIDGETDATA_ORDER
-		// The array was sorted, but we may have invalidated a parent to only add at the end and the sub tree before is also invalidated.
-		//A (B(1,2), C(5,6)) => A (B(1,2,3), C(5, 6) D(8))
-		//A and B are invalidated. They were sorted (A, B). Now we know that D and B process. B need to be process after.
-		//This should already work (because we are doing the reverse operation). Lets confirm.
+		// There was not child, but maybe it has some now
+		else
 		{
-			FMemMark Mark2(FMemStack::Get());
-			TArray<FChildOrderInvalidationData, TMemStackAllocator<>> SortedRebuildWigetData = RebuildWigetData;
-			FSlateInvalidationWidgetList* Self = this;
-			SortedRebuildWigetData.Sort([Self](const FChildOrderInvalidationData& A, const FChildOrderInvalidationData& B) -> bool
-				{
-					const FSlateInvalidationWidgetIndex AIndex = A.bRemove ? A.Range.InclusiveMin : A.WhereToCut;
-					const FSlateInvalidationWidgetIndex BIndex = B.bRemove ? B.Range.InclusiveMin : B.WhereToCut;
-					return FSlateInvalidationWidgetSortOrder(*Self, AIndex) < FSlateInvalidationWidgetSortOrder(*Self, BIndex);
-				});
-
-			// Is the sorted list in the same order
-			ensureAlways(SortedRebuildWigetData == RebuildWigetData);
+			// Nothing to remove.
+			//Result.ChildRemoved
+			// Add new child after this InvalidationWidget, so we want to break/cut the array.
+			StartChildIndex = 0;
+			CurrentOperation = EOperation::Cut;
 		}
-#endif //UE_SLATE_VERIFY_REBUILDWIDGETDATA_ORDER
 
+		// Remove child
 		{
 			SCOPED_NAMED_EVENT(Slate_InvalidationList_ProcessRemove, FColorList::Blue);
-			for (int32 Index = RebuildWigetData.Num() - 1; Index >= 0; --Index)
+			if (CurrentOperation == EOperation::Remove)
 			{
-				const FChildOrderInvalidationData& RebuildData = RebuildWigetData[Index];
-				if (RebuildData.bRemove)
-				{
-					_RemoveRangeFromSameParent(RebuildData.Range);
-				}
-				else
-				{
-					CutArray(RebuildData.WhereToCut);
-				}
+				ensureMsgf(ChildToRemoveRange.IsValid(), TEXT("The ChildToRemoveRange should be valid at this point."));
+				Callback.PreChildRemove(ChildToRemoveRange);
+				_RemoveRangeFromSameParent(ChildToRemoveRange);
 			}
-		}
-
-		{
-			SCOPED_NAMED_EVENT(Slate_InvalidationList_ProcessRebuild, FColorList::Blue);
-			for (const FChildOrderInvalidationData& RebuildData : RebuildWigetData)
+			else if (CurrentOperation == EOperation::Cut)
 			{
-				if (RebuildData.Widget)
-				{
-#if UE_SLATE_VERIFY_REMOVED_WIDGET_ARE_NOT_INVALIDATED
-					FSlateInvalidationWidgetIndex SlateInvalidationWidgetIndexRemoved = { GSlateInvalidationWidgetIndex_RemovedIndex, GSlateInvalidationWidgetIndex_RemovedIndex };
-					ensure(RebuildData.Widget->GetProxyHandle().GetWidgetIndex() != SlateInvalidationWidgetIndexRemoved);
-#endif
-					_RebuildWidgetListTree(RebuildData.Widget.ToSharedRef(), RebuildData.ChildAtToStartWith);
-				}
+				CutArray(InvalidationWidget.LeafMostChildIndex);
 			}
 		}
 
-
-#if UE_SLATE_VERIFY_INVALID_INVALIDATIONHANDLE
+		// Rebuild
+		if (StartChildIndex != INDEX_NONE)
 		{
-			for (TWeakPtr<SWidget>& WidgetInvalidated : VerifyWidgetInvalidationHandle)
+			FChildren* InvalidatedChildren = WidgetPtr->GetAllChildren();
+			check(InvalidatedChildren);
+			if (InvalidatedChildren->Num() > 0)
 			{
-				if (TSharedPtr<SWidget> Widget = WidgetInvalidated.Pin())
+				SCOPED_NAMED_EVENT(Slate_InvalidationList_ProcessRebuild, FColorList::Blue);
+				FSlateInvalidationWidgetIndex PreviousLeafMostChildIndex = InvalidationWidget.LeafMostChildIndex;
+				_RebuildWidgetListTree(WidgetPtr->AsShared(), StartChildIndex);
+
+				const FIndexRange BuiltRange = (InvalidationWidget.LeafMostChildIndex == PreviousLeafMostChildIndex)
+					? FIndexRange{ *this, PreviousLeafMostChildIndex, InvalidationWidget.LeafMostChildIndex }
+					: FIndexRange{ *this, IncrementIndex(PreviousLeafMostChildIndex), InvalidationWidget.LeafMostChildIndex };
+				if (BuiltRange.IsValid())
 				{
-					const FSlateInvalidationWidgetIndex WidgetIndex = Widget->GetProxyHandle().GetWidgetIndex();
-					ensureMsgf(WidgetIndex != FSlateInvalidationWidgetIndex::Invalid, TEXT("The widget '%s' requested a ChildOrder but didn't have a valid index and was not rebuilt by something else.")
-						, *FReflectionMetaData::GetWidgetDebugInfo(Widget.Get()));
+					Callback.ProxiesBuilt(BuiltRange);
 				}
 			}
 		}
-#endif
 	}
+
+	return bIsInvalidationWidgetStillValid;
 }
 
 
+void FSlateInvalidationWidgetList::ProcessAttributeRegistrationInvalidation(InvalidationWidgetType const& InvalidationWidget)
+{
+	SWidget* WidgetPtr = InvalidationWidget.GetWidget();
+	check(WidgetPtr);
+
+	bool bNeedsToBeThere = false;
+	if (FSlateAttributeMetaData* MetaData = FSlateAttributeMetaData::FindMetaData(*WidgetPtr))
+	{
+		bNeedsToBeThere = (MetaData->RegisteredNum() > 0);
+	}
+
+	if (bNeedsToBeThere)
+	{
+		Data[InvalidationWidget.Index.ArrayIndex].ElementIndexList_WidgetWithRegisteredSlateAttribute.InsertUnique(InvalidationWidget.Index.ElementIndex);
+	}
+	else
+	{
+		Data[InvalidationWidget.Index.ArrayIndex].ElementIndexList_WidgetWithRegisteredSlateAttribute.RemoveSingle(InvalidationWidget.Index.ElementIndex);
+	}
+}
+
 namespace
 {
 	template<typename TSlateInvalidationWidgetList, typename Predicate>
@@ -518,13 +544,13 @@ namespace
 }
 
 
-void FSlateInvalidationWidgetList::_FindChildren(FSlateInvalidationWidgetIndex WidgetIndex, TArray<SWidget*, TMemStackAllocator<>>& Widgets) const
+void FSlateInvalidationWidgetList::_FindChildren(FSlateInvalidationWidgetIndex WidgetIndex, TArray<FFindChildrenElement, TMemStackAllocator<>>& Widgets) const
 {
 	Widgets.Reserve(16);
 	const InvalidationWidgetType& InvalidationWidget = (*this)[WidgetIndex];
 	ForEachChildren(*this, InvalidationWidget, WidgetIndex, [&Widgets](const InvalidationWidgetType& ChildWidget)
 		{
-			Widgets.Add(ChildWidget.GetWidget());
+			Widgets.Emplace(ChildWidget.GetWidget(), ChildWidget.Index);
 		});
 }
 
@@ -662,17 +688,48 @@ void FSlateInvalidationWidgetList::RebuildOrderIndex(IndexType StartFrom)
 	check(WidgetListConfig.PreferedElementsNum < WidgetListConfig.MaxPreferedElementsNum);
 
 	FSlateInvalidationWidgetList* Self = this;
-	auto SetValue = [Self](IndexType ArrayIndex)
+	auto SetValue = [Self](IndexType ArrayIndex, int32 NewSortOrder)
 		{
 			ElementListType& ElementList = Self->Data[ArrayIndex].ElementList;
-			for (int32 ElementIndex = Self->Data[ArrayIndex].StartIndex; ElementIndex < ElementList.Num(); ++ElementIndex)
+			int32 ElementIndex = Self->Data[ArrayIndex].StartIndex;
+			const int32 PreviousSortOrder = Self->Data[ArrayIndex].SortOrder;
+			if (ElementIndex < ElementList.Num())
+			{
+				const FSlateInvalidationWidgetIndex FirstWidgetIndex = { ArrayIndex, (IndexType)ElementIndex };
+				const FSlateInvalidationWidgetIndex LastWidgetIndex = { ArrayIndex, (IndexType)(ElementIndex + ElementList.Num() - 1) };
+				const FIndexRange ResortRange = { *Self, FirstWidgetIndex, LastWidgetIndex };
+
+				Self->Data[ArrayIndex].SortOrder = NewSortOrder;
+
+				if (Self->CurrentInvalidationCallback)
+				{
+					const IProcessChildOrderInvalidationCallback::FReSortOperation SortOperation { ResortRange };
+					Self->CurrentInvalidationCallback->ProxiesPreResort(SortOperation);
+				}
+
+
+				for (; ElementIndex < ElementList.Num(); ++ElementIndex)
+				{
+					FSlateInvalidationWidgetIndex WidgetIndex = { ArrayIndex, (IndexType)ElementIndex };
+					FSlateInvalidationWidgetSortOrder SortOrder = { *Self, WidgetIndex };
+					if (ElementList[ElementIndex].GetWidget())
+					{
+						ElementList[ElementIndex].GetWidget()->SetFastPathSortOrder(SortOrder);
+					}
+				}
+
+				if (Self->CurrentInvalidationCallback)
+				{
+					Self->CurrentInvalidationCallback->ProxiesPostResort();
+				}
+			}
+			else
 			{
-				FSlateInvalidationWidgetIndex WidgetIndex = { ArrayIndex, (IndexType)ElementIndex };
-				FSlateInvalidationWidgetSortOrder SortOrder = { *Self, WidgetIndex };
-				ElementList[ElementIndex].GetWidget()->SetFastPathSortOrder(SortOrder);
+				Self->Data[ArrayIndex].SortOrder = NewSortOrder;
 			}
 		};
 
+	checkf(WidgetListConfig.PreferedElementsNum <= WidgetListConfig.SortOrderPaddingBetweenArray, TEXT("It should have been confirmed at initialization."));
 	for (int32 CurrentIndex = StartFrom; CurrentIndex != INDEX_NONE; CurrentIndex = Data[CurrentIndex].NextArrayIndex)
 	{
 		const int32 PreviousIndex = Data[CurrentIndex].PreviousArrayIndex;
@@ -680,13 +737,17 @@ void FSlateInvalidationWidgetList::RebuildOrderIndex(IndexType StartFrom)
 
 		if (PreviousIndex == INDEX_NONE)
 		{
-			Data[CurrentIndex].SortOrder = 0;
-			SetValue((IndexType)CurrentIndex);
+			SetValue((IndexType)CurrentIndex, 0);
 		}
 		else if (NextIndex == INDEX_NONE)
 		{
-			Data[CurrentIndex].SortOrder = Data[PreviousIndex].SortOrder + WidgetListConfig.SortOrderPaddingBetweenArray;
-			SetValue((IndexType)CurrentIndex);
+			const int32 PreviousMaxSortOrder = Data[PreviousIndex].SortOrder + WidgetListConfig.PreferedElementsNum;
+			const int32 CurrentSortOrder = Data[CurrentIndex].SortOrder;
+			if (CurrentSortOrder < PreviousMaxSortOrder)
+			{
+				const int32 NewSortOrder = Data[PreviousIndex].SortOrder + WidgetListConfig.SortOrderPaddingBetweenArray;
+				SetValue((IndexType)CurrentIndex, NewSortOrder);
+			}
 			break;
 		}
 		else
@@ -704,8 +765,8 @@ void FSlateInvalidationWidgetList::RebuildOrderIndex(IndexType StartFrom)
 			// Would the normal padding be valid
 			else if (PreviousMinSortOrder + WidgetListConfig.SortOrderPaddingBetweenArray < NextMinSortOrder)
 			{
-				Data[CurrentIndex].SortOrder = PreviousMaxSortOrder + WidgetListConfig.SortOrderPaddingBetweenArray;
-				SetValue((IndexType)CurrentIndex);
+				const int32 NewSortOrder = PreviousMinSortOrder + WidgetListConfig.SortOrderPaddingBetweenArray;
+				SetValue((IndexType)CurrentIndex, NewSortOrder);
 			}
 			// Would padding by half would be valid
 			else if (NextMinSortOrder > PreviousMaxSortOrder && NextMinSortOrder - PreviousMaxSortOrder >= WidgetListConfig.PreferedElementsNum)
@@ -715,14 +776,14 @@ void FSlateInvalidationWidgetList::RebuildOrderIndex(IndexType StartFrom)
 				const int32 NumSpacesAvailable = (NextMinSortOrder - PreviousMaxSortOrder) / WidgetListConfig.PreferedElementsNum;
 				const int32 NewCurrentOrderIndex = PreviousMaxSortOrder + (WidgetListConfig.PreferedElementsNum * (NumSpacesAvailable/2));
 				check(PreviousMaxSortOrder <= NewCurrentOrderIndex && NewCurrentOrderIndex + WidgetListConfig.PreferedElementsNum <= NextMinSortOrder);
-				Data[CurrentIndex].SortOrder = NewCurrentOrderIndex;
-				SetValue((IndexType)CurrentIndex);
+				const int32 NewSortOrder = NewCurrentOrderIndex;
+				SetValue((IndexType)CurrentIndex, NewSortOrder);
 			}
 			// Worst case, need to also rebuild the next array
 			else
 			{
-				Data[CurrentIndex].SortOrder = PreviousMaxSortOrder + FMath::Min(WidgetListConfig.PreferedElementsNum * 2, WidgetListConfig.SortOrderPaddingBetweenArray);
-				SetValue((IndexType)CurrentIndex);
+				const int32 NewSortOrder = PreviousMaxSortOrder + FMath::Min(WidgetListConfig.PreferedElementsNum * 2, WidgetListConfig.SortOrderPaddingBetweenArray);
+				SetValue((IndexType)CurrentIndex, NewSortOrder);
 			}
 		}
 		ensureMsgf(Data[CurrentIndex].SortOrder <= WidgetListConfig.MaxSortOrderPaddingBetweenArray
@@ -738,8 +799,9 @@ FSlateInvalidationWidgetList::IndexType FSlateInvalidationWidgetList::InsertData
 	{
 		check(AfterIndex == INDEX_NONE);
 		check(LastArrayIndex == INDEX_NONE);
-		AddArrayNodeIfNeeded(bReserveElementList);
-		return (IndexType)LastArrayIndex;
+		const IndexType NewLastArrayIndex = AddArrayNodeIfNeeded(bReserveElementList);
+		check(NewLastArrayIndex == LastArrayIndex);
+		return NewLastArrayIndex;
 	}
 	else
 	{
@@ -810,6 +872,7 @@ void FSlateInvalidationWidgetList::RemoveDataNode(IndexType Index)
 		LastArrayIndex = ArrayNode.PreviousArrayIndex;
 	}
 	ArrayNode.ElementList.Empty();
+	ArrayNode.ElementIndexList_WidgetWithRegisteredSlateAttribute.Empty();
 	Data.RemoveAt(Index);
 
 	// No need to rebuild the order when we remove.
@@ -860,8 +923,8 @@ void FSlateInvalidationWidgetList::_RemoveRangeFromSameParent(const FIndexRange
 		//(i)	There is no other child. Set Parent's LeafIndex to itself (recursive) [remove (F,G), leaf of E=E and leaf of A=E]
 		//(ii)	The Parent's LeafIndex is already set properly [remove (B,C), leaf of A is already G]
 		//(iii)	Parent's LeafIndex should be set to the previous sibling' leaf (recursive) [remove (G,G), leaf of E=F, A=F]
-		const InvalidationWidgetType& InvalidationWidgetStart = (*this)[Range.InclusiveMin];
-		const InvalidationWidgetType& InvalidationWidgetEnd = (*this)[Range.InclusiveMax];
+		const InvalidationWidgetType& InvalidationWidgetStart = (*this)[Range.GetInclusiveMinWidgetIndex()];
+		const InvalidationWidgetType& InvalidationWidgetEnd = (*this)[Range.GetInclusiveMaxWidgetIndex()];
 
 		// Parent index could be invalid if there is only one widget and we are removing it
 		if (InvalidationWidgetStart.ParentIndex != FSlateInvalidationWidgetIndex::Invalid)
@@ -872,7 +935,7 @@ void FSlateInvalidationWidgetList::_RemoveRangeFromSameParent(const FIndexRange
 			{
 				// _RemoveRangeFromSameParent doesn't support cross family removal
 				check(Range.Include(FSlateInvalidationWidgetSortOrder{ *this, (*this)[InvalidationWidgetEnd.ParentIndex].LeafMostChildIndex }));
-				const FSlateInvalidationWidgetIndex PreviousWidget = DecrementIndex(Range.InclusiveMin);
+				const FSlateInvalidationWidgetIndex PreviousWidget = DecrementIndex(Range.GetInclusiveMinWidgetIndex());
 				UpdateParentLeafIndex(InvalidationWidgetStart, InvalidationWidgetParentStart.LeafMostChildIndex, PreviousWidget);
 			}
 		}
@@ -887,16 +950,16 @@ void FSlateInvalidationWidgetList::_RemoveRangeFromSameParent(const FIndexRange
 	//(v)		1234 xxx8 90ab => 1234 8 90ab		=> cut, remove 5-8, (8 was moved)
 	//(vi)		1234 5xx8 90ab => 1234 5 8 90ab		=> cut, remove 6-8, (8 was moved)
 
-	const int32 NumberElementLeft = Data[Range.InclusiveMax.ArrayIndex].ElementList.Num()
-		//- Data[Range.InclusiveMax.ArrayIndex].StartIndex
-		- Range.InclusiveMax.ElementIndex
+	const int32 NumberElementLeft = Data[Range.GetInclusiveMaxWidgetIndex().ArrayIndex].ElementList.Num()
+		//- Data[Range.GetInclusiveMax().ArrayIndex].StartIndex
+		- Range.GetInclusiveMaxWidgetIndex().ElementIndex
 		- 1;
-	const bool bRangeIsInSameElementArray = Range.InclusiveMin.ArrayIndex == Range.InclusiveMax.ArrayIndex;
+	const bool bRangeIsInSameElementArray = Range.GetInclusiveMinWidgetIndex().ArrayIndex == Range.GetInclusiveMaxWidgetIndex().ArrayIndex;
 	const bool bShouldCutArray = NumberElementLeft < WidgetListConfig.NumberElementsLeftBeforeSplitting
-		|| (bRangeIsInSameElementArray && Data[Range.InclusiveMin.ArrayIndex].StartIndex != Range.InclusiveMin.ElementIndex);
+		|| (bRangeIsInSameElementArray && Data[Range.GetInclusiveMinWidgetIndex().ArrayIndex].StartIndex != Range.GetInclusiveMinWidgetIndex().ElementIndex);
 	if (bShouldCutArray)
 	{
-		_CutArray(Range.InclusiveMax);
+		_CutArray(Range.GetInclusiveMaxWidgetIndex());
 	}
 
 	// Destroy/Remove the data that is not needed anymore
@@ -904,7 +967,7 @@ void FSlateInvalidationWidgetList::_RemoveRangeFromSameParent(const FIndexRange
 		FSlateInvalidationWidgetList* Self = this;
 		auto SetFakeInvalidatWidgetHandle = [Self](IndexType ArrayIndex, int32 StartIndex, int32 Num)
 		{
-#if UE_SLATE_VERIFY_REMOVED_WIDGET_ARE_NOT_INVALIDATED
+#if UE_SLATE_WITH_WIDGETLIST_ASSIGNINVALIDPROXYWHENREMOVED
 			const FSlateInvalidationWidgetIndex SlateInvalidationWidgetIndexRemoved = { GSlateInvalidationWidgetIndex_RemovedIndex, GSlateInvalidationWidgetIndex_RemovedIndex };
 			for (int32 ElementIndex = StartIndex; ElementIndex < Num; ++ElementIndex)
 			{
@@ -914,33 +977,37 @@ void FSlateInvalidationWidgetList::_RemoveRangeFromSameParent(const FIndexRange
 					Widget->SetFastPathProxyHandle(FWidgetProxyHandle{ Self->Owner, SlateInvalidationWidgetIndexRemoved, FSlateInvalidationWidgetSortOrder(), GenerationNumber });
 				}
 			}
-#endif
+#endif //UE_SLATE_WITH_WIDGETLIST_ASSIGNINVALIDPROXYWHENREMOVED
 		};
 		auto ResetInvalidationWidget = [Self](IndexType ArrayIndex, int32 StartIndex, int32 Num)
 		{
-			ElementListType& ResetElementList = Self->Data[ArrayIndex].ElementList;
+			FArrayNode& ArrayNode = Self->Data[ArrayIndex];
+			ElementListType& ResetElementList = ArrayNode.ElementList;
 			for (int32 ElementIndex = StartIndex; ElementIndex < Num; ++ElementIndex)
 			{
 				ResetElementList[ElementIndex].ResetWidget();
 			}
+			ArrayNode.RemoveElementIndexBiggerOrEqualThan(StartIndex);
 		};
 
-		auto RemoveDataNodeIfNeeded = [Self](int32 ArrayIndex)
+		auto RemoveDataNodeIfNeeded = [Self](int32 ArrayIndex) -> bool
 		{
 			if (Self->Data[ArrayIndex].StartIndex >= Self->Data[ArrayIndex].ElementList.Num())
 			{
 				Self->RemoveDataNode(ArrayIndex);
+				return true;
 			}
+			return false;
 		};
 
 		// Remove the other arrays that are between min and max	(ie. i, ii)
-		//if (Range.InclusiveMax.ArrayIndex - Range.InclusiveMin.ArrayIndex > 1)
+		//if (Range.GetInclusiveMax().ArrayIndex - Range.GetInclusiveMin().ArrayIndex > 1)
 		if (!bRangeIsInSameElementArray)
 		{
-			const IndexType BeginArrayIndex = (IndexType)Data[Range.InclusiveMin.ArrayIndex].NextArrayIndex;
-			const IndexType EndArrayIndex = (IndexType)Data[Range.InclusiveMax.ArrayIndex].PreviousArrayIndex;
+			const IndexType BeginArrayIndex = (IndexType)Data[Range.GetInclusiveMinWidgetIndex().ArrayIndex].NextArrayIndex;
+			const IndexType EndArrayIndex = (IndexType)Data[Range.GetInclusiveMaxWidgetIndex().ArrayIndex].PreviousArrayIndex;
 
-			if (BeginArrayIndex != Range.InclusiveMax.ArrayIndex)
+			if (BeginArrayIndex != Range.GetInclusiveMaxWidgetIndex().ArrayIndex)
 			{
 				IndexType NextToRemove = BeginArrayIndex;
 				IndexType CurrentArrayIndex = NextToRemove;
@@ -958,73 +1025,88 @@ void FSlateInvalidationWidgetList::_RemoveRangeFromSameParent(const FIndexRange
 		if (bShouldCutArray && !bRangeIsInSameElementArray)
 		{
 			// The valid data in the array was moved (ie. ii)
-			SetFakeInvalidatWidgetHandle(Range.InclusiveMax.ArrayIndex, Data[Range.InclusiveMax.ArrayIndex].StartIndex, Range.InclusiveMax.ElementIndex + 1);
-			RemoveDataNode(Range.InclusiveMax.ArrayIndex);
+			SetFakeInvalidatWidgetHandle(Range.GetInclusiveMaxWidgetIndex().ArrayIndex, Data[Range.GetInclusiveMaxWidgetIndex().ArrayIndex].StartIndex, Range.GetInclusiveMaxWidgetIndex().ElementIndex + 1);
+			RemoveDataNode(Range.GetInclusiveMaxWidgetIndex().ArrayIndex);
 		}
 		else if (!bRangeIsInSameElementArray)
 		{
 			// Set StartIndex (ie. i)
-			check(Range.InclusiveMin.ArrayIndex != Range.InclusiveMax.ArrayIndex);
+			check(Range.GetInclusiveMinWidgetIndex().ArrayIndex != Range.GetInclusiveMaxWidgetIndex().ArrayIndex);
 
-			SetFakeInvalidatWidgetHandle(Range.InclusiveMax.ArrayIndex, Data[Range.InclusiveMax.ArrayIndex].StartIndex, Range.InclusiveMax.ElementIndex + 1);
-			ResetInvalidationWidget(Range.InclusiveMax.ArrayIndex, Data[Range.InclusiveMax.ArrayIndex].StartIndex, Range.InclusiveMax.ElementIndex + 1);
-			Data[Range.InclusiveMax.ArrayIndex].StartIndex = Range.InclusiveMax.ElementIndex + 1;
-			RemoveDataNodeIfNeeded(Range.InclusiveMax.ArrayIndex);
+			SetFakeInvalidatWidgetHandle(Range.GetInclusiveMaxWidgetIndex().ArrayIndex, Data[Range.GetInclusiveMaxWidgetIndex().ArrayIndex].StartIndex, Range.GetInclusiveMaxWidgetIndex().ElementIndex + 1);
+			ResetInvalidationWidget(Range.GetInclusiveMaxWidgetIndex().ArrayIndex, Data[Range.GetInclusiveMaxWidgetIndex().ArrayIndex].StartIndex, Range.GetInclusiveMaxWidgetIndex().ElementIndex + 1);
+			Data[Range.GetInclusiveMaxWidgetIndex().ArrayIndex].StartIndex = Range.GetInclusiveMaxWidgetIndex().ElementIndex + 1;
+			RemoveDataNodeIfNeeded(Range.GetInclusiveMaxWidgetIndex().ArrayIndex);
 		}
 
 		// Remove what is left of the Min array
 		if (bShouldCutArray || !bRangeIsInSameElementArray)
 		{
 			// RemoveAt Min to Num of the array. (ie. i, ii, iii, v, vi)
-			ElementListType& RemoveElementList = Data[Range.InclusiveMin.ArrayIndex].ElementList;
+			FArrayNode& ArrayNode = Data[Range.GetInclusiveMinWidgetIndex().ArrayIndex];
+			ElementListType& RemoveElementList = ArrayNode.ElementList;
 			if (bRangeIsInSameElementArray)
 			{
-				SetFakeInvalidatWidgetHandle(Range.InclusiveMin.ArrayIndex, Range.InclusiveMin.ElementIndex, Range.InclusiveMax.ElementIndex);
+				SetFakeInvalidatWidgetHandle(Range.GetInclusiveMinWidgetIndex().ArrayIndex, Range.GetInclusiveMinWidgetIndex().ElementIndex, Range.GetInclusiveMaxWidgetIndex().ElementIndex);
 			}
 			else
 			{
-				SetFakeInvalidatWidgetHandle(Range.InclusiveMin.ArrayIndex, Range.InclusiveMin.ElementIndex, RemoveElementList.Num());
+				SetFakeInvalidatWidgetHandle(Range.GetInclusiveMinWidgetIndex().ArrayIndex, Range.GetInclusiveMinWidgetIndex().ElementIndex, RemoveElementList.Num());
 			}
 
-			const int32 RemoveArrayAt = Range.InclusiveMin.ElementIndex;
+			const int32 RemoveArrayAt = Range.GetInclusiveMinWidgetIndex().ElementIndex;
 			RemoveElementList.RemoveAt(RemoveArrayAt, RemoveElementList.Num() - RemoveArrayAt, true);
-			RemoveDataNodeIfNeeded(Range.InclusiveMin.ArrayIndex);
+			if (!RemoveDataNodeIfNeeded(Range.GetInclusiveMinWidgetIndex().ArrayIndex))
+			{
+				ArrayNode.RemoveElementIndexBiggerOrEqualThan(RemoveArrayAt);
+			}
 		}
 		else
 		{
 			// Set StartIndex (ie. iv)
-			check(Range.InclusiveMin.ArrayIndex == Range.InclusiveMax.ArrayIndex);
-			check(Range.InclusiveMin.ElementIndex == Data[Range.InclusiveMin.ArrayIndex].StartIndex);
+			check(Range.GetInclusiveMinWidgetIndex().ArrayIndex == Range.GetInclusiveMaxWidgetIndex().ArrayIndex);
+			check(Range.GetInclusiveMinWidgetIndex().ElementIndex == Data[Range.GetInclusiveMinWidgetIndex().ArrayIndex].StartIndex);
 
-			SetFakeInvalidatWidgetHandle(Range.InclusiveMin.ArrayIndex, Range.InclusiveMin.ElementIndex, Range.InclusiveMax.ElementIndex + 1);
-			ResetInvalidationWidget(Range.InclusiveMin.ArrayIndex, Range.InclusiveMin.ElementIndex, Range.InclusiveMax.ElementIndex + 1);
-			Data[Range.InclusiveMin.ArrayIndex].StartIndex = Range.InclusiveMax.ElementIndex + 1;
-			RemoveDataNodeIfNeeded(Range.InclusiveMin.ArrayIndex);
+			SetFakeInvalidatWidgetHandle(Range.GetInclusiveMinWidgetIndex().ArrayIndex, Range.GetInclusiveMinWidgetIndex().ElementIndex, Range.GetInclusiveMaxWidgetIndex().ElementIndex + 1);
+			ResetInvalidationWidget(Range.GetInclusiveMinWidgetIndex().ArrayIndex, Range.GetInclusiveMinWidgetIndex().ElementIndex, Range.GetInclusiveMaxWidgetIndex().ElementIndex + 1);
+			Data[Range.GetInclusiveMinWidgetIndex().ArrayIndex].StartIndex = Range.GetInclusiveMaxWidgetIndex().ElementIndex + 1;
+			RemoveDataNodeIfNeeded(Range.GetInclusiveMinWidgetIndex().ArrayIndex);
 		}
 	}
 }
 
 
-void FSlateInvalidationWidgetList::CutArray(const FSlateInvalidationWidgetIndex WhereToCut)
+FSlateInvalidationWidgetList::FCutResult FSlateInvalidationWidgetList::CutArray(const FSlateInvalidationWidgetIndex WhereToCut)
 {
-	int32 OldElementIndexStart = _CutArray(WhereToCut);
+	FCutResult CutResult = _CutArray(WhereToCut);
 
 	// Remove the old data that is now moved to the new array
-	if (OldElementIndexStart != INDEX_NONE)
+	if (CutResult.OldElementIndexStart != INDEX_NONE)
 	{
+		FArrayNode& ArrayNode = Data[WhereToCut.ArrayIndex];
 		ElementListType& RemoveElementList = Data[WhereToCut.ArrayIndex].ElementList;
-		RemoveElementList.RemoveAt(OldElementIndexStart, RemoveElementList.Num() - OldElementIndexStart, true);
+		RemoveElementList.RemoveAt(CutResult.OldElementIndexStart, RemoveElementList.Num() - CutResult.OldElementIndexStart, true);
 		if (RemoveElementList.Num() == 0)
 		{
 			RemoveDataNode(WhereToCut.ArrayIndex);
 		}
+		else
+		{
+			ArrayNode.RemoveElementIndexBiggerOrEqualThan(CutResult.OldElementIndexStart);
+		}
 	}
+
+	return CutResult;
 }
 
 
-int32 FSlateInvalidationWidgetList::_CutArray(const FSlateInvalidationWidgetIndex WhereToCut)
+FSlateInvalidationWidgetList::FCutResult FSlateInvalidationWidgetList::_CutArray(const FSlateInvalidationWidgetIndex WhereToCut)
 {
-	SCOPED_NAMED_EVENT(Slate_InvalidationList_CutArray, FColor::Blue);
+	FCutResult Result;
+	// Widget proxies that got moved/re-indexed by the operation.
+	FIndexRange ProxiesReIndexed;
+	//The new index that now correspond to the previous ProxiesReIndexed.GetInclusiveMin()
+	FSlateInvalidationWidgetIndex ReIndexTarget = FSlateInvalidationWidgetIndex::Invalid;
 
 	// Should we cut the array and move the data to another array
 	// N.B. We can cut/move anywhere. Cross family may occur. Fix up everything.
@@ -1047,8 +1129,16 @@ int32 FSlateInvalidationWidgetList::_CutArray(const FSlateInvalidationWidgetInde
 			return NewWidgetIndex;
 		};
 
+		if (OldElementIndexStart < OldElementIndexEnd)
+		{
+			FSlateInvalidationWidgetIndex ResultStart = { OldArrayIndex, OldElementIndexStart };
+			FSlateInvalidationWidgetIndex ResultEnd = { OldArrayIndex, (IndexType)(OldElementIndexEnd - 1) };
+			ProxiesReIndexed = FIndexRange{ *this, ResultStart , ResultEnd };
+		}
+
+		Result.OldElementIndexStart = OldElementIndexStart;
+
 		// Copy and assign the new index to the widget
-//Todo @pat fix this with proper iterator so that we do not have to use operator[]
 		for (IndexType OldElementIndex = OldElementIndexStart; OldElementIndex < OldElementIndexEnd; ++OldElementIndex)
 		{
 			const FSlateInvalidationWidgetIndex MoveWidgetIndex = { OldArrayIndex, OldElementIndex };
@@ -1058,6 +1148,11 @@ int32 FSlateInvalidationWidgetList::_CutArray(const FSlateInvalidationWidgetInde
 			check(OldToNewIndex(MoveWidgetIndex) == NewWidgetIndex);
 			InvalidationWidgetType& NewInvalidationWidget = (*this)[NewWidgetIndex];
 
+			if (ReIndexTarget == FSlateInvalidationWidgetIndex::Invalid)
+			{
+				ReIndexTarget = NewWidgetIndex;
+			}
+
 			// Fix up Index
 			NewInvalidationWidget.Index = NewWidgetIndex;
 
@@ -1098,6 +1193,12 @@ int32 FSlateInvalidationWidgetList::_CutArray(const FSlateInvalidationWidgetInde
 			{
 				FSlateInvalidationWidgetSortOrder SortIndex = { *this, NewWidgetIndex };
 				Widget->SetFastPathProxyHandle(FWidgetProxyHandle{ Owner, NewWidgetIndex, SortIndex, GenerationNumber });
+
+				// Check for ElementIndexList
+				if (Widget->HasRegisteredSlateAttribute())
+				{
+					Data[NewArrayIndex].ElementIndexList_WidgetWithRegisteredSlateAttribute.AddUnsorted(NewElementIndex);
+				}
 			}
 		}
 
@@ -1116,13 +1217,18 @@ int32 FSlateInvalidationWidgetList::_CutArray(const FSlateInvalidationWidgetInde
 					NewChildInvalidationWidget.ParentIndex = NewWidgetIndex;
 				});
 		}
+	}
 
-		return OldElementIndexStart;
+	if (CurrentInvalidationCallback && ProxiesReIndexed.IsValid())
+	{
+		IProcessChildOrderInvalidationCallback::FReIndexOperation Operation{ ProxiesReIndexed, ReIndexTarget };
+		CurrentInvalidationCallback->ProxiesReIndexed(Operation);
 	}
-	return INDEX_NONE;
+
+	return Result;
 }
 
-#if WITH_SLATE_DEBUGGING
+#if UE_SLATE_WITH_INVALIDATIONWIDGETLIST_DEBUGGING
 FSlateInvalidationWidgetIndex FSlateInvalidationWidgetList::FindWidget(const TSharedRef<SWidget> WidgetToFind) const
 {
 	SWidget* WidgetToFindPtr = &WidgetToFind.Get();
@@ -1177,13 +1283,13 @@ TArray<TSharedPtr<SWidget>> FSlateInvalidationWidgetList::FindChildren(const TSh
 		}
 
 		FMemMark Mark(FMemStack::Get());
-		TArray<SWidget*, TMemStackAllocator<>> PreviousChildrenWidget;
+		TArray<FFindChildrenElement, TMemStackAllocator<>> PreviousChildrenWidget;
 		_FindChildren(WidgetIndex, PreviousChildrenWidget);
 
 		Result.Reserve(PreviousChildrenWidget.Num());
-		for (SWidget* WidgetPtr : PreviousChildrenWidget)
+		for (const FFindChildrenElement& ChildrenElement : PreviousChildrenWidget)
 		{
-			Result.Add(WidgetPtr ? WidgetPtr->AsShared() : TSharedPtr<SWidget>());
+			Result.Add(ChildrenElement.Get<0>() ? ChildrenElement.Get<0>()->AsShared() : TSharedPtr<SWidget>());
 		}
 	}
 	return Result;
@@ -1247,7 +1353,7 @@ void FSlateInvalidationWidgetList::LogWidgetsList()
 		const InvalidationWidgetType& InvalidateWidget = (*this)[Index];
 		if (SWidget* Widget = InvalidateWidget.GetWidget())
 		{
-			Builder << Widget->GetTag();
+			Builder << FReflectionMetaData::GetWidgetDebugInfo(Widget) << TEXT("(") << Widget->GetTag() << TEXT(")");
 		}
 		else
 		{
@@ -1258,7 +1364,7 @@ void FSlateInvalidationWidgetList::LogWidgetsList()
 		{
 			if (SWidget* Widget = (*this)[InvalidateWidget.ParentIndex].GetWidget())
 			{
-				Builder << Widget->GetTag();
+				Builder << FReflectionMetaData::GetWidgetDebugInfo(Widget) << TEXT("(") << Widget->GetTag() << TEXT(")");
 			}
 			else
 			{
@@ -1274,7 +1380,7 @@ void FSlateInvalidationWidgetList::LogWidgetsList()
 		{
 			if (SWidget* Widget = (*this)[InvalidateWidget.LeafMostChildIndex].GetWidget())
 			{
-				Builder << Widget->GetTag();
+				Builder << FReflectionMetaData::GetWidgetDebugInfo(Widget) << TEXT("(") << Widget->GetTag() << TEXT(")");
 			}
 			else
 			{
@@ -1327,6 +1433,7 @@ bool FSlateInvalidationWidgetList::VerifyWidgetsIndex() const
 	return bResult;
 }
 
+
 bool FSlateInvalidationWidgetList::VerifyProxiesWidget() const
 {
 	bool bResult = true;
@@ -1359,6 +1466,7 @@ bool FSlateInvalidationWidgetList::VerifyProxiesWidget() const
 	return bResult;
 }
 
+
 bool FSlateInvalidationWidgetList::VerifySortOrder() const
 {
 	bool bResult = true;
@@ -1378,8 +1486,59 @@ bool FSlateInvalidationWidgetList::VerifySortOrder() const
 	}
 	return bResult;
 }
-#endif //WITH_SLATE_DEBUGGING
 
-#undef UE_SLATE_VERIFY_INVALID_INVALIDATIONHANDLE
-#undef UE_SLATE_VERIFY_REBUILDWIDGETDATA_ORDER
-#undef UE_SLATE_VERIFY_REMOVED_WIDGET_ARE_NOT_INVALIDATED
+
+bool FSlateInvalidationWidgetList::VerifyElementIndexList() const
+{
+	bool bResult = true;
+	for (const FArrayNode& ArrayNode : Data)
+	{
+		for (IndexType ElementIndex : ArrayNode.ElementIndexList_WidgetWithRegisteredSlateAttribute)
+		{
+			if (!ArrayNode.ElementList.IsValidIndex(ElementIndex))
+			{
+				UE_LOG(LogSlate, Warning, TEXT("ElementIndex '%d' in the array of sort value '%d' is invalid.")
+					, ElementIndex, ArrayNode.SortOrder);
+				bResult = false;
+				break;
+			}
+
+			FSlateInvalidationWidgetList::InvalidationWidgetType const& InvalidationWidget = ArrayNode.ElementList[ElementIndex];
+			if (const SWidget* Widget = InvalidationWidget.GetWidget())
+			{
+				if (!Widget->HasRegisteredSlateAttribute())
+				{
+					UE_LOG(LogSlate, Warning, TEXT("ElementIndex '%d' in the array of sort value '%d' should not be there since widget '%s' doesn't have registered SlateAttribute.")
+						, ElementIndex, ArrayNode.SortOrder
+						, *FReflectionMetaData::GetWidgetDebugInfo(Widget));
+					bResult = false;
+				}
+
+				if (Widget->IsFastPathVisible())
+				{
+					ensureMsgf(InvalidationWidget.bDebug_AttributeUpdated, TEXT("The widget '%s' was not updated."), *FReflectionMetaData::GetWidgetDebugInfo(Widget));
+				}
+			}
+		}
+
+		for (int32 ElementIndex = ArrayNode.StartIndex; ElementIndex < ArrayNode.ElementList.Num(); ++ElementIndex)
+		{
+			if (!ArrayNode.ElementIndexList_WidgetWithRegisteredSlateAttribute.Contains(ElementIndex))
+			{
+				SWidget* Widget = ArrayNode.ElementList[ElementIndex].GetWidget();
+				if (Widget->HasRegisteredSlateAttribute())
+				{
+					UE_LOG(LogSlate, Warning, TEXT("Widget '%s' (element: '%d') in the array of sort value '%d' should be in ElementIndexList_WidgetWithRegisteredSlateAttribute.")
+						, *FReflectionMetaData::GetWidgetDebugInfo(Widget)
+						, ElementIndex, ArrayNode.SortOrder);
+					bResult = false;
+				}
+			}
+		}
+	}
+	return bResult;
+}
+#endif //UE_SLATE_WITH_INVALIDATIONWIDGETLIST_DEBUGGING
+
+#undef UE_SLATE_WITH_WIDGETLIST_ASSIGNINVALIDPROXYWHENREMOVED
+#undef UE_SLATE_WITH_WIDGETLIST_UPDATEONLYWHATISNEEDED
diff --git a/Engine/Source/Runtime/SlateCore/Private/FastUpdate/SlateInvalidationWidgetList.h b/Engine/Source/Runtime/SlateCore/Private/FastUpdate/SlateInvalidationWidgetList.h
--- a/Engine/Source/Runtime/SlateCore/Private/FastUpdate/SlateInvalidationWidgetList.h
+++ b/Engine/Source/Runtime/SlateCore/Private/FastUpdate/SlateInvalidationWidgetList.h
@@ -4,22 +4,62 @@
 
 #include "CoreMinimal.h"
 
-#include "Debugging/SlateDebugging.h"
+#include "FastUpdate/SlateElementSortedArray.h"
 #include "FastUpdate/SlateInvalidationRootHandle.h"
 #include "FastUpdate/SlateInvalidationWidgetIndex.h"
 #include "FastUpdate/SlateInvalidationWidgetSortOrder.h"
 #include "Widgets/SNullWidget.h"
 #include "Widgets/SWidget.h"
 
+#ifndef UE_SLATE_WITH_INVALIDATIONWIDGETLIST_DEBUGGING
+	#define UE_SLATE_WITH_INVALIDATIONWIDGETLIST_DEBUGGING !(UE_BUILD_SHIPPING || UE_BUILD_TEST)
+#endif
+
 class FSlateInvalidationWidgetList
 {
+	friend struct FSlateInvalidationWidgetSortOrder;
 public:
 	using InvalidationWidgetType = FWidgetProxy;
 private:
 	using IndexType = FSlateInvalidationWidgetIndex::IndexType;
 	using ElementListType = TArray<InvalidationWidgetType>;
+	using WidgetListType = TSlateElementSortedArray<IndexType>;
+
+
+public:
+	/** */
+	struct FIndexRange
+	{
+	private:
+		FSlateInvalidationWidgetIndex InclusiveMin = FSlateInvalidationWidgetIndex::Invalid;
+		FSlateInvalidationWidgetIndex InclusiveMax = FSlateInvalidationWidgetIndex::Invalid;
+		FSlateInvalidationWidgetSortOrder OrderMin;
+		FSlateInvalidationWidgetSortOrder OrderMax;
+
+	public:
+		FIndexRange() = default;
+		FIndexRange(const FSlateInvalidationWidgetList& Self, FSlateInvalidationWidgetIndex InFrom, FSlateInvalidationWidgetIndex InEnd)
+			: InclusiveMin(InFrom), InclusiveMax(InEnd)
+			, OrderMin(Self, InFrom), OrderMax(Self, InEnd)
+		{
+			check(OrderMin <= OrderMax);
+		}
+		bool Include(FSlateInvalidationWidgetSortOrder Other) const
+		{
+			return OrderMin <= Other && Other <= OrderMax;
+		}
+		bool IsValid() const { return InclusiveMin != FSlateInvalidationWidgetIndex::Invalid; }
+
+		FSlateInvalidationWidgetIndex GetInclusiveMinWidgetIndex() const { return InclusiveMin; }
+		FSlateInvalidationWidgetIndex GetInclusiveMaxWidgetIndex() const { return InclusiveMax; }
+		FSlateInvalidationWidgetSortOrder GetInclusiveMinWidgetSortOrder() const { return OrderMin; }
+		FSlateInvalidationWidgetSortOrder GetInclusiveMaxWidgetSortOrder() const { return OrderMax; }
+
+		bool operator==(const FIndexRange& Other) const { return Other.InclusiveMin == InclusiveMin && Other.InclusiveMax == InclusiveMax; }
+	};
 
 public:
+	/** */
 	struct FArguments
 	{
 		static int32 MaxPreferedElementsNum;
@@ -41,7 +81,7 @@ public:
 		 * The value should be between PreferedElementsNum and MaxSortOrderPaddingBetweenArray.
 		 */
 		int32 SortOrderPaddingBetweenArray = 1000;
-#if WITH_SLATE_DEBUGGING
+#if UE_SLATE_WITH_INVALIDATIONWIDGETLIST_DEBUGGING
 		/** Change the Widget Index when building the array. Use when building temporary list. */
 		bool bAssignedWidgetIndex = true;
 #endif
@@ -59,8 +99,50 @@ public:
 	/** Get the Generation number the list was built on. */
 	int32 GetGenerationNumber() const { return GenerationNumber; }
 
-	/** Process widget that have a ChildOrder invalidation. */
-	void ProcessChildOrderInvalidation(const TArray<TWeakPtr<SWidget>>& Invalidated);
+	/** */
+	struct IProcessChildOrderInvalidationCallback
+	{
+		struct FReIndexOperation
+		{
+			FReIndexOperation(FIndexRange const& InRange, FSlateInvalidationWidgetIndex InReIndexTarget) : Range(InRange), ReIndexTarget(InReIndexTarget) {}
+			FIndexRange const& GetRange() const { return Range; }
+			UE_NODISCARD FSlateInvalidationWidgetIndex ReIndex(FSlateInvalidationWidgetIndex Index) const;
+		private:
+			const FIndexRange& Range;
+			FSlateInvalidationWidgetIndex ReIndexTarget = FSlateInvalidationWidgetIndex::Invalid;
+		};
+		struct FReSortOperation
+		{
+			FReSortOperation(FIndexRange const& InRange) : Range(InRange) {}
+			FIndexRange const& GetRange() const { return Range; }
+		private:
+			FIndexRange const& Range;
+		};
+
+		/** Widget proxies that will be removed and will not be valid anymore. */
+		virtual void PreChildRemove(FIndexRange const& Range) {}
+		/** Widget proxies that got moved/re-indexed by the operation. */
+		virtual void ProxiesReIndexed(FReIndexOperation const& Operation) {}
+		/** Widget proxies that got resorted by the operation. */
+		virtual void ProxiesPreResort(FReSortOperation const& Operation) {}
+		/** Widget proxies that got resorted by the operation. */
+		virtual void ProxiesPostResort() {}
+		/** Widget proxies built by the operation. */
+		virtual void ProxiesBuilt(FIndexRange const& Range) {}
+	};
+
+	/**
+	 * Process widget that have a ChildOrder invalidation.
+	 * @returns true if the InvalidationWidget is still valid.
+	 */
+	bool ProcessChildOrderInvalidation(InvalidationWidgetType const& InvalidationWidget, IProcessChildOrderInvalidationCallback& Callback);
+
+
+	/**
+	 * Test, then adds or removes from the registered attribute list.
+	 * @returns true if the Widget go added to the registered attribute list.
+	 */
+	void ProcessAttributeRegistrationInvalidation(InvalidationWidgetType const& InvalidationWidget);
 
 	/** Performs an operation on all SWidget in the list. */
 	template<typename Predicate>
@@ -102,6 +184,49 @@ public:
 		}
 	}
 
+	/** Iterator that goes over all the widgets with registered attribute. */
+	struct FWidgetAttributeIterator
+	{
+	private:
+		FSlateInvalidationWidgetList const& WidgetList;
+		FSlateInvalidationWidgetIndex CurrentWidgetIndex;
+		FSlateInvalidationWidgetSortOrder CurrentWidgetSortOrder;
+		int32 AttributeIndex;
+
+		FSlateInvalidationWidgetIndex MoveToWidgetIndexOnNextAdvance;
+
+	public:
+		FWidgetAttributeIterator(FSlateInvalidationWidgetList const& InWidgetList);
+
+		//~ Handle operation
+		void PreChildRemove(FIndexRange const& Range);
+		void ReIndexed(IProcessChildOrderInvalidationCallback::FReIndexOperation const& Operation);
+		void PostResort();
+		void ProxiesBuilt(FIndexRange const& Range);
+		void FixCurrentWidgetIndex();
+		void Seek(FSlateInvalidationWidgetIndex SeekTo);
+
+		/** Get the current widget index the iterator is pointing to. */
+		FSlateInvalidationWidgetIndex GetCurrentIndex() const { return CurrentWidgetIndex; }
+		
+		/** Get the current widget sort order the iterator is pointing to. */
+		FSlateInvalidationWidgetSortOrder GetCurrentSortOrder() const { return CurrentWidgetSortOrder; }
+
+		/** Advance the iterator to the next valid widget index. */
+		void Advance();
+
+		/** Is the iterator pointing to a valid widget index. */
+		bool IsValid() const { return CurrentWidgetIndex != FSlateInvalidationWidgetIndex::Invalid; }
+
+	private:
+		void AdvanceArrayIndex(int32 ArrayIndex);
+	};
+
+	FWidgetAttributeIterator CreateWidgetAttributeIterator()
+	{
+		return FWidgetAttributeIterator(*this);
+	}
+
 public:
 	/** Returns reference to element at give index. */
 	InvalidationWidgetType& operator[](const FSlateInvalidationWidgetIndex Index)
@@ -148,10 +273,10 @@ public:
 	}
 
 	/** Increment a widget index to the next entry in the WidgetList. */
-	FSlateInvalidationWidgetIndex IncrementIndex(FSlateInvalidationWidgetIndex Index) const;
+	UE_NODISCARD FSlateInvalidationWidgetIndex IncrementIndex(FSlateInvalidationWidgetIndex Index) const;
 
 	/** Decrement a widget index to the next entry in the WidgetList. */
-	FSlateInvalidationWidgetIndex DecrementIndex(FSlateInvalidationWidgetIndex Index) const;
+	UE_NODISCARD FSlateInvalidationWidgetIndex DecrementIndex(FSlateInvalidationWidgetIndex Index) const;
 
 	/** Empties the WidgetList. */
 	void Empty();
@@ -160,7 +285,7 @@ public:
 	void Reset();
 
 public:
-#if WITH_SLATE_DEBUGGING
+#if UE_SLATE_WITH_INVALIDATIONWIDGETLIST_DEBUGGING
 	/** For testing purposes. Return the InvalidationWidgetIndex of the Widget within the InvalidationWidgetList. */
 	FSlateInvalidationWidgetIndex FindWidget(const TSharedRef<SWidget> Widget) const;
 
@@ -170,7 +295,7 @@ public:
 	void RemoveWidget(const TSharedRef<SWidget> WidgetToRemove);
 
 	/** For testing purposes. Use to test ProcessChildOrderInvalidation */
-	TArray<TSharedPtr<SWidget>> FindChildren(const TSharedRef<SWidget> Widget) const;
+	UE_NODISCARD TArray<TSharedPtr<SWidget>> FindChildren(const TSharedRef<SWidget> Widget) const;
 
 	/**
 	 * For testing purposes.
@@ -191,30 +316,10 @@ public:
 
 	/** For testing purposes. Verify that every the sorting order is increasing between widgets. */
 	bool VerifySortOrder() const;
-#endif
 
-private:
-	friend struct FSlateInvalidationWidgetSortOrder;
-
-	struct FIndexRange
-	{
-		FSlateInvalidationWidgetIndex InclusiveMin = FSlateInvalidationWidgetIndex::Invalid;
-		FSlateInvalidationWidgetIndex InclusiveMax = FSlateInvalidationWidgetIndex::Invalid;
-		FSlateInvalidationWidgetSortOrder OrderMin;
-		FSlateInvalidationWidgetSortOrder OrderMax;
-		FIndexRange() = default;
-		FIndexRange(const FSlateInvalidationWidgetList& Self, FSlateInvalidationWidgetIndex InFrom, FSlateInvalidationWidgetIndex InEnd)
-			: InclusiveMin(InFrom), InclusiveMax(InEnd)
-			, OrderMin(Self, InFrom), OrderMax(Self, InEnd)
-		{
-			check(OrderMin <= OrderMax);
-		}
-		bool Include(FSlateInvalidationWidgetSortOrder Other) const
-		{
-			return OrderMin <= Other && Other <= OrderMax;
-		}
-		bool operator==(const FIndexRange& Other) const { return Other.InclusiveMin == InclusiveMin && Other.InclusiveMax == InclusiveMax; }
-	};
+	/** For testing purposes. Verify that the ElementIndexList_ have valid indexes and are sorted. */
+	bool VerifyElementIndexList() const;
+#endif
 
 private:
 	bool ShouldDoRecursion(const SWidget* Widget) const
@@ -251,20 +356,29 @@ private:
 	}
 
 private:
-	IndexType AddArrayNodeIfNeeded(bool bReserveElementList);
-	IndexType InsertArrayNodeIfNeeded(IndexType AfterArrayIndex, bool bReserveElementList);
-	IndexType InsertDataNodeAfter(IndexType AfterIndex, bool bReserveElementList);
+	struct FCutResult
+	{
+		FCutResult() = default;
+
+		/** Where in the previous array the reindexed element starts. */
+		int32 OldElementIndexStart = INDEX_NONE;
+	};
+private:
+	UE_NODISCARD IndexType AddArrayNodeIfNeeded(bool bReserveElementList);
+	UE_NODISCARD IndexType InsertArrayNodeIfNeeded(IndexType AfterArrayIndex, bool bReserveElementList);
+	UE_NODISCARD IndexType InsertDataNodeAfter(IndexType AfterIndex, bool bReserveElementList);
 	void RemoveDataNode(IndexType Index);
 	void RebuildOrderIndex(IndexType StartFrom);
 	void UpdateParentLeafIndex(const InvalidationWidgetType& InvalidationWidget, FSlateInvalidationWidgetIndex OldIndex, FSlateInvalidationWidgetIndex NewIndex);
-	void CutArray(const FSlateInvalidationWidgetIndex WhereToCut);
+	FCutResult CutArray(const FSlateInvalidationWidgetIndex WhereToCut);
 
 private:
 	FSlateInvalidationWidgetIndex _BuildWidgetList_Recursive(TSharedRef<SWidget>& Widget, FSlateInvalidationWidgetIndex ParentIndex, IndexType& LastestIndex, bool bParentVisible, bool bParentVolatile);
 	void _RebuildWidgetListTree(TSharedRef<SWidget> Widget, int32 ChildAtIndex);
-	void _FindChildren(FSlateInvalidationWidgetIndex WidgetIndex, TArray<SWidget*, TMemStackAllocator<>>& Widgets) const;
+	using FFindChildrenElement = TPair<SWidget*, FSlateInvalidationWidgetIndex>;
+	void _FindChildren(FSlateInvalidationWidgetIndex WidgetIndex, TArray<FFindChildrenElement, TMemStackAllocator<>>& Widgets) const;
 	void _RemoveRangeFromSameParent(const FIndexRange Range);
-	int32 _CutArray(const FSlateInvalidationWidgetIndex WhereToCut);
+	FCutResult _CutArray(const FSlateInvalidationWidgetIndex WhereToCut);
 
 private:
 	struct FArrayNode
@@ -273,8 +387,11 @@ private:
 		int32 PreviousArrayIndex = INDEX_NONE;
 		int32 NextArrayIndex = INDEX_NONE;
 		int32 SortOrder = 0;
-		IndexType StartIndex = 0; // The array may start further in the array as a result of a split.
+		IndexType StartIndex = 0; // The array may start further in the ElementList as a result of a split.
 		ElementListType ElementList;
+		WidgetListType ElementIndexList_WidgetWithRegisteredSlateAttribute;
+
+		void RemoveElementIndexBiggerOrEqualThan(IndexType ElementIndex);
 	};
 	using ArrayNodeType = TSparseArray<FArrayNode>;
 
@@ -284,6 +401,6 @@ private:
 	int32 FirstArrayIndex = INDEX_NONE;
 	int32 LastArrayIndex = INDEX_NONE;
 	int32 GenerationNumber = INDEX_NONE;
-
+	IProcessChildOrderInvalidationCallback* CurrentInvalidationCallback = nullptr;
 	const FArguments WidgetListConfig;
 };
\ No newline at end of file
diff --git a/Engine/Source/Runtime/SlateCore/Private/FastUpdate/SlateInvalidationWidgetSortOrder.cpp b/Engine/Source/Runtime/SlateCore/Private/FastUpdate/SlateInvalidationWidgetSortOrder.cpp
--- a/Engine/Source/Runtime/SlateCore/Private/FastUpdate/SlateInvalidationWidgetSortOrder.cpp
+++ b/Engine/Source/Runtime/SlateCore/Private/FastUpdate/SlateInvalidationWidgetSortOrder.cpp
@@ -4,10 +4,24 @@
 
 #include "FastUpdate/SlateInvalidationWidgetList.h"
 
+#include <limits>
+
 FSlateInvalidationWidgetSortOrder::FSlateInvalidationWidgetSortOrder(const FSlateInvalidationWidgetList& List, FSlateInvalidationWidgetIndex Index)
 	: Order()
 {
 	// 22 bits for the order (array index)
 	// 10 bits for the element index
 	Order = (List.Data[Index.ArrayIndex].SortOrder << 10) + Index.ElementIndex;
-}
\ No newline at end of file
+}
+
+
+FSlateInvalidationWidgetSortOrder FSlateInvalidationWidgetSortOrder::LimitMax()
+{
+	return FSlateInvalidationWidgetSortOrder{ std::numeric_limits<uint32>::max() };
+}
+
+
+FSlateInvalidationWidgetSortOrder FSlateInvalidationWidgetSortOrder::LimitMin()
+{
+	return FSlateInvalidationWidgetSortOrder{0};
+}
diff --git a/Engine/Source/Runtime/SlateCore/Private/FastUpdate/WidgetProxy.cpp b/Engine/Source/Runtime/SlateCore/Private/FastUpdate/WidgetProxy.cpp
--- a/Engine/Source/Runtime/SlateCore/Private/FastUpdate/WidgetProxy.cpp
+++ b/Engine/Source/Runtime/SlateCore/Private/FastUpdate/WidgetProxy.cpp
@@ -28,9 +28,11 @@ FWidgetProxy::FWidgetProxy(TSharedRef<SWidget>& InWidget)
 	// Potentially unsafe to update visibility from the widget due to attribute bindings.  This is updated later when the widgets are sorted in ProcessInvalidation
 	, Visibility(EVisibility::Collapsed) 
 	, bUpdatedSinceLastInvalidate(false)
-	, bContainedByWidgetHeap(false)
+	, bContainedByWidgetPreHeap(false)
+	, bContainedByWidgetPostHeap(false)
 	, bDebug_LastFrameVisible(true)
 	, bDebug_LastFrameVisibleSet(false)
+	, bDebug_AttributeUpdated(false)
 {
 #if UE_SLATE_WITH_WIDGETPROXY_WIDGETTYPE
 	WidgetName = GetWidget()->GetType();
@@ -90,7 +92,7 @@ int32 FWidgetProxy::Update(const FPaintArgs& PaintArgs, FSlateWindowElementList&
 	return OutgoingLayerId;
 }
 
-bool FWidgetProxy::ProcessInvalidation(FSlateInvalidationWidgetHeap& UpdateList, FSlateInvalidationWidgetList& FastWidgetPathList, FSlateInvalidationRoot& Root)
+bool FWidgetProxy::ProcessPostInvalidation(FSlateInvalidationWidgetPostHeap& UpdateList, FSlateInvalidationWidgetList& FastWidgetPathList, FSlateInvalidationRoot& Root)
 {
 	bool bWidgetNeedsRepaint = false;
 	SWidget* WidgetPtr = GetWidget();
@@ -108,7 +110,7 @@ bool FWidgetProxy::ProcessInvalidation(FSlateInvalidationWidgetHeap& UpdateList,
 			FSlateDebugging::BroadcastWidgetInvalidate(ParentWidgetPtr, WidgetPtr, EInvalidateWidgetReason::Layout);
 #endif
 			UE_TRACE_SLATE_WIDGET_INVALIDATED(ParentWidgetPtr, WidgetPtr, EInvalidateWidgetReason::Layout);
-			UpdateList.PushUnique(ParentProxy);
+			UpdateList.HeapPushUnique(ParentProxy);
 		}
 		bWidgetNeedsRepaint = true;
 	}
@@ -158,7 +160,7 @@ bool FWidgetProxy::ProcessInvalidation(FSlateInvalidationWidgetHeap& UpdateList,
 					FSlateDebugging::BroadcastWidgetInvalidate(ParentProxy.GetWidget(), WidgetPtr, EInvalidateWidgetReason::Layout);
 #endif
 					UE_TRACE_SLATE_WIDGET_INVALIDATED(ParentProxy.GetWidget(), WidgetPtr, EInvalidateWidgetReason::Layout);
-					UpdateList.PushUnique(ParentProxy);
+					UpdateList.HeapPushUnique(ParentProxy);
 				}
 			}
 			else if (TSharedPtr<SWidget> ParentWidget = WidgetPtr->GetParentWidget())
@@ -184,7 +186,7 @@ bool FWidgetProxy::ProcessInvalidation(FSlateInvalidationWidgetHeap& UpdateList,
 	return bWidgetNeedsRepaint;
 }
 
-void FWidgetProxy::MarkProxyUpdatedThisFrame(FSlateInvalidationWidgetHeap& UpdateList)
+void FWidgetProxy::MarkProxyUpdatedThisFrame(FSlateInvalidationWidgetPostHeap& PostUpdateList)
 {
 	bUpdatedSinceLastInvalidate = true;
 
@@ -194,7 +196,7 @@ void FWidgetProxy::MarkProxyUpdatedThisFrame(FSlateInvalidationWidgetHeap& Updat
 		if (WidgetPtr && WidgetPtr->IsFastPathVisible())
 		{
 			// If there are any updates still needed add them to the next update list
-			UpdateList.PushUnique(*this);
+			PostUpdateList.PushBackUnique(*this);
 		}
 	}
 }
@@ -268,7 +270,6 @@ FWidgetProxyHandle::FWidgetProxyHandle(FSlateInvalidationWidgetIndex InIndex)
 	, WidgetSortOrder()
 	, GenerationNumber(INDEX_NONE)
 {
-
 }
 
 bool FWidgetProxyHandle::IsValid(const SWidget* Widget) const
@@ -301,34 +302,13 @@ const FWidgetProxy& FWidgetProxyHandle::GetProxy() const
 
 void FWidgetProxyHandle::MarkWidgetUpdatedThisFrame()
 {
-	GetInvalidationRoot()->GetFastPathWidgetList()[WidgetIndex].MarkProxyUpdatedThisFrame(*GetInvalidationRoot()->WidgetsNeedingUpdate);
+	GetInvalidationRoot()->GetFastPathWidgetList()[WidgetIndex].MarkProxyUpdatedThisFrame(*GetInvalidationRoot()->WidgetsNeedingPostUpdate);
 }
 
 void FWidgetProxyHandle::MarkWidgetDirty(EInvalidateWidgetReason InvalidateReason)
 {
 	FWidgetProxy& Proxy = GetInvalidationRoot()->GetFastPathWidgetList()[WidgetIndex];
-	SWidget* WidgetPtr = Proxy.GetWidget();
-
-	if (EnumHasAnyFlags(InvalidateReason, EInvalidateWidgetReason::ChildOrder))
-	{
-		GetInvalidationRoot()->InvalidateWidgetChildOrder(WidgetPtr->AsShared());
-	}
-
-	if (Proxy.CurrentInvalidateReason == EInvalidateWidgetReason::None)
-	{
-		GetInvalidationRoot()->WidgetsNeedingUpdate->PushUnique(Proxy);
-	}
-#if 0
-	else
-	{
-		ensure(InvalidationRoot->WidgetsNeedingUpdate.Contains(Proxy));
-	}
-#endif
-	Proxy.CurrentInvalidateReason |= InvalidateReason;
-#if WITH_SLATE_DEBUGGING
-	FSlateDebugging::BroadcastWidgetInvalidate(WidgetPtr, nullptr, InvalidateReason);
-#endif
-	UE_TRACE_SLATE_WIDGET_INVALIDATED(WidgetPtr, nullptr, InvalidateReason);
+	GetInvalidationRoot()->InvalidateWidget(Proxy, InvalidateReason);
 }
 
 void FWidgetProxyHandle::UpdateWidgetFlags(const SWidget* Widget, EWidgetUpdateFlags NewFlags)
@@ -344,7 +324,7 @@ void FWidgetProxyHandle::UpdateWidgetFlags(const SWidget* Widget, EWidgetUpdateF
 			// Add to update list if the widget is now tickable or has an active timer.
 			if (EnumHasAnyFlags(NewFlags, EWidgetUpdateFlags::AnyUpdate))
 			{
-				GetInvalidationRoot()->WidgetsNeedingUpdate->PushUnique(Proxy);
+				GetInvalidationRoot()->WidgetsNeedingPostUpdate->PushBackOrHeapUnique(Proxy);
 			}
 		}
 	}
diff --git a/Engine/Source/Runtime/SlateCore/Private/Test/SlateInvalidationWidgetListTest.cpp b/Engine/Source/Runtime/SlateCore/Private/Test/SlateInvalidationWidgetListTest.cpp
--- a/Engine/Source/Runtime/SlateCore/Private/Test/SlateInvalidationWidgetListTest.cpp
+++ b/Engine/Source/Runtime/SlateCore/Private/Test/SlateInvalidationWidgetListTest.cpp
@@ -14,8 +14,8 @@
 
 #define LOCTEXT_NAMESPACE "Slate.FastPath.InvalidationWidgetList"
 
-//IMPLEMENT_SIMPLE_AUTOMATION_TEST(FSlateInvalidationWidgetListTest, "Slate.FastPath.InvalidationWidgetList.AddBuildRemove", EAutomationTestFlags::ApplicationContextMask | EAutomationTestFlags::SmokeFilter)
-IMPLEMENT_SIMPLE_AUTOMATION_TEST(FSlateInvalidationWidgetListTest, "Slate.FastPath.InvalidationWidgetList.AddBuildRemove", EAutomationTestFlags::ApplicationContextMask | EAutomationTestFlags::EngineFilter)
+//IMPLEMENT_SIMPLE_AUTOMATION_TEST(FSlateInvalidationWidgetListTest, "Slate.InvalidationWidgetList.AddBuildRemove", EAutomationTestFlags::ApplicationContextMask | EAutomationTestFlags::SmokeFilter)
+IMPLEMENT_SIMPLE_AUTOMATION_TEST(FSlateInvalidationWidgetListTest, "Slate.InvalidationWidgetList.AddBuildRemove", EAutomationTestFlags::ApplicationContextMask | EAutomationTestFlags::EngineFilter)
 
 namespace UE
 {
@@ -24,6 +24,8 @@ namespace Slate
 namespace Private
 {
 
+static const bool bUpdateOnlyWhatIsNeeded = false;
+
 class SEmptyLeftWidget : public SLeafWidget
 {
 public:
@@ -42,20 +44,22 @@ TSharedRef<SVerticalBox> AddVerticalBox(TSharedPtr<SVerticalBox> VerticalBox, TC
 {
 	FName NewName = *FString::Printf(TEXT("TagVerticalBox-%c"), Letter);
 	TSharedRef<SVerticalBox> Result = SNew(SVerticalBox).Tag(NewName);
-	VerticalBox->AddSlot()[Result];
+	if (VerticalBox)
+	{
+		VerticalBox->AddSlot()[Result];
+	}
 	return Result;
 }
 
 TSharedRef<SWidget> AddEmptyWidget(TSharedPtr<SVerticalBox> VerticalBox, int32 Number)
 {
-	static FName NewName = TEXT("TagEmptyLeftWidget");
-	NewName.SetNumber(Number);
+	FName NewName = *FString::Printf(TEXT("TagLeafWidget-%d"), Number);
 	TSharedRef<SWidget> Result = SNew(SEmptyLeftWidget).Tag(NewName);
 	VerticalBox->AddSlot()[Result];
 	return Result;
 }
 
-TSharedRef<SVerticalBox> BuildTestUI_ChildOrder(TSharedPtr<SVerticalBox>& WidgetC, TSharedPtr<SVerticalBox>& WidgetF)
+TSharedRef<SVerticalBox> BuildTestUI_ChildOrder(TSharedPtr<SVerticalBox>& WidgetC, TSharedPtr<SVerticalBox>& WidgetE, TSharedPtr<SVerticalBox>& WidgetF)
 {
 	// A
 	//  B (1, 2, 3)
@@ -69,7 +73,7 @@ TSharedRef<SVerticalBox> BuildTestUI_ChildOrder(TSharedPtr<SVerticalBox>& Widget
 	//   I
 	//  J (14)
 
-	TSharedRef<SVerticalBox> Root = SNew(SVerticalBox);
+	TSharedRef<SVerticalBox> Root = AddVerticalBox(nullptr, TEXT('A'));
 	{
 		TSharedRef<SVerticalBox> Sub = AddVerticalBox(Root, TEXT('B'));
 		AddEmptyWidget(Sub, 1);
@@ -86,10 +90,10 @@ TSharedRef<SVerticalBox> BuildTestUI_ChildOrder(TSharedPtr<SVerticalBox>& Widget
 	Root->AddSlot()[SNullWidget::NullWidget];
 	AddVerticalBox(Root, TEXT('D'));
 	{
-		TSharedRef<SVerticalBox> Sub = AddVerticalBox(Root, TEXT('E'));
-		AddEmptyWidget(Sub, 8);
-		AddEmptyWidget(Sub, 9);
-		AddEmptyWidget(Sub, 10);
+		WidgetE = AddVerticalBox(Root, TEXT('E'));
+		AddEmptyWidget(WidgetE, 8);
+		AddEmptyWidget(WidgetE, 9);
+		AddEmptyWidget(WidgetE, 10);
 	}
 	{
 		WidgetF = AddVerticalBox(Root, TEXT('F'));
@@ -131,7 +135,7 @@ TSharedRef<SVerticalBox> BuildTestUI_Child(TSharedPtr<SVerticalBox>& WidgetA, TA
 	//  H (8)
 	//  I
 
-	WidgetA = SNew(SVerticalBox);
+	WidgetA = AddVerticalBox(nullptr, TEXT('A'));
 	{
 		WidgetB = AddVerticalBox(WidgetA, TEXT('B'));
 		ChildOfWidgetA.Add(WidgetB);
@@ -228,8 +232,8 @@ bool FSlateInvalidationWidgetListTest::RunTest(const FString& Parameters)
 		}
 
 		{
-			TSharedPtr<SVerticalBox> WidgetC, WidgetF;
-			TSharedRef<SVerticalBox> RootChildOrder = UE::Slate::Private::BuildTestUI_ChildOrder(WidgetC, WidgetF);
+			TSharedPtr<SVerticalBox> WidgetC, WidgetE, WidgetF;
+			TSharedRef<SVerticalBox> RootChildOrder = UE::Slate::Private::BuildTestUI_ChildOrder(WidgetC, WidgetE, WidgetF);
 
 			FSlateInvalidationWidgetList List = { FSlateInvalidationRootHandle(), Args };
 			List.BuildWidgetList(RootChildOrder);
@@ -240,7 +244,9 @@ bool FSlateInvalidationWidgetListTest::RunTest(const FString& Parameters)
 				List.RemoveWidget(WidgetF->GetAllChildren()->GetChildAt(1));
 				WidgetF->RemoveSlot(WidgetF->GetAllChildren()->GetChildAt(1));
 				{
-					FSlateInvalidationWidgetList TempList = { FSlateInvalidationRootHandle(), Args };
+					FSlateInvalidationWidgetList::FArguments ArgsWithoutAssign = Args;
+					ArgsWithoutAssign.bAssignedWidgetIndex = false;
+					FSlateInvalidationWidgetList TempList = { FSlateInvalidationRootHandle(), ArgsWithoutAssign };
 					TempList.BuildWidgetList(RootChildOrder);
 					if (!TempList.DeapCompare(List))
 					{
@@ -270,7 +276,9 @@ bool FSlateInvalidationWidgetListTest::RunTest(const FString& Parameters)
 				RootChildOrder->RemoveSlot(WidgetF.ToSharedRef());
 				RootChildOrder->RemoveSlot(WidgetC.ToSharedRef());
 				{
-					FSlateInvalidationWidgetList TempList = { FSlateInvalidationRootHandle(), Args };
+					FSlateInvalidationWidgetList::FArguments ArgsWithoutAssign = Args;
+					ArgsWithoutAssign.bAssignedWidgetIndex = false;
+					FSlateInvalidationWidgetList TempList = { FSlateInvalidationRootHandle(), ArgsWithoutAssign };
 					TempList.BuildWidgetList(RootChildOrder);
 					if (!TempList.DeapCompare(List))
 					{
@@ -286,7 +294,9 @@ bool FSlateInvalidationWidgetListTest::RunTest(const FString& Parameters)
 				List.RemoveWidget(List.FindWidget(RemovedWidget));
 
 				{
-					FSlateInvalidationWidgetList TempList = { FSlateInvalidationRootHandle(), Args };
+					FSlateInvalidationWidgetList::FArguments ArgsWithoutAssign = Args;
+					ArgsWithoutAssign.bAssignedWidgetIndex = false;
+					FSlateInvalidationWidgetList TempList = { FSlateInvalidationRootHandle(), ArgsWithoutAssign };
 					TempList.BuildWidgetList(RootChildOrder);
 					if (!TempList.DeapCompare(List))
 					{
@@ -296,31 +306,255 @@ bool FSlateInvalidationWidgetListTest::RunTest(const FString& Parameters)
 			}
 		}
 		{
-			TSharedPtr<SVerticalBox> WidgetC, WidgetF;
-			TSharedRef<SVerticalBox> RootChildOrder = UE::Slate::Private::BuildTestUI_ChildOrder(WidgetC, WidgetF);
+			TSharedPtr<SVerticalBox> WidgetC, WidgetE, WidgetF;
+			TSharedRef<SVerticalBox> RootChildOrder = UE::Slate::Private::BuildTestUI_ChildOrder(WidgetC, WidgetE, WidgetF);
 
 			// Child invalidation
 			FSlateInvalidationWidgetList List = { FSlateInvalidationRootHandle(), Args };
 			List.BuildWidgetList(RootChildOrder);
 			AddErrorIfFalse(List.VerifyWidgetsIndex(), TEXT("The widget list integrity has failed."));
 
-			WidgetF->RemoveSlot(WidgetF->GetAllChildren()->GetChildAt(0));
-			WidgetC->RemoveSlot(WidgetC->GetAllChildren()->GetChildAt(1));
+			struct FChildOrderInvalidationCallback_Count : FSlateInvalidationWidgetList::IProcessChildOrderInvalidationCallback
+			{
+				mutable int32 ChildRemoveCount = 0;
+				mutable int32 ReIndexedCount = 0;
+				mutable int32 ResortCount = 0;
+				mutable int32 BuiltCount = 0;
+				virtual void PreChildRemove(const FSlateInvalidationWidgetList::FIndexRange& Range) override { ++ChildRemoveCount; }
+				using FReIndexOperation = FSlateInvalidationWidgetList::IProcessChildOrderInvalidationCallback::FReIndexOperation;
+				virtual void ProxiesReIndexed(const FReIndexOperation& Operation) override { ++ReIndexedCount; }
+				using FReSortOperation = FSlateInvalidationWidgetList::IProcessChildOrderInvalidationCallback::FReSortOperation;
+				virtual void ProxiesPreResort(const FReSortOperation& Operation) override { ++ResortCount; }
+				virtual void ProxiesBuilt(const FSlateInvalidationWidgetList::FIndexRange& Range) override { ++BuiltCount; }
+			};
+
+			auto TestRemoveWidget = [&](TSharedPtr<SVerticalBox>& Widget, const TCHAR* Message) -> bool
+				{
+					FSlateInvalidationWidgetList::InvalidationWidgetType& InvalidationWidget = List[Widget->GetProxyHandle().GetWidgetIndex()];
+					FChildOrderInvalidationCallback_Count Callback;
+					const bool bIsProxyValid = List.ProcessChildOrderInvalidation(InvalidationWidget, Callback);
+					AddErrorIfFalse(bIsProxyValid, FString::Printf(TEXT("The proxy should be valid for '%s'."), Message));
+					AddErrorIfFalse(List.VerifyWidgetsIndex(), FString::Printf(TEXT("The widget list integrity has failed for '%s'."), Message));
+					AddErrorIfFalse(List.VerifySortOrder(), FString::Printf(TEXT("The widget list sort order has failed for '%s'."), Message));
 
-			TArray<TWeakPtr<SWidget>> InvalidationWidgetIndexes;
-			InvalidationWidgetIndexes.Add(WidgetF->AsShared());
-			InvalidationWidgetIndexes.Add(WidgetC->AsShared());
-			InvalidationWidgetIndexes.Add(WidgetC->AsShared());
-			List.ProcessChildOrderInvalidation(InvalidationWidgetIndexes);
-			AddErrorIfFalse(List.VerifyWidgetsIndex(), TEXT("The widget list integrity has failed."));
+					{
+						FSlateInvalidationWidgetList::FArguments ArgsWithoutAssign = Args;
+						ArgsWithoutAssign.bAssignedWidgetIndex = false;
+						FSlateInvalidationWidgetList TempList = { FSlateInvalidationRootHandle(), ArgsWithoutAssign };
+						TempList.BuildWidgetList(RootChildOrder);
+						if (!TempList.DeapCompare(List))
+						{
+							AddError(FString::Printf(TEXT("Was not able to process invalidation of widget '%s'."), Message));
+						}
+					}
+					return Callback.ChildRemoveCount > 0;
+				};
+
+			{
+				WidgetF->RemoveSlot(WidgetF->GetAllChildren()->GetChildAt(0));
+				bool bChildRemoved = TestRemoveWidget(WidgetF, TEXT("F"));
+				AddErrorIfFalse(bChildRemoved, TEXT("F has to removed 1 widget"));
+			}
+
+			{
+				WidgetE->RemoveSlot(WidgetE->GetAllChildren()->GetChildAt(1));
+				bool bChildRemoved = TestRemoveWidget(WidgetE, TEXT("E"));
+				AddErrorIfFalse(bChildRemoved, TEXT("E has to removed 1 widget"));
+			}
+
+			{
+				WidgetC->RemoveSlot(WidgetC->GetAllChildren()->GetChildAt(1));
+				WidgetC->RemoveSlot(WidgetC->GetAllChildren()->GetChildAt(1));
+				bool bChildRemoved = TestRemoveWidget(WidgetC, TEXT("C"));
+				AddErrorIfFalse(bChildRemoved, TEXT("C has to removed 2 widget"));
+			}
 
 			{
-				FSlateInvalidationWidgetList TempList = { FSlateInvalidationRootHandle(), Args };
-				TempList.BuildWidgetList(RootChildOrder);
-				if (!TempList.DeapCompare(List))
+				bool bChildRemoved = TestRemoveWidget(WidgetC, TEXT("C"));
+				AddErrorIfFalse(!bChildRemoved || !UE::Slate::Private::bUpdateOnlyWhatIsNeeded, TEXT("C no widget was removed"));
+			}
+		}
+	}
+
+
+	{
+		// A						//[0]A, B, 1, 2
+		//  B (1, 2, 3)				//[1]3, C, 4, 5
+		//  C (4, 5, 6, 7)			//[2]6, 7, D, E
+		//  Null					//[3]8, 9, 0, F
+		//  D						//[4]G, 1, 2, H
+		//  E (8, 9, 10)			//[5]3, I, J, 4
+		//  F
+		//   G (11, Null, 12)
+		//   H (13)
+		//   I
+		//  J (14)
+
+		TSharedPtr<SVerticalBox> WidgetC, WidgetE, WidgetF;
+		TSharedRef<SVerticalBox> RootChildOrder = UE::Slate::Private::BuildTestUI_ChildOrder(WidgetC, WidgetE, WidgetF);
+
+		FSlateInvalidationWidgetList::FArguments Arg = { 4,2 };
+		FSlateInvalidationWidgetList List = { FSlateInvalidationRootHandle(), Arg };
+		List.BuildWidgetList(RootChildOrder);
+		AddErrorIfFalse(List.VerifyWidgetsIndex(), TEXT("The widget list integrity has failed."));
+
+		// Test FProcessChildOrderInvalidationResult with G
+		{
+			struct FChildOrderInvalidationCallback_RemoveG : FSlateInvalidationWidgetList::IProcessChildOrderInvalidationCallback
+			{
+				virtual ~FChildOrderInvalidationCallback_RemoveG() = default;
+				FSlateInvalidationWidgetListTest* Self = nullptr;
+
+				TArray<TTuple<FString, FSlateInvalidationWidgetSortOrder, bool>> PreChildRemoveToCheck;
+				virtual void PreChildRemove(const FSlateInvalidationWidgetList::FIndexRange& Range) override
 				{
-					AddError(TEXT("Was not able to process invalidation C and F."));
+					Self->AddErrorIfFalse(Range.IsValid(), TEXT("The range is not valid."));
+
+					if (PreChildRemoveToCheck.Num() == 0)
+					{
+						Self->AddError(TEXT("No PreChildRemove was expected"));
+					}
+					else
+					{
+						for (const auto& Element : PreChildRemoveToCheck)
+						{
+							bool bIsIncluded = Range.Include(Element.Get<1>());
+							if (bIsIncluded != Element.Get<2>())
+							{
+								Self->AddError(FString::Printf(TEXT("'%s' should be %s in the PreChildRemove range"), *Element.Get<0>(), (bIsIncluded ? TEXT("included") : TEXT("not be included"))));
+							}
+						}
+					}
 				}
+
+				TArray<TTuple<FString, FSlateInvalidationWidgetSortOrder, bool>> ProxiesReIndexedToCheck;
+				using FReIndexOperation = FSlateInvalidationWidgetList::IProcessChildOrderInvalidationCallback::FReIndexOperation;
+				virtual void ProxiesReIndexed(const FReIndexOperation& Operation) override
+				{
+					Self->AddErrorIfFalse(Operation.GetRange().IsValid(), TEXT("The range is not valid."));
+
+					if (ProxiesReIndexedToCheck.Num() == 0)
+					{
+						Self->AddError(TEXT("No ProxiesReIndexed was expected"));
+					}
+					else
+					{
+						for (const auto& Element : ProxiesReIndexedToCheck)
+						{
+							bool bIsIncluded = Operation.GetRange().Include(Element.Get<1>());
+							if (bIsIncluded != Element.Get<2>())
+							{
+								Self->AddError(FString::Printf(TEXT("'%s' should be %s in the ProxiesReIndexed range"), *Element.Get<0>(), (bIsIncluded ? TEXT("included") : TEXT("not be included"))));
+							}
+						}
+					}
+				}
+
+				TArray<TTuple<FString, FSlateInvalidationWidgetSortOrder, bool>> ProxiesPreResortToCheck;
+				using FReSortOperation = FSlateInvalidationWidgetList::IProcessChildOrderInvalidationCallback::FReSortOperation;
+				virtual void ProxiesPreResort(const FReSortOperation& Operation) override
+				{
+					Self->AddErrorIfFalse(Operation.GetRange().IsValid(), TEXT("The range is not valid."));
+
+					if (ProxiesPreResortToCheck.Num() == 0)
+					{
+						Self->AddError(TEXT("No ProxiesPreResort was expected"));
+					}
+					else
+					{
+						for (const auto& Element : ProxiesPreResortToCheck)
+						{
+							bool bIsIncluded = Operation.GetRange().Include(Element.Get<1>());
+							if (bIsIncluded != Element.Get<2>())
+							{
+								Self->AddError(FString::Printf(TEXT("'%s' should be %s in the ProxiesPreResort range"), *Element.Get<0>(), (bIsIncluded ? TEXT("included") : TEXT("not be included"))));
+							}
+						}
+					}
+				}
+
+				TArray<TTuple<FString, FSlateInvalidationWidgetSortOrder, bool>> ProxiesBuiltToCheck;
+				virtual void ProxiesBuilt(const FSlateInvalidationWidgetList::FIndexRange& Range) override
+				{
+					Self->AddErrorIfFalse(Range.IsValid(), TEXT("The range is not valid."));
+
+					if (ProxiesBuiltToCheck.Num() == 0)
+					{
+						Self->AddError(TEXT("No ProxiesBuilt was expected"));
+					}
+					else
+					{
+						for (const auto& Element : ProxiesBuiltToCheck)
+						{
+							bool bIsIncluded = Range.Include(Element.Get<1>());
+							if (bIsIncluded != Element.Get<2>())
+							{
+								Self->AddError(FString::Printf(TEXT("'%s' should be %s in the ProxiesBuilt range"), *Element.Get<0>(), (bIsIncluded ? TEXT("included") : TEXT("not be included"))));
+							}
+						}
+					}
+				}
+			};
+
+			FChildOrderInvalidationCallback_RemoveG CallbackG;
+			CallbackG.Self = this;
+
+			const FSlateInvalidationWidgetIndex WidgetIndexC = WidgetC->GetProxyHandle().GetWidgetIndex();
+			const FSlateInvalidationWidgetIndex WidgetIndexF = WidgetF->GetProxyHandle().GetWidgetIndex();
+			const FSlateInvalidationWidgetIndex WidgetIndexG = WidgetF->GetAllChildren()->GetChildAt(0)->GetProxyHandle().GetWidgetIndex();
+			const FSlateInvalidationWidgetIndex WidgetIndex11 = List.IncrementIndex(WidgetIndexG);
+			AddErrorIfFalse(WidgetIndex11 == StaticCastSharedRef<SVerticalBox>(WidgetF->GetAllChildren()->GetChildAt(0))->GetAllChildren()->GetChildAt(0)->GetProxyHandle().GetWidgetIndex(), TEXT("Index of 11 is not identical."));
+			const FSlateInvalidationWidgetIndex WidgetIndex12 = List.IncrementIndex(WidgetIndex11);
+			AddErrorIfFalse(WidgetIndex12 == StaticCastSharedRef<SVerticalBox>(WidgetF->GetAllChildren()->GetChildAt(0))->GetAllChildren()->GetChildAt(2)->GetProxyHandle().GetWidgetIndex(), TEXT("Index of 12 is not identical."));
+			const FSlateInvalidationWidgetIndex WidgetIndexH = WidgetF->GetAllChildren()->GetChildAt(1)->GetProxyHandle().GetWidgetIndex();
+			const FSlateInvalidationWidgetIndex WidgetIndex13 = List.IncrementIndex(WidgetIndexH);
+			AddErrorIfFalse(WidgetIndex13 == StaticCastSharedRef<SVerticalBox>(WidgetF->GetAllChildren()->GetChildAt(1))->GetAllChildren()->GetChildAt(0)->GetProxyHandle().GetWidgetIndex(), TEXT("Index of 13 is not identical."));
+			const FSlateInvalidationWidgetIndex WidgetIndexI = List.IncrementIndex(WidgetIndex13);
+			const FSlateInvalidationWidgetIndex WidgetIndexJ = List.IncrementIndex(WidgetIndexI);
+			const FSlateInvalidationWidgetIndex WidgetIndex14 = List.IncrementIndex(WidgetIndexJ);
+
+			{
+				CallbackG.PreChildRemoveToCheck.Emplace(TEXT("C"), FSlateInvalidationWidgetSortOrder{ List, WidgetIndexC }, false);
+				CallbackG.PreChildRemoveToCheck.Emplace(TEXT("F"), FSlateInvalidationWidgetSortOrder{ List, WidgetIndexF }, false);
+				CallbackG.PreChildRemoveToCheck.Emplace(TEXT("G"), FSlateInvalidationWidgetSortOrder{ List, WidgetIndexG }, true);
+				CallbackG.PreChildRemoveToCheck.Emplace(TEXT("11"), FSlateInvalidationWidgetSortOrder{ List, WidgetIndex11 }, true);
+				CallbackG.PreChildRemoveToCheck.Emplace(TEXT("12"), FSlateInvalidationWidgetSortOrder{ List, WidgetIndex12 }, true);
+				CallbackG.PreChildRemoveToCheck.Emplace(TEXT("H"), FSlateInvalidationWidgetSortOrder{ List, WidgetIndexH }, true);
+				CallbackG.PreChildRemoveToCheck.Emplace(TEXT("13"), FSlateInvalidationWidgetSortOrder{ List, WidgetIndex13 }, true);
+				CallbackG.PreChildRemoveToCheck.Emplace(TEXT("I"), FSlateInvalidationWidgetSortOrder{ List, WidgetIndexI }, true);
+				CallbackG.PreChildRemoveToCheck.Emplace(TEXT("J"), FSlateInvalidationWidgetSortOrder{ List, WidgetIndexJ }, false);
+				CallbackG.PreChildRemoveToCheck.Emplace(TEXT("14"), FSlateInvalidationWidgetSortOrder{ List, WidgetIndex14 }, false);
+			}
+
+			{
+				//No ProxiesReIndexed expected
+			}
+
+			{
+				//No ProxiesPreResortToCheck expected
+			}
+
+			{
+
+				CallbackG.ProxiesBuiltToCheck.Emplace(TEXT("C"), FSlateInvalidationWidgetSortOrder{ List, WidgetIndexC }, false);
+				CallbackG.ProxiesBuiltToCheck.Emplace(TEXT("F"), FSlateInvalidationWidgetSortOrder{ List, WidgetIndexF }, false);
+				const FSlateInvalidationWidgetIndex NewWidgetIndexH = WidgetF->GetAllChildren()->GetChildAt(0)->GetProxyHandle().GetWidgetIndex();
+				CallbackG.ProxiesBuiltToCheck.Emplace(TEXT("H"), FSlateInvalidationWidgetSortOrder{ List, NewWidgetIndexH }, true);
+				const FSlateInvalidationWidgetIndex NewWidgetIndex13 = List.IncrementIndex(NewWidgetIndexH);
+				CallbackG.ProxiesBuiltToCheck.Emplace(TEXT("13"), FSlateInvalidationWidgetSortOrder{ List, NewWidgetIndex13 }, true);
+				const FSlateInvalidationWidgetIndex NewWidgetIndexI = List.IncrementIndex(NewWidgetIndex13);
+				CallbackG.ProxiesBuiltToCheck.Emplace(TEXT("I"), FSlateInvalidationWidgetSortOrder{ List, NewWidgetIndexI }, true);
+				const FSlateInvalidationWidgetIndex NewWidgetIndexJ = List.IncrementIndex(NewWidgetIndexI);
+				CallbackG.ProxiesBuiltToCheck.Emplace(TEXT("J"), FSlateInvalidationWidgetSortOrder{ List, NewWidgetIndexJ }, false);
+				const FSlateInvalidationWidgetIndex NewWidgetIndex14 = List.IncrementIndex(NewWidgetIndexJ);
+				CallbackG.ProxiesBuiltToCheck.Emplace(TEXT("14"), FSlateInvalidationWidgetSortOrder{ List, NewWidgetIndex14 }, false);
+			}
+
+			// Remove G
+			{
+				WidgetF->RemoveSlot(WidgetF->GetAllChildren()->GetChildAt(0));
+				bool bIsValid = List.ProcessChildOrderInvalidation(List[WidgetIndexF], CallbackG);
+				AddErrorIfFalse(bIsValid, TEXT("The F widget should still be valid"));
 			}
 		}
 	}
diff --git a/Engine/Source/Runtime/SlateCore/Public/FastUpdate/SlateInvalidationRoot.h b/Engine/Source/Runtime/SlateCore/Public/FastUpdate/SlateInvalidationRoot.h
--- a/Engine/Source/Runtime/SlateCore/Public/FastUpdate/SlateInvalidationRoot.h
+++ b/Engine/Source/Runtime/SlateCore/Public/FastUpdate/SlateInvalidationRoot.h
@@ -10,8 +10,9 @@
 #include "Rendering/DrawElements.h"
 
 struct FSlateCachedElementData;
-class FSlateInvalidationWidgetHeap;
 class FSlateInvalidationWidgetList;
+class FSlateInvalidationWidgetPreHeap;
+class FSlateInvalidationWidgetPostHeap;
 class FSlateWindowElementList;
 class FWidgetStyle;
 
@@ -115,8 +116,6 @@ public:
 
 	SLATECORE_API void Advanced_ResetInvalidation(bool bClearResourcesImmediately);
 
-	SLATECORE_API static void ClearAllWidgetUpdatesPending();
-
 #if WITH_SLATE_DEBUGGING
 	/** @return the last paint type the invalidation root handle used. */
 	ESlateInvalidationPaintType GetLastPaintType() const { return LastPaintType; }
@@ -144,22 +143,35 @@ private:
 
 	bool PaintFastPath(const FSlateInvalidationContext& Context);
 
-	/** Call to notify that the ordering of children below this Widget has changed and the fast path is no longer valid. */
-	void InvalidateWidgetChildOrder(TSharedRef<SWidget> Widget);
-	void ProcessChildOrderUpdate();
-	void BuildFastPathWidgetList(TSharedRef<SWidget> RootWidget);
+	/** Call when an invalidation occurred. */
+	void InvalidateWidget(FWidgetProxy& Proxy, EInvalidateWidgetReason InvalidateReason);
 
+	void BuildFastPathWidgetList(TSharedRef<SWidget> RootWidget);
 	void AdjustWidgetsDesktopGeometry(FVector2D WindowToDesktopTransform);
 
+	/** Update child order and slate attribute */
+	void ProcessPreUpdate();
+	/** Update layout, and update the FinalUpdateList */
+	bool ProcessPostUpdate();
+
 private:
 	/** List of all the Widget included by this SlateInvalidationRoot. */
-	FSlateInvalidationWidgetList* FastWidgetPathList;
-	/** Index to widgets which are dirty, volatile, or need some sort of per frame update (such as a tick or timer) */
-	FSlateInvalidationWidgetHeap* WidgetsNeedingUpdate;
+	TUniquePtr<FSlateInvalidationWidgetList> FastWidgetPathList;
+
+	/**
+	 * Index of widgets that are dirty (child order)
+	 * and need to be updated before the Post list.
+	 */
+	TUniquePtr<FSlateInvalidationWidgetPreHeap> WidgetsNeedingPreUpdate;
+
+	/**
+	 * Index to widgets that are dirty (volatile, or need some sort of per frame update such as a tick or timer)
+	 * and need to be updated after the Pre list.
+	 */
+	TUniquePtr<FSlateInvalidationWidgetPostHeap> WidgetsNeedingPostUpdate;
+
 	/** Index to widgets that will be updated. */
 	TArray<FSlateInvalidationWidgetIndex> FinalUpdateList;
-	/** Widget that has ChildOrder invalidation. */
-	TArray<TWeakPtr<SWidget>> WidgetsNeedingChildOrderUpdate;
 
 	FSlateCachedElementData* CachedElementData;
 
@@ -182,15 +194,15 @@ private:
 	bool bChildOrderInvalidated;
 	bool bNeedsSlowPath;
 	bool bNeedScreenPositionShift;
-	bool bProcessingChildOrderUpdate;
+	bool bProcessingPreUpdate;
+	bool bProcessingPostUpdate;
+	bool bBuildingWidgetList;
+	bool bProcessingChildOrderInvalidation;
 
 #if WITH_SLATE_DEBUGGING
 	ESlateInvalidationPaintType LastPaintType;
-	uint32 ProcessInvalidationFrameNumber;
 #endif
 #if UE_SLATE_DEBUGGING_CLEAR_ALL_FAST_PATH_DATA
 	TArray<const SWidget*> FastWidgetPathToClearedBecauseOfDelay;
 #endif
-
-	static TArray<FSlateInvalidationRoot*> ClearUpdateList;
 };
diff --git a/Engine/Source/Runtime/SlateCore/Public/FastUpdate/SlateInvalidationWidgetSortOrder.h b/Engine/Source/Runtime/SlateCore/Public/FastUpdate/SlateInvalidationWidgetSortOrder.h
--- a/Engine/Source/Runtime/SlateCore/Public/FastUpdate/SlateInvalidationWidgetSortOrder.h
+++ b/Engine/Source/Runtime/SlateCore/Public/FastUpdate/SlateInvalidationWidgetSortOrder.h
@@ -17,11 +17,15 @@ struct FSlateInvalidationWidgetSortOrder
 {
 private:
 	uint32 Order = 0;
+	FSlateInvalidationWidgetSortOrder(uint32 InOrder) : Order(InOrder) {}
 
 public:
 	FSlateInvalidationWidgetSortOrder() = default;
 	FSlateInvalidationWidgetSortOrder(const FSlateInvalidationWidgetList& List, FSlateInvalidationWidgetIndex Index);
 
+	static FSlateInvalidationWidgetSortOrder LimitMax();
+	static FSlateInvalidationWidgetSortOrder LimitMin();
+
 	bool operator< (const FSlateInvalidationWidgetSortOrder Other) const { return Order < Other.Order; }
 	bool operator<= (const FSlateInvalidationWidgetSortOrder Other) const { return Order <= Other.Order; }
 	bool operator> (const FSlateInvalidationWidgetSortOrder Other) const { return Order > Other.Order; }
diff --git a/Engine/Source/Runtime/SlateCore/Public/FastUpdate/WidgetProxy.h b/Engine/Source/Runtime/SlateCore/Public/FastUpdate/WidgetProxy.h
--- a/Engine/Source/Runtime/SlateCore/Public/FastUpdate/WidgetProxy.h
+++ b/Engine/Source/Runtime/SlateCore/Public/FastUpdate/WidgetProxy.h
@@ -17,7 +17,7 @@
 class SWidget;
 class FPaintArgs;
 struct FFastPathPerFrameData;
-class FSlateInvalidationWidgetHeap;
+class FSlateInvalidationWidgetPostHeap;
 class FSlateInvalidationWidgetList;
 struct FSlateWidgetPersistentState;
 class FSlateInvalidationRoot;
@@ -35,9 +35,9 @@ public:
 
 	int32 Update(const FPaintArgs& PaintArgs, FSlateWindowElementList& OutDrawElements);
 
-	bool ProcessInvalidation(FSlateInvalidationWidgetHeap& UpdateList, FSlateInvalidationWidgetList& FastPathWidgetList, FSlateInvalidationRoot& Root);
+	bool ProcessPostInvalidation(FSlateInvalidationWidgetPostHeap& UpdateList, FSlateInvalidationWidgetList& FastPathWidgetList, FSlateInvalidationRoot& Root);
 
-	void MarkProxyUpdatedThisFrame(FSlateInvalidationWidgetHeap& UpdateList);
+	void MarkProxyUpdatedThisFrame(FSlateInvalidationWidgetPostHeap& UpdateList);
 
 #if UE_SLATE_WITH_WIDGETPROXY_WEAKPTR
 	SWidget* GetWidget() const
@@ -87,11 +87,15 @@ public:
 	EVisibility Visibility;
 	/** Used to make sure we don't double process a widget that is invalidated.  (a widget can invalidate itself but an ancestor can end up painting that widget first thus rendering the child's own invalidate unnecessary */
 	uint8 bUpdatedSinceLastInvalidate : 1;
-	/** Is the widget already in a pending update list.  If it already is in an update list we don't bother adding it again */
-	uint8 bContainedByWidgetHeap : 1;
+	/** Is the widget already in a pending pre update list.  If it already is in an update list we don't bother adding it again */
+	uint8 bContainedByWidgetPreHeap : 1;
+	/** Is the widget already in a pending post update list.  If it already is in an update list we don't bother adding it again */
+	uint8 bContainedByWidgetPostHeap : 1;
 	/** Use with "Slate.InvalidationRoot.VerifyWidgetVisibility". Cached the last FastPathVisible value to find widgets that do not call Invalidate properly. */
 	uint8 bDebug_LastFrameVisible : 1;
 	uint8 bDebug_LastFrameVisibleSet : 1;
+	/** Use with "SLate.InvalidationRoot.VerifyWidgetAttribute". */
+	uint8 bDebug_AttributeUpdated : 1;
 };
 
 #if !UE_SLATE_WITH_WIDGETPROXY_WIDGETTYPE
diff --git a/Engine/Source/Runtime/SlateCore/Public/Widgets/InvalidateWidgetReason.h b/Engine/Source/Runtime/SlateCore/Public/Widgets/InvalidateWidgetReason.h
--- a/Engine/Source/Runtime/SlateCore/Public/Widgets/InvalidateWidgetReason.h
+++ b/Engine/Source/Runtime/SlateCore/Public/Widgets/InvalidateWidgetReason.h
@@ -31,7 +31,9 @@ enum class EInvalidateWidgetReason : uint8
 	 */
 	ChildOrder = 1 << 3,
 
-	/** A Widgets render transform changed */
+	/**
+	 * A Widgets render transform changed
+	 */
 	RenderTransform = 1 << 4,
 
 	/**
@@ -39,6 +41,11 @@ enum class EInvalidateWidgetReason : uint8
 	 */
 	Visibility = 1 << 5,
 
+	/**
+	 * Attributes got bound or unbound (it's used by the SlateAttributeMetaData)
+	 */
+	AttributeRegistration = 1 << 6,
+
 	/**
 	 * Use Paint invalidation if you're changing a normal property involving painting or sizing.
 	 * Additionally if the property that was changed affects Volatility in anyway, it's important
