Slate: Add SlateAttributes unit-tests.
Add Depedency, Prerequisit and Collapsed features for SlateAttributeDescriptor. Prevent SlateAttribute from invalidating the widget when it's constructing. Add explicit message when the SlateType is use wrongly (when possible it's done at compile time).

diff --git a/Engine/Source/Runtime/SlateCore/Private/Test/SlateAttributeTest.cpp b/Engine/Source/Runtime/SlateCore/Private/Test/SlateAttributeTest.cpp
new file mode 100644
--- /dev/null
+++ b/Engine/Source/Runtime/SlateCore/Private/Test/SlateAttributeTest.cpp
@@ -0,0 +1,429 @@
+// Copyright Epic Games, Inc. All Rights Reserved.
+
+#include "CoreMinimal.h"
+#include "Misc/AutomationTest.h"
+
+#include "Widgets/DeclarativeSyntaxSupport.h"
+#include "Widgets/SLeafWidget.h"
+#include "Widgets/SNullWidget.h"
+
+#if WITH_AUTOMATION_WORKER
+
+#define LOCTEXT_NAMESPACE "Slate.Attribute"
+
+IMPLEMENT_SIMPLE_AUTOMATION_TEST(FSlateAttributeTest, "Slate.Attribute", EAutomationTestFlags::ApplicationContextMask | EAutomationTestFlags::EngineFilter)
+
+namespace UE
+{
+namespace Slate
+{
+namespace Private
+{
+
+int32 CallbackForIntAttribute(int32 Value)
+{
+	return Value;
+}
+FVector2D CallbackForFVectorAttribute()
+{
+	return FVector2D(1,1);
+}
+
+class SAttributeLeftWidget_Parent : public SLeafWidget
+{
+	SLATE_DECLARE_WIDGET(SAttributeLeftWidget_Parent, SLeafWidget)
+public:
+	SLATE_BEGIN_ARGS(SAttributeLeftWidget_Parent) {}
+	SLATE_END_ARGS()
+
+
+	SAttributeLeftWidget_Parent()
+		: IntAttributeA(*this, 1)
+		, IntAttributeB(*this, 1)
+		, IntAttributeC(*this, 1)
+		, IntAttributeD(*this, 1)
+		, IntAttributeWithPredicate(*this, 1)
+		//, IntManagedAttributes{*this, 4}
+	{
+	}
+	void Construct(const FArguments& InArgs){}
+	virtual FVector2D ComputeDesiredSize(float) const override { return FVector2D{ 100, 100 }; }
+	virtual int32 OnPaint(const FPaintArgs& Args, const FGeometry& AllottedGeometry, const FSlateRect& MyCullingRect, FSlateWindowElementList& OutDrawElements, int32 LayerId, const FWidgetStyle& InWidgetStyle, bool bParentEnabled) const override
+	{
+		return LayerId;
+	}
+
+	int32 TestMember()
+	{
+		// this should not compile
+		{
+			//TSlateAttribute<int32> Other = IntAttribute;
+			//Other.Set(*this, 4);
+			//return Other.Get();
+		}
+		// this should compile but failed with check at runtime
+		{
+			//TSlateAttribute<int32> Other{ *this, 0 };
+			//Other.Set(*this, 4);
+			//return Other.Get();
+		}
+		{
+			int32 DefaultValue = 0;
+			TSlateAttribute<int32> Attribute = { *this, DefaultValue };
+		}
+		{
+			int32 DefaultValue = 0;
+			auto Getter1 = TAttribute<int32>::FGetter::CreateStatic(UE::Slate::Private::CallbackForIntAttribute, 1);
+			TSlateAttribute<int32> Attribute = { *this, DefaultValue, Getter1 };
+		}
+		{
+			int32 DefaultValue = 0;
+			auto Getter1 = TAttribute<int32>::FGetter::CreateStatic(UE::Slate::Private::CallbackForIntAttribute, 1);
+			TSlateAttribute<int32> Attribute = { *this, DefaultValue, MoveTemp(Getter1) };
+		} 
+		{
+			int32 DefaultValue = 0;
+			TSlateAttribute<int32> Attribute = { *this, DefaultValue, MakeAttributeLambda([](){return 5;}) };
+		}
+		{
+			FVector2D DefaultValue = FVector2D(0.f, 0.f);
+			TSlateAttribute<FVector2D> Attribute = { *this, MoveTemp(DefaultValue) };
+		}
+		{
+			FVector2D DefaultValue = FVector2D(0.f, 0.f);
+			auto Getter1 = TAttribute<FVector2D>::FGetter::CreateStatic(UE::Slate::Private::CallbackForFVectorAttribute);
+			TSlateAttribute<FVector2D> Attribute = { *this, MoveTemp(DefaultValue), Getter1 };
+		}
+		{
+			FVector2D DefaultValue = FVector2D(0.f, 0.f);
+			auto Getter1 = TAttribute<FVector2D>::FGetter::CreateStatic(UE::Slate::Private::CallbackForFVectorAttribute);
+			TSlateAttribute<FVector2D> Attribute = { *this, MoveTemp(DefaultValue), MoveTemp(Getter1) };
+		}
+		{
+			FVector2D DefaultValue = FVector2D(0.f, 0.f);
+			TSlateAttribute<FVector2D> Attribute = { *this, MoveTemp(DefaultValue), MakeAttributeLambda([]() {return FVector2D(1.f, 1.f);}) };
+		}
+
+		return 0; 
+	}
+
+	void TestExternalBind()
+	{
+		// Copy should not compile
+		{
+			//TSlateManagedAttribute<int32, EInvalidateWidgetReason::ChildOrder> ExtAtt1 {AsShared(), 10};
+			//IntManagedAttributes.Add(ExtAtt1);
+		}
+
+		// Move should compile
+		{
+			TSlateManagedAttribute<int32, EInvalidateWidgetReason::ChildOrder> ExtAtt1 {AsShared(), 10};
+			IntManagedAttributes.Add(MoveTemp(ExtAtt1));
+		}
+	}
+
+	TSlateAttribute<int32> IntAttributeA;
+	TSlateAttribute<int32> IntAttributeB;
+	TSlateAttribute<int32> IntAttributeC;
+	TSlateAttribute<int32> IntAttributeD;
+	TSlateAttributeWithPredicate<int32, TSlateAttributeInvalidationReason<EInvalidateWidgetReason::ChildOrder>> IntAttributeWithPredicate;
+	TArray<TSlateManagedAttribute<int32, EInvalidateWidgetReason::ChildOrder>> IntManagedAttributes;
+};
+
+SLATE_IMPLEMENT_WIDGET(SAttributeLeftWidget_Parent)
+void SAttributeLeftWidget_Parent::PrivateRegisterAttributes(FSlateAttributeInitializer& AttributeInitializer)
+{
+	// The update order is B, A, D, C
+	// C updates when D is invalidated, so D needs to be before C
+	// A updates after B, so B needs to be before A
+	SLATE_ADD_MEMBER_ATTRIBUTE_DEFINITION(AttributeInitializer, IntAttributeD, EInvalidateWidgetReason::ChildOrder);
+	SLATE_ADD_MEMBER_ATTRIBUTE_DEFINITION(AttributeInitializer, IntAttributeC, EInvalidateWidgetReason::ChildOrder)
+		.UpdateDependency(GET_MEMBER_NAME_CHECKED(PrivateThisType, IntAttributeD));
+	SLATE_ADD_MEMBER_ATTRIBUTE_DEFINITION(AttributeInitializer, IntAttributeB, EInvalidateWidgetReason::ChildOrder);
+	SLATE_ADD_MEMBER_ATTRIBUTE_DEFINITION(AttributeInitializer, IntAttributeA, EInvalidateWidgetReason::ChildOrder)
+		.UpdatePrerequisite(GET_MEMBER_NAME_CHECKED(PrivateThisType, IntAttributeB));
+
+	AttributeInitializer.OverrideInvalidationReason(GET_MEMBER_NAME_CHECKED(PrivateThisType, IntAttributeD), FSlateAttributeDescriptor::FInvalidateWidgetReasonAttribute{ EInvalidateWidgetReason::Paint});
+}
+
+
+class SAttributeLeftWidget_Child : public SAttributeLeftWidget_Parent
+{
+	SLATE_DECLARE_WIDGET(SAttributeLeftWidget_Child, SAttributeLeftWidget_Parent)
+public:
+	SLATE_BEGIN_ARGS(SAttributeLeftWidget_Child) {}
+	SLATE_END_ARGS()
+
+
+	SAttributeLeftWidget_Child()
+		: IntAttributeH(*this, 1)
+		, IntAttributeI(*this, 1)
+		, IntAttributeJ(*this, 1)
+		, IntAttributeK(*this, 1)
+		, IntAttributeL(*this, 1)
+	{
+	}
+
+	void Construct(const FArguments& InArgs) {}
+
+	TSlateAttribute<int32, EInvalidateWidgetReason::ChildOrder> IntAttributeH;
+	TSlateAttribute<int32> IntAttributeI;
+	TSlateAttribute<int32> IntAttributeJ;
+	TSlateAttribute<int32> IntAttributeK;
+	TSlateAttribute<int32> IntAttributeL;
+};
+
+SLATE_IMPLEMENT_WIDGET(SAttributeLeftWidget_Child)
+void SAttributeLeftWidget_Child::PrivateRegisterAttributes(FSlateAttributeInitializer& AttributeInitializer)
+{
+	// The update order is B, A, I, J, D, C, L, H, K
+	//SLATE_ADD_MEMBER_ATTRIBUTE_DEFINITION(AttributeInitializer, IntAttributeH, EInvalidateWidgetReason::ChildOrder);
+	SLATE_ADD_MEMBER_ATTRIBUTE_DEFINITION(AttributeInitializer, IntAttributeJ, EInvalidateWidgetReason::ChildOrder)
+		.UpdateDependency("IntAttributeA");
+	SLATE_ADD_MEMBER_ATTRIBUTE_DEFINITION(AttributeInitializer, IntAttributeK, EInvalidateWidgetReason::ChildOrder);
+	SLATE_ADD_MEMBER_ATTRIBUTE_DEFINITION(AttributeInitializer, IntAttributeI, EInvalidateWidgetReason::ChildOrder)
+		.UpdatePrerequisite("IntAttributeB");
+	SLATE_ADD_MEMBER_ATTRIBUTE_DEFINITION(AttributeInitializer, IntAttributeL, EInvalidateWidgetReason::ChildOrder)
+		.UpdatePrerequisite("IntAttributeC");
+}
+
+}
+}
+}
+
+
+bool FSlateAttributeTest::RunTest(const FString& Parameters)
+{
+	int32 OrderCounter = 0;
+	auto OrderLambda = [this, &OrderCounter]() -> int32
+	{
+		++OrderCounter;
+		return OrderCounter;
+	};
+	bool bWasUpdate = false;
+	int32 ReturnValue = 0;
+	auto UpdateLambda = [this, &bWasUpdate, &ReturnValue]() -> int32
+	{
+		bWasUpdate = true;
+		return ReturnValue;
+	};
+
+	{
+		TSharedRef<UE::Slate::Private::SAttributeLeftWidget_Parent> WidgetParent = SNew(UE::Slate::Private::SAttributeLeftWidget_Parent);
+
+		AddErrorIfFalse(&WidgetParent->GetWidgetClass() == &UE::Slate::Private::SAttributeLeftWidget_Parent::StaticWidgetClass()
+			, TEXT("The static data do not matches"));
+
+		FSlateAttributeDescriptor const& AttributeDescriptor = WidgetParent->GetWidgetClass().GetAttributeDescriptor();
+		AddErrorIfFalse(AttributeDescriptor.AttributeNum() == 4, TEXT(""));
+
+		const int32 IndexA = AttributeDescriptor.IndexOfMemberAttribute("IntAttributeA");
+		const int32 IndexB = AttributeDescriptor.IndexOfMemberAttribute("IntAttributeB");
+		const int32 IndexC = AttributeDescriptor.IndexOfMemberAttribute("IntAttributeC");
+		const int32 IndexD = AttributeDescriptor.IndexOfMemberAttribute("IntAttributeD");
+		const int32 IndexI = AttributeDescriptor.IndexOfMemberAttribute("IntAttributeI");
+		const int32 IndexJ = AttributeDescriptor.IndexOfMemberAttribute("IntAttributeJ");
+		const int32 IndexK = AttributeDescriptor.IndexOfMemberAttribute("IntAttributeK");
+
+		AddErrorIfFalse(IndexA != INDEX_NONE, TEXT("Could not find the Attribute A"));
+		AddErrorIfFalse(IndexB != INDEX_NONE, TEXT("Could not find the Attribute B"));
+		AddErrorIfFalse(IndexC != INDEX_NONE, TEXT("Could not find the Attribute C"));
+		AddErrorIfFalse(IndexD != INDEX_NONE, TEXT("Could not find the Attribute D"));
+		AddErrorIfFalse(IndexI == INDEX_NONE, TEXT("Was not supposed to find the Attribute I"));
+		AddErrorIfFalse(IndexJ == INDEX_NONE, TEXT("Was not supposed to find the Attribute J"));
+		AddErrorIfFalse(IndexK == INDEX_NONE, TEXT("Was not supposed to find the Attribute K"));
+
+		AddErrorIfFalse(&AttributeDescriptor.GetAttributeAtIndex(IndexA) == AttributeDescriptor.FindAttribute("IntAttributeA"), TEXT("Index and Attribute should return the same value."));
+		AddErrorIfFalse(&AttributeDescriptor.GetAttributeAtIndex(IndexB) == AttributeDescriptor.FindAttribute("IntAttributeB"), TEXT("Index and Attribute should return the same value."));
+		AddErrorIfFalse(&AttributeDescriptor.GetAttributeAtIndex(IndexC) == AttributeDescriptor.FindAttribute("IntAttributeC"), TEXT("Index and Attribute should return the same value."));
+		AddErrorIfFalse(&AttributeDescriptor.GetAttributeAtIndex(IndexD) == AttributeDescriptor.FindAttribute("IntAttributeD"), TEXT("Index and Attribute should return the same value."));
+		AddErrorIfFalse(AttributeDescriptor.FindAttribute("IntAttributeI") == nullptr, TEXT("Was not supposed to find the Attribute I"));
+		AddErrorIfFalse(AttributeDescriptor.FindAttribute("IntAttributeJ") == nullptr, TEXT("Was not supposed to find the Attribute J"));
+		AddErrorIfFalse(AttributeDescriptor.FindAttribute("IntAttributeK") == nullptr, TEXT("Was not supposed to find the Attribute K"));
+
+		//B, A, D, C
+		AddErrorIfFalse(AttributeDescriptor.GetAttributeAtIndex(IndexB).SortOrder < AttributeDescriptor.GetAttributeAtIndex(IndexA).SortOrder, TEXT("B should have a lower value than A"));
+		AddErrorIfFalse(AttributeDescriptor.GetAttributeAtIndex(IndexD).SortOrder < AttributeDescriptor.GetAttributeAtIndex(IndexC).SortOrder, TEXT("D should have a lower value than C"));
+		AddErrorIfFalse(AttributeDescriptor.GetAttributeAtIndex(IndexA).SortOrder < AttributeDescriptor.GetAttributeAtIndex(IndexD).SortOrder, TEXT("A should have a lower value than D"));
+
+		{
+			OrderCounter = 0;
+			WidgetParent->IntAttributeA.Assign(WidgetParent.Get(), MakeAttributeLambda(OrderLambda));
+			AddErrorIfFalse(WidgetParent->IntAttributeA.Get() == 1, TEXT("A It is not the expected value."));
+			WidgetParent->IntAttributeB.Assign(WidgetParent.Get(), MakeAttributeLambda(OrderLambda));
+			AddErrorIfFalse(WidgetParent->IntAttributeB.Get() == 2, TEXT("B It is not the expected value."));
+			WidgetParent->IntAttributeC.Assign(WidgetParent.Get(), MakeAttributeLambda(OrderLambda));
+			AddErrorIfFalse(WidgetParent->IntAttributeC.Get() == 3, TEXT("C It is not the expected value."));
+			WidgetParent->IntAttributeD.Assign(WidgetParent.Get(), MakeAttributeLambda(OrderLambda));
+			AddErrorIfFalse(WidgetParent->IntAttributeD.Get() == 4, TEXT("D It is not the expected value."));
+
+			OrderCounter = 0;
+			bWasUpdate = false;
+			ReturnValue = 4;
+			WidgetParent->InvalidatePrepass();
+			WidgetParent->SlatePrepass(1.f);
+			AddErrorIfFalse(WidgetParent->IntAttributeA.Get() == 2, TEXT("A It is not the expected value."));
+			AddErrorIfFalse(WidgetParent->IntAttributeB.Get() == 1, TEXT("B It is not the expected value."));
+			AddErrorIfFalse(WidgetParent->IntAttributeC.Get() == 4, TEXT("C It is not the expected value."));
+			AddErrorIfFalse(WidgetParent->IntAttributeD.Get() == 3, TEXT("D It is not the expected value."));
+		}
+
+		{
+			OrderCounter = 0;
+			bWasUpdate = false;
+			ReturnValue = 5;
+			WidgetParent->IntAttributeC.Assign(WidgetParent.Get(), MakeAttributeLambda(UpdateLambda));
+			AddErrorIfFalse(bWasUpdate, TEXT("C should be updated."));
+			AddErrorIfFalse(WidgetParent->IntAttributeC.Get() == 5, TEXT("C It is not the expected value."));
+		}
+
+		{
+			OrderCounter = 0;
+			bWasUpdate = false;
+			ReturnValue = 10;	//10 show that C didn't change
+			WidgetParent->InvalidatePrepass();
+			WidgetParent->SlatePrepass(1.f);
+			AddErrorIfFalse(WidgetParent->IntAttributeA.Get() == 2, TEXT("A It is not the expected value."));
+			AddErrorIfFalse(WidgetParent->IntAttributeB.Get() == 1, TEXT("B It is not the expected value."));
+			AddErrorIfFalse(!bWasUpdate, TEXT("C should not be updated."));
+			AddErrorIfFalse(WidgetParent->IntAttributeC.Get() == 5, TEXT("C It is not the expected value."));
+			AddErrorIfFalse(WidgetParent->IntAttributeD.Get() == 3, TEXT("D It is not the expected value."));
+		}
+
+		{
+			WidgetParent->IntAttributeD.Set(WidgetParent.Get(), 8);
+			AddErrorIfFalse(WidgetParent->IntAttributeC.Get() == 5, TEXT("C It is not the expected value."));
+			AddErrorIfFalse(WidgetParent->IntAttributeD.Get() == 8, TEXT("D It is not the expected value."));
+
+			OrderCounter = 0;
+			bWasUpdate = false;
+			ReturnValue = 10;
+			WidgetParent->InvalidatePrepass();
+			WidgetParent->SlatePrepass(1.f);
+			AddErrorIfFalse(WidgetParent->IntAttributeA.Get() == 2, TEXT("A It is not the expected value."));
+			AddErrorIfFalse(WidgetParent->IntAttributeB.Get() == 1, TEXT("B It is not the expected value."));
+			AddErrorIfFalse(bWasUpdate, TEXT("C should be updated becaude D was."));
+			AddErrorIfFalse(WidgetParent->IntAttributeC.Get() == 10, TEXT("C It is not the expected value."));
+			AddErrorIfFalse(WidgetParent->IntAttributeD.Get() == 8, TEXT("D It is not the expected value."));
+			AddErrorIfFalse(OrderCounter == 2, TEXT("There is no D attribute anymore."));
+			AddErrorIfFalse(OrderCounter == 2, TEXT("There is no D attribute anymore."));
+		}
+
+		{
+			OrderCounter = 0;
+			bWasUpdate = false;
+			ReturnValue = 10;
+			WidgetParent->InvalidatePrepass();
+			WidgetParent->SlatePrepass(1.f);
+			AddErrorIfFalse(!bWasUpdate, TEXT("C should not be updated."));
+		}
+	}
+
+	{
+		TSharedRef<UE::Slate::Private::SAttributeLeftWidget_Child> WidgetChild = SNew(UE::Slate::Private::SAttributeLeftWidget_Child);
+
+		AddErrorIfFalse(&WidgetChild->GetWidgetClass() == &UE::Slate::Private::SAttributeLeftWidget_Child::StaticWidgetClass()
+			, TEXT("The static data do not matches"));
+
+		FSlateAttributeDescriptor const& AttributeDescriptor = WidgetChild->GetWidgetClass().GetAttributeDescriptor();
+		AddErrorIfFalse(AttributeDescriptor.AttributeNum() == 8, TEXT("")); // H is not counted
+
+		const int32 IndexA = AttributeDescriptor.IndexOfMemberAttribute("IntAttributeA");
+		const int32 IndexB = AttributeDescriptor.IndexOfMemberAttribute("IntAttributeB");
+		const int32 IndexC = AttributeDescriptor.IndexOfMemberAttribute("IntAttributeC");
+		const int32 IndexD = AttributeDescriptor.IndexOfMemberAttribute("IntAttributeD");
+		const int32 IndexI = AttributeDescriptor.IndexOfMemberAttribute("IntAttributeI");
+		const int32 IndexJ = AttributeDescriptor.IndexOfMemberAttribute("IntAttributeJ");
+		const int32 IndexK = AttributeDescriptor.IndexOfMemberAttribute("IntAttributeK");
+		const int32 IndexL = AttributeDescriptor.IndexOfMemberAttribute("IntAttributeL");
+
+		AddErrorIfFalse(IndexA != INDEX_NONE, TEXT("Could not find the Attribute A"));
+		AddErrorIfFalse(IndexB != INDEX_NONE, TEXT("Could not find the Attribute B"));
+		AddErrorIfFalse(IndexC != INDEX_NONE, TEXT("Could not find the Attribute C"));
+		AddErrorIfFalse(IndexD != INDEX_NONE, TEXT("Could not find the Attribute D"));
+		AddErrorIfFalse(IndexI != INDEX_NONE, TEXT("Could not find the Attribute I"));
+		AddErrorIfFalse(IndexJ != INDEX_NONE, TEXT("Could not find the Attribute J"));
+		AddErrorIfFalse(IndexK != INDEX_NONE, TEXT("Could not find the Attribute K"));
+		AddErrorIfFalse(IndexL != INDEX_NONE, TEXT("Could not find the Attribute L"));
+
+		AddErrorIfFalse(&AttributeDescriptor.GetAttributeAtIndex(IndexA) == AttributeDescriptor.FindAttribute("IntAttributeA"), TEXT("Index and Attribute should return the same value."));
+		AddErrorIfFalse(&AttributeDescriptor.GetAttributeAtIndex(IndexB) == AttributeDescriptor.FindAttribute("IntAttributeB"), TEXT("Index and Attribute should return the same value."));
+		AddErrorIfFalse(&AttributeDescriptor.GetAttributeAtIndex(IndexC) == AttributeDescriptor.FindAttribute("IntAttributeC"), TEXT("Index and Attribute should return the same value."));
+		AddErrorIfFalse(&AttributeDescriptor.GetAttributeAtIndex(IndexD) == AttributeDescriptor.FindAttribute("IntAttributeD"), TEXT("Index and Attribute should return the same value."));
+		AddErrorIfFalse(&AttributeDescriptor.GetAttributeAtIndex(IndexI) == AttributeDescriptor.FindAttribute("IntAttributeI"), TEXT("Index and Attribute should return the same value."));
+		AddErrorIfFalse(&AttributeDescriptor.GetAttributeAtIndex(IndexJ) == AttributeDescriptor.FindAttribute("IntAttributeJ"), TEXT("Index and Attribute should return the same value."));
+		AddErrorIfFalse(&AttributeDescriptor.GetAttributeAtIndex(IndexK) == AttributeDescriptor.FindAttribute("IntAttributeK"), TEXT("Index and Attribute should return the same value."));
+		AddErrorIfFalse(&AttributeDescriptor.GetAttributeAtIndex(IndexL) == AttributeDescriptor.FindAttribute("IntAttributeL"), TEXT("Index and Attribute should return the same value."));
+		AddErrorIfFalse(AttributeDescriptor.FindAttribute("IntAttributeH") == nullptr, TEXT("H exist but is not defined."));
+
+		AddErrorIfFalse(AttributeDescriptor.GetAttributeAtIndex(IndexB).SortOrder < AttributeDescriptor.GetAttributeAtIndex(IndexA).SortOrder, TEXT("B should have a lower value than A"));
+		AddErrorIfFalse(AttributeDescriptor.GetAttributeAtIndex(IndexA).SortOrder <= AttributeDescriptor.GetAttributeAtIndex(IndexI).SortOrder, TEXT("A should have a lower value than I"));
+		AddErrorIfFalse(AttributeDescriptor.GetAttributeAtIndex(IndexI).SortOrder < AttributeDescriptor.GetAttributeAtIndex(IndexJ).SortOrder, TEXT("I should have a lower value than J"));
+		AddErrorIfFalse(AttributeDescriptor.GetAttributeAtIndex(IndexJ).SortOrder < AttributeDescriptor.GetAttributeAtIndex(IndexD).SortOrder, TEXT("J should have a lower value than D"));
+		AddErrorIfFalse(AttributeDescriptor.GetAttributeAtIndex(IndexD).SortOrder < AttributeDescriptor.GetAttributeAtIndex(IndexC).SortOrder, TEXT("D should have a lower value than C"));
+		AddErrorIfFalse(AttributeDescriptor.GetAttributeAtIndex(IndexC).SortOrder < AttributeDescriptor.GetAttributeAtIndex(IndexL).SortOrder, TEXT("C should have a lower value than L"));
+		AddErrorIfFalse(AttributeDescriptor.GetAttributeAtIndex(IndexL).SortOrder < AttributeDescriptor.GetAttributeAtIndex(IndexK).SortOrder, TEXT("L should have a lower value than K"));
+
+
+		{
+			OrderCounter = 49;
+			WidgetChild->IntAttributeA.Assign(WidgetChild.Get(), MakeAttributeLambda(OrderLambda));
+			AddErrorIfFalse(WidgetChild->IntAttributeA.Get() == 50, TEXT("A It is not the expected value."));
+			WidgetChild->IntAttributeB.Assign(WidgetChild.Get(), MakeAttributeLambda(OrderLambda));
+			AddErrorIfFalse(WidgetChild->IntAttributeB.Get() == 51, TEXT("B It is not the expected value."));
+			WidgetChild->IntAttributeC.Assign(WidgetChild.Get(), MakeAttributeLambda(OrderLambda));
+			AddErrorIfFalse(WidgetChild->IntAttributeC.Get() == 52, TEXT("C It is not the expected value."));
+			WidgetChild->IntAttributeD.Assign(WidgetChild.Get(), MakeAttributeLambda(OrderLambda));
+			AddErrorIfFalse(WidgetChild->IntAttributeD.Get() == 53, TEXT("D It is not the expected value."));
+			WidgetChild->IntAttributeH.Assign(WidgetChild.Get(), MakeAttributeLambda(OrderLambda));
+			AddErrorIfFalse(WidgetChild->IntAttributeH.Get() == 54, TEXT("I It is not the expected value."));
+			WidgetChild->IntAttributeI.Assign(WidgetChild.Get(), MakeAttributeLambda(OrderLambda));
+			AddErrorIfFalse(WidgetChild->IntAttributeI.Get() == 55, TEXT("I It is not the expected value."));
+			WidgetChild->IntAttributeJ.Assign(WidgetChild.Get(), MakeAttributeLambda(OrderLambda));
+			AddErrorIfFalse(WidgetChild->IntAttributeJ.Get() == 56, TEXT("J It is not the expected value."));
+			WidgetChild->IntAttributeK.Assign(WidgetChild.Get(), MakeAttributeLambda(OrderLambda));
+			AddErrorIfFalse(WidgetChild->IntAttributeK.Get() == 57, TEXT("K It is not the expected value."));
+			WidgetChild->IntAttributeL.Assign(WidgetChild.Get(), MakeAttributeLambda(OrderLambda));
+			AddErrorIfFalse(WidgetChild->IntAttributeL.Get() == 58, TEXT("L It is not the expected value."));
+
+
+			OrderCounter = 0;
+			bWasUpdate = false;
+			ReturnValue = 4;
+			WidgetChild->InvalidatePrepass();
+			WidgetChild->SlatePrepass(1.f);
+			AddErrorIfFalse(WidgetChild->IntAttributeA.Get() == 2 || WidgetChild->IntAttributeA.Get() == 3, TEXT("A It is not the expected value."));
+			AddErrorIfFalse(WidgetChild->IntAttributeB.Get() == 1, TEXT("B It is not the expected value."));
+			AddErrorIfFalse(WidgetChild->IntAttributeC.Get() == 6, TEXT("C It is not the expected value."));
+			AddErrorIfFalse(WidgetChild->IntAttributeD.Get() == 5, TEXT("D It is not the expected value."));
+			AddErrorIfFalse(WidgetChild->IntAttributeH.Get() == 8, TEXT("H It is not the expected value."));
+			AddErrorIfFalse(WidgetChild->IntAttributeI.Get() == 2 || WidgetChild->IntAttributeI.Get() == 3, TEXT("I It is not the expected value."));
+			AddErrorIfFalse(WidgetChild->IntAttributeJ.Get() == 4, TEXT("J It is not the expected value."));
+			AddErrorIfFalse(WidgetChild->IntAttributeK.Get() == 9, TEXT("K It is not the expected value."));
+			AddErrorIfFalse(WidgetChild->IntAttributeL.Get() == 7, TEXT("L It is not the expected value."));
+		}
+
+		{
+			OrderCounter = 0;
+			bWasUpdate = false;
+			ReturnValue = 4;
+			WidgetChild->InvalidatePrepass();
+			WidgetChild->SlatePrepass(1.f);
+			AddErrorIfFalse(WidgetChild->IntAttributeA.Get() == 2 || WidgetChild->IntAttributeA.Get() == 3, TEXT("A It is not the expected value."));
+			AddErrorIfFalse(WidgetChild->IntAttributeB.Get() == 1, TEXT("B It is not the expected value."));
+			AddErrorIfFalse(WidgetChild->IntAttributeC.Get() == 5, TEXT("C It is not the expected value.")); // will get updated because D changes
+			AddErrorIfFalse(WidgetChild->IntAttributeD.Get() == 4, TEXT("D It is not the expected value."));
+			AddErrorIfFalse(WidgetChild->IntAttributeH.Get() == 7, TEXT("H It is not the expected value."));
+			AddErrorIfFalse(WidgetChild->IntAttributeI.Get() == 2 || WidgetChild->IntAttributeI.Get() == 3, TEXT("I It is not the expected value."));
+			AddErrorIfFalse(WidgetChild->IntAttributeJ.Get() == 4, TEXT("J It is not the expected value.")); // should not get updated
+			AddErrorIfFalse(WidgetChild->IntAttributeK.Get() == 8, TEXT("K It is not the expected value."));
+			AddErrorIfFalse(WidgetChild->IntAttributeL.Get() == 6, TEXT("L It is not the expected value."));
+		}
+	}
+
+	return true;
+}
+
+#undef LOCTEXT_NAMESPACE 
+
+#endif //WITH_AUTOMATION_WORKER
diff --git a/Engine/Source/Runtime/SlateCore/Private/Types/SlateAttribute.cpp b/Engine/Source/Runtime/SlateCore/Private/Types/SlateAttribute.cpp
--- a/Engine/Source/Runtime/SlateCore/Private/Types/SlateAttribute.cpp
+++ b/Engine/Source/Runtime/SlateCore/Private/Types/SlateAttribute.cpp
@@ -24,6 +24,7 @@ namespace SlateAttributePrivate
 #endif
 	}
 
+
 	bool FSlateAttributeImpl::ProtectedIsWidgetInDestructionPath(SWidget* Widget) const
 	{
 #if UE_WITH_SLATE_DEBUG_WIDGETLIST
@@ -35,9 +36,17 @@ namespace SlateAttributePrivate
 		// if the shared instance exist, then the widget is not being destroyed
 		return !UE::Slate::FWidgetList::GetAllWidgets()[FoundIndex]->DoesSharedInstanceExist();
 #endif
-	return true;
+		return true;
+	}
+
+
+	bool FSlateAttributeImpl::ProtectedIsImplemented(const SWidget& OwningWidget) const
+	{
+		const FSlateAttributeDescriptor::OffsetType Offset = (FSlateAttributeDescriptor::OffsetType)((UPTRINT)(this) - (UPTRINT)(&OwningWidget));
+		return OwningWidget.GetWidgetClass().GetAttributeDescriptor().FindMemberAttribute(Offset) != nullptr;
 	}
 
+
 	void FSlateAttributeImpl::ProtectedUnregisterAttribute(SWidget& OwningWidget, ESlateAttributeType AttributeType) const
 	{
 		TestAttributeAddress(OwningWidget, *this, AttributeType);
@@ -58,7 +67,7 @@ namespace SlateAttributePrivate
 	{
 		TestAttributeAddress(OwningWidget, *this, AttributeType);
 
-		FSlateAttributeMetaData::InvalidateWidget(OwningWidget, *this, InvalidationReason);
+		FSlateAttributeMetaData::InvalidateWidget(OwningWidget, *this, AttributeType, InvalidationReason);
 	}
 
 
diff --git a/Engine/Source/Runtime/SlateCore/Private/Types/SlateAttributeDescriptor.cpp b/Engine/Source/Runtime/SlateCore/Private/Types/SlateAttributeDescriptor.cpp
--- a/Engine/Source/Runtime/SlateCore/Private/Types/SlateAttributeDescriptor.cpp
+++ b/Engine/Source/Runtime/SlateCore/Private/Types/SlateAttributeDescriptor.cpp
@@ -9,30 +9,33 @@ FSlateAttributeDescriptor::FInitializer::FAttributeEntry::FAttributeEntry(FSlate
 	: Descriptor(InDescriptor)
 	, AttributeIndex(InAttributeIndex)
 {}
-			
-FSlateAttributeDescriptor::FInitializer::FAttributeEntry& FSlateAttributeDescriptor::FInitializer::FAttributeEntry::SetPrerequisite(FName Prerequisite)
+
+	
+FSlateAttributeDescriptor::FInitializer::FAttributeEntry& FSlateAttributeDescriptor::FInitializer::FAttributeEntry::UpdatePrerequisite(FName Prerequisite)
 {
 	if (Descriptor.Attributes.IsValidIndex(AttributeIndex))
 	{
-		Descriptor.SetPrerequisite(Descriptor.Attributes[AttributeIndex], Prerequisite);
+		Descriptor.SetPrerequisite(Descriptor.Attributes[AttributeIndex], Prerequisite, false);
 	}
 	return *this;
 }
 
-FSlateAttributeDescriptor::FInitializer::FAttributeEntry& FSlateAttributeDescriptor::FInitializer::FAttributeEntry::UpdateEveryFrame()
+
+FSlateAttributeDescriptor::FInitializer::FAttributeEntry& FSlateAttributeDescriptor::FInitializer::FAttributeEntry::UpdateDependency(FName Dependency)
 {
 	if (Descriptor.Attributes.IsValidIndex(AttributeIndex))
 	{
-		Descriptor.Attributes[AttributeIndex].Dependency = FName();
+		Descriptor.SetPrerequisite(Descriptor.Attributes[AttributeIndex], Dependency, true);
 	}
 	return *this;
 }
 
-FSlateAttributeDescriptor::FInitializer::FAttributeEntry& FSlateAttributeDescriptor::FInitializer::FAttributeEntry::UpdateDependency(FName Dependency)
+
+FSlateAttributeDescriptor::FInitializer::FAttributeEntry& FSlateAttributeDescriptor::FInitializer::FAttributeEntry::UpdateWhenCollapsed()
 {
 	if (Descriptor.Attributes.IsValidIndex(AttributeIndex))
 	{
-		Descriptor.SetDependency(Descriptor.Attributes[AttributeIndex], Dependency);
+		Descriptor.SetUpdateWhenCollapsed(Descriptor.Attributes[AttributeIndex], true);
 	}
 	return *this;
 }
@@ -49,64 +52,181 @@ FSlateAttributeDescriptor::FInitializer::FInitializer(FSlateAttributeDescriptor&
 	InDescriptor.Attributes = ParentDescriptor.Attributes;
 }
 
+
 FSlateAttributeDescriptor::FInitializer::~FInitializer()
 {
-	// because adding the attribute is not required at the moment
+	// Update the sort order for the item that have prerequisite.
+	//Because adding the attribute is not required at the moment
 	//try to not change the order in which they were added
-	//for (int32 Index = 0; Index < Descriptor.Attributes.Num(); ++Index)
-	//{
-	//	if (!Descriptor.Attributes[Index].Prerequisite.IsNone())
-	//	{
-	//		// Find the Prerequisite index
-	//		const FName Prerequisite = Descriptor.Attributes[Index].Prerequisite;
-	//		const int32 FoundIndex = Descriptor.Attributes.IndexOfByPredicate([Prerequisite](const FAttribute& Other) { return Other.Name == Prerequisite; });
-	//		if (ensureAlwaysMsgf(Descriptor.Attributes.IsValidIndex(FoundIndex), TEXT("The Prerequisite '%s' doesn't exist"), *Prerequisite.ToString()))
-	//		{
-	//			Descriptor.Attributes[Index].SortOrder = Descriptor.Attributes[FoundIndex].SortOrder + 1;
-	//		}
-	//	}
-	//}
+
+	struct FPrerequisiteSort
+	{
+		FPrerequisiteSort() = default;
+		FPrerequisiteSort(int32 A, int32 B, int32 InDepth) : AttributeIndex(A), PrerequisitesIndex(B), Depth(InDepth) {}
+		int32 AttributeIndex;
+		int32 PrerequisitesIndex;
+		int32 Depth = -1;
+
+		void CalculateDepth(TArray<FPrerequisiteSort, TMemStackAllocator<>>& Prerequisites)
+		{
+			if (Depth < 0)
+			{
+				check(PrerequisitesIndex != INDEX_NONE);
+				if (Prerequisites[PrerequisitesIndex].Depth < 0)
+				{
+					// calculate the Depth recursively
+					Prerequisites[PrerequisitesIndex].CalculateDepth(Prerequisites);
+				}
+				Depth = Prerequisites[PrerequisitesIndex].Depth + 1;
+			}
+		}
+	};
+	struct FPrerequisiteSortPredicate
+	{
+		const TArray<FAttribute>& Attrbutes;
+		bool operator()(const FPrerequisiteSort& A, const FPrerequisiteSort& B) const
+		{
+			if (A.Depth != B.Depth)
+			{
+				return A.Depth < B.Depth;
+			}
+			if (A.PrerequisitesIndex == B.PrerequisitesIndex)
+			{
+				return Attrbutes[A.AttributeIndex].SortOrder < Attrbutes[B.AttributeIndex].SortOrder;
+			}
+
+			const int32 SortA = A.PrerequisitesIndex != INDEX_NONE ? Attrbutes[A.PrerequisitesIndex].SortOrder : Attrbutes[A.AttributeIndex].SortOrder;
+			const int32 SortB = B.PrerequisitesIndex != INDEX_NONE ? Attrbutes[B.PrerequisitesIndex].SortOrder : Attrbutes[B.AttributeIndex].SortOrder;
+			return SortA < SortB;
+		}
+	};
+
+	FMemMark Mark(FMemStack::Get());
+	TArray<FPrerequisiteSort, TMemStackAllocator<>> Prerequisites;
+	Prerequisites.Reserve(Descriptor.Attributes.Num());
+
+	bool bHavePrerequisite = false;
+	for (int32 Index = 0; Index < Descriptor.Attributes.Num(); ++Index)
+	{
+		FAttribute& Attribute = Descriptor.Attributes[Index];
+		Attribute.SortOrder = DefaultSortOrder(Attribute.Offset);
+
+		if (!Attribute.Prerequisite.IsNone())
+		{
+			// Find the Prerequisite index
+			const FName Prerequisite = Attribute.Prerequisite;
+			const int32 PrerequisiteIndex = Descriptor.Attributes.IndexOfByPredicate([Prerequisite](const FAttribute& Other) { return Other.Name == Prerequisite; });
+			if (ensureAlwaysMsgf(Descriptor.Attributes.IsValidIndex(PrerequisiteIndex), TEXT("The Prerequisite '%s' doesn't exist"), *Prerequisite.ToString()))
+			{
+				Prerequisites.Emplace(Index, PrerequisiteIndex, -1);
+				Descriptor.Attributes[PrerequisiteIndex].bIsADependencyForSomeoneElse = true;
+				bHavePrerequisite = true;
+			}
+			else
+			{
+				Prerequisites.Emplace(Index, INDEX_NONE, 0);
+			}
+		}
+		else
+		{
+			Prerequisites.Emplace(Index, INDEX_NONE, 0);
+		}
+	}
+
+	if (bHavePrerequisite)
+	{
+		// Get the depth order
+		for (FPrerequisiteSort& PrerequisiteSort : Prerequisites)
+		{
+			PrerequisiteSort.CalculateDepth(Prerequisites);
+		}
+
+		Prerequisites.Sort(FPrerequisiteSortPredicate{ Descriptor.Attributes });
+		int32 PreviousPrerequisiteIndex = INDEX_NONE;
+		int32 IncreaseCount = 1;
+		for (const FPrerequisiteSort& Element : Prerequisites)
+		{
+			if (Element.PrerequisitesIndex != INDEX_NONE)
+			{
+				if (PreviousPrerequisiteIndex == Element.PrerequisitesIndex)
+				{
+					++IncreaseCount;
+				}
+				Descriptor.Attributes[Element.AttributeIndex].SortOrder = Descriptor.Attributes[Element.PrerequisitesIndex].SortOrder + IncreaseCount;
+			}
+			PreviousPrerequisiteIndex = Element.PrerequisitesIndex;
+		}
+	}
 }
 
 
-FSlateAttributeDescriptor::FInitializer::FAttributeEntry FSlateAttributeDescriptor::FInitializer::AddMemberAttribute(FName AttributeName, OffsetType Offset)
+FSlateAttributeDescriptor::FInitializer::FAttributeEntry FSlateAttributeDescriptor::FInitializer::AddMemberAttribute(FName AttributeName, OffsetType Offset, const FInvalidateWidgetReasonAttribute& Reason)
 {
-	return Descriptor.AddMemberAttribute(AttributeName, Offset, TAttribute<EInvalidateWidgetReason>());
+	return Descriptor.AddMemberAttribute(AttributeName, Offset, Reason);
 }
 
 
-FSlateAttributeDescriptor::FInitializer::FAttributeEntry FSlateAttributeDescriptor::FInitializer::AddMemberAttribute(FName AttributeName, OffsetType Offset, TAttribute<EInvalidateWidgetReason> Reason)
+FSlateAttributeDescriptor::FInitializer::FAttributeEntry FSlateAttributeDescriptor::FInitializer::AddMemberAttribute(FName AttributeName, OffsetType Offset, FInvalidateWidgetReasonAttribute&& Reason)
 {
-	check(Reason.IsSet());
 	return Descriptor.AddMemberAttribute(AttributeName, Offset, MoveTemp(Reason));
 }
 
 
-void FSlateAttributeDescriptor::FInitializer::OverrideInvalidationReason(FName AttributeName, TAttribute<EInvalidateWidgetReason> Reason)
+void FSlateAttributeDescriptor::FInitializer::OverrideInvalidationReason(FName AttributeName, const FInvalidateWidgetReasonAttribute& Reason)
 {
-	check(Reason.IsSet());
-	Descriptor.OverrideInvalidationReason(AttributeName, MoveTemp(Reason));
+	Descriptor.OverrideInvalidationReason(AttributeName, Reason);
 }
 
 
-void FSlateAttributeDescriptor::FInitializer::SetPrerequisite(FName AttributeName, FName Prerequisite)
+void FSlateAttributeDescriptor::FInitializer::SetUpdateWhenCollapsed(FName AttributeName, bool bUpdateWhenCollapsed)
 {
-	check(!AttributeName.IsNone());
-
-	FAttribute* FoundAttribute = Descriptor.FindAttribute(AttributeName);
-	if (ensureAlwaysMsgf(FoundAttribute != nullptr, TEXT("The attribute '%s' doesn't exist."), *AttributeName.ToString()))
+	FAttribute* Attribute = Descriptor.FindAttribute(AttributeName);
+	if (ensureAlwaysMsgf(Attribute, TEXT("The attribute named '%s' doesn't exist"), *AttributeName.ToString()))
 	{
-		Descriptor.SetPrerequisite(*FoundAttribute, Prerequisite);
+		Descriptor.SetUpdateWhenCollapsed(*Attribute, bUpdateWhenCollapsed);
 	}
 }
 
 
+void FSlateAttributeDescriptor::FInitializer::OverrideInvalidationReason(FName AttributeName, FInvalidateWidgetReasonAttribute&& Reason)
+{
+	Descriptor.OverrideInvalidationReason(AttributeName, MoveTemp(Reason));
+}
+
+
+FSlateAttributeDescriptor::FAttribute const& FSlateAttributeDescriptor::GetAttributeAtIndex(int32 Index) const
+{
+	check(Attributes.IsValidIndex(Index));
+	FSlateAttributeDescriptor::FAttribute const& Result = Attributes[Index];
+	return Result;
+}
+
+
 FSlateAttributeDescriptor::FAttribute const* FSlateAttributeDescriptor::FindAttribute(FName AttributeName) const
 {
 	return Attributes.FindByPredicate([AttributeName](const FAttribute& Other) { return Other.Name == AttributeName; });
 }
 
 
+int32 FSlateAttributeDescriptor::IndexOfMemberAttribute(OffsetType AttributeOffset) const
+{
+	int32 FoundIndex = Attributes.IndexOfByPredicate([AttributeOffset](const FAttribute& Other) { return Other.Offset == AttributeOffset; });
+	check(FoundIndex == INDEX_NONE || Attributes[FoundIndex].bIsMemberAttribute);
+	return FoundIndex;
+}
+
+
+int32 FSlateAttributeDescriptor::IndexOfMemberAttribute(FName AttributeName) const
+{
+	int32 FoundIndex = Attributes.IndexOfByPredicate([AttributeName](const FAttribute& Other) { return Other.Name == AttributeName; });
+	if (ensure(FoundIndex == INDEX_NONE || Attributes[FoundIndex].bIsMemberAttribute))
+	{
+		return FoundIndex;
+	}
+	return INDEX_NONE;
+}
+
+
 FSlateAttributeDescriptor::FAttribute const* FSlateAttributeDescriptor::FindMemberAttribute(OffsetType AttributeOffset) const
 {
 	FSlateAttributeDescriptor::FAttribute const* Result = Attributes.FindByPredicate([AttributeOffset](const FAttribute& Other) { return Other.Offset == AttributeOffset; });
@@ -121,27 +241,26 @@ FSlateAttributeDescriptor::FAttribute* FSlateAttributeDescriptor::FindAttribute(
 }
 
 
-FSlateAttributeDescriptor::FInitializer::FAttributeEntry FSlateAttributeDescriptor::AddMemberAttribute(FName AttributeName, OffsetType Offset, TAttribute<EInvalidateWidgetReason>&& Reason)
+FSlateAttributeDescriptor::FInitializer::FAttributeEntry FSlateAttributeDescriptor::AddMemberAttribute(FName AttributeName, OffsetType Offset, FInvalidateWidgetReasonAttribute Reason)
 {
 	check(!AttributeName.IsNone());
 
 	int32 NewIndex = INDEX_NONE;
 	FAttribute const* FoundAttribute = FindAttribute(AttributeName);
-	if (ensureAlwaysMsgf(FoundAttribute == nullptr, TEXT("The attribute '%s' already exist."), *AttributeName.ToString()))
+	if (ensureAlwaysMsgf(FoundAttribute == nullptr, TEXT("The attribute '%s' already exist. (Do you have the correct parrent class in SLATE_DECLARE_WIDGET)"), *AttributeName.ToString()))
 	{
 		NewIndex = Attributes.AddZeroed();
 		FAttribute& NewAttribute = Attributes[NewIndex];
 		NewAttribute.Name = AttributeName;
 		NewAttribute.Offset = Offset;
 		NewAttribute.InvalidationReason = MoveTemp(Reason);
-		NewAttribute.SortOrder = DefaultSortOrder(Offset);
 		NewAttribute.bIsMemberAttribute = true;
 	}
 	return FInitializer::FAttributeEntry(*this, NewIndex);
 }
 
 
-void FSlateAttributeDescriptor::OverrideInvalidationReason(FName AttributeName, TAttribute<EInvalidateWidgetReason>&& Reason)
+void FSlateAttributeDescriptor::OverrideInvalidationReason(FName AttributeName, FInvalidateWidgetReasonAttribute Reason)
 {
 	check(!AttributeName.IsNone());
 
@@ -152,51 +271,13 @@ void FSlateAttributeDescriptor::OverrideInvalidationReason(FName AttributeName,
 	}
 }
 
-void FSlateAttributeDescriptor::SetDependency(FSlateAttributeDescriptor::FAttribute& Attribute, FName Dependency)
-{
-	if (Dependency.IsNone())
-	{
-		Attribute.Dependency = FName();
-	}
-	else
-	{
-		FAttribute const* FoundDependency = FindAttribute(Dependency);
-		if (ensureAlwaysMsgf(FoundDependency, TEXT("The Dependency '%s' doesn't exist for attribute '%s'"), *Dependency.ToString(), *Attribute.Name.ToString()))
-		{
-			Attribute.Dependency = Dependency;
 
-			// Verify recursion
-			{
-				FMemMark Mark(FMemStack::Get());
-				TArray<FName, TMemStackAllocator<>> Recursion;
-				Recursion.Reserve(Attributes.Num());
-				FAttribute const* CurrentAttribute = &Attribute;
-				while (!CurrentAttribute->Dependency.IsNone())
-				{
-					if (Recursion.Contains(CurrentAttribute->Name))
-					{
-						ensureAlwaysMsgf(false, TEXT("The Dependency '%s' would introduce an infinit loop with attribute '%s'."), *Dependency.ToString(), *Attribute.Name.ToString());
-						Attribute.Dependency = FName();
-						break;
-					}
-					Recursion.Add(CurrentAttribute->Name);
-					CurrentAttribute = FindAttribute(CurrentAttribute->Dependency);
-					check(CurrentAttribute);
-				}
-			}
-		}
-		else
-		{
-			Attribute.Dependency = FName();
-		}
-	}
-}
-
-void FSlateAttributeDescriptor::SetPrerequisite(FSlateAttributeDescriptor::FAttribute& Attribute, FName Prerequisite)
+void FSlateAttributeDescriptor::SetPrerequisite(FSlateAttributeDescriptor::FAttribute& Attribute, FName Prerequisite, bool bSetAsDependency)
 {
 	if (Prerequisite.IsNone())
 	{
 		Attribute.Prerequisite = FName();
+		Attribute.bIsPrerequisiteAlsoADependency = false;
 	}
 	else
 	{
@@ -204,7 +285,7 @@ void FSlateAttributeDescriptor::SetPrerequisite(FSlateAttributeDescriptor::FAttr
 		if (ensureAlwaysMsgf(FoundPrerequisite, TEXT("The prerequisite '%s' doesn't exist for attribute '%s'"), *Prerequisite.ToString(), *Attribute.Name.ToString()))
 		{
 			Attribute.Prerequisite = Prerequisite;
-			Attribute.SortOrder = FoundPrerequisite->SortOrder + 1;
+			Attribute.bIsPrerequisiteAlsoADependency = bSetAsDependency;
 
 			// Verify recursion
 			{
@@ -218,6 +299,7 @@ void FSlateAttributeDescriptor::SetPrerequisite(FSlateAttributeDescriptor::FAttr
 					{
 						ensureAlwaysMsgf(false, TEXT("The prerequsite '%s' would introduce an infinit loop with attribute '%s'."), *Prerequisite.ToString(), *Attribute.Name.ToString());
 						Attribute.Prerequisite = FName();
+						Attribute.bIsPrerequisiteAlsoADependency = false;
 						break;
 					}
 					Recursion.Add(CurrentAttribute->Name);
@@ -229,6 +311,13 @@ void FSlateAttributeDescriptor::SetPrerequisite(FSlateAttributeDescriptor::FAttr
 		else
 		{
 			Attribute.Prerequisite = FName();
+			Attribute.bIsPrerequisiteAlsoADependency = false;
 		}
 	}
 }
+
+
+void FSlateAttributeDescriptor::SetUpdateWhenCollapsed(FAttribute& Attribute, bool bUpdate)
+{
+	Attribute.bUpdateWhenCollapsed = bUpdate;
+}
\ No newline at end of file
diff --git a/Engine/Source/Runtime/SlateCore/Private/Types/SlateAttributeMetaData.cpp b/Engine/Source/Runtime/SlateCore/Private/Types/SlateAttributeMetaData.cpp
--- a/Engine/Source/Runtime/SlateCore/Private/Types/SlateAttributeMetaData.cpp
+++ b/Engine/Source/Runtime/SlateCore/Private/Types/SlateAttributeMetaData.cpp
@@ -2,6 +2,7 @@
 
 #include "Types/SlateAttributeMetaData.h"
 
+#include "Algo/BinarySearch.h"
 #include "Types/ReflectionMetadata.h"
 #include "Widgets/SWidget.h"
 
@@ -10,9 +11,11 @@
 
 namespace Private
 {
-	FSlateAttributeDescriptor::OffsetType FindOffet(const SWidget& OwningWidget, FSlateAttributeBase& Attribute)
+	FSlateAttributeDescriptor::OffsetType FindOffet(const SWidget& OwningWidget, const FSlateAttributeBase& Attribute)
 	{
-		return (FSlateAttributeDescriptor::OffsetType)((UPTRINT)(&Attribute) - (UPTRINT)(&OwningWidget));
+		UPTRINT Offset = (UPTRINT)(&Attribute) - (UPTRINT)(&OwningWidget);
+		ensure(Offset <= std::numeric_limits<FSlateAttributeDescriptor::OffsetType>::max());
+		return (FSlateAttributeDescriptor::OffsetType)(Offset);
 	}
 }
 
@@ -49,11 +52,8 @@ void FSlateAttributeMetaData::RegisterAttribute(SWidget& OwningWidget, FSlateAtt
 		TSharedRef<FSlateAttributeMetaData> NewAttributeMetaData = MakeShared<FSlateAttributeMetaData>();
 		NewAttributeMetaData->RegisterAttributeImpl(OwningWidget, Attribute, AttributeType, MoveTemp(Wrapper));
 		OwningWidget.bHasRegisteredSlateAttribute = true;
-		if (OwningWidget.GetProxyHandle().IsValid(&OwningWidget))
-		{
-			//OwningWidget.FastPathProxyHandle.MarkWidgetRegisteredSlateAttributeDirty(&OwningWidget);
-		}
 		OwningWidget.MetaData.Insert(NewAttributeMetaData, 0);
+		OwningWidget.Invalidate(EInvalidateWidgetReason::AttributeRegistration);
 	}
 }
 
@@ -64,36 +64,61 @@ void FSlateAttributeMetaData::RegisterAttributeImpl(SWidget& OwningWidget, FSlat
 	if (FoundIndex != INDEX_NONE)
 	{
 		Attributes[FoundIndex].Getter = MoveTemp(Getter);
+		Attributes[FoundIndex].bUpdatedOnce = false;
 	}
 	else
 	{
 		if (AttributeType == ESlateAttributeType::Member)
 		{
-			FSlateAttributeDescriptor::OffsetType Offset = Private::FindOffet(OwningWidget, Attribute);
-			uint32 SortOrder = FSlateAttributeDescriptor::DefaultSortOrder(Offset);
-
 			// MemberAttribute are optional for now but will be needed in the future
+			const FSlateAttributeDescriptor::OffsetType Offset = Private::FindOffet(OwningWidget, Attribute);
 			const FSlateAttributeDescriptor& Descriptor = OwningWidget.GetWidgetClass().GetAttributeDescriptor();
-			if (FSlateAttributeDescriptor::FAttribute const* FoundAttribute = Descriptor.FindMemberAttribute(Offset))
+			const int32 FoundMemberAttributeIndex = Descriptor.IndexOfMemberAttribute(Offset);
+
+			if (FoundMemberAttributeIndex != INDEX_NONE)
 			{
-				Attributes.Emplace(&Attribute, FoundAttribute->SortOrder, MoveTemp(Getter), FoundAttribute->InvalidationReason);
+				FSlateAttributeDescriptor::FAttribute const& FoundAttribute = Descriptor.GetAttributeAtIndex(FoundMemberAttributeIndex);
+				check(FoundMemberAttributeIndex < std::numeric_limits<int16>::max());
+
+				const int32 InsertLocation = Algo::LowerBoundBy(Attributes, FoundAttribute.SortOrder, [](const FGetterItem& Item) { return Item.SortOrder; }, TLess<>());
+
+				FGetterItem& GetterItem = Attributes.Insert_GetRef({ &Attribute, FoundAttribute.SortOrder, MoveTemp(Getter), (int16)FoundMemberAttributeIndex }, InsertLocation);
+				GetterItem.bIsMemberType = true;
+
+				// Do I have dependency or am I a dependency
+				if (!FoundAttribute.Prerequisite.IsNone() && FoundAttribute.bIsPrerequisiteAlsoADependency)
+				{
+					// I can only be updated if the prerequisite is updated
+					const int32 FoundDependencyAttributeIndex = Descriptor.IndexOfMemberAttribute(FoundAttribute.Prerequisite);
+					check(FoundDependencyAttributeIndex < std::numeric_limits<int16>::max());
+					GetterItem.CachedAttributeDependencyIndex = (int16)FoundDependencyAttributeIndex;
+				}
+				GetterItem.bIsADependencyForSomeoneElse = FoundAttribute.bIsADependencyForSomeoneElse;
+				GetterItem.bUpdateWhenCollapsed = FoundAttribute.bUpdateWhenCollapsed;
+				if (GetterItem.bUpdateWhenCollapsed)
+				{
+					++CollaspedAttributeCounter;
+				}
 			}
 			else
 			{
-				Attributes.Emplace(&Attribute, SortOrder, MoveTemp(Getter));
+				const uint32 SortOrder = FSlateAttributeDescriptor::DefaultSortOrder(Offset);
+
+				const  int32 InsertLocation = Algo::LowerBoundBy(Attributes, SortOrder, [](const FGetterItem& Item) { return Item.SortOrder; }, TLess<>());
+				FGetterItem& GetterItem = Attributes.Insert_GetRef({&Attribute, SortOrder, MoveTemp(Getter)}, InsertLocation);
+				GetterItem.bIsMemberType = true;
 			}
 		}
 		else if (AttributeType == ESlateAttributeType::Managed)
 		{
 			const uint32 ManagedSortOrder = std::numeric_limits<uint32>::max();
-			Attributes.Emplace(&Attribute, ManagedSortOrder, MoveTemp(Getter));
+			FGetterItem& GetterItem = Attributes.Emplace_GetRef(&Attribute, ManagedSortOrder, MoveTemp(Getter));
+			GetterItem.bIsManagedType = true;
 		}
 		else
 		{
 			ensureMsgf(false, TEXT("The SlateAttributeType is not supported"));
 		}
-
-		Attributes.Sort();
 	}
 }
 
@@ -107,11 +132,8 @@ bool FSlateAttributeMetaData::UnregisterAttribute(SWidget& OwningWidget, const F
 		{
 			check(bResult); // if the num is 0 then we should have remove an item.
 			OwningWidget.bHasRegisteredSlateAttribute = false;
-			if (OwningWidget.GetProxyHandle().IsValid(&OwningWidget))
-			{
-				//OwningWidget.FastPathProxyHandle.MarkWidgetRegisteredSlateAttributeDirty(&OwningWidget);
-			}
 			OwningWidget.MetaData.RemoveAtSwap(0);
+			OwningWidget.Invalidate(EInvalidateWidgetReason::AttributeRegistration);
 		}
 		return bResult;
 	}
@@ -124,31 +146,89 @@ bool FSlateAttributeMetaData::UnregisterAttributeImpl(const FSlateAttributeBase&
 	const int32 FoundIndex = IndexOfAttribute(Attribute);
 	if (FoundIndex != INDEX_NONE)
 	{
-		Attributes.RemoveAt(FoundIndex);
+		if (Attributes[FoundIndex].bUpdateWhenCollapsed)
+		{
+			check(CollaspedAttributeCounter > 0);
+			--CollaspedAttributeCounter;
+		}
+		Attributes.RemoveAt(FoundIndex); // keep the order valid
 		return true;
 	}
 	return false;
 }
 
 
-void FSlateAttributeMetaData::InvalidateWidget(SWidget& OwningWidget, const FSlateAttributeBase& Attribute, EInvalidateWidgetReason Reason)
+EInvalidateWidgetReason FSlateAttributeMetaData::FGetterItem::GetInvalidationReason(const SWidget& OwningWidget, EInvalidateWidgetReason Reason) const
+{
+	if (CachedAttributeDescriptorIndex != INDEX_NONE)
+	{
+		return OwningWidget.GetWidgetClass().GetAttributeDescriptor().GetAttributeAtIndex(CachedAttributeDescriptorIndex).InvalidationReason.Get(OwningWidget);
+	}
+	return Reason;
+}
+
+
+void FSlateAttributeMetaData::InvalidateWidget(SWidget& OwningWidget, const FSlateAttributeBase& Attribute, ESlateAttributeType AttributeType, EInvalidateWidgetReason Reason)
 {
+	// The widget is in construction phase. It's already invalidated... no need to keep invalidating it.
+	if (OwningWidget.bPauseAttributeInvalidation)
+	{
+		return;
+	}
+
 	if (FSlateAttributeMetaData* AttributeMetaData = FSlateAttributeMetaData::FindMetaData(OwningWidget))
 	{
 		const int32 FoundIndex = AttributeMetaData->IndexOfAttribute(Attribute);
 		if (FoundIndex != INDEX_NONE)
 		{
-			OwningWidget.Invalidate(AttributeMetaData->Attributes[FoundIndex].OverrideInvalidateReason.Get(Reason));
+			FGetterItem& GetterItem = AttributeMetaData->Attributes[FoundIndex];
+			Reason = GetterItem.GetInvalidationReason(OwningWidget, Reason);
+
+			// The dependency attribute need to be updated in the update loop (note that it may not be registered yet)
+			if (GetterItem.bIsADependencyForSomeoneElse)
+			{
+				GetterItem.bUpatedManually = true;
+				AttributeMetaData->bHasUpdatedManuallyFlagToReset = true;
+			}
 		}
-		else
+		// Not registered but may be defined in the Descriptor
+		else if (AttributeType == ESlateAttributeType::Member)
 		{
-			OwningWidget.Invalidate(Reason);
+			FSlateAttributeDescriptor const& AttributeDescriptor = OwningWidget.GetWidgetClass().GetAttributeDescriptor();
+			const FSlateAttributeDescriptor::OffsetType Offset = Private::FindOffet(OwningWidget, Attribute);
+			if (FSlateAttributeDescriptor::FAttribute const* FoundAttribute = AttributeDescriptor.FindMemberAttribute(Offset))
+			{
+				Reason = FoundAttribute->InvalidationReason.Get(OwningWidget);
+
+				if (FoundAttribute->bIsADependencyForSomeoneElse)
+				{
+					// Find if that dependency is registered, if not it is ok because every attribute is updated at least once
+					// Set UpdatedOnce to false to force a new update.
+					AttributeDescriptor.ForEachDependency(*FoundAttribute, [AttributeMetaData](int32 DependencyIndex)
+					{
+						FGetterItem* FoundOther = AttributeMetaData->Attributes.FindByPredicate([DependencyIndex](FGetterItem const& Other)
+						{
+							return Other.CachedAttributeDescriptorIndex == DependencyIndex;
+						});
+						if (FoundOther)
+						{
+							FoundOther->bUpdatedOnce = false;
+						}
+					});
+				}
+			}
 		}
 	}
-	else
+	else if (AttributeType == ESlateAttributeType::Member)
 	{
-		OwningWidget.Invalidate(Reason);
+		const FSlateAttributeDescriptor::OffsetType Offset = Private::FindOffet(OwningWidget, Attribute);
+		if (FSlateAttributeDescriptor::FAttribute const* FoundAttribute = OwningWidget.GetWidgetClass().GetAttributeDescriptor().FindMemberAttribute(Offset))
+		{
+			Reason = FoundAttribute->InvalidationReason.Get(OwningWidget);
+		}
 	}
+
+	OwningWidget.Invalidate(Reason);
 }
 
 
@@ -156,15 +236,83 @@ void FSlateAttributeMetaData::UpdateAttributes(SWidget& OwningWidget)
 {
 	if (FSlateAttributeMetaData* AttributeMetaData = FSlateAttributeMetaData::FindMetaData(OwningWidget))
 	{
-		EInvalidateWidgetReason InvalidationReason = EInvalidateWidgetReason::None;
-		for (FGetterItem& GetterItem: AttributeMetaData->Attributes)
+		AttributeMetaData->UpdateAttributes(OwningWidget, false);
+	}
+}
+
+
+void FSlateAttributeMetaData::UpdateCollapsedAttributes(SWidget& OwningWidget)
+{
+	if (FSlateAttributeMetaData* AttributeMetaData = FSlateAttributeMetaData::FindMetaData(OwningWidget))
+	{
+		// Does it have any collapsed attributes
+		if (AttributeMetaData->CollaspedAttributeCounter > 0)
 		{
-			ISlateAttributeGetter::FUpdateAttributeResult Result = GetterItem.Getter->UpdateAttribute(OwningWidget);
-			if (Result.bInvalidationRequested)
+			AttributeMetaData->UpdateAttributes(OwningWidget, true);
+		}
+	}
+}
+
+
+void FSlateAttributeMetaData::UpdateAttributes(SWidget& OwningWidget, bool bOnlyCollapsed)
+{
+	EInvalidateWidgetReason InvalidationReason = EInvalidateWidgetReason::None;
+	for (int32 Index = 0; Index < Attributes.Num(); ++Index)
+	{
+		FGetterItem& GetterItem = Attributes[Index];
+
+		// Only update attributes that needs to be updated when the widget is collapsed
+		if (bOnlyCollapsed && !GetterItem.bUpdateWhenCollapsed)
+		{
+			continue;
+		}
+
+		// Update every attribute at least once.
+		//Check if it has a dependency and if it was updated this frame (it could be from an UpdateNow)
+		if (GetterItem.CachedAttributeDependencyIndex != INDEX_NONE && GetterItem.bUpdatedOnce)
+		{
+			// Note that the dependency is maybe not registered and the attribute may have been invalidated manually
+
+			// Because the list is sorted, the dependency needs to be before this element.
+			bool bShouldUpdate = false;
+			bool bFound = false;
+			for (int32 OtherIndex = Index-1; OtherIndex >= 0; --OtherIndex)
+			{
+				const FGetterItem& OtherGetterItem = Attributes[OtherIndex];
+				if (OtherGetterItem.CachedAttributeDescriptorIndex == GetterItem.CachedAttributeDependencyIndex)
+				{
+					bFound = true;
+					bShouldUpdate = OtherGetterItem.bUpdatedThisFrame || OtherGetterItem.bUpatedManually;
+					break;
+				}
+			}
+
+			if (!bShouldUpdate)
 			{
-				InvalidationReason |= GetterItem.OverrideInvalidateReason.Get(Result.InvalidationReason);
+				continue;
 			}
 		}
+
+		ISlateAttributeGetter::FUpdateAttributeResult Result = GetterItem.Getter->UpdateAttribute(OwningWidget);
+		GetterItem.bUpdatedOnce = true;
+		GetterItem.bUpdatedThisFrame = Result.bInvalidationRequested;
+		if (Result.bInvalidationRequested && !OwningWidget.bPauseAttributeInvalidation)
+		{
+			InvalidationReason |= GetterItem.GetInvalidationReason(OwningWidget, Result.InvalidationReason);
+		}
+	}
+
+	if (bHasUpdatedManuallyFlagToReset)
+	{
+		for (FGetterItem& GetterItem : Attributes)
+		{
+			GetterItem.bUpatedManually = false;
+		}
+	}
+	bHasUpdatedManuallyFlagToReset = false;
+
+	if (!OwningWidget.bPauseAttributeInvalidation)
+	{
 		OwningWidget.Invalidate(InvalidationReason);
 	}
 }
@@ -178,11 +326,19 @@ void FSlateAttributeMetaData::UpdateAttribute(SWidget& OwningWidget, FSlateAttri
 		if (FoundIndex != INDEX_NONE)
 		{
 			FGetterItem& GetterItem = AttributeMetaData->Attributes[FoundIndex];
+			GetterItem.bUpdatedOnce = true;
 			check(GetterItem.Getter.Get());
 			ISlateAttributeGetter::FUpdateAttributeResult Result = GetterItem.Getter->UpdateAttribute(OwningWidget);
 			if (Result.bInvalidationRequested)
 			{
-				OwningWidget.Invalidate(GetterItem.OverrideInvalidateReason.Get(Result.InvalidationReason));
+				OwningWidget.Invalidate(GetterItem.GetInvalidationReason(OwningWidget, Result.InvalidationReason));
+
+				// The dependency attribute need to be updated in the update loop (note that it may not be registered yet)
+				if (GetterItem.bIsADependencyForSomeoneElse)
+				{
+					GetterItem.bUpatedManually = true;
+					AttributeMetaData->bHasUpdatedManuallyFlagToReset = true;
+				}
 			}
 		}
 	}
diff --git a/Engine/Source/Runtime/SlateCore/Private/Widgets/SWidget.cpp b/Engine/Source/Runtime/SlateCore/Private/Widgets/SWidget.cpp
--- a/Engine/Source/Runtime/SlateCore/Private/Widgets/SWidget.cpp
+++ b/Engine/Source/Runtime/SlateCore/Private/Widgets/SWidget.cpp
@@ -199,9 +199,6 @@ namespace SlateTraceMetaData
 SLATE_IMPLEMENT_WIDGET(SWidget)
 void SWidget::PrivateRegisterAttributes(FSlateAttributeInitializer& AttributeInitializer)
 {
-	AttributeInitializer.AddMemberAttribute(GET_MEMBER_NAME_CHECKED(PrivateThisType, EnabledState), STRUCT_OFFSET(PrivateThisType, EnabledState));
-	//SLATE_ADD_ATTRIBUTE_DEFINITION(AttributeInitializer, EnabledState);
-	//SLATE_ADD_ATTRIBUTE_DEFINITION(AttributeInitializer, Visibility, EInvalidateWidgetReason::All);
 }
 
 SWidget::SWidget()
@@ -217,6 +214,7 @@ SWidget::SWidget()
 	, bNeedsPrepass(true)
 	, bUpdatingDesiredSize(false)
 	, bHasRegisteredSlateAttribute(false)
+	, bPauseAttributeInvalidation(true)
 	, bHasCustomPrepass(false)
 	, bHasRelativeLayoutScale(false)
 	, bVolatilityAlwaysInvalidatesPrepass(false)
@@ -1639,8 +1637,18 @@ void SWidget::Prepass_ChildLoop(float InLayoutScaleMultiplier, FChildren* MyChil
 		}
 		else
 		{
+			if (HasRegisteredSlateAttribute() && !GSlateIsOnFastProcessInvalidation)
+			{
+				FSlateAttributeMetaData::UpdateCollapsedAttributes(*this);
+				if (Child->Visibility.Get() != EVisibility::Collapsed)
+				{
+					// Recur: Descend down the widget tree.
+					Child->Prepass_Internal(ChildLayoutScaleMultiplier);
+				}
+			}
+
 			// If the child widget is collapsed, we need to store the new layout scale it will have when 
-			// it is finally visible and invalidate it's prepass so that it gets that when its visiblity
+			// it is finally visible and invalidate it's prepass so that it gets that when its visibility
 			// is finally invalidated.
 			Child->InvalidatePrepass();
 			Child->PrepassLayoutScaleMultiplier = ChildLayoutScaleMultiplier;
diff --git a/Engine/Source/Runtime/SlateCore/Public/Types/SlateAttribute.h b/Engine/Source/Runtime/SlateCore/Public/Types/SlateAttribute.h
--- a/Engine/Source/Runtime/SlateCore/Public/Types/SlateAttribute.h
+++ b/Engine/Source/Runtime/SlateCore/Public/Types/SlateAttribute.h
@@ -12,15 +12,18 @@
 #ifndef UE_SLATE_WITH_MEMBER_ATTRIBUTE_DEBUGGING
 	#define UE_SLATE_WITH_MEMBER_ATTRIBUTE_DEBUGGING UE_BUILD_DEBUG
 #endif
+#ifndef UE_SLATE_WITH_ATTRIBUTE_INITIALIZATION_ON_BIND
+	#define UE_SLATE_WITH_ATTRIBUTE_INITIALIZATION_ON_BIND 0
+#endif
 
 /**
  * Use TSlateAttribute when it's a SWidget member.
- * Use TSlateManagedAttribute when it's a member inside an array or other moving structure, and the array is a SWidget member. They can only be moved (they can't be copied).
+ * Use TSlateManagedAttribute when it's a member inside an array or other moving structure, and the array is a SWidget member. They can only be moved (they can't be copied). THey consume more memory.
  * For everything else, use TAttribute.
  *
  *
  *
- * In Slate TAttributes are optimized for developer efficiency.
+ * In Slate, TAttributes are optimized for developer efficiency.
  * They enable widgets to poll for data instead of requiring the user to manually set the state on widgets.
  * Attributes generally work well when performance is not a concern but break down when performance is critical (like a game UI).
  *
@@ -30,18 +33,21 @@
  * TAttributes have a high memory overhead and are not cache-friendly.
  *
  * TSlateAttribute makes the attribute system viable for invalidation and more performance-friendly while keeping the benefits of attributes intact.
- * TSlateAttributes are updated once per frame before the Prepass. If the cached value of the TSlateAttribute changes, then it will invalidate the widget.
- * The TSlateAttributes are updated in the order they are defined in the SWidget (by default). TSlateManagedAttribute are updated in a random order (after TSlateAttribute).
- * The order can be defined/override by setting a Prerequisite (see bellow for an example).
+ * TSlateAttributes are updated once per frame in the Prepass update phase. If the cached value of the TSlateAttribute changes, then it will invalidate the widget.
+ * The TSlateAttributes are updated in the order the variables are defined in the SWidget definition (by default).
+ * TSlateManagedAttribute are updated in a random order (after TSlateAttribute).
+ * The update order of TSlateAttribute can be defined/override by setting a Prerequisite (see bellow for an example).
  * The invalidation reason can be a predicate (see bellow for an example).
- * The invalidation reason can be override per SWidget. Use override with precaution since it can break the invalidation of child widget.
+ * The invalidation reason can be override per SWidget. Use override with precaution since it can break the invalidation of widget's parent.
+ * The widget attributes are updated only if the widget is visible/not collapsed.
  *
  *
  * TSlateAttribute is not copyable and can only live inside a SWidget.
- * For performance reasons, we do not save the extra information that would be needed to be memory save in all contexts.
+ * For performance reasons, we do not save the extra information that would be needed to be "memory save" in all contexts.
  * If you create a TSlateAttribute that can be moved, you need to use TSlateManagedAttribute.
- * The managed version is as fast but use more memory and is less cache-friendly.
+ * TSlateManagedAttribute is as fast as TSlateAttribute but use more memory and is less cache-friendly.
  * Note, if you use TAttribute to change the state of a SWidget, you need to override bool ComputeVolatility() const.
+ * Note, ComputeVolatility() is not needed for TSlateAttribute and TSlateManagedAttribute
  *
  * TSlateAttribute request a SWidget pointer. The "this" pointer should ALWAYS be used.
  * (The SlateAttribute pointer is saved inside the SlateAttributeMetaData. The widget needs to be aware when the pointer changes.)
@@ -51,15 +57,17 @@
  *  // The new way of using attribute in your SWidget
  *	class SMyNewWidget : public SLeafWidget
  *	{
- *		SLATE_DECLARE_WIDGET(SMyNewWidget, SLeafWidget)	// I is optional.
-														// If added, you need to implement void PrivateRegisterAttributes(FSlateAttributeInitializer& AttributeInitializer).
+ *		SLATE_DECLARE_WIDGET(SMyNewWidget, SLeafWidget)	// It defined the static data needed by the SlateAttribute system.
+ *														// The invalidation reason is defined in the static function "void PrivateRegisterAttributes(FSlateAttributeInitializer& AttributeInitializer)"
+ *														// You need to implement the function PrivateRegisterAttributes.
+ *														// If you can't implement the SLATE_DECLARE_WIDGET pattern, use TSlateAttribute and provide a reason instead.
  *
  *		SLATE_BEGIN_ARGS(SMyNewWidget)
  *			SLATE_ATTRIBUTE(FLinearColor, Color)
  *		SLATE_END_ARGS()
  *
  *		SMyNewWidget()
- *			: NewWay(*this, FLinearColor::White) // Always use "this" pointer. No exception. If you can't, use TAttribute.
+ *			: NewWay(*this, FLinearColor::White) // Always use "this" pointer. No exception. If you can't, use TAttribute instead.
  *		{}
  *
  *		virtual int32 OnPaint( const FPaintArgs& Args, const FGeometry& AllottedGeometry, const FSlateRect& MyCullingRect, FSlateWindowElementList& OutDrawElements, int32 LayerId, const FWidgetStyle& InWidgetStyle, bool bParentEnabled ) const override
@@ -80,34 +88,44 @@
  *		void Construct(const FArguments& InArgs)
  *		{
  *			// Assigned the attribute if needed.
-			// If the TAttribute is bound, bound the TSlateAttribute
-			// It the TAttribute is set, set the TSlateAttribute, test if the value changed and invalidate the widget.
+ *			// If the TAttribute is bound, copy the TSlateAttribute to the TAttribute getter, update the value, test if the value have changed and, if so, invalidate the widget.
+ *			// It the TAttribute is set, remove previous getter, set the TSlateAttribute value, test if the value changed and, if so, invalidate the widget.
  *			// Else, nothing (use the previous value).
  *			NewWays.Bind(*this, InArgs._Color); // Always use the "this" pointer.
  *		}
  *
  *		// If NewWay is bounded and the values changed from the previous frame, the widget will be invalidated.
  *		//In this case, it will only paint the widget (no layout required).
- *		TSlateAttribute<FLinearColor, EInvalidationReason::Paint> NewWay;
+ *		TSlateAttribute<FLinearColor> NewWay;
+ *
+ *		// Note that you can put a predicate or put the invalidation as a template argument.
+ *		//Either you define the attribute in PrivateRegisterAttributes or you define the invalidation at declaration
+ *		//but you need to use one and one of the 2 methods. Compile and runtime check will check if you did it properly.
+ *		//Setting the reason as a template argument is less error prone (ie. missing from the PrivateRegisterAttributes, bad copy paste, bad reason...)
+ *		//Setting the reason in PrivateRegisterAttributes enable override and ease debugging later (ie. the attribute will be named)
+ *		//TSlateAttribute<FLinearColor, EInvalidationReason::Paint> NewWayWithout_SLATE_DECLARE_WIDGET_;
  *	};
  *
  *  .cpp
  *	// This is optional. It is still a good practice to implement it since it will allow user to override the default behavior and control the update order.
- *	// The WidgetReflector use that information.
+ *	//The WidgetReflector use that information.
+ *	//See NewWayWithout_SLATE_DECLARE_WIDGET_ for an example of how to use the system without the SLATE_DECLARE_WIDGET
+ *
  *	SLATE_IMPLEMENT_WIDGET(STextBlock)
  *	void STextBlock::PrivateRegisterAttributes(FSlateAttributeInitializer& AttributeInitializer)
  *	{
- *		SLATE_ADD_MEMBER_ATTRIBUTE_DEFINITION(AttributeInitializer, BoundText)
-			.SetPrerequisite("bEnabled"); // SetPrerequisite is not needed here, this is just an example to show how you could do it if needed.
-
-		//bEnabled invalidate paint, we need it to invalidate the Layout.
-		AttributeInitializer.OverrideInvalidationReason("bEnabled", EWidgetInvalidationReason::Layout | EWidgetInvalidationReason::Paint);
+ *		SLATE_ADD_MEMBER_ATTRIBUTE_DEFINITION(AttributeInitializer, NewWay, EInvalidationReason::Paint)
+ *			.SetPrerequisite("bEnabled"); // SetPrerequisite is not needed here. This is just an example to show how you could do it if needed.
+ *
+ *		//bEnabled invalidate paint, we need it to invalidate the Layout.
+ *		AttributeInitializer.OverrideInvalidationReason("bEnabled",
+			FSlateAttributeDescriptor::FInvalidateWidgetReasonAttribute{EWidgetInvalidationReason::Layout | EWidgetInvalidationReason::Paint});
  *	}
  *
  *
- * // We used to do this. Keep using this method when you are not able to provide the "this" pointer to TSlateAttribute.
- * class SMyOldWidget : public SLeafWidget
- * {
+ *	// We used to do this. Keep using this method when you are not able to provide the "this" pointer to TSlateAttribute.
+ *	class SMyOldWidget : public SLeafWidget
+ *	{
  *		SLATE_BEGIN_ARGS(SMyOldWidget)
  *			SLATE_ATTRIBUTE(FLinearColor, Color)
  *		SLATE_END_ARGS()
@@ -139,9 +157,11 @@
  * };
  */
 
+
 class SWidget;
 
-/** */
+
+/** Default predicate to compare FText. */
 struct TSlateAttributeFTextComparePredicate
 {
 	bool operator()(const FText& Lhs, const FText& Rhs) const
@@ -150,22 +170,32 @@ struct TSlateAttributeFTextComparePredicate
 	}
 };
 
-/** */
-struct FSlateAttributeBase
-{
-	
-};
 
-/** */
+/** Predicate that returns the InvalidationReason defined as argument type. */
 template<EInvalidateWidgetReason InvalidationReason>
 struct TSlateAttributeInvalidationReason
 {
 	static constexpr EInvalidateWidgetReason GetInvalidationReason(const SWidget&) { return InvalidationReason; }
 };
 
+
+/** Base struct of all SlateAttribute type. */
+struct FSlateAttributeBase
+{
+
+};
+
+
 /** */
 namespace SlateAttributePrivate
 {
+	/** Predicate used to identify if the InvalidationWidgetReason is defined in the attribute descriptor. */
+	struct FSlateAttributeNoInvalidationReason
+	{
+		static constexpr EInvalidateWidgetReason GetInvalidationReason(const SWidget&) { return EInvalidateWidgetReason::None; }
+	};
+
+
 	/** */
 	enum class ESlateAttributeType : uint8
 	{
@@ -210,6 +240,7 @@ namespace SlateAttributePrivate
 	{
 	protected:
 		bool ProtectedIsWidgetInDestructionPath(SWidget* Widget) const;
+		bool ProtectedIsImplemented(const SWidget& Widget) const;
 		void ProtectedUnregisterAttribute(SWidget& Widget, ESlateAttributeType AttributeType) const;
 		void ProtectedRegisterAttribute(SWidget& Widget, ESlateAttributeType AttributeType, TUniquePtr<ISlateAttributeGetter>&& Wrapper);
 		void ProtectedInvalidateWidget(SWidget& Widget, ESlateAttributeType AttributeType, EInvalidateWidgetReason InvalidationReason) const;
@@ -242,6 +273,13 @@ namespace SlateAttributePrivate
 		static EInvalidateWidgetReason GetInvalidationReason(const SWidget& Widget) { return FInvalidationReasonPredicate::GetInvalidationReason(Widget); }
 		
 	private:
+		void UpdateNowOnBind(SWidget& Widget)
+		{
+#if UE_SLATE_WITH_ATTRIBUTE_INITIALIZATION_ON_BIND
+			ProtectedUpdateNow(Widget, InAttributeType);
+#endif
+		}
+
 		void VerifyOwningWidget(const SWidget& Widget) const
 		{
 #if UE_SLATE_WITH_MEMBER_ATTRIBUTE_DEBUGGING
@@ -398,7 +436,11 @@ namespace SlateAttributePrivate
 			}
 		}
 
-		/** Bind the SlateAttribute to the Getter function. The SlateAttribute will be updated every frame from the Getter. */
+		/**
+		 * Bind the SlateAttribute to the Getter function.
+		 * (If enabled) Update the value from the Getter. If the value is different, then invalidate the widget.
+		 * The SlateAttribute will now be updated every frame from the Getter.
+		 */
 		void Bind(SWidget& Widget, const FGetter& Getter)
 		{
 			VerifyOwningWidget(Widget);
@@ -409,6 +451,7 @@ namespace SlateAttributePrivate
 				{
 					TUniquePtr<ISlateAttributeGetter> Wrapper = MakeUniqueGetter(*this, Getter);
 					ProtectedRegisterAttribute(Widget, InAttributeType, MoveTemp(Wrapper));
+					UpdateNowOnBind(Widget);
 				}
 			}
 			else
@@ -417,7 +460,11 @@ namespace SlateAttributePrivate
 			}
 		}
 
-		/** Bind the SlateAttribute to the Getter function. The SlateAttribute will be updated every frame from the Getter. */
+		/**
+		 * Bind the SlateAttribute to the Getter function.
+		 * (If enabled) Update the value from the Getter. If the value is different, then invalidate the widget.
+		 * The SlateAttribute will now be updated every frame from the Getter.
+		 */
 		void Bind(SWidget& Widget, FGetter&& Getter)
 		{
 			VerifyOwningWidget(Widget);
@@ -428,6 +475,7 @@ namespace SlateAttributePrivate
 				{
 					TUniquePtr<ISlateAttributeGetter> Wrapper = MakeUniqueGetter(*this, MoveTemp(Getter));
 					ProtectedRegisterAttribute(Widget, InAttributeType, MoveTemp(Wrapper));
+					UpdateNowOnBind(Widget);
 				}
 			}
 			else
@@ -436,7 +484,11 @@ namespace SlateAttributePrivate
 			}
 		}
 
-		/** Bind the SlateAttribute to the Getter function. The SlateAttribute will be updated every frame from the Getter. */
+		/**
+		 * Bind the SlateAttribute to the newly create getter function.
+		 * (If enabled) Update the value from the getter. If the value is different, then invalidate the widget.
+		 * The SlateAttribute will now be updated every frame from the Getter.
+		 */
 		template<typename WidgetType, typename U = typename std::enable_if<std::is_base_of<SWidget, WidgetType>::value>::type>
 		void Bind(WidgetType& Widget, typename FGetter::template TSPMethodDelegate_Const<WidgetType>::FMethodPtr MethodPtr)
 		{
@@ -444,10 +496,16 @@ namespace SlateAttributePrivate
 		}
 
 		/**
-		 * Bind the SlateAttribute to the Attribute Getter function (if it exist). The SlateAttribute will be updated every frame from the Getter.
-		 * Or set the value if the Attribute is bound.
+		 * Bind the SlateAttribute to the Attribute Getter function (if it exist).
+		 * (If enabled) Update the value from the getter. If the value is different, then invalidate the widget.
+		 * The SlateAttribute will now be updated every frame from the Getter.
+		 * Or
+		 * Set the SlateAttribute's value if the Attribute is not bound but is set.
 		 * This will Unbind any previously bound getter.
-		 * It may Invalidate the Widget. @see Set
+		 * If the value is different, then invalidate the widget.
+		 * Or
+		 * Unbind the SlateAttribute if the Attribute is not bound and not set.
+		 * @see Set
 		 */
 		void Assign(SWidget& Widget, const TAttribute<ObjectType>& OtherAttribute)
 		{
@@ -459,6 +517,7 @@ namespace SlateAttributePrivate
 				{
 					TUniquePtr<ISlateAttributeGetter> Wrapper = MakeUniqueGetter(*this, OtherAttribute.GetBinding());
 					ProtectedRegisterAttribute(Widget, InAttributeType, MoveTemp(Wrapper));
+					UpdateNowOnBind(Widget);
 				}
 			}
 			else if (OtherAttribute.IsSet())
@@ -472,10 +531,16 @@ namespace SlateAttributePrivate
 		}
 
 		/**
-		 * Bind the SlateAttribute to the Attribute Getter function (if it exist). The SlateAttribute will be updated every frame from the Getter.
-		 * Or set the value if the Attribute is bound. If the Attribute is not set the DefaultValue will be used.
+		 * Bind the SlateAttribute to the Attribute Getter function (if it exist).
+		 * (If enabled) Update the value from the getter. If the value is different, then invalidate the widget.
+		 * The SlateAttribute will now be updated every frame from the Getter.
+		 * Or
+		 * Set the SlateAttribute's value if the Attribute is not bound but is set.
 		 * This will Unbind any previously bound getter.
-		 * It may Invalidate the Widget. @see Set
+		 * If the value is different, then invalidate the widget.
+		 * Or
+		 * Unbind the SlateAttribute to DefaultValue if the Attribute is not bound and not set.
+		 * @see Set
 		 */
 		void Assign(SWidget& Widget, const TAttribute<ObjectType>& OtherAttribute, const ObjectType& DefaultValue)
 		{
@@ -487,6 +552,7 @@ namespace SlateAttributePrivate
 				{
 					TUniquePtr<ISlateAttributeGetter> Wrapper = MakeUniqueGetter(*this, OtherAttribute.GetBinding());
 					ProtectedRegisterAttribute(Widget, InAttributeType, MoveTemp(Wrapper));
+					UpdateNowOnBind(Widget);
 				}
 			}
 			else if (OtherAttribute.IsSet())
@@ -500,10 +566,16 @@ namespace SlateAttributePrivate
 		}
 
 		/**
-		 * Bind the SlateAttribute to the Attribute Getter function (if it exist). The SlateAttribute will be updated every frame from the Getter.
-		 * Or set the value if the Attribute is bound. If the Attribute is not set the DefaultValue will be used.
+		 * Bind the SlateAttribute to the Attribute Getter function (if it exist).
+		 * (If enabled) Update the value from the getter. If the value is different, then invalidate the widget.
+		 * The SlateAttribute will now be updated every frame from the Getter.
+		 * Or
+		 * Set the SlateAttribute's value if the Attribute is not bound but is set.
 		 * This will Unbind any previously bound getter.
-		 * It may Invalidate the Widget. @see Set
+		 * If the value is different, then invalidate the widget.
+		 * Or
+		 * Unbind the SlateAttribute to DefaultValue if the Attribute is not bound and not set.
+		 * @see Set
 		 */
 		void Assign(SWidget& Widget, const TAttribute<ObjectType>& OtherAttribute, ObjectType&& DefaultValue)
 		{
@@ -515,6 +587,7 @@ namespace SlateAttributePrivate
 				{
 					TUniquePtr<ISlateAttributeGetter> Wrapper = MakeUniqueGetter(*this, OtherAttribute.GetBinding());
 					ProtectedRegisterAttribute(Widget, InAttributeType, MoveTemp(Wrapper));
+					UpdateNowOnBind(Widget);
 				}
 			}
 			else if (OtherAttribute.IsSet())
@@ -664,18 +737,22 @@ public:
 	private:
 		using Super = TSlateAttributeBase<InObjectType, InInvalidationReasonPredicate, InComparePredicate, ESlateAttributeType::Member>;
 
-		//using TSlateAttributeBase<InObjectType, InInvalidationReasonPredicate, InComparePredicate, ESlateAttributeType::Member>::TSlateAttributeBase;
-
 		template<typename WidgetType, typename U = typename std::enable_if<std::is_base_of<SWidget, WidgetType>::value>::type>
 		static void VerifyAttributeAddress(WidgetType& Widget, TSlateMemberAttribute* Self)
 		{
 			checkf((UPTRINT)Self >= (UPTRINT)&Widget && (UPTRINT)Self < (UPTRINT)&Widget + sizeof(WidgetType),
 				TEXT("Use TAttribute or TSlateManagedAttribute instead. See SlateAttribute.h for more info."));
+			ensureAlwaysMsgf((HasDefinedInvalidationReason || Self->ProtectedIsImplemented(Widget)),
+				TEXT("The TSlateAttribute could not be found in the SlateAttributeDescriptor.\n")
+				TEXT("Use the SLATE_DECLARE_WIDGET and add the attribute in PrivateRegisterAttributes,\n")
+				TEXT("Or use TSlateAttribute with a valid Invalidation Reason instead."));
 		}
 
 	public:
 		using FGetter = typename Super::FGetter;
 		using ObjectType = typename Super::ObjectType;
+		static const bool IsMemberType = true;
+		static constexpr bool HasDefinedInvalidationReason = !std::is_same<InInvalidationReasonPredicate, FSlateAttributeNoInvalidationReason>::value;
 
 	public:
 		//~ You can only register Attribute that are defined in a SWidget.
@@ -779,6 +856,7 @@ public:
 		using FInvalidationReasonPredicate = typename Super::FInvalidationReasonPredicate;
 		using FGetter = typename Super::FGetter;
 		using FComparePredicate = typename Super::FComparePredicate;
+		static const bool IsMemberType = false;
 
 		static EInvalidateWidgetReason GetInvalidationReason(const SWidget& Widget) { return Super::GetInvalidationReason(Widget); }
 
diff --git a/Engine/Source/Runtime/SlateCore/Public/Types/SlateAttributeDescriptor.h b/Engine/Source/Runtime/SlateCore/Public/Types/SlateAttributeDescriptor.h
--- a/Engine/Source/Runtime/SlateCore/Public/Types/SlateAttributeDescriptor.h
+++ b/Engine/Source/Runtime/SlateCore/Public/Types/SlateAttributeDescriptor.h
@@ -10,25 +10,83 @@
 
 class FSlateWidgetClassData;
 
-
+/**
+ * Describes the static information about a Widget's type SlateAttributes.
+ **/
 class SLATECORE_API FSlateAttributeDescriptor
 {
 public:
-	using OffsetType = uint32; //UPTRINT
+	/**
+	 * A EInvalidationWidgetReason Attribute 
+	 * It can be explicitly initialize or can be a callback static function or lambda that returns the EInvalidationReason.
+	 * The signature of the function takes a const SWidget& as argument.
+	 */
+	struct FInvalidateWidgetReasonAttribute
+	{
+		using Arg1Type = const class SWidget&;
+		DECLARE_DELEGATE_RetVal_OneParam(EInvalidateWidgetReason, FGetter, Arg1Type);
+
+		FInvalidateWidgetReasonAttribute(const FInvalidateWidgetReasonAttribute&) = default;
+		FInvalidateWidgetReasonAttribute(FInvalidateWidgetReasonAttribute&&) = default;
+		FInvalidateWidgetReasonAttribute& operator=(const FInvalidateWidgetReasonAttribute&) = default;
+		FInvalidateWidgetReasonAttribute& operator=(FInvalidateWidgetReasonAttribute&&) = default;
+
+		/** Default constructor. */
+		explicit FInvalidateWidgetReasonAttribute(EInvalidateWidgetReason InReason)
+			: Reason(InReason)
+			, Getter()
+		{ }
+
+		template<typename... PayloadTypes>
+		explicit FInvalidateWidgetReasonAttribute(typename FGetter::FStaticDelegate::FFuncPtr InFuncPtr, PayloadTypes&&... InputPayload)
+			: Reason(EInvalidateWidgetReason::None)
+			, Getter(FGetter::CreateStatic(InFuncPtr, Forward<PayloadTypes>(InputPayload)...))
+		{
+		}
+
+		template<typename LambdaType, typename... PayloadTypes>
+		explicit FInvalidateWidgetReasonAttribute(LambdaType&& InCallable, PayloadTypes&&... InputPayload)
+			: Reason(EInvalidateWidgetReason::None)
+			, Getter(FGetter::CreateLambda(InCallable, Forward<PayloadTypes>(InputPayload)...))
+		{
+		}
+
+		bool IsBound() const
+		{
+			return Getter.IsBound();
+		}
+
+		EInvalidateWidgetReason Get(const SWidget& Widget) const
+		{
+			return IsBound() ? Getter.Execute(Widget) : Reason;
+		}
+
+	private:
+		EInvalidateWidgetReason Reason;
+		FGetter Getter;
+	};
+
+public:
+	using OffsetType = uint32;
 
-	static uint32 DefaultSortOrder(OffsetType Offset) { return Offset * 10; }
+	/** The default sort order that define in which order attributes will be updated. */
+	static uint32 DefaultSortOrder(OffsetType Offset) { return Offset * 100; }
 
+	/** */
 	struct FAttribute
 	{
 		FName Name;
-		OffsetType Offset;
+		OffsetType Offset = 0;
 		FName Prerequisite;
-		FName Dependency;
-		uint32 SortOrder;
-		TAttribute<EInvalidateWidgetReason> InvalidationReason;
-		bool bIsMemberAttribute;
+		uint32 SortOrder = 0;
+		FInvalidateWidgetReasonAttribute InvalidationReason = FInvalidateWidgetReasonAttribute(EInvalidateWidgetReason::None);
+		bool bIsMemberAttribute = false;
+		bool bIsPrerequisiteAlsoADependency = false;
+		bool bIsADependencyForSomeoneElse = false;
+		bool bUpdateWhenCollapsed = false;
 	};
 
+	/** Internal class to initialize the SlateAttributeDescriptor (Add attributes or modify existing attributes). */
 	struct SLATECORE_API FInitializer
 	{
 	private:
@@ -45,47 +103,113 @@ public:
 		{
 			FAttributeEntry(FSlateAttributeDescriptor& Descriptor, int32 InAttributeIndex);
 
-			/** Assign an order in which the attributes should be updated. */
-			FAttributeEntry& SetPrerequisite(FName Prerequisite);
-			/** Update the property every. */
-			FAttributeEntry& UpdateEveryFrame();
-			/** The property only need to be poll when the dependency changes. */
+			/**
+			 * Update the attribute after the prerequisite.
+			 * The order is guaranteed but other attributes may be updated in between.
+			 * No order is guaranteed if the prerequisite or this property is updated manually.
+			 */
+			FAttributeEntry& UpdatePrerequisite(FName Prerequisite);
+
+			/**
+			 * The property only needs to be updated when the dependency changes inside the update loop.
+			 * The property can still be set/updated manually.
+			 * If the dependency is updated manually, then the property will be updated in the next update loop.
+			 * It will implicitly set a prerequisite.
+			 */
 			FAttributeEntry& UpdateDependency(FName Dependency);
 
+			/**
+			 * Update the attribute when the widget is collapsed and its parent is not.
+			 */
+			FAttributeEntry& UpdateWhenCollapsed();
+
 		private:
 			FSlateAttributeDescriptor& Descriptor;
 			int32 AttributeIndex;
 		};
 
-		FAttributeEntry AddMemberAttribute(FName AttributeName, OffsetType Offset);
+		FAttributeEntry AddMemberAttribute(FName AttributeName, OffsetType Offset, const FInvalidateWidgetReasonAttribute& ReasonGetter);
+		FAttributeEntry AddMemberAttribute(FName AttributeName, OffsetType Offset, FInvalidateWidgetReasonAttribute&& ReasonGetter);
 
-		FAttributeEntry AddMemberAttribute(FName AttributeName, OffsetType Offset, TAttribute<EInvalidateWidgetReason> ReasonGetter);
+		void OverrideInvalidationReason(FName AttributeName, const FInvalidateWidgetReasonAttribute& Reason);
+		void OverrideInvalidationReason(FName AttributeName, FInvalidateWidgetReasonAttribute&& Reason);
 
-		void OverrideInvalidationReason(FName AttributeName, TAttribute<EInvalidateWidgetReason> Reason);
+		void SetUpdateWhenCollapsed(FName AttributeName, bool bUpdateWhenCollapsed);
 
-		void SetPrerequisite(FName AttributeName, FName Prerequisite);
 
 	private:
 		FSlateAttributeDescriptor& Descriptor;
 	};
 
+	/** @returns the number of Attributes registered. */
+	int32 AttributeNum() const { return Attributes.Num(); }
+
+	/** @returns the Attribute at the index previously found with IndexOfMemberAttribute */
+	FAttribute const& GetAttributeAtIndex(int32 Index) const;
+
+	/** @returns the Attribute with the corresponding name. */
 	FAttribute const* FindAttribute(FName AttributeName) const;
 
+	/** @returns the index of a SlateAttribute that have the corresponding memory offset. */
+	int32 IndexOfMemberAttribute(OffsetType AttributeOffset) const;
+
+	/** @returns the index of a SlateAttribute that have the corresponding memory offset. */
+	int32 IndexOfMemberAttribute(FName AttributeName) const;
+
+	/** @returns the Attribute of a SlateAttribute that have the corresponding memory offset. */
 	FAttribute const* FindMemberAttribute(OffsetType AttributeOffset) const;
 
+	/** Iterator over each dependency this attribute is responsible of. */
+	template<typename Predicate>
+	void ForEachDependency(FAttribute const& Attribute, Predicate Pred) const
+	{
+		checkf(&Attribute >= Attributes.GetData() && &Attribute <= Attributes.GetData()+Attributes.Num()
+			, TEXT("The attribute is not part of this Descriptor."));
+		if (Attribute.bIsADependencyForSomeoneElse)
+		{
+			ForEachDependencyImpl(Attribute.Name, (int32)(&Attribute - Attributes.GetData()), Pred);
+		}
+	}
+
 private:
 	FAttribute* FindAttribute(FName AttributeName);
 
-	FInitializer::FAttributeEntry AddMemberAttribute(FName AttributeName, OffsetType Offset, TAttribute<EInvalidateWidgetReason>&& ReasonGetter);
-	void OverrideInvalidationReason(FName AttributeName, TAttribute<EInvalidateWidgetReason>&& ReasonGetter);
-	void SetDependency(FAttribute& Attribute, FName Dependency);
-	void SetPrerequisite(FAttribute& Attribute, FName Prerequisite);
+	FInitializer::FAttributeEntry AddMemberAttribute(FName AttributeName, OffsetType Offset, FInvalidateWidgetReasonAttribute ReasonGetter);
+	void OverrideInvalidationReason(FName AttributeName, FInvalidateWidgetReasonAttribute ReasonGetter);
+	void SetPrerequisite(FAttribute& Attribute, FName Prerequisite, bool bSetAsDependency);
+	void SetUpdateWhenCollapsed(FAttribute& Attribute, bool bUpdate);
+
+	template<typename Predicate>
+	void ForEachDependencyImpl(FName const& LookForName, int32 Index, Predicate& Pred) const
+	{
+		++Index;
+		for (; Index < Attributes.Num(); ++Index)
+		{
+			FAttribute const& Other = Attributes[Index];
+			if (Other.Prerequisite == LookForName && Other.bIsPrerequisiteAlsoADependency)
+			{
+				Pred(Index);
+				if (Other.bIsADependencyForSomeoneElse)
+				{
+					ForEachDependencyImpl(Other.Name, Index, Pred);
+				}
+			}
+		}
+	}
 
 private:
 	TArray<FAttribute> Attributes;
 };
 
-#define SLATE_ADD_MEMBER_ATTRIBUTE_DEFINITION(_Initializer, _Property) {_Initializer.AddMemberAttribute(GET_MEMBER_NAME_CHECKED(PrivateThisType, _Property), STRUCT_OFFSET(PrivateThisType, _Property));}
-#define SLATE_ADD_MEMBER_ATTRIBUTE_DEFINITION_REASON(_Initializer, _Property, _Reason) {_Initializer.AddMemberAttribute(GET_MEMBER_NAME_CHECKED(PrivateThisType, _Property), STRUCT_OFFSET(PrivateThisType, _Property), _Reason));}
-
+/**
+ * Add a TSlateAttribute to the descriptor.
+ * @param _Initializer The FSlateAttributeInitializer from the PrivateRegisterAttributes function.
+ * @param _Property The TSlateAttribute property
+ * @param _Reason The EInvalidationWidgetReason or a static function/lambda that takes a const SWidget& and that returns the invalidation reason.
+ */
+#define SLATE_ADD_MEMBER_ATTRIBUTE_DEFINITION(_Initializer, _Property, _Reason) \
+		static_assert(decltype(_Property)::IsMemberType, "The SlateProperty is not a TSlateAttribute. Do not use SLATE_ADD_MEMBER_ATTRIBUTE_DEFINITION"); \
+		static_assert(!decltype(_Property)::HasDefinedInvalidationReason, "When implementing the SLATE_DECLARE_WIDGET pattern, use TSlateAttribute without the invalidation reason."); \
+		_Initializer.AddMemberAttribute(GET_MEMBER_NAME_CHECKED(PrivateThisType, _Property), STRUCT_OFFSET(PrivateThisType, _Property), FSlateAttributeDescriptor::FInvalidateWidgetReasonAttribute{_Reason})
+		//__VA_ARGS__
 using FSlateAttributeInitializer = FSlateAttributeDescriptor::FInitializer;
diff --git a/Engine/Source/Runtime/SlateCore/Public/Types/SlateAttributeMetaData.h b/Engine/Source/Runtime/SlateCore/Public/Types/SlateAttributeMetaData.h
--- a/Engine/Source/Runtime/SlateCore/Public/Types/SlateAttributeMetaData.h
+++ b/Engine/Source/Runtime/SlateCore/Public/Types/SlateAttributeMetaData.h
@@ -26,6 +26,7 @@ public:
 public:
 	static FSlateAttributeMetaData* FindMetaData(const SWidget& OwningWidget);
 	static void UpdateAttributes(SWidget& OwningWidget);
+	static void UpdateCollapsedAttributes(SWidget& OwningWidget);
 
 public:
 	bool IsBound(const FSlateAttributeBase& Attribute) const
@@ -40,7 +41,7 @@ private:
 	using ISlateAttributeGetter = SlateAttributePrivate::ISlateAttributeGetter;
 	static void RegisterAttribute(SWidget& OwningWidget, FSlateAttributeBase& Attribute, ESlateAttributeType AttributeType, TUniquePtr<ISlateAttributeGetter>&& Wrapper);
 	static bool UnregisterAttribute(SWidget& OwningWidget, const FSlateAttributeBase& Attribute);
-	static void InvalidateWidget(SWidget& OwningWidget, const FSlateAttributeBase& Attribute, EInvalidateWidgetReason Reason);
+	static void InvalidateWidget(SWidget& OwningWidget, const FSlateAttributeBase& Attribute, ESlateAttributeType AttributeType, EInvalidateWidgetReason Reason);
 	static void UpdateAttribute(SWidget& OwningWidget, FSlateAttributeBase& Attribute);
 	static bool IsAttributeBound(const SWidget& OwningWidget, const FSlateAttributeBase& Attribute);
 	static FDelegateHandle GetAttributeGetterHandle(const SWidget& OwningWidget, const FSlateAttributeBase& Attribute);
@@ -49,6 +50,7 @@ private:
 private:
 	void RegisterAttributeImpl(SWidget& OwningWidget, FSlateAttributeBase& Attribute, ESlateAttributeType AttributeType, TUniquePtr<ISlateAttributeGetter>&& Getter);
 	bool UnregisterAttributeImpl(const FSlateAttributeBase& Attribute);
+	void UpdateAttributes(SWidget& OwningWidget, bool bOnlyCollapsed);
 
 private:
 	int32 IndexOfAttribute(const FSlateAttributeBase& Attribute) const
@@ -61,29 +63,53 @@ private:
 	{
 		FGetterItem() = default;
 		FGetterItem(const FGetterItem&) = delete;
+		FGetterItem(FGetterItem&&) = default;
 		FGetterItem& operator=(const FGetterItem&) = delete;
 		FGetterItem(FSlateAttributeBase* InAttribute, uint32 InSortOrder, TUniquePtr<SlateAttributePrivate::ISlateAttributeGetter>&& InGetter)
 			: Attribute(InAttribute)
 			, Getter(MoveTemp(InGetter))
 			, SortOrder(InSortOrder)
-			, OverrideInvalidateReason()
+			, CachedAttributeDescriptorIndex(INDEX_NONE)
+			, CachedAttributeDependencyIndex(INDEX_NONE)
+			, Flags(0)
 		{ }
-		FGetterItem(FSlateAttributeBase* InAttribute, uint32 InSortOrder, TUniquePtr<SlateAttributePrivate::ISlateAttributeGetter> && InGetter, TAttribute<EInvalidateWidgetReason> InInvalidaitonReason)
+		FGetterItem(FSlateAttributeBase* InAttribute, uint32 InSortOrder, TUniquePtr<SlateAttributePrivate::ISlateAttributeGetter> && InGetter, int16 InAttributeDescriptorIndex)
 			: Attribute(InAttribute)
 			, Getter(MoveTemp(InGetter))
 			, SortOrder(InSortOrder)
-			, OverrideInvalidateReason(MoveTemp(InInvalidaitonReason))
+			, CachedAttributeDescriptorIndex(InAttributeDescriptorIndex)
+			, CachedAttributeDependencyIndex(INDEX_NONE)
+			, Flags(0)
 		{ }
 		FSlateAttributeBase* Attribute;
 		TUniquePtr<SlateAttributePrivate::ISlateAttributeGetter> Getter;
 		uint32 SortOrder;
-		TAttribute<EInvalidateWidgetReason> OverrideInvalidateReason;
+		int16 CachedAttributeDescriptorIndex;
+		int16 CachedAttributeDependencyIndex;
+		union
+		{
+			struct
+			{
+				int8 bUpdatedOnce : 1;
+				int8 bUpdatedThisFrame : 1;
+				int8 bUpatedManually : 1;
+				int8 bIsADependencyForSomeoneElse : 1;
+				int8 bIsMemberType : 1;
+				int8 bIsManagedType : 1;
+				int8 bUpdateWhenCollapsed : 1;
+			};
+			int8 Flags;
+		};
 
 		bool operator<(const FGetterItem& Other) const
 		{
 			return SortOrder < Other.SortOrder;
 		}
+
+		EInvalidateWidgetReason GetInvalidationReason(const SWidget&, EInvalidateWidgetReason Reason) const;
 	};
 
 	TArray<FGetterItem, TInlineAllocator<4>> Attributes;
+	bool bHasUpdatedManuallyFlagToReset = false;
+	uint8 CollaspedAttributeCounter = 0;
 };
diff --git a/Engine/Source/Runtime/SlateCore/Public/Widgets/DeclarativeSyntaxSupport.h b/Engine/Source/Runtime/SlateCore/Public/Widgets/DeclarativeSyntaxSupport.h
--- a/Engine/Source/Runtime/SlateCore/Public/Widgets/DeclarativeSyntaxSupport.h
+++ b/Engine/Source/Runtime/SlateCore/Public/Widgets/DeclarativeSyntaxSupport.h
@@ -1099,6 +1099,7 @@ struct TSlateDecl
 	{
 		_Widget->SWidgetConstruct(InArgs);
 		_RequiredArgs.CallConstruct(_Widget, InArgs);
+		_Widget->bPauseAttributeInvalidation = false;
 		_Widget->CacheVolatility();
 
 		return _Widget;
diff --git a/Engine/Source/Runtime/SlateCore/Public/Widgets/SWidget.h b/Engine/Source/Runtime/SlateCore/Public/Widgets/SWidget.h
--- a/Engine/Source/Runtime/SlateCore/Public/Widgets/SWidget.h
+++ b/Engine/Source/Runtime/SlateCore/Public/Widgets/SWidget.h
@@ -162,16 +162,16 @@ class SLATECORE_API SWidget
 	friend struct TSlateDecl;
 
 protected:
-	template<typename InObjectType, EInvalidateWidgetReason InInvalidationReasonValue, typename InComparePredicate = TEqualTo<>>
-	struct TSlateAttribute : public ::SlateAttributePrivate::TSlateMemberAttribute<InObjectType, TSlateAttributeInvalidationReason<InInvalidationReasonValue>, InComparePredicate>
+	template<typename InObjectType, EInvalidateWidgetReason InInvalidationReasonValue = EInvalidateWidgetReason::None, typename InComparePredicate = TEqualTo<>>
+	struct TSlateAttribute : public ::SlateAttributePrivate::TSlateMemberAttribute<
+		InObjectType,
+		typename std::conditional<InInvalidationReasonValue == EInvalidateWidgetReason::None, ::SlateAttributePrivate::FSlateAttributeNoInvalidationReason, TSlateAttributeInvalidationReason<InInvalidationReasonValue>>::type,
+		typename std::conditional<std::is_same<InObjectType, FText>::value, TSlateAttributeFTextComparePredicate, InComparePredicate>::type>
 	{
-		using ::SlateAttributePrivate::TSlateMemberAttribute<InObjectType, TSlateAttributeInvalidationReason<InInvalidationReasonValue>, InComparePredicate>::TSlateMemberAttribute;
-	};
-
-	template<EInvalidateWidgetReason InInvalidationReasonValue>
-	struct TSlateAttribute<FText, InInvalidationReasonValue> : public ::SlateAttributePrivate::TSlateMemberAttribute<FText, TSlateAttributeInvalidationReason<InInvalidationReasonValue>, TSlateAttributeFTextComparePredicate>
-	{
-		using ::SlateAttributePrivate::TSlateMemberAttribute<FText, TSlateAttributeInvalidationReason<InInvalidationReasonValue>, TSlateAttributeFTextComparePredicate>::TSlateMemberAttribute;
+		using ::SlateAttributePrivate::TSlateMemberAttribute<
+			InObjectType,
+			typename std::conditional<InInvalidationReasonValue == EInvalidateWidgetReason::None, ::SlateAttributePrivate::FSlateAttributeNoInvalidationReason, TSlateAttributeInvalidationReason<InInvalidationReasonValue>>::type,
+			typename std::conditional<std::is_same<InObjectType, FText>::value, TSlateAttributeFTextComparePredicate, InComparePredicate>::type>::TSlateMemberAttribute;
 	};
 
 	template<typename InObjectType, typename InInvalidationReasonPredicate, typename InComparePredicate = TEqualTo<>>
@@ -1083,7 +1083,7 @@ public:
 		if(RenderOpacity != InRenderOpacity)
 		{
 			RenderOpacity = InRenderOpacity;
-			Invalidate(EInvalidateWidget::Paint);
+			Invalidate(EInvalidateWidgetReason::Paint);
 		}
 	}
 
@@ -1160,7 +1160,7 @@ public:
 			Clipping = InClipping;
 			OnClippingChanged();
 			// @todo - Fast path should this be Paint?
-			Invalidate(EInvalidateWidget::Layout);
+			Invalidate(EInvalidateWidgetReason::Layout);
 		}
 	}
 
@@ -1180,7 +1180,7 @@ public:
 		{
 			CullingBoundsExtension = InCullingBoundsExtension;
 			// @todo - Fast path should this be Paint?
-			Invalidate(EInvalidateWidget::Layout);
+			Invalidate(EInvalidateWidgetReason::Layout);
 		}
 	}
 
@@ -1199,7 +1199,7 @@ public:
 		if (FlowDirectionPreference != InFlowDirectionPreference)
 		{
 			FlowDirectionPreference = InFlowDirectionPreference;
-			Invalidate(EInvalidateWidget::Paint);
+			Invalidate(EInvalidateWidgetReason::Paint);
 		}
 	}
 
@@ -1645,6 +1645,9 @@ private:
 	/** Is there at least one SlateAttribute currently registered. */
 	uint8 bHasRegisteredSlateAttribute : 1;
 
+	/** Prevents invalidation while we are in the construction phase. */
+	uint8 bPauseAttributeInvalidation : 1;
+
 protected:
 	uint8 bHasCustomPrepass : 1;
 
