Slate: Fix some implicit conversion.
#rnx

diff --git a/Engine/Source/Runtime/SlateCore/Private/Application/SlateApplicationBase.cpp b/Engine/Source/Runtime/SlateCore/Private/Application/SlateApplicationBase.cpp
--- a/Engine/Source/Runtime/SlateCore/Private/Application/SlateApplicationBase.cpp
+++ b/Engine/Source/Runtime/SlateCore/Private/Application/SlateApplicationBase.cpp
@@ -60,7 +60,7 @@ void FSlateApplicationBase::GetSafeZoneSize(FMargin& SafeZone, const FVector2D&
 	{
 		FDisplayMetrics Metrics;
 		GetCachedDisplayMetrics(Metrics);
-		ContainerSize = FVector2D(Metrics.PrimaryDisplayWidth, Metrics.PrimaryDisplayHeight);
+		ContainerSize = FVector2D((float)Metrics.PrimaryDisplayWidth, (float)Metrics.PrimaryDisplayHeight);
 	}
 
 	FMargin SafeZoneRatio;
diff --git a/Engine/Source/Runtime/SlateCore/Private/Fonts/FontCacheFreeType.h b/Engine/Source/Runtime/SlateCore/Private/Fonts/FontCacheFreeType.h
--- a/Engine/Source/Runtime/SlateCore/Private/Fonts/FontCacheFreeType.h
+++ b/Engine/Source/Runtime/SlateCore/Private/Fonts/FontCacheFreeType.h
@@ -259,7 +259,7 @@ public:
 	/**
 	 * Gets the memory size of the loaded font or 0 if the font is streamed
 	 */
-	FORCEINLINE uint32 GetAllocatedMemorySize() const
+	FORCEINLINE SIZE_T GetAllocatedMemorySize() const
 	{
 #if WITH_FREETYPE
 		return Memory.IsValid() ? Memory->GetData().GetAllocatedSize() : 0;
diff --git a/Engine/Source/Runtime/SlateCore/Private/Fonts/SlateFontRenderer.cpp b/Engine/Source/Runtime/SlateCore/Private/Fonts/SlateFontRenderer.cpp
--- a/Engine/Source/Runtime/SlateCore/Private/Fonts/SlateFontRenderer.cpp
+++ b/Engine/Source/Runtime/SlateCore/Private/Fonts/SlateFontRenderer.cpp
@@ -7,6 +7,8 @@
 #include "HAL/IConsoleManager.h"
 #include "SlateGlobals.h"
 
+#include <limits>
+
 DECLARE_CYCLE_STAT(TEXT("Freetype Render Glyph"), STAT_FreetypeRenderGlyph, STATGROUP_Slate);
 
 /**
@@ -92,17 +94,19 @@ void ResizeFontBitmap_StoreAverageOutputPixel(TArray<uint8>& DstData, const int3
 template <>
 void ResizeFontBitmap_StoreAverageOutputPixel<1>(TArray<uint8>& DstData, const int32 PixelIndex, const FVector4& StepColor, const int32 PixelCount)
 {
-	DstData[PixelIndex] = FMath::Clamp(FMath::TruncToInt(StepColor.X / (float)PixelCount), 0, 255);
+	ensure(PixelCount > 0);
+	DstData[PixelIndex] = FMath::Clamp<uint8>((uint8)(StepColor.X / (float)PixelCount), 0, 255);
 }
 
 template <>
 void ResizeFontBitmap_StoreAverageOutputPixel<4>(TArray<uint8>& DstData, const int32 PixelIndex, const FVector4& StepColor, const int32 PixelCount)
 {
+	ensure(PixelCount > 0);
 	const int32 DstDataIndex = PixelIndex * 4;
-	DstData[DstDataIndex+0] = FMath::Clamp(FMath::TruncToInt(StepColor.X / (float)PixelCount), 0, 255);
-	DstData[DstDataIndex+1] = FMath::Clamp(FMath::TruncToInt(StepColor.Y / (float)PixelCount), 0, 255);
-	DstData[DstDataIndex+2] = FMath::Clamp(FMath::TruncToInt(StepColor.Z / (float)PixelCount), 0, 255);
-	DstData[DstDataIndex+3] = FMath::Clamp(FMath::TruncToInt(StepColor.W / (float)PixelCount), 0, 255);
+	DstData[DstDataIndex + 0] = FMath::Clamp<uint8>(FMath::TruncToInt(StepColor.X / (float)PixelCount), 0, 255);
+	DstData[DstDataIndex + 1] = FMath::Clamp<uint8>(FMath::TruncToInt(StepColor.Y / (float)PixelCount), 0, 255);
+	DstData[DstDataIndex + 2] = FMath::Clamp<uint8>(FMath::TruncToInt(StepColor.Z / (float)PixelCount), 0, 255);
+	DstData[DstDataIndex + 3] = FMath::Clamp<uint8>(FMath::TruncToInt(StepColor.W / (float)PixelCount), 0, 255);
 }
 
 // Note: Copied from FImageUtils and updated to work efficiently with 1 or 4-bytes per-pixel, and include the alpha channel (which the FImageUtils version just clobbers)
@@ -463,7 +467,7 @@ struct FRasterizerSpanList
  */
 void RenderOutlineRows(FT_Library Library, FT_Outline* Outline, FRasterizerSpanList& OutSpansList)
 {
-	auto RasterizerCallback = [](const int32 Y, const int32 Count, const FT_Span* const Spans, void* const UserData)
+	auto RasterizerCallback = [](const int Y, const int Count, const FT_Span* const Spans, void* const UserData)
 	{
 		FRasterizerSpanList& UserDataSpanList = *static_cast<FRasterizerSpanList*>(UserData);
 
@@ -475,8 +479,8 @@ void RenderOutlineRows(FT_Library Library, FT_Outline* Outline, FRasterizerSpanL
 		{
 			const FT_Span& Span = Spans[SpanIndex];
 
-			BoundingBox += FVector2D(Span.x, Y);
-			BoundingBox += FVector2D(Span.x + Span.len - 1, Y);
+			BoundingBox += FVector2D((float)Span.x, (float)Y);
+			BoundingBox += FVector2D((float)(Span.x + Span.len - 1), (float)Y);
 
 			UserDataSpans.Add(FRasterizerSpan(Span.x, Y, Span.len, Span.coverage));
 		}
@@ -547,18 +551,18 @@ bool FSlateFontRenderer::GetRenderDataInternal(const FFreeTypeFaceGlyphData& InF
 		FVector2D Size = BoundingBox.GetSize();
 
 		//Note: we add 1 to width and height because the size of the rect is inclusive
-		int32 Width = FMath::TruncToInt(Size.X)+1;
-		int32 Height = FMath::TruncToInt(Size.Y)+1;
+		const int32 Width = FMath::TruncToInt(Size.X)+1;
+		const int32 Height = FMath::TruncToInt(Size.Y)+1;
+		ensure(Width <= std::numeric_limits<int16>::max());
+		ensure(Height <= std::numeric_limits<int16>::max());
+		OutRenderData.SizeX = (int16)Width;
+		OutRenderData.SizeY = (int16)Height;
 
 		OutRenderData.RawPixels.Reset();
-
-		OutRenderData.SizeX = Width;
-		OutRenderData.SizeY = Height;
-
 		OutRenderData.RawPixels.AddZeroed(Width*Height);
 
-		int32 XMin = BoundingBox.Min.X;
-		int32 YMin = BoundingBox.Min.Y;
+		const int32 XMin = FMath::TruncToInt(BoundingBox.Min.X);
+		const int32 YMin = FMath::TruncToInt(BoundingBox.Min.Y);
 
 		// Compute and copy the pixels for the total filled area of the glyph. 
 
@@ -590,7 +594,7 @@ bool FSlateFontRenderer::GetRenderDataInternal(const FFreeTypeFaceGlyphData& InF
 						if(InOutlineSettings.bSeparateFillAlpha)
 						{
 							// this method is better for transparent fill areas
-							Src = Span.Coverage ? FMath::Abs(Src - Span.Coverage) : 0;
+							Src = Span.Coverage ? FMath::Abs<uint8>(Src - Span.Coverage) : 0;
 						}
 						else
 						{
@@ -672,8 +676,8 @@ bool FSlateFontRenderer::GetRenderDataInternal(const FFreeTypeFaceGlyphData& InF
 			}
 		}
 
-		OutRenderData.SizeX = Bitmap->width;
-		OutRenderData.SizeY = Bitmap->rows;
+		OutRenderData.SizeX = (int16)Bitmap->width;
+		OutRenderData.SizeY = (int16)Bitmap->rows;
 
 		if(Bitmap == &TmpBitmap)
 		{
@@ -686,8 +690,10 @@ bool FSlateFontRenderer::GetRenderDataInternal(const FFreeTypeFaceGlyphData& InF
 		{
 			const int32 ScaledWidth = FMath::TruncToInt(OutRenderData.SizeX * BitmapAtlasScale);
 			const int32 ScaledHeight = FMath::TruncToInt(OutRenderData.SizeY * BitmapAtlasScale);
-			TArray<uint8> ScaledRawPixels;
+			ensure(ScaledWidth <= std::numeric_limits<int16>::max());
+			ensure(ScaledHeight <= std::numeric_limits<int16>::max());
 
+			TArray<uint8> ScaledRawPixels;
 			if (OutRenderData.bIsGrayscale)
 			{
 				check(BytesPerPixel == 1);
@@ -699,8 +705,8 @@ bool FSlateFontRenderer::GetRenderDataInternal(const FFreeTypeFaceGlyphData& InF
 				SlateFontRendererUtils::ResizeFontBitmap<4>(OutRenderData.SizeX, OutRenderData.SizeY, OutRenderData.RawPixels, ScaledWidth, ScaledHeight, ScaledRawPixels);
 			}
 
-			OutRenderData.SizeX = ScaledWidth;
-			OutRenderData.SizeY = ScaledHeight;
+			OutRenderData.SizeX = (int16)ScaledWidth;
+			OutRenderData.SizeY = (int16)ScaledHeight;
 			OutRenderData.RawPixels = MoveTemp(ScaledRawPixels);
 		}
 
@@ -709,8 +715,12 @@ bool FSlateFontRenderer::GetRenderDataInternal(const FFreeTypeFaceGlyphData& InF
 		ScaledOutlineSize = 0.0f;
 	}
 
-	OutRenderData.HorizontalOffset = FMath::RoundToInt(Slot->bitmap_left * BitmapAtlasScale);
-	OutRenderData.VerticalOffset = FMath::RoundToInt(Slot->bitmap_top + ScaledOutlineSize) * BitmapAtlasScale;
+	const int32 HorizontalOffset = FMath::RoundToInt(Slot->bitmap_left * BitmapAtlasScale);
+	const int32 VerticalOffset = FMath::RoundToInt(Slot->bitmap_top + ScaledOutlineSize) * BitmapAtlasScale;
+	ensure(HorizontalOffset <= std::numeric_limits<int16>::max());
+	ensure(VerticalOffset <= std::numeric_limits<int16>::max());
+	OutRenderData.HorizontalOffset = (int16)HorizontalOffset;
+	OutRenderData.VerticalOffset = (int16)VerticalOffset;
 
 	return true;
 }
diff --git a/Engine/Source/Runtime/SlateCore/Private/Styling/StarshipCoreStyle.cpp b/Engine/Source/Runtime/SlateCore/Private/Styling/StarshipCoreStyle.cpp
--- a/Engine/Source/Runtime/SlateCore/Private/Styling/StarshipCoreStyle.cpp
+++ b/Engine/Source/Runtime/SlateCore/Private/Styling/StarshipCoreStyle.cpp
@@ -232,8 +232,8 @@ TSharedRef<ISlateStyle> FStarshipCoreStyle::Create()
 		.SetNormalForeground(  FStyleColors::ForegroundHover)
 		.SetHoveredForeground( FStyleColors::ForegroundHover)
 		.SetPressedForeground( FStyleColors::ForegroundHover)
-		.SetNormalPadding(FMargin(0))
-		.SetPressedPadding(FMargin(0));
+		.SetNormalPadding(FMargin(0.f))
+		.SetPressedPadding(FMargin(0.f));
 
 
 	// Convenient transparent/invisible elements
@@ -252,7 +252,7 @@ TSharedRef<ISlateStyle> FStarshipCoreStyle::Create()
 
 	// Demo Recording
 	{
-		Style->Set("DemoRecording.CursorPing", new IMAGE_BRUSH("Common/CursorPing", FVector2D(31, 31)));
+		Style->Set("DemoRecording.CursorPing", new IMAGE_BRUSH("Common/CursorPing", FVector2D(31.f, 31.f)));
 	}
 
 	// Error Reporting
@@ -276,14 +276,14 @@ TSharedRef<ISlateStyle> FStarshipCoreStyle::Create()
 		// chrome icon or not (we are testing this)
 #if 0
 		Style->Set("AppIcon", new IMAGE_BRUSH_SVG("Starship/Common/unreal", FVector2D(36, 36), FStyleColors::Foreground));
-		Style->Set("AppIconPadding", FMargin(11, 11, 3, 5));
+		Style->Set("AppIconPadding", FMargin(11.f, 11.f, 3.f, 5.f));
 #else
 		Style->Set("AppIcon", new IMAGE_BRUSH("Starship/Common/UELogo", FVector2D(45, 45), FStyleColors::White));
-		Style->Set("AppIconPadding", FMargin(5, 5, 5, 5));
+		Style->Set("AppIconPadding", FMargin(5.f, 5.f, 5.f, 5.f));
 #endif
 
 		Style->Set("AppIcon.Small", new IMAGE_BRUSH_SVG("Starship/Common/unreal-small", Icon24x24, FStyleColors::Foreground));
-		Style->Set("AppIconPadding.Small", FMargin(4, 4, 0, 0));
+		Style->Set("AppIconPadding.Small", FMargin(4.f, 4.f, 0.f, 0.f));
 
 		Style->Set("Checker", new IMAGE_BRUSH("Starship/Common/Checker", Icon16x16, FLinearColor::White, ESlateBrushTileType::Both));
 
@@ -390,8 +390,8 @@ TSharedRef<ISlateStyle> FStarshipCoreStyle::Create()
 	}
 
 	// SVerticalBox Drag& Drop icon
-	Style->Set("VerticalBoxDragIndicator", new IMAGE_BRUSH("Common/VerticalBoxDragIndicator", FVector2D(6, 45)));
-	Style->Set("VerticalBoxDragIndicatorShort", new IMAGE_BRUSH("Common/VerticalBoxDragIndicatorShort", FVector2D(6, 15)));
+	Style->Set("VerticalBoxDragIndicator", new IMAGE_BRUSH("Common/VerticalBoxDragIndicator", FVector2D(6.f, 45.f)));
+	Style->Set("VerticalBoxDragIndicatorShort", new IMAGE_BRUSH("Common/VerticalBoxDragIndicatorShort", FVector2D(6.f, 15.f)));
 
 	SetupButtonStyles(Style);
 	SetupComboButtonStyles(Style);
@@ -433,7 +433,7 @@ TSharedRef<ISlateStyle> FStarshipCoreStyle::Create()
 		Style->Set("Border", new FSlateColorBrush(FStyleColors::Panel));
 
 		FLinearColor TransBackground = FStyleColors::Panel.GetSpecifiedColor();
-		TransBackground.A = .5;
+		TransBackground.A = 0.5f;
 
 		Style->Set("FloatingBorder", new FSlateRoundedBoxBrush(TransBackground, 8.f));
 	}
@@ -441,9 +441,9 @@ TSharedRef<ISlateStyle> FStarshipCoreStyle::Create()
 	// SHyperlink defaults...
 	{
 		FButtonStyle HyperlinkButton = FButtonStyle()
-			.SetNormal(BORDER_BRUSH("Old/HyperlinkDotted", FMargin(0, 0, 0, 3 / 16.0f)))
+			.SetNormal(BORDER_BRUSH("Old/HyperlinkDotted", FMargin(0.f, 0.f, 0.f, 3.f / 16.f)))
 			.SetPressed(FSlateNoResource())
-			.SetHovered(BORDER_BRUSH("Old/HyperlinkUnderline", FMargin(0, 0, 0, 3 / 16.0f)));
+			.SetHovered(BORDER_BRUSH("Old/HyperlinkUnderline", FMargin(0.f, 0.f, 0.f, 3.f / 16.0f)));
 
 		FHyperlinkStyle Hyperlink = FHyperlinkStyle()
 			.SetUnderlineStyle(HyperlinkButton)
@@ -456,8 +456,8 @@ TSharedRef<ISlateStyle> FStarshipCoreStyle::Create()
 	{
 		Style->Set("ProgressBar", FProgressBarStyle()
 			.SetBackgroundImage(FSlateColorBrush(FStyleColors::Foldout))
-			.SetFillImage(IMAGE_BRUSH("Starship/CoreWidgets/ProgressBar/ProgressMarquee", FVector2D(20, 12), FStyleColors::Primary, ESlateBrushTileType::Horizontal))
-			.SetMarqueeImage(IMAGE_BRUSH("Starship/CoreWidgets/ProgressBar/ProgressMarquee", FVector2D(20, 12), FStyleColors::Primary, ESlateBrushTileType::Horizontal))
+			.SetFillImage(IMAGE_BRUSH("Starship/CoreWidgets/ProgressBar/ProgressMarquee", FVector2D(20.f, 12.f), FStyleColors::Primary, ESlateBrushTileType::Horizontal))
+			.SetMarqueeImage(IMAGE_BRUSH("Starship/CoreWidgets/ProgressBar/ProgressMarquee", FVector2D(20.f, 12.f), FStyleColors::Primary, ESlateBrushTileType::Horizontal))
 			.SetEnableFillAnimation(true)
 		);
 	}
@@ -465,8 +465,8 @@ TSharedRef<ISlateStyle> FStarshipCoreStyle::Create()
 
 	// SThrobber, SCircularThrobber defaults...
 	{
-		Style->Set("Throbber.Chunk", new IMAGE_BRUSH("Common/Throbber_Piece", FVector2D(16, 16)));
-		Style->Set("Throbber.CircleChunk", new IMAGE_BRUSH("Common/Throbber_Piece", FVector2D(8, 8)));
+		Style->Set("Throbber.Chunk", new IMAGE_BRUSH("Common/Throbber_Piece", Icon16x16));
+		Style->Set("Throbber.CircleChunk", new IMAGE_BRUSH("Common/Throbber_Piece", Icon8x8));
 	}
 
 	// SExpandableArea defaults...
@@ -476,14 +476,14 @@ TSharedRef<ISlateStyle> FStarshipCoreStyle::Create()
 			.SetExpandedImage(IMAGE_BRUSH_SVG("Starship/Common/chevron-down", Icon16x16, DefaultForeground))
 		);
 		Style->Set("ExpandableArea.TitleFont", StyleFonts.SmallBold);
-		Style->Set("ExpandableArea.Border", new FSlateRoundedBoxBrush(FStyleColors::Panel, 4) );
+		Style->Set("ExpandableArea.Border", new FSlateRoundedBoxBrush(FStyleColors::Panel, 4.f) );
 	}
 
 	// SSlider and SVolumeControl defaults...
 	{
 		FSliderStyle SliderStyle = FSliderStyle()
-			.SetNormalBarImage(   FSlateRoundedBoxBrush(FStyleColors::Input, 2.0, FStyleColors::Input, 1.0))
-			.SetHoveredBarImage(  FSlateRoundedBoxBrush(FStyleColors::Input, 2.0, FStyleColors::Input, 1.0))
+			.SetNormalBarImage(   FSlateRoundedBoxBrush(FStyleColors::Input, 2.0f, FStyleColors::Input, 1.0f))
+			.SetHoveredBarImage(  FSlateRoundedBoxBrush(FStyleColors::Input, 2.0f, FStyleColors::Input, 1.0f))
 			.SetNormalThumbImage(  FSlateRoundedBoxBrush(FStyleColors::Hover2, Icon8x8) )
 			.SetHoveredThumbImage( FSlateRoundedBoxBrush(FStyleColors::ForegroundHover, Icon8x8) )
 			.SetBarThickness(4.0f);
@@ -542,7 +542,7 @@ TSharedRef<ISlateStyle> FStarshipCoreStyle::Create()
 
 		Style->Set("NumericEntrySpinBox.Decorator", new BOX_BRUSH("Common/TextBoxLabelBorder", FMargin(5.0f / 16.0f)));
 
-		Style->Set("NumericEntrySpinBox.NarrowDecorator", new IMAGE_BRUSH_SVG("Starship/CoreWidgets/NumericEntryBox/NarrowDecorator", FVector2D(2,16)));
+		Style->Set("NumericEntrySpinBox.NarrowDecorator", new IMAGE_BRUSH_SVG("Starship/CoreWidgets/NumericEntryBox/NarrowDecorator", FVector2D(2.0f,16.0f)));
 	}
 
 	SetupColorPickerStyles(Style);
@@ -578,8 +578,8 @@ TSharedRef<ISlateStyle> FStarshipCoreStyle::Create()
 			.SetNormal(FSlateNoResource())
 			.SetPressed(BOX_BRUSH("Common/RoundedSelection_16x", 4.0f / 16.0f, SelectionColor_Pressed))
 			.SetHovered(BOX_BRUSH("Common/RoundedSelection_16x", 4.0f / 16.0f, SelectionColor))
-			.SetNormalPadding(FMargin(0, 0))
-			.SetPressedPadding(FMargin(0, 0))
+			.SetNormalPadding(FMargin(0.f, 0.f))
+			.SetPressedPadding(FMargin(0.f, 0.f))
 		);
 	}
 
@@ -587,8 +587,8 @@ TSharedRef<ISlateStyle> FStarshipCoreStyle::Create()
 	{
 		Style->Set("Wizard.PageTitle", FTextBlockStyle(NormalText)
 			.SetFont(DEFAULT_FONT("BoldCondensed", 28))
-			.SetShadowOffset(FVector2D(1, 1))
-			.SetShadowColorAndOpacity(FLinearColor(0, 0, 0, 0.9f))
+			.SetShadowOffset(FVector2D(1.f, 1.f))
+			.SetShadowColorAndOpacity(FLinearColor(0.f, 0.f, 0.f, 0.9f))
 		);
 	}
 
@@ -611,8 +611,8 @@ TSharedRef<ISlateStyle> FStarshipCoreStyle::Create()
 
 	// SHeader defaults...
 	{
-		Style->Set("Header.Pre", new BOX_BRUSH("Common/Separator", FMargin(1 / 4.0f, 0, 2 / 4.0f, 0), FLinearColor(1, 1, 1, 0.5f)));
-		Style->Set("Header.Post", new BOX_BRUSH("Common/Separator", FMargin(2 / 4.0f, 0, 1 / 4.0f, 0), FLinearColor(1, 1, 1, 0.5f)));
+		Style->Set("Header.Pre", new BOX_BRUSH("Common/Separator", FMargin(1.f / 4.0f, 0.f, 2.f / 4.0f, 0.f), FLinearColor(1.f, 1.f, 1.f, 0.5f)));
+		Style->Set("Header.Post", new BOX_BRUSH("Common/Separator", FMargin(2.f / 4.0f, 0.f, 1.f / 4.0f, 0.f), FLinearColor(1.f, 1.f, 1.f, 0.5f)));
 	}
 
 	SetupDockingStyles(Style);
@@ -620,10 +620,10 @@ TSharedRef<ISlateStyle> FStarshipCoreStyle::Create()
 	// SScrollBox defaults...
 	{
 		Style->Set("ScrollBox", FScrollBoxStyle()
-			.SetTopShadowBrush(BOX_BRUSH("Common/ScrollBoxShadowTop", FVector2D(16, 8), FMargin(0.5, 1, 0.5, 0)))
-			.SetBottomShadowBrush(BOX_BRUSH("Common/ScrollBoxShadowBottom", FVector2D(16, 8), FMargin(0.5, 0, 0.5, 1)))
-			.SetLeftShadowBrush(BOX_BRUSH("Common/ScrollBoxShadowLeft", FVector2D(8, 16), FMargin(1, 0.5, 0, 0.5)))
-			.SetRightShadowBrush(BOX_BRUSH("Common/ScrollBoxShadowRight", FVector2D(8, 16), FMargin(0, 0.5, 1, 0.5)))
+			.SetTopShadowBrush(BOX_BRUSH("Common/ScrollBoxShadowTop", FVector2D(16.f, 8.f), FMargin(0.5f, 1.f, 0.5f, 0.f)))
+			.SetBottomShadowBrush(BOX_BRUSH("Common/ScrollBoxShadowBottom", FVector2D(16.f, 8.f), FMargin(0.5f, 0.f, 0.5f, 1.f)))
+			.SetLeftShadowBrush(BOX_BRUSH("Common/ScrollBoxShadowLeft", FVector2D(8.f, 16.f), FMargin(1.f, 0.5f, 0.f, 0.5f)))
+			.SetRightShadowBrush(BOX_BRUSH("Common/ScrollBoxShadowRight", FVector2D(8.f, 16.f), FMargin(0.f, 0.5f, 1.f, 0.5f)))
 			.SetBarThickness(8.0)
 		);
 	}
@@ -631,8 +631,8 @@ TSharedRef<ISlateStyle> FStarshipCoreStyle::Create()
 	// SScrollBorder defaults...
 	{
 		Style->Set("ScrollBorder", FScrollBorderStyle()
-			.SetTopShadowBrush(BOX_BRUSH("Common/ScrollBorderShadowTop", FVector2D(16, 8), FMargin(0.5, 1, 0.5, 0)))
-			.SetBottomShadowBrush(BOX_BRUSH("Common/ScrollBorderShadowBottom", FVector2D(16, 8), FMargin(0.5, 0, 0.5, 1)))
+			.SetTopShadowBrush(BOX_BRUSH("Common/ScrollBorderShadowTop", FVector2D(16.f, 8.f), FMargin(0.5f, 1.f, 0.5f, 0.f)))
+			.SetBottomShadowBrush(BOX_BRUSH("Common/ScrollBorderShadowBottom", FVector2D(16.f, 8.f), FMargin(0.5f, 0.f, 0.5f, 1.f)))
 		);
 	}
 
@@ -675,11 +675,11 @@ TSharedRef<ISlateStyle> FStarshipCoreStyle::Create()
 			.SetInactiveTitleBrush(FSlateNoResource())
 			.SetFlashTitleBrush(IMAGE_BRUSH("Common/Window/WindowTitle_Flashing", Icon24x24, FLinearColor(1.0f, 1.0f, 1.0f, 1.0f), ESlateBrushTileType::Horizontal))
 			.SetBackgroundBrush(FSlateColorBrush(FStyleColors::Recessed))
-			.SetBorderBrush(FSlateRoundedBoxBrush(FStyleColors::Recessed, 2.0, FStyleColors::WindowBorder, 2.0))
-			.SetOutlineBrush(FSlateRoundedBoxBrush(FStyleColors::Recessed, 2.0, FStyleColors::InputOutline, 1.0))
+			.SetBorderBrush(FSlateRoundedBoxBrush(FStyleColors::Recessed, 2.0f, FStyleColors::WindowBorder, 2.0f))
+			.SetOutlineBrush(FSlateRoundedBoxBrush(FStyleColors::Recessed, 2.0f, FStyleColors::InputOutline, 1.0f))
 			.SetChildBackgroundBrush(FSlateColorBrush(FStyleColors::Recessed))
 			.SetCornerRadius(2)
-			.SetBorderPadding(FMargin(3.0, 3.0, 3.0, 3.0));
+			.SetBorderPadding(FMargin(3.0f, 3.0f, 3.0f, 3.0f));
 
 
 		Style->Set("Window", Window);
@@ -944,7 +944,7 @@ void FStarshipCoreStyle::SetupTextStyles(TSharedRef<FStyle>& Style)
 		.SetFont(DEFAULT_FONT("Regular", 24))
 		.SetColorAndOpacity(FLinearColor::Black)
 		.SetShadowOffset(FVector2D(0.0f, 1.0f))
-		.SetShadowColorAndOpacity(FLinearColor(0.8f, 0.8f, 0.8f, 0.5))
+		.SetShadowColorAndOpacity(FLinearColor(0.8f, 0.8f, 0.8f, 0.5f))
 	);
 
 	const FEditableTextBoxStyle DarkEditableTextBoxStyle = FEditableTextBoxStyle()
@@ -973,7 +973,7 @@ void FStarshipCoreStyle::SetupTextStyles(TSharedRef<FStyle>& Style)
 	{
 		FSlateBrush* SelectionBackground = new FSlateColorBrush(FStyleColors::Highlight);
 		FSlateBrush* SelectionTarget = new BOX_BRUSH("Old/DashedBorder", FMargin(6.0f / 32.0f), FLinearColor(0.0f, 0.0f, 0.0f, 0.75f));
-		FSlateBrush* CompositionBackground = new BORDER_BRUSH("Old/HyperlinkDotted", FMargin(0, 0, 0, 3 / 16.0f));
+		FSlateBrush* CompositionBackground = new BORDER_BRUSH("Old/HyperlinkDotted", FMargin(0.f, 0.f, 0.f, 3.f / 16.0f));
 
 
 		const FEditableTextStyle NormalEditableTextStyle = FEditableTextStyle()
@@ -1020,7 +1020,7 @@ void FStarshipCoreStyle::SetupTextStyles(TSharedRef<FStyle>& Style)
 			.SetDownArrowImage(IMAGE_BRUSH_SVG("Starship/Common/arrow-south", Icon8x8, FStyleColors::Foreground))
 			.SetGlassImage(IMAGE_BRUSH_SVG("Starship/Common/search", Icon16x16))
 			.SetClearImage(IMAGE_BRUSH_SVG("Starship/Common/close", Icon16x16))
-			.SetImagePadding(FMargin(3.f, 0.f, -2.f, 0.0))
+			.SetImagePadding(FMargin(3.f, 0.f, -2.f, 0.f))
 			.SetLeftAlignButtons(true)
 		);
 	}
@@ -1094,8 +1094,8 @@ void FStarshipCoreStyle::SetupButtonStyles(TSharedRef<FStyle>& Style)
 		.SetHoveredForeground(FStyleColors::ForegroundHover)
 		.SetPressedForeground(FStyleColors::ForegroundHover)
 		.SetDisabledForeground(FStyleColors::Foreground)
-		.SetNormalPadding(FMargin(2, 2, 2, 2))
-		.SetPressedPadding(FMargin(2, 3, 2, 1));
+		.SetNormalPadding(FMargin(2.f, 2.f, 2.f, 2.f))
+		.SetPressedPadding(FMargin(2.f, 3.f, 2.f, 1.f));
 
 
 	const FButtonStyle SecondaryButton = FButtonStyle()
@@ -1107,8 +1107,8 @@ void FStarshipCoreStyle::SetupButtonStyles(TSharedRef<FStyle>& Style)
 		.SetHoveredForeground(FStyleColors::ForegroundHover)
 		.SetPressedForeground(FStyleColors::ForegroundHover)
 		.SetDisabledForeground(FStyleColors::Foreground)
-		.SetNormalPadding(FMargin(8, 4.5, 8, 3.5))
-		.SetPressedPadding(FMargin(8, 5, 6, 3));
+		.SetNormalPadding(FMargin(8.f, 4.5f, 8.f, 3.5f))
+		.SetPressedPadding(FMargin(8.f, 5.f, 6.f, 3.f));
 	{
 
 		const FTextBlockStyle& NormalText = Style->GetWidgetStyle<FTextBlockStyle>("NormalText");
@@ -1139,8 +1139,8 @@ void FStarshipCoreStyle::SetupButtonStyles(TSharedRef<FStyle>& Style)
 			.SetHoveredForeground(FLinearColor::White)
 			.SetPressedForeground(FLinearColor::White)
 			.SetDisabledForeground(FLinearColor::White)
-			.SetNormalPadding(FMargin(0))
-			.SetPressedPadding(FMargin(0));
+			.SetNormalPadding(FMargin(0.f))
+			.SetPressedPadding(FMargin(0.f));
 
 		Style->Set("InvisibleButton", InvisibleButton);
 	}
@@ -1166,7 +1166,7 @@ void FStarshipCoreStyle::SetupComboButtonStyles(TSharedRef<FStyle>& Style)
 		.SetButtonStyle(ComboButtonButton)
 		.SetContentPadding(0.f)
 		.SetDownArrowImage(IMAGE_BRUSH_SVG("Starship/CoreWidgets/ComboBox/wide-chevron-down", FVector2D(20.f, 16.f)))
-		.SetMenuBorderBrush(FSlateRoundedBoxBrush(FStyleColors::Dropdown, 0.0, WindowHighlight, 1.0))
+		.SetMenuBorderBrush(FSlateRoundedBoxBrush(FStyleColors::Dropdown, 0.0f, WindowHighlight, 1.0f))
 		.SetMenuBorderPadding(0.f);
 	Style->Set("ComboButton", ComboButton);
 
@@ -1183,7 +1183,7 @@ void FStarshipCoreStyle::SetupComboButtonStyles(TSharedRef<FStyle>& Style)
 	FComboButtonStyle SimpleComboButton = FComboButtonStyle()
 		.SetButtonStyle(SimpleButton)
 		.SetContentPadding(0.f)
-		.SetShadowOffset(FVector2D(0.0, 0.0))
+		.SetShadowOffset(FVector2D(0.0f, 0.0f))
 		.SetDownArrowImage(IMAGE_BRUSH_SVG("Starship/CoreWidgets/ComboBox/corner-dropdown", FVector2D(7.0f, 7.0f)))
 		.SetDownArrowPadding(FMargin(0.0f))
 		.SetDownArrowAlignment(EVerticalAlignment::VAlign_Bottom)
@@ -1483,7 +1483,7 @@ void FStarshipCoreStyle::SetupCheckboxStyles(TSharedRef<FStyle>& Style)
 		.SetUndeterminedImage(IMAGE_BRUSH_SVG("/Starship/CoreWidgets/CheckBox/radio-off", Icon16x16, FStyleColors::White25))
 		.SetUndeterminedHoveredImage(IMAGE_BRUSH_SVG("/Starship/CoreWidgets/CheckBox/radio-off", Icon16x16, FStyleColors::ForegroundHover))
 		.SetUndeterminedPressedImage(IMAGE_BRUSH_SVG("/Starship/CoreWidgets/CheckBox/radio-off", Icon16x16, FStyleColors::ForegroundHover))
-		.SetPadding(FMargin(4.0));
+		.SetPadding(FMargin(4.0f));
 	Style->Set("RadioButton", BasicRadioButtonStyle);
 }
 
@@ -1496,7 +1496,7 @@ void FStarshipCoreStyle::SetupDockingStyles(TSharedRef<FStyle>& Style)
 
 	// SDockTab, SDockingTarget, SDockingTabStack defaults...
 	Style->Set("Docking.Background", new BOX_BRUSH("Old/Menu_Background", FMargin(8.0f / 64.0f)));
-	Style->Set("Docking.Border", new FSlateRoundedBoxBrush(FStyleColors::Panel, 4));
+	Style->Set("Docking.Border", new FSlateRoundedBoxBrush(FStyleColors::Panel, 4.f));
 
 	Style->Set("Docking.UnhideTabwellButton", FButtonStyle(Button)
 		.SetNormal(IMAGE_BRUSH_SVG("Starship/Docking/show-tab-well", Icon8x8, FStyleColors::Primary))
@@ -1532,7 +1532,7 @@ void FStarshipCoreStyle::SetupDockingStyles(TSharedRef<FStyle>& Style)
 		.SetTabWellBrush(FSlateColorBrush(FStyleColors::Background))
 		.SetFlashColor(TabFlashColor)
 
-		.SetTabPadding(FMargin(10, 3, 10, 4))
+		.SetTabPadding(FMargin(10.f, 3.f, 10.f, 4.f))
 		.SetOverlapWidth(0.0f)
 
 		.SetNormalForegroundColor(FStyleColors::Foreground)
@@ -1556,7 +1556,7 @@ void FStarshipCoreStyle::SetupDockingStyles(TSharedRef<FStyle>& Style)
 		.SetContentAreaBrush(FSlateColorBrush(FStyleColors::Background))
 		.SetTabWellBrush(FSlateColorBrush(FStyleColors::Background))
 
-		.SetTabPadding(FMargin(10, 7, 10, 8))
+		.SetTabPadding(FMargin(10.f, 7.f, 10.f, 8.f))
 		.SetOverlapWidth(0.f)
 		.SetFlashColor(TabFlashColor)
 
@@ -1595,20 +1595,20 @@ void FStarshipCoreStyle::SetupDockingStyles(TSharedRef<FStyle>& Style)
 
 	Style->Set("Docking.SidebarButton.Opened", SidebarTabButtonOpened);
 
-	Style->Set("Docking.Sidebar.DrawerShadow", new BOX_BRUSH("/Starship/Docking/drawer-shadow", FMargin(8/64.f), FLinearColor(0, 0, 0, 1)));
+	Style->Set("Docking.Sidebar.DrawerShadow", new BOX_BRUSH("/Starship/Docking/drawer-shadow", FMargin(8.f /64.f), FLinearColor(0.f, 0.f, 0.f, 1.f)));
 	Style->Set("Docking.Sidebar.DrawerBackground", new FSlateColorBrush(FStyleColors::Panel));
 	Style->Set("Docking.Sidebar.Background", new FSlateColorBrush(FStyleColors::Recessed));
 	Style->Set("Docking.Sidebar.Border", new FSlateRoundedBoxBrush(FSlateColor(FLinearColor::Transparent), 5.0f, FStyleColors::Hover, 1.0f) );
 
 	// Dock Cross
-	Style->Set("Docking.Cross.DockLeft", new IMAGE_BRUSH("/Docking/OuterDockingIndicator", FVector2D(6, 6), FLinearColor(1.0f, 0.35f, 0.0f, 0.25f)));
-	Style->Set("Docking.Cross.DockLeft_Hovered", new IMAGE_BRUSH("/Docking/OuterDockingIndicator", FVector2D(6, 6), FLinearColor(1.0f, 0.35f, 0.0f)));
-	Style->Set("Docking.Cross.DockTop", new IMAGE_BRUSH("/Docking/OuterDockingIndicator", FVector2D(6, 6), FLinearColor(1.0f, 0.35f, 0.0f, 0.25f)));
-	Style->Set("Docking.Cross.DockTop_Hovered", new IMAGE_BRUSH("/Docking/OuterDockingIndicator", FVector2D(6, 6), FLinearColor(1.0f, 0.35f, 0.0f)));
-	Style->Set("Docking.Cross.DockRight", new IMAGE_BRUSH("/Docking/OuterDockingIndicator", FVector2D(6, 6), FLinearColor(1.0f, 0.35f, 0.0f, 0.25f)));
-	Style->Set("Docking.Cross.DockRight_Hovered", new IMAGE_BRUSH("/Docking/OuterDockingIndicator", FVector2D(6, 6), FLinearColor(1.0f, 0.35f, 0.0f)));
-	Style->Set("Docking.Cross.DockBottom", new IMAGE_BRUSH("/Docking/OuterDockingIndicator", FVector2D(6, 6), FLinearColor(1.0f, 0.35f, 0.0f, 0.25f)));
-	Style->Set("Docking.Cross.DockBottom_Hovered", new IMAGE_BRUSH("/Docking/OuterDockingIndicator", FVector2D(6, 6), FLinearColor(1.0f, 0.35f, 0.0f)));
+	Style->Set("Docking.Cross.DockLeft", new IMAGE_BRUSH("/Docking/OuterDockingIndicator", FVector2D(6.f, 6.f), FLinearColor(1.0f, 0.35f, 0.0f, 0.25f)));
+	Style->Set("Docking.Cross.DockLeft_Hovered", new IMAGE_BRUSH("/Docking/OuterDockingIndicator", FVector2D(6.f, 6.f), FLinearColor(1.0f, 0.35f, 0.0f)));
+	Style->Set("Docking.Cross.DockTop", new IMAGE_BRUSH("/Docking/OuterDockingIndicator", FVector2D(6.f, 6.f), FLinearColor(1.0f, 0.35f, 0.0f, 0.25f)));
+	Style->Set("Docking.Cross.DockTop_Hovered", new IMAGE_BRUSH("/Docking/OuterDockingIndicator", FVector2D(6.f, 6.f), FLinearColor(1.0f, 0.35f, 0.0f)));
+	Style->Set("Docking.Cross.DockRight", new IMAGE_BRUSH("/Docking/OuterDockingIndicator", FVector2D(6.f, 6.f), FLinearColor(1.0f, 0.35f, 0.0f, 0.25f)));
+	Style->Set("Docking.Cross.DockRight_Hovered", new IMAGE_BRUSH("/Docking/OuterDockingIndicator", FVector2D(6.f, 6.f), FLinearColor(1.0f, 0.35f, 0.0f)));
+	Style->Set("Docking.Cross.DockBottom", new IMAGE_BRUSH("/Docking/OuterDockingIndicator", FVector2D(6.f, 6.f), FLinearColor(1.0f, 0.35f, 0.0f, 0.25f)));
+	Style->Set("Docking.Cross.DockBottom_Hovered", new IMAGE_BRUSH("/Docking/OuterDockingIndicator", FVector2D(6.f, 6.f), FLinearColor(1.0f, 0.35f, 0.0f)));
 	Style->Set("Docking.Cross.DockCenter", new IMAGE_BRUSH("/Docking/DockingIndicator_Center", Icon64x64, FLinearColor(1.0f, 0.35f, 0.0f, 0.25f)));
 	Style->Set("Docking.Cross.DockCenter_Hovered", new IMAGE_BRUSH("/Docking/DockingIndicator_Center", Icon64x64, FLinearColor(1.0f, 0.35f, 0.0f)));
 
@@ -1636,7 +1636,7 @@ void FStarshipCoreStyle::SetupColorPickerStyles(TSharedRef<FStyle>& Style)
 		Style->Set("ColorPicker.Font", FStyleFonts::Get().Normal);
 		Style->Set("ColorPicker.Mode", new IMAGE_BRUSH("Common/ColorPicker_Mode_16x", Icon16x16));
 		Style->Set("ColorPicker.Separator", new IMAGE_BRUSH("Common/ColorPicker_Separator", FVector2D(2.0f, 2.0f)));
-		Style->Set("ColorPicker.Selector", new IMAGE_BRUSH("Common/Circle", FVector2D(8, 8)));
+		Style->Set("ColorPicker.Selector", new IMAGE_BRUSH("Common/Circle", Icon8x8));
 		Style->Set("ColorPicker.Slider", FSliderStyle()
 			.SetDisabledThumbImage(IMAGE_BRUSH("Common/ColorPicker_SliderHandle", FVector2D(8.0f, 32.0f)))
 			.SetNormalThumbImage(IMAGE_BRUSH("Common/ColorPicker_SliderHandle", FVector2D(8.0f, 32.0f)))
@@ -1646,8 +1646,8 @@ void FStarshipCoreStyle::SetupColorPickerStyles(TSharedRef<FStyle>& Style)
 
 	// SColorSpectrum defaults...
 	{
-		Style->Set("ColorSpectrum.Spectrum", new IMAGE_BRUSH("Common/ColorSpectrum", FVector2D(256, 256)));
-		Style->Set("ColorSpectrum.Selector", new IMAGE_BRUSH("Common/Circle", FVector2D(8, 8)));
+		Style->Set("ColorSpectrum.Spectrum", new IMAGE_BRUSH("Common/ColorSpectrum", FVector2D(256.f, 256.f)));
+		Style->Set("ColorSpectrum.Selector", new IMAGE_BRUSH("Common/Circle", FVector2D(8.f, 8.f)));
 	}
 
 	// SColorThemes defaults...
@@ -1657,14 +1657,14 @@ void FStarshipCoreStyle::SetupColorPickerStyles(TSharedRef<FStyle>& Style)
 
 	// SColorWheel defaults...
 	{
-		Style->Set("ColorWheel.HueValueCircle", new IMAGE_BRUSH("Common/ColorWheel", FVector2D(192, 192)));
-		Style->Set("ColorWheel.Selector", new IMAGE_BRUSH("Common/Circle", FVector2D(8, 8)));
+		Style->Set("ColorWheel.HueValueCircle", new IMAGE_BRUSH("Common/ColorWheel", FVector2D(192.f, 192.f)));
+		Style->Set("ColorWheel.Selector", new IMAGE_BRUSH("Common/Circle", FVector2D(8.f, 8.f)));
 	}
 
 	// SColorGradingWheel defaults...
 	{
-		Style->Set("ColorGradingWheel.HueValueCircle", new IMAGE_BRUSH("Common/ColorGradingWheel", FVector2D(192, 192)));
-		Style->Set("ColorGradingWheel.Selector", new IMAGE_BRUSH("Common/Circle", FVector2D(8, 8)));
+		Style->Set("ColorGradingWheel.HueValueCircle", new IMAGE_BRUSH("Common/ColorGradingWheel", FVector2D(192.f, 192.f)));
+		Style->Set("ColorGradingWheel.Selector", new IMAGE_BRUSH("Common/Circle", FVector2D(8.f, 8.f)));
 	}
 }
 
@@ -1700,9 +1700,9 @@ void FStarshipCoreStyle::SetupTableViewStyles(TSharedRef<FStyle>& Style)
 		.SetTextColor(FStyleColors::Foreground)
 		.SetSelectedTextColor(FStyleColors::ForegroundInverted)
 
-		.SetDropIndicator_Above(BOX_BRUSH("Common/DropZoneIndicator_Above", FMargin(10.0f / 16.0f, 10.0f / 16.0f, 0, 0), SelectionColor))
+		.SetDropIndicator_Above(BOX_BRUSH("Common/DropZoneIndicator_Above", FMargin(10.0f / 16.0f, 10.0f / 16.0f, 0.f, 0.f), SelectionColor))
 		.SetDropIndicator_Onto(BOX_BRUSH("Common/DropZoneIndicator_Onto", FMargin(4.0f / 16.0f), SelectionColor))
-		.SetDropIndicator_Below(BOX_BRUSH("Common/DropZoneIndicator_Below", FMargin(10.0f / 16.0f, 0, 0, 10.0f / 16.0f), SelectionColor));
+		.SetDropIndicator_Below(BOX_BRUSH("Common/DropZoneIndicator_Below", FMargin(10.0f / 16.0f, 0.f, 0.f, 10.0f / 16.0f), SelectionColor));
 
 	Style->Set("TableView.Row", DefaultTableRowStyle);
 
@@ -1751,7 +1751,7 @@ void FStarshipCoreStyle::SetupTableViewStyles(TSharedRef<FStyle>& Style)
 		.SetColumnStyle(TableColumnHeaderStyle)
 		.SetLastColumnStyle(TableLastColumnHeaderStyle)
 		.SetColumnSplitterStyle(TableHeaderSplitterStyle)
-		.SetSplitterHandleSize(1.0)
+		.SetSplitterHandleSize(1.0f)
 		.SetBackgroundBrush(FSlateColorBrush(FStyleColors::Recessed))
 		.SetForegroundColor(FStyleColors::Foreground)
 		.SetHorizontalSeparatorBrush(FSlateColorBrush(FStyleColors::Recessed))
@@ -1782,7 +1782,7 @@ void FStarshipCoreStyle::SetupMultiboxStyles(TSharedRef<FStyle>& Style)
 		FToolBarStyle NormalToolbarStyle =
 			FToolBarStyle()
 			.SetBackground(FSlateColorBrush(FStyleColors::Panel))
-			.SetBackgroundPadding(FMargin(0, 4, 0, 4))
+			.SetBackgroundPadding(FMargin(0.f, 4.f, 0.f, 4.f))
 			.SetExpandBrush(IMAGE_BRUSH("Icons/toolbar_expand_16x", Icon16x16))
 			.SetComboButtonPadding(FMargin(4.0f, 0.0f))
 			.SetButtonPadding(FMargin(2.0f, 0.f))
@@ -1812,8 +1812,8 @@ void FStarshipCoreStyle::SetupMultiboxStyles(TSharedRef<FStyle>& Style)
 			.SetNormal(FSlateNoResource())
 			.SetPressed(BOX_BRUSH("Common/RoundedSelection_16x", 4.0f / 16.0f, SelectionColor_Pressed))
 			.SetHovered(BOX_BRUSH("Common/RoundedSelection_16x", 4.0f / 16.0f, SelectionColor))
-			.SetNormalPadding(FMargin(2, 2, 2, 2))
-			.SetPressedPadding(FMargin(2, 3, 2, 1))
+			.SetNormalPadding(FMargin(2.f, 2.f, 2.f, 2.f))
+			.SetPressedPadding(FMargin(2.f, 3.f, 2.f, 1.f))
 			.SetNormalForeground(FSlateColor::UseForeground())
 			.SetPressedForeground(FSlateColor::UseForeground())
 			.SetHoveredForeground(FSlateColor::UseForeground())
@@ -1842,7 +1842,7 @@ void FStarshipCoreStyle::SetupMultiboxStyles(TSharedRef<FStyle>& Style)
 		FToolBarStyle SlimToolbarStyle =
 			FToolBarStyle()
 			.SetBackground(FSlateColorBrush(FStyleColors::Panel))
-			.SetBackgroundPadding(FMargin(0,4,0,4))
+			.SetBackgroundPadding(FMargin(0.f, 4.f, 0.f, 4.f))
 			.SetExpandBrush(IMAGE_BRUSH("Icons/toolbar_expand_16x", Icon16x16))
 			.SetComboButtonPadding(FMargin(6.0f, 0.0f))
 			.SetButtonPadding(FMargin(4.0f, 4.0f))
@@ -1851,7 +1851,7 @@ void FStarshipCoreStyle::SetupMultiboxStyles(TSharedRef<FStyle>& Style)
 			.SetSeparatorPadding(FMargin(8.f, -4.0f))
 			.SetLabelStyle(FTextBlockStyle(NormalText))
 			.SetComboButtonStyle(SlimToolBarComboButton)
-			.SetLabelPadding(FMargin(5, 0, 0, 0))
+			.SetLabelPadding(FMargin(5.f, 0.f, 0.f, 0.f))
 			.SetEditableTextStyle(FEditableTextBoxStyle(NormalEditableTextBoxStyle));
 
 		FCheckBoxStyle SlimToolBarToggleButtonCheckBoxStyle = FCheckBoxStyle()
diff --git a/Engine/Source/Runtime/SlateCore/Public/Fonts/FontCache.h b/Engine/Source/Runtime/SlateCore/Public/Fonts/FontCache.h
--- a/Engine/Source/Runtime/SlateCore/Public/Fonts/FontCache.h
+++ b/Engine/Source/Runtime/SlateCore/Public/Fonts/FontCache.h
@@ -417,7 +417,7 @@ private:
 			return (GlyphDataArray.IsValidIndex(InternalIndex)) ? &GlyphDataArray[InternalIndex] : nullptr;
 		}
 
-		FORCEINLINE uint32 GetAllocatedSize() const
+		FORCEINLINE SIZE_T GetAllocatedSize() const
 		{
 			return GlyphDataArray.GetAllocatedSize();
 		}
