Slate: Fix incorrect use of array index when fixng the itterator slate attribute.

diff --git a/Engine/Source/Runtime/SlateCore/Private/FastUpdate/SlateInvalidationRoot.cpp b/Engine/Source/Runtime/SlateCore/Private/FastUpdate/SlateInvalidationRoot.cpp
--- a/Engine/Source/Runtime/SlateCore/Private/FastUpdate/SlateInvalidationRoot.cpp
+++ b/Engine/Source/Runtime/SlateCore/Private/FastUpdate/SlateInvalidationRoot.cpp
@@ -139,6 +139,32 @@ FAutoConsoleVariableRef CVarSlateInvalidationWidgetListNumElementLeftBeforeSplit
 	GSlateInvalidationWidgetListNumberElementLeftBeforeSplitting,
 	TEXT("With Global Invalidation, when splitting, only split the array when the number of element left is under X."));
 
+/**
+ *
+ */
+ namespace Slate
+ {
+	 UE_NODISCARD EInvalidateWidgetReason EInvalidateWidgetReason_RemovePreUpdate(EInvalidateWidgetReason InvalidateReason)
+	{
+		static_assert(std::is_same<std::underlying_type_t<EInvalidateWidgetReason>, uint8>::value, "EInvalidateWidgetReason is not a uint8");
+		const uint8 AnyPostUpdate = (0xFF & ~(uint8)EInvalidateWidgetReason::AttributeRegistration);
+		return (EInvalidateWidgetReason)(((uint8)InvalidateReason) & AnyPostUpdate);
+	}
+
+	bool EInvalidateWidgetReason_HasPreUpdateFlag(EInvalidateWidgetReason InvalidateReason)
+	{
+		return EnumHasAnyFlags(InvalidateReason, EInvalidateWidgetReason::AttributeRegistration | EInvalidateWidgetReason::ChildOrder);
+	}
+
+	bool EInvalidateWidgetReason_HasPostUpdateFlag(EInvalidateWidgetReason InvalidateReason)
+	{
+		static_assert(std::is_same<std::underlying_type_t<EInvalidateWidgetReason>, uint8>::value, "EInvalidateWidgetReason is not a uint8");
+		const uint8 AnyPostUpdate = (0xFF & ~(uint8)EInvalidateWidgetReason::AttributeRegistration);
+		return (((uint8)InvalidateReason & AnyPostUpdate) != 0);
+	}
+ }
+
+
 /**
  *
  */
@@ -265,13 +291,12 @@ void FSlateInvalidationRoot::InvalidateWidget(FWidgetProxy& Proxy, EInvalidateWi
 				}
 			}
 		}
-		else if (EnumHasAnyFlags(InvalidateReason, EInvalidateWidgetReason::AttributeRegistration))
+		else if (Slate::EInvalidateWidgetReason_HasPreUpdateFlag(InvalidateReason))
 		{
 			WidgetsNeedingPreUpdate->HeapPushUnique(Proxy);
 		}
 
-		const uint8 AnyPostUpdate = (0xFF & ~(uint8)EInvalidateWidgetReason::AttributeRegistration);
-		if (((uint8)InvalidateReason & AnyPostUpdate) != 0)
+		if (Slate::EInvalidateWidgetReason_HasPostUpdateFlag(InvalidateReason))
 		{
 			WidgetsNeedingPostUpdate->PushBackOrHeapUnique(Proxy);
 		}
@@ -654,10 +679,12 @@ void FSlateInvalidationRoot::ProcessPreUpdate()
 						if (EnumHasAnyFlags(InvalidationWidget.CurrentInvalidateReason, EInvalidateWidgetReason::AttributeRegistration) && bIsInvalidationWidgetValid)
 						{
 							FastWidgetPathList->ProcessAttributeRegistrationInvalidation(InvalidationWidget);
-							AttributeItt.Seek(WidgetIndex);
+							AttributeItt.Seek(InvalidationWidget.Index);
 						}
 
 						AttributeItt.FixCurrentWidgetIndex();
+
+						InvalidationWidget.CurrentInvalidateReason = Slate::EInvalidateWidgetReason_RemovePreUpdate(InvalidationWidget.CurrentInvalidateReason);
 					}
 				}
 				else
@@ -666,7 +693,7 @@ void FSlateInvalidationRoot::ProcessPreUpdate()
 					SWidget* WidgetPtr = InvalidationWidget.GetWidget();
 					if (ensureMsgf(WidgetPtr, TEXT("Child order invalidation should have been called before we process this widget.")))
 					{
-						if (WidgetPtr->IsFastPathVisible())
+						auto UpdateAttribute = [&InvalidationWidget, WidgetPtr]()
 						{
 #if UE_SLATE_WITH_INVALIDATIONWIDGETLIST_DEBUGGING
 							if (GSlateInvalidationRootVerifySlateAttribute)
@@ -676,6 +703,24 @@ void FSlateInvalidationRoot::ProcessPreUpdate()
 #endif
 							FSlateAttributeMetaData::UpdateAttributes(*WidgetPtr);
 							InvalidationWidget.bDebug_AttributeUpdated = true;
+						};
+
+						if (WidgetPtr->IsFastPathVisible())
+						{
+							UpdateAttribute();
+						}
+						else if (InvalidationWidget.ParentIndex != FSlateInvalidationWidgetIndex::Invalid)
+						{
+							// Is my parent visible, should I update the collapsed attributes
+							const FSlateInvalidationWidgetList::InvalidationWidgetType& ParentInvalidationWidget = (*FastWidgetPathList)[InvalidationWidget.ParentIndex];
+							if (ParentInvalidationWidget.Visibility.IsVisible() && ensure(ParentInvalidationWidget.GetWidget()) && ParentInvalidationWidget.GetWidget()->IsFastPathVisible())
+							{
+								FSlateAttributeMetaData::UpdateCollapsedAttributes(*WidgetPtr);
+								if (WidgetPtr->IsFastPathVisible())
+								{
+									UpdateAttribute();
+								}
+							}
 						}
 					}
 					AttributeItt.Advance();
diff --git a/Engine/Source/Runtime/SlateCore/Private/FastUpdate/SlateInvalidationWidgetList.cpp b/Engine/Source/Runtime/SlateCore/Private/FastUpdate/SlateInvalidationWidgetList.cpp
--- a/Engine/Source/Runtime/SlateCore/Private/FastUpdate/SlateInvalidationWidgetList.cpp
+++ b/Engine/Source/Runtime/SlateCore/Private/FastUpdate/SlateInvalidationWidgetList.cpp
@@ -118,7 +118,8 @@ void FSlateInvalidationWidgetList::FWidgetAttributeIterator::FixCurrentWidgetInd
 		AttributeIndex = ArrayNode.ElementIndexList_WidgetWithRegisteredSlateAttribute.FindLowerBound(MoveToWidgetIndexOnNextAdvance.ElementIndex);
 		if (AttributeIndex == INDEX_NONE)
 		{
-			AdvanceArrayIndex(ArrayIndex);
+			int32 NextArrayIndex = ArrayNode.NextArrayIndex;
+			AdvanceArrayIndex(NextArrayIndex);
 		}
 		else
 		{
@@ -606,7 +607,7 @@ FSlateInvalidationWidgetIndex FSlateInvalidationWidgetList::DecrementIndex(FSlat
 		check(Data[Index.ArrayIndex].PreviousArrayIndex < FSlateInvalidationWidgetIndex::Invalid.ArrayIndex);
 		Index.ArrayIndex = (IndexType)Data[Index.ArrayIndex].PreviousArrayIndex;
 		check(Data[Index.ArrayIndex].ElementList.Num() > 0);
-		Index.ElementIndex = Data[Index.ArrayIndex].ElementList.Num() - 1;
+		Index.ElementIndex = (IndexType)(Data[Index.ArrayIndex].ElementList.Num() - 1);
 	}
 	else
 	{
@@ -641,8 +642,8 @@ FSlateInvalidationWidgetList::IndexType FSlateInvalidationWidgetList::AddArrayNo
 	{
 		if (Data.Num() + 1 == FSlateInvalidationWidgetIndex::Invalid.ArrayIndex)
 		{
-			ensure(false);
-			return LastArrayIndex;
+			ensureAlwaysMsgf(false, TEXT("The widget array is split more time that we support. Widget will not be updated properly. Try to increase Slate.InvalidationList.MaxArrayElements"));
+			return (IndexType)LastArrayIndex;
 		}
 		const int32 Index = Data.Add(FArrayNode());
 		check(Index < std::numeric_limits<IndexType>::max());
@@ -664,7 +665,7 @@ FSlateInvalidationWidgetList::IndexType FSlateInvalidationWidgetList::AddArrayNo
 			FirstArrayIndex = LastArrayIndex;
 		}
 	}
-	return LastArrayIndex;
+	return (IndexType)LastArrayIndex;
 }
 
 
@@ -809,8 +810,8 @@ FSlateInvalidationWidgetList::IndexType FSlateInvalidationWidgetList::InsertData
 
 		if (Data.Num() + 1 == FSlateInvalidationWidgetIndex::Invalid.ArrayIndex)
 		{
-			ensure(false);
-			return LastArrayIndex;
+			ensureAlwaysMsgf(false, TEXT("The widget array is split more time that we support. Widget will not be updated properly. Try to increase Slate.InvalidationList.MaxArrayElements"));
+			return (IndexType)LastArrayIndex;
 		}
 
 		const int32 NewIndex = Data.Add(FArrayNode());
@@ -834,7 +835,7 @@ FSlateInvalidationWidgetList::IndexType FSlateInvalidationWidgetList::InsertData
 				LastArrayIndex = NewIndex;
 			}
 
-			RebuildOrderIndex(NewIndex);
+			RebuildOrderIndex((IndexType)NewIndex);
 		}
 		else
 		{
@@ -979,7 +980,7 @@ void FSlateInvalidationWidgetList::_RemoveRangeFromSameParent(const FIndexRange
 			}
 #endif //UE_SLATE_WITH_WIDGETLIST_ASSIGNINVALIDPROXYWHENREMOVED
 		};
-		auto ResetInvalidationWidget = [Self](IndexType ArrayIndex, int32 StartIndex, int32 Num)
+		auto ResetInvalidationWidget = [Self](IndexType ArrayIndex, IndexType StartIndex, int32 Num)
 		{
 			FArrayNode& ArrayNode = Self->Data[ArrayIndex];
 			ElementListType& ResetElementList = ArrayNode.ElementList;
@@ -990,7 +991,7 @@ void FSlateInvalidationWidgetList::_RemoveRangeFromSameParent(const FIndexRange
 			ArrayNode.RemoveElementIndexBiggerOrEqualThan(StartIndex);
 		};
 
-		auto RemoveDataNodeIfNeeded = [Self](int32 ArrayIndex) -> bool
+		auto RemoveDataNodeIfNeeded = [Self](IndexType ArrayIndex) -> bool
 		{
 			if (Self->Data[ArrayIndex].StartIndex >= Self->Data[ArrayIndex].ElementList.Num())
 			{
@@ -1054,7 +1055,7 @@ void FSlateInvalidationWidgetList::_RemoveRangeFromSameParent(const FIndexRange
 				SetFakeInvalidatWidgetHandle(Range.GetInclusiveMinWidgetIndex().ArrayIndex, Range.GetInclusiveMinWidgetIndex().ElementIndex, RemoveElementList.Num());
 			}
 
-			const int32 RemoveArrayAt = Range.GetInclusiveMinWidgetIndex().ElementIndex;
+			const IndexType RemoveArrayAt = Range.GetInclusiveMinWidgetIndex().ElementIndex;
 			RemoveElementList.RemoveAt(RemoveArrayAt, RemoveElementList.Num() - RemoveArrayAt, true);
 			if (!RemoveDataNodeIfNeeded(Range.GetInclusiveMinWidgetIndex().ArrayIndex))
 			{
@@ -1092,7 +1093,7 @@ FSlateInvalidationWidgetList::FCutResult FSlateInvalidationWidgetList::CutArray(
 		}
 		else
 		{
-			ArrayNode.RemoveElementIndexBiggerOrEqualThan(CutResult.OldElementIndexStart);
+			ArrayNode.RemoveElementIndexBiggerOrEqualThan((IndexType)CutResult.OldElementIndexStart);
 		}
 	}
 
@@ -1115,7 +1116,7 @@ FSlateInvalidationWidgetList::FCutResult FSlateInvalidationWidgetList::_CutArray
 		//From where to where we are moving the item
 		const IndexType OldArrayIndex = WhereToCut.ArrayIndex;
 		const IndexType OldElementIndexStart = WhereToCut.ElementIndex + 1;
-		const IndexType OldElementIndexEnd = Data[WhereToCut.ArrayIndex].ElementList.Num();
+		const IndexType OldElementIndexEnd = (IndexType)Data[WhereToCut.ArrayIndex].ElementList.Num();
 		const IndexType NewArrayIndex = InsertDataNodeAfter(OldArrayIndex, false);
 		const int32 NewExpectedElementArraySize = OldElementIndexEnd - OldElementIndexStart;
 		Data[NewArrayIndex].ElementList.Reserve(NewExpectedElementArraySize);
@@ -1523,7 +1524,7 @@ bool FSlateInvalidationWidgetList::VerifyElementIndexList() const
 
 		for (int32 ElementIndex = ArrayNode.StartIndex; ElementIndex < ArrayNode.ElementList.Num(); ++ElementIndex)
 		{
-			if (!ArrayNode.ElementIndexList_WidgetWithRegisteredSlateAttribute.Contains(ElementIndex))
+			if (!ArrayNode.ElementIndexList_WidgetWithRegisteredSlateAttribute.Contains((IndexType)ElementIndex))
 			{
 				SWidget* Widget = ArrayNode.ElementList[ElementIndex].GetWidget();
 				if (Widget->HasRegisteredSlateAttribute())
