Fix a wrong assumption made by the Static Analysis.
#rnx

diff --git a/Engine/Source/Runtime/SlateCore/Private/Test/SlateAttributeTest.cpp b/Engine/Source/Runtime/SlateCore/Private/Test/SlateAttributeTest.cpp
--- a/Engine/Source/Runtime/SlateCore/Private/Test/SlateAttributeTest.cpp
+++ b/Engine/Source/Runtime/SlateCore/Private/Test/SlateAttributeTest.cpp
@@ -305,8 +305,7 @@ bool FSlateAttributeTest::RunTest(const FString& Parameters)
 			AddErrorIfFalse(bWasUpdate, TEXT("C should be updated becaude D was."));
 			AddErrorIfFalse(WidgetParent->IntAttributeC.Get() == 10, TEXT("C It is not the expected value."));
 			AddErrorIfFalse(WidgetParent->IntAttributeD.Get() == 8, TEXT("D It is not the expected value."));
-			AddErrorIfFalse(OrderCounter == 2, TEXT("There is no D attribute anymore."));
-			AddErrorIfFalse(OrderCounter == 2, TEXT("There is no D attribute anymore."));
+			AddErrorIfFalse(OrderCounter == 2, TEXT("There is no D attribute anymore.")); //-V547
 		}
 
 		{
