Slate: Fix implicit conversion. Change FSlateElementBatch::SceneIndex to a int8 (the same used in DrawElement)

diff --git a/Engine/Source/Runtime/SlateCore/Private/Rendering/DrawElements.cpp b/Engine/Source/Runtime/SlateCore/Private/Rendering/DrawElements.cpp
--- a/Engine/Source/Runtime/SlateCore/Private/Rendering/DrawElements.cpp
+++ b/Engine/Source/Runtime/SlateCore/Private/Rendering/DrawElements.cpp
@@ -11,9 +11,7 @@
 #include "Debugging/SlateDebugging.h"
 #include "Application/SlateApplicationBase.h"
 
-#if UE_SLATE_VERIFY_PIXELSIZE
 #include <limits>
-#endif
 
 DECLARE_CYCLE_STAT(TEXT("FSlateDrawElement::Make Time"), STAT_SlateDrawElementMakeTime, STATGROUP_SlateVerbose);
 DECLARE_CYCLE_STAT(TEXT("FSlateDrawElement::MakeCustomVerts Time"), STAT_SlateDrawElementMakeCustomVertsTime, STATGROUP_Slate);
@@ -170,12 +168,16 @@ void FSlateDrawElement::Init(FSlateWindowElementList& ElementList, EElementType
 	// Calculate the layout to render transform as this is needed by several calculations downstream.
 	const FSlateLayoutTransform InverseLayoutTransform(Inverse(FSlateLayoutTransform(Scale, Position)));
 
-	// This is a workaround because we want to keep track of the various Scenes 
-	// in use throughout the UI. We keep a synchronized set with the render thread on the SlateRenderer and 
-	// use indices to synchronize between them.
-	FSlateRenderer* Renderer = FSlateApplicationBase::Get().GetRenderer();
-	checkSlow(Renderer);
-	SceneIndex = Renderer->GetCurrentSceneIndex();
+	{
+		// This is a workaround because we want to keep track of the various Scenes 
+		// in use throughout the UI. We keep a synchronized set with the render thread on the SlateRenderer and 
+		// use indices to synchronize between them.
+		FSlateRenderer* Renderer = FSlateApplicationBase::Get().GetRenderer();
+		checkSlow(Renderer);
+		int32 SceneIndexLong = Renderer->GetCurrentSceneIndex();
+		ensureMsgf(SceneIndexLong <= std::numeric_limits<int8>::max(), TEXT("The Scene index is saved as an int8 in the DrawElements."));
+		SceneIndex = (int8)SceneIndexLong;
+	}
 
 	BatchFlags = ESlateBatchDrawFlag::None;
 	BatchFlags |= static_cast<ESlateBatchDrawFlag>(static_cast<uint32>(InDrawEffects) & static_cast<uint32>(ESlateDrawEffect::NoBlending | ESlateDrawEffect::PreMultipliedAlpha | ESlateDrawEffect::NoGamma | ESlateDrawEffect::InvertAlpha));
diff --git a/Engine/Source/Runtime/SlateCore/Private/Rendering/ElementBatcher.cpp b/Engine/Source/Runtime/SlateCore/Private/Rendering/ElementBatcher.cpp
--- a/Engine/Source/Runtime/SlateCore/Private/Rendering/ElementBatcher.cpp
+++ b/Engine/Source/Runtime/SlateCore/Private/Rendering/ElementBatcher.cpp
@@ -1016,12 +1016,10 @@ void FSlateElementBatcher::AddTextElement(const FSlateDrawElement& DrawElement)
 	SlateElementBatcher::CheckUObject(DrawElementPayload, OutlineFontMaterial);
 #endif
 
-	bool bOutlineFont = OutlineSettings.OutlineSize > 0.0f;
+	bool bOutlineFont = OutlineSettings.OutlineSize > 0;
+	const int32 OutlineSize = OutlineSettings.OutlineSize;
 
-	const float OutlineSize = OutlineSettings.OutlineSize;
-
-
-	auto BuildFontGeometry = [&](const FFontOutlineSettings& InOutlineSettings, const FColor& InTint, const UObject* FontMaterial, int32 InLayer, int32 InOutlineHorizontalOffset)
+	auto BuildFontGeometry = [&](const FFontOutlineSettings& InOutlineSettings, const FColor& InTint, const UObject* FontMaterial, int32 InLayer, float InOutlineHorizontalOffset)
 	{
 		FCharacterList& CharacterList = FontCache.GetCharacterList(DrawElementPayload.GetFontInfo(), FontScale, InOutlineSettings);
 
@@ -1062,7 +1060,7 @@ void FSlateElementBatcher::AddTextElement(const FSlateDrawElement& DrawElement)
 		LineX = PosX;
 
 		const bool bIsFontMaterial = FontMaterial != nullptr;
-		const bool bEnableOutline = InOutlineSettings.OutlineSize > 0.0f;
+		const bool bEnableOutline = InOutlineSettings.OutlineSize > 0;
 
 		uint32 NumChars = Len;
 
@@ -1192,10 +1190,10 @@ void FSlateElementBatcher::AddTextElement(const FSlateDrawElement& DrawElement)
 	if (bOutlineFont)
 	{
 		// Build geometry for the outline
-		BuildFontGeometry(OutlineSettings, PackVertexColor(OutlineSettings.OutlineColor), OutlineFontMaterial, Layer, 0);
+		BuildFontGeometry(OutlineSettings, PackVertexColor(OutlineSettings.OutlineColor), OutlineFontMaterial, Layer, 0.f);
 
 		//The fill area was measured without an outline so it must be shifted by the scaled outline size
-		float HorizontalOffset = FMath::RoundToFloat(OutlineSize * FontScale);
+		const float HorizontalOffset = FMath::RoundToFloat((float)OutlineSize * FontScale);
 
 		// Build geometry for the base font which is always rendered on top of the outline
 		BuildFontGeometry(FFontOutlineSettings::NoOutline, BaseTint, BaseFontMaterial, Layer + 1, HorizontalOffset);
@@ -1203,7 +1201,7 @@ void FSlateElementBatcher::AddTextElement(const FSlateDrawElement& DrawElement)
 	else
 	{
 		// No outline, draw normally
-		BuildFontGeometry(FFontOutlineSettings::NoOutline, BaseTint, BaseFontMaterial, Layer, 0);
+		BuildFontGeometry(FFontOutlineSettings::NoOutline, BaseTint, BaseFontMaterial, Layer, 0.f);
 	}
 }
 
@@ -1250,11 +1248,10 @@ void FSlateElementBatcher::AddShapedTextElement( const FSlateDrawElement& DrawEl
 	const UObject* BaseFontMaterial = ShapedGlyphSequence->GetFontMaterial();
 	const UObject* OutlineFontMaterial = OutlineSettings.OutlineMaterial;
 
-	bool bOutlineFont = OutlineSettings.OutlineSize > 0.0f;
-
-	const float OutlineSize = OutlineSettings.OutlineSize;
+	bool bOutlineFont = OutlineSettings.OutlineSize > 0;
+	const int32 OutlineSize = OutlineSettings.OutlineSize;
 
-	auto BuildFontGeometry = [&](const FFontOutlineSettings& InOutlineSettings, const FColor& InTint, const UObject* FontMaterial, int32 InLayer, int32 InHorizontalOffset)
+	auto BuildFontGeometry = [&](const FFontOutlineSettings& InOutlineSettings, const FColor& InTint, const UObject* FontMaterial, int32 InLayer, float InHorizontalOffset)
 	{
 		FVector2D TopLeft(0, 0);
 
@@ -1278,7 +1275,7 @@ void FSlateElementBatcher::AddShapedTextElement( const FSlateDrawElement& DrawEl
 		float InvTextureSizeY = 0;
 
 		const bool bIsFontMaterial = FontMaterial != nullptr;
-		const bool bEnableOutline = InOutlineSettings.OutlineSize > 0.0f;
+		const bool bEnableOutline = InOutlineSettings.OutlineSize > 0;
 
 		// Optimize by culling
 		// Todo: this doesnt work with cached clipping
@@ -1414,10 +1411,10 @@ void FSlateElementBatcher::AddShapedTextElement( const FSlateDrawElement& DrawEl
 	if (bOutlineFont)
 	{
 		// Build geometry for the outline
-		BuildFontGeometry(OutlineSettings, PackVertexColor(DrawElementPayload.GetOutlineTint()), OutlineFontMaterial, Layer, 0);
+		BuildFontGeometry(OutlineSettings, PackVertexColor(DrawElementPayload.GetOutlineTint()), OutlineFontMaterial, Layer, 0.f);
 		
 		//The fill area was measured without an outline so it must be shifted by the scaled outline size
-		float HorizontalOffset = FMath::RoundToFloat(OutlineSize * FontScale);
+		const float HorizontalOffset = FMath::RoundToFloat((float)OutlineSize * FontScale);
 
 		// Build geometry for the base font which is always rendered on top of the outline 
 		BuildFontGeometry(FFontOutlineSettings::NoOutline, BaseTint, BaseFontMaterial, Layer+1, HorizontalOffset);
@@ -1425,7 +1422,7 @@ void FSlateElementBatcher::AddShapedTextElement( const FSlateDrawElement& DrawEl
 	else
 	{
 		// No outline
-		BuildFontGeometry(FFontOutlineSettings::NoOutline, BaseTint, BaseFontMaterial, Layer, 0);
+		BuildFontGeometry(FFontOutlineSettings::NoOutline, BaseTint, BaseFontMaterial, Layer, 0.f);
 	}
 }
 
@@ -1651,11 +1648,11 @@ struct FLineBuilder
 	{
 		const uint32 FirstVertIndex = RenderBatch.GetNumVertices();
 
-		RenderBatch.AddVertex(FSlateVertex::Make<ESlateVertexRounding::Disabled>(RenderTransform, CapOrigin, FVector2D(0.5, 0.0f), FVector2D::ZeroVector, Color));
-		RenderBatch.AddVertex(FSlateVertex::Make<ESlateVertexRounding::Disabled>(RenderTransform, CapOrigin + CapDirection + Up, FVector2D(0.0, 0.0f), FVector2D::ZeroVector, Color));
-		RenderBatch.AddVertex(FSlateVertex::Make<ESlateVertexRounding::Disabled>(RenderTransform, CapOrigin + CapDirection - Up, FVector2D(0.0, 0.0f), FVector2D::ZeroVector, Color));
-		RenderBatch.AddVertex(FSlateVertex::Make<ESlateVertexRounding::Disabled>(RenderTransform, CapOrigin + Up, FVector2D(0.0, 0.0f), FVector2D::ZeroVector, Color));
-		RenderBatch.AddVertex(FSlateVertex::Make<ESlateVertexRounding::Disabled>(RenderTransform, CapOrigin - Up, FVector2D(0.0, 0.0f), FVector2D::ZeroVector, Color));
+		RenderBatch.AddVertex(FSlateVertex::Make<ESlateVertexRounding::Disabled>(RenderTransform, CapOrigin, FVector2D(0.5f, 0.0f), FVector2D::ZeroVector, Color));
+		RenderBatch.AddVertex(FSlateVertex::Make<ESlateVertexRounding::Disabled>(RenderTransform, CapOrigin + CapDirection + Up, FVector2D(0.0f, 0.0f), FVector2D::ZeroVector, Color));
+		RenderBatch.AddVertex(FSlateVertex::Make<ESlateVertexRounding::Disabled>(RenderTransform, CapOrigin + CapDirection - Up, FVector2D(0.0f, 0.0f), FVector2D::ZeroVector, Color));
+		RenderBatch.AddVertex(FSlateVertex::Make<ESlateVertexRounding::Disabled>(RenderTransform, CapOrigin + Up, FVector2D(0.0f, 0.0f), FVector2D::ZeroVector, Color));
+		RenderBatch.AddVertex(FSlateVertex::Make<ESlateVertexRounding::Disabled>(RenderTransform, CapOrigin - Up, FVector2D(0.0f, 0.0f), FVector2D::ZeroVector, Color));
 
 		RenderBatch.AddIndex(FirstVertIndex + 0);
 		RenderBatch.AddIndex(FirstVertIndex + 3);
@@ -2285,7 +2282,7 @@ void FSlateElementBatcher::AddBorderElement( const FSlateDrawElement& DrawElemen
 		TextureWidth = Resource->GetWidth();
 		TextureHeight = Resource->GetHeight();
 	}
-	FVector2D TextureSizeLocalSpace = TransformVector(DrawElement.GetInverseLayoutTransform(), FVector2D(TextureWidth, TextureHeight));
+	FVector2D TextureSizeLocalSpace = TransformVector(DrawElement.GetInverseLayoutTransform(), FVector2D((float)TextureWidth, (float)TextureHeight));
  
 	// Texel offset
 	const FVector2D HalfTexel( PixelCenterOffset/TextureWidth, PixelCenterOffset/TextureHeight );
@@ -2595,7 +2592,7 @@ void FSlateElementBatcher::AddPostProcessPass(const FSlateDrawElement& DrawEleme
 	// These could be negative with rotation or negative scales.  This is not supported yet
 	if(SizeUV.X > 0 && SizeUV.Y > 0)
 	{
-		FShaderParams Params = FShaderParams::MakePixelShaderParams(FVector4(WorldTopLeft, WorldBotRight), FVector4(Payload.PostProcessData.X, Payload.PostProcessData.Y, Payload.DownsampleAmount, 0));
+		FShaderParams Params = FShaderParams::MakePixelShaderParams(FVector4(WorldTopLeft, WorldBotRight), FVector4(Payload.PostProcessData.X, Payload.PostProcessData.Y, (float)Payload.DownsampleAmount, 0.f));
 
 		CreateRenderBatch(Layer, Params, nullptr, ESlateDrawPrimitive::TriangleList, ESlateShader::PostProcess, ESlateDrawEffect::None, ESlateBatchDrawFlag::None, DrawElement);
 	}
diff --git a/Engine/Source/Runtime/SlateCore/Public/Rendering/ElementBatcher.h b/Engine/Source/Runtime/SlateCore/Public/Rendering/ElementBatcher.h
--- a/Engine/Source/Runtime/SlateCore/Public/Rendering/ElementBatcher.h
+++ b/Engine/Source/Runtime/SlateCore/Public/Rendering/ElementBatcher.h
@@ -59,7 +59,7 @@ public:
 	int32 GetInstanceCount() const { return BatchKey.InstanceCount; }
 	uint32 GetInstanceOffset() const { return BatchKey.InstanceOffset; }
 	const ISlateUpdatableInstanceBuffer* GetInstanceData() const { return BatchKey.InstanceData; }
-	int32 GetSceneIndex() const { return BatchKey.SceneIndex; }
+	int8 GetSceneIndex() const { return BatchKey.SceneIndex; }
 private:
 	struct FBatchKey
 	{
@@ -73,9 +73,9 @@ private:
 		const int32 InstanceCount;
 		const uint32 InstanceOffset;
 		const ISlateUpdatableInstanceBuffer* InstanceData;
-		const int32 SceneIndex;
+		const int8 SceneIndex;
 
-		FBatchKey(const FShaderParams& InShaderParams, ESlateShader InShaderType, ESlateDrawPrimitive InDrawPrimitiveType, ESlateDrawEffect InDrawEffects, ESlateBatchDrawFlag InDrawFlags, const FClipStateHandle InClipStateHandle, int32 InInstanceCount, uint32 InInstanceOffset, ISlateUpdatableInstanceBuffer* InInstanceBuffer, int32 InSceneIndex)
+		FBatchKey(const FShaderParams& InShaderParams, ESlateShader InShaderType, ESlateDrawPrimitive InDrawPrimitiveType, ESlateDrawEffect InDrawEffects, ESlateBatchDrawFlag InDrawFlags, const FClipStateHandle InClipStateHandle, int32 InInstanceCount, uint32 InInstanceOffset, ISlateUpdatableInstanceBuffer* InInstanceBuffer, int8 InSceneIndex)
 			: ShaderParams(InShaderParams)
 			, DrawFlags(InDrawFlags)
 			, ShaderType(InShaderType)
diff --git a/Engine/Source/Runtime/SlateRHIRenderer/Private/SlateRHIRenderingPolicy.cpp b/Engine/Source/Runtime/SlateRHIRenderer/Private/SlateRHIRenderingPolicy.cpp
--- a/Engine/Source/Runtime/SlateRHIRenderer/Private/SlateRHIRenderingPolicy.cpp
+++ b/Engine/Source/Runtime/SlateRHIRenderer/Private/SlateRHIRenderingPolicy.cpp
@@ -1008,7 +1008,7 @@ void FSlateRHIRenderingPolicy::DrawElements(
 
 				check(RenderBatch.NumIndices > 0);
 				// Note: This code is only executed if the engine is loaded (in early loading screens attempting to use a material is unsupported
-				int ActiveSceneIndex = RenderBatch.SceneIndex;
+				int32 ActiveSceneIndex = (int32)RenderBatch.SceneIndex;
 
 				// We are assuming at this point that the SceneIndex from the batch is either -1, meaning no scene or a valid scene.
 				// We set up the "no scene" option as the last SceneView in the array above.
