Fix crash on startup in on some platforms. Don't use FMemStack during static initialization.

The crash is caused by initialization order of FRunnableThread::RunnableTlsSlot and slate globals.
FMemStack::Get() uses FTlsAutoCleanup::Register() which uses RunnableTlsSlot which may not be initialized yet.

diff --git a/Engine/Source/Runtime/SlateCore/Private/Types/SlateAttributeDescriptor.cpp b/Engine/Source/Runtime/SlateCore/Private/Types/SlateAttributeDescriptor.cpp
--- a/Engine/Source/Runtime/SlateCore/Private/Types/SlateAttributeDescriptor.cpp
+++ b/Engine/Source/Runtime/SlateCore/Private/Types/SlateAttributeDescriptor.cpp
@@ -2,7 +2,7 @@
 
 #include "Types/SlateAttributeDescriptor.h"
 #include "Types/SlateAttribute.h"
-#include "Misc/MemStack.h"
+#include "Containers/ArrayView.h"
 
 
 FSlateAttributeDescriptor::FInitializer::FAttributeEntry::FAttributeEntry(FSlateAttributeDescriptor& InDescriptor, int32 InAttributeIndex)
@@ -67,7 +67,7 @@ FSlateAttributeDescriptor::FInitializer::~FInitializer()
 		int32 PrerequisitesIndex;
 		int32 Depth = -1;
 
-		void CalculateDepth(TArray<FPrerequisiteSort, TMemStackAllocator<>>& Prerequisites)
+		void CalculateDepth(TArrayView<FPrerequisiteSort> Prerequisites)
 		{
 			if (Depth < 0)
 			{
@@ -101,8 +101,7 @@ FSlateAttributeDescriptor::FInitializer::~FInitializer()
 		}
 	};
 
-	FMemMark Mark(FMemStack::Get());
-	TArray<FPrerequisiteSort, TMemStackAllocator<>> Prerequisites;
+	TArray<FPrerequisiteSort, TInlineAllocator<16>> Prerequisites;
 	Prerequisites.Reserve(Descriptor.Attributes.Num());
 
 	bool bHavePrerequisite = false;
@@ -289,8 +288,7 @@ void FSlateAttributeDescriptor::SetPrerequisite(FSlateAttributeDescriptor::FAttr
 
 			// Verify recursion
 			{
-				FMemMark Mark(FMemStack::Get());
-				TArray<FName, TMemStackAllocator<>> Recursion;
+				TArray<FName, TInlineAllocator<16>> Recursion;
 				Recursion.Reserve(Attributes.Num());
 				FAttribute const* CurrentAttribute = &Attribute;
 				while (!CurrentAttribute->Prerequisite.IsNone())
