Slate: Update SlateAttribute when building the SlateInvalidtionWidgetList since SWidget may changes their visibility. That implies that the AttributeIterator now needs to move at the end of the newly built array. The code will be moved into a utility function once we do the visibility optimization (visibility flags are not on the FWidgetProxy). Confirmed the correct use of bPauseAttributeInvalidation in SlateAttributeMetadata.

diff --git a/Engine/Source/Runtime/SlateCore/Private/FastUpdate/SlateInvalidationWidgetList.cpp b/Engine/Source/Runtime/SlateCore/Private/FastUpdate/SlateInvalidationWidgetList.cpp
--- a/Engine/Source/Runtime/SlateCore/Private/FastUpdate/SlateInvalidationWidgetList.cpp
+++ b/Engine/Source/Runtime/SlateCore/Private/FastUpdate/SlateInvalidationWidgetList.cpp
@@ -102,7 +102,8 @@ void FSlateInvalidationWidgetList::FWidgetAttributeIterator::PostResort()
 
 void FSlateInvalidationWidgetList::FWidgetAttributeIterator::ProxiesBuilt(FIndexRange const& Range)
 {
-	MoveToWidgetIndexOnNextAdvance = Range.GetInclusiveMinWidgetIndex();
+	// The last built item is already updated. We want to update the following item next.
+	MoveToWidgetIndexOnNextAdvance = WidgetList.IncrementIndex(Range.GetInclusiveMaxWidgetIndex());
 }
 
 
@@ -217,8 +218,44 @@ FSlateInvalidationWidgetIndex FSlateInvalidationWidgetList::_BuildWidgetList_Rec
 	LastestIndex = NewIndex.ArrayIndex;
 
 	FSlateInvalidationWidgetIndex LeafMostChildIndex = NewIndex;
-	const EVisibility Visibility = Widget->GetVisibility();
-	const bool bParentAndSelfVisible = bParentVisible && Visibility.IsVisible();
+	EVisibility Visibility = Widget->GetVisibility();
+	bool bParentAndSelfVisible = bParentVisible && Visibility.IsVisible();
+
+	bool bUpdateSlateAttribute = false;
+	if (Widget->HasRegisteredSlateAttribute())
+	{
+
+		// The list is already sorted at this point. Add to the end.
+		Data[LastestIndex].ElementIndexList_WidgetWithRegisteredSlateAttribute.AddUnsorted(NewIndex.ElementIndex);
+
+		// Update attributes now because the visibility flag may change
+		{
+			Widget->bPauseAttributeInvalidation = true;
+			if (bParentAndSelfVisible)
+			{
+				FSlateAttributeMetaData::UpdateAttributes(Widget.Get());
+				bUpdateSlateAttribute = true;
+
+				Visibility = Widget->GetVisibility();
+				bParentAndSelfVisible = bParentVisible && Visibility.IsVisible();
+			}
+			else if (bParentVisible)
+			{
+				FSlateAttributeMetaData::UpdateCollapsedAttributes(Widget.Get());
+				bUpdateSlateAttribute = true;
+
+				Visibility = Widget->GetVisibility();
+				bParentAndSelfVisible = bParentVisible && Visibility.IsVisible();
+				if (bParentAndSelfVisible)
+				{
+					FSlateAttributeMetaData::UpdateAttributes(Widget.Get());
+					Visibility = Widget->GetVisibility();
+					bParentAndSelfVisible = bParentVisible && Visibility.IsVisible();
+				}
+			}
+			Widget->bPauseAttributeInvalidation = false;
+		}
+	}
 
 	{
 		InvalidationWidgetType& WidgetProxy = (*this)[NewIndex];
@@ -226,12 +263,7 @@ FSlateInvalidationWidgetIndex FSlateInvalidationWidgetList::_BuildWidgetList_Rec
 		WidgetProxy.ParentIndex = ParentIndex;
 		WidgetProxy.LeafMostChildIndex = LeafMostChildIndex;
 		WidgetProxy.Visibility = Visibility;
-	}
-
-	if (Widget->HasRegisteredSlateAttribute())
-	{
-		// The list is already sorted at this point. Add to the end.
-		Data[LastestIndex].ElementIndexList_WidgetWithRegisteredSlateAttribute.AddUnsorted(NewIndex.ElementIndex);
+		WidgetProxy.bDebug_AttributeUpdated = bUpdateSlateAttribute;
 	}
 
 #if UE_SLATE_WITH_INVALIDATIONWIDGETLIST_DEBUGGING
diff --git a/Engine/Source/Runtime/SlateCore/Private/Types/SlateAttributeMetaData.cpp b/Engine/Source/Runtime/SlateCore/Private/Types/SlateAttributeMetaData.cpp
--- a/Engine/Source/Runtime/SlateCore/Private/Types/SlateAttributeMetaData.cpp
+++ b/Engine/Source/Runtime/SlateCore/Private/Types/SlateAttributeMetaData.cpp
@@ -170,7 +170,11 @@ EInvalidateWidgetReason FSlateAttributeMetaData::FGetterItem::GetInvalidationRea
 
 void FSlateAttributeMetaData::InvalidateWidget(SWidget& OwningWidget, const FSlateAttributeBase& Attribute, ESlateAttributeType AttributeType, EInvalidateWidgetReason Reason)
 {
-	// The widget is in construction phase. It's already invalidated... no need to keep invalidating it.
+	// The widget is in the construction phase or is building in the WidgetList.
+	//It's already invalidated... no need to keep invalidating it.
+	//N.B. no needs to set the bUpatedManually in this case because
+	//	1. they are in construction, so they will all be called anyway
+	//	2. they are in WidgetList, so the SlateAttribute.Set will not be called
 	if (OwningWidget.bPauseAttributeInvalidation)
 	{
 		return;
@@ -331,7 +335,10 @@ void FSlateAttributeMetaData::UpdateAttribute(SWidget& OwningWidget, FSlateAttri
 			ISlateAttributeGetter::FUpdateAttributeResult Result = GetterItem.Getter->UpdateAttribute(OwningWidget);
 			if (Result.bInvalidationRequested)
 			{
-				OwningWidget.Invalidate(GetterItem.GetInvalidationReason(OwningWidget, Result.InvalidationReason));
+				if (!OwningWidget.bPauseAttributeInvalidation)
+				{
+					OwningWidget.Invalidate(GetterItem.GetInvalidationReason(OwningWidget, Result.InvalidationReason));
+				}
 
 				// The dependency attribute need to be updated in the update loop (note that it may not be registered yet)
 				if (GetterItem.bIsADependencyForSomeoneElse)
