Slate: Remove deprecated EInvalidateWidget. Force TPanelChildren to have a valid owner (used by SlateAttribute). Remove unused variable in SWidget. Fix some bad implicit cast.

diff --git a/Engine/Source/Editor/UMGEditor/Private/Hierarchy/SHierarchyViewItem.cpp b/Engine/Source/Editor/UMGEditor/Private/Hierarchy/SHierarchyViewItem.cpp
--- a/Engine/Source/Editor/UMGEditor/Private/Hierarchy/SHierarchyViewItem.cpp
+++ b/Engine/Source/Editor/UMGEditor/Private/Hierarchy/SHierarchyViewItem.cpp
@@ -1085,7 +1085,7 @@ FText FHierarchyWidget::GetImageToolTipText() const
 	if ( WidgetTemplate )
 	{
 		UClass* WidgetClass = WidgetTemplate->GetClass();
-		if ( WidgetClass->IsChildOf( UUserWidget::StaticClass() ) )
+		if ( WidgetClass->IsChildOf( UUserWidget::StaticClass() ) && WidgetClass->ClassGeneratedBy )
 		{
 			auto& Description = Cast<UWidgetBlueprint>( WidgetClass->ClassGeneratedBy )->BlueprintDescription;
 			if ( Description.Len() > 0 )
diff --git a/Engine/Source/Runtime/SlateCore/Private/Styling/SlateBrush.cpp b/Engine/Source/Runtime/SlateCore/Private/Styling/SlateBrush.cpp
--- a/Engine/Source/Runtime/SlateCore/Private/Styling/SlateBrush.cpp
+++ b/Engine/Source/Runtime/SlateCore/Private/Styling/SlateBrush.cpp
@@ -4,7 +4,18 @@
 #include "SlateGlobals.h"
 #include "Application/SlateApplicationBase.h"
 
-FSlateBrush::FSlateBrush( ESlateBrushDrawType::Type InDrawType, const FName InResourceName, const FMargin& InMargin, ESlateBrushTileType::Type InTiling, ESlateBrushImageType::Type InImageType, const FVector2D& InImageSize, const FLinearColor& InTint, UObject* InObjectResource, bool bInDynamicallyLoaded )
+
+FSlateBrush::FSlateBrush( ESlateBrushDrawType::Type InDrawType, 
+						  const FName InResourceName, 
+						  const FMargin& InMargin, 
+						  ESlateBrushTileType::Type InTiling, 
+						  ESlateBrushImageType::Type InImageType, 
+						  const FVector2D& InImageSize, 
+						  const FLinearColor& InTint, 
+						  UObject* InObjectResource, 
+						  bool bInDynamicallyLoaded
+						)
+
 	: ImageSize( InImageSize )
 	, Margin( InMargin )
 #if WITH_EDITORONLY_DATA
@@ -25,11 +36,21 @@ FSlateBrush::FSlateBrush( ESlateBrushDrawType::Type InDrawType, const FName InRe
 	//Useful for debugging style breakages
 	//if ( !bHasUObject_DEPRECATED && InResourceName.IsValid() && InResourceName != NAME_None )
 	//{
-	//	checkf( FPaths::FileExists( InResourceName.ToString() ), *FPaths::ConvertRelativePathToFull( InResourceName.ToString() ) );
+	//	checkf( FPaths::FileExists( InResourceName.ToString() ), TEXT("The resource '%s' doesn't exist"), *FPaths::ConvertRelativePathToFull( InResourceName.ToString() ) );
 	//}
 }
 
-FSlateBrush::FSlateBrush( ESlateBrushDrawType::Type InDrawType, const FName InResourceName, const FMargin& InMargin, ESlateBrushTileType::Type InTiling, ESlateBrushImageType::Type InImageType, const FVector2D& InImageSize, const TSharedRef< FLinearColor >& InTint, UObject* InObjectResource, bool bInDynamicallyLoaded )
+FSlateBrush::FSlateBrush( ESlateBrushDrawType::Type InDrawType,
+ 						  const FName InResourceName,
+ 						  const FMargin& InMargin,
+ 						  ESlateBrushTileType::Type InTiling,
+ 						  ESlateBrushImageType::Type InImageType,
+ 						  const FVector2D& InImageSize,
+ 						  const TSharedRef< FLinearColor >& InTint,
+ 						  UObject* InObjectResource, 
+ 						  bool bInDynamicallyLoaded
+ 						)
+
 	: ImageSize( InImageSize )
 	, Margin( InMargin )
 #if WITH_EDITORONLY_DATA
@@ -50,11 +71,20 @@ FSlateBrush::FSlateBrush( ESlateBrushDrawType::Type InDrawType, const FName InRe
 	//Useful for debugging style breakages
 	//if ( !bHasUObject_DEPRECATED && InResourceName.IsValid() && InResourceName != NAME_None )
 	//{
-	//	checkf( FPaths::FileExists( InResourceName.ToString() ), *FPaths::ConvertRelativePathToFull( InResourceName.ToString() ) );
+	//	checkf( FPaths::FileExists( InResourceName.ToString() ), TEXT("The resource '%s' doesn't exist"), *FPaths::ConvertRelativePathToFull( InResourceName.ToString() ) );
 	//}
 }
 
-FSlateBrush::FSlateBrush( ESlateBrushDrawType::Type InDrawType, const FName InResourceName, const FMargin& InMargin, ESlateBrushTileType::Type InTiling, ESlateBrushImageType::Type InImageType, const FVector2D& InImageSize, const FSlateColor& InTint, UObject* InObjectResource, bool bInDynamicallyLoaded )
+FSlateBrush::FSlateBrush( ESlateBrushDrawType::Type InDrawType, 
+						  const FName InResourceName, 
+						  const FMargin& InMargin,
+						  ESlateBrushTileType::Type InTiling, 
+						  ESlateBrushImageType::Type InImageType, 
+						  const FVector2D& InImageSize, const FSlateColor& InTint, 
+						  UObject* InObjectResource, 
+						  bool bInDynamicallyLoaded
+ 						)
+
 	: ImageSize(InImageSize)
 	, Margin(InMargin)
 #if WITH_EDITORONLY_DATA
@@ -75,7 +105,7 @@ FSlateBrush::FSlateBrush( ESlateBrushDrawType::Type InDrawType, const FName InRe
 	//Useful for debugging style breakages
 	//if ( !bHasUObject_DEPRECATED && InResourceName.IsValid() && InResourceName != NAME_None )
 	//{
-	//	checkf( FPaths::FileExists( InResourceName.ToString() ), *FPaths::ConvertRelativePathToFull( InResourceName.ToString() ) );
+	//	checkf( FPaths::FileExists( InResourceName.ToString() ), TEXT("The resource '%s' doesn't exist"), *FPaths::ConvertRelativePathToFull( InResourceName.ToString() ) );
 	//}
 }
 
@@ -84,12 +114,18 @@ const FString FSlateBrush::UTextureIdentifier()
 	return FString(TEXT("texture:/"));
 }
 
-void FSlateBrush::UpdateRenderingResource() const
+void FSlateBrush::UpdateRenderingResource(FVector2D LocalSize, float DrawScale) const
 {
 	if (DrawAs != ESlateBrushDrawType::NoDrawType && (ResourceName != NAME_None || ResourceObject != nullptr))
 	{
-		ResourceHandle = FSlateApplicationBase::Get().GetRenderer()->GetResourceHandle(*this);
+		// Always re-acquire a handle if the current handle is invalid or if its vector graphics.
+		// For vector graphics we will rebuild the handle only if the shape needs to be rasterized again and the new size and scale
+		if (!ResourceHandle.IsValid() || ImageType == ESlateBrushImageType::Vector)
+		{
+			ResourceHandle = FSlateApplicationBase::Get().GetRenderer()->GetResourceHandle(*this, LocalSize, DrawScale);
+		}
 	}
+
 }
 
 bool FSlateBrush::CanRenderResourceObject(UObject* InResourceObject) const
diff --git a/Engine/Source/Runtime/SlateCore/Private/Widgets/SWidget.cpp b/Engine/Source/Runtime/SlateCore/Private/Widgets/SWidget.cpp
--- a/Engine/Source/Runtime/SlateCore/Private/Widgets/SWidget.cpp
+++ b/Engine/Source/Runtime/SlateCore/Private/Widgets/SWidget.cpp
@@ -212,7 +212,6 @@ SWidget::SWidget()
 	, bInheritedVolatility(false)
 	, bInvisibleDueToParentOrSelfVisibility(false)
 	, bNeedsPrepass(true)
-	, bUpdatingDesiredSize(false)
 	, bHasRegisteredSlateAttribute(false)
 	, bPauseAttributeInvalidation(true)
 	, bHasCustomPrepass(false)
@@ -1408,7 +1407,11 @@ int32 SWidget::Paint(const FPaintArgs& Args, const FGeometry& AllottedGeometry,
 	PersistentState.DesktopGeometry = DesktopSpaceGeometry;
 	PersistentState.WidgetStyle = InWidgetStyle;
 	PersistentState.CullingBounds = MyCullingRect;
-	PersistentState.IncomingUserIndex = Args.GetHittestGrid().GetUserIndex();
+
+	const int32 IncomingUserIndex = Args.GetHittestGrid().GetUserIndex();
+	ensure(IncomingUserIndex <= std::numeric_limits<int8>::max()); // shorten to save memory
+	PersistentState.IncomingUserIndex = (int8)IncomingUserIndex;
+
 	PersistentState.IncomingFlowDirection = GSlateFlowDirection;
 
 	FPaintArgs UpdatedArgs = Args.WithNewParent(this);
diff --git a/Engine/Source/Runtime/SlateCore/Public/Layout/Children.h b/Engine/Source/Runtime/SlateCore/Public/Layout/Children.h
--- a/Engine/Source/Runtime/SlateCore/Public/Layout/Children.h
+++ b/Engine/Source/Runtime/SlateCore/Public/Layout/Children.h
@@ -220,7 +220,7 @@ public:
 		WidgetPtr = InWidget;
 		if (Owner) 
 		{ 
-			Owner->Invalidate(EInvalidateWidget::ChildOrder);
+			Owner->Invalidate(EInvalidateWidgetReason::ChildOrder);
 
 			if (InWidget.IsValid() && InWidget != SNullWidget::NullWidget)
 			{
@@ -355,6 +355,9 @@ public:
 		, bEmptying(false)
 	{
 	}
+
+	//~ TPanelChildren should have a valid pointer
+	TPanelChildren(std::nullptr_t InOwner) = delete;
 	
 	virtual int32 Num() const override
 	{
@@ -454,7 +457,7 @@ public:
 
 		if (Owner)
 		{
-			Owner->Invalidate(EInvalidateWidget::ChildOrder);
+			Owner->Invalidate(EInvalidateWidgetReason::ChildOrder);
 		}
 	}
 
@@ -477,7 +480,7 @@ public:
 		::Sort(TIndirectArray< SlotType >::GetData(), TIndirectArray<SlotType>::Num(), Predicate);
 		if (Owner)
 		{
-			Owner->Invalidate(EInvalidateWidget::ChildOrder);
+			Owner->Invalidate(EInvalidateWidgetReason::ChildOrder);
 		}
 	}
 
@@ -487,7 +490,7 @@ public:
 		::StableSort(TIndirectArray< SlotType >::GetData(), TIndirectArray< SlotType >::Num(), Predicate);
 		if (Owner)
 		{
-			Owner->Invalidate(EInvalidateWidget::ChildOrder);
+			Owner->Invalidate(EInvalidateWidgetReason::ChildOrder);
 		}
 	}
 
@@ -496,7 +499,7 @@ public:
 		TIndirectArray< SlotType >::Swap(IndexA, IndexB);
 		if (Owner)
 		{
-			Owner->Invalidate(EInvalidateWidget::ChildOrder);
+			Owner->Invalidate(EInvalidateWidgetReason::ChildOrder);
 		}
 	}
 
@@ -663,7 +666,7 @@ public:
 	{
 		if (Owner && bChangesInvalidatePrepass)
 		{ 
-			Owner->Invalidate(EInvalidateWidget::ChildOrder);
+			Owner->Invalidate(EInvalidateWidgetReason::ChildOrder);
 		}
 
 		int32 Index = TArray< TSharedRef<ChildType> >::Add(Child);
@@ -712,7 +715,7 @@ public:
 	{
 		if (Owner && bChangesInvalidatePrepass) 
 		{
-			Owner->Invalidate(EInvalidateWidget::ChildOrder);
+			Owner->Invalidate(EInvalidateWidgetReason::ChildOrder);
 		}
 
 		TArray< TSharedRef<ChildType> >::Insert(Child, Index);
@@ -769,7 +772,7 @@ public:
 		TArray< TSharedRef<ChildType> >::Sort( Predicate );
 		if (Owner && bChangesInvalidatePrepass)
 		{
-			Owner->Invalidate(EInvalidateWidget::ChildOrder);
+			Owner->Invalidate(EInvalidateWidgetReason::ChildOrder);
 		}
 	}
 
@@ -778,7 +781,7 @@ public:
 		TIndirectArray< ChildType >::Swap(IndexA, IndexB);
 		if (Owner && bChangesInvalidatePrepass)
 		{
-			Owner->Invalidate(EInvalidateWidget::ChildOrder);
+			Owner->Invalidate(EInvalidateWidgetReason::ChildOrder);
 		}
 	}
 
diff --git a/Engine/Source/Runtime/SlateCore/Public/Widgets/SBoxPanel.h b/Engine/Source/Runtime/SlateCore/Public/Widgets/SBoxPanel.h
--- a/Engine/Source/Runtime/SlateCore/Public/Widgets/SBoxPanel.h
+++ b/Engine/Source/Runtime/SlateCore/Public/Widgets/SBoxPanel.h
@@ -207,7 +207,7 @@ public:
 		SHorizontalBox::FSlot& NewSlot = *new SHorizontalBox::FSlot();
 		this->Children.Add( &NewSlot );
 
-		Invalidate(EInvalidateWidget::Layout);
+		Invalidate(EInvalidateWidgetReason::Layout);
 
 		return NewSlot;
 	}
@@ -221,7 +221,7 @@ public:
 		SHorizontalBox::FSlot& NewSlot = *new SHorizontalBox::FSlot();
 		this->Children.Insert(&NewSlot, Index);
 
-		Invalidate(EInvalidateWidget::Layout);
+		Invalidate(EInvalidateWidgetReason::Layout);
 
 		return NewSlot;
 	}
@@ -352,7 +352,7 @@ public:
 		SVerticalBox::FSlot& NewSlot = *new SVerticalBox::FSlot();
 		this->Children.Add( &NewSlot );
 
-		Invalidate(EInvalidateWidget::Layout);
+		Invalidate(EInvalidateWidgetReason::Layout);
 
 		return NewSlot;
 	}
@@ -366,7 +366,7 @@ public:
 		SVerticalBox::FSlot& NewSlot = *new SVerticalBox::FSlot();
 		this->Children.Insert(&NewSlot, Index);
 
-		Invalidate(EInvalidateWidget::Layout);
+		Invalidate(EInvalidateWidgetReason::Layout);
 
 		return NewSlot;
 	}
diff --git a/Engine/Source/Runtime/SlateCore/Public/Widgets/SWidget.h b/Engine/Source/Runtime/SlateCore/Public/Widgets/SWidget.h
--- a/Engine/Source/Runtime/SlateCore/Public/Widgets/SWidget.h
+++ b/Engine/Source/Runtime/SlateCore/Public/Widgets/SWidget.h
@@ -920,7 +920,8 @@ public:
 
 	/**
 	 * @return is this widget visible, hidden or collapsed.
-	 * @note this widget can be visible but if a parent is hidden or collapsed, it would not show on screen. */
+	 * @note this widget can be visible but if a parent is hidden or collapsed, it would not show on screen.
+	 */
 	FORCEINLINE EVisibility GetVisibility() const { return Visibility.Get(); }
 
 	/** @param InVisibility  should this widget be */
@@ -1015,7 +1016,7 @@ public:
 		if (bForceVolatile != bForce)
 		{
 			bForceVolatile = bForce;
-			Invalidate(EInvalidateWidgetReason::Volatility);
+			Invalidate(EInvalidateWidgetReason::PaintAndVolatility);
 		}
 	}
 
@@ -1639,9 +1640,6 @@ private:
 	/** Are we currently updating the desired size? */
 	uint8 bNeedsPrepass : 1;
 
-	/** Are we currently updating the desired size? */
-	mutable uint8 bUpdatingDesiredSize : 1;
-
 	/** Is there at least one SlateAttribute currently registered. */
 	uint8 bHasRegisteredSlateAttribute : 1;
 
@@ -1821,7 +1819,7 @@ bool SetWidgetAttribute(SWidget& ThisWidget, TAttribute<TargetValueType>& Target
 		EInvalidateWidgetReason InvalidateReason = BaseInvalidationReason;
 		if (bBoundnessChanged)
 		{
-			InvalidateReason |= EInvalidateWidgetReason::Volatility;
+			InvalidateReason |= EInvalidateWidgetReason::PaintAndVolatility;
 		}
 
 		ThisWidget.Invalidate(InvalidateReason);
