Slate: Add SlateAttribute info to the Slate.Debug.LogAllWidgets command.

diff --git a/Engine/Source/Runtime/SlateCore/Private/Debugging/WidgetList.cpp b/Engine/Source/Runtime/SlateCore/Private/Debugging/WidgetList.cpp
--- a/Engine/Source/Runtime/SlateCore/Private/Debugging/WidgetList.cpp
+++ b/Engine/Source/Runtime/SlateCore/Private/Debugging/WidgetList.cpp
@@ -13,6 +13,7 @@
 #include "Misc/OutputDeviceFile.h"
 #include "Misc/Paths.h"
 #include "Misc/StringBuilder.h"
+#include "Types/SlateAttributeMetaData.h"
 #include "Types/SlateCursorMetaData.h"
 #include "Types/SlateMouseEventsMetaData.h"
 #include "Types/ReflectionMetadata.h"
@@ -39,6 +40,7 @@ struct FLogAllWidgetsDebugInfoFlags
 	bool bParent;
 	bool bToolTip;
 	bool bCursor;
+	bool bSlateAttribute;
 	bool bMouseEventsHandler;
 
 	void Parse(const FString& Arg)
@@ -77,6 +79,11 @@ struct FLogAllWidgetsDebugInfoFlags
 		{
 			return;
 		}
+		
+		if (FParse::Bool(*Arg, TEXT("SlateAttribute="), bSlateAttribute))
+		{
+			return;
+		}
 
 		if (FParse::Bool(*Arg, TEXT("MouseEvents="), bMouseEventsHandler))
 		{
@@ -118,6 +125,10 @@ void LogAllWidgetsDebugInfoImpl(FOutputDevice& Ar, const FLogAllWidgetsDebugInfo
 	{
 		MessageBuilder << TEXT(";CursorIsSet;CursorValue");
 	}
+	if (DebugInfoFlags.bSlateAttribute)
+	{
+		MessageBuilder << TEXT(";NumSlateAttribute");
+	}
 	if (DebugInfoFlags.bMouseEventsHandler)
 	{
 		MessageBuilder << TEXT(";MouseButtonDown;MouseButtonUp;MouseMove;MouseDblClick;MouseEnter;MouseLeave");
@@ -239,6 +250,19 @@ void LogAllWidgetsDebugInfoImpl(FOutputDevice& Ar, const FLogAllWidgetsDebugInfo
 			}
 		}
 
+		if (DebugInfoFlags.bSlateAttribute)
+		{
+			if (FSlateAttributeMetaData* MetaData = FSlateAttributeMetaData::FindMetaData(*Widget))
+			{
+				MessageBuilder << TEXT(";");
+				MessageBuilder << MetaData->RegisteredNum();
+			}
+			else
+			{
+				MessageBuilder << TEXT(";0");
+			}
+		}
+
 		if (DebugInfoFlags.bMouseEventsHandler)
 		{
 			if (TSharedPtr<FSlateMouseEventsMetaData> Data = Widget->GetMetaData<FSlateMouseEventsMetaData>())
diff --git a/Engine/Source/Runtime/SlateCore/Private/Debugging/WidgetList.h b/Engine/Source/Runtime/SlateCore/Private/Debugging/WidgetList.h
--- a/Engine/Source/Runtime/SlateCore/Private/Debugging/WidgetList.h
+++ b/Engine/Source/Runtime/SlateCore/Private/Debugging/WidgetList.h
@@ -3,14 +3,15 @@
 #pragma once
 
 #include "CoreMinimal.h"
+#include "Types/SlateAttribute.h"
 #include "Debugging/SlateDebugging.h"
 
 // Enabled to build a list of all the SWidget currently constructed
 #ifndef UE_WITH_SLATE_DEBUG_WIDGETLIST
-	#define UE_WITH_SLATE_DEBUG_WIDGETLIST 0
+	#define UE_WITH_SLATE_DEBUG_WIDGETLIST 0 || UE_SLATE_WITH_MEMBER_ATTRIBUTE_DEBUGGING
 #endif
 
-#if !WITH_SLATE_DEBUGGING && UE_WITH_SLATE_DEBUG_WIDGETLIST
+#if UE_WITH_SLATE_DEBUG_WIDGETLIST && (UE_BUILD_SHIPPING || UE_BUILD_TEST)
 	#error "UE_WITH_SLATE_DEBUG_WIDGETLIST is defined in a none debug build"
 #endif
 
