Slate: Fix crash with SlateAttribute and GlobalInvalidation when the ChildOrder is bigger than the itterator index.

diff --git a/Engine/Source/Runtime/SlateCore/Private/FastUpdate/SlateInvalidationWidgetList.cpp b/Engine/Source/Runtime/SlateCore/Private/FastUpdate/SlateInvalidationWidgetList.cpp
--- a/Engine/Source/Runtime/SlateCore/Private/FastUpdate/SlateInvalidationWidgetList.cpp
+++ b/Engine/Source/Runtime/SlateCore/Private/FastUpdate/SlateInvalidationWidgetList.cpp
@@ -83,7 +83,7 @@ void FSlateInvalidationWidgetList::FWidgetAttributeIterator::ReIndexed(IProcessC
 			MoveToWidgetIndexOnNextAdvance = Operation.ReIndex(MoveToWidgetIndexOnNextAdvance);
 		}
 	}
-	else
+	else if (CurrentWidgetIndex != FSlateInvalidationWidgetIndex::Invalid)
 	{
 		CurrentWidgetIndex = Operation.ReIndex(CurrentWidgetIndex);
 		CurrentWidgetSortOrder = FSlateInvalidationWidgetSortOrder{WidgetList, CurrentWidgetIndex };
@@ -93,7 +93,7 @@ void FSlateInvalidationWidgetList::FWidgetAttributeIterator::ReIndexed(IProcessC
 
 void FSlateInvalidationWidgetList::FWidgetAttributeIterator::PostResort()
 {
-	if (MoveToWidgetIndexOnNextAdvance != FSlateInvalidationWidgetIndex::Invalid)
+	if (MoveToWidgetIndexOnNextAdvance != FSlateInvalidationWidgetIndex::Invalid && CurrentWidgetIndex != FSlateInvalidationWidgetIndex::Invalid)
 	{
 		CurrentWidgetSortOrder = FSlateInvalidationWidgetSortOrder{ WidgetList, CurrentWidgetIndex };
 	}
