Make text edit boxes have a separate foreground color for when the text field is focused

diff --git a/Engine/Source/Runtime/Slate/Private/Widgets/Input/SEditableTextBox.cpp b/Engine/Source/Runtime/Slate/Private/Widgets/Input/SEditableTextBox.cpp
--- a/Engine/Source/Runtime/Slate/Private/Widgets/Input/SEditableTextBox.cpp
+++ b/Engine/Source/Runtime/Slate/Private/Widgets/Input/SEditableTextBox.cpp
@@ -32,6 +32,7 @@ void SEditableTextBox::Construct( const FArguments& InArgs )
 	ForegroundColorOverride = InArgs._ForegroundColor;
 	BackgroundColorOverride = InArgs._BackgroundColor;
 	ReadOnlyForegroundColorOverride = InArgs._ReadOnlyForegroundColor;
+	FocusedForegroundColorOverride = InArgs._FocusedForegroundColor;
 	OnTextChanged = InArgs._OnTextChanged;
 	OnVerifyTextChanged = InArgs._OnVerifyTextChanged;
 	OnTextCommitted = InArgs._OnTextCommitted;
@@ -303,6 +304,10 @@ FSlateColor SEditableTextBox::DetermineForegroundColor() const
 
 		return Style->ReadOnlyForegroundColor;
 	}
+	else if(HasKeyboardFocus())
+	{
+		return FocusedForegroundColorOverride.IsSet() ? FocusedForegroundColorOverride.Get() : Style->FocusedForegroundColor;
+	}
 	else
 	{
 		return ForegroundColorOverride.IsSet() ? ForegroundColorOverride.Get() : Style->ForegroundColor;
@@ -354,12 +359,15 @@ void SEditableTextBox::SetTextBoxBackgroundColor(const TAttribute<FSlateColor>&
 	BackgroundColorOverride = InBackgroundColor;
 }
 
-
 void SEditableTextBox::SetReadOnlyForegroundColor(const TAttribute<FSlateColor>& InReadOnlyForegroundColor)
 {
 	ReadOnlyForegroundColorOverride = InReadOnlyForegroundColor;
 }
 
+void SEditableTextBox::SetFocusedForegroundColor(const TAttribute<FSlateColor>& InFocusedForegroundColor)
+{
+	FocusedForegroundColorOverride = InFocusedForegroundColor;
+}
 
 void SEditableTextBox::SetMinimumDesiredWidth(const TAttribute<float>& InMinimumDesiredWidth)
 {
diff --git a/Engine/Source/Runtime/Slate/Private/Widgets/Input/SMultiLineEditableTextBox.cpp b/Engine/Source/Runtime/Slate/Private/Widgets/Input/SMultiLineEditableTextBox.cpp
--- a/Engine/Source/Runtime/Slate/Private/Widgets/Input/SMultiLineEditableTextBox.cpp
+++ b/Engine/Source/Runtime/Slate/Private/Widgets/Input/SMultiLineEditableTextBox.cpp
@@ -47,6 +47,7 @@ void SMultiLineEditableTextBox::Construct( const FArguments& InArgs )
 	ForegroundColorOverride = InArgs._ForegroundColor;
 	BackgroundColorOverride = InArgs._BackgroundColor;
 	ReadOnlyForegroundColorOverride = InArgs._ReadOnlyForegroundColor;
+	FocusedForegroundColorOverride = InArgs._FocusedForegroundColor;
 	bSelectWordOnMouseDoubleClick = InArgs._SelectWordOnMouseDoubleClick;
 
 	bHasExternalHScrollBar = InArgs._HScrollBar.IsValid();
@@ -216,7 +217,7 @@ FSlateColor SMultiLineEditableTextBox::DetermineForegroundColor() const
 
 	FSlateColor Result = FSlateColor::UseStyle();
 
-	if ( EditableText->IsTextReadOnly() )
+	if (EditableText->IsTextReadOnly())
 	{
 		if (ReadOnlyForegroundColorOverride.IsSet())
 		{
@@ -236,6 +237,26 @@ FSlateColor SMultiLineEditableTextBox::DetermineForegroundColor() const
 			return Result;
 		}
 	}
+	else if (HasKeyboardFocus())
+	{
+		if (FocusedForegroundColorOverride.IsSet())
+		{
+			Result = FocusedForegroundColorOverride.Get();
+		}
+		else if (ForegroundColorOverride.IsSet())
+		{
+			Result = ForegroundColorOverride.Get();
+		}
+
+		if (Result == FSlateColor::UseStyle())
+		{
+			return Style->FocusedForegroundColor;
+		}
+		else
+		{
+			return Result;
+		}
+	}
 	else
 	{
 		if (ForegroundColorOverride.IsSet())
diff --git a/Engine/Source/Runtime/Slate/Public/Widgets/Input/SEditableTextBox.h b/Engine/Source/Runtime/Slate/Public/Widgets/Input/SEditableTextBox.h
--- a/Engine/Source/Runtime/Slate/Public/Widgets/Input/SEditableTextBox.h
+++ b/Engine/Source/Runtime/Slate/Public/Widgets/Input/SEditableTextBox.h
@@ -38,6 +38,7 @@ public:
 		, _Font()
 		, _ForegroundColor()
 		, _ReadOnlyForegroundColor()
+		, _FocusedForegroundColor()
 		, _IsReadOnly( false )
 		, _IsPassword( false )
 		, _IsCaretMovedWhenGainFocus ( true )
@@ -79,6 +80,9 @@ public:
 		/** Text color and opacity when read-only (overrides Style) */
 		SLATE_ATTRIBUTE( FSlateColor, ReadOnlyForegroundColor )
 
+		/** Text color and opacity when this box has keyboard focus (overrides Style) */
+		SLATE_ATTRIBUTE(FSlateColor, FocusedForegroundColor)
+
 		/** Sets whether this text box can actually be modified interactively by the user */
 		SLATE_ATTRIBUTE( bool, IsReadOnly )
 
@@ -244,6 +248,13 @@ public:
 	 */
 	void SetReadOnlyForegroundColor(const TAttribute<FSlateColor>& InReadOnlyForegroundColor);
 
+	/**
+	 * Sets the text color and opacity when this box has keyboard focus(overrides Style)
+	 *
+	 * @param  InFocusedForegroundColor 	The focused color and opacity
+	 */
+	void SetFocusedForegroundColor(const TAttribute<FSlateColor>& InFocusedForegroundColor);
+
 	/**
 	 * Sets the minimum width that a text box should be.
 	 *
@@ -401,6 +412,9 @@ protected:
 	/** Read-only foreground color (overrides style) */
 	TAttribute<FSlateColor> ReadOnlyForegroundColorOverride;
 
+	/** Focused foreground color (overrides style) */
+	TAttribute<FSlateColor> FocusedForegroundColorOverride;
+
 	/** Allows for inserting additional widgets that extend the functionality of the text box */
 	TSharedPtr<SHorizontalBox> Box;
 
diff --git a/Engine/Source/Runtime/Slate/Public/Widgets/Input/SMultiLineEditableTextBox.h b/Engine/Source/Runtime/Slate/Public/Widgets/Input/SMultiLineEditableTextBox.h
--- a/Engine/Source/Runtime/Slate/Public/Widgets/Input/SMultiLineEditableTextBox.h
+++ b/Engine/Source/Runtime/Slate/Public/Widgets/Input/SMultiLineEditableTextBox.h
@@ -44,6 +44,7 @@ public:
 		, _Font()
 		, _ForegroundColor()
 		, _ReadOnlyForegroundColor()
+		, _FocusedForegroundColor()
 		, _Justification(ETextJustify::Left)
 		, _LineHeightPercentage(1.0f)
 		, _IsReadOnly( false )
@@ -101,6 +102,9 @@ public:
 		/** Text color and opacity when read-only (overrides Style) */
 		SLATE_ATTRIBUTE( FSlateColor, ReadOnlyForegroundColor )
 
+		/** Text color and opacity when this box has keyboard focus (overrides Style) */
+		SLATE_ATTRIBUTE(FSlateColor, FocusedForegroundColor)
+
 		/** How the text should be aligned with the margin. */
 		SLATE_ATTRIBUTE(ETextJustify::Type, Justification)
 
@@ -486,6 +490,9 @@ protected:
 	/** Read-only foreground color (overrides style) */
 	TAttribute<FSlateColor> ReadOnlyForegroundColorOverride;
 
+	/** Focused foreground color (overrides style) */
+	TAttribute<FSlateColor> FocusedForegroundColorOverride;
+
 	/** Whether to disable the context menu */
 	TAttribute< bool > AllowContextMenu;
 
diff --git a/Engine/Source/Runtime/SlateCore/Private/Styling/SlateTypes.cpp b/Engine/Source/Runtime/SlateCore/Private/Styling/SlateTypes.cpp
--- a/Engine/Source/Runtime/SlateCore/Private/Styling/SlateTypes.cpp
+++ b/Engine/Source/Runtime/SlateCore/Private/Styling/SlateTypes.cpp
@@ -252,6 +252,7 @@ FEditableTextBoxStyle::FEditableTextBoxStyle()
 	, ForegroundColor(FSlateColor::UseForeground())
 	, BackgroundColor(FLinearColor::White)
 	, ReadOnlyForegroundColor(FSlateColor::UseForeground())
+	, FocusedForegroundColor(FSlateColor::UseForeground())
 {
 }
 
diff --git a/Engine/Source/Runtime/SlateCore/Private/Styling/StarshipCoreStyle.cpp b/Engine/Source/Runtime/SlateCore/Private/Styling/StarshipCoreStyle.cpp
--- a/Engine/Source/Runtime/SlateCore/Private/Styling/StarshipCoreStyle.cpp
+++ b/Engine/Source/Runtime/SlateCore/Private/Styling/StarshipCoreStyle.cpp
@@ -996,9 +996,10 @@ void FStarshipCoreStyle::SetupTextStyles(TSharedRef<FStyle>& Style)
 		.SetBackgroundImageReadOnly(FSlateRoundedBoxBrush(FStyleColors::Header, InputFocusRadius, FStyleColors::InputOutline, InputFocusThickness))
 		.SetFont(StyleFonts.Normal)
 		.SetPadding(FMargin(12.f, 4.0f, 12.f, 5.0f)) // The padding should be 4 top, 5 bottom
-		.SetForegroundColor(FStyleColors::White)
+		.SetForegroundColor(FStyleColors::Foreground)
 		.SetBackgroundColor(FStyleColors::White)
 		.SetReadOnlyForegroundColor(FSlateColor::UseForeground())
+		.SetFocusedForegroundColor(FStyleColors::White)
 		.SetScrollBarStyle(ScrollBar);
 	{
 		Style->Set("NormalEditableTextBox", NormalEditableTextBoxStyle);
@@ -1248,13 +1249,13 @@ void FStarshipCoreStyle::SetupCheckboxStyles(TSharedRef<FStyle>& Style)
 	const FCheckBoxStyle BasicCheckBoxStyle = FCheckBoxStyle()
 		.SetCheckBoxType(ESlateCheckBoxType::CheckBox)
 
-		.SetForegroundColor(FLinearColor::White)
-		.SetHoveredForegroundColor(FLinearColor::White)
-		.SetPressedForegroundColor(FLinearColor::White)
-		.SetCheckedForegroundColor(FLinearColor::White)
-		.SetCheckedHoveredForegroundColor(FLinearColor::White)
-		.SetCheckedPressedForegroundColor(FLinearColor::White)
-		.SetUndeterminedForegroundColor(FLinearColor::White)
+		.SetForegroundColor(FStyleColors::Foreground)
+		.SetHoveredForegroundColor(FStyleColors::ForegroundHover)
+		.SetPressedForegroundColor(FStyleColors::ForegroundHover)
+		.SetCheckedForegroundColor(FStyleColors::Foreground)
+		.SetCheckedHoveredForegroundColor(FStyleColors::ForegroundHover)
+		.SetCheckedPressedForegroundColor(FStyleColors::ForegroundHover)
+		.SetUndeterminedForegroundColor(FStyleColors::Foreground)
 
 		.SetUncheckedImage(FSlateNoResource())
 		.SetUncheckedHoveredImage(FSlateNoResource())
diff --git a/Engine/Source/Runtime/SlateCore/Public/Styling/SlateTypes.h b/Engine/Source/Runtime/SlateCore/Public/Styling/SlateTypes.h
--- a/Engine/Source/Runtime/SlateCore/Public/Styling/SlateTypes.h
+++ b/Engine/Source/Runtime/SlateCore/Public/Styling/SlateTypes.h
@@ -934,7 +934,12 @@ struct SLATECORE_API FEditableTextBoxStyle : public FSlateWidgetStyle
 	/** The read-only foreground color of text in read-only mode. */
 	UPROPERTY(EditAnywhere, BlueprintReadWrite, Category=Appearance)
 	FSlateColor ReadOnlyForegroundColor;
-	FEditableTextBoxStyle& SetReadOnlyForegroundColor(const FSlateColor& InReadOnlyForegroundColor) { ReadOnlyForegroundColor = InReadOnlyForegroundColor; return *this; }
+	FEditableTextBoxStyle& SetReadOnlyForegroundColor(const FSlateColor& InReadOnlyForegroundColor) {ReadOnlyForegroundColor = InReadOnlyForegroundColor; return *this; }
+
+	/** The foreground color of text when the edit box has keyboard focus. */
+	UPROPERTY(EditAnywhere, BlueprintReadWrite, Category = Appearance)
+	FSlateColor FocusedForegroundColor;
+	FEditableTextBoxStyle& SetFocusedForegroundColor(const FSlateColor& InFocusedForegroundColor) {FocusedForegroundColor = InFocusedForegroundColor; return *this; }
 
 	/** Padding around the horizontal scrollbar */
 	UPROPERTY(EditAnywhere, BlueprintReadWrite, Category=Appearance)
