Adds Dropdown Outline Color and fixes ComboButton dropdown menu borders

diff --git a/Engine/Source/Runtime/SlateCore/Private/Styling/StarshipCoreStyle.cpp b/Engine/Source/Runtime/SlateCore/Private/Styling/StarshipCoreStyle.cpp
--- a/Engine/Source/Runtime/SlateCore/Private/Styling/StarshipCoreStyle.cpp
+++ b/Engine/Source/Runtime/SlateCore/Private/Styling/StarshipCoreStyle.cpp
@@ -802,6 +802,7 @@ void FStarshipCoreStyle::SetupColors(TSharedRef<FStyle>& Style)
 	Style->Set("Colors.Panel", FStyleColors::Panel);
 	Style->Set("Colors.Header", FStyleColors::Header);
 	Style->Set("Colors.Dropdown", FStyleColors::Dropdown);
+	Style->Set("Colors.DropdownOutline", FStyleColors::DropdownOutline);
 	Style->Set("Colors.Hover", FStyleColors::Hover);
 	Style->Set("Colors.Hover2", FStyleColors::Hover2);
 	Style->Set("Colors.White", FStyleColors::White);
@@ -849,6 +850,7 @@ void FStarshipCoreStyle::SetupColors(TSharedRef<FStyle>& Style)
 	Style->Set("Brushes.Panel", new FSlateColorBrush(FStyleColors::Panel));
 	Style->Set("Brushes.Header", new FSlateColorBrush(FStyleColors::Header));
 	Style->Set("Brushes.Dropdown", new FSlateColorBrush(FStyleColors::Dropdown));
+	Style->Set("Brushes.DropdownOutline", new FSlateColorBrush(FStyleColors::DropdownOutline));
 	Style->Set("Brushes.Hover", new FSlateColorBrush(FStyleColors::Hover));
 	Style->Set("Brushes.Hover2", new FSlateColorBrush(FStyleColors::Hover2));
 	Style->Set("Brushes.White", new FSlateColorBrush(FStyleColors::White));
@@ -1167,7 +1169,7 @@ void FStarshipCoreStyle::SetupComboButtonStyles(TSharedRef<FStyle>& Style)
 		.SetButtonStyle(ComboButtonButton)
 		.SetContentPadding(0.f)
 		.SetDownArrowImage(IMAGE_BRUSH_SVG("Starship/CoreWidgets/ComboBox/wide-chevron-down", FVector2D(20.f, 16.f)))
-		.SetMenuBorderBrush(FSlateRoundedBoxBrush(FStyleColors::Dropdown, 0.0f, WindowHighlight, 1.0f))
+		.SetMenuBorderBrush(FSlateRoundedBoxBrush(FStyleColors::Dropdown, 0.0, FStyleColors::DropdownOutline, 1.0))
 		.SetMenuBorderPadding(0.f);
 	Style->Set("ComboButton", ComboButton);
 
@@ -1188,7 +1190,7 @@ void FStarshipCoreStyle::SetupComboButtonStyles(TSharedRef<FStyle>& Style)
 		.SetDownArrowImage(IMAGE_BRUSH_SVG("Starship/CoreWidgets/ComboBox/corner-dropdown", FVector2D(7.0f, 7.0f)))
 		.SetDownArrowPadding(FMargin(0.0f))
 		.SetDownArrowAlignment(EVerticalAlignment::VAlign_Bottom)
-		.SetMenuBorderBrush(FSlateColorBrush(FStyleColors::Dropdown))
+		.SetMenuBorderBrush(FSlateRoundedBoxBrush(FStyleColors::Dropdown, 0.0, FStyleColors::DropdownOutline, 1.0))
 		.SetMenuBorderPadding(0.0f);
 	Style->Set("SimpleComboButton", SimpleComboButton);
 
@@ -1920,7 +1922,7 @@ void FStarshipCoreStyle::SetupMultiboxStyles(TSharedRef<FStyle>& Style)
 
 
 		Style->Set("Menu.Background", new FSlateColorBrush(FStyleColors::Dropdown));
-		Style->Set("Menu.Outline", new BORDER_BRUSH("Common/Window/WindowOutline", FMargin(1.0f / 32.0f), WindowHighlight));
+		Style->Set("Menu.Outline", new FSlateRoundedBoxBrush(FStyleColors::Transparent, 0.0, FStyleColors::DropdownOutline, 1.f));
 		Style->Set("Menu.Icon", new IMAGE_BRUSH("Icons/icon_tab_toolbar_16px", Icon16x16));
 		Style->Set("Menu.Expand", new IMAGE_BRUSH("Icons/toolbar_expand_16x", Icon16x16));
 		Style->Set("Menu.SubMenuIndicator", new IMAGE_BRUSH_SVG("Starship/Common/chevron-right", Icon16x16, FStyleColors::Foreground));
diff --git a/Engine/Source/Runtime/SlateCore/Private/Styling/StyleColors.cpp b/Engine/Source/Runtime/SlateCore/Private/Styling/StyleColors.cpp
--- a/Engine/Source/Runtime/SlateCore/Private/Styling/StyleColors.cpp
+++ b/Engine/Source/Runtime/SlateCore/Private/Styling/StyleColors.cpp
@@ -22,6 +22,7 @@ const FSlateColor FStyleColors::Background = EStyleColor::Background;
 const FSlateColor FStyleColors::Panel = EStyleColor::Panel;
 const FSlateColor FStyleColors::Header = EStyleColor::Header;
 const FSlateColor FStyleColors::Dropdown = EStyleColor::Dropdown;
+const FSlateColor FStyleColors::DropdownOutline = EStyleColor::DropdownOutline;
 const FSlateColor FStyleColors::Hover = EStyleColor::Hover;
 const FSlateColor FStyleColors::Hover2 = EStyleColor::Hover2;
 
@@ -81,6 +82,7 @@ void USlateThemeManager::InitalizeDefaults()
 	SetDefaultColor(EStyleColor::Panel, COLOR("#242424FF"));
 	SetDefaultColor(EStyleColor::Header, COLOR("#2F2F2FFF"));
 	SetDefaultColor(EStyleColor::Dropdown, COLOR("#383838FF"));
+	SetDefaultColor(EStyleColor::DropdownOutline, COLOR("#4C4C4CFF"));
 	SetDefaultColor(EStyleColor::Hover, COLOR("#575757FF"));
 	SetDefaultColor(EStyleColor::Hover2, COLOR("#808080FF"));
 
diff --git a/Engine/Source/Runtime/SlateCore/Public/Styling/StarshipCoreStyle.h b/Engine/Source/Runtime/SlateCore/Public/Styling/StarshipCoreStyle.h
--- a/Engine/Source/Runtime/SlateCore/Public/Styling/StarshipCoreStyle.h
+++ b/Engine/Source/Runtime/SlateCore/Public/Styling/StarshipCoreStyle.h
@@ -42,7 +42,6 @@ namespace CoreStyleConstants
 	const float InputFocusRadius = 4.f;
 	const float InputFocusThickness = 1.0f;
 
-	const FLinearColor WindowHighlight(FColor::FromHex("#4C4C4CFF"));
 }
 
 /**
diff --git a/Engine/Source/Runtime/SlateCore/Public/Styling/StyleColors.h b/Engine/Source/Runtime/SlateCore/Public/Styling/StyleColors.h
--- a/Engine/Source/Runtime/SlateCore/Public/Styling/StyleColors.h
+++ b/Engine/Source/Runtime/SlateCore/Public/Styling/StyleColors.h
@@ -29,6 +29,7 @@ enum class EStyleColor : uint8
 	Panel,
 	Header,
 	Dropdown,
+	DropdownOutline,
 	Hover,
 	Hover2,
 	White,
@@ -292,6 +293,7 @@ struct SLATECORE_API FStyleColors
 	static const FSlateColor Panel;
 	static const FSlateColor Header;
 	static const FSlateColor Dropdown;
+	static const FSlateColor DropdownOutline;
 	static const FSlateColor Hover;
 	static const FSlateColor Hover2;
 	static const FSlateColor White;
