Fix unity build
#rnx

diff --git a/Engine/Source/Runtime/SlateCore/Private/Fonts/SlateFontInfo.cpp b/Engine/Source/Runtime/SlateCore/Private/Fonts/SlateFontInfo.cpp
--- a/Engine/Source/Runtime/SlateCore/Private/Fonts/SlateFontInfo.cpp
+++ b/Engine/Source/Runtime/SlateCore/Private/Fonts/SlateFontInfo.cpp
@@ -223,3 +223,8 @@ void FSlateFontInfo::UpgradeLegacyFontInfo(FName LegacyFontName, EFontHinting Le
 		? FLegacySlateFontInfoCache::Get().GetSystemFont()
 		: FLegacySlateFontInfoCache::Get().GetCompositeFont(LegacyFontName, LegacyHinting);
 }
+
+uint16 FSlateFontInfo::GetClampSize() const
+{
+	return (uint16)FMath::Clamp<int32>(Size, 0, std::numeric_limits<uint16>::max());
+}
diff --git a/Engine/Source/Runtime/SlateCore/Public/Fonts/SlateFontInfo.h b/Engine/Source/Runtime/SlateCore/Public/Fonts/SlateFontInfo.h
--- a/Engine/Source/Runtime/SlateCore/Public/Fonts/SlateFontInfo.h
+++ b/Engine/Source/Runtime/SlateCore/Public/Fonts/SlateFontInfo.h
@@ -6,7 +6,6 @@
 #include "UObject/ObjectMacros.h"
 #include "UObject/Class.h"
 #include "Fonts/CompositeFont.h"
-#include <limits>
 #include "SlateFontInfo.generated.h"
 
 /**
@@ -278,7 +277,7 @@ public:
 	const FCompositeFont* GetCompositeFont() const;
 
 	/** Get the font size clamp for the font renderer (on 16bits) */
-	uint16 GetClampSize() const { return (uint16)FMath::Clamp<int32>(Size, 0, std::numeric_limits<uint16>::max()); }
+	uint16 GetClampSize() const;
 
 	/**
 	 * Calculates a type hash value for a font info.
