Slate: Improvement to the SlateAttribute API. Add FSlateDeprecatedTAttribute and SLATE_ADD_MEMBER_ATTRIBUTE_DEFINITION_WITH_NAME to work with deprecated attribute. Add TSlateMemberAttributeReference to help user of a widget, they do not need to provide the correct this pointer and prevent assigning the attribute directly.
 Complete a set of unit test for SlateAttribute. Revert changed to SWidet::TSlateAttribute with FText, that prevent the user from assigning another Compare predicate than the one hardcoded.

diff --git a/Engine/Source/Runtime/SlateCore/Private/Test/SlateAttributeTest.cpp b/Engine/Source/Runtime/SlateCore/Private/Test/SlateAttributeTest.cpp
--- a/Engine/Source/Runtime/SlateCore/Private/Test/SlateAttributeTest.cpp
+++ b/Engine/Source/Runtime/SlateCore/Private/Test/SlateAttributeTest.cpp
@@ -7,6 +7,8 @@
 #include "Widgets/SLeafWidget.h"
 #include "Widgets/SNullWidget.h"
 
+#include <type_traits>
+
 #if WITH_AUTOMATION_WORKER
 
 #define LOCTEXT_NAMESPACE "Slate.Attribute"
@@ -20,6 +22,39 @@ namespace Slate
 namespace Private
 {
 
+struct FConstructionCounter
+{
+	FConstructionCounter() : Value(0) { ++DefaultConstructionCounter; }
+	FConstructionCounter(int32 InValue) : Value(InValue) { ++DefaultConstructionCounter; }
+	FConstructionCounter(const FConstructionCounter& Other) : Value(Other.Value) { ++CopyConstructionCounter; }
+	FConstructionCounter(FConstructionCounter&& Other) :Value(Other.Value) { ++MoveConstructionCounter; }
+	FConstructionCounter& operator=(const FConstructionCounter& Other) { Value = Other.Value; ++CopyOperatorCounter; return *this; }
+	FConstructionCounter& operator=(FConstructionCounter&& Other) { Value = Other.Value; ++MoveOperatorCounter;  return *this; }
+
+	int32 Value;
+	bool operator== (const FConstructionCounter& Other) { return Other.Value == Value; }
+
+	static int32 DefaultConstructionCounter;
+	static int32 CopyConstructionCounter;
+	static int32 MoveConstructionCounter;
+	static int32 CopyOperatorCounter;
+	static int32 MoveOperatorCounter;
+	static void ResetCounter()
+	{
+		DefaultConstructionCounter = 0;
+		CopyConstructionCounter = 0;
+		MoveConstructionCounter = 0;
+		CopyOperatorCounter = 0;
+		MoveOperatorCounter = 0;
+	}
+};
+int32 FConstructionCounter::DefaultConstructionCounter = 0;
+int32 FConstructionCounter::CopyConstructionCounter = 0;
+int32 FConstructionCounter::MoveConstructionCounter = 0;
+int32 FConstructionCounter::CopyOperatorCounter = 0;
+int32 FConstructionCounter::MoveOperatorCounter = 0;
+
+
 int32 CallbackForIntAttribute(int32 Value)
 {
 	return Value;
@@ -38,14 +73,17 @@ public:
 
 
 	SAttributeLeftWidget_Parent()
-		: IntAttributeA(*this, 1)
-		, IntAttributeB(*this, 1)
-		, IntAttributeC(*this, 1)
-		, IntAttributeD(*this, 1)
-		, IntAttributeWithPredicate(*this, 1)
-		//, IntManagedAttributes{*this, 4}
+		: IntAttributeA(*this, 99)
+		, IntAttributeB(*this, 99)
+		, IntAttributeC(*this, 99)
+		, IntAttributeD(*this, 99)
 	{
+		static_assert(std::is_same<TSlateAttribute<bool>, typename TSlateAttributeRef<bool>::SlateAttributeType>::value, "TSlateAttributeRef doesn't have the same type as TSlateAttribute for bool");
+		static_assert(std::is_same<TSlateAttribute<int32>, typename TSlateAttributeRef<int32>::SlateAttributeType>::value, "TSlateAttributeRef doesn't have the same type as TSlateAttribute for int32");
+		static_assert(std::is_same<TSlateAttribute<FText>, typename TSlateAttributeRef<FText>::SlateAttributeType>::value, "TSlateAttributeRef doesn't have the same type as TSlateAttribute for FText");
+		static_assert(std::is_same<TSlateAttribute<FVector>, typename TSlateAttributeRef<FVector>::SlateAttributeType>::value, "TSlateAttributeRef doesn't have the same type as TSlateAttribute for FVector");
 	}
+
 	void Construct(const FArguments& InArgs){}
 	virtual FVector2D ComputeDesiredSize(float) const override { return FVector2D{ 100, 100 }; }
 	virtual int32 OnPaint(const FPaintArgs& Args, const FGeometry& AllottedGeometry, const FSlateRect& MyCullingRect, FSlateWindowElementList& OutDrawElements, int32 LayerId, const FWidgetStyle& InWidgetStyle, bool bParentEnabled) const override
@@ -53,80 +91,11 @@ public:
 		return LayerId;
 	}
 
-	int32 TestMember()
-	{
-		// this should not compile
-		{
-			//TSlateAttribute<int32> Other = IntAttribute;
-			//Other.Set(*this, 4);
-			//return Other.Get();
-		}
-		// this should compile but failed with check at runtime
-		{
-			//TSlateAttribute<int32> Other{ *this, 0 };
-			//Other.Set(*this, 4);
-			//return Other.Get();
-		}
-		{
-			int32 DefaultValue = 0;
-			TSlateAttribute<int32> Attribute = { *this, DefaultValue };
-		}
-		{
-			int32 DefaultValue = 0;
-			auto Getter1 = TAttribute<int32>::FGetter::CreateStatic(UE::Slate::Private::CallbackForIntAttribute, 1);
-			TSlateAttribute<int32> Attribute = { *this, DefaultValue, Getter1 };
-		}
-		{
-			int32 DefaultValue = 0;
-			auto Getter1 = TAttribute<int32>::FGetter::CreateStatic(UE::Slate::Private::CallbackForIntAttribute, 1);
-			TSlateAttribute<int32> Attribute = { *this, DefaultValue, MoveTemp(Getter1) };
-		} 
-		{
-			int32 DefaultValue = 0;
-			TSlateAttribute<int32> Attribute = { *this, DefaultValue, MakeAttributeLambda([](){return 5;}) };
-		}
-		{
-			FVector2D DefaultValue = FVector2D(0.f, 0.f);
-			TSlateAttribute<FVector2D> Attribute = { *this, MoveTemp(DefaultValue) };
-		}
-		{
-			FVector2D DefaultValue = FVector2D(0.f, 0.f);
-			auto Getter1 = TAttribute<FVector2D>::FGetter::CreateStatic(UE::Slate::Private::CallbackForFVectorAttribute);
-			TSlateAttribute<FVector2D> Attribute = { *this, MoveTemp(DefaultValue), Getter1 };
-		}
-		{
-			FVector2D DefaultValue = FVector2D(0.f, 0.f);
-			auto Getter1 = TAttribute<FVector2D>::FGetter::CreateStatic(UE::Slate::Private::CallbackForFVectorAttribute);
-			TSlateAttribute<FVector2D> Attribute = { *this, MoveTemp(DefaultValue), MoveTemp(Getter1) };
-		}
-		{
-			FVector2D DefaultValue = FVector2D(0.f, 0.f);
-			TSlateAttribute<FVector2D> Attribute = { *this, MoveTemp(DefaultValue), MakeAttributeLambda([]() {return FVector2D(1.f, 1.f);}) };
-		}
-
-		return 0; 
-	}
-
-	void TestExternalBind()
-	{
-		// Copy should not compile
-		{
-			//TSlateManagedAttribute<int32, EInvalidateWidgetReason::ChildOrder> ExtAtt1 {AsShared(), 10};
-			//IntManagedAttributes.Add(ExtAtt1);
-		}
-
-		// Move should compile
-		{
-			TSlateManagedAttribute<int32, EInvalidateWidgetReason::ChildOrder> ExtAtt1 {AsShared(), 10};
-			IntManagedAttributes.Add(MoveTemp(ExtAtt1));
-		}
-	}
 
 	TSlateAttribute<int32> IntAttributeA;
 	TSlateAttribute<int32> IntAttributeB;
 	TSlateAttribute<int32> IntAttributeC;
 	TSlateAttribute<int32> IntAttributeD;
-	TSlateAttributeWithPredicate<int32, TSlateAttributeInvalidationReason<EInvalidateWidgetReason::ChildOrder>> IntAttributeWithPredicate;
 	TArray<TSlateManagedAttribute<int32, EInvalidateWidgetReason::ChildOrder>> IntManagedAttributes;
 };
 
@@ -156,11 +125,11 @@ public:
 
 
 	SAttributeLeftWidget_Child()
-		: IntAttributeH(*this, 1)
-		, IntAttributeI(*this, 1)
-		, IntAttributeJ(*this, 1)
-		, IntAttributeK(*this, 1)
-		, IntAttributeL(*this, 1)
+		: IntAttributeH(*this, 99)
+		, IntAttributeI(*this, 99)
+		, IntAttributeJ(*this, 99)
+		, IntAttributeK(*this, 99)
+		, IntAttributeL(*this, 99)
 	{
 	}
 
@@ -249,13 +218,13 @@ bool FSlateAttributeTest::RunTest(const FString& Parameters)
 		{
 			OrderCounter = 0;
 			WidgetParent->IntAttributeA.Assign(WidgetParent.Get(), MakeAttributeLambda(OrderLambda));
-			AddErrorIfFalse(WidgetParent->IntAttributeA.Get() == 1, TEXT("A It is not the expected value."));
+			AddErrorIfFalse(WidgetParent->IntAttributeA.Get() == 99, TEXT("A It is not the expected value."));
 			WidgetParent->IntAttributeB.Assign(WidgetParent.Get(), MakeAttributeLambda(OrderLambda));
-			AddErrorIfFalse(WidgetParent->IntAttributeB.Get() == 2, TEXT("B It is not the expected value."));
+			AddErrorIfFalse(WidgetParent->IntAttributeB.Get() == 99, TEXT("B It is not the expected value."));
 			WidgetParent->IntAttributeC.Assign(WidgetParent.Get(), MakeAttributeLambda(OrderLambda));
-			AddErrorIfFalse(WidgetParent->IntAttributeC.Get() == 3, TEXT("C It is not the expected value."));
+			AddErrorIfFalse(WidgetParent->IntAttributeC.Get() == 99, TEXT("C It is not the expected value."));
 			WidgetParent->IntAttributeD.Assign(WidgetParent.Get(), MakeAttributeLambda(OrderLambda));
-			AddErrorIfFalse(WidgetParent->IntAttributeD.Get() == 4, TEXT("D It is not the expected value."));
+			AddErrorIfFalse(WidgetParent->IntAttributeD.Get() == 99, TEXT("D It is not the expected value."));
 
 			OrderCounter = 0;
 			bWasUpdate = false;
@@ -273,6 +242,10 @@ bool FSlateAttributeTest::RunTest(const FString& Parameters)
 			bWasUpdate = false;
 			ReturnValue = 5;
 			WidgetParent->IntAttributeC.Assign(WidgetParent.Get(), MakeAttributeLambda(UpdateLambda));
+			AddErrorIfFalse(!bWasUpdate, TEXT("C should not have be updated."));
+			AddErrorIfFalse(WidgetParent->IntAttributeC.Get() == 4, TEXT("C It is not the expected value."));
+			WidgetParent->InvalidatePrepass();
+			WidgetParent->SlatePrepass(1.f);
 			AddErrorIfFalse(bWasUpdate, TEXT("C should be updated."));
 			AddErrorIfFalse(WidgetParent->IntAttributeC.Get() == 5, TEXT("C It is not the expected value."));
 		}
@@ -305,7 +278,9 @@ bool FSlateAttributeTest::RunTest(const FString& Parameters)
 			AddErrorIfFalse(bWasUpdate, TEXT("C should be updated becaude D was."));
 			AddErrorIfFalse(WidgetParent->IntAttributeC.Get() == 10, TEXT("C It is not the expected value."));
 			AddErrorIfFalse(WidgetParent->IntAttributeD.Get() == 8, TEXT("D It is not the expected value."));
+#ifndef PVS_STUDIO // Build machine refuses to disable this warning
 			AddErrorIfFalse(OrderCounter == 2, TEXT("There is no D attribute anymore.")); //-V547
+#endif
 		}
 
 		{
@@ -367,23 +342,23 @@ bool FSlateAttributeTest::RunTest(const FString& Parameters)
 		{
 			OrderCounter = 49;
 			WidgetChild->IntAttributeA.Assign(WidgetChild.Get(), MakeAttributeLambda(OrderLambda));
-			AddErrorIfFalse(WidgetChild->IntAttributeA.Get() == 50, TEXT("A It is not the expected value."));
+			AddErrorIfFalse(WidgetChild->IntAttributeA.Get() == 99, TEXT("A It is not the expected value."));
 			WidgetChild->IntAttributeB.Assign(WidgetChild.Get(), MakeAttributeLambda(OrderLambda));
-			AddErrorIfFalse(WidgetChild->IntAttributeB.Get() == 51, TEXT("B It is not the expected value."));
+			AddErrorIfFalse(WidgetChild->IntAttributeB.Get() == 99, TEXT("B It is not the expected value."));
 			WidgetChild->IntAttributeC.Assign(WidgetChild.Get(), MakeAttributeLambda(OrderLambda));
-			AddErrorIfFalse(WidgetChild->IntAttributeC.Get() == 52, TEXT("C It is not the expected value."));
+			AddErrorIfFalse(WidgetChild->IntAttributeC.Get() == 99, TEXT("C It is not the expected value."));
 			WidgetChild->IntAttributeD.Assign(WidgetChild.Get(), MakeAttributeLambda(OrderLambda));
-			AddErrorIfFalse(WidgetChild->IntAttributeD.Get() == 53, TEXT("D It is not the expected value."));
+			AddErrorIfFalse(WidgetChild->IntAttributeD.Get() == 99, TEXT("D It is not the expected value."));
 			WidgetChild->IntAttributeH.Assign(WidgetChild.Get(), MakeAttributeLambda(OrderLambda));
-			AddErrorIfFalse(WidgetChild->IntAttributeH.Get() == 54, TEXT("I It is not the expected value."));
+			AddErrorIfFalse(WidgetChild->IntAttributeH.Get() == 99, TEXT("I It is not the expected value."));
 			WidgetChild->IntAttributeI.Assign(WidgetChild.Get(), MakeAttributeLambda(OrderLambda));
-			AddErrorIfFalse(WidgetChild->IntAttributeI.Get() == 55, TEXT("I It is not the expected value."));
+			AddErrorIfFalse(WidgetChild->IntAttributeI.Get() == 99, TEXT("I It is not the expected value."));
 			WidgetChild->IntAttributeJ.Assign(WidgetChild.Get(), MakeAttributeLambda(OrderLambda));
-			AddErrorIfFalse(WidgetChild->IntAttributeJ.Get() == 56, TEXT("J It is not the expected value."));
+			AddErrorIfFalse(WidgetChild->IntAttributeJ.Get() == 99, TEXT("J It is not the expected value."));
 			WidgetChild->IntAttributeK.Assign(WidgetChild.Get(), MakeAttributeLambda(OrderLambda));
-			AddErrorIfFalse(WidgetChild->IntAttributeK.Get() == 57, TEXT("K It is not the expected value."));
+			AddErrorIfFalse(WidgetChild->IntAttributeK.Get() == 99, TEXT("K It is not the expected value."));
 			WidgetChild->IntAttributeL.Assign(WidgetChild.Get(), MakeAttributeLambda(OrderLambda));
-			AddErrorIfFalse(WidgetChild->IntAttributeL.Get() == 58, TEXT("L It is not the expected value."));
+			AddErrorIfFalse(WidgetChild->IntAttributeL.Get() == 99, TEXT("L It is not the expected value."));
 
 
 			OrderCounter = 0;
@@ -420,6 +395,312 @@ bool FSlateAttributeTest::RunTest(const FString& Parameters)
 		}
 	}
 
+
+	// Make sure we call all the functions
+	{
+		{
+			// This should just compile to TSlateAttribute
+			struct SAttributeAttribute : public SLeafWidget
+			{
+				SAttributeAttribute()
+					: Toto(6)
+					, AttributeA(*this)
+					, AttributeB(*this, 5)
+					, AttributeC(*this, MoveTemp(Toto))
+				{ }
+				SLATE_BEGIN_ARGS(SAttributeAttribute) {}  SLATE_END_ARGS()
+				virtual int32 OnPaint(const FPaintArgs& Args, const FGeometry& AllottedGeometry, const FSlateRect& MyCullingRect, FSlateWindowElementList& OutDrawElements, int32 LayerId, const FWidgetStyle& InWidgetStyle, bool bParentEnabled) const override { return LayerId; }
+				virtual FVector2D ComputeDesiredSize(float) const override { return FVector2D(0.f, 0.f); }
+				void Construct(const FArguments&) {	}
+
+				int32 Callback() const { return 0; }
+
+				int32 Toto;
+				TSlateAttribute<int32, EInvalidateWidgetReason::Paint> AttributeA;
+				TSlateAttribute<int32, EInvalidateWidgetReason::Paint> AttributeB;
+				TSlateAttribute<int32, EInvalidateWidgetReason::Paint> AttributeC;
+			};
+			TSharedPtr<SAttributeAttribute> Widget = SNew(SAttributeAttribute);
+
+			{
+				int32 Hello = 7;
+				int32 Return1 = Widget->AttributeA.Get();
+				Widget->AttributeA.UpdateNow(*Widget);
+				Widget->AttributeA.Set(*Widget, 6);
+				Widget->AttributeA.Set(*Widget, MoveTemp(Hello));
+			}
+			{
+				auto Getter1 = TAttribute<int32>::FGetter::CreateStatic(UE::Slate::Private::CallbackForIntAttribute, 1);
+				Widget->AttributeA.Bind(*Widget, Getter1);
+				Widget->AttributeA.Bind(*Widget, MoveTemp(Getter1));
+				Widget->AttributeA.Bind(*Widget, &SAttributeAttribute::Callback);
+			}
+			{
+				int32 TmpInt1 = 7;
+				int32 TmpInt2 = 7;
+				auto Getter1 = TAttribute<int32>::FGetter::CreateStatic(UE::Slate::Private::CallbackForIntAttribute, 1);
+				TAttribute<int32> Attribute1 = TAttribute<int32>::Create(Getter1);
+				TAttribute<int32> Attribute2 = TAttribute<int32>::Create(Getter1);
+				TAttribute<int32> Attribute3 = TAttribute<int32>::Create(Getter1);
+				Widget->AttributeA.Assign(*Widget, Attribute1);
+				Widget->AttributeA.Assign(*Widget, MoveTemp(Attribute1));
+				Widget->AttributeA.Assign(*Widget, Attribute2, 7);
+				Widget->AttributeA.Assign(*Widget, MoveTemp(Attribute2), 7);
+				Widget->AttributeA.Assign(*Widget, Attribute3, MoveTemp(TmpInt1));
+				Widget->AttributeA.Assign(*Widget, MoveTemp(Attribute3), MoveTemp(TmpInt2));
+			}
+			{
+				bool bIsBound1 = Widget->AttributeA.IsBound(*Widget);
+				bool bIsIdentical1 = Widget->AttributeA.IsIdenticalTo(*Widget, Widget->AttributeA);
+				auto Getter1 = TAttribute<int32>::FGetter::CreateStatic(UE::Slate::Private::CallbackForIntAttribute, 1);
+				TAttribute<int32> Attribute1 = TAttribute<int32>::Create(Getter1);
+				bool bIsIdentical2 = Widget->AttributeA.IsIdenticalTo(*Widget, Attribute1);
+			}
+		}
+		{
+			typedef UE::Slate::Private::FConstructionCounter FLocalConstructionCounter;
+
+			// This should just compile to TSlateManagedAttribute
+			struct SAttributeAttribute : public SLeafWidget
+			{
+				SLATE_BEGIN_ARGS(SAttributeAttribute) {} SLATE_END_ARGS()
+				virtual int32 OnPaint(const FPaintArgs& Args, const FGeometry& AllottedGeometry, const FSlateRect& MyCullingRect, FSlateWindowElementList& OutDrawElements, int32 LayerId, const FWidgetStyle& InWidgetStyle, bool bParentEnabled) const override { return LayerId; }
+				virtual FVector2D ComputeDesiredSize(float) const override { return FVector2D(0.f, 0.f); }
+				void Construct(const FArguments&) {}
+
+				FLocalConstructionCounter ReturnDefaultCounter() const { return FLocalConstructionCounter(0); }
+
+				using ManagedSlateAttributeType = TSlateManagedAttribute<FLocalConstructionCounter, EInvalidateWidgetReason::Layout>;
+			};
+			TSharedPtr<SAttributeAttribute>	Widget = SNew(SAttributeAttribute);
+
+
+			auto AddErrorIfCounterDoNotMatches = [this](int32 Construct, int32 Copy, int32 Move, int32 CopyAssign, int32 MoveAssign, const TCHAR* Message)
+			{
+				bool bSuccess = FLocalConstructionCounter::DefaultConstructionCounter == Construct
+					&& FLocalConstructionCounter::CopyConstructionCounter == Copy
+					&& FLocalConstructionCounter::MoveConstructionCounter == Move
+					&& FLocalConstructionCounter::CopyOperatorCounter == CopyAssign
+					&& FLocalConstructionCounter::MoveOperatorCounter == MoveAssign;
+				AddErrorIfFalse(bSuccess, Message);
+			};
+
+			{
+				FLocalConstructionCounter::ResetCounter();
+				SAttributeAttribute::ManagedSlateAttributeType Attribute = SAttributeAttribute::ManagedSlateAttributeType(Widget.ToSharedRef() );
+				AddErrorIfCounterDoNotMatches(1, 0, 0, 0, 0, TEXT("Default & Copy constructor was not used."));
+			}
+			{
+				FLocalConstructionCounter Counter = 1;
+				FLocalConstructionCounter::ResetCounter();
+				SAttributeAttribute::ManagedSlateAttributeType Attribute = SAttributeAttribute::ManagedSlateAttributeType(Widget.ToSharedRef(), Counter);
+				AddErrorIfCounterDoNotMatches(0, 1, 0, 0, 0, TEXT("Default & Copy constructor was not used."));
+			}
+			{
+				FLocalConstructionCounter::ResetCounter();
+				FLocalConstructionCounter Counter = 1;
+				SAttributeAttribute::ManagedSlateAttributeType Attribute = SAttributeAttribute::ManagedSlateAttributeType(Widget.ToSharedRef(), MoveTemp(Counter));
+				AddErrorIfCounterDoNotMatches(1, 0, 1, 0, 0, TEXT("Default & Move constructor was not used."));
+			}
+			{
+				TAttribute<FLocalConstructionCounter>::FGetter Getter1 = TAttribute<FLocalConstructionCounter>::FGetter::CreateLambda([](){return FLocalConstructionCounter(1);});
+				FLocalConstructionCounter Counter = 1;
+				FLocalConstructionCounter::ResetCounter();
+				SAttributeAttribute::ManagedSlateAttributeType Attribute = SAttributeAttribute::ManagedSlateAttributeType{Widget.ToSharedRef(), Getter1, Counter};
+				AddErrorIfCounterDoNotMatches(0, 1, 0, 0, 0, TEXT("Getter & Copy constructor was not used."));
+			}
+			{
+				TAttribute<FLocalConstructionCounter>::FGetter Getter1 = TAttribute<FLocalConstructionCounter>::FGetter::CreateLambda([]() {return FLocalConstructionCounter(1);});
+				FLocalConstructionCounter Counter = 1;
+				FLocalConstructionCounter::ResetCounter();
+				SAttributeAttribute::ManagedSlateAttributeType Attribute = SAttributeAttribute::ManagedSlateAttributeType{Widget.ToSharedRef(), Getter1, MoveTemp(Counter)};
+				AddErrorIfCounterDoNotMatches(0, 0, 1, 0, 0, TEXT("Getter & Move constructor was not used."));
+			}
+			{
+				TAttribute<FLocalConstructionCounter>::FGetter Getter1 = TAttribute<FLocalConstructionCounter>::FGetter::CreateLambda([]() {return FLocalConstructionCounter(1);});
+				FLocalConstructionCounter Counter = 1;
+				FLocalConstructionCounter::ResetCounter();
+				SAttributeAttribute::ManagedSlateAttributeType Attribute = SAttributeAttribute::ManagedSlateAttributeType(Widget.ToSharedRef(), MoveTemp(Getter1), Counter);
+				AddErrorIfCounterDoNotMatches(0, 1, 0, 0, 0, TEXT("Move Getter & Copy constructor was not used."));
+			}
+			{
+				TAttribute<FLocalConstructionCounter>::FGetter Getter1 = TAttribute<FLocalConstructionCounter>::FGetter::CreateLambda([]() {return FLocalConstructionCounter(1);});
+				FLocalConstructionCounter Counter = 1;
+				FLocalConstructionCounter::ResetCounter();
+				SAttributeAttribute::ManagedSlateAttributeType Attribute = SAttributeAttribute::ManagedSlateAttributeType(Widget.ToSharedRef(), MoveTemp(Getter1), MoveTemp(Counter));
+				AddErrorIfCounterDoNotMatches(0, 0, 1, 0, 0, TEXT("Move Getter & Move constructor was not used."));
+			}
+			{
+				TAttribute<FLocalConstructionCounter> Attribute1 = MakeAttributeLambda([]() {return FLocalConstructionCounter(1);});
+				FLocalConstructionCounter Counter;
+				FLocalConstructionCounter::ResetCounter();
+				SAttributeAttribute::ManagedSlateAttributeType Attribute = SAttributeAttribute::ManagedSlateAttributeType(Widget.ToSharedRef(), Attribute1, Counter);
+				AddErrorIfCounterDoNotMatches(0, 1, 0, 0, 0, TEXT("Attribute & Copy constructor was not used."));
+			}
+			{
+				TAttribute<FLocalConstructionCounter> Attribute1 = MakeAttributeLambda([]() {return FLocalConstructionCounter(1);});
+				FLocalConstructionCounter Counter = 1;
+				FLocalConstructionCounter::ResetCounter();
+				SAttributeAttribute::ManagedSlateAttributeType Attribute = SAttributeAttribute::ManagedSlateAttributeType(Widget.ToSharedRef(), MoveTemp(Attribute1), MoveTemp(Counter));
+				AddErrorIfCounterDoNotMatches(0, 0, 1, 0, 0, TEXT("Move Attribute & Move constructor was not used."));
+			}
+			{
+				SAttributeAttribute::ManagedSlateAttributeType Attribute = SAttributeAttribute::ManagedSlateAttributeType(Widget.ToSharedRef());
+				FLocalConstructionCounter::ResetCounter();
+				FLocalConstructionCounter Result = Attribute.Get();
+				Attribute.UpdateNow();
+				AddErrorIfCounterDoNotMatches(0, 1, 0, 0, 0, TEXT("Get and UpdateNow failed."));
+			}
+			{
+				SAttributeAttribute::ManagedSlateAttributeType Attribute = SAttributeAttribute::ManagedSlateAttributeType(Widget.ToSharedRef());
+				FLocalConstructionCounter Counter = 1;
+				FLocalConstructionCounter::ResetCounter();
+				Attribute.Set(Counter);
+				AddErrorIfCounterDoNotMatches(0, 0, 0, 1, 0, TEXT("Set Copy failed."));
+			}
+			{
+				SAttributeAttribute::ManagedSlateAttributeType Attribute = SAttributeAttribute::ManagedSlateAttributeType(Widget.ToSharedRef());
+				FLocalConstructionCounter Counter = 1;
+				FLocalConstructionCounter::ResetCounter();
+				Attribute.Set(MoveTemp(Counter));
+				AddErrorIfCounterDoNotMatches(0, 0, 0, 0, 1, TEXT("Set Move failed."));
+			}
+			{
+				SAttributeAttribute::ManagedSlateAttributeType Attribute = SAttributeAttribute::ManagedSlateAttributeType(Widget.ToSharedRef());
+				auto Getter1 = TAttribute<FLocalConstructionCounter>::FGetter::CreateLambda([]() {return FLocalConstructionCounter(1);});
+				FLocalConstructionCounter::ResetCounter();
+				Attribute.Bind(Getter1);
+				AddErrorIfCounterDoNotMatches(0, 0, 0, 0, 0, TEXT("Bind Copy failed."));
+			}
+			{
+				SAttributeAttribute::ManagedSlateAttributeType Attribute = SAttributeAttribute::ManagedSlateAttributeType(Widget.ToSharedRef());
+				auto Getter1 = TAttribute<FLocalConstructionCounter>::FGetter::CreateLambda([]() {return FLocalConstructionCounter(1);});
+				FLocalConstructionCounter::ResetCounter();
+				Attribute.Bind(MoveTemp(Getter1));
+				AddErrorIfCounterDoNotMatches(0, 0, 0, 0, 0, TEXT("Bind Copy failed."));
+			}
+			// Test Assign
+			{
+				SAttributeAttribute::ManagedSlateAttributeType Attribute = SAttributeAttribute::ManagedSlateAttributeType(Widget.ToSharedRef());
+				TAttribute<FLocalConstructionCounter> Attribute1;
+				FLocalConstructionCounter::ResetCounter();
+				Attribute.Assign(Attribute1);
+				AddErrorIfCounterDoNotMatches(0, 0, 0, 0, 0, TEXT("Assign Copy failed."));
+				Attribute1.Set({1});
+				FLocalConstructionCounter::ResetCounter();
+				Attribute.Assign(Attribute1);
+				AddErrorIfCounterDoNotMatches(0, 0, 0, 1, 0, TEXT("Assign Copy failed."));
+				Attribute1.Set({ 1 });
+				FLocalConstructionCounter::ResetCounter();
+				Attribute.Assign(Attribute1);
+				AddErrorIfCounterDoNotMatches(0, 0, 0, 0, 0, TEXT("Assign Move failed."));
+				Attribute1.Set({ 2 });
+				FLocalConstructionCounter::ResetCounter();
+				Attribute.Assign(MoveTemp(Attribute1));
+				AddErrorIfCounterDoNotMatches(0, 0, 0, 0, 1, TEXT("Assign Move failed."));
+			}
+			// Test unbinded Attribute
+			{
+				SAttributeAttribute::ManagedSlateAttributeType Attribute = SAttributeAttribute::ManagedSlateAttributeType(Widget.ToSharedRef());
+				TAttribute<FLocalConstructionCounter> Attribute1;
+				FLocalConstructionCounter Counter = 1;
+				FLocalConstructionCounter::ResetCounter();
+				Attribute.Assign(Attribute1, Counter);
+				AddErrorIfCounterDoNotMatches(0, 0, 0, 1, 0, TEXT("Assign Copy/Copy failed."));
+			}
+			{
+				SAttributeAttribute::ManagedSlateAttributeType Attribute = SAttributeAttribute::ManagedSlateAttributeType(Widget.ToSharedRef());
+				TAttribute<FLocalConstructionCounter> Attribute1;
+				FLocalConstructionCounter Counter = 1;
+				FLocalConstructionCounter::ResetCounter();
+				Attribute.Assign(Attribute1, MoveTemp(Counter));
+				AddErrorIfCounterDoNotMatches(0, 0, 0, 0, 1, TEXT("Assign Copy/Move failed."));
+			}
+			{
+				SAttributeAttribute::ManagedSlateAttributeType Attribute = SAttributeAttribute::ManagedSlateAttributeType(Widget.ToSharedRef());
+				TAttribute<FLocalConstructionCounter> Attribute1;
+				FLocalConstructionCounter Counter = 1;
+				FLocalConstructionCounter::ResetCounter();
+				Attribute.Assign(MoveTemp(Attribute1), Counter);
+				AddErrorIfCounterDoNotMatches(0, 0, 0, 1, 0, TEXT("Assign Move/CCopy failed."));
+			}
+			{
+				SAttributeAttribute::ManagedSlateAttributeType Attribute = SAttributeAttribute::ManagedSlateAttributeType(Widget.ToSharedRef());
+				TAttribute<FLocalConstructionCounter> Attribute1;
+				FLocalConstructionCounter Counter = 1;
+				FLocalConstructionCounter::ResetCounter();
+				Attribute.Assign(MoveTemp(Attribute1), MoveTemp(Counter));
+				AddErrorIfCounterDoNotMatches(0, 0, 0, 0, 1, TEXT("Assign Move/Move failed."));
+			}
+			// Test binded Attribute
+			{
+				SAttributeAttribute::ManagedSlateAttributeType Attribute = SAttributeAttribute::ManagedSlateAttributeType(Widget.ToSharedRef());
+				TAttribute<FLocalConstructionCounter> Attribute1 = MakeAttributeLambda([]() {return FLocalConstructionCounter(1);});
+				FLocalConstructionCounter Counter = 1;
+				FLocalConstructionCounter::ResetCounter();
+				Attribute.Assign(Attribute1, Counter);
+				AddErrorIfCounterDoNotMatches(0, 0, 0, 0, 0, TEXT("Bind Copy with binded attribute failed."));
+			}
+			{
+				SAttributeAttribute::ManagedSlateAttributeType Attribute = SAttributeAttribute::ManagedSlateAttributeType(Widget.ToSharedRef());
+				TAttribute<FLocalConstructionCounter> Attribute1 = MakeAttributeLambda([]() {return FLocalConstructionCounter(1);});
+				FLocalConstructionCounter Counter = 1;
+				FLocalConstructionCounter::ResetCounter();
+				Attribute.Assign(Attribute1, MoveTemp(Counter));
+				AddErrorIfCounterDoNotMatches(0, 0, 0, 0, 0, TEXT("Bind Move with binded attribute failed."));
+			}
+			{
+				SAttributeAttribute::ManagedSlateAttributeType Attribute = SAttributeAttribute::ManagedSlateAttributeType(Widget.ToSharedRef());
+				TAttribute<FLocalConstructionCounter> Attribute1 = MakeAttributeLambda([]() {return FLocalConstructionCounter(1);});
+				FLocalConstructionCounter Counter = 1;
+				FLocalConstructionCounter::ResetCounter();
+				Attribute.Assign(MoveTemp(Attribute1), Counter);
+				AddErrorIfCounterDoNotMatches(0, 0, 0, 0, 0, TEXT("Assign Copy with binded attribute failed."));
+			}
+			{
+				SAttributeAttribute::ManagedSlateAttributeType Attribute = SAttributeAttribute::ManagedSlateAttributeType(Widget.ToSharedRef());
+				TAttribute<FLocalConstructionCounter> Attribute1 = MakeAttributeLambda([]() {return FLocalConstructionCounter(1);});
+				FLocalConstructionCounter Counter = 1;
+				FLocalConstructionCounter::ResetCounter();
+				Attribute.Assign(MoveTemp(Attribute1), MoveTemp(Counter));
+				AddErrorIfCounterDoNotMatches(0, 0, 0, 0, 0, TEXT("Assign Move with binded attribute failed."));
+			}
+			// Test set Attribute
+			{
+				SAttributeAttribute::ManagedSlateAttributeType Attribute = SAttributeAttribute::ManagedSlateAttributeType(Widget.ToSharedRef());
+				FLocalConstructionCounter Counter = 1;
+				TAttribute<FLocalConstructionCounter> Attribute1 = Counter;
+				FLocalConstructionCounter::ResetCounter();
+				Attribute.Assign(Attribute1, Counter);
+				AddErrorIfCounterDoNotMatches(0, 0, 0, 1, 0, TEXT("Assign Set Copy failed."));
+			}
+			{
+				SAttributeAttribute::ManagedSlateAttributeType Attribute = SAttributeAttribute::ManagedSlateAttributeType(Widget.ToSharedRef());
+				FLocalConstructionCounter Counter = 1;
+				TAttribute<FLocalConstructionCounter> Attribute1 = Counter;
+				FLocalConstructionCounter::ResetCounter();
+				Attribute.Assign(Attribute1, MoveTemp(Counter));
+				AddErrorIfCounterDoNotMatches(0, 0, 0, 1, 0, TEXT("Assign Set Copy failed."));
+			}
+			{
+				SAttributeAttribute::ManagedSlateAttributeType Attribute = SAttributeAttribute::ManagedSlateAttributeType(Widget.ToSharedRef());
+				FLocalConstructionCounter Counter = 1;
+				TAttribute<FLocalConstructionCounter> Attribute1 = Counter;
+				FLocalConstructionCounter::ResetCounter();
+				Attribute.Assign(MoveTemp(Attribute1), Counter);
+				AddErrorIfCounterDoNotMatches(0, 0, 0, 0, 1, TEXT("Assign Set Move failed."));
+			}
+			{
+				SAttributeAttribute::ManagedSlateAttributeType Attribute = SAttributeAttribute::ManagedSlateAttributeType(Widget.ToSharedRef());
+				FLocalConstructionCounter Counter = 1;
+				TAttribute<FLocalConstructionCounter> Attribute1 = Counter;
+				FLocalConstructionCounter::ResetCounter();
+				Attribute.Assign(MoveTemp(Attribute1), MoveTemp(Counter));
+				AddErrorIfCounterDoNotMatches(0, 0, 0, 0, 1, TEXT("Assign Set Move failed."));
+			}
+		}
+	}
+
 	return true;
 }
 
diff --git a/Engine/Source/Runtime/SlateCore/Private/Types/SlateAttribute.cpp b/Engine/Source/Runtime/SlateCore/Private/Types/SlateAttribute.cpp
--- a/Engine/Source/Runtime/SlateCore/Private/Types/SlateAttribute.cpp
+++ b/Engine/Source/Runtime/SlateCore/Private/Types/SlateAttribute.cpp
@@ -79,6 +79,14 @@ namespace SlateAttributePrivate
 	}
 
 
+	ISlateAttributeGetter* FSlateAttributeImpl::ProtectedFindGetter(const SWidget& OwningWidget, ESlateAttributeType AttributeType) const
+	{
+		TestAttributeAddress(OwningWidget, *this, AttributeType);
+
+		return FSlateAttributeMetaData::GetAttributeGetter(OwningWidget, *this);
+	}
+	
+	
 	FDelegateHandle FSlateAttributeImpl::ProtectedFindGetterHandle(const SWidget& OwningWidget, ESlateAttributeType AttributeType) const
 	{
 		TestAttributeAddress(OwningWidget, *this, AttributeType);
diff --git a/Engine/Source/Runtime/SlateCore/Private/Types/SlateAttributeMetaData.cpp b/Engine/Source/Runtime/SlateCore/Private/Types/SlateAttributeMetaData.cpp
--- a/Engine/Source/Runtime/SlateCore/Private/Types/SlateAttributeMetaData.cpp
+++ b/Engine/Source/Runtime/SlateCore/Private/Types/SlateAttributeMetaData.cpp
@@ -362,6 +362,20 @@ bool FSlateAttributeMetaData::IsAttributeBound(const SWidget& OwningWidget, cons
 }
 
 
+SlateAttributePrivate::ISlateAttributeGetter* FSlateAttributeMetaData::GetAttributeGetter(const SWidget& OwningWidget, const FSlateAttributeBase& Attribute)
+{
+	if (FSlateAttributeMetaData* AttributeMetaData = FSlateAttributeMetaData::FindMetaData(OwningWidget))
+	{
+		const int32 FoundIndex = AttributeMetaData->IndexOfAttribute(Attribute);
+		if (FoundIndex != INDEX_NONE)
+		{
+			return AttributeMetaData->Attributes[FoundIndex].Getter.Get();
+		}
+	}
+	return nullptr;
+}
+
+
 FDelegateHandle FSlateAttributeMetaData::GetAttributeGetterHandle(const SWidget& OwningWidget, const FSlateAttributeBase& Attribute)
 {
 	if (FSlateAttributeMetaData* AttributeMetaData = FSlateAttributeMetaData::FindMetaData(OwningWidget))
diff --git a/Engine/Source/Runtime/SlateCore/Public/Types/SlateAttribute.h b/Engine/Source/Runtime/SlateCore/Public/Types/SlateAttribute.h
--- a/Engine/Source/Runtime/SlateCore/Public/Types/SlateAttribute.h
+++ b/Engine/Source/Runtime/SlateCore/Public/Types/SlateAttribute.h
@@ -179,6 +179,41 @@ struct TSlateAttributeInvalidationReason
 };
 
 
+/** A structure used to help the user identify deprecated TAttribute that are now TSlateAttribute. */
+template<typename ObjectType>
+struct FSlateDeprecatedTAttribute
+{
+	FSlateDeprecatedTAttribute() = default;
+
+	using FGetter = typename TAttribute<ObjectType>::FGetter;
+
+	template<typename OtherType>
+	FSlateDeprecatedTAttribute(const OtherType& InInitialValue)	{ }
+
+	FSlateDeprecatedTAttribute(ObjectType&& InInitialValue)	{ }
+
+	template<class SourceType>
+	FSlateDeprecatedTAttribute(TSharedRef<SourceType> InUserObject, typename FGetter::template TSPMethodDelegate_Const< SourceType >::FMethodPtr InMethodPtr) {}
+
+	template< class SourceType >
+	FSlateDeprecatedTAttribute(SourceType* InUserObject, typename FGetter::template TSPMethodDelegate_Const< SourceType >::FMethodPtr InMethodPtr) { }
+
+	bool IsSet() const { return false; }
+
+	template<typename OtherType>
+	void Set(const OtherType& InNewValue) {}
+
+	const ObjectType& Get(const ObjectType& DefaultValue) const { return DefaultValue; }
+	const ObjectType& Get() const { static ObjectType Temp; return Temp; }
+	FGetter GetBinding() const { return false; }
+
+	void Bind(const FGetter& InGetter) {}
+	bool IsBound() const { return false; }
+
+	bool IdenticalTo(const TAttribute<ObjectType>& InOther) const { return false; }
+};
+
+
 /** Base struct of all SlateAttribute type. */
 struct FSlateAttributeBase
 {
@@ -206,6 +241,8 @@ namespace SlateAttributePrivate
 	class ISlateAttributeGetter;
 	template<typename ObjectType, typename InvalidationReasonPredicate, typename FComparePredicate, ESlateAttributeType AttributeType>
 	struct TSlateAttributeBase;
+	template<typename AttributeMemberType>
+	struct TSlateMemberAttributeRef;
 
 
 	/** */
@@ -245,6 +282,7 @@ namespace SlateAttributePrivate
 		void ProtectedRegisterAttribute(SWidget& Widget, ESlateAttributeType AttributeType, TUniquePtr<ISlateAttributeGetter>&& Wrapper);
 		void ProtectedInvalidateWidget(SWidget& Widget, ESlateAttributeType AttributeType, EInvalidateWidgetReason InvalidationReason) const;
 		bool ProtectedIsBound(const SWidget& Widget, ESlateAttributeType AttributeType) const;
+		ISlateAttributeGetter* ProtectedFindGetter(const SWidget& Widget, ESlateAttributeType AttributeType) const;
 		FDelegateHandle ProtectedFindGetterHandle(const SWidget& Widget, ESlateAttributeType AttributeType) const;
 		bool ProtectedIsIdenticalTo(const SWidget& Widget, ESlateAttributeType AttributeType, const FSlateAttributeBase& Other, const bool bHasSameValue) const;
 		bool ProtectedIsIdenticalToAttribute(const SWidget& Widget, ESlateAttributeType AttributeType, const void* Other, const bool bHasSameValue) const;
@@ -265,6 +303,9 @@ namespace SlateAttributePrivate
 	struct TSlateAttributeBase : public FSlateAttributeImpl
 	{
 	public:
+		template<typename AttributeMemberType>
+		friend struct TSlateMemberAttributeRef;
+
 		using ObjectType = InObjectType;
 		using FInvalidationReasonPredicate = InInvalidationReasonPredicate;
 		using FGetter = typename TAttribute<ObjectType>::FGetter;
@@ -315,7 +356,7 @@ namespace SlateAttributePrivate
 		{
 		}
 
-		TSlateAttributeBase(SWidget& Widget, const ObjectType& InitialValue, const FGetter& Getter)
+		TSlateAttributeBase(SWidget& Widget, const FGetter& Getter, const ObjectType& InitialValue)
 			: Value(InitialValue)
 #if UE_SLATE_WITH_MEMBER_ATTRIBUTE_DEBUGGING
 			, Debug_OwningWidget(&Widget)
@@ -323,38 +364,35 @@ namespace SlateAttributePrivate
 		{
 			if (Getter.IsBound())
 			{
-				TUniquePtr<ISlateAttributeGetter> Wrapper = MakeUniqueGetter(*this, Getter);
-				ProtectedRegisterAttribute(Widget, InAttributeType, MoveTemp(Wrapper));
+				ConstructWrapper(Widget, Getter);
 			}
 		}
 
-		TSlateAttributeBase(SWidget& Widget, const ObjectType& InitialValue, FGetter&& Getter)
-			: Value(InitialValue)
+		TSlateAttributeBase(SWidget& Widget, const FGetter& Getter, ObjectType&& InitialValue)
+			: Value(MoveTemp(InitialValue))
 #if UE_SLATE_WITH_MEMBER_ATTRIBUTE_DEBUGGING
 			, Debug_OwningWidget(&Widget)
 #endif
 		{
 			if (Getter.IsBound())
 			{
-				TUniquePtr<ISlateAttributeGetter> Wrapper = MakeUniqueGetter(*this, MoveTemp(Getter));
-				ProtectedRegisterAttribute(Widget, InAttributeType, MoveTemp(Wrapper));
+				ConstructWrapper(Widget, Getter);
 			}
 		}
 
-		TSlateAttributeBase(SWidget& Widget, const ObjectType& InitialValue, const TAttribute<ObjectType>& Attribute)
-			: Value((Attribute.IsSet() && !Attribute.IsBound()) ? Attribute.Get() : InitialValue)
+		TSlateAttributeBase(SWidget& Widget, FGetter&& Getter, const ObjectType& InitialValue)
+			: Value(InitialValue)
 #if UE_SLATE_WITH_MEMBER_ATTRIBUTE_DEBUGGING
 			, Debug_OwningWidget(&Widget)
 #endif
 		{
-			if (Attribute.IsBound())
+			if (Getter.IsBound())
 			{
-				TUniquePtr<ISlateAttributeGetter> Wrapper = MakeUniqueGetter(*this, Attribute.GetBinding());
-				ProtectedRegisterAttribute(Widget, InAttributeType, MoveTemp(Wrapper));
+				ConstructWrapper(Widget, MoveTemp(Getter));
 			}
 		}
 
-		TSlateAttributeBase(SWidget& Widget, ObjectType&& InitialValue, const FGetter& Getter)
+		TSlateAttributeBase(SWidget& Widget, FGetter&& Getter, ObjectType&& InitialValue)
 			: Value(MoveTemp(InitialValue))
 #if UE_SLATE_WITH_MEMBER_ATTRIBUTE_DEBUGGING
 			, Debug_OwningWidget(&Widget)
@@ -362,40 +400,37 @@ namespace SlateAttributePrivate
 		{
 			if (Getter.IsBound())
 			{
-				TUniquePtr<ISlateAttributeGetter> Wrapper = MakeUniqueGetter(*this, Getter);
-				ProtectedRegisterAttribute(Widget, InAttributeType, MoveTemp(Wrapper));
+				ConstructWrapper(Widget, MoveTemp(Getter));
 			}
 		}
 
-		TSlateAttributeBase(SWidget& Widget, ObjectType&& InitialValue, FGetter&& Getter)
-			: Value(MoveTemp(InitialValue))
+		TSlateAttributeBase(SWidget& Widget, const TAttribute<ObjectType>& Attribute, const ObjectType& InitialValue)
+			: Value((Attribute.IsSet() && !Attribute.IsBound()) ? Attribute.Get() : InitialValue)
 #if UE_SLATE_WITH_MEMBER_ATTRIBUTE_DEBUGGING
 			, Debug_OwningWidget(&Widget)
 #endif
 		{
-			if (Getter.IsBound())
+			if (Attribute.IsBound())
 			{
-				TUniquePtr<ISlateAttributeGetter> Wrapper = MakeUniqueGetter(*this, MoveTemp(Getter));
-				ProtectedRegisterAttribute(Widget, InAttributeType, MoveTemp(Wrapper));
+				ConstructWrapper(Widget, Attribute.GetBinding());
 			}
 		}
 
-		TSlateAttributeBase(SWidget& Widget, ObjectType&& InitialValue, const TAttribute<ObjectType>& Attribute)
-			: Value((Attribute.IsSet() && !Attribute.IsBound()) ? Attribute.Get() : MoveTemp(InitialValue))
+		TSlateAttributeBase(SWidget& Widget, TAttribute<ObjectType>&& Attribute, ObjectType&& InitialValue)
+			: Value((Attribute.IsSet() && !Attribute.IsBound()) ? MoveTemp(Attribute.Steal().template Get<ObjectType>()) : MoveTemp(InitialValue))
 #if UE_SLATE_WITH_MEMBER_ATTRIBUTE_DEBUGGING
 			, Debug_OwningWidget(&Widget)
 #endif
 		{
 			if (Attribute.IsBound())
 			{
-				TUniquePtr<ISlateAttributeGetter> Wrapper = MakeUniqueGetter(*this, Attribute.GetBinding());
-				ProtectedRegisterAttribute(Widget, InAttributeType, MoveTemp(Wrapper));
+				ConstructWrapper(Widget, MoveTemp(Attribute.Steal().template Get<FGetter>()));
 			}
 		}
 
 	public:
-		/** @return the SlateAttribute value. If the SlateAttribute is bound, value is cached only at the end of the frame. */ 
-		const ObjectType& Get() const
+		/** @return the SlateAttribute cached value. If the SlateAttribute is bound, the value will be cached at the end of the every frame. */ 
+		UE_NODISCARD const ObjectType& Get() const
 		{
 			return Value;
 		}
@@ -436,6 +471,7 @@ namespace SlateAttributePrivate
 			}
 		}
 
+	public:
 		/**
 		 * Bind the SlateAttribute to the Getter function.
 		 * (If enabled) Update the value from the Getter. If the value is different, then invalidate the widget.
@@ -446,13 +482,7 @@ namespace SlateAttributePrivate
 			VerifyOwningWidget(Widget);
 			if (Getter.IsBound())
 			{
-				const FDelegateHandle PreviousGetterHandle = ProtectedFindGetterHandle(Widget, InAttributeType);
-				if (PreviousGetterHandle != Getter.GetHandle())
-				{
-					TUniquePtr<ISlateAttributeGetter> Wrapper = MakeUniqueGetter(*this, Getter);
-					ProtectedRegisterAttribute(Widget, InAttributeType, MoveTemp(Wrapper));
-					UpdateNowOnBind(Widget);
-				}
+				AssignBinding(Widget, Getter);
 			}
 			else
 			{
@@ -470,13 +500,7 @@ namespace SlateAttributePrivate
 			VerifyOwningWidget(Widget);
 			if (Getter.IsBound())
 			{
-				const FDelegateHandle PreviousGetterHandle = ProtectedFindGetterHandle(Widget, InAttributeType);
-				if (PreviousGetterHandle != Getter.GetHandle())
-				{
-					TUniquePtr<ISlateAttributeGetter> Wrapper = MakeUniqueGetter(*this, MoveTemp(Getter));
-					ProtectedRegisterAttribute(Widget, InAttributeType, MoveTemp(Wrapper));
-					UpdateNowOnBind(Widget);
-				}
+				AssignBinding(Widget, MoveTemp(Getter));
 			}
 			else
 			{
@@ -512,13 +536,7 @@ namespace SlateAttributePrivate
 			VerifyOwningWidget(Widget);
 			if (OtherAttribute.IsBound())
 			{
-				const FDelegateHandle PreviousGetterHandle = ProtectedFindGetterHandle(Widget, InAttributeType);
-				if (PreviousGetterHandle != OtherAttribute.GetBinding().GetHandle())
-				{
-					TUniquePtr<ISlateAttributeGetter> Wrapper = MakeUniqueGetter(*this, OtherAttribute.GetBinding());
-					ProtectedRegisterAttribute(Widget, InAttributeType, MoveTemp(Wrapper));
-					UpdateNowOnBind(Widget);
-				}
+				AssignBinding(Widget, OtherAttribute.GetBinding());
 			}
 			else if (OtherAttribute.IsSet())
 			{
@@ -530,30 +548,29 @@ namespace SlateAttributePrivate
 			}
 		}
 
-		/**
-		 * Bind the SlateAttribute to the Attribute Getter function (if it exist).
-		 * (If enabled) Update the value from the getter. If the value is different, then invalidate the widget.
-		 * The SlateAttribute will now be updated every frame from the Getter.
-		 * Or
-		 * Set the SlateAttribute's value if the Attribute is not bound but is set.
-		 * This will Unbind any previously bound getter.
-		 * If the value is different, then invalidate the widget.
-		 * Or
-		 * Unbind the SlateAttribute to DefaultValue if the Attribute is not bound and not set.
-		 * @see Set
-		 */
+		void Assign(SWidget& Widget, TAttribute<ObjectType>&& OtherAttribute)
+		{
+			VerifyOwningWidget(Widget);
+			if (OtherAttribute.IsBound())
+			{
+				AssignBinding(Widget, MoveTemp(OtherAttribute.Steal().template Get<FGetter>()));
+			}
+			else if (OtherAttribute.IsSet())
+			{
+				Set(Widget, MoveTemp(OtherAttribute.Steal().template Get<ObjectType>()));
+			}
+			else
+			{
+				ProtectedUnregisterAttribute(Widget, InAttributeType);
+			}
+		}
+
 		void Assign(SWidget& Widget, const TAttribute<ObjectType>& OtherAttribute, const ObjectType& DefaultValue)
 		{
 			VerifyOwningWidget(Widget);
 			if (OtherAttribute.IsBound())
 			{
-				const FDelegateHandle PreviousGetterHandle = ProtectedFindGetterHandle(Widget, InAttributeType);
-				if (PreviousGetterHandle != OtherAttribute.GetBinding().GetHandle())
-				{
-					TUniquePtr<ISlateAttributeGetter> Wrapper = MakeUniqueGetter(*this, OtherAttribute.GetBinding());
-					ProtectedRegisterAttribute(Widget, InAttributeType, MoveTemp(Wrapper));
-					UpdateNowOnBind(Widget);
-				}
+				AssignBinding(Widget, OtherAttribute.GetBinding());
 			}
 			else if (OtherAttribute.IsSet())
 			{
@@ -565,30 +582,29 @@ namespace SlateAttributePrivate
 			}
 		}
 
-		/**
-		 * Bind the SlateAttribute to the Attribute Getter function (if it exist).
-		 * (If enabled) Update the value from the getter. If the value is different, then invalidate the widget.
-		 * The SlateAttribute will now be updated every frame from the Getter.
-		 * Or
-		 * Set the SlateAttribute's value if the Attribute is not bound but is set.
-		 * This will Unbind any previously bound getter.
-		 * If the value is different, then invalidate the widget.
-		 * Or
-		 * Unbind the SlateAttribute to DefaultValue if the Attribute is not bound and not set.
-		 * @see Set
-		 */
+		void Assign(SWidget& Widget, TAttribute<ObjectType>&& OtherAttribute, const ObjectType& DefaultValue)
+		{
+			VerifyOwningWidget(Widget);
+			if (OtherAttribute.IsBound())
+			{
+				AssignBinding(Widget, MoveTemp(OtherAttribute.Steal().template Get<FGetter>()));
+			}
+			else if (OtherAttribute.IsSet())
+			{
+				Set(Widget, MoveTemp(OtherAttribute.Steal().template Get<ObjectType>()));
+			}
+			else
+			{
+				Set(Widget, DefaultValue);
+			}
+		}
+
 		void Assign(SWidget& Widget, const TAttribute<ObjectType>& OtherAttribute, ObjectType&& DefaultValue)
 		{
 			VerifyOwningWidget(Widget);
 			if (OtherAttribute.IsBound())
 			{
-				const FDelegateHandle PreviousGetterHandle = ProtectedFindGetterHandle(Widget, InAttributeType);
-				if (PreviousGetterHandle != OtherAttribute.GetBinding().GetHandle())
-				{
-					TUniquePtr<ISlateAttributeGetter> Wrapper = MakeUniqueGetter(*this, OtherAttribute.GetBinding());
-					ProtectedRegisterAttribute(Widget, InAttributeType, MoveTemp(Wrapper));
-					UpdateNowOnBind(Widget);
-				}
+				AssignBinding(Widget, OtherAttribute.GetBinding());
 			}
 			else if (OtherAttribute.IsSet())
 			{
@@ -600,6 +616,23 @@ namespace SlateAttributePrivate
 			}
 		}
 
+		void Assign(SWidget& Widget, TAttribute<ObjectType>&& OtherAttribute, ObjectType&& DefaultValue)
+		{
+			VerifyOwningWidget(Widget);
+			if (OtherAttribute.IsBound())
+			{
+				AssignBinding(Widget, MoveTemp(OtherAttribute.Steal().template Get<FGetter>()));
+			}
+			else if (OtherAttribute.IsSet())
+			{
+				Set(Widget, MoveTemp(OtherAttribute.Steal().template Get<ObjectType>()));
+			}
+			else
+			{
+				Set(Widget, MoveTemp(DefaultValue));
+			}
+		}
+
 		/**
 		 * Remove the Getter function.
 		 * The Slate Attribute will not be updated anymore and will keep its current cached value.
@@ -610,16 +643,26 @@ namespace SlateAttributePrivate
 			ProtectedUnregisterAttribute(Widget, InAttributeType);
 		}
 
-public:
+	public:
+		/** Build a Attribute from this SlateAttribute. */
+		UE_NODISCARD TAttribute<ObjectType> ToAttribute(const SWidget& Widget) const
+		{
+			if (ISlateAttributeGetter* Delegate = ProtectedFindGetter(Widget, InAttributeType))
+			{
+				return TAttribute<ObjectType>::Create(static_cast<FSlateAttributeGetterWrapper<TSlateAttributeBase>*>(Delegate)->GetDelegate());
+			}
+			return TAttribute<ObjectType>(Get());
+		}
+
 		/** @return True if the SlateAttribute is bound to a getter function. */
-		bool IsBound(const SWidget& Widget) const
+		UE_NODISCARD bool IsBound(const SWidget& Widget) const
 		{
 			VerifyOwningWidget(Widget);
 			return ProtectedIsBound(Widget, InAttributeType);
 		}
 
 		/** @return True if they have the same Getter or the same value. */
-		bool IsIdenticalTo(const SWidget& Widget, const TSlateAttributeBase& Other) const
+		UE_NODISCARD bool IsIdenticalTo(const SWidget& Widget, const TSlateAttributeBase& Other) const
 		{
 			VerifyOwningWidget(Widget);
 			FDelegateHandle ThisDelegateHandle = ProtectedFindGetterHandle(Widget, InAttributeType);
@@ -630,13 +673,13 @@ public:
 				{
 					return true;
 				}
-				return Get() == Other.Get();
+				return FComparePredicate{}.operator()(Get(), Other.Get());
 			}
 			return false;
 		}
 
-		/**  if they have the same Getter or, if the Attribute is set, the same value. */
-		bool IsIdenticalTo(const SWidget& Widget, const TAttribute<ObjectType>& Other) const
+		/** @return True if they have the same Getter or, if the Attribute is set, the same value. */
+		UE_NODISCARD bool IsIdenticalTo(const SWidget& Widget, const TAttribute<ObjectType>& Other) const
 		{
 			VerifyOwningWidget(Widget);
 			FDelegateHandle ThisDelegateHandle = ProtectedFindGetterHandle(Widget, InAttributeType);
@@ -644,7 +687,40 @@ public:
 			{
 				return Other.GetBinding().GetHandle() == ThisDelegateHandle;
 			}
-			return !ThisDelegateHandle.IsValid() && Other.IsSet() && Get() == Other.Get();
+			return !ThisDelegateHandle.IsValid() && Other.IsSet() && FComparePredicate{}.operator()(Get(), Other.Get());
+		}
+
+	private:
+		void ConstructWrapper(SWidget& Widget, const FGetter& Getter)
+		{
+			TUniquePtr<ISlateAttributeGetter> Wrapper = MakeUniqueGetter(*this, Getter);
+			ProtectedRegisterAttribute(Widget, InAttributeType, MoveTemp(Wrapper));
+			UpdateNowOnBind(Widget);
+		}
+
+		void ConstructWrapper(SWidget& Widget, FGetter&& Getter)
+		{
+			TUniquePtr<ISlateAttributeGetter> Wrapper = MakeUniqueGetter(*this, MoveTemp(Getter));
+			ProtectedRegisterAttribute(Widget, InAttributeType, MoveTemp(Wrapper));
+			UpdateNowOnBind(Widget);
+		}
+
+		void AssignBinding(SWidget& Widget, const FGetter& Getter)
+		{
+			const FDelegateHandle PreviousGetterHandle = ProtectedFindGetterHandle(Widget, InAttributeType);
+			if (PreviousGetterHandle != Getter.GetHandle())
+			{
+				ConstructWrapper(Widget, Getter);
+			}
+		}
+
+		void AssignBinding(SWidget& Widget, FGetter&& Getter)
+		{
+			const FDelegateHandle PreviousGetterHandle = ProtectedFindGetterHandle(Widget, InAttributeType);
+			if (PreviousGetterHandle != Getter.GetHandle())
+			{
+				ConstructWrapper(Widget, MoveTemp(Getter));
+			}
 		}
 
 	private:
@@ -701,6 +777,11 @@ public:
 				return Getter.GetHandle();
 			}
 
+			const FGetter& GetDelegate() const
+			{
+				return Getter;
+			}
+
 		private:
 			/** Getter function to fetch the new value of the SlateAttribute. */
 			FGetter Getter;
@@ -778,67 +859,25 @@ public:
 #endif
 
 		template<typename WidgetType, typename U = typename std::enable_if<std::is_base_of<SWidget, WidgetType>::value>::type>
-		TSlateMemberAttribute(WidgetType& Widget)
+		explicit TSlateMemberAttribute(WidgetType& Widget)
 			: Super(Widget)
 		{
 			VerifyAttributeAddress(Widget, this);
 		}
 
 		template<typename WidgetType, typename U = typename std::enable_if<std::is_base_of<SWidget, WidgetType>::value>::type>
-		TSlateMemberAttribute(WidgetType& Widget, const ObjectType& InValue)
+		explicit TSlateMemberAttribute(WidgetType& Widget, const ObjectType& InValue)
 			: Super(Widget, InValue)
 		{
 			VerifyAttributeAddress(Widget, this);
 		}
 
 		template<typename WidgetType, typename U = typename std::enable_if<std::is_base_of<SWidget, WidgetType>::value>::type>
-		TSlateMemberAttribute(WidgetType& Widget, ObjectType&& InValue)
+		explicit TSlateMemberAttribute(WidgetType& Widget, ObjectType&& InValue)
 			: Super(Widget, MoveTemp(InValue))
 		{
 			VerifyAttributeAddress(Widget, this);
 		}
-
-		template<typename WidgetType, typename U = typename std::enable_if<std::is_base_of<SWidget, WidgetType>::value>::type>
-		TSlateMemberAttribute(WidgetType& Widget, const ObjectType& InitialValue, const FGetter& Getter)
-			: Super(Widget, InitialValue, Getter)
-		{
-			VerifyAttributeAddress(Widget, this);
-		}
-
-		template<typename WidgetType, typename U = typename std::enable_if<std::is_base_of<SWidget, WidgetType>::value>::type>
-		TSlateMemberAttribute(WidgetType& Widget, const ObjectType& InitialValue, FGetter&& Getter)
-			: Super(Widget, InitialValue, MoveTemp(Getter))
-		{
-			VerifyAttributeAddress(Widget, this);
-		}
-
-		template<typename WidgetType, typename U = typename std::enable_if<std::is_base_of<SWidget, WidgetType>::value>::type>
-		TSlateMemberAttribute(WidgetType& Widget, const ObjectType& InitialValue, const TAttribute<ObjectType>& Attribute)
-			: Super(Widget, InitialValue, Attribute)
-		{
-			VerifyAttributeAddress(Widget, this);
-		}
-
-		template<typename WidgetType, typename U = typename std::enable_if<std::is_base_of<SWidget, WidgetType>::value>::type>
-		TSlateMemberAttribute(WidgetType& Widget, ObjectType&& InitialValue, const FGetter& Getter)
-			: Super(Widget, MoveTemp(InitialValue), Getter)
-		{
-			VerifyAttributeAddress(Widget, this);
-		}
-
-		template<typename WidgetType, typename U = typename std::enable_if<std::is_base_of<SWidget, WidgetType>::value>::type>
-		TSlateMemberAttribute(WidgetType& Widget, ObjectType&& InitialValue, FGetter&& Getter)
-			: Super(Widget, MoveTemp(InitialValue), MoveTemp(Getter))
-		{
-			VerifyAttributeAddress(Widget, this);
-		}
-
-		template<typename WidgetType, typename U = typename std::enable_if<std::is_base_of<SWidget, WidgetType>::value>::type>
-		TSlateMemberAttribute(WidgetType& Widget, ObjectType&& InitialValue, const TAttribute<ObjectType>& Attribute)
-			: Super(Widget, MoveTemp(InitialValue), Attribute)
-		{
-			VerifyAttributeAddress(Widget, this);
-		}
 	};
 
 
@@ -892,55 +931,55 @@ public:
 			return *this;
 		}
 
-		TSlateManagedAttribute(TSharedRef<SWidget> Widget)
+		explicit TSlateManagedAttribute(TSharedRef<SWidget> Widget)
 			: Super(Widget.Get())
 			, ManagedWidget(Widget)
 		{ }
 
-		TSlateManagedAttribute(TSharedRef<SWidget> Widget, const ObjectType& InValue)
+		explicit TSlateManagedAttribute(TSharedRef<SWidget> Widget, const ObjectType& InValue)
 			: Super(Widget.Get(), InValue)
 			, ManagedWidget(Widget)
 		{
 		}
 
-		TSlateManagedAttribute(TSharedRef<SWidget> Widget, ObjectType&& InValue)
+		explicit TSlateManagedAttribute(TSharedRef<SWidget> Widget, ObjectType&& InValue)
 			: Super(Widget.Get(), MoveTemp(InValue))
 			, ManagedWidget(Widget)
 		{
 		}
 
-		TSlateManagedAttribute(TSharedRef<SWidget> Widget, const ObjectType& InitialValue, const FGetter& Getter)
-			: Super(Widget.Get(), InitialValue, Getter)
+		explicit TSlateManagedAttribute(TSharedRef<SWidget> Widget, const FGetter& Getter, const ObjectType& InitialValue)
+			: Super(Widget.Get(), Getter, InitialValue)
 			, ManagedWidget(Widget)
 		{
 		}
 
-		TSlateManagedAttribute(TSharedRef<SWidget> Widget, const ObjectType& InitialValue, FGetter&& Getter)
-			: Super(Widget.Get(), InitialValue, MoveTemp(Getter))
+		explicit TSlateManagedAttribute(TSharedRef<SWidget> Widget, const FGetter& Getter, ObjectType&& InitialValue)
+			: Super(Widget.Get(), Getter, MoveTemp(InitialValue))
 			, ManagedWidget(Widget)
 		{
 		}
 
-		TSlateManagedAttribute(TSharedRef<SWidget> Widget, const ObjectType& InitialValue, const TAttribute<ObjectType>& Attribute)
-			: Super(Widget.Get(), InitialValue, Attribute)
+		explicit TSlateManagedAttribute(TSharedRef<SWidget> Widget, FGetter&& Getter, const ObjectType& InitialValue)
+			: Super(Widget.Get(), MoveTemp(Getter), InitialValue)
 			, ManagedWidget(Widget)
 		{
 		}
 
-		TSlateManagedAttribute(TSharedRef<SWidget> Widget, ObjectType&& InitialValue, const FGetter& Getter)
-			: Super(Widget.Get(), MoveTemp(InitialValue), Getter)
+		explicit TSlateManagedAttribute(TSharedRef<SWidget> Widget, FGetter&& Getter, ObjectType&& InitialValue)
+			: Super(Widget.Get(), MoveTemp(Getter), MoveTemp(InitialValue))
 			, ManagedWidget(Widget)
 		{
 		}
 
-		TSlateManagedAttribute(TSharedRef<SWidget> Widget, ObjectType&& InitialValue, FGetter&& Getter)
-			: Super(Widget.Get(), MoveTemp(InitialValue), MoveTemp(Getter))
+		explicit TSlateManagedAttribute(TSharedRef<SWidget> Widget, const TAttribute<ObjectType>& Attribute, const ObjectType& InitialValue)
+			: Super(Widget.Get(), Attribute, InitialValue)
 			, ManagedWidget(Widget)
 		{
 		}
 
-		TSlateManagedAttribute(TSharedRef<SWidget> Widget, ObjectType&& InitialValue, const TAttribute<ObjectType>& Attribute)
-			: Super(Widget.Get(), MoveTemp(InitialValue), Attribute)
+		explicit TSlateManagedAttribute(TSharedRef<SWidget> Widget, TAttribute<ObjectType>&& Attribute, ObjectType&& InitialValue)
+			: Super(Widget.Get(), MoveTemp(Attribute), MoveTemp(InitialValue))
 			, ManagedWidget(Widget)
 		{
 		}
@@ -992,20 +1031,19 @@ public:
 			}
 		}
 
-		template<typename WidgetType, typename U = typename std::enable_if<std::is_base_of<SWidget, WidgetType>::value>::type>
-		void Bind(typename FGetter::template TSPMethodDelegate_Const<WidgetType>::FMethodPtr MethodPtr)
+		void Assign(const TAttribute<ObjectType>& OtherAttribute)
 		{
 			if (TSharedPtr<SWidget> Pin = ManagedWidget.Pin())
 			{
-				Super::Bind(*Pin.Get(), FGetter::CreateSP(Pin, MethodPtr));
+				Super::Assign(*Pin.Get(), OtherAttribute);
 			}
 		}
 
-		void Assign(const TAttribute<ObjectType>& OtherAttribute)
+		void Assign(TAttribute<ObjectType>&& OtherAttribute)
 		{
 			if (TSharedPtr<SWidget> Pin = ManagedWidget.Pin())
 			{
-				Super::Assign(*Pin.Get(), OtherAttribute);
+				Super::Assign(*Pin.Get(), MoveTemp(OtherAttribute));
 			}
 		}
 
@@ -1025,6 +1063,22 @@ public:
 			}
 		}
 
+		void Assign(TAttribute<ObjectType>&& OtherAttribute, const ObjectType& DefaultValue)
+		{
+			if (TSharedPtr<SWidget> Pin = ManagedWidget.Pin())
+			{
+				Super::Assign(*Pin.Get(), MoveTemp(OtherAttribute), DefaultValue);
+			}
+		}
+
+		void Assign(TAttribute<ObjectType>&& OtherAttribute, ObjectType&& DefaultValue)
+		{
+			if (TSharedPtr<SWidget> Pin = ManagedWidget.Pin())
+			{
+				Super::Assign(*Pin.Get(), MoveTemp(OtherAttribute), MoveTemp(DefaultValue));
+			}
+		}
+
 		void Unbind()
 		{
 			if (TSharedPtr<SWidget> Pin = ManagedWidget.Pin())
@@ -1065,4 +1119,151 @@ public:
 		TWeakPtr<SWidget> ManagedWidget;
 	};
 
+
+	/*
+	 * A reference to a SlateAttribute that can be returned and saved for later.
+	 */
+	template<typename AttributeMemberType>
+	struct TSlateMemberAttributeRef
+	{
+	public:
+		using SlateAttributeType = AttributeMemberType;
+		using ObjectType = typename AttributeMemberType::ObjectType;
+		using AttributeType = TAttribute<ObjectType>;
+
+	private:
+		template<typename WidgetType>
+		static void VerifyAttributeAddress(WidgetType const& InWidget, AttributeMemberType const& InAttribute)
+		{
+			checkf((UPTRINT)&InAttribute >= (UPTRINT)&InWidget && (UPTRINT)&InAttribute < (UPTRINT)&InWidget + sizeof(WidgetType),
+				TEXT("The attribute is not a member of the widget."));
+			InAttribute.VerifyOwningWidget(InWidget);
+		}
+
+	public:
+		/** Constructor */
+		TSlateMemberAttributeRef() = default;
+
+		template<typename WidgetType, typename V = typename std::enable_if<std::is_base_of<SWidget, WidgetType>::value>::type>
+		explicit TSlateMemberAttributeRef(WidgetType const& InOwner, AttributeMemberType const& InAttribute)
+			: Owner(InOwner.AsShared())
+			, Attribute(&InAttribute)
+		{
+			VerifyAttributeAddress(InOwner, InAttribute);
+		}
+
+	public:
+		/** @return if the reference is valid. A reference can be invalid if the SWidget is destroyed. */
+		UE_NODISCARD bool IsValid() const
+		{
+			return Owner.IsValid();
+		}
+
+		/** @return the SlateAttribute cached value; undefined when IsValid() returns false. */
+		UE_NODISCARD const ObjectType& Get() const
+		{
+			if (TSharedPtr<const SWidget> Pin = Owner.Pin())
+			{
+				return Attribute->Get();
+			}
+			checkf(false, TEXT("It is an error to call GetValue() on an unset TSlateMemberAttributeRef. Please either check IsValid() or use Get(DefaultValue) instead."));
+			static ObjectType Tmp;
+			return Tmp;
+		}
+
+		/** @return the SlateAttribute cached value or the DefaultValue if the reference is invalid. */
+		UE_NODISCARD const ObjectType& Get(const ObjectType& DefaultValue) const
+		{
+			if (TSharedPtr<const SWidget> Pin = Owner.Pin())
+			{
+				return Attribute->Get();
+			}
+			return DefaultValue;
+		}
+
+		/** Update the cached value and invalidate the widget if needed. */
+		void UpdateValue()
+		{
+			if (TSharedPtr<const SWidget> Pin = Owner.Pin())
+			{
+				const_cast<AttributeMemberType*>(Attribute)->UpdateNow(*Pin.Get());
+			}
+		}
+
+		/**
+		 * Assumes the reference is valid.
+		 * Shorthand for the boilerplace code: MyAttribute.UpdateValueNow(); MyAttribute.Get();
+		 */
+		UE_NODISCARD const ObjectType& UpdateAndGet()
+		{
+			if (TSharedPtr<const SWidget> Pin = Owner.Pin())
+			{
+				const_cast<AttributeMemberType*>(Attribute)->UpdateNow(*Pin.Get());
+				return Attribute->Get();
+			}
+			checkf(false, TEXT("It is an error to call GetValue() on an unset TSlateMemberAttributeRef. Please either check IsValid() or use Get(DefaultValue) instead."));
+		}
+
+		/**
+		 * Shorthand for the boilerplace code: MyAttribute.UpdateValueNow(); MyAttribute.Get(DefaultValue);
+		 * @return the SlateAttribute cached value or the DefaultValue if the reference is invalid.
+		 */
+		UE_NODISCARD const ObjectType& UpdateAndGet(const ObjectType& DefaultValue)
+		{
+			if (TSharedPtr<const SWidget> Pin = Owner.Pin())
+			{
+				const_cast<AttributeMemberType*>(Attribute)->UpdateNow(*Pin.Get());
+				return Attribute->Get();
+			}
+			return DefaultValue;
+		}
+
+		/** Build a Attribute from this SlateAttribute. */
+		UE_NODISCARD TAttribute<ObjectType> ToAttribute() const
+		{
+			if (TSharedPtr<const SWidget> Pin = Owner.Pin())
+			{
+				return Attribute->ToAttribute(*Pin.Get());
+			}
+			return TAttribute<ObjectType>();
+		}
+
+		/** @return True if the SlateAttribute is bound to a getter function. */
+		UE_NODISCARD bool IsBound() const
+		{
+			if (TSharedPtr<const SWidget> Pin = Owner.Pin())
+			{
+				return Attribute->IsBound(*Pin.Get());
+			}
+			return false;
+		}
+
+		/** @return True if they have the same Getter or the same value. */
+		UE_NODISCARD bool IsIdenticalTo(const TSlateMemberAttributeRef& Other) const
+		{
+			TSharedPtr<const SWidget> SelfPin = Owner.Pin();
+			TSharedPtr<const SWidget> OtherPin = Other.Owner.Pin();
+			if (SelfPin == OtherPin && SelfPin)
+			{
+				return Attribute->IsIdenticalTo(*SelfPin.Get(), *Other.Attribute);
+			}
+			return SelfPin == OtherPin;
+		}
+
+		/** @return True if they have the same Getter or, if the Attribute is set, the same value. */
+		UE_NODISCARD bool IsIdenticalTo(const TAttribute<ObjectType>& Other) const
+		{
+			if (TSharedPtr<const SWidget> Pin = Owner.Pin())
+			{
+				return Attribute->IsIdenticalTo(*Pin.Get(), Other);
+			}
+			return !Other.IsSet(); // if the other is not set, then both are invalid.
+		}
+
+	private:
+		TWeakPtr<const SWidget> Owner;
+		AttributeMemberType const* Attribute = nullptr;
+	 };
+
+
 } // SlateAttributePrivate
diff --git a/Engine/Source/Runtime/SlateCore/Public/Types/SlateAttributeDescriptor.h b/Engine/Source/Runtime/SlateCore/Public/Types/SlateAttributeDescriptor.h
--- a/Engine/Source/Runtime/SlateCore/Public/Types/SlateAttributeDescriptor.h
+++ b/Engine/Source/Runtime/SlateCore/Public/Types/SlateAttributeDescriptor.h
@@ -207,9 +207,12 @@ private:
  * @param _Property The TSlateAttribute property
  * @param _Reason The EInvalidationWidgetReason or a static function/lambda that takes a const SWidget& and that returns the invalidation reason.
  */
-#define SLATE_ADD_MEMBER_ATTRIBUTE_DEFINITION(_Initializer, _Property, _Reason) \
+#define SLATE_ADD_MEMBER_ATTRIBUTE_DEFINITION_WITH_NAME(_Initializer, _Name, _Property, _Reason) \
 		static_assert(decltype(_Property)::IsMemberType, "The SlateProperty is not a TSlateAttribute. Do not use SLATE_ADD_MEMBER_ATTRIBUTE_DEFINITION"); \
 		static_assert(!decltype(_Property)::HasDefinedInvalidationReason, "When implementing the SLATE_DECLARE_WIDGET pattern, use TSlateAttribute without the invalidation reason."); \
-		_Initializer.AddMemberAttribute(GET_MEMBER_NAME_CHECKED(PrivateThisType, _Property), STRUCT_OFFSET(PrivateThisType, _Property), FSlateAttributeDescriptor::FInvalidateWidgetReasonAttribute{_Reason})
-		//__VA_ARGS__
+		_Initializer.AddMemberAttribute(_Name, STRUCT_OFFSET(PrivateThisType, _Property), FSlateAttributeDescriptor::FInvalidateWidgetReasonAttribute{_Reason})
+
+#define SLATE_ADD_MEMBER_ATTRIBUTE_DEFINITION(_Initializer, _Property, _Reason) \
+	SLATE_ADD_MEMBER_ATTRIBUTE_DEFINITION_WITH_NAME(_Initializer, GET_MEMBER_NAME_CHECKED(PrivateThisType, _Property), _Property, _Reason)
+
 using FSlateAttributeInitializer = FSlateAttributeDescriptor::FInitializer;
diff --git a/Engine/Source/Runtime/SlateCore/Public/Types/SlateAttributeMetaData.h b/Engine/Source/Runtime/SlateCore/Public/Types/SlateAttributeMetaData.h
--- a/Engine/Source/Runtime/SlateCore/Public/Types/SlateAttributeMetaData.h
+++ b/Engine/Source/Runtime/SlateCore/Public/Types/SlateAttributeMetaData.h
@@ -44,6 +44,7 @@ private:
 	static void InvalidateWidget(SWidget& OwningWidget, const FSlateAttributeBase& Attribute, ESlateAttributeType AttributeType, EInvalidateWidgetReason Reason);
 	static void UpdateAttribute(SWidget& OwningWidget, FSlateAttributeBase& Attribute);
 	static bool IsAttributeBound(const SWidget& OwningWidget, const FSlateAttributeBase& Attribute);
+	static SlateAttributePrivate::ISlateAttributeGetter* GetAttributeGetter(const SWidget& OwningWidget, const FSlateAttributeBase& Attribute);
 	static FDelegateHandle GetAttributeGetterHandle(const SWidget& OwningWidget, const FSlateAttributeBase& Attribute);
 	static void MoveAttribute(const SWidget& OwningWidget, FSlateAttributeBase& NewAttribute, ESlateAttributeType AttributeType, const FSlateAttributeBase* PreviousAttribute);
 
diff --git a/Engine/Source/Runtime/SlateCore/Public/Widgets/SWidget.h b/Engine/Source/Runtime/SlateCore/Public/Widgets/SWidget.h
--- a/Engine/Source/Runtime/SlateCore/Public/Widgets/SWidget.h
+++ b/Engine/Source/Runtime/SlateCore/Public/Widgets/SWidget.h
@@ -162,40 +162,97 @@ class SLATECORE_API SWidget
 	friend struct TSlateDecl;
 
 protected:
+	/**
+	 * A SlateAttribute that is member variable of a SWidget.
+	 * @usage: TSlateAttribute<int32> MyAttribute1; TSlateAttribute<int32, EInvalidateWidgetReason::Paint> MyAttribute2; TSlateAttribute<int32, EInvalidateWidgetReason::Paint, TEqualTo<>> MyAttribute3;
+	 */
 	template<typename InObjectType, EInvalidateWidgetReason InInvalidationReasonValue = EInvalidateWidgetReason::None, typename InComparePredicate = TEqualTo<>>
 	struct TSlateAttribute : public ::SlateAttributePrivate::TSlateMemberAttribute<
 		InObjectType,
 		typename std::conditional<InInvalidationReasonValue == EInvalidateWidgetReason::None, ::SlateAttributePrivate::FSlateAttributeNoInvalidationReason, TSlateAttributeInvalidationReason<InInvalidationReasonValue>>::type,
-		typename std::conditional<std::is_same<InObjectType, FText>::value, TSlateAttributeFTextComparePredicate, InComparePredicate>::type>
+		InComparePredicate>
 	{
 		using ::SlateAttributePrivate::TSlateMemberAttribute<
 			InObjectType,
 			typename std::conditional<InInvalidationReasonValue == EInvalidateWidgetReason::None, ::SlateAttributePrivate::FSlateAttributeNoInvalidationReason, TSlateAttributeInvalidationReason<InInvalidationReasonValue>>::type,
-			typename std::conditional<std::is_same<InObjectType, FText>::value, TSlateAttributeFTextComparePredicate, InComparePredicate>::type>::TSlateMemberAttribute;
+			InComparePredicate>::TSlateMemberAttribute;
 	};
 
-	template<typename InObjectType, typename InInvalidationReasonPredicate, typename InComparePredicate = TEqualTo<>>
-	struct TSlateAttributeWithPredicate : public ::SlateAttributePrivate::TSlateMemberAttribute<InObjectType, InInvalidationReasonPredicate, InComparePredicate>
+	//~ Override for FText that use the TSlateAttributeFTextComparePredicate to compare FText
+	template<>
+	struct TSlateAttribute<FText, EInvalidateWidgetReason::None> : public ::SlateAttributePrivate::TSlateMemberAttribute<FText, ::SlateAttributePrivate::FSlateAttributeNoInvalidationReason, TSlateAttributeFTextComparePredicate>
 	{
-		using ::SlateAttributePrivate::TSlateMemberAttribute<InObjectType, InInvalidationReasonPredicate, InComparePredicate>::TSlateMemberAttribute;
+		using ::SlateAttributePrivate::TSlateMemberAttribute<FText, ::SlateAttributePrivate::FSlateAttributeNoInvalidationReason, TSlateAttributeFTextComparePredicate>::TSlateMemberAttribute;
+	};
+
+	//~ Override for FText that use the TSlateAttributeFTextComparePredicate to compare FText
+	template<EInvalidateWidgetReason InInvalidationReasonValue>
+	struct TSlateAttribute<FText, InInvalidationReasonValue> : public ::SlateAttributePrivate::TSlateMemberAttribute<
+		FText,
+		typename std::conditional<InInvalidationReasonValue == EInvalidateWidgetReason::None, ::SlateAttributePrivate::FSlateAttributeNoInvalidationReason, TSlateAttributeInvalidationReason<InInvalidationReasonValue>>::type,
+		TSlateAttributeFTextComparePredicate>
+	{
+		using ::SlateAttributePrivate::TSlateMemberAttribute<
+			FText,
+			typename std::conditional<InInvalidationReasonValue == EInvalidateWidgetReason::None, ::SlateAttributePrivate::FSlateAttributeNoInvalidationReason, TSlateAttributeInvalidationReason<InInvalidationReasonValue>>::type,
+			TSlateAttributeFTextComparePredicate>::TSlateMemberAttribute;
 	};
 
-	template<typename InInvalidationReasonPredicate>
-	struct TSlateAttributeWithPredicate<FText, InInvalidationReasonPredicate> : public ::SlateAttributePrivate::TSlateMemberAttribute<FText, InInvalidationReasonPredicate, TSlateAttributeFTextComparePredicate>
+	/**
+	 * A SlateAttribute that is NOT a member variable of a SWidget.
+	 * @usage: TSlateManagedAttribute<int32> MyAttribute1; TSlateManagedAttribute<int32, EInvalidateWidgetReason::Paint> MyAttribute2; TSlateManagedAttribute<int32, EInvalidateWidgetReason::Paint, TEqualTo<>> MyAttribute3;
+	 */
+	template<typename InObjectType, EInvalidateWidgetReason InInvalidationReasonValue = EInvalidateWidgetReason::None, typename InComparePredicate = TEqualTo<>>
+	struct TSlateManagedAttribute : public ::SlateAttributePrivate::TSlateManagedAttribute<
+		InObjectType,
+		typename std::conditional<InInvalidationReasonValue == EInvalidateWidgetReason::None, ::SlateAttributePrivate::FSlateAttributeNoInvalidationReason, TSlateAttributeInvalidationReason<InInvalidationReasonValue>>::type,
+		InComparePredicate>
+	{
+		using ::SlateAttributePrivate::TSlateManagedAttribute<
+			InObjectType,
+			typename std::conditional<InInvalidationReasonValue == EInvalidateWidgetReason::None, ::SlateAttributePrivate::FSlateAttributeNoInvalidationReason, TSlateAttributeInvalidationReason<InInvalidationReasonValue>>::type,
+			InComparePredicate>::TSlateManagedAttribute;
+	};
+
+	//~ Override for FText that use the TSlateAttributeFTextComparePredicate to compare FText
+	template<>
+	struct TSlateManagedAttribute<FText, EInvalidateWidgetReason::None> : public ::SlateAttributePrivate::TSlateManagedAttribute<FText, ::SlateAttributePrivate::FSlateAttributeNoInvalidationReason, TSlateAttributeFTextComparePredicate>
+	{
+		using ::SlateAttributePrivate::TSlateManagedAttribute<FText, ::SlateAttributePrivate::FSlateAttributeNoInvalidationReason, TSlateAttributeFTextComparePredicate>::TSlateManagedAttribute;
+	};
+
+	//~ Override for FText that use the TSlateAttributeFTextComparePredicate to compare FText
+	template<EInvalidateWidgetReason InInvalidationReasonValue>
+	struct TSlateManagedAttribute<FText, InInvalidationReasonValue> : public ::SlateAttributePrivate::TSlateManagedAttribute<
+		FText,
+		typename std::conditional<InInvalidationReasonValue == EInvalidateWidgetReason::None, ::SlateAttributePrivate::FSlateAttributeNoInvalidationReason, TSlateAttributeInvalidationReason<InInvalidationReasonValue>>::type,
+		TSlateAttributeFTextComparePredicate>
+	{
+		using ::SlateAttributePrivate::TSlateManagedAttribute<
+			FText,
+			typename std::conditional<InInvalidationReasonValue == EInvalidateWidgetReason::None, ::SlateAttributePrivate::FSlateAttributeNoInvalidationReason, TSlateAttributeInvalidationReason<InInvalidationReasonValue>>::type,
+			TSlateAttributeFTextComparePredicate>::TSlateManagedAttribute;
+	};
+
+	/** A Reference to a TSlateAttribute. */
+	template<typename InObjectType, EInvalidateWidgetReason InInvalidationReasonValue = EInvalidateWidgetReason::None, typename InComparePredicate = TEqualTo<>>
+	struct TSlateAttributeRef : public ::SlateAttributePrivate::TSlateMemberAttributeRef<TSlateAttribute<InObjectType, InInvalidationReasonValue, InComparePredicate>>
 	{
-		using ::SlateAttributePrivate::TSlateMemberAttribute<FText, InInvalidationReasonPredicate, TSlateAttributeFTextComparePredicate>::TSlateMemberAttribute;
+		using ::SlateAttributePrivate::TSlateMemberAttributeRef<TSlateAttribute<InObjectType, InInvalidationReasonValue, InComparePredicate>>::TSlateMemberAttributeRef;
 	};
 
-	template<typename InObjectType, EInvalidateWidgetReason InInvalidationReasonValue, typename InComparePredicate = TEqualTo<>>
-	struct TSlateManagedAttribute : public ::SlateAttributePrivate::TSlateManagedAttribute<InObjectType, TSlateAttributeInvalidationReason<InInvalidationReasonValue>, InComparePredicate>
+	//~ Override for FText that use the TSlateAttributeFTextComparePredicate to compare FText
+	template<>
+	struct TSlateAttributeRef<FText, EInvalidateWidgetReason::None> : public ::SlateAttributePrivate::TSlateMemberAttributeRef<TSlateAttribute<FText>>
 	{
-		using ::SlateAttributePrivate::TSlateManagedAttribute<InObjectType, TSlateAttributeInvalidationReason<InInvalidationReasonValue>, InComparePredicate>::TSlateManagedAttribute;
+		using ::SlateAttributePrivate::TSlateMemberAttributeRef<TSlateAttribute<FText>>::TSlateMemberAttributeRef;
 	};
 
-	template<typename InObjectType, typename InInvalidationReasonPredicate, typename InComparePredicate = TEqualTo<>>
-	struct TSlateManagedAttributeWithPredicate : public ::SlateAttributePrivate::TSlateManagedAttribute<InObjectType, InInvalidationReasonPredicate, InComparePredicate>
+	//~ Override for FText that use the TSlateAttributeFTextComparePredicate to compare FText
+	template<EInvalidateWidgetReason InInvalidationReasonValue>
+	struct TSlateAttributeRef<FText, InInvalidationReasonValue> : public ::SlateAttributePrivate::TSlateMemberAttributeRef<TSlateAttribute<FText, InInvalidationReasonValue>>
 	{
-		using ::SlateAttributePrivate::TSlateManagedAttribute<InObjectType, InInvalidationReasonPredicate, InComparePredicate>::TSlateManagedAttribute;
+		using ::SlateAttributePrivate::TSlateMemberAttributeRef<TSlateAttribute<FText, InInvalidationReasonValue>>::TSlateMemberAttributeRef;
 	};
 
 public:
