Fix CIS unity build
#rnx

diff --git a/Engine/Source/Runtime/SlateCore/Private/Fonts/SlateTextShaper.cpp b/Engine/Source/Runtime/SlateCore/Private/Fonts/SlateTextShaper.cpp
--- a/Engine/Source/Runtime/SlateCore/Private/Fonts/SlateTextShaper.cpp
+++ b/Engine/Source/Runtime/SlateCore/Private/Fonts/SlateTextShaper.cpp
@@ -6,6 +6,8 @@
 #include "Internationalization/BreakIterator.h"
 #include "SlateGlobals.h"
 
+#include <limits>
+
 DECLARE_CYCLE_STAT(TEXT("Shape Bidirectional Text"), STAT_SlateShapeBidirectionalText, STATGROUP_Slate);
 DECLARE_CYCLE_STAT(TEXT("Shape Unidirectional Text"), STAT_SlateShapeUnidirectionalText, STATGROUP_Slate);
 
