Slate: Convert SImage and other image widget Attributes into SlateAttributes. Make sure other image widget use the invalidation system correctly.

diff --git a/Engine/Plugins/Compositing/Composure/Source/ComposureLayersEditor/Private/Widgets/SCompPreviewImage.h b/Engine/Plugins/Compositing/Composure/Source/ComposureLayersEditor/Private/Widgets/SCompPreviewImage.h
--- a/Engine/Plugins/Compositing/Composure/Source/ComposureLayersEditor/Private/Widgets/SCompPreviewImage.h
+++ b/Engine/Plugins/Compositing/Composure/Source/ComposureLayersEditor/Private/Widgets/SCompPreviewImage.h
@@ -37,7 +37,7 @@ public:
 	//~ Begin SWidget interface
 	virtual int32 OnPaint(const FPaintArgs& Args, const FGeometry& AllottedGeometry, const FSlateRect& MyCullingRect, FSlateWindowElementList& OutDrawElements, int32 LayerId, const FWidgetStyle& InWidgetStyle, bool bParentEnabled) const override
 	{
-		const FSlateBrush* ImageBrush = Image.Get();
+		const FSlateBrush* ImageBrush = GetImageAttribute().Get();
 
 		if ((ImageBrush != nullptr) && (ImageBrush->DrawAs != ESlateBrushDrawType::NoDrawType))
 		{
@@ -47,7 +47,7 @@ public:
 				PaintEffects |= ESlateDrawEffect::DisabledEffect;
 			}
 
-			const FLinearColor FinalColorAndOpacity(InWidgetStyle.GetColorAndOpacityTint() * ColorAndOpacity.Get().GetColor(InWidgetStyle) * ImageBrush->GetTint(InWidgetStyle));
+			const FLinearColor FinalColorAndOpacity(InWidgetStyle.GetColorAndOpacityTint() * GetColorAndOpacityAttribute().Get().GetColor(InWidgetStyle) * ImageBrush->GetTint(InWidgetStyle));
 
 			FSlateDrawElement::MakeBox(OutDrawElements, LayerId, AllottedGeometry.ToPaintGeometry(), ImageBrush, PaintEffects, FinalColorAndOpacity);
 		}
diff --git a/Engine/Source/Editor/WorldBrowser/Private/Tiles/SWorldTileItem.cpp b/Engine/Source/Editor/WorldBrowser/Private/Tiles/SWorldTileItem.cpp
--- a/Engine/Source/Editor/WorldBrowser/Private/Tiles/SWorldTileItem.cpp
+++ b/Engine/Source/Editor/WorldBrowser/Private/Tiles/SWorldTileItem.cpp
@@ -21,7 +21,7 @@
 //----------------------------------------------------------------
 int32 SWorldTileImage::OnPaint(const FPaintArgs& Args, const FGeometry& AllottedGeometry, const FSlateRect& MyCullingRect, FSlateWindowElementList& OutDrawElements, int32 LayerId, const FWidgetStyle& InWidgetStyle, bool bParentEnabled) const
 {
-	const FSlateBrush* ImageBrush = Image.Get();
+	const FSlateBrush* ImageBrush = GetImageAttribute().Get();
 
 	if ((ImageBrush != nullptr) && (ImageBrush->DrawAs != ESlateBrushDrawType::NoDrawType))
 	{
diff --git a/Engine/Source/Runtime/Slate/Private/Widgets/Images/SSpinningImage.cpp b/Engine/Source/Runtime/Slate/Private/Widgets/Images/SSpinningImage.cpp
--- a/Engine/Source/Runtime/Slate/Private/Widgets/Images/SSpinningImage.cpp
+++ b/Engine/Source/Runtime/Slate/Private/Widgets/Images/SSpinningImage.cpp
@@ -6,29 +6,29 @@
 
 void SSpinningImage::Construct(const FArguments& InArgs)
 {
-	Image = FInvalidatableBrushAttribute(InArgs._Image);
-	ColorAndOpacity = InArgs._ColorAndOpacity;
+	SetImage(InArgs._Image);
+	SetColorAndOpacity(InArgs._ColorAndOpacity);
 
 	if (InArgs._OnMouseButtonDown.IsBound())
 	{
 		SetOnMouseButtonDown(InArgs._OnMouseButtonDown);
 	}
 	
-	SpinAnimationSequence = FCurveSequence( 0.f, InArgs._Period );
-	SpinAnimationSequence.Play( this->AsShared(), true );
+	SpinAnimationSequence = FCurveSequence(0.f, InArgs._Period);
+	SpinAnimationSequence.Play(AsShared(), true, 0.f, false);
 }
 
 // Override SImage's OnPaint to draw rotated
 int32 SSpinningImage::OnPaint( const FPaintArgs& Args, const FGeometry& AllottedGeometry, const FSlateRect& MyCullingRect, FSlateWindowElementList& OutDrawElements, int32 LayerId, const FWidgetStyle& InWidgetStyle, bool bParentEnabled ) const
 {
-	const FSlateBrush* ImageBrush = Image.Get();
+	const FSlateBrush* ImageBrush = GetImageAttribute().Get();
 
 	if ((ImageBrush != nullptr) && (ImageBrush->DrawAs != ESlateBrushDrawType::NoDrawType))
 	{
 		const bool bIsEnabled = ShouldBeEnabled(bParentEnabled);
 		const ESlateDrawEffect DrawEffects = bIsEnabled ? ESlateDrawEffect::None : ESlateDrawEffect::DisabledEffect;
 
-		const FLinearColor FinalColorAndOpacity( InWidgetStyle.GetColorAndOpacityTint() * ColorAndOpacity.Get().GetColor(InWidgetStyle) * ImageBrush->GetTint( InWidgetStyle ) );
+		const FLinearColor FinalColorAndOpacity( InWidgetStyle.GetColorAndOpacityTint() * GetColorAndOpacityAttribute().Get().GetColor(InWidgetStyle) * ImageBrush->GetTint( InWidgetStyle ) );
 		
 		const float Angle = SpinAnimationSequence.GetLerp() * 2.0f * PI;
 
@@ -46,3 +46,8 @@ int32 SSpinningImage::OnPaint( const FPaintArgs& Args, const FGeometry& Allotted
 	}
 	return LayerId;
 }
+
+bool SSpinningImage::ComputeVolatility() const
+{
+	return SImage::ComputeVolatility() || SpinAnimationSequence.IsPlaying();
+}
\ No newline at end of file
diff --git a/Engine/Source/Runtime/Slate/Private/Widgets/Images/SThrobber.cpp b/Engine/Source/Runtime/Slate/Private/Widgets/Images/SThrobber.cpp
--- a/Engine/Source/Runtime/Slate/Private/Widgets/Images/SThrobber.cpp
+++ b/Engine/Source/Runtime/Slate/Private/Widgets/Images/SThrobber.cpp
@@ -25,17 +25,13 @@ void SThrobber::Construct(const FArguments& InArgs)
 
 void SThrobber::ConstructPieces()
 {
-	ThrobberCurve.Empty();
+	ThrobberCurve.Reset(NumPieces);
 	AnimCurves = FCurveSequence();
+	HBox->ClearChildren();
 	for (int32 PieceIndex = 0; PieceIndex < NumPieces; ++PieceIndex)
 	{
 		ThrobberCurve.Add(AnimCurves.AddCurve(PieceIndex*0.05f, 1.5f));
-	}
-	AnimCurves.Play(this->AsShared(), true);
 
-	HBox->ClearChildren();
-	for (int32 PieceIndex = 0; PieceIndex < NumPieces; ++PieceIndex)
-	{
 		HBox->AddSlot()
 		.AutoWidth()
 		[
@@ -51,6 +47,12 @@ void SThrobber::ConstructPieces()
 			]
 		];
 	}
+	AnimCurves.Play(AsShared(), true, 0.f, false);
+}
+
+bool SThrobber::ComputeVolatility() const
+{
+	return SCompoundWidget::ComputeVolatility() || Animate != EAnimation::None;
 }
 
 const FSlateBrush* SThrobber::GetPieceBrush() const
@@ -60,23 +62,36 @@ const FSlateBrush* SThrobber::GetPieceBrush() const
 
 void SThrobber::SetPieceImage(const FSlateBrush* InPieceImage)
 {
-	PieceImage = InPieceImage;
+	if (PieceImage != InPieceImage)
+	{
+		PieceImage = InPieceImage;
+		Invalidate(EInvalidateWidgetReason::Paint);
+	}
 }
 
 void SThrobber::SetNumPieces(int InNumPieces)
 {
-	NumPieces = InNumPieces;
-	ConstructPieces();
+	if (NumPieces != InNumPieces)
+	{
+		NumPieces = InNumPieces;
+		ConstructPieces();
+		Invalidate(EInvalidateWidgetReason::Layout);
+	}
 }
 
 void SThrobber::SetAnimate(EAnimation InAnimate)
 {
-	Animate = InAnimate;
+	if (Animate != InAnimate)
+	{
+		Animate = InAnimate;
+		// invalidating the volatility will trigger a last paint, if Animate == 0
+		Invalidate(EInvalidateWidgetReason::PaintAndVolatility);
+	}
 }
 
 FVector2D SThrobber::GetPieceScale(int32 PieceIndex) const
 {
-	const float Value = FMath::Sin( 2 * PI * ThrobberCurve[PieceIndex].GetLerp());
+	const float Value = FMath::Sin(2.f * PI * ThrobberCurve[PieceIndex].GetLerp());
 	
 	const bool bAnimateHorizontally = ((0 != (Animate & Horizontal)));
 	const bool bAnimateVertically = (0 != (Animate & Vertical));
@@ -92,7 +107,7 @@ FLinearColor SThrobber::GetPieceColor(int32 PieceIndex) const
 	const bool bAnimateOpacity = (0 != (Animate & Opacity));
 	if (bAnimateOpacity)
 	{
-		const float Value = FMath::Sin( 2 * PI * ThrobberCurve[PieceIndex].GetLerp());
+		const float Value = FMath::Sin(2.f * PI * ThrobberCurve[PieceIndex].GetLerp());
 		return FLinearColor(1.0f,1.0f,1.0f, Value);
 	}
 	else
@@ -102,51 +117,79 @@ FLinearColor SThrobber::GetPieceColor(int32 PieceIndex) const
 }
 
 // SCircularThrobber
+
 const float SCircularThrobber::MinimumPeriodValue = SMALL_NUMBER;
 
+SLATE_IMPLEMENT_WIDGET(SCircularThrobber)
+void SCircularThrobber::PrivateRegisterAttributes(FSlateAttributeInitializer& AttributeInitializer)
+{
+	SLATE_ADD_MEMBER_ATTRIBUTE_DEFINITION(AttributeInitializer, ColorAndOpacity, EInvalidateWidgetReason::Paint);
+}
+
+
+SCircularThrobber::SCircularThrobber()
+	: ColorAndOpacity(*this)
+{}
+
 void SCircularThrobber::Construct(const FArguments& InArgs)
 {
 	PieceImage = InArgs._PieceImage;
 	NumPieces = InArgs._NumPieces;
 	Period = InArgs._Period;
 	Radius = InArgs._Radius;
-	ColorAndOpacity = InArgs._ColorAndOpacity;
+	bColorAndOpacitySet = InArgs._ColorAndOpacity.IsSet();
+	ColorAndOpacity.Assign(*this, InArgs._ColorAndOpacity);
 
 	ConstructSequence();
 }
 
 void SCircularThrobber::SetPieceImage(const FSlateBrush* InPieceImage)
 {
-	PieceImage = InPieceImage;
+	if (PieceImage != InPieceImage)
+	{
+		PieceImage = InPieceImage;
+		Invalidate(EInvalidateWidgetReason::Paint);
+	}
 }
 
 void SCircularThrobber::SetNumPieces(const int32 InNumPieces)
 {
-	NumPieces = InNumPieces;
+	if (NumPieces != InNumPieces)
+	{
+		NumPieces = InNumPieces;
+		Invalidate(EInvalidateWidgetReason::Paint);
+	}
 }
 
 void SCircularThrobber::SetPeriod(const float InPeriod)
 {
-	Period = InPeriod;
-	ConstructSequence();
+	if (!FMath::IsNearlyEqual(Period, InPeriod))
+	{
+		Period = InPeriod;
+		ConstructSequence();
+	}
 }
 
 void SCircularThrobber::SetRadius(const float InRadius)
 {
-	Radius = InRadius;
+	if (!FMath::IsNearlyEqual(Radius, InRadius))
+	{
+		Radius = InRadius;
+		Invalidate(EInvalidateWidgetReason::Layout);
+	}
 }
 
 void SCircularThrobber::ConstructSequence()
 {
 	Sequence = FCurveSequence();
 	Curve = Sequence.AddCurve(0.f, FMath::Max(Period, MinimumPeriodValue));
-	Sequence.Play(this->AsShared(), true);
+	Sequence.Play(AsShared(), true, 0.0f, false);
 }
 
 int32 SCircularThrobber::OnPaint( const FPaintArgs& Args, const FGeometry& AllottedGeometry, const FSlateRect& MyCullingRect, FSlateWindowElementList& OutDrawElements, int32 LayerId, const FWidgetStyle& InWidgetStyle, bool bParentEnabled ) const
 {
 	FLinearColor FinalColorAndOpacity;
-	if (ColorAndOpacity.IsSet())
+	if (bColorAndOpacitySet)
 	{
 		FinalColorAndOpacity = ColorAndOpacity.Get().GetColor(InWidgetStyle);
 	}
@@ -177,3 +220,8 @@ FVector2D SCircularThrobber::ComputeDesiredSize( float ) const
 {
 	return FVector2D(Radius, Radius) * 2;
 }
+
+bool SCircularThrobber::ComputeVolatility() const
+{
+	return Super::ComputeVolatility() || Sequence.IsPlaying();
+}
diff --git a/Engine/Source/Runtime/Slate/Public/Widgets/Images/SSpinningImage.h b/Engine/Source/Runtime/Slate/Public/Widgets/Images/SSpinningImage.h
--- a/Engine/Source/Runtime/Slate/Public/Widgets/Images/SSpinningImage.h
+++ b/Engine/Source/Runtime/Slate/Public/Widgets/Images/SSpinningImage.h
@@ -40,8 +40,11 @@ public:
 	/* Construct this widget */
 	void Construct( const FArguments& InArgs );
 
+	//~ Begin SWidget interface
 	virtual int32 OnPaint( const FPaintArgs& Args, const FGeometry& AllottedGeometry, const FSlateRect& MyCullingRect, FSlateWindowElementList& OutDrawElements, int32 LayerId, const FWidgetStyle& InWidgetStyle, bool bParentEnabled ) const override;
-
+	virtual bool ComputeVolatility() const override;
+	//~ End SWidget interface
+	
 private:
 	/** The sequence to drive the spinning animation */
 	FCurveSequence SpinAnimationSequence;
diff --git a/Engine/Source/Runtime/Slate/Public/Widgets/Images/SThrobber.h b/Engine/Source/Runtime/Slate/Public/Widgets/Images/SThrobber.h
--- a/Engine/Source/Runtime/Slate/Public/Widgets/Images/SThrobber.h
+++ b/Engine/Source/Runtime/Slate/Public/Widgets/Images/SThrobber.h
@@ -60,6 +60,10 @@ public:
 	/** Sets which aspects of the throbber to animate */
 	void SetAnimate(EAnimation InAnimate);
 
+	//~ Begin SWidget interface
+	virtual bool ComputeVolatility() const override;
+	//~ End SWidget interface
+
 private:
 
 	FVector2D GetPieceScale(int32 PieceIndex) const;
@@ -83,7 +87,7 @@ private:
 	const FSlateBrush* PieceImage;
 
 	/** The number of pieces to display. */
-	int NumPieces;
+	int32 NumPieces;
 
 	/** Flags for which aspects of the throbber to animate. */
 	EAnimation Animate;
@@ -96,6 +100,8 @@ private:
 class SLATE_API SCircularThrobber
 	: public SLeafWidget
 {
+	SLATE_DECLARE_WIDGET(SCircularThrobber, SLeafWidget)
+
 public:
 	static const float MinimumPeriodValue;
 
@@ -119,6 +125,8 @@ public:
 
 	SLATE_END_ARGS()
 
+	SCircularThrobber();
+
 	/** Constructs the widget */
 	void Construct(const FArguments& InArgs);
 
@@ -134,8 +142,11 @@ public:
 	/** Sets the radius of the circle */
 	void SetRadius(float InRadius);
 
+	//~ Begin SWidget interface
 	virtual int32 OnPaint( const FPaintArgs& Args, const FGeometry& AllottedGeometry, const FSlateRect& MyCullingRect, FSlateWindowElementList& OutDrawElements, int32 LayerId, const FWidgetStyle& InWidgetStyle, bool bParentEnabled ) const override;
 	virtual FVector2D ComputeDesiredSize(float) const override;
+	virtual bool ComputeVolatility() const override;
+	//~ End SWidget interface
 
 private:
 
@@ -161,5 +172,6 @@ private:
 	float Radius;
 
 	/** Color and opacity of the throbber images. */
-	TAttribute<FSlateColor> ColorAndOpacity;
+	TSlateAttribute<FSlateColor> ColorAndOpacity;
+	bool bColorAndOpacitySet;
 };
diff --git a/Engine/Source/Runtime/SlateCore/Private/Widgets/Images/SImage.cpp b/Engine/Source/Runtime/SlateCore/Private/Widgets/Images/SImage.cpp
--- a/Engine/Source/Runtime/SlateCore/Private/Widgets/Images/SImage.cpp
+++ b/Engine/Source/Runtime/SlateCore/Private/Widgets/Images/SImage.cpp
@@ -7,24 +7,45 @@
 #include "Widgets/Accessibility/SlateCoreAccessibleWidgets.h"
 #endif
 
+
+SLATE_IMPLEMENT_WIDGET(SImage)
+void SImage::PrivateRegisterAttributes(FSlateAttributeInitializer& AttributeInitializer)
+{
+	SLATE_ADD_MEMBER_ATTRIBUTE_DEFINITION_WITH_NAME(AttributeInitializer, "Image", ImageAttribute, EInvalidateWidgetReason::Layout);
+	SLATE_ADD_MEMBER_ATTRIBUTE_DEFINITION_WITH_NAME(AttributeInitializer, "ColorAndOpacity", ColorAndOpacityAttribute, EInvalidateWidgetReason::Paint);
+	SLATE_ADD_MEMBER_ATTRIBUTE_DEFINITION_WITH_NAME(AttributeInitializer, "DesiredSizeOverride", DesiredSizeOverrideAttribute, EInvalidateWidgetReason::Layout);
+}
+
+SImage::SImage()
+	: ImageAttribute(*this)
+	, ColorAndOpacityAttribute(*this)
+	, DesiredSizeOverrideAttribute(*this)
+{
+	SetCanTick(false);
+	bCanSupportFocus = false;
+}
+
 void SImage::Construct( const FArguments& InArgs )
 {
-	Image = FInvalidatableBrushAttribute(InArgs._Image);
-	ColorAndOpacity = InArgs._ColorAndOpacity;
+	ImageAttribute.Assign(*this, InArgs._Image);
+	ColorAndOpacityAttribute.Assign(*this, InArgs._ColorAndOpacity);
 	bFlipForRightToLeftFlowDirection = InArgs._FlipForRightToLeftFlowDirection;
+
+	DesiredSizeOverrideAttribute.Assign(*this, InArgs._DesiredSizeOverride);
+
 	SetOnMouseButtonDown(InArgs._OnMouseButtonDown);
 }
 
 int32 SImage::OnPaint( const FPaintArgs& Args, const FGeometry& AllottedGeometry, const FSlateRect& MyCullingRect, FSlateWindowElementList& OutDrawElements, int32 LayerId, const FWidgetStyle& InWidgetStyle, bool bParentEnabled ) const
 {
-	const FSlateBrush* ImageBrush = Image.GetImage().Get();
+	const FSlateBrush* ImageBrush = ImageAttribute.Get();
 
 	if ((ImageBrush != nullptr) && (ImageBrush->DrawAs != ESlateBrushDrawType::NoDrawType))
 	{
 		const bool bIsEnabled = ShouldBeEnabled(bParentEnabled);
 		const ESlateDrawEffect DrawEffects = bIsEnabled ? ESlateDrawEffect::None : ESlateDrawEffect::DisabledEffect;
 
-		const FLinearColor FinalColorAndOpacity( InWidgetStyle.GetColorAndOpacityTint() * ColorAndOpacity.Get().GetColor(InWidgetStyle) * ImageBrush->GetTint( InWidgetStyle ) );
+		const FLinearColor FinalColorAndOpacity( InWidgetStyle.GetColorAndOpacityTint() * ColorAndOpacityAttribute.Get().GetColor(InWidgetStyle) * ImageBrush->GetTint( InWidgetStyle ) );
 
 		if (bFlipForRightToLeftFlowDirection && GSlateFlowDirection == EFlowDirection::RightToLeft)
 		{
@@ -42,27 +63,34 @@ int32 SImage::OnPaint( const FPaintArgs& Args, const FGeometry& AllottedGeometry
 
 FVector2D SImage::ComputeDesiredSize( float ) const
 {
-	const FSlateBrush* ImageBrush = Image.Get();
+	const FSlateBrush* ImageBrush = ImageAttribute.Get();
 	if (ImageBrush != nullptr)
 	{
-		return ImageBrush->ImageSize;
+		const TOptional<FVector2D>& CurrentSizeOverride = DesiredSizeOverrideAttribute.Get();
+
+		return CurrentSizeOverride.IsSet() ? CurrentSizeOverride.GetValue() : ImageBrush->ImageSize;
 	}
 	return FVector2D::ZeroVector;
 }
 
-void SImage::SetColorAndOpacity( const TAttribute<FSlateColor>& InColorAndOpacity )
+void SImage::SetColorAndOpacity(TAttribute<FSlateColor> InColorAndOpacity )
 {
-	SetAttribute(ColorAndOpacity, InColorAndOpacity, EInvalidateWidgetReason::Paint);
+	ColorAndOpacityAttribute.Assign(*this, MoveTemp(InColorAndOpacity));
 }
 
-void SImage::SetColorAndOpacity( FLinearColor InColorAndOpacity )
+void SImage::SetColorAndOpacity(FLinearColor InColorAndOpacity)
 {
-	SetColorAndOpacity(TAttribute<FSlateColor>(InColorAndOpacity));
+	ColorAndOpacityAttribute.Set(*this, InColorAndOpacity);
 }
 
 void SImage::SetImage(TAttribute<const FSlateBrush*> InImage)
 {
-	Image.SetImage(*this, InImage);
+	ImageAttribute.Assign(*this, MoveTemp(InImage));
+}
+
+void SImage::SetDesiredSizeOverride(TAttribute<TOptional<FVector2D>> InDesiredSizeOverride)
+{
+	DesiredSizeOverrideAttribute.Assign(*this, MoveTemp(InDesiredSizeOverride));
 }
 
 #if WITH_ACCESSIBILITY
diff --git a/Engine/Source/Runtime/SlateCore/Private/Widgets/Images/SLayeredImage.cpp b/Engine/Source/Runtime/SlateCore/Private/Widgets/Images/SLayeredImage.cpp
--- a/Engine/Source/Runtime/SlateCore/Private/Widgets/Images/SLayeredImage.cpp
+++ b/Engine/Source/Runtime/SlateCore/Private/Widgets/Images/SLayeredImage.cpp
@@ -7,14 +7,24 @@ void SLayeredImage::Construct(const FArguments& InArgs, const TArray<ImageLayer>
 {
 	SImage::Construct(InArgs);
 
-	Layers = InLayers;
+	Layers.Reserve(InLayers.Num());
+
+	for(const ImageLayer& Layer : InLayers)
+	{
+		AddLayer(Layer.Get<0>(), Layer.Get<1>());
+	}
 }
 
 void SLayeredImage::Construct(const SLayeredImage::FArguments& InArgs, TArray<ImageLayer>&& InLayers)
 {
-	SImage::Construct(InArgs);
+	Layers.Reserve(InLayers.Num());
 
-	Layers = MoveTemp(InLayers);
+	for(ImageLayer& Layer : InLayers)
+	{
+		BrushAttributeType TmpBrush = BrushAttributeType{ AsShared(), MoveTemp(Layer.Get<0>()), nullptr };
+		ColorAttributeType TmpColor = ColorAttributeType{ AsShared(), MoveTemp(Layer.Get<1>()), FLinearColor::White };
+		Layers.Emplace(MoveTemp(TmpBrush), MoveTemp(TmpColor));
+	}
 }
 
 void SLayeredImage::Construct(const FArguments& InArgs, TAttribute<const FSlateBrush*> Brush, TAttribute<FSlateColor> Color)
@@ -30,12 +40,12 @@ void SLayeredImage::Construct(const FArguments& InArgs, int32 NumLayers)
 
 	if (NumLayers > 0)
 	{
-		Layers.AddDefaulted(NumLayers);
-
-		// replace the default fuschia color with white
-		for (ImageLayer& Layer : Layers)
+		Layers.Reset(NumLayers);
+		for (int32 Index = 0; Index < NumLayers; ++Index)
 		{
-			Layer.Value.Set(FLinearColor::White);
+			BrushAttributeType TmpBrush = BrushAttributeType(AsShared(), nullptr);
+			ColorAttributeType TmpColor = ColorAttributeType(AsShared(), FLinearColor::White);
+			Layers.Emplace(MoveTemp(TmpBrush), MoveTemp(TmpColor));
 		}
 	}
 }
@@ -56,7 +66,7 @@ int32 SLayeredImage::OnPaint(const FPaintArgs& Args, const FGeometry& AllottedGe
 	const ESlateDrawEffect DrawEffects = bIsEnabled ? ESlateDrawEffect::None : ESlateDrawEffect::DisabledEffect;
 	
 	// draw rest of the images, we reuse the LayerId because images are assumed to note overlap:
-	for (const ImageLayer& Layer : Layers)
+	for (const InnerImageLayerType& Layer : Layers)
 	{
 		const FSlateBrush* LayerImageResolved = Layer.Key.Get();
 		if (LayerImageResolved && LayerImageResolved->DrawAs != ESlateBrushDrawType::NoDrawType)
@@ -69,9 +79,18 @@ int32 SLayeredImage::OnPaint(const FPaintArgs& Args, const FGeometry& AllottedGe
 	return LayerId;
 }
 
+void SLayeredImage::AddLayer(TAttribute<const FSlateBrush*> Brush)
+{
+	BrushAttributeType TmpBrush = BrushAttributeType(AsShared(), MoveTemp(Brush), nullptr);
+	ColorAttributeType TmpColor = ColorAttributeType(AsShared(), FLinearColor::White);
+	Layers.Emplace(MoveTemp(TmpBrush), MoveTemp(TmpColor));
+}
+
 void SLayeredImage::AddLayer(TAttribute<const FSlateBrush*> Brush, TAttribute<FSlateColor> Color)
 {
-	Layers.Emplace(MoveTemp(Brush), MoveTemp(Color));
+	BrushAttributeType TmpBrush = BrushAttributeType(AsShared(), MoveTemp(Brush), nullptr);
+	ColorAttributeType TmpColor = ColorAttributeType(AsShared(), MoveTemp(Color), FLinearColor::White);
+	Layers.Emplace(MoveTemp(TmpBrush), MoveTemp(TmpColor));
 }
 
 void SLayeredImage::SetFromSlateIcon(const FSlateIcon& InIcon)
@@ -102,12 +121,11 @@ bool SLayeredImage::IsValidIndex(int32 Index) const
 	// Index 0 is our local SImage
 	return Index == 0 || Layers.IsValidIndex(Index - 1);
 }
-
 const FSlateBrush* SLayeredImage::GetLayerBrush(int32 Index) const
 {
 	if (Index == 0)
 	{
-		return Image.Get();
+		return GetImageAttribute().Get();
 	}
 	else if (Layers.IsValidIndex(Index - 1))
 	{
@@ -119,31 +137,15 @@ const FSlateBrush* SLayeredImage::GetLayerBrush(int32 Index) const
 	}
 }
 
-void SLayeredImage::SetLayerBrush(int32 Index, const TAttribute<const FSlateBrush*>& Brush)
-{
-	if (Index == 0)
-	{
-		Image = Brush;
-	}
-	else if (Layers.IsValidIndex(Index - 1))
-	{
-		Layers[Index - 1].Key = Brush;
-	}
-	else
-	{
-		// That layer doesn't exist
-	}
-}
-
-void SLayeredImage::SetLayerBrush(int32 Index, TAttribute<const FSlateBrush*>&& Brush)
+void SLayeredImage::SetLayerBrush(int32 Index, TAttribute<const FSlateBrush*> Brush)
 {
 	if (Index == 0)
 	{
-		Image = MoveTemp(Brush);
+		SetImage(MoveTemp(Brush));
 	}
 	else if (Layers.IsValidIndex(Index - 1))
 	{
-		Layers[Index - 1].Key = MoveTemp(Brush);
+		Layers[Index - 1].Key.Assign(MoveTemp(Brush));
 	}
 	else
 	{
@@ -155,7 +157,7 @@ FSlateColor SLayeredImage::GetLayerColor(int32 Index) const
 {
 	if (Index == 0)
 	{
-		return ColorAndOpacity.Get();
+		return GetColorAndOpacityAttribute().Get();
 	}
 	else if (Layers.IsValidIndex(Index - 1))
 	{
@@ -167,15 +169,15 @@ FSlateColor SLayeredImage::GetLayerColor(int32 Index) const
 	}
 }
 
-void SLayeredImage::SetLayerColor(int32 Index, const TAttribute<FSlateColor>& Color)
+void SLayeredImage::SetLayerColor(int32 Index, TAttribute<FSlateColor> Color)
 {
 	if (Index == 0)
 	{
-		ColorAndOpacity = Color;
+		SetColorAndOpacity(MoveTemp(Color));
 	}
 	else if (Layers.IsValidIndex(Index - 1))
 	{
-		Layers[Index - 1].Value = Color;
+		Layers[Index - 1].Value.Assign(MoveTemp(Color));
 	}
 	else
 	{
@@ -183,18 +185,3 @@ void SLayeredImage::SetLayerColor(int32 Index, const TAttribute<FSlateColor>& Co
 	}
 }
 
-void SLayeredImage::SetLayerColor(int32 Index, TAttribute<FSlateColor>&& Color)
-{
-	if (Index == 0)
-	{
-		ColorAndOpacity = MoveTemp(Color);
-	}
-	else if (Layers.IsValidIndex(Index - 1))
-	{
-		Layers[Index - 1].Value = MoveTemp(Color);
-	}
-	else
-	{
-		// That layer doesn't exist!
-	}
-}
diff --git a/Engine/Source/Runtime/SlateCore/Public/Widgets/Images/SImage.h b/Engine/Source/Runtime/SlateCore/Public/Widgets/Images/SImage.h
--- a/Engine/Source/Runtime/SlateCore/Public/Widgets/Images/SImage.h
+++ b/Engine/Source/Runtime/SlateCore/Public/Widgets/Images/SImage.h
@@ -20,6 +20,8 @@ class FSlateWindowElementList;
 class SLATECORE_API SImage
 	: public SLeafWidget
 {
+	SLATE_DECLARE_WIDGET(SImage, SLeafWidget)
+
 public:
 	SLATE_BEGIN_ARGS( SImage )
 		: _Image( FCoreStyle::Get().GetDefaultBrush() )
@@ -28,10 +30,13 @@ public:
 		{ }
 
 		/** Image resource */
-		SLATE_ATTRIBUTE( const FSlateBrush*, Image )
+		SLATE_ATTRIBUTE(const FSlateBrush*, Image)
 
 		/** Color and opacity */
-		SLATE_ATTRIBUTE( FSlateColor, ColorAndOpacity )
+		SLATE_ATTRIBUTE(FSlateColor, ColorAndOpacity)
+
+		/** When specified, ignore the brushes size and report the DesiredSizeOverride as the desired image size. */
+		SLATE_ATTRIBUTE(TOptional<FVector2D>, DesiredSizeOverride)
 
 		/** Flips the image if the localization's flow direction is RightToLeft */
 		SLATE_ARGUMENT( bool, FlipForRightToLeftFlowDirection )
@@ -41,11 +46,7 @@ public:
 	SLATE_END_ARGS()
 
 	/** Constructor */
-	SImage()
-	{
-		SetCanTick(false);
-		bCanSupportFocus = false;
-	}
+	SImage();
 
 	/**
 	 * Construct this widget
@@ -56,15 +57,18 @@ public:
 
 public:
 
-	/** See the ColorAndOpacity attribute */
-	void SetColorAndOpacity( const TAttribute<FSlateColor>& InColorAndOpacity );
-	
+	/** Set the ColorAndOpacity attribute */
+	void SetColorAndOpacity(TAttribute<FSlateColor> InColorAndOpacity);
+
 	/** See the ColorAndOpacity attribute */
 	void SetColorAndOpacity( FLinearColor InColorAndOpacity );
 
-	/** See the Image attribute */
+	/** Set the Image attribute */
 	void SetImage(TAttribute<const FSlateBrush*> InImage);
 
+	/** Set SizeOverride attribute */
+	void SetDesiredSizeOverride(TAttribute<TOptional<FVector2D>> InDesiredSizeOverride);
+
 public:
 
 	// SWidget overrides
@@ -78,17 +82,34 @@ protected:
 	virtual FVector2D ComputeDesiredSize(float) const override;
 	// End SWidget overrides.
 
-protected:
+	/** @return an attribute reference of Image */
+	TSlateAttributeRef<const FSlateBrush*> GetImageAttribute() const { return TSlateAttributeRef<FSlateBrush const*>(*this, ImageAttribute); }
 
-	/** The slate brush to draw for the image that we can invalidate. */
+	/** @return an attribute reference of ColorAndOpacity */
+	TSlateAttributeRef<FSlateColor> GetColorAndOpacityAttribute() const { return TSlateAttributeRef<FSlateColor>(*this, ColorAndOpacityAttribute); }
+
+	/** @return an attribute reference of DesiredSizeOverride */
+	TSlateAttributeRef<TOptional<FVector2D>> GetDesiredSizeOverrideAttribute() const { return TSlateAttributeRef<TOptional<FVector2D>>(*this, DesiredSizeOverrideAttribute); }
+
+#if WITH_EDITORONLY_DATA
+	UE_DEPRECATED(5.0, "Direct access to Image is now deprecated. Use the setter or getter.")
 	FInvalidatableBrushAttribute Image;
+	UE_DEPRECATED(5.0, "Direct access to ColorAndOpacity is now deprecated. Use the setter or getter.")
+	FSlateDeprecatedTAttribute<FSlateColor> ColorAndOpacity;
+	UE_DEPRECATED(5.0, "Direct access to DesiredSizeOverride is now deprecated. Use the setter or getter.")
+	FSlateDeprecatedTAttribute<TOptional<FVector2D>> DesiredSizeOverride;
+#endif
 
-	/** Color and opacity scale for this image */
-	TAttribute<FSlateColor> ColorAndOpacity;
+private:
+	/** The slate brush to draw for the ImageAttribute that we can invalidate. */
+	TSlateAttribute<const FSlateBrush*> ImageAttribute;
+
+	/** Color and opacity scale for this ImageAttribute */
+	TSlateAttribute<FSlateColor> ColorAndOpacityAttribute;
+
+	/** When specified, ignore the content's desired size and report the.HeightOverride as the Box's desired height. */
+	TSlateAttribute<TOptional<FVector2D>> DesiredSizeOverrideAttribute;
 
 	/** Flips the image if the localization's flow direction is RightToLeft */
 	bool bFlipForRightToLeftFlowDirection;
-
-	/** Invoked when the mouse is pressed in the image */
-	FPointerEventHandler OnMouseButtonDownHandler;
 };
diff --git a/Engine/Source/Runtime/SlateCore/Public/Widgets/Images/SLayeredImage.h b/Engine/Source/Runtime/SlateCore/Public/Widgets/Images/SLayeredImage.h
--- a/Engine/Source/Runtime/SlateCore/Public/Widgets/Images/SLayeredImage.h
+++ b/Engine/Source/Runtime/SlateCore/Public/Widgets/Images/SLayeredImage.h
@@ -32,8 +32,11 @@ public:
 	virtual int32 OnPaint(const FPaintArgs& Args, const FGeometry& AllottedGeometry, const FSlateRect& MyCullingRect, FSlateWindowElementList& OutDrawElements, int32 LayerId, const FWidgetStyle& InWidgetStyle, bool bParentEnabled) const override;
 	//~ End SWidget Interface
 
+	/** Adds a new layer to the top of the Image given a brush and color. Use White as the default color. */
+	void AddLayer(TAttribute<const FSlateBrush*> Brush);
+
 	/** Adds a new layer to the top of the Image given a brush and color */
-	void AddLayer(TAttribute<const FSlateBrush*> Brush, TAttribute<FSlateColor> Color = FLinearColor::White);
+	void AddLayer(TAttribute<const FSlateBrush*> Brush, TAttribute<FSlateColor> Color);
 
 	/** Sets the base image and any overlay layers defined in the Slate icon. Note: removes all layers */
 	void SetFromSlateIcon(const FSlateIcon& InIcon);
@@ -52,28 +55,28 @@ public:
 	 * 
 	 * @return Null if the layer index is invalid. Otherwise, the brush for the given layer
 	 */
+	UE_DEPRECATED(5.0, "GetLayerBrush is not accessible anymore since it's attribute value may not have been updated yet.")
 	const FSlateBrush* GetLayerBrush(int32 Index) const;
 
 	/** Sets the brush for a given layer, if it exists */
-	void SetLayerBrush(int32 Index, const TAttribute<const FSlateBrush*>& Brush);
-
-	/** Sets the brush for a given layer, if it exists */
-	void SetLayerBrush(int32 Index, TAttribute<const FSlateBrush*>&& Brush);
+	void SetLayerBrush(int32 Index, TAttribute<const FSlateBrush*> Brush);
 
 	/** 
 	 * Gets the Color for a given layer
 	 *
 	 * @return Uninitialized fuschia color if the layer index is invalid. Otherwise the color for the given layer
 	 */
+	UE_DEPRECATED(5.0, "GetLayerColor is not accessible anymore since it's attribute value may not have been updated yet.")
 	FSlateColor GetLayerColor(int32 Index) const;
 
 	/** Sets the color for a given layer, if it exists. */
-	void SetLayerColor(int32 Index, const TAttribute<FSlateColor>& Color);
-
-	/** Sets the color for a given layer, if it exists. */
-	void SetLayerColor(int32 Index, TAttribute<FSlateColor>&& Color);
+	void SetLayerColor(int32 Index, TAttribute<FSlateColor> Color);
 
 private:
+	typedef TSlateManagedAttribute<const FSlateBrush*, EInvalidateWidgetReason::Paint> BrushAttributeType;
+	typedef TSlateManagedAttribute<FSlateColor, EInvalidateWidgetReason::Paint> ColorAttributeType;
+	typedef TTuple<BrushAttributeType, ColorAttributeType> InnerImageLayerType;
+		
 	/** An array to hold the additional draw layers */
-	TArray<ImageLayer,TInlineAllocator<2>> Layers;
+	TArray<InnerImageLayerType,TInlineAllocator<2>> Layers;
 };
