Slate: Convert SCompoundWidget Attributes into SlateAttributes.

diff --git a/Engine/Plugins/FX/Niagara/Source/NiagaraEditor/Private/Widgets/SNiagaraScriptGraph.cpp b/Engine/Plugins/FX/Niagara/Source/NiagaraEditor/Private/Widgets/SNiagaraScriptGraph.cpp
--- a/Engine/Plugins/FX/Niagara/Source/NiagaraEditor/Private/Widgets/SNiagaraScriptGraph.cpp
+++ b/Engine/Plugins/FX/Niagara/Source/NiagaraEditor/Private/Widgets/SNiagaraScriptGraph.cpp
@@ -44,7 +44,7 @@ void SNiagaraScriptGraph::Construct(const FArguments& InArgs, TSharedRef<FNiagar
 	bUpdatingGraphSelectionFromViewModel = false;
 
 	GraphTitle = InArgs._GraphTitle;
-	ForegroundColor = InArgs._ForegroundColor;
+	SetForegroundColor(InArgs._ForegroundColor);
 
 	GraphEditor = ConstructGraphEditor();
 	if (InArgs._ZoomToFitOnLoad)
diff --git a/Engine/Plugins/VirtualProduction/DMX/DMXEngine/Source/DMXEditor/Private/Widgets/Common/SSpinBoxVertical.h b/Engine/Plugins/VirtualProduction/DMX/DMXEngine/Source/DMXEditor/Private/Widgets/Common/SSpinBoxVertical.h
--- a/Engine/Plugins/VirtualProduction/DMX/DMXEngine/Source/DMXEditor/Private/Widgets/Common/SSpinBoxVertical.h
+++ b/Engine/Plugins/VirtualProduction/DMX/DMXEngine/Source/DMXEditor/Private/Widgets/Common/SSpinBoxVertical.h
@@ -164,7 +164,7 @@ public:
 
 		Style = InArgs._Style;
 
-		ForegroundColor = InArgs._Style->ForegroundColor;
+		SetForegroundColor(InArgs._Style->ForegroundColor);
 		Interface = InArgs._TypeInterface.IsValid() ? InArgs._TypeInterface : MakeShared<TDefaultNumericTypeInterface<NumericType>>();
 
 		ValueAttribute = InArgs._Value;
diff --git a/Engine/Source/Editor/EditorWidgets/Private/SObjectNameEditableTextBox.cpp b/Engine/Source/Editor/EditorWidgets/Private/SObjectNameEditableTextBox.cpp
--- a/Engine/Source/Editor/EditorWidgets/Private/SObjectNameEditableTextBox.cpp
+++ b/Engine/Source/Editor/EditorWidgets/Private/SObjectNameEditableTextBox.cpp
@@ -94,7 +94,7 @@ int32 SObjectNameEditableTextBox::OnPaint( const FPaintArgs& Args, const FGeomet
 		float EffectOpacity = EffectAlpha;
 
 		// Figure out a universally visible highlight color.
-		FLinearColor HighlightTargetColorAndOpacity = ( (FLinearColor::White - ColorAndOpacity.Get())*0.5f + FLinearColor(+0.4f, +0.1f, -0.2f)) * InWidgetStyle.GetColorAndOpacityTint();
+		FLinearColor HighlightTargetColorAndOpacity = ( (FLinearColor::White - GetColorAndOpacity())*0.5f + FLinearColor(+0.4f, +0.1f, -0.2f)) * InWidgetStyle.GetColorAndOpacityTint();
 		HighlightTargetColorAndOpacity.A = HighlightTargetOpacity * EffectOpacity * 255.0f;
 
 		// Compute the bounds offset of the highlight target from where the highlight target spring
diff --git a/Engine/Source/Editor/SceneOutliner/Private/SOutlinerTreeView.cpp b/Engine/Source/Editor/SceneOutliner/Private/SOutlinerTreeView.cpp
--- a/Engine/Source/Editor/SceneOutliner/Private/SOutlinerTreeView.cpp
+++ b/Engine/Source/Editor/SceneOutliner/Private/SOutlinerTreeView.cpp
@@ -366,7 +366,7 @@ int32 SSceneOutlinerTreeRow::OnPaint( const FPaintArgs& Args, const FGeometry& A
 		float EffectOpacity = EffectAlpha;
 
 		// Figure out a universally visible highlight color.
-		FLinearColor HighlightTargetColorAndOpacity = ( (FLinearColor::White - ColorAndOpacity.Get())*0.5f + FLinearColor(+0.4f, +0.1f, -0.2f)) * InWidgetStyle.GetColorAndOpacityTint();
+		FLinearColor HighlightTargetColorAndOpacity = ( (FLinearColor::White - GetColorAndOpacity())*0.5f + FLinearColor(+0.4f, +0.1f, -0.2f)) * InWidgetStyle.GetColorAndOpacityTint();
 		HighlightTargetColorAndOpacity.A = HighlightTargetOpacity * EffectOpacity * 255.0f;
 
 		// Compute the bounds offset of the highlight target from where the highlight target spring
diff --git a/Engine/Source/Editor/Sequencer/Private/SAnimationOutlinerTreeNode.cpp b/Engine/Source/Editor/Sequencer/Private/SAnimationOutlinerTreeNode.cpp
--- a/Engine/Source/Editor/Sequencer/Private/SAnimationOutlinerTreeNode.cpp
+++ b/Engine/Source/Editor/Sequencer/Private/SAnimationOutlinerTreeNode.cpp
@@ -98,7 +98,7 @@ void SAnimationOutlinerTreeNode::Construct( const FArguments& InArgs, TSharedRef
 
 	auto NodeHeight = [=]() -> FOptionalSize { return DisplayNode->GetNodeHeight(); };
 
-	ForegroundColor.Bind(this, &SAnimationOutlinerTreeNode::GetForegroundBasedOnSelection);
+	SetForegroundColor(MakeAttributeSP(this, &SAnimationOutlinerTreeNode::GetForegroundBasedOnSelection));
 
 	TSharedRef<SWidget>	FinalWidget = 
 		SNew( SBorder )
diff --git a/Engine/Source/Editor/UMGEditor/Private/Designer/SDisappearingBar.cpp b/Engine/Source/Editor/UMGEditor/Private/Designer/SDisappearingBar.cpp
--- a/Engine/Source/Editor/UMGEditor/Private/Designer/SDisappearingBar.cpp
+++ b/Engine/Source/Editor/UMGEditor/Private/Designer/SDisappearingBar.cpp
@@ -10,7 +10,7 @@
 void SDisappearingBar::Construct(const FArguments& InArgs)
 {
 	FadeCurve = FCurveSequence(0, 0.25f);
-	ColorAndOpacity = TAttribute<FLinearColor>::Create(TAttribute<FLinearColor>::FGetter::CreateSP(this, &SDisappearingBar::GetFadeColorAndOpacity));
+	SetColorAndOpacity(TAttribute<FLinearColor>::Create(TAttribute<FLinearColor>::FGetter::CreateSP(this, &SDisappearingBar::GetFadeColorAndOpacity)));
 
 	ChildSlot
 	[
diff --git a/Engine/Source/Runtime/Slate/Private/Framework/Docking/SDockingTarget.cpp b/Engine/Source/Runtime/Slate/Private/Framework/Docking/SDockingTarget.cpp
--- a/Engine/Source/Runtime/Slate/Private/Framework/Docking/SDockingTarget.cpp
+++ b/Engine/Source/Runtime/Slate/Private/Framework/Docking/SDockingTarget.cpp
@@ -81,7 +81,7 @@ void SDockingTarget::OnDragEnter( const FGeometry& MyGeometry, const FDragDropEv
 	{
 		bIsDragHovered = true;
 		// Provide feedback that this target is hovered
-		ColorAndOpacity = FCoreStyle::Get().GetColor(TEXT("Docking.Cross.HoveredTint"));
+		SetColorAndOpacity(FCoreStyle::Get().GetColor(TEXT("Docking.Cross.HoveredTint")));
 
 		// Let the drag and drop operation know that we are currently hovering over this node
 		DragDropOperation->SetHoveredTarget( FDockingDragOperation::FDockTarget(OwnerNode.Pin(), GetDockDirection()), DragDropEvent );
@@ -95,7 +95,7 @@ void SDockingTarget::OnDragLeave( const FDragDropEvent& DragDropEvent )
 	{
 		bIsDragHovered = false;
 		// We are no longer hovered, remove the hover feedback
-		ColorAndOpacity = FCoreStyle::Get().GetColor(TEXT("Docking.Cross.Tint"));
+		SetColorAndOpacity(FCoreStyle::Get().GetColor(TEXT("Docking.Cross.Tint")));
 		
 		// Let the drag and drop operation know that we are no longer hovering over any nodes
 		DragDropOperation->SetHoveredTarget( FDockingDragOperation::FDockTarget(), DragDropEvent );
diff --git a/Engine/Source/Runtime/Slate/Private/Framework/MultiBox/MultiBox.cpp b/Engine/Source/Runtime/Slate/Private/Framework/MultiBox/MultiBox.cpp
--- a/Engine/Source/Runtime/Slate/Private/Framework/MultiBox/MultiBox.cpp
+++ b/Engine/Source/Runtime/Slate/Private/Framework/MultiBox/MultiBox.cpp
@@ -538,7 +538,7 @@ bool FMultiBox::IsInEditMode() const
 void SMultiBoxWidget::Construct( const FArguments& InArgs )
 {
 	LinkedBoxManager = MakeShared<FLinkedBoxManager>();
-	ContentScale = InArgs._ContentScale;
+	SetContentScale(InArgs._ContentScale);
 }
 
 TSharedRef<ITableRow> SMultiBoxWidget::GenerateTiles(TSharedPtr<SWidget> Item, const TSharedRef<STableViewBase>& OwnerTable)
diff --git a/Engine/Source/Runtime/Slate/Private/Widgets/Docking/SDockTab.cpp b/Engine/Source/Runtime/Slate/Private/Widgets/Docking/SDockTab.cpp
--- a/Engine/Source/Runtime/Slate/Private/Widgets/Docking/SDockTab.cpp
+++ b/Engine/Source/Runtime/Slate/Private/Widgets/Docking/SDockTab.cpp
@@ -231,9 +231,10 @@ bool SDockTab::IsForeground() const
 
 FSlateColor SDockTab::GetForegroundColor() const
 {
-	if (ForegroundColor.IsBound() || ForegroundColor.Get() != FSlateColor::UseStyle())
+	auto LocalForegroundColorAttribute = GetForegroundColorAttribute();
+	if (LocalForegroundColorAttribute.IsBound() || LocalForegroundColorAttribute.Get() != FSlateColor::UseStyle())
 	{
-		return ForegroundColor.Get();
+		return LocalForegroundColorAttribute.Get();
 	}
 	else if (IsActive())
 	{
diff --git a/Engine/Source/Runtime/Slate/Private/Widgets/Input/SButton.cpp b/Engine/Source/Runtime/Slate/Private/Widgets/Input/SButton.cpp
--- a/Engine/Source/Runtime/Slate/Private/Widgets/Input/SButton.cpp
+++ b/Engine/Source/Runtime/Slate/Private/Widgets/Input/SButton.cpp
@@ -421,7 +421,7 @@ bool SButton::IsInteractable() const
 bool SButton::ComputeVolatility() const
 {
 	// Note: we need to be careful with button volatility.  The parent SBorder class always has bound delegates to button but the following are the only thing that actually would be bound that would not be caught by an Invalidate call alone
-	return ContentScale.IsBound() || DesiredSizeScale.IsBound() || BorderBackgroundColor.IsBound() || ContentPadding.IsBound() || ForegroundColor.IsBound();
+	return DesiredSizeScale.IsBound() || BorderBackgroundColor.IsBound() || ContentPadding.IsBound();
 }
 
 TEnumAsByte<EButtonClickMethod::Type> SButton::GetClickMethodFromInputType(const FPointerEvent& MouseEvent) const
@@ -535,6 +535,34 @@ void SButton::SetPressMethod(EButtonPressMethod::Type InPressMethod)
 	PressMethod = InPressMethod;
 }
 
+FSlateColor SButton::GetDisabledForegroundColor() const
+{
+	return Style->DisabledForeground;	
+}
+
+FSlateColor SButton::GetForegroundColor() const
+{
+	FSlateColor UserColor = GetForegroundColorAttribute().Get();
+
+	if(UserColor == FSlateColor::UseStyle())
+	{
+		if (IsPressed())
+		{
+			return Style->PressedForeground;
+		}
+		else if (IsHovered())
+		{
+			return Style->HoveredForeground;
+		}
+		else
+		{
+			return Style->NormalForeground;
+		}
+	}
+
+	return UserColor;
+}
+
 #if WITH_ACCESSIBILITY
 TSharedRef<FSlateAccessibleWidget> SButton::CreateAccessibleWidget()
 {
diff --git a/Engine/Source/Runtime/Slate/Private/Widgets/Layout/SBorder.cpp b/Engine/Source/Runtime/Slate/Private/Widgets/Layout/SBorder.cpp
--- a/Engine/Source/Runtime/Slate/Private/Widgets/Layout/SBorder.cpp
+++ b/Engine/Source/Runtime/Slate/Private/Widgets/Layout/SBorder.cpp
@@ -21,8 +21,8 @@ void SBorder::Construct( const SBorder::FArguments& InArgs )
 		bCanSupportFocus = false;
 	}
 
-	ContentScale = InArgs._ContentScale;
-	ColorAndOpacity = InArgs._ColorAndOpacity;
+	SetContentScale(InArgs._ContentScale);
+	SetColorAndOpacity(InArgs._ColorAndOpacity);
 	DesiredSizeScale = InArgs._DesiredSizeScale;
 
 	ShowDisabledEffect = InArgs._ShowEffectWhenDisabled;
@@ -31,7 +31,7 @@ void SBorder::Construct( const SBorder::FArguments& InArgs )
 
 	BorderImage = InArgs._BorderImage;
 	BorderBackgroundColor = InArgs._BorderBackgroundColor;
-	ForegroundColor = InArgs._ForegroundColor;
+	SetForegroundColor(InArgs._ForegroundColor);
 
 	if (InArgs._OnMouseButtonDown.IsBound())
 	{
diff --git a/Engine/Source/Runtime/Slate/Private/Widgets/Layout/SFxWidget.cpp b/Engine/Source/Runtime/Slate/Private/Widgets/Layout/SFxWidget.cpp
--- a/Engine/Source/Runtime/Slate/Private/Widgets/Layout/SFxWidget.cpp
+++ b/Engine/Source/Runtime/Slate/Private/Widgets/Layout/SFxWidget.cpp
@@ -12,7 +12,7 @@ void SFxWidget::Construct( const FArguments& InArgs )
 	LayoutScale = InArgs._LayoutScale;
 	VisualOffset = InArgs._VisualOffset;
 	bIgnoreClipping = InArgs._IgnoreClipping;
-	ColorAndOpacity = InArgs._ColorAndOpacity;
+	SetColorAndOpacity(InArgs._ColorAndOpacity);
 	
 	this->ChildSlot
 	.HAlign(InArgs._HAlign)
@@ -42,16 +42,6 @@ void SFxWidget::SetRenderScale( float InScale )
 	RenderScale = InScale;
 }
 
-void SFxWidget::SetColorAndOpacity( const TAttribute<FLinearColor>& InColorAndOpacity )
-{
-	ColorAndOpacity = InColorAndOpacity;
-}
-
-void SFxWidget::SetColorAndOpacity( FLinearColor InColorAndOpacity )
-{
-	ColorAndOpacity = InColorAndOpacity;
-}
-
 /**
  * This widget was created before render transforms existed for each widget, and it chose to apply the render transform AFTER the layout transform.
  * This means leveraging the render transform of FGeometry would be expensive, as we would need to use Concat(LayoutTransform, RenderTransform, Inverse(LayoutTransform).
@@ -73,6 +63,8 @@ int32 SFxWidget::OnPaint( const FPaintArgs& Args, const FGeometry& AllottedGeome
 	// There may be zero elements in this array if our child collapsed/hidden
 	if( ArrangedChildren.Num() > 0 )
 	{
+		const bool bShouldBeEnabled = ShouldBeEnabled(bParentEnabled);
+
 		// We can only have one direct descendant.
 		check( ArrangedChildren.Num() == 1 );
 		const FArrangedWidget& TheChild = ArrangedChildren[0];
@@ -83,10 +75,10 @@ int32 SFxWidget::OnPaint( const FPaintArgs& Args, const FGeometry& AllottedGeome
 			: MyCullingRect.IntersectionWith(ModifiedGeometry.GetLayoutBoundingRect());
 
 		FWidgetStyle CompoundedWidgetStyle = FWidgetStyle(InWidgetStyle)
-			.BlendColorAndOpacityTint(ColorAndOpacity.Get())
-			.SetForegroundColor( ForegroundColor );
+			.BlendColorAndOpacityTint(GetColorAndOpacity())
+			.SetForegroundColor(bShouldBeEnabled ? GetForegroundColor() : GetDisabledForegroundColor());
 
-		return TheChild.Widget->Paint( Args.WithNewParent(this), TheChild.Geometry, ChildClippingRect, OutDrawElements, LayerId + 1, CompoundedWidgetStyle, ShouldBeEnabled( bParentEnabled ) );
+		return TheChild.Widget->Paint( Args.WithNewParent(this), TheChild.Geometry, ChildClippingRect, OutDrawElements, LayerId + 1, CompoundedWidgetStyle, bShouldBeEnabled );
 	}
 	return LayerId;
 
diff --git a/Engine/Source/Runtime/Slate/Private/Widgets/Layout/SPopup.cpp b/Engine/Source/Runtime/Slate/Private/Widgets/Layout/SPopup.cpp
--- a/Engine/Source/Runtime/Slate/Private/Widgets/Layout/SPopup.cpp
+++ b/Engine/Source/Runtime/Slate/Private/Widgets/Layout/SPopup.cpp
@@ -20,16 +20,18 @@ int32 SPopup::OnPaint(const FPaintArgs& Args, const FGeometry& AllottedGeometry,
 	// There may be zero elements in this array if our child collapsed/hidden
 	if (ArrangedChildren.Num() > 0)
 	{
+		const bool bShouldBeEnabled = ShouldBeEnabled(bParentEnabled);
+
 		check(ArrangedChildren.Num() == 1);
 		FArrangedWidget& TheChild = ArrangedChildren[0];
 
 		FWidgetStyle CompoundedWidgetStyle = FWidgetStyle(InWidgetStyle)
-			.BlendColorAndOpacityTint(ColorAndOpacity.Get())
-			.SetForegroundColor(GetForegroundColor());
+			.BlendColorAndOpacityTint(GetColorAndOpacity())
+			.SetForegroundColor(bShouldBeEnabled ? GetForegroundColor() : GetDisabledForegroundColor());
 
 		// An SPopup just queues up its children to be painted after everything in this window is done painting.
 		OutDrawElements.QueueDeferredPainting(
-			FSlateWindowElementList::FDeferredPaint(TheChild.Widget, Args.WithNewParent(this), TheChild.Geometry, CompoundedWidgetStyle, ShouldBeEnabled(bParentEnabled))
+			FSlateWindowElementList::FDeferredPaint(TheChild.Widget, Args.WithNewParent(this), TheChild.Geometry, CompoundedWidgetStyle, bShouldBeEnabled)
 		);
 	}
 	return LayerId;
diff --git a/Engine/Source/Runtime/Slate/Private/Widgets/Notifications/SErrorHint.cpp b/Engine/Source/Runtime/Slate/Private/Widgets/Notifications/SErrorHint.cpp
--- a/Engine/Source/Runtime/Slate/Private/Widgets/Notifications/SErrorHint.cpp
+++ b/Engine/Source/Runtime/Slate/Private/Widgets/Notifications/SErrorHint.cpp
@@ -6,7 +6,7 @@
 
 void SErrorHint::Construct(const FArguments& InArgs)
 {
-	ContentScale.Bind( this, &SErrorHint::GetDesiredSizeScale );
+	SetContentScale(MakeAttributeSP(this, &SErrorHint::GetDesiredSizeScale));
 
 	ExpandAnimation = FCurveSequence(0.0f, 0.15f);
 
diff --git a/Engine/Source/Runtime/Slate/Public/Widgets/Input/SSpinBox.h b/Engine/Source/Runtime/Slate/Public/Widgets/Input/SSpinBox.h
--- a/Engine/Source/Runtime/Slate/Public/Widgets/Input/SSpinBox.h
+++ b/Engine/Source/Runtime/Slate/Public/Widgets/Input/SSpinBox.h
@@ -169,7 +169,7 @@ public:
 
 		Style = InArgs._Style;
 
-		ForegroundColor = InArgs._Style->ForegroundColor;
+		SetForegroundColor(InArgs._Style->ForegroundColor);
 		Interface = InArgs._TypeInterface.IsValid() ? InArgs._TypeInterface : MakeShareable( new TDefaultNumericTypeInterface<NumericType> );
 
 		ValueAttribute = InArgs._Value;
diff --git a/Engine/Source/Runtime/Slate/Public/Widgets/Layout/SFxWidget.h b/Engine/Source/Runtime/Slate/Public/Widgets/Layout/SFxWidget.h
--- a/Engine/Source/Runtime/Slate/Public/Widgets/Layout/SFxWidget.h
+++ b/Engine/Source/Runtime/Slate/Public/Widgets/Layout/SFxWidget.h
@@ -73,12 +73,6 @@ public:
 	/** @see RenderScale */
 	void SetRenderScale( float InScale );
 
-	/** @see ColorAndOpacity */
-	void SetColorAndOpacity( const TAttribute<FLinearColor>& InColorAndOpacity );
-
-	/** @see ColorAndOpacity */
-	void SetColorAndOpacity( FLinearColor InColorAndOpacity );
-
 	
 
 protected:
diff --git a/Engine/Source/Runtime/Slate/Public/Widgets/Views/STableRow.h b/Engine/Source/Runtime/Slate/Public/Widgets/Views/STableRow.h
--- a/Engine/Source/Runtime/Slate/Public/Widgets/Views/STableRow.h
+++ b/Engine/Source/Runtime/Slate/Public/Widgets/Views/STableRow.h
@@ -1137,7 +1137,7 @@ protected:
 
 		this->BorderImage = FInvalidatableBrushAttribute(TAttribute<const FSlateBrush*>(this, &STableRow::GetBorder));
 
-		this->ForegroundColor = TAttribute<FSlateColor>( this, &STableRow::GetForegroundBasedOnSelection );
+		this->SetForegroundColor(TAttribute<FSlateColor>( this, &STableRow::GetForegroundBasedOnSelection ));
 
 		this->OnCanAcceptDrop = InArgs._OnCanAcceptDrop;
 		this->OnAcceptDrop = InArgs._OnAcceptDrop;
diff --git a/Engine/Source/Runtime/SlateCore/Private/Widgets/SCompoundWidget.cpp b/Engine/Source/Runtime/SlateCore/Private/Widgets/SCompoundWidget.cpp
--- a/Engine/Source/Runtime/SlateCore/Private/Widgets/SCompoundWidget.cpp
+++ b/Engine/Source/Runtime/SlateCore/Private/Widgets/SCompoundWidget.cpp
@@ -8,6 +8,14 @@
 
 DECLARE_CYCLE_STAT(TEXT("Child Paint"), STAT_ChildPaint, STATGROUP_SlateVeryVerbose);
 
+SLATE_IMPLEMENT_WIDGET(SCompoundWidget)
+void SCompoundWidget::PrivateRegisterAttributes(FSlateAttributeInitializer& AttributeInitializer)
+{
+	SLATE_ADD_MEMBER_ATTRIBUTE_DEFINITION_WITH_NAME(AttributeInitializer, "ContentScale", ContentScaleAttribute, EInvalidateWidgetReason::Layout);
+	SLATE_ADD_MEMBER_ATTRIBUTE_DEFINITION_WITH_NAME(AttributeInitializer, "ColorAndOpacity", ColorAndOpacityAttribute, EInvalidateWidgetReason::Paint);
+	SLATE_ADD_MEMBER_ATTRIBUTE_DEFINITION_WITH_NAME(AttributeInitializer, "ForegroundColor", ForegroundColorAttribute, EInvalidateWidgetReason::Paint);
+}
+
 int32 SCompoundWidget::OnPaint( const FPaintArgs& Args, const FGeometry& AllottedGeometry, const FSlateRect& MyCullingRect, FSlateWindowElementList& OutDrawElements, int32 LayerId, const FWidgetStyle& InWidgetStyle, bool bParentEnabled ) const
 {
 
@@ -20,19 +28,21 @@ int32 SCompoundWidget::OnPaint( const FPaintArgs& Args, const FGeometry& Allotte
 	// There may be zero elements in this array if our child collapsed/hidden
 	if( ArrangedChildren.Num() > 0 )
 	{
+		const bool bShouldBeEnabled = ShouldBeEnabled(bParentEnabled);
+
 		check( ArrangedChildren.Num() == 1 );
 		FArrangedWidget& TheChild = ArrangedChildren[0];
 
 		FWidgetStyle CompoundedWidgetStyle = FWidgetStyle(InWidgetStyle)
-			.BlendColorAndOpacityTint(ColorAndOpacity.Get())
-			.SetForegroundColor( GetForegroundColor() );
+			.BlendColorAndOpacityTint(GetColorAndOpacity())
+			.SetForegroundColor(bShouldBeEnabled ? GetForegroundColor() : GetDisabledForegroundColor() );
 
 		int32 Layer = 0;
 		{
 #if WITH_VERY_VERBOSE_SLATE_STATS
 			SCOPE_CYCLE_COUNTER(STAT_ChildPaint);
 #endif
-			Layer = TheChild.Widget->Paint( Args.WithNewParent(this), TheChild.Geometry, MyCullingRect, OutDrawElements, LayerId + 1, CompoundedWidgetStyle, ShouldBeEnabled( bParentEnabled ) );
+			Layer = TheChild.Widget->Paint( Args.WithNewParent(this), TheChild.Geometry, MyCullingRect, OutDrawElements, LayerId + 1, CompoundedWidgetStyle, bShouldBeEnabled);
 		}
 		return Layer;
 	}
@@ -58,19 +68,19 @@ FVector2D SCompoundWidget::ComputeDesiredSize( float ) const
 
 void SCompoundWidget::OnArrangeChildren( const FGeometry& AllottedGeometry, FArrangedChildren& ArrangedChildren ) const
 {
-	ArrangeSingleChild(GSlateFlowDirection, AllottedGeometry, ArrangedChildren, ChildSlot, ContentScale);
+	ArrangeSingleChild(GSlateFlowDirection, AllottedGeometry, ArrangedChildren, ChildSlot, GetContentScale());
 }
 
 FSlateColor SCompoundWidget::GetForegroundColor() const
 {
-	return ForegroundColor.Get();
+	return ForegroundColorAttribute.Get();
 }
 
 SCompoundWidget::SCompoundWidget()
 	: ChildSlot(this)
-	, ContentScale( FVector2D(1.0f,1.0f) )
-	, ColorAndOpacity( FLinearColor::White )
-	, ForegroundColor( FSlateColor::UseForeground() )
+	, ContentScaleAttribute( *this, FVector2D(1.0f,1.0f) )
+	, ColorAndOpacityAttribute( *this, FLinearColor::White )
+	, ForegroundColorAttribute( *this, FSlateColor::UseForeground() )
 {
 }
 
diff --git a/Engine/Source/Runtime/SlateCore/Private/Widgets/SWindow.cpp b/Engine/Source/Runtime/SlateCore/Private/Widgets/SWindow.cpp
--- a/Engine/Source/Runtime/SlateCore/Private/Widgets/SWindow.cpp
+++ b/Engine/Source/Runtime/SlateCore/Private/Widgets/SWindow.cpp
@@ -530,7 +530,7 @@ EHorizontalAlignment SWindow::GetTitleAlignment()
 
 void SWindow::ConstructWindowInternals()
 {
-	ForegroundColor = FAppStyle::Get().GetSlateColor("DefaultForeground");
+	SetForegroundColor(FAppStyle::Get().GetSlateColor("DefaultForeground"));
 
 	// Setup widget that represents the main area of the window.  That is, everything inside the window's border.
 	TSharedRef< SVerticalBox > MainWindowArea =
diff --git a/Engine/Source/Runtime/SlateCore/Public/Layout/LayoutUtils.h b/Engine/Source/Runtime/SlateCore/Public/Layout/LayoutUtils.h
--- a/Engine/Source/Runtime/SlateCore/Public/Layout/LayoutUtils.h
+++ b/Engine/Source/Runtime/SlateCore/Public/Layout/LayoutUtils.h
@@ -191,6 +191,31 @@ static void ArrangeSingleChild(EFlowDirection InFlowDirection, const FGeometry&
 	}
 }
 
+template<typename SlotType>
+static void ArrangeSingleChild(const FGeometry& AllottedGeometry, FArrangedChildren& ArrangedChildren, const SlotType& ChildSlot, const FVector2D& ContentScale)
+{
+	ArrangeSingleChild<SlotType>(EFlowDirection::LeftToRight, AllottedGeometry, ArrangedChildren, ChildSlot, ContentScale);
+}
+
+template<typename SlotType>
+static void ArrangeSingleChild(EFlowDirection InFlowDirection, const FGeometry& AllottedGeometry, FArrangedChildren& ArrangedChildren, const SlotType& ChildSlot, const FVector2D& ContentScale)
+{
+	const EVisibility ChildVisibility = ChildSlot.GetWidget()->GetVisibility();
+	if (ArrangedChildren.Accepts(ChildVisibility))
+	{
+		const FVector2D ThisContentScale = ContentScale;
+		const FMargin SlotPadding(LayoutPaddingWithFlow(InFlowDirection, ChildSlot.SlotPadding.Get()));
+		const AlignmentArrangeResult XResult = AlignChild<Orient_Horizontal>(InFlowDirection, AllottedGeometry.GetLocalSize().X, ChildSlot, SlotPadding, ThisContentScale.X);
+		const AlignmentArrangeResult YResult = AlignChild<Orient_Vertical>(AllottedGeometry.GetLocalSize().Y, ChildSlot, SlotPadding, ThisContentScale.Y);
+
+		ArrangedChildren.AddWidget(ChildVisibility, AllottedGeometry.MakeChild(
+			ChildSlot.GetWidget(),
+			FVector2D(XResult.Offset, YResult.Offset),
+			FVector2D(XResult.Size, YResult.Size)
+		));
+	}
+}
+
 static FMargin LayoutPaddingWithFlow(EFlowDirection InLayoutFlow, const FMargin& InPadding)
 {
 	FMargin ReturnPadding(InPadding);
diff --git a/Engine/Source/Runtime/SlateCore/Public/Widgets/SCompoundWidget.h b/Engine/Source/Runtime/SlateCore/Public/Widgets/SCompoundWidget.h
--- a/Engine/Source/Runtime/SlateCore/Public/Widgets/SCompoundWidget.h
+++ b/Engine/Source/Runtime/SlateCore/Public/Widgets/SCompoundWidget.h
@@ -5,6 +5,7 @@
 #include "CoreMinimal.h"
 #include "Misc/Attribute.h"
 #include "Styling/SlateColor.h"
+#include "Types/SlateAttribute.h"
 #include "Layout/Visibility.h"
 #include "Widgets/SWidget.h"
 #include "Layout/Children.h"
@@ -19,6 +20,8 @@ class FSlateWindowElementList;
  */
 class SLATECORE_API SCompoundWidget : public SWidget
 {
+	SLATE_DECLARE_WIDGET(SCompoundWidget, SWidget)
+
 public:
 
 	/**
@@ -28,7 +31,7 @@ public:
 	 */
 	const FVector2D GetContentScale() const
 	{
-		return ContentScale.Get();
+		return ContentScaleAttribute.Get();
 	}
 
 	/**
@@ -36,9 +39,9 @@ public:
 	 *
 	 * @param InContentScale Content scaling factor.
 	 */
-	void SetContentScale( const TAttribute< FVector2D >& InContentScale )
+	void SetContentScale( TAttribute<FVector2D> InContentScale )
 	{
-		SetAttribute(ContentScale, InContentScale, EInvalidateWidgetReason::Layout);
+		ContentScaleAttribute.Assign(*this, MoveTemp(InContentScale));
 	}
 
 	/**
@@ -46,7 +49,7 @@ public:
 	 */
 	FLinearColor GetColorAndOpacity() const
 	{
-		return ColorAndOpacity.Get();
+		return ColorAndOpacityAttribute.Get();
 	}
 
 	/**
@@ -54,9 +57,9 @@ public:
 	 *
 	 * @param InColorAndOpacity The ColorAndOpacity to be applied to this widget and all its contents.
 	 */
-	void SetColorAndOpacity( const TAttribute<FLinearColor>& InColorAndOpacity )
+	void SetColorAndOpacity( TAttribute<FLinearColor> InColorAndOpacity )
 	{
-		SetAttribute(ColorAndOpacity, InColorAndOpacity, EInvalidateWidgetReason::Paint);
+		ColorAndOpacityAttribute.Assign(*this, MoveTemp(InColorAndOpacity));
 	}
 
 	/**
@@ -64,9 +67,9 @@ public:
 	 *
 	 * @param InColor The color to set.
 	 */
-	void SetForegroundColor( const TAttribute<FSlateColor>& InForegroundColor )
+	void SetForegroundColor( TAttribute<FSlateColor> InForegroundColor )
 	{
-		SetAttribute(ForegroundColor, InForegroundColor, EInvalidateWidgetReason::Paint);
+		ForegroundColorAttribute.Assign(*this, MoveTemp(InForegroundColor));
 	}
 
 public:
@@ -85,6 +88,15 @@ protected:
 	virtual FVector2D ComputeDesiredSize(float) const override;
 	// End SWidget overrides.
 
+	/** @return an attribute reference of ContentScale */
+	TSlateAttributeRef<FVector2D> GetContentScaleAttribute() const { return TSlateAttributeRef<FVector2D>(*this, ContentScaleAttribute); }
+
+	/** @return an attribute reference of ColorAndOpacity */
+	TSlateAttributeRef<FLinearColor> GetColorAndOpacityAttribute() const { return TSlateAttributeRef<FLinearColor>{*this, ColorAndOpacityAttribute}; }
+
+	/** @return an attribute reference of ForegroundColor */
+	TSlateAttributeRef<FSlateColor> GetForegroundColorAttribute() const {	return TSlateAttributeRef<FSlateColor>{*this, ForegroundColorAttribute}; }
+
 protected:
 
 	/** Disallow public construction */
@@ -93,12 +105,22 @@ protected:
 	/** The slot that contains this widget's descendants.*/
 	FSimpleSlot ChildSlot;
 
+//#if WITH_EDITORONLY_DATA
+//	UE_DEPRECATED(5.0, "Direct access to ContentScale is now deprecated. Use the setter or getter.")
+//	FSlateDeprecatedTAttribute<FVector2D> ContentScale;
+//	UE_DEPRECATED(5.0, "Direct access to ColorAndOpacity is now deprecated. Use the setter or getter.")
+//	FSlateDeprecatedTAttribute<FLinearColor> ColorAndOpacity;
+//	UE_DEPRECATED(5.0, "Direct access to ForegroundColor is now deprecated. Use the setter or getter.")
+//	FSlateDeprecatedTAttribute<FSlateColor> ForegroundColor;
+//#endif
+
+private:
 	/** The layout scale to apply to this widget's contents; useful for animation. */
-	TAttribute<FVector2D> ContentScale;
+	TSlateAttribute<FVector2D> ContentScaleAttribute;
 
 	/** The color and opacity to apply to this widget and all its descendants. */
-	TAttribute<FLinearColor> ColorAndOpacity;
+	TSlateAttribute<FLinearColor> ColorAndOpacityAttribute;
 
 	/** Optional foreground color that will be inherited by all of this widget's contents */
-	TAttribute<FSlateColor> ForegroundColor;
+	TSlateAttribute<FSlateColor> ForegroundColorAttribute;
 };
diff --git a/Engine/Source/Runtime/UMG/Private/Slate/SRetainerWidget.cpp b/Engine/Source/Runtime/UMG/Private/Slate/SRetainerWidget.cpp
--- a/Engine/Source/Runtime/UMG/Private/Slate/SRetainerWidget.cpp
+++ b/Engine/Source/Runtime/UMG/Private/Slate/SRetainerWidget.cpp
@@ -564,7 +564,7 @@ int32 SRetainerWidget::OnPaint(const FPaintArgs& Args, const FGeometry& Allotted
 
 			if (RenderTarget->GetSurfaceWidth() >= 1 && RenderTarget->GetSurfaceHeight() >= 1)
 			{
-				const FLinearColor ComputedColorAndOpacity(Context.WidgetStyle.GetColorAndOpacityTint() * ColorAndOpacity.Get() * SurfaceBrush.GetTint(Context.WidgetStyle));
+				const FLinearColor ComputedColorAndOpacity(Context.WidgetStyle.GetColorAndOpacityTint() * GetColorAndOpacity() * SurfaceBrush.GetTint(Context.WidgetStyle));
 				// Retainer widget uses pre-multiplied alpha, so pre-multiply the color by the alpha to respect opacity.
 				const FLinearColor PremultipliedColorAndOpacity(ComputedColorAndOpacity * ComputedColorAndOpacity.A);
 
