Slate: Add FOneSimpleMemberChild similar to FSimpleSlot but with SlateAttribute. Convert SDIScaler to SlateAttribute.

diff --git a/Engine/Source/Runtime/Slate/Private/Widgets/Layout/SDPIScaler.cpp b/Engine/Source/Runtime/Slate/Private/Widgets/Layout/SDPIScaler.cpp
--- a/Engine/Source/Runtime/Slate/Private/Widgets/Layout/SDPIScaler.cpp
+++ b/Engine/Source/Runtime/Slate/Private/Widgets/Layout/SDPIScaler.cpp
@@ -3,8 +3,15 @@
 #include "Widgets/Layout/SDPIScaler.h"
 #include "Layout/ArrangedChildren.h"
 
+SLATE_IMPLEMENT_WIDGET(SDPIScaler)
+void SDPIScaler::PrivateRegisterAttributes(FSlateAttributeInitializer& AttributeInitializer)
+{
+	SLATE_ADD_MEMBER_ATTRIBUTE_DEFINITION_WITH_NAME(AttributeInitializer, "DPIScale", DPIScaleAttribute, EInvalidateWidgetReason::Layout);
+}
+
 SDPIScaler::SDPIScaler()
-	: ChildSlot(this)
+	: ChildSlot(*this)
+	, DPIScaleAttribute(*this)
 {
 	SetCanTick(false);
 	bCanSupportFocus = false;
@@ -13,7 +20,7 @@ SDPIScaler::SDPIScaler()
 
 void SDPIScaler::Construct( const FArguments& InArgs )
 {
-	DPIScale = InArgs._DPIScale;
+	SetDPIScale(InArgs._DPIScale);
 
 	ChildSlot
 	[
@@ -26,7 +33,7 @@ void SDPIScaler::OnArrangeChildren( const FGeometry& AllottedGeometry, FArranged
 	const EVisibility MyVisibility = this->GetVisibility();
 	if ( ArrangedChildren.Accepts( MyVisibility ) )
 	{
-		const float MyDPIScale = DPIScale.Get();
+		const float MyDPIScale = DPIScaleAttribute.Get();
 
 		ArrangedChildren.AddWidget( AllottedGeometry.MakeChild(
 			this->ChildSlot.GetWidget(),
@@ -39,7 +46,7 @@ void SDPIScaler::OnArrangeChildren( const FGeometry& AllottedGeometry, FArranged
 	
 FVector2D SDPIScaler::ComputeDesiredSize( float ) const
 {
-	return DPIScale.Get() * ChildSlot.GetWidget()->GetDesiredSize();
+	return DPIScaleAttribute.Get() * ChildSlot.GetWidget()->GetDesiredSize();
 }
 
 FChildren* SDPIScaler::GetChildren()
@@ -57,7 +64,7 @@ void SDPIScaler::SetContent(TSharedRef<SWidget> InContent)
 
 void SDPIScaler::SetDPIScale(TAttribute<float> InDPIScale)
 {
-	if (SetAttribute(DPIScale, InDPIScale, EInvalidateWidgetReason::Layout))
+	if (DPIScaleAttribute.Assign(*this, MoveTemp(InDPIScale)))
 	{
 		InvalidatePrepass();
 	}
@@ -65,5 +72,5 @@ void SDPIScaler::SetDPIScale(TAttribute<float> InDPIScale)
 
 float SDPIScaler::GetRelativeLayoutScale(int32 ChildIndex, float LayoutScaleMultiplier) const
 {
-	return DPIScale.Get();
+	return DPIScaleAttribute.Get();
 }
diff --git a/Engine/Source/Runtime/Slate/Public/Widgets/Layout/SDPIScaler.h b/Engine/Source/Runtime/Slate/Public/Widgets/Layout/SDPIScaler.h
--- a/Engine/Source/Runtime/Slate/Public/Widgets/Layout/SDPIScaler.h
+++ b/Engine/Source/Runtime/Slate/Public/Widgets/Layout/SDPIScaler.h
@@ -17,6 +17,8 @@ class FArrangedChildren;
  */
 class SLATE_API SDPIScaler : public SPanel
 {
+	SLATE_DECLARE_WIDGET(SDPIScaler, SPanel)
+
 public:
 	SLATE_BEGIN_ARGS( SDPIScaler )
 		: _DPIScale( 1.0f )
@@ -56,13 +58,21 @@ protected:
 
 	virtual float GetRelativeLayoutScale(int32 ChildIndex, float LayoutScaleMultiplier) const override;
 
+	TSlateAttributeRef<float> GetDPIScaleAttribute() const { return TSlateAttributeRef<float>(*this, DPIScaleAttribute); }
+
+#if WITH_EDITORONLY_DATA
+	UE_DEPRECATED(5.0, "Direct access to DPIScale is now deprecated. Use the setter or getter.")
+	FSlateDeprecatedTAttribute<float> DPIScale;
+#endif
+
+	/** The content being scaled. */
+	FOneSimpleMemberChild ChildSlot;
+
+private:
 	/**
 	 * The scaling factor from 1:1 pixel to slate unit (SU) ratio.
 	 * e.g. Value of 2 would mean that there are two pixels for every slate unit in every dimension.
 	 *      This means that every 1x1 SU square is represented by 2x2=4 pixels.
 	 */
-	TAttribute<float> DPIScale;
-
-	/** The content being scaled. */
-	FSimpleSlot ChildSlot;
+	TSlateAttribute<float> DPIScaleAttribute;
 };
diff --git a/Engine/Source/Runtime/SlateCore/Public/Layout/Children.h b/Engine/Source/Runtime/SlateCore/Public/Layout/Children.h
--- a/Engine/Source/Runtime/SlateCore/Public/Layout/Children.h
+++ b/Engine/Source/Runtime/SlateCore/Public/Layout/Children.h
@@ -19,6 +19,7 @@
  * FChildren is intended to be returned by the GetChildren() method.
  * 
  */
+PRAGMA_DISABLE_DEPRECATION_WARNINGS
 class SLATECORE_API FChildren
 {
 public:
@@ -43,9 +44,16 @@ protected:
 protected:
 	virtual ~FChildren(){}
 
+	PRAGMA_DISABLE_DEPRECATION_WARNINGS
+	SWidget* GetOwner() { return Owner; }
+	const SWidget* GetOwner() const { return Owner; }
+	PRAGMA_ENABLE_DEPRECATION_WARNINGS
+
 protected:
+	UE_DEPRECATED(5.0, "Direct access to Owner is now deprecated. Use the getter.")
 	SWidget* Owner;
 };
+PRAGMA_ENABLE_DEPRECATION_WARNINGS
 
 
 /**
@@ -218,13 +226,13 @@ public:
 	void AttachWidget(const TSharedPtr<SWidget>& InWidget)
 	{
 		WidgetPtr = InWidget;
-		if (Owner) 
+		if (GetOwner()) 
 		{ 
-			Owner->Invalidate(EInvalidateWidgetReason::ChildOrder);
+			GetOwner()->Invalidate(EInvalidateWidgetReason::ChildOrder);
 
 			if (InWidget.IsValid() && InWidget != SNullWidget::NullWidget)
 			{
-				InWidget->AssignParentWidget(Owner->AsShared());
+				InWidget->AssignParentWidget(GetOwner()->AsShared());
 			}
 		}
 	}
@@ -237,7 +245,7 @@ public:
 
 			if (Widget != SNullWidget::NullWidget)
 			{
-				Widget->ConditionallyDetatchParentWidget(Owner);
+				Widget->ConditionallyDetatchParentWidget(GetOwner());
 			}
 
 			WidgetPtr.Reset();
@@ -383,9 +391,9 @@ public:
 
 		int32 Index = TIndirectArray< SlotType >::Add(Slot);
 
-		if (Owner)
+		if (GetOwner())
 		{
-			Slot->AttachWidgetParent(Owner);
+			Slot->AttachWidgetParent(GetOwner());
 		}
 
 		return Index;
@@ -430,9 +438,9 @@ public:
 			TIndirectArray< SlotType >::Insert(Slot, Index);
 
 			// Don't do parent manipulation if this panel has no owner.
-			if (Owner)
+			if (GetOwner())
 			{
-				Slot->AttachWidgetParent(Owner);
+				Slot->AttachWidgetParent(GetOwner());
 			}
 		}
 	}
@@ -455,9 +463,9 @@ public:
 			TIndirectArray< SlotType >::RemoveAt(IndexToMove);
 		}
 
-		if (Owner)
+		if (GetOwner())
 		{
-			Owner->Invalidate(EInvalidateWidgetReason::ChildOrder);
+			GetOwner()->Invalidate(EInvalidateWidgetReason::ChildOrder);
 		}
 	}
 
@@ -478,9 +486,9 @@ public:
 	void Sort( const PREDICATE_CLASS& Predicate )
 	{
 		::Sort(TIndirectArray< SlotType >::GetData(), TIndirectArray<SlotType>::Num(), Predicate);
-		if (Owner)
+		if (GetOwner())
 		{
-			Owner->Invalidate(EInvalidateWidgetReason::ChildOrder);
+			GetOwner()->Invalidate(EInvalidateWidgetReason::ChildOrder);
 		}
 	}
 
@@ -488,18 +496,18 @@ public:
 	void StableSort(const PREDICATE_CLASS& Predicate)
 	{
 		::StableSort(TIndirectArray< SlotType >::GetData(), TIndirectArray< SlotType >::Num(), Predicate);
-		if (Owner)
+		if (GetOwner())
 		{
-			Owner->Invalidate(EInvalidateWidgetReason::ChildOrder);
+			GetOwner()->Invalidate(EInvalidateWidgetReason::ChildOrder);
 		}
 	}
 
 	void Swap( int32 IndexA, int32 IndexB )
 	{
 		TIndirectArray< SlotType >::Swap(IndexA, IndexB);
-		if (Owner)
+		if (GetOwner())
 		{
-			Owner->Invalidate(EInvalidateWidgetReason::ChildOrder);
+			GetOwner()->Invalidate(EInvalidateWidgetReason::ChildOrder);
 		}
 	}
 
@@ -664,18 +672,18 @@ public:
 
 	int32 Add( const TSharedRef<ChildType>& Child )
 	{
-		if (Owner && bChangesInvalidatePrepass)
+		if (GetOwner() && bChangesInvalidatePrepass)
 		{ 
-			Owner->Invalidate(EInvalidateWidgetReason::ChildOrder);
+			GetOwner()->Invalidate(EInvalidateWidgetReason::ChildOrder);
 		}
 
 		int32 Index = TArray< TSharedRef<ChildType> >::Add(Child);
 
-		if (Owner)
+		if (GetOwner())
 		{
 			if (Child != SNullWidget::NullWidget)
 			{
-				Child->AssignParentWidget(Owner->AsShared());
+				Child->AssignParentWidget(GetOwner()->AsShared());
 			}
 		}
 
@@ -690,7 +698,7 @@ public:
 			TSharedRef<SWidget> Child = GetChildAt(ChildIndex);
 			if (Child != SNullWidget::NullWidget)
 			{
-				Child->ConditionallyDetatchParentWidget(Owner);
+				Child->ConditionallyDetatchParentWidget(GetOwner());
 			}
 		}
 
@@ -704,7 +712,7 @@ public:
 			TSharedRef<SWidget> Child = GetChildAt(ChildIndex);
 			if (Child != SNullWidget::NullWidget)
 			{
-				Child->ConditionallyDetatchParentWidget(Owner);
+				Child->ConditionallyDetatchParentWidget(GetOwner());
 			}
 		}
 
@@ -713,18 +721,18 @@ public:
 
 	void Insert(const TSharedRef<ChildType>& Child, int32 Index)
 	{
-		if (Owner && bChangesInvalidatePrepass) 
+		if (GetOwner() && bChangesInvalidatePrepass) 
 		{
-			Owner->Invalidate(EInvalidateWidgetReason::ChildOrder);
+			GetOwner()->Invalidate(EInvalidateWidgetReason::ChildOrder);
 		}
 
 		TArray< TSharedRef<ChildType> >::Insert(Child, Index);
 
-		if (Owner)
+		if (GetOwner())
 		{
 			if (Child != SNullWidget::NullWidget)
 			{
-				Child->AssignParentWidget(Owner->AsShared());
+				Child->AssignParentWidget(GetOwner()->AsShared());
 			}
 		}
 	}
@@ -733,7 +741,7 @@ public:
 	{
 		if (Child != SNullWidget::NullWidget)
 		{
-			Child->ConditionallyDetatchParentWidget(Owner);
+			Child->ConditionallyDetatchParentWidget(GetOwner());
 		}
 
 		const int32 NumFoundAndRemoved = TArray< TSharedRef<ChildType> >::Remove( Child );
@@ -745,7 +753,7 @@ public:
 		TSharedRef<SWidget> Child = GetChildAt(Index);
 		if (Child != SNullWidget::NullWidget)
 		{
-			Child->ConditionallyDetatchParentWidget(Owner);
+			Child->ConditionallyDetatchParentWidget(GetOwner());
 		}
 
 		TArray< TSharedRef<ChildType> >::RemoveAt( Index );
@@ -770,18 +778,18 @@ public:
 	void Sort( const PREDICATE_CLASS& Predicate )
 	{
 		TArray< TSharedRef<ChildType> >::Sort( Predicate );
-		if (Owner && bChangesInvalidatePrepass)
+		if (GetOwner() && bChangesInvalidatePrepass)
 		{
-			Owner->Invalidate(EInvalidateWidgetReason::ChildOrder);
+			GetOwner()->Invalidate(EInvalidateWidgetReason::ChildOrder);
 		}
 	}
 
 	void Swap( int32 IndexA, int32 IndexB )
 	{
 		TIndirectArray< ChildType >::Swap(IndexA, IndexB);
-		if (Owner && bChangesInvalidatePrepass)
+		if (GetOwner() && bChangesInvalidatePrepass)
 		{
-			Owner->Invalidate(EInvalidateWidgetReason::ChildOrder);
+			GetOwner()->Invalidate(EInvalidateWidgetReason::ChildOrder);
 		}
 	}
 
@@ -790,6 +798,56 @@ private:
 };
 
 
+/** A FChildren that support only one slot and alignment of content and padding */
+class FOneSimpleMemberChild : public TSupportsOneChildMixin<FOneSimpleMemberChild>, public TSupportsContentAlignmentMixin<FOneSimpleMemberChild>
+{
+public:
+	template<typename WidgetType, typename V = typename std::enable_if<std::is_base_of<SWidget, WidgetType>::value>::type>
+	FOneSimpleMemberChild(WidgetType& InParent)
+		: TSupportsOneChildMixin<FOneSimpleMemberChild>(&InParent)
+		, TSupportsContentAlignmentMixin<FOneSimpleMemberChild>(HAlign_Fill, VAlign_Fill)
+		, SlotPadding(InParent)
+	{
+	}
+
+	FOneSimpleMemberChild& Padding(TAttribute<FMargin> InPadding)
+	{
+		SlotPadding.Assign(*FChildren::GetOwner(), MoveTemp(InPadding));
+		return *this;
+	}
+
+	FOneSimpleMemberChild& Padding(float Uniform)
+	{
+		SlotPadding.Set(*FChildren::GetOwner(), FMargin(Uniform));
+		return *this;
+	}
+
+	FOneSimpleMemberChild& Padding(float Horizontal, float Vertical)
+	{
+		SlotPadding.Set(*FChildren::GetOwner(), FMargin(Horizontal, Vertical));
+		return *this;
+	}
+
+	FOneSimpleMemberChild& Padding(float Left, float Top, float Right, float Bottom)
+	{
+		SlotPadding.Set(*FChildren::GetOwner(), FMargin(Left, Top, Right, Bottom));
+		return *this;
+	}
+
+	const FMargin& GetSlotPadding() const { return SlotPadding.Get(); }
+
+public:
+	using SlotPaddingAttributeType = SlateAttributePrivate::TSlateMemberAttribute<FMargin, TSlateAttributeInvalidationReason<EInvalidateWidgetReason::Layout>>;
+	using SlotPaddingAttributeRefType = SlateAttributePrivate::TSlateMemberAttributeRef<SlotPaddingAttributeType>;
+
+	template<typename WidgetType, typename V = typename std::enable_if<std::is_base_of<SWidget, WidgetType>::value>::type>
+	SlotPaddingAttributeRefType GetSlotPaddingAttribute() const { return SlotPaddingAttributeRefType(*(static_cast<const WidgetType*>(FChildren::GetOwner())), SlotPadding); }
+
+private:
+	SlotPaddingAttributeType SlotPadding;
+};
+
+
 /** Required to implement GetChildren() in a way that can dynamically return the currently active child. */
 template<typename SlotType>
 class TOneDynamicChild : public FChildren
diff --git a/Engine/Source/Runtime/SlateCore/Public/Types/SlateAttribute.h b/Engine/Source/Runtime/SlateCore/Public/Types/SlateAttribute.h
--- a/Engine/Source/Runtime/SlateCore/Public/Types/SlateAttribute.h
+++ b/Engine/Source/Runtime/SlateCore/Public/Types/SlateAttribute.h
@@ -443,8 +443,11 @@ namespace SlateAttributePrivate
 		}
 
 	public:
-		/** Unbind the SlateAttribute and set its value. It may invalidate the Widget if the value is different. */
-		void Set(SWidget& Widget, const ObjectType& NewValue)
+		/**
+		 * Unbind the SlateAttribute and set its value. It may invalidate the Widget if the value is different.
+		 * @return true if the value is considered different and an invalidation occurred.
+		 */
+		bool Set(SWidget& Widget, const ObjectType& NewValue)
 		{
 			VerifyOwningWidget(Widget);
 			ProtectedUnregisterAttribute(Widget, InAttributeType);
@@ -455,10 +458,14 @@ namespace SlateAttributePrivate
 				Value = NewValue;
 				ProtectedInvalidateWidget(Widget, InAttributeType, GetInvalidationReason(Widget));
 			}
+			return !bIsIdentical;
 		}
 
-		/** Unbind the SlateAttribute and set its value. It may invalidate the Widget if the value is different. */
-		void Set(SWidget& Widget, ObjectType&& NewValue)
+		/**
+		 * Unbind the SlateAttribute and set its value. It may invalidate the Widget if the value is different.
+		 * @return true if the value is considered different and an invalidation occurred.
+		 */
+		bool Set(SWidget& Widget, ObjectType&& NewValue)
 		{
 			VerifyOwningWidget(Widget);
 			ProtectedUnregisterAttribute(Widget, InAttributeType);
@@ -469,6 +476,7 @@ namespace SlateAttributePrivate
 				Value = MoveTemp(NewValue);
 				ProtectedInvalidateWidget(Widget, InAttributeType, GetInvalidationReason(Widget));
 			}
+			return !bIsIdentical;
 		}
 
 	public:
@@ -530,106 +538,109 @@ namespace SlateAttributePrivate
 		 * Or
 		 * Unbind the SlateAttribute if the Attribute is not bound and not set.
 		 * @see Set
+		 * @return true if the getter was assigned or if the value is considered different and an invalidation occurred.
 		 */
-		void Assign(SWidget& Widget, const TAttribute<ObjectType>& OtherAttribute)
+		bool Assign(SWidget& Widget, const TAttribute<ObjectType>& OtherAttribute)
 		{
 			VerifyOwningWidget(Widget);
 			if (OtherAttribute.IsBound())
 			{
-				AssignBinding(Widget, OtherAttribute.GetBinding());
+				return AssignBinding(Widget, OtherAttribute.GetBinding());
 			}
 			else if (OtherAttribute.IsSet())
 			{
-				Set(Widget, OtherAttribute.Get());
+				return Set(Widget, OtherAttribute.Get());
 			}
 			else
 			{
 				ProtectedUnregisterAttribute(Widget, InAttributeType);
+				return false;
 			}
 		}
 
-		void Assign(SWidget& Widget, TAttribute<ObjectType>&& OtherAttribute)
+		bool Assign(SWidget& Widget, TAttribute<ObjectType>&& OtherAttribute)
 		{
 			VerifyOwningWidget(Widget);
 			if (OtherAttribute.IsBound())
 			{
-				AssignBinding(Widget, MoveTemp(OtherAttribute.Steal().template Get<FGetter>()));
+				return AssignBinding(Widget, MoveTemp(OtherAttribute.Steal().template Get<FGetter>()));
 			}
 			else if (OtherAttribute.IsSet())
 			{
-				Set(Widget, MoveTemp(OtherAttribute.Steal().template Get<ObjectType>()));
+				return Set(Widget, MoveTemp(OtherAttribute.Steal().template Get<ObjectType>()));
 			}
 			else
 			{
 				ProtectedUnregisterAttribute(Widget, InAttributeType);
+				return false;
 			}
 		}
 
-		void Assign(SWidget& Widget, const TAttribute<ObjectType>& OtherAttribute, const ObjectType& DefaultValue)
+		bool Assign(SWidget& Widget, const TAttribute<ObjectType>& OtherAttribute, const ObjectType& DefaultValue)
 		{
 			VerifyOwningWidget(Widget);
 			if (OtherAttribute.IsBound())
 			{
-				AssignBinding(Widget, OtherAttribute.GetBinding());
+				return AssignBinding(Widget, OtherAttribute.GetBinding());
 			}
 			else if (OtherAttribute.IsSet())
 			{
-				Set(Widget, OtherAttribute.Get());
+				return Set(Widget, OtherAttribute.Get());
 			}
 			else
 			{
-				Set(Widget, DefaultValue);
+				return Set(Widget, DefaultValue);
 			}
 		}
 
-		void Assign(SWidget& Widget, TAttribute<ObjectType>&& OtherAttribute, const ObjectType& DefaultValue)
+		bool Assign(SWidget& Widget, TAttribute<ObjectType>&& OtherAttribute, const ObjectType& DefaultValue)
 		{
 			VerifyOwningWidget(Widget);
 			if (OtherAttribute.IsBound())
 			{
-				AssignBinding(Widget, MoveTemp(OtherAttribute.Steal().template Get<FGetter>()));
+				return AssignBinding(Widget, MoveTemp(OtherAttribute.Steal().template Get<FGetter>()));
 			}
 			else if (OtherAttribute.IsSet())
 			{
-				Set(Widget, MoveTemp(OtherAttribute.Steal().template Get<ObjectType>()));
+				return Set(Widget, MoveTemp(OtherAttribute.Steal().template Get<ObjectType>()));
 			}
 			else
 			{
-				Set(Widget, DefaultValue);
+				return Set(Widget, DefaultValue);
 			}
 		}
 
-		void Assign(SWidget& Widget, const TAttribute<ObjectType>& OtherAttribute, ObjectType&& DefaultValue)
+		bool Assign(SWidget& Widget, const TAttribute<ObjectType>& OtherAttribute, ObjectType&& DefaultValue)
 		{
 			VerifyOwningWidget(Widget);
 			if (OtherAttribute.IsBound())
 			{
-				AssignBinding(Widget, OtherAttribute.GetBinding());
+				return AssignBinding(Widget, OtherAttribute.GetBinding());
 			}
 			else if (OtherAttribute.IsSet())
 			{
-				Set(Widget, OtherAttribute.Get());
+				return Set(Widget, OtherAttribute.Get());
 			}
 			else
 			{
-				Set(Widget, MoveTemp(DefaultValue));
+				return Set(Widget, MoveTemp(DefaultValue));
 			}
 		}
 
-		void Assign(SWidget& Widget, TAttribute<ObjectType>&& OtherAttribute, ObjectType&& DefaultValue)
+		bool Assign(SWidget& Widget, TAttribute<ObjectType>&& OtherAttribute, ObjectType&& DefaultValue)
 		{
 			VerifyOwningWidget(Widget);
 			if (OtherAttribute.IsBound())
 			{
-				AssignBinding(Widget, MoveTemp(OtherAttribute.Steal().template Get<FGetter>()));
+				return AssignBinding(Widget, MoveTemp(OtherAttribute.Steal().template Get<FGetter>()));
 			}
 			else if (OtherAttribute.IsSet())
 			{
-				Set(Widget, MoveTemp(OtherAttribute.Steal().template Get<ObjectType>()));
+				return Set(Widget, MoveTemp(OtherAttribute.Steal().template Get<ObjectType>()));
 			}
 			else
 			{
-				Set(Widget, MoveTemp(DefaultValue));
+				return Set(Widget, MoveTemp(DefaultValue));
 			}
 		}
 
@@ -705,22 +716,26 @@ namespace SlateAttributePrivate
 			UpdateNowOnBind(Widget);
 		}
 
-		void AssignBinding(SWidget& Widget, const FGetter& Getter)
+		bool AssignBinding(SWidget& Widget, const FGetter& Getter)
 		{
 			const FDelegateHandle PreviousGetterHandle = ProtectedFindGetterHandle(Widget, InAttributeType);
 			if (PreviousGetterHandle != Getter.GetHandle())
 			{
 				ConstructWrapper(Widget, Getter);
+				return true;
 			}
+			return false;
 		}
 
-		void AssignBinding(SWidget& Widget, FGetter&& Getter)
+		bool AssignBinding(SWidget& Widget, FGetter&& Getter)
 		{
 			const FDelegateHandle PreviousGetterHandle = ProtectedFindGetterHandle(Widget, InAttributeType);
 			if (PreviousGetterHandle != Getter.GetHandle())
 			{
+				return true;
 				ConstructWrapper(Widget, MoveTemp(Getter));
 			}
+			return false;
 		}
 
 	private:
@@ -1133,7 +1148,7 @@ namespace SlateAttributePrivate
 
 	private:
 		template<typename WidgetType>
-		static void VerifyAttributeAddress(WidgetType const& InWidget, AttributeMemberType const& InAttribute)
+		static void VerifyAttributeAddress(const WidgetType& InWidget, const AttributeMemberType& InAttribute)
 		{
 			checkf((UPTRINT)&InAttribute >= (UPTRINT)&InWidget && (UPTRINT)&InAttribute < (UPTRINT)&InWidget + sizeof(WidgetType),
 				TEXT("The attribute is not a member of the widget."));
@@ -1142,12 +1157,10 @@ namespace SlateAttributePrivate
 
 	public:
 		/** Constructor */
-		TSlateMemberAttributeRef() = default;
-
 		template<typename WidgetType, typename V = typename std::enable_if<std::is_base_of<SWidget, WidgetType>::value>::type>
-		explicit TSlateMemberAttributeRef(WidgetType const& InOwner, AttributeMemberType const& InAttribute)
+		explicit TSlateMemberAttributeRef(const WidgetType& InOwner, const AttributeMemberType& InAttribute)
 			: Owner(InOwner.AsShared())
-			, Attribute(&InAttribute)
+			, Attribute(InAttribute)
 		{
 			VerifyAttributeAddress(InOwner, InAttribute);
 		}
@@ -1164,7 +1177,7 @@ namespace SlateAttributePrivate
 		{
 			if (TSharedPtr<const SWidget> Pin = Owner.Pin())
 			{
-				return Attribute->Get();
+				return Attribute.Get();
 			}
 			checkf(false, TEXT("It is an error to call GetValue() on an unset TSlateMemberAttributeRef. Please either check IsValid() or use Get(DefaultValue) instead."));
 			static ObjectType Tmp;
@@ -1176,7 +1189,7 @@ namespace SlateAttributePrivate
 		{
 			if (TSharedPtr<const SWidget> Pin = Owner.Pin())
 			{
-				return Attribute->Get();
+				return Attribute.Get();
 			}
 			return DefaultValue;
 		}
@@ -1186,7 +1199,7 @@ namespace SlateAttributePrivate
 		{
 			if (TSharedPtr<const SWidget> Pin = Owner.Pin())
 			{
-				const_cast<AttributeMemberType*>(Attribute)->UpdateNow(*Pin.Get());
+				const_cast<AttributeMemberType>(Attribute).UpdateNow(*Pin.Get());
 			}
 		}
 
@@ -1198,8 +1211,8 @@ namespace SlateAttributePrivate
 		{
 			if (TSharedPtr<const SWidget> Pin = Owner.Pin())
 			{
-				const_cast<AttributeMemberType*>(Attribute)->UpdateNow(*Pin.Get());
-				return Attribute->Get();
+				const_cast<AttributeMemberType>(Attribute)->UpdateNow(*Pin.Get());
+				return Attribute.Get();
 			}
 			checkf(false, TEXT("It is an error to call GetValue() on an unset TSlateMemberAttributeRef. Please either check IsValid() or use Get(DefaultValue) instead."));
 		}
@@ -1212,8 +1225,8 @@ namespace SlateAttributePrivate
 		{
 			if (TSharedPtr<const SWidget> Pin = Owner.Pin())
 			{
-				const_cast<AttributeMemberType*>(Attribute)->UpdateNow(*Pin.Get());
-				return Attribute->Get();
+				const_cast<AttributeMemberType>(Attribute)->UpdateNow(*Pin.Get());
+				return Attribute.Get();
 			}
 			return DefaultValue;
 		}
@@ -1223,7 +1236,7 @@ namespace SlateAttributePrivate
 		{
 			if (TSharedPtr<const SWidget> Pin = Owner.Pin())
 			{
-				return Attribute->ToAttribute(*Pin.Get());
+				return Attribute.ToAttribute(*Pin.Get());
 			}
 			return TAttribute<ObjectType>();
 		}
@@ -1233,7 +1246,7 @@ namespace SlateAttributePrivate
 		{
 			if (TSharedPtr<const SWidget> Pin = Owner.Pin())
 			{
-				return Attribute->IsBound(*Pin.Get());
+				return Attribute.IsBound(*Pin.Get());
 			}
 			return false;
 		}
@@ -1245,7 +1258,7 @@ namespace SlateAttributePrivate
 			TSharedPtr<const SWidget> OtherPin = Other.Owner.Pin();
 			if (SelfPin == OtherPin && SelfPin)
 			{
-				return Attribute->IsIdenticalTo(*SelfPin.Get(), *Other.Attribute);
+				return Attribute.IsIdenticalTo(*SelfPin.Get(), Other.Attribute);
 			}
 			return SelfPin == OtherPin;
 		}
@@ -1255,14 +1268,14 @@ namespace SlateAttributePrivate
 		{
 			if (TSharedPtr<const SWidget> Pin = Owner.Pin())
 			{
-				return Attribute->IsIdenticalTo(*Pin.Get(), Other);
+				return Attribute.IsIdenticalTo(*Pin.Get(), Other);
 			}
 			return !Other.IsSet(); // if the other is not set, then both are invalid.
 		}
 
 	private:
 		TWeakPtr<const SWidget> Owner;
-		AttributeMemberType const* Attribute = nullptr;
+		const AttributeMemberType& Attribute;
 	 };
 
 
