Fix CIS with multiple deprecation warnings
#rn none
#rnx

diff --git a/Engine/Source/Runtime/SlateCore/Public/Layout/Children.h b/Engine/Source/Runtime/SlateCore/Public/Layout/Children.h
--- a/Engine/Source/Runtime/SlateCore/Public/Layout/Children.h
+++ b/Engine/Source/Runtime/SlateCore/Public/Layout/Children.h
@@ -44,10 +44,8 @@ protected:
 protected:
 	virtual ~FChildren(){}
 
-	PRAGMA_DISABLE_DEPRECATION_WARNINGS
 	SWidget* GetOwner() { return Owner; }
 	const SWidget* GetOwner() const { return Owner; }
-	PRAGMA_ENABLE_DEPRECATION_WARNINGS
 
 protected:
 	UE_DEPRECATED(5.0, "Direct access to Owner is now deprecated. Use the getter.")
