Slate: Fix crash when the attribute array index is not the same array as the childorder array index.

diff --git a/Engine/Source/Runtime/SlateCore/Private/FastUpdate/SlateInvalidationWidgetList.cpp b/Engine/Source/Runtime/SlateCore/Private/FastUpdate/SlateInvalidationWidgetList.cpp
--- a/Engine/Source/Runtime/SlateCore/Private/FastUpdate/SlateInvalidationWidgetList.cpp
+++ b/Engine/Source/Runtime/SlateCore/Private/FastUpdate/SlateInvalidationWidgetList.cpp
@@ -35,10 +35,13 @@ int32 FSlateInvalidationWidgetList::FArguments::MaxSortOrderPaddingBetweenArray
 /** */
 FSlateInvalidationWidgetIndex FSlateInvalidationWidgetList::IProcessChildOrderInvalidationCallback::FReIndexOperation::ReIndex(FSlateInvalidationWidgetIndex Index) const
 {
-	check(Range.GetInclusiveMinWidgetIndex().ArrayIndex == Range.GetInclusiveMaxWidgetIndex().ArrayIndex);
-	check(Index.ArrayIndex == Range.GetInclusiveMaxWidgetIndex().ArrayIndex);
-	const IndexType NewElementIndex = Index.ElementIndex - Range.GetInclusiveMinWidgetIndex().ElementIndex + ReIndexTarget.ElementIndex;
-	return { ReIndexTarget.ArrayIndex, NewElementIndex };
+	if (Index.ArrayIndex == Range.GetInclusiveMaxWidgetIndex().ArrayIndex)
+	{
+		check(Range.GetInclusiveMinWidgetIndex().ArrayIndex == Range.GetInclusiveMaxWidgetIndex().ArrayIndex);
+		const IndexType NewElementIndex = Index.ElementIndex - Range.GetInclusiveMinWidgetIndex().ElementIndex + ReIndexTarget.ElementIndex;
+		return { ReIndexTarget.ArrayIndex, NewElementIndex };
+	}
+	return Index;
 }
 
 
