Revert changelist 15730745

diff --git a/Engine/Source/Runtime/Slate/Private/Widgets/Layout/SDPIScaler.cpp b/Engine/Source/Runtime/Slate/Private/Widgets/Layout/SDPIScaler.cpp
--- a/Engine/Source/Runtime/Slate/Private/Widgets/Layout/SDPIScaler.cpp
+++ b/Engine/Source/Runtime/Slate/Private/Widgets/Layout/SDPIScaler.cpp
@@ -3,15 +3,8 @@
 #include "Widgets/Layout/SDPIScaler.h"
 #include "Layout/ArrangedChildren.h"
 
-SLATE_IMPLEMENT_WIDGET(SDPIScaler)
-void SDPIScaler::PrivateRegisterAttributes(FSlateAttributeInitializer& AttributeInitializer)
-{
-	SLATE_ADD_MEMBER_ATTRIBUTE_DEFINITION_WITH_NAME(AttributeInitializer, "DPIScale", DPIScaleAttribute, EInvalidateWidgetReason::Layout);
-}
-
 SDPIScaler::SDPIScaler()
-	: ChildSlot(*this)
-	, DPIScaleAttribute(*this)
+	: ChildSlot(this)
 {
 	SetCanTick(false);
 	bCanSupportFocus = false;
@@ -20,7 +13,7 @@ SDPIScaler::SDPIScaler()
 
 void SDPIScaler::Construct( const FArguments& InArgs )
 {
-	SetDPIScale(InArgs._DPIScale);
+	DPIScale = InArgs._DPIScale;
 
 	ChildSlot
 	[
@@ -33,7 +26,7 @@ void SDPIScaler::OnArrangeChildren( const FGeometry& AllottedGeometry, FArranged
 	const EVisibility MyVisibility = this->GetVisibility();
 	if ( ArrangedChildren.Accepts( MyVisibility ) )
 	{
-		const float MyDPIScale = DPIScaleAttribute.Get();
+		const float MyDPIScale = DPIScale.Get();
 
 		ArrangedChildren.AddWidget( AllottedGeometry.MakeChild(
 			this->ChildSlot.GetWidget(),
@@ -46,7 +39,7 @@ void SDPIScaler::OnArrangeChildren( const FGeometry& AllottedGeometry, FArranged
 	
 FVector2D SDPIScaler::ComputeDesiredSize( float ) const
 {
-	return DPIScaleAttribute.Get() * ChildSlot.GetWidget()->GetDesiredSize();
+	return DPIScale.Get() * ChildSlot.GetWidget()->GetDesiredSize();
 }
 
 FChildren* SDPIScaler::GetChildren()
@@ -64,7 +57,7 @@ void SDPIScaler::SetContent(TSharedRef<SWidget> InContent)
 
 void SDPIScaler::SetDPIScale(TAttribute<float> InDPIScale)
 {
-	if (DPIScaleAttribute.Assign(*this, MoveTemp(InDPIScale)))
+	if (SetAttribute(DPIScale, InDPIScale, EInvalidateWidgetReason::Layout))
 	{
 		InvalidatePrepass();
 	}
@@ -72,5 +65,5 @@ void SDPIScaler::SetDPIScale(TAttribute<float> InDPIScale)
 
 float SDPIScaler::GetRelativeLayoutScale(int32 ChildIndex, float LayoutScaleMultiplier) const
 {
-	return DPIScaleAttribute.Get();
+	return DPIScale.Get();
 }
diff --git a/Engine/Source/Runtime/Slate/Public/Widgets/Layout/SDPIScaler.h b/Engine/Source/Runtime/Slate/Public/Widgets/Layout/SDPIScaler.h
--- a/Engine/Source/Runtime/Slate/Public/Widgets/Layout/SDPIScaler.h
+++ b/Engine/Source/Runtime/Slate/Public/Widgets/Layout/SDPIScaler.h
@@ -17,8 +17,6 @@ class FArrangedChildren;
  */
 class SLATE_API SDPIScaler : public SPanel
 {
-	SLATE_DECLARE_WIDGET(SDPIScaler, SPanel)
-
 public:
 	SLATE_BEGIN_ARGS( SDPIScaler )
 		: _DPIScale( 1.0f )
@@ -58,21 +56,13 @@ protected:
 
 	virtual float GetRelativeLayoutScale(int32 ChildIndex, float LayoutScaleMultiplier) const override;
 
-	TSlateAttributeRef<float> GetDPIScaleAttribute() const { return TSlateAttributeRef<float>(*this, DPIScaleAttribute); }
-
-#if WITH_EDITORONLY_DATA
-	UE_DEPRECATED(5.0, "Direct access to DPIScale is now deprecated. Use the setter or getter.")
-	FSlateDeprecatedTAttribute<float> DPIScale;
-#endif
-
-	/** The content being scaled. */
-	FOneSimpleMemberChild ChildSlot;
-
-private:
 	/**
 	 * The scaling factor from 1:1 pixel to slate unit (SU) ratio.
 	 * e.g. Value of 2 would mean that there are two pixels for every slate unit in every dimension.
 	 *      This means that every 1x1 SU square is represented by 2x2=4 pixels.
 	 */
-	TSlateAttribute<float> DPIScaleAttribute;
+	TAttribute<float> DPIScale;
+
+	/** The content being scaled. */
+	FSimpleSlot ChildSlot;
 };
diff --git a/Engine/Source/Runtime/SlateCore/Public/Layout/Children.h b/Engine/Source/Runtime/SlateCore/Public/Layout/Children.h
--- a/Engine/Source/Runtime/SlateCore/Public/Layout/Children.h
+++ b/Engine/Source/Runtime/SlateCore/Public/Layout/Children.h
@@ -19,7 +19,6 @@
  * FChildren is intended to be returned by the GetChildren() method.
  * 
  */
-PRAGMA_DISABLE_DEPRECATION_WARNINGS
 class SLATECORE_API FChildren
 {
 public:
@@ -44,14 +43,9 @@ protected:
 protected:
 	virtual ~FChildren(){}
 
-	SWidget* GetOwner() { return Owner; }
-	const SWidget* GetOwner() const { return Owner; }
-
 protected:
-	UE_DEPRECATED(5.0, "Direct access to Owner is now deprecated. Use the getter.")
 	SWidget* Owner;
 };
-PRAGMA_ENABLE_DEPRECATION_WARNINGS
 
 
 /**
@@ -224,13 +218,13 @@ public:
 	void AttachWidget(const TSharedPtr<SWidget>& InWidget)
 	{
 		WidgetPtr = InWidget;
-		if (GetOwner()) 
+		if (Owner) 
 		{ 
-			GetOwner()->Invalidate(EInvalidateWidgetReason::ChildOrder);
+			Owner->Invalidate(EInvalidateWidgetReason::ChildOrder);
 
 			if (InWidget.IsValid() && InWidget != SNullWidget::NullWidget)
 			{
-				InWidget->AssignParentWidget(GetOwner()->AsShared());
+				InWidget->AssignParentWidget(Owner->AsShared());
 			}
 		}
 	}
@@ -243,7 +237,7 @@ public:
 
 			if (Widget != SNullWidget::NullWidget)
 			{
-				Widget->ConditionallyDetatchParentWidget(GetOwner());
+				Widget->ConditionallyDetatchParentWidget(Owner);
 			}
 
 			WidgetPtr.Reset();
@@ -389,9 +383,9 @@ public:
 
 		int32 Index = TIndirectArray< SlotType >::Add(Slot);
 
-		if (GetOwner())
+		if (Owner)
 		{
-			Slot->AttachWidgetParent(GetOwner());
+			Slot->AttachWidgetParent(Owner);
 		}
 
 		return Index;
@@ -436,9 +430,9 @@ public:
 			TIndirectArray< SlotType >::Insert(Slot, Index);
 
 			// Don't do parent manipulation if this panel has no owner.
-			if (GetOwner())
+			if (Owner)
 			{
-				Slot->AttachWidgetParent(GetOwner());
+				Slot->AttachWidgetParent(Owner);
 			}
 		}
 	}
@@ -461,9 +455,9 @@ public:
 			TIndirectArray< SlotType >::RemoveAt(IndexToMove);
 		}
 
-		if (GetOwner())
+		if (Owner)
 		{
-			GetOwner()->Invalidate(EInvalidateWidgetReason::ChildOrder);
+			Owner->Invalidate(EInvalidateWidgetReason::ChildOrder);
 		}
 	}
 
@@ -484,9 +478,9 @@ public:
 	void Sort( const PREDICATE_CLASS& Predicate )
 	{
 		::Sort(TIndirectArray< SlotType >::GetData(), TIndirectArray<SlotType>::Num(), Predicate);
-		if (GetOwner())
+		if (Owner)
 		{
-			GetOwner()->Invalidate(EInvalidateWidgetReason::ChildOrder);
+			Owner->Invalidate(EInvalidateWidgetReason::ChildOrder);
 		}
 	}
 
@@ -494,18 +488,18 @@ public:
 	void StableSort(const PREDICATE_CLASS& Predicate)
 	{
 		::StableSort(TIndirectArray< SlotType >::GetData(), TIndirectArray< SlotType >::Num(), Predicate);
-		if (GetOwner())
+		if (Owner)
 		{
-			GetOwner()->Invalidate(EInvalidateWidgetReason::ChildOrder);
+			Owner->Invalidate(EInvalidateWidgetReason::ChildOrder);
 		}
 	}
 
 	void Swap( int32 IndexA, int32 IndexB )
 	{
 		TIndirectArray< SlotType >::Swap(IndexA, IndexB);
-		if (GetOwner())
+		if (Owner)
 		{
-			GetOwner()->Invalidate(EInvalidateWidgetReason::ChildOrder);
+			Owner->Invalidate(EInvalidateWidgetReason::ChildOrder);
 		}
 	}
 
@@ -670,18 +664,18 @@ public:
 
 	int32 Add( const TSharedRef<ChildType>& Child )
 	{
-		if (GetOwner() && bChangesInvalidatePrepass)
+		if (Owner && bChangesInvalidatePrepass)
 		{ 
-			GetOwner()->Invalidate(EInvalidateWidgetReason::ChildOrder);
+			Owner->Invalidate(EInvalidateWidgetReason::ChildOrder);
 		}
 
 		int32 Index = TArray< TSharedRef<ChildType> >::Add(Child);
 
-		if (GetOwner())
+		if (Owner)
 		{
 			if (Child != SNullWidget::NullWidget)
 			{
-				Child->AssignParentWidget(GetOwner()->AsShared());
+				Child->AssignParentWidget(Owner->AsShared());
 			}
 		}
 
@@ -696,7 +690,7 @@ public:
 			TSharedRef<SWidget> Child = GetChildAt(ChildIndex);
 			if (Child != SNullWidget::NullWidget)
 			{
-				Child->ConditionallyDetatchParentWidget(GetOwner());
+				Child->ConditionallyDetatchParentWidget(Owner);
 			}
 		}
 
@@ -710,7 +704,7 @@ public:
 			TSharedRef<SWidget> Child = GetChildAt(ChildIndex);
 			if (Child != SNullWidget::NullWidget)
 			{
-				Child->ConditionallyDetatchParentWidget(GetOwner());
+				Child->ConditionallyDetatchParentWidget(Owner);
 			}
 		}
 
@@ -719,18 +713,18 @@ public:
 
 	void Insert(const TSharedRef<ChildType>& Child, int32 Index)
 	{
-		if (GetOwner() && bChangesInvalidatePrepass) 
+		if (Owner && bChangesInvalidatePrepass) 
 		{
-			GetOwner()->Invalidate(EInvalidateWidgetReason::ChildOrder);
+			Owner->Invalidate(EInvalidateWidgetReason::ChildOrder);
 		}
 
 		TArray< TSharedRef<ChildType> >::Insert(Child, Index);
 
-		if (GetOwner())
+		if (Owner)
 		{
 			if (Child != SNullWidget::NullWidget)
 			{
-				Child->AssignParentWidget(GetOwner()->AsShared());
+				Child->AssignParentWidget(Owner->AsShared());
 			}
 		}
 	}
@@ -739,7 +733,7 @@ public:
 	{
 		if (Child != SNullWidget::NullWidget)
 		{
-			Child->ConditionallyDetatchParentWidget(GetOwner());
+			Child->ConditionallyDetatchParentWidget(Owner);
 		}
 
 		const int32 NumFoundAndRemoved = TArray< TSharedRef<ChildType> >::Remove( Child );
@@ -751,7 +745,7 @@ public:
 		TSharedRef<SWidget> Child = GetChildAt(Index);
 		if (Child != SNullWidget::NullWidget)
 		{
-			Child->ConditionallyDetatchParentWidget(GetOwner());
+			Child->ConditionallyDetatchParentWidget(Owner);
 		}
 
 		TArray< TSharedRef<ChildType> >::RemoveAt( Index );
@@ -776,18 +770,18 @@ public:
 	void Sort( const PREDICATE_CLASS& Predicate )
 	{
 		TArray< TSharedRef<ChildType> >::Sort( Predicate );
-		if (GetOwner() && bChangesInvalidatePrepass)
+		if (Owner && bChangesInvalidatePrepass)
 		{
-			GetOwner()->Invalidate(EInvalidateWidgetReason::ChildOrder);
+			Owner->Invalidate(EInvalidateWidgetReason::ChildOrder);
 		}
 	}
 
 	void Swap( int32 IndexA, int32 IndexB )
 	{
 		TIndirectArray< ChildType >::Swap(IndexA, IndexB);
-		if (GetOwner() && bChangesInvalidatePrepass)
+		if (Owner && bChangesInvalidatePrepass)
 		{
-			GetOwner()->Invalidate(EInvalidateWidgetReason::ChildOrder);
+			Owner->Invalidate(EInvalidateWidgetReason::ChildOrder);
 		}
 	}
 
@@ -796,56 +790,6 @@ private:
 };
 
 
-/** A FChildren that support only one slot and alignment of content and padding */
-class FOneSimpleMemberChild : public TSupportsOneChildMixin<FOneSimpleMemberChild>, public TSupportsContentAlignmentMixin<FOneSimpleMemberChild>
-{
-public:
-	template<typename WidgetType, typename V = typename std::enable_if<std::is_base_of<SWidget, WidgetType>::value>::type>
-	FOneSimpleMemberChild(WidgetType& InParent)
-		: TSupportsOneChildMixin<FOneSimpleMemberChild>(&InParent)
-		, TSupportsContentAlignmentMixin<FOneSimpleMemberChild>(HAlign_Fill, VAlign_Fill)
-		, SlotPadding(InParent)
-	{
-	}
-
-	FOneSimpleMemberChild& Padding(TAttribute<FMargin> InPadding)
-	{
-		SlotPadding.Assign(*FChildren::GetOwner(), MoveTemp(InPadding));
-		return *this;
-	}
-
-	FOneSimpleMemberChild& Padding(float Uniform)
-	{
-		SlotPadding.Set(*FChildren::GetOwner(), FMargin(Uniform));
-		return *this;
-	}
-
-	FOneSimpleMemberChild& Padding(float Horizontal, float Vertical)
-	{
-		SlotPadding.Set(*FChildren::GetOwner(), FMargin(Horizontal, Vertical));
-		return *this;
-	}
-
-	FOneSimpleMemberChild& Padding(float Left, float Top, float Right, float Bottom)
-	{
-		SlotPadding.Set(*FChildren::GetOwner(), FMargin(Left, Top, Right, Bottom));
-		return *this;
-	}
-
-	const FMargin& GetSlotPadding() const { return SlotPadding.Get(); }
-
-public:
-	using SlotPaddingAttributeType = SlateAttributePrivate::TSlateMemberAttribute<FMargin, TSlateAttributeInvalidationReason<EInvalidateWidgetReason::Layout>>;
-	using SlotPaddingAttributeRefType = SlateAttributePrivate::TSlateMemberAttributeRef<SlotPaddingAttributeType>;
-
-	template<typename WidgetType, typename V = typename std::enable_if<std::is_base_of<SWidget, WidgetType>::value>::type>
-	SlotPaddingAttributeRefType GetSlotPaddingAttribute() const { return SlotPaddingAttributeRefType(*(static_cast<const WidgetType*>(FChildren::GetOwner())), SlotPadding); }
-
-private:
-	SlotPaddingAttributeType SlotPadding;
-};
-
-
 /** Required to implement GetChildren() in a way that can dynamically return the currently active child. */
 template<typename SlotType>
 class TOneDynamicChild : public FChildren
diff --git a/Engine/Source/Runtime/SlateCore/Public/Types/SlateAttribute.h b/Engine/Source/Runtime/SlateCore/Public/Types/SlateAttribute.h
--- a/Engine/Source/Runtime/SlateCore/Public/Types/SlateAttribute.h
+++ b/Engine/Source/Runtime/SlateCore/Public/Types/SlateAttribute.h
@@ -443,11 +443,8 @@ namespace SlateAttributePrivate
 		}
 
 	public:
-		/**
-		 * Unbind the SlateAttribute and set its value. It may invalidate the Widget if the value is different.
-		 * @return true if the value is considered different and an invalidation occurred.
-		 */
-		bool Set(SWidget& Widget, const ObjectType& NewValue)
+		/** Unbind the SlateAttribute and set its value. It may invalidate the Widget if the value is different. */
+		void Set(SWidget& Widget, const ObjectType& NewValue)
 		{
 			VerifyOwningWidget(Widget);
 			ProtectedUnregisterAttribute(Widget, InAttributeType);
@@ -458,14 +455,10 @@ namespace SlateAttributePrivate
 				Value = NewValue;
 				ProtectedInvalidateWidget(Widget, InAttributeType, GetInvalidationReason(Widget));
 			}
-			return !bIsIdentical;
 		}
 
-		/**
-		 * Unbind the SlateAttribute and set its value. It may invalidate the Widget if the value is different.
-		 * @return true if the value is considered different and an invalidation occurred.
-		 */
-		bool Set(SWidget& Widget, ObjectType&& NewValue)
+		/** Unbind the SlateAttribute and set its value. It may invalidate the Widget if the value is different. */
+		void Set(SWidget& Widget, ObjectType&& NewValue)
 		{
 			VerifyOwningWidget(Widget);
 			ProtectedUnregisterAttribute(Widget, InAttributeType);
@@ -476,7 +469,6 @@ namespace SlateAttributePrivate
 				Value = MoveTemp(NewValue);
 				ProtectedInvalidateWidget(Widget, InAttributeType, GetInvalidationReason(Widget));
 			}
-			return !bIsIdentical;
 		}
 
 	public:
@@ -538,109 +530,106 @@ namespace SlateAttributePrivate
 		 * Or
 		 * Unbind the SlateAttribute if the Attribute is not bound and not set.
 		 * @see Set
-		 * @return true if the getter was assigned or if the value is considered different and an invalidation occurred.
 		 */
-		bool Assign(SWidget& Widget, const TAttribute<ObjectType>& OtherAttribute)
+		void Assign(SWidget& Widget, const TAttribute<ObjectType>& OtherAttribute)
 		{
 			VerifyOwningWidget(Widget);
 			if (OtherAttribute.IsBound())
 			{
-				return AssignBinding(Widget, OtherAttribute.GetBinding());
+				AssignBinding(Widget, OtherAttribute.GetBinding());
 			}
 			else if (OtherAttribute.IsSet())
 			{
-				return Set(Widget, OtherAttribute.Get());
+				Set(Widget, OtherAttribute.Get());
 			}
 			else
 			{
 				ProtectedUnregisterAttribute(Widget, InAttributeType);
-				return false;
 			}
 		}
 
-		bool Assign(SWidget& Widget, TAttribute<ObjectType>&& OtherAttribute)
+		void Assign(SWidget& Widget, TAttribute<ObjectType>&& OtherAttribute)
 		{
 			VerifyOwningWidget(Widget);
 			if (OtherAttribute.IsBound())
 			{
-				return AssignBinding(Widget, MoveTemp(OtherAttribute.Steal().template Get<FGetter>()));
+				AssignBinding(Widget, MoveTemp(OtherAttribute.Steal().template Get<FGetter>()));
 			}
 			else if (OtherAttribute.IsSet())
 			{
-				return Set(Widget, MoveTemp(OtherAttribute.Steal().template Get<ObjectType>()));
+				Set(Widget, MoveTemp(OtherAttribute.Steal().template Get<ObjectType>()));
 			}
 			else
 			{
 				ProtectedUnregisterAttribute(Widget, InAttributeType);
-				return false;
 			}
 		}
 
-		bool Assign(SWidget& Widget, const TAttribute<ObjectType>& OtherAttribute, const ObjectType& DefaultValue)
+		void Assign(SWidget& Widget, const TAttribute<ObjectType>& OtherAttribute, const ObjectType& DefaultValue)
 		{
 			VerifyOwningWidget(Widget);
 			if (OtherAttribute.IsBound())
 			{
-				return AssignBinding(Widget, OtherAttribute.GetBinding());
+				AssignBinding(Widget, OtherAttribute.GetBinding());
 			}
 			else if (OtherAttribute.IsSet())
 			{
-				return Set(Widget, OtherAttribute.Get());
+				Set(Widget, OtherAttribute.Get());
 			}
 			else
 			{
-				return Set(Widget, DefaultValue);
+				Set(Widget, DefaultValue);
 			}
 		}
 
-		bool Assign(SWidget& Widget, TAttribute<ObjectType>&& OtherAttribute, const ObjectType& DefaultValue)
+		void Assign(SWidget& Widget, TAttribute<ObjectType>&& OtherAttribute, const ObjectType& DefaultValue)
 		{
 			VerifyOwningWidget(Widget);
 			if (OtherAttribute.IsBound())
 			{
-				return AssignBinding(Widget, MoveTemp(OtherAttribute.Steal().template Get<FGetter>()));
+				AssignBinding(Widget, MoveTemp(OtherAttribute.Steal().template Get<FGetter>()));
 			}
 			else if (OtherAttribute.IsSet())
 			{
-				return Set(Widget, MoveTemp(OtherAttribute.Steal().template Get<ObjectType>()));
+				Set(Widget, MoveTemp(OtherAttribute.Steal().template Get<ObjectType>()));
 			}
 			else
 			{
-				return Set(Widget, DefaultValue);
+				Set(Widget, DefaultValue);
 			}
 		}
 
-		bool Assign(SWidget& Widget, const TAttribute<ObjectType>& OtherAttribute, ObjectType&& DefaultValue)
+		void Assign(SWidget& Widget, const TAttribute<ObjectType>& OtherAttribute, ObjectType&& DefaultValue)
 		{
 			VerifyOwningWidget(Widget);
 			if (OtherAttribute.IsBound())
 			{
-				return AssignBinding(Widget, OtherAttribute.GetBinding());
+				AssignBinding(Widget, OtherAttribute.GetBinding());
 			}
 			else if (OtherAttribute.IsSet())
 			{
-				return Set(Widget, OtherAttribute.Get());
+				Set(Widget, OtherAttribute.Get());
 			}
 			else
 			{
-				return Set(Widget, MoveTemp(DefaultValue));
+				Set(Widget, MoveTemp(DefaultValue));
 			}
 		}
 
-		bool Assign(SWidget& Widget, TAttribute<ObjectType>&& OtherAttribute, ObjectType&& DefaultValue)
+		void Assign(SWidget& Widget, TAttribute<ObjectType>&& OtherAttribute, ObjectType&& DefaultValue)
 		{
 			VerifyOwningWidget(Widget);
 			if (OtherAttribute.IsBound())
 			{
-				return AssignBinding(Widget, MoveTemp(OtherAttribute.Steal().template Get<FGetter>()));
+				AssignBinding(Widget, MoveTemp(OtherAttribute.Steal().template Get<FGetter>()));
 			}
 			else if (OtherAttribute.IsSet())
 			{
-				return Set(Widget, MoveTemp(OtherAttribute.Steal().template Get<ObjectType>()));
+				Set(Widget, MoveTemp(OtherAttribute.Steal().template Get<ObjectType>()));
 			}
 			else
 			{
-				return Set(Widget, MoveTemp(DefaultValue));
+				Set(Widget, MoveTemp(DefaultValue));
 			}
 		}
 
@@ -716,26 +705,22 @@ namespace SlateAttributePrivate
 			UpdateNowOnBind(Widget);
 		}
 
-		bool AssignBinding(SWidget& Widget, const FGetter& Getter)
+		void AssignBinding(SWidget& Widget, const FGetter& Getter)
 		{
 			const FDelegateHandle PreviousGetterHandle = ProtectedFindGetterHandle(Widget, InAttributeType);
 			if (PreviousGetterHandle != Getter.GetHandle())
 			{
 				ConstructWrapper(Widget, Getter);
-				return true;
 			}
-			return false;
 		}
 
-		bool AssignBinding(SWidget& Widget, FGetter&& Getter)
+		void AssignBinding(SWidget& Widget, FGetter&& Getter)
 		{
 			const FDelegateHandle PreviousGetterHandle = ProtectedFindGetterHandle(Widget, InAttributeType);
 			if (PreviousGetterHandle != Getter.GetHandle())
 			{
-				return true;
 				ConstructWrapper(Widget, MoveTemp(Getter));
 			}
-			return false;
 		}
 
 	private:
@@ -1148,7 +1133,7 @@ namespace SlateAttributePrivate
 
 	private:
 		template<typename WidgetType>
-		static void VerifyAttributeAddress(const WidgetType& InWidget, const AttributeMemberType& InAttribute)
+		static void VerifyAttributeAddress(WidgetType const& InWidget, AttributeMemberType const& InAttribute)
 		{
 			checkf((UPTRINT)&InAttribute >= (UPTRINT)&InWidget && (UPTRINT)&InAttribute < (UPTRINT)&InWidget + sizeof(WidgetType),
 				TEXT("The attribute is not a member of the widget."));
@@ -1157,10 +1142,12 @@ namespace SlateAttributePrivate
 
 	public:
 		/** Constructor */
+		TSlateMemberAttributeRef() = default;
+
 		template<typename WidgetType, typename V = typename std::enable_if<std::is_base_of<SWidget, WidgetType>::value>::type>
-		explicit TSlateMemberAttributeRef(const WidgetType& InOwner, const AttributeMemberType& InAttribute)
+		explicit TSlateMemberAttributeRef(WidgetType const& InOwner, AttributeMemberType const& InAttribute)
 			: Owner(InOwner.AsShared())
-			, Attribute(InAttribute)
+			, Attribute(&InAttribute)
 		{
 			VerifyAttributeAddress(InOwner, InAttribute);
 		}
@@ -1177,7 +1164,7 @@ namespace SlateAttributePrivate
 		{
 			if (TSharedPtr<const SWidget> Pin = Owner.Pin())
 			{
-				return Attribute.Get();
+				return Attribute->Get();
 			}
 			checkf(false, TEXT("It is an error to call GetValue() on an unset TSlateMemberAttributeRef. Please either check IsValid() or use Get(DefaultValue) instead."));
 			static ObjectType Tmp;
@@ -1189,7 +1176,7 @@ namespace SlateAttributePrivate
 		{
 			if (TSharedPtr<const SWidget> Pin = Owner.Pin())
 			{
-				return Attribute.Get();
+				return Attribute->Get();
 			}
 			return DefaultValue;
 		}
@@ -1199,7 +1186,7 @@ namespace SlateAttributePrivate
 		{
 			if (TSharedPtr<const SWidget> Pin = Owner.Pin())
 			{
-				const_cast<AttributeMemberType>(Attribute).UpdateNow(*Pin.Get());
+				const_cast<AttributeMemberType*>(Attribute)->UpdateNow(*Pin.Get());
 			}
 		}
 
@@ -1211,8 +1198,8 @@ namespace SlateAttributePrivate
 		{
 			if (TSharedPtr<const SWidget> Pin = Owner.Pin())
 			{
-				const_cast<AttributeMemberType>(Attribute)->UpdateNow(*Pin.Get());
-				return Attribute.Get();
+				const_cast<AttributeMemberType*>(Attribute)->UpdateNow(*Pin.Get());
+				return Attribute->Get();
 			}
 			checkf(false, TEXT("It is an error to call GetValue() on an unset TSlateMemberAttributeRef. Please either check IsValid() or use Get(DefaultValue) instead."));
 		}
@@ -1225,8 +1212,8 @@ namespace SlateAttributePrivate
 		{
 			if (TSharedPtr<const SWidget> Pin = Owner.Pin())
 			{
-				const_cast<AttributeMemberType>(Attribute)->UpdateNow(*Pin.Get());
-				return Attribute.Get();
+				const_cast<AttributeMemberType*>(Attribute)->UpdateNow(*Pin.Get());
+				return Attribute->Get();
 			}
 			return DefaultValue;
 		}
@@ -1236,7 +1223,7 @@ namespace SlateAttributePrivate
 		{
 			if (TSharedPtr<const SWidget> Pin = Owner.Pin())
 			{
-				return Attribute.ToAttribute(*Pin.Get());
+				return Attribute->ToAttribute(*Pin.Get());
 			}
 			return TAttribute<ObjectType>();
 		}
@@ -1246,7 +1233,7 @@ namespace SlateAttributePrivate
 		{
 			if (TSharedPtr<const SWidget> Pin = Owner.Pin())
 			{
-				return Attribute.IsBound(*Pin.Get());
+				return Attribute->IsBound(*Pin.Get());
 			}
 			return false;
 		}
@@ -1258,7 +1245,7 @@ namespace SlateAttributePrivate
 			TSharedPtr<const SWidget> OtherPin = Other.Owner.Pin();
 			if (SelfPin == OtherPin && SelfPin)
 			{
-				return Attribute.IsIdenticalTo(*SelfPin.Get(), Other.Attribute);
+				return Attribute->IsIdenticalTo(*SelfPin.Get(), *Other.Attribute);
 			}
 			return SelfPin == OtherPin;
 		}
@@ -1268,14 +1255,14 @@ namespace SlateAttributePrivate
 		{
 			if (TSharedPtr<const SWidget> Pin = Owner.Pin())
 			{
-				return Attribute.IsIdenticalTo(*Pin.Get(), Other);
+				return Attribute->IsIdenticalTo(*Pin.Get(), Other);
 			}
 			return !Other.IsSet(); // if the other is not set, then both are invalid.
 		}
 
 	private:
 		TWeakPtr<const SWidget> Owner;
-		const AttributeMemberType& Attribute;
+		AttributeMemberType const* Attribute = nullptr;
 	 };
 
 
