Added missing favorites icon

diff --git a/Engine/Source/Runtime/SlateCore/Private/Styling/StarshipCoreStyle.cpp b/Engine/Source/Runtime/SlateCore/Private/Styling/StarshipCoreStyle.cpp
--- a/Engine/Source/Runtime/SlateCore/Private/Styling/StarshipCoreStyle.cpp
+++ b/Engine/Source/Runtime/SlateCore/Private/Styling/StarshipCoreStyle.cpp
@@ -357,7 +357,7 @@ TSharedRef<ISlateStyle> FStarshipCoreStyle::Create()
 
 		Style->Set("Icons.Refresh", new IMAGE_BRUSH_SVG("Starship/Common/refresh", Icon16x16));
 
-
+		Style->Set("Icons.Star", new IMAGE_BRUSH_SVG("Starship/Common/Favorite", Icon16x16));
 	}
 
 	// Tool panels
