Slate: Add Deprecation message for deprecated attribute in SCompoundWidget.
#rnx

diff --git a/Engine/Source/Runtime/SlateCore/Public/Widgets/SCompoundWidget.h b/Engine/Source/Runtime/SlateCore/Public/Widgets/SCompoundWidget.h
--- a/Engine/Source/Runtime/SlateCore/Public/Widgets/SCompoundWidget.h
+++ b/Engine/Source/Runtime/SlateCore/Public/Widgets/SCompoundWidget.h
@@ -105,14 +105,14 @@ protected:
 	/** The slot that contains this widget's descendants.*/
 	FSimpleSlot ChildSlot;
 
-//#if WITH_EDITORONLY_DATA
-//	UE_DEPRECATED(5.0, "Direct access to ContentScale is now deprecated. Use the setter or getter.")
-//	FSlateDeprecatedTAttribute<FVector2D> ContentScale;
-//	UE_DEPRECATED(5.0, "Direct access to ColorAndOpacity is now deprecated. Use the setter or getter.")
-//	FSlateDeprecatedTAttribute<FLinearColor> ColorAndOpacity;
-//	UE_DEPRECATED(5.0, "Direct access to ForegroundColor is now deprecated. Use the setter or getter.")
-//	FSlateDeprecatedTAttribute<FSlateColor> ForegroundColor;
-//#endif
+#if WITH_EDITORONLY_DATA
+	UE_DEPRECATED(5.0, "Direct access to ContentScale is now deprecated. Use the setter or getter.")
+	FSlateDeprecatedTAttribute<FVector2D> ContentScale;
+	UE_DEPRECATED(5.0, "Direct access to ColorAndOpacity is now deprecated. Use the setter or getter.")
+	FSlateDeprecatedTAttribute<FLinearColor> ColorAndOpacity;
+	UE_DEPRECATED(5.0, "Direct access to ForegroundColor is now deprecated. Use the setter or getter.")
+	FSlateDeprecatedTAttribute<FSlateColor> ForegroundColor;
+#endif
 
 private:
 	/** The layout scale to apply to this widget's contents; useful for animation. */
