Slate: Fix bugs with GlobalInvalidation and SlateAttributes.
Rework the PreInvalidate algo to make it more clear the order of the different invalidations should be processed.
Add a flag to FWidgetAttributeIterator since MoveToWidgeIndexOnNextAdvance was using Invalid as the invalid state and the last element.
Rework the failed condition for the "Slate.InvalidationRoot.Verify*" to stop the process and only output once the error.
Only add the widget to the FinalUpdateList if an update is needed.

diff --git a/Engine/Source/Runtime/SlateCore/Private/FastUpdate/SlateInvalidationRoot.cpp b/Engine/Source/Runtime/SlateCore/Private/FastUpdate/SlateInvalidationRoot.cpp
--- a/Engine/Source/Runtime/SlateCore/Private/FastUpdate/SlateInvalidationRoot.cpp
+++ b/Engine/Source/Runtime/SlateCore/Private/FastUpdate/SlateInvalidationRoot.cpp
@@ -121,7 +121,7 @@ static FAutoConsoleVariableRef CVarSlateInvalidationRootVerifySlateAttributes(
 	TEXT("Each frame, verify that the widgets that have registered attribute are correctly updated once and the list contains all the widgets.")
 );
 void VerifySlateAttribute_Pre(FSlateInvalidationWidgetList& FastWidgetPathList);
-void VerifySlateAttribute_Post(FSlateInvalidationWidgetList const& FastWidgetPathList);
+void VerifySlateAttribute_Post(const FSlateInvalidationWidgetList& FastWidgetPathList);
 
 #endif //UE_SLATE_WITH_INVALIDATIONWIDGETLIST_DEBUGGING
 
@@ -144,13 +144,6 @@ FAutoConsoleVariableRef CVarSlateInvalidationWidgetListNumElementLeftBeforeSplit
  */
  namespace Slate
  {
-	 UE_NODISCARD EInvalidateWidgetReason EInvalidateWidgetReason_RemovePreUpdate(EInvalidateWidgetReason InvalidateReason)
-	{
-		static_assert(std::is_same<std::underlying_type_t<EInvalidateWidgetReason>, uint8>::value, "EInvalidateWidgetReason is not a uint8");
-		const uint8 AnyPostUpdate = (0xFF & ~(uint8)EInvalidateWidgetReason::AttributeRegistration);
-		return (EInvalidateWidgetReason)(((uint8)InvalidateReason) & AnyPostUpdate);
-	}
-
 	bool EInvalidateWidgetReason_HasPreUpdateFlag(EInvalidateWidgetReason InvalidateReason)
 	{
 		return EnumHasAnyFlags(InvalidateReason, EInvalidateWidgetReason::AttributeRegistration | EInvalidateWidgetReason::ChildOrder);
@@ -633,74 +626,80 @@ void FSlateInvalidationRoot::ProcessPreUpdate()
 			while((AttributeItt.IsValid() || WidgetsNeedingPreUpdate->Num() > 0) && !bNeedsSlowPath)
 			{
 				FSlateInvalidationWidgetSortOrder AttributeSortOrder = (AttributeItt.IsValid()) ? AttributeItt.GetCurrentSortOrder() : FSlateInvalidationWidgetSortOrder::LimitMax();
-				FSlateInvalidationWidgetSortOrder NeedUpdateSortOrder = (WidgetsNeedingPreUpdate->Num() > 0) ? WidgetsNeedingPreUpdate->HeapPeekElement().GetWidgetSortOrder() : FSlateInvalidationWidgetSortOrder::LimitMax();
+				FSlateInvalidationWidgetSortOrder NeedsUpdateSortOrder = (WidgetsNeedingPreUpdate->Num() > 0) ? WidgetsNeedingPreUpdate->HeapPeekElement().GetWidgetSortOrder() : FSlateInvalidationWidgetSortOrder::LimitMax();
 
-				if (NeedUpdateSortOrder < AttributeSortOrder)
+				if (AttributeSortOrder == FSlateInvalidationWidgetSortOrder::LimitMax() && NeedsUpdateSortOrder == FSlateInvalidationWidgetSortOrder::LimitMax())
 				{
-					const FSlateInvalidationWidgetIndex WidgetIndex = WidgetsNeedingPreUpdate->HeapPop();
+					checkf(false, TEXT("An element inside the lists has an invalid sort order. Something went wrong."));
+					WidgetsNeedingPreUpdate->Reset(true);
+					bNeedsSlowPath = true;
+					break;
+				}
+
+				// Process in order
+				//1.Invalidation AttributeRegistration of NeedsUpdate
+				//2.UpdateAttributes of AttributeSortOrder
+				//3.Invalidation ChildOrder of NeedsUpdateSortOrder
+
+				if (AttributeSortOrder == NeedsUpdateSortOrder)
+				{
+					// Process AttributeRegistration invalidation.
+					const FSlateInvalidationWidgetIndex WidgetIndex = WidgetsNeedingPreUpdate->HeapPeek();
 					FSlateInvalidationWidgetList::InvalidationWidgetType& InvalidationWidget = (*FastWidgetPathList)[WidgetIndex];
-					// It could have been destroyed since then
-					if (InvalidationWidget.GetWidget())
+					SWidget* WidgetPtr = InvalidationWidget.GetWidget();
+					if (ensureMsgf(WidgetPtr, TEXT("A parent remove this widget. Child order invalidation should have been called before we process this widget.")))
 					{
-#if WITH_SLATE_DEBUGGING
-						if (GSlateInvalidationRootDumpPreInvalidationList)
-						{
-							LogPreInvalidationItem(*FastWidgetPathList, WidgetIndex);
-						}
-#endif
-
-						bool bIsInvalidationWidgetValid = true;
-						if (EnumHasAnyFlags(InvalidationWidget.CurrentInvalidateReason, EInvalidateWidgetReason::ChildOrder))
-						{
-#if 0
-						FMemMark Mark(FMemStack::Get());
-						TArray<TTuple<FSlateInvalidationWidgetIndex, FSlateInvalidationWidgetSortOrder, TWeakPtr<SWidget>>, TMemStackAllocator<>> PreviousPreUpdate;
-						TArray<TTuple<FSlateInvalidationWidgetIndex, FSlateInvalidationWidgetSortOrder, TWeakPtr<SWidget>>, TMemStackAllocator<>> PreviousPostUpdate;
-						PreviousPreUpdate.Reserve(WidgetsNeedingPreUpdate->Num());
-						PreviousPostUpdate.Reserve(WidgetsNeedingPostUpdate->Num());
-						for (const auto& Element : WidgetsNeedingPreUpdate->GetRaw())
+						if (EnumHasAnyFlags(InvalidationWidget.CurrentInvalidateReason, EInvalidateWidgetReason::AttributeRegistration))
 						{
-							ensureMsgf(FastWidgetPathList->IsValidIndex(Element.GetWidgetIndex()), TEXT("F"));
-							PreviousPreUpdate.Emplace(Element.GetWidgetIndex(), Element.GetWidgetSortOrder(), (*FastWidgetPathList)[Element.GetWidgetIndex()].GetWidgetAsShared());
-						}
-						for (const auto& Element : WidgetsNeedingPostUpdate->GetRaw())
-						{
-							ensureMsgf(FastWidgetPathList->IsValidIndex(Element.GetWidgetIndex()), TEXT("gF"));
-							PreviousPostUpdate.Emplace(Element.GetWidgetIndex(), Element.GetWidgetSortOrder(), (*FastWidgetPathList)[Element.GetWidgetIndex()].GetWidgetAsShared());
-						}
+#if WITH_SLATE_DEBUGGING
+							if (GSlateInvalidationRootDumpPreInvalidationList)
+							{
+								LogPreInvalidationItem(*FastWidgetPathList, WidgetIndex);
+							}
 #endif
 
-							TGuardValue<bool> Tasdfdmp(bProcessingChildOrderInvalidation, true);
-							bIsInvalidationWidgetValid = FastWidgetPathList->ProcessChildOrderInvalidation(InvalidationWidget, ChildOriderInvalidationCallback);
-
-							// This widget may not be valid anymore (got removed because it doesn't fulfill the requirement anymore ie. NullWidget).
-						}
-
-						if (EnumHasAnyFlags(InvalidationWidget.CurrentInvalidateReason, EInvalidateWidgetReason::AttributeRegistration) && bIsInvalidationWidgetValid)
-						{
 							FastWidgetPathList->ProcessAttributeRegistrationInvalidation(InvalidationWidget);
-							AttributeItt.Seek(InvalidationWidget.Index);
-						}
+							if (WidgetPtr->HasRegisteredSlateAttribute())
+							{
+								// This should be the next Attribute to update
+								AttributeItt.Seek(InvalidationWidget.Index);
+							}
+							EnumRemoveFlags(InvalidationWidget.CurrentInvalidateReason, EInvalidateWidgetReason::AttributeRegistration);
 
-						AttributeItt.FixCurrentWidgetIndex();
+							// Do we still need to update this element, if not, then remove it from the update list.
+							if (!Slate::EInvalidateWidgetReason_HasPreUpdateFlag(InvalidationWidget.CurrentInvalidateReason))
+							{
+								WidgetsNeedingPreUpdate->HeapPopDiscard();
+							}
 
-						InvalidationWidget.CurrentInvalidateReason = Slate::EInvalidateWidgetReason_RemovePreUpdate(InvalidationWidget.CurrentInvalidateReason);
+							// Stop and process the next element.
+							continue;
+						}
+					}
+					else
+					{
+						WidgetsNeedingPreUpdate->HeapPopDiscard();
+						continue;
 					}
 				}
-				else
+
+				if (AttributeSortOrder <= NeedsUpdateSortOrder)
 				{
+					// Update Attributes
 					FSlateInvalidationWidgetList::InvalidationWidgetType& InvalidationWidget = (*FastWidgetPathList)[AttributeItt.GetCurrentIndex()];
 					SWidget* WidgetPtr = InvalidationWidget.GetWidget();
-					if (ensureMsgf(WidgetPtr, TEXT("Child order invalidation should have been called before we process this widget.")))
+					if (ensureMsgf(WidgetPtr, TEXT("A parent remove this widget. Child order invalidation should have been called before we process this widget.")))
 					{
 						auto UpdateAttribute = [&InvalidationWidget, WidgetPtr]()
 						{
 #if UE_SLATE_WITH_INVALIDATIONWIDGETLIST_DEBUGGING
 							if (GSlateInvalidationRootVerifySlateAttribute)
 							{
-								ensureMsgf(InvalidationWidget.bDebug_AttributeUpdated == false, TEXT("Attribute should only be updated once per frame."));
+								ensureAlwaysMsgf(InvalidationWidget.bDebug_AttributeUpdated == false, TEXT("Attribute should only be updated once per frame."));
 							}
 #endif
+
+							// It is possible that the widget doesn't have an attribute registered anymore.
 							FSlateAttributeMetaData::UpdateAttributes(*WidgetPtr);
 							InvalidationWidget.bDebug_AttributeUpdated = true;
 						};
@@ -724,6 +723,56 @@ void FSlateInvalidationRoot::ProcessPreUpdate()
 						}
 					}
 					AttributeItt.Advance();
+					continue; // Do the loop again, in case a ChildOrder invalidation occurred.
+				}
+
+				checkf(NeedsUpdateSortOrder < AttributeSortOrder, TEXT("Else case"));
+				{
+					// Process ChildOrder invalidation.
+
+					const FSlateInvalidationWidgetIndex WidgetIndex = WidgetsNeedingPreUpdate->HeapPop();
+					FSlateInvalidationWidgetList::InvalidationWidgetType& InvalidationWidget = (*FastWidgetPathList)[WidgetIndex];
+					// It could have been destroyed since then
+					if (InvalidationWidget.GetWidget())
+					{
+#if WITH_SLATE_DEBUGGING
+						if (GSlateInvalidationRootDumpPreInvalidationList)
+						{
+							LogPreInvalidationItem(*FastWidgetPathList, WidgetIndex);
+						}
+#endif
+
+						if (ensureAlways(EnumHasAnyFlags(InvalidationWidget.CurrentInvalidateReason, EInvalidateWidgetReason::ChildOrder)))
+						{
+// Uncomment to see to be able to compare the list before and after when debugging
+#if 0
+							FMemMark Mark(FMemStack::Get());
+							TArray<TTuple<FSlateInvalidationWidgetIndex, FSlateInvalidationWidgetSortOrder, TWeakPtr<SWidget>>, TMemStackAllocator<>> PreviousPreUpdate;
+							TArray<TTuple<FSlateInvalidationWidgetIndex, FSlateInvalidationWidgetSortOrder, TWeakPtr<SWidget>>, TMemStackAllocator<>> PreviousPostUpdate;
+							PreviousPreUpdate.Reserve(WidgetsNeedingPreUpdate->Num());
+							PreviousPostUpdate.Reserve(WidgetsNeedingPostUpdate->Num());
+							for (const auto& Element : WidgetsNeedingPreUpdate->GetRaw())
+							{
+								ensureAlwaysMsgf(FastWidgetPathList->IsValidIndex(Element.GetWidgetIndex()), TEXT("The element is invalid"));
+								PreviousPreUpdate.Emplace(Element.GetWidgetIndex(), Element.GetWidgetSortOrder(), (*FastWidgetPathList)[Element.GetWidgetIndex()].GetWidgetAsShared());
+							}
+							for (const auto& Element : WidgetsNeedingPostUpdate->GetRaw())
+							{
+								ensureAlwaysMsgf(FastWidgetPathList->IsValidIndex(Element.GetWidgetIndex()), TEXT("The element is invalid."));
+								PreviousPostUpdate.Emplace(Element.GetWidgetIndex(), Element.GetWidgetSortOrder(), (*FastWidgetPathList)[Element.GetWidgetIndex()].GetWidgetAsShared());
+							}
+#endif
+
+							TGuardValue<bool> ProcessChildOrderInvalidationGuardValue(bProcessingChildOrderInvalidation, true);
+							FastWidgetPathList->ProcessChildOrderInvalidation(InvalidationWidget, ChildOriderInvalidationCallback);
+
+							// This widget may not be valid anymore (got removed because it doesn't fulfill the requirement anymore ie. NullWidget).
+
+							AttributeItt.FixCurrentWidgetIndex();
+							// We need to keep it to run the layout calculation in FWidgetProxy::ProcessPostInvalidation
+							//EnumRemoveFlags(InvalidationWidget.CurrentInvalidateReason, EInvalidateWidgetReason::ChildOrder);
+						}
+					}
 				}
 			}
 		}
@@ -773,7 +822,6 @@ bool FSlateInvalidationRoot::ProcessPostUpdate()
 	while (WidgetsNeedingPostUpdate->Num() && !bNeedsSlowPath)
 	{
 		const FSlateInvalidationWidgetIndex WidgetIndex = WidgetsNeedingPostUpdate->HeapPop();
-		FinalUpdateList.Add(WidgetIndex);
 		FWidgetProxy& WidgetProxy = (*FastWidgetPathList)[WidgetIndex];
 
 		// Reset each widgets paint state
@@ -801,6 +849,11 @@ bool FSlateInvalidationRoot::ProcessPostUpdate()
 			}
 
 			bWidgetsNeedRepaint |= WidgetProxy.ProcessPostInvalidation(*WidgetsNeedingPostUpdate, *FastWidgetPathList, *this);
+
+			if (WidgetPtr->HasAnyUpdateFlags(EWidgetUpdateFlags::AnyUpdate))
+			{
+				FinalUpdateList.Add(WidgetIndex);
+			}
 		}
 	}
 	WidgetsNeedingPostUpdate->Reset(true);
@@ -958,7 +1011,11 @@ void LogPreInvalidationItem(const FSlateInvalidationWidgetList& FastWidgetPathLi
 {
 	const FWidgetProxy& Proxy = FastWidgetPathList[WidgetIndex];
 
-	if (EnumHasAnyFlags(Proxy.CurrentInvalidateReason, EInvalidateWidgetReason::ChildOrder))
+	if (EnumHasAnyFlags(Proxy.CurrentInvalidateReason, EInvalidateWidgetReason::AttributeRegistration))
+	{
+		UE_LOG(LogSlate, Log, TEXT("  AttributeRegistration %s"), *FReflectionMetaData::GetWidgetDebugInfo(Proxy.GetWidget()));
+	}
+	else if (EnumHasAnyFlags(Proxy.CurrentInvalidateReason, EInvalidateWidgetReason::ChildOrder))
 	{
 		UE_LOG(LogSlate, Log, TEXT("  Child Order %s"), *FReflectionMetaData::GetWidgetDebugInfo(Proxy.GetWidget()));
 	}
@@ -1000,6 +1057,16 @@ void LogPostInvalidationItem(const FSlateInvalidationWidgetList& FastWidgetPathL
 #endif
 
 #if UE_SLATE_WITH_INVALIDATIONWIDGETLIST_DEBUGGING
+
+#define UE_SLATE_LOG_ERROR_IF_FALSE(Test, FlagToReset, Message, ...) \
+	ensureAlwaysMsgf((Test), Message, ##__VA_ARGS__); \
+	if (!(Test)) \
+	{ \
+		UE_LOG(LogSlate, Error, Message, ##__VA_ARGS__); \
+		FlagToReset = false; \
+		return; \
+	}
+
 void VerifyWidgetList(TSharedRef<SWidget> RootWidget, FSlateInvalidationRootHandle InvalidationRootHandle, FSlateInvalidationWidgetList& WidgetList)
 {
 	FSlateInvalidationWidgetList List(InvalidationRootHandle, FSlateInvalidationWidgetList::FArguments{ 128, 128, 1000, false });
@@ -1012,7 +1079,7 @@ void VerifyWidgetList(TSharedRef<SWidget> RootWidget, FSlateInvalidationRootHand
 		UE_LOG(LogSlate, Log, TEXT("**-- Invaliation Root List --**"));
 		WidgetList.LogWidgetsList();
 
-		ensureMsgf(false, TEXT("The updated list doesn't match a newly created list."));
+		UE_SLATE_LOG_ERROR_IF_FALSE(false, GSlateInvalidationRootVerifyWidgetList, TEXT("The updated list doesn't match a newly created list."));
 	}
 }
 
@@ -1021,7 +1088,8 @@ void VerifyHittest(SWidget* InvalidationRootWidget, FSlateInvalidationWidgetList
 	check(InvalidationRootWidget);
 	check(HittestGrid);
 
-	ensureAlwaysMsgf(WidgetList.VerifySortOrder()
+	UE_SLATE_LOG_ERROR_IF_FALSE(WidgetList.VerifySortOrder()
+		, GSlateInvalidationRootVerifyHittestGrid
 		, TEXT("The array's sort order for InvalidationRoot '%d' is not respected.")
 		, *FReflectionMetaData::GetWidgetPath(InvalidationRootWidget));
 
@@ -1042,11 +1110,10 @@ void VerifyHittest(SWidget* InvalidationRootWidget, FSlateInvalidationWidgetList
 	for (const FHittestGrid::FWidgetSortData& Data : WeakHittestGridSortDatas)
 	{
 		TSharedPtr<SWidget> Widget = Data.WeakWidget.Pin();
-		if (ensureMsgf(Widget, TEXT("A widget is invalid in the HittestGrid")))
-		{
-			FHittestWidgetSortData SortData = { Widget.Get(), Data.PrimarySort, Data.SecondarySort };
-			HittestGridSortDatas.Add(MoveTemp(SortData));
-		}
+		UE_SLATE_LOG_ERROR_IF_FALSE(Widget, GSlateInvalidationRootVerifyHittestGrid, TEXT("A widget is invalid in the HittestGrid"));
+
+		FHittestWidgetSortData SortData = { Widget.Get(), Data.PrimarySort, Data.SecondarySort };
+		HittestGridSortDatas.Add(MoveTemp(SortData));
 	}
 
 	// The order in the WidgetList is sorted. It's not the case of the HittestGrid.
@@ -1068,7 +1135,8 @@ void VerifyHittest(SWidget* InvalidationRootWidget, FSlateInvalidationWidgetList
 					return;
 				}
 
-				ensureMsgf(Widget->GetProxyHandle().GetWidgetSortOrder() == HittestGridSortDatas[FoundHittestIndex].SecondarySort
+				UE_SLATE_LOG_ERROR_IF_FALSE(Widget->GetProxyHandle().GetWidgetSortOrder() == HittestGridSortDatas[FoundHittestIndex].SecondarySort
+					, GSlateInvalidationRootVerifyHittestGrid
 					, TEXT("The SecondarySort of widget '%s' doesn't match the SecondarySort inside the hittestgrid.")
 					, *FReflectionMetaData::GetWidgetPath(Widget));
 
@@ -1088,7 +1156,9 @@ void VerifyHittest(SWidget* InvalidationRootWidget, FSlateInvalidationWidgetList
 		HittestGridSortDatas.RemoveAtSwap(FoundHittestIndex);
 	}
 
-	ensureMsgf(HittestGridSortDatas.Num() == 0, TEXT("The hittest grid of Root '%s' has widgets that are not inside the InvalidationRoot's widget list")
+	UE_SLATE_LOG_ERROR_IF_FALSE(HittestGridSortDatas.Num() == 0
+		, GSlateInvalidationRootVerifyHittestGrid
+		, TEXT("The hittest grid of Root '%s' has widgets that are not inside the InvalidationRoot's widget list")
 		, *FReflectionMetaData::GetWidgetPath(InvalidationRootWidget));
 }
 
@@ -1112,14 +1182,18 @@ void VerifyWidgetVisibility(FSlateInvalidationWidgetList& WidgetList)
 						// It's possible that one of the parent is volatile
 						if (!Widget->IsVolatile() && !Widget->IsVolatileIndirectly())
 						{
-							ensureMsgf(false, TEXT("Widget '%s' should be %s.")
+							UE_SLATE_LOG_ERROR_IF_FALSE(false
+								, GSlateInvalidationRootVerifyWidgetVisibility
+								, TEXT("Widget '%s' should be %s.")
 								, *FReflectionMetaData::GetWidgetDebugInfo(Widget)
 								, (bSouldBeFastPathVisible ? TEXT("visible") : TEXT("hidden")));
 						}
 					}
 
 					const bool bHasValidCachedElementHandle = Widget->IsFastPathVisible() || !Widget->GetPersistentState().CachedElementHandle.HasCachedElements();
-					ensureMsgf(bHasValidCachedElementHandle, TEXT("Widget '%s' has cached element and is not visibled.")
+					UE_SLATE_LOG_ERROR_IF_FALSE(bHasValidCachedElementHandle
+						, GSlateInvalidationRootVerifyWidgetVisibility
+						, TEXT("Widget '%s' has cached element and is not visibled.")
 						, *FReflectionMetaData::GetWidgetDebugInfo(Widget));
 
 					// Cache last frame visibility
@@ -1141,30 +1215,38 @@ void VerifyWidgetVolatile(FSlateInvalidationWidgetList& WidgetList, TArray<FSlat
 					const bool bWasVolatile = Widget->IsVolatile();
 					Widget->CacheVolatility();
 					const bool bIsVolatile = Widget->IsVolatile();
-					ensureMsgf(bWasVolatile == bIsVolatile, TEXT("Widget '%s' volatily changed without an invalidation.")
+					UE_SLATE_LOG_ERROR_IF_FALSE(bWasVolatile == bIsVolatile
+						, GSlateInvalidationRootVerifyWidgetVolatile
+						, TEXT("Widget '%s' volatily changed without an invalidation.")
 						, *FReflectionMetaData::GetWidgetDebugInfo(Widget));
 				}
 
 				const TSharedPtr<const SWidget> ParentWidget = Widget->GetParentWidget();
-				if (ensure(ParentWidget))
+				UE_SLATE_LOG_ERROR_IF_FALSE(ParentWidget
+					, GSlateInvalidationRootVerifyWidgetVolatile
+					, TEXT("Parent widget of widget '%s' is invalid.")
+					, *FReflectionMetaData::GetWidgetDebugInfo(Widget));
+
 				{
 					const bool bShouldBeVolatileIndirectly = ParentWidget->IsVolatileIndirectly() || ParentWidget->IsVolatile();
-					ensureMsgf(Widget->IsVolatileIndirectly() == bShouldBeVolatileIndirectly, TEXT("Widget '%s' should be set as %s.")
+					UE_SLATE_LOG_ERROR_IF_FALSE(Widget->IsVolatileIndirectly() == bShouldBeVolatileIndirectly
+					, GSlateInvalidationRootVerifyWidgetVolatile
+						, TEXT("Widget '%s' should be set as %s.")
 						, *FReflectionMetaData::GetWidgetDebugInfo(Widget)
 						, (bShouldBeVolatileIndirectly ? TEXT("volatile indirectly") : TEXT("not volatile indirectly")));
 				}
 
+				if (Widget->IsVolatile() && !Widget->IsVolatileIndirectly())
 				{
-					if (Widget->IsVolatile() && !Widget->IsVolatileIndirectly())
-					{
-						//const bool bHasUpdateFlag = Widget->HasAnyUpdateFlags(EWidgetUpdateFlags::NeedsVolatilePaint);
-						//ensureMsgf(bHasUpdateFlag, TEXT("Widget '%s' is volatile but doesn't have the update flag NeedsVolatilePaint.")
-						//	, *FReflectionMetaData::GetWidgetDebugInfo(Widget));
-
-						const bool bIsContains = FinalUpdateList.Contains(Widget->GetProxyHandle().GetWidgetIndex());
-						ensureMsgf(bIsContains, TEXT("Widget '%s' is volatile but is not in the update list.")
-							, *FReflectionMetaData::GetWidgetDebugInfo(Widget));
-					}
+					//const bool bHasUpdateFlag = Widget->HasAnyUpdateFlags(EWidgetUpdateFlags::NeedsVolatilePaint);
+					//ensureMsgf(bHasUpdateFlag, TEXT("Widget '%s' is volatile but doesn't have the update flag NeedsVolatilePaint.")
+					//	, *FReflectionMetaData::GetWidgetDebugInfo(Widget));
+
+					const bool bIsContains = FinalUpdateList.Contains(Widget->GetProxyHandle().GetWidgetIndex());
+					UE_SLATE_LOG_ERROR_IF_FALSE(bIsContains
+						, GSlateInvalidationRootVerifyWidgetVolatile
+						, TEXT("Widget '%s' is volatile but is not in the update list.")
+						, *FReflectionMetaData::GetWidgetDebugInfo(Widget));
 				}
 			}
 		});
@@ -1183,9 +1265,12 @@ void VerifyWidgetsUpdateList_BeforeProcessPreUpdate(const TSharedRef<SWidget>& R
 
 	for (FSlateInvalidationWidgetIndex WidgetIndex : FinalUpdateList)
 	{
-		ensureMsgf(FastWidgetPathList->IsValidIndex(WidgetIndex), TEXT("An element is not valid."));
+		UE_SLATE_LOG_ERROR_IF_FALSE(FastWidgetPathList->IsValidIndex(WidgetIndex)
+			, GSlateInvalidationRootVerifyWidgetsUpdateList
+			, TEXT("An element is not valid."));
 		const FSlateInvalidationWidgetList::InvalidationWidgetType& InvalidationWidget = (*FastWidgetPathList)[WidgetIndex];
-		ensureMsgf(InvalidationWidget.CurrentInvalidateReason != EInvalidateWidgetReason::None || InvalidationWidget.UpdateFlags != EWidgetUpdateFlags::None
+		UE_SLATE_LOG_ERROR_IF_FALSE(InvalidationWidget.CurrentInvalidateReason != EInvalidateWidgetReason::None || InvalidationWidget.UpdateFlags != EWidgetUpdateFlags::None
+			, GSlateInvalidationRootVerifyWidgetsUpdateList
 			, TEXT("A widget is in the update list but doesn't have an update flag set."));
 	}
 
@@ -1193,14 +1278,23 @@ void VerifyWidgetsUpdateList_BeforeProcessPreUpdate(const TSharedRef<SWidget>& R
 
 	auto CheckIfValidInvalidation = [FastWidgetPathList](const FSlateInvalidationWidgetPreHeap::FElement& Element)
 	{
-		ensureMsgf(FastWidgetPathList->IsValidIndex(Element.GetWidgetIndex()), TEXT("An element is not valid."));
+		UE_SLATE_LOG_ERROR_IF_FALSE(FastWidgetPathList->IsValidIndex(Element.GetWidgetIndex())
+			, GSlateInvalidationRootVerifyWidgetsUpdateList
+			, TEXT("An element is not valid."));
 		const FSlateInvalidationWidgetList::InvalidationWidgetType& InvalidationWidget = (*FastWidgetPathList)[Element.GetWidgetIndex()];
-		ensureMsgf(InvalidationWidget.CurrentInvalidateReason != EInvalidateWidgetReason::None || InvalidationWidget.UpdateFlags != EWidgetUpdateFlags::None
+		UE_SLATE_LOG_ERROR_IF_FALSE(InvalidationWidget.CurrentInvalidateReason != EInvalidateWidgetReason::None || InvalidationWidget.UpdateFlags != EWidgetUpdateFlags::None
+			, GSlateInvalidationRootVerifyWidgetsUpdateList
 			, TEXT("A widget is in the update list but doesn't have an update flag set."));
 		if (SWidget* Widget = (*FastWidgetPathList)[Element.GetWidgetIndex()].GetWidget())
 		{
-			ensureMsgf(Widget->GetProxyHandle().GetWidgetSortOrder() == Element.GetWidgetSortOrder(), TEXT("The sort order of the widget do not matches what is in the heap."));
-			ensureMsgf(Widget->GetProxyHandle().GetWidgetIndex() == Element.GetWidgetIndex(), TEXT("The widget index of the widget do not matches what is in the heap."));
+			UE_SLATE_LOG_ERROR_IF_FALSE(Widget->GetProxyHandle().GetWidgetSortOrder() == Element.GetWidgetSortOrder()
+				, GSlateInvalidationRootVerifyWidgetsUpdateList
+				, TEXT("The sort order of the widget do not matches what is in the heap.")
+				, *FReflectionMetaData::GetWidgetDebugInfo(Widget));
+			UE_SLATE_LOG_ERROR_IF_FALSE(Widget->GetProxyHandle().GetWidgetIndex() == Element.GetWidgetIndex()
+				, GSlateInvalidationRootVerifyWidgetsUpdateList
+				, TEXT("The widget index of the widget do not matches what is in the heap.")
+				, *FReflectionMetaData::GetWidgetDebugInfo(Widget));
 		}
 	};
 	WidgetsNeedingPreUpdate->ForEachIndexes(CheckIfValidInvalidation);
@@ -1208,10 +1302,12 @@ void VerifyWidgetsUpdateList_BeforeProcessPreUpdate(const TSharedRef<SWidget>& R
 
 	FastWidgetPathList->ForEachInvalidationWidget([WidgetsNeedingPreUpdate, WidgetsNeedingPostUpdate](FSlateInvalidationWidgetList::InvalidationWidgetType& InvalidationWidget)
 		{
-			ensureMsgf(WidgetsNeedingPreUpdate->Contains_Debug(InvalidationWidget.Index) == InvalidationWidget.bContainedByWidgetPreHeap
-				, TEXT("A widget is or is not in the PreUpdate but the flag say otherwise."));
-			ensureMsgf(WidgetsNeedingPostUpdate->Contains_Debug(InvalidationWidget.Index) == InvalidationWidget.bContainedByWidgetPostHeap
-				, TEXT("A widget is or is not in the PostUpdate but the flag say otherwise."));
+			UE_SLATE_LOG_ERROR_IF_FALSE(WidgetsNeedingPreUpdate->Contains_Debug(InvalidationWidget.Index) == InvalidationWidget.bContainedByWidgetPreHeap
+				, GSlateInvalidationRootVerifyWidgetsUpdateList
+				, TEXT("Widget '%s' is or is not in the PreUpdate but the flag say otherwise."));
+			UE_SLATE_LOG_ERROR_IF_FALSE(WidgetsNeedingPostUpdate->Contains_Debug(InvalidationWidget.Index) == InvalidationWidget.bContainedByWidgetPostHeap
+				, GSlateInvalidationRootVerifyWidgetsUpdateList
+				, TEXT("Widget '%s' is or is not in the PostUpdate but the flag say otherwise."));
 		});
 }
 
@@ -1234,12 +1330,15 @@ void VerifySlateAttribute_Pre(FSlateInvalidationWidgetList& FastWidgetPathList)
 		});
 }
 
-void VerifySlateAttribute_Post(FSlateInvalidationWidgetList const& FastWidgetPathList)
+void VerifySlateAttribute_Post(const FSlateInvalidationWidgetList& FastWidgetPathList)
 {
-	FastWidgetPathList.VerifyElementIndexList();
+	const bool bElementIndexListValid = FastWidgetPathList.VerifyElementIndexList();
+	UE_SLATE_LOG_ERROR_IF_FALSE(bElementIndexListValid
+		, GSlateInvalidationRootVerifySlateAttribute
+		, TEXT("The VerifySlateAttribute failed in post."));
 }
 
-
+#undef UE_SLATE_LOG_ERROR_IF_FALSE
 #endif //UE_SLATE_WITH_INVALIDATIONWIDGETLIST_DEBUGGING
 
 /**
diff --git a/Engine/Source/Runtime/SlateCore/Private/FastUpdate/SlateInvalidationWidgetHeap.h b/Engine/Source/Runtime/SlateCore/Private/FastUpdate/SlateInvalidationWidgetHeap.h
--- a/Engine/Source/Runtime/SlateCore/Private/FastUpdate/SlateInvalidationWidgetHeap.h
+++ b/Engine/Source/Runtime/SlateCore/Private/FastUpdate/SlateInvalidationWidgetHeap.h
@@ -82,6 +82,12 @@ public:
 		OwnerList[Result].bContainedByWidgetPreHeap = false;
 		return Result;
 	}
+
+	/** Removes the biggest WidgetIndex from the list. */
+	void HeapPopDiscard()
+	{
+		Heap.HeapPopDiscard(SortPredicate(), false);
+	}
 	
 	/** Returns the biggest WidgetIndex from the list. */
 	inline FSlateInvalidationWidgetIndex HeapPeek() const
diff --git a/Engine/Source/Runtime/SlateCore/Private/FastUpdate/SlateInvalidationWidgetList.cpp b/Engine/Source/Runtime/SlateCore/Private/FastUpdate/SlateInvalidationWidgetList.cpp
--- a/Engine/Source/Runtime/SlateCore/Private/FastUpdate/SlateInvalidationWidgetList.cpp
+++ b/Engine/Source/Runtime/SlateCore/Private/FastUpdate/SlateInvalidationWidgetList.cpp
@@ -46,12 +46,13 @@ FSlateInvalidationWidgetIndex FSlateInvalidationWidgetList::IProcessChildOrderIn
 
 
 /** */
-FSlateInvalidationWidgetList::FWidgetAttributeIterator::FWidgetAttributeIterator(FSlateInvalidationWidgetList const& InWidgetList)
+FSlateInvalidationWidgetList::FWidgetAttributeIterator::FWidgetAttributeIterator(const FSlateInvalidationWidgetList& InWidgetList)
 	: WidgetList(InWidgetList)
 	, CurrentWidgetIndex(FSlateInvalidationWidgetIndex::Invalid)
 	, CurrentWidgetSortOrder(FSlateInvalidationWidgetSortOrder::LimitMax())
 	, AttributeIndex(0)
 	, MoveToWidgetIndexOnNextAdvance(FSlateInvalidationWidgetIndex::Invalid)
+	, bNeedsWidgetFixUp(false)
 {
 	int32 ArrayIndex = WidgetList.FirstArrayIndex;
 	while(ArrayIndex != INDEX_NONE)
@@ -70,14 +71,17 @@ FSlateInvalidationWidgetList::FWidgetAttributeIterator::FWidgetAttributeIterator
 
 void FSlateInvalidationWidgetList::FWidgetAttributeIterator::PreChildRemove(const FIndexRange& Range)
 {
-	if (CurrentWidgetSortOrder < Range.GetInclusiveMinWidgetSortOrder())
+	// The current widget will be removed or rebuild. Attributes will be updated in the Building loop.
+	//Move to the last item (that item, may be re-indexed soon)
+	if (CurrentWidgetSortOrder <= Range.GetInclusiveMaxWidgetSortOrder())
 	{
 		MoveToWidgetIndexOnNextAdvance = WidgetList.IncrementIndex(Range.GetInclusiveMaxWidgetIndex());
+		bNeedsWidgetFixUp = true;
 	}
 }
 
 
-void FSlateInvalidationWidgetList::FWidgetAttributeIterator::ReIndexed(IProcessChildOrderInvalidationCallback::FReIndexOperation const& Operation)
+void FSlateInvalidationWidgetList::FWidgetAttributeIterator::ReIndexed(const IProcessChildOrderInvalidationCallback::FReIndexOperation& Operation)
 {
 	if (MoveToWidgetIndexOnNextAdvance != FSlateInvalidationWidgetIndex::Invalid)
 	{
@@ -89,7 +93,7 @@ void FSlateInvalidationWidgetList::FWidgetAttributeIterator::ReIndexed(IProcessC
 	else if (CurrentWidgetIndex != FSlateInvalidationWidgetIndex::Invalid)
 	{
 		CurrentWidgetIndex = Operation.ReIndex(CurrentWidgetIndex);
-		CurrentWidgetSortOrder = FSlateInvalidationWidgetSortOrder{WidgetList, CurrentWidgetIndex };
+		CurrentWidgetSortOrder = FSlateInvalidationWidgetSortOrder{ WidgetList, CurrentWidgetIndex };
 	}
 }
 
@@ -103,10 +107,11 @@ void FSlateInvalidationWidgetList::FWidgetAttributeIterator::PostResort()
 }
 
 
-void FSlateInvalidationWidgetList::FWidgetAttributeIterator::ProxiesBuilt(FIndexRange const& Range)
+void FSlateInvalidationWidgetList::FWidgetAttributeIterator::ProxiesBuilt(const FIndexRange& Range)
 {
 	// The last built item is already updated. We want to update the following item next.
 	MoveToWidgetIndexOnNextAdvance = WidgetList.IncrementIndex(Range.GetInclusiveMaxWidgetIndex());
+	bNeedsWidgetFixUp = true;
 }
 
 
@@ -130,14 +135,35 @@ void FSlateInvalidationWidgetList::FWidgetAttributeIterator::FixCurrentWidgetInd
 			CurrentWidgetIndex = FSlateInvalidationWidgetIndex{ (IndexType)ArrayIndex, ArrayNode.ElementIndexList_WidgetWithRegisteredSlateAttribute[AttributeIndex] };
 			CurrentWidgetSortOrder = FSlateInvalidationWidgetSortOrder{ WidgetList, CurrentWidgetIndex };
 		}
-		MoveToWidgetIndexOnNextAdvance = FSlateInvalidationWidgetIndex::Invalid;
 	}
+	else if (bNeedsWidgetFixUp)
+	{
+		AttributeIndex = 0;
+		CurrentWidgetIndex = FSlateInvalidationWidgetIndex::Invalid;
+		CurrentWidgetSortOrder = FSlateInvalidationWidgetSortOrder::LimitMax();
+	}
+
+	MoveToWidgetIndexOnNextAdvance = FSlateInvalidationWidgetIndex::Invalid;
+	bNeedsWidgetFixUp = false;
 }
 
 
 void FSlateInvalidationWidgetList::FWidgetAttributeIterator::Seek(FSlateInvalidationWidgetIndex SeekTo)
 {
-	MoveToWidgetIndexOnNextAdvance = SeekTo;
+	check(SeekTo != FSlateInvalidationWidgetIndex::Invalid);
+
+	const FArrayNode& ArrayNode = WidgetList.Data[SeekTo.ArrayIndex];
+	AttributeIndex = ArrayNode.ElementIndexList_WidgetWithRegisteredSlateAttribute.FindLowerBound(SeekTo.ElementIndex);
+	if (AttributeIndex == INDEX_NONE)
+	{
+		int32 NextArrayIndex = ArrayNode.NextArrayIndex;
+		AdvanceArrayIndex(NextArrayIndex);
+	}
+	else
+	{
+		CurrentWidgetIndex = SeekTo;
+		CurrentWidgetSortOrder = FSlateInvalidationWidgetSortOrder{ WidgetList, CurrentWidgetIndex };
+	}
 }
 
 
@@ -194,6 +220,23 @@ void FSlateInvalidationWidgetList::FArrayNode::RemoveElementIndexBiggerOrEqualTh
 	}
 }
 
+void FSlateInvalidationWidgetList::FArrayNode::RemoveElementIndexBetweenOrEqualThan(IndexType StartElementIndex, IndexType EndElementIndex)
+{
+	const int32 StartFoundIndex = ElementIndexList_WidgetWithRegisteredSlateAttribute.FindLowerBound(StartElementIndex);
+	if (ElementIndexList_WidgetWithRegisteredSlateAttribute.IsValidIndex(StartFoundIndex))
+	{
+		const int32 EndFoundIndex = ElementIndexList_WidgetWithRegisteredSlateAttribute.FindUpperBound(EndElementIndex);
+		if (ElementIndexList_WidgetWithRegisteredSlateAttribute.IsValidIndex(EndFoundIndex))
+		{
+			ElementIndexList_WidgetWithRegisteredSlateAttribute.RemoveAt(StartFoundIndex, EndFoundIndex - StartFoundIndex);
+		}
+		else
+		{
+			ElementIndexList_WidgetWithRegisteredSlateAttribute.RemoveAt(StartFoundIndex, ElementIndexList_WidgetWithRegisteredSlateAttribute.Num() - StartFoundIndex);
+		}
+	}
+}
+
 
 /** */
 FSlateInvalidationWidgetList::FSlateInvalidationWidgetList(FSlateInvalidationRootHandle InOwner, const FArguments& Args)
@@ -370,7 +413,7 @@ void FSlateInvalidationWidgetList::_RebuildWidgetListTree(TSharedRef<SWidget> Wi
 }
 
 
-bool FSlateInvalidationWidgetList::ProcessChildOrderInvalidation(InvalidationWidgetType const& InvalidationWidget, IProcessChildOrderInvalidationCallback& Callback)
+bool FSlateInvalidationWidgetList::ProcessChildOrderInvalidation(const InvalidationWidgetType& InvalidationWidget, IProcessChildOrderInvalidationCallback& Callback)
 {
 	SCOPE_CYCLE_COUNTER(STAT_WidgetList_ProcessChildOrderInvalidation);
 
@@ -511,12 +554,15 @@ bool FSlateInvalidationWidgetList::ProcessChildOrderInvalidation(InvalidationWid
 			if (InvalidatedChildren->Num() > 0)
 			{
 				SCOPED_NAMED_EVENT(Slate_InvalidationList_ProcessRebuild, FColorList::Blue);
+				FSlateInvalidationWidgetIndex PreviousWidgetIndex = WidgetPtr->GetProxyHandle().GetWidgetIndex();
 				FSlateInvalidationWidgetIndex PreviousLeafMostChildIndex = InvalidationWidget.LeafMostChildIndex;
 				_RebuildWidgetListTree(WidgetPtr->AsShared(), StartChildIndex);
 
-				const FIndexRange BuiltRange = (InvalidationWidget.LeafMostChildIndex == PreviousLeafMostChildIndex)
-					? FIndexRange{ *this, PreviousLeafMostChildIndex, InvalidationWidget.LeafMostChildIndex }
-					: FIndexRange{ *this, IncrementIndex(PreviousLeafMostChildIndex), InvalidationWidget.LeafMostChildIndex };
+
+				const InvalidationWidgetType& NewInvalidationWidget = (*this)[PreviousWidgetIndex];
+				const FIndexRange BuiltRange = (NewInvalidationWidget.LeafMostChildIndex == PreviousLeafMostChildIndex)
+					? FIndexRange{ *this, PreviousLeafMostChildIndex, NewInvalidationWidget.LeafMostChildIndex }
+					: FIndexRange{ *this, IncrementIndex(PreviousLeafMostChildIndex), NewInvalidationWidget.LeafMostChildIndex };
 				if (BuiltRange.IsValid())
 				{
 					Callback.ProxiesBuilt(BuiltRange);
@@ -529,18 +575,12 @@ bool FSlateInvalidationWidgetList::ProcessChildOrderInvalidation(InvalidationWid
 }
 
 
-void FSlateInvalidationWidgetList::ProcessAttributeRegistrationInvalidation(InvalidationWidgetType const& InvalidationWidget)
+void FSlateInvalidationWidgetList::ProcessAttributeRegistrationInvalidation(const InvalidationWidgetType& InvalidationWidget)
 {
 	SWidget* WidgetPtr = InvalidationWidget.GetWidget();
 	check(WidgetPtr);
 
-	bool bNeedsToBeThere = false;
-	if (FSlateAttributeMetaData* MetaData = FSlateAttributeMetaData::FindMetaData(*WidgetPtr))
-	{
-		bNeedsToBeThere = (MetaData->RegisteredNum() > 0);
-	}
-
-	if (bNeedsToBeThere)
+	if (WidgetPtr->HasRegisteredSlateAttribute())
 	{
 		Data[InvalidationWidget.Index.ArrayIndex].ElementIndexList_WidgetWithRegisteredSlateAttribute.InsertUnique(InvalidationWidget.Index.ElementIndex);
 	}
@@ -1017,13 +1057,16 @@ void FSlateInvalidationWidgetList::_RemoveRangeFromSameParent(const FIndexRange
 		};
 		auto ResetInvalidationWidget = [Self](IndexType ArrayIndex, IndexType StartIndex, int32 Num)
 		{
-			FArrayNode& ArrayNode = Self->Data[ArrayIndex];
-			ElementListType& ResetElementList = ArrayNode.ElementList;
-			for (int32 ElementIndex = StartIndex; ElementIndex < Num; ++ElementIndex)
+			if (Num > 0)
 			{
-				ResetElementList[ElementIndex].ResetWidget();
+				FArrayNode& ArrayNode = Self->Data[ArrayIndex];
+				ElementListType& ResetElementList = ArrayNode.ElementList;
+				for (int32 ElementIndex = StartIndex; ElementIndex < Num; ++ElementIndex)
+				{
+					ResetElementList[ElementIndex].ResetWidget();
+				}
+				ArrayNode.RemoveElementIndexBetweenOrEqualThan(StartIndex, (IndexType)(StartIndex + Num - 1));
 			}
-			ArrayNode.RemoveElementIndexBiggerOrEqualThan(StartIndex);
 		};
 
 		auto RemoveDataNodeIfNeeded = [Self](IndexType ArrayIndex) -> bool
@@ -1527,51 +1570,76 @@ bool FSlateInvalidationWidgetList::VerifySortOrder() const
 bool FSlateInvalidationWidgetList::VerifyElementIndexList() const
 {
 	bool bResult = true;
-	for (const FArrayNode& ArrayNode : Data)
+
+	// Test that the iterator return the same result as a new calculated list.
+	if (bResult)
 	{
-		for (IndexType ElementIndex : ArrayNode.ElementIndexList_WidgetWithRegisteredSlateAttribute)
+		TArray<const SWidget*> WidgetListWithIterator;
+		FSlateInvalidationWidgetList::FWidgetAttributeIterator AttributeItt = CreateWidgetAttributeIterator();
+		while (AttributeItt.IsValid())
 		{
-			if (!ArrayNode.ElementList.IsValidIndex(ElementIndex))
+			const SWidget* WidgetPtr = (*this)[AttributeItt.GetCurrentIndex()].GetWidget();
+			if (WidgetPtr)
 			{
-				UE_LOG(LogSlate, Warning, TEXT("ElementIndex '%d' in the array of sort value '%d' is invalid.")
-					, ElementIndex, ArrayNode.SortOrder);
-				bResult = false;
-				break;
+				WidgetListWithIterator.Add(WidgetPtr);
 			}
+			AttributeItt.Advance();
+		}
+
 
-			FSlateInvalidationWidgetList::InvalidationWidgetType const& InvalidationWidget = ArrayNode.ElementList[ElementIndex];
-			if (const SWidget* Widget = InvalidationWidget.GetWidget())
+		TArray<const SWidget*> WidgetListWithForEachWidget;
+		WidgetListWithForEachWidget.Reset(WidgetListWithIterator.Num());
+		int32 Index = 0;
+		const_cast<FSlateInvalidationWidgetList*>(this)->ForEachWidget([&](const SWidget* Widget)
 			{
-				if (!Widget->HasRegisteredSlateAttribute())
+				if (Widget->HasRegisteredSlateAttribute())
 				{
-					UE_LOG(LogSlate, Warning, TEXT("ElementIndex '%d' in the array of sort value '%d' should not be there since widget '%s' doesn't have registered SlateAttribute.")
-						, ElementIndex, ArrayNode.SortOrder
-						, *FReflectionMetaData::GetWidgetDebugInfo(Widget));
-					bResult = false;
+					WidgetListWithForEachWidget.Add(Widget);
 				}
+			});
 
-				if (Widget->IsFastPathVisible())
-				{
-					ensureMsgf(InvalidationWidget.bDebug_AttributeUpdated, TEXT("The widget '%s' was not updated."), *FReflectionMetaData::GetWidgetDebugInfo(Widget));
-				}
-			}
+		bResult = WidgetListWithIterator == WidgetListWithForEachWidget;
+		if (!bResult)
+		{
+			UE_LOG(LogSlate, Warning, TEXT("The 2 lists are not identical. With Attribute: %d. With ForEach: %d"), WidgetListWithIterator.Num(), WidgetListWithForEachWidget.Num());
 		}
+	}
 
-		for (int32 ElementIndex = ArrayNode.StartIndex; ElementIndex < ArrayNode.ElementList.Num(); ++ElementIndex)
+	// Test if bDebug_AttributeUpdated was updated.
+	if (bResult)
+	{
+		for (const FArrayNode& ArrayNode : Data)
 		{
-			if (!ArrayNode.ElementIndexList_WidgetWithRegisteredSlateAttribute.Contains((IndexType)ElementIndex))
+			for (IndexType ElementIndex : ArrayNode.ElementIndexList_WidgetWithRegisteredSlateAttribute)
 			{
-				SWidget* Widget = ArrayNode.ElementList[ElementIndex].GetWidget();
-				if (Widget->HasRegisteredSlateAttribute())
+				if (!ArrayNode.ElementList.IsValidIndex(ElementIndex))
 				{
-					UE_LOG(LogSlate, Warning, TEXT("Widget '%s' (element: '%d') in the array of sort value '%d' should be in ElementIndexList_WidgetWithRegisteredSlateAttribute.")
-						, *FReflectionMetaData::GetWidgetDebugInfo(Widget)
+					UE_LOG(LogSlate, Warning, TEXT("ElementIndex '%d' in the array of sort value '%d' is invalid.")
 						, ElementIndex, ArrayNode.SortOrder);
 					bResult = false;
+					break;
+				}
+
+				const FSlateInvalidationWidgetList::InvalidationWidgetType& InvalidationWidget = ArrayNode.ElementList[ElementIndex];
+				if (const SWidget* Widget = InvalidationWidget.GetWidget())
+				{
+					if (!Widget->HasRegisteredSlateAttribute())
+					{
+						UE_LOG(LogSlate, Warning, TEXT("ElementIndex '%d' in the array of sort value '%d' should not be there since widget '%s' doesn't have registered SlateAttribute.")
+							, ElementIndex, ArrayNode.SortOrder
+							, *FReflectionMetaData::GetWidgetDebugInfo(Widget));
+						bResult = false;
+					}
+
+					if (Widget->IsFastPathVisible() && !InvalidationWidget.bDebug_AttributeUpdated)
+					{
+						UE_LOG(LogSlate, Warning, TEXT("The widget '%s' was not updated."), *FReflectionMetaData::GetWidgetDebugInfo(Widget));
+					}
 				}
 			}
 		}
 	}
+
 	return bResult;
 }
 #endif //UE_SLATE_WITH_INVALIDATIONWIDGETLIST_DEBUGGING
diff --git a/Engine/Source/Runtime/SlateCore/Private/FastUpdate/SlateInvalidationWidgetList.h b/Engine/Source/Runtime/SlateCore/Private/FastUpdate/SlateInvalidationWidgetList.h
--- a/Engine/Source/Runtime/SlateCore/Private/FastUpdate/SlateInvalidationWidgetList.h
+++ b/Engine/Source/Runtime/SlateCore/Private/FastUpdate/SlateInvalidationWidgetList.h
@@ -104,8 +104,8 @@ public:
 	{
 		struct FReIndexOperation
 		{
-			FReIndexOperation(FIndexRange const& InRange, FSlateInvalidationWidgetIndex InReIndexTarget) : Range(InRange), ReIndexTarget(InReIndexTarget) {}
-			FIndexRange const& GetRange() const { return Range; }
+			FReIndexOperation(const FIndexRange& InRange, FSlateInvalidationWidgetIndex InReIndexTarget) : Range(InRange), ReIndexTarget(InReIndexTarget) {}
+			const FIndexRange& GetRange() const { return Range; }
 			UE_NODISCARD FSlateInvalidationWidgetIndex ReIndex(FSlateInvalidationWidgetIndex Index) const;
 		private:
 			const FIndexRange& Range;
@@ -113,36 +113,36 @@ public:
 		};
 		struct FReSortOperation
 		{
-			FReSortOperation(FIndexRange const& InRange) : Range(InRange) {}
-			FIndexRange const& GetRange() const { return Range; }
+			FReSortOperation(const FIndexRange& InRange) : Range(InRange) {}
+			const FIndexRange& GetRange() const { return Range; }
 		private:
-			FIndexRange const& Range;
+			const FIndexRange& Range;
 		};
 
 		/** Widget proxies that will be removed and will not be valid anymore. */
-		virtual void PreChildRemove(FIndexRange const& Range) {}
+		virtual void PreChildRemove(const FIndexRange& Range) {}
 		/** Widget proxies that got moved/re-indexed by the operation. */
-		virtual void ProxiesReIndexed(FReIndexOperation const& Operation) {}
+		virtual void ProxiesReIndexed(const FReIndexOperation& Operation) {}
 		/** Widget proxies that got resorted by the operation. */
-		virtual void ProxiesPreResort(FReSortOperation const& Operation) {}
+		virtual void ProxiesPreResort(const FReSortOperation& Operation) {}
 		/** Widget proxies that got resorted by the operation. */
 		virtual void ProxiesPostResort() {}
 		/** Widget proxies built by the operation. */
-		virtual void ProxiesBuilt(FIndexRange const& Range) {}
+		virtual void ProxiesBuilt(const FIndexRange& Range) {}
 	};
 
 	/**
 	 * Process widget that have a ChildOrder invalidation.
 	 * @returns true if the InvalidationWidget is still valid.
 	 */
-	bool ProcessChildOrderInvalidation(InvalidationWidgetType const& InvalidationWidget, IProcessChildOrderInvalidationCallback& Callback);
+	bool ProcessChildOrderInvalidation(const InvalidationWidgetType& InvalidationWidget, IProcessChildOrderInvalidationCallback& Callback);
 
 
 	/**
 	 * Test, then adds or removes from the registered attribute list.
 	 * @returns true if the Widget go added to the registered attribute list.
 	 */
-	void ProcessAttributeRegistrationInvalidation(InvalidationWidgetType const& InvalidationWidget);
+	void ProcessAttributeRegistrationInvalidation(const InvalidationWidgetType& InvalidationWidget);
 
 	/** Performs an operation on all SWidget in the list. */
 	template<typename Predicate>
@@ -188,21 +188,23 @@ public:
 	struct FWidgetAttributeIterator
 	{
 	private:
-		FSlateInvalidationWidgetList const& WidgetList;
+		const FSlateInvalidationWidgetList& WidgetList;
 		FSlateInvalidationWidgetIndex CurrentWidgetIndex;
 		FSlateInvalidationWidgetSortOrder CurrentWidgetSortOrder;
 		int32 AttributeIndex;
 
 		FSlateInvalidationWidgetIndex MoveToWidgetIndexOnNextAdvance;
+		/** A fix-up is required. MoveToWidgetIndexOnNextAdvance cannot be used since it may point to the last element (invalid) */
+		bool bNeedsWidgetFixUp;
 
 	public:
-		FWidgetAttributeIterator(FSlateInvalidationWidgetList const& InWidgetList);
+		FWidgetAttributeIterator(const FSlateInvalidationWidgetList& InWidgetList);
 
 		//~ Handle operation
-		void PreChildRemove(FIndexRange const& Range);
-		void ReIndexed(IProcessChildOrderInvalidationCallback::FReIndexOperation const& Operation);
+		void PreChildRemove(const FIndexRange& Range);
+		void ReIndexed(const IProcessChildOrderInvalidationCallback::FReIndexOperation& Operation);
 		void PostResort();
-		void ProxiesBuilt(FIndexRange const& Range);
+		void ProxiesBuilt(const FIndexRange& Range);
 		void FixCurrentWidgetIndex();
 		void Seek(FSlateInvalidationWidgetIndex SeekTo);
 
@@ -222,7 +224,7 @@ public:
 		void AdvanceArrayIndex(int32 ArrayIndex);
 	};
 
-	FWidgetAttributeIterator CreateWidgetAttributeIterator()
+	FWidgetAttributeIterator CreateWidgetAttributeIterator() const
 	{
 		return FWidgetAttributeIterator(*this);
 	}
@@ -392,6 +394,7 @@ private:
 		WidgetListType ElementIndexList_WidgetWithRegisteredSlateAttribute;
 
 		void RemoveElementIndexBiggerOrEqualThan(IndexType ElementIndex);
+		void RemoveElementIndexBetweenOrEqualThan(IndexType StartElementIndex, IndexType EndElementIndex);
 	};
 	using ArrayNodeType = TSparseArray<FArrayNode>;
 
diff --git a/Engine/Source/Runtime/SlateCore/Private/Types/SlateAttributeMetaData.cpp b/Engine/Source/Runtime/SlateCore/Private/Types/SlateAttributeMetaData.cpp
--- a/Engine/Source/Runtime/SlateCore/Private/Types/SlateAttributeMetaData.cpp
+++ b/Engine/Source/Runtime/SlateCore/Private/Types/SlateAttributeMetaData.cpp
@@ -53,7 +53,10 @@ void FSlateAttributeMetaData::RegisterAttribute(SWidget& OwningWidget, FSlateAtt
 		NewAttributeMetaData->RegisterAttributeImpl(OwningWidget, Attribute, AttributeType, MoveTemp(Wrapper));
 		OwningWidget.bHasRegisteredSlateAttribute = true;
 		OwningWidget.MetaData.Insert(NewAttributeMetaData, 0);
-		OwningWidget.Invalidate(EInvalidateWidgetReason::AttributeRegistration);
+		if (!OwningWidget.bPauseAttributeInvalidation)
+		{
+			OwningWidget.Invalidate(EInvalidateWidgetReason::AttributeRegistration);
+		}
 	}
 }
 
@@ -133,7 +136,10 @@ bool FSlateAttributeMetaData::UnregisterAttribute(SWidget& OwningWidget, const F
 			check(bResult); // if the num is 0 then we should have remove an item.
 			OwningWidget.bHasRegisteredSlateAttribute = false;
 			OwningWidget.MetaData.RemoveAtSwap(0);
-			OwningWidget.Invalidate(EInvalidateWidgetReason::AttributeRegistration);
+			if (!OwningWidget.bPauseAttributeInvalidation)
+			{
+				OwningWidget.Invalidate(EInvalidateWidgetReason::AttributeRegistration);
+			}
 		}
 		return bResult;
 	}
diff --git a/Engine/Source/Runtime/SlateCore/Public/FastUpdate/WidgetProxy.h b/Engine/Source/Runtime/SlateCore/Public/FastUpdate/WidgetProxy.h
--- a/Engine/Source/Runtime/SlateCore/Public/FastUpdate/WidgetProxy.h
+++ b/Engine/Source/Runtime/SlateCore/Public/FastUpdate/WidgetProxy.h
@@ -94,7 +94,7 @@ public:
 	/** Use with "Slate.InvalidationRoot.VerifyWidgetVisibility". Cached the last FastPathVisible value to find widgets that do not call Invalidate properly. */
 	uint8 bDebug_LastFrameVisible : 1;
 	uint8 bDebug_LastFrameVisibleSet : 1;
-	/** Use with "SLate.InvalidationRoot.VerifyWidgetAttribute". */
+	/** Use with "Slate.InvalidationRoot.VerifyWidgetAttribute". */
 	uint8 bDebug_AttributeUpdated : 1;
 };
 
