Slate: Add SWidget to the SlateAttribute's compare predicate.

diff --git a/Engine/Source/Runtime/SlateCore/Private/Test/SlateAttributeTest.cpp b/Engine/Source/Runtime/SlateCore/Private/Test/SlateAttributeTest.cpp
--- a/Engine/Source/Runtime/SlateCore/Private/Test/SlateAttributeTest.cpp
+++ b/Engine/Source/Runtime/SlateCore/Private/Test/SlateAttributeTest.cpp
@@ -32,7 +32,7 @@ struct FConstructionCounter
 	FConstructionCounter& operator=(FConstructionCounter&& Other) { Value = Other.Value; ++MoveOperatorCounter;  return *this; }
 
 	int32 Value;
-	bool operator== (const FConstructionCounter& Other) { return Other.Value == Value; }
+	bool operator== (const FConstructionCounter& Other) const { return Other.Value == Value; }
 
 	static int32 DefaultConstructionCounter;
 	static int32 CopyConstructionCounter;
@@ -451,10 +451,10 @@ bool FSlateAttributeTest::RunTest(const FString& Parameters)
 			}
 			{
 				bool bIsBound1 = Widget->AttributeA.IsBound(*Widget);
-				bool bIsIdentical1 = Widget->AttributeA.IsIdenticalTo(*Widget, Widget->AttributeA);
+				bool bIsIdentical1 = Widget->AttributeA.IdenticalTo(*Widget, Widget->AttributeA);
 				auto Getter1 = TAttribute<int32>::FGetter::CreateStatic(UE::Slate::Private::CallbackForIntAttribute, 1);
 				TAttribute<int32> Attribute1 = TAttribute<int32>::Create(Getter1);
-				bool bIsIdentical2 = Widget->AttributeA.IsIdenticalTo(*Widget, Attribute1);
+				bool bIsIdentical2 = Widget->AttributeA.IdenticalTo(*Widget, Attribute1);
 			}
 		}
 		{
@@ -597,7 +597,7 @@ bool FSlateAttributeTest::RunTest(const FString& Parameters)
 				Attribute1.Set({ 2 });
 				FLocalConstructionCounter::ResetCounter();
 				Attribute.Assign(MoveTemp(Attribute1));
-				AddErrorIfCounterDoNotMatches(0, 0, 0, 0, 1, TEXT("Assign Move failed."));
+				AddErrorIfCounterDoNotMatches(0, 0, 2, 0, 1, TEXT("Assign Move failed."));
 			}
 			// Test unbinded Attribute
 			{
@@ -688,7 +688,7 @@ bool FSlateAttributeTest::RunTest(const FString& Parameters)
 				TAttribute<FLocalConstructionCounter> Attribute1 = Counter;
 				FLocalConstructionCounter::ResetCounter();
 				Attribute.Assign(MoveTemp(Attribute1), Counter);
-				AddErrorIfCounterDoNotMatches(0, 0, 0, 0, 1, TEXT("Assign Set Move failed."));
+				AddErrorIfCounterDoNotMatches(0, 0, 2, 0, 1, TEXT("Assign Set Move failed."));
 			}
 			{
 				SAttributeAttribute::ManagedSlateAttributeType Attribute = SAttributeAttribute::ManagedSlateAttributeType(Widget.ToSharedRef());
@@ -696,7 +696,7 @@ bool FSlateAttributeTest::RunTest(const FString& Parameters)
 				TAttribute<FLocalConstructionCounter> Attribute1 = Counter;
 				FLocalConstructionCounter::ResetCounter();
 				Attribute.Assign(MoveTemp(Attribute1), MoveTemp(Counter));
-				AddErrorIfCounterDoNotMatches(0, 0, 0, 0, 1, TEXT("Assign Set Move failed."));
+				AddErrorIfCounterDoNotMatches(0, 0, 2, 0, 1, TEXT("Assign Set Move failed."));
 			}
 		}
 	}
diff --git a/Engine/Source/Runtime/SlateCore/Public/Layout/Children.h b/Engine/Source/Runtime/SlateCore/Public/Layout/Children.h
--- a/Engine/Source/Runtime/SlateCore/Public/Layout/Children.h
+++ b/Engine/Source/Runtime/SlateCore/Public/Layout/Children.h
@@ -835,7 +835,7 @@ public:
 	const FMargin& GetSlotPadding() const { return SlotPadding.Get(); }
 
 public:
-	using SlotPaddingAttributeType = SlateAttributePrivate::TSlateMemberAttribute<FMargin, TSlateAttributeInvalidationReason<EInvalidateWidgetReason::Layout>>;
+	using SlotPaddingAttributeType = SlateAttributePrivate::TSlateMemberAttribute<FMargin, TSlateAttributeInvalidationReason<EInvalidateWidgetReason::Layout>, TSlateAttributeComparePredicate<>>;
 	using SlotPaddingAttributeRefType = SlateAttributePrivate::TSlateMemberAttributeRef<SlotPaddingAttributeType>;
 
 	template<typename WidgetType, typename V = typename std::enable_if<std::is_base_of<SWidget, WidgetType>::value>::type>
diff --git a/Engine/Source/Runtime/SlateCore/Public/Types/SlateAttribute.h b/Engine/Source/Runtime/SlateCore/Public/Types/SlateAttribute.h
--- a/Engine/Source/Runtime/SlateCore/Public/Types/SlateAttribute.h
+++ b/Engine/Source/Runtime/SlateCore/Public/Types/SlateAttribute.h
@@ -161,10 +161,23 @@
 class SWidget;
 
 
+/** Default predicate to compare of Object for SlateAttribute. */
+template<typename ComparePredicate = TEqualTo<>>
+struct TSlateAttributeComparePredicate
+{
+	template<typename ObjectType>
+	static bool IdenticalTo(const SWidget&, ObjectType&& Lhs, ObjectType&& Rhs)
+	{
+		// If you have an error here, do you have a operator== const?
+		return ComparePredicate{}(Forward<ObjectType>(Lhs), Forward<ObjectType>(Rhs));
+	}
+};
+
+
 /** Default predicate to compare FText. */
 struct TSlateAttributeFTextComparePredicate
 {
-	bool operator()(const FText& Lhs, const FText& Rhs) const
+	static bool IdenticalTo(const SWidget&, const FText& Lhs, const FText& Rhs)
 	{
 		return Lhs.IdenticalTo(Rhs, ETextIdenticalModeFlags::DeepCompare | ETextIdenticalModeFlags::LexicalCompareInvariants);
 	}
@@ -312,6 +325,7 @@ namespace SlateAttributePrivate
 		using FComparePredicate = InComparePredicateType;
 
 		static EInvalidateWidgetReason GetInvalidationReason(const SWidget& Widget) { return FInvalidationReasonPredicate::GetInvalidationReason(Widget); }
+		static bool IdenticalTo(const SWidget& Widget, const ObjectType& Lhs, const ObjectType& Rhs) { return FComparePredicate::IdenticalTo(Widget, Lhs, Rhs); }
 		
 	private:
 		void UpdateNowOnBind(SWidget& Widget)
@@ -452,7 +466,7 @@ namespace SlateAttributePrivate
 			VerifyOwningWidget(Widget);
 			ProtectedUnregisterAttribute(Widget, InAttributeType);
 
-			const bool bIsIdentical = FComparePredicate{}.operator()(Value, NewValue);
+			const bool bIsIdentical = IdenticalTo(Widget, Value, NewValue);
 			if (!bIsIdentical)
 			{
 				Value = NewValue;
@@ -470,7 +484,7 @@ namespace SlateAttributePrivate
 			VerifyOwningWidget(Widget);
 			ProtectedUnregisterAttribute(Widget, InAttributeType);
 
-			const bool bIsIdentical = FComparePredicate{}.operator()(Value, NewValue);
+			const bool bIsIdentical = IdenticalTo(Widget, Value, NewValue);
 			if (!bIsIdentical)
 			{
 				Value = MoveTemp(NewValue);
@@ -673,7 +687,7 @@ namespace SlateAttributePrivate
 		}
 
 		/** @return True if they have the same Getter or the same value. */
-		UE_NODISCARD bool IsIdenticalTo(const SWidget& Widget, const TSlateAttributeBase& Other) const
+		UE_NODISCARD bool IdenticalTo(const SWidget& Widget, const TSlateAttributeBase& Other) const
 		{
 			VerifyOwningWidget(Widget);
 			FDelegateHandle ThisDelegateHandle = ProtectedFindGetterHandle(Widget, InAttributeType);
@@ -684,13 +698,13 @@ namespace SlateAttributePrivate
 				{
 					return true;
 				}
-				return FComparePredicate{}.operator()(Get(), Other.Get());
+				return IdenticalTo(Widget, Get(), Other.Get());
 			}
 			return false;
 		}
 
 		/** @return True if they have the same Getter or, if the Attribute is set, the same value. */
-		UE_NODISCARD bool IsIdenticalTo(const SWidget& Widget, const TAttribute<ObjectType>& Other) const
+		UE_NODISCARD bool IdenticalTo(const SWidget& Widget, const TAttribute<ObjectType>& Other) const
 		{
 			VerifyOwningWidget(Widget);
 			FDelegateHandle ThisDelegateHandle = ProtectedFindGetterHandle(Widget, InAttributeType);
@@ -698,7 +712,7 @@ namespace SlateAttributePrivate
 			{
 				return Other.GetBinding().GetHandle() == ThisDelegateHandle;
 			}
-			return !ThisDelegateHandle.IsValid() && Other.IsSet() && FComparePredicate{}.operator()(Get(), Other.Get());
+			return !ThisDelegateHandle.IsValid() && Other.IsSet() && IdenticalTo(Widget, Get(), Other.Get());
 		}
 
 	private:
@@ -768,7 +782,8 @@ namespace SlateAttributePrivate
 			{
 				ObjectType NewValue = Getter.Execute();
 
-				if (!(FComparePredicate{}.operator()(Attribute->Value, NewValue)))
+				const bool bIsIdentical = Attribute->IdenticalTo(Widget, Attribute->Value, NewValue);
+				if (!bIsIdentical)
 				{
 					// Set the value on the widget
 					Attribute->Value = MoveTemp(NewValue);
@@ -827,7 +842,7 @@ namespace SlateAttributePrivate
 	/**
 	 *
 	 */
-	template<typename InObjectType, typename InInvalidationReasonPredicate, typename InComparePredicate = TEqualTo<>>
+	template<typename InObjectType, typename InInvalidationReasonPredicate, typename InComparePredicate>
 	struct TSlateMemberAttribute : public TSlateAttributeBase<InObjectType, InInvalidationReasonPredicate, InComparePredicate, ESlateAttributeType::Member>
 	{
 	private:
@@ -899,7 +914,7 @@ namespace SlateAttributePrivate
 	/**
 	 *
 	 */
-	template<typename InObjectType, typename InInvalidationReasonPredicate, typename InComparePredicate = TEqualTo<>>
+	template<typename InObjectType, typename InInvalidationReasonPredicate, typename InComparePredicate>
 	struct TSlateManagedAttribute : protected TSlateAttributeBase<InObjectType, InInvalidationReasonPredicate, InComparePredicate, ESlateAttributeType::Managed>
 	{
 	private:
@@ -913,6 +928,7 @@ namespace SlateAttributePrivate
 		static const bool IsMemberType = false;
 
 		static EInvalidateWidgetReason GetInvalidationReason(const SWidget& Widget) { return Super::GetInvalidationReason(Widget); }
+		static bool IdenticalTo(const SWidget& Widget, const ObjectType& Lhs, const ObjectType& Rhs) { return Super::IdenticalTo(Widget, Lhs, Rhs); }
 
 	public:
 		TSlateManagedAttribute() = delete;
@@ -1112,20 +1128,20 @@ namespace SlateAttributePrivate
 			return false;
 		}
 
-		bool IsIdenticalTo(const TSlateManagedAttribute& Other) const
+		bool IdenticalTo(const TSlateManagedAttribute& Other) const
 		{
 			if (TSharedPtr<SWidget> Pin = ManagedWidget.Pin())
 			{
-				return Super::IsIdenticalTo(*Pin.Get(), Other);
+				return Super::IdenticalTo(*Pin.Get(), Other);
 			}
 			return false;
 		}
 
-		bool IsIdenticalTo(const TAttribute<ObjectType>& Other) const
+		bool IdenticalTo(const TAttribute<ObjectType>& Other) const
 		{
 			if (TSharedPtr<SWidget> Pin = ManagedWidget.Pin())
 			{
-				return Super::IsIdenticalTo(*Pin.Get(), Other);
+				return Super::IdenticalTo(*Pin.Get(), Other);
 			}
 			return false;
 		}
@@ -1252,23 +1268,23 @@ namespace SlateAttributePrivate
 		}
 
 		/** @return True if they have the same Getter or the same value. */
-		UE_NODISCARD bool IsIdenticalTo(const TSlateMemberAttributeRef& Other) const
+		UE_NODISCARD bool IdenticalTo(const TSlateMemberAttributeRef& Other) const
 		{
 			TSharedPtr<const SWidget> SelfPin = Owner.Pin();
 			TSharedPtr<const SWidget> OtherPin = Other.Owner.Pin();
 			if (SelfPin == OtherPin && SelfPin)
 			{
-				return Attribute.IsIdenticalTo(*SelfPin.Get(), Other.Attribute);
+				return Attribute.IdenticalTo(*SelfPin.Get(), Other.Attribute);
 			}
 			return SelfPin == OtherPin;
 		}
 
 		/** @return True if they have the same Getter or, if the Attribute is set, the same value. */
-		UE_NODISCARD bool IsIdenticalTo(const TAttribute<ObjectType>& Other) const
+		UE_NODISCARD bool IdenticalTo(const TAttribute<ObjectType>& Other) const
 		{
 			if (TSharedPtr<const SWidget> Pin = Owner.Pin())
 			{
-				return Attribute.IsIdenticalTo(*Pin.Get(), Other);
+				return Attribute.IdenticalTo(*Pin.Get(), Other);
 			}
 			return !Other.IsSet(); // if the other is not set, then both are invalid.
 		}
diff --git a/Engine/Source/Runtime/SlateCore/Public/Widgets/SWidget.h b/Engine/Source/Runtime/SlateCore/Public/Widgets/SWidget.h
--- a/Engine/Source/Runtime/SlateCore/Public/Widgets/SWidget.h
+++ b/Engine/Source/Runtime/SlateCore/Public/Widgets/SWidget.h
@@ -164,9 +164,9 @@ class SLATECORE_API SWidget
 protected:
 	/**
 	 * A SlateAttribute that is member variable of a SWidget.
-	 * @usage: TSlateAttribute<int32> MyAttribute1; TSlateAttribute<int32, EInvalidateWidgetReason::Paint> MyAttribute2; TSlateAttribute<int32, EInvalidateWidgetReason::Paint, TEqualTo<>> MyAttribute3;
+	 * @usage: TSlateAttribute<int32> MyAttribute1; TSlateAttribute<int32, EInvalidateWidgetReason::Paint> MyAttribute2; TSlateAttribute<int32, EInvalidateWidgetReason::Paint, TSlateAttributeComparePredicate<>> MyAttribute3;
 	 */
-	template<typename InObjectType, EInvalidateWidgetReason InInvalidationReasonValue = EInvalidateWidgetReason::None, typename InComparePredicate = TEqualTo<>>
+	template<typename InObjectType, EInvalidateWidgetReason InInvalidationReasonValue = EInvalidateWidgetReason::None, typename InComparePredicate = TSlateAttributeComparePredicate<>>
 	struct TSlateAttribute : public ::SlateAttributePrivate::TSlateMemberAttribute<
 		InObjectType,
 		typename std::conditional<InInvalidationReasonValue == EInvalidateWidgetReason::None, ::SlateAttributePrivate::FSlateAttributeNoInvalidationReason, TSlateAttributeInvalidationReason<InInvalidationReasonValue>>::type,
@@ -200,9 +200,9 @@ protected:
 
 	/**
 	 * A SlateAttribute that is NOT a member variable of a SWidget.
-	 * @usage: TSlateManagedAttribute<int32> MyAttribute1; TSlateManagedAttribute<int32, EInvalidateWidgetReason::Paint> MyAttribute2; TSlateManagedAttribute<int32, EInvalidateWidgetReason::Paint, TEqualTo<>> MyAttribute3;
+	 * @usage: TSlateManagedAttribute<int32> MyAttribute1; TSlateManagedAttribute<int32, EInvalidateWidgetReason::Paint> MyAttribute2; TSlateManagedAttribute<int32, EInvalidateWidgetReason::Paint, TSlateAttributeComparePredicate<>> MyAttribute3;
 	 */
-	template<typename InObjectType, EInvalidateWidgetReason InInvalidationReasonValue = EInvalidateWidgetReason::None, typename InComparePredicate = TEqualTo<>>
+	template<typename InObjectType, EInvalidateWidgetReason InInvalidationReasonValue = EInvalidateWidgetReason::None, typename InComparePredicate = TSlateAttributeComparePredicate<>>
 	struct TSlateManagedAttribute : public ::SlateAttributePrivate::TSlateManagedAttribute<
 		InObjectType,
 		typename std::conditional<InInvalidationReasonValue == EInvalidateWidgetReason::None, ::SlateAttributePrivate::FSlateAttributeNoInvalidationReason, TSlateAttributeInvalidationReason<InInvalidationReasonValue>>::type,
@@ -235,7 +235,7 @@ protected:
 	};
 
 	/** A Reference to a TSlateAttribute. */
-	template<typename InObjectType, EInvalidateWidgetReason InInvalidationReasonValue = EInvalidateWidgetReason::None, typename InComparePredicate = TEqualTo<>>
+	template<typename InObjectType, EInvalidateWidgetReason InInvalidationReasonValue = EInvalidateWidgetReason::None, typename InComparePredicate = TSlateAttributeComparePredicate<>>
 	struct TSlateAttributeRef : public ::SlateAttributePrivate::TSlateMemberAttributeRef<TSlateAttribute<InObjectType, InInvalidationReasonValue, InComparePredicate>>
 	{
 		using ::SlateAttributePrivate::TSlateMemberAttributeRef<TSlateAttribute<InObjectType, InInvalidationReasonValue, InComparePredicate>>::TSlateMemberAttributeRef;
