Slate: Remove duplicated UpdateFlag in WidgetProxy.

diff --git a/Engine/Source/Runtime/SlateCore/Private/FastUpdate/SlateInvalidationRoot.cpp b/Engine/Source/Runtime/SlateCore/Private/FastUpdate/SlateInvalidationRoot.cpp
--- a/Engine/Source/Runtime/SlateCore/Private/FastUpdate/SlateInvalidationRoot.cpp
+++ b/Engine/Source/Runtime/SlateCore/Private/FastUpdate/SlateInvalidationRoot.cpp
@@ -112,7 +112,8 @@ static FAutoConsoleVariableRef CVarSlateInvalidationRootVerifyWidgetsUpdateList(
 	TEXT("Each frame, verify that pre and post update list contains the correct information and they are sorted.")
 );
 void VerifyWidgetsUpdateList_BeforeProcessPreUpdate(const TSharedRef<SWidget>&, FSlateInvalidationWidgetList*, FSlateInvalidationWidgetPreHeap*, FSlateInvalidationWidgetPostHeap*, TArray<FSlateInvalidationWidgetIndex>&);
-void VerifyWidgetsUpdateList_BeoreProcessPostUpdate(const TSharedRef<SWidget>&, FSlateInvalidationWidgetList*, FSlateInvalidationWidgetPreHeap*, FSlateInvalidationWidgetPostHeap*, TArray<FSlateInvalidationWidgetIndex>&);
+void VerifyWidgetsUpdateList_BeforeProcessPostUpdate(const TSharedRef<SWidget>&, FSlateInvalidationWidgetList*, FSlateInvalidationWidgetPreHeap*, FSlateInvalidationWidgetPostHeap*, TArray<FSlateInvalidationWidgetIndex>&);
+void VerifyWidgetsUpdateList_AfterProcessPostUpdate(const TSharedRef<SWidget>&, FSlateInvalidationWidgetList*, FSlateInvalidationWidgetPreHeap*, FSlateInvalidationWidgetPostHeap*, TArray<FSlateInvalidationWidgetIndex>&);
 
 bool GSlateInvalidationRootVerifySlateAttribute = false;
 static FAutoConsoleVariableRef CVarSlateInvalidationRootVerifySlateAttributes(
@@ -120,8 +121,8 @@ static FAutoConsoleVariableRef CVarSlateInvalidationRootVerifySlateAttributes(
 	GSlateInvalidationRootVerifySlateAttribute,
 	TEXT("Each frame, verify that the widgets that have registered attribute are correctly updated once and the list contains all the widgets.")
 );
-void VerifySlateAttribute_Pre(FSlateInvalidationWidgetList& FastWidgetPathList);
-void VerifySlateAttribute_Post(const FSlateInvalidationWidgetList& FastWidgetPathList);
+void VerifySlateAttribute_BeforeProcessPreUpdate(FSlateInvalidationWidgetList& FastWidgetPathList);
+void VerifySlateAttribute_AfterProcessPreUpdate(const FSlateInvalidationWidgetList& FastWidgetPathList);
 
 #endif //UE_SLATE_WITH_INVALIDATIONWIDGETLIST_DEBUGGING
 
@@ -446,7 +447,7 @@ bool FSlateInvalidationRoot::PaintFastPath(const FSlateInvalidationContext& Cont
 			SWidget* WidgetPtr = WidgetProxy.GetWidget();
 			if (!WidgetProxy.bUpdatedSinceLastInvalidate && WidgetPtr && WidgetPtr->IsFastPathVisible())
 			{
-				bWidgetsNeededRepaint = bWidgetsNeededRepaint || EnumHasAnyFlags(WidgetProxy.UpdateFlags, EWidgetUpdateFlags::NeedsRepaint | EWidgetUpdateFlags::NeedsVolatilePaint);
+				bWidgetsNeededRepaint = bWidgetsNeededRepaint || WidgetPtr->HasAnyUpdateFlags(EWidgetUpdateFlags::NeedsRepaint | EWidgetUpdateFlags::NeedsVolatilePaint);
 
 				const int32 NewLayerId = WidgetProxy.Update(*Context.PaintArgs, *Context.WindowElementList);
 				CachedMaxLayerId = FMath::Max(NewLayerId, CachedMaxLayerId);
@@ -507,7 +508,7 @@ void FSlateInvalidationRoot::ProcessPreUpdate()
 	}
 	if (GSlateInvalidationRootVerifySlateAttribute)
 	{
-		VerifySlateAttribute_Pre(*FastWidgetPathList);
+		VerifySlateAttribute_BeforeProcessPreUpdate(*FastWidgetPathList);
 	}
 #endif
 
@@ -791,7 +792,7 @@ void FSlateInvalidationRoot::ProcessPreUpdate()
 	}
 	if (GSlateInvalidationRootVerifySlateAttribute)
 	{
-		VerifySlateAttribute_Post(*FastWidgetPathList);
+		VerifySlateAttribute_AfterProcessPreUpdate(*FastWidgetPathList);
 	}
 #endif //UE_SLATE_WITH_INVALIDATIONWIDGETLIST_DEBUGGING
 }
@@ -803,7 +804,7 @@ bool FSlateInvalidationRoot::ProcessPostUpdate()
 #if UE_SLATE_WITH_INVALIDATIONWIDGETLIST_DEBUGGING
 	if (GSlateInvalidationRootVerifyWidgetsUpdateList)
 	{
-		VerifyWidgetsUpdateList_BeoreProcessPostUpdate(GetRootWidget(), FastWidgetPathList.Get(), WidgetsNeedingPreUpdate.Get(), WidgetsNeedingPostUpdate.Get(), FinalUpdateList);
+		VerifyWidgetsUpdateList_BeforeProcessPostUpdate(GetRootWidget(), FastWidgetPathList.Get(), WidgetsNeedingPreUpdate.Get(), WidgetsNeedingPostUpdate.Get(), FinalUpdateList);
 	}
 #endif
 
@@ -858,6 +859,13 @@ bool FSlateInvalidationRoot::ProcessPostUpdate()
 	}
 	WidgetsNeedingPostUpdate->Reset(true);
 
+#if UE_SLATE_WITH_INVALIDATIONWIDGETLIST_DEBUGGING
+	if (GSlateInvalidationRootVerifyWidgetsUpdateList && !bNeedsSlowPath)
+	{
+		VerifyWidgetsUpdateList_AfterProcessPostUpdate(GetRootWidget(), FastWidgetPathList.Get(), WidgetsNeedingPreUpdate.Get(), WidgetsNeedingPostUpdate.Get(), FinalUpdateList);
+	}
+#endif
+
 	return bWidgetsNeedRepaint;
 }
 
@@ -979,22 +987,22 @@ void DumpUpdateList(const FSlateInvalidationWidgetList& FastWidgetPathList, cons
 		const SWidget* WidgetPtr = WidgetProxy.GetWidget();
 		if (WidgetPtr)
 		{
-			if (EnumHasAnyFlags(WidgetProxy.UpdateFlags, EWidgetUpdateFlags::NeedsVolatilePaint))
+			if (WidgetPtr->HasAnyUpdateFlags(EWidgetUpdateFlags::NeedsVolatilePaint))
 			{
 				UE_LOG(LogSlate, Log, TEXT("Volatile Repaint %s"), *FReflectionMetaData::GetWidgetDebugInfo(WidgetPtr));
 			}
-			else if (EnumHasAnyFlags(WidgetProxy.UpdateFlags, EWidgetUpdateFlags::NeedsRepaint))
+			else if (WidgetPtr->HasAnyUpdateFlags(EWidgetUpdateFlags::NeedsRepaint))
 			{
 				UE_LOG(LogSlate, Log, TEXT("Repaint %s"), *FReflectionMetaData::GetWidgetDebugInfo(WidgetPtr));
 			}
 			else if (WidgetPtr->IsFastPathVisible())
 			{
-				if (EnumHasAnyFlags(WidgetProxy.UpdateFlags, EWidgetUpdateFlags::NeedsActiveTimerUpdate))
+				if (WidgetPtr->HasAnyUpdateFlags(EWidgetUpdateFlags::NeedsActiveTimerUpdate))
 				{
 					UE_LOG(LogSlate, Log, TEXT("ActiveTimer %s"), *FReflectionMetaData::GetWidgetDebugInfo(WidgetPtr));
 				}
 
-				if (EnumHasAnyFlags(WidgetProxy.UpdateFlags, EWidgetUpdateFlags::NeedsTick))
+				if (WidgetPtr->HasAnyUpdateFlags(EWidgetUpdateFlags::NeedsTick))
 				{
 					UE_LOG(LogSlate, Log, TEXT("Tick %s"), *FReflectionMetaData::GetWidgetDebugInfo(WidgetPtr));
 				}
@@ -1049,7 +1057,7 @@ void LogPostInvalidationItem(const FSlateInvalidationWidgetList& FastWidgetPathL
 	{
 		UE_LOG(LogSlate, Log, TEXT("  Paint %s"), *FReflectionMetaData::GetWidgetDebugInfo(Proxy.GetWidget()));
 	}
-	else if (!EnumHasAnyFlags(Proxy.UpdateFlags, EWidgetUpdateFlags::AnyUpdate))
+	else if (!Proxy.GetWidget()->HasAnyUpdateFlags(EWidgetUpdateFlags::AnyUpdate))
 	{
 		UE_LOG(LogSlate, Log, TEXT("  [?] %s"), *FReflectionMetaData::GetWidgetDebugInfo(Proxy.GetWidget()));
 	}
@@ -1231,16 +1239,17 @@ void VerifyWidgetVolatile(FSlateInvalidationWidgetList& WidgetList, TArray<FSlat
 					const bool bShouldBeVolatileIndirectly = ParentWidget->IsVolatileIndirectly() || ParentWidget->IsVolatile();
 					UE_SLATE_LOG_ERROR_IF_FALSE(Widget->IsVolatileIndirectly() == bShouldBeVolatileIndirectly
 					, GSlateInvalidationRootVerifyWidgetVolatile
-						, TEXT("Widget '%s' should be set as %s.")
-						, *FReflectionMetaData::GetWidgetDebugInfo(Widget)
-						, (bShouldBeVolatileIndirectly ? TEXT("volatile indirectly") : TEXT("not volatile indirectly")));
+					, TEXT("Widget '%s' should be set as %s.")
+					, *FReflectionMetaData::GetWidgetDebugInfo(Widget)
+					, (bShouldBeVolatileIndirectly ? TEXT("volatile indirectly") : TEXT("not volatile indirectly")));
 				}
 
 				if (Widget->IsVolatile() && !Widget->IsVolatileIndirectly())
 				{
-					//const bool bHasUpdateFlag = Widget->HasAnyUpdateFlags(EWidgetUpdateFlags::NeedsVolatilePaint);
-					//ensureMsgf(bHasUpdateFlag, TEXT("Widget '%s' is volatile but doesn't have the update flag NeedsVolatilePaint.")
-					//	, *FReflectionMetaData::GetWidgetDebugInfo(Widget));
+					UE_SLATE_LOG_ERROR_IF_FALSE(Widget->HasAnyUpdateFlags(EWidgetUpdateFlags::NeedsVolatilePaint)
+						, GSlateInvalidationRootVerifyWidgetVolatile
+						, TEXT("Widget '%s' is volatile but doesn't have the update flag NeedsVolatilePaint.")
+						, *FReflectionMetaData::GetWidgetDebugInfo(Widget));
 
 					const bool bIsContains = FinalUpdateList.Contains(Widget->GetProxyHandle().GetWidgetIndex());
 					UE_SLATE_LOG_ERROR_IF_FALSE(bIsContains
@@ -1267,62 +1276,133 @@ void VerifyWidgetsUpdateList_BeforeProcessPreUpdate(const TSharedRef<SWidget>& R
 	{
 		UE_SLATE_LOG_ERROR_IF_FALSE(FastWidgetPathList->IsValidIndex(WidgetIndex)
 			, GSlateInvalidationRootVerifyWidgetsUpdateList
-			, TEXT("An element is not valid."));
-		const FSlateInvalidationWidgetList::InvalidationWidgetType& InvalidationWidget = (*FastWidgetPathList)[WidgetIndex];
-		UE_SLATE_LOG_ERROR_IF_FALSE(InvalidationWidget.CurrentInvalidateReason != EInvalidateWidgetReason::None || InvalidationWidget.UpdateFlags != EWidgetUpdateFlags::None
-			, GSlateInvalidationRootVerifyWidgetsUpdateList
-			, TEXT("A widget is in the update list but doesn't have an update flag set."));
+			, TEXT("A WidgetIndex is invalid. The Widget can be invalid (because it's not been processed yet)."));
 	}
 
-	ensureMsgf(WidgetsNeedingPreUpdate->IsValidHeap_Debug(), TEXT("The PreUpdate list need to stay a valid heap"));
+	UE_SLATE_LOG_ERROR_IF_FALSE(WidgetsNeedingPreUpdate->IsValidHeap_Debug()
+		, GSlateInvalidationRootVerifyWidgetsUpdateList
+		, TEXT("The PreUpdate list need to stay a valid heap"));
 
-	auto CheckIfValidInvalidation = [FastWidgetPathList](const FSlateInvalidationWidgetPreHeap::FElement& Element)
+	WidgetsNeedingPreUpdate->ForEachIndexes([FastWidgetPathList](const FSlateInvalidationWidgetPreHeap::FElement& Element)
+		{
+			UE_SLATE_LOG_ERROR_IF_FALSE(FastWidgetPathList->IsValidIndex(Element.GetWidgetIndex())
+				, GSlateInvalidationRootVerifyWidgetsUpdateList
+				, TEXT("An element is not valid."));
+			const FSlateInvalidationWidgetList::InvalidationWidgetType& InvalidationWidget = (*FastWidgetPathList)[Element.GetWidgetIndex()];
+			if (SWidget* Widget = (*FastWidgetPathList)[Element.GetWidgetIndex()].GetWidget())
+			{
+				UE_SLATE_LOG_ERROR_IF_FALSE(Widget->GetProxyHandle().GetWidgetSortOrder() == Element.GetWidgetSortOrder()
+					, GSlateInvalidationRootVerifyWidgetsUpdateList
+					, TEXT("The sort order of the widget '%s' do not matches what is in the heap.")
+					, *FReflectionMetaData::GetWidgetDebugInfo(Widget));
+				UE_SLATE_LOG_ERROR_IF_FALSE(Widget->GetProxyHandle().GetWidgetIndex() == Element.GetWidgetIndex()
+					, GSlateInvalidationRootVerifyWidgetsUpdateList
+					, TEXT("The widget index of the widget '%s' do not matches what is in the heap.")
+					, *FReflectionMetaData::GetWidgetDebugInfo(Widget));
+			}
+		});
+
+	FastWidgetPathList->ForEachInvalidationWidget([WidgetsNeedingPreUpdate](FSlateInvalidationWidgetList::InvalidationWidgetType& InvalidationWidget)
+		{
+			UE_SLATE_LOG_ERROR_IF_FALSE(WidgetsNeedingPreUpdate->Contains_Debug(InvalidationWidget.Index) == InvalidationWidget.bContainedByWidgetPreHeap
+				, GSlateInvalidationRootVerifyWidgetsUpdateList
+				, TEXT("Widget '%s' is or is not in the PreUpdate but the flag say otherwise.")
+				, *FReflectionMetaData::GetWidgetDebugInfo(InvalidationWidget.GetWidget()));
+		});
+}
+
+void VerifyWidgetsUpdateList_BeforeProcessPostUpdate(const TSharedRef<SWidget>& RootWidget,
+	FSlateInvalidationWidgetList* FastWidgetPathList,
+	FSlateInvalidationWidgetPreHeap* WidgetsNeedingPreUpdate,
+	FSlateInvalidationWidgetPostHeap* WidgetsNeedingPostUpdate,
+	TArray<FSlateInvalidationWidgetIndex>& FinalUpdateList)
+{
+	if (FastWidgetPathList->GetRoot().Pin() != RootWidget)
 	{
-		UE_SLATE_LOG_ERROR_IF_FALSE(FastWidgetPathList->IsValidIndex(Element.GetWidgetIndex())
-			, GSlateInvalidationRootVerifyWidgetsUpdateList
-			, TEXT("An element is not valid."));
-		const FSlateInvalidationWidgetList::InvalidationWidgetType& InvalidationWidget = (*FastWidgetPathList)[Element.GetWidgetIndex()];
-		UE_SLATE_LOG_ERROR_IF_FALSE(InvalidationWidget.CurrentInvalidateReason != EInvalidateWidgetReason::None || InvalidationWidget.UpdateFlags != EWidgetUpdateFlags::None
-			, GSlateInvalidationRootVerifyWidgetsUpdateList
-			, TEXT("A widget is in the update list but doesn't have an update flag set."));
-		if (SWidget* Widget = (*FastWidgetPathList)[Element.GetWidgetIndex()].GetWidget())
+		return;
+	}
+
+	UE_SLATE_LOG_ERROR_IF_FALSE(WidgetsNeedingPostUpdate->IsValidHeap_Debug()
+		, GSlateInvalidationRootVerifyWidgetsUpdateList
+		, TEXT("The PostUpdate list need to stay a valid heap"));
+
+	UE_SLATE_LOG_ERROR_IF_FALSE(WidgetsNeedingPreUpdate->Num() == 0
+		, GSlateInvalidationRootVerifyWidgetsUpdateList
+		, TEXT("The PreUpdate list should be empty"));
+
+	UE_SLATE_LOG_ERROR_IF_FALSE(FinalUpdateList.Num() == 0
+		, GSlateInvalidationRootVerifyWidgetsUpdateList
+		, TEXT("The Final Update list should be empty."));
+
+	WidgetsNeedingPostUpdate->ForEachIndexes([FastWidgetPathList](const FSlateInvalidationWidgetPreHeap::FElement& Element)
 		{
+			UE_SLATE_LOG_ERROR_IF_FALSE(FastWidgetPathList->IsValidIndex(Element.GetWidgetIndex())
+				, GSlateInvalidationRootVerifyWidgetsUpdateList
+				, TEXT("An element is not valid."));
+
+			const FSlateInvalidationWidgetList::InvalidationWidgetType& InvalidationWidget = (*FastWidgetPathList)[Element.GetWidgetIndex()];
+			const SWidget* Widget = InvalidationWidget.GetWidget();
+
+			UE_SLATE_LOG_ERROR_IF_FALSE(Widget
+				, GSlateInvalidationRootVerifyWidgetsUpdateList
+				, TEXT("Widget should be valid (should have been cleaned by PreProcess)."));
 			UE_SLATE_LOG_ERROR_IF_FALSE(Widget->GetProxyHandle().GetWidgetSortOrder() == Element.GetWidgetSortOrder()
 				, GSlateInvalidationRootVerifyWidgetsUpdateList
-				, TEXT("The sort order of the widget do not matches what is in the heap.")
+				, TEXT("The sort order of the widget '%s' do not matches what is in the heap.")
 				, *FReflectionMetaData::GetWidgetDebugInfo(Widget));
 			UE_SLATE_LOG_ERROR_IF_FALSE(Widget->GetProxyHandle().GetWidgetIndex() == Element.GetWidgetIndex()
 				, GSlateInvalidationRootVerifyWidgetsUpdateList
-				, TEXT("The widget index of the widget do not matches what is in the heap.")
+				, TEXT("The widget index of the widget '%s' do not matches what is in the heap.")
 				, *FReflectionMetaData::GetWidgetDebugInfo(Widget));
-		}
-	};
-	WidgetsNeedingPreUpdate->ForEachIndexes(CheckIfValidInvalidation);
-	WidgetsNeedingPostUpdate->ForEachIndexes(CheckIfValidInvalidation);
+		});
 
-	FastWidgetPathList->ForEachInvalidationWidget([WidgetsNeedingPreUpdate, WidgetsNeedingPostUpdate](FSlateInvalidationWidgetList::InvalidationWidgetType& InvalidationWidget)
+	FastWidgetPathList->ForEachInvalidationWidget([WidgetsNeedingPostUpdate](FSlateInvalidationWidgetList::InvalidationWidgetType& InvalidationWidget)
 		{
-			UE_SLATE_LOG_ERROR_IF_FALSE(WidgetsNeedingPreUpdate->Contains_Debug(InvalidationWidget.Index) == InvalidationWidget.bContainedByWidgetPreHeap
-				, GSlateInvalidationRootVerifyWidgetsUpdateList
-				, TEXT("Widget '%s' is or is not in the PreUpdate but the flag say otherwise."));
 			UE_SLATE_LOG_ERROR_IF_FALSE(WidgetsNeedingPostUpdate->Contains_Debug(InvalidationWidget.Index) == InvalidationWidget.bContainedByWidgetPostHeap
 				, GSlateInvalidationRootVerifyWidgetsUpdateList
-				, TEXT("Widget '%s' is or is not in the PostUpdate but the flag say otherwise."));
+				, TEXT("Widget '%s' is or is not in the PostUpdate but the flag say otherwise.")
+				, *FReflectionMetaData::GetWidgetDebugInfo(InvalidationWidget.GetWidget()));
 		});
 }
 
-void VerifyWidgetsUpdateList_BeoreProcessPostUpdate(const TSharedRef<SWidget>& RootWidget,
+void VerifyWidgetsUpdateList_AfterProcessPostUpdate(const TSharedRef<SWidget>& RootWidget,
 	FSlateInvalidationWidgetList* FastWidgetPathList,
 	FSlateInvalidationWidgetPreHeap* WidgetsNeedingPreUpdate,
 	FSlateInvalidationWidgetPostHeap* WidgetsNeedingPostUpdate,
 	TArray<FSlateInvalidationWidgetIndex>& FinalUpdateList)
 {
-	ensureMsgf(WidgetsNeedingPostUpdate->IsValidHeap_Debug(), TEXT("The PostUpdate list need to stay a valid heap"));
+	if (FastWidgetPathList->GetRoot().Pin() != RootWidget)
+	{
+		return;
+	}
+
+	UE_SLATE_LOG_ERROR_IF_FALSE(WidgetsNeedingPreUpdate->Num() == 0
+		, GSlateInvalidationRootVerifyWidgetsUpdateList
+		, TEXT("The list of Pre Update should already been processed."));
+	UE_SLATE_LOG_ERROR_IF_FALSE(WidgetsNeedingPostUpdate->Num() == 0
+		, GSlateInvalidationRootVerifyWidgetsUpdateList
+		, TEXT("The list of Post Update should already been processed."));
 
-	VerifyWidgetsUpdateList_BeforeProcessPreUpdate(RootWidget, FastWidgetPathList, WidgetsNeedingPreUpdate, WidgetsNeedingPostUpdate, FinalUpdateList);
+	for(FSlateInvalidationWidgetIndex WidgetIndex : FinalUpdateList)
+	{
+		const FSlateInvalidationWidgetList::InvalidationWidgetType& InvalidationWidget = (*FastWidgetPathList)[WidgetIndex];
+		const SWidget* Widget = InvalidationWidget.GetWidget();
+
+		UE_SLATE_LOG_ERROR_IF_FALSE(Widget
+			, GSlateInvalidationRootVerifyWidgetsUpdateList
+			, TEXT("Widget should be valid (should have been cleaned by PreProcess)."));
+		UE_SLATE_LOG_ERROR_IF_FALSE(InvalidationWidget.CurrentInvalidateReason == EInvalidateWidgetReason::None
+			, GSlateInvalidationRootVerifyWidgetsUpdateList
+			, TEXT("The widget '%s' is in the update list and it still has a Invalidation Reason.")
+			, * FReflectionMetaData::GetWidgetDebugInfo(Widget));
+		UE_SLATE_LOG_ERROR_IF_FALSE(Widget->HasAnyUpdateFlags(EWidgetUpdateFlags::AnyUpdate)
+			, GSlateInvalidationRootVerifyWidgetsUpdateList
+			, TEXT("The widget '%s' is in the update list but doesn't have an update flag set.")
+			, *FReflectionMetaData::GetWidgetDebugInfo(Widget));
+	}
 }
 
-void VerifySlateAttribute_Pre(FSlateInvalidationWidgetList& FastWidgetPathList)
+void VerifySlateAttribute_BeforeProcessPreUpdate(FSlateInvalidationWidgetList& FastWidgetPathList)
 {
 	FastWidgetPathList.ForEachInvalidationWidget([](FSlateInvalidationWidgetList::InvalidationWidgetType& InvalidationWidget)
 		{
@@ -1330,7 +1410,7 @@ void VerifySlateAttribute_Pre(FSlateInvalidationWidgetList& FastWidgetPathList)
 		});
 }
 
-void VerifySlateAttribute_Post(const FSlateInvalidationWidgetList& FastWidgetPathList)
+void VerifySlateAttribute_AfterProcessPreUpdate(const FSlateInvalidationWidgetList& FastWidgetPathList)
 {
 	const bool bElementIndexListValid = FastWidgetPathList.VerifyElementIndexList();
 	UE_SLATE_LOG_ERROR_IF_FALSE(bElementIndexListValid
diff --git a/Engine/Source/Runtime/SlateCore/Private/FastUpdate/WidgetProxy.cpp b/Engine/Source/Runtime/SlateCore/Private/FastUpdate/WidgetProxy.cpp
--- a/Engine/Source/Runtime/SlateCore/Private/FastUpdate/WidgetProxy.cpp
+++ b/Engine/Source/Runtime/SlateCore/Private/FastUpdate/WidgetProxy.cpp
@@ -23,7 +23,6 @@ FWidgetProxy::FWidgetProxy(TSharedRef<SWidget>& InWidget)
 	, Index(FSlateInvalidationWidgetIndex::Invalid)
 	, ParentIndex(FSlateInvalidationWidgetIndex::Invalid)
 	, LeafMostChildIndex(FSlateInvalidationWidgetIndex::Invalid)
-	, UpdateFlags(EWidgetUpdateFlags::None)
 	, CurrentInvalidateReason(EInvalidateWidgetReason::None)
 	// Potentially unsafe to update visibility from the widget due to attribute bindings.  This is updated later when the widgets are sorted in ProcessInvalidation
 	, Visibility(EVisibility::Collapsed) 
@@ -59,21 +58,21 @@ int32 FWidgetProxy::Update(const FPaintArgs& PaintArgs, FSlateWindowElementList&
 
 	// If Outgoing layer id remains index none, there was no change
 	int32 OutgoingLayerId = INDEX_NONE;
-	if (EnumHasAnyFlags(UpdateFlags, EWidgetUpdateFlags::NeedsRepaint|EWidgetUpdateFlags::NeedsVolatilePaint))
+	if (CurrentWidget->HasAnyUpdateFlags(EWidgetUpdateFlags::NeedsRepaint|EWidgetUpdateFlags::NeedsVolatilePaint))
 	{
 		check(CurrentWidget->IsFastPathVisible());
 		OutgoingLayerId = Repaint(PaintArgs, OutDrawElements);
 	}
 	else if(CurrentWidget->IsFastPathVisible())
 	{
-		EWidgetUpdateFlags PreviousUpdateFlag = UpdateFlags;
-		if (EnumHasAnyFlags(UpdateFlags, EWidgetUpdateFlags::NeedsActiveTimerUpdate))
+		EWidgetUpdateFlags PreviousUpdateFlag = CurrentWidget->UpdateFlags;
+		if (CurrentWidget->HasAnyUpdateFlags(EWidgetUpdateFlags::NeedsActiveTimerUpdate))
 		{
 			SCOPE_CYCLE_COUNTER(STAT_SlateExecuteActiveTimers);
 			CurrentWidget->ExecuteActiveTimers(PaintArgs.GetCurrentTime(), PaintArgs.GetDeltaTime());
 		}
 
-		if (EnumHasAnyFlags(UpdateFlags, EWidgetUpdateFlags::NeedsTick))
+		if (CurrentWidget->HasAnyUpdateFlags(EWidgetUpdateFlags::NeedsTick))
 		{
 			const FSlateWidgetPersistentState& MyState = CurrentWidget->GetPersistentState();
 
@@ -139,7 +138,6 @@ bool FWidgetProxy::ProcessPostInvalidation(FSlateInvalidationWidgetPostHeap& Upd
 		{
 			// Set the value directly instead of calling AddUpdateFlags as an optimization
 			WidgetPtr->UpdateFlags |= EWidgetUpdateFlags::NeedsRepaint;
-			UpdateFlags |= EWidgetUpdateFlags::NeedsRepaint;
 		}
 
 		// If the desired size changed, invalidate the parent if it is visible
@@ -176,7 +174,6 @@ bool FWidgetProxy::ProcessPostInvalidation(FSlateInvalidationWidgetPostHeap& Upd
 		SCOPE_CYCLE_SWIDGET(WidgetPtr);
 		// Set the value directly instead of calling AddUpdateFlags as an optimization
 		WidgetPtr->UpdateFlags |= EWidgetUpdateFlags::NeedsRepaint;
-		UpdateFlags |= EWidgetUpdateFlags::NeedsRepaint;
 
 		bWidgetNeedsRepaint = true;
 	}
@@ -190,14 +187,11 @@ void FWidgetProxy::MarkProxyUpdatedThisFrame(FSlateInvalidationWidgetPostHeap& P
 {
 	bUpdatedSinceLastInvalidate = true;
 
-	if(EnumHasAnyFlags(UpdateFlags, EWidgetUpdateFlags::AnyUpdate))
+	SWidget* WidgetPtr = GetWidget();
+	if (WidgetPtr && WidgetPtr->IsFastPathVisible() && WidgetPtr->HasAnyUpdateFlags(EWidgetUpdateFlags::AnyUpdate))
 	{
-		SWidget* WidgetPtr = GetWidget();
-		if (WidgetPtr && WidgetPtr->IsFastPathVisible())
-		{
-			// If there are any updates still needed add them to the next update list
-			PostUpdateList.PushBackUnique(*this);
-		}
+		// If there are any updates still needed add them to the next update list
+		PostUpdateList.PushBackUnique(*this);
 	}
 }
 
@@ -207,7 +201,7 @@ int32 FWidgetProxy::Repaint(const FPaintArgs& PaintArgs, FSlateWindowElementList
 	check(WidgetPtr);
 
 #if WITH_SLATE_DEBUGGING
-	SCOPED_NAMED_EVENT_FSTRING(FReflectionMetaData::GetWidgetDebugInfo(WidgetPtr), EnumHasAnyFlags(UpdateFlags, EWidgetUpdateFlags::NeedsRepaint) ? FColor::Orange : FColor::Red);
+	SCOPED_NAMED_EVENT_FSTRING(FReflectionMetaData::GetWidgetDebugInfo(WidgetPtr), WidgetPtr->HasAnyUpdateFlags(EWidgetUpdateFlags::NeedsRepaint) ? FColor::Orange : FColor::Red);
 #endif
 
 	const FSlateWidgetPersistentState& MyState = WidgetPtr->GetPersistentState();
@@ -233,7 +227,7 @@ int32 FWidgetProxy::Repaint(const FPaintArgs& PaintArgs, FSlateWindowElementList
 
 	if (GSlateEnableGlobalInvalidation)
 	{
-		if (EnumHasAnyFlags(UpdateFlags, EWidgetUpdateFlags::NeedsVolatilePaint))
+		if (WidgetPtr->HasAnyUpdateFlags(EWidgetUpdateFlags::NeedsVolatilePaint))
 		{
 			if (WidgetPtr->ShouldInvalidatePrepassDueToVolatility())
 			{
@@ -313,19 +307,13 @@ void FWidgetProxyHandle::MarkWidgetDirty(EInvalidateWidgetReason InvalidateReaso
 
 void FWidgetProxyHandle::UpdateWidgetFlags(const SWidget* Widget, EWidgetUpdateFlags NewFlags)
 {
-	if (IsValid(Widget))
+	// Add to update list if the widget is now tickable or has an active timer.
+	if (Widget->IsFastPathVisible() && Widget->HasAnyUpdateFlags(EWidgetUpdateFlags::AnyUpdate))
 	{
-		FWidgetProxy& Proxy = GetInvalidationRoot()->GetFastPathWidgetList()[WidgetIndex];
-
-		if (Widget->IsFastPathVisible())
+		if (IsValid(Widget))
 		{
-			Proxy.UpdateFlags = NewFlags;
-
-			// Add to update list if the widget is now tickable or has an active timer.
-			if (EnumHasAnyFlags(NewFlags, EWidgetUpdateFlags::AnyUpdate))
-			{
-				GetInvalidationRoot()->WidgetsNeedingPostUpdate->PushBackOrHeapUnique(Proxy);
-			}
+			FWidgetProxy& Proxy = GetInvalidationRoot()->GetFastPathWidgetList()[WidgetIndex];
+			GetInvalidationRoot()->WidgetsNeedingPostUpdate->PushBackOrHeapUnique(Proxy);
 		}
 	}
 }
diff --git a/Engine/Source/Runtime/SlateCore/Public/FastUpdate/WidgetProxy.h b/Engine/Source/Runtime/SlateCore/Public/FastUpdate/WidgetProxy.h
--- a/Engine/Source/Runtime/SlateCore/Public/FastUpdate/WidgetProxy.h
+++ b/Engine/Source/Runtime/SlateCore/Public/FastUpdate/WidgetProxy.h
@@ -81,7 +81,6 @@ public:
 	FSlateInvalidationWidgetIndex Index;
 	FSlateInvalidationWidgetIndex ParentIndex;
 	FSlateInvalidationWidgetIndex LeafMostChildIndex;
-	EWidgetUpdateFlags UpdateFlags;
 	EInvalidateWidgetReason CurrentInvalidateReason;
 	/** The widgets own visibility */
 	EVisibility Visibility;
