Adding initial value to FStyleColorsList

diff --git a/Engine/Source/Runtime/SlateCore/Public/Styling/StyleColors.h b/Engine/Source/Runtime/SlateCore/Public/Styling/StyleColors.h
--- a/Engine/Source/Runtime/SlateCore/Public/Styling/StyleColors.h
+++ b/Engine/Source/Runtime/SlateCore/Public/Styling/StyleColors.h
@@ -96,7 +96,7 @@ struct FStyleColorList
 	GENERATED_BODY()
 
 	UPROPERTY(EditAnywhere, Config, Category = Colors)
-	FLinearColor StyleColors[(int32)EStyleColor::MAX];
+	FLinearColor StyleColors[(int32)EStyleColor::MAX] = {FLinearColor::Transparent};
 
 	FText DisplayNames[(int32)EStyleColor::MAX];
 };
