Commented out EStyleMessageSeverity::CriticalError from the enum because it conflicts with a global symbol defined somewhere in the Windows.h header forest. This now shows up as a shadowing error in clang. It also turns out to be not used anywhere.

#rnx

diff --git a/Engine/Source/Runtime/SlateCore/Public/Styling/ISlateStyle.h b/Engine/Source/Runtime/SlateCore/Public/Styling/ISlateStyle.h
--- a/Engine/Source/Runtime/SlateCore/Public/Styling/ISlateStyle.h
+++ b/Engine/Source/Runtime/SlateCore/Public/Styling/ISlateStyle.h
@@ -173,7 +173,7 @@ protected:
 	 */
 	enum EStyleMessageSeverity
 	{
-		CriticalError		= 0,
+		// CriticalError		= 0,
 		Error				= 1,
 		PerformanceWarning	= 2,
 		Warning				= 3,
