Reverting CL 15884404

#trivial
#rnx

diff --git a/Engine/Source/Runtime/SlateCore/Public/Styling/ISlateStyle.h b/Engine/Source/Runtime/SlateCore/Public/Styling/ISlateStyle.h
--- a/Engine/Source/Runtime/SlateCore/Public/Styling/ISlateStyle.h
+++ b/Engine/Source/Runtime/SlateCore/Public/Styling/ISlateStyle.h
@@ -173,7 +173,7 @@ protected:
 	 */
 	enum EStyleMessageSeverity
 	{
-		// CriticalError		= 0,
+		CriticalError		= 0,
 		Error				= 1,
 		PerformanceWarning	= 2,
 		Warning				= 3,
