Slate: Transform ISlateStyle::EStyleMessageSeverity into a enum class to prevent error in compilation.

diff --git a/Engine/Source/Runtime/Engine/Private/Slate/SlateGameResources.cpp b/Engine/Source/Runtime/Engine/Private/Slate/SlateGameResources.cpp
--- a/Engine/Source/Runtime/Engine/Private/Slate/SlateGameResources.cpp
+++ b/Engine/Source/Runtime/Engine/Private/Slate/SlateGameResources.cpp
@@ -249,7 +249,7 @@ void FSlateGameResources::AddAssetToCache( UObject* InStyleObject, bool bEnsureU
 
 		if ( ExistingAsset != NULL )
 		{
-			Log( ISlateStyle::Error, FText::Format( NSLOCTEXT("SlateWidgetStyleSet", "LoadingError", "Encountered multiple Slate Widget Styles with the same name. Name: '{0}', First Asset: '{1}',  Second Asset: '{2}'."),
+			Log(EStyleMessageSeverity::Error, FText::Format( NSLOCTEXT("SlateWidgetStyleSet", "LoadingError", "Encountered multiple Slate Widget Styles with the same name. Name: '{0}', First Asset: '{1}',  Second Asset: '{2}'."),
 				FText::FromName( StyleName ),
 				FText::FromString( InStyleObject->GetPathName() ),
 				FText::FromString( InStyleObject->GetPathName() ) ) );
diff --git a/Engine/Source/Runtime/SlateCore/Private/Styling/SlateStyleSet.cpp b/Engine/Source/Runtime/SlateCore/Private/Styling/SlateStyleSet.cpp
--- a/Engine/Source/Runtime/SlateCore/Private/Styling/SlateStyleSet.cpp
+++ b/Engine/Source/Runtime/SlateCore/Private/Styling/SlateStyleSet.cpp
@@ -189,7 +189,7 @@ float FSlateStyleSet::GetFloat(const FName PropertyName, const ANSICHAR* Specifi
 	if (Result == nullptr) 
 	{
 		const ISlateStyle* ReportingStyle = RequestingStyle != nullptr ? RequestingStyle : this;
-		ReportingStyle->LogMissingResource(ISlateStyle::Warning, FText::Format(NSLOCTEXT("SlateStyleSet", "UknownSlateFloat", "Unable to find float property '{0}' in style."), FText::FromName(PropertyName)), PropertyName);
+		ReportingStyle->LogMissingResource(EStyleMessageSeverity::Warning, FText::Format(NSLOCTEXT("SlateStyleSet", "UknownSlateFloat", "Unable to find float property '{0}' in style."), FText::FromName(PropertyName)), PropertyName);
 	}
 
 	return Result ? *Result : DefaultValue;
@@ -210,7 +210,7 @@ FVector2D FSlateStyleSet::GetVector(const FName PropertyName, const ANSICHAR* Sp
 	if (Result == nullptr)
 	{
 		const ISlateStyle* ReportingStyle = RequestingStyle != nullptr ? RequestingStyle : this;
-		ReportingStyle->LogMissingResource(ISlateStyle::Warning, FText::Format(NSLOCTEXT("SlateStyleSet", "UknownSlateVector", "Unable to find vector property '{0}' in style."), FText::FromName(PropertyName)), PropertyName);
+		ReportingStyle->LogMissingResource(EStyleMessageSeverity::Warning, FText::Format(NSLOCTEXT("SlateStyleSet", "UknownSlateVector", "Unable to find vector property '{0}' in style."), FText::FromName(PropertyName)), PropertyName);
 	}
 
 	return Result ? *Result : DefaultValue;
@@ -232,7 +232,7 @@ const FLinearColor& FSlateStyleSet::GetColor(const FName PropertyName, const ANS
 	if ( Result == nullptr) 
 	{
 		const ISlateStyle* ReportingStyle = RequestingStyle != nullptr ? RequestingStyle : this;
-		ReportingStyle->LogMissingResource(ISlateStyle::Warning, FText::Format(NSLOCTEXT("SlateStyleSet", "UknownColor", "Unable to find Color '{0}'."), FText::FromName(LookupName)), LookupName);
+		ReportingStyle->LogMissingResource(EStyleMessageSeverity::Warning, FText::Format(NSLOCTEXT("SlateStyleSet", "UknownColor", "Unable to find Color '{0}'."), FText::FromName(LookupName)), LookupName);
 	}
 
 	return Result ? *Result : DefaultValue;
@@ -260,7 +260,7 @@ const FSlateColor FSlateStyleSet::GetSlateColor(const FName PropertyName, const
 	if ( Result == nullptr && LinearResult == nullptr)
 	{
 		const ISlateStyle* ReportingStyle = RequestingStyle != nullptr ? RequestingStyle : this;
-		ReportingStyle->LogMissingResource(ISlateStyle::Warning, FText::Format(NSLOCTEXT("SlateStyleSet", "UknownSlateColor", "Unable to find SlateColor '{0}'."), FText::FromName(StyleName)), StyleName);
+		ReportingStyle->LogMissingResource(EStyleMessageSeverity::Warning, FText::Format(NSLOCTEXT("SlateStyleSet", "UknownSlateColor", "Unable to find SlateColor '{0}'."), FText::FromName(StyleName)), StyleName);
 	}
 
 	return Result ? *Result : LinearResult ? *LinearResult : DefaultValue;
@@ -282,7 +282,7 @@ const FMargin& FSlateStyleSet::GetMargin(const FName PropertyName, const ANSICHA
 	if ( Result == nullptr )
 	{
 		const ISlateStyle* ReportingStyle = RequestingStyle != nullptr ? RequestingStyle : this;
-		ReportingStyle->LogMissingResource(ISlateStyle::Warning, FText::Format(NSLOCTEXT("SlateStyleSet", "UknownMargin", "Unable to find Margin '{0}'."), FText::FromName(StyleName)), StyleName);
+		ReportingStyle->LogMissingResource(EStyleMessageSeverity::Warning, FText::Format(NSLOCTEXT("SlateStyleSet", "UknownMargin", "Unable to find Margin '{0}'."), FText::FromName(StyleName)), StyleName);
 	}
 
 	return Result ? *Result : DefaultValue;
@@ -315,7 +315,7 @@ const FSlateBrush* FSlateStyleSet::GetBrush(const FName PropertyName, const ANSI
 	if ( Result == nullptr )
 	{
 		const ISlateStyle* ReportingStyle = RequestingStyle != nullptr ? RequestingStyle : this;
-		ReportingStyle->LogMissingResource(ISlateStyle::Warning, FText::Format(NSLOCTEXT("SlateStyleSet", "UknownBrush", "Unable to find Brush '{0}'."), FText::FromName(StyleName)), StyleName);
+		ReportingStyle->LogMissingResource(EStyleMessageSeverity::Warning, FText::Format(NSLOCTEXT("SlateStyleSet", "UknownBrush", "Unable to find Brush '{0}'."), FText::FromName(StyleName)), StyleName);
 	}
 
 	return Result ? Result : GetDefaultBrush();
@@ -410,7 +410,7 @@ const FSlateSound& FSlateStyleSet::GetSound(const FName PropertyName, const ANSI
 	if ( Result == nullptr )
 	{
 		const ISlateStyle* ReportingStyle = RequestingStyle != nullptr ? RequestingStyle : this;
-		ReportingStyle->LogMissingResource(ISlateStyle::Warning, FText::Format(NSLOCTEXT("SlateStyleSet", "UknownSound", "Unable to find Sound '{0}'."), FText::FromName(StyleName)), StyleName);
+		ReportingStyle->LogMissingResource(EStyleMessageSeverity::Warning, FText::Format(NSLOCTEXT("SlateStyleSet", "UknownSound", "Unable to find Sound '{0}'."), FText::FromName(StyleName)), StyleName);
 	}
 
 	return Result ? *Result : FStyleDefaults::GetSound();
@@ -450,7 +450,7 @@ const FSlateWidgetStyle* FSlateStyleSet::GetWidgetStyleInternal(const FName Desi
 	{
 		if (bWarnIfNotFound)
 		{
-			Log(ISlateStyle::Warning, FText::Format(NSLOCTEXT("SlateStyleSet", "UnknownWidgetStyle", "Unable to find Slate Widget Style '{0}'. Using {1} defaults instead."), FText::FromName(StyleName), FText::FromName(DesiredTypeName)));
+			Log(EStyleMessageSeverity::Warning, FText::Format(NSLOCTEXT("SlateStyleSet", "UnknownWidgetStyle", "Unable to find Slate Widget Style '{0}'. Using {1} defaults instead."), FText::FromName(StyleName), FText::FromName(DesiredTypeName)));
 		}
 
 		return DefaultStyle;
@@ -462,7 +462,7 @@ const FSlateWidgetStyle* FSlateStyleSet::GetWidgetStyleInternal(const FName Desi
 	{
 		if (bWarnIfNotFound)
 		{
-			Log(ISlateStyle::Error, FText::Format(NSLOCTEXT("SlateStyleSet", "WrongWidgetStyleType", "The Slate Widget Style '{0}' is not of the desired type. Desired: '{1}', Actual: '{2}'"), FText::FromName(StyleName), FText::FromName(DesiredTypeName), FText::FromName(Style->GetTypeName())));
+			Log(EStyleMessageSeverity::Error, FText::Format(NSLOCTEXT("SlateStyleSet", "WrongWidgetStyleType", "The Slate Widget Style '{0}' is not of the desired type. Desired: '{1}', Actual: '{2}'"), FText::FromName(StyleName), FText::FromName(DesiredTypeName), FText::FromName(Style->GetTypeName())));
 		}
 
 		return nullptr;
@@ -536,19 +536,19 @@ TSet<FName> FSlateStyleSet::GetStyleKeys() const
 
 void FSlateStyleSet::Log(ISlateStyle::EStyleMessageSeverity Severity, const FText& Message) const
 {
-	if ( Severity == ISlateStyle::Error )
+	if ( Severity == EStyleMessageSeverity::Error )
 	{
 		UE_LOG(LogSlateStyle, Error, TEXT("%s"), *Message.ToString());
 	}
-	else if ( Severity == ISlateStyle::PerformanceWarning )
+	else if ( Severity == EStyleMessageSeverity::PerformanceWarning )
 	{
 		UE_LOG(LogSlateStyle, Warning, TEXT("%s"), *Message.ToString());
 	}
-	else if ( Severity == ISlateStyle::Warning )
+	else if ( Severity == EStyleMessageSeverity::Warning )
 	{
 		UE_LOG(LogSlateStyle, Warning, TEXT("%s"), *Message.ToString());
 	}
-	else if ( Severity == ISlateStyle::Info )
+	else if ( Severity == EStyleMessageSeverity::Info )
 	{
 		UE_LOG(LogSlateStyle, Log, TEXT("%s"), *Message.ToString());
 	}
@@ -634,7 +634,7 @@ void FSlateStyleSet::LogUnusedBrushResources()
 
 		if ( !IsUsed )
 		{
-			Log(ISlateStyle::Warning, FText::FromString(FilePath));
+			Log(EStyleMessageSeverity::Warning, FText::FromString(FilePath));
 		}
 	}
 }
diff --git a/Engine/Source/Runtime/SlateCore/Public/Styling/ISlateStyle.h b/Engine/Source/Runtime/SlateCore/Public/Styling/ISlateStyle.h
--- a/Engine/Source/Runtime/SlateCore/Public/Styling/ISlateStyle.h
+++ b/Engine/Source/Runtime/SlateCore/Public/Styling/ISlateStyle.h
@@ -171,7 +171,7 @@ protected:
 	 * The severity of the message type 
 	 * Ordered according to their severity
 	 */
-	enum EStyleMessageSeverity
+	enum class EStyleMessageSeverity : uint8
 	{
 		CriticalError		= 0,
 		Error				= 1,
