Missing style changes for notifications

diff --git a/Engine/Source/Runtime/SlateCore/Private/Styling/StarshipCoreStyle.cpp b/Engine/Source/Runtime/SlateCore/Private/Styling/StarshipCoreStyle.cpp
--- a/Engine/Source/Runtime/SlateCore/Private/Styling/StarshipCoreStyle.cpp
+++ b/Engine/Source/Runtime/SlateCore/Private/Styling/StarshipCoreStyle.cpp
@@ -615,8 +615,8 @@ TSharedRef<ISlateStyle> FStarshipCoreStyle::Create()
 		Style->Set("NotificationList.FontBold", DEFAULT_FONT("Bold", 12));
 		Style->Set("NotificationList.FontLight", DEFAULT_FONT("Light", 12));
 		Style->Set("NotificationList.ItemBackground", new FSlateColorBrush(FStyleColors::Notifications));
-		
-		Style->Set("NotificationList.ItemBackground_Border", new FSlateRoundedBoxBrush(FStyleColors::Transparent, 5.0f, FStyleColors::Black, 1.0f));
+
+		Style->Set("NotificationList.ItemBackground_Border", new FSlateRoundedBoxBrush(FStyleColors::Notifications, 5.0f, FStyleColors::Black, 1.0f));
 		Style->Set("NotificationList.ItemBackground_Border_Transparent", new BOX_BRUSH("Old/Notification_Border_Flash", FMargin(8.0f / 64.0f)));
 
 		Style->Set("NotificationList.Throbber", new IMAGE_BRUSH("Starship/Notifications/Throbber", FVector2D(40, 40)));
