Slate: Add option to prevent Slate Attributes from being called automatically on the SWidget. It is useful when the owner of the widget controls the invalidation (ie. ScrollBox).
Add PrePass invalidaiton. It implies a Layout and replaces the InvalidatePrepass function.

diff --git a/Engine/Plugins/Enterprise/DataprepEditor/Source/DataprepEditor/Private/Widgets/DataprepWidgets.cpp b/Engine/Plugins/Enterprise/DataprepEditor/Source/DataprepEditor/Private/Widgets/DataprepWidgets.cpp
--- a/Engine/Plugins/Enterprise/DataprepEditor/Source/DataprepEditor/Private/Widgets/DataprepWidgets.cpp
+++ b/Engine/Plugins/Enterprise/DataprepEditor/Source/DataprepEditor/Private/Widgets/DataprepWidgets.cpp
@@ -409,7 +409,7 @@ void SDataprepDetailsView::OnObjectReplaced(const TMap<UObject*, UObject*>& Repl
 void SDataprepDetailsView::ForceRefresh()
 {
 	// ueent_hotfix Hotfix for 4.24 (Remove the ui flickering)
-	InvalidatePrepass();
+	Invalidate(EInvalidateWidgetReason::Prepass);
 	bRefreshObjectToDisplay = true;
 }
 
diff --git a/Engine/Source/Runtime/Engine/Private/Slate/SGameLayerManager.cpp b/Engine/Source/Runtime/Engine/Private/Slate/SGameLayerManager.cpp
--- a/Engine/Source/Runtime/Engine/Private/Slate/SGameLayerManager.cpp
+++ b/Engine/Source/Runtime/Engine/Private/Slate/SGameLayerManager.cpp
@@ -497,8 +497,7 @@ void SGameLayerManager::AddOrUpdatePlayerLayers(const FGeometry& AllottedGeometr
 
 	if (CachedInverseDPIScale != InverseDPIScale)
 	{
-		InvalidatePrepass();
-		Invalidate(EInvalidateWidget::Layout);
+		Invalidate(EInvalidateWidget::Prepass);
 		CachedInverseDPIScale = InverseDPIScale;
 	}
 
diff --git a/Engine/Source/Runtime/Slate/Private/Widgets/Layout/SDPIScaler.cpp b/Engine/Source/Runtime/Slate/Private/Widgets/Layout/SDPIScaler.cpp
--- a/Engine/Source/Runtime/Slate/Private/Widgets/Layout/SDPIScaler.cpp
+++ b/Engine/Source/Runtime/Slate/Private/Widgets/Layout/SDPIScaler.cpp
@@ -6,7 +6,7 @@
 SLATE_IMPLEMENT_WIDGET(SDPIScaler)
 void SDPIScaler::PrivateRegisterAttributes(FSlateAttributeInitializer& AttributeInitializer)
 {
-	SLATE_ADD_MEMBER_ATTRIBUTE_DEFINITION_WITH_NAME(AttributeInitializer, "DPIScale", DPIScaleAttribute, EInvalidateWidgetReason::Layout);
+	SLATE_ADD_MEMBER_ATTRIBUTE_DEFINITION_WITH_NAME(AttributeInitializer, "DPIScale", DPIScaleAttribute, EInvalidateWidgetReason::Prepass);
 }
 
 SDPIScaler::SDPIScaler()
@@ -70,10 +70,7 @@ void SDPIScaler::SetContent(TSharedRef<SWidget> InContent)
 
 void SDPIScaler::SetDPIScale(TAttribute<float> InDPIScale)
 {
-	if (DPIScaleAttribute.Assign(*this, MoveTemp(InDPIScale), 1.f))
-	{
-		InvalidatePrepass();
-	}
+	DPIScaleAttribute.Assign(*this, MoveTemp(InDPIScale), 1.f);
 }
 
 float SDPIScaler::GetRelativeLayoutScale(int32 ChildIndex, float LayoutScaleMultiplier) const
diff --git a/Engine/Source/Runtime/Slate/Private/Widgets/Layout/SScaleBox.cpp b/Engine/Source/Runtime/Slate/Private/Widgets/Layout/SScaleBox.cpp
--- a/Engine/Source/Runtime/Slate/Private/Widgets/Layout/SScaleBox.cpp
+++ b/Engine/Source/Runtime/Slate/Private/Widgets/Layout/SScaleBox.cpp
@@ -89,8 +89,7 @@ bool SScaleBox::CustomPrepass(float LayoutScaleMultiplier)
 
 	if (bNeedsNormalizingPrepassOrLocalGeometry)
 	{
-		ChildSlotWidget.Invalidate(EInvalidateWidgetReason::Layout);
-		ChildSlotWidget.InvalidatePrepass();
+		ChildSlotWidget.Invalidate(EInvalidateWidgetReason::Prepass);
 	}
 
 	// Extract the incoming scale out of the layout scale if 
@@ -268,8 +267,7 @@ int32 SScaleBox::OnPaint( const FPaintArgs& Args, const FGeometry& AllottedGeome
 
 		if (DoesScaleRequireNormalizingPrepassOrLocalGeometry())
 		{
-			const_cast<SScaleBox*>(this)->Invalidate(EInvalidateWidgetReason::Layout);
-			const_cast<SScaleBox*>(this)->InvalidatePrepass();
+			const_cast<SScaleBox*>(this)->Invalidate(EInvalidateWidgetReason::Prepass);
 		}
 	}
 
@@ -337,8 +335,7 @@ void SScaleBox::SetStretchDirection(EStretchDirection::Type InStretchDirection)
 {
 	if (SetAttribute(StretchDirection, TAttribute<EStretchDirection::Type>(InStretchDirection), EInvalidateWidgetReason::Layout))
 	{
-		Invalidate(EInvalidateWidgetReason::Layout);
-		InvalidatePrepass();
+		Invalidate(EInvalidateWidgetReason::Prepass);
 	}
 }
 
@@ -354,10 +351,9 @@ void SScaleBox::SetStretch(EStretch::Type InStretch)
 
 void SScaleBox::SetUserSpecifiedScale(float InUserSpecifiedScale)
 {
-	if (SetAttribute(UserSpecifiedScale, TAttribute<float>(InUserSpecifiedScale), EInvalidateWidgetReason::Layout))
+	if (SetAttribute(UserSpecifiedScale, TAttribute<float>(InUserSpecifiedScale), EInvalidateWidgetReason::Prepass))
 	{
-		Invalidate(EInvalidateWidgetReason::Layout);
-		InvalidatePrepass();
+		Invalidate(EInvalidateWidgetReason::Prepass);
 	}
 }
 
@@ -365,7 +361,7 @@ void SScaleBox::SetIgnoreInheritedScale(bool InIgnoreInheritedScale)
 {
 	if (SetAttribute(IgnoreInheritedScale, TAttribute<bool>(InIgnoreInheritedScale), EInvalidateWidgetReason::Layout))
 	{
-		InvalidatePrepass();
+		Invalidate(EInvalidateWidgetReason::Prepass);
 	}
 }
 
@@ -459,8 +455,7 @@ void SScaleBox::RefreshSafeZoneScale()
 
 	SafeZoneScale = 1.f - ScaleDownBy;
 
-	Invalidate(EInvalidateWidgetReason::Layout);
-	InvalidatePrepass();
+	Invalidate(EInvalidateWidgetReason::Prepass);
 }
 
 #if WITH_EDITOR
diff --git a/Engine/Source/Runtime/Slate/Private/Widgets/Text/SRichTextBlock.cpp b/Engine/Source/Runtime/Slate/Private/Widgets/Text/SRichTextBlock.cpp
--- a/Engine/Source/Runtime/Slate/Private/Widgets/Text/SRichTextBlock.cpp
+++ b/Engine/Source/Runtime/Slate/Private/Widgets/Text/SRichTextBlock.cpp
@@ -107,8 +107,7 @@ void SRichTextBlock::OnArrangeChildren(const FGeometry& AllottedGeometry, FArran
 void SRichTextBlock::SetText( const TAttribute<FText>& InTextAttr )
 {
 	BoundText = InTextAttr;
-	Invalidate(EInvalidateWidget::LayoutAndVolatility);
-	InvalidatePrepass();
+	Invalidate(EInvalidateWidget::LayoutAndVolatility|EInvalidateWidgetReason::Prepass);
 }
 
 void SRichTextBlock::SetHighlightText( const TAttribute<FText>& InHighlightText )
@@ -136,8 +135,7 @@ void SRichTextBlock::SetWrapTextAt(const TAttribute<float>& InWrapTextAt)
 
 void SRichTextBlock::SetAutoWrapText(const TAttribute<bool>& InAutoWrapText)
 {
-	SetAttribute(AutoWrapText, InAutoWrapText, EInvalidateWidgetReason::Layout);
-	InvalidatePrepass();
+	SetAttribute(AutoWrapText, InAutoWrapText, EInvalidateWidgetReason::Prepass);
 }
 
 void SRichTextBlock::SetWrappingPolicy(const TAttribute<ETextWrappingPolicy>& InWrappingPolicy)
@@ -188,8 +186,7 @@ void SRichTextBlock::SetDecoratorStyleSet(const ISlateStyle* NewDecoratorStyleSe
 void SRichTextBlock::SetTextBlockScale(const float NewTextBlockScale)
 {
 	TextBlockScale = NewTextBlockScale;
-	Invalidate(EInvalidateWidget::Layout);
-	InvalidatePrepass();
+	Invalidate(EInvalidateWidget::Prepass);
 }
 
 void SRichTextBlock::Refresh()
diff --git a/Engine/Source/Runtime/Slate/Public/Widgets/Views/SListView.h b/Engine/Source/Runtime/Slate/Public/Widgets/Views/SListView.h
--- a/Engine/Source/Runtime/Slate/Public/Widgets/Views/SListView.h
+++ b/Engine/Source/Runtime/Slate/Public/Widgets/Views/SListView.h
@@ -1318,7 +1318,7 @@ public:
 
 		// We rely on the widgets desired size in order to determine how many will fit on screen.
 		const TSharedRef<SWidget> NewlyGeneratedWidget = WidgetForItem->AsWidget();
-		NewlyGeneratedWidget->InvalidatePrepass();
+		NewlyGeneratedWidget->MarkPrepassAsDirty();
 		NewlyGeneratedWidget->SlatePrepass(LayoutScaleMultiplier);
 
 		// We have a widget for this item; add it to the panel so that it is part of the UI.
diff --git a/Engine/Source/Runtime/SlateCore/Private/FastUpdate/SlateInvalidationRoot.cpp b/Engine/Source/Runtime/SlateCore/Private/FastUpdate/SlateInvalidationRoot.cpp
--- a/Engine/Source/Runtime/SlateCore/Private/FastUpdate/SlateInvalidationRoot.cpp
+++ b/Engine/Source/Runtime/SlateCore/Private/FastUpdate/SlateInvalidationRoot.cpp
@@ -238,8 +238,7 @@ void FSlateInvalidationRoot::InvalidateRootChildOrder(const SWidget* Investigato
 
 	// Invalidate all proxy handles
 	FastWidgetPathList->Reset();
-	InvalidationRootWidget->InvalidatePrepass();
-	InvalidationRootWidget->Invalidate(EInvalidateWidgetReason::Layout);
+	InvalidationRootWidget->Invalidate(EInvalidateWidgetReason::Prepass);
 	bNeedsSlowPath = true;
 
 #if WITH_SLATE_DEBUGGING
@@ -250,8 +249,7 @@ void FSlateInvalidationRoot::InvalidateRootChildOrder(const SWidget* Investigato
 
 void FSlateInvalidationRoot::InvalidateRootLayout(const SWidget* Investigator)
 {
-	InvalidationRootWidget->InvalidatePrepass();
-	InvalidationRootWidget->Invalidate(EInvalidateWidgetReason::Layout);
+	InvalidationRootWidget->Invalidate(EInvalidateWidgetReason::Prepass);
 	bNeedsSlowPath = true;
 
 #if WITH_SLATE_DEBUGGING
@@ -276,7 +274,7 @@ void FSlateInvalidationRoot::InvalidateWidget(FWidgetProxy& Proxy, EInvalidateWi
 				bChildOrderInvalidated = true;
 				if (!InvalidationRootWidget->Advanced_IsWindow())
 				{
-					InvalidationRootWidget->InvalidatePrepass();
+					InvalidationRootWidget->MarkPrepassAsDirty();
 				}
 
 				if (!GSlateEnableGlobalInvalidation && !InvalidationRootWidget->Advanced_IsWindow())
@@ -521,8 +519,7 @@ void FSlateInvalidationRoot::ProcessPreUpdate()
 
 		// Add the root to the update list (to prepass and paint it)
 		check(RootWidget->GetProxyHandle().IsValid(&RootWidget.Get()));
-		RootWidget->InvalidatePrepass();
-		InvalidateWidget(RootWidget->FastPathProxyHandle.GetProxy(), EInvalidateWidgetReason::Layout);
+		RootWidget->Invalidate(EInvalidateWidgetReason::Prepass);
 	}
 	else
 	{
@@ -660,7 +657,7 @@ void FSlateInvalidationRoot::ProcessPreUpdate()
 #endif
 
 							FastWidgetPathList->ProcessAttributeRegistrationInvalidation(InvalidationWidget);
-							if (WidgetPtr->HasRegisteredSlateAttribute())
+							if (WidgetPtr->HasRegisteredSlateAttribute() && WidgetPtr->IsAttributesUpdatesEnabled())
 							{
 								// This should be the next Attribute to update
 								AttributeItt.Seek(InvalidationWidget.Index);
@@ -963,8 +960,6 @@ void FSlateInvalidationRoot::Advanced_ResetInvalidation(bool bClearResourcesImme
 {
 	InvalidateRootChildOrder();
 
-	InvalidationRootWidget->InvalidatePrepass();
-
 	if (bClearResourcesImmediately)
 	{
 		ClearAllFastPathData(true);
diff --git a/Engine/Source/Runtime/SlateCore/Private/FastUpdate/SlateInvalidationWidgetList.cpp b/Engine/Source/Runtime/SlateCore/Private/FastUpdate/SlateInvalidationWidgetList.cpp
--- a/Engine/Source/Runtime/SlateCore/Private/FastUpdate/SlateInvalidationWidgetList.cpp
+++ b/Engine/Source/Runtime/SlateCore/Private/FastUpdate/SlateInvalidationWidgetList.cpp
@@ -268,7 +268,7 @@ FSlateInvalidationWidgetIndex FSlateInvalidationWidgetList::_BuildWidgetList_Rec
 	bool bParentAndSelfVisible = bParentVisible && Visibility.IsVisible();
 
 	bool bUpdateSlateAttribute = false;
-	if (Widget->HasRegisteredSlateAttribute())
+	if (Widget->HasRegisteredSlateAttribute() && Widget->IsAttributesUpdatesEnabled())
 	{
 
 		// The list is already sorted at this point. Add to the end.
@@ -276,10 +276,9 @@ FSlateInvalidationWidgetIndex FSlateInvalidationWidgetList::_BuildWidgetList_Rec
 
 		// Update attributes now because the visibility flag may change
 		{
-			Widget->bPauseAttributeInvalidation = true;
 			if (bParentAndSelfVisible)
 			{
-				FSlateAttributeMetaData::UpdateAttributes(Widget.Get());
+				FSlateAttributeMetaData::UpdateAttributes(Widget.Get(), false);
 				bUpdateSlateAttribute = true;
 
 				Visibility = Widget->GetVisibility();
@@ -287,19 +286,18 @@ FSlateInvalidationWidgetIndex FSlateInvalidationWidgetList::_BuildWidgetList_Rec
 			}
 			else if (bParentVisible)
 			{
-				FSlateAttributeMetaData::UpdateCollapsedAttributes(Widget.Get());
+				FSlateAttributeMetaData::UpdateCollapsedAttributes(Widget.Get(), false);
 				bUpdateSlateAttribute = true;
 
 				Visibility = Widget->GetVisibility();
 				bParentAndSelfVisible = bParentVisible && Visibility.IsVisible();
 				if (bParentAndSelfVisible)
 				{
-					FSlateAttributeMetaData::UpdateAttributes(Widget.Get());
+					FSlateAttributeMetaData::UpdateAttributes(Widget.Get(), false);
 					Visibility = Widget->GetVisibility();
 					bParentAndSelfVisible = bParentVisible && Visibility.IsVisible();
 				}
 			}
-			Widget->bPauseAttributeInvalidation = false;
 		}
 	}
 
@@ -580,7 +578,7 @@ void FSlateInvalidationWidgetList::ProcessAttributeRegistrationInvalidation(cons
 	SWidget* WidgetPtr = InvalidationWidget.GetWidget();
 	check(WidgetPtr);
 
-	if (WidgetPtr->HasRegisteredSlateAttribute())
+	if (WidgetPtr->HasRegisteredSlateAttribute() && WidgetPtr->IsAttributesUpdatesEnabled())
 	{
 		Data[InvalidationWidget.Index.ArrayIndex].ElementIndexList_WidgetWithRegisteredSlateAttribute.InsertUnique(InvalidationWidget.Index.ElementIndex);
 	}
@@ -1274,7 +1272,7 @@ FSlateInvalidationWidgetList::FCutResult FSlateInvalidationWidgetList::_CutArray
 				Widget->SetFastPathProxyHandle(FWidgetProxyHandle{ Owner, NewWidgetIndex, SortIndex, GenerationNumber });
 
 				// Check for ElementIndexList
-				if (Widget->HasRegisteredSlateAttribute())
+				if (Widget->HasRegisteredSlateAttribute() && Widget->IsAttributesUpdatesEnabled())
 				{
 					Data[NewArrayIndex].ElementIndexList_WidgetWithRegisteredSlateAttribute.AddUnsorted(NewElementIndex);
 				}
@@ -1592,7 +1590,7 @@ bool FSlateInvalidationWidgetList::VerifyElementIndexList() const
 		int32 Index = 0;
 		const_cast<FSlateInvalidationWidgetList*>(this)->ForEachWidget([&](const SWidget* Widget)
 			{
-				if (Widget->HasRegisteredSlateAttribute())
+				if (Widget->HasRegisteredSlateAttribute() && Widget->IsAttributesUpdatesEnabled())
 				{
 					WidgetListWithForEachWidget.Add(Widget);
 				}
@@ -1623,7 +1621,8 @@ bool FSlateInvalidationWidgetList::VerifyElementIndexList() const
 				const FSlateInvalidationWidgetList::InvalidationWidgetType& InvalidationWidget = ArrayNode.ElementList[ElementIndex];
 				if (const SWidget* Widget = InvalidationWidget.GetWidget())
 				{
-					if (!Widget->HasRegisteredSlateAttribute())
+					const bool bShouldBeThere = Widget->HasRegisteredSlateAttribute() && Widget->IsAttributesUpdatesEnabled();
+					if (!bShouldBeThere)
 					{
 						UE_LOG(LogSlate, Warning, TEXT("ElementIndex '%d' in the array of sort value '%d' should not be there since widget '%s' doesn't have registered SlateAttribute.")
 							, ElementIndex, ArrayNode.SortOrder
diff --git a/Engine/Source/Runtime/SlateCore/Private/FastUpdate/WidgetProxy.cpp b/Engine/Source/Runtime/SlateCore/Private/FastUpdate/WidgetProxy.cpp
--- a/Engine/Source/Runtime/SlateCore/Private/FastUpdate/WidgetProxy.cpp
+++ b/Engine/Source/Runtime/SlateCore/Private/FastUpdate/WidgetProxy.cpp
@@ -103,7 +103,7 @@ bool FWidgetProxy::ProcessPostInvalidation(FSlateInvalidationWidgetPostHeap& Upd
 		FWidgetProxy& ParentProxy = FastWidgetPathList[ParentIndex];
 		if (SWidget* ParentWidgetPtr = ParentProxy.GetWidget())
 		{
-			ParentWidgetPtr->InvalidatePrepass();
+			ParentWidgetPtr->MarkPrepassAsDirty();
 			ParentProxy.CurrentInvalidateReason |= EInvalidateWidgetReason::Layout;
 #if WITH_SLATE_DEBUGGING
 			FSlateDebugging::BroadcastWidgetInvalidate(ParentWidgetPtr, WidgetPtr, EInvalidateWidgetReason::Layout);
@@ -231,7 +231,7 @@ int32 FWidgetProxy::Repaint(const FPaintArgs& PaintArgs, FSlateWindowElementList
 		{
 			if (WidgetPtr->ShouldInvalidatePrepassDueToVolatility())
 			{
-				WidgetPtr->InvalidatePrepass();
+				WidgetPtr->MarkPrepassAsDirty();
 			}
 			WidgetPtr->SlatePrepass(WidgetPtr->GetPrepassLayoutScaleMultiplier());
 		}
diff --git a/Engine/Source/Runtime/SlateCore/Private/Test/SlateAttributeTest.cpp b/Engine/Source/Runtime/SlateCore/Private/Test/SlateAttributeTest.cpp
--- a/Engine/Source/Runtime/SlateCore/Private/Test/SlateAttributeTest.cpp
+++ b/Engine/Source/Runtime/SlateCore/Private/Test/SlateAttributeTest.cpp
@@ -229,7 +229,7 @@ bool FSlateAttributeTest::RunTest(const FString& Parameters)
 			OrderCounter = 0;
 			bWasUpdate = false;
 			ReturnValue = 4;
-			WidgetParent->InvalidatePrepass();
+			WidgetParent->MarkPrepassAsDirty();
 			WidgetParent->SlatePrepass(1.f);
 			AddErrorIfFalse(WidgetParent->IntAttributeA.Get() == 2, TEXT("A It is not the expected value."));
 			AddErrorIfFalse(WidgetParent->IntAttributeB.Get() == 1, TEXT("B It is not the expected value."));
@@ -244,7 +244,7 @@ bool FSlateAttributeTest::RunTest(const FString& Parameters)
 			WidgetParent->IntAttributeC.Assign(WidgetParent.Get(), MakeAttributeLambda(UpdateLambda));
 			AddErrorIfFalse(!bWasUpdate, TEXT("C should not have be updated."));
 			AddErrorIfFalse(WidgetParent->IntAttributeC.Get() == 4, TEXT("C It is not the expected value."));
-			WidgetParent->InvalidatePrepass();
+			WidgetParent->MarkPrepassAsDirty();
 			WidgetParent->SlatePrepass(1.f);
 			AddErrorIfFalse(bWasUpdate, TEXT("C should be updated."));
 			AddErrorIfFalse(WidgetParent->IntAttributeC.Get() == 5, TEXT("C It is not the expected value."));
@@ -254,7 +254,7 @@ bool FSlateAttributeTest::RunTest(const FString& Parameters)
 			OrderCounter = 0;
 			bWasUpdate = false;
 			ReturnValue = 10;	//10 show that C didn't change
-			WidgetParent->InvalidatePrepass();
+			WidgetParent->MarkPrepassAsDirty();
 			WidgetParent->SlatePrepass(1.f);
 			AddErrorIfFalse(WidgetParent->IntAttributeA.Get() == 2, TEXT("A It is not the expected value."));
 			AddErrorIfFalse(WidgetParent->IntAttributeB.Get() == 1, TEXT("B It is not the expected value."));
@@ -271,7 +271,7 @@ bool FSlateAttributeTest::RunTest(const FString& Parameters)
 			OrderCounter = 0;
 			bWasUpdate = false;
 			ReturnValue = 10;
-			WidgetParent->InvalidatePrepass();
+			WidgetParent->MarkPrepassAsDirty();
 			WidgetParent->SlatePrepass(1.f);
 			AddErrorIfFalse(WidgetParent->IntAttributeA.Get() == 2, TEXT("A It is not the expected value."));
 			AddErrorIfFalse(WidgetParent->IntAttributeB.Get() == 1, TEXT("B It is not the expected value."));
@@ -287,7 +287,7 @@ bool FSlateAttributeTest::RunTest(const FString& Parameters)
 			OrderCounter = 0;
 			bWasUpdate = false;
 			ReturnValue = 10;
-			WidgetParent->InvalidatePrepass();
+			WidgetParent->MarkPrepassAsDirty();
 			WidgetParent->SlatePrepass(1.f);
 			AddErrorIfFalse(!bWasUpdate, TEXT("C should not be updated."));
 		}
@@ -364,7 +364,7 @@ bool FSlateAttributeTest::RunTest(const FString& Parameters)
 			OrderCounter = 0;
 			bWasUpdate = false;
 			ReturnValue = 4;
-			WidgetChild->InvalidatePrepass();
+			WidgetChild->MarkPrepassAsDirty();
 			WidgetChild->SlatePrepass(1.f);
 			AddErrorIfFalse(WidgetChild->IntAttributeA.Get() == 2 || WidgetChild->IntAttributeA.Get() == 3, TEXT("A It is not the expected value."));
 			AddErrorIfFalse(WidgetChild->IntAttributeB.Get() == 1, TEXT("B It is not the expected value."));
@@ -381,7 +381,7 @@ bool FSlateAttributeTest::RunTest(const FString& Parameters)
 			OrderCounter = 0;
 			bWasUpdate = false;
 			ReturnValue = 4;
-			WidgetChild->InvalidatePrepass();
+			WidgetChild->MarkPrepassAsDirty();
 			WidgetChild->SlatePrepass(1.f);
 			AddErrorIfFalse(WidgetChild->IntAttributeA.Get() == 2 || WidgetChild->IntAttributeA.Get() == 3, TEXT("A It is not the expected value."));
 			AddErrorIfFalse(WidgetChild->IntAttributeB.Get() == 1, TEXT("B It is not the expected value."));
diff --git a/Engine/Source/Runtime/SlateCore/Private/Types/SlateAttributeMetaData.cpp b/Engine/Source/Runtime/SlateCore/Private/Types/SlateAttributeMetaData.cpp
--- a/Engine/Source/Runtime/SlateCore/Private/Types/SlateAttributeMetaData.cpp
+++ b/Engine/Source/Runtime/SlateCore/Private/Types/SlateAttributeMetaData.cpp
@@ -53,7 +53,7 @@ void FSlateAttributeMetaData::RegisterAttribute(SWidget& OwningWidget, FSlateAtt
 		NewAttributeMetaData->RegisterAttributeImpl(OwningWidget, Attribute, AttributeType, MoveTemp(Wrapper));
 		OwningWidget.bHasRegisteredSlateAttribute = true;
 		OwningWidget.MetaData.Insert(NewAttributeMetaData, 0);
-		if (!OwningWidget.bPauseAttributeInvalidation)
+		if (OwningWidget.IsConstructionCompleted() && OwningWidget.IsAttributesUpdatesEnabled())
 		{
 			OwningWidget.Invalidate(EInvalidateWidgetReason::AttributeRegistration);
 		}
@@ -136,7 +136,7 @@ bool FSlateAttributeMetaData::UnregisterAttribute(SWidget& OwningWidget, const F
 			check(bResult); // if the num is 0 then we should have remove an item.
 			OwningWidget.bHasRegisteredSlateAttribute = false;
 			OwningWidget.MetaData.RemoveAtSwap(0);
-			if (!OwningWidget.bPauseAttributeInvalidation)
+			if (OwningWidget.IsConstructionCompleted() && OwningWidget.IsAttributesUpdatesEnabled())
 			{
 				OwningWidget.Invalidate(EInvalidateWidgetReason::AttributeRegistration);
 			}
@@ -181,7 +181,7 @@ void FSlateAttributeMetaData::InvalidateWidget(SWidget& OwningWidget, const FSla
 	//N.B. no needs to set the bUpatedManually in this case because
 	//	1. they are in construction, so they will all be called anyway
 	//	2. they are in WidgetList, so the SlateAttribute.Set will not be called
-	if (OwningWidget.bPauseAttributeInvalidation)
+	if (!OwningWidget.IsConstructionCompleted())
 	{
 		return;
 	}
@@ -243,28 +243,40 @@ void FSlateAttributeMetaData::InvalidateWidget(SWidget& OwningWidget, const FSla
 
 
 void FSlateAttributeMetaData::UpdateAttributes(SWidget& OwningWidget)
+{
+	UpdateAttributes(OwningWidget, OwningWidget.IsConstructionCompleted());
+}
+
+
+void FSlateAttributeMetaData::UpdateCollapsedAttributes(SWidget& OwningWidget)
+{
+	UpdateCollapsedAttributes(OwningWidget, OwningWidget.IsConstructionCompleted());	
+}
+
+
+void FSlateAttributeMetaData::UpdateAttributes(SWidget& OwningWidget, bool bAllowInvalidation)
 {
 	if (FSlateAttributeMetaData* AttributeMetaData = FSlateAttributeMetaData::FindMetaData(OwningWidget))
 	{
-		AttributeMetaData->UpdateAttributes(OwningWidget, false);
+		AttributeMetaData->UpdateAttributes(OwningWidget, false, bAllowInvalidation);
 	}
 }
 
 
-void FSlateAttributeMetaData::UpdateCollapsedAttributes(SWidget& OwningWidget)
+void FSlateAttributeMetaData::UpdateCollapsedAttributes(SWidget& OwningWidget, bool bAllowInvalidation)
 {
 	if (FSlateAttributeMetaData* AttributeMetaData = FSlateAttributeMetaData::FindMetaData(OwningWidget))
 	{
 		// Does it have any collapsed attributes
 		if (AttributeMetaData->CollaspedAttributeCounter > 0)
 		{
-			AttributeMetaData->UpdateAttributes(OwningWidget, true);
+			AttributeMetaData->UpdateAttributes(OwningWidget, true, bAllowInvalidation);
 		}
 	}
 }
 
 
-void FSlateAttributeMetaData::UpdateAttributes(SWidget& OwningWidget, bool bOnlyCollapsed)
+void FSlateAttributeMetaData::UpdateAttributes(SWidget& OwningWidget, bool bOnlyCollapsed, bool bAllowInvalidation)
 {
 	EInvalidateWidgetReason InvalidationReason = EInvalidateWidgetReason::None;
 	for (int32 Index = 0; Index < Attributes.Num(); ++Index)
@@ -306,7 +318,7 @@ void FSlateAttributeMetaData::UpdateAttributes(SWidget& OwningWidget, bool bOnly
 		ISlateAttributeGetter::FUpdateAttributeResult Result = GetterItem.Getter->UpdateAttribute(OwningWidget);
 		GetterItem.bUpdatedOnce = true;
 		GetterItem.bUpdatedThisFrame = Result.bInvalidationRequested;
-		if (Result.bInvalidationRequested && !OwningWidget.bPauseAttributeInvalidation)
+		if (Result.bInvalidationRequested && bAllowInvalidation)
 		{
 			InvalidationReason |= GetterItem.GetInvalidationReason(OwningWidget, Result.InvalidationReason);
 		}
@@ -321,7 +333,7 @@ void FSlateAttributeMetaData::UpdateAttributes(SWidget& OwningWidget, bool bOnly
 	}
 	bHasUpdatedManuallyFlagToReset = false;
 
-	if (!OwningWidget.bPauseAttributeInvalidation)
+	if (bAllowInvalidation)
 	{
 		OwningWidget.Invalidate(InvalidationReason);
 	}
@@ -341,7 +353,7 @@ void FSlateAttributeMetaData::UpdateAttribute(SWidget& OwningWidget, FSlateAttri
 			ISlateAttributeGetter::FUpdateAttributeResult Result = GetterItem.Getter->UpdateAttribute(OwningWidget);
 			if (Result.bInvalidationRequested)
 			{
-				if (!OwningWidget.bPauseAttributeInvalidation)
+				if (OwningWidget.IsConstructionCompleted())
 				{
 					OwningWidget.Invalidate(GetterItem.GetInvalidationReason(OwningWidget, Result.InvalidationReason));
 				}
diff --git a/Engine/Source/Runtime/SlateCore/Private/Widgets/SWidget.cpp b/Engine/Source/Runtime/SlateCore/Private/Widgets/SWidget.cpp
--- a/Engine/Source/Runtime/SlateCore/Private/Widgets/SWidget.cpp
+++ b/Engine/Source/Runtime/SlateCore/Private/Widgets/SWidget.cpp
@@ -213,7 +213,8 @@ SWidget::SWidget()
 	, bInvisibleDueToParentOrSelfVisibility(false)
 	, bNeedsPrepass(true)
 	, bHasRegisteredSlateAttribute(false)
-	, bPauseAttributeInvalidation(true)
+	, bEnabledAttributesUpdate(true)
+	, bIsDeclarativeSyntaxConstructionCompleted(false)
 	, bHasCustomPrepass(false)
 	, bHasRelativeLayoutScale(false)
 	, bVolatilityAlwaysInvalidatesPrepass(false)
@@ -358,6 +359,7 @@ void SWidget::SWidgetConstruct(const FSlateBaseNamedArgs& Args)
 	bForceVolatile = Args._ForceVolatile;
 	Clipping = Args._Clipping;
 	FlowDirectionPreference = Args._FlowDirectionPreference;
+	bEnabledAttributesUpdate = Args._EnabledAttributesUpdate;
 	MetaData.Append(Args.MetaData);
 
 	if (Args._ToolTip.IsSet())
@@ -690,9 +692,7 @@ void SWidget::SlatePrepass(float InLayoutScaleMultiplier)
 
 void SWidget::InvalidatePrepass()
 {
-	SLATE_CROSS_THREAD_CHECK();
-
-	bNeedsPrepass = true;
+	MarkPrepassAsDirty();
 }
 
 void SWidget::InvalidateChildRemovedFromTree(SWidget& Child)
@@ -1195,7 +1195,13 @@ void SWidget::Invalidate(EInvalidateWidgetReason InvalidateReason)
 
 	if (EnumHasAnyFlags(InvalidateReason, EInvalidateWidgetReason::ChildOrder) || !PrepassLayoutScaleMultiplier.IsSet())
 	{
-		InvalidatePrepass();
+		MarkPrepassAsDirty();
+	}
+
+	if (EnumHasAnyFlags(InvalidateReason, EInvalidateWidgetReason::Prepass))
+	{
+		bNeedsPrepass = true;
+		InvalidateReason |= EInvalidateWidgetReason::Layout;
 	}
 
 	if(FastPathProxyHandle.IsValid(this))
@@ -1594,7 +1600,7 @@ void SWidget::Prepass_Internal(float InLayoutScaleMultiplier)
 {
 	PrepassLayoutScaleMultiplier = InLayoutScaleMultiplier;
 
-	if (HasRegisteredSlateAttribute() && !GSlateIsOnFastProcessInvalidation)
+	if (HasRegisteredSlateAttribute() && IsAttributesUpdatesEnabled() && !GSlateIsOnFastProcessInvalidation)
 	{
 		FSlateAttributeMetaData::UpdateAttributes(*this);
 	}
@@ -1633,28 +1639,33 @@ void SWidget::Prepass_ChildLoop(float InLayoutScaleMultiplier, FChildren* MyChil
 
 		const TSharedRef<SWidget>& Child = MyChildren->GetChildAt(ChildIndex);
 
-		if (Child->Visibility.Get() != EVisibility::Collapsed)
+		if (Child->GetVisibility() != EVisibility::Collapsed)
 		{
 			// Recur: Descend down the widget tree.
 			Child->Prepass_Internal(ChildLayoutScaleMultiplier);
 		}
 		else
 		{
-			if (HasRegisteredSlateAttribute() && !GSlateIsOnFastProcessInvalidation)
+			bool bIsCollapsed = true;
+			if (Child->HasRegisteredSlateAttribute() && Child->IsAttributesUpdatesEnabled() && !GSlateIsOnFastProcessInvalidation)
 			{
-				FSlateAttributeMetaData::UpdateCollapsedAttributes(*this);
-				if (Child->Visibility.Get() != EVisibility::Collapsed)
-				{
-					// Recur: Descend down the widget tree.
-					Child->Prepass_Internal(ChildLayoutScaleMultiplier);
-				}
+				FSlateAttributeMetaData::UpdateCollapsedAttributes(Child.Get());
+				bIsCollapsed = Child->GetVisibility() == EVisibility::Collapsed;
 			}
 
-			// If the child widget is collapsed, we need to store the new layout scale it will have when 
-			// it is finally visible and invalidate it's prepass so that it gets that when its visibility
-			// is finally invalidated.
-			Child->InvalidatePrepass();
-			Child->PrepassLayoutScaleMultiplier = ChildLayoutScaleMultiplier;
+			if (bIsCollapsed)
+			{
+				// If the child widget is collapsed, we need to store the new layout scale it will have when 
+				// it is finally visible and invalidate it's prepass so that it gets that when its visibility
+				// is finally invalidated.
+				Child->MarkPrepassAsDirty();
+				Child->PrepassLayoutScaleMultiplier = ChildLayoutScaleMultiplier;
+			}
+			else
+			{
+				// Recur: Descend down the widget tree.
+				Child->Prepass_Internal(ChildLayoutScaleMultiplier);
+			}
 		}
 	}
 }
diff --git a/Engine/Source/Runtime/SlateCore/Public/Types/SlateAttributeMetaData.h b/Engine/Source/Runtime/SlateCore/Public/Types/SlateAttributeMetaData.h
--- a/Engine/Source/Runtime/SlateCore/Public/Types/SlateAttributeMetaData.h
+++ b/Engine/Source/Runtime/SlateCore/Public/Types/SlateAttributeMetaData.h
@@ -24,9 +24,28 @@ public:
 	FSlateAttributeMetaData& operator=(const FSlateAttributeMetaData&) = delete;
 
 public:
+	/** @return the instance associated to the SWidget (if it exists). */
 	static FSlateAttributeMetaData* FindMetaData(const SWidget& OwningWidget);
+	/**
+	 * Update all the attributes.
+	 * Invalidate the widget if it has finished its construction phase.
+	 */
 	static void UpdateAttributes(SWidget& OwningWidget);
+	/**
+	 * Update attributes that are mark to be updated when the widget is collapsed.
+	 * Invalidate the widget if it has finished its construction phase.
+	 */
 	static void UpdateCollapsedAttributes(SWidget& OwningWidget);
+	/**
+	 * Update attributes that are mark to be updated when the widget is collapsed.
+	 * @param bAllowInvalidation if we should allow the widget to be invalidated.
+	 */
+	static void UpdateAttributes(SWidget& OwningWidget, bool bAllowInvalidation);
+	/**
+	 * Update attributes that are mark to be updated when the widget is collapsed.
+	 * @param bAllowInvalidation if we should allow the widget to be invalidated.
+	 */
+	static void UpdateCollapsedAttributes(SWidget& OwningWidget, bool bAllowInvalidation);
 
 public:
 	bool IsBound(const FSlateAttributeBase& Attribute) const
@@ -34,7 +53,9 @@ public:
 		return IndexOfAttribute(Attribute) != INDEX_NONE;
 	}
 
-	int32 RegisteredNum() const { return Attributes.Num(); }
+	int32 RegisteredAttributeCount() const { return Attributes.Num(); }
+
+	int32 RegisteredCollaspedAttributeCount() const { return CollaspedAttributeCounter; }
 
 private:
 	using ESlateAttributeType = SlateAttributePrivate::ESlateAttributeType;
@@ -51,7 +72,7 @@ private:
 private:
 	void RegisterAttributeImpl(SWidget& OwningWidget, FSlateAttributeBase& Attribute, ESlateAttributeType AttributeType, TUniquePtr<ISlateAttributeGetter>&& Getter);
 	bool UnregisterAttributeImpl(const FSlateAttributeBase& Attribute);
-	void UpdateAttributes(SWidget& OwningWidget, bool bOnlyCollapsed);
+	void UpdateAttributes(SWidget& OwningWidget, bool bOnlyCollapsed, bool bAllowInvalidation);
 
 private:
 	int32 IndexOfAttribute(const FSlateAttributeBase& Attribute) const
diff --git a/Engine/Source/Runtime/SlateCore/Public/Widgets/DeclarativeSyntaxSupport.h b/Engine/Source/Runtime/SlateCore/Public/Widgets/DeclarativeSyntaxSupport.h
--- a/Engine/Source/Runtime/SlateCore/Public/Widgets/DeclarativeSyntaxSupport.h
+++ b/Engine/Source/Runtime/SlateCore/Public/Widgets/DeclarativeSyntaxSupport.h
@@ -782,12 +782,14 @@ struct FSlateBaseNamedArgs
 	SLATE_PRIVATE_ATTRIBUTE_VARIABLE(FText, ToolTipText);
 	SLATE_PRIVATE_ATTRIBUTE_VARIABLE(TSharedPtr<IToolTip>, ToolTip);
 	SLATE_PRIVATE_ATTRIBUTE_VARIABLE(TOptional<EMouseCursor::Type>, Cursor);
-	SLATE_PRIVATE_ATTRIBUTE_VARIABLE(bool, IsEnabled) = true;
 	SLATE_PRIVATE_ATTRIBUTE_VARIABLE(EVisibility, Visibility) = EVisibility::Visible;
-	SLATE_PRIVATE_ARGUMENT_VARIABLE(float, RenderOpacity) = 1.f;
+	SLATE_PRIVATE_ATTRIBUTE_VARIABLE(bool, IsEnabled) = true;
 	SLATE_PRIVATE_ARGUMENT_VARIABLE(bool, ForceVolatile) = false;
+	/** If true, bound Slate Attributes will be updated once per frame. */
+	SLATE_PRIVATE_ARGUMENT_VARIABLE(bool, EnabledAttributesUpdate) = true;
 	SLATE_PRIVATE_ARGUMENT_VARIABLE(EWidgetClipping, Clipping) = EWidgetClipping::Inherit;
 	SLATE_PRIVATE_ARGUMENT_VARIABLE(EFlowDirectionPreference, FlowDirectionPreference) = EFlowDirectionPreference::Inherit;
+	SLATE_PRIVATE_ARGUMENT_VARIABLE(float, RenderOpacity) = 1.f;
 	SLATE_PRIVATE_ATTRIBUTE_VARIABLE(TOptional<FSlateRenderTransform>, RenderTransform);
 	SLATE_PRIVATE_ATTRIBUTE_VARIABLE(FVector2D, RenderTransformPivot) = FVector2D::ZeroVector;
 	SLATE_PRIVATE_ARGUMENT_VARIABLE(FName, Tag);
@@ -808,10 +810,10 @@ struct TSlateBaseNamedArgs : public FSlateBaseNamedArgs
 	SLATE_PRIVATE_ATTRIBUTE_FUNCTION(TOptional<EMouseCursor::Type>, Cursor)
 	SLATE_PRIVATE_ATTRIBUTE_FUNCTION(bool, IsEnabled)
 	SLATE_PRIVATE_ATTRIBUTE_FUNCTION(EVisibility, Visibility)
-	SLATE_PRIVATE_ARGUMENT_FUNCTION(float, RenderOpacity)
 	SLATE_PRIVATE_ARGUMENT_FUNCTION(bool, ForceVolatile)
 	SLATE_PRIVATE_ARGUMENT_FUNCTION(EWidgetClipping, Clipping)
 	SLATE_PRIVATE_ARGUMENT_FUNCTION(EFlowDirectionPreference, FlowDirectionPreference)
+	SLATE_PRIVATE_ARGUMENT_FUNCTION(float, RenderOpacity)
 	SLATE_PRIVATE_ATTRIBUTE_FUNCTION(TOptional<FSlateRenderTransform>, RenderTransform)
 	SLATE_PRIVATE_ATTRIBUTE_FUNCTION(FVector2D, RenderTransformPivot)
 	SLATE_PRIVATE_ARGUMENT_FUNCTION(FName, Tag)
@@ -1099,8 +1101,8 @@ struct TSlateDecl
 	{
 		_Widget->SWidgetConstruct(InArgs);
 		_RequiredArgs.CallConstruct(_Widget, InArgs);
-		_Widget->bPauseAttributeInvalidation = false;
 		_Widget->CacheVolatility();
+		_Widget->bIsDeclarativeSyntaxConstructionCompleted = true;
 
 		return _Widget;
 	}
diff --git a/Engine/Source/Runtime/SlateCore/Public/Widgets/InvalidateWidgetReason.h b/Engine/Source/Runtime/SlateCore/Public/Widgets/InvalidateWidgetReason.h
--- a/Engine/Source/Runtime/SlateCore/Public/Widgets/InvalidateWidgetReason.h
+++ b/Engine/Source/Runtime/SlateCore/Public/Widgets/InvalidateWidgetReason.h
@@ -46,6 +46,11 @@ enum class EInvalidateWidgetReason : uint8
 	 */
 	AttributeRegistration = 1 << 6,
 
+	/**
+	 * Re-cache desired size of all of this widget's children recursively (this implies layout)
+	 */
+	Prepass = 1 << 7,
+
 	/**
 	 * Use Paint invalidation if you're changing a normal property involving painting or sizing.
 	 * Additionally if the property that was changed affects Volatility in anyway, it's important
diff --git a/Engine/Source/Runtime/SlateCore/Public/Widgets/SWidget.h b/Engine/Source/Runtime/SlateCore/Public/Widgets/SWidget.h
--- a/Engine/Source/Runtime/SlateCore/Public/Widgets/SWidget.h
+++ b/Engine/Source/Runtime/SlateCore/Public/Widgets/SWidget.h
@@ -701,7 +701,10 @@ public:
 	void SetCanTick(bool bInCanTick) { bInCanTick ? AddUpdateFlags(EWidgetUpdateFlags::NeedsTick) : RemoveUpdateFlags(EWidgetUpdateFlags::NeedsTick); }
 	bool GetCanTick() const { return HasAnyUpdateFlags(EWidgetUpdateFlags::NeedsTick); }
 
+	/** @return true if the widgets has any bound slate attribute. */
 	bool HasRegisteredSlateAttribute() const { return bHasRegisteredSlateAttribute; }
+	/** @return true if the widgets will update its registered slate attributes automatically or they need to be updated manually. */
+	bool IsAttributesUpdatesEnabled() const { return bEnabledAttributesUpdate; }
 
 	const FSlateWidgetPersistentState& GetPersistentState() const { return PersistentState; }
 	const FWidgetProxyHandle GetProxyHandle() const { return FastPathProxyHandle; }
@@ -1095,8 +1098,15 @@ public:
 		bCachedVolatile = bForceVolatile || ComputeVolatility();
 	}
 
+	UE_DEPRECATED(5.0, "InvalidatePrepass is deprecated. Use the Invalidate(EInvalidateWidgetReason::Prepass) or use MarkPrepassAsDirty()")
 	void InvalidatePrepass();
 
+	/**
+	 * In fast path, if the widget is mark, do a full Prepass on its next update to calculate it's desired size.
+	 * This does not invalidate the widget.
+	 */
+	void MarkPrepassAsDirty() { bNeedsPrepass = true; }
+
 protected:
 
 #if SLATE_CULL_WIDGETS
@@ -1487,6 +1497,9 @@ protected:
 	/** Construct a SWidget based on initial parameters. */
 	void SWidgetConstruct(const FSlateBaseNamedArgs& Args);
 
+	/** Is the widget construction is completed (did we called completed the Construct() function) */
+	bool IsConstructionCompleted() const { return bIsDeclarativeSyntaxConstructionCompleted; }
+
 	/** 
 	 * Find the geometry of a descendant widget. This method assumes that WidgetsToFind are a descendants of this widget.
 	 * Note that not all widgets are guaranteed to be found; OutResult will contain null entries for missing widgets.
@@ -1701,8 +1714,11 @@ private:
 	/** Is there at least one SlateAttribute currently registered. */
 	uint8 bHasRegisteredSlateAttribute : 1;
 
-	/** Prevents invalidation while we are in the construction phase. */
-	uint8 bPauseAttributeInvalidation : 1;
+	/** Are bound Slate Attributes will be updated once per frame. */
+	uint8 bEnabledAttributesUpdate : 1;
+
+	/** The SNew or SAssignedNew construction is completed. */
+	uint8 bIsDeclarativeSyntaxConstructionCompleted : 1;
 
 protected:
 	uint8 bHasCustomPrepass : 1;
diff --git a/Engine/Source/Runtime/UMG/Private/UserWidget.cpp b/Engine/Source/Runtime/UMG/Private/UserWidget.cpp
--- a/Engine/Source/Runtime/UMG/Private/UserWidget.cpp
+++ b/Engine/Source/Runtime/UMG/Private/UserWidget.cpp
@@ -1409,7 +1409,7 @@ void UUserWidget::PostTickActionsAndAnimation(float InDeltaTime)
 		TSharedPtr<SWidget> CachedWidget = GetCachedWidget();
 		if (CachedWidget.IsValid())
 		{
-			CachedWidget->InvalidatePrepass();
+			CachedWidget->Invalidate(EInvalidateWidgetReason::Prepass);
 		}
 	}
 
