Fix build. Missing file from CL15929257
#rnx

diff --git a/Engine/Source/Runtime/SlateCore/Private/Debugging/WidgetList.cpp b/Engine/Source/Runtime/SlateCore/Private/Debugging/WidgetList.cpp
--- a/Engine/Source/Runtime/SlateCore/Private/Debugging/WidgetList.cpp
+++ b/Engine/Source/Runtime/SlateCore/Private/Debugging/WidgetList.cpp
@@ -127,7 +127,7 @@ void LogAllWidgetsDebugInfoImpl(FOutputDevice& Ar, const FLogAllWidgetsDebugInfo
 	}
 	if (DebugInfoFlags.bSlateAttribute)
 	{
-		MessageBuilder << TEXT(";NumSlateAttribute");
+		MessageBuilder << TEXT(";NumSlateAttribute;NumSlateCollaspedAttribute");
 	}
 	if (DebugInfoFlags.bMouseEventsHandler)
 	{
@@ -255,11 +255,13 @@ void LogAllWidgetsDebugInfoImpl(FOutputDevice& Ar, const FLogAllWidgetsDebugInfo
 			if (FSlateAttributeMetaData* MetaData = FSlateAttributeMetaData::FindMetaData(*Widget))
 			{
 				MessageBuilder << TEXT(";");
-				MessageBuilder << MetaData->RegisteredNum();
+				MessageBuilder << MetaData->RegisteredAttributeCount();
+				MessageBuilder << TEXT(";");
+				MessageBuilder << MetaData->RegisteredCollaspedAttributeCount();
 			}
 			else
 			{
-				MessageBuilder << TEXT(";0");
+				MessageBuilder << TEXT(";0;0");
 			}
 		}
 
