Fix unreal icon artifacts in project browser

diff --git a/Engine/Source/Editor/ContentBrowser/Private/AssetViewWidgets.cpp b/Engine/Source/Editor/ContentBrowser/Private/AssetViewWidgets.cpp
--- a/Engine/Source/Editor/ContentBrowser/Private/AssetViewWidgets.cpp
+++ b/Engine/Source/Editor/ContentBrowser/Private/AssetViewWidgets.cpp
@@ -1732,7 +1732,7 @@ void SAssetTileItem::Construct( const FArguments& InArgs )
 						[
 							SNew(SVerticalBox)
 							+ SVerticalBox::Slot()
-							.Padding(2.0f,0.0f,0.0f,0.0f)
+							.Padding(2.0f,2.0f,0.0f,0.0f)
 							.VAlign(VAlign_Top)
 							[
 								SNew(SBox)
diff --git a/Engine/Source/Editor/EditorStyle/Private/StarshipStyle.cpp b/Engine/Source/Editor/EditorStyle/Private/StarshipStyle.cpp
--- a/Engine/Source/Editor/EditorStyle/Private/StarshipStyle.cpp
+++ b/Engine/Source/Editor/EditorStyle/Private/StarshipStyle.cpp
@@ -1585,7 +1585,6 @@ void FStarshipEditorStyle::FStyle::SetupGeneralStyles()
 	// GameProjectDialog
 	if (IncludeEditorSpecificStyles())
 	{
-		Set( "GameProjectDialog.DefaultGameThumbnail", new IMAGE_BRUSH( "GameProjectDialog/default_game_thumbnail", Icon64x64 ) );
 		Set( "GameProjectDialog.BlankProjectThumbnail", new IMAGE_BRUSH( "GameProjectDialog/blank_project_thumbnail", Icon128x128 ) );
 		Set( "GameProjectDialog.BlankProjectPreview", new IMAGE_BRUSH( "GameProjectDialog/blank_project_preview", FVector2D(400, 200) ) );
 
diff --git a/Engine/Source/Editor/GameProjectGeneration/Private/SProjectBrowser.cpp b/Engine/Source/Editor/GameProjectGeneration/Private/SProjectBrowser.cpp
--- a/Engine/Source/Editor/GameProjectGeneration/Private/SProjectBrowser.cpp
+++ b/Engine/Source/Editor/GameProjectGeneration/Private/SProjectBrowser.cpp
@@ -149,13 +149,14 @@ public:
 						.HeightOverride(102)
 						[
 							SNew(SBorder)
-							.Padding(0)
 							.BorderImage(FAppStyle::Get().GetBrush("ProjectBrowser.ProjectTile.ThumbnailAreaBackground"))
-							.HAlign(HAlign_Center)
-							.VAlign(VAlign_Center)
+							.HAlign(ProjectItem->ProjectThumbnail ? HAlign_Center : HAlign_Fill)
+							.VAlign(ProjectItem->ProjectThumbnail ? VAlign_Center : VAlign_Fill)
+							.Padding(ProjectItem->ProjectThumbnail ? FMargin(0) : FMargin(12.0f) )
 							[
 								SNew(SImage)
-								.Image(ProjectItem->ProjectThumbnail ? ProjectItem->ProjectThumbnail.Get() : FAppStyle::Get().GetBrush("GameProjectDialog.DefaultGameThumbnail"))
+								.Image(ProjectItem->ProjectThumbnail ? ProjectItem->ProjectThumbnail.Get() : FAppStyle::Get().GetBrush("UnrealCircle.Thin"))
+								.ColorAndOpacity(FAppStyle::Get().GetSlateColor("Colors.Foreground"))
 							]
 						]
 					]
@@ -179,7 +180,7 @@ public:
 							]
 							+SVerticalBox::Slot()
 							.AutoHeight()
-							.Padding(0.0f, 4.0f, 0.0f, 4.0f)
+							.Padding(0.0f, 4.0f, 0.0f, 0.0f)
 							.VAlign(VAlign_Bottom)
 							[
 								SNew(STextBlock)
diff --git a/Engine/Source/Editor/GameProjectGeneration/Private/SProjectDialog.cpp b/Engine/Source/Editor/GameProjectGeneration/Private/SProjectDialog.cpp
--- a/Engine/Source/Editor/GameProjectGeneration/Private/SProjectDialog.cpp
+++ b/Engine/Source/Editor/GameProjectGeneration/Private/SProjectDialog.cpp
@@ -258,7 +258,7 @@ private:
 		{
 			return ItemPtr->Thumbnail.Get();
 		}
-		return FEditorStyle::GetBrush("GameProjectDialog.DefaultGameThumbnail");
+		return FEditorStyle::GetBrush("UnrealCircle.Thin");
 	}
 	
 };
diff --git a/Engine/Source/Runtime/Slate/Private/Widgets/Notifications/SNotificationList.cpp b/Engine/Source/Runtime/Slate/Private/Widgets/Notifications/SNotificationList.cpp
--- a/Engine/Source/Runtime/Slate/Private/Widgets/Notifications/SNotificationList.cpp
+++ b/Engine/Source/Runtime/Slate/Private/Widgets/Notifications/SNotificationList.cpp
@@ -719,7 +719,7 @@ TSharedRef<SNotificationItem> SNotificationList::AddNotification(const FNotifica
 		}
 		else
 		{
-			static const FSlateBrush* CachedImage = FAppStyle::Get().GetBrush("AppIconFlat");
+			static const FSlateBrush* CachedImage = FAppStyle::Get().GetBrush("UnrealCircle.Thick");
 
 			// Create notification.
 			NewItem = SNew(SNotificationItemImpl)
diff --git a/Engine/Source/Runtime/SlateCore/Private/Styling/StarshipCoreStyle.cpp b/Engine/Source/Runtime/SlateCore/Private/Styling/StarshipCoreStyle.cpp
--- a/Engine/Source/Runtime/SlateCore/Private/Styling/StarshipCoreStyle.cpp
+++ b/Engine/Source/Runtime/SlateCore/Private/Styling/StarshipCoreStyle.cpp
@@ -282,7 +282,8 @@ TSharedRef<ISlateStyle> FStarshipCoreStyle::Create()
 		Style->Set("AppIconPadding", FMargin(5.f, 5.f, 5.f, 5.f));
 #endif
 
-		Style->Set("AppIconFlat", new IMAGE_BRUSH_SVG("Starship/Common/unreal", FVector2D(32, 32), FStyleColors::Foreground));
+		Style->Set("UnrealCircle.Thick", new IMAGE_BRUSH_SVG("Starship/Common/unreal-circle-thick", FVector2D(32, 32), FStyleColors::Foreground));
+		Style->Set("UnrealCircle.Thin", new IMAGE_BRUSH_SVG("Starship/Common/unreal-circle-thin", FVector2D(32, 32)));
 		Style->Set("AppIcon.Small", new IMAGE_BRUSH_SVG("Starship/Common/unreal-small", Icon24x24, FStyleColors::Foreground));
 		Style->Set("AppIconPadding.Small", FMargin(4.f, 4.f, 0.f, 0.f));
 
@@ -621,7 +622,7 @@ TSharedRef<ISlateStyle> FStarshipCoreStyle::Create()
 
 		Style->Set("NotificationList.Throbber", new IMAGE_BRUSH("Starship/Notifications/Throbber", FVector2D(40, 40)));
 
-		Style->Set("NotificationList.Watermark", new IMAGE_BRUSH_SVG("Starship/Notifications/Watermark", FVector2D(128, 128)));
+		Style->Set("NotificationList.Watermark", new IMAGE_BRUSH_SVG("Starship/Common/unreal-circle-thin", FVector2D(128, 128)));
 
 		// Should be deleted. everywhere using this is looking for a generic "alert" icon
 		Style->Set("NotificationList.DefaultMessage", new IMAGE_BRUSH("Common/EventMessage_Default", Icon40x40));
