Slate: Add callback when a widget is invalidated by a SlateAttribute.

diff --git a/Engine/Source/Runtime/SlateCore/Private/Test/SlateAttributeTest.cpp b/Engine/Source/Runtime/SlateCore/Private/Test/SlateAttributeTest.cpp
--- a/Engine/Source/Runtime/SlateCore/Private/Test/SlateAttributeTest.cpp
+++ b/Engine/Source/Runtime/SlateCore/Private/Test/SlateAttributeTest.cpp
@@ -156,13 +156,124 @@ void SAttributeLeftWidget_Child::PrivateRegisterAttributes(FSlateAttributeInitia
 		.UpdatePrerequisite("IntAttributeC");
 }
 
+class SAttributeLeftWidget_OnInvalidationParent : public SLeafWidget
+{
+	SLATE_DECLARE_WIDGET(SAttributeLeftWidget_OnInvalidationParent, SLeafWidget)
+public:
+	SLATE_BEGIN_ARGS(SAttributeLeftWidget_OnInvalidationParent) {}
+	SLATE_END_ARGS()
+
+
+	SAttributeLeftWidget_OnInvalidationParent()
+		: IntAttributeA(*this, 99)
+		, IntAttributeB(*this, 99)
+		, IntAttributeC(*this, 99)
+	{
+	}
+
+	void Construct(const FArguments& InArgs) {}
+	virtual FVector2D ComputeDesiredSize(float) const override { return FVector2D{ 100, 100 }; }
+	virtual int32 OnPaint(const FPaintArgs& Args, const FGeometry& AllottedGeometry, const FSlateRect& MyCullingRect, FSlateWindowElementList& OutDrawElements, int32 LayerId, const FWidgetStyle& InWidgetStyle, bool bParentEnabled) const override
+	{
+		return LayerId;
+	}
+
+	void SetAttributeA(TAttribute<int32> ValueA) { IntAttributeA.Assign(*this, MoveTemp(ValueA)); }
+	void SetAttributeB(TAttribute<int32> ValueB) { IntAttributeB.Assign(*this, MoveTemp(ValueB)); }
+	void SetAttributeC(TAttribute<int32> ValueC) { IntAttributeC.Assign(*this, MoveTemp(ValueC)); }
+	virtual void SetCallbackValue(int32 Value) { CallbackValueA = Value; CallbackValueB = Value; CallbackValueC = Value; }
+
+	TSlateAttribute<int32> IntAttributeA;
+	TSlateAttribute<int32> IntAttributeB;
+	TSlateAttribute<int32> IntAttributeC;
+	int32 CallbackValueA;
+	int32 CallbackValueB;
+	int32 CallbackValueC;
+};
+
+SLATE_IMPLEMENT_WIDGET(SAttributeLeftWidget_OnInvalidationParent)
+void SAttributeLeftWidget_OnInvalidationParent::PrivateRegisterAttributes(FSlateAttributeInitializer& AttributeInitializer)
+{
+	SLATE_ADD_MEMBER_ATTRIBUTE_DEFINITION(AttributeInitializer, IntAttributeA, EInvalidateWidgetReason::ChildOrder)
+		.OnInvalidation(FSlateAttributeDescriptor::FInvalidationDelegate::CreateLambda([](SWidget& Widget)
+			{
+				SAttributeLeftWidget_OnInvalidationParent& Instance = static_cast<SAttributeLeftWidget_OnInvalidationParent&>(Widget);
+				Instance.CallbackValueA = Instance.IntAttributeA.Get();
+			}));
+	SLATE_ADD_MEMBER_ATTRIBUTE_DEFINITION(AttributeInitializer, IntAttributeB, EInvalidateWidgetReason::ChildOrder)
+		.OnInvalidation(FSlateAttributeDescriptor::FInvalidationDelegate::CreateLambda([](SWidget& Widget)
+			{
+				SAttributeLeftWidget_OnInvalidationParent& Instance = static_cast<SAttributeLeftWidget_OnInvalidationParent&>(Widget);
+				Instance.CallbackValueB = Instance.IntAttributeB.Get();
+			}));
+	SLATE_ADD_MEMBER_ATTRIBUTE_DEFINITION(AttributeInitializer, IntAttributeC, EInvalidateWidgetReason::ChildOrder)
+		.OnInvalidation(FSlateAttributeDescriptor::FInvalidationDelegate::CreateLambda([](SWidget& Widget)
+			{
+				SAttributeLeftWidget_OnInvalidationParent& Instance = static_cast<SAttributeLeftWidget_OnInvalidationParent&>(Widget);
+				Instance.CallbackValueC = Instance.IntAttributeC.Get();
+			}));
 }
-}
+
+class SAttributeLeftWidget_OnInvalidationChild : public SAttributeLeftWidget_OnInvalidationParent
+{
+	SLATE_DECLARE_WIDGET(SAttributeLeftWidget_OnInvalidationChild, SAttributeLeftWidget_OnInvalidationParent)
+public:
+	SLATE_BEGIN_ARGS(SAttributeLeftWidget_OnInvalidationChild) {}
+	SLATE_END_ARGS()
+
+	void Construct(const FArguments& InArgs) {}
+	virtual FVector2D ComputeDesiredSize(float) const override { return FVector2D{ 100, 100 }; }
+	virtual int32 OnPaint(const FPaintArgs& Args, const FGeometry& AllottedGeometry, const FSlateRect& MyCullingRect, FSlateWindowElementList& OutDrawElements, int32 LayerId, const FWidgetStyle& InWidgetStyle, bool bParentEnabled) const override
+	{
+		return LayerId;
+	}
+
+	virtual void SetCallbackValue(int32 Value) override
+	{
+		Super::SetCallbackValue(Value);
+		OverrideCallbackValueA = Value;
+		OverrideCallbackValueB = Value;
+		OverrideCallbackValueC = Value;
+	}
+
+	int32 OverrideCallbackValueA;
+	int32 OverrideCallbackValueB;
+	int32 OverrideCallbackValueC;
+};
+
+SLATE_IMPLEMENT_WIDGET(SAttributeLeftWidget_OnInvalidationChild)
+void SAttributeLeftWidget_OnInvalidationChild::PrivateRegisterAttributes(FSlateAttributeInitializer& AttributeInitializer)
+{
+	AttributeInitializer.OverrideOnInvalidation("IntAttributeA"
+		, FSlateAttributeDescriptor::EInvalidationDelegateOverrideType::ReplacePrevious
+		, FSlateAttributeDescriptor::FInvalidationDelegate::CreateLambda([](SWidget& Widget)
+			{
+				SAttributeLeftWidget_OnInvalidationChild& Instance = static_cast<SAttributeLeftWidget_OnInvalidationChild&>(Widget);
+				Instance.OverrideCallbackValueA = Instance.IntAttributeA.Get();
+			}));
+	AttributeInitializer.OverrideOnInvalidation("IntAttributeB"
+		, FSlateAttributeDescriptor::EInvalidationDelegateOverrideType::ExecuteAfterPrevious
+		, FSlateAttributeDescriptor::FInvalidationDelegate::CreateLambda([](SWidget& Widget)
+			{
+				SAttributeLeftWidget_OnInvalidationChild& Instance = static_cast<SAttributeLeftWidget_OnInvalidationChild&>(Widget);
+				Instance.OverrideCallbackValueB = Instance.CallbackValueB;
+			}));
+	AttributeInitializer.OverrideOnInvalidation("IntAttributeC"
+		, FSlateAttributeDescriptor::EInvalidationDelegateOverrideType::ExecuteBeforePrevious
+		, FSlateAttributeDescriptor::FInvalidationDelegate::CreateLambda([](SWidget& Widget)
+			{
+				SAttributeLeftWidget_OnInvalidationChild& Instance = static_cast<SAttributeLeftWidget_OnInvalidationChild&>(Widget);
+				Instance.OverrideCallbackValueC = Instance.CallbackValueC;
+			}));
 }
 
+}
+}
+}
 
 bool FSlateAttributeTest::RunTest(const FString& Parameters)
 {
+	const int NumberOfAttributeInSWidget = 2;
 	int32 OrderCounter = 0;
 	auto OrderLambda = [this, &OrderCounter]() -> int32
 	{
@@ -184,7 +295,7 @@ bool FSlateAttributeTest::RunTest(const FString& Parameters)
 			, TEXT("The static data do not matches"));
 
 		FSlateAttributeDescriptor const& AttributeDescriptor = WidgetParent->GetWidgetClass().GetAttributeDescriptor();
-		AddErrorIfFalse(AttributeDescriptor.AttributeNum() == 4, TEXT(""));
+		AddErrorIfFalse(AttributeDescriptor.AttributeNum() == 4 + NumberOfAttributeInSWidget, TEXT(""));
 
 		const int32 IndexA = AttributeDescriptor.IndexOfMemberAttribute("IntAttributeA");
 		const int32 IndexB = AttributeDescriptor.IndexOfMemberAttribute("IntAttributeB");
@@ -300,7 +411,7 @@ bool FSlateAttributeTest::RunTest(const FString& Parameters)
 			, TEXT("The static data do not matches"));
 
 		FSlateAttributeDescriptor const& AttributeDescriptor = WidgetChild->GetWidgetClass().GetAttributeDescriptor();
-		AddErrorIfFalse(AttributeDescriptor.AttributeNum() == 8, TEXT("")); // H is not counted
+		AddErrorIfFalse(AttributeDescriptor.AttributeNum() == 8 + NumberOfAttributeInSWidget, TEXT("")); // H is not counted
 
 		const int32 IndexA = AttributeDescriptor.IndexOfMemberAttribute("IntAttributeA");
 		const int32 IndexB = AttributeDescriptor.IndexOfMemberAttribute("IntAttributeB");
@@ -701,6 +812,91 @@ bool FSlateAttributeTest::RunTest(const FString& Parameters)
 		}
 	}
 
+	// Test OnInvalidation
+	{
+		{
+			TSharedRef<UE::Slate::Private::SAttributeLeftWidget_OnInvalidationParent> WidgetParent = SNew(UE::Slate::Private::SAttributeLeftWidget_OnInvalidationParent);
+			WidgetParent->SetCallbackValue(10);
+			WidgetParent->MarkPrepassAsDirty();
+			WidgetParent->SlatePrepass(1.f);
+			AddErrorIfFalse(WidgetParent->CallbackValueA == 10, TEXT("The callback for value A was triggered."));
+			AddErrorIfFalse(WidgetParent->CallbackValueB == 10, TEXT("The callback for value B was triggered."));
+			AddErrorIfFalse(WidgetParent->CallbackValueC == 10, TEXT("The callback for value C was triggered."));
+		}
+		{
+			TSharedRef<UE::Slate::Private::SAttributeLeftWidget_OnInvalidationChild> WidgetChild = SNew(UE::Slate::Private::SAttributeLeftWidget_OnInvalidationChild);
+			WidgetChild->SetCallbackValue(10);
+			WidgetChild->MarkPrepassAsDirty();
+			WidgetChild->SlatePrepass(1.f);
+			AddErrorIfFalse(WidgetChild->CallbackValueA == 10, TEXT("The callback for value A was triggered."));
+			AddErrorIfFalse(WidgetChild->CallbackValueB == 10, TEXT("The callback for value B was triggered."));
+			AddErrorIfFalse(WidgetChild->CallbackValueC == 10, TEXT("The callback for value C was triggered."));
+			AddErrorIfFalse(WidgetChild->OverrideCallbackValueA == 10, TEXT("The callback for override value A was triggered."));
+			AddErrorIfFalse(WidgetChild->OverrideCallbackValueB == 10, TEXT("The callback for override value B was triggered."));
+			AddErrorIfFalse(WidgetChild->OverrideCallbackValueC == 10, TEXT("The callback for override value C was triggered."));
+		}
+
+		{
+			TSharedRef<UE::Slate::Private::SAttributeLeftWidget_OnInvalidationParent> WidgetParent = SNew(UE::Slate::Private::SAttributeLeftWidget_OnInvalidationParent);
+			WidgetParent->SetCallbackValue(10);
+			WidgetParent->SetAttributeA(4);
+			AddErrorIfFalse(WidgetParent->CallbackValueA == 4, TEXT("The callback for value A was not triggered."));
+			AddErrorIfFalse(WidgetParent->CallbackValueB == 10, TEXT("The callback for value B was triggered."));
+			AddErrorIfFalse(WidgetParent->CallbackValueC == 10, TEXT("The callback for value C was triggered."));
+			WidgetParent->SetAttributeB(5);
+			AddErrorIfFalse(WidgetParent->CallbackValueB == 5, TEXT("The callback for value B was not triggered."));
+			AddErrorIfFalse(WidgetParent->CallbackValueC == 10, TEXT("The callback for value C was triggered."));
+			WidgetParent->SetAttributeC(6);
+			AddErrorIfFalse(WidgetParent->CallbackValueC == 6, TEXT("The callback for value C was not triggered."));
+		}
+
+		{
+			TSharedRef<UE::Slate::Private::SAttributeLeftWidget_OnInvalidationChild> WidgetChild = SNew(UE::Slate::Private::SAttributeLeftWidget_OnInvalidationChild);
+			WidgetChild->SetCallbackValue(10);
+			WidgetChild->CallbackValueB = 15;
+			WidgetChild->SetAttributeA(4);
+			AddErrorIfFalse(WidgetChild->CallbackValueA == 10, TEXT("The callback for value A was triggered."));
+			AddErrorIfFalse(WidgetChild->OverrideCallbackValueA == 4, TEXT("The callback for override value A was not triggered."));
+			AddErrorIfFalse(WidgetChild->CallbackValueB == 15, TEXT("The callback for value B was triggered."));
+			AddErrorIfFalse(WidgetChild->OverrideCallbackValueB == 10, TEXT("The callback for override value B was triggered."));
+			AddErrorIfFalse(WidgetChild->CallbackValueC == 10, TEXT("The callback for value C was triggered."));
+			AddErrorIfFalse(WidgetChild->OverrideCallbackValueC == 10, TEXT("The callback for override value C was triggered."));
+			WidgetChild->SetAttributeB(5);
+			AddErrorIfFalse(WidgetChild->CallbackValueB == 5, TEXT("The callback for value B was not triggered."));
+			AddErrorIfFalse(WidgetChild->OverrideCallbackValueB == 15, TEXT("The callback for override value B was not triggered."));
+			AddErrorIfFalse(WidgetChild->CallbackValueC == 10, TEXT("The callback for value C was triggered."));
+			AddErrorIfFalse(WidgetChild->OverrideCallbackValueC == 10, TEXT("The callback for override value C was triggered."));
+			WidgetChild->SetAttributeC(6);
+			AddErrorIfFalse(WidgetChild->CallbackValueC == 6, TEXT("The callback for value C was not triggered."));
+			AddErrorIfFalse(WidgetChild->OverrideCallbackValueC == 6, TEXT("The callback for override value C was not triggered."));
+		}
+		{
+			TSharedRef<UE::Slate::Private::SAttributeLeftWidget_OnInvalidationChild> WidgetChild = SNew(UE::Slate::Private::SAttributeLeftWidget_OnInvalidationChild);
+			WidgetChild->SetCallbackValue(10);
+			WidgetChild->CallbackValueB = 15;
+			WidgetChild->SetAttributeA(MakeAttributeLambda([](){ return 4; }));
+			WidgetChild->SetAttributeB(MakeAttributeLambda([](){ return 5; }));
+			WidgetChild->SetAttributeC(MakeAttributeLambda([](){ return 6; }));
+
+			AddErrorIfFalse(WidgetChild->CallbackValueA == 10, TEXT("The callback for value A was triggered."));
+			AddErrorIfFalse(WidgetChild->OverrideCallbackValueA == 10, TEXT("The callback for override value A was triggered."));
+			AddErrorIfFalse(WidgetChild->CallbackValueB == 15, TEXT("The callback for value B was triggered."));
+			AddErrorIfFalse(WidgetChild->OverrideCallbackValueB == 10, TEXT("The callback for override value B was triggered."));
+			AddErrorIfFalse(WidgetChild->CallbackValueC == 10, TEXT("The callback for value C was triggered."));
+			AddErrorIfFalse(WidgetChild->OverrideCallbackValueC == 10, TEXT("The callback for override value C was triggered."));
+
+			WidgetChild->MarkPrepassAsDirty();
+			WidgetChild->SlatePrepass(1.f);
+
+			AddErrorIfFalse(WidgetChild->CallbackValueA == 10, TEXT("The callback for value A was triggered."));
+			AddErrorIfFalse(WidgetChild->OverrideCallbackValueA == 4, TEXT("The callback for override value A was not triggered."));
+			AddErrorIfFalse(WidgetChild->CallbackValueB == 5, TEXT("The callback for value B was not triggered."));
+			AddErrorIfFalse(WidgetChild->OverrideCallbackValueB == 15, TEXT("The callback for override value B was not triggered."));
+			AddErrorIfFalse(WidgetChild->CallbackValueC == 6, TEXT("The callback for value C was not triggered."));
+			AddErrorIfFalse(WidgetChild->OverrideCallbackValueC == 6, TEXT("The callback for override value C was not triggered."));
+		}
+	}
+
 	return true;
 }
 
diff --git a/Engine/Source/Runtime/SlateCore/Private/Types/SlateAttributeDescriptor.cpp b/Engine/Source/Runtime/SlateCore/Private/Types/SlateAttributeDescriptor.cpp
--- a/Engine/Source/Runtime/SlateCore/Private/Types/SlateAttributeDescriptor.cpp
+++ b/Engine/Source/Runtime/SlateCore/Private/Types/SlateAttributeDescriptor.cpp
@@ -4,7 +4,10 @@
 #include "Types/SlateAttribute.h"
 #include "Containers/ArrayView.h"
 
+#include <limits>
 
+
+/** */
 FSlateAttributeDescriptor::FInitializer::FAttributeEntry::FAttributeEntry(FSlateAttributeDescriptor& InDescriptor, int32 InAttributeIndex)
 	: Descriptor(InDescriptor)
 	, AttributeIndex(InAttributeIndex)
@@ -41,6 +44,17 @@ FSlateAttributeDescriptor::FInitializer::FAttributeEntry& FSlateAttributeDescrip
 }
 
 
+FSlateAttributeDescriptor::FInitializer::FAttributeEntry& FSlateAttributeDescriptor::FInitializer::FAttributeEntry::OnInvalidation(FInvalidationDelegate InCallback)
+{
+	if (Descriptor.Attributes.IsValidIndex(AttributeIndex))
+	{
+		Descriptor.Attributes[AttributeIndex].OnInvalidation = MoveTemp(InCallback);
+	}
+	return *this;
+}
+
+
+/** */
 FSlateAttributeDescriptor::FInitializer::FInitializer(FSlateAttributeDescriptor& InDescriptor)
 	: Descriptor(InDescriptor)
 {
@@ -55,6 +69,8 @@ FSlateAttributeDescriptor::FInitializer::FInitializer(FSlateAttributeDescriptor&
 
 FSlateAttributeDescriptor::FInitializer::~FInitializer()
 {
+	checkf(Descriptor.Attributes.Num() < std::numeric_limits<int16>::max(), TEXT("There is too many attributes. The index is saved as an int16 in FSlateAttributeMetaData."));
+
 	// Update the sort order for the item that have prerequisite.
 	//Because adding the attribute is not required at the moment
 	//try to not change the order in which they were added
@@ -177,6 +193,18 @@ void FSlateAttributeDescriptor::FInitializer::OverrideInvalidationReason(FName A
 }
 
 
+void FSlateAttributeDescriptor::FInitializer::OverrideInvalidationReason(FName AttributeName, FInvalidateWidgetReasonAttribute&& Reason)
+{
+	Descriptor.OverrideInvalidationReason(AttributeName, MoveTemp(Reason));
+}
+
+
+void FSlateAttributeDescriptor::FInitializer::OverrideOnInvalidation(FName AttributeName, EInvalidationDelegateOverrideType OverrideType, FInvalidationDelegate Callback)
+{
+	Descriptor.OverrideOnInvalidation(AttributeName, OverrideType, MoveTemp(Callback));
+}
+
+
 void FSlateAttributeDescriptor::FInitializer::SetUpdateWhenCollapsed(FName AttributeName, bool bUpdateWhenCollapsed)
 {
 	FAttribute* Attribute = Descriptor.FindAttribute(AttributeName);
@@ -187,12 +215,7 @@ void FSlateAttributeDescriptor::FInitializer::SetUpdateWhenCollapsed(FName Attri
 }
 
 
-void FSlateAttributeDescriptor::FInitializer::OverrideInvalidationReason(FName AttributeName, FInvalidateWidgetReasonAttribute&& Reason)
-{
-	Descriptor.OverrideInvalidationReason(AttributeName, MoveTemp(Reason));
-}
-
-
+/** */
 FSlateAttributeDescriptor::FAttribute const& FSlateAttributeDescriptor::GetAttributeAtIndex(int32 Index) const
 {
 	check(Attributes.IsValidIndex(Index));
@@ -271,6 +294,48 @@ void FSlateAttributeDescriptor::OverrideInvalidationReason(FName AttributeName,
 }
 
 
+void FSlateAttributeDescriptor::OverrideOnInvalidation(FName AttributeName, EInvalidationDelegateOverrideType OverrideType, FInvalidationDelegate Callback)
+{
+	check(!AttributeName.IsNone());
+
+	FAttribute* FoundAttribute = FindAttribute(AttributeName);
+	if (ensureAlwaysMsgf(FoundAttribute != nullptr, TEXT("The attribute 's' doesn't exist."), *AttributeName.ToString()))
+	{
+		switch(OverrideType)
+		{
+		case EInvalidationDelegateOverrideType::ReplacePrevious:
+			FoundAttribute->OnInvalidation = MoveTemp(Callback);
+			break;
+		case EInvalidationDelegateOverrideType::ExecuteAfterPrevious:
+		case EInvalidationDelegateOverrideType::ExecuteBeforePrevious:
+			if (FoundAttribute->OnInvalidation.IsBound() && Callback.IsBound())
+			{
+				FInvalidationDelegate Previous = FoundAttribute->OnInvalidation;
+				FoundAttribute->OnInvalidation = FInvalidationDelegate::CreateLambda([Previous{MoveTemp(Previous)}, Callback{MoveTemp(Callback)}, OverrideType](SWidget& Widget)
+					{
+						if (OverrideType == EInvalidationDelegateOverrideType::ExecuteBeforePrevious)
+						{
+							Previous.ExecuteIfBound(Widget);
+						}
+						Callback.ExecuteIfBound(Widget);
+						if (OverrideType == EInvalidationDelegateOverrideType::ExecuteAfterPrevious)
+						{
+							Previous.ExecuteIfBound(Widget);
+						}
+					});
+			}
+			else if (Callback.IsBound())
+			{
+				FoundAttribute->OnInvalidation = MoveTemp(Callback);
+			}
+			break;
+		default:
+			check(false);
+		}
+	}
+}
+
+
 void FSlateAttributeDescriptor::SetPrerequisite(FSlateAttributeDescriptor::FAttribute& Attribute, FName Prerequisite, bool bSetAsDependency)
 {
 	if (Prerequisite.IsNone())
diff --git a/Engine/Source/Runtime/SlateCore/Private/Types/SlateAttributeMetaData.cpp b/Engine/Source/Runtime/SlateCore/Private/Types/SlateAttributeMetaData.cpp
--- a/Engine/Source/Runtime/SlateCore/Private/Types/SlateAttributeMetaData.cpp
+++ b/Engine/Source/Runtime/SlateCore/Private/Types/SlateAttributeMetaData.cpp
@@ -164,13 +164,41 @@ bool FSlateAttributeMetaData::UnregisterAttributeImpl(const FSlateAttributeBase&
 }
 
 
-EInvalidateWidgetReason FSlateAttributeMetaData::FGetterItem::GetInvalidationReason(const SWidget& OwningWidget, EInvalidateWidgetReason Reason) const
+TArray<FName> FSlateAttributeMetaData::GetAttributeNames(const SWidget& OwningWidget) const
+{
+	TArray<FName> Names;
+	Names.Reserve(Attributes.Num());
+	for (const FGetterItem& Getter : Attributes)
+	{
+		const FName Name = Getter.GetAttributeName(OwningWidget);
+		if (Name.IsValid())
+		{
+			Names.Add(Name);
+		}
+	}
+	return Names;
+}
+
+
+FSlateAttributeMetaData::FGetterItem::FInvalidationDetail FSlateAttributeMetaData::FGetterItem::GetInvalidationDetail(const SWidget& OwningWidget, EInvalidateWidgetReason Reason) const
+{
+	if (CachedAttributeDescriptorIndex != INDEX_NONE)
+	{
+		const FSlateAttributeDescriptor::FAttribute& DescriptorAttribute = OwningWidget.GetWidgetClass().GetAttributeDescriptor().GetAttributeAtIndex(CachedAttributeDescriptorIndex);
+		return FInvalidationDetail{&DescriptorAttribute.OnInvalidation, DescriptorAttribute.InvalidationReason.Get(OwningWidget)};
+	}
+	return FInvalidationDetail{nullptr, Reason};
+}
+
+
+FName FSlateAttributeMetaData::FGetterItem::GetAttributeName(const SWidget& OwningWidget) const
 {
 	if (CachedAttributeDescriptorIndex != INDEX_NONE)
 	{
-		return OwningWidget.GetWidgetClass().GetAttributeDescriptor().GetAttributeAtIndex(CachedAttributeDescriptorIndex).InvalidationReason.Get(OwningWidget);
+		const FSlateAttributeDescriptor::FAttribute& DescriptorAttribute = OwningWidget.GetWidgetClass().GetAttributeDescriptor().GetAttributeAtIndex(CachedAttributeDescriptorIndex);
+		return DescriptorAttribute.Name;
 	}
-	return Reason;
+	return FName();
 }
 
 
@@ -186,13 +214,19 @@ void FSlateAttributeMetaData::InvalidateWidget(SWidget& OwningWidget, const FSla
 		return;
 	}
 
+	const FSlateAttributeDescriptor::FInvalidationDelegate* OnInvalidationCallback = nullptr;
+
 	if (FSlateAttributeMetaData* AttributeMetaData = FSlateAttributeMetaData::FindMetaData(OwningWidget))
 	{
 		const int32 FoundIndex = AttributeMetaData->IndexOfAttribute(Attribute);
 		if (FoundIndex != INDEX_NONE)
 		{
 			FGetterItem& GetterItem = AttributeMetaData->Attributes[FoundIndex];
-			Reason = GetterItem.GetInvalidationReason(OwningWidget, Reason);
+			{
+				const FGetterItem::FInvalidationDetail Detail = GetterItem.GetInvalidationDetail(OwningWidget, Reason);
+				OnInvalidationCallback = Detail.Get<0>();
+				Reason = Detail.Get<1>();
+			}
 
 			// The dependency attribute need to be updated in the update loop (note that it may not be registered yet)
 			if (GetterItem.bIsADependencyForSomeoneElse)
@@ -209,6 +243,7 @@ void FSlateAttributeMetaData::InvalidateWidget(SWidget& OwningWidget, const FSla
 			if (FSlateAttributeDescriptor::FAttribute const* FoundAttribute = AttributeDescriptor.FindMemberAttribute(Offset))
 			{
 				Reason = FoundAttribute->InvalidationReason.Get(OwningWidget);
+				OnInvalidationCallback = &FoundAttribute->OnInvalidation;
 
 				if (FoundAttribute->bIsADependencyForSomeoneElse)
 				{
@@ -235,10 +270,15 @@ void FSlateAttributeMetaData::InvalidateWidget(SWidget& OwningWidget, const FSla
 		if (FSlateAttributeDescriptor::FAttribute const* FoundAttribute = OwningWidget.GetWidgetClass().GetAttributeDescriptor().FindMemberAttribute(Offset))
 		{
 			Reason = FoundAttribute->InvalidationReason.Get(OwningWidget);
+			OnInvalidationCallback = &FoundAttribute->OnInvalidation;
 		}
 	}
 
 	OwningWidget.Invalidate(Reason);
+	if (OnInvalidationCallback)
+	{
+		OnInvalidationCallback->ExecuteIfBound(OwningWidget);
+	}
 }
 
 
@@ -258,7 +298,7 @@ void FSlateAttributeMetaData::UpdateAttributes(SWidget& OwningWidget, bool bAllo
 {
 	if (FSlateAttributeMetaData* AttributeMetaData = FSlateAttributeMetaData::FindMetaData(OwningWidget))
 	{
-		AttributeMetaData->UpdateAttributes(OwningWidget, false, bAllowInvalidation);
+		AttributeMetaData->UpdateAttributesImpl(OwningWidget, false, bAllowInvalidation);
 	}
 }
 
@@ -270,14 +310,17 @@ void FSlateAttributeMetaData::UpdateCollapsedAttributes(SWidget& OwningWidget, b
 		// Does it have any collapsed attributes
 		if (AttributeMetaData->CollaspedAttributeCounter > 0)
 		{
-			AttributeMetaData->UpdateAttributes(OwningWidget, true, bAllowInvalidation);
+			AttributeMetaData->UpdateAttributesImpl(OwningWidget, true, bAllowInvalidation);
 		}
 	}
 }
 
 
-void FSlateAttributeMetaData::UpdateAttributes(SWidget& OwningWidget, bool bOnlyCollapsed, bool bAllowInvalidation)
+void FSlateAttributeMetaData::UpdateAttributesImpl(SWidget& OwningWidget, bool bOnlyCollapsed, bool bAllowInvalidation)
 {
+	// List of all the callback we need to execute after the invalidation
+	TArray<const FSlateAttributeDescriptor::FInvalidationDelegate*, TInlineAllocator<8>> AllInvalidationDelegates;
+
 	EInvalidateWidgetReason InvalidationReason = EInvalidateWidgetReason::None;
 	for (int32 Index = 0; Index < Attributes.Num(); ++Index)
 	{
@@ -320,7 +363,12 @@ void FSlateAttributeMetaData::UpdateAttributes(SWidget& OwningWidget, bool bOnly
 		GetterItem.bUpdatedThisFrame = Result.bInvalidationRequested;
 		if (Result.bInvalidationRequested && bAllowInvalidation)
 		{
-			InvalidationReason |= GetterItem.GetInvalidationReason(OwningWidget, Result.InvalidationReason);
+			const FGetterItem::FInvalidationDetail Detail = GetterItem.GetInvalidationDetail(OwningWidget, Result.InvalidationReason);
+			if (Detail.Get<0>())
+			{
+				AllInvalidationDelegates.Add(Detail.Get<0>());
+			}
+			InvalidationReason |= Detail.Get<1>();
 		}
 	}
 
@@ -336,6 +384,10 @@ void FSlateAttributeMetaData::UpdateAttributes(SWidget& OwningWidget, bool bOnly
 	if (bAllowInvalidation)
 	{
 		OwningWidget.Invalidate(InvalidationReason);
+		for (const FSlateAttributeDescriptor::FInvalidationDelegate* OnInvalidation : AllInvalidationDelegates)
+		{
+			OnInvalidation->ExecuteIfBound(OwningWidget);
+		}
 	}
 }
 
@@ -355,7 +407,12 @@ void FSlateAttributeMetaData::UpdateAttribute(SWidget& OwningWidget, FSlateAttri
 			{
 				if (OwningWidget.IsConstructionCompleted())
 				{
-					OwningWidget.Invalidate(GetterItem.GetInvalidationReason(OwningWidget, Result.InvalidationReason));
+					const FGetterItem::FInvalidationDetail Detail = GetterItem.GetInvalidationDetail(OwningWidget, Result.InvalidationReason);
+					OwningWidget.Invalidate(Detail.Get<1>());
+					if (Detail.Get<0>())
+					{
+						Detail.Get<0>()->ExecuteIfBound(OwningWidget);
+					}
 				}
 
 				// The dependency attribute need to be updated in the update loop (note that it may not be registered yet)
diff --git a/Engine/Source/Runtime/SlateCore/Public/Types/SlateAttributeDescriptor.h b/Engine/Source/Runtime/SlateCore/Public/Types/SlateAttributeDescriptor.h
--- a/Engine/Source/Runtime/SlateCore/Public/Types/SlateAttributeDescriptor.h
+++ b/Engine/Source/Runtime/SlateCore/Public/Types/SlateAttributeDescriptor.h
@@ -66,6 +66,20 @@ public:
 		FGetter Getter;
 	};
 
+	/** */
+	DECLARE_DELEGATE_OneParam(FInvalidationDelegate, SWidget& /*Widget*/);
+
+	/** */
+	enum class EInvalidationDelegateOverrideType
+	{
+		/** Replace the callback that the base class defined. */
+		ReplacePrevious,
+		/** Execute the callback that the base class defined, then execute the new callback. */
+		ExecuteAfterPrevious,
+		/** Execute the new callback, then execute the callback that the base class defined. */
+		ExecuteBeforePrevious,
+	};
+
 public:
 	using OffsetType = uint32;
 
@@ -80,6 +94,7 @@ public:
 		FName Prerequisite;
 		uint32 SortOrder = 0;
 		FInvalidateWidgetReasonAttribute InvalidationReason = FInvalidateWidgetReasonAttribute(EInvalidateWidgetReason::None);
+		FInvalidationDelegate OnInvalidation;
 		bool bIsMemberAttribute = false;
 		bool bIsPrerequisiteAlsoADependency = false;
 		bool bIsADependencyForSomeoneElse = false;
@@ -123,6 +138,13 @@ public:
 			 */
 			FAttributeEntry& UpdateWhenCollapsed();
 
+			/**
+			 * Be notified when the attribute value changed and the invalidation occurs.
+			 * It will not be called when the SWidget is in his construction phase.
+			 * @see SWidget::IsConstructionCompleted
+			 */
+			FAttributeEntry& OnInvalidation(FInvalidationDelegate Callback);
+
 		private:
 			FSlateAttributeDescriptor& Descriptor;
 			int32 AttributeIndex;
@@ -131,9 +153,15 @@ public:
 		FAttributeEntry AddMemberAttribute(FName AttributeName, OffsetType Offset, const FInvalidateWidgetReasonAttribute& ReasonGetter);
 		FAttributeEntry AddMemberAttribute(FName AttributeName, OffsetType Offset, FInvalidateWidgetReasonAttribute&& ReasonGetter);
 
+		/** Change the InvalidationReason of an attribute defined in a base class. */
 		void OverrideInvalidationReason(FName AttributeName, const FInvalidateWidgetReasonAttribute& Reason);
+		/** Change the InvalidationReason of an attribute defined in a base class. */
 		void OverrideInvalidationReason(FName AttributeName, FInvalidateWidgetReasonAttribute&& Reason);
 
+		/** Change the FInvalidationDelegate of an attribute defined in a base class. */
+		void OverrideOnInvalidation(FName AttributeName, EInvalidationDelegateOverrideType OverrideType, FInvalidationDelegate Callback);
+
+		/** Change the update type of an attribute defined in a base class. */
 		void SetUpdateWhenCollapsed(FName AttributeName, bool bUpdateWhenCollapsed);
 
 
@@ -176,6 +204,7 @@ private:
 
 	FInitializer::FAttributeEntry AddMemberAttribute(FName AttributeName, OffsetType Offset, FInvalidateWidgetReasonAttribute ReasonGetter);
 	void OverrideInvalidationReason(FName AttributeName, FInvalidateWidgetReasonAttribute ReasonGetter);
+	void OverrideOnInvalidation(FName AttributeName, EInvalidationDelegateOverrideType OverrideType, FInvalidationDelegate Callback);
 	void SetPrerequisite(FAttribute& Attribute, FName Prerequisite, bool bSetAsDependency);
 	void SetUpdateWhenCollapsed(FAttribute& Attribute, bool bUpdate);
 
diff --git a/Engine/Source/Runtime/SlateCore/Public/Types/SlateAttributeMetaData.h b/Engine/Source/Runtime/SlateCore/Public/Types/SlateAttributeMetaData.h
--- a/Engine/Source/Runtime/SlateCore/Public/Types/SlateAttributeMetaData.h
+++ b/Engine/Source/Runtime/SlateCore/Public/Types/SlateAttributeMetaData.h
@@ -6,6 +6,7 @@
 #include "Algo/BinarySearch.h"
 #include "Types/ISlateMetaData.h"
 #include "Types/SlateAttribute.h"
+#include "Types/SlateAttributeDescriptor.h"
 #include "Widgets/InvalidateWidgetReason.h"
 
 
@@ -57,6 +58,9 @@ public:
 
 	int32 RegisteredCollaspedAttributeCount() const { return CollaspedAttributeCounter; }
 
+	/** Get the name of all the atrributes, if available. */
+	TArray<FName> GetAttributeNames(const SWidget& OwningWidget) const;
+
 private:
 	using ESlateAttributeType = SlateAttributePrivate::ESlateAttributeType;
 	using ISlateAttributeGetter = SlateAttributePrivate::ISlateAttributeGetter;
@@ -72,7 +76,7 @@ private:
 private:
 	void RegisterAttributeImpl(SWidget& OwningWidget, FSlateAttributeBase& Attribute, ESlateAttributeType AttributeType, TUniquePtr<ISlateAttributeGetter>&& Getter);
 	bool UnregisterAttributeImpl(const FSlateAttributeBase& Attribute);
-	void UpdateAttributes(SWidget& OwningWidget, bool bOnlyCollapsed, bool bAllowInvalidation);
+	void UpdateAttributesImpl(SWidget& OwningWidget, bool bOnlyCollapsed, bool bAllowInvalidation);
 
 private:
 	int32 IndexOfAttribute(const FSlateAttributeBase& Attribute) const
@@ -128,7 +132,11 @@ private:
 			return SortOrder < Other.SortOrder;
 		}
 
-		EInvalidateWidgetReason GetInvalidationReason(const SWidget&, EInvalidateWidgetReason Reason) const;
+		using FInvalidationDetail = TTuple<const FSlateAttributeDescriptor::FInvalidationDelegate*, EInvalidateWidgetReason>;
+		FInvalidationDetail GetInvalidationDetail(const SWidget&, EInvalidateWidgetReason Reason) const;
+
+		/** If available, return the name of the attribute. */
+		FName GetAttributeName(const SWidget& OwningWidget) const;
 	};
 
 	TArray<FGetterItem, TInlineAllocator<4>> Attributes;
