Slate: Convert SWidget's Enabled and Visibility flag to SlateAttribute.

diff --git a/Engine/Plugins/Experimental/ControlRig/Source/ControlRigEditor/Private/Graph/SControlRigGraphNode.cpp b/Engine/Plugins/Experimental/ControlRig/Source/ControlRigEditor/Private/Graph/SControlRigGraphNode.cpp
--- a/Engine/Plugins/Experimental/ControlRig/Source/ControlRigEditor/Private/Graph/SControlRigGraphNode.cpp
+++ b/Engine/Plugins/Experimental/ControlRig/Source/ControlRigEditor/Private/Graph/SControlRigGraphNode.cpp
@@ -390,13 +390,11 @@ TSharedRef<SWidget> SControlRigGraphNode::CreateTitleWidget(TSharedPtr<SNodeTitl
 {
 	NodeTitle = InNodeTitle;
 
-	auto WidgetRef = SGraphNode::CreateTitleWidget(NodeTitle);
-	auto VisibilityAttribute = TAttribute<EVisibility>::Create(
-		TAttribute<EVisibility>::FGetter::CreateSP(this, &SControlRigGraphNode::GetTitleVisibility));
-	WidgetRef->SetVisibility(VisibilityAttribute);
+	TSharedRef<SWidget> WidgetRef = SGraphNode::CreateTitleWidget(NodeTitle);
+	WidgetRef->SetVisibility(MakeAttributeSP(this, &SControlRigGraphNode::GetTitleVisibility));
 	if (NodeTitle.IsValid())
 	{
-		NodeTitle->SetVisibility(VisibilityAttribute);
+		NodeTitle->SetVisibility(MakeAttributeSP(this, &SControlRigGraphNode::GetTitleVisibility));
 	}
 
 	return SNew(SHorizontalBox)
diff --git a/Engine/Plugins/NetcodeUnitTest/NetcodeUnitTest/Source/NetcodeUnitTest/Private/UI/SLogWidget.cpp b/Engine/Plugins/NetcodeUnitTest/NetcodeUnitTest/Source/NetcodeUnitTest/Private/UI/SLogWidget.cpp
--- a/Engine/Plugins/NetcodeUnitTest/NetcodeUnitTest/Source/NetcodeUnitTest/Private/UI/SLogWidget.cpp
+++ b/Engine/Plugins/NetcodeUnitTest/NetcodeUnitTest/Source/NetcodeUnitTest/Private/UI/SLogWidget.cpp
@@ -46,7 +46,7 @@ IMPLEMENT_GET_PRIVATE_VAR(SEditableTextBox, EditableText, TSharedPtr<SEditableTe
 IMPLEMENT_GET_PRIVATE_VAR(SButton, Style, const FButtonStyle*);
 
 // Enable reading SBorder::GetBorderBackgroundColor
-IMPLEMENT_GET_PROTECTED_FUNC_CONST(SBorder, GetBorderBackgroundColor, const FSlateColor&, void,, const);
+IMPLEMENT_GET_PROTECTED_FUNC_CONST(SBorder, GetBorderBackgroundColor, FSlateColor, void,, const);
 
 // Enable access to SDockTab::GetCurrentStyle, using the CALL_PROTECTED macro
 IMPLEMENT_GET_PROTECTED_FUNC_CONST(SDockTab, GetCurrentStyle, const FDockTabStyle&, void,, const);
diff --git a/Engine/Source/Editor/GameProjectGeneration/Private/SGetSuggestedIDEWidget.cpp b/Engine/Source/Editor/GameProjectGeneration/Private/SGetSuggestedIDEWidget.cpp
--- a/Engine/Source/Editor/GameProjectGeneration/Private/SGetSuggestedIDEWidget.cpp
+++ b/Engine/Source/Editor/GameProjectGeneration/Private/SGetSuggestedIDEWidget.cpp
@@ -18,7 +18,7 @@ TSharedPtr<SNotificationItem> SGetSuggestedIDEWidget::IDEDownloadNotification;
 
 void SGetSuggestedIDEWidget::Construct(const FArguments& InArgs)
 {
-	Visibility = InArgs._VisibilityOverride.IsSet() ? InArgs._VisibilityOverride : TAttribute<EVisibility>(this, &SGetSuggestedIDEWidget::GetVisibility);
+	SetVisibility(InArgs._VisibilityOverride.IsSet() ? InArgs._VisibilityOverride : TAttribute<EVisibility>(this, &SGetSuggestedIDEWidget::GetVisibility));
 
 	ChildSlot
 	[
diff --git a/Engine/Source/Editor/GraphEditor/Private/KismetNodes/SGraphNodeK2Copy.cpp b/Engine/Source/Editor/GraphEditor/Private/KismetNodes/SGraphNodeK2Copy.cpp
--- a/Engine/Source/Editor/GraphEditor/Private/KismetNodes/SGraphNodeK2Copy.cpp
+++ b/Engine/Source/Editor/GraphEditor/Private/KismetNodes/SGraphNodeK2Copy.cpp
@@ -33,7 +33,7 @@ public:
 		bUsePinColorForText = InArgs._UsePinColorForText;
 		this->SetCursor(EMouseCursor::Default);
 
-		Visibility = TAttribute<EVisibility>(this, &SCopyNodeGraphPin::GetPinVisiblity);
+		SetVisibility(MakeAttributeSP(this, &SCopyNodeGraphPin::GetPinVisiblity));
 
 		GraphPinObj = InPin;
 		check(GraphPinObj != NULL);
diff --git a/Engine/Source/Editor/GraphEditor/Private/KismetNodes/SGraphNodeK2Event.cpp b/Engine/Source/Editor/GraphEditor/Private/KismetNodes/SGraphNodeK2Event.cpp
--- a/Engine/Source/Editor/GraphEditor/Private/KismetNodes/SGraphNodeK2Event.cpp
+++ b/Engine/Source/Editor/GraphEditor/Private/KismetNodes/SGraphNodeK2Event.cpp
@@ -43,13 +43,11 @@ EVisibility SGraphNodeK2Event::GetTitleVisibility() const
 
 TSharedRef<SWidget> SGraphNodeK2Event::CreateTitleWidget(TSharedPtr<SNodeTitle> NodeTitle)
 {
-	auto WidgetRef = SGraphNodeK2Default::CreateTitleWidget(NodeTitle);
-	auto VisibilityAttribute = TAttribute<EVisibility>::Create(
-		TAttribute<EVisibility>::FGetter::CreateSP(this, &SGraphNodeK2Event::GetTitleVisibility));
-	WidgetRef->SetVisibility(VisibilityAttribute);
+	TSharedRef<SWidget> WidgetRef = SGraphNodeK2Default::CreateTitleWidget(NodeTitle);
+	WidgetRef->SetVisibility(MakeAttributeSP(this, &SGraphNodeK2Event::GetTitleVisibility));
 	if (NodeTitle.IsValid())
 	{
-		NodeTitle->SetVisibility(VisibilityAttribute);
+		NodeTitle->SetVisibility(MakeAttributeSP(this, &SGraphNodeK2Event::GetTitleVisibility));
 	}
 
 	return WidgetRef;
diff --git a/Engine/Source/Editor/GraphEditor/Private/SGraphPin.cpp b/Engine/Source/Editor/GraphEditor/Private/SGraphPin.cpp
--- a/Engine/Source/Editor/GraphEditor/Private/SGraphPin.cpp
+++ b/Engine/Source/Editor/GraphEditor/Private/SGraphPin.cpp
@@ -150,7 +150,7 @@ void SGraphPin::Construct(const FArguments& InArgs, UEdGraphPin* InPin)
 	bUsePinColorForText = InArgs._UsePinColorForText;
 	this->SetCursor(EMouseCursor::Default);
 
-	Visibility = TAttribute<EVisibility>(this, &SGraphPin::GetPinVisiblity);
+	SetVisibility(MakeAttributeSP(this, &SGraphPin::GetPinVisiblity));
 
 	GraphPinObj = InPin;
 	check(GraphPinObj != NULL);
@@ -168,10 +168,10 @@ void SGraphPin::Construct(const FArguments& InArgs, UEdGraphPin* InPin)
 
 	// Create the pin icon widget
 	TSharedRef<SWidget> PinWidgetRef = SPinTypeSelector::ConstructPinTypeImage(
-		TAttribute<const FSlateBrush*>::Create( TAttribute<const FSlateBrush*>::FGetter::CreateRaw(this, &SGraphPin::GetPinIcon ) ),
-		TAttribute<FSlateColor>::Create( TAttribute<FSlateColor>::FGetter::CreateRaw(this, &SGraphPin::GetPinColor) ),
-		TAttribute<const FSlateBrush*>::Create( TAttribute<const FSlateBrush*>::FGetter::CreateRaw(this, &SGraphPin::GetSecondaryPinIcon ) ),
-		TAttribute<FSlateColor>::Create( TAttribute<FSlateColor>::FGetter::CreateRaw(this, &SGraphPin::GetSecondaryPinColor) ));
+		MakeAttributeSP(this, &SGraphPin::GetPinIcon ),
+		MakeAttributeSP(this, &SGraphPin::GetPinColor),
+		MakeAttributeSP(this, &SGraphPin::GetSecondaryPinIcon),
+		MakeAttributeSP(this, &SGraphPin::GetSecondaryPinColor));
 	PinImage = PinWidgetRef;
 
 	PinWidgetRef->SetCursor( 
diff --git a/Engine/Source/Editor/MaterialEditor/Private/SMaterialEditorStatsWidget.cpp b/Engine/Source/Editor/MaterialEditor/Private/SMaterialEditorStatsWidget.cpp
--- a/Engine/Source/Editor/MaterialEditor/Private/SMaterialEditorStatsWidget.cpp
+++ b/Engine/Source/Editor/MaterialEditor/Private/SMaterialEditorStatsWidget.cpp
@@ -574,7 +574,7 @@ void SMaterialEditorStatsWidget::Construct(const FArguments& InArgs)
 
 	auto MessageArea = BuildMessageArea();
 
-	Visibility = EVisibility::SelfHitTestInvisible;
+	SetVisibility(EVisibility::SelfHitTestInvisible);
 
 	// construct default column headers
 	PlatformColumnHeader = SNew(SHeaderRow);
diff --git a/Engine/Source/Editor/Persona/Private/SAnimNotifyPanel.cpp b/Engine/Source/Editor/Persona/Private/SAnimNotifyPanel.cpp
--- a/Engine/Source/Editor/Persona/Private/SAnimNotifyPanel.cpp
+++ b/Engine/Source/Editor/Persona/Private/SAnimNotifyPanel.cpp
@@ -667,7 +667,7 @@ void SAnimNotifyPair::Construct(const FArguments& InArgs)
 
 	float ScaleMult = 1.0f;
 	FVector2D NodeSize = NodePtr->ComputeDesiredSize(ScaleMult);
-	Visibility = EVisibility::SelfHitTestInvisible;
+	SetVisibility(EVisibility::SelfHitTestInvisible);
 
 	this->ChildSlot
 		[
diff --git a/Engine/Source/Editor/PropertyEditor/Private/DetailItemNode.cpp b/Engine/Source/Editor/PropertyEditor/Private/DetailItemNode.cpp
--- a/Engine/Source/Editor/PropertyEditor/Private/DetailItemNode.cpp
+++ b/Engine/Source/Editor/PropertyEditor/Private/DetailItemNode.cpp
@@ -711,21 +711,23 @@ FPropertyPath FDetailItemNode::GetPropertyPath() const
 
 TAttribute<bool> FDetailItemNode::IsPropertyEditingEnabled() const
 {
-	return TAttribute<bool>::Create([this]()
-		{
-			bool IsParentEnabledValue = IsParentEnabled.Get(true);
-			if (Customization.HasCustomWidget())
-			{
-				IDetailsViewPrivate* DetailsView = GetDetailsView();
-				if (DetailsView)
-				{
-					return IsParentEnabledValue && 
-						!DetailsView->IsCustomRowReadOnly(FName(*Customization.WidgetDecl->FilterTextString.ToString()), FName(*GetParentCategory()->GetDisplayName().ToString()));
-				}
-			}
+	return MakeAttributeSP(this, &FDetailItemNode::IsPropertyEditingEnabledImpl);
+}
 
-			return IsParentEnabledValue;
-		});
+bool FDetailItemNode::IsPropertyEditingEnabledImpl() const
+{
+	bool IsParentEnabledValue = IsParentEnabled.Get(true);
+	if (Customization.HasCustomWidget())
+	{
+		IDetailsViewPrivate* DetailsView = GetDetailsView();
+		if (DetailsView)
+		{
+			return IsParentEnabledValue &&
+				!DetailsView->IsCustomRowReadOnly(FName(*Customization.WidgetDecl->FilterTextString.ToString()), FName(*GetParentCategory()->GetDisplayName().ToString()));
+		}
+	}
+	
+	return IsParentEnabledValue;
 }
 
 TSharedPtr<FPropertyNode> FDetailItemNode::GetPropertyNode() const
diff --git a/Engine/Source/Editor/PropertyEditor/Private/DetailItemNode.h b/Engine/Source/Editor/PropertyEditor/Private/DetailItemNode.h
--- a/Engine/Source/Editor/PropertyEditor/Private/DetailItemNode.h
+++ b/Engine/Source/Editor/PropertyEditor/Private/DetailItemNode.h
@@ -97,6 +97,11 @@ private:
 	 */
 	void InitGroup();
 
+	/**
+	 * Implementation of IsPropertyEditingEnabled
+	 */
+	bool IsPropertyEditingEnabledImpl() const;
+
 private:
 	/** Customization on this node */
 	FDetailLayoutCustomization Customization;
diff --git a/Engine/Source/Editor/PropertyEditor/Public/DetailWidgetRow.h b/Engine/Source/Editor/PropertyEditor/Public/DetailWidgetRow.h
--- a/Engine/Source/Editor/PropertyEditor/Public/DetailWidgetRow.h
+++ b/Engine/Source/Editor/PropertyEditor/Public/DetailWidgetRow.h
@@ -79,7 +79,7 @@ private:
 
 		void Construct( const FArguments& InArgs )
 		{
-			Visibility = EVisibility::Collapsed;
+			SetVisibility(EVisibility::Collapsed);
 		}
 
 	};
diff --git a/Engine/Source/Runtime/GameMenuBuilder/Private/SGameMenuPageWidget.cpp b/Engine/Source/Runtime/GameMenuBuilder/Private/SGameMenuPageWidget.cpp
--- a/Engine/Source/Runtime/GameMenuBuilder/Private/SGameMenuPageWidget.cpp
+++ b/Engine/Source/Runtime/GameMenuBuilder/Private/SGameMenuPageWidget.cpp
@@ -157,7 +157,7 @@ void SGameMenuPageWidget::Construct(const FArguments& InArgs)
 	bGameMenu = InArgs._GameMenu;
 	UIScale.Bind(this, &SGameMenuPageWidget::GetUIScale);
 	//ControllerHideMenuKey = EKeys::Gamepad_Special_Right;
-	Visibility.Bind(this, &SGameMenuPageWidget::GetSlateVisibility);
+	SetVisibility(MakeAttributeSP(this, &SGameMenuPageWidget::GetSlateVisibility));
 		
 	// Create the title widget
 	TSharedRef<SHorizontalBox> TitleBoxWidget = SNew(SHorizontalBox)
diff --git a/Engine/Source/Runtime/Slate/Private/Framework/Application/MenuStack.cpp b/Engine/Source/Runtime/Slate/Private/Framework/Application/MenuStack.cpp
--- a/Engine/Source/Runtime/Slate/Private/Framework/Application/MenuStack.cpp
+++ b/Engine/Source/Runtime/Slate/Private/Framework/Application/MenuStack.cpp
@@ -95,7 +95,7 @@ namespace MenuStackInternal
 		void Construct(const FArguments& InArgs)
 		{
 			// The visibility of the content wrapper should match that of the provided content
-			Visibility = AccessWidgetVisibilityAttribute(InArgs._MenuContent.Widget);
+			SetVisibility(AccessWidgetVisibilityAttribute(InArgs._MenuContent.Widget));
 
 			OnKeyDownDelegate = InArgs._OnKeyDown;
 			OnMenuLostFocus = InArgs._OnMenuLostFocus;
diff --git a/Engine/Source/Runtime/Slate/Private/Framework/Docking/SDockingNode.cpp b/Engine/Source/Runtime/Slate/Private/Framework/Docking/SDockingNode.cpp
--- a/Engine/Source/Runtime/Slate/Private/Framework/Docking/SDockingNode.cpp
+++ b/Engine/Source/Runtime/Slate/Private/Framework/Docking/SDockingNode.cpp
@@ -31,7 +31,7 @@ void SDockingNode::SetSizeCoefficient( float InSizeCoefficient )
 
 void SDockingNode::OnLiveTabAdded()
 {
-	this->Visibility = EVisibility::Visible;
+	SetVisibility(EVisibility::Visible);
 
 	TSharedPtr<SDockingNode> ParentNode = ParentNodePtr.Pin();
 	if (ParentNode.IsValid())
diff --git a/Engine/Source/Runtime/Slate/Private/Framework/Docking/SDockingSplitter.cpp b/Engine/Source/Runtime/Slate/Private/Framework/Docking/SDockingSplitter.cpp
--- a/Engine/Source/Runtime/Slate/Private/Framework/Docking/SDockingSplitter.cpp
+++ b/Engine/Source/Runtime/Slate/Private/Framework/Docking/SDockingSplitter.cpp
@@ -181,7 +181,7 @@ SDockingNode::ECleanupRetVal SDockingSplitter::CleanUpNodes()
 		const bool bIsDockArea = !(this->ParentNodePtr.IsValid());
 		if (!bIsDockArea)
 		{
-			this->Visibility = EVisibility::Collapsed;
+			SetVisibility(EVisibility::Collapsed);
 		}		
 	}
 		
diff --git a/Engine/Source/Runtime/Slate/Private/Framework/Docking/SDockingTabStack.cpp b/Engine/Source/Runtime/Slate/Private/Framework/Docking/SDockingTabStack.cpp
--- a/Engine/Source/Runtime/Slate/Private/Framework/Docking/SDockingTabStack.cpp
+++ b/Engine/Source/Runtime/Slate/Private/Framework/Docking/SDockingTabStack.cpp
@@ -604,7 +604,7 @@ SDockingNode::ECleanupRetVal SDockingTabStack::CleanUpNodes()
 	}
 	else if (Tabs.Num() > 0)
 	{
-		Visibility = EVisibility::Collapsed;
+		SetVisibility(EVisibility::Collapsed);
 		return HistoryTabsUnderNode;
 	}
 	else
diff --git a/Engine/Source/Runtime/Slate/Private/Widgets/Notifications/SErrorHint.cpp b/Engine/Source/Runtime/Slate/Private/Widgets/Notifications/SErrorHint.cpp
--- a/Engine/Source/Runtime/Slate/Private/Widgets/Notifications/SErrorHint.cpp
+++ b/Engine/Source/Runtime/Slate/Private/Widgets/Notifications/SErrorHint.cpp
@@ -10,8 +10,8 @@ void SErrorHint::Construct(const FArguments& InArgs)
 
 	ExpandAnimation = FCurveSequence(0.0f, 0.15f);
 
-	CustomVisibility = Visibility;
-	Visibility = TAttribute<EVisibility>( SharedThis(this), &SErrorHint::MyVisibility );
+	CustomVisibility = GetVisibilityAttribute().ToAttribute();
+	SetVisibility(MakeAttributeSP(this, &SErrorHint::MyVisibility));
 
 	ChildSlot
 	.HAlign(HAlign_Center)
diff --git a/Engine/Source/Runtime/Slate/Private/Widgets/Notifications/SErrorText.cpp b/Engine/Source/Runtime/Slate/Private/Widgets/Notifications/SErrorText.cpp
--- a/Engine/Source/Runtime/Slate/Private/Widgets/Notifications/SErrorText.cpp
+++ b/Engine/Source/Runtime/Slate/Private/Widgets/Notifications/SErrorText.cpp
@@ -7,8 +7,8 @@ void SErrorText::Construct(const FArguments& InArgs)
 {
 	ExpandAnimation = FCurveSequence(0.0f, 0.15f);
 
-	CustomVisibility = Visibility;
-	Visibility = TAttribute<EVisibility>( SharedThis(this), &SErrorText::MyVisibility );
+	CustomVisibility = GetVisibilityAttribute().ToAttribute();
+	SetVisibility(MakeAttributeSP(this, &SErrorText::MyVisibility));
 
 	SBorder::Construct( SBorder::FArguments()
 		.BorderBackgroundColor( InArgs._BackgroundColor )
diff --git a/Engine/Source/Runtime/Slate/Public/Widgets/Layout/SBorder.h b/Engine/Source/Runtime/Slate/Public/Widgets/Layout/SBorder.h
--- a/Engine/Source/Runtime/Slate/Public/Widgets/Layout/SBorder.h
+++ b/Engine/Source/Runtime/Slate/Public/Widgets/Layout/SBorder.h
@@ -136,11 +136,11 @@ protected:
 	virtual FVector2D ComputeDesiredSize(float) const override;
 	//~End SWidget overrides.
 
-	const FSlateBrush* GetBorderImageAttribute() const { return BorderImageAttribute.Get(); }
-	const FSlateColor& GetBorderBackgroundColor() const { return BorderBackgroundColorAttribute.Get(); }
-	const FVector2D& GetDesiredSizeScale() const { return DesiredSizeScaleAttribute.Get(); }
+	const FSlateBrush* GetBorderImage() const { return BorderImageAttribute.Get(); }
+	FSlateColor GetBorderBackgroundColor() const { return BorderBackgroundColorAttribute.Get(); }
+	FVector2D GetDesiredSizeScale() const { return DesiredSizeScaleAttribute.Get(); }
 	bool GetShowDisabledEffect() const { return ShowDisabledEffectAttribute.Get(); }
-	TSlateAttributeRef<const FSlateBrush*> GetBorderImageAttributeAttribute() const { return TSlateAttributeRef<const FSlateBrush*>(*this, BorderImageAttribute); }
+	TSlateAttributeRef<const FSlateBrush*> GetBorderImageAttribute() const { return TSlateAttributeRef<const FSlateBrush*>(*this, BorderImageAttribute); }
 	TSlateAttributeRef<FSlateColor> GetBorderBackgroundColorAttribute() const { return TSlateAttributeRef<FSlateColor>(*this, BorderBackgroundColorAttribute); }
 	TSlateAttributeRef<FVector2D> GetDesiredSizeScaleAttribute() const { return TSlateAttributeRef<FVector2D>(*this, DesiredSizeScaleAttribute); }
 	TSlateAttributeRef<bool> GetShowDisabledEffectAttribute() const { return TSlateAttributeRef<bool>(*this, ShowDisabledEffectAttribute); }
diff --git a/Engine/Source/Runtime/SlateCore/Private/Widgets/SNullWidget.cpp b/Engine/Source/Runtime/SlateCore/Private/Widgets/SNullWidget.cpp
--- a/Engine/Source/Runtime/SlateCore/Private/Widgets/SNullWidget.cpp
+++ b/Engine/Source/Runtime/SlateCore/Private/Widgets/SNullWidget.cpp
@@ -35,7 +35,7 @@ public:
 private:
 	virtual void SetVisibility( TAttribute<EVisibility> InVisibility ) override final
 	{
-		ensureMsgf( false, TEXT("Attempting to SetVisibility() on SNullWidget. Mutating SNullWidget is not allowed.") );
+		ensureMsgf(!IsConstructionCompleted(), TEXT("Attempting to SetVisibility() on SNullWidget. Mutating SNullWidget is not allowed.") );
 	}
 public:
 	
diff --git a/Engine/Source/Runtime/SlateCore/Private/Widgets/SWidget.cpp b/Engine/Source/Runtime/SlateCore/Private/Widgets/SWidget.cpp
--- a/Engine/Source/Runtime/SlateCore/Private/Widgets/SWidget.cpp
+++ b/Engine/Source/Runtime/SlateCore/Private/Widgets/SWidget.cpp
@@ -203,6 +203,9 @@ namespace SlateTraceMetaData
 SLATE_IMPLEMENT_WIDGET(SWidget)
 void SWidget::PrivateRegisterAttributes(FSlateAttributeInitializer& AttributeInitializer)
 {
+	SLATE_ADD_MEMBER_ATTRIBUTE_DEFINITION_WITH_NAME(AttributeInitializer, "EnabledState", EnabledStateAttribute, EInvalidateWidgetReason::Paint);
+	SLATE_ADD_MEMBER_ATTRIBUTE_DEFINITION_WITH_NAME(AttributeInitializer, "Visibility", VisibilityAttribute, EInvalidateWidgetReason::Visibility)
+		.UpdateWhenCollapsed();
 }
 
 SWidget::SWidget()
@@ -232,9 +235,9 @@ SWidget::SWidget()
 	// Note we are defaulting to tick for backwards compatibility
 	, UpdateFlags(EWidgetUpdateFlags::NeedsTick)
 	, DesiredSize()
+	, EnabledStateAttribute(*this, true)
+	, VisibilityAttribute(*this, EVisibility::Visible)
 	, CullingBoundsExtension()
-	, EnabledState(true)
-	, Visibility(EVisibility::Visible)
 	, RenderOpacity(1.0f)
 	, RenderTransform()
 	, RenderTransformPivot(FVector2D::ZeroVector)
@@ -354,8 +357,8 @@ PRAGMA_ENABLE_DEPRECATION_WARNINGS
 
 void SWidget::SWidgetConstruct(const FSlateBaseNamedArgs& Args)
 {
-	EnabledState = Args._IsEnabled;
-	Visibility = Args._Visibility;
+	SetEnabled(Args._IsEnabled);
+	VisibilityAttribute.Assign(*this, Args._Visibility); // SetVisibility is virtual, assign directly to stay backward compatible
 	RenderOpacity = Args._RenderOpacity;
 	RenderTransform = Args._RenderTransform;
 	RenderTransformPivot = Args._RenderTransformPivot;
@@ -1174,20 +1177,19 @@ bool SWidget::IsDirectlyHovered() const
 
 void SWidget::SetVisibility(TAttribute<EVisibility> InVisibility)
 {
-	SetAttribute(Visibility, InVisibility, EInvalidateWidgetReason::Visibility);
+	VisibilityAttribute.Assign(*this, MoveTemp(InVisibility));
 }
 
 void SWidget::Invalidate(EInvalidateWidgetReason InvalidateReason)
 {
 	SLATE_CROSS_THREAD_CHECK();
 
-	if (InvalidateReason == EInvalidateWidgetReason::None)
+	if (InvalidateReason == EInvalidateWidgetReason::None /*|| !IsConstructionCompleted()*/)
 	{
 		return;
 	}
 
 	SCOPED_NAMED_EVENT_TEXT("SWidget::Invalidate", FColor::Orange);
-	const bool bWasVolatile = IsVolatileIndirectly() || IsVolatile();
 
 	// Backwards compatibility fix:  Its no longer valid to just invalidate volatility since we need to repaint to cache elements if a widget becomes non-volatile. So after volatility changes force repaint
 	if (InvalidateReason == EInvalidateWidgetReason::Volatility)
diff --git a/Engine/Source/Runtime/SlateCore/Private/Widgets/SWindow.cpp b/Engine/Source/Runtime/SlateCore/Private/Widgets/SWindow.cpp
--- a/Engine/Source/Runtime/SlateCore/Private/Widgets/SWindow.cpp
+++ b/Engine/Source/Runtime/SlateCore/Private/Widgets/SWindow.cpp
@@ -2084,7 +2084,7 @@ int32 SWindow::PaintWindow( double CurrentTime, float DeltaTime, FSlateWindowEle
 	// Always set the window geometry and visibility
 	PersistentState.AllottedGeometry = GetWindowGeometryInWindow();
 	PersistentState.CullingBounds = GetClippingRectangleInWindow();
-	if (!Visibility.IsBound())
+	if (!GetVisibilityAttribute().IsBound())
 	{
 		SetVisibility(GetWindowVisibility());
 	}
diff --git a/Engine/Source/Runtime/SlateCore/Public/Types/SlateAttribute.h b/Engine/Source/Runtime/SlateCore/Public/Types/SlateAttribute.h
--- a/Engine/Source/Runtime/SlateCore/Public/Types/SlateAttribute.h
+++ b/Engine/Source/Runtime/SlateCore/Public/Types/SlateAttribute.h
@@ -221,6 +221,8 @@ struct FSlateDeprecatedTAttribute
 	FGetter GetBinding() const { return false; }
 
 	void Bind(const FGetter& InGetter) {}
+	template<class SourceType>
+	void Bind(SourceType* InUserObject, typename FGetter::template TSPMethodDelegate_Const< SourceType >::FMethodPtr InMethodPtr) {}
 	bool IsBound() const { return false; }
 
 	bool IdenticalTo(const TAttribute<ObjectType>& InOther) const { return false; }
diff --git a/Engine/Source/Runtime/SlateCore/Public/Widgets/SWidget.h b/Engine/Source/Runtime/SlateCore/Public/Widgets/SWidget.h
--- a/Engine/Source/Runtime/SlateCore/Public/Widgets/SWidget.h
+++ b/Engine/Source/Runtime/SlateCore/Public/Widgets/SWidget.h
@@ -936,15 +936,15 @@ public:
 	 *
 	 * @param InEnabledState	An attribute containing the enabled state or a delegate to call to get the enabled state.
 	 */
-	void SetEnabled(const TAttribute<bool>& InEnabledState)
+	void SetEnabled(TAttribute<bool> InEnabledState)
 	{
-		SetAttribute(EnabledState, InEnabledState, EInvalidateWidgetReason::Paint);
+		EnabledStateAttribute.Assign(*this, MoveTemp(InEnabledState));
 	}
 
 	/** @return Whether or not this widget is enabled */
 	FORCEINLINE bool IsEnabled() const
 	{
-		return EnabledState.Get();
+		return EnabledStateAttribute.Get();
 	}
 
 	/** @return Is this widget interactive or not? Defaults to false */
@@ -986,7 +986,7 @@ public:
 	 * @return is this widget visible, hidden or collapsed.
 	 * @note this widget can be visible but if a parent is hidden or collapsed, it would not show on screen.
 	 */
-	FORCEINLINE EVisibility GetVisibility() const { return Visibility.Get(); }
+	FORCEINLINE EVisibility GetVisibility() const { return VisibilityAttribute.Get(); }
 
 	/** @param InVisibility  should this widget be */
 	virtual void SetVisibility(TAttribute<EVisibility> InVisibility);
@@ -1563,7 +1563,7 @@ protected:
 	 */
 	virtual bool ComputeVolatility() const
 	{
-		return Visibility.IsBound() || EnabledState.IsBound() || RenderTransform.IsBound();
+		return RenderTransform.IsBound();
 	}
 
 	/**
@@ -1571,9 +1571,9 @@ protected:
 	 * 
 	 * @param Widget The widget to get the visibility attribute of
 	 */
-	static const TAttribute<EVisibility>& AccessWidgetVisibilityAttribute(const TSharedRef<SWidget>& Widget)
+	static TAttribute<EVisibility> AccessWidgetVisibilityAttribute(const TSharedRef<SWidget>& Widget)
 	{
-		return Widget->Visibility;
+		return Widget->VisibilityAttribute.ToAttribute(Widget.Get());
 	}
 
 	/**
@@ -1667,6 +1667,11 @@ protected:
 		return SetWidgetAttribute(*this, TargetValue, SourceValue, BaseInvalidationReason);
 	}
 
+	/** @return an attribute reference of EnabledStateAttribute */
+	TSlateAttributeRef<bool> GetEnabledStateAttribute() const { return TSlateAttributeRef<bool>(*this, EnabledStateAttribute); }
+	/** @return an attribute reference of VisibilityAttribute */
+	TSlateAttributeRef<EVisibility> GetVisibilityAttribute() const { return TSlateAttributeRef<EVisibility>(*this, VisibilityAttribute); }
+
 protected:
 	/** Dtor ensures that active timer handles are UnRegistered with the SlateApplication. */
 	virtual ~SWidget();
@@ -1765,24 +1770,26 @@ protected:
 	}
 
 private:
+
 	/** Flow direction preference */
 	EFlowDirectionPreference FlowDirectionPreference;
 
 	/** The different updates this widget needs next frame. */
 	EWidgetUpdateFlags UpdateFlags;
 
-#if WITH_SLATE_DEBUGGING
-	/** The last time this widget got painted. */
-	uint32 LastPaintFrame = 0;
-#endif
-
 	mutable FSlateWidgetPersistentState PersistentState;
 
+	/** The list of active timer handles for this widget. */
+	TArray<TSharedRef<FActiveTimerHandle>> ActiveTimers;
+
 	/** Stores the ideal size this widget wants to be. */
 	TOptional<FVector2D> DesiredSize;
 
-	/** The list of active timer handles for this widget. */
-	TArray<TSharedRef<FActiveTimerHandle>> ActiveTimers;
+	/** Whether or not this widget is enabled */
+	TSlateAttribute<bool> EnabledStateAttribute;
+
+	/** Is this widget visible, hidden or collapsed */
+	TSlateAttribute<EVisibility> VisibilityAttribute;
 
 protected:
 
@@ -1795,11 +1802,14 @@ protected:
 	*/
 	FMargin CullingBoundsExtension;
 
+#if WITH_EDITORONLY_DATA
 	/** Whether or not this widget is enabled */
-	TAttribute< bool > EnabledState;
-
+	UE_DEPRECATED(5.0, "Direct access to EnabledState is now deprecated. Use the setter or getter.")
+	FSlateDeprecatedTAttribute<bool> EnabledState;
 	/** Is this widget visible, hidden or collapsed */
-	TAttribute< EVisibility > Visibility;
+	UE_DEPRECATED(5.0, "Direct access to Visibility is now deprecated. Use the setter or getter.")
+	FSlateDeprecatedTAttribute<EVisibility> Visibility;
+#endif
 
 	/** The opacity of the widget. Automatically applied during rendering. */
 	float RenderOpacity;
@@ -1817,16 +1827,23 @@ private:
 	/** Pointer to this widgets parent widget.  If it is null this is a root widget or it is not in the widget tree */
 	TWeakPtr<SWidget> ParentWidgetPtr;
 
+	/** Tag for this widget */
+	FName Tag;
+
 	/** Debugging information on the type of widget we're creating for the Widget Reflector. */
 	FName TypeOfWidget;
 
+private: 
+
 #if !UE_BUILD_SHIPPING
 	/** Full file path (and line) in which this widget was created */
 	FName CreatedInLocation;
 #endif
 
-	/** Tag for this widget */
-	FName Tag;
+#if WITH_SLATE_DEBUGGING
+	/** The last time this widget got painted. */
+	uint32 LastPaintFrame = 0;
+#endif
 
 #if !(UE_BUILD_SHIPPING || UE_BUILD_TEST)
 	UE_DEPRECATED(4.27, "Access to SWidget::Cursor is deprecated and will not function. Call SetCursor/GetCursor instead")
@@ -1834,8 +1851,6 @@ private:
 	TAttribute<TOptional<EMouseCursor::Type>> Cursor;
 #endif
 
-private:
-
 #if UE_SLATE_WITH_WIDGET_UNIQUE_IDENTIFIER
 	/** The widget's id */
 	uint64 UniqueIdentifier;
