Adding back missing icons for the search bar style

diff --git a/Engine/Source/Runtime/SlateCore/Private/Styling/StarshipCoreStyle.cpp b/Engine/Source/Runtime/SlateCore/Private/Styling/StarshipCoreStyle.cpp
--- a/Engine/Source/Runtime/SlateCore/Private/Styling/StarshipCoreStyle.cpp
+++ b/Engine/Source/Runtime/SlateCore/Private/Styling/StarshipCoreStyle.cpp
@@ -1061,8 +1061,8 @@ void FStarshipCoreStyle::SetupTextStyles(TSharedRef<FStyle>& Style)
 
 		Style->Set("SearchBox", FSearchBoxStyle()
 			.SetTextBoxStyle(SearchBoxEditStyle)
-			.SetUpArrowImage(IMAGE_BRUSH_SVG("Starship/Common/arrow-north", Icon8x8, FStyleColors::Foreground))
-			.SetDownArrowImage(IMAGE_BRUSH_SVG("Starship/Common/arrow-south", Icon8x8, FStyleColors::Foreground))
+			.SetUpArrowImage(IMAGE_BRUSH_SVG("Starship/Common/chevron-up", Icon16x16, FStyleColors::Foreground))
+			.SetDownArrowImage(IMAGE_BRUSH_SVG("Starship/Common/chevron-down", Icon16x16, FStyleColors::Foreground))
 			.SetGlassImage(IMAGE_BRUSH_SVG("Starship/Common/search", Icon16x16))
 			.SetClearImage(IMAGE_BRUSH_SVG("Starship/Common/close", Icon16x16))
 			.SetImagePadding(FMargin(3.f, 0.f, -2.f, 0.f))
