Slate: Fix ProtectedIsWidgetInDestructionPath when destroying base widget's attribute.
#rnx

diff --git a/Engine/Source/Runtime/SlateCore/Private/Types/SlateAttribute.cpp b/Engine/Source/Runtime/SlateCore/Private/Types/SlateAttribute.cpp
--- a/Engine/Source/Runtime/SlateCore/Private/Types/SlateAttribute.cpp
+++ b/Engine/Source/Runtime/SlateCore/Private/Types/SlateAttribute.cpp
@@ -31,7 +31,11 @@ namespace SlateAttributePrivate
 		int32 FoundIndex = UE::Slate::FWidgetList::GetAllWidgets().Find(Widget);
 		if (FoundIndex == INDEX_NONE)
 		{
-			return false; // the widget is already destroyed
+			// The widget is already destroyed (this is bad) or we are destroying the SWidget base class SlateAttribute (this is ok)
+			UPTRINT SlateAttributePtr = (UPTRINT)this;
+			UPTRINT WidgetPtr = (UPTRINT)Widget;
+			const bool bIsBaseWidgetAttribute = SlateAttributePtr >= WidgetPtr && SlateAttributePtr <= WidgetPtr + sizeof(SWidget);
+			return bIsBaseWidgetAttribute;
 		}
 		// if the shared instance exist, then the widget is not being destroyed
 		return !UE::Slate::FWidgetList::GetAllWidgets()[FoundIndex]->DoesSharedInstanceExist();
diff --git a/Engine/Source/Runtime/SlateCore/Public/Types/SlateAttribute.h b/Engine/Source/Runtime/SlateCore/Public/Types/SlateAttribute.h
--- a/Engine/Source/Runtime/SlateCore/Public/Types/SlateAttribute.h
+++ b/Engine/Source/Runtime/SlateCore/Public/Types/SlateAttribute.h
@@ -881,7 +881,7 @@ namespace SlateAttributePrivate
 		~TSlateMemberAttribute()
 		{
 			/**
-			 * The parent should now be destroyed and the shared pointer should invalidate.
+			 * The parent should now be destroyed and the shared pointer should be invalidate.
 			 * If you hit the check, that means the TSlateAttribute is not a variable member of SWidget.
 			 * It will introduced bad memory access.
 			 * See documentation above.
