Slate: Move visibility flag to WidgetProxy. Keep track of the Collapsed and the Visible state to update the Attribute as in SlowPath. Convert RenderTransform and RenderTransformPivot to SlateAttribute. Make the Visibility attribute update first (needs more work in the SlateAttributeInitializer to enforce it). Update Attribute in 2 steps. Update the Collapsed first (that can change the visibility), then update the "not collapsed" attribute.

diff --git a/Engine/Source/Developer/SlateReflector/Private/Models/WidgetReflectorNode.cpp b/Engine/Source/Developer/SlateReflector/Private/Models/WidgetReflectorNode.cpp
--- a/Engine/Source/Developer/SlateReflector/Private/Models/WidgetReflectorNode.cpp
+++ b/Engine/Source/Developer/SlateReflector/Private/Models/WidgetReflectorNode.cpp
@@ -910,7 +910,7 @@ bool FWidgetReflectorNodeUtils::GetWidgetVisibility(const TSharedPtr<const SWidg
 
 bool FWidgetReflectorNodeUtils::GetWidgetVisibilityInherited(const TSharedPtr<const SWidget>& InWidget)
 {
-	return InWidget.IsValid() ? InWidget->IsFastPathVisible() : false;
+	return InWidget.IsValid() ? InWidget->GetProxyHandle().GetWidgetVisibility(InWidget.Get()).IsVisible() : false;
 }
 
 bool FWidgetReflectorNodeUtils::GetWidgetFocusable(const TSharedPtr<const SWidget>& InWidget)
diff --git a/Engine/Source/Runtime/SlateCore/Private/FastUpdate/SlateInvalidationRoot.cpp b/Engine/Source/Runtime/SlateCore/Private/FastUpdate/SlateInvalidationRoot.cpp
--- a/Engine/Source/Runtime/SlateCore/Private/FastUpdate/SlateInvalidationRoot.cpp
+++ b/Engine/Source/Runtime/SlateCore/Private/FastUpdate/SlateInvalidationRoot.cpp
@@ -443,7 +443,7 @@ bool FSlateInvalidationRoot::PaintFastPath(const FSlateInvalidationContext& Cont
 
 			// Check visibility, it may have been in the update list but a parent who was also in the update list already updated it.
 			SWidget* WidgetPtr = WidgetProxy.GetWidget();
-			if (!WidgetProxy.bUpdatedSinceLastInvalidate && WidgetPtr && WidgetPtr->IsFastPathVisible())
+			if (!WidgetProxy.bUpdatedSinceLastInvalidate && WidgetProxy.Visibility.IsVisible() && WidgetPtr)
 			{
 				bWidgetsNeededRepaint = bWidgetsNeededRepaint || WidgetPtr->HasAnyUpdateFlags(EWidgetUpdateFlags::NeedsRepaint | EWidgetUpdateFlags::NeedsVolatilePaint);
 
@@ -688,35 +688,17 @@ void FSlateInvalidationRoot::ProcessPreUpdate()
 					SWidget* WidgetPtr = InvalidationWidget.GetWidget();
 					if (ensureMsgf(WidgetPtr, TEXT("A parent remove this widget. Child order invalidation should have been called before we process this widget.")))
 					{
-						auto UpdateAttribute = [&InvalidationWidget, WidgetPtr]()
+						if (!InvalidationWidget.Visibility.IsCollapseIndirectly())
 						{
-#if UE_SLATE_WITH_INVALIDATIONWIDGETLIST_DEBUGGING
-							if (GSlateInvalidationRootVerifySlateAttribute)
+							// if my parent is not collapse, then update my visible state
+							FSlateAttributeMetaData::UpdateCollapsedAttributes(*WidgetPtr);
+							if (!InvalidationWidget.Visibility.IsCollapsed())
 							{
-								ensureAlwaysMsgf(InvalidationWidget.bDebug_AttributeUpdated == false, TEXT("Attribute should only be updated once per frame."));
-							}
+#if UE_SLATE_WITH_INVALIDATIONWIDGETLIST_DEBUGGING
+								ensureAlwaysMsgf(!GSlateInvalidationRootVerifySlateAttribute || InvalidationWidget.bDebug_AttributeUpdated == false, TEXT("Attribute should only be updated once per frame."));
+								InvalidationWidget.bDebug_AttributeUpdated = true;
 #endif
-
-							// It is possible that the widget doesn't have an attribute registered anymore.
-							FSlateAttributeMetaData::UpdateAttributes(*WidgetPtr);
-							InvalidationWidget.bDebug_AttributeUpdated = true;
-						};
-
-						if (WidgetPtr->IsFastPathVisible())
-						{
-							UpdateAttribute();
-						}
-						else if (InvalidationWidget.ParentIndex != FSlateInvalidationWidgetIndex::Invalid)
-						{
-							// Is my parent visible, should I update the collapsed attributes
-							const FSlateInvalidationWidgetList::InvalidationWidgetType& ParentInvalidationWidget = (*FastWidgetPathList)[InvalidationWidget.ParentIndex];
-							if (ParentInvalidationWidget.Visibility.IsVisible() && ensure(ParentInvalidationWidget.GetWidget()) && ParentInvalidationWidget.GetWidget()->IsFastPathVisible())
-							{
-								FSlateAttributeMetaData::UpdateCollapsedAttributes(*WidgetPtr);
-								if (WidgetPtr->IsFastPathVisible())
-								{
-									UpdateAttribute();
-								}
+								FSlateAttributeMetaData::UpdateExpandedAttributes(*WidgetPtr);
 							}
 						}
 					}
@@ -836,14 +818,17 @@ bool FSlateInvalidationRoot::ProcessPostUpdate()
 			}
 #endif
 
-			const bool bIsInvalidationRoot = WidgetPtr->Advanced_IsInvalidationRoot();
-			if (bIsInvalidationRoot && WidgetPtr != InvalidationRootWidget)
+			if (WidgetProxy.Visibility.IsVisible())
 			{
-				FSlateInvalidationRoot* InvalidationRoot = const_cast<FSlateInvalidationRoot*>(WidgetPtr->Advanced_AsInvalidationRoot());
-				check(InvalidationRoot);
-				// Prevent reentering call
-				FSlateInvalidationWidgetPostHeap::FScopeWidgetCannotBeAdded Guard{ *WidgetsNeedingPostUpdate, WidgetProxy };
-				InvalidationRoot->ProcessInvalidation();
+				const bool bIsInvalidationRoot = WidgetPtr->Advanced_IsInvalidationRoot();
+				if (bIsInvalidationRoot && WidgetPtr != InvalidationRootWidget)
+				{
+					FSlateInvalidationRoot* InvalidationRoot = const_cast<FSlateInvalidationRoot*>(WidgetPtr->Advanced_AsInvalidationRoot());
+					check(InvalidationRoot);
+					// Prevent reentering call
+					FSlateInvalidationWidgetPostHeap::FScopeWidgetCannotBeAdded Guard{ *WidgetsNeedingPostUpdate, WidgetProxy };
+					InvalidationRoot->ProcessInvalidation();
+				}
 			}
 
 			bWidgetsNeedRepaint |= WidgetProxy.ProcessPostInvalidation(*WidgetsNeedingPostUpdate, *FastWidgetPathList, *this);
@@ -980,7 +965,7 @@ void DumpUpdateList(const FSlateInvalidationWidgetList& FastWidgetPathList, cons
 
 		const FWidgetProxy& WidgetProxy = FastWidgetPathList[MyIndex];
 		const SWidget* WidgetPtr = WidgetProxy.GetWidget();
-		if (WidgetPtr)
+		if (WidgetProxy.Visibility.IsVisible() && WidgetPtr)
 		{
 			if (WidgetPtr->HasAnyUpdateFlags(EWidgetUpdateFlags::NeedsVolatilePaint))
 			{
@@ -990,7 +975,7 @@ void DumpUpdateList(const FSlateInvalidationWidgetList& FastWidgetPathList, cons
 			{
 				UE_LOG(LogSlate, Log, TEXT("Repaint %s"), *FReflectionMetaData::GetWidgetDebugInfo(WidgetPtr));
 			}
-			else if (WidgetPtr->IsFastPathVisible())
+			else
 			{
 				if (WidgetPtr->HasAnyUpdateFlags(EWidgetUpdateFlags::NeedsActiveTimerUpdate))
 				{
@@ -1167,40 +1152,71 @@ void VerifyHittest(SWidget* InvalidationRootWidget, FSlateInvalidationWidgetList
 
 void VerifyWidgetVisibility(FSlateInvalidationWidgetList& WidgetList)
 {
-	WidgetList.ForEachInvalidationWidget([](FSlateInvalidationWidgetList::InvalidationWidgetType& InvalidationWidget)
+	WidgetList.ForEachInvalidationWidget([&WidgetList](FSlateInvalidationWidgetList::InvalidationWidgetType& InvalidationWidget)
 		{
-			if (InvalidationWidget.ParentIndex != FSlateInvalidationWidgetIndex::Invalid)
+			if (SWidget* Widget = InvalidationWidget.GetWidget())
 			{
-				if (SWidget* Widget = InvalidationWidget.GetWidget())
 				{
-					bool bSouldBeFastPathVisible = Widget->GetVisibility().IsVisible();
-					TSharedPtr<SWidget> ParentWidget = Widget->GetParentWidget();
-					if (ensure(ParentWidget))
+					const EVisibility WidgetVisibility = Widget->GetVisibility();
+					bool bParentIsVisible = true;
+					bool bParentIsCollapsed = false;
+
+					const TSharedPtr<SWidget> ParentWidget = Widget->GetParentWidget();
+					if (InvalidationWidget.ParentIndex != FSlateInvalidationWidgetIndex::Invalid)
 					{
-						bSouldBeFastPathVisible = bSouldBeFastPathVisible && ParentWidget->IsFastPathVisible();
+						// Confirm that we have the correct parent
+						UE_SLATE_LOG_ERROR_IF_FALSE(WidgetList.IsValidIndex(InvalidationWidget.ParentIndex)
+							, GSlateInvalidationRootVerifyWidgetVisibility
+							, TEXT("The widget's parent index '%s' is invalid")
+							, *FReflectionMetaData::GetWidgetDebugInfo(Widget));
+
+						const FSlateInvalidationWidgetList::InvalidationWidgetType& ParentInvalidationWidget = WidgetList[InvalidationWidget.ParentIndex];
+						UE_SLATE_LOG_ERROR_IF_FALSE(ParentWidget.Get() == ParentInvalidationWidget.GetWidget()
+							, GSlateInvalidationRootVerifyWidgetVisibility
+							, TEXT("The widget's parent '%s' is not '%s'")
+							, *FReflectionMetaData::GetWidgetDebugInfo(Widget)
+							, *FReflectionMetaData::GetWidgetDebugInfo(ParentWidget.Get()));
+
+						bParentIsVisible = ParentInvalidationWidget.Visibility.IsVisible();
+						bParentIsCollapsed = ParentInvalidationWidget.Visibility.IsCollapsed();
 					}
-
-					if (Widget->IsFastPathVisible() != bSouldBeFastPathVisible)
+					else
 					{
-						// It's possible that one of the parent is volatile
-						if (!Widget->IsVolatile() && !Widget->IsVolatileIndirectly())
-						{
-							UE_SLATE_LOG_ERROR_IF_FALSE(false
-								, GSlateInvalidationRootVerifyWidgetVisibility
-								, TEXT("Widget '%s' should be %s.")
-								, *FReflectionMetaData::GetWidgetDebugInfo(Widget)
-								, (bSouldBeFastPathVisible ? TEXT("visible") : TEXT("hidden")));
-						}
+						UE_SLATE_LOG_ERROR_IF_FALSE(ParentWidget == nullptr || ParentWidget->Advanced_IsInvalidationRoot()
+							, GSlateInvalidationRootVerifyWidgetVisibility
+							, TEXT("The widget's parent is valid and is not an invalidation root.")
+							, *FReflectionMetaData::GetWidgetDebugInfo(Widget));
 					}
 
-					const bool bHasValidCachedElementHandle = Widget->IsFastPathVisible() || !Widget->GetPersistentState().CachedElementHandle.HasCachedElements();
-					UE_SLATE_LOG_ERROR_IF_FALSE(bHasValidCachedElementHandle
+					UE_SLATE_LOG_ERROR_IF_FALSE(InvalidationWidget.Visibility.AreAncestorsVisible() == bParentIsVisible
 						, GSlateInvalidationRootVerifyWidgetVisibility
-						, TEXT("Widget '%s' has cached element and is not visibled.")
+						, TEXT("Widget '%s' AreAncestorsVisible flag is wrong.")
 						, *FReflectionMetaData::GetWidgetDebugInfo(Widget));
-
+					UE_SLATE_LOG_ERROR_IF_FALSE(InvalidationWidget.Visibility.IsVisible() == (bParentIsVisible && WidgetVisibility.IsVisible())
+						, GSlateInvalidationRootVerifyWidgetVisibility
+						, TEXT("Widget '%s' IsVisible flag is wrong.")
+						, *FReflectionMetaData::GetWidgetDebugInfo(Widget));
+					UE_SLATE_LOG_ERROR_IF_FALSE(InvalidationWidget.Visibility.IsCollapsed() == bParentIsCollapsed || WidgetVisibility == EVisibility::Collapsed
+						, GSlateInvalidationRootVerifyWidgetVisibility
+						, TEXT("Widget '%s' IsCollapsed flag is wrong.")
+						, *FReflectionMetaData::GetWidgetDebugInfo(Widget));
+					UE_SLATE_LOG_ERROR_IF_FALSE(InvalidationWidget.Visibility.IsCollapseIndirectly() == bParentIsCollapsed
+					, GSlateInvalidationRootVerifyWidgetVisibility
+						, TEXT("Widget '%s' IsCollapseIndirectly flag is wrong.")
+						, *FReflectionMetaData::GetWidgetDebugInfo(Widget));
+				}
+				{
+					if (!InvalidationWidget.Visibility.IsVisible())
+					{
+						UE_SLATE_LOG_ERROR_IF_FALSE(!Widget->GetPersistentState().CachedElementHandle.HasCachedElements()
+							, GSlateInvalidationRootVerifyWidgetVisibility
+							, TEXT("Widget '%s' has cached element and is not visibled.")
+							, *FReflectionMetaData::GetWidgetDebugInfo(Widget));
+					}
+				}
+				{
 					// Cache last frame visibility
-					InvalidationWidget.bDebug_LastFrameVisible = Widget->IsFastPathVisible();
+					InvalidationWidget.bDebug_LastFrameVisible = InvalidationWidget.Visibility.IsVisible();
 					InvalidationWidget.bDebug_LastFrameVisibleSet = true;
 				}
 			}
@@ -1210,9 +1226,9 @@ void VerifyWidgetVisibility(FSlateInvalidationWidgetList& WidgetList)
 void VerifyWidgetVolatile(FSlateInvalidationWidgetList& WidgetList, TArray<FSlateInvalidationWidgetIndex>& FinalUpdateList)
 {
 	SWidget* Root = WidgetList.GetRoot().Pin().Get();
-	WidgetList.ForEachWidget([Root , &FinalUpdateList](SWidget* Widget)
+	WidgetList.ForEachWidget([Root, &FinalUpdateList](SWidget* Widget)
 		{
-			if (Widget != Root)
+			if (Widget != Root && GSlateInvalidationRootVerifyWidgetVolatile)
 			{
 				{
 					const bool bWasVolatile = Widget->IsVolatile();
@@ -1246,11 +1262,15 @@ void VerifyWidgetVolatile(FSlateInvalidationWidgetList& WidgetList, TArray<FSlat
 						, TEXT("Widget '%s' is volatile but doesn't have the update flag NeedsVolatilePaint.")
 						, *FReflectionMetaData::GetWidgetDebugInfo(Widget));
 
-					const bool bIsContains = FinalUpdateList.Contains(Widget->GetProxyHandle().GetWidgetIndex());
-					UE_SLATE_LOG_ERROR_IF_FALSE(bIsContains
-						, GSlateInvalidationRootVerifyWidgetVolatile
-						, TEXT("Widget '%s' is volatile but is not in the update list.")
-						, *FReflectionMetaData::GetWidgetDebugInfo(Widget));
+					if (Widget->GetProxyHandle().IsValid(Widget))
+					{
+						const bool bIsVisible = Widget->GetProxyHandle().GetProxy().Visibility.IsVisible();
+						const bool bIsContains = FinalUpdateList.Contains(Widget->GetProxyHandle().GetWidgetIndex());
+						UE_SLATE_LOG_ERROR_IF_FALSE(bIsContains || !bIsVisible
+							, GSlateInvalidationRootVerifyWidgetVolatile
+							, TEXT("Widget '%s' is volatile but is not in the update list.")
+							, *FReflectionMetaData::GetWidgetDebugInfo(Widget));
+					}
 				}
 			}
 		});
diff --git a/Engine/Source/Runtime/SlateCore/Private/FastUpdate/SlateInvalidationWidgetList.cpp b/Engine/Source/Runtime/SlateCore/Private/FastUpdate/SlateInvalidationWidgetList.cpp
--- a/Engine/Source/Runtime/SlateCore/Private/FastUpdate/SlateInvalidationWidgetList.cpp
+++ b/Engine/Source/Runtime/SlateCore/Private/FastUpdate/SlateInvalidationWidgetList.cpp
@@ -257,57 +257,50 @@ FSlateInvalidationWidgetList::FSlateInvalidationWidgetList(FSlateInvalidationRoo
 }
 
 
-FSlateInvalidationWidgetIndex FSlateInvalidationWidgetList::_BuildWidgetList_Recursive(TSharedRef<SWidget>& Widget, FSlateInvalidationWidgetIndex ParentIndex, IndexType& LastestIndex, bool bParentVisible, bool bParentVolatile)
+FSlateInvalidationWidgetIndex FSlateInvalidationWidgetList::Internal_BuildWidgetList_Recursive(TSharedRef<SWidget>& Widget, FSlateInvalidationWidgetIndex ParentIndex, IndexType& LastestIndex, FSlateInvalidationWidgetVisibility ParentVisibility, bool bParentVolatile)
 {
 	const bool bIsEmpty = IsEmpty();
 	const FSlateInvalidationWidgetIndex NewIndex = EmplaceInsertAfter(LastestIndex, Widget);
 	LastestIndex = NewIndex.ArrayIndex;
 
 	FSlateInvalidationWidgetIndex LeafMostChildIndex = NewIndex;
-	EVisibility Visibility = Widget->GetVisibility();
-	bool bParentAndSelfVisible = bParentVisible && Visibility.IsVisible();
+	FSlateInvalidationWidgetVisibility NewVisibility {ParentVisibility, Widget->GetVisibility()};
 
 	bool bUpdateSlateAttribute = false;
-	if (Widget->HasRegisteredSlateAttribute() && Widget->IsAttributesUpdatesEnabled())
+	if (ShouldBeAddedToAttributeList(Widget))
 	{
-
 		// The list is already sorted at this point. Add to the end.
 		Data[LastestIndex].ElementIndexList_WidgetWithRegisteredSlateAttribute.AddUnsorted(NewIndex.ElementIndex);
 
 		// Update attributes now because the visibility flag may change
 		{
-			if (bParentAndSelfVisible)
-			{
-				FSlateAttributeMetaData::UpdateAttributes(Widget.Get(), false);
-				bUpdateSlateAttribute = true;
-
-				Visibility = Widget->GetVisibility();
-				bParentAndSelfVisible = bParentVisible && Visibility.IsVisible();
-			}
-			else if (bParentVisible)
+			if (!NewVisibility.IsCollapseIndirectly())
 			{
 				FSlateAttributeMetaData::UpdateCollapsedAttributes(Widget.Get(), false);
-				bUpdateSlateAttribute = true;
+				NewVisibility.SetVisibility(ParentVisibility, Widget->GetVisibility());
 
-				Visibility = Widget->GetVisibility();
-				bParentAndSelfVisible = bParentVisible && Visibility.IsVisible();
-				if (bParentAndSelfVisible)
+				if (!NewVisibility.IsCollapsed())
 				{
-					FSlateAttributeMetaData::UpdateAttributes(Widget.Get(), false);
-					Visibility = Widget->GetVisibility();
-					bParentAndSelfVisible = bParentVisible && Visibility.IsVisible();
+					FSlateAttributeMetaData::UpdateExpandedAttributes(Widget.Get(), false);
+					NewVisibility.SetVisibility(ParentVisibility, Widget->GetVisibility());
+					bUpdateSlateAttribute = true;
 				}
 			}
 		}
 	}
 
+	const bool bIsInvalidationRoot = Widget->Advanced_IsInvalidationRoot();
+
 	{
 		InvalidationWidgetType& WidgetProxy = (*this)[NewIndex];
 		WidgetProxy.Index = NewIndex;
 		WidgetProxy.ParentIndex = ParentIndex;
 		WidgetProxy.LeafMostChildIndex = LeafMostChildIndex;
-		WidgetProxy.Visibility = Visibility;
+		WidgetProxy.Visibility = NewVisibility;
+		WidgetProxy.bIsInvalidationRoot = bIsInvalidationRoot;
+#if UE_SLATE_WITH_INVALIDATIONWIDGETLIST_DEBUGGING
 		WidgetProxy.bDebug_AttributeUpdated = bUpdateSlateAttribute;
+#endif
 	}
 
 #if UE_SLATE_WITH_INVALIDATIONWIDGETLIST_DEBUGGING
@@ -315,14 +308,14 @@ FSlateInvalidationWidgetIndex FSlateInvalidationWidgetList::_BuildWidgetList_Rec
 #endif
 	{
 		const FSlateInvalidationWidgetSortOrder SortIndex = { *this, NewIndex };
-		Widget->SetFastPathProxyHandle(FWidgetProxyHandle{ Owner, NewIndex, SortIndex, GenerationNumber }, !bParentAndSelfVisible, bParentVolatile);
+		Widget->SetFastPathProxyHandle(FWidgetProxyHandle{ Owner, NewIndex, SortIndex, GenerationNumber }, NewVisibility, bParentVolatile);
 	}
 
 	const bool bParentOrSelfVolatile = bParentVolatile || Widget->IsVolatile();
 
 	// N.B. The SInvalidationBox needs a valid Proxy to decide if he's a root or not.
 	//const bool bDoRecursion = ShouldDoRecursion(Widget);
-	const bool bDoRecursion = !Widget->Advanced_IsInvalidationRoot() || bIsEmpty;
+	const bool bDoRecursion = !bIsInvalidationRoot || bIsEmpty;
 	if (bDoRecursion)
 	{
 		FChildren* Children = Widget->GetAllChildren();
@@ -335,7 +328,7 @@ FSlateInvalidationWidgetIndex FSlateInvalidationWidgetList::_BuildWidgetList_Rec
 			if (bShouldAdd)
 			{
 				check(NewWidget->GetParentWidget() == Widget);
-				LeafMostChildIndex = _BuildWidgetList_Recursive(NewWidget, NewIndex, LastestIndex, bParentAndSelfVisible, bParentOrSelfVolatile);
+				LeafMostChildIndex = Internal_BuildWidgetList_Recursive(NewWidget, NewIndex, LastestIndex, NewVisibility, bParentOrSelfVolatile);
 			}
 		}
 
@@ -358,15 +351,17 @@ void FSlateInvalidationWidgetList::BuildWidgetList(TSharedRef<SWidget> InRoot)
 	{
 		GenerationNumber = InvalidationRoot->GetFastPathGenerationNumber();
 	}
-
 	const bool bShouldAdd = ShouldBeAdded(InRoot);
 	if (bShouldAdd)
 	{
-		const TSharedPtr<SWidget> Parent = InRoot->GetParentWidget();
-		const bool bParentVisible = Parent.IsValid() ? Parent->GetVisibility().IsVisible() : true;
+		FSlateInvalidationWidgetVisibility ParentVisibility = EVisibility::Visible;
+		if (const TSharedPtr<SWidget> Parent = InRoot->GetParentWidget())
+		{
+			ParentVisibility = Parent->GetVisibility();
+		}
 		const bool bParentVolatile = false;
 		IndexType LatestArrayIndex = FSlateInvalidationWidgetIndex::Invalid.ArrayIndex;
-		_BuildWidgetList_Recursive(InRoot, FSlateInvalidationWidgetIndex::Invalid, LatestArrayIndex, bParentVisible, bParentVolatile);
+		Internal_BuildWidgetList_Recursive(InRoot, FSlateInvalidationWidgetIndex::Invalid, LatestArrayIndex, ParentVisibility, bParentVolatile);
 	}
 }
 
@@ -378,15 +373,16 @@ void FSlateInvalidationWidgetList::_RebuildWidgetListTree(TSharedRef<SWidget> Wi
 	const FSlateInvalidationWidgetIndex WidgetIndex = Widget->GetProxyHandle().GetWidgetIndex();
 	if (bShouldAddWidget && ChildAtIndex < ParentChildren->Num() && WidgetIndex != FSlateInvalidationWidgetIndex::Invalid)
 	{
-		// Since we are going to add item, the array may get invalidated. Do not use after _BuildWidgetList_Recursive
-		ensure((*this)[WidgetIndex].GetWidget() == &Widget.Get());
-		FSlateInvalidationWidgetIndex PreviousLeafIndex = (*this)[WidgetIndex].LeafMostChildIndex;
+		// Since we are going to add item, the array may get invalidated. Do not use after Internal_BuildWidgetList_Recursive
+		const InvalidationWidgetType& CurrentInvalidationWidget = (*this)[WidgetIndex];
+		ensure(CurrentInvalidationWidget.GetWidget() == &Widget.Get());
+		FSlateInvalidationWidgetIndex PreviousLeafIndex = CurrentInvalidationWidget.LeafMostChildIndex;
 		FSlateInvalidationWidgetIndex NewLeafIndex = PreviousLeafIndex;
 		IndexType LastestArrayIndex = PreviousLeafIndex.ArrayIndex;
 		const bool bDoRecursion = ShouldDoRecursion(Widget);
 		if (bDoRecursion)
 		{
-			const bool bParentVisible = Widget->IsFastPathVisible() && Widget->GetVisibility().IsVisible();
+			const FSlateInvalidationWidgetVisibility CurrentVisibility = CurrentInvalidationWidget.Visibility;
 			const bool bParentVolatile = Widget->IsVolatileIndirectly() || Widget->IsVolatile();
 			int32 NumChildren = ParentChildren->Num();
 			for (int32 Index = ChildAtIndex; Index < NumChildren; ++Index)
@@ -396,7 +392,7 @@ void FSlateInvalidationWidgetList::_RebuildWidgetListTree(TSharedRef<SWidget> Wi
 				if (bShouldAddChild)
 				{
 					check(ChildWidget->GetParentWidget() == Widget);
-					NewLeafIndex = _BuildWidgetList_Recursive(ChildWidget, WidgetIndex, LastestArrayIndex, bParentVisible, bParentVolatile);
+					NewLeafIndex = Internal_BuildWidgetList_Recursive(ChildWidget, WidgetIndex, LastestArrayIndex, CurrentVisibility, bParentVolatile);
 				}
 			}
 		}
@@ -578,7 +574,7 @@ void FSlateInvalidationWidgetList::ProcessAttributeRegistrationInvalidation(cons
 	SWidget* WidgetPtr = InvalidationWidget.GetWidget();
 	check(WidgetPtr);
 
-	if (WidgetPtr->HasRegisteredSlateAttribute() && WidgetPtr->IsAttributesUpdatesEnabled())
+	if (ShouldBeAddedToAttributeList(WidgetPtr))
 	{
 		Data[InvalidationWidget.Index.ArrayIndex].ElementIndexList_WidgetWithRegisteredSlateAttribute.InsertUnique(InvalidationWidget.Index.ElementIndex);
 	}
@@ -588,7 +584,7 @@ void FSlateInvalidationWidgetList::ProcessAttributeRegistrationInvalidation(cons
 	}
 }
 
-namespace
+namespace InvalidationList
 {
 	template<typename TSlateInvalidationWidgetList, typename Predicate>
 	void ForEachChildren(TSlateInvalidationWidgetList& Self, const typename TSlateInvalidationWidgetList::InvalidationWidgetType& InvalidationWidget, FSlateInvalidationWidgetIndex WidgetIndex, Predicate InPredicate)
@@ -622,7 +618,7 @@ void FSlateInvalidationWidgetList::_FindChildren(FSlateInvalidationWidgetIndex W
 {
 	Widgets.Reserve(16);
 	const InvalidationWidgetType& InvalidationWidget = (*this)[WidgetIndex];
-	ForEachChildren(*this, InvalidationWidget, WidgetIndex, [&Widgets](const InvalidationWidgetType& ChildWidget)
+	InvalidationList::ForEachChildren(*this, InvalidationWidget, WidgetIndex, [&Widgets](const InvalidationWidgetType& ChildWidget)
 		{
 			Widgets.Emplace(ChildWidget.GetWidget(), ChildWidget.Index);
 		});
@@ -1048,7 +1044,7 @@ void FSlateInvalidationWidgetList::_RemoveRangeFromSameParent(const FIndexRange
 				InvalidationWidgetType& InvalidationWidget = Self->Data[ArrayIndex].ElementList[ElementIndex];
 				if (SWidget* Widget = InvalidationWidget.GetWidget())
 				{
-					Widget->SetFastPathProxyHandle(FWidgetProxyHandle{ Self->Owner, SlateInvalidationWidgetIndexRemoved, FSlateInvalidationWidgetSortOrder(), GenerationNumber });
+					Widget->SetFastPathProxyHandle(FWidgetProxyHandle{ Self->Owner, SlateInvalidationWidgetIndexRemoved, FSlateInvalidationWidgetSortOrder() });
 				}
 			}
 #endif //UE_SLATE_WITH_WIDGETLIST_ASSIGNINVALIDPROXYWHENREMOVED
@@ -1272,7 +1268,7 @@ FSlateInvalidationWidgetList::FCutResult FSlateInvalidationWidgetList::_CutArray
 				Widget->SetFastPathProxyHandle(FWidgetProxyHandle{ Owner, NewWidgetIndex, SortIndex, GenerationNumber });
 
 				// Check for ElementIndexList
-				if (Widget->HasRegisteredSlateAttribute() && Widget->IsAttributesUpdatesEnabled())
+				if (ShouldBeAddedToAttributeList(Widget))
 				{
 					Data[NewArrayIndex].ElementIndexList_WidgetWithRegisteredSlateAttribute.AddUnsorted(NewElementIndex);
 				}
@@ -1289,7 +1285,7 @@ FSlateInvalidationWidgetList::FCutResult FSlateInvalidationWidgetList::_CutArray
 
 			// The parent is only invalid when it's the root. We can't remove it, but we shouldn't move it.
 			check(NewInvalidationWidget.ParentIndex != FSlateInvalidationWidgetIndex::Invalid);
-			ForEachChildren(*this, NewInvalidationWidget, NewWidgetIndex, [NewWidgetIndex](InvalidationWidgetType& NewChildInvalidationWidget)
+			InvalidationList::ForEachChildren(*this, NewInvalidationWidget, NewWidgetIndex, [NewWidgetIndex](InvalidationWidgetType& NewChildInvalidationWidget)
 				{
 					NewChildInvalidationWidget.ParentIndex = NewWidgetIndex;
 				});
@@ -1590,7 +1586,7 @@ bool FSlateInvalidationWidgetList::VerifyElementIndexList() const
 		int32 Index = 0;
 		const_cast<FSlateInvalidationWidgetList*>(this)->ForEachWidget([&](const SWidget* Widget)
 			{
-				if (Widget->HasRegisteredSlateAttribute() && Widget->IsAttributesUpdatesEnabled())
+				if (ShouldBeAddedToAttributeList(Widget))
 				{
 					WidgetListWithForEachWidget.Add(Widget);
 				}
@@ -1621,8 +1617,7 @@ bool FSlateInvalidationWidgetList::VerifyElementIndexList() const
 				const FSlateInvalidationWidgetList::InvalidationWidgetType& InvalidationWidget = ArrayNode.ElementList[ElementIndex];
 				if (const SWidget* Widget = InvalidationWidget.GetWidget())
 				{
-					const bool bShouldBeThere = Widget->HasRegisteredSlateAttribute() && Widget->IsAttributesUpdatesEnabled();
-					if (!bShouldBeThere)
+					if (!Widget->HasRegisteredSlateAttribute())
 					{
 						UE_LOG(LogSlate, Warning, TEXT("ElementIndex '%d' in the array of sort value '%d' should not be there since widget '%s' doesn't have registered SlateAttribute.")
 							, ElementIndex, ArrayNode.SortOrder
@@ -1630,7 +1625,7 @@ bool FSlateInvalidationWidgetList::VerifyElementIndexList() const
 						bResult = false;
 					}
 
-					if (Widget->IsFastPathVisible() && !InvalidationWidget.bDebug_AttributeUpdated)
+					if (InvalidationWidget.Visibility.IsVisible() && !InvalidationWidget.bDebug_AttributeUpdated)
 					{
 						UE_LOG(LogSlate, Warning, TEXT("The widget '%s' was not updated."), *FReflectionMetaData::GetWidgetDebugInfo(Widget));
 					}
diff --git a/Engine/Source/Runtime/SlateCore/Private/FastUpdate/SlateInvalidationWidgetList.h b/Engine/Source/Runtime/SlateCore/Private/FastUpdate/SlateInvalidationWidgetList.h
--- a/Engine/Source/Runtime/SlateCore/Private/FastUpdate/SlateInvalidationWidgetList.h
+++ b/Engine/Source/Runtime/SlateCore/Private/FastUpdate/SlateInvalidationWidgetList.h
@@ -11,10 +11,6 @@
 #include "Widgets/SNullWidget.h"
 #include "Widgets/SWidget.h"
 
-#ifndef UE_SLATE_WITH_INVALIDATIONWIDGETLIST_DEBUGGING
-	#define UE_SLATE_WITH_INVALIDATIONWIDGETLIST_DEBUGGING !(UE_BUILD_SHIPPING || UE_BUILD_TEST)
-#endif
-
 class FSlateInvalidationWidgetList
 {
 	friend struct FSlateInvalidationWidgetSortOrder;
@@ -340,6 +336,14 @@ private:
 	{
 		return Widget != SNullWidget::NullWidget;
 	}
+	bool ShouldBeAddedToAttributeList(const SWidget* Widget) const
+	{
+		return Widget->HasRegisteredSlateAttribute() && Widget->IsAttributesUpdatesEnabled();
+	}
+	bool ShouldBeAddedToAttributeList(const TSharedRef<SWidget>& Widget) const
+	{
+		return Widget->HasRegisteredSlateAttribute() && Widget->IsAttributesUpdatesEnabled();
+	}
 
 private:
 	template <typename... ArgsType>
@@ -375,7 +379,7 @@ private:
 	FCutResult CutArray(const FSlateInvalidationWidgetIndex WhereToCut);
 
 private:
-	FSlateInvalidationWidgetIndex _BuildWidgetList_Recursive(TSharedRef<SWidget>& Widget, FSlateInvalidationWidgetIndex ParentIndex, IndexType& LastestIndex, bool bParentVisible, bool bParentVolatile);
+	FSlateInvalidationWidgetIndex Internal_BuildWidgetList_Recursive(TSharedRef<SWidget>& Widget, FSlateInvalidationWidgetIndex ParentIndex, IndexType& LastestIndex, FSlateInvalidationWidgetVisibility ParentVisibility, bool bParentVolatile);
 	void _RebuildWidgetListTree(TSharedRef<SWidget> Widget, int32 ChildAtIndex);
 	using FFindChildrenElement = TPair<SWidget*, FSlateInvalidationWidgetIndex>;
 	void _FindChildren(FSlateInvalidationWidgetIndex WidgetIndex, TArray<FFindChildrenElement, TMemStackAllocator<>>& Widgets) const;
diff --git a/Engine/Source/Runtime/SlateCore/Private/FastUpdate/WidgetProxy.cpp b/Engine/Source/Runtime/SlateCore/Private/FastUpdate/WidgetProxy.cpp
--- a/Engine/Source/Runtime/SlateCore/Private/FastUpdate/WidgetProxy.cpp
+++ b/Engine/Source/Runtime/SlateCore/Private/FastUpdate/WidgetProxy.cpp
@@ -24,14 +24,8 @@ FWidgetProxy::FWidgetProxy(TSharedRef<SWidget>& InWidget)
 	, ParentIndex(FSlateInvalidationWidgetIndex::Invalid)
 	, LeafMostChildIndex(FSlateInvalidationWidgetIndex::Invalid)
 	, CurrentInvalidateReason(EInvalidateWidgetReason::None)
-	// Potentially unsafe to update visibility from the widget due to attribute bindings.  This is updated later when the widgets are sorted in ProcessInvalidation
-	, Visibility(EVisibility::Collapsed) 
-	, bUpdatedSinceLastInvalidate(false)
-	, bContainedByWidgetPreHeap(false)
-	, bContainedByWidgetPostHeap(false)
-	, bDebug_LastFrameVisible(true)
-	, bDebug_LastFrameVisibleSet(false)
-	, bDebug_AttributeUpdated(false)
+	, Visibility()
+	, PrivateFlags(false)
 {
 #if UE_SLATE_WITH_WIDGETPROXY_WIDGETTYPE
 	WidgetName = GetWidget()->GetType();
@@ -49,21 +43,16 @@ TSharedPtr<SWidget> FWidgetProxy::GetWidgetAsShared() const
 
 int32 FWidgetProxy::Update(const FPaintArgs& PaintArgs, FSlateWindowElementList& OutDrawElements)
 {
-// Commenting this since it could be triggered in specific cases where Widgte->UpdateFlags in reset but the Widget Proxy is still in the  update list
-//#if WITH_SLATE_DEBUGGING
-//	ensure(UpdateFlags == Widget->UpdateFlags);
-//#endif
-
 	TSharedPtr<SWidget> CurrentWidget = GetWidgetAsShared();
+	check(Visibility.IsVisible());
 
 	// If Outgoing layer id remains index none, there was no change
 	int32 OutgoingLayerId = INDEX_NONE;
 	if (CurrentWidget->HasAnyUpdateFlags(EWidgetUpdateFlags::NeedsRepaint|EWidgetUpdateFlags::NeedsVolatilePaint))
 	{
-		check(CurrentWidget->IsFastPathVisible());
 		OutgoingLayerId = Repaint(PaintArgs, OutDrawElements);
 	}
-	else if(CurrentWidget->IsFastPathVisible())
+	else
 	{
 		EWidgetUpdateFlags PreviousUpdateFlag = CurrentWidget->UpdateFlags;
 		if (CurrentWidget->HasAnyUpdateFlags(EWidgetUpdateFlags::NeedsActiveTimerUpdate))
@@ -96,7 +85,7 @@ bool FWidgetProxy::ProcessPostInvalidation(FSlateInvalidationWidgetPostHeap& Upd
 	bool bWidgetNeedsRepaint = false;
 	SWidget* WidgetPtr = GetWidget();
 
-	if (WidgetPtr->IsFastPathVisible() && ParentIndex != FSlateInvalidationWidgetIndex::Invalid && !WidgetPtr->PrepassLayoutScaleMultiplier.IsSet())
+	if (Visibility.IsVisible() && ParentIndex != FSlateInvalidationWidgetIndex::Invalid && !WidgetPtr->PrepassLayoutScaleMultiplier.IsSet())
 	{
 		SCOPE_CYCLE_SWIDGET(WidgetPtr);
 		// If this widget has never been prepassed make sure the parent prepasses it to set the correct multiplier
@@ -119,7 +108,7 @@ bool FWidgetProxy::ProcessPostInvalidation(FSlateInvalidationWidgetPostHeap& Upd
 		// When layout changes compute a new desired size for this widget
 		FVector2D CurrentDesiredSize = WidgetPtr->GetDesiredSize();
 		FVector2D NewDesiredSize = FVector2D::ZeroVector;
-		if (Visibility != EVisibility::Collapsed)
+		if (!Visibility.IsCollapsed())
 		{
 			if (WidgetPtr->NeedsPrepass())
 			{
@@ -188,7 +177,7 @@ void FWidgetProxy::MarkProxyUpdatedThisFrame(FSlateInvalidationWidgetPostHeap& P
 	bUpdatedSinceLastInvalidate = true;
 
 	SWidget* WidgetPtr = GetWidget();
-	if (WidgetPtr && WidgetPtr->IsFastPathVisible() && WidgetPtr->HasAnyUpdateFlags(EWidgetUpdateFlags::AnyUpdate))
+	if (WidgetPtr && Visibility.IsVisible() && WidgetPtr->HasAnyUpdateFlags(EWidgetUpdateFlags::AnyUpdate))
 	{
 		// If there are any updates still needed add them to the next update list
 		PostUpdateList.PushBackUnique(*this);
@@ -284,6 +273,15 @@ bool FWidgetProxyHandle::HasValidInvalidationRootOwnership(const SWidget* Widget
 		&& InvalidationRoot->GetFastPathWidgetList()[WidgetIndex].GetWidget() == Widget;
 }
 
+FSlateInvalidationWidgetVisibility FWidgetProxyHandle::GetWidgetVisibility(const SWidget* Widget) const
+{
+	if (IsValid(Widget))
+	{
+		return GetProxy().Visibility;
+	}
+	return FSlateInvalidationWidgetVisibility();
+}
+
 FWidgetProxy& FWidgetProxyHandle::GetProxy()
 {
 	return GetInvalidationRoot()->GetFastPathWidgetList()[WidgetIndex];
@@ -307,12 +305,11 @@ void FWidgetProxyHandle::MarkWidgetDirty(EInvalidateWidgetReason InvalidateReaso
 
 void FWidgetProxyHandle::UpdateWidgetFlags(const SWidget* Widget, EWidgetUpdateFlags NewFlags)
 {
-	// Add to update list if the widget is now tickable or has an active timer.
-	if (Widget->IsFastPathVisible() && Widget->HasAnyUpdateFlags(EWidgetUpdateFlags::AnyUpdate))
+	if (Widget->HasAnyUpdateFlags(EWidgetUpdateFlags::AnyUpdate) && IsValid(Widget))
 	{
-		if (IsValid(Widget))
+		FWidgetProxy& Proxy = GetInvalidationRoot()->GetFastPathWidgetList()[WidgetIndex];
+		if (Proxy.Visibility.IsVisible())
 		{
-			FWidgetProxy& Proxy = GetInvalidationRoot()->GetFastPathWidgetList()[WidgetIndex];
 			GetInvalidationRoot()->WidgetsNeedingPostUpdate->PushBackOrHeapUnique(Proxy);
 		}
 	}
diff --git a/Engine/Source/Runtime/SlateCore/Private/Types/SlateAttributeMetaData.cpp b/Engine/Source/Runtime/SlateCore/Private/Types/SlateAttributeMetaData.cpp
--- a/Engine/Source/Runtime/SlateCore/Private/Types/SlateAttributeMetaData.cpp
+++ b/Engine/Source/Runtime/SlateCore/Private/Types/SlateAttributeMetaData.cpp
@@ -294,11 +294,17 @@ void FSlateAttributeMetaData::UpdateCollapsedAttributes(SWidget& OwningWidget)
 }
 
 
+void FSlateAttributeMetaData::UpdateExpandedAttributes(SWidget& OwningWidget)
+{
+	UpdateExpandedAttributes(OwningWidget, OwningWidget.IsConstructionCompleted());
+}
+
+
 void FSlateAttributeMetaData::UpdateAttributes(SWidget& OwningWidget, bool bAllowInvalidation)
 {
 	if (FSlateAttributeMetaData* AttributeMetaData = FSlateAttributeMetaData::FindMetaData(OwningWidget))
 	{
-		AttributeMetaData->UpdateAttributesImpl(OwningWidget, false, bAllowInvalidation);
+		AttributeMetaData->UpdateAttributesImpl(OwningWidget, EUpdateType::All, bAllowInvalidation);
 	}
 }
 
@@ -310,13 +316,23 @@ void FSlateAttributeMetaData::UpdateCollapsedAttributes(SWidget& OwningWidget, b
 		// Does it have any collapsed attributes
 		if (AttributeMetaData->CollaspedAttributeCounter > 0)
 		{
-			AttributeMetaData->UpdateAttributesImpl(OwningWidget, true, bAllowInvalidation);
+			AttributeMetaData->UpdateAttributesImpl(OwningWidget, EUpdateType::Collapsed, bAllowInvalidation);
 		}
 	}
 }
 
 
-void FSlateAttributeMetaData::UpdateAttributesImpl(SWidget& OwningWidget, bool bOnlyCollapsed, bool bAllowInvalidation)
+void FSlateAttributeMetaData::UpdateExpandedAttributes(SWidget& OwningWidget, bool bAllowInvalidation)
+{
+	if (FSlateAttributeMetaData* AttributeMetaData = FSlateAttributeMetaData::FindMetaData(OwningWidget))
+	{
+		AttributeMetaData->UpdateAttributesImpl(OwningWidget, EUpdateType::Expanded, bAllowInvalidation);
+	}
+}
+
+
+
+void FSlateAttributeMetaData::UpdateAttributesImpl(SWidget& OwningWidget, EUpdateType UpdateType, bool bAllowInvalidation)
 {
 	// List of all the callback we need to execute after the invalidation
 	TArray<const FSlateAttributeDescriptor::FInvalidationDelegate*, TInlineAllocator<8>> AllInvalidationDelegates;
@@ -327,7 +343,11 @@ void FSlateAttributeMetaData::UpdateAttributesImpl(SWidget& OwningWidget, bool b
 		FGetterItem& GetterItem = Attributes[Index];
 
 		// Only update attributes that needs to be updated when the widget is collapsed
-		if (bOnlyCollapsed && !GetterItem.bUpdateWhenCollapsed)
+		if (UpdateType == EUpdateType::Collapsed && !GetterItem.bUpdateWhenCollapsed)
+		{
+			continue;
+		}
+		if (UpdateType == EUpdateType::Expanded && GetterItem.bUpdateWhenCollapsed)
 		{
 			continue;
 		}
diff --git a/Engine/Source/Runtime/SlateCore/Private/Widgets/SWidget.cpp b/Engine/Source/Runtime/SlateCore/Private/Widgets/SWidget.cpp
--- a/Engine/Source/Runtime/SlateCore/Private/Widgets/SWidget.cpp
+++ b/Engine/Source/Runtime/SlateCore/Private/Widgets/SWidget.cpp
@@ -159,40 +159,6 @@ void SWidget::CreateStatID() const
 
 #endif
 
-void SWidget::UpdateWidgetProxy(int32 NewLayerId, FSlateCachedElementsHandle& CacheHandle)
-{
-#if WITH_SLATE_DEBUGGING
-	check(!CacheHandle.IsValid() || CacheHandle.IsOwnedByWidget(this));
-#endif
-
-	// Account for the case when the widget gets a new handle for some reason.  This should really never happen
-	if (PersistentState.CachedElementHandle.IsValid() && PersistentState.CachedElementHandle != CacheHandle)
-	{
-		ensureMsgf(!CacheHandle.IsValid(), TEXT("Widget was assigned a new cache handle.  This is not expected to happen."));
-		PersistentState.CachedElementHandle.RemoveFromCache();
-	}
-	PersistentState.CachedElementHandle = CacheHandle;
-
-	if (FastPathProxyHandle.IsValid(this))
-	{
-		FWidgetProxy& MyProxy = FastPathProxyHandle.GetProxy();
-
-		MyProxy.Visibility = GetVisibility();
-
-		PersistentState.OutgoingLayerId = NewLayerId;
-
-		if ((IsVolatile() && !IsVolatileIndirectly()) || (Advanced_IsInvalidationRoot() && !Advanced_IsWindow()))
-		{
-			AddUpdateFlags(EWidgetUpdateFlags::NeedsVolatilePaint);
-		}
-		else
-		{
-			RemoveUpdateFlags(EWidgetUpdateFlags::NeedsVolatilePaint);
-		}
-		FastPathProxyHandle.MarkWidgetUpdatedThisFrame();
-	}
-}
-
 #if UE_SLATE_WITH_WIDGET_UNIQUE_IDENTIFIER
 namespace SlateTraceMetaData
 {
@@ -203,9 +169,13 @@ namespace SlateTraceMetaData
 SLATE_IMPLEMENT_WIDGET(SWidget)
 void SWidget::PrivateRegisterAttributes(FSlateAttributeInitializer& AttributeInitializer)
 {
-	SLATE_ADD_MEMBER_ATTRIBUTE_DEFINITION_WITH_NAME(AttributeInitializer, "EnabledState", EnabledStateAttribute, EInvalidateWidgetReason::Paint);
+	// Visibility should be the first Attribute in the list.
+	//The order in which SlateAttribute are declared in the .h dictates of the order.
 	SLATE_ADD_MEMBER_ATTRIBUTE_DEFINITION_WITH_NAME(AttributeInitializer, "Visibility", VisibilityAttribute, EInvalidateWidgetReason::Visibility)
 		.UpdateWhenCollapsed();
+	SLATE_ADD_MEMBER_ATTRIBUTE_DEFINITION_WITH_NAME(AttributeInitializer, "EnabledState", EnabledStateAttribute, EInvalidateWidgetReason::Paint);
+	SLATE_ADD_MEMBER_ATTRIBUTE_DEFINITION_WITH_NAME(AttributeInitializer, "RenderTransform", RenderTransformAttribute, EInvalidateWidgetReason::Layout | EInvalidateWidgetReason::RenderTransform);
+	SLATE_ADD_MEMBER_ATTRIBUTE_DEFINITION_WITH_NAME(AttributeInitializer, "RenderTransformPivot", RenderTransformPivotAttribute, EInvalidateWidgetReason::Layout | EInvalidateWidgetReason::RenderTransform);
 }
 
 SWidget::SWidget()
@@ -217,7 +187,6 @@ SWidget::SWidget()
 	, bForceVolatile(false)
 	, bCachedVolatile(false)
 	, bInheritedVolatility(false)
-	, bInvisibleDueToParentOrSelfVisibility(false)
 	, bNeedsPrepass(true)
 	, bHasRegisteredSlateAttribute(false)
 	, bEnabledAttributesUpdate(true)
@@ -235,12 +204,12 @@ SWidget::SWidget()
 	// Note we are defaulting to tick for backwards compatibility
 	, UpdateFlags(EWidgetUpdateFlags::NeedsTick)
 	, DesiredSize()
-	, EnabledStateAttribute(*this, true)
 	, VisibilityAttribute(*this, EVisibility::Visible)
+	, EnabledStateAttribute(*this, true)
+	, RenderTransformPivotAttribute(*this, FVector2D::ZeroVector)
+	, RenderTransformAttribute(*this)
 	, CullingBoundsExtension()
 	, RenderOpacity(1.0f)
-	, RenderTransform()
-	, RenderTransformPivot(FVector2D::ZeroVector)
 #if UE_SLATE_WITH_WIDGET_UNIQUE_IDENTIFIER
 	, UniqueIdentifier(++SlateTraceMetaData::UniqueIdGenerator)
 #endif
@@ -360,8 +329,8 @@ void SWidget::SWidgetConstruct(const FSlateBaseNamedArgs& Args)
 	SetEnabled(Args._IsEnabled);
 	VisibilityAttribute.Assign(*this, Args._Visibility); // SetVisibility is virtual, assign directly to stay backward compatible
 	RenderOpacity = Args._RenderOpacity;
-	RenderTransform = Args._RenderTransform;
-	RenderTransformPivot = Args._RenderTransformPivot;
+	SetRenderTransform(Args._RenderTransform);
+	SetRenderTransformPivot(Args._RenderTransformPivot);
 	Tag = Args._Tag;
 	bForceVolatile = Args._ForceVolatile;
 	Clipping = Args._Clipping;
@@ -693,6 +662,10 @@ void SWidget::SlatePrepass(float InLayoutScaleMultiplier)
 
 	if (!GSlateIsOnFastUpdatePath || bNeedsPrepass)
 	{
+		if (HasRegisteredSlateAttribute() && IsAttributesUpdatesEnabled() && !GSlateIsOnFastProcessInvalidation)
+		{
+			FSlateAttributeMetaData::UpdateAttributes(*this);
+		}
 		Prepass_Internal(InLayoutScaleMultiplier);
 	}
 }
@@ -708,7 +681,7 @@ void SWidget::InvalidateChildRemovedFromTree(SWidget& Child)
 	if (FSlateInvalidationRoot* ChildInvalidationRoot = Child.FastPathProxyHandle.GetInvalidationRootHandle().GetInvalidationRoot())
 	{
 		SCOPED_NAMED_EVENT(SWidget_InvalidateChildRemovedFromTree, FColor::Red);
-		Child.UpdateFastPathVisibility(false, true, ChildInvalidationRoot->GetHittestGrid());
+		Child.UpdateFastPathWidgetRemoved(ChildInvalidationRoot->GetHittestGrid());
 	}
 }
 
@@ -769,6 +742,37 @@ bool SWidget::ConditionallyDetatchParentWidget(SWidget* InExpectedParent)
 	return false;
 }
 
+void SWidget::UpdateWidgetProxy(int32 NewLayerId, FSlateCachedElementsHandle& CacheHandle)
+{
+#if WITH_SLATE_DEBUGGING
+	check(!CacheHandle.IsValid() || CacheHandle.IsOwnedByWidget(this));
+#endif
+
+	// Account for the case when the widget gets a new handle for some reason.  This should really never happen
+	if (PersistentState.CachedElementHandle.IsValid() && PersistentState.CachedElementHandle != CacheHandle)
+	{
+		ensureMsgf(!CacheHandle.IsValid(), TEXT("Widget was assigned a new cache handle.  This is not expected to happen."));
+		PersistentState.CachedElementHandle.RemoveFromCache();
+	}
+	PersistentState.CachedElementHandle = CacheHandle;
+	PersistentState.OutgoingLayerId = NewLayerId;
+
+	if (FastPathProxyHandle.IsValid(this))
+	{
+		FWidgetProxy& MyProxy = FastPathProxyHandle.GetProxy();
+		MyProxy.Visibility = GetVisibility();
+		if ((IsVolatile() && !IsVolatileIndirectly()) || (Advanced_IsInvalidationRoot() && !Advanced_IsWindow()))
+		{
+			AddUpdateFlags(EWidgetUpdateFlags::NeedsVolatilePaint);
+		}
+		else
+		{
+			RemoveUpdateFlags(EWidgetUpdateFlags::NeedsVolatilePaint);
+		}
+		FastPathProxyHandle.MarkWidgetUpdatedThisFrame();
+	}
+}
+
 void SWidget::SetFastPathSortOrder(const FSlateInvalidationWidgetSortOrder InSortOrder)
 {
 	if (InSortOrder != FastPathProxyHandle.GetWidgetSortOrder())
@@ -786,16 +790,15 @@ void SWidget::SetFastPathSortOrder(const FSlateInvalidationWidgetSortOrder InSor
 	}
 }
 
-void SWidget::SetFastPathProxyHandle(const FWidgetProxyHandle& Handle, bool bInInvisibleDueToParentOrSelfVisibility, bool bParentVolatile)
+void SWidget::SetFastPathProxyHandle(const FWidgetProxyHandle& Handle, FSlateInvalidationWidgetVisibility InvalidationVisibility, bool bParentVolatile)
 {
 	check(this != &SNullWidget::NullWidget.Get());
 
 	FastPathProxyHandle = Handle;
 
-	bInvisibleDueToParentOrSelfVisibility = bInInvisibleDueToParentOrSelfVisibility;
 	bInheritedVolatility = bParentVolatile;
 
-	if (bInvisibleDueToParentOrSelfVisibility && PersistentState.CachedElementHandle.IsValid())
+	if (!InvalidationVisibility.IsVisible() && PersistentState.CachedElementHandle.IsValid())
 	{
 #if WITH_SLATE_DEBUGGING
 		check(PersistentState.CachedElementHandle.IsOwnedByWidget(this));
@@ -805,10 +808,8 @@ void SWidget::SetFastPathProxyHandle(const FWidgetProxyHandle& Handle, bool bInI
 
 	if (IsVolatile() && !IsVolatileIndirectly())
 	{
-		if (!HasAnyUpdateFlags(EWidgetUpdateFlags::NeedsVolatilePaint))
-		{
-			AddUpdateFlags(EWidgetUpdateFlags::NeedsVolatilePaint);
-		}
+		// Always add to the list, since it may have been removed with a ChildOrder invalidation.
+		AddUpdateFlags(EWidgetUpdateFlags::NeedsVolatilePaint);
 	}
 	else
 	{
@@ -819,13 +820,10 @@ void SWidget::SetFastPathProxyHandle(const FWidgetProxyHandle& Handle, bool bInI
 	}
 }
 
-void SWidget::UpdateFastPathVisibility(bool bParentVisible, bool bWidgetRemoved, FHittestGrid* ParentHittestGrid)
+void SWidget::UpdateFastPathVisibility(FSlateInvalidationWidgetVisibility ParentVisibility, FHittestGrid* ParentHittestGrid)
 {
 	const EVisibility CurrentVisibility = GetVisibility();
-	const bool bParentAndSelfVisible = bParentVisible && CurrentVisibility.IsVisible();
-	const bool bWasInvisible = bInvisibleDueToParentOrSelfVisibility;
-	bInvisibleDueToParentOrSelfVisibility = !bParentAndSelfVisible;
-	const bool bVisibilityChanged = bWasInvisible != bInvisibleDueToParentOrSelfVisibility;
+	const FSlateInvalidationWidgetVisibility NewVisibility = FSlateInvalidationWidgetVisibility{ ParentVisibility, CurrentVisibility };
 
 	FHittestGrid* HittestGridToRemoveFrom = ParentHittestGrid;
 	if (FastPathProxyHandle.IsValid(this))
@@ -833,17 +831,7 @@ void SWidget::UpdateFastPathVisibility(bool bParentVisible, bool bWidgetRemoved,
 		// Try and remove this from the current handles hit test grid.  If we are in a nested invalidation situation the hittest grid may have changed
 		HittestGridToRemoveFrom = FastPathProxyHandle.GetInvalidationRoot()->GetHittestGrid();
 		FWidgetProxy& Proxy = FastPathProxyHandle.GetProxy();
-		Proxy.Visibility = CurrentVisibility;
-	}
-	else if (bWidgetRemoved)
-	{
-		// The widget can be deleted before the next FastWidgetPathList is built. Remove it now from its InvalidationRoot
-		if (FSlateInvalidationRoot* InvalidationRoot = FastPathProxyHandle.GetInvalidationRootHandle().GetInvalidationRoot())
-		{
-			InvalidationRoot->OnWidgetDestroyed(this);
-		}
-
-		FastPathProxyHandle = FWidgetProxyHandle();
+		Proxy.Visibility = NewVisibility;
 	}
 
 	if (HittestGridToRemoveFrom)
@@ -851,22 +839,43 @@ void SWidget::UpdateFastPathVisibility(bool bParentVisible, bool bWidgetRemoved,
 		HittestGridToRemoveFrom->RemoveWidget(SharedThis(this));
 	}
 
-	if (bWidgetRemoved)
+	PersistentState.CachedElementHandle.ClearCachedElements();
+
+	// Loop through children
+	FChildren* MyChildren = GetAllChildren();
+	const int32 NumChildren = MyChildren->Num();
+	for (int32 ChildIndex = 0; ChildIndex < NumChildren; ++ChildIndex)
 	{
-		PersistentState.CachedElementHandle.RemoveFromCache();
+		const TSharedRef<SWidget>& Child = MyChildren->GetChildAt(ChildIndex);
+		Child->UpdateFastPathVisibility(NewVisibility, HittestGridToRemoveFrom);
 	}
-	else
+}
+
+void SWidget::UpdateFastPathWidgetRemoved(FHittestGrid* ParentHittestGrid)
+{
+	FHittestGrid* HittestGridToRemoveFrom = ParentHittestGrid;
+
+	if (FSlateInvalidationRoot* InvalidationRoot = FastPathProxyHandle.GetInvalidationRootHandle().GetInvalidationRoot())
+	{
+		HittestGridToRemoveFrom = FastPathProxyHandle.GetInvalidationRoot()->GetHittestGrid();
+		InvalidationRoot->OnWidgetDestroyed(this);
+	}
+	FastPathProxyHandle = FWidgetProxyHandle();
+
+	if (HittestGridToRemoveFrom)
 	{
-		PersistentState.CachedElementHandle.ClearCachedElements();
+		HittestGridToRemoveFrom->RemoveWidget(SharedThis(this));
 	}
 
+	PersistentState.CachedElementHandle.RemoveFromCache();
+
 	// Loop through children
 	FChildren* MyChildren = GetAllChildren();
 	const int32 NumChildren = MyChildren->Num();
 	for (int32 ChildIndex = 0; ChildIndex < NumChildren; ++ChildIndex)
 	{
 		const TSharedRef<SWidget>& Child = MyChildren->GetChildAt(ChildIndex);
-		Child->UpdateFastPathVisibility(bParentAndSelfVisible, bWidgetRemoved, HittestGridToRemoveFrom);
+		Child->UpdateFastPathWidgetRemoved(HittestGridToRemoveFrom);
 	}
 }
 
@@ -890,8 +899,6 @@ void SWidget::UpdateFastPathVolatility(bool bParentVolatile)
 		const TSharedRef<SWidget>& Child = MyChildren->GetChildAt(ChildIndex);
 		Child->UpdateFastPathVolatility(bParentVolatile || IsVolatile());
 	}
-
-
 }
 
 void SWidget::CacheDesiredSize(float InLayoutScaleMultiplier)
@@ -1180,6 +1187,11 @@ void SWidget::SetVisibility(TAttribute<EVisibility> InVisibility)
 	VisibilityAttribute.Assign(*this, MoveTemp(InVisibility));
 }
 
+bool SWidget::IsFastPathVisible() const
+{
+	return FastPathProxyHandle.GetWidgetVisibility(this).IsVisible();
+}
+
 void SWidget::Invalidate(EInvalidateWidgetReason InvalidateReason)
 {
 	SLATE_CROSS_THREAD_CHECK();
@@ -1216,8 +1228,7 @@ void SWidget::Invalidate(EInvalidateWidgetReason InvalidateReason)
 		if (EnumHasAnyFlags(InvalidateReason, EInvalidateWidgetReason::Visibility))
 		{
 			SCOPED_NAMED_EVENT(SWidget_UpdateFastPathVisibility, FColor::Red);
-			TSharedPtr<SWidget> ParentWidget = GetParentWidget();
-			UpdateFastPathVisibility(ParentWidget.IsValid() ? !ParentWidget->bInvisibleDueToParentOrSelfVisibility : false, false, FastPathProxyHandle.GetInvalidationRoot()->GetHittestGrid());
+			UpdateFastPathVisibility(FastPathProxyHandle.GetProxy().Visibility.MimicAsParent(), FastPathProxyHandle.GetInvalidationRoot()->GetHittestGrid());
 		}
 
 		if (bVolatilityChanged)
@@ -1438,7 +1449,7 @@ int32 SWidget::Paint(const FPaintArgs& Args, const FGeometry& AllottedGeometry,
 #if WITH_SLATE_DEBUGGING
 	if (!FastPathProxyHandle.IsValid(this) && PersistentState.CachedElementHandle.IsValid())
 	{
-		ensure(!bInvisibleDueToParentOrSelfVisibility);
+		ensure(FastPathProxyHandle.GetProxy().Visibility.IsVisible());
 	}
 #endif
 
@@ -1606,11 +1617,6 @@ void SWidget::Prepass_Internal(float InLayoutScaleMultiplier)
 {
 	PrepassLayoutScaleMultiplier = InLayoutScaleMultiplier;
 
-	if (HasRegisteredSlateAttribute() && IsAttributesUpdatesEnabled() && !GSlateIsOnFastProcessInvalidation)
-	{
-		FSlateAttributeMetaData::UpdateAttributes(*this);
-	}
-
 	bool bShouldPrepassChildren = true;
 	if (bHasCustomPrepass)
 	{
@@ -1645,33 +1651,28 @@ void SWidget::Prepass_ChildLoop(float InLayoutScaleMultiplier, FChildren* MyChil
 
 		const TSharedRef<SWidget>& Child = MyChildren->GetChildAt(ChildIndex);
 
+		const bool bUpdateAttributes = Child->HasRegisteredSlateAttribute() && Child->IsAttributesUpdatesEnabled() && !GSlateIsOnFastProcessInvalidation;
+		if (bUpdateAttributes)
+		{
+			FSlateAttributeMetaData::UpdateCollapsedAttributes(Child.Get());
+		}
+
 		if (Child->GetVisibility() != EVisibility::Collapsed)
 		{
+			if (bUpdateAttributes)
+			{
+				FSlateAttributeMetaData::UpdateExpandedAttributes(Child.Get());
+			}
 			// Recur: Descend down the widget tree.
 			Child->Prepass_Internal(ChildLayoutScaleMultiplier);
 		}
 		else
 		{
-			bool bIsCollapsed = true;
-			if (Child->HasRegisteredSlateAttribute() && Child->IsAttributesUpdatesEnabled() && !GSlateIsOnFastProcessInvalidation)
-			{
-				FSlateAttributeMetaData::UpdateCollapsedAttributes(Child.Get());
-				bIsCollapsed = Child->GetVisibility() == EVisibility::Collapsed;
-			}
-
-			if (bIsCollapsed)
-			{
-				// If the child widget is collapsed, we need to store the new layout scale it will have when 
-				// it is finally visible and invalidate it's prepass so that it gets that when its visibility
-				// is finally invalidated.
-				Child->MarkPrepassAsDirty();
-				Child->PrepassLayoutScaleMultiplier = ChildLayoutScaleMultiplier;
-			}
-			else
-			{
-				// Recur: Descend down the widget tree.
-				Child->Prepass_Internal(ChildLayoutScaleMultiplier);
-			}
+			// If the child widget is collapsed, we need to store the new layout scale it will have when 
+			// it is finally visible and invalidate it's prepass so that it gets that when its visibility
+			// is finally invalidated.
+			Child->MarkPrepassAsDirty();
+			Child->PrepassLayoutScaleMultiplier = ChildLayoutScaleMultiplier;
 		}
 	}
 }
diff --git a/Engine/Source/Runtime/SlateCore/Public/FastUpdate/WidgetProxy.h b/Engine/Source/Runtime/SlateCore/Public/FastUpdate/WidgetProxy.h
--- a/Engine/Source/Runtime/SlateCore/Public/FastUpdate/WidgetProxy.h
+++ b/Engine/Source/Runtime/SlateCore/Public/FastUpdate/WidgetProxy.h
@@ -27,6 +27,80 @@ enum class EInvalidateWidgetReason : uint8;
 #define UE_SLATE_WITH_WIDGETPROXY_WEAKPTR 0
 #define UE_SLATE_VERIFY_WIDGETPROXY_WEAKPTR_STALE 0
 #define UE_SLATE_WITH_WIDGETPROXY_WIDGETTYPE 0
+#define UE_SLATE_WITH_INVALIDATIONWIDGETLIST_DEBUGGING !(UE_BUILD_SHIPPING || UE_BUILD_TEST)
+
+
+struct FSlateInvalidationWidgetVisibility
+{
+public:
+	FSlateInvalidationWidgetVisibility()
+		: Flags(0)
+	{ }
+	FSlateInvalidationWidgetVisibility(EVisibility InVisibility)
+		: bAncestorsVisible(true)
+		, bVisible(InVisibility.IsVisible())
+		, bAncestorCollapse(false)
+		, bCollapse(InVisibility == EVisibility::Collapsed)
+	{ }
+	FSlateInvalidationWidgetVisibility(FSlateInvalidationWidgetVisibility ParentFlags, EVisibility InVisibility)
+		: bAncestorsVisible(ParentFlags.IsVisible())
+		, bVisible(InVisibility.IsVisible())
+		, bAncestorCollapse(ParentFlags.IsCollapsed())
+		, bCollapse(InVisibility == EVisibility::Collapsed)
+	{ }
+	FSlateInvalidationWidgetVisibility(const FSlateInvalidationWidgetVisibility& Other) : Flags(Other.Flags) {  }
+	FSlateInvalidationWidgetVisibility& operator=(const FSlateInvalidationWidgetVisibility& Other) { Flags = Other.Flags; return *this; }
+
+	/** @returns true when all the widget ancestors are visible and the widget itself is visible. */
+	bool IsVisible() const { return bAncestorsVisible && bVisible; }
+	/** @returns true when all the widget ancestors are visible but the widget itself may not be visible. */
+	bool AreAncestorsVisible() const { return bAncestorsVisible; }
+	/** @returns true when at least one of the widget's ancestor is collapse or the widget itself is collapse. */
+	bool IsCollapsed() const { return bAncestorCollapse || bCollapse; }
+	/** @returns true when at least one of the widget's ancestor is collapse but the widget itself may not be collapse. */
+	bool IsCollapseIndirectly() const { return bAncestorCollapse; }
+
+	void SetVisibility(FSlateInvalidationWidgetVisibility ParentFlags, EVisibility InVisibility)
+	{
+		*this = FSlateInvalidationWidgetVisibility(ParentFlags, InVisibility);
+	}
+
+	void SetAncestorsVisibility(FSlateInvalidationWidgetVisibility ParentFlags)
+	{
+		bAncestorsVisible = ParentFlags.IsVisible();
+		bAncestorCollapse = ParentFlags.IsCollapsed();
+	}
+
+	/** Assign the ancestors value to the widget values. Mimicking as it would be the parent. */
+	FSlateInvalidationWidgetVisibility MimicAsParent() const
+	{
+		FSlateInvalidationWidgetVisibility Result;
+		Result.bAncestorsVisible = bAncestorsVisible;
+		Result.bVisible = bAncestorsVisible;
+		Result.bCollapse = bAncestorCollapse;
+		Result.bAncestorCollapse = bAncestorCollapse;
+		return Result;
+	}
+
+	bool operator==(FSlateInvalidationWidgetVisibility Other) const { return Other.CompareFlags == CompareFlags; }
+	bool operator!=(FSlateInvalidationWidgetVisibility Other) const { return Other.CompareFlags != CompareFlags; }
+
+private:
+	union
+	{
+		struct 
+		{
+			uint8 bAncestorsVisible : 1;	// all ancestors are visible
+			uint8 bVisible : 1;
+			uint8 bAncestorCollapse : 1;	// at least one ancestor is collapse
+			uint8 bCollapse : 1;
+		};
+		uint8 CompareFlags : 4;
+		uint8 Flags;
+	};
+};
+static_assert(sizeof(FSlateInvalidationWidgetVisibility) == sizeof(uint8), "FSlateInvalidationWidgetVisibility should be size of uint8");
+
 
 class FWidgetProxy
 {
@@ -73,6 +147,7 @@ private:
 #else
 	SWidget* Widget;
 #endif
+
 #if UE_SLATE_WITH_WIDGETPROXY_WIDGETTYPE
 	FName WidgetType;
 #endif
@@ -82,19 +157,35 @@ public:
 	FSlateInvalidationWidgetIndex ParentIndex;
 	FSlateInvalidationWidgetIndex LeafMostChildIndex;
 	EInvalidateWidgetReason CurrentInvalidateReason;
-	/** The widgets own visibility */
-	EVisibility Visibility;
-	/** Used to make sure we don't double process a widget that is invalidated.  (a widget can invalidate itself but an ancestor can end up painting that widget first thus rendering the child's own invalidate unnecessary */
-	uint8 bUpdatedSinceLastInvalidate : 1;
-	/** Is the widget already in a pending pre update list.  If it already is in an update list we don't bother adding it again */
-	uint8 bContainedByWidgetPreHeap : 1;
-	/** Is the widget already in a pending post update list.  If it already is in an update list we don't bother adding it again */
-	uint8 bContainedByWidgetPostHeap : 1;
-	/** Use with "Slate.InvalidationRoot.VerifyWidgetVisibility". Cached the last FastPathVisible value to find widgets that do not call Invalidate properly. */
-	uint8 bDebug_LastFrameVisible : 1;
-	uint8 bDebug_LastFrameVisibleSet : 1;
-	/** Use with "Slate.InvalidationRoot.VerifyWidgetAttribute". */
-	uint8 bDebug_AttributeUpdated : 1;
+	FSlateInvalidationWidgetVisibility Visibility;
+
+	union
+	{
+		struct
+		{
+		public:
+			/** Used to make sure we don't double process a widget that is invalidated.  (a widget can invalidate itself but an ancestor can end up painting that widget first thus rendering the child's own invalidate unnecessary */
+			uint8 bUpdatedSinceLastInvalidate : 1;
+			/** Is the widget already in a pending pre update list.  If it already is in an update list we don't bother adding it again */
+			uint8 bContainedByWidgetPreHeap : 1;
+			/** Is the widget already in a pending post update list.  If it already is in an update list we don't bother adding it again */
+			uint8 bContainedByWidgetPostHeap : 1;
+			/** Is the widget an Invalidation Root. Cached value of SWidget::Advanced_IsInvalidationRoot */
+			uint8 bContainedByWidgetVolatileList : 1;
+			/** Is the widget an Invalidation Root. Cached value of SWidget::Advanced_IsInvalidationRoot */
+			uint8 bIsInvalidationRoot : 1;
+
+		public:
+#if UE_SLATE_WITH_INVALIDATIONWIDGETLIST_DEBUGGING
+			/** Use with "Slate.InvalidationRoot.VerifyWidgetVisibility". Cached the last FastPathVisible value to find widgets that do not call Invalidate properly. */
+			uint8 bDebug_LastFrameVisible : 1;
+			uint8 bDebug_LastFrameVisibleSet : 1;
+			/** Use with "Slate.InvalidationRoot.VerifyWidgetAttribute". */
+			uint8 bDebug_AttributeUpdated : 1;
+#endif
+		};
+		uint8 PrivateFlags;
+	};
 };
 
 #if !UE_SLATE_WITH_WIDGETPROXY_WIDGETTYPE
@@ -164,6 +255,8 @@ public:
 	FSlateInvalidationWidgetIndex GetWidgetIndex() const { return WidgetIndex; }
 	FSlateInvalidationWidgetSortOrder GetWidgetSortOrder() const { return WidgetSortOrder; }
 
+	SLATECORE_API FSlateInvalidationWidgetVisibility GetWidgetVisibility(const SWidget* Widget) const;
+
 	FWidgetProxy& GetProxy();
 	const FWidgetProxy& GetProxy() const;
 
diff --git a/Engine/Source/Runtime/SlateCore/Public/Types/SlateAttributeMetaData.h b/Engine/Source/Runtime/SlateCore/Public/Types/SlateAttributeMetaData.h
--- a/Engine/Source/Runtime/SlateCore/Public/Types/SlateAttributeMetaData.h
+++ b/Engine/Source/Runtime/SlateCore/Public/Types/SlateAttributeMetaData.h
@@ -37,6 +37,11 @@ public:
 	 * Invalidate the widget if it has finished its construction phase.
 	 */
 	static void UpdateCollapsedAttributes(SWidget& OwningWidget);
+	/**
+	 * Update attributes that are mark to be updated when the widget is NOT collapsed.
+	 * Invalidate the widget if it has finished its construction phase.
+	 */
+	static void UpdateExpandedAttributes(SWidget& OwningWidget);
 	/**
 	 * Update attributes that are mark to be updated when the widget is collapsed.
 	 * @param bAllowInvalidation if we should allow the widget to be invalidated.
@@ -47,6 +52,11 @@ public:
 	 * @param bAllowInvalidation if we should allow the widget to be invalidated.
 	 */
 	static void UpdateCollapsedAttributes(SWidget& OwningWidget, bool bAllowInvalidation);
+	/**
+	 * Update attributes that are mark to be updated when the widget is NOT collapsed.
+	 * @param bAllowInvalidation if we should allow the widget to be invalidated.
+	 */
+	static void UpdateExpandedAttributes(SWidget& OwningWidget, bool bAllowInvalidation);
 
 public:
 	bool IsBound(const FSlateAttributeBase& Attribute) const
@@ -58,7 +68,7 @@ public:
 
 	int32 RegisteredCollaspedAttributeCount() const { return CollaspedAttributeCounter; }
 
-	/** Get the name of all the atrributes, if available. */
+	/** Get the name of all the attributes, if available. */
 	TArray<FName> GetAttributeNames(const SWidget& OwningWidget) const;
 
 private:
@@ -74,9 +84,15 @@ private:
 	static void MoveAttribute(const SWidget& OwningWidget, FSlateAttributeBase& NewAttribute, ESlateAttributeType AttributeType, const FSlateAttributeBase* PreviousAttribute);
 
 private:
+	enum class EUpdateType
+	{
+		All,
+		Collapsed,
+		Expanded, // not mark as collapsed
+	};
 	void RegisterAttributeImpl(SWidget& OwningWidget, FSlateAttributeBase& Attribute, ESlateAttributeType AttributeType, TUniquePtr<ISlateAttributeGetter>&& Getter);
 	bool UnregisterAttributeImpl(const FSlateAttributeBase& Attribute);
-	void UpdateAttributesImpl(SWidget& OwningWidget, bool bOnlyCollapsed, bool bAllowInvalidation);
+	void UpdateAttributesImpl(SWidget& OwningWidget, EUpdateType UpdateType, bool bAllowInvalidation);
 
 private:
 	int32 IndexOfAttribute(const FSlateAttributeBase& Attribute) const
diff --git a/Engine/Source/Runtime/SlateCore/Public/Widgets/SWidget.h b/Engine/Source/Runtime/SlateCore/Public/Widgets/SWidget.h
--- a/Engine/Source/Runtime/SlateCore/Public/Widgets/SWidget.h
+++ b/Engine/Source/Runtime/SlateCore/Public/Widgets/SWidget.h
@@ -764,11 +764,11 @@ protected:
 
 private:
 	void SetFastPathProxyHandle(const FWidgetProxyHandle& Handle) { FastPathProxyHandle = Handle; }
-	void SetFastPathProxyHandle(const FWidgetProxyHandle& Handle, bool bInvisibleDueToParentOrSelfVisibility, bool bParentVolatile);
+	void SetFastPathProxyHandle(const FWidgetProxyHandle& Handle, FSlateInvalidationWidgetVisibility Visibility, bool bParentVolatile);
 	void SetFastPathSortOrder(const FSlateInvalidationWidgetSortOrder SortOrder);
 
-	void UpdateFastPathVisibility(bool bParentVisible, bool bWidgetRemoved, FHittestGrid* ParentHittestGrid);
-
+	void UpdateFastPathVisibility(FSlateInvalidationWidgetVisibility ParentVisibility, FHittestGrid* ParentHittestGrid);
+	void UpdateFastPathWidgetRemoved(FHittestGrid* ParentHittestGrid);
 	void UpdateFastPathVolatility(bool bParentVolatile);
 
 	/**
@@ -995,7 +995,8 @@ public:
 	 * @return is the widget visible and his parents also visible.
 	 * @note only valid if the widget is contained by an InvalidationRoot (the proxy is valid).
 	 */
-	bool IsFastPathVisible() const { return !bInvisibleDueToParentOrSelfVisibility; }
+	UE_DEPRECATED(5.0, "IsFastPathVisible is deprecated and should not be used.")
+	bool IsFastPathVisible() const;
 
 #if WITH_ACCESSIBILITY
 	/**
@@ -1167,19 +1168,19 @@ public:
 	/** @return the render transform of the widget. */
 	FORCEINLINE const TOptional<FSlateRenderTransform>& GetRenderTransform() const
 	{
-		return RenderTransform.Get();
+		return RenderTransformAttribute.Get();
 	}
 
 	FORCEINLINE TOptional<FSlateRenderTransform> GetRenderTransformWithRespectToFlowDirection() const
 	{
 		if (LIKELY(GSlateFlowDirection == EFlowDirection::LeftToRight))
 		{
-			return RenderTransform.Get();
+			return RenderTransformAttribute.Get();
 		}
 		else
 		{
 			// If we're going right to left, flip the X translation on render transforms.
-			TOptional<FSlateRenderTransform> Transform = RenderTransform.Get();
+			TOptional<FSlateRenderTransform> Transform = RenderTransformAttribute.Get();
 			if (Transform.IsSet())
 			{
 				FVector2D Translation = Transform.GetValue().GetTranslation();
@@ -1193,12 +1194,12 @@ public:
 	{
 		if (LIKELY(GSlateFlowDirection == EFlowDirection::LeftToRight))
 		{
-			return RenderTransformPivot.Get();
+			return RenderTransformPivotAttribute.Get();
 		}
 		else
 		{
 			// If we're going right to left, flip the X's pivot mirrored about 0.5.
-			FVector2D TransformPivot = RenderTransformPivot.Get();
+			FVector2D TransformPivot = RenderTransformPivotAttribute.Get();
 			TransformPivot.X = 0.5f + (0.5f - TransformPivot.X);
 			return TransformPivot;
 		}
@@ -1207,19 +1208,19 @@ public:
 	/** @param InTransform the render transform to set for the widget (transforms from widget's local space). TOptional<> to allow code to skip expensive overhead if there is no render transform applied. */
 	FORCEINLINE void SetRenderTransform(TAttribute<TOptional<FSlateRenderTransform>> InTransform)
 	{
-		SetAttribute(RenderTransform, InTransform, EInvalidateWidgetReason::Layout | EInvalidateWidgetReason::RenderTransform);
+		RenderTransformAttribute.Assign(*this, MoveTemp(InTransform));
 	}
 
 	/** @return the pivot point of the render transform. */
 	FORCEINLINE FVector2D GetRenderTransformPivot() const
 	{
-		return RenderTransformPivot.Get();
+		return RenderTransformPivotAttribute.Get();
 	}
 
 	/** @param InTransformPivot Sets the pivot point of the widget's render transform (in normalized local space). */
 	FORCEINLINE void SetRenderTransformPivot(TAttribute<FVector2D> InTransformPivot)
 	{
-		SetAttribute(RenderTransformPivot, InTransformPivot, EInvalidateWidgetReason::Layout | EInvalidateWidgetReason::RenderTransform);
+		RenderTransformPivotAttribute.Assign(*this, MoveTemp(InTransformPivot));
 	}
 
 	/**
@@ -1561,10 +1562,7 @@ protected:
 	 * Recomputes the volatility of the widget.  If you have additional state you automatically want to make
 	 * the widget volatile, you should sample that information here.
 	 */
-	virtual bool ComputeVolatility() const
-	{
-		return RenderTransform.IsBound();
-	}
+	virtual bool ComputeVolatility() const { return false; }
 
 	/**
 	 * Protected static helper to allow widgets to access the visibility attribute of other widgets directly
@@ -1671,6 +1669,10 @@ protected:
 	TSlateAttributeRef<bool> GetEnabledStateAttribute() const { return TSlateAttributeRef<bool>(*this, EnabledStateAttribute); }
 	/** @return an attribute reference of VisibilityAttribute */
 	TSlateAttributeRef<EVisibility> GetVisibilityAttribute() const { return TSlateAttributeRef<EVisibility>(*this, VisibilityAttribute); }
+	/** @return an attribute reference of RenderTransformAttribute */
+	TSlateAttributeRef<TOptional<FSlateRenderTransform>> GetRenderTransformAttribute() const { return TSlateAttributeRef<TOptional<FSlateRenderTransform>>(*this, RenderTransformAttribute); }
+	/** @return an attribute reference of RenderTransformPivotAttribute */
+	TSlateAttributeRef<FVector2D> GetRenderTransformPivotAttribute() const { return TSlateAttributeRef<FVector2D>(*this, RenderTransformPivotAttribute); }
 
 protected:
 	/** Dtor ensures that active timer handles are UnRegistered with the SlateApplication. */
@@ -1717,9 +1719,6 @@ private:
 	/** If we're owned by a volatile widget, we need inherit that volatility and use as part of our volatility, but don't cache it. */
 	uint8 bInheritedVolatility : 1;
 
-	/** If the widget is hidden or collapsed to ancestor visibility */
-	uint8 bInvisibleDueToParentOrSelfVisibility : 1;
-
 	/** Are we currently updating the desired size? */
 	uint8 bNeedsPrepass : 1;
 
@@ -1785,11 +1784,17 @@ private:
 	/** Stores the ideal size this widget wants to be. */
 	TOptional<FVector2D> DesiredSize;
 
+	/** Is this widget visible, hidden or collapsed */
+	TSlateAttribute<EVisibility> VisibilityAttribute;
+
 	/** Whether or not this widget is enabled */
 	TSlateAttribute<bool> EnabledStateAttribute;
 
-	/** Is this widget visible, hidden or collapsed */
-	TSlateAttribute<EVisibility> VisibilityAttribute;
+	/** Render transform pivot of this widget (in normalized local space) */
+	TSlateAttribute<FVector2D> RenderTransformPivotAttribute;
+
+	/** Render transform of this widget. TOptional<> to allow code to skip expensive overhead if there is no render transform applied. */
+	TSlateAttribute<TOptional<FSlateRenderTransform>> RenderTransformAttribute;
 
 protected:
 
@@ -1809,17 +1814,17 @@ protected:
 	/** Is this widget visible, hidden or collapsed */
 	UE_DEPRECATED(5.0, "Direct access to Visibility is now deprecated. Use the setter or getter.")
 	FSlateDeprecatedTAttribute<EVisibility> Visibility;
+	/** Render transform of this widget. TOptional<> to allow code to skip expensive overhead if there is no render transform applied. */
+	UE_DEPRECATED(5.0, "Direct access to RenderTransform is now deprecated. Use the setter or getter.")
+	FSlateDeprecatedTAttribute< TOptional<FSlateRenderTransform> > RenderTransform;
+	/** Render transform pivot of this widget (in normalized local space) */
+	UE_DEPRECATED(5.0, "Direct access to RenderTransformPivot is now deprecated. Use the setter or getter.")
+	TAttribute<FVector2D> RenderTransformPivot;
 #endif
 
 	/** The opacity of the widget. Automatically applied during rendering. */
 	float RenderOpacity;
 
-	/** Render transform of this widget. TOptional<> to allow code to skip expensive overhead if there is no render transform applied. */
-	TAttribute< TOptional<FSlateRenderTransform> > RenderTransform;
-
-	/** Render transform pivot of this widget (in normalized local space) */
-	TAttribute< FVector2D > RenderTransformPivot;
-
 private:
 	/** Metadata associated with this widget. */
 	TArray<TSharedRef<ISlateMetaData>> MetaData;
