Slate: Fix ensure test when the CachedElementHandle is valid.
#rnx

diff --git a/Engine/Source/Runtime/SlateCore/Private/Widgets/SWidget.cpp b/Engine/Source/Runtime/SlateCore/Private/Widgets/SWidget.cpp
--- a/Engine/Source/Runtime/SlateCore/Private/Widgets/SWidget.cpp
+++ b/Engine/Source/Runtime/SlateCore/Private/Widgets/SWidget.cpp
@@ -1447,7 +1447,7 @@ int32 SWidget::Paint(const FPaintArgs& Args, const FGeometry& AllottedGeometry,
 #endif
 
 #if WITH_SLATE_DEBUGGING
-	if (!FastPathProxyHandle.IsValid(this) && PersistentState.CachedElementHandle.IsValid())
+	if (FastPathProxyHandle.IsValid(this) && PersistentState.CachedElementHandle.IsValid())
 	{
 		ensure(FastPathProxyHandle.GetProxy().Visibility.IsVisible());
 	}
