Slate: Convert SCompoundWidget and SBox attributes to SlateAttribute. Mark TSupportsContentPaddingMixin::SlotPadding as deprecated to move it to private later.

diff --git a/Engine/Plugins/Experimental/CommonUI/Source/CommonUI/Private/CommonTextBlock.cpp b/Engine/Plugins/Experimental/CommonUI/Source/CommonUI/Private/CommonTextBlock.cpp
--- a/Engine/Plugins/Experimental/CommonUI/Source/CommonUI/Private/CommonTextBlock.cpp
+++ b/Engine/Plugins/Experimental/CommonUI/Source/CommonUI/Private/CommonTextBlock.cpp
@@ -30,7 +30,7 @@ void STextScroller::OnArrangeChildren(const FGeometry& AllottedGeometry, FArrang
 	if (ArrangedChildren.Accepts(ChildVisibility))
 	{
 		const FVector2D ThisContentScale = GetContentScale();
-		const FMargin SlotPadding(LayoutPaddingWithFlow(GSlateFlowDirection, ChildSlot.SlotPadding.Get()));
+		const FMargin SlotPadding(LayoutPaddingWithFlow(GSlateFlowDirection, ChildSlot.GetPadding()));
 		const AlignmentArrangeResult XResult = AlignChild<Orient_Horizontal>(GSlateFlowDirection, AllottedGeometry.GetLocalSize().X, ChildSlot, SlotPadding, ThisContentScale.X, false);
 		const AlignmentArrangeResult YResult = AlignChild<Orient_Vertical>(AllottedGeometry.GetLocalSize().Y, ChildSlot, SlotPadding, ThisContentScale.Y);
 		const FVector2D MinimumSize = ChildSlot.GetWidget()->GetDesiredSize();
diff --git a/Engine/Plugins/Experimental/DigitalMultiplex/DMXPixelMapping/Source/DMXPixelMappingEditor/Private/Widgets/SDMXPixelMappingZoomPan.cpp b/Engine/Plugins/Experimental/DigitalMultiplex/DMXPixelMapping/Source/DMXPixelMappingEditor/Private/Widgets/SDMXPixelMappingZoomPan.cpp
--- a/Engine/Plugins/Experimental/DigitalMultiplex/DMXPixelMapping/Source/DMXPixelMappingEditor/Private/Widgets/SDMXPixelMappingZoomPan.cpp
+++ b/Engine/Plugins/Experimental/DigitalMultiplex/DMXPixelMapping/Source/DMXPixelMappingEditor/Private/Widgets/SDMXPixelMappingZoomPan.cpp
@@ -29,7 +29,7 @@ void SDMXPixelMappingZoomPan::OnArrangeChildren(const FGeometry& AllottedGeometr
 	const EVisibility ChildVisibility = ChildSlot.GetWidget()->GetVisibility();
 	if ( ArrangedChildren.Accepts(ChildVisibility) )
 	{
-		const FMargin SlotPadding(ChildSlot.SlotPadding.Get());
+		const FMargin SlotPadding(ChildSlot.GetPadding());
 		AlignmentArrangeResult XResult = AlignChild<Orient_Horizontal>(AllottedGeometry.Size.X, ChildSlot, SlotPadding, 1);
 		AlignmentArrangeResult YResult = AlignChild<Orient_Vertical>(AllottedGeometry.Size.Y, ChildSlot, SlotPadding, 1);
 
diff --git a/Engine/Source/Editor/UMGEditor/Private/Designer/SZoomPan.cpp b/Engine/Source/Editor/UMGEditor/Private/Designer/SZoomPan.cpp
--- a/Engine/Source/Editor/UMGEditor/Private/Designer/SZoomPan.cpp
+++ b/Engine/Source/Editor/UMGEditor/Private/Designer/SZoomPan.cpp
@@ -32,7 +32,7 @@ void SZoomPan::OnArrangeChildren(const FGeometry& AllottedGeometry, FArrangedChi
 	const EVisibility ChildVisibility = ChildSlot.GetWidget()->GetVisibility();
 	if ( ArrangedChildren.Accepts(ChildVisibility) )
 	{
-		const FMargin SlotPadding(ChildSlot.SlotPadding.Get());
+		const FMargin SlotPadding(ChildSlot.GetPadding());
 		AlignmentArrangeResult XResult = AlignChild<Orient_Horizontal>(AllottedGeometry.Size.X, ChildSlot, SlotPadding, 1);
 		AlignmentArrangeResult YResult = AlignChild<Orient_Vertical>(AllottedGeometry.Size.Y, ChildSlot, SlotPadding, 1);
 
diff --git a/Engine/Source/Runtime/Slate/Private/Framework/Docking/STabSidebar.cpp b/Engine/Source/Runtime/Slate/Private/Framework/Docking/STabSidebar.cpp
--- a/Engine/Source/Runtime/Slate/Private/Framework/Docking/STabSidebar.cpp
+++ b/Engine/Source/Runtime/Slate/Private/Framework/Docking/STabSidebar.cpp
@@ -553,9 +553,9 @@ void STabSidebar::OpenDrawerInternal(TSharedRef<SDockTab> ForTab)
 
 		const float DPIScale = MyWindow->GetDPIScaleFactor();
 
-		const float TopOffset = ((ChildSlot.SlotPadding.Get().Top + MyGeometry.GetAbsolutePosition().Y) - WindowGeometry.GetAbsolutePosition().Y);
+		const float TopOffset = ((ChildSlot.GetPadding().Top + MyGeometry.GetAbsolutePosition().Y) - WindowGeometry.GetAbsolutePosition().Y);
 
-		const float BottomOffset = ((ChildSlot.SlotPadding.Get().Bottom + WindowGeometry.GetAbsolutePositionAtCoordinates(FVector2D::UnitVector).Y) - MyGeometry.GetAbsolutePositionAtCoordinates(FVector2D::UnitVector).Y);
+		const float BottomOffset = ((ChildSlot.GetPadding().Bottom + WindowGeometry.GetAbsolutePositionAtCoordinates(FVector2D::UnitVector).Y) - MyGeometry.GetAbsolutePositionAtCoordinates(FVector2D::UnitVector).Y);
 
 		FMargin SlotPadding(
 			Location == ESidebarLocation::Left ? MinDrawerSize : 0.0f,
diff --git a/Engine/Source/Runtime/Slate/Private/Widgets/Layout/SBackgroundBlur.cpp b/Engine/Source/Runtime/Slate/Private/Widgets/Layout/SBackgroundBlur.cpp
--- a/Engine/Source/Runtime/Slate/Private/Widgets/Layout/SBackgroundBlur.cpp
+++ b/Engine/Source/Runtime/Slate/Private/Widgets/Layout/SBackgroundBlur.cpp
@@ -93,13 +93,9 @@ void SBackgroundBlur::SetVAlign(EVerticalAlignment VAlign)
 	}
 }
 
-void SBackgroundBlur::SetPadding(const TAttribute<FMargin>& InPadding)
+void SBackgroundBlur::SetPadding(TAttribute<FMargin> InPadding)
 {
-	if (!ChildSlot.SlotPadding.IdenticalTo(InPadding))
-	{
-		ChildSlot.SlotPadding = InPadding;
-		Invalidate(EInvalidateWidget::Layout);
-	}
+	ChildSlot.Padding(MoveTemp(InPadding));
 }
 
 bool SBackgroundBlur::IsUsingLowQualityFallbackBrush() const
diff --git a/Engine/Source/Runtime/Slate/Private/Widgets/Layout/SBorder.cpp b/Engine/Source/Runtime/Slate/Private/Widgets/Layout/SBorder.cpp
--- a/Engine/Source/Runtime/Slate/Private/Widgets/Layout/SBorder.cpp
+++ b/Engine/Source/Runtime/Slate/Private/Widgets/Layout/SBorder.cpp
@@ -164,7 +164,7 @@ void SBorder::SetVAlign(EVerticalAlignment VAlign)
 
 void SBorder::SetPadding(const TAttribute<FMargin>& InPadding)
 {
-	SetAttribute(ChildSlot.SlotPadding, InPadding, EInvalidateWidgetReason::Layout);
+	ChildSlot.SetPadding(InPadding);
 }
 
 void SBorder::SetShowEffectWhenDisabled(TAttribute<bool> InShowEffectWhenDisabled)
diff --git a/Engine/Source/Runtime/Slate/Private/Widgets/Layout/SBox.cpp b/Engine/Source/Runtime/Slate/Private/Widgets/Layout/SBox.cpp
--- a/Engine/Source/Runtime/Slate/Private/Widgets/Layout/SBox.cpp
+++ b/Engine/Source/Runtime/Slate/Private/Widgets/Layout/SBox.cpp
@@ -8,6 +8,7 @@
 SLATE_IMPLEMENT_WIDGET(SBox)
 void SBox::PrivateRegisterAttributes(FSlateAttributeInitializer& AttributeInitializer)
 {
+	AttributeInitializer.AddMemberAttribute("SlotPadding", STRUCT_OFFSET(PrivateThisType, ChildSlot) + FBoxSlot::GetSlotPaddingAttributeOffset(), FSlateAttributeDescriptor::FInvalidateWidgetReasonAttribute{ EInvalidateWidgetReason::Layout });
 	SLATE_ADD_MEMBER_ATTRIBUTE_DEFINITION(AttributeInitializer, WidthOverride, EInvalidateWidgetReason::Layout);
 	SLATE_ADD_MEMBER_ATTRIBUTE_DEFINITION(AttributeInitializer, HeightOverride, EInvalidateWidgetReason::Layout);
 	SLATE_ADD_MEMBER_ATTRIBUTE_DEFINITION(AttributeInitializer, MinDesiredWidth, EInvalidateWidgetReason::Layout);
@@ -19,7 +20,7 @@ void SBox::PrivateRegisterAttributes(FSlateAttributeInitializer& AttributeInitia
 }
 
 SBox::SBox()
-	: ChildSlot(this)
+	: ChildSlot(*this)
 	, WidthOverride(*this)
 	, HeightOverride(*this)
 	, MinDesiredWidth(*this)
@@ -58,17 +59,13 @@ void SBox::Construct( const FArguments& InArgs )
 void SBox::SetContent(const TSharedRef< SWidget >& InContent)
 {
 	ChildSlot
-		[
-			InContent
-		];
-
-	// TODO SlateGI - This seems no longer needed.
-	Invalidate(EInvalidateWidget::Layout);
+	[
+		InContent
+	];
 }
 
 void SBox::SetHAlign(EHorizontalAlignment HAlign)
 {
-	// TODO SlateGI - Fix Slots
 	if (ChildSlot.HAlignment != HAlign)
 	{
 		ChildSlot.HAlignment = HAlign;
@@ -78,7 +75,6 @@ void SBox::SetHAlign(EHorizontalAlignment HAlign)
 
 void SBox::SetVAlign(EVerticalAlignment VAlign)
 {
-	// TODO SlateGI - Fix Slots
 	if (ChildSlot.VAlignment != VAlign)
 	{
 		ChildSlot.VAlignment = VAlign;
@@ -88,12 +84,7 @@ void SBox::SetVAlign(EVerticalAlignment VAlign)
 
 void SBox::SetPadding(const TAttribute<FMargin>& InPadding)
 {
-	// TODO SlateGI - Fix Slots
-	if (!ChildSlot.SlotPadding.IdenticalTo(InPadding))
-	{
-		ChildSlot.SlotPadding = InPadding;
-		Invalidate(EInvalidateWidget::LayoutAndVolatility);
-	}
+	ChildSlot.SetPadding(InPadding);
 }
 
 void SBox::SetWidthOverride(TAttribute<FOptionalSize> InWidthOverride)
@@ -157,7 +148,7 @@ FVector2D SBox::ComputeDesiredSize( float ) const
 float SBox::ComputeDesiredWidth() const
 {
 	// If the user specified a fixed width or height, those values override the Box's content.
-	const FVector2D& UnmodifiedChildDesiredSize = ChildSlot.GetWidget()->GetDesiredSize() + ChildSlot.SlotPadding.Get().GetDesiredSize();
+	const FVector2D& UnmodifiedChildDesiredSize = ChildSlot.GetWidget()->GetDesiredSize() + ChildSlot.GetPadding().GetDesiredSize();
 	const FOptionalSize CurrentMinDesiredWidth = MinDesiredWidth.Get();
 	const FOptionalSize CurrentMaxDesiredWidth = MaxDesiredWidth.Get();
 
@@ -179,7 +170,7 @@ float SBox::ComputeDesiredWidth() const
 float SBox::ComputeDesiredHeight() const
 {
 	// If the user specified a fixed width or height, those values override the Box's content.
-	const FVector2D& UnmodifiedChildDesiredSize = ChildSlot.GetWidget()->GetDesiredSize() + ChildSlot.SlotPadding.Get().GetDesiredSize();
+	const FVector2D& UnmodifiedChildDesiredSize = ChildSlot.GetWidget()->GetDesiredSize() + ChildSlot.GetPadding().GetDesiredSize();
 
 	const FOptionalSize CurrentMinDesiredHeight = MinDesiredHeight.Get();
 	const FOptionalSize CurrentMaxDesiredHeight = MaxDesiredHeight.Get();
@@ -206,7 +197,7 @@ void SBox::OnArrangeChildren( const FGeometry& AllottedGeometry, FArrangedChildr
 	{
 		const FOptionalSize CurrentMinAspectRatio = MinAspectRatio.Get();
 		const FOptionalSize CurrentMaxAspectRatio = MaxAspectRatio.Get();
-		const FMargin SlotPadding(ChildSlot.SlotPadding.Get());
+		const FMargin SlotPadding(ChildSlot.GetPadding());
 		bool bAlignChildren = true;
 
 		AlignmentArrangeResult XAlignmentResult(0, 0);
diff --git a/Engine/Source/Runtime/Slate/Private/Widgets/Layout/SDPIScaler.cpp b/Engine/Source/Runtime/Slate/Private/Widgets/Layout/SDPIScaler.cpp
--- a/Engine/Source/Runtime/Slate/Private/Widgets/Layout/SDPIScaler.cpp
+++ b/Engine/Source/Runtime/Slate/Private/Widgets/Layout/SDPIScaler.cpp
@@ -6,6 +6,7 @@
 SLATE_IMPLEMENT_WIDGET(SDPIScaler)
 void SDPIScaler::PrivateRegisterAttributes(FSlateAttributeInitializer& AttributeInitializer)
 {
+	AttributeInitializer.AddMemberAttribute("SlotPadding", STRUCT_OFFSET(PrivateThisType, ChildSlot) + FOneSimpleMemberChild::GetSlotPaddingAttributeOffset(), FSlateAttributeDescriptor::FInvalidateWidgetReasonAttribute{ EInvalidateWidgetReason::Layout });
 	SLATE_ADD_MEMBER_ATTRIBUTE_DEFINITION_WITH_NAME(AttributeInitializer, "DPIScale", DPIScaleAttribute, EInvalidateWidgetReason::Prepass);
 }
 
diff --git a/Engine/Source/Runtime/Slate/Private/Widgets/Layout/SExpandableArea.cpp b/Engine/Source/Runtime/Slate/Private/Widgets/Layout/SExpandableArea.cpp
--- a/Engine/Source/Runtime/Slate/Private/Widgets/Layout/SExpandableArea.cpp
+++ b/Engine/Source/Runtime/Slate/Private/Widgets/Layout/SExpandableArea.cpp
@@ -195,7 +195,7 @@ FVector2D SExpandableArea::ComputeDesiredSize( float ) const
 	EVisibility ChildVisibility = ChildSlot.GetWidget()->GetVisibility();
 	if ( ChildVisibility != EVisibility::Collapsed )
 	{
-		FVector2D SlotWidgetDesiredSize = ChildSlot.GetWidget()->GetDesiredSize() + ChildSlot.SlotPadding.Get().GetDesiredSize();
+		FVector2D SlotWidgetDesiredSize = ChildSlot.GetWidget()->GetDesiredSize() + ChildSlot.GetPadding().GetDesiredSize();
 		
 		// Only clamp if the user specified a min width
 		if( MinWidth > 0.0f )
diff --git a/Engine/Source/Runtime/Slate/Private/Widgets/Layout/SGridPanel.cpp b/Engine/Source/Runtime/Slate/Private/Widgets/Layout/SGridPanel.cpp
--- a/Engine/Source/Runtime/Slate/Private/Widgets/Layout/SGridPanel.cpp
+++ b/Engine/Source/Runtime/Slate/Private/Widgets/Layout/SGridPanel.cpp
@@ -201,7 +201,7 @@ void SGridPanel::OnArrangeChildren( const FGeometry& AllottedGeometry, FArranged
 
 			// Do the standard arrangement of elements within a slot
 			// Takes care of alignment and padding.
-			const FMargin SlotPadding(CurSlot.SlotPadding.Get());
+			const FMargin SlotPadding(CurSlot.GetPadding());
 			AlignmentArrangeResult XAxisResult = AlignChild<Orient_Horizontal>( CellSize.X, CurSlot, SlotPadding );
 			AlignmentArrangeResult YAxisResult = AlignChild<Orient_Vertical>( CellSize.Y, CurSlot, SlotPadding );
 
@@ -391,7 +391,7 @@ void SGridPanel::ComputeDesiredCellSizes( TArray<float>& OutColumns, TArray<floa
 		if (CurSlot.GetWidget()->GetVisibility() != EVisibility::Collapsed)
 		{
 			// The slots wants to be as big as its content along with the required padding.
-			const FVector2D SlotDesiredSize = CurSlot.GetWidget()->GetDesiredSize() + CurSlot.SlotPadding.Get().GetDesiredSize();
+			const FVector2D SlotDesiredSize = CurSlot.GetWidget()->GetDesiredSize() + CurSlot.GetPadding().GetDesiredSize();
 
 			// If the slot has a (colspan,rowspan) of (1,1) it will only affect that slot.
 			// For larger spans, the slots size will be evenly distributed across all the affected slots.
diff --git a/Engine/Source/Runtime/Slate/Private/Widgets/Layout/SResponsiveGridPanel.cpp b/Engine/Source/Runtime/Slate/Private/Widgets/Layout/SResponsiveGridPanel.cpp
--- a/Engine/Source/Runtime/Slate/Private/Widgets/Layout/SResponsiveGridPanel.cpp
+++ b/Engine/Source/Runtime/Slate/Private/Widgets/Layout/SResponsiveGridPanel.cpp
@@ -253,7 +253,7 @@ void SResponsiveGridPanel::OnArrangeChildren( const FGeometry& AllottedGeometry,
 
 			// Do the standard arrangement of elements within a slot
 			// Takes care of alignment and padding.
-			FMargin SlotPadding(CurSlot.SlotPadding.Get());
+			FMargin SlotPadding(CurSlot.GetPadding());
 
 			AlignmentArrangeResult XAxisResult = AlignChild<Orient_Horizontal>(CellSize.X, CurSlot, SlotPadding);
 			AlignmentArrangeResult YAxisResult = AlignChild<Orient_Vertical>(CellSize.Y, CurSlot, SlotPadding);
@@ -348,7 +348,7 @@ void SResponsiveGridPanel::ComputeDesiredCellSizes(float AvailableWidth, TArray<
 			}
 
 			// The slots want to be as big as its content along with the required padding.
-			const FVector2D SlotDesiredSize = CurSlot.GetWidget()->GetDesiredSize() + CurSlot.SlotPadding.Get().GetDesiredSize();
+			const FVector2D SlotDesiredSize = CurSlot.GetWidget()->GetDesiredSize() + CurSlot.GetPadding().GetDesiredSize();
 
 			// If the slot has a (colspan,rowspan) of (1,1) it will only affect that cell.
 			// For larger spans, the slots size will be evenly distributed across all the affected cells.
diff --git a/Engine/Source/Runtime/Slate/Private/Widgets/Layout/SScaleBox.cpp b/Engine/Source/Runtime/Slate/Private/Widgets/Layout/SScaleBox.cpp
--- a/Engine/Source/Runtime/Slate/Private/Widgets/Layout/SScaleBox.cpp
+++ b/Engine/Source/Runtime/Slate/Private/Widgets/Layout/SScaleBox.cpp
@@ -227,7 +227,7 @@ void SScaleBox::OnArrangeChildren(const FGeometry& AllottedGeometry, FArrangedCh
 		// If we're just filling, there's no scale applied, we're just filling the area.
 		if (CurrentStretch != EStretch::Fill)
 		{
-			const FMargin SlotPadding(ChildSlot.SlotPadding.Get());
+			const FMargin SlotPadding(ChildSlot.GetPadding());
 			AlignmentArrangeResult XResult = AlignChild<Orient_Horizontal>(AreaSize.X, ChildSlot, SlotPadding, FinalScale, false);
 			AlignmentArrangeResult YResult = AlignChild<Orient_Vertical>(AreaSize.Y, ChildSlot, SlotPadding, FinalScale, false);
 
diff --git a/Engine/Source/Runtime/Slate/Private/Widgets/Layout/SScrollBox.cpp b/Engine/Source/Runtime/Slate/Private/Widgets/Layout/SScrollBox.cpp
--- a/Engine/Source/Runtime/Slate/Private/Widgets/Layout/SScrollBox.cpp
+++ b/Engine/Source/Runtime/Slate/Private/Widgets/Layout/SScrollBox.cpp
@@ -65,11 +65,11 @@ FVector2D SScrollPanel::ComputeDesiredSize(float) const
 			if (Orientation == Orient_Vertical)
 			{
 				ThisDesiredSize.X = FMath::Max(ChildDesiredSize.X, ThisDesiredSize.X);
-				ThisDesiredSize.Y += ChildDesiredSize.Y + ThisSlot.SlotPadding.Get().GetTotalSpaceAlong<Orient_Vertical>();
+				ThisDesiredSize.Y += ChildDesiredSize.Y + ThisSlot.GetPadding().GetTotalSpaceAlong<Orient_Vertical>();
 			}
 			else
 			{
-				ThisDesiredSize.X += ChildDesiredSize.X + ThisSlot.SlotPadding.Get().GetTotalSpaceAlong<Orient_Horizontal>();
+				ThisDesiredSize.X += ChildDesiredSize.X + ThisSlot.GetPadding().GetTotalSpaceAlong<Orient_Horizontal>();
 				ThisDesiredSize.Y = FMath::Max(ChildDesiredSize.Y, ThisDesiredSize.Y);
 			}
 		}
@@ -85,7 +85,7 @@ FVector2D SScrollPanel::ComputeDesiredSize(float) const
 
 float SScrollPanel::ArrangeChildVerticalAndReturnOffset(const FGeometry& AllottedGeometry, FArrangedChildren& ArrangedChildren, const SScrollBox::FSlot& ThisSlot, float CurChildOffset) const
 {
-	const FMargin& ThisPadding = ThisSlot.SlotPadding.Get();
+	const FMargin& ThisPadding = ThisSlot.GetPadding();
 	const FVector2D& WidgetDesiredSize = ThisSlot.GetWidget()->GetDesiredSize();
 	const float ThisSlotDesiredHeight = WidgetDesiredSize.Y + ThisPadding.GetTotalSpaceAlong<Orient_Vertical>();
 
@@ -99,7 +99,7 @@ float SScrollPanel::ArrangeChildVerticalAndReturnOffset(const FGeometry& Allotte
 
 float SScrollPanel::ArrangeChildHorizontalAndReturnOffset(const FGeometry& AllottedGeometry, FArrangedChildren& ArrangedChildren, const SScrollBox::FSlot& ThisSlot, float CurChildOffset) const
 {
-	const FMargin& ThisPadding = ThisSlot.SlotPadding.Get();
+	const FMargin& ThisPadding = ThisSlot.GetPadding();
 	const FVector2D& WidgetDesiredSize = ThisSlot.GetWidget()->GetDesiredSize();
 	const float ThisSlotDesiredWidth = WidgetDesiredSize.X + ThisPadding.GetTotalSpaceAlong<Orient_Horizontal>();
 
diff --git a/Engine/Source/Runtime/Slate/Private/Widgets/Layout/SWidgetSwitcher.cpp b/Engine/Source/Runtime/Slate/Private/Widgets/Layout/SWidgetSwitcher.cpp
--- a/Engine/Source/Runtime/Slate/Private/Widgets/Layout/SWidgetSwitcher.cpp
+++ b/Engine/Source/Runtime/Slate/Private/Widgets/Layout/SWidgetSwitcher.cpp
@@ -183,7 +183,7 @@ FVector2D SWidgetSwitcher::ComputeDesiredSize(float) const
 		const EVisibility ChildVisibility = Widget->GetVisibility();
 		if (ChildVisibility != EVisibility::Collapsed)
 		{
-			return Widget->GetDesiredSize() + ActiveSlotPtr->SlotPadding.Get().GetDesiredSize();
+			return Widget->GetDesiredSize() + ActiveSlotPtr->GetPadding().GetDesiredSize();
 		}
 	}
 
diff --git a/Engine/Source/Runtime/Slate/Private/Widgets/Layout/SWindowTitleBarArea.cpp b/Engine/Source/Runtime/Slate/Private/Widgets/Layout/SWindowTitleBarArea.cpp
--- a/Engine/Source/Runtime/Slate/Private/Widgets/Layout/SWindowTitleBarArea.cpp
+++ b/Engine/Source/Runtime/Slate/Private/Widgets/Layout/SWindowTitleBarArea.cpp
@@ -191,7 +191,7 @@ void SWindowTitleBarArea::SetVAlign(EVerticalAlignment VAlign)
 
 void SWindowTitleBarArea::SetPadding(const TAttribute<FMargin>& InPadding)
 {
-	ChildSlot.SlotPadding = InPadding;
+	ChildSlot.Padding(InPadding);
 }
 
 FVector2D SWindowTitleBarArea::ComputeDesiredSize( float ) const
@@ -200,7 +200,7 @@ FVector2D SWindowTitleBarArea::ComputeDesiredSize( float ) const
 
 	if ( ChildVisibility != EVisibility::Collapsed )
 	{
-		return ChildSlot.GetWidget()->GetDesiredSize() + ChildSlot.SlotPadding.Get().GetDesiredSize();
+		return ChildSlot.GetWidget()->GetDesiredSize() + ChildSlot.GetPadding().GetDesiredSize();
 	}
 	
 	return FVector2D::ZeroVector;
@@ -211,7 +211,7 @@ void SWindowTitleBarArea::OnArrangeChildren( const FGeometry& AllottedGeometry,
 	const EVisibility& MyCurrentVisibility = this->GetVisibility();
 	if ( ArrangedChildren.Accepts( MyCurrentVisibility ) )
 	{
-		const FMargin SlotPadding(ChildSlot.SlotPadding.Get());
+		const FMargin SlotPadding(ChildSlot.GetPadding());
 		AlignmentArrangeResult XAlignmentResult = AlignChild<Orient_Horizontal>( AllottedGeometry.GetLocalSize().X, ChildSlot, SlotPadding );
 		AlignmentArrangeResult YAlignmentResult = AlignChild<Orient_Vertical>( AllottedGeometry.GetLocalSize().Y, ChildSlot, SlotPadding );
 
diff --git a/Engine/Source/Runtime/Slate/Private/Widgets/Layout/SWrapBox.cpp b/Engine/Source/Runtime/Slate/Private/Widgets/Layout/SWrapBox.cpp
--- a/Engine/Source/Runtime/Slate/Private/Widgets/Layout/SWrapBox.cpp
+++ b/Engine/Source/Runtime/Slate/Private/Widgets/Layout/SWrapBox.cpp
@@ -157,7 +157,7 @@ void SWrapBox::FChildArranger::Arrange()
 			Offset.X += ((WrapBox.Orientation == EOrientation::Orient_Horizontal) * WrapBox.InnerSlotPadding.X);
 		}
 
-		const FVector2D DesiredSizeOfSlot = Slot.SlotPadding.Get().GetDesiredSize() + Widget->GetDesiredSize();
+		const FVector2D DesiredSizeOfSlot = Slot.GetPadding().GetDesiredSize() + Widget->GetDesiredSize();
 
 		// Populate arrangement data with default size and offset at right end of current line.
 		ArrangementData.SlotOffset.X = Offset.X;
@@ -344,8 +344,8 @@ void SWrapBox::OnArrangeChildren(const FGeometry& AllottedGeometry, FArrangedChi
 	FChildArranger::Arrange(*this, [&](const FSlot& Slot, const FChildArranger::FArrangementData& ArrangementData)
 	{
 		// Calculate offset and size in slot using alignment.
-		const AlignmentArrangeResult XResult = AlignChild<Orient_Horizontal>(ArrangementData.SlotSize.X, Slot, Slot.SlotPadding.Get());
-		const AlignmentArrangeResult YResult = AlignChild<Orient_Vertical>(ArrangementData.SlotSize.Y, Slot, Slot.SlotPadding.Get());
+		const AlignmentArrangeResult XResult = AlignChild<Orient_Horizontal>(ArrangementData.SlotSize.X, Slot, Slot.GetPadding());
+		const AlignmentArrangeResult YResult = AlignChild<Orient_Vertical>(ArrangementData.SlotSize.Y, Slot, Slot.GetPadding());
 
 		// Note: Alignment offset is relative to slot offset.
 		const FVector2D PostAlignmentOffset = ArrangementData.SlotOffset + FVector2D(XResult.Offset, YResult.Offset);
diff --git a/Engine/Source/Runtime/Slate/Public/Widgets/Layout/SBackgroundBlur.h b/Engine/Source/Runtime/Slate/Public/Widgets/Layout/SBackgroundBlur.h
--- a/Engine/Source/Runtime/Slate/Public/Widgets/Layout/SBackgroundBlur.h
+++ b/Engine/Source/Runtime/Slate/Public/Widgets/Layout/SBackgroundBlur.h
@@ -53,7 +53,7 @@ public:
 	
 	void SetHAlign(EHorizontalAlignment HAlign);
 	void SetVAlign(EVerticalAlignment VAlign);
-	void SetPadding(const TAttribute<FMargin>& InPadding);
+	void SetPadding(TAttribute<FMargin> InPadding);
 
 	bool IsUsingLowQualityFallbackBrush() const;
 
diff --git a/Engine/Source/Runtime/Slate/Public/Widgets/Layout/SBox.h b/Engine/Source/Runtime/Slate/Public/Widgets/Layout/SBox.h
--- a/Engine/Source/Runtime/Slate/Public/Widgets/Layout/SBox.h
+++ b/Engine/Source/Runtime/Slate/Public/Widgets/Layout/SBox.h
@@ -24,14 +24,9 @@ class SLATE_API SBox : public SPanel
 	SLATE_DECLARE_WIDGET(SBox, SPanel)
 
 public:
-	class FBoxSlot : public TSupportsOneChildMixin<FBoxSlot>, public TSupportsContentAlignmentMixin<FBoxSlot>, public TSupportsContentPaddingMixin<FBoxSlot>
+	class FBoxSlot : public FOneSimpleMemberChild
 	{
-	public:
-		FBoxSlot(SWidget* InOwner)
-			: TSupportsOneChildMixin<FBoxSlot>(InOwner)
-			, TSupportsContentAlignmentMixin<FBoxSlot>(HAlign_Fill, VAlign_Fill)
-		{
-		}
+		using FOneSimpleMemberChild::FOneSimpleMemberChild;
 	};
 
 	SLATE_BEGIN_ARGS(SBox)
diff --git a/Engine/Source/Runtime/Slate/Public/Widgets/Layout/SLinkedBox.cpp b/Engine/Source/Runtime/Slate/Public/Widgets/Layout/SLinkedBox.cpp
--- a/Engine/Source/Runtime/Slate/Public/Widgets/Layout/SLinkedBox.cpp
+++ b/Engine/Source/Runtime/Slate/Public/Widgets/Layout/SLinkedBox.cpp
@@ -101,7 +101,7 @@ FVector2D SLinkedBox::GetChildrensDesiredSize() const
 	EVisibility ChildVisibility = ChildSlot.GetWidget()->GetVisibility();
 	if ( ChildVisibility != EVisibility::Collapsed )
 	{
-		return ChildSlot.GetWidget()->GetDesiredSize() + ChildSlot.SlotPadding.Get().GetDesiredSize();
+		return ChildSlot.GetWidget()->GetDesiredSize() + ChildSlot.GetPadding().GetDesiredSize();
 	}
 	return FVector2D::ZeroVector;
 }
diff --git a/Engine/Source/Runtime/Slate/Public/Widgets/SViewport.h b/Engine/Source/Runtime/Slate/Public/Widgets/SViewport.h
--- a/Engine/Source/Runtime/Slate/Public/Widgets/SViewport.h
+++ b/Engine/Source/Runtime/Slate/Public/Widgets/SViewport.h
@@ -246,12 +246,6 @@ public:
 	virtual FNavigationReply OnNavigation( const FGeometry& MyGeometry, const FNavigationEvent& InNavigationEvent ) override;
 
 private:
-
-	FSimpleSlot& Decl_GetContent()
-	{
-		return ChildSlot;
-	}
-
 	// Viewports shouldn't show focus
 	virtual const FSlateBrush* GetFocusBrush() const override
 	{
diff --git a/Engine/Source/Runtime/Slate/Public/Widgets/Views/STableRow.h b/Engine/Source/Runtime/Slate/Public/Widgets/Views/STableRow.h
--- a/Engine/Source/Runtime/Slate/Public/Widgets/Views/STableRow.h
+++ b/Engine/Source/Runtime/Slate/Public/Widgets/Views/STableRow.h
@@ -229,7 +229,7 @@ public:
 		if ( InOwnerTableMode == ETableViewMode::List || InOwnerTableMode == ETableViewMode::Tile )
 		{
 			// -- Row is in a ListView or the user --
-			FSimpleSlot* InnerContentSlotNativePtr = nullptr;
+			FOneSimpleMemberChild* InnerContentSlotNativePtr = nullptr;
 
 			// We just need to hold on to this row's content.
 			this->ChildSlot
diff --git a/Engine/Source/Runtime/SlateCore/Private/Widgets/SCompoundWidget.cpp b/Engine/Source/Runtime/SlateCore/Private/Widgets/SCompoundWidget.cpp
--- a/Engine/Source/Runtime/SlateCore/Private/Widgets/SCompoundWidget.cpp
+++ b/Engine/Source/Runtime/SlateCore/Private/Widgets/SCompoundWidget.cpp
@@ -11,6 +11,7 @@ DECLARE_CYCLE_STAT(TEXT("Child Paint"), STAT_ChildPaint, STATGROUP_SlateVeryVerb
 SLATE_IMPLEMENT_WIDGET(SCompoundWidget)
 void SCompoundWidget::PrivateRegisterAttributes(FSlateAttributeInitializer& AttributeInitializer)
 {
+	AttributeInitializer.AddMemberAttribute("SlotPadding", STRUCT_OFFSET(PrivateThisType, ChildSlot) + FOneSimpleMemberChild::GetSlotPaddingAttributeOffset(), FSlateAttributeDescriptor::FInvalidateWidgetReasonAttribute{ EInvalidateWidgetReason::Layout });
 	SLATE_ADD_MEMBER_ATTRIBUTE_DEFINITION_WITH_NAME(AttributeInitializer, "ContentScale", ContentScaleAttribute, EInvalidateWidgetReason::Layout);
 	SLATE_ADD_MEMBER_ATTRIBUTE_DEFINITION_WITH_NAME(AttributeInitializer, "ColorAndOpacity", ColorAndOpacityAttribute, EInvalidateWidgetReason::Paint);
 	SLATE_ADD_MEMBER_ATTRIBUTE_DEFINITION_WITH_NAME(AttributeInitializer, "ForegroundColor", ForegroundColorAttribute, EInvalidateWidgetReason::Paint);
@@ -60,7 +61,7 @@ FVector2D SCompoundWidget::ComputeDesiredSize( float ) const
 	EVisibility ChildVisibility = ChildSlot.GetWidget()->GetVisibility();
 	if ( ChildVisibility != EVisibility::Collapsed )
 	{
-		return ChildSlot.GetWidget()->GetDesiredSize() + ChildSlot.SlotPadding.Get().GetDesiredSize();
+		return ChildSlot.GetWidget()->GetDesiredSize() + ChildSlot.GetPadding().GetDesiredSize();
 	}
 	
 	return FVector2D::ZeroVector;
@@ -77,7 +78,7 @@ FSlateColor SCompoundWidget::GetForegroundColor() const
 }
 
 SCompoundWidget::SCompoundWidget()
-	: ChildSlot(this)
+	: ChildSlot(*this)
 	, ContentScaleAttribute( *this, FVector2D(1.0f,1.0f) )
 	, ColorAndOpacityAttribute( *this, FLinearColor::White )
 	, ForegroundColorAttribute( *this, FSlateColor::UseForeground() )
diff --git a/Engine/Source/Runtime/SlateCore/Public/Layout/Children.h b/Engine/Source/Runtime/SlateCore/Public/Layout/Children.h
--- a/Engine/Source/Runtime/SlateCore/Public/Layout/Children.h
+++ b/Engine/Source/Runtime/SlateCore/Public/Layout/Children.h
@@ -295,13 +295,14 @@ public:
 	EVerticalAlignment VAlignment;
 };
 
+PRAGMA_DISABLE_DEPRECATION_WARNINGS
 template <typename MixedIntoType>
 class TSupportsContentPaddingMixin
 {
 public:
-	MixedIntoType& Padding( const TAttribute<FMargin> InPadding )
+	MixedIntoType& Padding( TAttribute<FMargin> InPadding )
 	{
-		SlotPadding = InPadding;
+		SlotPadding = MoveTemp(InPadding);
 		return *(static_cast<MixedIntoType*>(this));
 	}
 
@@ -323,8 +324,17 @@ public:
 		return *(static_cast<MixedIntoType*>(this));
 	}
 
+	void SetPadding(TAttribute<FMargin> InPadding)
+	{
+		SlotPadding = MoveTemp(InPadding);
+	}
+	const FMargin& GetPadding() const { return SlotPadding.Get(); }
+
+public:
+	UE_DEPRECATED(5.0, "Direct access to SlotPadding is now deprecated. Use the getter.")
 	TAttribute< FMargin > SlotPadding;
 };
+PRAGMA_ENABLE_DEPRECATION_WARNINGS
 
 /** A slot that support alignment of content and padding */
 class SLATECORE_API FSimpleSlot : public TSupportsOneChildMixin<FSimpleSlot>, public TSupportsContentAlignmentMixin<FSimpleSlot>, public TSupportsContentPaddingMixin<FSimpleSlot>
@@ -808,45 +818,60 @@ public:
 	FOneSimpleMemberChild(WidgetType& InParent)
 		: TSupportsOneChildMixin<FOneSimpleMemberChild>(&InParent)
 		, TSupportsContentAlignmentMixin<FOneSimpleMemberChild>(HAlign_Fill, VAlign_Fill)
-		, SlotPadding(InParent)
+		, SlotPaddingAttribute(InParent)
 	{
 	}
 
+	//~ TSlateAttribute cannot be copied
+	FOneSimpleMemberChild(const FOneSimpleMemberChild&) = delete;
+	FOneSimpleMemberChild& operator=(const FOneSimpleMemberChild&) = delete;
+
 	FOneSimpleMemberChild& Padding(TAttribute<FMargin> InPadding)
 	{
-		SlotPadding.Assign(*FChildren::GetOwner(), MoveTemp(InPadding));
+		SlotPaddingAttribute.Assign(*FChildren::GetOwner(), MoveTemp(InPadding));
 		return *this;
 	}
 
 	FOneSimpleMemberChild& Padding(float Uniform)
 	{
-		SlotPadding.Set(*FChildren::GetOwner(), FMargin(Uniform));
+		SlotPaddingAttribute.Set(*FChildren::GetOwner(), FMargin(Uniform));
 		return *this;
 	}
 
 	FOneSimpleMemberChild& Padding(float Horizontal, float Vertical)
 	{
-		SlotPadding.Set(*FChildren::GetOwner(), FMargin(Horizontal, Vertical));
+		SlotPaddingAttribute.Set(*FChildren::GetOwner(), FMargin(Horizontal, Vertical));
 		return *this;
 	}
 
 	FOneSimpleMemberChild& Padding(float Left, float Top, float Right, float Bottom)
 	{
-		SlotPadding.Set(*FChildren::GetOwner(), FMargin(Left, Top, Right, Bottom));
+		SlotPaddingAttribute.Set(*FChildren::GetOwner(), FMargin(Left, Top, Right, Bottom));
 		return *this;
 	}
 
-	const FMargin& GetSlotPadding() const { return SlotPadding.Get(); }
+	void SetPadding(TAttribute<FMargin> InPadding)
+	{
+		SlotPaddingAttribute.Assign(*FChildren::GetOwner(), MoveTemp(InPadding));
+	}
+	const FMargin& GetPadding() const { return SlotPaddingAttribute.Get(); }
+
+#if WITH_EDITORONLY_DATA
+	UE_DEPRECATED(5.0, "Direct access to SlotPadding is now deprecated. Use the setter or getter.")
+	FSlateDeprecatedTAttribute<FMargin> SlotPadding;
+#endif
 
 public:
-	using SlotPaddingAttributeType = SlateAttributePrivate::TSlateMemberAttribute<FMargin, TSlateAttributeInvalidationReason<EInvalidateWidgetReason::Layout>, TSlateAttributeComparePredicate<>>;
+	using SlotPaddingAttributeType = SlateAttributePrivate::TSlateMemberAttribute<FMargin, ::SlateAttributePrivate::FSlateAttributeNoInvalidationReason, TSlateAttributeComparePredicate<>>;
 	using SlotPaddingAttributeRefType = SlateAttributePrivate::TSlateMemberAttributeRef<SlotPaddingAttributeType>;
 
+	static SIZE_T GetSlotPaddingAttributeOffset() { return STRUCT_OFFSET(FOneSimpleMemberChild, SlotPaddingAttribute); }
+
 	template<typename WidgetType, typename V = typename std::enable_if<std::is_base_of<SWidget, WidgetType>::value>::type>
-	SlotPaddingAttributeRefType GetSlotPaddingAttribute() const { return SlotPaddingAttributeRefType(*(static_cast<const WidgetType*>(FChildren::GetOwner())), SlotPadding); }
+	SlotPaddingAttributeRefType GetSlotPaddingAttribute() const { return SlotPaddingAttributeRefType(*(static_cast<const WidgetType*>(FChildren::GetOwner())), SlotPaddingAttribute); }
 
 private:
-	SlotPaddingAttributeType SlotPadding;
+	SlotPaddingAttributeType SlotPaddingAttribute;
 };
 
 
diff --git a/Engine/Source/Runtime/SlateCore/Public/Layout/LayoutUtils.h b/Engine/Source/Runtime/SlateCore/Public/Layout/LayoutUtils.h
--- a/Engine/Source/Runtime/SlateCore/Public/Layout/LayoutUtils.h
+++ b/Engine/Source/Runtime/SlateCore/Public/Layout/LayoutUtils.h
@@ -179,7 +179,7 @@ static void ArrangeSingleChild(EFlowDirection InFlowDirection, const FGeometry&
 	if ( ArrangedChildren.Accepts(ChildVisibility) )
 	{
 		const FVector2D ThisContentScale = ContentScale.Get();
-		const FMargin SlotPadding(LayoutPaddingWithFlow(InFlowDirection, ChildSlot.SlotPadding.Get()));
+		const FMargin SlotPadding(LayoutPaddingWithFlow(InFlowDirection, ChildSlot.GetPadding()));
 		const AlignmentArrangeResult XResult = AlignChild<Orient_Horizontal>(InFlowDirection, AllottedGeometry.GetLocalSize().X, ChildSlot, SlotPadding, ThisContentScale.X);
 		const AlignmentArrangeResult YResult = AlignChild<Orient_Vertical>(AllottedGeometry.GetLocalSize().Y, ChildSlot, SlotPadding, ThisContentScale.Y);
 
@@ -204,7 +204,7 @@ static void ArrangeSingleChild(EFlowDirection InFlowDirection, const FGeometry&
 	if (ArrangedChildren.Accepts(ChildVisibility))
 	{
 		const FVector2D ThisContentScale = ContentScale;
-		const FMargin SlotPadding(LayoutPaddingWithFlow(InFlowDirection, ChildSlot.SlotPadding.Get()));
+		const FMargin SlotPadding(LayoutPaddingWithFlow(InFlowDirection, ChildSlot.GetPadding()));
 		const AlignmentArrangeResult XResult = AlignChild<Orient_Horizontal>(InFlowDirection, AllottedGeometry.GetLocalSize().X, ChildSlot, SlotPadding, ThisContentScale.X);
 		const AlignmentArrangeResult YResult = AlignChild<Orient_Vertical>(AllottedGeometry.GetLocalSize().Y, ChildSlot, SlotPadding, ThisContentScale.Y);
 
diff --git a/Engine/Source/Runtime/SlateCore/Public/Widgets/SCompoundWidget.h b/Engine/Source/Runtime/SlateCore/Public/Widgets/SCompoundWidget.h
--- a/Engine/Source/Runtime/SlateCore/Public/Widgets/SCompoundWidget.h
+++ b/Engine/Source/Runtime/SlateCore/Public/Widgets/SCompoundWidget.h
@@ -103,7 +103,7 @@ protected:
 	SCompoundWidget();
 
 	/** The slot that contains this widget's descendants.*/
-	FSimpleSlot ChildSlot;
+	FOneSimpleMemberChild ChildSlot;
 
 #if WITH_EDITORONLY_DATA
 	UE_DEPRECATED(5.0, "Direct access to ContentScale is now deprecated. Use the setter or getter.")
