Slate: Add an option to update the attributes and cache the invalidation for later. Update the visibility attribute before doing layout out of the regular SWidget paint path.

diff --git a/Engine/Source/Developer/AutomationDriver/Private/Locators/SlateWidgetLocatorByPath.cpp b/Engine/Source/Developer/AutomationDriver/Private/Locators/SlateWidgetLocatorByPath.cpp
--- a/Engine/Source/Developer/AutomationDriver/Private/Locators/SlateWidgetLocatorByPath.cpp
+++ b/Engine/Source/Developer/AutomationDriver/Private/Locators/SlateWidgetLocatorByPath.cpp
@@ -203,8 +203,9 @@ public:
 			const FArrangedWidget& Candidate = State.Path.Widgets.Last();
 
 			const bool bAllow3DWidgets = true;
+			const bool bUpdateVisibilityAttributes = true;
 			FArrangedChildren ArrangedChildren(VisibilityFilter, bAllow3DWidgets);
-			Candidate.Widget->ArrangeChildren(Candidate.Geometry, ArrangedChildren);
+			Candidate.Widget->ArrangeChildren(Candidate.Geometry, ArrangedChildren, bUpdateVisibilityAttributes);
 
 			for (int32 ChildIndex = 0; ChildIndex < ArrangedChildren.Num(); ++ChildIndex)
 			{
diff --git a/Engine/Source/Developer/AutomationDriver/Private/Locators/SlateWidgetLocatorByUniqueTag.cpp b/Engine/Source/Developer/AutomationDriver/Private/Locators/SlateWidgetLocatorByUniqueTag.cpp
--- a/Engine/Source/Developer/AutomationDriver/Private/Locators/SlateWidgetLocatorByUniqueTag.cpp
+++ b/Engine/Source/Developer/AutomationDriver/Private/Locators/SlateWidgetLocatorByUniqueTag.cpp
@@ -54,8 +54,9 @@ public:
 			const FArrangedWidget& Candidate = State.Path.Widgets.Last();
 
 			const bool bAllow3DWidgets = true;
+			const bool bUpdateVisibilityAttributes = true;
 			FArrangedChildren ArrangedChildren(VisibilityFilter, bAllow3DWidgets);
-			Candidate.Widget->ArrangeChildren(Candidate.Geometry, ArrangedChildren);
+			Candidate.Widget->ArrangeChildren(Candidate.Geometry, ArrangedChildren, bUpdateVisibilityAttributes);
 
 			for (int32 ChildIndex = 0; ChildIndex < ArrangedChildren.Num(); ++ChildIndex)
 			{
diff --git a/Engine/Source/Editor/WorldBrowser/Private/Tiles/SWorldComposition.cpp b/Engine/Source/Editor/WorldBrowser/Private/Tiles/SWorldComposition.cpp
--- a/Engine/Source/Editor/WorldBrowser/Private/Tiles/SWorldComposition.cpp
+++ b/Engine/Source/Editor/WorldBrowser/Private/Tiles/SWorldComposition.cpp
@@ -595,7 +595,8 @@ protected:
 		}
 
 		FArrangedChildren ArrangedChildren(EVisibility::Visible);
-		ArrangeChildren(MyGeometry, ArrangedChildren);
+		const bool bUpdateVisibilityAttributes = true;
+		ArrangeChildren(MyGeometry, ArrangedChildren, bUpdateVisibilityAttributes);
 
 		const int32 NodeUnderMouseIndex = SWidget::FindChildUnderMouse( ArrangedChildren, MouseEvent );
 		if (NodeUnderMouseIndex != INDEX_NONE)
diff --git a/Engine/Source/Runtime/SlateCore/Private/Application/SlateWindowHelper.cpp b/Engine/Source/Runtime/SlateCore/Private/Application/SlateWindowHelper.cpp
--- a/Engine/Source/Runtime/SlateCore/Private/Application/SlateWindowHelper.cpp
+++ b/Engine/Source/Runtime/SlateCore/Private/Application/SlateWindowHelper.cpp
@@ -4,6 +4,7 @@
 #include "Layout/ArrangedChildren.h"
 #include "SlateGlobals.h"
 #include "Layout/WidgetPath.h"
+#include "Types/SlateAttributeMetaData.h"
 
 
 DECLARE_CYCLE_STAT( TEXT("FindPathToWidget"), STAT_FindPathToWidget, STATGROUP_Slate );
@@ -101,6 +102,8 @@ bool FSlateWindowHelper::FindPathToWidget( const TArray<TSharedRef<SWindow>>& Wi
 		OutWidgetPath.Widgets.SetFilter(VisibilityFilter);
 		while (true)
 		{
+			// Update the Widgets visibility before getting the ArrangeChildren
+			FSlateAttributeMetaData::UpdateCollapsedAttributes(*CurWidget.Get(), FSlateAttributeMetaData::EInvalidationPermission::DelayInvalidation);
 			EVisibility CurWidgetVisibility = CurWidget->GetVisibility();
 			if (OutWidgetPath.Widgets.Accepts(CurWidgetVisibility))
 			{
diff --git a/Engine/Source/Runtime/SlateCore/Private/FastUpdate/SlateInvalidationRoot.cpp b/Engine/Source/Runtime/SlateCore/Private/FastUpdate/SlateInvalidationRoot.cpp
--- a/Engine/Source/Runtime/SlateCore/Private/FastUpdate/SlateInvalidationRoot.cpp
+++ b/Engine/Source/Runtime/SlateCore/Private/FastUpdate/SlateInvalidationRoot.cpp
@@ -691,14 +691,14 @@ void FSlateInvalidationRoot::ProcessPreUpdate()
 						if (!InvalidationWidget.Visibility.IsCollapseIndirectly())
 						{
 							// if my parent is not collapse, then update my visible state
-							FSlateAttributeMetaData::UpdateCollapsedAttributes(*WidgetPtr);
+							FSlateAttributeMetaData::UpdateCollapsedAttributes(*WidgetPtr, FSlateAttributeMetaData::EInvalidationPermission::AllowInvalidation);
 							if (!InvalidationWidget.Visibility.IsCollapsed())
 							{
 #if UE_SLATE_WITH_INVALIDATIONWIDGETLIST_DEBUGGING
 								ensureAlwaysMsgf(!GSlateInvalidationRootVerifySlateAttribute || InvalidationWidget.bDebug_AttributeUpdated == false, TEXT("Attribute should only be updated once per frame."));
 								InvalidationWidget.bDebug_AttributeUpdated = true;
 #endif
-								FSlateAttributeMetaData::UpdateExpandedAttributes(*WidgetPtr);
+								FSlateAttributeMetaData::UpdateExpandedAttributes(*WidgetPtr, FSlateAttributeMetaData::EInvalidationPermission::AllowInvalidation);
 							}
 						}
 					}
diff --git a/Engine/Source/Runtime/SlateCore/Private/FastUpdate/SlateInvalidationWidgetList.cpp b/Engine/Source/Runtime/SlateCore/Private/FastUpdate/SlateInvalidationWidgetList.cpp
--- a/Engine/Source/Runtime/SlateCore/Private/FastUpdate/SlateInvalidationWidgetList.cpp
+++ b/Engine/Source/Runtime/SlateCore/Private/FastUpdate/SlateInvalidationWidgetList.cpp
@@ -276,12 +276,12 @@ FSlateInvalidationWidgetIndex FSlateInvalidationWidgetList::Internal_BuildWidget
 		{
 			if (!NewVisibility.IsCollapseIndirectly())
 			{
-				FSlateAttributeMetaData::UpdateCollapsedAttributes(Widget.Get(), false);
+				FSlateAttributeMetaData::UpdateCollapsedAttributes(Widget.Get(), FSlateAttributeMetaData::EInvalidationPermission::DenyAndClearDelayedInvalidation);
 				NewVisibility.SetVisibility(ParentVisibility, Widget->GetVisibility());
 
 				if (!NewVisibility.IsCollapsed())
 				{
-					FSlateAttributeMetaData::UpdateExpandedAttributes(Widget.Get(), false);
+					FSlateAttributeMetaData::UpdateExpandedAttributes(Widget.Get(), FSlateAttributeMetaData::EInvalidationPermission::DenyAndClearDelayedInvalidation);
 					NewVisibility.SetVisibility(ParentVisibility, Widget->GetVisibility());
 					bUpdateSlateAttribute = true;
 				}
diff --git a/Engine/Source/Runtime/SlateCore/Private/Layout/WidgetPath.cpp b/Engine/Source/Runtime/SlateCore/Private/Layout/WidgetPath.cpp
--- a/Engine/Source/Runtime/SlateCore/Private/Layout/WidgetPath.cpp
+++ b/Engine/Source/Runtime/SlateCore/Private/Layout/WidgetPath.cpp
@@ -2,6 +2,7 @@
 
 #include "Layout/WidgetPath.h"
 #include "SlateGlobals.h"
+#include "Types/SlateAttributeMetaData.h"
 
 DECLARE_CYCLE_STAT(TEXT("Weak-To-Strong WidgetPath"), STAT_WeakToStrong_WidgetPath, STATGROUP_Slate);
 
@@ -191,7 +192,8 @@ bool FWidgetPath::MoveFocus(int32 PathLevel, EUINavigation NavigationType, bool
 	
 		// Arrange the children so we can iterate through them regardless of widget type.
 		FArrangedChildren ArrangedChildren(EVisibility::Visible);
-		Widgets[PathLevel].Widget->ArrangeChildren( Widgets[PathLevel].Geometry, ArrangedChildren );
+		const bool bUpdateVisibilityAttributes = true;
+		Widgets[PathLevel].Widget->ArrangeChildren( Widgets[PathLevel].Geometry, ArrangedChildren,  bUpdateVisibilityAttributes);
 
 		// Don't continue if there are no children in the widget.
 		if (ArrangedChildren.Num() > 0)
@@ -308,23 +310,27 @@ FWeakWidgetPath::EPathResolutionResult::Result FWeakWidgetPath::ToWidgetPath( FW
 				TSharedPtr<SWidget> CurWidget = WidgetPtrs[WidgetIndex];
 
 				bool bFoundChild = false;
-				if (CurWidget.IsValid() && EVisibility::DoesVisibilityPassFilter(CurWidget->GetVisibility(), VisibilityFilter))
+				if (CurWidget.IsValid())
 				{
-					TSharedRef<SWidget> CurWidgetRef = CurWidget.ToSharedRef();
-					const TSharedPtr<SWidget>& ChildWidgetPtr = WidgetPtrs[WidgetIndex + 1];
-
-					if (ChildWidgetPtr.IsValid() && ChildWidgetPtr->GetParentWidget() == CurWidgetRef)
+					FSlateAttributeMetaData::UpdateCollapsedAttributes(*CurWidget.Get(), FSlateAttributeMetaData::EInvalidationPermission::DelayInvalidation);
+					if (EVisibility::DoesVisibilityPassFilter(CurWidget->GetVisibility(), VisibilityFilter))
 					{
-						if (PointerEvent && !VirtualPointerPos.IsValid())
+						TSharedRef<SWidget> CurWidgetRef = CurWidget.ToSharedRef();
+						const TSharedPtr<SWidget>& ChildWidgetPtr = WidgetPtrs[WidgetIndex + 1];
+
+						if (ChildWidgetPtr.IsValid() && ChildWidgetPtr->GetParentWidget() == CurWidgetRef)
 						{
-							VirtualPointerPos = CurWidget->TranslateMouseCoordinateForCustomHitTestChild(ChildWidgetPtr.ToSharedRef(), ParentGeometry, PointerEvent->GetScreenSpacePosition(), PointerEvent->GetLastScreenSpacePosition());
-						}
+							if (PointerEvent && !VirtualPointerPos.IsValid())
+							{
+								VirtualPointerPos = CurWidget->TranslateMouseCoordinateForCustomHitTestChild(ChildWidgetPtr.ToSharedRef(), ParentGeometry, PointerEvent->GetScreenSpacePosition(), PointerEvent->GetLastScreenSpacePosition());
+							}
 
-						bFoundChild = true;
-						// Remember the widget, the associated geometry, and the pointer position in a transformed space.
-						PathWithGeometries.Add(FWidgetAndPointer(FArrangedWidget(ChildWidgetPtr.ToSharedRef(), ChildWidgetPtr->GetCachedGeometry()), VirtualPointerPos));
-						// The next child in the vertical slice will be arranged with respect to its parent's geometry.
-						ParentGeometry = CurWidgetRef->GetCachedGeometry();
+							bFoundChild = true;
+							// Remember the widget, the associated geometry, and the pointer position in a transformed space.
+							PathWithGeometries.Add(FWidgetAndPointer(FArrangedWidget(ChildWidgetPtr.ToSharedRef(), ChildWidgetPtr->GetCachedGeometry()), VirtualPointerPos));
+							// The next child in the vertical slice will be arranged with respect to its parent's geometry.
+							ParentGeometry = CurWidgetRef->GetCachedGeometry();
+						}
 					}
 				}
 
@@ -370,7 +376,8 @@ FWeakWidgetPath::EPathResolutionResult::Result FWeakWidgetPath::ToWidgetPath( FW
 				{
 					// Arrange the widget's children to find their geometries.
 					ArrangedChildren.Empty();
-					CurWidget->ArrangeChildren(ParentGeometry, ArrangedChildren);
+					const bool bUpdateVisibilityAttributes = true;
+					CurWidget->ArrangeChildren(ParentGeometry, ArrangedChildren, bUpdateVisibilityAttributes);
 
 					// Find the next widget in the path among the arranged children.
 					for (int32 SearchIndex = 0; !bFoundChild && SearchIndex < ArrangedChildren.Num(); ++SearchIndex)
diff --git a/Engine/Source/Runtime/SlateCore/Private/Types/SlateAttributeMetaData.cpp b/Engine/Source/Runtime/SlateCore/Private/Types/SlateAttributeMetaData.cpp
--- a/Engine/Source/Runtime/SlateCore/Private/Types/SlateAttributeMetaData.cpp
+++ b/Engine/Source/Runtime/SlateCore/Private/Types/SlateAttributeMetaData.cpp
@@ -3,6 +3,7 @@
 #include "Types/SlateAttributeMetaData.h"
 
 #include "Algo/BinarySearch.h"
+#include "Layout/Children.h"
 #include "Types/ReflectionMetadata.h"
 #include "Widgets/SWidget.h"
 
@@ -29,14 +30,17 @@ FSlateAttributeMetaData* FSlateAttributeMetaData::FindMetaData(const SWidget& Ow
 		check(SlateMetaData->IsOfType<FSlateAttributeMetaData>());
 		return &(StaticCastSharedRef<FSlateAttributeMetaData>(SlateMetaData).Get());
 	}
+#if WITH_SLATE_DEBUGGING
 	else if (OwningWidget.MetaData.Num() > 0)
 	{
 		TSharedRef<ISlateMetaData> SlateMetaData = OwningWidget.MetaData[0];
 		if (SlateMetaData->IsOfType<FSlateAttributeMetaData>())
 		{
+			ensureMsgf(false, TEXT("bHasRegisteredSlateAttribute should be set on the SWidget '%s'"), *FReflectionMetaData::GetWidgetDebugInfo(OwningWidget));
 			return &(StaticCastSharedRef<FSlateAttributeMetaData>(SlateMetaData).Get());
 		}
 	}
+#endif
 	return nullptr;
 }
 
@@ -225,7 +229,8 @@ void FSlateAttributeMetaData::InvalidateWidget(SWidget& OwningWidget, const FSla
 			{
 				const FGetterItem::FInvalidationDetail Detail = GetterItem.GetInvalidationDetail(OwningWidget, Reason);
 				OnInvalidationCallback = Detail.Get<0>();
-				Reason = Detail.Get<1>();
+				Reason = Detail.Get<1>() | AttributeMetaData->CachedInvalidationReason;
+				AttributeMetaData->CachedInvalidationReason = EInvalidateWidgetReason::None;
 			}
 
 			// The dependency attribute need to be updated in the update loop (note that it may not be registered yet)
@@ -242,8 +247,9 @@ void FSlateAttributeMetaData::InvalidateWidget(SWidget& OwningWidget, const FSla
 			const FSlateAttributeDescriptor::OffsetType Offset = Private::FindOffet(OwningWidget, Attribute);
 			if (FSlateAttributeDescriptor::FAttribute const* FoundAttribute = AttributeDescriptor.FindMemberAttribute(Offset))
 			{
-				Reason = FoundAttribute->InvalidationReason.Get(OwningWidget);
 				OnInvalidationCallback = &FoundAttribute->OnInvalidation;
+				Reason = FoundAttribute->InvalidationReason.Get(OwningWidget) | AttributeMetaData->CachedInvalidationReason;
+				AttributeMetaData->CachedInvalidationReason = EInvalidateWidgetReason::None;
 
 				if (FoundAttribute->bIsADependencyForSomeoneElse)
 				{
@@ -282,61 +288,56 @@ void FSlateAttributeMetaData::InvalidateWidget(SWidget& OwningWidget, const FSla
 }
 
 
-void FSlateAttributeMetaData::UpdateAttributes(SWidget& OwningWidget)
-{
-	UpdateAttributes(OwningWidget, OwningWidget.IsConstructionCompleted());
-}
-
-
-void FSlateAttributeMetaData::UpdateCollapsedAttributes(SWidget& OwningWidget)
-{
-	UpdateCollapsedAttributes(OwningWidget, OwningWidget.IsConstructionCompleted());	
-}
-
-
-void FSlateAttributeMetaData::UpdateExpandedAttributes(SWidget& OwningWidget)
-{
-	UpdateExpandedAttributes(OwningWidget, OwningWidget.IsConstructionCompleted());
-}
-
-
-void FSlateAttributeMetaData::UpdateAttributes(SWidget& OwningWidget, bool bAllowInvalidation)
+void FSlateAttributeMetaData::UpdateAttributes(SWidget& OwningWidget, EInvalidationPermission InvalidationStyle)
 {
 	if (FSlateAttributeMetaData* AttributeMetaData = FSlateAttributeMetaData::FindMetaData(OwningWidget))
 	{
-		AttributeMetaData->UpdateAttributesImpl(OwningWidget, EUpdateType::All, bAllowInvalidation);
+		AttributeMetaData->UpdateAttributesImpl(OwningWidget, EUpdateType::All, InvalidationStyle);
 	}
 }
 
 
-void FSlateAttributeMetaData::UpdateCollapsedAttributes(SWidget& OwningWidget, bool bAllowInvalidation)
+void FSlateAttributeMetaData::UpdateCollapsedAttributes(SWidget& OwningWidget, EInvalidationPermission InvalidationStyle)
 {
 	if (FSlateAttributeMetaData* AttributeMetaData = FSlateAttributeMetaData::FindMetaData(OwningWidget))
 	{
 		// Does it have any collapsed attributes
 		if (AttributeMetaData->CollaspedAttributeCounter > 0)
 		{
-			AttributeMetaData->UpdateAttributesImpl(OwningWidget, EUpdateType::Collapsed, bAllowInvalidation);
+			AttributeMetaData->UpdateAttributesImpl(OwningWidget, EUpdateType::Collapsed, InvalidationStyle);
 		}
 	}
 }
 
 
-void FSlateAttributeMetaData::UpdateExpandedAttributes(SWidget& OwningWidget, bool bAllowInvalidation)
+void FSlateAttributeMetaData::UpdateExpandedAttributes(SWidget& OwningWidget, EInvalidationPermission InvalidationStyle)
 {
 	if (FSlateAttributeMetaData* AttributeMetaData = FSlateAttributeMetaData::FindMetaData(OwningWidget))
 	{
-		AttributeMetaData->UpdateAttributesImpl(OwningWidget, EUpdateType::Expanded, bAllowInvalidation);
+		AttributeMetaData->UpdateAttributesImpl(OwningWidget, EUpdateType::Expanded, InvalidationStyle);
 	}
 }
 
 
+void FSlateAttributeMetaData::UpdateChildrenCollapsedAttributes(SWidget& OwningWidget, EInvalidationPermission InvalidationStyle)
+{
+	FChildren* Children = OwningWidget.GetChildren();
+	const int32 NumChildren = Children->Num();
+	for (int32 ChildIndex = 0; ChildIndex < NumChildren; ++ChildIndex)
+	{
+		const TSharedRef<SWidget>& Child = Children->GetChildAt(ChildIndex);
+		UpdateCollapsedAttributes(Child.Get(), InvalidationStyle);
+	}
+}
+
 
-void FSlateAttributeMetaData::UpdateAttributesImpl(SWidget& OwningWidget, EUpdateType UpdateType, bool bAllowInvalidation)
+void FSlateAttributeMetaData::UpdateAttributesImpl(SWidget& OwningWidget, EUpdateType UpdateType, EInvalidationPermission InvalidationStyle)
 {
 	// List of all the callback we need to execute after the invalidation
 	TArray<const FSlateAttributeDescriptor::FInvalidationDelegate*, TInlineAllocator<8>> AllInvalidationDelegates;
 
+	bool bInvalidateIfNeeded = (InvalidationStyle == EInvalidationPermission::AllowInvalidation) || (InvalidationStyle == EInvalidationPermission::AllowInvalidationIfConstructed && OwningWidget.IsConstructionCompleted());
+	bool bAllowInvalidation = bInvalidateIfNeeded || InvalidationStyle == EInvalidationPermission::DelayInvalidation;
 	EInvalidateWidgetReason InvalidationReason = EInvalidateWidgetReason::None;
 	for (int32 Index = 0; Index < Attributes.Num(); ++Index)
 	{
@@ -401,13 +402,22 @@ void FSlateAttributeMetaData::UpdateAttributesImpl(SWidget& OwningWidget, EUpdat
 	}
 	bHasUpdatedManuallyFlagToReset = false;
 
-	if (bAllowInvalidation)
+	if (bInvalidateIfNeeded)
 	{
-		OwningWidget.Invalidate(InvalidationReason);
+		OwningWidget.Invalidate(InvalidationReason | CachedInvalidationReason);
 		for (const FSlateAttributeDescriptor::FInvalidationDelegate* OnInvalidation : AllInvalidationDelegates)
 		{
 			OnInvalidation->ExecuteIfBound(OwningWidget);
 		}
+		CachedInvalidationReason = EInvalidateWidgetReason::None;
+	}
+	else if (InvalidationStyle == EInvalidationPermission::DelayInvalidation)
+	{
+		CachedInvalidationReason |= InvalidationReason;
+	}
+	else if (InvalidationStyle == EInvalidationPermission::DenyAndClearDelayedInvalidation)
+	{
+		CachedInvalidationReason = EInvalidateWidgetReason::None;
 	}
 }
 
@@ -428,11 +438,12 @@ void FSlateAttributeMetaData::UpdateAttribute(SWidget& OwningWidget, FSlateAttri
 				if (OwningWidget.IsConstructionCompleted())
 				{
 					const FGetterItem::FInvalidationDetail Detail = GetterItem.GetInvalidationDetail(OwningWidget, Result.InvalidationReason);
-					OwningWidget.Invalidate(Detail.Get<1>());
+					OwningWidget.Invalidate(Detail.Get<1>() | AttributeMetaData->CachedInvalidationReason);
 					if (Detail.Get<0>())
 					{
 						Detail.Get<0>()->ExecuteIfBound(OwningWidget);
 					}
+					AttributeMetaData->CachedInvalidationReason = EInvalidateWidgetReason::None;
 				}
 
 				// The dependency attribute need to be updated in the update loop (note that it may not be registered yet)
diff --git a/Engine/Source/Runtime/SlateCore/Private/Widgets/SWidget.cpp b/Engine/Source/Runtime/SlateCore/Private/Widgets/SWidget.cpp
--- a/Engine/Source/Runtime/SlateCore/Private/Widgets/SWidget.cpp
+++ b/Engine/Source/Runtime/SlateCore/Private/Widgets/SWidget.cpp
@@ -664,7 +664,7 @@ void SWidget::SlatePrepass(float InLayoutScaleMultiplier)
 	{
 		if (HasRegisteredSlateAttribute() && IsAttributesUpdatesEnabled() && !GSlateIsOnFastProcessInvalidation)
 		{
-			FSlateAttributeMetaData::UpdateAttributes(*this);
+			FSlateAttributeMetaData::UpdateAttributes(*this, FSlateAttributeMetaData::EInvalidationPermission::AllowInvalidationIfConstructed);
 		}
 		Prepass_Internal(InLayoutScaleMultiplier);
 	}
@@ -1605,11 +1605,19 @@ float SWidget::GetRelativeLayoutScale(int32 ChildIndex, float LayoutScaleMultipl
 	return 1.0f;
 }
 
-void SWidget::ArrangeChildren(const FGeometry& AllottedGeometry, FArrangedChildren& ArrangedChildren) const
+void SWidget::ArrangeChildren(const FGeometry& AllottedGeometry, FArrangedChildren& ArrangedChildren, bool bUpdateAttributes) const
 {
 #if WITH_VERY_VERBOSE_SLATE_STATS
 	SCOPED_NAMED_EVENT(SWidget_ArrangeChildren, FColor::Black);
 #endif
+
+	if (bUpdateAttributes)
+	{
+		// Update the Widgets visibility before getting the ArrangeChildren
+		//const-casting for TSlateAttribute has the same behavior as previously with TAttribute. The const was hidden from the user.
+		FSlateAttributeMetaData::UpdateChildrenCollapsedAttributes(const_cast<SWidget&>(*this), FSlateAttributeMetaData::EInvalidationPermission::DelayInvalidation);
+	}
+
 	OnArrangeChildren(AllottedGeometry, ArrangedChildren);
 }
 
@@ -1643,7 +1651,7 @@ void SWidget::Prepass_Internal(float InLayoutScaleMultiplier)
 void SWidget::Prepass_ChildLoop(float InLayoutScaleMultiplier, FChildren* MyChildren)
 {
 	const int32 NumChildren = MyChildren->Num();
-	for (int32 ChildIndex = 0; ChildIndex < MyChildren->Num(); ++ChildIndex)
+	for (int32 ChildIndex = 0; ChildIndex < NumChildren; ++ChildIndex)
 	{
 		const float ChildLayoutScaleMultiplier = bHasRelativeLayoutScale
 			? InLayoutScaleMultiplier * GetRelativeLayoutScale(ChildIndex, InLayoutScaleMultiplier)
@@ -1654,14 +1662,14 @@ void SWidget::Prepass_ChildLoop(float InLayoutScaleMultiplier, FChildren* MyChil
 		const bool bUpdateAttributes = Child->HasRegisteredSlateAttribute() && Child->IsAttributesUpdatesEnabled() && !GSlateIsOnFastProcessInvalidation;
 		if (bUpdateAttributes)
 		{
-			FSlateAttributeMetaData::UpdateCollapsedAttributes(Child.Get());
+			FSlateAttributeMetaData::UpdateCollapsedAttributes(Child.Get(), FSlateAttributeMetaData::EInvalidationPermission::AllowInvalidationIfConstructed);
 		}
 
 		if (Child->GetVisibility() != EVisibility::Collapsed)
 		{
 			if (bUpdateAttributes)
 			{
-				FSlateAttributeMetaData::UpdateExpandedAttributes(Child.Get());
+				FSlateAttributeMetaData::UpdateExpandedAttributes(Child.Get(), FSlateAttributeMetaData::EInvalidationPermission::AllowInvalidationIfConstructed);
 			}
 			// Recur: Descend down the widget tree.
 			Child->Prepass_Internal(ChildLayoutScaleMultiplier);
diff --git a/Engine/Source/Runtime/SlateCore/Public/Layout/WidgetPath.h b/Engine/Source/Runtime/SlateCore/Public/Layout/WidgetPath.h
--- a/Engine/Source/Runtime/SlateCore/Public/Layout/WidgetPath.h
+++ b/Engine/Source/Runtime/SlateCore/Public/Layout/WidgetPath.h
@@ -8,7 +8,6 @@
 #include "Layout/ArrangedWidget.h"
 #include "Widgets/SWidget.h"
 #include "Widgets/SWindow.h"
-#include "Layout/ArrangedChildren.h"
 
 /** Matches widgets against InWidget */
 struct FWidgetMatcher
diff --git a/Engine/Source/Runtime/SlateCore/Public/Layout/WidgetPath.inl b/Engine/Source/Runtime/SlateCore/Public/Layout/WidgetPath.inl
--- a/Engine/Source/Runtime/SlateCore/Public/Layout/WidgetPath.inl
+++ b/Engine/Source/Runtime/SlateCore/Public/Layout/WidgetPath.inl
@@ -22,8 +22,9 @@ template<typename MatchRuleType>
 bool FWidgetPath::SearchForWidgetRecursively( const MatchRuleType& MatchRule, const FArrangedWidget& InCandidate, FArrangedChildren& OutReversedPath, EVisibility VisibilityFilter )
 {
 	const bool bAllow3DWidgets = true;
+	const bool bUpdateVisibilityAttributes = true;
 	FArrangedChildren ArrangedChildren(VisibilityFilter, bAllow3DWidgets);
-	InCandidate.Widget->ArrangeChildren( InCandidate.Geometry, ArrangedChildren );
+	InCandidate.Widget->ArrangeChildren( InCandidate.Geometry, ArrangedChildren, bUpdateVisibilityAttributes );
 
 	for( int32 ChildIndex = 0; ChildIndex < ArrangedChildren.Num(); ++ChildIndex )
 	{
@@ -48,8 +49,9 @@ template<typename MatchRuleType>
 bool FWidgetPath::SearchForWidgetRecursively_Reverse( const MatchRuleType& MatchRule, const FArrangedWidget& InCandidate, FArrangedChildren& OutReversedPath, EVisibility VisibilityFilter )
 {
 	const bool bAllow3DWidgets = true;
+	const bool bUpdateVisibilityAttributes = true;
 	FArrangedChildren ArrangedChildren(VisibilityFilter, bAllow3DWidgets);
-	InCandidate.Widget->ArrangeChildren( InCandidate.Geometry, ArrangedChildren );
+	InCandidate.Widget->ArrangeChildren( InCandidate.Geometry, ArrangedChildren, bUpdateVisibilityAttributes );
 
 	for( int32 ChildIndex = ArrangedChildren.Num()-1; ChildIndex >=0 ; --ChildIndex )
 	{
diff --git a/Engine/Source/Runtime/SlateCore/Public/Types/SlateAttributeMetaData.h b/Engine/Source/Runtime/SlateCore/Public/Types/SlateAttributeMetaData.h
--- a/Engine/Source/Runtime/SlateCore/Public/Types/SlateAttributeMetaData.h
+++ b/Engine/Source/Runtime/SlateCore/Public/Types/SlateAttributeMetaData.h
@@ -27,36 +27,44 @@ public:
 public:
 	/** @return the instance associated to the SWidget (if it exists). */
 	static FSlateAttributeMetaData* FindMetaData(const SWidget& OwningWidget);
+
+	enum class EInvalidationPermission : uint8
+	{
+		/** Invalidate the widget if it's needed and it's construction phase is completed. */
+		AllowInvalidationIfConstructed,
+		/** Invalidate the widget if it's needed. */
+		AllowInvalidation,
+		/** Cache the invalidation. On any future update, if it's needed, invalidate the widget. */
+		DelayInvalidation,
+		/** Never invalidate the widget. */
+		DenyInvalidation,
+		/** Never invalidate the widget and clear any delayed invalidation. */
+		DenyAndClearDelayedInvalidation,
+	};
+
 	/**
 	 * Update all the attributes.
-	 * Invalidate the widget if it has finished its construction phase.
+	 * @param InvalidationStyle if we should invalidate the widget.
 	 */
-	static void UpdateAttributes(SWidget& OwningWidget);
+	static void UpdateAttributes(SWidget& OwningWidget, EInvalidationPermission InvalidationStyle);
 	/**
 	 * Update attributes that are mark to be updated when the widget is collapsed.
-	 * Invalidate the widget if it has finished its construction phase.
+	 * These attributes are usually responsible to change visibility of the widget.
+	 * @param InvalidationStyle if we should invalidate the widget.
 	 */
-	static void UpdateCollapsedAttributes(SWidget& OwningWidget);
+	static void UpdateCollapsedAttributes(SWidget& OwningWidget, EInvalidationPermission InvalidationStyle);
 	/**
 	 * Update attributes that are mark to be updated when the widget is NOT collapsed.
-	 * Invalidate the widget if it has finished its construction phase.
+	 * These attributes usually do not change the visibility of the widget.
+	 * @param InvalidationStyle if we should invalidate the widget.
 	 */
-	static void UpdateExpandedAttributes(SWidget& OwningWidget);
+	static void UpdateExpandedAttributes(SWidget& OwningWidget, EInvalidationPermission InvalidationStyle);
 	/**
-	 * Update attributes that are mark to be updated when the widget is collapsed.
-	 * @param bAllowInvalidation if we should allow the widget to be invalidated.
-	 */
-	static void UpdateAttributes(SWidget& OwningWidget, bool bAllowInvalidation);
-	/**
-	 * Update attributes that are mark to be updated when the widget is collapsed.
-	 * @param bAllowInvalidation if we should allow the widget to be invalidated.
-	 */
-	static void UpdateCollapsedAttributes(SWidget& OwningWidget, bool bAllowInvalidation);
-	/**
-	 * Update attributes that are mark to be updated when the widget is NOT collapsed.
-	 * @param bAllowInvalidation if we should allow the widget to be invalidated.
+	 * Update the children attributes that are mark to be updated when the widget is collapsed.
+	 * These attributes are usually responsible to change visibility of the widget.
+	 * @param InvalidationStyle if we should invalidate the widget.
 	 */
-	static void UpdateExpandedAttributes(SWidget& OwningWidget, bool bAllowInvalidation);
+	static void UpdateChildrenCollapsedAttributes(SWidget& OwningWidget, EInvalidationPermission InvalidationStyle);
 
 public:
 	bool IsBound(const FSlateAttributeBase& Attribute) const
@@ -92,7 +100,7 @@ private:
 	};
 	void RegisterAttributeImpl(SWidget& OwningWidget, FSlateAttributeBase& Attribute, ESlateAttributeType AttributeType, TUniquePtr<ISlateAttributeGetter>&& Getter);
 	bool UnregisterAttributeImpl(const FSlateAttributeBase& Attribute);
-	void UpdateAttributesImpl(SWidget& OwningWidget, EUpdateType UpdateType, bool bAllowInvalidation);
+	void UpdateAttributesImpl(SWidget& OwningWidget, EUpdateType UpdateType, EInvalidationPermission InvaldiationStyle);
 
 private:
 	int32 IndexOfAttribute(const FSlateAttributeBase& Attribute) const
@@ -156,6 +164,11 @@ private:
 	};
 
 	TArray<FGetterItem, TInlineAllocator<4>> Attributes;
+	//~ There is a possibility that the widget has a CachedInvalidationReason and a parent become collapsed.
+	//~The invalidation will probably never get executed but
+	//~1. The widget is collapsed indirectly, so we do not care if it's invalidated.
+	//~2. The parent widget will clear this widget PersistentState.
+	EInvalidateWidgetReason CachedInvalidationReason = EInvalidateWidgetReason::None;
 	bool bHasUpdatedManuallyFlagToReset = false;
 	uint8 CollaspedAttributeCounter = 0;
 };
diff --git a/Engine/Source/Runtime/SlateCore/Public/Widgets/SWidget.h b/Engine/Source/Runtime/SlateCore/Public/Widgets/SWidget.h
--- a/Engine/Source/Runtime/SlateCore/Public/Widgets/SWidget.h
+++ b/Engine/Source/Runtime/SlateCore/Public/Widgets/SWidget.h
@@ -848,11 +848,16 @@ public:
 	 * ArrangeChildren function.
 	 * Compute the Geometry of all the children and add populate the ArrangedChildren list with their values.
 	 * Each type of Layout panel should arrange children based on desired behavior.
+	 * 
+	 * Optionally, update the collapsed attributes (attributes that affect the visibility) of the children before executing the virtual ArrangeChildren function.
+	 * The visibility attribute is updated once per frame (see SlatePrepass).
+	 * Use the option when you are calling ArrangeChildren outside of the regular SWidget Paint/Tick.
 	 *
 	 * @param AllottedGeometry    The geometry allotted for this widget by its parent.
 	 * @param ArrangedChildren    The array to which to add the WidgetGeometries that represent the arranged children.
+	 * @param bUpdateAttributes   Update the collapsed attributes.
 	 */
-	void ArrangeChildren(const FGeometry& AllottedGeometry, FArrangedChildren& ArrangedChildren) const;
+	void ArrangeChildren(const FGeometry& AllottedGeometry, FArrangedChildren& ArrangedChildren, bool bUpdateAttributes = false) const;
 
 	/**
 	 * Every widget that has children must implement this method. This allows for iteration over the Widget's
