remove MemStack constructor with unused int argument

diff --git a/Engine/Source/Programs/UnrealLightmass/Private/LightmassCore/Math/LMOctree.h b/Engine/Source/Programs/UnrealLightmass/Private/LightmassCore/Math/LMOctree.h
--- a/Engine/Source/Programs/UnrealLightmass/Private/LightmassCore/Math/LMOctree.h
+++ b/Engine/Source/Programs/UnrealLightmass/Private/LightmassCore/Math/LMOctree.h
@@ -617,7 +617,7 @@ public:
 
 private:
 	/** Lightmass octree is grow-only. Use a stack allocator to coalesce allocations for better allocation performance. */
-	FMemStackBase NodeAllocator = FMemStackBase(0);
+	FMemStackBase NodeAllocator;
 
 	/** The octree's root node. */
 	FNode RootNode;
diff --git a/Engine/Source/Runtime/Core/Private/ProfilingDebugging/CpuProfilerTrace.cpp b/Engine/Source/Runtime/Core/Private/ProfilingDebugging/CpuProfilerTrace.cpp
--- a/Engine/Source/Runtime/Core/Private/ProfilingDebugging/CpuProfilerTrace.cpp
+++ b/Engine/Source/Runtime/Core/Private/ProfilingDebugging/CpuProfilerTrace.cpp
@@ -66,7 +66,6 @@ struct FCpuProfilerTraceInternal
 		: public FTlsAutoCleanup
 	{
 		FThreadBuffer()
-			: DynamicScopeNamesMemory(0)
 		{
 		}
 
diff --git a/Engine/Source/Runtime/Core/Private/Serialization/MemoryImage.cpp b/Engine/Source/Runtime/Core/Private/Serialization/MemoryImage.cpp
--- a/Engine/Source/Runtime/Core/Private/Serialization/MemoryImage.cpp
+++ b/Engine/Source/Runtime/Core/Private/Serialization/MemoryImage.cpp
@@ -1004,7 +1004,7 @@ public:
 		return Instance;
 	}
 
-	FHashedNameRegistry() : MemStack(0)
+	FHashedNameRegistry()
 	{
 		const char NoneString[] = "None";
 		EmptyString = RegisterString(FName(), NoneString, sizeof(NoneString), "", 0u);
diff --git a/Engine/Source/Runtime/Core/Public/Misc/MemStack.h b/Engine/Source/Runtime/Core/Public/Misc/MemStack.h
--- a/Engine/Source/Runtime/Core/Public/Misc/MemStack.h
+++ b/Engine/Source/Runtime/Core/Public/Misc/MemStack.h
@@ -79,11 +79,6 @@ public:
 	{
 	}
 
-	FMemStackBase(int unused) : FMemStackBase()
-	{
-		// @todo : remove uses of this constructor
-	}
-
 	~FMemStackBase()
 	{
 		check((GIsCriticalError || !NumMarks));
diff --git a/Engine/Source/Runtime/Engine/Private/Materials/MaterialHLSLTree.cpp b/Engine/Source/Runtime/Engine/Private/Materials/MaterialHLSLTree.cpp
--- a/Engine/Source/Runtime/Engine/Private/Materials/MaterialHLSLTree.cpp
+++ b/Engine/Source/Runtime/Engine/Private/Materials/MaterialHLSLTree.cpp
@@ -12,7 +12,7 @@
 #include "Materials/MaterialExpressionMaterialFunctionCall.h"
 #include "HLSLTree/HLSLTreeCommon.h"
 
-FMaterialHLSLTree::FMaterialHLSLTree() : Allocator(0)
+FMaterialHLSLTree::FMaterialHLSLTree()
 {
 }
 
diff --git a/Engine/Source/Runtime/RHI/Private/RHICommandList.cpp b/Engine/Source/Runtime/RHI/Private/RHICommandList.cpp
--- a/Engine/Source/Runtime/RHI/Private/RHICommandList.cpp
+++ b/Engine/Source/Runtime/RHI/Private/RHICommandList.cpp
@@ -909,7 +909,7 @@ FRHICommandListBase::FRHICommandListBase(FRHIGPUMask InGPUMask)
 	, UID(UINT32_MAX)
 	, Context(nullptr)
 	, ComputeContext(nullptr)
-	, MemManager(0)
+	, MemManager()
 	, bAsyncPSOCompileAllowed(true)
 	, GPUMask(InGPUMask)
 	, InitialGPUMask(InGPUMask)
diff --git a/Engine/Source/Runtime/RenderCore/Public/RenderGraphAllocator.h b/Engine/Source/Runtime/RenderCore/Public/RenderGraphAllocator.h
--- a/Engine/Source/Runtime/RenderCore/Public/RenderGraphAllocator.h
+++ b/Engine/Source/Runtime/RenderCore/Public/RenderGraphAllocator.h
@@ -67,7 +67,7 @@ private:
 		T Alloc;
 	};
 
-	FMemStackBase MemStack{0};
+	FMemStackBase MemStack;
 	TArray<FTrackedAlloc*> TrackedAllocs;
 
 	template <uint32>
diff --git a/Engine/Source/Runtime/Renderer/Private/RayTracing/RayTracingMaterialHitShaders.h b/Engine/Source/Runtime/Renderer/Private/RayTracing/RayTracingMaterialHitShaders.h
--- a/Engine/Source/Runtime/Renderer/Private/RayTracing/RayTracingMaterialHitShaders.h
+++ b/Engine/Source/Runtime/Renderer/Private/RayTracing/RayTracingMaterialHitShaders.h
@@ -144,7 +144,6 @@ class FRayTracingLocalShaderBindingWriter
 public:
 
 	FRayTracingLocalShaderBindingWriter()
-		: ParameterMemory(0)
 	{}
 
 	FRayTracingLocalShaderBindingWriter(const FRayTracingLocalShaderBindingWriter&) = delete;
diff --git a/Engine/Source/Runtime/SlateCore/Private/Rendering/DrawElements.cpp b/Engine/Source/Runtime/SlateCore/Private/Rendering/DrawElements.cpp
--- a/Engine/Source/Runtime/SlateCore/Private/Rendering/DrawElements.cpp
+++ b/Engine/Source/Runtime/SlateCore/Private/Rendering/DrawElements.cpp
@@ -115,7 +115,7 @@ static bool ShouldCull(const FSlateWindowElementList& ElementList, const FPaintG
 FSlateWindowElementList::FSlateWindowElementList(const TSharedPtr<SWindow>& InPaintWindow)
 	: WeakPaintWindow(InPaintWindow)
 	, RawPaintWindow(InPaintWindow.Get())
-	, MemManager(0)
+	, MemManager()
 #if STATS
 	, MemManagerAllocatedMemory(0)
 #endif
