Slate: Fix attributes update order in GlobalInvalidation.

diff --git a/Engine/Source/Runtime/SlateCore/Private/FastUpdate/SlateInvalidationRoot.cpp b/Engine/Source/Runtime/SlateCore/Private/FastUpdate/SlateInvalidationRoot.cpp
--- a/Engine/Source/Runtime/SlateCore/Private/FastUpdate/SlateInvalidationRoot.cpp
+++ b/Engine/Source/Runtime/SlateCore/Private/FastUpdate/SlateInvalidationRoot.cpp
@@ -639,54 +639,12 @@ void FSlateInvalidationRoot::ProcessPreUpdate()
 				//2.UpdateAttributes of AttributeSortOrder
 				//3.Invalidation ChildOrder of NeedsUpdateSortOrder
 
-				if (AttributeSortOrder == NeedsUpdateSortOrder)
-				{
-					// Process AttributeRegistration invalidation.
-					const FSlateInvalidationWidgetIndex WidgetIndex = WidgetsNeedingPreUpdate->HeapPeek();
-					FSlateInvalidationWidgetList::InvalidationWidgetType& InvalidationWidget = (*FastWidgetPathList)[WidgetIndex];
-					SWidget* WidgetPtr = InvalidationWidget.GetWidget();
-					if (ensureMsgf(WidgetPtr, TEXT("A parent remove this widget. Child order invalidation should have been called before we process this widget.")))
-					{
-						if (EnumHasAnyFlags(InvalidationWidget.CurrentInvalidateReason, EInvalidateWidgetReason::AttributeRegistration))
-						{
-#if WITH_SLATE_DEBUGGING
-							if (GSlateInvalidationRootDumpPreInvalidationList)
-							{
-								LogPreInvalidationItem(*FastWidgetPathList, WidgetIndex);
-							}
-#endif
-
-							FastWidgetPathList->ProcessAttributeRegistrationInvalidation(InvalidationWidget);
-							if (WidgetPtr->HasRegisteredSlateAttribute() && WidgetPtr->IsAttributesUpdatesEnabled())
-							{
-								// This should be the next Attribute to update
-								AttributeItt.Seek(InvalidationWidget.Index);
-							}
-							EnumRemoveFlags(InvalidationWidget.CurrentInvalidateReason, EInvalidateWidgetReason::AttributeRegistration);
-
-							// Do we still need to update this element, if not, then remove it from the update list.
-							if (!Slate::EInvalidateWidgetReason_HasPreUpdateFlag(InvalidationWidget.CurrentInvalidateReason))
-							{
-								WidgetsNeedingPreUpdate->HeapPopDiscard();
-							}
-
-							// Stop and process the next element.
-							continue;
-						}
-					}
-					else
-					{
-						WidgetsNeedingPreUpdate->HeapPopDiscard();
-						continue;
-					}
-				}
-
 				if (AttributeSortOrder <= NeedsUpdateSortOrder)
 				{
 					// Update Attributes
+					//Note the attribute may still be in the list and will get remove in next loop tick. UpdateCollapsedAttributes and UpdateExpandedAttributes won't do anything.
 					FSlateInvalidationWidgetList::InvalidationWidgetType& InvalidationWidget = (*FastWidgetPathList)[AttributeItt.GetCurrentIndex()];
-					SWidget* WidgetPtr = InvalidationWidget.GetWidget();
-					if (ensureMsgf(WidgetPtr, TEXT("A parent remove this widget. Child order invalidation should have been called before we process this widget.")))
+					if (SWidget* WidgetPtr = InvalidationWidget.GetWidget())
 					{
 						if (!InvalidationWidget.Visibility.IsCollapseIndirectly())
 						{
@@ -699,21 +657,31 @@ void FSlateInvalidationRoot::ProcessPreUpdate()
 								InvalidationWidget.bDebug_AttributeUpdated = true;
 #endif
 								FSlateAttributeMetaData::UpdateExpandedAttributes(*WidgetPtr, FSlateAttributeMetaData::EInvalidationPermission::AllowInvalidation);
+								AttributeItt.Advance();
+							}
+							else
+							{
+								AttributeItt.AdvanceToNextSibling();
 							}
 						}
+						else
+						{
+							AttributeItt.AdvanceToNextParent();
+						}
+					}
+					else
+					{
+						AttributeItt.Advance();
 					}
-					AttributeItt.Advance();
-					continue; // Do the loop again, in case a ChildOrder invalidation occurred.
 				}
-
-				checkf(NeedsUpdateSortOrder < AttributeSortOrder, TEXT("Else case"));
+				else
 				{
 					// Process ChildOrder invalidation.
 
-					const FSlateInvalidationWidgetIndex WidgetIndex = WidgetsNeedingPreUpdate->HeapPop();
+					const FSlateInvalidationWidgetIndex WidgetIndex = WidgetsNeedingPreUpdate->HeapPeek();
 					FSlateInvalidationWidgetList::InvalidationWidgetType& InvalidationWidget = (*FastWidgetPathList)[WidgetIndex];
-					// It could have been destroyed since then
-					if (InvalidationWidget.GetWidget())
+					// It could have been destroyed
+					if (SWidget* WidgetPtr = InvalidationWidget.GetWidget())
 					{
 #if WITH_SLATE_DEBUGGING
 						if (GSlateInvalidationRootDumpPreInvalidationList)
@@ -722,7 +690,27 @@ void FSlateInvalidationRoot::ProcessPreUpdate()
 						}
 #endif
 
-						if (ensureAlways(EnumHasAnyFlags(InvalidationWidget.CurrentInvalidateReason, EInvalidateWidgetReason::ChildOrder)))
+						if (EnumHasAnyFlags(InvalidationWidget.CurrentInvalidateReason, EInvalidateWidgetReason::AttributeRegistration))
+						{
+							FastWidgetPathList->ProcessAttributeRegistrationInvalidation(InvalidationWidget);
+							EnumRemoveFlags(InvalidationWidget.CurrentInvalidateReason, EInvalidateWidgetReason::AttributeRegistration);
+
+							// This element was removed or added, seek will assign the correct widget to be ticked next.
+							AttributeItt.Seek(InvalidationWidget.Index);
+							if (FastWidgetPathList->ShouldBeAddedToAttributeList(WidgetPtr))
+							{
+								// Do we still need to update this element, if not, then remove it from the update list.
+								if (!Slate::EInvalidateWidgetReason_HasPreUpdateFlag(InvalidationWidget.CurrentInvalidateReason))
+								{
+									WidgetsNeedingPreUpdate->HeapPopDiscard();
+								}
+
+								// We should update the attribute of this proxy before doing the ChildOrder (if any).
+								continue;
+							}
+						}
+
+						if (EnumHasAnyFlags(InvalidationWidget.CurrentInvalidateReason, EInvalidateWidgetReason::ChildOrder))
 						{
 // Uncomment to see to be able to compare the list before and after when debugging
 #if 0
@@ -753,6 +741,7 @@ void FSlateInvalidationRoot::ProcessPreUpdate()
 							//EnumRemoveFlags(InvalidationWidget.CurrentInvalidateReason, EInvalidateWidgetReason::ChildOrder);
 						}
 					}
+					WidgetsNeedingPreUpdate->HeapPopDiscard();
 				}
 			}
 		}
diff --git a/Engine/Source/Runtime/SlateCore/Private/FastUpdate/SlateInvalidationWidgetHeap.h b/Engine/Source/Runtime/SlateCore/Private/FastUpdate/SlateInvalidationWidgetHeap.h
--- a/Engine/Source/Runtime/SlateCore/Private/FastUpdate/SlateInvalidationWidgetHeap.h
+++ b/Engine/Source/Runtime/SlateCore/Private/FastUpdate/SlateInvalidationWidgetHeap.h
@@ -75,8 +75,9 @@ public:
 	}
 
 	/** Returns and removes the biggest WidgetIndex from the list. */
-	FSlateInvalidationWidgetIndex HeapPop()
+	UE_NODISCARD FSlateInvalidationWidgetIndex HeapPop()
 	{
+		check(Heap.Num() > 0);
 		FSlateInvalidationWidgetIndex Result = Heap.HeapTop().GetWidgetIndex();
 		Heap.HeapPopDiscard(SortPredicate(), false);
 		OwnerList[Result].bContainedByWidgetPreHeap = false;
@@ -86,18 +87,23 @@ public:
 	/** Removes the biggest WidgetIndex from the list. */
 	void HeapPopDiscard()
 	{
+		check(Heap.Num() > 0);
+		FSlateInvalidationWidgetIndex Result = Heap.HeapTop().GetWidgetIndex();
 		Heap.HeapPopDiscard(SortPredicate(), false);
+		OwnerList[Result].bContainedByWidgetPreHeap = false;
 	}
 	
 	/** Returns the biggest WidgetIndex from the list. */
 	inline FSlateInvalidationWidgetIndex HeapPeek() const
 	{
+		check(Heap.Num() > 0);
 		return Heap.HeapTop().GetWidgetIndex();
 	}
 	
 	/** Returns the biggest WidgetIndex from the list. */
 	inline const FElement& HeapPeekElement() const
 	{
+		check(Heap.Num() > 0);
 		return Heap.HeapTop();
 	}
 
diff --git a/Engine/Source/Runtime/SlateCore/Private/FastUpdate/SlateInvalidationWidgetList.cpp b/Engine/Source/Runtime/SlateCore/Private/FastUpdate/SlateInvalidationWidgetList.cpp
--- a/Engine/Source/Runtime/SlateCore/Private/FastUpdate/SlateInvalidationWidgetList.cpp
+++ b/Engine/Source/Runtime/SlateCore/Private/FastUpdate/SlateInvalidationWidgetList.cpp
@@ -138,9 +138,7 @@ void FSlateInvalidationWidgetList::FWidgetAttributeIterator::FixCurrentWidgetInd
 	}
 	else if (bNeedsWidgetFixUp)
 	{
-		AttributeIndex = 0;
-		CurrentWidgetIndex = FSlateInvalidationWidgetIndex::Invalid;
-		CurrentWidgetSortOrder = FSlateInvalidationWidgetSortOrder::LimitMax();
+		Clear();
 	}
 
 	MoveToWidgetIndexOnNextAdvance = FSlateInvalidationWidgetIndex::Invalid;
@@ -161,8 +159,9 @@ void FSlateInvalidationWidgetList::FWidgetAttributeIterator::Seek(FSlateInvalida
 	}
 	else
 	{
-		CurrentWidgetIndex = SeekTo;
+		CurrentWidgetIndex = FSlateInvalidationWidgetIndex{ (IndexType)SeekTo.ArrayIndex, ArrayNode.ElementIndexList_WidgetWithRegisteredSlateAttribute[AttributeIndex] };
 		CurrentWidgetSortOrder = FSlateInvalidationWidgetSortOrder{ WidgetList, CurrentWidgetIndex };
+		check(WidgetList.IsValidIndex(CurrentWidgetIndex));
 	}
 }
 
@@ -185,16 +184,54 @@ void FSlateInvalidationWidgetList::FWidgetAttributeIterator::Advance()
 		CurrentWidgetIndex = FSlateInvalidationWidgetIndex{ (IndexType)ArrayIndex, ArrayNode.ElementIndexList_WidgetWithRegisteredSlateAttribute[AttributeIndex] };
 		CurrentWidgetSortOrder = (CurrentWidgetIndex != FSlateInvalidationWidgetIndex::Invalid)
 			? FSlateInvalidationWidgetSortOrder{ WidgetList, CurrentWidgetIndex }
-			: FSlateInvalidationWidgetSortOrder::LimitMax();
+		: FSlateInvalidationWidgetSortOrder::LimitMax();
+	}
+}
+
+
+void FSlateInvalidationWidgetList::FWidgetAttributeIterator::AdvanceToNextSibling()
+{
+	check(MoveToWidgetIndexOnNextAdvance == FSlateInvalidationWidgetIndex::Invalid);
+
+	const FSlateInvalidationWidgetIndex SeekToIndex = WidgetList.IncrementIndex(WidgetList[CurrentWidgetIndex].LeafMostChildIndex);
+	if (SeekToIndex != FSlateInvalidationWidgetIndex::Invalid)
+	{
+		Seek(SeekToIndex);
+	}
+	else
+	{
+		Clear();
+	}
+}
+
+
+void FSlateInvalidationWidgetList::FWidgetAttributeIterator::AdvanceToNextParent()
+{
+	check(MoveToWidgetIndexOnNextAdvance == FSlateInvalidationWidgetIndex::Invalid);
+
+	const FSlateInvalidationWidgetIndex ParentIndex = WidgetList[CurrentWidgetIndex].ParentIndex;
+	if (ParentIndex != FSlateInvalidationWidgetIndex::Invalid)
+	{
+		const FSlateInvalidationWidgetIndex SeekToIndex = WidgetList.IncrementIndex(WidgetList[ParentIndex].LeafMostChildIndex);
+		if (SeekToIndex != FSlateInvalidationWidgetIndex::Invalid)
+		{
+			Seek(SeekToIndex);
+		}
+		else
+		{
+			Clear();
+		}
+	}
+	else
+	{
+		Clear();
 	}
 }
 
 
 void FSlateInvalidationWidgetList::FWidgetAttributeIterator::AdvanceArrayIndex(int32 ArrayIndex)
 {
-	AttributeIndex = 0;
-	CurrentWidgetIndex = FSlateInvalidationWidgetIndex::Invalid;
-	CurrentWidgetSortOrder = FSlateInvalidationWidgetSortOrder::LimitMax();
+	Clear();
 
 	while (ArrayIndex != INDEX_NONE)
 	{
@@ -210,6 +247,14 @@ void FSlateInvalidationWidgetList::FWidgetAttributeIterator::AdvanceArrayIndex(i
 }
 
 
+void FSlateInvalidationWidgetList::FWidgetAttributeIterator::Clear()
+{
+	AttributeIndex = 0;
+	CurrentWidgetIndex = FSlateInvalidationWidgetIndex::Invalid;
+	CurrentWidgetSortOrder = FSlateInvalidationWidgetSortOrder::LimitMax();
+}
+
+
 /** */
 void FSlateInvalidationWidgetList::FArrayNode::RemoveElementIndexBiggerOrEqualThan(IndexType ElementIndex)
 {
@@ -366,7 +411,7 @@ void FSlateInvalidationWidgetList::BuildWidgetList(TSharedRef<SWidget> InRoot)
 }
 
 
-void FSlateInvalidationWidgetList::_RebuildWidgetListTree(TSharedRef<SWidget> Widget, int32 ChildAtIndex)
+void FSlateInvalidationWidgetList::Internal_RebuildWidgetListTree(TSharedRef<SWidget> Widget, int32 ChildAtIndex)
 {
 	const bool bShouldAddWidget = ShouldBeAdded(Widget);
 	FChildren* ParentChildren = Widget->GetAllChildren();
@@ -462,7 +507,7 @@ bool FSlateInvalidationWidgetList::ProcessChildOrderInvalidation(const Invalidat
 			// Find all it's previous children
 			FMemMark Mark2(FMemStack::Get());
 			TArray<FFindChildrenElement, TMemStackAllocator<>> PreviousChildrenWidget;
-			_FindChildren(InvalidationWidget.Index, PreviousChildrenWidget);
+			Internal_FindChildren(InvalidationWidget.Index, PreviousChildrenWidget);
 
 			FChildren* InvalidatedChildren = WidgetPtr->GetAllChildren();
 			check(InvalidatedChildren);
@@ -532,7 +577,7 @@ bool FSlateInvalidationWidgetList::ProcessChildOrderInvalidation(const Invalidat
 			{
 				ensureMsgf(ChildToRemoveRange.IsValid(), TEXT("The ChildToRemoveRange should be valid at this point."));
 				Callback.PreChildRemove(ChildToRemoveRange);
-				_RemoveRangeFromSameParent(ChildToRemoveRange);
+				Internal_RemoveRangeFromSameParent(ChildToRemoveRange);
 			}
 			else if (CurrentOperation == EOperation::Cut)
 			{
@@ -550,7 +595,7 @@ bool FSlateInvalidationWidgetList::ProcessChildOrderInvalidation(const Invalidat
 				SCOPED_NAMED_EVENT(Slate_InvalidationList_ProcessRebuild, FColorList::Blue);
 				FSlateInvalidationWidgetIndex PreviousWidgetIndex = WidgetPtr->GetProxyHandle().GetWidgetIndex();
 				FSlateInvalidationWidgetIndex PreviousLeafMostChildIndex = InvalidationWidget.LeafMostChildIndex;
-				_RebuildWidgetListTree(WidgetPtr->AsShared(), StartChildIndex);
+				Internal_RebuildWidgetListTree(WidgetPtr->AsShared(), StartChildIndex);
 
 
 				const InvalidationWidgetType& NewInvalidationWidget = (*this)[PreviousWidgetIndex];
@@ -614,7 +659,7 @@ namespace InvalidationList
 }
 
 
-void FSlateInvalidationWidgetList::_FindChildren(FSlateInvalidationWidgetIndex WidgetIndex, TArray<FFindChildrenElement, TMemStackAllocator<>>& Widgets) const
+void FSlateInvalidationWidgetList::Internal_FindChildren(FSlateInvalidationWidgetIndex WidgetIndex, TArray<FFindChildrenElement, TMemStackAllocator<>>& Widgets) const
 {
 	Widgets.Reserve(16);
 	const InvalidationWidgetType& InvalidationWidget = (*this)[WidgetIndex];
@@ -985,7 +1030,7 @@ void FSlateInvalidationWidgetList::UpdateParentLeafIndex(const InvalidationWidge
  * Can use to remove range (B,D) or (C,C) or (C,D) or (E, G) or (F,G) or (B,G).
  * Cannot use (B,E) or (B,C) or (B,F)
  */
-void FSlateInvalidationWidgetList::_RemoveRangeFromSameParent(const FIndexRange Range)
+void FSlateInvalidationWidgetList::Internal_RemoveRangeFromSameParent(const FIndexRange Range)
 {
 	// Fix up Parent's LeafIndex if they are in Range
 	{
@@ -1323,7 +1368,7 @@ void FSlateInvalidationWidgetList::RemoveWidget(const FSlateInvalidationWidgetIn
 		const InvalidationWidgetType& InvalidationWidget = (*this)[WidgetIndex];
 
 		const FIndexRange Range = { *this, WidgetIndex, InvalidationWidget.LeafMostChildIndex };
-		_RemoveRangeFromSameParent(Range);
+		Internal_RemoveRangeFromSameParent(Range);
 	}
 }
 
@@ -1337,7 +1382,7 @@ void FSlateInvalidationWidgetList::RemoveWidget(const TSharedRef<SWidget> Widget
 		{
 			const InvalidationWidgetType& InvalidationWidget = (*this)[WidgetIndex];
 			const FIndexRange Range = { *this, WidgetIndex, InvalidationWidget.LeafMostChildIndex };
-			_RemoveRangeFromSameParent(Range);
+			Internal_RemoveRangeFromSameParent(Range);
 		}
 	}
 
@@ -1357,7 +1402,7 @@ TArray<TSharedPtr<SWidget>> FSlateInvalidationWidgetList::FindChildren(const TSh
 
 		FMemMark Mark(FMemStack::Get());
 		TArray<FFindChildrenElement, TMemStackAllocator<>> PreviousChildrenWidget;
-		_FindChildren(WidgetIndex, PreviousChildrenWidget);
+		Internal_FindChildren(WidgetIndex, PreviousChildrenWidget);
 
 		Result.Reserve(PreviousChildrenWidget.Num());
 		for (const FFindChildrenElement& ChildrenElement : PreviousChildrenWidget)
@@ -1624,10 +1669,10 @@ bool FSlateInvalidationWidgetList::VerifyElementIndexList() const
 							, *FReflectionMetaData::GetWidgetDebugInfo(Widget));
 						bResult = false;
 					}
-
-					if (InvalidationWidget.Visibility.IsVisible() && !InvalidationWidget.bDebug_AttributeUpdated)
+					else if (!InvalidationWidget.Visibility.IsCollapsed() && !InvalidationWidget.bDebug_AttributeUpdated)
 					{
 						UE_LOG(LogSlate, Warning, TEXT("The widget '%s' was not updated."), *FReflectionMetaData::GetWidgetDebugInfo(Widget));
+						bResult = false;
 					}
 				}
 			}
diff --git a/Engine/Source/Runtime/SlateCore/Private/FastUpdate/SlateInvalidationWidgetList.h b/Engine/Source/Runtime/SlateCore/Private/FastUpdate/SlateInvalidationWidgetList.h
--- a/Engine/Source/Runtime/SlateCore/Private/FastUpdate/SlateInvalidationWidgetList.h
+++ b/Engine/Source/Runtime/SlateCore/Private/FastUpdate/SlateInvalidationWidgetList.h
@@ -11,6 +11,9 @@
 #include "Widgets/SNullWidget.h"
 #include "Widgets/SWidget.h"
 
+#define UE_SLATE_WITH_INVALIDATIONWIDGETLIST_RANGECHECK UE_BUILD_DEBUG
+
+
 class FSlateInvalidationWidgetList
 {
 	friend struct FSlateInvalidationWidgetSortOrder;
@@ -38,7 +41,9 @@ public:
 			: InclusiveMin(InFrom), InclusiveMax(InEnd)
 			, OrderMin(Self, InFrom), OrderMax(Self, InEnd)
 		{
+#if UE_SLATE_WITH_INVALIDATIONWIDGETLIST_RANGECHECK
 			check(OrderMin <= OrderMax);
+#endif
 		}
 		bool Include(FSlateInvalidationWidgetSortOrder Other) const
 		{
@@ -213,11 +218,18 @@ public:
 		/** Advance the iterator to the next valid widget index. */
 		void Advance();
 
+		/** Advance the iterator to the next valid widget index that is a child of this widget. */
+		void AdvanceToNextSibling();
+
+		/** Advance the iterator to the next valid widget index that is a sibling of this widget's parent. */
+		void AdvanceToNextParent();
+
 		/** Is the iterator pointing to a valid widget index. */
 		bool IsValid() const { return CurrentWidgetIndex != FSlateInvalidationWidgetIndex::Invalid; }
 
 	private:
 		void AdvanceArrayIndex(int32 ArrayIndex);
+		void Clear();
 	};
 
 	FWidgetAttributeIterator CreateWidgetAttributeIterator() const
@@ -229,14 +241,18 @@ public:
 	/** Returns reference to element at give index. */
 	InvalidationWidgetType& operator[](const FSlateInvalidationWidgetIndex Index)
 	{
+#if UE_SLATE_WITH_INVALIDATIONWIDGETLIST_RANGECHECK
 		check(IsValidIndex(Index));
+#endif
 		return Data[Index.ArrayIndex].ElementList[Index.ElementIndex];
 	}
 
 	/** Returns reference to element at give index. */
 	const InvalidationWidgetType& operator[](const FSlateInvalidationWidgetIndex Index) const
 	{
+#if UE_SLATE_WITH_INVALIDATIONWIDGETLIST_RANGECHECK
 		check(IsValidIndex(Index));
+#endif
 		return Data[Index.ArrayIndex].ElementList[Index.ElementIndex];
 	}
 
@@ -319,7 +335,7 @@ public:
 	bool VerifyElementIndexList() const;
 #endif
 
-private:
+public:
 	bool ShouldDoRecursion(const SWidget* Widget) const
 	{
 		return !Widget->Advanced_IsInvalidationRoot() || IsEmpty();
@@ -328,19 +344,19 @@ private:
 	{
 		return !Widget->Advanced_IsInvalidationRoot() || IsEmpty();
 	}
-	bool ShouldBeAdded(const SWidget* Widget) const
+	static bool ShouldBeAdded(const SWidget* Widget)
 	{
 		return Widget != &(SNullWidget::NullWidget.Get());
 	}
-	bool ShouldBeAdded(const TSharedRef<SWidget>& Widget) const
+	static bool ShouldBeAdded(const TSharedRef<SWidget>& Widget)
 	{
 		return Widget != SNullWidget::NullWidget;
 	}
-	bool ShouldBeAddedToAttributeList(const SWidget* Widget) const
+	static bool ShouldBeAddedToAttributeList(const SWidget* Widget)
 	{
 		return Widget->HasRegisteredSlateAttribute() && Widget->IsAttributesUpdatesEnabled();
 	}
-	bool ShouldBeAddedToAttributeList(const TSharedRef<SWidget>& Widget) const
+	static bool ShouldBeAddedToAttributeList(const TSharedRef<SWidget>& Widget)
 	{
 		return Widget->HasRegisteredSlateAttribute() && Widget->IsAttributesUpdatesEnabled();
 	}
@@ -380,10 +396,10 @@ private:
 
 private:
 	FSlateInvalidationWidgetIndex Internal_BuildWidgetList_Recursive(TSharedRef<SWidget>& Widget, FSlateInvalidationWidgetIndex ParentIndex, IndexType& LastestIndex, FSlateInvalidationWidgetVisibility ParentVisibility, bool bParentVolatile);
-	void _RebuildWidgetListTree(TSharedRef<SWidget> Widget, int32 ChildAtIndex);
+	void Internal_RebuildWidgetListTree(TSharedRef<SWidget> Widget, int32 ChildAtIndex);
 	using FFindChildrenElement = TPair<SWidget*, FSlateInvalidationWidgetIndex>;
-	void _FindChildren(FSlateInvalidationWidgetIndex WidgetIndex, TArray<FFindChildrenElement, TMemStackAllocator<>>& Widgets) const;
-	void _RemoveRangeFromSameParent(const FIndexRange Range);
+	void Internal_FindChildren(FSlateInvalidationWidgetIndex WidgetIndex, TArray<FFindChildrenElement, TMemStackAllocator<>>& Widgets) const;
+	void Internal_RemoveRangeFromSameParent(const FIndexRange Range);
 	FCutResult _CutArray(const FSlateInvalidationWidgetIndex WhereToCut);
 
 private:
diff --git a/Engine/Source/Runtime/SlateCore/Public/FastUpdate/WidgetProxy.h b/Engine/Source/Runtime/SlateCore/Public/FastUpdate/WidgetProxy.h
--- a/Engine/Source/Runtime/SlateCore/Public/FastUpdate/WidgetProxy.h
+++ b/Engine/Source/Runtime/SlateCore/Public/FastUpdate/WidgetProxy.h
@@ -41,12 +41,14 @@ public:
 		, bVisible(InVisibility.IsVisible())
 		, bAncestorCollapse(false)
 		, bCollapse(InVisibility == EVisibility::Collapsed)
+		, FlagPadding(0)
 	{ }
 	FSlateInvalidationWidgetVisibility(FSlateInvalidationWidgetVisibility ParentFlags, EVisibility InVisibility)
 		: bAncestorsVisible(ParentFlags.IsVisible())
 		, bVisible(InVisibility.IsVisible())
 		, bAncestorCollapse(ParentFlags.IsCollapsed())
 		, bCollapse(InVisibility == EVisibility::Collapsed)
+		, FlagPadding(0)
 	{ }
 	FSlateInvalidationWidgetVisibility(const FSlateInvalidationWidgetVisibility& Other) : Flags(Other.Flags) {  }
 	FSlateInvalidationWidgetVisibility& operator=(const FSlateInvalidationWidgetVisibility& Other) { Flags = Other.Flags; return *this; }
@@ -82,8 +84,8 @@ public:
 		return Result;
 	}
 
-	bool operator==(FSlateInvalidationWidgetVisibility Other) const { return Other.CompareFlags == CompareFlags; }
-	bool operator!=(FSlateInvalidationWidgetVisibility Other) const { return Other.CompareFlags != CompareFlags; }
+	bool operator==(FSlateInvalidationWidgetVisibility Other) const { return Other.Flags == Flags; }
+	bool operator!=(FSlateInvalidationWidgetVisibility Other) const { return Other.Flags != Flags; }
 
 private:
 	union
@@ -94,8 +96,8 @@ private:
 			uint8 bVisible : 1;
 			uint8 bAncestorCollapse : 1;	// at least one ancestor is collapse
 			uint8 bCollapse : 1;
+			uint8 FlagPadding : 4;
 		};
-		uint8 CompareFlags : 4;
 		uint8 Flags;
 	};
 };
