Slate: Remove OnInvalidation feature. The invalidation is not always triggered when the value is cached or when it's built from a FastPath widget list.

diff --git a/Engine/Source/Runtime/SlateCore/Private/Test/SlateAttributeTest.cpp b/Engine/Source/Runtime/SlateCore/Private/Test/SlateAttributeTest.cpp
--- a/Engine/Source/Runtime/SlateCore/Private/Test/SlateAttributeTest.cpp
+++ b/Engine/Source/Runtime/SlateCore/Private/Test/SlateAttributeTest.cpp
@@ -191,82 +191,6 @@ public:
 	int32 CallbackValueC;
 };
 
-SLATE_IMPLEMENT_WIDGET(SAttributeLeftWidget_OnInvalidationParent)
-void SAttributeLeftWidget_OnInvalidationParent::PrivateRegisterAttributes(FSlateAttributeInitializer& AttributeInitializer)
-{
-	SLATE_ADD_MEMBER_ATTRIBUTE_DEFINITION(AttributeInitializer, IntAttributeA, EInvalidateWidgetReason::ChildOrder)
-		.OnInvalidation(FSlateAttributeDescriptor::FInvalidationDelegate::CreateLambda([](SWidget& Widget)
-			{
-				SAttributeLeftWidget_OnInvalidationParent& Instance = static_cast<SAttributeLeftWidget_OnInvalidationParent&>(Widget);
-				Instance.CallbackValueA = Instance.IntAttributeA.Get();
-			}));
-	SLATE_ADD_MEMBER_ATTRIBUTE_DEFINITION(AttributeInitializer, IntAttributeB, EInvalidateWidgetReason::ChildOrder)
-		.OnInvalidation(FSlateAttributeDescriptor::FInvalidationDelegate::CreateLambda([](SWidget& Widget)
-			{
-				SAttributeLeftWidget_OnInvalidationParent& Instance = static_cast<SAttributeLeftWidget_OnInvalidationParent&>(Widget);
-				Instance.CallbackValueB = Instance.IntAttributeB.Get();
-			}));
-	SLATE_ADD_MEMBER_ATTRIBUTE_DEFINITION(AttributeInitializer, IntAttributeC, EInvalidateWidgetReason::ChildOrder)
-		.OnInvalidation(FSlateAttributeDescriptor::FInvalidationDelegate::CreateLambda([](SWidget& Widget)
-			{
-				SAttributeLeftWidget_OnInvalidationParent& Instance = static_cast<SAttributeLeftWidget_OnInvalidationParent&>(Widget);
-				Instance.CallbackValueC = Instance.IntAttributeC.Get();
-			}));
-}
-
-class SAttributeLeftWidget_OnInvalidationChild : public SAttributeLeftWidget_OnInvalidationParent
-{
-	SLATE_DECLARE_WIDGET(SAttributeLeftWidget_OnInvalidationChild, SAttributeLeftWidget_OnInvalidationParent)
-public:
-	SLATE_BEGIN_ARGS(SAttributeLeftWidget_OnInvalidationChild) {}
-	SLATE_END_ARGS()
-
-	void Construct(const FArguments& InArgs) {}
-	virtual FVector2D ComputeDesiredSize(float) const override { return FVector2D{ 100, 100 }; }
-	virtual int32 OnPaint(const FPaintArgs& Args, const FGeometry& AllottedGeometry, const FSlateRect& MyCullingRect, FSlateWindowElementList& OutDrawElements, int32 LayerId, const FWidgetStyle& InWidgetStyle, bool bParentEnabled) const override
-	{
-		return LayerId;
-	}
-
-	virtual void SetCallbackValue(int32 Value) override
-	{
-		Super::SetCallbackValue(Value);
-		OverrideCallbackValueA = Value;
-		OverrideCallbackValueB = Value;
-		OverrideCallbackValueC = Value;
-	}
-
-	int32 OverrideCallbackValueA;
-	int32 OverrideCallbackValueB;
-	int32 OverrideCallbackValueC;
-};
-
-SLATE_IMPLEMENT_WIDGET(SAttributeLeftWidget_OnInvalidationChild)
-void SAttributeLeftWidget_OnInvalidationChild::PrivateRegisterAttributes(FSlateAttributeInitializer& AttributeInitializer)
-{
-	AttributeInitializer.OverrideOnInvalidation("IntAttributeA"
-		, FSlateAttributeDescriptor::EInvalidationDelegateOverrideType::ReplacePrevious
-		, FSlateAttributeDescriptor::FInvalidationDelegate::CreateLambda([](SWidget& Widget)
-			{
-				SAttributeLeftWidget_OnInvalidationChild& Instance = static_cast<SAttributeLeftWidget_OnInvalidationChild&>(Widget);
-				Instance.OverrideCallbackValueA = Instance.IntAttributeA.Get();
-			}));
-	AttributeInitializer.OverrideOnInvalidation("IntAttributeB"
-		, FSlateAttributeDescriptor::EInvalidationDelegateOverrideType::ExecuteAfterPrevious
-		, FSlateAttributeDescriptor::FInvalidationDelegate::CreateLambda([](SWidget& Widget)
-			{
-				SAttributeLeftWidget_OnInvalidationChild& Instance = static_cast<SAttributeLeftWidget_OnInvalidationChild&>(Widget);
-				Instance.OverrideCallbackValueB = Instance.CallbackValueB;
-			}));
-	AttributeInitializer.OverrideOnInvalidation("IntAttributeC"
-		, FSlateAttributeDescriptor::EInvalidationDelegateOverrideType::ExecuteBeforePrevious
-		, FSlateAttributeDescriptor::FInvalidationDelegate::CreateLambda([](SWidget& Widget)
-			{
-				SAttributeLeftWidget_OnInvalidationChild& Instance = static_cast<SAttributeLeftWidget_OnInvalidationChild&>(Widget);
-				Instance.OverrideCallbackValueC = Instance.CallbackValueC;
-			}));
-}
-
 }
 }
 }
@@ -812,91 +736,6 @@ bool FSlateAttributeTest::RunTest(const FString& Parameters)
 		}
 	}
 
-	// Test OnInvalidation
-	{
-		{
-			TSharedRef<UE::Slate::Private::SAttributeLeftWidget_OnInvalidationParent> WidgetParent = SNew(UE::Slate::Private::SAttributeLeftWidget_OnInvalidationParent);
-			WidgetParent->SetCallbackValue(10);
-			WidgetParent->MarkPrepassAsDirty();
-			WidgetParent->SlatePrepass(1.f);
-			AddErrorIfFalse(WidgetParent->CallbackValueA == 10, TEXT("The callback for value A was triggered."));
-			AddErrorIfFalse(WidgetParent->CallbackValueB == 10, TEXT("The callback for value B was triggered."));
-			AddErrorIfFalse(WidgetParent->CallbackValueC == 10, TEXT("The callback for value C was triggered."));
-		}
-		{
-			TSharedRef<UE::Slate::Private::SAttributeLeftWidget_OnInvalidationChild> WidgetChild = SNew(UE::Slate::Private::SAttributeLeftWidget_OnInvalidationChild);
-			WidgetChild->SetCallbackValue(10);
-			WidgetChild->MarkPrepassAsDirty();
-			WidgetChild->SlatePrepass(1.f);
-			AddErrorIfFalse(WidgetChild->CallbackValueA == 10, TEXT("The callback for value A was triggered."));
-			AddErrorIfFalse(WidgetChild->CallbackValueB == 10, TEXT("The callback for value B was triggered."));
-			AddErrorIfFalse(WidgetChild->CallbackValueC == 10, TEXT("The callback for value C was triggered."));
-			AddErrorIfFalse(WidgetChild->OverrideCallbackValueA == 10, TEXT("The callback for override value A was triggered."));
-			AddErrorIfFalse(WidgetChild->OverrideCallbackValueB == 10, TEXT("The callback for override value B was triggered."));
-			AddErrorIfFalse(WidgetChild->OverrideCallbackValueC == 10, TEXT("The callback for override value C was triggered."));
-		}
-
-		{
-			TSharedRef<UE::Slate::Private::SAttributeLeftWidget_OnInvalidationParent> WidgetParent = SNew(UE::Slate::Private::SAttributeLeftWidget_OnInvalidationParent);
-			WidgetParent->SetCallbackValue(10);
-			WidgetParent->SetAttributeA(4);
-			AddErrorIfFalse(WidgetParent->CallbackValueA == 4, TEXT("The callback for value A was not triggered."));
-			AddErrorIfFalse(WidgetParent->CallbackValueB == 10, TEXT("The callback for value B was triggered."));
-			AddErrorIfFalse(WidgetParent->CallbackValueC == 10, TEXT("The callback for value C was triggered."));
-			WidgetParent->SetAttributeB(5);
-			AddErrorIfFalse(WidgetParent->CallbackValueB == 5, TEXT("The callback for value B was not triggered."));
-			AddErrorIfFalse(WidgetParent->CallbackValueC == 10, TEXT("The callback for value C was triggered."));
-			WidgetParent->SetAttributeC(6);
-			AddErrorIfFalse(WidgetParent->CallbackValueC == 6, TEXT("The callback for value C was not triggered."));
-		}
-
-		{
-			TSharedRef<UE::Slate::Private::SAttributeLeftWidget_OnInvalidationChild> WidgetChild = SNew(UE::Slate::Private::SAttributeLeftWidget_OnInvalidationChild);
-			WidgetChild->SetCallbackValue(10);
-			WidgetChild->CallbackValueB = 15;
-			WidgetChild->SetAttributeA(4);
-			AddErrorIfFalse(WidgetChild->CallbackValueA == 10, TEXT("The callback for value A was triggered."));
-			AddErrorIfFalse(WidgetChild->OverrideCallbackValueA == 4, TEXT("The callback for override value A was not triggered."));
-			AddErrorIfFalse(WidgetChild->CallbackValueB == 15, TEXT("The callback for value B was triggered."));
-			AddErrorIfFalse(WidgetChild->OverrideCallbackValueB == 10, TEXT("The callback for override value B was triggered."));
-			AddErrorIfFalse(WidgetChild->CallbackValueC == 10, TEXT("The callback for value C was triggered."));
-			AddErrorIfFalse(WidgetChild->OverrideCallbackValueC == 10, TEXT("The callback for override value C was triggered."));
-			WidgetChild->SetAttributeB(5);
-			AddErrorIfFalse(WidgetChild->CallbackValueB == 5, TEXT("The callback for value B was not triggered."));
-			AddErrorIfFalse(WidgetChild->OverrideCallbackValueB == 15, TEXT("The callback for override value B was not triggered."));
-			AddErrorIfFalse(WidgetChild->CallbackValueC == 10, TEXT("The callback for value C was triggered."));
-			AddErrorIfFalse(WidgetChild->OverrideCallbackValueC == 10, TEXT("The callback for override value C was triggered."));
-			WidgetChild->SetAttributeC(6);
-			AddErrorIfFalse(WidgetChild->CallbackValueC == 6, TEXT("The callback for value C was not triggered."));
-			AddErrorIfFalse(WidgetChild->OverrideCallbackValueC == 6, TEXT("The callback for override value C was not triggered."));
-		}
-		{
-			TSharedRef<UE::Slate::Private::SAttributeLeftWidget_OnInvalidationChild> WidgetChild = SNew(UE::Slate::Private::SAttributeLeftWidget_OnInvalidationChild);
-			WidgetChild->SetCallbackValue(10);
-			WidgetChild->CallbackValueB = 15;
-			WidgetChild->SetAttributeA(MakeAttributeLambda([](){ return 4; }));
-			WidgetChild->SetAttributeB(MakeAttributeLambda([](){ return 5; }));
-			WidgetChild->SetAttributeC(MakeAttributeLambda([](){ return 6; }));
-
-			AddErrorIfFalse(WidgetChild->CallbackValueA == 10, TEXT("The callback for value A was triggered."));
-			AddErrorIfFalse(WidgetChild->OverrideCallbackValueA == 10, TEXT("The callback for override value A was triggered."));
-			AddErrorIfFalse(WidgetChild->CallbackValueB == 15, TEXT("The callback for value B was triggered."));
-			AddErrorIfFalse(WidgetChild->OverrideCallbackValueB == 10, TEXT("The callback for override value B was triggered."));
-			AddErrorIfFalse(WidgetChild->CallbackValueC == 10, TEXT("The callback for value C was triggered."));
-			AddErrorIfFalse(WidgetChild->OverrideCallbackValueC == 10, TEXT("The callback for override value C was triggered."));
-
-			WidgetChild->MarkPrepassAsDirty();
-			WidgetChild->SlatePrepass(1.f);
-
-			AddErrorIfFalse(WidgetChild->CallbackValueA == 10, TEXT("The callback for value A was triggered."));
-			AddErrorIfFalse(WidgetChild->OverrideCallbackValueA == 4, TEXT("The callback for override value A was not triggered."));
-			AddErrorIfFalse(WidgetChild->CallbackValueB == 5, TEXT("The callback for value B was not triggered."));
-			AddErrorIfFalse(WidgetChild->OverrideCallbackValueB == 15, TEXT("The callback for override value B was not triggered."));
-			AddErrorIfFalse(WidgetChild->CallbackValueC == 6, TEXT("The callback for value C was not triggered."));
-			AddErrorIfFalse(WidgetChild->OverrideCallbackValueC == 6, TEXT("The callback for override value C was not triggered."));
-		}
-	}
-
 	return true;
 }
 
diff --git a/Engine/Source/Runtime/SlateCore/Private/Types/SlateAttributeDescriptor.cpp b/Engine/Source/Runtime/SlateCore/Private/Types/SlateAttributeDescriptor.cpp
--- a/Engine/Source/Runtime/SlateCore/Private/Types/SlateAttributeDescriptor.cpp
+++ b/Engine/Source/Runtime/SlateCore/Private/Types/SlateAttributeDescriptor.cpp
@@ -44,16 +44,6 @@ FSlateAttributeDescriptor::FInitializer::FAttributeEntry& FSlateAttributeDescrip
 }
 
 
-FSlateAttributeDescriptor::FInitializer::FAttributeEntry& FSlateAttributeDescriptor::FInitializer::FAttributeEntry::OnInvalidation(FInvalidationDelegate InCallback)
-{
-	if (Descriptor.Attributes.IsValidIndex(AttributeIndex))
-	{
-		Descriptor.Attributes[AttributeIndex].OnInvalidation = MoveTemp(InCallback);
-	}
-	return *this;
-}
-
-
 /** */
 FSlateAttributeDescriptor::FInitializer::FInitializer(FSlateAttributeDescriptor& InDescriptor)
 	: Descriptor(InDescriptor)
@@ -199,12 +189,6 @@ void FSlateAttributeDescriptor::FInitializer::OverrideInvalidationReason(FName A
 }
 
 
-void FSlateAttributeDescriptor::FInitializer::OverrideOnInvalidation(FName AttributeName, EInvalidationDelegateOverrideType OverrideType, FInvalidationDelegate Callback)
-{
-	Descriptor.OverrideOnInvalidation(AttributeName, OverrideType, MoveTemp(Callback));
-}
-
-
 void FSlateAttributeDescriptor::FInitializer::SetUpdateWhenCollapsed(FName AttributeName, bool bUpdateWhenCollapsed)
 {
 	FAttribute* Attribute = Descriptor.FindAttribute(AttributeName);
@@ -294,48 +278,6 @@ void FSlateAttributeDescriptor::OverrideInvalidationReason(FName AttributeName,
 }
 
 
-void FSlateAttributeDescriptor::OverrideOnInvalidation(FName AttributeName, EInvalidationDelegateOverrideType OverrideType, FInvalidationDelegate Callback)
-{
-	check(!AttributeName.IsNone());
-
-	FAttribute* FoundAttribute = FindAttribute(AttributeName);
-	if (ensureAlwaysMsgf(FoundAttribute != nullptr, TEXT("The attribute 's' doesn't exist."), *AttributeName.ToString()))
-	{
-		switch(OverrideType)
-		{
-		case EInvalidationDelegateOverrideType::ReplacePrevious:
-			FoundAttribute->OnInvalidation = MoveTemp(Callback);
-			break;
-		case EInvalidationDelegateOverrideType::ExecuteAfterPrevious:
-		case EInvalidationDelegateOverrideType::ExecuteBeforePrevious:
-			if (FoundAttribute->OnInvalidation.IsBound() && Callback.IsBound())
-			{
-				FInvalidationDelegate Previous = FoundAttribute->OnInvalidation;
-				FoundAttribute->OnInvalidation = FInvalidationDelegate::CreateLambda([Previous{MoveTemp(Previous)}, Callback{MoveTemp(Callback)}, OverrideType](SWidget& Widget)
-					{
-						if (OverrideType == EInvalidationDelegateOverrideType::ExecuteBeforePrevious)
-						{
-							Previous.ExecuteIfBound(Widget);
-						}
-						Callback.ExecuteIfBound(Widget);
-						if (OverrideType == EInvalidationDelegateOverrideType::ExecuteAfterPrevious)
-						{
-							Previous.ExecuteIfBound(Widget);
-						}
-					});
-			}
-			else if (Callback.IsBound())
-			{
-				FoundAttribute->OnInvalidation = MoveTemp(Callback);
-			}
-			break;
-		default:
-			check(false);
-		}
-	}
-}
-
-
 void FSlateAttributeDescriptor::SetPrerequisite(FSlateAttributeDescriptor::FAttribute& Attribute, FName Prerequisite, bool bSetAsDependency)
 {
 	if (Prerequisite.IsNone())
diff --git a/Engine/Source/Runtime/SlateCore/Private/Types/SlateAttributeMetaData.cpp b/Engine/Source/Runtime/SlateCore/Private/Types/SlateAttributeMetaData.cpp
--- a/Engine/Source/Runtime/SlateCore/Private/Types/SlateAttributeMetaData.cpp
+++ b/Engine/Source/Runtime/SlateCore/Private/Types/SlateAttributeMetaData.cpp
@@ -184,14 +184,14 @@ TArray<FName> FSlateAttributeMetaData::GetAttributeNames(const SWidget& OwningWi
 }
 
 
-FSlateAttributeMetaData::FGetterItem::FInvalidationDetail FSlateAttributeMetaData::FGetterItem::GetInvalidationDetail(const SWidget& OwningWidget, EInvalidateWidgetReason Reason) const
+EInvalidateWidgetReason FSlateAttributeMetaData::FGetterItem::GetInvalidationDetail(const SWidget& OwningWidget, EInvalidateWidgetReason Reason) const
 {
 	if (CachedAttributeDescriptorIndex != INDEX_NONE)
 	{
 		const FSlateAttributeDescriptor::FAttribute& DescriptorAttribute = OwningWidget.GetWidgetClass().GetAttributeDescriptor().GetAttributeAtIndex(CachedAttributeDescriptorIndex);
-		return FInvalidationDetail{&DescriptorAttribute.OnInvalidation, DescriptorAttribute.InvalidationReason.Get(OwningWidget)};
+		return DescriptorAttribute.InvalidationReason.Get(OwningWidget);
 	}
-	return FInvalidationDetail{nullptr, Reason};
+	return Reason;
 }
 
 
@@ -218,8 +218,6 @@ void FSlateAttributeMetaData::InvalidateWidget(SWidget& OwningWidget, const FSla
 		return;
 	}
 
-	const FSlateAttributeDescriptor::FInvalidationDelegate* OnInvalidationCallback = nullptr;
-
 	if (FSlateAttributeMetaData* AttributeMetaData = FSlateAttributeMetaData::FindMetaData(OwningWidget))
 	{
 		const int32 FoundIndex = AttributeMetaData->IndexOfAttribute(Attribute);
@@ -227,9 +225,7 @@ void FSlateAttributeMetaData::InvalidateWidget(SWidget& OwningWidget, const FSla
 		{
 			FGetterItem& GetterItem = AttributeMetaData->Attributes[FoundIndex];
 			{
-				const FGetterItem::FInvalidationDetail Detail = GetterItem.GetInvalidationDetail(OwningWidget, Reason);
-				OnInvalidationCallback = Detail.Get<0>();
-				Reason = Detail.Get<1>() | AttributeMetaData->CachedInvalidationReason;
+				Reason = GetterItem.GetInvalidationDetail(OwningWidget, Reason) | AttributeMetaData->CachedInvalidationReason;
 				AttributeMetaData->CachedInvalidationReason = EInvalidateWidgetReason::None;
 			}
 
@@ -247,7 +243,6 @@ void FSlateAttributeMetaData::InvalidateWidget(SWidget& OwningWidget, const FSla
 			const FSlateAttributeDescriptor::OffsetType Offset = Private::FindOffet(OwningWidget, Attribute);
 			if (FSlateAttributeDescriptor::FAttribute const* FoundAttribute = AttributeDescriptor.FindMemberAttribute(Offset))
 			{
-				OnInvalidationCallback = &FoundAttribute->OnInvalidation;
 				Reason = FoundAttribute->InvalidationReason.Get(OwningWidget) | AttributeMetaData->CachedInvalidationReason;
 				AttributeMetaData->CachedInvalidationReason = EInvalidateWidgetReason::None;
 
@@ -276,15 +271,10 @@ void FSlateAttributeMetaData::InvalidateWidget(SWidget& OwningWidget, const FSla
 		if (FSlateAttributeDescriptor::FAttribute const* FoundAttribute = OwningWidget.GetWidgetClass().GetAttributeDescriptor().FindMemberAttribute(Offset))
 		{
 			Reason = FoundAttribute->InvalidationReason.Get(OwningWidget);
-			OnInvalidationCallback = &FoundAttribute->OnInvalidation;
 		}
 	}
 
 	OwningWidget.Invalidate(Reason);
-	if (OnInvalidationCallback)
-	{
-		OnInvalidationCallback->ExecuteIfBound(OwningWidget);
-	}
 }
 
 
@@ -333,9 +323,6 @@ void FSlateAttributeMetaData::UpdateChildrenCollapsedAttributes(SWidget& OwningW
 
 void FSlateAttributeMetaData::UpdateAttributesImpl(SWidget& OwningWidget, EUpdateType UpdateType, EInvalidationPermission InvalidationStyle)
 {
-	// List of all the callback we need to execute after the invalidation
-	TArray<const FSlateAttributeDescriptor::FInvalidationDelegate*, TInlineAllocator<8>> AllInvalidationDelegates;
-
 	bool bInvalidateIfNeeded = (InvalidationStyle == EInvalidationPermission::AllowInvalidation) || (InvalidationStyle == EInvalidationPermission::AllowInvalidationIfConstructed && OwningWidget.IsConstructionCompleted());
 	bool bAllowInvalidation = bInvalidateIfNeeded || InvalidationStyle == EInvalidationPermission::DelayInvalidation;
 	EInvalidateWidgetReason InvalidationReason = EInvalidateWidgetReason::None;
@@ -384,12 +371,7 @@ void FSlateAttributeMetaData::UpdateAttributesImpl(SWidget& OwningWidget, EUpdat
 		GetterItem.bUpdatedThisFrame = Result.bInvalidationRequested;
 		if (Result.bInvalidationRequested && bAllowInvalidation)
 		{
-			const FGetterItem::FInvalidationDetail Detail = GetterItem.GetInvalidationDetail(OwningWidget, Result.InvalidationReason);
-			if (Detail.Get<0>())
-			{
-				AllInvalidationDelegates.Add(Detail.Get<0>());
-			}
-			InvalidationReason |= Detail.Get<1>();
+			InvalidationReason |= GetterItem.GetInvalidationDetail(OwningWidget, Result.InvalidationReason);
 		}
 	}
 
@@ -405,10 +387,6 @@ void FSlateAttributeMetaData::UpdateAttributesImpl(SWidget& OwningWidget, EUpdat
 	if (bInvalidateIfNeeded)
 	{
 		OwningWidget.Invalidate(InvalidationReason | CachedInvalidationReason);
-		for (const FSlateAttributeDescriptor::FInvalidationDelegate* OnInvalidation : AllInvalidationDelegates)
-		{
-			OnInvalidation->ExecuteIfBound(OwningWidget);
-		}
 		CachedInvalidationReason = EInvalidateWidgetReason::None;
 	}
 	else if (InvalidationStyle == EInvalidationPermission::DelayInvalidation)
@@ -437,12 +415,8 @@ void FSlateAttributeMetaData::UpdateAttribute(SWidget& OwningWidget, FSlateAttri
 			{
 				if (OwningWidget.IsConstructionCompleted())
 				{
-					const FGetterItem::FInvalidationDetail Detail = GetterItem.GetInvalidationDetail(OwningWidget, Result.InvalidationReason);
-					OwningWidget.Invalidate(Detail.Get<1>() | AttributeMetaData->CachedInvalidationReason);
-					if (Detail.Get<0>())
-					{
-						Detail.Get<0>()->ExecuteIfBound(OwningWidget);
-					}
+					EInvalidateWidgetReason Reason = GetterItem.GetInvalidationDetail(OwningWidget, Result.InvalidationReason);
+					OwningWidget.Invalidate(Reason | AttributeMetaData->CachedInvalidationReason);
 					AttributeMetaData->CachedInvalidationReason = EInvalidateWidgetReason::None;
 				}
 
diff --git a/Engine/Source/Runtime/SlateCore/Public/Types/SlateAttributeDescriptor.h b/Engine/Source/Runtime/SlateCore/Public/Types/SlateAttributeDescriptor.h
--- a/Engine/Source/Runtime/SlateCore/Public/Types/SlateAttributeDescriptor.h
+++ b/Engine/Source/Runtime/SlateCore/Public/Types/SlateAttributeDescriptor.h
@@ -66,9 +66,6 @@ public:
 		FGetter Getter;
 	};
 
-	/** */
-	DECLARE_DELEGATE_OneParam(FInvalidationDelegate, SWidget& /*Widget*/);
-
 	/** */
 	enum class EInvalidationDelegateOverrideType
 	{
@@ -94,7 +91,6 @@ public:
 		FName Prerequisite;
 		uint32 SortOrder = 0;
 		FInvalidateWidgetReasonAttribute InvalidationReason = FInvalidateWidgetReasonAttribute(EInvalidateWidgetReason::None);
-		FInvalidationDelegate OnInvalidation;
 		bool bIsMemberAttribute = false;
 		bool bIsPrerequisiteAlsoADependency = false;
 		bool bIsADependencyForSomeoneElse = false;
@@ -138,13 +134,6 @@ public:
 			 */
 			FAttributeEntry& UpdateWhenCollapsed();
 
-			/**
-			 * Be notified when the attribute value changed and the invalidation occurs.
-			 * It will not be called when the SWidget is in his construction phase.
-			 * @see SWidget::IsConstructionCompleted
-			 */
-			FAttributeEntry& OnInvalidation(FInvalidationDelegate Callback);
-
 		private:
 			FSlateAttributeDescriptor& Descriptor;
 			int32 AttributeIndex;
@@ -158,9 +147,6 @@ public:
 		/** Change the InvalidationReason of an attribute defined in a base class. */
 		void OverrideInvalidationReason(FName AttributeName, FInvalidateWidgetReasonAttribute&& Reason);
 
-		/** Change the FInvalidationDelegate of an attribute defined in a base class. */
-		void OverrideOnInvalidation(FName AttributeName, EInvalidationDelegateOverrideType OverrideType, FInvalidationDelegate Callback);
-
 		/** Change the update type of an attribute defined in a base class. */
 		void SetUpdateWhenCollapsed(FName AttributeName, bool bUpdateWhenCollapsed);
 
@@ -204,7 +190,6 @@ private:
 
 	FInitializer::FAttributeEntry AddMemberAttribute(FName AttributeName, OffsetType Offset, FInvalidateWidgetReasonAttribute ReasonGetter);
 	void OverrideInvalidationReason(FName AttributeName, FInvalidateWidgetReasonAttribute ReasonGetter);
-	void OverrideOnInvalidation(FName AttributeName, EInvalidationDelegateOverrideType OverrideType, FInvalidationDelegate Callback);
 	void SetPrerequisite(FAttribute& Attribute, FName Prerequisite, bool bSetAsDependency);
 	void SetUpdateWhenCollapsed(FAttribute& Attribute, bool bUpdate);
 
diff --git a/Engine/Source/Runtime/SlateCore/Public/Types/SlateAttributeMetaData.h b/Engine/Source/Runtime/SlateCore/Public/Types/SlateAttributeMetaData.h
--- a/Engine/Source/Runtime/SlateCore/Public/Types/SlateAttributeMetaData.h
+++ b/Engine/Source/Runtime/SlateCore/Public/Types/SlateAttributeMetaData.h
@@ -156,8 +156,7 @@ private:
 			return SortOrder < Other.SortOrder;
 		}
 
-		using FInvalidationDetail = TTuple<const FSlateAttributeDescriptor::FInvalidationDelegate*, EInvalidateWidgetReason>;
-		FInvalidationDetail GetInvalidationDetail(const SWidget&, EInvalidateWidgetReason Reason) const;
+		EInvalidateWidgetReason GetInvalidationDetail(const SWidget&, EInvalidateWidgetReason Reason) const;
 
 		/** If available, return the name of the attribute. */
 		FName GetAttributeName(const SWidget& OwningWidget) const;
