Add back in separate cvars for font atlas flush params. Each flushable atlas can now have their own params

diff --git a/Engine/Source/Runtime/SlateCore/Private/Fonts/FontCache.cpp b/Engine/Source/Runtime/SlateCore/Private/Fonts/FontCache.cpp
--- a/Engine/Source/Runtime/SlateCore/Private/Fonts/FontCache.cpp
+++ b/Engine/Source/Runtime/SlateCore/Private/Fonts/FontCache.cpp
@@ -19,6 +19,26 @@ DECLARE_DWORD_ACCUMULATOR_STAT(TEXT("Num Font Non-Atlased Textures"), STAT_Slate
 DECLARE_MEMORY_STAT(TEXT("Shaped Glyph Sequence Memory"), STAT_SlateShapedGlyphSequenceMemory, STATGROUP_SlateMemory);
 DEFINE_STAT(STAT_SlateFontMeasureCacheMemory);
 
+static FAtlasFlushParams FontAtlasFlushParams;
+FAutoConsoleVariableRef CVarMaxAtlasPagesBeforeFlush(
+	TEXT("Slate.MaxFontAtlasPagesBeforeFlush"),
+	FontAtlasFlushParams.InitialMaxAtlasPagesBeforeFlushRequest,
+	TEXT("The number of font atlas textures created and used before we flush the font cache if a texture atlas is full"));
+
+FAutoConsoleVariableRef CVarMaxFontNonAtlasPagesBeforeFlush(
+	TEXT("Slate.MaxFontNonAtlasTexturesBeforeFlush"),
+	FontAtlasFlushParams.InitialMaxNonAtlasPagesBeforeFlushRequest,
+	TEXT("The number of large glyph font textures initially."));
+
+FAutoConsoleVariableRef CVarGrowFontAtlasFrameWindow(
+	TEXT("Slate.GrowFontAtlasFrameWindow"),
+	FontAtlasFlushParams.GrowAtlasFrameWindow,
+	TEXT("The number of frames within the font atlas will resize rather than flush."));
+
+FAutoConsoleVariableRef CVarGrowFontNonAtlasFrameWindow(
+	TEXT("Slate.GrowFontNonAtlasFrameWindow"),
+	FontAtlasFlushParams.GrowNonAtlasFrameWindow,
+	TEXT("The number of frames within the large font glyph pool will resize rather than flush."));
 
 static int32 UnloadFreeTypeDataOnFlush = 1;
 FAutoConsoleVariableRef CVarUnloadFreeTypeDataOnFlush(
@@ -776,7 +796,8 @@ FCharacterEntry FCharacterList::MakeCharacterEntry(TCHAR Character, const FChara
 }
 
 FSlateFontCache::FSlateFontCache( TSharedRef<ISlateFontAtlasFactory> InFontAtlasFactory, ESlateTextureAtlasThreadId InOwningThread)
-	: FTLibrary( new FFreeTypeLibrary() )
+	: FSlateFlushableAtlasCache(&FontAtlasFlushParams)
+	, FTLibrary( new FFreeTypeLibrary() )
 	, FTCacheDirectory( new FFreeTypeCacheDirectory() )
 	, CompositeFontCache( new FCompositeFontCache( FTLibrary.Get() ) )
 	, FontRenderer( new FSlateFontRenderer( FTLibrary.Get(), FTCacheDirectory.Get(), CompositeFontCache.Get() ) )
diff --git a/Engine/Source/Runtime/SlateCore/Private/Rendering/SlateVectorGraphicsCache.cpp b/Engine/Source/Runtime/SlateCore/Private/Rendering/SlateVectorGraphicsCache.cpp
--- a/Engine/Source/Runtime/SlateCore/Private/Rendering/SlateVectorGraphicsCache.cpp
+++ b/Engine/Source/Runtime/SlateCore/Private/Rendering/SlateVectorGraphicsCache.cpp
@@ -6,8 +6,30 @@
 #include "ProfilingDebugging/ScopedTimers.h"
 #include "Async/ParallelFor.h"
 
+static FAtlasFlushParams SVGAtlasFlushParams;
+FAutoConsoleVariableRef CVarMaxSVGAtlasPagesBeforeFlush(
+	TEXT("Slate.MaxSVGAtlasPagesBeforeFlush"),
+	SVGAtlasFlushParams.InitialMaxAtlasPagesBeforeFlushRequest,
+	TEXT("The number of atlas textures created and used before we flush the cache if a texture atlas is full"));
+
+FAutoConsoleVariableRef CVarMaxSVGNonAtlasPagesBeforeFlush(
+	TEXT("Slate.MaxSVGNonAtlasTexturesBeforeFlush"),
+	SVGAtlasFlushParams.InitialMaxNonAtlasPagesBeforeFlushRequest,
+	TEXT("The number of large textures initially."));
+
+FAutoConsoleVariableRef CVarGrowSVGAtlasFrameWindow(
+	TEXT("Slate.GrowSVGAtlasFrameWindow"),
+	SVGAtlasFlushParams.GrowAtlasFrameWindow,
+	TEXT("The number of frames within the atlas will resize rather than flush."));
+
+FAutoConsoleVariableRef CVarGrowSVGNonAtlasFrameWindow(
+	TEXT("Slate.GrowSVGNonAtlasFrameWindow"),
+	SVGAtlasFlushParams.GrowNonAtlasFrameWindow,
+	TEXT("The number of frames within the large pool will resize rather than flush."));
+
 FSlateVectorGraphicsCache::FSlateVectorGraphicsCache(TSharedPtr<ISlateTextureAtlasFactory> InAtlasFactory, bool bInNeedRedBlueSwap)
-	: AtlasFactory(InAtlasFactory)
+	: FSlateFlushableAtlasCache(&SVGAtlasFlushParams)
+	, AtlasFactory(InAtlasFactory)
 	, bNeedRedBlueSwap(bInNeedRedBlueSwap)
 	, bFlushRequested(false)
 {
diff --git a/Engine/Source/Runtime/SlateCore/Private/Textures/TextureAtlas.cpp b/Engine/Source/Runtime/SlateCore/Private/Textures/TextureAtlas.cpp
--- a/Engine/Source/Runtime/SlateCore/Private/Textures/TextureAtlas.cpp
+++ b/Engine/Source/Runtime/SlateCore/Private/Textures/TextureAtlas.cpp
@@ -13,29 +13,6 @@ DEFINE_STAT(STAT_SlateTextureGPUMemory);
 DEFINE_STAT(STAT_SlateTextureDataMemory);
 DECLARE_MEMORY_STAT(TEXT("Texture Atlas Memory (CPU)"), STAT_SlateTextureAtlasMemory, STATGROUP_SlateMemory);
 
-static int32 GInitialMaxAtlasPagesBeforeFlushRequest = 1;
-FAutoConsoleVariableRef CVarMaxAtlasPagesBeforeFlush(
-	TEXT("Slate.MaxAtlasPagesBeforeFlush"),
-	GInitialMaxAtlasPagesBeforeFlushRequest,
-	TEXT("The number of atlas textures created and used before we flush the cache if a texture atlas is full"));
-
-static int32 GInitialMaxNonAtlasedTexturesBeforeFlushRequest = 1;
-FAutoConsoleVariableRef CVarMaxFontNonAtlasPagesBeforeFlush(
-	TEXT("Slate.MaxNonAtlasTexturesBeforeFlush"),
-	GInitialMaxNonAtlasedTexturesBeforeFlushRequest,
-	TEXT("The number of large textures initially."));
-
-static int32 GGrowAtlasFrameWindow = 1;
-FAutoConsoleVariableRef CVarGrowFontAtlasFrameWindow(
-	TEXT("Slate.GrowAtlasFrameWindow"),
-	GGrowAtlasFrameWindow,
-	TEXT("The number of frames within the atlas will resize rather than flush."));
-
-static int32 GGrowNonAtlasFrameWindow = 1;
-FAutoConsoleVariableRef CVarGrowFontNonAtlasFrameWindow(
-	TEXT("Slate.GrowNonAtlasFrameWindow"),
-	GGrowNonAtlasFrameWindow,
-	TEXT("The number of frames within the large pool will resize rather than flush."));
 
 ESlateTextureAtlasThreadId GetCurrentSlateTextureAtlasThreadId()
 {
@@ -390,38 +367,30 @@ const FAtlasedTextureSlot* FSlateTextureAtlas::FindSlotForTexture(uint32 InWidth
 
 
 
-FSlateFlushableAtlasCache::FSlateFlushableAtlasCache()
-	: InitialMaxAtlasPagesBeforeFlushRequest(GInitialMaxAtlasPagesBeforeFlushRequest)
-	, InitialMaxNonAtlasPagesBeforeFlushRequest(GInitialMaxNonAtlasedTexturesBeforeFlushRequest)
-	, CurrentMaxGrayscaleAtlasPagesBeforeFlushRequest(GInitialMaxAtlasPagesBeforeFlushRequest)
-	, CurrentMaxColorAtlasPagesBeforeFlushRequest(GInitialMaxAtlasPagesBeforeFlushRequest)
-	, CurrentMaxNonAtlasedTexturesBeforeFlushRequest(GInitialMaxNonAtlasedTexturesBeforeFlushRequest)
+FSlateFlushableAtlasCache::FSlateFlushableAtlasCache(const FAtlasFlushParams* InFlushParams)
+	: FlushParams(InFlushParams)
 {
-}
+	check(FlushParams);
 
-FSlateFlushableAtlasCache::FSlateFlushableAtlasCache(int32 InInitialMaxAtlasPagesBeforeFlushRequest, int32 InInitialMaxNonAtlasedTexturesBeforeFlushRequest)
-	: InitialMaxAtlasPagesBeforeFlushRequest(InInitialMaxAtlasPagesBeforeFlushRequest)
-	, InitialMaxNonAtlasPagesBeforeFlushRequest(InInitialMaxNonAtlasedTexturesBeforeFlushRequest)
-	, CurrentMaxGrayscaleAtlasPagesBeforeFlushRequest(InInitialMaxAtlasPagesBeforeFlushRequest)
-	, CurrentMaxColorAtlasPagesBeforeFlushRequest(InInitialMaxAtlasPagesBeforeFlushRequest)
-	, CurrentMaxNonAtlasedTexturesBeforeFlushRequest(InInitialMaxNonAtlasedTexturesBeforeFlushRequest)
-{
+	CurrentMaxGrayscaleAtlasPagesBeforeFlushRequest = FlushParams->InitialMaxAtlasPagesBeforeFlushRequest;
+	CurrentMaxColorAtlasPagesBeforeFlushRequest = FlushParams->InitialMaxAtlasPagesBeforeFlushRequest;
+	CurrentMaxNonAtlasedTexturesBeforeFlushRequest = FlushParams->InitialMaxNonAtlasPagesBeforeFlushRequest;
 }
 
 void FSlateFlushableAtlasCache::ResetFlushCounters()
 {
-	CurrentMaxGrayscaleAtlasPagesBeforeFlushRequest = InitialMaxAtlasPagesBeforeFlushRequest;
-	CurrentMaxColorAtlasPagesBeforeFlushRequest = InitialMaxAtlasPagesBeforeFlushRequest;
-	CurrentMaxNonAtlasedTexturesBeforeFlushRequest = InitialMaxNonAtlasPagesBeforeFlushRequest;
+	CurrentMaxGrayscaleAtlasPagesBeforeFlushRequest = FlushParams->InitialMaxAtlasPagesBeforeFlushRequest;
+	CurrentMaxColorAtlasPagesBeforeFlushRequest = FlushParams->InitialMaxAtlasPagesBeforeFlushRequest;
+	CurrentMaxNonAtlasedTexturesBeforeFlushRequest = FlushParams->InitialMaxNonAtlasPagesBeforeFlushRequest;
 	FrameCounterLastFlushRequest = GFrameCounter;
 }
 
 void FSlateFlushableAtlasCache::UpdateFlushCounters(int32 NumGrayscale, int32 NumColor, int32 NumNonAtlased)
 {
 	bool bFlushRequested = false;
-	bFlushRequested |= UpdateInternal(NumGrayscale, CurrentMaxGrayscaleAtlasPagesBeforeFlushRequest, InitialMaxAtlasPagesBeforeFlushRequest, GGrowAtlasFrameWindow);
-	bFlushRequested |= UpdateInternal(NumColor, CurrentMaxColorAtlasPagesBeforeFlushRequest, InitialMaxAtlasPagesBeforeFlushRequest, GGrowAtlasFrameWindow);
-	bFlushRequested |= UpdateInternal(NumNonAtlased, CurrentMaxNonAtlasedTexturesBeforeFlushRequest, InitialMaxNonAtlasPagesBeforeFlushRequest, GGrowNonAtlasFrameWindow);
+	bFlushRequested |= UpdateInternal(NumGrayscale, CurrentMaxGrayscaleAtlasPagesBeforeFlushRequest, FlushParams->InitialMaxAtlasPagesBeforeFlushRequest, FlushParams->GrowAtlasFrameWindow);
+	bFlushRequested |= UpdateInternal(NumColor, CurrentMaxColorAtlasPagesBeforeFlushRequest, FlushParams->InitialMaxAtlasPagesBeforeFlushRequest, FlushParams->GrowAtlasFrameWindow);
+	bFlushRequested |= UpdateInternal(NumNonAtlased, CurrentMaxNonAtlasedTexturesBeforeFlushRequest, FlushParams->InitialMaxNonAtlasPagesBeforeFlushRequest, FlushParams->GrowNonAtlasFrameWindow);
 
 	if (bFlushRequested)
 	{
diff --git a/Engine/Source/Runtime/SlateCore/Public/Textures/TextureAtlas.h b/Engine/Source/Runtime/SlateCore/Public/Textures/TextureAtlas.h
--- a/Engine/Source/Runtime/SlateCore/Public/Textures/TextureAtlas.h
+++ b/Engine/Source/Runtime/SlateCore/Public/Textures/TextureAtlas.h
@@ -256,13 +256,20 @@ public:
 	virtual void ReleaseTextureAtlases(const TArray<TUniquePtr<FSlateTextureAtlas>>& InTextureAtlases, const TArray<TUniquePtr<FSlateShaderResource>>& InNonAtlasedTextures, const bool bWaitForRelease) const = 0;
 };
 
+/** Parameters for flushable atlases that dictate when the atlas is allowed to flush after it becomes full */
+struct FAtlasFlushParams
+{
+	int32 InitialMaxAtlasPagesBeforeFlushRequest = 1;
+	int32 InitialMaxNonAtlasPagesBeforeFlushRequest = 1;
+	int32 GrowAtlasFrameWindow = 1;
+	int32 GrowNonAtlasFrameWindow = 1;
+};
 
 /** Base class for any atlas cache which has flushing logic to keep the number of in use pages small */
 class FSlateFlushableAtlasCache
 {
 public:
-	FSlateFlushableAtlasCache();
-	FSlateFlushableAtlasCache(int32 InInitialMaxAtlasPagesBeforeFlushRequest, int32 InInitialMaxNonAtlasedTexturesBeforeFlushRequest);
+	FSlateFlushableAtlasCache(const FAtlasFlushParams* InFlushParams);
 
 	virtual ~FSlateFlushableAtlasCache() {}
 
@@ -282,8 +289,8 @@ public:
 private:
 	bool UpdateInternal(int32 CurrentNum, int32& MaxNum, int32 InitialMax, int32 FrameWindowNum);
 private:
-	int32 InitialMaxAtlasPagesBeforeFlushRequest;
-	int32 InitialMaxNonAtlasPagesBeforeFlushRequest;
+	/** Flush params that dictate when this atlas can flush.*/
+	const FAtlasFlushParams* FlushParams;
 
 	/** Number of grayscale atlas pages we can have before we request that the cache be flushed */
 	int32 CurrentMaxGrayscaleAtlasPagesBeforeFlushRequest;
