Slate: Rename "Collapsed" SlateAttribute to "AffectVisibility". It represent more the intention. If the parent is collapsed, the child SlateAttribute was not updated even if they were marked as "update when collapsed". Update the widget reflector to show the name of the SlateAttribute (only if the widget is live). Call GetRelativeLayoutScale after the attribute are updated.

diff --git a/Engine/Source/Developer/SlateReflector/Private/Models/WidgetReflectorNode.cpp b/Engine/Source/Developer/SlateReflector/Private/Models/WidgetReflectorNode.cpp
--- a/Engine/Source/Developer/SlateReflector/Private/Models/WidgetReflectorNode.cpp
+++ b/Engine/Source/Developer/SlateReflector/Private/Models/WidgetReflectorNode.cpp
@@ -949,7 +949,7 @@ int32 FWidgetReflectorNodeUtils::GetWidgetAttributeCount(const TSharedPtr<const
 	{
 		if (FSlateAttributeMetaData* MetaData = FSlateAttributeMetaData::FindMetaData(*InWidget.Get()))
 		{
-			return MetaData->RegisteredAttributeCount();
+			return MetaData->GetRegisteredAttributeCount();
 		}
 	}
 	return 0;
@@ -961,7 +961,7 @@ int32 FWidgetReflectorNodeUtils::GetWidgetCollapsedAttributeCount(const TSharedP
 	{
 		if (FSlateAttributeMetaData* MetaData = FSlateAttributeMetaData::FindMetaData(*InWidget.Get()))
 		{
-			return MetaData->RegisteredCollaspedAttributeCount();
+			return MetaData->GetRegisteredAffectVisibilityAttributeCount();
 		}
 	}
 	return 0;
diff --git a/Engine/Source/Developer/SlateReflector/Private/Widgets/SWidgetReflectorToolTipWidget.cpp b/Engine/Source/Developer/SlateReflector/Private/Widgets/SWidgetReflectorToolTipWidget.cpp
--- a/Engine/Source/Developer/SlateReflector/Private/Widgets/SWidgetReflectorToolTipWidget.cpp
+++ b/Engine/Source/Developer/SlateReflector/Private/Widgets/SWidgetReflectorToolTipWidget.cpp
@@ -4,6 +4,7 @@
 
 #include "Layout/Geometry.h"
 #include "Styling/CoreStyle.h"
+#include "Types/SlateAttributeMetaData.h"
 #include "Widgets/SBoxPanel.h"
 #include "Widgets/Layout/SBorder.h"
 #include "Widgets/Layout/SGridPanel.h"
@@ -76,8 +77,15 @@ void SReflectorToolTipWidget::Construct( const FArguments& InArgs )
 		BuildLabelAndValue(LOCTEXT("IsVolatile", "Is Volatile"), { this, &SReflectorToolTipWidget::GetIsVolatile });
 		BuildLabelAndValue(LOCTEXT("IsVolatileIndirectly", "Is Volatile Indirectly"), { this, &SReflectorToolTipWidget::GetIsVolatileIndirectly });
 		BuildLabelAndValue(LOCTEXT("HasActiveTimers", "Has Active Timer"), { this, &SReflectorToolTipWidget::GetHasActiveTimers });
-		BuildLabelAndValue(LOCTEXT("HasRegisteredAttribute", "Attribute Count"), { this, &SReflectorToolTipWidget::GetAttributeCount });
-		BuildLabelAndValue(LOCTEXT("HasCollapsedAttribute", "Collapsed Attribute Count"), { this, &SReflectorToolTipWidget::GetCollapsedAttributeCount });
+		if (WidgetInfo->GetNodeType() == EWidgetReflectorNodeType::Live)
+		{
+			BuildLabelAndValue(LOCTEXT("HasRegisteredAttribute", "Attributes"), { this, &SReflectorToolTipWidget::GetLiveAttributeName });
+		}
+		else
+		{
+			BuildLabelAndValue(LOCTEXT("HasRegisteredAttribute", "Attributes"), { this, &SReflectorToolTipWidget::GetAttributeCount });
+			BuildLabelAndValue(LOCTEXT("HasCollapsedAttribute", "Affect Visibility Attributes"), { this, &SReflectorToolTipWidget::GetCollapsedAttributeCount });
+		}
 
 		if (bIsInsideInvalidationRoot)
 		{
@@ -119,4 +127,27 @@ void SReflectorToolTipWidget::Construct( const FArguments& InArgs )
 	}
 }
 
+FText SReflectorToolTipWidget::GetLiveAttributeName() const
+{
+	if (TSharedPtr<SWidget> Widget = WidgetInfo->GetLiveWidget())
+	{
+		if (FSlateAttributeMetaData* AttributeMetaData = FSlateAttributeMetaData::FindMetaData(*Widget.Get()))
+		{
+			TArray<FName> AttributeNames = AttributeMetaData->GetAttributeNames(*Widget.Get());
+			TArray<FText> AttributeNamesAsText;
+			AttributeNamesAsText.Reserve(AttributeNames.Num());
+			for (FName AttributeName : AttributeNames)
+			{
+				AttributeNamesAsText.Add(FText::FromName(AttributeName));
+			}
+			return FText::FormatOrdered(LOCTEXT("LiveAttributeNames", "{0} ({1} affect widget visibility) \n  {2}")
+				, AttributeMetaData->GetRegisteredAttributeCount()
+				, AttributeMetaData->GetRegisteredAffectVisibilityAttributeCount()
+				, FText::Join(LOCTEXT("LiveAttributeNamesDelimiter", "\n  "), AttributeNamesAsText));
+		}
+	}
+
+	return FText::AsNumber(0);
+}
+
 #undef LOCTEXT_NAMESPACE
diff --git a/Engine/Source/Developer/SlateReflector/Private/Widgets/SWidgetReflectorToolTipWidget.h b/Engine/Source/Developer/SlateReflector/Private/Widgets/SWidgetReflectorToolTipWidget.h
--- a/Engine/Source/Developer/SlateReflector/Private/Widgets/SWidgetReflectorToolTipWidget.h
+++ b/Engine/Source/Developer/SlateReflector/Private/Widgets/SWidgetReflectorToolTipWidget.h
@@ -69,6 +69,8 @@ private:
 		return (WidgetInfo->GetWidgetHasActiveTimers()) ? TrueText : FalseText;
 	}
 
+	FText GetLiveAttributeName() const;
+
 	FText GetAttributeCount() const
 	{
 		return FText::AsNumber(WidgetInfo->GetWidgetAttributeCount());
diff --git a/Engine/Source/Runtime/SlateCore/Private/Application/SlateWindowHelper.cpp b/Engine/Source/Runtime/SlateCore/Private/Application/SlateWindowHelper.cpp
--- a/Engine/Source/Runtime/SlateCore/Private/Application/SlateWindowHelper.cpp
+++ b/Engine/Source/Runtime/SlateCore/Private/Application/SlateWindowHelper.cpp
@@ -103,7 +103,7 @@ bool FSlateWindowHelper::FindPathToWidget( const TArray<TSharedRef<SWindow>>& Wi
 		while (true)
 		{
 			// Update the Widgets visibility before getting the ArrangeChildren
-			FSlateAttributeMetaData::UpdateCollapsedAttributes(*CurWidget.Get(), FSlateAttributeMetaData::EInvalidationPermission::DelayInvalidation);
+			FSlateAttributeMetaData::UpdateOnlyVisibilityAttributes(*CurWidget.Get(), FSlateAttributeMetaData::EInvalidationPermission::DelayInvalidation);
 			EVisibility CurWidgetVisibility = CurWidget->GetVisibility();
 			if (OutWidgetPath.Widgets.Accepts(CurWidgetVisibility))
 			{
diff --git a/Engine/Source/Runtime/SlateCore/Private/Debugging/WidgetList.cpp b/Engine/Source/Runtime/SlateCore/Private/Debugging/WidgetList.cpp
--- a/Engine/Source/Runtime/SlateCore/Private/Debugging/WidgetList.cpp
+++ b/Engine/Source/Runtime/SlateCore/Private/Debugging/WidgetList.cpp
@@ -255,9 +255,9 @@ void LogAllWidgetsDebugInfoImpl(FOutputDevice& Ar, const FLogAllWidgetsDebugInfo
 			if (FSlateAttributeMetaData* MetaData = FSlateAttributeMetaData::FindMetaData(*Widget))
 			{
 				MessageBuilder << TEXT(";");
-				MessageBuilder << MetaData->RegisteredAttributeCount();
+				MessageBuilder << MetaData->GetRegisteredAttributeCount();
 				MessageBuilder << TEXT(";");
-				MessageBuilder << MetaData->RegisteredCollaspedAttributeCount();
+				MessageBuilder << MetaData->GetRegisteredAffectVisibilityAttributeCount();
 			}
 			else
 			{
diff --git a/Engine/Source/Runtime/SlateCore/Private/FastUpdate/SlateInvalidationRoot.cpp b/Engine/Source/Runtime/SlateCore/Private/FastUpdate/SlateInvalidationRoot.cpp
--- a/Engine/Source/Runtime/SlateCore/Private/FastUpdate/SlateInvalidationRoot.cpp
+++ b/Engine/Source/Runtime/SlateCore/Private/FastUpdate/SlateInvalidationRoot.cpp
@@ -649,14 +649,14 @@ void FSlateInvalidationRoot::ProcessPreUpdate()
 						if (!InvalidationWidget.Visibility.IsCollapseIndirectly())
 						{
 							// if my parent is not collapse, then update my visible state
-							FSlateAttributeMetaData::UpdateCollapsedAttributes(*WidgetPtr, FSlateAttributeMetaData::EInvalidationPermission::AllowInvalidation);
+							FSlateAttributeMetaData::UpdateOnlyVisibilityAttributes(*WidgetPtr, FSlateAttributeMetaData::EInvalidationPermission::AllowInvalidation);
 							if (!InvalidationWidget.Visibility.IsCollapsed())
 							{
 #if UE_SLATE_WITH_INVALIDATIONWIDGETLIST_DEBUGGING
 								ensureAlwaysMsgf(!GSlateInvalidationRootVerifySlateAttribute || InvalidationWidget.bDebug_AttributeUpdated == false, TEXT("Attribute should only be updated once per frame."));
 								InvalidationWidget.bDebug_AttributeUpdated = true;
 #endif
-								FSlateAttributeMetaData::UpdateExpandedAttributes(*WidgetPtr, FSlateAttributeMetaData::EInvalidationPermission::AllowInvalidation);
+								FSlateAttributeMetaData::UpdateExceptVisibilityAttributes(*WidgetPtr, FSlateAttributeMetaData::EInvalidationPermission::AllowInvalidation);
 								AttributeItt.Advance();
 							}
 							else
diff --git a/Engine/Source/Runtime/SlateCore/Private/FastUpdate/SlateInvalidationWidgetList.cpp b/Engine/Source/Runtime/SlateCore/Private/FastUpdate/SlateInvalidationWidgetList.cpp
--- a/Engine/Source/Runtime/SlateCore/Private/FastUpdate/SlateInvalidationWidgetList.cpp
+++ b/Engine/Source/Runtime/SlateCore/Private/FastUpdate/SlateInvalidationWidgetList.cpp
@@ -321,12 +321,12 @@ FSlateInvalidationWidgetIndex FSlateInvalidationWidgetList::Internal_BuildWidget
 		{
 			if (!NewVisibility.IsCollapseIndirectly())
 			{
-				FSlateAttributeMetaData::UpdateCollapsedAttributes(Widget.Get(), FSlateAttributeMetaData::EInvalidationPermission::DenyAndClearDelayedInvalidation);
+				FSlateAttributeMetaData::UpdateOnlyVisibilityAttributes(Widget.Get(), FSlateAttributeMetaData::EInvalidationPermission::DenyAndClearDelayedInvalidation);
 				NewVisibility.SetVisibility(ParentVisibility, Widget->GetVisibility());
 
 				if (!NewVisibility.IsCollapsed())
 				{
-					FSlateAttributeMetaData::UpdateExpandedAttributes(Widget.Get(), FSlateAttributeMetaData::EInvalidationPermission::DenyAndClearDelayedInvalidation);
+					FSlateAttributeMetaData::UpdateExceptVisibilityAttributes(Widget.Get(), FSlateAttributeMetaData::EInvalidationPermission::DenyAndClearDelayedInvalidation);
 					NewVisibility.SetVisibility(ParentVisibility, Widget->GetVisibility());
 					bUpdateSlateAttribute = true;
 				}
diff --git a/Engine/Source/Runtime/SlateCore/Private/Layout/WidgetPath.cpp b/Engine/Source/Runtime/SlateCore/Private/Layout/WidgetPath.cpp
--- a/Engine/Source/Runtime/SlateCore/Private/Layout/WidgetPath.cpp
+++ b/Engine/Source/Runtime/SlateCore/Private/Layout/WidgetPath.cpp
@@ -312,7 +312,7 @@ FWeakWidgetPath::EPathResolutionResult::Result FWeakWidgetPath::ToWidgetPath( FW
 				bool bFoundChild = false;
 				if (CurWidget.IsValid())
 				{
-					FSlateAttributeMetaData::UpdateCollapsedAttributes(*CurWidget.Get(), FSlateAttributeMetaData::EInvalidationPermission::DelayInvalidation);
+					FSlateAttributeMetaData::UpdateOnlyVisibilityAttributes(*CurWidget.Get(), FSlateAttributeMetaData::EInvalidationPermission::DelayInvalidation);
 					if (EVisibility::DoesVisibilityPassFilter(CurWidget->GetVisibility(), VisibilityFilter))
 					{
 						TSharedRef<SWidget> CurWidgetRef = CurWidget.ToSharedRef();
diff --git a/Engine/Source/Runtime/SlateCore/Private/Test/SlateAttributeTest.cpp b/Engine/Source/Runtime/SlateCore/Private/Test/SlateAttributeTest.cpp
--- a/Engine/Source/Runtime/SlateCore/Private/Test/SlateAttributeTest.cpp
+++ b/Engine/Source/Runtime/SlateCore/Private/Test/SlateAttributeTest.cpp
@@ -130,6 +130,7 @@ public:
 		, IntAttributeJ(*this, 99)
 		, IntAttributeK(*this, 99)
 		, IntAttributeL(*this, 99)
+		, IntAttributeM(*this, 99)
 	{
 	}
 
@@ -140,12 +141,13 @@ public:
 	TSlateAttribute<int32> IntAttributeJ;
 	TSlateAttribute<int32> IntAttributeK;
 	TSlateAttribute<int32> IntAttributeL;
+	TSlateAttribute<int32> IntAttributeM;
 };
 
 SLATE_IMPLEMENT_WIDGET(SAttributeLeftWidget_Child)
 void SAttributeLeftWidget_Child::PrivateRegisterAttributes(FSlateAttributeInitializer& AttributeInitializer)
 {
-	// The update order is B, A, I, J, D, C, L, H, K
+	// The update order is M, B, A, I, J, D, C, L, H, K
 	//SLATE_ADD_MEMBER_ATTRIBUTE_DEFINITION(AttributeInitializer, IntAttributeH, EInvalidateWidgetReason::ChildOrder);
 	SLATE_ADD_MEMBER_ATTRIBUTE_DEFINITION(AttributeInitializer, IntAttributeJ, EInvalidateWidgetReason::ChildOrder)
 		.UpdateDependency("IntAttributeA");
@@ -154,6 +156,9 @@ void SAttributeLeftWidget_Child::PrivateRegisterAttributes(FSlateAttributeInitia
 		.UpdatePrerequisite("IntAttributeB");
 	SLATE_ADD_MEMBER_ATTRIBUTE_DEFINITION(AttributeInitializer, IntAttributeL, EInvalidateWidgetReason::ChildOrder)
 		.UpdatePrerequisite("IntAttributeC");
+	SLATE_ADD_MEMBER_ATTRIBUTE_DEFINITION(AttributeInitializer, IntAttributeM, EInvalidateWidgetReason::ChildOrder)
+		.UpdatePrerequisite("Visibility")
+		.AffectVisibility();
 }
 
 class SAttributeLeftWidget_OnInvalidationParent : public SLeafWidget
@@ -197,7 +202,7 @@ public:
 
 bool FSlateAttributeTest::RunTest(const FString& Parameters)
 {
-	const int NumberOfAttributeInSWidget = 2;
+	const int NumberOfAttributeInSWidget = 4;
 	int32 OrderCounter = 0;
 	auto OrderLambda = [this, &OrderCounter]() -> int32
 	{
@@ -335,7 +340,7 @@ bool FSlateAttributeTest::RunTest(const FString& Parameters)
 			, TEXT("The static data do not matches"));
 
 		FSlateAttributeDescriptor const& AttributeDescriptor = WidgetChild->GetWidgetClass().GetAttributeDescriptor();
-		AddErrorIfFalse(AttributeDescriptor.AttributeNum() == 8 + NumberOfAttributeInSWidget, TEXT("")); // H is not counted
+		AddErrorIfFalse(AttributeDescriptor.AttributeNum() == 9 + NumberOfAttributeInSWidget, TEXT("")); // H is not counted
 
 		const int32 IndexA = AttributeDescriptor.IndexOfMemberAttribute("IntAttributeA");
 		const int32 IndexB = AttributeDescriptor.IndexOfMemberAttribute("IntAttributeB");
@@ -345,6 +350,7 @@ bool FSlateAttributeTest::RunTest(const FString& Parameters)
 		const int32 IndexJ = AttributeDescriptor.IndexOfMemberAttribute("IntAttributeJ");
 		const int32 IndexK = AttributeDescriptor.IndexOfMemberAttribute("IntAttributeK");
 		const int32 IndexL = AttributeDescriptor.IndexOfMemberAttribute("IntAttributeL");
+		const int32 IndexM = AttributeDescriptor.IndexOfMemberAttribute("IntAttributeM");
 
 		AddErrorIfFalse(IndexA != INDEX_NONE, TEXT("Could not find the Attribute A"));
 		AddErrorIfFalse(IndexB != INDEX_NONE, TEXT("Could not find the Attribute B"));
@@ -354,6 +360,7 @@ bool FSlateAttributeTest::RunTest(const FString& Parameters)
 		AddErrorIfFalse(IndexJ != INDEX_NONE, TEXT("Could not find the Attribute J"));
 		AddErrorIfFalse(IndexK != INDEX_NONE, TEXT("Could not find the Attribute K"));
 		AddErrorIfFalse(IndexL != INDEX_NONE, TEXT("Could not find the Attribute L"));
+		AddErrorIfFalse(IndexM != INDEX_NONE, TEXT("Could not find the Attribute M"));
 
 		AddErrorIfFalse(&AttributeDescriptor.GetAttributeAtIndex(IndexA) == AttributeDescriptor.FindAttribute("IntAttributeA"), TEXT("Index and Attribute should return the same value."));
 		AddErrorIfFalse(&AttributeDescriptor.GetAttributeAtIndex(IndexB) == AttributeDescriptor.FindAttribute("IntAttributeB"), TEXT("Index and Attribute should return the same value."));
@@ -363,8 +370,10 @@ bool FSlateAttributeTest::RunTest(const FString& Parameters)
 		AddErrorIfFalse(&AttributeDescriptor.GetAttributeAtIndex(IndexJ) == AttributeDescriptor.FindAttribute("IntAttributeJ"), TEXT("Index and Attribute should return the same value."));
 		AddErrorIfFalse(&AttributeDescriptor.GetAttributeAtIndex(IndexK) == AttributeDescriptor.FindAttribute("IntAttributeK"), TEXT("Index and Attribute should return the same value."));
 		AddErrorIfFalse(&AttributeDescriptor.GetAttributeAtIndex(IndexL) == AttributeDescriptor.FindAttribute("IntAttributeL"), TEXT("Index and Attribute should return the same value."));
+		AddErrorIfFalse(&AttributeDescriptor.GetAttributeAtIndex(IndexM) == AttributeDescriptor.FindAttribute("IntAttributeM"), TEXT("Index and Attribute should return the same value."));
 		AddErrorIfFalse(AttributeDescriptor.FindAttribute("IntAttributeH") == nullptr, TEXT("H exist but is not defined."));
 
+		AddErrorIfFalse(AttributeDescriptor.GetAttributeAtIndex(IndexM).SortOrder < AttributeDescriptor.GetAttributeAtIndex(IndexB).SortOrder, TEXT("M should have a lower value than B"));
 		AddErrorIfFalse(AttributeDescriptor.GetAttributeAtIndex(IndexB).SortOrder < AttributeDescriptor.GetAttributeAtIndex(IndexA).SortOrder, TEXT("B should have a lower value than A"));
 		AddErrorIfFalse(AttributeDescriptor.GetAttributeAtIndex(IndexA).SortOrder <= AttributeDescriptor.GetAttributeAtIndex(IndexI).SortOrder, TEXT("A should have a lower value than I"));
 		AddErrorIfFalse(AttributeDescriptor.GetAttributeAtIndex(IndexI).SortOrder < AttributeDescriptor.GetAttributeAtIndex(IndexJ).SortOrder, TEXT("I should have a lower value than J"));
@@ -394,6 +403,8 @@ bool FSlateAttributeTest::RunTest(const FString& Parameters)
 			AddErrorIfFalse(WidgetChild->IntAttributeK.Get() == 99, TEXT("K It is not the expected value."));
 			WidgetChild->IntAttributeL.Assign(WidgetChild.Get(), MakeAttributeLambda(OrderLambda));
 			AddErrorIfFalse(WidgetChild->IntAttributeL.Get() == 99, TEXT("L It is not the expected value."));
+			WidgetChild->IntAttributeM.Assign(WidgetChild.Get(), MakeAttributeLambda(OrderLambda));
+			AddErrorIfFalse(WidgetChild->IntAttributeM.Get() == 99, TEXT("L It is not the expected value."));
 
 
 			OrderCounter = 0;
@@ -401,15 +412,16 @@ bool FSlateAttributeTest::RunTest(const FString& Parameters)
 			ReturnValue = 4;
 			WidgetChild->MarkPrepassAsDirty();
 			WidgetChild->SlatePrepass(1.f);
-			AddErrorIfFalse(WidgetChild->IntAttributeA.Get() == 2 || WidgetChild->IntAttributeA.Get() == 3, TEXT("A It is not the expected value."));
-			AddErrorIfFalse(WidgetChild->IntAttributeB.Get() == 1, TEXT("B It is not the expected value."));
-			AddErrorIfFalse(WidgetChild->IntAttributeC.Get() == 6, TEXT("C It is not the expected value."));
-			AddErrorIfFalse(WidgetChild->IntAttributeD.Get() == 5, TEXT("D It is not the expected value."));
-			AddErrorIfFalse(WidgetChild->IntAttributeH.Get() == 8, TEXT("H It is not the expected value."));
-			AddErrorIfFalse(WidgetChild->IntAttributeI.Get() == 2 || WidgetChild->IntAttributeI.Get() == 3, TEXT("I It is not the expected value."));
-			AddErrorIfFalse(WidgetChild->IntAttributeJ.Get() == 4, TEXT("J It is not the expected value."));
-			AddErrorIfFalse(WidgetChild->IntAttributeK.Get() == 9, TEXT("K It is not the expected value."));
-			AddErrorIfFalse(WidgetChild->IntAttributeL.Get() == 7, TEXT("L It is not the expected value."));
+			AddErrorIfFalse(WidgetChild->IntAttributeA.Get() == 3 || WidgetChild->IntAttributeA.Get() == 4, TEXT("A It is not the expected value."));
+			AddErrorIfFalse(WidgetChild->IntAttributeB.Get() == 2, TEXT("B It is not the expected value."));
+			AddErrorIfFalse(WidgetChild->IntAttributeC.Get() == 7, TEXT("C It is not the expected value."));
+			AddErrorIfFalse(WidgetChild->IntAttributeD.Get() == 6, TEXT("D It is not the expected value."));
+			AddErrorIfFalse(WidgetChild->IntAttributeH.Get() == 9, TEXT("H It is not the expected value."));
+			AddErrorIfFalse(WidgetChild->IntAttributeI.Get() == 3 || WidgetChild->IntAttributeI.Get() == 4, TEXT("I It is not the expected value."));
+			AddErrorIfFalse(WidgetChild->IntAttributeJ.Get() == 5, TEXT("J It is not the expected value."));
+			AddErrorIfFalse(WidgetChild->IntAttributeK.Get() == 10, TEXT("K It is not the expected value."));
+			AddErrorIfFalse(WidgetChild->IntAttributeL.Get() == 8, TEXT("L It is not the expected value."));
+			AddErrorIfFalse(WidgetChild->IntAttributeM.Get() == 1, TEXT("L It is not the expected value."));
 		}
 
 		{
@@ -418,15 +430,16 @@ bool FSlateAttributeTest::RunTest(const FString& Parameters)
 			ReturnValue = 4;
 			WidgetChild->MarkPrepassAsDirty();
 			WidgetChild->SlatePrepass(1.f);
-			AddErrorIfFalse(WidgetChild->IntAttributeA.Get() == 2 || WidgetChild->IntAttributeA.Get() == 3, TEXT("A It is not the expected value."));
-			AddErrorIfFalse(WidgetChild->IntAttributeB.Get() == 1, TEXT("B It is not the expected value."));
-			AddErrorIfFalse(WidgetChild->IntAttributeC.Get() == 5, TEXT("C It is not the expected value.")); // will get updated because D changes
-			AddErrorIfFalse(WidgetChild->IntAttributeD.Get() == 4, TEXT("D It is not the expected value."));
-			AddErrorIfFalse(WidgetChild->IntAttributeH.Get() == 7, TEXT("H It is not the expected value."));
-			AddErrorIfFalse(WidgetChild->IntAttributeI.Get() == 2 || WidgetChild->IntAttributeI.Get() == 3, TEXT("I It is not the expected value."));
-			AddErrorIfFalse(WidgetChild->IntAttributeJ.Get() == 4, TEXT("J It is not the expected value.")); // should not get updated
-			AddErrorIfFalse(WidgetChild->IntAttributeK.Get() == 8, TEXT("K It is not the expected value."));
-			AddErrorIfFalse(WidgetChild->IntAttributeL.Get() == 6, TEXT("L It is not the expected value."));
+			AddErrorIfFalse(WidgetChild->IntAttributeA.Get() == 3 || WidgetChild->IntAttributeA.Get() == 4, TEXT("A It is not the expected value."));
+			AddErrorIfFalse(WidgetChild->IntAttributeB.Get() == 2, TEXT("B It is not the expected value."));
+			AddErrorIfFalse(WidgetChild->IntAttributeC.Get() == 6, TEXT("C It is not the expected value.")); // will get updated because D changes
+			AddErrorIfFalse(WidgetChild->IntAttributeD.Get() == 5, TEXT("D It is not the expected value."));
+			AddErrorIfFalse(WidgetChild->IntAttributeH.Get() == 8, TEXT("H It is not the expected value."));
+			AddErrorIfFalse(WidgetChild->IntAttributeI.Get() == 3 || WidgetChild->IntAttributeI.Get() == 4, TEXT("I It is not the expected value."));
+			AddErrorIfFalse(WidgetChild->IntAttributeJ.Get() == 5, TEXT("J It is not the expected value.")); // should not get updated
+			AddErrorIfFalse(WidgetChild->IntAttributeK.Get() == 9, TEXT("K It is not the expected value."));
+			AddErrorIfFalse(WidgetChild->IntAttributeL.Get() == 7, TEXT("L It is not the expected value."));
+			AddErrorIfFalse(WidgetChild->IntAttributeM.Get() == 1, TEXT("M It is not the expected value."));
 		}
 	}
 
diff --git a/Engine/Source/Runtime/SlateCore/Private/Types/SlateAttributeDescriptor.cpp b/Engine/Source/Runtime/SlateCore/Private/Types/SlateAttributeDescriptor.cpp
--- a/Engine/Source/Runtime/SlateCore/Private/Types/SlateAttributeDescriptor.cpp
+++ b/Engine/Source/Runtime/SlateCore/Private/Types/SlateAttributeDescriptor.cpp
@@ -34,11 +34,11 @@ FSlateAttributeDescriptor::FInitializer::FAttributeEntry& FSlateAttributeDescrip
 }
 
 
-FSlateAttributeDescriptor::FInitializer::FAttributeEntry& FSlateAttributeDescriptor::FInitializer::FAttributeEntry::UpdateWhenCollapsed()
+FSlateAttributeDescriptor::FInitializer::FAttributeEntry& FSlateAttributeDescriptor::FInitializer::FAttributeEntry::AffectVisibility()
 {
 	if (Descriptor.Attributes.IsValidIndex(AttributeIndex))
 	{
-		Descriptor.SetUpdateWhenCollapsed(Descriptor.Attributes[AttributeIndex], true);
+		Descriptor.SetAffectVisibility(Descriptor.Attributes[AttributeIndex], true);
 	}
 	return *this;
 }
@@ -61,9 +61,18 @@ FSlateAttributeDescriptor::FInitializer::~FInitializer()
 {
 	checkf(Descriptor.Attributes.Num() < std::numeric_limits<int16>::max(), TEXT("There is too many attributes. The index is saved as an int16 in FSlateAttributeMetaData."));
 
+	// Confirm that the Visibility attribute is marked as "bAffectVisibility"
+	{
+		const FAttribute* FoundVisibilityAttribute = Descriptor.FindAttribute("Visibility");
+		checkf(FoundVisibilityAttribute, TEXT("The visibility attribute doesn't exist."));
+		checkf(FoundVisibilityAttribute->bAffectVisibility, TEXT("The Visibility attribute must be marked as 'Affect Visibility'"));
+	}
+
+
 	// Update the sort order for the item that have prerequisite.
 	//Because adding the attribute is not required at the moment
-	//try to not change the order in which they were added
+	//try to not change the order in which they were added.
+	//The AffectVisibility attributes must be in the front of the list.
 
 	struct FPrerequisiteSort
 	{
@@ -92,22 +101,31 @@ FSlateAttributeDescriptor::FInitializer::~FInitializer()
 		const TArray<FAttribute>& Attrbutes;
 		bool operator()(const FPrerequisiteSort& A, const FPrerequisiteSort& B) const
 		{
+			const FAttribute& AttributeA = Attrbutes[A.AttributeIndex];
+			const FAttribute& AttributeB = Attrbutes[B.AttributeIndex];
+
+			if (AttributeA.bAffectVisibility != AttributeB.bAffectVisibility)
+			{
+				return AttributeA.bAffectVisibility;
+			}
+
 			if (A.Depth != B.Depth)
 			{
 				return A.Depth < B.Depth;
 			}
+
 			if (A.PrerequisitesIndex == B.PrerequisitesIndex)
 			{
 				return Attrbutes[A.AttributeIndex].SortOrder < Attrbutes[B.AttributeIndex].SortOrder;
 			}
 
-			const int32 SortA = A.PrerequisitesIndex != INDEX_NONE ? Attrbutes[A.PrerequisitesIndex].SortOrder : Attrbutes[A.AttributeIndex].SortOrder;
-			const int32 SortB = B.PrerequisitesIndex != INDEX_NONE ? Attrbutes[B.PrerequisitesIndex].SortOrder : Attrbutes[B.AttributeIndex].SortOrder;
+			const int32 SortA = A.PrerequisitesIndex != INDEX_NONE ? Attrbutes[A.PrerequisitesIndex].SortOrder : AttributeA.SortOrder;
+			const int32 SortB = B.PrerequisitesIndex != INDEX_NONE ? Attrbutes[B.PrerequisitesIndex].SortOrder : AttributeB.SortOrder;
 			return SortA < SortB;
 		}
 	};
 
-	TArray<FPrerequisiteSort, TInlineAllocator<16>> Prerequisites;
+	TArray<FPrerequisiteSort, TInlineAllocator<32>> Prerequisites;
 	Prerequisites.Reserve(Descriptor.Attributes.Num());
 
 	bool bHavePrerequisite = false;
@@ -162,6 +180,27 @@ FSlateAttributeDescriptor::FInitializer::~FInitializer()
 			PreviousPrerequisiteIndex = Element.PrerequisitesIndex;
 		}
 	}
+
+	// Confirm that the attributes marked as "AffectVisibility" are in front of the list.
+#if 1
+	{
+		TArray<FAttribute, TInlineAllocator<32>> Copy;
+		Copy.Append(Descriptor.Attributes);
+		Copy.Sort([](const FAttribute& A, const FAttribute& B){ return A.SortOrder < B.SortOrder; });
+		bool bLookingForAffectVisibility = true;
+		for (const FAttribute& Attribute : Copy)
+		{
+			if (!Attribute.bAffectVisibility)
+			{
+				bLookingForAffectVisibility = false;
+			}
+			else if (!bLookingForAffectVisibility)
+			{
+				checkf(false, TEXT("Attribute marked as 'AffectVisibility' should be at the start of the update list or depends on the Visibility attribute."));
+			}
+		}
+	}
+#endif
 }
 
 
@@ -189,12 +228,12 @@ void FSlateAttributeDescriptor::FInitializer::OverrideInvalidationReason(FName A
 }
 
 
-void FSlateAttributeDescriptor::FInitializer::SetUpdateWhenCollapsed(FName AttributeName, bool bUpdateWhenCollapsed)
+void FSlateAttributeDescriptor::FInitializer::SetAffectVisibility(FName AttributeName, bool bAffectVisibility)
 {
 	FAttribute* Attribute = Descriptor.FindAttribute(AttributeName);
 	if (ensureAlwaysMsgf(Attribute, TEXT("The attribute named '%s' doesn't exist"), *AttributeName.ToString()))
 	{
-		Descriptor.SetUpdateWhenCollapsed(*Attribute, bUpdateWhenCollapsed);
+		Descriptor.SetAffectVisibility(*Attribute, bAffectVisibility);
 	}
 }
 
@@ -322,7 +361,12 @@ void FSlateAttributeDescriptor::SetPrerequisite(FSlateAttributeDescriptor::FAttr
 }
 
 
-void FSlateAttributeDescriptor::SetUpdateWhenCollapsed(FAttribute& Attribute, bool bUpdate)
+void FSlateAttributeDescriptor::SetAffectVisibility(FAttribute& Attribute, bool bInAffectVisibility)
 {
-	Attribute.bUpdateWhenCollapsed = bUpdate;
+	if (Attribute.Name == "Visibility")
+	{
+		checkf(bInAffectVisibility, TEXT("The Visibility attribute must be marked at 'Affect Visibility'"));
+	}
+
+	Attribute.bAffectVisibility = bInAffectVisibility;
 }
\ No newline at end of file
diff --git a/Engine/Source/Runtime/SlateCore/Private/Types/SlateAttributeMetaData.cpp b/Engine/Source/Runtime/SlateCore/Private/Types/SlateAttributeMetaData.cpp
--- a/Engine/Source/Runtime/SlateCore/Private/Types/SlateAttributeMetaData.cpp
+++ b/Engine/Source/Runtime/SlateCore/Private/Types/SlateAttributeMetaData.cpp
@@ -101,10 +101,10 @@ void FSlateAttributeMetaData::RegisterAttributeImpl(SWidget& OwningWidget, FSlat
 					GetterItem.CachedAttributeDependencyIndex = (int16)FoundDependencyAttributeIndex;
 				}
 				GetterItem.bIsADependencyForSomeoneElse = FoundAttribute.bIsADependencyForSomeoneElse;
-				GetterItem.bUpdateWhenCollapsed = FoundAttribute.bUpdateWhenCollapsed;
-				if (GetterItem.bUpdateWhenCollapsed)
+				GetterItem.bAffectVisibility = FoundAttribute.bAffectVisibility;
+				if (GetterItem.bAffectVisibility)
 				{
-					++CollaspedAttributeCounter;
+					++AffectVisibilityCounter;
 				}
 			}
 			else
@@ -156,10 +156,10 @@ bool FSlateAttributeMetaData::UnregisterAttributeImpl(const FSlateAttributeBase&
 	const int32 FoundIndex = IndexOfAttribute(Attribute);
 	if (FoundIndex != INDEX_NONE)
 	{
-		if (Attributes[FoundIndex].bUpdateWhenCollapsed)
+		if (Attributes[FoundIndex].bAffectVisibility)
 		{
-			check(CollaspedAttributeCounter > 0);
-			--CollaspedAttributeCounter;
+			check(AffectVisibilityCounter > 0);
+			--AffectVisibilityCounter;
 		}
 		Attributes.RemoveAt(FoundIndex); // keep the order valid
 		return true;
@@ -168,16 +168,19 @@ bool FSlateAttributeMetaData::UnregisterAttributeImpl(const FSlateAttributeBase&
 }
 
 
-TArray<FName> FSlateAttributeMetaData::GetAttributeNames(const SWidget& OwningWidget) const
+TArray<FName> FSlateAttributeMetaData::GetAttributeNames(const SWidget& OwningWidget)
 {
 	TArray<FName> Names;
-	Names.Reserve(Attributes.Num());
-	for (const FGetterItem& Getter : Attributes)
+	if (FSlateAttributeMetaData* AttributeMetaData = FSlateAttributeMetaData::FindMetaData(OwningWidget))
 	{
-		const FName Name = Getter.GetAttributeName(OwningWidget);
-		if (Name.IsValid())
+		Names.Reserve(AttributeMetaData->Attributes.Num());
+		for (const FGetterItem& Getter : AttributeMetaData->Attributes)
 		{
-			Names.Add(Name);
+			const FName Name = Getter.GetAttributeName(OwningWidget);
+			if (Name.IsValid())
+			{
+				Names.Add(Name);
+			}
 		}
 	}
 	return Names;
@@ -278,68 +281,58 @@ void FSlateAttributeMetaData::InvalidateWidget(SWidget& OwningWidget, const FSla
 }
 
 
-void FSlateAttributeMetaData::UpdateAttributes(SWidget& OwningWidget, EInvalidationPermission InvalidationStyle)
+void FSlateAttributeMetaData::UpdateAllAttributes(SWidget& OwningWidget, EInvalidationPermission InvalidationStyle)
 {
 	if (FSlateAttributeMetaData* AttributeMetaData = FSlateAttributeMetaData::FindMetaData(OwningWidget))
 	{
-		AttributeMetaData->UpdateAttributesImpl(OwningWidget, EUpdateType::All, InvalidationStyle);
+		AttributeMetaData->UpdateAttributesImpl(OwningWidget, InvalidationStyle, 0, AttributeMetaData->Attributes.Num());
 	}
 }
 
 
-void FSlateAttributeMetaData::UpdateCollapsedAttributes(SWidget& OwningWidget, EInvalidationPermission InvalidationStyle)
+void FSlateAttributeMetaData::UpdateOnlyVisibilityAttributes(SWidget& OwningWidget, EInvalidationPermission InvalidationStyle)
 {
 	if (FSlateAttributeMetaData* AttributeMetaData = FSlateAttributeMetaData::FindMetaData(OwningWidget))
 	{
-		// Does it have any collapsed attributes
-		if (AttributeMetaData->CollaspedAttributeCounter > 0)
-		{
-			AttributeMetaData->UpdateAttributesImpl(OwningWidget, EUpdateType::Collapsed, InvalidationStyle);
-		}
+		AttributeMetaData->UpdateAttributesImpl(OwningWidget, InvalidationStyle, 0, AttributeMetaData->AffectVisibilityCounter);
 	}
 }
 
 
-void FSlateAttributeMetaData::UpdateExpandedAttributes(SWidget& OwningWidget, EInvalidationPermission InvalidationStyle)
+void FSlateAttributeMetaData::UpdateExceptVisibilityAttributes(SWidget& OwningWidget, EInvalidationPermission InvalidationStyle)
 {
 	if (FSlateAttributeMetaData* AttributeMetaData = FSlateAttributeMetaData::FindMetaData(OwningWidget))
 	{
-		AttributeMetaData->UpdateAttributesImpl(OwningWidget, EUpdateType::Expanded, InvalidationStyle);
+		AttributeMetaData->UpdateAttributesImpl(OwningWidget, InvalidationStyle, AttributeMetaData->AffectVisibilityCounter, AttributeMetaData->Attributes.Num());
 	}
 }
 
 
-void FSlateAttributeMetaData::UpdateChildrenCollapsedAttributes(SWidget& OwningWidget, EInvalidationPermission InvalidationStyle)
+void FSlateAttributeMetaData::UpdateChildrenOnlyVisibilityAttributes(SWidget& OwningWidget, EInvalidationPermission InvalidationStyle, bool bRecursive)
 {
 	FChildren* Children = OwningWidget.GetChildren();
 	const int32 NumChildren = Children->Num();
 	for (int32 ChildIndex = 0; ChildIndex < NumChildren; ++ChildIndex)
 	{
 		const TSharedRef<SWidget>& Child = Children->GetChildAt(ChildIndex);
-		UpdateCollapsedAttributes(Child.Get(), InvalidationStyle);
+		UpdateOnlyVisibilityAttributes(Child.Get(), InvalidationStyle);
+		if (bRecursive)
+		{
+			UpdateChildrenOnlyVisibilityAttributes(Child.Get(), InvalidationStyle, bRecursive);
+		}
 	}
 }
 
 
-void FSlateAttributeMetaData::UpdateAttributesImpl(SWidget& OwningWidget, EUpdateType UpdateType, EInvalidationPermission InvalidationStyle)
+void FSlateAttributeMetaData::UpdateAttributesImpl(SWidget& OwningWidget, EInvalidationPermission InvalidationStyle, int32 StartIndex, int32 IndexNum)
 {
-	bool bInvalidateIfNeeded = (InvalidationStyle == EInvalidationPermission::AllowInvalidation) || (InvalidationStyle == EInvalidationPermission::AllowInvalidationIfConstructed && OwningWidget.IsConstructionCompleted());
-	bool bAllowInvalidation = bInvalidateIfNeeded || InvalidationStyle == EInvalidationPermission::DelayInvalidation;
+	const bool bInvalidateIfNeeded = (InvalidationStyle == EInvalidationPermission::AllowInvalidation) || (InvalidationStyle == EInvalidationPermission::AllowInvalidationIfConstructed && OwningWidget.IsConstructionCompleted());
+	const bool bAllowInvalidation = bInvalidateIfNeeded || InvalidationStyle == EInvalidationPermission::DelayInvalidation;
 	EInvalidateWidgetReason InvalidationReason = EInvalidateWidgetReason::None;
-	for (int32 Index = 0; Index < Attributes.Num(); ++Index)
+	for (int32 Index = StartIndex; Index < IndexNum; ++Index)
 	{
 		FGetterItem& GetterItem = Attributes[Index];
 
-		// Only update attributes that needs to be updated when the widget is collapsed
-		if (UpdateType == EUpdateType::Collapsed && !GetterItem.bUpdateWhenCollapsed)
-		{
-			continue;
-		}
-		if (UpdateType == EUpdateType::Expanded && GetterItem.bUpdateWhenCollapsed)
-		{
-			continue;
-		}
-
 		// Update every attribute at least once.
 		//Check if it has a dependency and if it was updated this frame (it could be from an UpdateNow)
 		if (GetterItem.CachedAttributeDependencyIndex != INDEX_NONE && GetterItem.bUpdatedOnce)
diff --git a/Engine/Source/Runtime/SlateCore/Private/Widgets/SWidget.cpp b/Engine/Source/Runtime/SlateCore/Private/Widgets/SWidget.cpp
--- a/Engine/Source/Runtime/SlateCore/Private/Widgets/SWidget.cpp
+++ b/Engine/Source/Runtime/SlateCore/Private/Widgets/SWidget.cpp
@@ -172,7 +172,7 @@ void SWidget::PrivateRegisterAttributes(FSlateAttributeInitializer& AttributeIni
 	// Visibility should be the first Attribute in the list.
 	//The order in which SlateAttribute are declared in the .h dictates of the order.
 	SLATE_ADD_MEMBER_ATTRIBUTE_DEFINITION_WITH_NAME(AttributeInitializer, "Visibility", VisibilityAttribute, EInvalidateWidgetReason::Visibility)
-		.UpdateWhenCollapsed();
+		.AffectVisibility();
 	SLATE_ADD_MEMBER_ATTRIBUTE_DEFINITION_WITH_NAME(AttributeInitializer, "EnabledState", EnabledStateAttribute, EInvalidateWidgetReason::Paint);
 	SLATE_ADD_MEMBER_ATTRIBUTE_DEFINITION_WITH_NAME(AttributeInitializer, "RenderTransform", RenderTransformAttribute, EInvalidateWidgetReason::Layout | EInvalidateWidgetReason::RenderTransform);
 	SLATE_ADD_MEMBER_ATTRIBUTE_DEFINITION_WITH_NAME(AttributeInitializer, "RenderTransformPivot", RenderTransformPivotAttribute, EInvalidateWidgetReason::Layout | EInvalidateWidgetReason::RenderTransform);
@@ -664,7 +664,7 @@ void SWidget::SlatePrepass(float InLayoutScaleMultiplier)
 	{
 		if (HasRegisteredSlateAttribute() && IsAttributesUpdatesEnabled() && !GSlateIsOnFastProcessInvalidation)
 		{
-			FSlateAttributeMetaData::UpdateAttributes(*this, FSlateAttributeMetaData::EInvalidationPermission::AllowInvalidationIfConstructed);
+			FSlateAttributeMetaData::UpdateAllAttributes(*this, FSlateAttributeMetaData::EInvalidationPermission::AllowInvalidationIfConstructed);
 		}
 		Prepass_Internal(InLayoutScaleMultiplier);
 	}
@@ -1615,7 +1615,7 @@ void SWidget::ArrangeChildren(const FGeometry& AllottedGeometry, FArrangedChildr
 	{
 		// Update the Widgets visibility before getting the ArrangeChildren
 		//const-casting for TSlateAttribute has the same behavior as previously with TAttribute. The const was hidden from the user.
-		FSlateAttributeMetaData::UpdateChildrenCollapsedAttributes(const_cast<SWidget&>(*this), FSlateAttributeMetaData::EInvalidationPermission::DelayInvalidation);
+		FSlateAttributeMetaData::UpdateChildrenOnlyVisibilityAttributes(const_cast<SWidget&>(*this), FSlateAttributeMetaData::EInvalidationPermission::DelayInvalidation, false);
 	}
 
 	OnArrangeChildren(AllottedGeometry, ArrangedChildren);
@@ -1653,24 +1653,25 @@ void SWidget::Prepass_ChildLoop(float InLayoutScaleMultiplier, FChildren* MyChil
 	const int32 NumChildren = MyChildren->Num();
 	for (int32 ChildIndex = 0; ChildIndex < NumChildren; ++ChildIndex)
 	{
-		const float ChildLayoutScaleMultiplier = bHasRelativeLayoutScale
-			? InLayoutScaleMultiplier * GetRelativeLayoutScale(ChildIndex, InLayoutScaleMultiplier)
-			: InLayoutScaleMultiplier;
-
 		const TSharedRef<SWidget>& Child = MyChildren->GetChildAt(ChildIndex);
 
 		const bool bUpdateAttributes = Child->HasRegisteredSlateAttribute() && Child->IsAttributesUpdatesEnabled() && !GSlateIsOnFastProcessInvalidation;
 		if (bUpdateAttributes)
 		{
-			FSlateAttributeMetaData::UpdateCollapsedAttributes(Child.Get(), FSlateAttributeMetaData::EInvalidationPermission::AllowInvalidationIfConstructed);
+			FSlateAttributeMetaData::UpdateOnlyVisibilityAttributes(Child.Get(), FSlateAttributeMetaData::EInvalidationPermission::AllowInvalidationIfConstructed);
 		}
 
 		if (Child->GetVisibility() != EVisibility::Collapsed)
 		{
 			if (bUpdateAttributes)
 			{
-				FSlateAttributeMetaData::UpdateExpandedAttributes(Child.Get(), FSlateAttributeMetaData::EInvalidationPermission::AllowInvalidationIfConstructed);
+				FSlateAttributeMetaData::UpdateExceptVisibilityAttributes(Child.Get(), FSlateAttributeMetaData::EInvalidationPermission::AllowInvalidationIfConstructed);
 			}
+
+			const float ChildLayoutScaleMultiplier = bHasRelativeLayoutScale
+				? InLayoutScaleMultiplier * GetRelativeLayoutScale(ChildIndex, InLayoutScaleMultiplier)
+				: InLayoutScaleMultiplier;
+
 			// Recur: Descend down the widget tree.
 			Child->Prepass_Internal(ChildLayoutScaleMultiplier);
 		}
@@ -1680,7 +1681,9 @@ void SWidget::Prepass_ChildLoop(float InLayoutScaleMultiplier, FChildren* MyChil
 			// it is finally visible and invalidate it's prepass so that it gets that when its visibility
 			// is finally invalidated.
 			Child->MarkPrepassAsDirty();
-			Child->PrepassLayoutScaleMultiplier = ChildLayoutScaleMultiplier;
+			Child->PrepassLayoutScaleMultiplier = bHasRelativeLayoutScale
+				? InLayoutScaleMultiplier * GetRelativeLayoutScale(ChildIndex, InLayoutScaleMultiplier)
+				: InLayoutScaleMultiplier;
 		}
 	}
 }
diff --git a/Engine/Source/Runtime/SlateCore/Public/Types/SlateAttributeDescriptor.h b/Engine/Source/Runtime/SlateCore/Public/Types/SlateAttributeDescriptor.h
--- a/Engine/Source/Runtime/SlateCore/Public/Types/SlateAttributeDescriptor.h
+++ b/Engine/Source/Runtime/SlateCore/Public/Types/SlateAttributeDescriptor.h
@@ -66,17 +66,6 @@ public:
 		FGetter Getter;
 	};
 
-	/** */
-	enum class EInvalidationDelegateOverrideType
-	{
-		/** Replace the callback that the base class defined. */
-		ReplacePrevious,
-		/** Execute the callback that the base class defined, then execute the new callback. */
-		ExecuteAfterPrevious,
-		/** Execute the new callback, then execute the callback that the base class defined. */
-		ExecuteBeforePrevious,
-	};
-
 public:
 	using OffsetType = uint32;
 
@@ -94,7 +83,7 @@ public:
 		bool bIsMemberAttribute = false;
 		bool bIsPrerequisiteAlsoADependency = false;
 		bool bIsADependencyForSomeoneElse = false;
-		bool bUpdateWhenCollapsed = false;
+		bool bAffectVisibility = false;
 	};
 
 	/** Internal class to initialize the SlateAttributeDescriptor (Add attributes or modify existing attributes). */
@@ -130,9 +119,11 @@ public:
 			FAttributeEntry& UpdateDependency(FName Dependency);
 
 			/**
-			 * Update the attribute when the widget is collapsed and its parent is not.
+			 * The attribute affect the visibility of the widget.
+			 * We only update the attributes that can change the visibility of the widget when the widget is collapsed.
+			 * Attributes that affect visibility must have the Visibility attribute as a Prerequisite or the Visibility attribute must have it as a Prerequisite.
 			 */
-			FAttributeEntry& UpdateWhenCollapsed();
+			FAttributeEntry& AffectVisibility();
 
 		private:
 			FSlateAttributeDescriptor& Descriptor;
@@ -148,7 +139,7 @@ public:
 		void OverrideInvalidationReason(FName AttributeName, FInvalidateWidgetReasonAttribute&& Reason);
 
 		/** Change the update type of an attribute defined in a base class. */
-		void SetUpdateWhenCollapsed(FName AttributeName, bool bUpdateWhenCollapsed);
+		void SetAffectVisibility(FName AttributeName, bool bAffectVisibility);
 
 
 	private:
@@ -159,10 +150,10 @@ public:
 	int32 AttributeNum() const { return Attributes.Num(); }
 
 	/** @returns the Attribute at the index previously found with IndexOfMemberAttribute */
-	FAttribute const& GetAttributeAtIndex(int32 Index) const;
+	const FAttribute& GetAttributeAtIndex(int32 Index) const;
 
 	/** @returns the Attribute with the corresponding name. */
-	FAttribute const* FindAttribute(FName AttributeName) const;
+	const FAttribute* FindAttribute(FName AttributeName) const;
 
 	/** @returns the index of a SlateAttribute that have the corresponding memory offset. */
 	int32 IndexOfMemberAttribute(OffsetType AttributeOffset) const;
@@ -171,7 +162,7 @@ public:
 	int32 IndexOfMemberAttribute(FName AttributeName) const;
 
 	/** @returns the Attribute of a SlateAttribute that have the corresponding memory offset. */
-	FAttribute const* FindMemberAttribute(OffsetType AttributeOffset) const;
+	const FAttribute* FindMemberAttribute(OffsetType AttributeOffset) const;
 
 	/** Iterator over each dependency this attribute is responsible of. */
 	template<typename Predicate>
@@ -191,7 +182,7 @@ private:
 	FInitializer::FAttributeEntry AddMemberAttribute(FName AttributeName, OffsetType Offset, FInvalidateWidgetReasonAttribute ReasonGetter);
 	void OverrideInvalidationReason(FName AttributeName, FInvalidateWidgetReasonAttribute ReasonGetter);
 	void SetPrerequisite(FAttribute& Attribute, FName Prerequisite, bool bSetAsDependency);
-	void SetUpdateWhenCollapsed(FAttribute& Attribute, bool bUpdate);
+	void SetAffectVisibility(FAttribute& Attribute, bool bUpdate);
 
 	template<typename Predicate>
 	void ForEachDependencyImpl(FName const& LookForName, int32 Index, Predicate& Pred) const
diff --git a/Engine/Source/Runtime/SlateCore/Public/Types/SlateAttributeMetaData.h b/Engine/Source/Runtime/SlateCore/Public/Types/SlateAttributeMetaData.h
--- a/Engine/Source/Runtime/SlateCore/Public/Types/SlateAttributeMetaData.h
+++ b/Engine/Source/Runtime/SlateCore/Public/Types/SlateAttributeMetaData.h
@@ -46,25 +46,22 @@ public:
 	 * Update all the attributes.
 	 * @param InvalidationStyle if we should invalidate the widget.
 	 */
-	static void UpdateAttributes(SWidget& OwningWidget, EInvalidationPermission InvalidationStyle);
+	static void UpdateAllAttributes(SWidget& OwningWidget, EInvalidationPermission InvalidationStyle);
 	/**
-	 * Update attributes that are mark to be updated when the widget is collapsed.
-	 * These attributes are usually responsible to change visibility of the widget.
+	 * Update attributes that are responsible to change visibility of the widget.
 	 * @param InvalidationStyle if we should invalidate the widget.
 	 */
-	static void UpdateCollapsedAttributes(SWidget& OwningWidget, EInvalidationPermission InvalidationStyle);
+	static void UpdateOnlyVisibilityAttributes(SWidget& OwningWidget, EInvalidationPermission InvalidationStyle);
 	/**
-	 * Update attributes that are mark to be updated when the widget is NOT collapsed.
-	 * These attributes usually do not change the visibility of the widget.
+	 * Update attributes that are NOT responsible to change visibility of the widget.
 	 * @param InvalidationStyle if we should invalidate the widget.
 	 */
-	static void UpdateExpandedAttributes(SWidget& OwningWidget, EInvalidationPermission InvalidationStyle);
+	static void UpdateExceptVisibilityAttributes(SWidget& OwningWidget, EInvalidationPermission InvalidationStyle);
 	/**
-	 * Update the children attributes that are mark to be updated when the widget is collapsed.
-	 * These attributes are usually responsible to change visibility of the widget.
+	 * Execute UpdateOnlyVisibilityAttributes on every children of the widget.
 	 * @param InvalidationStyle if we should invalidate the widget.
 	 */
-	static void UpdateChildrenCollapsedAttributes(SWidget& OwningWidget, EInvalidationPermission InvalidationStyle);
+	static void UpdateChildrenOnlyVisibilityAttributes(SWidget& OwningWidget, EInvalidationPermission InvalidationStyle, bool bRecursive);
 
 public:
 	bool IsBound(const FSlateAttributeBase& Attribute) const
@@ -72,12 +69,12 @@ public:
 		return IndexOfAttribute(Attribute) != INDEX_NONE;
 	}
 
-	int32 RegisteredAttributeCount() const { return Attributes.Num(); }
+	int32 GetRegisteredAttributeCount() const { return Attributes.Num(); }
 
-	int32 RegisteredCollaspedAttributeCount() const { return CollaspedAttributeCounter; }
+	int32 GetRegisteredAffectVisibilityAttributeCount() const { return AffectVisibilityCounter; }
 
 	/** Get the name of all the attributes, if available. */
-	TArray<FName> GetAttributeNames(const SWidget& OwningWidget) const;
+	static TArray<FName> GetAttributeNames(const SWidget& OwningWidget);
 
 private:
 	using ESlateAttributeType = SlateAttributePrivate::ESlateAttributeType;
@@ -92,15 +89,9 @@ private:
 	static void MoveAttribute(const SWidget& OwningWidget, FSlateAttributeBase& NewAttribute, ESlateAttributeType AttributeType, const FSlateAttributeBase* PreviousAttribute);
 
 private:
-	enum class EUpdateType
-	{
-		All,
-		Collapsed,
-		Expanded, // not mark as collapsed
-	};
 	void RegisterAttributeImpl(SWidget& OwningWidget, FSlateAttributeBase& Attribute, ESlateAttributeType AttributeType, TUniquePtr<ISlateAttributeGetter>&& Getter);
 	bool UnregisterAttributeImpl(const FSlateAttributeBase& Attribute);
-	void UpdateAttributesImpl(SWidget& OwningWidget, EUpdateType UpdateType, EInvalidationPermission InvaldiationStyle);
+	void UpdateAttributesImpl(SWidget& OwningWidget, EInvalidationPermission InvaldiationStyle, int32 StartIndex, int32 IndexNum);
 
 private:
 	int32 IndexOfAttribute(const FSlateAttributeBase& Attribute) const
@@ -146,7 +137,7 @@ private:
 				int8 bIsADependencyForSomeoneElse : 1;
 				int8 bIsMemberType : 1;
 				int8 bIsManagedType : 1;
-				int8 bUpdateWhenCollapsed : 1;
+				int8 bAffectVisibility : 1;
 			};
 			int8 Flags;
 		};
@@ -169,5 +160,5 @@ private:
 	//~2. The parent widget will clear this widget PersistentState.
 	EInvalidateWidgetReason CachedInvalidationReason = EInvalidateWidgetReason::None;
 	bool bHasUpdatedManuallyFlagToReset = false;
-	uint8 CollaspedAttributeCounter = 0;
+	uint8 AffectVisibilityCounter = 0;
 };
