TMethodPtr<U, ...> and TConstMethodPtr<U, ...> template aliases added to delegates to replace the T*Delegate_Const<U, ...>::FMethodPtr types.

diff --git a/Engine/Plugins/Experimental/EnhancedInput/Source/EnhancedInput/Public/EnhancedInputComponent.h b/Engine/Plugins/Experimental/EnhancedInput/Source/EnhancedInput/Public/EnhancedInputComponent.h
--- a/Engine/Plugins/Experimental/EnhancedInput/Source/EnhancedInput/Public/EnhancedInputComponent.h
+++ b/Engine/Plugins/Experimental/EnhancedInput/Source/EnhancedInput/Public/EnhancedInputComponent.h
@@ -56,7 +56,7 @@ public:
 	/** Binds a native delegate, hidden for script delegates */
 	template<	typename UserClass,
 				typename TSig = TSignature>
-	void BindDelegate(UserClass* Object, typename TSig::template TUObjectMethodDelegate<UserClass>::FMethodPtr Func)
+	void BindDelegate(UserClass* Object, typename TSig::template TMethodPtr<UserClass> Func)
 	{
 		Unbind();
 		Delegate = MakeShared<TSig>(TSig::CreateUObject(Object, Func));
@@ -356,13 +356,13 @@ public:
 	/**
 	 * Binds a delegate function matching any of the handler signatures to a UInputAction assigned via UInputMappingContext to the owner of this component.
 	 */
-#define DEFINE_BIND_ACTION(HANDLER_SIG)																																									\
-	template<class UserClass>																																											\
-	FEnhancedInputActionEventBinding& BindAction(const UInputAction* Action, ETriggerEvent TriggerEvent, UserClass* Object, typename HANDLER_SIG::TUObjectMethodDelegate< UserClass >::FMethodPtr Func) \
-	{																																																	\
-		TUniquePtr<FEnhancedInputActionEventDelegateBinding<HANDLER_SIG>> AB = MakeUnique<FEnhancedInputActionEventDelegateBinding<HANDLER_SIG>>(Action, TriggerEvent);									\
-		AB->Delegate.BindDelegate<UserClass>(Object, Func);																																				\
-		return *EnhancedActionEventBindings.Add_GetRef(MoveTemp(AB));																																	\
+#define DEFINE_BIND_ACTION(HANDLER_SIG)																																			\
+	template<class UserClass>																																					\
+	FEnhancedInputActionEventBinding& BindAction(const UInputAction* Action, ETriggerEvent TriggerEvent, UserClass* Object, typename HANDLER_SIG::TMethodPtr< UserClass > Func) \
+	{																																											\
+		TUniquePtr<FEnhancedInputActionEventDelegateBinding<HANDLER_SIG>> AB = MakeUnique<FEnhancedInputActionEventDelegateBinding<HANDLER_SIG>>(Action, TriggerEvent);			\
+		AB->Delegate.BindDelegate<UserClass>(Object, Func);																														\
+		return *EnhancedActionEventBindings.Add_GetRef(MoveTemp(AB));																											\
 	}
 
 	DEFINE_BIND_ACTION(FEnhancedInputActionHandlerSignature);
@@ -406,7 +406,7 @@ public:
 	 * Binds a chord event to a delegate function in development builds only.
 	 */
 	template<class UserClass>
-	FInputDebugKeyBinding& BindDebugKey(const FInputChord Chord, const EInputEvent KeyEvent, UserClass* Object, typename FInputDebugKeyHandlerSignature::TUObjectMethodDelegate< UserClass >::FMethodPtr Func, bool bExecuteWhenPaused = true)
+	FInputDebugKeyBinding& BindDebugKey(const FInputChord Chord, const EInputEvent KeyEvent, UserClass* Object, typename FInputDebugKeyHandlerSignature::TMethodPtr< UserClass > Func, bool bExecuteWhenPaused = true)
 	{
 #ifdef DEV_ONLY_KEY_BINDINGS_AVAILABLE
 		TUniquePtr<FInputDebugKeyDelegateBinding<FInputDebugKeyHandlerSignature>> KB = MakeUnique<FInputDebugKeyDelegateBinding<FInputDebugKeyHandlerSignature>>(Chord, KeyEvent, bExecuteWhenPaused);
@@ -436,26 +436,26 @@ public:
 
 	// Delete all InputComponent binding helpers. Indicates intentions going forward and improves intellisense/VAX when working with EnhancedInputCompoennts.
 	template<class UserClass>
-	FInputActionBinding& BindAction(const FName ActionName, const EInputEvent KeyEvent, UserClass* Object, typename FInputActionHandlerSignature::TUObjectMethodDelegate< UserClass >::FMethodPtr Func) = delete;
+	FInputActionBinding& BindAction(const FName ActionName, const EInputEvent KeyEvent, UserClass* Object, typename FInputActionHandlerSignature::TMethodPtr< UserClass > Func) = delete;
 	template<class UserClass>
-	FInputActionBinding& BindAction(const FName ActionName, const EInputEvent KeyEvent, UserClass* Object, typename FInputActionHandlerWithKeySignature::TUObjectMethodDelegate< UserClass >::FMethodPtr Func) = delete;
+	FInputActionBinding& BindAction(const FName ActionName, const EInputEvent KeyEvent, UserClass* Object, typename FInputActionHandlerWithKeySignature::TMethodPtr< UserClass > Func) = delete;
 	template< class DelegateType, class UserClass, typename... VarTypes >
-	FInputActionBinding& BindAction(const FName ActionName, const EInputEvent KeyEvent, UserClass* Object, typename DelegateType::template TUObjectMethodDelegate< UserClass >::FMethodPtr Func, VarTypes... Vars) = delete;
+	FInputActionBinding& BindAction(const FName ActionName, const EInputEvent KeyEvent, UserClass* Object, typename DelegateType::template TMethodPtr< UserClass > Func, VarTypes... Vars) = delete;
 	template<class UserClass>
-	FInputAxisBinding& BindAxis(const FName AxisName, UserClass* Object, typename FInputAxisHandlerSignature::TUObjectMethodDelegate< UserClass >::FMethodPtr Func) = delete;
+	FInputAxisBinding& BindAxis(const FName AxisName, UserClass* Object, typename FInputAxisHandlerSignature::TMethodPtr< UserClass > Func) = delete;
 	FInputAxisBinding& BindAxis(const FName AxisName) = delete;
 	template<class UserClass>
-	FInputAxisKeyBinding& BindAxisKey(const FKey AxisKey, UserClass* Object, typename FInputAxisHandlerSignature::TUObjectMethodDelegate< UserClass >::FMethodPtr Func) = delete;
+	FInputAxisKeyBinding& BindAxisKey(const FKey AxisKey, UserClass* Object, typename FInputAxisHandlerSignature::TMethodPtr< UserClass > Func) = delete;
 	FInputAxisKeyBinding& BindAxisKey(const FKey AxisKey) = delete;
 	template<class UserClass>
-	FInputVectorAxisBinding& BindVectorAxis(const FKey AxisKey, UserClass* Object, typename FInputVectorAxisHandlerSignature::TUObjectMethodDelegate< UserClass >::FMethodPtr Func) = delete;
+	FInputVectorAxisBinding& BindVectorAxis(const FKey AxisKey, UserClass* Object, typename FInputVectorAxisHandlerSignature::TMethodPtr< UserClass > Func) = delete;
 	FInputVectorAxisBinding& BindVectorAxis(const FKey AxisKey) = delete;
 	template<class UserClass>
-	FInputKeyBinding& BindKey(const FInputChord Chord, const EInputEvent KeyEvent, UserClass* Object, typename FInputActionHandlerSignature::TUObjectMethodDelegate< UserClass >::FMethodPtr Func) = delete;
+	FInputKeyBinding& BindKey(const FInputChord Chord, const EInputEvent KeyEvent, UserClass* Object, typename FInputActionHandlerSignature::TMethodPtr< UserClass > Func) = delete;
 	template<class UserClass>
-	FInputKeyBinding& BindKey(const FKey Key, const EInputEvent KeyEvent, UserClass* Object, typename FInputActionHandlerSignature::TUObjectMethodDelegate< UserClass >::FMethodPtr Func) = delete;
+	FInputKeyBinding& BindKey(const FKey Key, const EInputEvent KeyEvent, UserClass* Object, typename FInputActionHandlerSignature::TMethodPtr< UserClass > Func) = delete;
 	template<class UserClass>
-	FInputTouchBinding& BindTouch(const EInputEvent KeyEvent, UserClass* Object, typename FInputTouchHandlerSignature::TUObjectMethodDelegate< UserClass >::FMethodPtr Func) = delete;
+	FInputTouchBinding& BindTouch(const EInputEvent KeyEvent, UserClass* Object, typename FInputTouchHandlerSignature::TMethodPtr< UserClass > Func) = delete;
 	template<class UserClass>
-	FInputGestureBinding& BindGesture(const FKey GestureKey, UserClass* Object, typename FInputGestureHandlerSignature::TUObjectMethodDelegate< UserClass >::FMethodPtr Func) = delete;
+	FInputGestureBinding& BindGesture(const FKey GestureKey, UserClass* Object, typename FInputGestureHandlerSignature::TMethodPtr< UserClass > Func) = delete;
 };
diff --git a/Engine/Source/Developer/GameplayDebugger/Public/GameplayDebuggerAddonBase.h b/Engine/Source/Developer/GameplayDebugger/Public/GameplayDebuggerAddonBase.h
--- a/Engine/Source/Developer/GameplayDebugger/Public/GameplayDebuggerAddonBase.h
+++ b/Engine/Source/Developer/GameplayDebugger/Public/GameplayDebuggerAddonBase.h
@@ -42,7 +42,7 @@ protected:
 
 	/** creates new key binding handler: single key press */
 	template<class UserClass>
-	bool BindKeyPress(FName KeyName, UserClass* KeyHandlerObject, typename FGameplayDebuggerInputHandler::FHandler::TRawMethodDelegate< UserClass >::FMethodPtr KeyHandlerFunc, EGameplayDebuggerInputMode InputMode = EGameplayDebuggerInputMode::Local)
+	bool BindKeyPress(FName KeyName, UserClass* KeyHandlerObject, typename FGameplayDebuggerInputHandler::FHandler::TMethodPtr< UserClass > KeyHandlerFunc, EGameplayDebuggerInputMode InputMode = EGameplayDebuggerInputMode::Local)
 	{
 		FGameplayDebuggerInputHandler NewHandler;
 		NewHandler.KeyName = KeyName;
@@ -60,7 +60,7 @@ protected:
 
 	/** creates new key binding handler: key press with modifiers */
 	template<class UserClass>
-	bool BindKeyPress(FName KeyName, FGameplayDebuggerInputModifier KeyModifer, UserClass* KeyHandlerObject, typename FGameplayDebuggerInputHandler::FHandler::TRawMethodDelegate< UserClass >::FMethodPtr KeyHandlerFunc, EGameplayDebuggerInputMode InputMode = EGameplayDebuggerInputMode::Local)
+	bool BindKeyPress(FName KeyName, FGameplayDebuggerInputModifier KeyModifer, UserClass* KeyHandlerObject, typename FGameplayDebuggerInputHandler::FHandler::TMethodPtr< UserClass > KeyHandlerFunc, EGameplayDebuggerInputMode InputMode = EGameplayDebuggerInputMode::Local)
 	{
 		FGameplayDebuggerInputHandler NewHandler;
 		NewHandler.KeyName = KeyName;
@@ -79,7 +79,7 @@ protected:
 
 	/** creates new key binding handler: customizable key press, stored in config files */
 	template<class UserClass>
-	bool BindKeyPress(const FGameplayDebuggerInputHandlerConfig& InputConfig, UserClass* KeyHandlerObject, typename FGameplayDebuggerInputHandler::FHandler::TRawMethodDelegate< UserClass >::FMethodPtr KeyHandlerFunc, EGameplayDebuggerInputMode InputMode = EGameplayDebuggerInputMode::Local)
+	bool BindKeyPress(const FGameplayDebuggerInputHandlerConfig& InputConfig, UserClass* KeyHandlerObject, typename FGameplayDebuggerInputHandler::FHandler::TMethodPtr< UserClass > KeyHandlerFunc, EGameplayDebuggerInputMode InputMode = EGameplayDebuggerInputMode::Local)
 	{
 		FGameplayDebuggerInputHandler NewHandler;
 		NewHandler.KeyName = InputConfig.KeyName;
diff --git a/Engine/Source/Developer/TraceInsights/Private/Insights/Widgets/SStartPageWindow.cpp b/Engine/Source/Developer/TraceInsights/Private/Insights/Widgets/SStartPageWindow.cpp
--- a/Engine/Source/Developer/TraceInsights/Private/Insights/Widgets/SStartPageWindow.cpp
+++ b/Engine/Source/Developer/TraceInsights/Private/Insights/Widgets/SStartPageWindow.cpp
@@ -483,7 +483,7 @@ public:
 	}
 
 	void AddGridPanelRow(TSharedPtr<SGridPanel> Grid, int32 Row, const FText& Name,
-		typename TAttribute<FText>::FGetter::template TRawMethodDelegate_Const<STraceListRow>::FMethodPtr Value) const
+		typename TAttribute<FText>::FGetter::template TConstMethodPtr<STraceListRow> Value) const
 	{
 		Grid->AddSlot(0, Row)
 			.Padding(2.0f)
diff --git a/Engine/Source/Runtime/AIModule/Classes/EnvironmentQuery/EnvQueryManager.h b/Engine/Source/Runtime/AIModule/Classes/EnvironmentQuery/EnvQueryManager.h
--- a/Engine/Source/Runtime/AIModule/Classes/EnvironmentQuery/EnvQueryManager.h
+++ b/Engine/Source/Runtime/AIModule/Classes/EnvironmentQuery/EnvQueryManager.h
@@ -44,12 +44,12 @@ struct AIMODULE_API FEnvQueryRequest
 	FORCEINLINE FEnvQueryRequest& SetWorldOverride(UWorld* InWorld) { World = InWorld; return *this; }
 
 	template< class UserClass >	
-	FORCEINLINE int32 Execute(EEnvQueryRunMode::Type Mode, UserClass* InObj, typename FQueryFinishedSignature::TUObjectMethodDelegate< UserClass >::FMethodPtr InMethod)
+	FORCEINLINE int32 Execute(EEnvQueryRunMode::Type Mode, UserClass* InObj, typename FQueryFinishedSignature::TMethodPtr< UserClass > InMethod)
 	{
 		return Execute(Mode, FQueryFinishedSignature::CreateUObject(InObj, InMethod));
 	}
 	template< class UserClass >	
-	FORCEINLINE int32 Execute(EEnvQueryRunMode::Type Mode, UserClass* InObj, typename FQueryFinishedSignature::TUObjectMethodDelegate_Const< UserClass >::FMethodPtr InMethod)
+	FORCEINLINE int32 Execute(EEnvQueryRunMode::Type Mode, UserClass* InObj, typename FQueryFinishedSignature::TConstMethodPtr< UserClass > InMethod)
 	{
 		return Execute(Mode, FQueryFinishedSignature::CreateUObject(InObj, InMethod));
 	}
diff --git a/Engine/Source/Runtime/Core/Public/Delegates/DelegateSignatureImpl.inl b/Engine/Source/Runtime/Core/Public/Delegates/DelegateSignatureImpl.inl
--- a/Engine/Source/Runtime/Core/Public/Delegates/DelegateSignatureImpl.inl
+++ b/Engine/Source/Runtime/Core/Public/Delegates/DelegateSignatureImpl.inl
@@ -69,6 +69,10 @@ public:
 	typedef InRetValType RetValType;
 	typedef InRetValType TFuncType(ParamTypes...);
 
+	/* Helper typedefs for getting a member function pointer type for the delegate with a given payload */
+	template <typename UserClass, typename... VarTypes> using TMethodPtr      = typename TMemFunPtrType<false, UserClass, RetValType(ParamTypes..., VarTypes...)>::Type;
+	template <typename UserClass, typename... VarTypes> using TConstMethodPtr = typename TMemFunPtrType<true,  UserClass, RetValType(ParamTypes..., VarTypes...)>::Type;
+
 	/** Declare the user's "fast" shared pointer-based delegate instance types. */
 	template <typename UserClass                                                                            > struct TSPMethodDelegate                 : TBaseSPMethodDelegateInstance<false, UserClass, ESPMode::Fast, FuncType, UserPolicy                                        > { typedef TBaseSPMethodDelegateInstance<false, UserClass, ESPMode::Fast, FuncType, UserPolicy                                        > Super; TSPMethodDelegate                (const TSharedRef< UserClass, ESPMode::Fast >& InUserObject, typename Super::FMethodPtr InMethodPtr                                                            ) : Super(InUserObject, InMethodPtr                        ) {} };
 	template <typename UserClass                                                                            > struct TSPMethodDelegate_Const           : TBaseSPMethodDelegateInstance<true , UserClass, ESPMode::Fast, FuncType, UserPolicy                                        > { typedef TBaseSPMethodDelegateInstance<true , UserClass, ESPMode::Fast, FuncType, UserPolicy                                        > Super; TSPMethodDelegate_Const          (const TSharedRef< UserClass, ESPMode::Fast >& InUserObject, typename Super::FMethodPtr InMethodPtr                                                            ) : Super(InUserObject, InMethodPtr                        ) {} };
diff --git a/Engine/Source/Runtime/Core/Public/Misc/Attribute.h b/Engine/Source/Runtime/Core/Public/Misc/Attribute.h
--- a/Engine/Source/Runtime/Core/Public/Misc/Attribute.h
+++ b/Engine/Source/Runtime/Core/Public/Misc/Attribute.h
@@ -32,7 +32,8 @@ public:
 		: Value()         // NOTE: Potentially uninitialized for atomics!!
 		, bIsSet(false)
 		, Getter()
-	{ }
+	{
+	}
 
 	/**
 	 * Construct implicitly from an initial value
@@ -44,7 +45,8 @@ public:
 		: Value( (ObjectType)InInitialValue )
 		, bIsSet(true)
 		, Getter()
-	{ }
+	{
+	}
 
 	/** 
 	 * Construct implicitly from moving an initial value
@@ -55,7 +57,8 @@ public:
 		: Value(MoveTemp(InInitialValue))
 		, bIsSet(true)
 		, Getter()
-	{ }
+	{
+	}
 
 	/**
 	 * Constructs by binding an arbitrary function that will be called to generate this attribute's value on demand. 
@@ -65,13 +68,14 @@ public:
 	 * @param  InUserObject  Shared Pointer to the instance of the class that contains the member function you want to bind.  The attribute will only retain a weak pointer to this class.
 	 * @param  InMethodPtr Member function to bind.  The function's structure (return value, arguments, etc) must match IBoundAttributeDelegate's definition.
 	 */
-	template< class SourceType >	
-	TAttribute( TSharedRef< SourceType > InUserObject, typename FGetter::template TSPMethodDelegate_Const< SourceType >::FMethodPtr InMethodPtr )
-		: Value()		
+	template< class SourceType >
+	TAttribute( TSharedRef< SourceType > InUserObject, typename FGetter::template TConstMethodPtr< SourceType > InMethodPtr )
+		: Value()
 		, bIsSet(true)
 		, Getter(FGetter::CreateSP(InUserObject, InMethodPtr))
-	{ }
-	
+	{
+	}
+
 	/**
 	 * Constructs by binding an arbitrary function that will be called to generate this attribute's value on demand. 
 	 * After binding, the attribute will no longer have a value that can be accessed directly, and instead the bound
@@ -80,12 +84,13 @@ public:
 	 * @param  InUserObject  Shared Pointer to the instance of the class that contains the member function you want to bind.  The attribute will only retain a weak pointer to this class.
 	 * @param  InMethodPtr Member function to bind.  The function's structure (return value, arguments, etc) must match IBoundAttributeDelegate's definition.
 	 */
-	template< class SourceType >	
-	TAttribute( SourceType* InUserObject, typename FGetter::template TSPMethodDelegate_Const< SourceType >::FMethodPtr InMethodPtr )
-		: Value()		
+	template< class SourceType >
+	TAttribute( SourceType* InUserObject, typename FGetter::template TConstMethodPtr< SourceType > InMethodPtr )
+		: Value()
 		, bIsSet(true)
 		, Getter(FGetter::CreateSP(InUserObject, InMethodPtr))
-	{ }
+	{
+	}
 
 	/**
 	 * Static: Creates an attribute that's pre-bound to the specified 'getter' delegate
@@ -291,8 +296,8 @@ public:
 	 * @param  InUserObject  Instance of the class that contains the member function you want to bind.
 	 * @param  InMethodPtr Member function to bind.  The function's structure (return value, arguments, etc) must match IBoundAttributeDelegate's definition.
 	 */
-	template< class SourceType >	
-	void BindRaw( SourceType* InUserObject, typename FGetter::template TRawMethodDelegate_Const< SourceType >::FMethodPtr InMethodPtr )
+	template< class SourceType >
+	void BindRaw( SourceType* InUserObject, typename FGetter::template TConstMethodPtr< SourceType > InMethodPtr )
 	{
 		bIsSet = true;
 		Getter.BindRaw( InUserObject, InMethodPtr );
@@ -306,13 +311,13 @@ public:
 	 * @param  InUserObject  Shared Pointer to the instance of the class that contains the member function you want to bind.  The attribute will only retain a weak pointer to this class.
 	 * @param  InMethodPtr Member function to bind.  The function's structure (return value, arguments, etc) must match IBoundAttributeDelegate's definition.
 	 */
-	template< class SourceType >	
-	void Bind( TSharedRef< SourceType > InUserObject, typename FGetter::template TSPMethodDelegate_Const< SourceType >::FMethodPtr InMethodPtr )
+	template< class SourceType >
+	void Bind( TSharedRef< SourceType > InUserObject, typename FGetter::template TConstMethodPtr< SourceType > InMethodPtr )
 	{
 		bIsSet = true;
 		Getter.BindSP( InUserObject, InMethodPtr );
 	}
-	
+
 	/**
 	 * Binds an arbitrary function that will be called to generate this attribute's value on demand. 
 	 * After binding, the attribute will no longer have a value that can be accessed directly, and instead the bound
@@ -321,13 +326,13 @@ public:
 	 * @param  InUserObject  Shared Pointer to the instance of the class that contains the member function you want to bind.  The attribute will only retain a weak pointer to this class.
 	 * @param  InMethodPtr Member function to bind.  The function's structure (return value, arguments, etc) must match IBoundAttributeDelegate's definition.
 	 */
-	template< class SourceType >	
-	void Bind( SourceType* InUserObject, typename FGetter::template TSPMethodDelegate_Const< SourceType >::FMethodPtr InMethodPtr )
+	template< class SourceType >
+	void Bind( SourceType* InUserObject, typename FGetter::template TConstMethodPtr< SourceType > InMethodPtr )
 	{
 		bIsSet = true;
 		Getter.BindSP( InUserObject, InMethodPtr );
 	}
-	
+
 	/**
 	 * Binds an arbitrary function that will be called to generate this attribute's value on demand. 
 	 * After binding, the attribute will no longer have a value that can be accessed directly, and instead the bound
@@ -337,7 +342,7 @@ public:
 	 * @param  InMethodPtr Member function to bind.  The function's structure (return value, arguments, etc) must match IBoundAttributeDelegate's definition.
 	 */
 	template< class SourceType >	
-	void BindUObject( SourceType* InUserObject, typename FGetter::template TUObjectMethodDelegate_Const< SourceType >::FMethodPtr InMethodPtr )
+	void BindUObject( SourceType* InUserObject, typename FGetter::template TConstMethodPtr< SourceType > InMethodPtr )
 	{
 		bIsSet = true;
 		Getter.BindUObject( InUserObject, InMethodPtr );
diff --git a/Engine/Source/Runtime/Engine/Classes/Components/InputComponent.h b/Engine/Source/Runtime/Engine/Classes/Components/InputComponent.h
--- a/Engine/Source/Runtime/Engine/Classes/Components/InputComponent.h
+++ b/Engine/Source/Runtime/Engine/Classes/Components/InputComponent.h
@@ -48,7 +48,7 @@ struct TInputUnifiedDelegate
 
 	/** Binds a native delegate and unbinds any bound dynamic delegate */
 	template< class UserClass >
-	inline void BindDelegate(UserClass* Object, typename DelegateType::template TUObjectMethodDelegate< UserClass >::FMethodPtr Func)
+	inline void BindDelegate(UserClass* Object, typename DelegateType::template TMethodPtr< UserClass > Func)
 	{
 		FuncDynDelegate.Reset();;
 		FuncDelegate = MakeShared<DelegateType>(DelegateType::CreateUObject(Object, Func));
@@ -207,7 +207,7 @@ struct FInputActionUnifiedDelegate
 
 	/** Binds a native delegate and unbinds any bound dynamic delegate */
 	template< class UserClass >
-	inline void BindDelegate(UserClass* Object, typename FInputActionHandlerSignature::template TUObjectMethodDelegate< UserClass >::FMethodPtr Func)
+	inline void BindDelegate(UserClass* Object, typename FInputActionHandlerSignature::template TMethodPtr< UserClass > Func)
 	{
 		Unbind();
 		BoundDelegateType = EBoundDelegate::Delegate;
@@ -215,7 +215,7 @@ struct FInputActionUnifiedDelegate
 	}
 
 	template< class UserClass >
-	inline void BindDelegate(UserClass* Object, typename FInputActionHandlerWithKeySignature::template TUObjectMethodDelegate< UserClass >::FMethodPtr Func)
+	inline void BindDelegate(UserClass* Object, typename FInputActionHandlerWithKeySignature::template TMethodPtr< UserClass > Func)
 	{
 		Unbind();
 		BoundDelegateType = EBoundDelegate::DelegateWithKey;
@@ -223,7 +223,7 @@ struct FInputActionUnifiedDelegate
 	}
 
 	template< class DelegateType, class UserClass, typename... VarTypes >
-	inline void BindDelegate(UserClass* Object, typename DelegateType::template TUObjectMethodDelegate< UserClass >::FMethodPtr Func, VarTypes... Vars)
+	inline void BindDelegate(UserClass* Object, typename DelegateType::template TMethodPtr< UserClass > Func, VarTypes... Vars)
 	{
 		Unbind();
 		BoundDelegateType = EBoundDelegate::Delegate;
@@ -864,7 +864,7 @@ public:
 	 * Returned reference is only guaranteed to be valid until another action is bound.
 	 */
 	template<class UserClass>
-	FInputActionBinding& BindAction( const FName ActionName, const EInputEvent KeyEvent, UserClass* Object, typename FInputActionHandlerSignature::TUObjectMethodDelegate< UserClass >::FMethodPtr Func )
+	FInputActionBinding& BindAction( const FName ActionName, const EInputEvent KeyEvent, UserClass* Object, typename FInputActionHandlerSignature::TMethodPtr< UserClass > Func )
 	{
 		FInputActionBinding AB( ActionName, KeyEvent );
 		AB.ActionDelegate.BindDelegate(Object, Func);
@@ -876,7 +876,7 @@ public:
 	 * Returned reference is only guaranteed to be valid until another action is bound.
 	 */
 	template<class UserClass>
-	FInputActionBinding& BindAction( const FName ActionName, const EInputEvent KeyEvent, UserClass* Object, typename FInputActionHandlerWithKeySignature::TUObjectMethodDelegate< UserClass >::FMethodPtr Func )
+	FInputActionBinding& BindAction( const FName ActionName, const EInputEvent KeyEvent, UserClass* Object, typename FInputActionHandlerWithKeySignature::TMethodPtr< UserClass > Func )
 	{
 		FInputActionBinding AB( ActionName, KeyEvent );
 		AB.ActionDelegate.BindDelegate(Object, Func);
@@ -888,7 +888,7 @@ public:
 	* Returned reference is only guaranteed to be valid until another action is bound.
 	*/
 	template< class DelegateType, class UserClass, typename... VarTypes >
-	FInputActionBinding& BindAction( const FName ActionName, const EInputEvent KeyEvent, UserClass* Object, typename DelegateType::template TUObjectMethodDelegate< UserClass >::FMethodPtr Func, VarTypes... Vars )
+	FInputActionBinding& BindAction( const FName ActionName, const EInputEvent KeyEvent, UserClass* Object, typename DelegateType::template TMethodPtr< UserClass > Func, VarTypes... Vars )
 	{
 		FInputActionBinding AB( ActionName, KeyEvent );
 		AB.ActionDelegate.BindDelegate<DelegateType>(Object, Func, Vars...);
@@ -900,7 +900,7 @@ public:
 	 * Returned reference is only guaranteed to be valid until another axis is bound.
 	 */
 	template<class UserClass>
-	FInputAxisBinding& BindAxis( const FName AxisName, UserClass* Object, typename FInputAxisHandlerSignature::TUObjectMethodDelegate< UserClass >::FMethodPtr Func )
+	FInputAxisBinding& BindAxis( const FName AxisName, UserClass* Object, typename FInputAxisHandlerSignature::TMethodPtr< UserClass > Func )
 	{
 		FInputAxisBinding AB( AxisName );
 		AB.AxisDelegate.BindDelegate(Object, Func);
@@ -925,7 +925,7 @@ public:
 	 * Returned reference is only guaranteed to be valid until another axis key is bound.
 	 */
 	template<class UserClass>
-	FInputAxisKeyBinding& BindAxisKey( const FKey AxisKey, UserClass* Object, typename FInputAxisHandlerSignature::TUObjectMethodDelegate< UserClass >::FMethodPtr Func )
+	FInputAxisKeyBinding& BindAxisKey( const FKey AxisKey, UserClass* Object, typename FInputAxisHandlerSignature::TMethodPtr< UserClass > Func )
 	{
 		FInputAxisKeyBinding AB(AxisKey);
 		AB.AxisDelegate.BindDelegate(Object, Func);
@@ -950,7 +950,7 @@ public:
 	 * Returned reference is only guaranteed to be valid until another vector axis key is bound.
 	 */
 	template<class UserClass>
-	FInputVectorAxisBinding& BindVectorAxis( const FKey AxisKey, UserClass* Object, typename FInputVectorAxisHandlerSignature::TUObjectMethodDelegate< UserClass >::FMethodPtr Func )
+	FInputVectorAxisBinding& BindVectorAxis( const FKey AxisKey, UserClass* Object, typename FInputVectorAxisHandlerSignature::TMethodPtr< UserClass > Func )
 	{
 		FInputVectorAxisBinding AB(AxisKey);
 		AB.AxisDelegate.BindDelegate(Object, Func);
@@ -975,7 +975,7 @@ public:
 	 * Returned reference is only guaranteed to be valid until another input key is bound.
 	 */
 	template<class UserClass>
-	FInputKeyBinding& BindKey( const FInputChord Chord, const EInputEvent KeyEvent, UserClass* Object, typename FInputActionHandlerSignature::TUObjectMethodDelegate< UserClass >::FMethodPtr Func )
+	FInputKeyBinding& BindKey( const FInputChord Chord, const EInputEvent KeyEvent, UserClass* Object, typename FInputActionHandlerSignature::TMethodPtr< UserClass > Func )
 	{
 		FInputKeyBinding KB(Chord, KeyEvent);
 		KB.KeyDelegate.BindDelegate(Object, Func);
@@ -988,7 +988,7 @@ public:
 	 * Returned reference is only guaranteed to be valid until another input key is bound.
 	 */
 	template<class UserClass>
-	FInputKeyBinding& BindKey( const FKey Key, const EInputEvent KeyEvent, UserClass* Object, typename FInputActionHandlerSignature::TUObjectMethodDelegate< UserClass >::FMethodPtr Func )
+	FInputKeyBinding& BindKey( const FKey Key, const EInputEvent KeyEvent, UserClass* Object, typename FInputActionHandlerSignature::TMethodPtr< UserClass > Func )
 	{
 		return BindKey(FInputChord(Key, false, false, false, false), KeyEvent, Object, Func);
 	}
@@ -998,7 +998,7 @@ public:
 	 * Returned reference is only guaranteed to be valid until another input key is bound.
 	 */
 	template <class UserClass>
-	FInputKeyBinding& BindKey(const FKey Key, const EInputEvent KeyEvent, UserClass* Object, typename FInputActionHandlerWithKeySignature::TUObjectMethodDelegate<UserClass>::FMethodPtr Func)
+	FInputKeyBinding& BindKey(const FKey Key, const EInputEvent KeyEvent, UserClass* Object, typename FInputActionHandlerWithKeySignature::TMethodPtr<UserClass> Func)
 	{
 		FInputKeyBinding KB(FInputChord(Key, false, false, false, false), KeyEvent);
 		KB.KeyDelegate.BindDelegate(Object, Func);
@@ -1011,7 +1011,7 @@ public:
 	 * Returned reference is only guaranteed to be valid until another touch event is bound.
 	 */
 	template<class UserClass>
-	FInputTouchBinding& BindTouch( const EInputEvent KeyEvent, UserClass* Object, typename FInputTouchHandlerSignature::TUObjectMethodDelegate< UserClass >::FMethodPtr Func )
+	FInputTouchBinding& BindTouch( const EInputEvent KeyEvent, UserClass* Object, typename FInputTouchHandlerSignature::TMethodPtr< UserClass > Func )
 	{
 		FInputTouchBinding TB(KeyEvent);
 		TB.TouchDelegate.BindDelegate(Object, Func);
@@ -1024,7 +1024,7 @@ public:
 	 * Returned reference is only guaranteed to be valid until another gesture event is bound.
 	 */
 	template<class UserClass>
-	FInputGestureBinding& BindGesture( const FKey GestureKey, UserClass* Object, typename FInputGestureHandlerSignature::TUObjectMethodDelegate< UserClass >::FMethodPtr Func )
+	FInputGestureBinding& BindGesture( const FKey GestureKey, UserClass* Object, typename FInputGestureHandlerSignature::TMethodPtr< UserClass > Func )
 	{
 		FInputGestureBinding GB(GestureKey);
 		GB.GestureDelegate.BindDelegate(Object, Func);
diff --git a/Engine/Source/Runtime/Engine/Public/TimerManager.h b/Engine/Source/Runtime/Engine/Public/TimerManager.h
--- a/Engine/Source/Runtime/Engine/Public/TimerManager.h
+++ b/Engine/Source/Runtime/Engine/Public/TimerManager.h
@@ -188,12 +188,12 @@ public:
 	 * @param InFirstDelay			The time (in seconds) for the first iteration of a looping timer. If < 0.f InRate will be used.
 	 */
 	template< class UserClass >
-	FORCEINLINE void SetTimer(FTimerHandle& InOutHandle, UserClass* InObj, typename FTimerDelegate::TUObjectMethodDelegate< UserClass >::FMethodPtr InTimerMethod, float InRate, bool InbLoop = false, float InFirstDelay = -1.f)
+	FORCEINLINE void SetTimer(FTimerHandle& InOutHandle, UserClass* InObj, typename FTimerDelegate::TMethodPtr< UserClass > InTimerMethod, float InRate, bool InbLoop = false, float InFirstDelay = -1.f)
 	{
 		InternalSetTimer(InOutHandle, FTimerUnifiedDelegate( FTimerDelegate::CreateUObject(InObj, InTimerMethod) ), InRate, InbLoop, InFirstDelay);
 	}
 	template< class UserClass >
-	FORCEINLINE void SetTimer(FTimerHandle& InOutHandle, UserClass* InObj, typename FTimerDelegate::TUObjectMethodDelegate_Const< UserClass >::FMethodPtr InTimerMethod, float InRate, bool InbLoop = false, float InFirstDelay = -1.f)
+	FORCEINLINE void SetTimer(FTimerHandle& InOutHandle, UserClass* InObj, typename FTimerDelegate::TConstMethodPtr< UserClass > InTimerMethod, float InRate, bool InbLoop = false, float InFirstDelay = -1.f)
 	{
 		InternalSetTimer(InOutHandle, FTimerUnifiedDelegate( FTimerDelegate::CreateUObject(InObj, InTimerMethod) ), InRate, InbLoop, InFirstDelay);
 	}
@@ -226,12 +226,12 @@ public:
 	 * @param inTimerMethod			Method to call when timer fires.
 	 */
 	template< class UserClass >
-	FORCEINLINE FTimerHandle SetTimerForNextTick(UserClass* inObj, typename FTimerDelegate::TUObjectMethodDelegate< UserClass >::FMethodPtr inTimerMethod)
+	FORCEINLINE FTimerHandle SetTimerForNextTick(UserClass* inObj, typename FTimerDelegate::TMethodPtr< UserClass > inTimerMethod)
 	{
 		return InternalSetTimerForNextTick(FTimerUnifiedDelegate(FTimerDelegate::CreateUObject(inObj, inTimerMethod)));
 	}
 	template< class UserClass >
-	FORCEINLINE FTimerHandle SetTimerForNextTick(UserClass* inObj, typename FTimerDelegate::TUObjectMethodDelegate_Const< UserClass >::FMethodPtr inTimerMethod)
+	FORCEINLINE FTimerHandle SetTimerForNextTick(UserClass* inObj, typename FTimerDelegate::TConstMethodPtr< UserClass > inTimerMethod)
 	{
 		return InternalSetTimerForNextTick(FTimerUnifiedDelegate(FTimerDelegate::CreateUObject(inObj, inTimerMethod)));
 	}
diff --git a/Engine/Source/Runtime/GameMenuBuilder/Public/GameMenuPage.h b/Engine/Source/Runtime/GameMenuBuilder/Public/GameMenuPage.h
--- a/Engine/Source/Runtime/GameMenuBuilder/Public/GameMenuPage.h
+++ b/Engine/Source/Runtime/GameMenuBuilder/Public/GameMenuPage.h
@@ -78,7 +78,7 @@ public:
 	* @returns  SharedRef to the MenuItem that was created.
 	*/
 	template< class UserClass >
-	FORCEINLINE TSharedRef<FGameMenuItem> AddMenuItem(const FText& InText, UserClass* InObj, typename FGameMenuItem::FOnConfirmMenuItem::TSPMethodDelegate< UserClass >::FMethodPtr InMethod)
+	FORCEINLINE TSharedRef<FGameMenuItem> AddMenuItem(const FText& InText, UserClass* InObj, typename FGameMenuItem::FOnConfirmMenuItem::TMethodPtr< UserClass > InMethod)
 	{
 		TSharedPtr<FGameMenuItem> Item = MakeShareable(new FGameMenuItem(InText));
 		Item->OnConfirmMenuItem.BindSP(InObj, InMethod);
@@ -96,7 +96,7 @@ public:
 	 * @returns  SharedRef to the MenuItem that was created.
 	 */
 	template< class UserClass >
-	FORCEINLINE TSharedRef<FGameMenuItem> AddMenuItemWithOptions(const FText& Text, const TArray<FText>& OptionsList, UserClass* InObj, typename FGameMenuItem::FOnOptionChanged::TSPMethodDelegate< UserClass >::FMethodPtr InMethod)
+	FORCEINLINE TSharedRef<FGameMenuItem> AddMenuItemWithOptions(const FText& Text, const TArray<FText>& OptionsList, UserClass* InObj, typename FGameMenuItem::FOnOptionChanged::TMethodPtr< UserClass > InMethod)
 	{
 		TSharedPtr<FGameMenuItem> Item = MakeShareable(new FGameMenuItem(Text, OptionsList));
 		Item->OnOptionChanged.BindSP(InObj, InMethod);
@@ -114,7 +114,7 @@ public:
 	 * @returns  SharedRef to the MenuItem that was created.
 	 */
 	template< class UserClass >
-	FORCEINLINE TSharedRef<FGameMenuItem> AddCustomMenuItem(const FText& Text, TSharedPtr<SGameMenuItemWidget> CustomWidget, UserClass* InObj, typename FGameMenuItem::FOnOptionChanged::TSPMethodDelegate< UserClass >::FMethodPtr InMethod)
+	FORCEINLINE TSharedRef<FGameMenuItem> AddCustomMenuItem(const FText& Text, TSharedPtr<SGameMenuItemWidget> CustomWidget, UserClass* InObj, typename FGameMenuItem::FOnOptionChanged::TMethodPtr< UserClass > InMethod)
 	{
 		TSharedPtr<FGameMenuItem> Item = MakeShareable(new FGameMenuItem(CustomWidget));
 		Item->OnOptionChanged.BindSP(InObj, InMethod);
@@ -129,7 +129,7 @@ public:
 	* @param	InMethod	Method to call when selection changes.
 	*/
 	template< class UserClass >
-	FORCEINLINE void SetCancelHandler(UserClass* InObj, typename FGameMenuPage::FOnMenuGoBack::TSPMethodDelegate< UserClass >::FMethodPtr InMethod)
+	FORCEINLINE void SetCancelHandler(UserClass* InObj, typename FGameMenuPage::FOnMenuGoBack::TMethodPtr< UserClass > InMethod)
 	{
 		OnGoBackCancel.BindSP(InObj, InMethod);
 	}
@@ -141,7 +141,7 @@ public:
 	* @param	InMethod	Method to call when menu has been hidden.
 	*/
 	template< class UserClass >
-	FORCEINLINE void SetOnHiddenHandler(UserClass* InObj, typename FGameMenuPage::FOnMenuHidden::TSPMethodDelegate< UserClass >::FMethodPtr InMethod)
+	FORCEINLINE void SetOnHiddenHandler(UserClass* InObj, typename FGameMenuPage::FOnMenuHidden::TMethodPtr< UserClass > InMethod)
 	{
 		OnMenuHidden.BindSP(InObj, InMethod);
 	}
@@ -153,7 +153,7 @@ public:
 	* @param	InMethod	Method to call when selection changes.
 	*/
 	template< class UserClass >
-	FORCEINLINE void SetAcceptHandler(UserClass* InObj, typename FGameMenuPage::FOnMenuGoBack::TSPMethodDelegate< UserClass >::FMethodPtr InMethod)
+	FORCEINLINE void SetAcceptHandler(UserClass* InObj, typename FGameMenuPage::FOnMenuGoBack::TMethodPtr< UserClass > InMethod)
 	{
 		OnGoBack.BindSP(InObj, InMethod);
 	}
@@ -165,7 +165,7 @@ public:
 	* @param	InMethod	Method to call when menu is about to open.
 	*/
 	template< class UserClass >
-	FORCEINLINE void SetOnOpenHandler(UserClass* InObj, typename FGameMenuPage::FOnMenuHidden::TSPMethodDelegate< UserClass >::FMethodPtr InMethod)
+	FORCEINLINE void SetOnOpenHandler(UserClass* InObj, typename FGameMenuPage::FOnMenuHidden::TMethodPtr< UserClass > InMethod)
 	{
 		OnMenuOpening.BindSP(InObj, InMethod);
 	}
diff --git a/Engine/Source/Runtime/NavigationSystem/Public/NavLinkCustomComponent.h b/Engine/Source/Runtime/NavigationSystem/Public/NavLinkCustomComponent.h
--- a/Engine/Source/Runtime/NavigationSystem/Public/NavLinkCustomComponent.h
+++ b/Engine/Source/Runtime/NavigationSystem/Public/NavLinkCustomComponent.h
@@ -114,23 +114,23 @@ class NAVIGATIONSYSTEM_API UNavLinkCustomComponent : public UNavRelevantComponen
 	// helper functions for setting delegates
 
 	template< class UserClass >	
-	FORCEINLINE void SetMoveReachedLink(UserClass* TargetOb, typename FOnMoveReachedLink::TUObjectMethodDelegate< UserClass >::FMethodPtr InFunc)
+	FORCEINLINE void SetMoveReachedLink(UserClass* TargetOb, typename FOnMoveReachedLink::TMethodPtr< UserClass > InFunc)
 	{
 		SetMoveReachedLink(FOnMoveReachedLink::CreateUObject(TargetOb, InFunc));
 	}
 	template< class UserClass >	
-	FORCEINLINE void SetMoveReachedLink(UserClass* TargetOb, typename FOnMoveReachedLink::TUObjectMethodDelegate_Const< UserClass >::FMethodPtr InFunc)
+	FORCEINLINE void SetMoveReachedLink(UserClass* TargetOb, typename FOnMoveReachedLink::TConstMethodPtr< UserClass > InFunc)
 	{
 		SetMoveReachedLink(FOnMoveReachedLink::CreateUObject(TargetOb, InFunc));
 	}
 
 	template< class UserClass >	
-	FORCEINLINE void SetBroadcastFilter(UserClass* TargetOb, typename FBroadcastFilter::TUObjectMethodDelegate< UserClass >::FMethodPtr InFunc)
+	FORCEINLINE void SetBroadcastFilter(UserClass* TargetOb, typename FBroadcastFilter::TMethodPtr< UserClass > InFunc)
 	{
 		SetBroadcastFilter(FBroadcastFilter::CreateUObject(TargetOb, InFunc));
 	}
 	template< class UserClass >	
-	FORCEINLINE void SetBroadcastFilter(UserClass* TargetOb, typename FBroadcastFilter::TUObjectMethodDelegate_Const< UserClass >::FMethodPtr InFunc)
+	FORCEINLINE void SetBroadcastFilter(UserClass* TargetOb, typename FBroadcastFilter::TConstMethodPtr< UserClass > InFunc)
 	{
 		SetBroadcastFilter(FBroadcastFilter::CreateUObject(TargetOb, InFunc));
 	}
diff --git a/Engine/Source/Runtime/Slate/Public/Widgets/Text/SRichTextBlock.h b/Engine/Source/Runtime/Slate/Public/Widgets/Text/SRichTextBlock.h
--- a/Engine/Source/Runtime/Slate/Public/Widgets/Text/SRichTextBlock.h
+++ b/Engine/Source/Runtime/Slate/Public/Widgets/Text/SRichTextBlock.h
@@ -132,7 +132,7 @@ public:
 	}
 
 	template< class UserClass >
-	static TSharedRef< ITextDecorator >  WidgetDecorator( const FString RunName, UserClass* InUserObjectPtr, typename FWidgetDecorator::FCreateWidget::TSPMethodDelegate_Const< UserClass >::FMethodPtr InFunc )
+	static TSharedRef< ITextDecorator >  WidgetDecorator( const FString RunName, UserClass* InUserObjectPtr, typename FWidgetDecorator::FCreateWidget::TConstMethodPtr< UserClass > InFunc )
 	{
 		return FWidgetDecorator::Create( RunName, FWidgetDecorator::FCreateWidget::CreateSP( InUserObjectPtr, InFunc ) );
 	}
@@ -148,7 +148,7 @@ public:
 	}
 
 	template< class UserClass >
-	static TSharedRef< ITextDecorator > HyperlinkDecorator( const FString Id, UserClass* InUserObjectPtr, typename FSlateHyperlinkRun::FOnClick::TSPMethodDelegate< UserClass >::FMethodPtr NavigateFunc )
+	static TSharedRef< ITextDecorator > HyperlinkDecorator( const FString Id, UserClass* InUserObjectPtr, typename FSlateHyperlinkRun::FOnClick::TMethodPtr< UserClass > NavigateFunc )
 	{
 		return FHyperlinkDecorator::Create( Id, FSlateHyperlinkRun::FOnClick::CreateSP( InUserObjectPtr, NavigateFunc ) );
 	}
diff --git a/Engine/Source/Runtime/SlateCore/Public/Types/SlateAttribute.h b/Engine/Source/Runtime/SlateCore/Public/Types/SlateAttribute.h
--- a/Engine/Source/Runtime/SlateCore/Public/Types/SlateAttribute.h
+++ b/Engine/Source/Runtime/SlateCore/Public/Types/SlateAttribute.h
@@ -206,10 +206,10 @@ struct FSlateDeprecatedTAttribute
 	FSlateDeprecatedTAttribute(ObjectType&& InInitialValue)	{ }
 
 	template<class SourceType>
-	FSlateDeprecatedTAttribute(TSharedRef<SourceType> InUserObject, typename FGetter::template TSPMethodDelegate_Const< SourceType >::FMethodPtr InMethodPtr) {}
+	FSlateDeprecatedTAttribute(TSharedRef<SourceType> InUserObject, typename FGetter::template TConstMethodPtr< SourceType > InMethodPtr) {}
 
 	template< class SourceType >
-	FSlateDeprecatedTAttribute(SourceType* InUserObject, typename FGetter::template TSPMethodDelegate_Const< SourceType >::FMethodPtr InMethodPtr) { }
+	FSlateDeprecatedTAttribute(SourceType* InUserObject, typename FGetter::template TConstMethodPtr< SourceType > InMethodPtr) { }
 
 	bool IsSet() const { return false; }
 
@@ -222,7 +222,7 @@ struct FSlateDeprecatedTAttribute
 
 	void Bind(const FGetter& InGetter) {}
 	template<class SourceType>
-	void Bind(SourceType* InUserObject, typename FGetter::template TSPMethodDelegate_Const< SourceType >::FMethodPtr InMethodPtr) {}
+	void Bind(SourceType* InUserObject, typename FGetter::template TConstMethodPtr< SourceType > InMethodPtr) {}
 	bool IsBound() const { return false; }
 
 	bool IdenticalTo(const TAttribute<ObjectType>& InOther) const { return false; }
@@ -538,7 +538,7 @@ namespace SlateAttributePrivate
 		 * The SlateAttribute will now be updated every frame from the Getter.
 		 */
 		template<typename WidgetType, typename U = typename std::enable_if<std::is_base_of<SWidget, WidgetType>::value>::type>
-		void Bind(WidgetType& Widget, typename FGetter::template TSPMethodDelegate_Const<WidgetType>::FMethodPtr MethodPtr)
+		void Bind(WidgetType& Widget, typename FGetter::template TConstMethodPtr<WidgetType> MethodPtr)
 		{
 			Bind(Widget, FGetter::CreateSP(&Widget, MethodPtr));
 		}
diff --git a/Engine/Source/Runtime/SlateCore/Public/Widgets/DeclarativeSyntaxSupport.h b/Engine/Source/Runtime/SlateCore/Public/Widgets/DeclarativeSyntaxSupport.h
--- a/Engine/Source/Runtime/SlateCore/Public/Widgets/DeclarativeSyntaxSupport.h
+++ b/Engine/Source/Runtime/SlateCore/Public/Widgets/DeclarativeSyntaxSupport.h
@@ -162,31 +162,31 @@ template<typename WidgetType> struct TSlateBaseNamedArgs;
 	\
 		/* Bind attribute with delegate to a raw C++ class method */ \
 		template< class UserClass >	\
-		WidgetArgsType& AttrName##_Raw( UserClass* InUserObject, typename TAttribute< AttrType >::FGetter::template TRawMethodDelegate_Const< UserClass >::FMethodPtr InFunc )	\
+		WidgetArgsType& AttrName##_Raw( UserClass* InUserObject, typename TAttribute< AttrType >::FGetter::template TConstMethodPtr< UserClass > InFunc )	\
 		{ \
 			_##AttrName = TAttribute< AttrType >::Create( TAttribute< AttrType >::FGetter::CreateRaw( InUserObject, InFunc ) ); \
 			return this->Me(); \
 		} \
 		template< class UserClass, typename Var1Type >	\
-		WidgetArgsType& AttrName##_Raw( UserClass* InUserObject, typename TAttribute< AttrType >::FGetter::template TRawMethodDelegate_OneVar_Const< UserClass, Var1Type  >::FMethodPtr InFunc, Var1Type Var1 )	\
+		WidgetArgsType& AttrName##_Raw( UserClass* InUserObject, typename TAttribute< AttrType >::FGetter::template TConstMethodPtr< UserClass, Var1Type  > InFunc, Var1Type Var1 )	\
 		{ \
 			_##AttrName = TAttribute< AttrType >::Create( TAttribute< AttrType >::FGetter::CreateRaw( InUserObject, InFunc, Var1 ) ); \
 			return this->Me(); \
 		} \
 		template< class UserClass, typename Var1Type, typename Var2Type >	\
-		WidgetArgsType& AttrName##_Raw( UserClass* InUserObject, typename TAttribute< AttrType >::FGetter::template TRawMethodDelegate_TwoVars_Const< UserClass, Var1Type, Var2Type >::FMethodPtr InFunc, Var1Type Var1, Var2Type Var2 )	\
+		WidgetArgsType& AttrName##_Raw( UserClass* InUserObject, typename TAttribute< AttrType >::FGetter::template TConstMethodPtr< UserClass, Var1Type, Var2Type > InFunc, Var1Type Var1, Var2Type Var2 )	\
 		{ \
 			_##AttrName = TAttribute< AttrType >::Create( TAttribute< AttrType >::FGetter::CreateRaw( InUserObject, InFunc, Var1, Var2 ) ); \
 			return this->Me(); \
 		} \
 		template< class UserClass, typename Var1Type, typename Var2Type, typename Var3Type >	\
-		WidgetArgsType& AttrName##_Raw( UserClass* InUserObject, typename TAttribute< AttrType >::FGetter::template TRawMethodDelegate_ThreeVars_Const< UserClass, Var1Type, Var2Type, Var3Type >::FMethodPtr InFunc, Var1Type Var1, Var2Type Var2, Var3Type Var3 )	\
+		WidgetArgsType& AttrName##_Raw( UserClass* InUserObject, typename TAttribute< AttrType >::FGetter::template TConstMethodPtr< UserClass, Var1Type, Var2Type, Var3Type > InFunc, Var1Type Var1, Var2Type Var2, Var3Type Var3 )	\
 		{ \
 			_##AttrName = TAttribute< AttrType >::Create( TAttribute< AttrType >::FGetter::CreateRaw( InUserObject, InFunc, Var1, Var2, Var3 ) ); \
 			return this->Me(); \
 		} \
 		template< class UserClass, typename Var1Type, typename Var2Type, typename Var3Type, typename Var4Type >	\
-		WidgetArgsType& AttrName##_Raw( UserClass* InUserObject, typename TAttribute< AttrType >::FGetter::template TRawMethodDelegate_FourVars_Const< UserClass, Var1Type, Var2Type, Var3Type, Var4Type >::FMethodPtr InFunc, Var1Type Var1, Var2Type Var2, Var3Type Var3, Var4Type Var4 )	\
+		WidgetArgsType& AttrName##_Raw( UserClass* InUserObject, typename TAttribute< AttrType >::FGetter::template TConstMethodPtr< UserClass, Var1Type, Var2Type, Var3Type, Var4Type > InFunc, Var1Type Var1, Var2Type Var2, Var3Type Var3, Var4Type Var4 )	\
 		{ \
 			_##AttrName = TAttribute< AttrType >::Create( TAttribute< AttrType >::FGetter::CreateRaw( InUserObject, InFunc, Var1, Var2, Var3, Var4 ) ); \
 			return this->Me(); \
@@ -194,31 +194,31 @@ template<typename WidgetType> struct TSlateBaseNamedArgs;
 	\
 		/* Bind attribute with delegate to a shared pointer-based class method.  Slate mostly uses shared pointers so we use an overload for this type of binding. */ \
 		template< class UserClass >	\
-		WidgetArgsType& AttrName( TSharedRef< UserClass > InUserObjectRef, typename TAttribute< AttrType >::FGetter::template TSPMethodDelegate_Const< UserClass >::FMethodPtr InFunc )	\
+		WidgetArgsType& AttrName( TSharedRef< UserClass > InUserObjectRef, typename TAttribute< AttrType >::FGetter::template TConstMethodPtr< UserClass > InFunc )	\
 		{ \
 			_##AttrName = TAttribute< AttrType >( InUserObjectRef, InFunc ); \
 			return this->Me(); \
 		} \
 		template< class UserClass, typename Var1Type >	\
-		WidgetArgsType& AttrName( TSharedRef< UserClass > InUserObjectRef, typename TAttribute< AttrType >::FGetter::template TSPMethodDelegate_OneVar_Const< UserClass, Var1Type >::FMethodPtr InFunc, Var1Type Var1 )	\
+		WidgetArgsType& AttrName( TSharedRef< UserClass > InUserObjectRef, typename TAttribute< AttrType >::FGetter::template TConstMethodPtr< UserClass, Var1Type > InFunc, Var1Type Var1 )	\
 		{ \
 			_##AttrName = TAttribute< AttrType >::Create( TAttribute< AttrType >::FGetter::CreateSP( InUserObjectRef, InFunc, Var1 ) ); \
 			return this->Me(); \
 		} \
 		template< class UserClass, typename Var1Type, typename Var2Type >	\
-		WidgetArgsType& AttrName( TSharedRef< UserClass > InUserObjectRef, typename TAttribute< AttrType >::FGetter::template TSPMethodDelegate_TwoVars_Const< UserClass, Var1Type, Var2Type >::FMethodPtr InFunc, Var1Type Var1, Var2Type Var2 )	\
+		WidgetArgsType& AttrName( TSharedRef< UserClass > InUserObjectRef, typename TAttribute< AttrType >::FGetter::template TConstMethodPtr< UserClass, Var1Type, Var2Type > InFunc, Var1Type Var1, Var2Type Var2 )	\
 		{ \
 			_##AttrName = TAttribute< AttrType >::Create( TAttribute< AttrType >::FGetter::CreateSP( InUserObjectRef, InFunc, Var1, Var2 ) ); \
 			return this->Me(); \
 		} \
 		template< class UserClass, typename Var1Type, typename Var2Type, typename Var3Type >	\
-		WidgetArgsType& AttrName( TSharedRef< UserClass > InUserObjectRef, typename TAttribute< AttrType >::FGetter::template TSPMethodDelegate_ThreeVars_Const< UserClass, Var1Type, Var2Type, Var3Type >::FMethodPtr InFunc, Var1Type Var1, Var2Type Var2, Var3Type Var3 )	\
+		WidgetArgsType& AttrName( TSharedRef< UserClass > InUserObjectRef, typename TAttribute< AttrType >::FGetter::template TConstMethodPtr< UserClass, Var1Type, Var2Type, Var3Type > InFunc, Var1Type Var1, Var2Type Var2, Var3Type Var3 )	\
 		{ \
 			_##AttrName = TAttribute< AttrType >::Create( TAttribute< AttrType >::FGetter::CreateSP( InUserObjectRef, InFunc, Var1, Var2, Var3 ) ); \
 			return this->Me(); \
 		} \
 		template< class UserClass, typename Var1Type, typename Var2Type, typename Var3Type, typename Var4Type >	\
-		WidgetArgsType& AttrName( TSharedRef< UserClass > InUserObjectRef, typename TAttribute< AttrType >::FGetter::template TSPMethodDelegate_FourVars_Const< UserClass, Var1Type, Var2Type, Var3Type, Var4Type >::FMethodPtr InFunc, Var1Type Var1, Var2Type Var2, Var3Type Var3, Var4Type Var4 )	\
+		WidgetArgsType& AttrName( TSharedRef< UserClass > InUserObjectRef, typename TAttribute< AttrType >::FGetter::template TConstMethodPtr< UserClass, Var1Type, Var2Type, Var3Type, Var4Type > InFunc, Var1Type Var1, Var2Type Var2, Var3Type Var3, Var4Type Var4 )	\
 		{ \
 			_##AttrName = TAttribute< AttrType >::Create( TAttribute< AttrType >::FGetter::CreateSP( InUserObjectRef, InFunc, Var1, Var2, Var3, Var4 ) ); \
 			return this->Me(); \
@@ -226,31 +226,31 @@ template<typename WidgetType> struct TSlateBaseNamedArgs;
 	\
 		/* Bind attribute with delegate to a shared pointer-based class method.  Slate mostly uses shared pointers so we use an overload for this type of binding. */ \
 		template< class UserClass >	\
-		WidgetArgsType& AttrName( UserClass* InUserObject, typename TAttribute< AttrType >::FGetter::template TSPMethodDelegate_Const< UserClass >::FMethodPtr InFunc )	\
+		WidgetArgsType& AttrName( UserClass* InUserObject, typename TAttribute< AttrType >::FGetter::template TConstMethodPtr< UserClass > InFunc )	\
 		{ \
 			_##AttrName = TAttribute< AttrType >::Create( TAttribute< AttrType >::FGetter::CreateSP( InUserObject, InFunc ) ); \
 			return this->Me(); \
 		} \
 		template< class UserClass, typename Var1Type >	\
-		WidgetArgsType& AttrName( UserClass* InUserObject, typename TAttribute< AttrType >::FGetter::template TSPMethodDelegate_OneVar_Const< UserClass, Var1Type >::FMethodPtr InFunc, Var1Type Var1 )	\
+		WidgetArgsType& AttrName( UserClass* InUserObject, typename TAttribute< AttrType >::FGetter::template TConstMethodPtr< UserClass, Var1Type > InFunc, Var1Type Var1 )	\
 		{ \
 			_##AttrName = TAttribute< AttrType >::Create( TAttribute< AttrType >::FGetter::CreateSP( InUserObject, InFunc, Var1 ) ); \
 			return this->Me(); \
 		} \
 		template< class UserClass, typename Var1Type, typename Var2Type >	\
-		WidgetArgsType& AttrName( UserClass* InUserObject, typename TAttribute< AttrType >::FGetter::template TSPMethodDelegate_TwoVars_Const< UserClass, Var1Type, Var2Type >::FMethodPtr InFunc, Var1Type Var1, Var2Type Var2 )	\
+		WidgetArgsType& AttrName( UserClass* InUserObject, typename TAttribute< AttrType >::FGetter::template TConstMethodPtr< UserClass, Var1Type, Var2Type > InFunc, Var1Type Var1, Var2Type Var2 )	\
 		{ \
 			_##AttrName = TAttribute< AttrType >::Create( TAttribute< AttrType >::FGetter::CreateSP( InUserObject, InFunc, Var1, Var2 ) ); \
 			return this->Me(); \
 		} \
 		template< class UserClass, typename Var1Type, typename Var2Type, typename Var3Type >	\
-		WidgetArgsType& AttrName( UserClass* InUserObject, typename TAttribute< AttrType >::FGetter::template TSPMethodDelegate_ThreeVars_Const< UserClass, Var1Type, Var2Type, Var3Type >::FMethodPtr InFunc, Var1Type Var1, Var2Type Var2, Var3Type Var3 )	\
+		WidgetArgsType& AttrName( UserClass* InUserObject, typename TAttribute< AttrType >::FGetter::template TConstMethodPtr< UserClass, Var1Type, Var2Type, Var3Type > InFunc, Var1Type Var1, Var2Type Var2, Var3Type Var3 )	\
 		{ \
 			_##AttrName = TAttribute< AttrType >::Create( TAttribute< AttrType >::FGetter::CreateSP( InUserObject, InFunc, Var1, Var2, Var3 ) ); \
 			return this->Me(); \
 		} \
 		template< class UserClass, typename Var1Type, typename Var2Type, typename Var3Type, typename Var4Type >	\
-		WidgetArgsType& AttrName( UserClass* InUserObject, typename TAttribute< AttrType >::FGetter::template TSPMethodDelegate_FourVars_Const< UserClass, Var1Type, Var2Type, Var3Type, Var4Type >::FMethodPtr InFunc, Var1Type Var1, Var2Type Var2, Var3Type Var3, Var4Type Var4 )	\
+		WidgetArgsType& AttrName( UserClass* InUserObject, typename TAttribute< AttrType >::FGetter::template TConstMethodPtr< UserClass, Var1Type, Var2Type, Var3Type, Var4Type > InFunc, Var1Type Var1, Var2Type Var2, Var3Type Var3, Var4Type Var4 )	\
 		{ \
 			_##AttrName = TAttribute< AttrType >::Create( TAttribute< AttrType >::FGetter::CreateSP( InUserObject, InFunc, Var1, Var2, Var3, Var4 ) ); \
 			return this->Me(); \
@@ -258,31 +258,31 @@ template<typename WidgetType> struct TSlateBaseNamedArgs;
 	\
 		/* Bind attribute with delegate to a UObject-based class method */ \
 		template< class UserClass >	\
-		WidgetArgsType& AttrName##_UObject( UserClass* InUserObject, typename TAttribute< AttrType >::FGetter::template TUObjectMethodDelegate_Const< UserClass >::FMethodPtr InFunc )	\
+		WidgetArgsType& AttrName##_UObject( UserClass* InUserObject, typename TAttribute< AttrType >::FGetter::template TConstMethodPtr< UserClass > InFunc )	\
 		{ \
 			_##AttrName = TAttribute< AttrType >::Create( TAttribute< AttrType >::FGetter::CreateUObject( InUserObject, InFunc ) ); \
 			return this->Me(); \
 		} \
 		template< class UserClass, typename Var1Type >	\
-		WidgetArgsType& AttrName##_UObject( UserClass* InUserObject, typename TAttribute< AttrType >::FGetter::template TUObjectMethodDelegate_OneVar_Const< UserClass, Var1Type >::FMethodPtr InFunc, Var1Type Var1 )	\
+		WidgetArgsType& AttrName##_UObject( UserClass* InUserObject, typename TAttribute< AttrType >::FGetter::template TConstMethodPtr< UserClass, Var1Type > InFunc, Var1Type Var1 )	\
 		{ \
 			_##AttrName = TAttribute< AttrType >::Create( TAttribute< AttrType >::FGetter::CreateUObject( InUserObject, InFunc, Var1 ) ); \
 			return this->Me(); \
 		} \
 		template< class UserClass, typename Var1Type, typename Var2Type  >	\
-		WidgetArgsType& AttrName##_UObject( UserClass* InUserObject, typename TAttribute< AttrType >::FGetter::template TUObjectMethodDelegate_TwoVars_Const< UserClass, Var1Type, Var2Type >::FMethodPtr InFunc, Var1Type Var1, Var2Type Var2 )	\
+		WidgetArgsType& AttrName##_UObject( UserClass* InUserObject, typename TAttribute< AttrType >::FGetter::template TConstMethodPtr< UserClass, Var1Type, Var2Type > InFunc, Var1Type Var1, Var2Type Var2 )	\
 		{ \
 			_##AttrName = TAttribute< AttrType >::Create( TAttribute< AttrType >::FGetter::CreateUObject( InUserObject, InFunc, Var1, Var2 ) ); \
 			return this->Me(); \
 		} \
 		template< class UserClass, typename Var1Type, typename Var2Type, typename Var3Type >	\
-		WidgetArgsType& AttrName##_UObject( UserClass* InUserObject, typename TAttribute< AttrType >::FGetter::template TUObjectMethodDelegate_ThreeVars_Const< UserClass, Var1Type, Var2Type, Var3Type >::FMethodPtr InFunc, Var1Type Var1, Var2Type Var2, Var3Type Var3 )	\
+		WidgetArgsType& AttrName##_UObject( UserClass* InUserObject, typename TAttribute< AttrType >::FGetter::template TConstMethodPtr< UserClass, Var1Type, Var2Type, Var3Type > InFunc, Var1Type Var1, Var2Type Var2, Var3Type Var3 )	\
 		{ \
 			_##AttrName = TAttribute< AttrType >::Create( TAttribute< AttrType >::FGetter::CreateUObject( InUserObject, InFunc, Var1, Var2, Var3 ) ); \
 			return this->Me(); \
 		} \
 		template< class UserClass, typename Var1Type, typename Var2Type, typename Var3Type, typename Var4Type >	\
-		WidgetArgsType& AttrName##_UObject( UserClass* InUserObject, typename TAttribute< AttrType >::FGetter::template TUObjectMethodDelegate_FourVars_Const< UserClass, Var1Type, Var2Type, Var3Type, Var4Type >::FMethodPtr InFunc, Var1Type Var1, Var2Type Var2, Var3Type Var3, Var4Type Var4 )	\
+		WidgetArgsType& AttrName##_UObject( UserClass* InUserObject, typename TAttribute< AttrType >::FGetter::template TConstMethodPtr< UserClass, Var1Type, Var2Type, Var3Type, Var4Type > InFunc, Var1Type Var1, Var2Type Var2, Var3Type Var3, Var4Type Var4 )	\
 		{ \
 			_##AttrName = TAttribute< AttrType >::Create( TAttribute< AttrType >::FGetter::CreateUObject( InUserObject, InFunc, Var1, Var2, Var3, Var4 ) ); \
 			return this->Me(); \
@@ -525,61 +525,61 @@ struct NamedSlotProperty
 		\
 		/* Set event delegate to a raw C++ class method */ \
 		template< class UserClass >	\
-		WidgetArgsType& EventName##_Raw( UserClass* InUserObject, typename DelegateName::template TRawMethodDelegate< UserClass >::FMethodPtr InFunc )	\
+		WidgetArgsType& EventName##_Raw( UserClass* InUserObject, typename DelegateName::template TMethodPtr< UserClass > InFunc )	\
 		{ \
 			_##EventName = DelegateName::CreateRaw( InUserObject, InFunc ); \
 			return *this; \
 		} \
 		template< class UserClass >	\
-		WidgetArgsType& EventName##_Raw( UserClass* InUserObject, typename DelegateName::template TRawMethodDelegate_Const< UserClass >::FMethodPtr InFunc )	\
+		WidgetArgsType& EventName##_Raw( UserClass* InUserObject, typename DelegateName::template TConstMethodPtr< UserClass > InFunc )	\
 		{ \
 			_##EventName = DelegateName::CreateRaw( InUserObject, InFunc ); \
 			return *this; \
 		} \
 		template< class UserClass, typename Var1Type >	\
-		WidgetArgsType& EventName##_Raw( UserClass* InUserObject, typename DelegateName::template TRawMethodDelegate_OneVar< UserClass, Var1Type >::FMethodPtr InFunc, Var1Type Var1 )	\
+		WidgetArgsType& EventName##_Raw( UserClass* InUserObject, typename DelegateName::template TMethodPtr< UserClass, Var1Type > InFunc, Var1Type Var1 )	\
 		{ \
 			_##EventName = DelegateName::CreateRaw( InUserObject, InFunc, Var1 ); \
 			return *this; \
 		} \
 		template< class UserClass, typename Var1Type >	\
-		WidgetArgsType& EventName##_Raw( UserClass* InUserObject, typename DelegateName::template TRawMethodDelegate_OneVar_Const< UserClass, Var1Type >::FMethodPtr InFunc, Var1Type Var1 )	\
+		WidgetArgsType& EventName##_Raw( UserClass* InUserObject, typename DelegateName::template TConstMethodPtr< UserClass, Var1Type > InFunc, Var1Type Var1 )	\
 		{ \
 			_##EventName = DelegateName::CreateRaw( InUserObject, InFunc, Var1 ); \
 			return *this; \
 		} \
 		template< class UserClass, typename Var1Type, typename Var2Type >	\
-		WidgetArgsType& EventName##_Raw( UserClass* InUserObject, typename DelegateName::template TRawMethodDelegate_TwoVars< UserClass, Var1Type, Var2Type >::FMethodPtr InFunc, Var1Type Var1, Var2Type Var2 )	\
+		WidgetArgsType& EventName##_Raw( UserClass* InUserObject, typename DelegateName::template TMethodPtr< UserClass, Var1Type, Var2Type > InFunc, Var1Type Var1, Var2Type Var2 )	\
 		{ \
 			_##EventName = DelegateName::CreateRaw( InUserObject, InFunc, Var1, Var2 ); \
 			return *this; \
 		} \
 		template< class UserClass, typename Var1Type, typename Var2Type >	\
-		WidgetArgsType& EventName##_Raw( UserClass* InUserObject, typename DelegateName::template TRawMethodDelegate_TwoVars_Const< UserClass, Var1Type, Var2Type >::FMethodPtr InFunc, Var1Type Var1, Var2Type Var2 )	\
+		WidgetArgsType& EventName##_Raw( UserClass* InUserObject, typename DelegateName::template TConstMethodPtr< UserClass, Var1Type, Var2Type > InFunc, Var1Type Var1, Var2Type Var2 )	\
 		{ \
 			_##EventName = DelegateName::CreateRaw( InUserObject, InFunc, Var1, Var2 ); \
 			return *this; \
 		} \
 		template< class UserClass, typename Var1Type, typename Var2Type, typename Var3Type >	\
-		WidgetArgsType& EventName##_Raw( UserClass* InUserObject, typename DelegateName::template TRawMethodDelegate_ThreeVars< UserClass, Var1Type, Var2Type, Var3Type >::FMethodPtr InFunc, Var1Type Var1, Var2Type Var2, Var3Type Var3 )	\
+		WidgetArgsType& EventName##_Raw( UserClass* InUserObject, typename DelegateName::template TMethodPtr< UserClass, Var1Type, Var2Type, Var3Type > InFunc, Var1Type Var1, Var2Type Var2, Var3Type Var3 )	\
 		{ \
 			_##EventName = DelegateName::CreateRaw( InUserObject, InFunc, Var1, Var2, Var3 ); \
 			return *this; \
 		} \
 		template< class UserClass, typename Var1Type, typename Var2Type, typename Var3Type >	\
-		WidgetArgsType& EventName##_Raw( UserClass* InUserObject, typename DelegateName::template TRawMethodDelegate_ThreeVars_Const< UserClass, Var1Type, Var2Type, Var3Type >::FMethodPtr InFunc, Var1Type Var1, Var2Type Var2, Var3Type Var3 )	\
+		WidgetArgsType& EventName##_Raw( UserClass* InUserObject, typename DelegateName::template TConstMethodPtr< UserClass, Var1Type, Var2Type, Var3Type > InFunc, Var1Type Var1, Var2Type Var2, Var3Type Var3 )	\
 		{ \
 			_##EventName = DelegateName::CreateRaw( InUserObject, InFunc, Var1, Var2, Var3 ); \
 			return *this; \
 		} \
 		template< class UserClass, typename Var1Type, typename Var2Type, typename Var3Type, typename Var4Type >	\
-		WidgetArgsType& EventName##_Raw( UserClass* InUserObject, typename DelegateName::template TRawMethodDelegate_FourVars< UserClass, Var1Type, Var2Type, Var3Type, Var4Type >::FMethodPtr InFunc, Var1Type Var1, Var2Type Var2, Var3Type Var3, Var4Type Var4 )	\
+		WidgetArgsType& EventName##_Raw( UserClass* InUserObject, typename DelegateName::template TMethodPtr< UserClass, Var1Type, Var2Type, Var3Type, Var4Type > InFunc, Var1Type Var1, Var2Type Var2, Var3Type Var3, Var4Type Var4 )	\
 		{ \
 			_##EventName = DelegateName::CreateRaw( InUserObject, InFunc, Var1, Var2, Var3, Var4 ); \
 			return *this; \
 		} \
 		template< class UserClass, typename Var1Type, typename Var2Type, typename Var3Type, typename Var4Type >	\
-		WidgetArgsType& EventName##_Raw( UserClass* InUserObject, typename DelegateName::template TRawMethodDelegate_FourVars_Const< UserClass, Var1Type, Var2Type, Var3Type, Var4Type >::FMethodPtr InFunc, Var1Type Var1, Var2Type Var2, Var3Type Var3, Var4Type Var4 )	\
+		WidgetArgsType& EventName##_Raw( UserClass* InUserObject, typename DelegateName::template TConstMethodPtr< UserClass, Var1Type, Var2Type, Var3Type, Var4Type > InFunc, Var1Type Var1, Var2Type Var2, Var3Type Var3, Var4Type Var4 )	\
 		{ \
 			_##EventName = DelegateName::CreateRaw( InUserObject, InFunc, Var1, Var2, Var3, Var4 ); \
 			return *this; \
@@ -587,61 +587,61 @@ struct NamedSlotProperty
 		\
 		/* Set event delegate to a shared pointer-based class method.  Slate mostly uses shared pointers so we use an overload for this type of binding. */ \
 		template< class UserClass >	\
-		WidgetArgsType& EventName( TSharedRef< UserClass > InUserObjectRef, typename DelegateName::template TSPMethodDelegate< UserClass >::FMethodPtr InFunc )	\
+		WidgetArgsType& EventName( TSharedRef< UserClass > InUserObjectRef, typename DelegateName::template TMethodPtr< UserClass > InFunc )	\
 		{ \
 			_##EventName = DelegateName::CreateSP( InUserObjectRef, InFunc ); \
 			return *this; \
 		} \
 		template< class UserClass >	\
-		WidgetArgsType& EventName( TSharedRef< UserClass > InUserObjectRef, typename DelegateName::template TSPMethodDelegate_Const< UserClass >::FMethodPtr InFunc )	\
+		WidgetArgsType& EventName( TSharedRef< UserClass > InUserObjectRef, typename DelegateName::template TConstMethodPtr< UserClass > InFunc )	\
 		{ \
 			_##EventName = DelegateName::CreateSP( InUserObjectRef, InFunc ); \
 			return *this; \
 		} \
 		template< class UserClass, typename Var1Type >	\
-		WidgetArgsType& EventName( TSharedRef< UserClass > InUserObjectRef, typename DelegateName::template TSPMethodDelegate_OneVar< UserClass, Var1Type >::FMethodPtr InFunc, Var1Type Var1 )	\
+		WidgetArgsType& EventName( TSharedRef< UserClass > InUserObjectRef, typename DelegateName::template TMethodPtr< UserClass, Var1Type > InFunc, Var1Type Var1 )	\
 		{ \
 			_##EventName = DelegateName::CreateSP( InUserObjectRef, InFunc, Var1 ); \
 			return *this; \
 		} \
 		template< class UserClass, typename Var1Type >	\
-		WidgetArgsType& EventName( TSharedRef< UserClass > InUserObjectRef, typename DelegateName::template TSPMethodDelegate_OneVar_Const< UserClass, Var1Type >::FMethodPtr InFunc, Var1Type Var1 )	\
+		WidgetArgsType& EventName( TSharedRef< UserClass > InUserObjectRef, typename DelegateName::template TConstMethodPtr< UserClass, Var1Type > InFunc, Var1Type Var1 )	\
 		{ \
 			_##EventName = DelegateName::CreateSP( InUserObjectRef, InFunc, Var1 ); \
 			return *this; \
 		} \
 		template< class UserClass, typename Var1Type, typename Var2Type >	\
-		WidgetArgsType& EventName( TSharedRef< UserClass > InUserObjectRef, typename DelegateName::template TSPMethodDelegate_TwoVars< UserClass, Var1Type, Var2Type >::FMethodPtr InFunc, Var1Type Var1, Var2Type Var2 )	\
+		WidgetArgsType& EventName( TSharedRef< UserClass > InUserObjectRef, typename DelegateName::template TMethodPtr< UserClass, Var1Type, Var2Type > InFunc, Var1Type Var1, Var2Type Var2 )	\
 		{ \
 			_##EventName = DelegateName::CreateSP( InUserObjectRef, InFunc, Var1, Var2 ); \
 			return *this; \
 		} \
 		template< class UserClass, typename Var1Type, typename Var2Type >	\
-		WidgetArgsType& EventName( TSharedRef< UserClass > InUserObjectRef, typename DelegateName::template TSPMethodDelegate_TwoVars_Const< UserClass, Var1Type, Var2Type >::FMethodPtr InFunc, Var1Type Var1, Var2Type Var2 )	\
+		WidgetArgsType& EventName( TSharedRef< UserClass > InUserObjectRef, typename DelegateName::template TConstMethodPtr< UserClass, Var1Type, Var2Type > InFunc, Var1Type Var1, Var2Type Var2 )	\
 		{ \
 			_##EventName = DelegateName::CreateSP( InUserObjectRef, InFunc, Var1, Var2 ); \
 			return *this; \
 		} \
 		template< class UserClass, typename Var1Type, typename Var2Type, typename Var3Type >	\
-		WidgetArgsType& EventName( TSharedRef< UserClass > InUserObjectRef, typename DelegateName::template TSPMethodDelegate_ThreeVars< UserClass, Var1Type, Var2Type, Var3Type >::FMethodPtr InFunc, Var1Type Var1, Var2Type Var2, Var3Type Var3 )	\
+		WidgetArgsType& EventName( TSharedRef< UserClass > InUserObjectRef, typename DelegateName::template TMethodPtr< UserClass, Var1Type, Var2Type, Var3Type > InFunc, Var1Type Var1, Var2Type Var2, Var3Type Var3 )	\
 		{ \
 			_##EventName = DelegateName::CreateSP( InUserObjectRef, InFunc, Var1, Var2, Var3 ); \
 			return *this; \
 		} \
 		template< class UserClass, typename Var1Type, typename Var2Type, typename Var3Type >	\
-		WidgetArgsType& EventName( TSharedRef< UserClass > InUserObjectRef, typename DelegateName::template TSPMethodDelegate_ThreeVars_Const< UserClass, Var1Type, Var2Type, Var3Type >::FMethodPtr InFunc, Var1Type Var1, Var2Type Var2, Var3Type Var3 )	\
+		WidgetArgsType& EventName( TSharedRef< UserClass > InUserObjectRef, typename DelegateName::template TConstMethodPtr< UserClass, Var1Type, Var2Type, Var3Type > InFunc, Var1Type Var1, Var2Type Var2, Var3Type Var3 )	\
 		{ \
 			_##EventName = DelegateName::CreateSP( InUserObjectRef, InFunc, Var1, Var2, Var3 ); \
 			return *this; \
 		} \
 		template< class UserClass, typename Var1Type, typename Var2Type, typename Var3Type, typename Var4Type >	\
-		WidgetArgsType& EventName( TSharedRef< UserClass > InUserObjectRef, typename DelegateName::template TSPMethodDelegate_FourVars< UserClass, Var1Type, Var2Type, Var3Type, Var4Type >::FMethodPtr InFunc, Var1Type Var1, Var2Type Var2, Var3Type Var3, Var4Type Var4 )	\
+		WidgetArgsType& EventName( TSharedRef< UserClass > InUserObjectRef, typename DelegateName::template TMethodPtr< UserClass, Var1Type, Var2Type, Var3Type, Var4Type > InFunc, Var1Type Var1, Var2Type Var2, Var3Type Var3, Var4Type Var4 )	\
 		{ \
 			_##EventName = DelegateName::CreateSP( InUserObjectRef, InFunc, Var1, Var2, Var3, Var4 ); \
 			return *this; \
 		} \
 		template< class UserClass, typename Var1Type, typename Var2Type, typename Var3Type, typename Var4Type >	\
-		WidgetArgsType& EventName( TSharedRef< UserClass > InUserObjectRef, typename DelegateName::template TSPMethodDelegate_FourVars_Const< UserClass, Var1Type, Var2Type, Var3Type, Var4Type >::FMethodPtr InFunc, Var1Type Var1, Var2Type Var2, Var3Type Var3, Var4Type Var4 )	\
+		WidgetArgsType& EventName( TSharedRef< UserClass > InUserObjectRef, typename DelegateName::template TConstMethodPtr< UserClass, Var1Type, Var2Type, Var3Type, Var4Type > InFunc, Var1Type Var1, Var2Type Var2, Var3Type Var3, Var4Type Var4 )	\
 		{ \
 			_##EventName = DelegateName::CreateSP( InUserObjectRef, InFunc, Var1, Var2, Var3, Var4 ); \
 			return *this; \
@@ -649,61 +649,61 @@ struct NamedSlotProperty
 		\
 		/* Set event delegate to a shared pointer-based class method.  Slate mostly uses shared pointers so we use an overload for this type of binding. */ \
 		template< class UserClass >	\
-		WidgetArgsType& EventName( UserClass* InUserObject, typename DelegateName::template TSPMethodDelegate< UserClass >::FMethodPtr InFunc )	\
+		WidgetArgsType& EventName( UserClass* InUserObject, typename DelegateName::template TMethodPtr< UserClass > InFunc )	\
 		{ \
 			_##EventName = DelegateName::CreateSP( InUserObject, InFunc ); \
 			return *this; \
 		} \
 		template< class UserClass >	\
-		WidgetArgsType& EventName( UserClass* InUserObject, typename DelegateName::template TSPMethodDelegate_Const< UserClass >::FMethodPtr InFunc )	\
+		WidgetArgsType& EventName( UserClass* InUserObject, typename DelegateName::template TConstMethodPtr< UserClass > InFunc )	\
 		{ \
 			_##EventName = DelegateName::CreateSP( InUserObject, InFunc ); \
 			return *this; \
 		} \
 		template< class UserClass, typename Var1Type >	\
-		WidgetArgsType& EventName( UserClass* InUserObject, typename DelegateName::template TSPMethodDelegate_OneVar< UserClass, Var1Type >::FMethodPtr InFunc, Var1Type Var1 )	\
+		WidgetArgsType& EventName( UserClass* InUserObject, typename DelegateName::template TMethodPtr< UserClass, Var1Type > InFunc, Var1Type Var1 )	\
 		{ \
 			_##EventName = DelegateName::CreateSP( InUserObject, InFunc, Var1 ); \
 			return *this; \
 		} \
 		template< class UserClass, typename Var1Type >	\
-		WidgetArgsType& EventName( UserClass* InUserObject, typename DelegateName::template TSPMethodDelegate_OneVar_Const< UserClass, Var1Type >::FMethodPtr InFunc, Var1Type Var1 )	\
+		WidgetArgsType& EventName( UserClass* InUserObject, typename DelegateName::template TConstMethodPtr< UserClass, Var1Type > InFunc, Var1Type Var1 )	\
 		{ \
 			_##EventName = DelegateName::CreateSP( InUserObject, InFunc, Var1 ); \
 			return *this; \
 		} \
 		template< class UserClass, typename Var1Type, typename Var2Type >	\
-		WidgetArgsType& EventName( UserClass* InUserObject, typename DelegateName::template TSPMethodDelegate_TwoVars< UserClass, Var1Type, Var2Type >::FMethodPtr InFunc, Var1Type Var1, Var2Type Var2 )	\
+		WidgetArgsType& EventName( UserClass* InUserObject, typename DelegateName::template TMethodPtr< UserClass, Var1Type, Var2Type > InFunc, Var1Type Var1, Var2Type Var2 )	\
 		{ \
 			_##EventName = DelegateName::CreateSP( InUserObject, InFunc, Var1, Var2 ); \
 			return *this; \
 		} \
 		template< class UserClass, typename Var1Type, typename Var2Type >	\
-		WidgetArgsType& EventName( UserClass* InUserObject, typename DelegateName::template TSPMethodDelegate_TwoVars_Const< UserClass, Var1Type, Var2Type >::FMethodPtr InFunc, Var1Type Var1, Var2Type Var2 )	\
+		WidgetArgsType& EventName( UserClass* InUserObject, typename DelegateName::template TConstMethodPtr< UserClass, Var1Type, Var2Type > InFunc, Var1Type Var1, Var2Type Var2 )	\
 		{ \
 			_##EventName = DelegateName::CreateSP( InUserObject, InFunc, Var1, Var2 ); \
 			return *this; \
 		} \
 		template< class UserClass, typename Var1Type, typename Var2Type, typename Var3Type >	\
-		WidgetArgsType& EventName( UserClass* InUserObject, typename DelegateName::template TSPMethodDelegate_ThreeVars< UserClass, Var1Type, Var2Type, Var3Type >::FMethodPtr InFunc, Var1Type Var1, Var2Type Var2, Var3Type Var3 )	\
+		WidgetArgsType& EventName( UserClass* InUserObject, typename DelegateName::template TMethodPtr< UserClass, Var1Type, Var2Type, Var3Type > InFunc, Var1Type Var1, Var2Type Var2, Var3Type Var3 )	\
 		{ \
 			_##EventName = DelegateName::CreateSP( InUserObject, InFunc, Var1, Var2, Var3 ); \
 			return *this; \
 		} \
 		template< class UserClass, typename Var1Type, typename Var2Type, typename Var3Type >	\
-		WidgetArgsType& EventName( UserClass* InUserObject, typename DelegateName::template TSPMethodDelegate_ThreeVars_Const< UserClass, Var1Type, Var2Type, Var3Type >::FMethodPtr InFunc, Var1Type Var1, Var2Type Var2, Var3Type Var3 )	\
+		WidgetArgsType& EventName( UserClass* InUserObject, typename DelegateName::template TConstMethodPtr< UserClass, Var1Type, Var2Type, Var3Type > InFunc, Var1Type Var1, Var2Type Var2, Var3Type Var3 )	\
 		{ \
 			_##EventName = DelegateName::CreateSP( InUserObject, InFunc, Var1, Var2, Var3 ); \
 			return *this; \
 		} \
 		template< class UserClass, typename Var1Type, typename Var2Type, typename Var3Type, typename Var4Type >	\
-		WidgetArgsType& EventName( UserClass* InUserObject, typename DelegateName::template TSPMethodDelegate_FourVars< UserClass, Var1Type, Var2Type, Var3Type, Var4Type >::FMethodPtr InFunc, Var1Type Var1, Var2Type Var2, Var3Type Var3, Var4Type Var4 )	\
+		WidgetArgsType& EventName( UserClass* InUserObject, typename DelegateName::template TMethodPtr< UserClass, Var1Type, Var2Type, Var3Type, Var4Type > InFunc, Var1Type Var1, Var2Type Var2, Var3Type Var3, Var4Type Var4 )	\
 		{ \
 			_##EventName = DelegateName::CreateSP( InUserObject, InFunc, Var1, Var2, Var3, Var4 ); \
 			return *this; \
 		} \
 		template< class UserClass, typename Var1Type, typename Var2Type, typename Var3Type, typename Var4Type >	\
-		WidgetArgsType& EventName( UserClass* InUserObject, typename DelegateName::template TSPMethodDelegate_FourVars_Const< UserClass, Var1Type, Var2Type, Var3Type, Var4Type >::FMethodPtr InFunc, Var1Type Var1, Var2Type Var2, Var3Type Var3, Var4Type Var4 )	\
+		WidgetArgsType& EventName( UserClass* InUserObject, typename DelegateName::template TConstMethodPtr< UserClass, Var1Type, Var2Type, Var3Type, Var4Type > InFunc, Var1Type Var1, Var2Type Var2, Var3Type Var3, Var4Type Var4 )	\
 		{ \
 			_##EventName = DelegateName::CreateSP( InUserObject, InFunc, Var1, Var2, Var3, Var4 ); \
 			return *this; \
@@ -711,61 +711,61 @@ struct NamedSlotProperty
 		\
 		/* Set event delegate to a UObject-based class method */ \
 		template< class UserClass >	\
-		WidgetArgsType& EventName##_UObject( UserClass* InUserObject, typename DelegateName::template TUObjectMethodDelegate< UserClass >::FMethodPtr InFunc )	\
+		WidgetArgsType& EventName##_UObject( UserClass* InUserObject, typename DelegateName::template TMethodPtr< UserClass > InFunc )	\
 		{ \
 			_##EventName = DelegateName::CreateUObject( InUserObject, InFunc ); \
 			return *this; \
 		} \
 		template< class UserClass >	\
-		WidgetArgsType& EventName##_UObject( UserClass* InUserObject, typename DelegateName::template TUObjectMethodDelegate_Const< UserClass >::FMethodPtr InFunc )	\
+		WidgetArgsType& EventName##_UObject( UserClass* InUserObject, typename DelegateName::template TConstMethodPtr< UserClass > InFunc )	\
 		{ \
 			_##EventName = DelegateName::CreateUObject( InUserObject, InFunc ); \
 			return *this; \
 		} \
 		template< class UserClass, typename Var1Type >	\
-		WidgetArgsType& EventName##_UObject( UserClass* InUserObject, typename DelegateName::template TUObjectMethodDelegate_OneVar< UserClass, Var1Type >::FMethodPtr InFunc, Var1Type Var1 )	\
+		WidgetArgsType& EventName##_UObject( UserClass* InUserObject, typename DelegateName::template TMethodPtr< UserClass, Var1Type > InFunc, Var1Type Var1 )	\
 		{ \
 			_##EventName = DelegateName::CreateUObject( InUserObject, InFunc, Var1 ); \
 			return *this; \
 		} \
 		template< class UserClass, typename Var1Type >	\
-		WidgetArgsType& EventName##_UObject( UserClass* InUserObject, typename DelegateName::template TUObjectMethodDelegate_OneVar_Const< UserClass, Var1Type >::FMethodPtr InFunc, Var1Type Var1 )	\
+		WidgetArgsType& EventName##_UObject( UserClass* InUserObject, typename DelegateName::template TConstMethodPtr< UserClass, Var1Type > InFunc, Var1Type Var1 )	\
 		{ \
 			_##EventName = DelegateName::CreateUObject( InUserObject, InFunc, Var1 ); \
 			return *this; \
 		} \
 		template< class UserClass, typename Var1Type, typename Var2Type >	\
-		WidgetArgsType& EventName##_UObject( UserClass* InUserObject, typename DelegateName::template TUObjectMethodDelegate_TwoVars< UserClass, Var1Type, Var2Type >::FMethodPtr InFunc, Var1Type Var1, Var2Type Var2 )	\
+		WidgetArgsType& EventName##_UObject( UserClass* InUserObject, typename DelegateName::template TMethodPtr< UserClass, Var1Type, Var2Type > InFunc, Var1Type Var1, Var2Type Var2 )	\
 		{ \
 			_##EventName = DelegateName::CreateUObject( InUserObject, InFunc, Var1, Var2 ); \
 			return *this; \
 		} \
 		template< class UserClass, typename Var1Type, typename Var2Type >	\
-		WidgetArgsType& EventName##_UObject( UserClass* InUserObject, typename DelegateName::template TUObjectMethodDelegate_TwoVars_Const< UserClass, Var1Type, Var2Type >::FMethodPtr InFunc, Var1Type Var1, Var2Type Var2 )	\
+		WidgetArgsType& EventName##_UObject( UserClass* InUserObject, typename DelegateName::template TConstMethodPtr< UserClass, Var1Type, Var2Type > InFunc, Var1Type Var1, Var2Type Var2 )	\
 		{ \
 			_##EventName = DelegateName::CreateUObject( InUserObject, InFunc, Var1, Var2 ); \
 			return *this; \
 		} \
 		template< class UserClass, typename Var1Type, typename Var2Type, typename Var3Type >	\
-		WidgetArgsType& EventName##_UObject( UserClass* InUserObject, typename DelegateName::template TUObjectMethodDelegate_ThreeVars< UserClass, Var1Type, Var2Type, Var3Type >::FMethodPtr InFunc, Var1Type Var1, Var2Type Var2, Var3Type Var3 )	\
+		WidgetArgsType& EventName##_UObject( UserClass* InUserObject, typename DelegateName::template TMethodPtr< UserClass, Var1Type, Var2Type, Var3Type > InFunc, Var1Type Var1, Var2Type Var2, Var3Type Var3 )	\
 		{ \
 			_##EventName = DelegateName::CreateUObject( InUserObject, InFunc, Var1, Var2, Var3 ); \
 			return *this; \
 		} \
 		template< class UserClass, typename Var1Type, typename Var2Type, typename Var3Type >	\
-		WidgetArgsType& EventName##_UObject( UserClass* InUserObject, typename DelegateName::template TUObjectMethodDelegate_ThreeVars_Const< UserClass, Var1Type, Var2Type, Var3Type >::FMethodPtr InFunc, Var1Type Var1, Var2Type Var2, Var3Type Var3 )	\
+		WidgetArgsType& EventName##_UObject( UserClass* InUserObject, typename DelegateName::template TConstMethodPtr< UserClass, Var1Type, Var2Type, Var3Type > InFunc, Var1Type Var1, Var2Type Var2, Var3Type Var3 )	\
 		{ \
 			_##EventName = DelegateName::CreateUObject( InUserObject, InFunc, Var1, Var2, Var3 ); \
 			return *this; \
 		} \
 		template< class UserClass, typename Var1Type, typename Var2Type, typename Var3Type, typename Var4Type >	\
-		WidgetArgsType& EventName##_UObject( UserClass* InUserObject, typename DelegateName::template TUObjectMethodDelegate_FourVars< UserClass, Var1Type, Var2Type, Var3Type, Var4Type >::FMethodPtr InFunc, Var1Type Var1, Var2Type Var2, Var3Type Var3, Var4Type Var4 )	\
+		WidgetArgsType& EventName##_UObject( UserClass* InUserObject, typename DelegateName::template TMethodPtr< UserClass, Var1Type, Var2Type, Var3Type, Var4Type > InFunc, Var1Type Var1, Var2Type Var2, Var3Type Var3, Var4Type Var4 )	\
 		{ \
 			_##EventName = DelegateName::CreateUObject( InUserObject, InFunc, Var1, Var2, Var3, Var4 ); \
 			return *this; \
 		} \
 		template< class UserClass, typename Var1Type, typename Var2Type, typename Var3Type, typename Var4Type >	\
-		WidgetArgsType& EventName##_UObject( UserClass* InUserObject, typename DelegateName::template TUObjectMethodDelegate_FourVars_Const< UserClass, Var1Type, Var2Type, Var3Type, Var4Type >::FMethodPtr InFunc, Var1Type Var1, Var2Type Var2, Var3Type Var3, Var4Type Var4 )	\
+		WidgetArgsType& EventName##_UObject( UserClass* InUserObject, typename DelegateName::template TConstMethodPtr< UserClass, Var1Type, Var2Type, Var3Type, Var4Type > InFunc, Var1Type Var1, Var2Type Var2, Var3Type Var3, Var4Type Var4 )	\
 		{ \
 			_##EventName = DelegateName::CreateUObject( InUserObject, InFunc, Var1, Var2, Var3, Var4 ); \
 			return *this; \
diff --git a/Engine/Source/Runtime/UMG/Public/Components/TreeView.h b/Engine/Source/Runtime/UMG/Public/Components/TreeView.h
--- a/Engine/Source/Runtime/UMG/Public/Components/TreeView.h
+++ b/Engine/Source/Runtime/UMG/Public/Components/TreeView.h
@@ -38,13 +38,13 @@ public:
 	void CollapseAll();
 
 	template <typename ObjectT>
-	void SetOnGetItemChildren(ObjectT* InUserObject, typename TSlateDelegates<UObject*>::FOnGetChildren::TUObjectMethodDelegate<ObjectT>::FMethodPtr InMethodPtr)
+	void SetOnGetItemChildren(ObjectT* InUserObject, typename TSlateDelegates<UObject*>::FOnGetChildren::TMethodPtr<ObjectT> InMethodPtr)
 	{
 		static_assert(TIsDerivedFrom<ObjectT, UObject>::IsDerived, "Only UObject ptrs can be passed directly when binding to SetOnGetItemChildren. Pass a shared ref for non-UObject binders.");
 		OnGetItemChildren = TSlateDelegates<UObject*>::FOnGetChildren::CreateUObject(InUserObject, InMethodPtr);
 	}
 	template <typename ObjectT>
-	void SetOnGetItemChildren(TSharedRef<ObjectT> InUserObject, typename TSlateDelegates<UObject*>::FOnGetChildren::TUObjectMethodDelegate<ObjectT>::FMethodPtr InMethodPtr)
+	void SetOnGetItemChildren(TSharedRef<ObjectT> InUserObject, typename TSlateDelegates<UObject*>::FOnGetChildren::TMethodPtr<ObjectT> InMethodPtr)
 	{
 		OnGetItemChildren = TSlateDelegates<UObject*>::FOnGetChildren::CreateSP(InUserObject, InMethodPtr);
 	}
