Slate: SlateAttributeMemberRef now takes a TSharedRef instead of a reference to prevent circular inclusion. Split the SlateAttribute.h file into multiple files (other types are in progress and the file is getting bulky). Save some space in TSlateAttributeMetaData::FGetterItem (more variable are coming up). The maximum number of the defined attribute is now 255 (the biggest SWidget has ~25 attributes). Fix issue with bHasUpdatedManuallyFlagToReset that would reset the flag when we update the AffectVisibility attribute (the other attributes didn't have the chance to get updated). Prevent invalidation when the widget is not yet constructed (save some process time and space in the InvalidationRoot). Set bHasRegisteredSlateAttribute in ~SWidget to prevent invalidation when the widget is getting destroyed.

diff --git a/Engine/Source/Runtime/Slate/Public/Widgets/Colors/SColorGradingWheel.h b/Engine/Source/Runtime/Slate/Public/Widgets/Colors/SColorGradingWheel.h
--- a/Engine/Source/Runtime/Slate/Public/Widgets/Colors/SColorGradingWheel.h
+++ b/Engine/Source/Runtime/Slate/Public/Widgets/Colors/SColorGradingWheel.h
@@ -96,13 +96,13 @@ protected:
 	void SetExponentDisplacementAttribute(TAttribute<float> InExponentDisplacement);
 
 	/** @return an attribute reference of SelectedColor */
-	TSlateAttributeRef<FLinearColor> GetSelectedColorAttribute() const { return TSlateAttributeRef<FLinearColor>(*this, SelectedColorAttribute); }
+	TSlateAttributeRef<FLinearColor> GetSelectedColorAttribute() const { return TSlateAttributeRef<FLinearColor>(SharedThis(this), SelectedColorAttribute); }
 
 	/** @return an attribute reference of DesiredWheelSize */
-	TSlateAttributeRef<int32> GetDesiredWheelSizeAttribute() const { return TSlateAttributeRef<int32>(*this, DesiredWheelSizeAttribute); }
+	TSlateAttributeRef<int32> GetDesiredWheelSizeAttribute() const { return TSlateAttributeRef<int32>(SharedThis(this), DesiredWheelSizeAttribute); }
 
 	/** @return an attribute reference of ExponentDisplacement */
-	TSlateAttributeRef<float> GetExponentDisplacementAttribute() const { return TSlateAttributeRef<float>(*this, ExponentDisplacementAttribute); }
+	TSlateAttributeRef<float> GetExponentDisplacementAttribute() const { return TSlateAttributeRef<float>(SharedThis(this), ExponentDisplacementAttribute); }
 
 	/** The color wheel image to show. */
 	const FSlateBrush* Image;
diff --git a/Engine/Source/Runtime/Slate/Public/Widgets/Layout/SBackgroundBlur.h b/Engine/Source/Runtime/Slate/Public/Widgets/Layout/SBackgroundBlur.h
--- a/Engine/Source/Runtime/Slate/Public/Widgets/Layout/SBackgroundBlur.h
+++ b/Engine/Source/Runtime/Slate/Public/Widgets/Layout/SBackgroundBlur.h
@@ -65,10 +65,10 @@ protected:
 
 
 	/** @return an attribute reference of ColorAndOpacity */
-	TSlateAttributeRef<float> GetBlurStrengthAttribute() const { return TSlateAttributeRef<float>{*this, BlurStrengthAttribute}; }
+	TSlateAttributeRef<float> GetBlurStrengthAttribute() const { return TSlateAttributeRef<float>{SharedThis(this), BlurStrengthAttribute}; }
 
 	/** @return an attribute reference of ForegroundColor */
-	TSlateAttributeRef<TOptional<int32>> GetBlurRadiusAttribute() const { return TSlateAttributeRef<TOptional<int32>>{*this, BlurRadiusAttribute}; }
+	TSlateAttributeRef<TOptional<int32>> GetBlurRadiusAttribute() const { return TSlateAttributeRef<TOptional<int32>>{SharedThis(this), BlurRadiusAttribute}; }
 
 #if WITH_EDITORONLY_DATA
 	FSlateDeprecatedTAttribute<float> BlurStrength;
diff --git a/Engine/Source/Runtime/Slate/Public/Widgets/Layout/SBorder.h b/Engine/Source/Runtime/Slate/Public/Widgets/Layout/SBorder.h
--- a/Engine/Source/Runtime/Slate/Public/Widgets/Layout/SBorder.h
+++ b/Engine/Source/Runtime/Slate/Public/Widgets/Layout/SBorder.h
@@ -140,10 +140,10 @@ protected:
 	FSlateColor GetBorderBackgroundColor() const { return BorderBackgroundColorAttribute.Get(); }
 	FVector2D GetDesiredSizeScale() const { return DesiredSizeScaleAttribute.Get(); }
 	bool GetShowDisabledEffect() const { return ShowDisabledEffectAttribute.Get(); }
-	TSlateAttributeRef<const FSlateBrush*> GetBorderImageAttribute() const { return TSlateAttributeRef<const FSlateBrush*>(*this, BorderImageAttribute); }
-	TSlateAttributeRef<FSlateColor> GetBorderBackgroundColorAttribute() const { return TSlateAttributeRef<FSlateColor>(*this, BorderBackgroundColorAttribute); }
-	TSlateAttributeRef<FVector2D> GetDesiredSizeScaleAttribute() const { return TSlateAttributeRef<FVector2D>(*this, DesiredSizeScaleAttribute); }
-	TSlateAttributeRef<bool> GetShowDisabledEffectAttribute() const { return TSlateAttributeRef<bool>(*this, ShowDisabledEffectAttribute); }
+	TSlateAttributeRef<const FSlateBrush*> GetBorderImageAttribute() const { return TSlateAttributeRef<const FSlateBrush*>(SharedThis(this), BorderImageAttribute); }
+	TSlateAttributeRef<FSlateColor> GetBorderBackgroundColorAttribute() const { return TSlateAttributeRef<FSlateColor>(SharedThis(this), BorderBackgroundColorAttribute); }
+	TSlateAttributeRef<FVector2D> GetDesiredSizeScaleAttribute() const { return TSlateAttributeRef<FVector2D>(SharedThis(this), DesiredSizeScaleAttribute); }
+	TSlateAttributeRef<bool> GetShowDisabledEffectAttribute() const { return TSlateAttributeRef<bool>(SharedThis(this), ShowDisabledEffectAttribute); }
 
 #if WITH_EDITORONLY_DATA
 	UE_DEPRECATED(5.0, "Direct access to BorderImage is now deprecated. Use the setter or getter.")
diff --git a/Engine/Source/Runtime/Slate/Public/Widgets/Layout/SDPIScaler.h b/Engine/Source/Runtime/Slate/Public/Widgets/Layout/SDPIScaler.h
--- a/Engine/Source/Runtime/Slate/Public/Widgets/Layout/SDPIScaler.h
+++ b/Engine/Source/Runtime/Slate/Public/Widgets/Layout/SDPIScaler.h
@@ -58,7 +58,7 @@ protected:
 
 	virtual float GetRelativeLayoutScale(int32 ChildIndex, float LayoutScaleMultiplier) const override;
 
-	TSlateAttributeRef<float> GetDPIScaleAttribute() const { return TSlateAttributeRef<float>(*this, DPIScaleAttribute); }
+	TSlateAttributeRef<float> GetDPIScaleAttribute() const { return TSlateAttributeRef<float>(SharedThis(this), DPIScaleAttribute); }
 
 #if WITH_EDITORONLY_DATA
 	UE_DEPRECATED(5.0, "Direct access to DPIScale is now deprecated. Use the setter or getter.")
diff --git a/Engine/Source/Runtime/SlateCore/Private/Test/SlateAttributeTest.cpp b/Engine/Source/Runtime/SlateCore/Private/Test/SlateAttributeTest.cpp
--- a/Engine/Source/Runtime/SlateCore/Private/Test/SlateAttributeTest.cpp
+++ b/Engine/Source/Runtime/SlateCore/Private/Test/SlateAttributeTest.cpp
@@ -8,6 +8,7 @@
 #include "Widgets/SNullWidget.h"
 
 #include <type_traits>
+#include <limits>
 
 #if WITH_AUTOMATION_WORKER
 
@@ -224,7 +225,7 @@ bool FSlateAttributeTest::RunTest(const FString& Parameters)
 			, TEXT("The static data do not matches"));
 
 		FSlateAttributeDescriptor const& AttributeDescriptor = WidgetParent->GetWidgetClass().GetAttributeDescriptor();
-		AddErrorIfFalse(AttributeDescriptor.AttributeNum() == 4 + NumberOfAttributeInSWidget, TEXT(""));
+		AddErrorIfFalse(AttributeDescriptor.GetAttributeNum() == 4 + NumberOfAttributeInSWidget, TEXT("Invalid number of attributes"));
 
 		const int32 IndexA = AttributeDescriptor.IndexOfMemberAttribute("IntAttributeA");
 		const int32 IndexB = AttributeDescriptor.IndexOfMemberAttribute("IntAttributeB");
@@ -340,7 +341,7 @@ bool FSlateAttributeTest::RunTest(const FString& Parameters)
 			, TEXT("The static data do not matches"));
 
 		FSlateAttributeDescriptor const& AttributeDescriptor = WidgetChild->GetWidgetClass().GetAttributeDescriptor();
-		AddErrorIfFalse(AttributeDescriptor.AttributeNum() == 9 + NumberOfAttributeInSWidget, TEXT("")); // H is not counted
+		AddErrorIfFalse(AttributeDescriptor.GetAttributeNum() == 9 + NumberOfAttributeInSWidget, TEXT("Invalid number of attributes")); // H is not counted
 
 		const int32 IndexA = AttributeDescriptor.IndexOfMemberAttribute("IntAttributeA");
 		const int32 IndexB = AttributeDescriptor.IndexOfMemberAttribute("IntAttributeB");
@@ -441,6 +442,42 @@ bool FSlateAttributeTest::RunTest(const FString& Parameters)
 			AddErrorIfFalse(WidgetChild->IntAttributeL.Get() == 7, TEXT("L It is not the expected value."));
 			AddErrorIfFalse(WidgetChild->IntAttributeM.Get() == 1, TEXT("M It is not the expected value."));
 		}
+
+		// Check the ForEachDependency result
+		{
+			const FSlateAttributeDescriptor& ChildDescriptor = WidgetChild->GetWidgetClass().GetAttributeDescriptor();
+			auto DependencyTest = [this, ChildDescriptor](auto& AttributeInstance, TSharedRef<SWidget> Widget, int32 Expected, FStringView VariableName)
+			{
+				auto FindOffet = [](const SWidget& OwningWidget, const FSlateAttributeBase& Attribute)
+				{
+					UPTRINT Offset = (UPTRINT)(&Attribute) - (UPTRINT)(&OwningWidget);
+					ensure(Offset <= std::numeric_limits<FSlateAttributeDescriptor::OffsetType>::max());
+					return (FSlateAttributeDescriptor::OffsetType)(Offset);
+				};
+
+				int32 Count = 0;
+				const FSlateAttributeDescriptor::FAttribute* FoundAttribute = ChildDescriptor.FindMemberAttribute(FindOffet(Widget.Get(), AttributeInstance));
+				if (FoundAttribute == nullptr)
+				{
+					AddError(FString::Printf(TEXT("Could not find attribute '%s'"), VariableName.GetData()));
+				}
+				else
+				{
+					ChildDescriptor.ForEachDependentsOn(*FoundAttribute, [&Count](uint8 Index) {++Count;});
+					AddErrorIfFalse(Count == Expected, FString::Printf(TEXT("'%s' doesn't have the correct number of dependency (returned %d, expected %d).")
+						, VariableName.GetData(), Count, Expected));
+				}
+			};
+			DependencyTest(WidgetChild->IntAttributeA, WidgetChild, 1, TEXT("A"));//J
+			DependencyTest(WidgetChild->IntAttributeB, WidgetChild, 0, TEXT("B"));//AIJ, they are prerequisite, not dependency
+			DependencyTest(WidgetChild->IntAttributeC, WidgetChild, 0, TEXT("C"));//L
+			DependencyTest(WidgetChild->IntAttributeD, WidgetChild, 1, TEXT("D"));//CL
+			DependencyTest(WidgetChild->IntAttributeI, WidgetChild, 0, TEXT("I"));
+			DependencyTest(WidgetChild->IntAttributeJ, WidgetChild, 0, TEXT("J"));
+			DependencyTest(WidgetChild->IntAttributeK, WidgetChild, 0, TEXT("K"));
+			DependencyTest(WidgetChild->IntAttributeL, WidgetChild, 0, TEXT("L"));
+			DependencyTest(WidgetChild->IntAttributeM, WidgetChild, 0, TEXT("M"));
+		}
 	}
 
 
diff --git a/Engine/Source/Runtime/SlateCore/Private/Types/SlateAttributeDescriptor.cpp b/Engine/Source/Runtime/SlateCore/Private/Types/SlateAttributeDescriptor.cpp
--- a/Engine/Source/Runtime/SlateCore/Private/Types/SlateAttributeDescriptor.cpp
+++ b/Engine/Source/Runtime/SlateCore/Private/Types/SlateAttributeDescriptor.cpp
@@ -59,7 +59,8 @@ FSlateAttributeDescriptor::FInitializer::FInitializer(FSlateAttributeDescriptor&
 
 FSlateAttributeDescriptor::FInitializer::~FInitializer()
 {
-	checkf(Descriptor.Attributes.Num() < std::numeric_limits<int16>::max(), TEXT("There is too many attributes. The index is saved as an int16 in FSlateAttributeMetaData."));
+	checkf(Descriptor.Attributes.Num() <= std::numeric_limits<uint8>::max()-1
+		, TEXT("There are too many attributes '%d'. The index is saved as an uint8 in FSlateAttributeMetaData."), Descriptor.Attributes.Num());
 
 	// Confirm that the Visibility attribute is marked as "bAffectVisibility"
 	{
@@ -179,16 +180,15 @@ FSlateAttributeDescriptor::FInitializer::~FInitializer()
 			}
 			PreviousPrerequisiteIndex = Element.PrerequisitesIndex;
 		}
+
+		Descriptor.Attributes.Sort([](const FAttribute& A, const FAttribute& B) { return A.SortOrder < B.SortOrder; });
 	}
 
 	// Confirm that the attributes marked as "AffectVisibility" are in front of the list.
-#if 1
+#if 0
 	{
-		TArray<FAttribute, TInlineAllocator<32>> Copy;
-		Copy.Append(Descriptor.Attributes);
-		Copy.Sort([](const FAttribute& A, const FAttribute& B){ return A.SortOrder < B.SortOrder; });
 		bool bLookingForAffectVisibility = true;
-		for (const FAttribute& Attribute : Copy)
+		for (const FAttribute& Attribute : Descriptor.Attributes)
 		{
 			if (!Attribute.bAffectVisibility)
 			{
diff --git a/Engine/Source/Runtime/SlateCore/Private/Types/SlateAttributeMetaData.cpp b/Engine/Source/Runtime/SlateCore/Private/Types/SlateAttributeMetaData.cpp
--- a/Engine/Source/Runtime/SlateCore/Private/Types/SlateAttributeMetaData.cpp
+++ b/Engine/Source/Runtime/SlateCore/Private/Types/SlateAttributeMetaData.cpp
@@ -10,6 +10,9 @@
 #include <limits>
 
 
+const FSlateAttributeMetaData::FGetterItem::FAttributeIndex FSlateAttributeMetaData::FGetterItem::InvalidAttributeIndex = std::numeric_limits<FSlateAttributeMetaData::FGetterItem::FAttributeIndex>::max();
+
+
 namespace Private
 {
 	FSlateAttributeDescriptor::OffsetType FindOffet(const SWidget& OwningWidget, const FSlateAttributeBase& Attribute)
@@ -85,20 +88,21 @@ void FSlateAttributeMetaData::RegisterAttributeImpl(SWidget& OwningWidget, FSlat
 			if (FoundMemberAttributeIndex != INDEX_NONE)
 			{
 				FSlateAttributeDescriptor::FAttribute const& FoundAttribute = Descriptor.GetAttributeAtIndex(FoundMemberAttributeIndex);
-				check(FoundMemberAttributeIndex < std::numeric_limits<int16>::max());
+				check(FoundMemberAttributeIndex < std::numeric_limits<FGetterItem::FAttributeIndex>::max());
 
 				const int32 InsertLocation = Algo::LowerBoundBy(Attributes, FoundAttribute.SortOrder, [](const FGetterItem& Item) { return Item.SortOrder; }, TLess<>());
-
-				FGetterItem& GetterItem = Attributes.Insert_GetRef({ &Attribute, FoundAttribute.SortOrder, MoveTemp(Getter), (int16)FoundMemberAttributeIndex }, InsertLocation);
-				GetterItem.bIsMemberType = true;
+				FGetterItem& GetterItem = Attributes.Insert_GetRef({ &Attribute, FoundAttribute.SortOrder, MoveTemp(Getter), (FGetterItem::FAttributeIndex)FoundMemberAttributeIndex }, InsertLocation);
+				GetterItem.AttributeType = ESlateAttributeType::Member;
 
 				// Do I have dependency or am I a dependency
-				if (!FoundAttribute.Prerequisite.IsNone() && FoundAttribute.bIsPrerequisiteAlsoADependency)
+				if (!FoundAttribute.Prerequisite.IsNone() && FoundAttribute.bIsPrerequisiteAlsoADependency) 
 				{
 					// I can only be updated if the prerequisite is updated
 					const int32 FoundDependencyAttributeIndex = Descriptor.IndexOfMemberAttribute(FoundAttribute.Prerequisite);
-					check(FoundDependencyAttributeIndex < std::numeric_limits<int16>::max());
-					GetterItem.CachedAttributeDependencyIndex = (int16)FoundDependencyAttributeIndex;
+					if (FoundDependencyAttributeIndex != INDEX_NONE)
+					{
+						GetterItem.CachedAttributeDependencyIndex = (FGetterItem::FAttributeIndex)FoundDependencyAttributeIndex;
+					}
 				}
 				GetterItem.bIsADependencyForSomeoneElse = FoundAttribute.bIsADependencyForSomeoneElse;
 				GetterItem.bAffectVisibility = FoundAttribute.bAffectVisibility;
@@ -113,14 +117,14 @@ void FSlateAttributeMetaData::RegisterAttributeImpl(SWidget& OwningWidget, FSlat
 
 				const  int32 InsertLocation = Algo::LowerBoundBy(Attributes, SortOrder, [](const FGetterItem& Item) { return Item.SortOrder; }, TLess<>());
 				FGetterItem& GetterItem = Attributes.Insert_GetRef({&Attribute, SortOrder, MoveTemp(Getter)}, InsertLocation);
-				GetterItem.bIsMemberType = true;
+				GetterItem.AttributeType = ESlateAttributeType::Member;
 			}
 		}
 		else if (AttributeType == ESlateAttributeType::Managed)
 		{
 			const uint32 ManagedSortOrder = std::numeric_limits<uint32>::max();
 			FGetterItem& GetterItem = Attributes.Emplace_GetRef(&Attribute, ManagedSortOrder, MoveTemp(Getter));
-			GetterItem.bIsManagedType = true;
+			GetterItem.AttributeType = ESlateAttributeType::Managed;
 		}
 		else
 		{
@@ -189,7 +193,7 @@ TArray<FName> FSlateAttributeMetaData::GetAttributeNames(const SWidget& OwningWi
 
 EInvalidateWidgetReason FSlateAttributeMetaData::FGetterItem::GetInvalidationDetail(const SWidget& OwningWidget, EInvalidateWidgetReason Reason) const
 {
-	if (CachedAttributeDescriptorIndex != INDEX_NONE)
+	if (CachedAttributeDescriptorIndex != FGetterItem::InvalidAttributeIndex)
 	{
 		const FSlateAttributeDescriptor::FAttribute& DescriptorAttribute = OwningWidget.GetWidgetClass().GetAttributeDescriptor().GetAttributeAtIndex(CachedAttributeDescriptorIndex);
 		return DescriptorAttribute.InvalidationReason.Get(OwningWidget);
@@ -200,7 +204,7 @@ EInvalidateWidgetReason FSlateAttributeMetaData::FGetterItem::GetInvalidationDet
 
 FName FSlateAttributeMetaData::FGetterItem::GetAttributeName(const SWidget& OwningWidget) const
 {
-	if (CachedAttributeDescriptorIndex != INDEX_NONE)
+	if (CachedAttributeDescriptorIndex != FGetterItem::InvalidAttributeIndex)
 	{
 		const FSlateAttributeDescriptor::FAttribute& DescriptorAttribute = OwningWidget.GetWidgetClass().GetAttributeDescriptor().GetAttributeAtIndex(CachedAttributeDescriptorIndex);
 		return DescriptorAttribute.Name;
@@ -236,10 +240,10 @@ void FSlateAttributeMetaData::InvalidateWidget(SWidget& OwningWidget, const FSla
 			if (GetterItem.bIsADependencyForSomeoneElse)
 			{
 				GetterItem.bUpatedManually = true;
-				AttributeMetaData->bHasUpdatedManuallyFlagToReset = true;
+				AttributeMetaData->SetNeedToResetFlag(FoundIndex);
 			}
 		}
-		// Not registered but may be defined in the Descriptor
+		// Not registered/bound but may be defined in the Descriptor
 		else if (AttributeType == ESlateAttributeType::Member)
 		{
 			FSlateAttributeDescriptor const& AttributeDescriptor = OwningWidget.GetWidgetClass().GetAttributeDescriptor();
@@ -253,11 +257,13 @@ void FSlateAttributeMetaData::InvalidateWidget(SWidget& OwningWidget, const FSla
 				{
 					// Find if that dependency is registered, if not it is ok because every attribute is updated at least once
 					// Set UpdatedOnce to false to force a new update.
-					AttributeDescriptor.ForEachDependency(*FoundAttribute, [AttributeMetaData](int32 DependencyIndex)
+					AttributeDescriptor.ForEachDependentsOn(*FoundAttribute, [AttributeMetaData](int32 DependencyIndex)
 					{
 						FGetterItem* FoundOther = AttributeMetaData->Attributes.FindByPredicate([DependencyIndex](FGetterItem const& Other)
 						{
-							return Other.CachedAttributeDescriptorIndex == DependencyIndex;
+							check(DependencyIndex != INDEX_NONE);
+							check(DependencyIndex < std::numeric_limits<FGetterItem::FAttributeIndex>::max());
+							return Other.CachedAttributeDescriptorIndex == (FGetterItem::FAttributeIndex)DependencyIndex;
 						});
 						if (FoundOther)
 						{
@@ -286,6 +292,15 @@ void FSlateAttributeMetaData::UpdateAllAttributes(SWidget& OwningWidget, EInvali
 	if (FSlateAttributeMetaData* AttributeMetaData = FSlateAttributeMetaData::FindMetaData(OwningWidget))
 	{
 		AttributeMetaData->UpdateAttributesImpl(OwningWidget, InvalidationStyle, 0, AttributeMetaData->Attributes.Num());
+		if (AttributeMetaData->ResetFlag != EResetFlags::None)
+		{
+			for (FGetterItem& GetterItem : AttributeMetaData->Attributes)
+			{
+				GetterItem.bUpatedManually = false;
+				GetterItem.bUpdatedThisFrame = false;
+			}
+			AttributeMetaData->ResetFlag = EResetFlags::None;
+		}
 	}
 }
 
@@ -294,7 +309,22 @@ void FSlateAttributeMetaData::UpdateOnlyVisibilityAttributes(SWidget& OwningWidg
 {
 	if (FSlateAttributeMetaData* AttributeMetaData = FSlateAttributeMetaData::FindMetaData(OwningWidget))
 	{
-		AttributeMetaData->UpdateAttributesImpl(OwningWidget, InvalidationStyle, 0, AttributeMetaData->AffectVisibilityCounter);
+		if (AttributeMetaData->AffectVisibilityCounter > 0)
+		{
+			const int32 StartIndex = 0;
+			const int32 EndIndex = AttributeMetaData->AffectVisibilityCounter;
+			AttributeMetaData->UpdateAttributesImpl(OwningWidget, InvalidationStyle, StartIndex, EndIndex);
+			if (EnumHasAllFlags(AttributeMetaData->ResetFlag, EResetFlags::NeedToReset_ExceptVisibility))
+			{
+				for (int32 Index = StartIndex; Index < EndIndex; ++Index)
+				{
+					FGetterItem& GetterItem = AttributeMetaData->Attributes[Index];
+					GetterItem.bUpatedManually = false;
+					GetterItem.bUpdatedThisFrame = false;
+				}
+				EnumRemoveFlags(AttributeMetaData->ResetFlag, EResetFlags::NeedToReset_ExceptVisibility);
+			}
+		}
 	}
 }
 
@@ -303,7 +333,22 @@ void FSlateAttributeMetaData::UpdateExceptVisibilityAttributes(SWidget& OwningWi
 {
 	if (FSlateAttributeMetaData* AttributeMetaData = FSlateAttributeMetaData::FindMetaData(OwningWidget))
 	{
-		AttributeMetaData->UpdateAttributesImpl(OwningWidget, InvalidationStyle, AttributeMetaData->AffectVisibilityCounter, AttributeMetaData->Attributes.Num());
+		if (AttributeMetaData->AffectVisibilityCounter < AttributeMetaData->Attributes.Num())
+		{
+			const int32 StartIndex = AttributeMetaData->AffectVisibilityCounter;
+			const int32 EndIndex = AttributeMetaData->Attributes.Num();
+			AttributeMetaData->UpdateAttributesImpl(OwningWidget, InvalidationStyle, StartIndex, EndIndex);
+			if (EnumHasAllFlags(AttributeMetaData->ResetFlag, EResetFlags::NeedToReset_ExceptVisibility))
+			{
+				for (int32 Index = StartIndex; Index < EndIndex; ++Index)
+				{
+					FGetterItem& GetterItem = AttributeMetaData->Attributes[Index];
+					GetterItem.bUpatedManually = false;
+					GetterItem.bUpdatedThisFrame = false;
+				}
+				EnumRemoveFlags(AttributeMetaData->ResetFlag, EResetFlags::NeedToReset_ExceptVisibility);
+			}
+		}
 	}
 }
 
@@ -335,7 +380,7 @@ void FSlateAttributeMetaData::UpdateAttributesImpl(SWidget& OwningWidget, EInval
 
 		// Update every attribute at least once.
 		//Check if it has a dependency and if it was updated this frame (it could be from an UpdateNow)
-		if (GetterItem.CachedAttributeDependencyIndex != INDEX_NONE && GetterItem.bUpdatedOnce)
+		if (GetterItem.CachedAttributeDependencyIndex != FGetterItem::InvalidAttributeIndex && GetterItem.bUpdatedOnce)
 		{
 			// Note that the dependency is maybe not registered and the attribute may have been invalidated manually
 
@@ -364,19 +409,11 @@ void FSlateAttributeMetaData::UpdateAttributesImpl(SWidget& OwningWidget, EInval
 		GetterItem.bUpdatedThisFrame = Result.bInvalidationRequested;
 		if (Result.bInvalidationRequested && bAllowInvalidation)
 		{
+			SetNeedToResetFlag(Index);
 			InvalidationReason |= GetterItem.GetInvalidationDetail(OwningWidget, Result.InvalidationReason);
 		}
 	}
 
-	if (bHasUpdatedManuallyFlagToReset)
-	{
-		for (FGetterItem& GetterItem : Attributes)
-		{
-			GetterItem.bUpatedManually = false;
-		}
-	}
-	bHasUpdatedManuallyFlagToReset = false;
-
 	if (bInvalidateIfNeeded)
 	{
 		OwningWidget.Invalidate(InvalidationReason | CachedInvalidationReason);
@@ -417,7 +454,7 @@ void FSlateAttributeMetaData::UpdateAttribute(SWidget& OwningWidget, FSlateAttri
 				if (GetterItem.bIsADependencyForSomeoneElse)
 				{
 					GetterItem.bUpatedManually = true;
-					AttributeMetaData->bHasUpdatedManuallyFlagToReset = true;
+					AttributeMetaData->SetNeedToResetFlag(FoundIndex);
 				}
 			}
 		}
diff --git a/Engine/Source/Runtime/SlateCore/Private/Widgets/SWidget.cpp b/Engine/Source/Runtime/SlateCore/Private/Widgets/SWidget.cpp
--- a/Engine/Source/Runtime/SlateCore/Private/Widgets/SWidget.cpp
+++ b/Engine/Source/Runtime/SlateCore/Private/Widgets/SWidget.cpp
@@ -237,6 +237,8 @@ SWidget::~SWidget()
 	}
 #endif
 
+	bHasRegisteredSlateAttribute = false;
+
 	// Unregister all ActiveTimers so they aren't left stranded in the Application's list.
 	if (FSlateApplicationBase::IsInitialized())
 	{
@@ -1196,7 +1198,7 @@ void SWidget::Invalidate(EInvalidateWidgetReason InvalidateReason)
 {
 	SLATE_CROSS_THREAD_CHECK();
 
-	if (InvalidateReason == EInvalidateWidgetReason::None /*|| !IsConstructionCompleted()*/)
+	if (InvalidateReason == EInvalidateWidgetReason::None || !IsConstructionCompleted())
 	{
 		return;
 	}
diff --git a/Engine/Source/Runtime/SlateCore/Public/Types/Attributes/SlateAttributeBase.h b/Engine/Source/Runtime/SlateCore/Public/Types/Attributes/SlateAttributeBase.h
new file mode 100644
--- /dev/null
+++ b/Engine/Source/Runtime/SlateCore/Public/Types/Attributes/SlateAttributeBase.h
@@ -0,0 +1,573 @@
+// Copyright Epic Games, Inc. All Rights Reserved.
+
+#pragma once
+
+/** */
+namespace SlateAttributePrivate
+{
+	/**
+	 * Attribute object
+	 * InObjectType - Type of the value to store
+	 * InInvalidationReasonPredicate - Predicate that returns the type of invalidation to do when the value changes (e.g layout or paint)
+	 *								The invalidation can be overridden per widget. (This use memory allocation. See FSlateAttributeMetadata.)
+	 * InComparePredicateType - Predicate to compare the cached value with the Getter.
+	 * bInIsExternal - The attribute life is not controlled by the SWidget.
+	 */
+	template<typename InObjectType, typename InInvalidationReasonPredicate, typename InComparePredicateType, ESlateAttributeType InAttributeType>
+	struct TSlateAttributeBase : public FSlateAttributeImpl
+	{
+	public:
+		template<typename AttributeMemberType>
+		friend struct TSlateMemberAttributeRef;
+		friend FSlateAttributeContainer;
+
+		using ObjectType = InObjectType;
+		using FInvalidationReasonPredicate = InInvalidationReasonPredicate;
+		using FGetter = typename TAttribute<ObjectType>::FGetter;
+		using FComparePredicate = InComparePredicateType;
+
+		static EInvalidateWidgetReason GetInvalidationReason(const SWidget& Widget) { return FInvalidationReasonPredicate::GetInvalidationReason(Widget); }
+		static bool IdenticalTo(const SWidget& Widget, const ObjectType& Lhs, const ObjectType& Rhs) { return FComparePredicate::IdenticalTo(Widget, Lhs, Rhs); }
+
+	private:
+		void UpdateNowOnBind(SWidget& Widget)
+		{
+#if UE_SLATE_WITH_ATTRIBUTE_INITIALIZATION_ON_BIND
+			ProtectedUpdateNow(Widget, InAttributeType);
+#endif
+		}
+
+		void VerifyOwningWidget(const SWidget& Widget) const
+		{
+#if UE_SLATE_WITH_MEMBER_ATTRIBUTE_DEBUGGING
+			checkf(&Widget == Debug_OwningWidget || Debug_OwningWidget == nullptr,
+				TEXT("The Owning Widget is not the same as used at construction. This will cause bad memory access."));
+#endif
+		}
+
+	public:
+#if UE_SLATE_WITH_MEMBER_ATTRIBUTE_DEBUGGING
+		TSlateAttributeBase()
+			: Debug_OwningWidget(nullptr)
+		{
+		}
+#else
+		TSlateAttributeBase() = default;
+#endif
+
+		TSlateAttributeBase(const ObjectType& InValue)
+			: Value(InValue)
+#if UE_SLATE_WITH_MEMBER_ATTRIBUTE_DEBUGGING
+			, Debug_OwningWidget(nullptr)
+#endif
+		{
+		}
+
+		TSlateAttributeBase(ObjectType&& InValue)
+			: Value(MoveTemp(InValue))
+#if UE_SLATE_WITH_MEMBER_ATTRIBUTE_DEBUGGING
+			, Debug_OwningWidget(nullptr)
+#endif
+		{
+		}
+
+		TSlateAttributeBase(SWidget& Widget)
+			: Value()
+#if UE_SLATE_WITH_MEMBER_ATTRIBUTE_DEBUGGING
+			, Debug_OwningWidget(&Widget)
+#endif
+		{
+		}
+
+		TSlateAttributeBase(SWidget& Widget, const ObjectType& InValue)
+			: Value(InValue)
+#if UE_SLATE_WITH_MEMBER_ATTRIBUTE_DEBUGGING
+			, Debug_OwningWidget(&Widget)
+#endif
+		{
+		}
+
+		TSlateAttributeBase(SWidget& Widget, ObjectType&& InValue)
+			: Value(MoveTemp(InValue))
+#if UE_SLATE_WITH_MEMBER_ATTRIBUTE_DEBUGGING
+			, Debug_OwningWidget(&Widget)
+#endif
+		{
+		}
+
+		TSlateAttributeBase(SWidget& Widget, const FGetter& Getter, const ObjectType& InitialValue)
+			: Value(InitialValue)
+#if UE_SLATE_WITH_MEMBER_ATTRIBUTE_DEBUGGING
+			, Debug_OwningWidget(&Widget)
+#endif
+		{
+			if (Getter.IsBound())
+			{
+				ConstructWrapper(Widget, Getter);
+			}
+		}
+
+		TSlateAttributeBase(SWidget& Widget, const FGetter& Getter, ObjectType&& InitialValue)
+			: Value(MoveTemp(InitialValue))
+#if UE_SLATE_WITH_MEMBER_ATTRIBUTE_DEBUGGING
+			, Debug_OwningWidget(&Widget)
+#endif
+		{
+			if (Getter.IsBound())
+			{
+				ConstructWrapper(Widget, Getter);
+			}
+		}
+
+		TSlateAttributeBase(SWidget& Widget, FGetter&& Getter, const ObjectType& InitialValue)
+			: Value(InitialValue)
+#if UE_SLATE_WITH_MEMBER_ATTRIBUTE_DEBUGGING
+			, Debug_OwningWidget(&Widget)
+#endif
+		{
+			if (Getter.IsBound())
+			{
+				ConstructWrapper(Widget, MoveTemp(Getter));
+			}
+		}
+
+		TSlateAttributeBase(SWidget& Widget, FGetter&& Getter, ObjectType&& InitialValue)
+			: Value(MoveTemp(InitialValue))
+#if UE_SLATE_WITH_MEMBER_ATTRIBUTE_DEBUGGING
+			, Debug_OwningWidget(&Widget)
+#endif
+		{
+			if (Getter.IsBound())
+			{
+				ConstructWrapper(Widget, MoveTemp(Getter));
+			}
+		}
+
+		TSlateAttributeBase(SWidget& Widget, const TAttribute<ObjectType>& Attribute, const ObjectType& InitialValue)
+			: Value((Attribute.IsSet() && !Attribute.IsBound()) ? Attribute.Get() : InitialValue)
+#if UE_SLATE_WITH_MEMBER_ATTRIBUTE_DEBUGGING
+			, Debug_OwningWidget(&Widget)
+#endif
+		{
+			if (Attribute.IsBound())
+			{
+				ConstructWrapper(Widget, Attribute.GetBinding());
+			}
+		}
+
+		TSlateAttributeBase(SWidget& Widget, TAttribute<ObjectType>&& Attribute, ObjectType&& InitialValue)
+			: Value((Attribute.IsSet() && !Attribute.IsBound()) ? MoveTemp(Attribute.Steal().template Get<ObjectType>()) : MoveTemp(InitialValue))
+#if UE_SLATE_WITH_MEMBER_ATTRIBUTE_DEBUGGING
+			, Debug_OwningWidget(&Widget)
+#endif
+		{
+			if (Attribute.IsBound())
+			{
+				ConstructWrapper(Widget, MoveTemp(Attribute.Steal().template Get<FGetter>()));
+			}
+		}
+
+	public:
+		/** @return the SlateAttribute cached value. If the SlateAttribute is bound, the value will be cached at the end of the every frame. */
+		UE_NODISCARD const ObjectType& Get() const
+		{
+			return Value;
+		}
+
+		/** Update the cached value and invalidate the widget if needed. */
+		void UpdateNow(SWidget& Widget)
+		{
+			VerifyOwningWidget(Widget);
+			ProtectedUpdateNow(Widget, InAttributeType);
+		}
+
+	public:
+		/**
+		 * Unbind the SlateAttribute and set its value. It may invalidate the Widget if the value is different.
+		 * @return true if the value is considered different and an invalidation occurred.
+		 */
+		bool Set(SWidget& Widget, const ObjectType& NewValue)
+		{
+			VerifyOwningWidget(Widget);
+			ProtectedUnregisterAttribute(Widget, InAttributeType);
+
+			const bool bIsIdentical = IdenticalTo(Widget, Value, NewValue);
+			if (!bIsIdentical)
+			{
+				Value = NewValue;
+				ProtectedInvalidateWidget(Widget, InAttributeType, GetInvalidationReason(Widget));
+			}
+			return !bIsIdentical;
+		}
+
+		/**
+		 * Unbind the SlateAttribute and set its value. It may invalidate the Widget if the value is different.
+		 * @return true if the value is considered different and an invalidation occurred.
+		 */
+		bool Set(SWidget& Widget, ObjectType&& NewValue)
+		{
+			VerifyOwningWidget(Widget);
+			ProtectedUnregisterAttribute(Widget, InAttributeType);
+
+			const bool bIsIdentical = IdenticalTo(Widget, Value, NewValue);
+			if (!bIsIdentical)
+			{
+				Value = MoveTemp(NewValue);
+				ProtectedInvalidateWidget(Widget, InAttributeType, GetInvalidationReason(Widget));
+			}
+			return !bIsIdentical;
+		}
+
+	public:
+		/**
+		 * Bind the SlateAttribute to the Getter function.
+		 * (If enabled) Update the value from the Getter. If the value is different, then invalidate the widget.
+		 * The SlateAttribute will now be updated every frame from the Getter.
+		 */
+		void Bind(SWidget& Widget, const FGetter& Getter)
+		{
+			VerifyOwningWidget(Widget);
+			if (Getter.IsBound())
+			{
+				AssignBinding(Widget, Getter);
+			}
+			else
+			{
+				ProtectedUnregisterAttribute(Widget, InAttributeType);
+			}
+		}
+
+		/**
+		 * Bind the SlateAttribute to the Getter function.
+		 * (If enabled) Update the value from the Getter. If the value is different, then invalidate the widget.
+		 * The SlateAttribute will now be updated every frame from the Getter.
+		 */
+		void Bind(SWidget& Widget, FGetter&& Getter)
+		{
+			VerifyOwningWidget(Widget);
+			if (Getter.IsBound())
+			{
+				AssignBinding(Widget, MoveTemp(Getter));
+			}
+			else
+			{
+				ProtectedUnregisterAttribute(Widget, InAttributeType);
+			}
+		}
+
+		/**
+		 * Bind the SlateAttribute to the newly create getter function.
+		 * (If enabled) Update the value from the getter. If the value is different, then invalidate the widget.
+		 * The SlateAttribute will now be updated every frame from the Getter.
+		 */
+		template<typename WidgetType, typename U = typename std::enable_if<std::is_base_of<SWidget, WidgetType>::value>::type>
+		void Bind(WidgetType& Widget, typename FGetter::template TConstMethodPtr<WidgetType> MethodPtr)
+		{
+			Bind(Widget, FGetter::CreateSP(&Widget, MethodPtr));
+		}
+
+		/**
+		 * Bind the SlateAttribute to the Attribute Getter function (if it exist).
+		 * (If enabled) Update the value from the getter. If the value is different, then invalidate the widget.
+		 * The SlateAttribute will now be updated every frame from the Getter.
+		 * Or
+		 * Set the SlateAttribute's value if the Attribute is not bound but is set.
+		 * This will Unbind any previously bound getter.
+		 * If the value is different, then invalidate the widget.
+		 * Or
+		 * Unbind the SlateAttribute if the Attribute is not bound and not set.
+		 * @see Set
+		 * @return true if the getter was assigned or if the value is considered different and an invalidation occurred.
+		 */
+		bool Assign(SWidget& Widget, const TAttribute<ObjectType>& OtherAttribute)
+		{
+			VerifyOwningWidget(Widget);
+			if (OtherAttribute.IsBound())
+			{
+				return AssignBinding(Widget, OtherAttribute.GetBinding());
+			}
+			else if (OtherAttribute.IsSet())
+			{
+				return Set(Widget, OtherAttribute.Get());
+			}
+			else
+			{
+				ProtectedUnregisterAttribute(Widget, InAttributeType);
+				return false;
+			}
+		}
+
+		bool Assign(SWidget& Widget, TAttribute<ObjectType>&& OtherAttribute)
+		{
+			VerifyOwningWidget(Widget);
+			if (OtherAttribute.IsBound())
+			{
+				return AssignBinding(Widget, MoveTemp(OtherAttribute.Steal().template Get<FGetter>()));
+			}
+			else if (OtherAttribute.IsSet())
+			{
+				return Set(Widget, MoveTemp(OtherAttribute.Steal().template Get<ObjectType>()));
+			}
+			else
+			{
+				ProtectedUnregisterAttribute(Widget, InAttributeType);
+				return false;
+			}
+		}
+
+		bool Assign(SWidget& Widget, const TAttribute<ObjectType>& OtherAttribute, const ObjectType& DefaultValue)
+		{
+			VerifyOwningWidget(Widget);
+			if (OtherAttribute.IsBound())
+			{
+				return AssignBinding(Widget, OtherAttribute.GetBinding());
+			}
+			else if (OtherAttribute.IsSet())
+			{
+				return Set(Widget, OtherAttribute.Get());
+			}
+			else
+			{
+				return Set(Widget, DefaultValue);
+			}
+		}
+
+		bool Assign(SWidget& Widget, TAttribute<ObjectType>&& OtherAttribute, const ObjectType& DefaultValue)
+		{
+			VerifyOwningWidget(Widget);
+			if (OtherAttribute.IsBound())
+			{
+				return AssignBinding(Widget, MoveTemp(OtherAttribute.Steal().template Get<FGetter>()));
+			}
+			else if (OtherAttribute.IsSet())
+			{
+				return Set(Widget, MoveTemp(OtherAttribute.Steal().template Get<ObjectType>()));
+			}
+			else
+			{
+				return Set(Widget, DefaultValue);
+			}
+		}
+
+		bool Assign(SWidget& Widget, const TAttribute<ObjectType>& OtherAttribute, ObjectType&& DefaultValue)
+		{
+			VerifyOwningWidget(Widget);
+			if (OtherAttribute.IsBound())
+			{
+				return AssignBinding(Widget, OtherAttribute.GetBinding());
+			}
+			else if (OtherAttribute.IsSet())
+			{
+				return Set(Widget, OtherAttribute.Get());
+			}
+			else
+			{
+				return Set(Widget, MoveTemp(DefaultValue));
+			}
+		}
+
+		bool Assign(SWidget& Widget, TAttribute<ObjectType>&& OtherAttribute, ObjectType&& DefaultValue)
+		{
+			VerifyOwningWidget(Widget);
+			if (OtherAttribute.IsBound())
+			{
+				return AssignBinding(Widget, MoveTemp(OtherAttribute.Steal().template Get<FGetter>()));
+			}
+			else if (OtherAttribute.IsSet())
+			{
+				return Set(Widget, MoveTemp(OtherAttribute.Steal().template Get<ObjectType>()));
+			}
+			else
+			{
+				return Set(Widget, MoveTemp(DefaultValue));
+			}
+		}
+
+		/**
+		 * Remove the Getter function.
+		 * The Slate Attribute will not be updated anymore and will keep its current cached value.
+		 */
+		void Unbind(SWidget& Widget)
+		{
+			VerifyOwningWidget(Widget);
+			ProtectedUnregisterAttribute(Widget, InAttributeType);
+		}
+
+	public:
+		/** Build a Attribute from this SlateAttribute. */
+		UE_NODISCARD TAttribute<ObjectType> ToAttribute(const SWidget& Widget) const
+		{
+			if (ISlateAttributeGetter* Delegate = ProtectedFindGetter(Widget, InAttributeType))
+			{
+				return TAttribute<ObjectType>::Create(static_cast<FSlateAttributeGetterWrapper<TSlateAttributeBase>*>(Delegate)->GetDelegate());
+			}
+			return TAttribute<ObjectType>(Get());
+		}
+
+		/** @return True if the SlateAttribute is bound to a getter function. */
+		UE_NODISCARD bool IsBound(const SWidget& Widget) const
+		{
+			VerifyOwningWidget(Widget);
+			return ProtectedIsBound(Widget, InAttributeType);
+		}
+
+		/** @return True if they have the same Getter or the same value. */
+		UE_NODISCARD bool IdenticalTo(const SWidget& Widget, const TSlateAttributeBase& Other) const
+		{
+			VerifyOwningWidget(Widget);
+			FDelegateHandle ThisDelegateHandle = ProtectedFindGetterHandle(Widget, InAttributeType);
+			FDelegateHandle OthersDelegateHandle = ProtectedFindGetterHandle(Widget, InAttributeType);
+			if (ThisDelegateHandle == OthersDelegateHandle)
+			{
+				if (ThisDelegateHandle.IsValid())
+				{
+					return true;
+				}
+				return IdenticalTo(Widget, Get(), Other.Get());
+			}
+			return false;
+		}
+
+		/** @return True if they have the same Getter or, if the Attribute is set, the same value. */
+		UE_NODISCARD bool IdenticalTo(const SWidget& Widget, const TAttribute<ObjectType>& Other) const
+		{
+			VerifyOwningWidget(Widget);
+			FDelegateHandle ThisDelegateHandle = ProtectedFindGetterHandle(Widget, InAttributeType);
+			if (Other.IsBound())
+			{
+				return Other.GetBinding().GetHandle() == ThisDelegateHandle;
+			}
+			return !ThisDelegateHandle.IsValid() && Other.IsSet() && IdenticalTo(Widget, Get(), Other.Get());
+		}
+
+	private:
+		void ConstructWrapper(SWidget& Widget, const FGetter& Getter)
+		{
+			TUniquePtr<ISlateAttributeGetter> Wrapper = MakeUniqueGetter(*this, Getter);
+			ProtectedRegisterAttribute(Widget, InAttributeType, MoveTemp(Wrapper));
+			UpdateNowOnBind(Widget);
+		}
+
+		void ConstructWrapper(SWidget& Widget, FGetter&& Getter)
+		{
+			TUniquePtr<ISlateAttributeGetter> Wrapper = MakeUniqueGetter(*this, MoveTemp(Getter));
+			ProtectedRegisterAttribute(Widget, InAttributeType, MoveTemp(Wrapper));
+			UpdateNowOnBind(Widget);
+		}
+
+		bool AssignBinding(SWidget& Widget, const FGetter& Getter)
+		{
+			const FDelegateHandle PreviousGetterHandle = ProtectedFindGetterHandle(Widget, InAttributeType);
+			if (PreviousGetterHandle != Getter.GetHandle())
+			{
+				ConstructWrapper(Widget, Getter);
+				return true;
+			}
+			return false;
+		}
+
+		bool AssignBinding(SWidget& Widget, FGetter&& Getter)
+		{
+			const FDelegateHandle PreviousGetterHandle = ProtectedFindGetterHandle(Widget, InAttributeType);
+			if (PreviousGetterHandle != Getter.GetHandle())
+			{
+				ConstructWrapper(Widget, MoveTemp(Getter));
+				return true;
+			}
+			return false;
+		}
+
+	private:
+		template<typename SlateAttributeType>
+		class FSlateAttributeGetterWrapper : public ISlateAttributeGetter
+		{
+		public:
+			using ObjectType = typename SlateAttributeType::ObjectType;
+			using FGetter = typename TAttribute<ObjectType>::FGetter;
+			using FComparePredicate = typename SlateAttributeType::FComparePredicate;
+
+			FSlateAttributeGetterWrapper() = delete;
+			FSlateAttributeGetterWrapper(const FSlateAttributeGetterWrapper&) = delete;
+			FSlateAttributeGetterWrapper& operator= (const FSlateAttributeGetterWrapper&) = delete;
+			virtual ~FSlateAttributeGetterWrapper() = default;
+
+		public:
+			FSlateAttributeGetterWrapper(SlateAttributeType& InOwningAttribute, const FGetter& InGetterDelegate)
+				: Getter(InGetterDelegate)
+				, Attribute(&InOwningAttribute)
+			{
+			}
+			FSlateAttributeGetterWrapper(SlateAttributeType& InOwningAttribute, FGetter&& InGetterDelegate)
+				: Getter(MoveTemp(InGetterDelegate))
+				, Attribute(&InOwningAttribute)
+			{
+			}
+
+			virtual FUpdateAttributeResult UpdateAttribute(const SWidget& Widget) override
+			{
+				ObjectType NewValue = Getter.Execute();
+
+				const bool bIsIdentical = Attribute->IdenticalTo(Widget, Attribute->Value, NewValue);
+				if (!bIsIdentical)
+				{
+					// Set the value on the widget
+					Attribute->Value = MoveTemp(NewValue);
+					return Attribute->GetInvalidationReason(Widget);
+				}
+				return FUpdateAttributeResult();
+			}
+
+			virtual const FSlateAttributeBase& GetAttribute() const override
+			{
+				return *Attribute;
+			}
+
+			virtual void SetAttribute(FSlateAttributeBase& InBase) override
+			{
+				Attribute = &static_cast<SlateAttributeType&>(InBase);
+			}
+
+			virtual FDelegateHandle GetDelegateHandle() const override
+			{
+				return Getter.GetHandle();
+			}
+
+			const FGetter& GetDelegate() const
+			{
+				return Getter;
+			}
+
+		private:
+			/** Getter function to fetch the new value of the SlateAttribute. */
+			FGetter Getter;
+			/** The SlateAttribute of the SWidget owning the value. */
+			SlateAttributeType* Attribute;
+		};
+
+	private:
+		static TUniquePtr<ISlateAttributeGetter> MakeUniqueGetter(TSlateAttributeBase& Attribute, const FGetter& Getter)
+		{
+			return MakeUnique<FSlateAttributeGetterWrapper<TSlateAttributeBase>>(Attribute, Getter);
+		}
+
+		static TUniquePtr<ISlateAttributeGetter> MakeUniqueGetter(TSlateAttributeBase& Attribute, FGetter&& Getter)
+		{
+			return MakeUnique<FSlateAttributeGetterWrapper<TSlateAttributeBase>>(Attribute, MoveTemp(Getter));
+		}
+
+	private:
+		ObjectType Value;
+
+	protected:
+#if UE_SLATE_WITH_MEMBER_ATTRIBUTE_DEBUGGING
+		void Debug_SetOwningWidget(SWidget* InWidget)
+		{
+			Debug_OwningWidget = InWidget;
+		}
+
+		SWidget* Debug_OwningWidget;
+#endif
+	};
+
+
+} // SlateAttributePrivate
diff --git a/Engine/Source/Runtime/SlateCore/Public/Types/Attributes/SlateAttributeDefinition.h b/Engine/Source/Runtime/SlateCore/Public/Types/Attributes/SlateAttributeDefinition.h
new file mode 100644
--- /dev/null
+++ b/Engine/Source/Runtime/SlateCore/Public/Types/Attributes/SlateAttributeDefinition.h
@@ -0,0 +1,89 @@
+// Copyright Epic Games, Inc. All Rights Reserved.
+
+#pragma once
+
+#include "CoreMinimal.h"
+#include "Misc/Attribute.h"
+#include "Templates/EqualTo.h"
+#include "Widgets/InvalidateWidgetReason.h"
+
+#include <type_traits>
+
+
+class SWidget;
+
+/** */
+namespace SlateAttributePrivate
+{
+	/** Predicate used to identify if the InvalidationWidgetReason is defined in the attribute descriptor. */
+	struct FSlateAttributeNoInvalidationReason
+	{
+		static constexpr EInvalidateWidgetReason GetInvalidationReason(const SWidget&) { return EInvalidateWidgetReason::None; }
+	};
+
+
+	/** */
+	enum class ESlateAttributeType : uint8
+	{
+		Member		= 0,	// Member of a SWidget (are not allowed to move).
+		Managed		= 1,	// External to the SWidget, global variable or member that can moved.
+		Unused0		= 2,
+		Unused1		= 3,
+		// We use the attribute type in a bit field in FSlateAttributeMetaData. Only 4 value are allowed.
+	};
+
+
+	class ISlateAttributeGetter;
+	template<typename ObjectType, typename InvalidationReasonPredicate, typename FComparePredicate, ESlateAttributeType AttributeType>
+	struct TSlateAttributeBase;
+	struct FSlateAttributeContainer;
+	template<typename AttributeMemberType>
+	struct TSlateMemberAttributeRef;
+
+
+	/** */
+	class ISlateAttributeGetter
+	{
+	public:
+		struct FUpdateAttributeResult
+		{
+			FUpdateAttributeResult(EInvalidateWidgetReason InInvalidationReason)
+				: InvalidationReason(InInvalidationReason)
+				, bInvalidationRequested(true)
+			{}
+			FUpdateAttributeResult()
+				: InvalidationReason(EInvalidateWidgetReason::None)
+				, bInvalidationRequested(false)
+			{}
+			EInvalidateWidgetReason InvalidationReason;
+			bool bInvalidationRequested;
+		};
+
+
+		virtual FUpdateAttributeResult UpdateAttribute(const SWidget& Widget) = 0;
+		virtual const FSlateAttributeBase& GetAttribute() const = 0;
+		virtual void SetAttribute(FSlateAttributeBase&) = 0;
+		virtual FDelegateHandle GetDelegateHandle() const = 0;
+		virtual ~ISlateAttributeGetter() = default;
+	};
+
+
+	/** */
+	struct SLATECORE_API FSlateAttributeImpl : public FSlateAttributeBase
+	{
+	protected:
+		bool ProtectedIsWidgetInDestructionPath(SWidget* Widget) const;
+		bool ProtectedIsImplemented(const SWidget& Widget) const;
+		void ProtectedUnregisterAttribute(SWidget& Widget, ESlateAttributeType AttributeType) const;
+		void ProtectedRegisterAttribute(SWidget& Widget, ESlateAttributeType AttributeType, TUniquePtr<ISlateAttributeGetter>&& Wrapper);
+		void ProtectedInvalidateWidget(SWidget& Widget, ESlateAttributeType AttributeType, EInvalidateWidgetReason InvalidationReason) const;
+		bool ProtectedIsBound(const SWidget& Widget, ESlateAttributeType AttributeType) const;
+		ISlateAttributeGetter* ProtectedFindGetter(const SWidget& Widget, ESlateAttributeType AttributeType) const;
+		FDelegateHandle ProtectedFindGetterHandle(const SWidget& Widget, ESlateAttributeType AttributeType) const;
+		bool ProtectedIsIdenticalTo(const SWidget& Widget, ESlateAttributeType AttributeType, const FSlateAttributeBase& Other, const bool bHasSameValue) const;
+		bool ProtectedIsIdenticalToAttribute(const SWidget& Widget, ESlateAttributeType AttributeType, const void* Other, const bool bHasSameValue) const;
+		void ProtectedUpdateNow(SWidget& Widget, ESlateAttributeType AttributeType);
+		void ProtectedMoveAttribute(SWidget& Widget, ESlateAttributeType AttributeType, const FSlateAttributeBase* Other);
+	};
+
+} // SlateAttributePrivate
diff --git a/Engine/Source/Runtime/SlateCore/Public/Types/Attributes/SlateAttributeManaged.h b/Engine/Source/Runtime/SlateCore/Public/Types/Attributes/SlateAttributeManaged.h
new file mode 100644
--- /dev/null
+++ b/Engine/Source/Runtime/SlateCore/Public/Types/Attributes/SlateAttributeManaged.h
@@ -0,0 +1,247 @@
+// Copyright Epic Games, Inc. All Rights Reserved.
+
+#pragma once
+
+/** */
+namespace SlateAttributePrivate
+{
+	/**
+	 *
+	 */
+	template<typename InObjectType, typename InInvalidationReasonPredicate, typename InComparePredicate>
+	struct TSlateManagedAttribute : protected TSlateAttributeBase<InObjectType, InInvalidationReasonPredicate, InComparePredicate, ESlateAttributeType::Managed>
+	{
+	private:
+		using Super = TSlateAttributeBase<InObjectType, InInvalidationReasonPredicate, InComparePredicate, ESlateAttributeType::Managed>;
+
+	public:
+		using ObjectType = typename Super::ObjectType;
+		using FInvalidationReasonPredicate = typename Super::FInvalidationReasonPredicate;
+		using FGetter = typename Super::FGetter;
+		using FComparePredicate = typename Super::FComparePredicate;
+		static const bool IsMemberType = false;
+
+		static EInvalidateWidgetReason GetInvalidationReason(const SWidget& Widget) { return Super::GetInvalidationReason(Widget); }
+		static bool IdenticalTo(const SWidget& Widget, const ObjectType& Lhs, const ObjectType& Rhs) { return Super::IdenticalTo(Widget, Lhs, Rhs); }
+
+	public:
+		TSlateManagedAttribute() = delete;
+		TSlateManagedAttribute(const TSlateManagedAttribute&) = delete;
+		TSlateManagedAttribute& operator=(const TSlateManagedAttribute&) = delete;
+
+		~TSlateManagedAttribute()
+		{
+			Unbind();
+		}
+
+		TSlateManagedAttribute(TSlateManagedAttribute&& Other)
+			: Super(MoveTemp(Other))
+			, ManagedWidget(MoveTemp(Other.ManagedWidget))
+		{
+			if (TSharedPtr<SWidget> Pin = ManagedWidget.Pin())
+			{
+				FSlateAttributeImpl::ProtectedMoveAttribute(*Pin.Get(), ESlateAttributeType::Managed, &Other);
+			}
+		}
+
+		TSlateManagedAttribute& operator=(TSlateManagedAttribute&& Other)
+		{
+			Super::operator=(MoveTemp(Other));
+			ManagedWidget = MoveTemp(Other.ManagedWidget);
+
+			if (TSharedPtr<SWidget> Pin = ManagedWidget.Pin())
+			{
+				FSlateAttributeImpl::ProtectedMoveAttribute(*Pin.Get(), ESlateAttributeType::Managed, &Other);
+			}
+			return *this;
+		}
+
+		explicit TSlateManagedAttribute(TSharedRef<SWidget> Widget)
+			: Super(Widget.Get())
+			, ManagedWidget(Widget)
+		{ }
+
+		explicit TSlateManagedAttribute(TSharedRef<SWidget> Widget, const ObjectType& InValue)
+			: Super(Widget.Get(), InValue)
+			, ManagedWidget(Widget)
+		{
+		}
+
+		explicit TSlateManagedAttribute(TSharedRef<SWidget> Widget, ObjectType&& InValue)
+			: Super(Widget.Get(), MoveTemp(InValue))
+			, ManagedWidget(Widget)
+		{
+		}
+
+		explicit TSlateManagedAttribute(TSharedRef<SWidget> Widget, const FGetter& Getter, const ObjectType& InitialValue)
+			: Super(Widget.Get(), Getter, InitialValue)
+			, ManagedWidget(Widget)
+		{
+		}
+
+		explicit TSlateManagedAttribute(TSharedRef<SWidget> Widget, const FGetter& Getter, ObjectType&& InitialValue)
+			: Super(Widget.Get(), Getter, MoveTemp(InitialValue))
+			, ManagedWidget(Widget)
+		{
+		}
+
+		explicit TSlateManagedAttribute(TSharedRef<SWidget> Widget, FGetter&& Getter, const ObjectType& InitialValue)
+			: Super(Widget.Get(), MoveTemp(Getter), InitialValue)
+			, ManagedWidget(Widget)
+		{
+		}
+
+		explicit TSlateManagedAttribute(TSharedRef<SWidget> Widget, FGetter&& Getter, ObjectType&& InitialValue)
+			: Super(Widget.Get(), MoveTemp(Getter), MoveTemp(InitialValue))
+			, ManagedWidget(Widget)
+		{
+		}
+
+		explicit TSlateManagedAttribute(TSharedRef<SWidget> Widget, const TAttribute<ObjectType>& Attribute, const ObjectType& InitialValue)
+			: Super(Widget.Get(), Attribute, InitialValue)
+			, ManagedWidget(Widget)
+		{
+		}
+
+		explicit TSlateManagedAttribute(TSharedRef<SWidget> Widget, TAttribute<ObjectType>&& Attribute, ObjectType&& InitialValue)
+			: Super(Widget.Get(), MoveTemp(Attribute), MoveTemp(InitialValue))
+			, ManagedWidget(Widget)
+		{
+		}
+
+	public:
+		const ObjectType& Get() const
+		{
+			return Super::Get();
+		}
+
+		void UpdateNow()
+		{
+			if (TSharedPtr<SWidget> Pin = ManagedWidget.Pin())
+			{
+				Super::UpdateNow(*Pin.Get());
+			}
+		}
+
+	public:
+		void Set(const ObjectType& NewValue)
+		{
+			if (TSharedPtr<SWidget> Pin = ManagedWidget.Pin())
+			{
+				Super::Set(*Pin.Get(), NewValue);
+			}
+		}
+
+		void Set(ObjectType&& NewValue)
+		{
+			if (TSharedPtr<SWidget> Pin = ManagedWidget.Pin())
+			{
+				Super::Set(*Pin.Get(), MoveTemp(NewValue));
+			}
+		}
+
+		void Bind(const FGetter& Getter)
+		{
+			if (TSharedPtr<SWidget> Pin = ManagedWidget.Pin())
+			{
+				Super::Bind(*Pin.Get(), Getter);
+			}
+		}
+
+		void Bind(FGetter&& Getter)
+		{
+			if (TSharedPtr<SWidget> Pin = ManagedWidget.Pin())
+			{
+				Super::Bind(*Pin.Get(), MoveTemp(Getter));
+			}
+		}
+
+		void Assign(const TAttribute<ObjectType>& OtherAttribute)
+		{
+			if (TSharedPtr<SWidget> Pin = ManagedWidget.Pin())
+			{
+				Super::Assign(*Pin.Get(), OtherAttribute);
+			}
+		}
+
+		void Assign(TAttribute<ObjectType>&& OtherAttribute)
+		{
+			if (TSharedPtr<SWidget> Pin = ManagedWidget.Pin())
+			{
+				Super::Assign(*Pin.Get(), MoveTemp(OtherAttribute));
+			}
+		}
+
+		void Assign(const TAttribute<ObjectType>& OtherAttribute, const ObjectType& DefaultValue)
+		{
+			if (TSharedPtr<SWidget> Pin = ManagedWidget.Pin())
+			{
+				Super::Assign(*Pin.Get(), OtherAttribute, DefaultValue);
+			}
+		}
+
+		void Assign(const TAttribute<ObjectType>& OtherAttribute, ObjectType&& DefaultValue)
+		{
+			if (TSharedPtr<SWidget> Pin = ManagedWidget.Pin())
+			{
+				Super::Assign(*Pin.Get(), OtherAttribute, MoveTemp(DefaultValue));
+			}
+		}
+
+		void Assign(TAttribute<ObjectType>&& OtherAttribute, const ObjectType& DefaultValue)
+		{
+			if (TSharedPtr<SWidget> Pin = ManagedWidget.Pin())
+			{
+				Super::Assign(*Pin.Get(), MoveTemp(OtherAttribute), DefaultValue);
+			}
+		}
+
+		void Assign(TAttribute<ObjectType>&& OtherAttribute, ObjectType&& DefaultValue)
+		{
+			if (TSharedPtr<SWidget> Pin = ManagedWidget.Pin())
+			{
+				Super::Assign(*Pin.Get(), MoveTemp(OtherAttribute), MoveTemp(DefaultValue));
+			}
+		}
+
+		void Unbind()
+		{
+			if (TSharedPtr<SWidget> Pin = ManagedWidget.Pin())
+			{
+				Super::Unbind(*Pin.Get());
+			}
+		}
+
+	public:
+		bool IsBound() const
+		{
+			if (TSharedPtr<SWidget> Pin = ManagedWidget.Pin())
+			{
+				return Super::IsBound(*Pin.Get());
+			}
+			return false;
+		}
+
+		bool IdenticalTo(const TSlateManagedAttribute& Other) const
+		{
+			if (TSharedPtr<SWidget> Pin = ManagedWidget.Pin())
+			{
+				return Super::IdenticalTo(*Pin.Get(), Other);
+			}
+			return false;
+		}
+
+		bool IdenticalTo(const TAttribute<ObjectType>& Other) const
+		{
+			if (TSharedPtr<SWidget> Pin = ManagedWidget.Pin())
+			{
+				return Super::IdenticalTo(*Pin.Get(), Other);
+			}
+			return false;
+		}
+
+	private:
+		TWeakPtr<SWidget> ManagedWidget;
+	};
+
+} // SlateAttributePrivate
diff --git a/Engine/Source/Runtime/SlateCore/Public/Types/Attributes/SlateAttributeMember.h b/Engine/Source/Runtime/SlateCore/Public/Types/Attributes/SlateAttributeMember.h
new file mode 100644
--- /dev/null
+++ b/Engine/Source/Runtime/SlateCore/Public/Types/Attributes/SlateAttributeMember.h
@@ -0,0 +1,235 @@
+// Copyright Epic Games, Inc. All Rights Reserved.
+
+#pragma once
+
+/** */
+namespace SlateAttributePrivate
+{
+	/**
+	 *
+	 */
+	template<typename InObjectType, typename InInvalidationReasonPredicate, typename InComparePredicate>
+	struct TSlateMemberAttribute : public TSlateAttributeBase<InObjectType, InInvalidationReasonPredicate, InComparePredicate, ESlateAttributeType::Member>
+	{
+	private:
+		using Super = TSlateAttributeBase<InObjectType, InInvalidationReasonPredicate, InComparePredicate, ESlateAttributeType::Member>;
+
+		template<typename WidgetType, typename U = typename std::enable_if<std::is_base_of<SWidget, WidgetType>::value>::type>
+		static void VerifyAttributeAddress(const WidgetType& Widget, const TSlateMemberAttribute& Self)
+		{
+			checkf((UPTRINT)&Self >= (UPTRINT)&Widget && (UPTRINT)&Self < (UPTRINT)&Widget + sizeof(WidgetType),
+				TEXT("Use TAttribute or TSlateManagedAttribute instead. See SlateAttribute.h for more info."));
+			ensureAlwaysMsgf((HasDefinedInvalidationReason || Self.ProtectedIsImplemented(Widget)),
+				TEXT("The TSlateAttribute could not be found in the SlateAttributeDescriptor.\n")
+				TEXT("Use the SLATE_DECLARE_WIDGET and add the attribute in PrivateRegisterAttributes,\n")
+				TEXT("Or use TSlateAttribute with a valid Invalidation Reason instead."));
+		}
+
+	public:
+		using FGetter = typename Super::FGetter;
+		using ObjectType = typename Super::ObjectType;
+		static const bool IsMemberType = true;
+		static constexpr bool HasDefinedInvalidationReason = !std::is_same<InInvalidationReasonPredicate, FSlateAttributeNoInvalidationReason>::value;
+
+	public:
+		//~ You can only register Attribute that are defined in a SWidget (member of the class).
+		//~ Use the constructor with Widget pointer.
+		//~ See documentation in SlateAttribute.h for more info.
+		TSlateMemberAttribute() = delete;
+		TSlateMemberAttribute(const TSlateMemberAttribute&) = delete;
+		TSlateMemberAttribute(TSlateMemberAttribute&&) = delete;
+		TSlateMemberAttribute& operator=(const TSlateMemberAttribute&) = delete;
+		TSlateMemberAttribute& operator=(TSlateMemberAttribute&&) = delete;
+		void* operator new(size_t) = delete;
+		void* operator new[](size_t) = delete;
+
+#if UE_SLATE_WITH_MEMBER_ATTRIBUTE_DEBUGGING
+		~TSlateMemberAttribute()
+		{
+			/**
+			 * The parent should now be destroyed and the shared pointer should be invalidate.
+			 * If you hit the check, that means the TSlateAttribute is not a variable member of SWidget.
+			 * It will introduced bad memory access.
+			 * See documentation in SlateAttribute.h.
+			*/
+			checkf(Super::ProtectedIsWidgetInDestructionPath(Super::Debug_OwningWidget), TEXT("The Owning widget should be invalid."));
+		}
+#endif
+
+		template<typename WidgetType, typename U = typename std::enable_if<std::is_base_of<SWidget, WidgetType>::value>::type>
+		explicit TSlateMemberAttribute(WidgetType& Widget)
+			: Super(Widget)
+		{
+			VerifyAttributeAddress(Widget, *this);
+		}
+
+		template<typename WidgetType, typename U = typename std::enable_if<std::is_base_of<SWidget, WidgetType>::value>::type>
+		explicit TSlateMemberAttribute(WidgetType& Widget, const ObjectType& InValue)
+			: Super(Widget, InValue)
+		{
+			VerifyAttributeAddress(Widget, *this);
+		}
+
+		template<typename WidgetType, typename U = typename std::enable_if<std::is_base_of<SWidget, WidgetType>::value>::type>
+		explicit TSlateMemberAttribute(WidgetType& Widget, ObjectType&& InValue)
+			: Super(Widget, MoveTemp(InValue))
+		{
+			VerifyAttributeAddress(Widget, *this);
+		}
+	};
+
+
+
+	enum ESlateMemberAttributeRefNoCheckParam { SlateMemberAttributeRefNoCheckParam };
+
+	/*
+	 * A reference to a SlateAttribute that can be returned and saved for later.
+	 */
+	template<typename AttributeMemberType>
+	struct TSlateMemberAttributeRef
+	{
+	public:
+		using SlateAttributeType = AttributeMemberType;
+		using ObjectType = typename AttributeMemberType::ObjectType;
+		using AttributeType = TAttribute<ObjectType>;
+
+	private:
+		template<typename WidgetType>
+		static void VerifyAttributeAddress(const WidgetType& InWidget, const AttributeMemberType& InAttribute)
+		{
+			checkf((UPTRINT)&InAttribute >= (UPTRINT)&InWidget && (UPTRINT)&InAttribute < (UPTRINT)&InWidget + sizeof(WidgetType),
+				TEXT("The attribute is not a member of the widget."));
+			InAttribute.VerifyOwningWidget(InWidget);
+		}
+
+	public:
+		/** Constructor */
+		template<typename WidgetType, typename V = typename std::enable_if<std::is_base_of<SWidget, WidgetType>::value>::type>
+		explicit TSlateMemberAttributeRef(const TSharedRef<WidgetType>& InOwner, const AttributeMemberType& InAttribute)
+			: Owner(InOwner)
+			, Attribute(InAttribute)
+		{
+			VerifyAttributeAddress(InOwner.Get(), InAttribute);
+		}
+
+		/** Constructor */
+		explicit TSlateMemberAttributeRef(ESlateMemberAttributeRefNoCheckParam, const TSharedRef<SWidget>& InOwner, const AttributeMemberType& InAttribute)
+			: Owner(InOwner)
+			, Attribute(InAttribute)
+		{
+		}
+
+	public:
+		/** @return if the reference is valid. A reference can be invalid if the SWidget is destroyed. */
+		UE_NODISCARD bool IsValid() const
+		{
+			return Owner.IsValid();
+		}
+
+		/** @return the SlateAttribute cached value; undefined when IsValid() returns false. */
+		UE_NODISCARD const ObjectType& Get() const
+		{
+			if (TSharedPtr<const SWidget> Pin = Owner.Pin())
+			{
+				return Attribute.Get();
+			}
+			checkf(false, TEXT("It is an error to call GetValue() on an unset TSlateMemberAttributeRef. Please either check IsValid() or use Get(DefaultValue) instead."));
+			static ObjectType Tmp;
+			return Tmp;
+		}
+
+		/** @return the SlateAttribute cached value or the DefaultValue if the reference is invalid. */
+		UE_NODISCARD const ObjectType& Get(const ObjectType& DefaultValue) const
+		{
+			if (TSharedPtr<const SWidget> Pin = Owner.Pin())
+			{
+				return Attribute.Get();
+			}
+			return DefaultValue;
+		}
+
+		/** Update the cached value and invalidate the widget if needed. */
+		void UpdateValue()
+		{
+			if (TSharedPtr<const SWidget> Pin = Owner.Pin())
+			{
+				const_cast<AttributeMemberType>(Attribute).UpdateNow(*Pin.Get());
+			}
+		}
+
+		/**
+		 * Assumes the reference is valid.
+		 * Shorthand for the boilerplace code: MyAttribute.UpdateValueNow(); MyAttribute.Get();
+		 */
+		UE_NODISCARD const ObjectType& UpdateAndGet()
+		{
+			if (TSharedPtr<const SWidget> Pin = Owner.Pin())
+			{
+				const_cast<AttributeMemberType>(Attribute)->UpdateNow(*Pin.Get());
+				return Attribute.Get();
+			}
+			checkf(false, TEXT("It is an error to call GetValue() on an unset TSlateMemberAttributeRef. Please either check IsValid() or use Get(DefaultValue) instead."));
+		}
+
+		/**
+		 * Shorthand for the boilerplace code: MyAttribute.UpdateValueNow(); MyAttribute.Get(DefaultValue);
+		 * @return the SlateAttribute cached value or the DefaultValue if the reference is invalid.
+		 */
+		UE_NODISCARD const ObjectType& UpdateAndGet(const ObjectType& DefaultValue)
+		{
+			if (TSharedPtr<const SWidget> Pin = Owner.Pin())
+			{
+				const_cast<AttributeMemberType>(Attribute)->UpdateNow(*Pin.Get());
+				return Attribute.Get();
+			}
+			return DefaultValue;
+		}
+
+		/** Build a Attribute from this SlateAttribute. */
+		UE_NODISCARD TAttribute<ObjectType> ToAttribute() const
+		{
+			if (TSharedPtr<const SWidget> Pin = Owner.Pin())
+			{
+				return Attribute.ToAttribute(*Pin.Get());
+			}
+			return TAttribute<ObjectType>();
+		}
+
+		/** @return True if the SlateAttribute is bound to a getter function. */
+		UE_NODISCARD bool IsBound() const
+		{
+			if (TSharedPtr<const SWidget> Pin = Owner.Pin())
+			{
+				return Attribute.IsBound(*Pin.Get());
+			}
+			return false;
+		}
+
+		/** @return True if they have the same Getter or the same value. */
+		UE_NODISCARD bool IdenticalTo(const TSlateMemberAttributeRef& Other) const
+		{
+			TSharedPtr<const SWidget> SelfPin = Owner.Pin();
+			TSharedPtr<const SWidget> OtherPin = Other.Owner.Pin();
+			if (SelfPin == OtherPin && SelfPin)
+			{
+				return Attribute.IdenticalTo(*SelfPin.Get(), Other.Attribute);
+			}
+			return SelfPin == OtherPin;
+		}
+
+		/** @return True if they have the same Getter or, if the Attribute is set, the same value. */
+		UE_NODISCARD bool IdenticalTo(const TAttribute<ObjectType>& Other) const
+		{
+			if (TSharedPtr<const SWidget> Pin = Owner.Pin())
+			{
+				return Attribute.IdenticalTo(*Pin.Get(), Other);
+			}
+			return !Other.IsSet(); // if the other is not set, then both are invalid.
+		}
+
+	private:
+		TWeakPtr<const SWidget> Owner;
+		const AttributeMemberType& Attribute;
+	 };
+
+} // SlateAttributePrivate
diff --git a/Engine/Source/Runtime/SlateCore/Public/Types/SlateAttribute.h b/Engine/Source/Runtime/SlateCore/Public/Types/SlateAttribute.h
--- a/Engine/Source/Runtime/SlateCore/Public/Types/SlateAttribute.h
+++ b/Engine/Source/Runtime/SlateCore/Public/Types/SlateAttribute.h
@@ -70,7 +70,7 @@
  *			: NewWay(*this, FLinearColor::White) // Always use "this" pointer. No exception. If you can't, use TAttribute instead.
  *		{}
  *
- *		virtual int32 OnPaint( const FPaintArgs& Args, const FGeometry& AllottedGeometry, const FSlateRect& MyCullingRect, FSlateWindowElementList& OutDrawElements, int32 LayerId, const FWidgetStyle& InWidgetStyle, bool bParentEnabled ) const override
+ *		virtual int32 OnPaint(const FPaintArgs& Args, const FGeometry& AllottedGeometry, const FSlateRect& MyCullingRect, FSlateWindowElementList& OutDrawElements, int32 LayerId, const FWidgetStyle& InWidgetStyle, bool bParentEnabled) const override
  *		{
  *			++LayerId;
  *			FSlateDrawElement::MakeDebugQuad(
@@ -91,7 +91,7 @@
  *			// If the TAttribute is bound, copy the TSlateAttribute to the TAttribute getter, update the value, test if the value have changed and, if so, invalidate the widget.
  *			// It the TAttribute is set, remove previous getter, set the TSlateAttribute value, test if the value changed and, if so, invalidate the widget.
  *			// Else, nothing (use the previous value).
- *			NewWays.Bind(*this, InArgs._Color); // Always use the "this" pointer.
+ *			NewWays.Assign(*this, InArgs._Color); // Always use the "this" pointer.
  *		}
  *
  *		// If NewWay is bounded and the values changed from the previous frame, the widget will be invalidated.
@@ -111,8 +111,8 @@
  *	//The WidgetReflector use that information.
  *	//See NewWayWithout_SLATE_DECLARE_WIDGET_ for an example of how to use the system without the SLATE_DECLARE_WIDGET
  *
- *	SLATE_IMPLEMENT_WIDGET(STextBlock)
- *	void STextBlock::PrivateRegisterAttributes(FSlateAttributeInitializer& AttributeInitializer)
+ *	SLATE_IMPLEMENT_WIDGET(SMyNewWidget)
+ *	void SMyNewWidget::PrivateRegisterAttributes(FSlateAttributeInitializer& AttributeInitializer)
  *	{
  *		SLATE_ADD_MEMBER_ATTRIBUTE_DEFINITION(AttributeInitializer, NewWay, EInvalidationReason::Paint)
  *			.SetPrerequisite("bEnabled"); // SetPrerequisite is not needed here. This is just an example to show how you could do it if needed.
@@ -236,1065 +236,7 @@ struct FSlateAttributeBase
 };
 
 
-/** */
-namespace SlateAttributePrivate
-{
-	/** Predicate used to identify if the InvalidationWidgetReason is defined in the attribute descriptor. */
-	struct FSlateAttributeNoInvalidationReason
-	{
-		static constexpr EInvalidateWidgetReason GetInvalidationReason(const SWidget&) { return EInvalidateWidgetReason::None; }
-	};
-
-
-	/** */
-	enum class ESlateAttributeType : uint8
-	{
-		Member = 0,
-		Managed = 1,
-	};
-
-	class ISlateAttributeGetter;
-	template<typename ObjectType, typename InvalidationReasonPredicate, typename FComparePredicate, ESlateAttributeType AttributeType>
-	struct TSlateAttributeBase;
-	template<typename AttributeMemberType>
-	struct TSlateMemberAttributeRef;
-
-
-	/** */
-	class ISlateAttributeGetter
-	{
-	public:
-		struct FUpdateAttributeResult
-		{
-			FUpdateAttributeResult(EInvalidateWidgetReason InInvalidationReason)
-				: InvalidationReason(InInvalidationReason)
-				, bInvalidationRequested(true)
-			{}
-			FUpdateAttributeResult()
-				: InvalidationReason(EInvalidateWidgetReason::None)
-				, bInvalidationRequested(false)
-			{}
-			EInvalidateWidgetReason InvalidationReason;
-			bool bInvalidationRequested;
-		};
-
-
-		virtual FUpdateAttributeResult UpdateAttribute(const SWidget& Widget) = 0;
-		virtual const FSlateAttributeBase& GetAttribute() const = 0;
-		virtual void SetAttribute(FSlateAttributeBase&) = 0;
-		virtual FDelegateHandle GetDelegateHandle() const = 0;
-		virtual ~ISlateAttributeGetter() = default;
-	};
-
-
-	/** */
-	struct SLATECORE_API FSlateAttributeImpl : public FSlateAttributeBase
-	{
-	protected:
-		bool ProtectedIsWidgetInDestructionPath(SWidget* Widget) const;
-		bool ProtectedIsImplemented(const SWidget& Widget) const;
-		void ProtectedUnregisterAttribute(SWidget& Widget, ESlateAttributeType AttributeType) const;
-		void ProtectedRegisterAttribute(SWidget& Widget, ESlateAttributeType AttributeType, TUniquePtr<ISlateAttributeGetter>&& Wrapper);
-		void ProtectedInvalidateWidget(SWidget& Widget, ESlateAttributeType AttributeType, EInvalidateWidgetReason InvalidationReason) const;
-		bool ProtectedIsBound(const SWidget& Widget, ESlateAttributeType AttributeType) const;
-		ISlateAttributeGetter* ProtectedFindGetter(const SWidget& Widget, ESlateAttributeType AttributeType) const;
-		FDelegateHandle ProtectedFindGetterHandle(const SWidget& Widget, ESlateAttributeType AttributeType) const;
-		bool ProtectedIsIdenticalTo(const SWidget& Widget, ESlateAttributeType AttributeType, const FSlateAttributeBase& Other, const bool bHasSameValue) const;
-		bool ProtectedIsIdenticalToAttribute(const SWidget& Widget, ESlateAttributeType AttributeType, const void* Other, const bool bHasSameValue) const;
-		void ProtectedUpdateNow(SWidget& Widget, ESlateAttributeType AttributeType);
-		void ProtectedMoveAttribute(SWidget& Widget, ESlateAttributeType AttributeType, const FSlateAttributeBase* Other);
-	};
-
-
-	/**
-	 * Attribute object
-	 * InObjectType - Type of the value to store
-	 * InInvalidationReasonPredicate - Predicate that returns the type of invalidation to do when the value changes (e.g layout or paint)
-	 *								The invalidation can be overridden per widget. (This use memory allocation. See FSlateAttributeMetadata.)
-	 * InComparePredicateType - Predicate to compare the cached value with the Getter.
-	 * bInIsExternal - The attribute life is not controlled by the SWidget.
-	 */
-	template<typename InObjectType, typename InInvalidationReasonPredicate, typename InComparePredicateType, ESlateAttributeType InAttributeType>
-	struct TSlateAttributeBase : public FSlateAttributeImpl
-	{
-	public:
-		template<typename AttributeMemberType>
-		friend struct TSlateMemberAttributeRef;
-
-		using ObjectType = InObjectType;
-		using FInvalidationReasonPredicate = InInvalidationReasonPredicate;
-		using FGetter = typename TAttribute<ObjectType>::FGetter;
-		using FComparePredicate = InComparePredicateType;
-
-		static EInvalidateWidgetReason GetInvalidationReason(const SWidget& Widget) { return FInvalidationReasonPredicate::GetInvalidationReason(Widget); }
-		static bool IdenticalTo(const SWidget& Widget, const ObjectType& Lhs, const ObjectType& Rhs) { return FComparePredicate::IdenticalTo(Widget, Lhs, Rhs); }
-		
-	private:
-		void UpdateNowOnBind(SWidget& Widget)
-		{
-#if UE_SLATE_WITH_ATTRIBUTE_INITIALIZATION_ON_BIND
-			ProtectedUpdateNow(Widget, InAttributeType);
-#endif
-		}
-
-		void VerifyOwningWidget(const SWidget& Widget) const
-		{
-#if UE_SLATE_WITH_MEMBER_ATTRIBUTE_DEBUGGING
-			checkf(&Widget == Debug_OwningWidget,
-				TEXT("The Owning Widget is not the same as used at construction. This will cause bad memory access."));
-#endif
-		}
-
-	public:
-		TSlateAttributeBase() = default;
-
-		TSlateAttributeBase(SWidget& Widget)
-			: Value()
-#if UE_SLATE_WITH_MEMBER_ATTRIBUTE_DEBUGGING
-			, Debug_OwningWidget(&Widget)
-#endif
-		{
-		}
-
-		TSlateAttributeBase(SWidget& Widget, const ObjectType& InValue)
-			: Value(InValue)
-#if UE_SLATE_WITH_MEMBER_ATTRIBUTE_DEBUGGING
-			, Debug_OwningWidget(&Widget)
-#endif
-		{
-		}
-
-		TSlateAttributeBase(SWidget& Widget, ObjectType&& InValue)
-			: Value(MoveTemp(InValue))
-#if UE_SLATE_WITH_MEMBER_ATTRIBUTE_DEBUGGING
-			, Debug_OwningWidget(&Widget)
-#endif
-		{
-		}
-
-		TSlateAttributeBase(SWidget& Widget, const FGetter& Getter, const ObjectType& InitialValue)
-			: Value(InitialValue)
-#if UE_SLATE_WITH_MEMBER_ATTRIBUTE_DEBUGGING
-			, Debug_OwningWidget(&Widget)
-#endif
-		{
-			if (Getter.IsBound())
-			{
-				ConstructWrapper(Widget, Getter);
-			}
-		}
-
-		TSlateAttributeBase(SWidget& Widget, const FGetter& Getter, ObjectType&& InitialValue)
-			: Value(MoveTemp(InitialValue))
-#if UE_SLATE_WITH_MEMBER_ATTRIBUTE_DEBUGGING
-			, Debug_OwningWidget(&Widget)
-#endif
-		{
-			if (Getter.IsBound())
-			{
-				ConstructWrapper(Widget, Getter);
-			}
-		}
-
-		TSlateAttributeBase(SWidget& Widget, FGetter&& Getter, const ObjectType& InitialValue)
-			: Value(InitialValue)
-#if UE_SLATE_WITH_MEMBER_ATTRIBUTE_DEBUGGING
-			, Debug_OwningWidget(&Widget)
-#endif
-		{
-			if (Getter.IsBound())
-			{
-				ConstructWrapper(Widget, MoveTemp(Getter));
-			}
-		}
-
-		TSlateAttributeBase(SWidget& Widget, FGetter&& Getter, ObjectType&& InitialValue)
-			: Value(MoveTemp(InitialValue))
-#if UE_SLATE_WITH_MEMBER_ATTRIBUTE_DEBUGGING
-			, Debug_OwningWidget(&Widget)
-#endif
-		{
-			if (Getter.IsBound())
-			{
-				ConstructWrapper(Widget, MoveTemp(Getter));
-			}
-		}
-
-		TSlateAttributeBase(SWidget& Widget, const TAttribute<ObjectType>& Attribute, const ObjectType& InitialValue)
-			: Value((Attribute.IsSet() && !Attribute.IsBound()) ? Attribute.Get() : InitialValue)
-#if UE_SLATE_WITH_MEMBER_ATTRIBUTE_DEBUGGING
-			, Debug_OwningWidget(&Widget)
-#endif
-		{
-			if (Attribute.IsBound())
-			{
-				ConstructWrapper(Widget, Attribute.GetBinding());
-			}
-		}
-
-		TSlateAttributeBase(SWidget& Widget, TAttribute<ObjectType>&& Attribute, ObjectType&& InitialValue)
-			: Value((Attribute.IsSet() && !Attribute.IsBound()) ? MoveTemp(Attribute.Steal().template Get<ObjectType>()) : MoveTemp(InitialValue))
-#if UE_SLATE_WITH_MEMBER_ATTRIBUTE_DEBUGGING
-			, Debug_OwningWidget(&Widget)
-#endif
-		{
-			if (Attribute.IsBound())
-			{
-				ConstructWrapper(Widget, MoveTemp(Attribute.Steal().template Get<FGetter>()));
-			}
-		}
-
-	public:
-		/** @return the SlateAttribute cached value. If the SlateAttribute is bound, the value will be cached at the end of the every frame. */ 
-		UE_NODISCARD const ObjectType& Get() const
-		{
-			return Value;
-		}
-		
-		/** Update the cached value and invalidate the widget if needed. */
-		void UpdateNow(SWidget& Widget)
-		{
-			VerifyOwningWidget(Widget);
-			ProtectedUpdateNow(Widget, InAttributeType);
-		}
-
-	public:
-		/**
-		 * Unbind the SlateAttribute and set its value. It may invalidate the Widget if the value is different.
-		 * @return true if the value is considered different and an invalidation occurred.
-		 */
-		bool Set(SWidget& Widget, const ObjectType& NewValue)
-		{
-			VerifyOwningWidget(Widget);
-			ProtectedUnregisterAttribute(Widget, InAttributeType);
-
-			const bool bIsIdentical = IdenticalTo(Widget, Value, NewValue);
-			if (!bIsIdentical)
-			{
-				Value = NewValue;
-				ProtectedInvalidateWidget(Widget, InAttributeType, GetInvalidationReason(Widget));
-			}
-			return !bIsIdentical;
-		}
-
-		/**
-		 * Unbind the SlateAttribute and set its value. It may invalidate the Widget if the value is different.
-		 * @return true if the value is considered different and an invalidation occurred.
-		 */
-		bool Set(SWidget& Widget, ObjectType&& NewValue)
-		{
-			VerifyOwningWidget(Widget);
-			ProtectedUnregisterAttribute(Widget, InAttributeType);
-
-			const bool bIsIdentical = IdenticalTo(Widget, Value, NewValue);
-			if (!bIsIdentical)
-			{
-				Value = MoveTemp(NewValue);
-				ProtectedInvalidateWidget(Widget, InAttributeType, GetInvalidationReason(Widget));
-			}
-			return !bIsIdentical;
-		}
-
-	public:
-		/**
-		 * Bind the SlateAttribute to the Getter function.
-		 * (If enabled) Update the value from the Getter. If the value is different, then invalidate the widget.
-		 * The SlateAttribute will now be updated every frame from the Getter.
-		 */
-		void Bind(SWidget& Widget, const FGetter& Getter)
-		{
-			VerifyOwningWidget(Widget);
-			if (Getter.IsBound())
-			{
-				AssignBinding(Widget, Getter);
-			}
-			else
-			{
-				ProtectedUnregisterAttribute(Widget, InAttributeType);
-			}
-		}
-
-		/**
-		 * Bind the SlateAttribute to the Getter function.
-		 * (If enabled) Update the value from the Getter. If the value is different, then invalidate the widget.
-		 * The SlateAttribute will now be updated every frame from the Getter.
-		 */
-		void Bind(SWidget& Widget, FGetter&& Getter)
-		{
-			VerifyOwningWidget(Widget);
-			if (Getter.IsBound())
-			{
-				AssignBinding(Widget, MoveTemp(Getter));
-			}
-			else
-			{
-				ProtectedUnregisterAttribute(Widget, InAttributeType);
-			}
-		}
-
-		/**
-		 * Bind the SlateAttribute to the newly create getter function.
-		 * (If enabled) Update the value from the getter. If the value is different, then invalidate the widget.
-		 * The SlateAttribute will now be updated every frame from the Getter.
-		 */
-		template<typename WidgetType, typename U = typename std::enable_if<std::is_base_of<SWidget, WidgetType>::value>::type>
-		void Bind(WidgetType& Widget, typename FGetter::template TConstMethodPtr<WidgetType> MethodPtr)
-		{
-			Bind(Widget, FGetter::CreateSP(&Widget, MethodPtr));
-		}
-
-		/**
-		 * Bind the SlateAttribute to the Attribute Getter function (if it exist).
-		 * (If enabled) Update the value from the getter. If the value is different, then invalidate the widget.
-		 * The SlateAttribute will now be updated every frame from the Getter.
-		 * Or
-		 * Set the SlateAttribute's value if the Attribute is not bound but is set.
-		 * This will Unbind any previously bound getter.
-		 * If the value is different, then invalidate the widget.
-		 * Or
-		 * Unbind the SlateAttribute if the Attribute is not bound and not set.
-		 * @see Set
-		 * @return true if the getter was assigned or if the value is considered different and an invalidation occurred.
-		 */
-		bool Assign(SWidget& Widget, const TAttribute<ObjectType>& OtherAttribute)
-		{
-			VerifyOwningWidget(Widget);
-			if (OtherAttribute.IsBound())
-			{
-				return AssignBinding(Widget, OtherAttribute.GetBinding());
-			}
-			else if (OtherAttribute.IsSet())
-			{
-				return Set(Widget, OtherAttribute.Get());
-			}
-			else
-			{
-				ProtectedUnregisterAttribute(Widget, InAttributeType);
-				return false;
-			}
-		}
-
-		bool Assign(SWidget& Widget, TAttribute<ObjectType>&& OtherAttribute)
-		{
-			VerifyOwningWidget(Widget);
-			if (OtherAttribute.IsBound())
-			{
-				return AssignBinding(Widget, MoveTemp(OtherAttribute.Steal().template Get<FGetter>()));
-			}
-			else if (OtherAttribute.IsSet())
-			{
-				return Set(Widget, MoveTemp(OtherAttribute.Steal().template Get<ObjectType>()));
-			}
-			else
-			{
-				ProtectedUnregisterAttribute(Widget, InAttributeType);
-				return false;
-			}
-		}
-
-		bool Assign(SWidget& Widget, const TAttribute<ObjectType>& OtherAttribute, const ObjectType& DefaultValue)
-		{
-			VerifyOwningWidget(Widget);
-			if (OtherAttribute.IsBound())
-			{
-				return AssignBinding(Widget, OtherAttribute.GetBinding());
-			}
-			else if (OtherAttribute.IsSet())
-			{
-				return Set(Widget, OtherAttribute.Get());
-			}
-			else
-			{
-				return Set(Widget, DefaultValue);
-			}
-		}
-
-		bool Assign(SWidget& Widget, TAttribute<ObjectType>&& OtherAttribute, const ObjectType& DefaultValue)
-		{
-			VerifyOwningWidget(Widget);
-			if (OtherAttribute.IsBound())
-			{
-				return AssignBinding(Widget, MoveTemp(OtherAttribute.Steal().template Get<FGetter>()));
-			}
-			else if (OtherAttribute.IsSet())
-			{
-				return Set(Widget, MoveTemp(OtherAttribute.Steal().template Get<ObjectType>()));
-			}
-			else
-			{
-				return Set(Widget, DefaultValue);
-			}
-		}
-
-		bool Assign(SWidget& Widget, const TAttribute<ObjectType>& OtherAttribute, ObjectType&& DefaultValue)
-		{
-			VerifyOwningWidget(Widget);
-			if (OtherAttribute.IsBound())
-			{
-				return AssignBinding(Widget, OtherAttribute.GetBinding());
-			}
-			else if (OtherAttribute.IsSet())
-			{
-				return Set(Widget, OtherAttribute.Get());
-			}
-			else
-			{
-				return Set(Widget, MoveTemp(DefaultValue));
-			}
-		}
-
-		bool Assign(SWidget& Widget, TAttribute<ObjectType>&& OtherAttribute, ObjectType&& DefaultValue)
-		{
-			VerifyOwningWidget(Widget);
-			if (OtherAttribute.IsBound())
-			{
-				return AssignBinding(Widget, MoveTemp(OtherAttribute.Steal().template Get<FGetter>()));
-			}
-			else if (OtherAttribute.IsSet())
-			{
-				return Set(Widget, MoveTemp(OtherAttribute.Steal().template Get<ObjectType>()));
-			}
-			else
-			{
-				return Set(Widget, MoveTemp(DefaultValue));
-			}
-		}
-
-		/**
-		 * Remove the Getter function.
-		 * The Slate Attribute will not be updated anymore and will keep its current cached value.
-		 */
-		void Unbind(SWidget& Widget)
-		{
-			VerifyOwningWidget(Widget);
-			ProtectedUnregisterAttribute(Widget, InAttributeType);
-		}
-
-	public:
-		/** Build a Attribute from this SlateAttribute. */
-		UE_NODISCARD TAttribute<ObjectType> ToAttribute(const SWidget& Widget) const
-		{
-			if (ISlateAttributeGetter* Delegate = ProtectedFindGetter(Widget, InAttributeType))
-			{
-				return TAttribute<ObjectType>::Create(static_cast<FSlateAttributeGetterWrapper<TSlateAttributeBase>*>(Delegate)->GetDelegate());
-			}
-			return TAttribute<ObjectType>(Get());
-		}
-
-		/** @return True if the SlateAttribute is bound to a getter function. */
-		UE_NODISCARD bool IsBound(const SWidget& Widget) const
-		{
-			VerifyOwningWidget(Widget);
-			return ProtectedIsBound(Widget, InAttributeType);
-		}
-
-		/** @return True if they have the same Getter or the same value. */
-		UE_NODISCARD bool IdenticalTo(const SWidget& Widget, const TSlateAttributeBase& Other) const
-		{
-			VerifyOwningWidget(Widget);
-			FDelegateHandle ThisDelegateHandle = ProtectedFindGetterHandle(Widget, InAttributeType);
-			FDelegateHandle OthersDelegateHandle = ProtectedFindGetterHandle(Widget, InAttributeType);
-			if (ThisDelegateHandle == OthersDelegateHandle)
-			{
-				if (ThisDelegateHandle.IsValid())
-				{
-					return true;
-				}
-				return IdenticalTo(Widget, Get(), Other.Get());
-			}
-			return false;
-		}
-
-		/** @return True if they have the same Getter or, if the Attribute is set, the same value. */
-		UE_NODISCARD bool IdenticalTo(const SWidget& Widget, const TAttribute<ObjectType>& Other) const
-		{
-			VerifyOwningWidget(Widget);
-			FDelegateHandle ThisDelegateHandle = ProtectedFindGetterHandle(Widget, InAttributeType);
-			if (Other.IsBound())
-			{
-				return Other.GetBinding().GetHandle() == ThisDelegateHandle;
-			}
-			return !ThisDelegateHandle.IsValid() && Other.IsSet() && IdenticalTo(Widget, Get(), Other.Get());
-		}
-
-	private:
-		void ConstructWrapper(SWidget& Widget, const FGetter& Getter)
-		{
-			TUniquePtr<ISlateAttributeGetter> Wrapper = MakeUniqueGetter(*this, Getter);
-			ProtectedRegisterAttribute(Widget, InAttributeType, MoveTemp(Wrapper));
-			UpdateNowOnBind(Widget);
-		}
-
-		void ConstructWrapper(SWidget& Widget, FGetter&& Getter)
-		{
-			TUniquePtr<ISlateAttributeGetter> Wrapper = MakeUniqueGetter(*this, MoveTemp(Getter));
-			ProtectedRegisterAttribute(Widget, InAttributeType, MoveTemp(Wrapper));
-			UpdateNowOnBind(Widget);
-		}
-
-		bool AssignBinding(SWidget& Widget, const FGetter& Getter)
-		{
-			const FDelegateHandle PreviousGetterHandle = ProtectedFindGetterHandle(Widget, InAttributeType);
-			if (PreviousGetterHandle != Getter.GetHandle())
-			{
-				ConstructWrapper(Widget, Getter);
-				return true;
-			}
-			return false;
-		}
-
-		bool AssignBinding(SWidget& Widget, FGetter&& Getter)
-		{
-			const FDelegateHandle PreviousGetterHandle = ProtectedFindGetterHandle(Widget, InAttributeType);
-			if (PreviousGetterHandle != Getter.GetHandle())
-			{
-				ConstructWrapper(Widget, MoveTemp(Getter));
-				return true;
-			}
-			return false;
-		}
-
-	private:
-		template<typename SlateAttributeType>
-		class FSlateAttributeGetterWrapper : public ISlateAttributeGetter
-		{
-		public:
-			using ObjectType = typename SlateAttributeType::ObjectType;
-			using FGetter = typename TAttribute<ObjectType>::FGetter;
-			using FComparePredicate = typename SlateAttributeType::FComparePredicate;
-
-			FSlateAttributeGetterWrapper() = delete;
-			FSlateAttributeGetterWrapper(const FSlateAttributeGetterWrapper&) = delete;
-			FSlateAttributeGetterWrapper& operator= (const FSlateAttributeGetterWrapper&) = delete;
-			virtual ~FSlateAttributeGetterWrapper() = default;
-
-		public:
-			FSlateAttributeGetterWrapper(SlateAttributeType& InOwningAttribute, const FGetter& InGetterDelegate)
-				: Getter(InGetterDelegate)
-				, Attribute(&InOwningAttribute)
-			{
-			}
-			FSlateAttributeGetterWrapper(SlateAttributeType& InOwningAttribute, FGetter&& InGetterDelegate)
-				: Getter(MoveTemp(InGetterDelegate))
-				, Attribute(&InOwningAttribute)
-			{
-			}
-
-			virtual FUpdateAttributeResult UpdateAttribute(const SWidget& Widget) override
-			{
-				ObjectType NewValue = Getter.Execute();
-
-				const bool bIsIdentical = Attribute->IdenticalTo(Widget, Attribute->Value, NewValue);
-				if (!bIsIdentical)
-				{
-					// Set the value on the widget
-					Attribute->Value = MoveTemp(NewValue);
-					return Attribute->GetInvalidationReason(Widget);
-				}
-				return FUpdateAttributeResult();
-			}
-
-			virtual const FSlateAttributeBase& GetAttribute() const override
-			{
-				return *Attribute;
-			}
-
-			virtual void SetAttribute(FSlateAttributeBase& InBase) override
-			{
-				Attribute = &static_cast<SlateAttributeType&>(InBase);
-			}
-
-			virtual FDelegateHandle GetDelegateHandle() const override
-			{
-				return Getter.GetHandle();
-			}
-
-			const FGetter& GetDelegate() const
-			{
-				return Getter;
-			}
-
-		private:
-			/** Getter function to fetch the new value of the SlateAttribute. */
-			FGetter Getter;
-			/** The SlateAttribute of the SWidget owning the value. */
-			SlateAttributeType* Attribute;
-		};
-
-		static TUniquePtr<ISlateAttributeGetter> MakeUniqueGetter(TSlateAttributeBase& Attribute, const FGetter& Getter)
-		{
-			return MakeUnique<FSlateAttributeGetterWrapper<TSlateAttributeBase>>(Attribute, Getter);
-		}
-
-		static TUniquePtr<ISlateAttributeGetter> MakeUniqueGetter(TSlateAttributeBase& Attribute, FGetter&& Getter)
-		{
-			return MakeUnique<FSlateAttributeGetterWrapper<TSlateAttributeBase>>(Attribute, MoveTemp(Getter));
-		}
-
-	private:
-		ObjectType Value;
-
-	protected:
-#if UE_SLATE_WITH_MEMBER_ATTRIBUTE_DEBUGGING
-		SWidget* Debug_OwningWidget;
-#endif
-	};
-
-
-	/**
-	 *
-	 */
-	template<typename InObjectType, typename InInvalidationReasonPredicate, typename InComparePredicate>
-	struct TSlateMemberAttribute : public TSlateAttributeBase<InObjectType, InInvalidationReasonPredicate, InComparePredicate, ESlateAttributeType::Member>
-	{
-	private:
-		using Super = TSlateAttributeBase<InObjectType, InInvalidationReasonPredicate, InComparePredicate, ESlateAttributeType::Member>;
-
-		template<typename WidgetType, typename U = typename std::enable_if<std::is_base_of<SWidget, WidgetType>::value>::type>
-		static void VerifyAttributeAddress(WidgetType& Widget, TSlateMemberAttribute* Self)
-		{
-			checkf((UPTRINT)Self >= (UPTRINT)&Widget && (UPTRINT)Self < (UPTRINT)&Widget + sizeof(WidgetType),
-				TEXT("Use TAttribute or TSlateManagedAttribute instead. See SlateAttribute.h for more info."));
-			ensureAlwaysMsgf((HasDefinedInvalidationReason || Self->ProtectedIsImplemented(Widget)),
-				TEXT("The TSlateAttribute could not be found in the SlateAttributeDescriptor.\n")
-				TEXT("Use the SLATE_DECLARE_WIDGET and add the attribute in PrivateRegisterAttributes,\n")
-				TEXT("Or use TSlateAttribute with a valid Invalidation Reason instead."));
-		}
-
-	public:
-		using FGetter = typename Super::FGetter;
-		using ObjectType = typename Super::ObjectType;
-		static const bool IsMemberType = true;
-		static constexpr bool HasDefinedInvalidationReason = !std::is_same<InInvalidationReasonPredicate, FSlateAttributeNoInvalidationReason>::value;
-
-	public:
-		//~ You can only register Attribute that are defined in a SWidget.
-		//~ Use the constructor with Widget pointer.
-		//~ See documentation on top of this file for more info.
-		TSlateMemberAttribute() = delete;
-		TSlateMemberAttribute(const TSlateMemberAttribute&) = delete;
-		TSlateMemberAttribute(TSlateMemberAttribute&&) = delete;
-		TSlateMemberAttribute& operator=(const TSlateMemberAttribute&) = delete;
-		TSlateMemberAttribute& operator=(TSlateMemberAttribute&&) = delete;
-
-#if UE_SLATE_WITH_MEMBER_ATTRIBUTE_DEBUGGING
-		~TSlateMemberAttribute()
-		{
-			/**
-			 * The parent should now be destroyed and the shared pointer should be invalidate.
-			 * If you hit the check, that means the TSlateAttribute is not a variable member of SWidget.
-			 * It will introduced bad memory access.
-			 * See documentation above.
-			*/
-			checkf(Super::ProtectedIsWidgetInDestructionPath(Super::Debug_OwningWidget), TEXT("The Owning widget should be invalid."));
-		}
-#endif
-
-		template<typename WidgetType, typename U = typename std::enable_if<std::is_base_of<SWidget, WidgetType>::value>::type>
-		explicit TSlateMemberAttribute(WidgetType& Widget)
-			: Super(Widget)
-		{
-			VerifyAttributeAddress(Widget, this);
-		}
-
-		template<typename WidgetType, typename U = typename std::enable_if<std::is_base_of<SWidget, WidgetType>::value>::type>
-		explicit TSlateMemberAttribute(WidgetType& Widget, const ObjectType& InValue)
-			: Super(Widget, InValue)
-		{
-			VerifyAttributeAddress(Widget, this);
-		}
-
-		template<typename WidgetType, typename U = typename std::enable_if<std::is_base_of<SWidget, WidgetType>::value>::type>
-		explicit TSlateMemberAttribute(WidgetType& Widget, ObjectType&& InValue)
-			: Super(Widget, MoveTemp(InValue))
-		{
-			VerifyAttributeAddress(Widget, this);
-		}
-	};
-
-
-	/**
-	 *
-	 */
-	template<typename InObjectType, typename InInvalidationReasonPredicate, typename InComparePredicate>
-	struct TSlateManagedAttribute : protected TSlateAttributeBase<InObjectType, InInvalidationReasonPredicate, InComparePredicate, ESlateAttributeType::Managed>
-	{
-	private:
-		using Super = TSlateAttributeBase<InObjectType, InInvalidationReasonPredicate, InComparePredicate, ESlateAttributeType::Managed>;
-
-	public:
-		using ObjectType = typename Super::ObjectType;
-		using FInvalidationReasonPredicate = typename Super::FInvalidationReasonPredicate;
-		using FGetter = typename Super::FGetter;
-		using FComparePredicate = typename Super::FComparePredicate;
-		static const bool IsMemberType = false;
-
-		static EInvalidateWidgetReason GetInvalidationReason(const SWidget& Widget) { return Super::GetInvalidationReason(Widget); }
-		static bool IdenticalTo(const SWidget& Widget, const ObjectType& Lhs, const ObjectType& Rhs) { return Super::IdenticalTo(Widget, Lhs, Rhs); }
-
-	public:
-		TSlateManagedAttribute() = delete;
-		TSlateManagedAttribute(const TSlateManagedAttribute&) = delete;
-		TSlateManagedAttribute& operator=(const TSlateManagedAttribute&) = delete;
-
-		~TSlateManagedAttribute()
-		{
-			Unbind();
-		}
-
-		TSlateManagedAttribute(TSlateManagedAttribute&& Other)
-			: Super(MoveTemp(Other))
-			, ManagedWidget(MoveTemp(Other.ManagedWidget))
-		{
-			if (TSharedPtr<SWidget> Pin = ManagedWidget.Pin())
-			{
-				FSlateAttributeImpl::ProtectedMoveAttribute(*Pin.Get(), ESlateAttributeType::Managed, &Other);
-			}
-		}
-
-		TSlateManagedAttribute& operator=(TSlateManagedAttribute&& Other)
-		{
-			Super::operator=(MoveTemp(Other));
-			ManagedWidget = MoveTemp(Other.ManagedWidget);
-
-			if (TSharedPtr<SWidget> Pin = ManagedWidget.Pin())
-			{
-				FSlateAttributeImpl::ProtectedMoveAttribute(*Pin.Get(), ESlateAttributeType::Managed, &Other);
-			}
-			return *this;
-		}
-
-		explicit TSlateManagedAttribute(TSharedRef<SWidget> Widget)
-			: Super(Widget.Get())
-			, ManagedWidget(Widget)
-		{ }
-
-		explicit TSlateManagedAttribute(TSharedRef<SWidget> Widget, const ObjectType& InValue)
-			: Super(Widget.Get(), InValue)
-			, ManagedWidget(Widget)
-		{
-		}
-
-		explicit TSlateManagedAttribute(TSharedRef<SWidget> Widget, ObjectType&& InValue)
-			: Super(Widget.Get(), MoveTemp(InValue))
-			, ManagedWidget(Widget)
-		{
-		}
-
-		explicit TSlateManagedAttribute(TSharedRef<SWidget> Widget, const FGetter& Getter, const ObjectType& InitialValue)
-			: Super(Widget.Get(), Getter, InitialValue)
-			, ManagedWidget(Widget)
-		{
-		}
-
-		explicit TSlateManagedAttribute(TSharedRef<SWidget> Widget, const FGetter& Getter, ObjectType&& InitialValue)
-			: Super(Widget.Get(), Getter, MoveTemp(InitialValue))
-			, ManagedWidget(Widget)
-		{
-		}
-
-		explicit TSlateManagedAttribute(TSharedRef<SWidget> Widget, FGetter&& Getter, const ObjectType& InitialValue)
-			: Super(Widget.Get(), MoveTemp(Getter), InitialValue)
-			, ManagedWidget(Widget)
-		{
-		}
-
-		explicit TSlateManagedAttribute(TSharedRef<SWidget> Widget, FGetter&& Getter, ObjectType&& InitialValue)
-			: Super(Widget.Get(), MoveTemp(Getter), MoveTemp(InitialValue))
-			, ManagedWidget(Widget)
-		{
-		}
-
-		explicit TSlateManagedAttribute(TSharedRef<SWidget> Widget, const TAttribute<ObjectType>& Attribute, const ObjectType& InitialValue)
-			: Super(Widget.Get(), Attribute, InitialValue)
-			, ManagedWidget(Widget)
-		{
-		}
-
-		explicit TSlateManagedAttribute(TSharedRef<SWidget> Widget, TAttribute<ObjectType>&& Attribute, ObjectType&& InitialValue)
-			: Super(Widget.Get(), MoveTemp(Attribute), MoveTemp(InitialValue))
-			, ManagedWidget(Widget)
-		{
-		}
-
-	public:
-		const ObjectType& Get() const
-		{
-			return Super::Get();
-		}
-
-		void UpdateNow()
-		{
-			if (TSharedPtr<SWidget> Pin = ManagedWidget.Pin())
-			{
-				Super::UpdateNow(*Pin.Get());
-			}
-		}
-
-	public:
-		void Set(const ObjectType& NewValue)
-		{
-			if (TSharedPtr<SWidget> Pin = ManagedWidget.Pin())
-			{
-				Super::Set(*Pin.Get(), NewValue);
-			}
-		}
-
-		void Set(ObjectType&& NewValue)
-		{
-			if (TSharedPtr<SWidget> Pin = ManagedWidget.Pin())
-			{
-				Super::Set(*Pin.Get(), MoveTemp(NewValue));
-			}
-		}
-
-		void Bind(const FGetter& Getter)
-		{
-			if (TSharedPtr<SWidget> Pin = ManagedWidget.Pin())
-			{
-				Super::Bind(*Pin.Get(), Getter);
-			}
-		}
-
-		void Bind(FGetter&& Getter)
-		{
-			if (TSharedPtr<SWidget> Pin = ManagedWidget.Pin())
-			{
-				Super::Bind(*Pin.Get(), MoveTemp(Getter));
-			}
-		}
-
-		void Assign(const TAttribute<ObjectType>& OtherAttribute)
-		{
-			if (TSharedPtr<SWidget> Pin = ManagedWidget.Pin())
-			{
-				Super::Assign(*Pin.Get(), OtherAttribute);
-			}
-		}
-
-		void Assign(TAttribute<ObjectType>&& OtherAttribute)
-		{
-			if (TSharedPtr<SWidget> Pin = ManagedWidget.Pin())
-			{
-				Super::Assign(*Pin.Get(), MoveTemp(OtherAttribute));
-			}
-		}
-
-		void Assign(const TAttribute<ObjectType>& OtherAttribute, const ObjectType& DefaultValue)
-		{
-			if (TSharedPtr<SWidget> Pin = ManagedWidget.Pin())
-			{
-				Super::Assign(*Pin.Get(), OtherAttribute, DefaultValue);
-			}
-		}
-
-		void Assign(const TAttribute<ObjectType>& OtherAttribute, ObjectType&& DefaultValue)
-		{
-			if (TSharedPtr<SWidget> Pin = ManagedWidget.Pin())
-			{
-				Super::Assign(*Pin.Get(), OtherAttribute, MoveTemp(DefaultValue));
-			}
-		}
-
-		void Assign(TAttribute<ObjectType>&& OtherAttribute, const ObjectType& DefaultValue)
-		{
-			if (TSharedPtr<SWidget> Pin = ManagedWidget.Pin())
-			{
-				Super::Assign(*Pin.Get(), MoveTemp(OtherAttribute), DefaultValue);
-			}
-		}
-
-		void Assign(TAttribute<ObjectType>&& OtherAttribute, ObjectType&& DefaultValue)
-		{
-			if (TSharedPtr<SWidget> Pin = ManagedWidget.Pin())
-			{
-				Super::Assign(*Pin.Get(), MoveTemp(OtherAttribute), MoveTemp(DefaultValue));
-			}
-		}
-
-		void Unbind()
-		{
-			if (TSharedPtr<SWidget> Pin = ManagedWidget.Pin())
-			{
-				Super::Unbind(*Pin.Get());
-			}
-		}
-
-	public:
-		bool IsBound() const
-		{
-			if (TSharedPtr<SWidget> Pin = ManagedWidget.Pin())
-			{
-				return Super::IsBound(*Pin.Get());
-			}
-			return false;
-		}
-
-		bool IdenticalTo(const TSlateManagedAttribute& Other) const
-		{
-			if (TSharedPtr<SWidget> Pin = ManagedWidget.Pin())
-			{
-				return Super::IdenticalTo(*Pin.Get(), Other);
-			}
-			return false;
-		}
-
-		bool IdenticalTo(const TAttribute<ObjectType>& Other) const
-		{
-			if (TSharedPtr<SWidget> Pin = ManagedWidget.Pin())
-			{
-				return Super::IdenticalTo(*Pin.Get(), Other);
-			}
-			return false;
-		}
-
-	private:
-		TWeakPtr<SWidget> ManagedWidget;
-	};
-
-
-	/*
-	 * A reference to a SlateAttribute that can be returned and saved for later.
-	 */
-	template<typename AttributeMemberType>
-	struct TSlateMemberAttributeRef
-	{
-	public:
-		using SlateAttributeType = AttributeMemberType;
-		using ObjectType = typename AttributeMemberType::ObjectType;
-		using AttributeType = TAttribute<ObjectType>;
-
-	private:
-		template<typename WidgetType>
-		static void VerifyAttributeAddress(const WidgetType& InWidget, const AttributeMemberType& InAttribute)
-		{
-			checkf((UPTRINT)&InAttribute >= (UPTRINT)&InWidget && (UPTRINT)&InAttribute < (UPTRINT)&InWidget + sizeof(WidgetType),
-				TEXT("The attribute is not a member of the widget."));
-			InAttribute.VerifyOwningWidget(InWidget);
-		}
-
-	public:
-		/** Constructor */
-		template<typename WidgetType, typename V = typename std::enable_if<std::is_base_of<SWidget, WidgetType>::value>::type>
-		explicit TSlateMemberAttributeRef(const WidgetType& InOwner, const AttributeMemberType& InAttribute)
-			: Owner(InOwner.AsShared())
-			, Attribute(InAttribute)
-		{
-			VerifyAttributeAddress(InOwner, InAttribute);
-		}
-
-	public:
-		/** @return if the reference is valid. A reference can be invalid if the SWidget is destroyed. */
-		UE_NODISCARD bool IsValid() const
-		{
-			return Owner.IsValid();
-		}
-
-		/** @return the SlateAttribute cached value; undefined when IsValid() returns false. */
-		UE_NODISCARD const ObjectType& Get() const
-		{
-			if (TSharedPtr<const SWidget> Pin = Owner.Pin())
-			{
-				return Attribute.Get();
-			}
-			checkf(false, TEXT("It is an error to call GetValue() on an unset TSlateMemberAttributeRef. Please either check IsValid() or use Get(DefaultValue) instead."));
-			static ObjectType Tmp;
-			return Tmp;
-		}
-
-		/** @return the SlateAttribute cached value or the DefaultValue if the reference is invalid. */
-		UE_NODISCARD const ObjectType& Get(const ObjectType& DefaultValue) const
-		{
-			if (TSharedPtr<const SWidget> Pin = Owner.Pin())
-			{
-				return Attribute.Get();
-			}
-			return DefaultValue;
-		}
-
-		/** Update the cached value and invalidate the widget if needed. */
-		void UpdateValue()
-		{
-			if (TSharedPtr<const SWidget> Pin = Owner.Pin())
-			{
-				const_cast<AttributeMemberType>(Attribute).UpdateNow(*Pin.Get());
-			}
-		}
-
-		/**
-		 * Assumes the reference is valid.
-		 * Shorthand for the boilerplace code: MyAttribute.UpdateValueNow(); MyAttribute.Get();
-		 */
-		UE_NODISCARD const ObjectType& UpdateAndGet()
-		{
-			if (TSharedPtr<const SWidget> Pin = Owner.Pin())
-			{
-				const_cast<AttributeMemberType>(Attribute)->UpdateNow(*Pin.Get());
-				return Attribute.Get();
-			}
-			checkf(false, TEXT("It is an error to call GetValue() on an unset TSlateMemberAttributeRef. Please either check IsValid() or use Get(DefaultValue) instead."));
-		}
-
-		/**
-		 * Shorthand for the boilerplace code: MyAttribute.UpdateValueNow(); MyAttribute.Get(DefaultValue);
-		 * @return the SlateAttribute cached value or the DefaultValue if the reference is invalid.
-		 */
-		UE_NODISCARD const ObjectType& UpdateAndGet(const ObjectType& DefaultValue)
-		{
-			if (TSharedPtr<const SWidget> Pin = Owner.Pin())
-			{
-				const_cast<AttributeMemberType>(Attribute)->UpdateNow(*Pin.Get());
-				return Attribute.Get();
-			}
-			return DefaultValue;
-		}
-
-		/** Build a Attribute from this SlateAttribute. */
-		UE_NODISCARD TAttribute<ObjectType> ToAttribute() const
-		{
-			if (TSharedPtr<const SWidget> Pin = Owner.Pin())
-			{
-				return Attribute.ToAttribute(*Pin.Get());
-			}
-			return TAttribute<ObjectType>();
-		}
-
-		/** @return True if the SlateAttribute is bound to a getter function. */
-		UE_NODISCARD bool IsBound() const
-		{
-			if (TSharedPtr<const SWidget> Pin = Owner.Pin())
-			{
-				return Attribute.IsBound(*Pin.Get());
-			}
-			return false;
-		}
-
-		/** @return True if they have the same Getter or the same value. */
-		UE_NODISCARD bool IdenticalTo(const TSlateMemberAttributeRef& Other) const
-		{
-			TSharedPtr<const SWidget> SelfPin = Owner.Pin();
-			TSharedPtr<const SWidget> OtherPin = Other.Owner.Pin();
-			if (SelfPin == OtherPin && SelfPin)
-			{
-				return Attribute.IdenticalTo(*SelfPin.Get(), Other.Attribute);
-			}
-			return SelfPin == OtherPin;
-		}
-
-		/** @return True if they have the same Getter or, if the Attribute is set, the same value. */
-		UE_NODISCARD bool IdenticalTo(const TAttribute<ObjectType>& Other) const
-		{
-			if (TSharedPtr<const SWidget> Pin = Owner.Pin())
-			{
-				return Attribute.IdenticalTo(*Pin.Get(), Other);
-			}
-			return !Other.IsSet(); // if the other is not set, then both are invalid.
-		}
-
-	private:
-		TWeakPtr<const SWidget> Owner;
-		const AttributeMemberType& Attribute;
-	 };
-
-
-} // SlateAttributePrivate
+#include "Types/Attributes/SlateAttributeDefinition.h"
+#include "Types/Attributes/SlateAttributeBase.h"
+#include "Types/Attributes/SlateAttributeManaged.h"
+#include "Types/Attributes/SlateAttributeMember.h"
diff --git a/Engine/Source/Runtime/SlateCore/Public/Types/SlateAttributeDescriptor.h b/Engine/Source/Runtime/SlateCore/Public/Types/SlateAttributeDescriptor.h
--- a/Engine/Source/Runtime/SlateCore/Public/Types/SlateAttributeDescriptor.h
+++ b/Engine/Source/Runtime/SlateCore/Public/Types/SlateAttributeDescriptor.h
@@ -147,7 +147,7 @@ public:
 	};
 
 	/** @returns the number of Attributes registered. */
-	int32 AttributeNum() const { return Attributes.Num(); }
+	int32 GetAttributeNum() const { return Attributes.Num(); }
 
 	/** @returns the Attribute at the index previously found with IndexOfMemberAttribute */
 	const FAttribute& GetAttributeAtIndex(int32 Index) const;
@@ -166,13 +166,13 @@ public:
 
 	/** Iterator over each dependency this attribute is responsible of. */
 	template<typename Predicate>
-	void ForEachDependency(FAttribute const& Attribute, Predicate Pred) const
+	void ForEachDependentsOn(FAttribute const& Attribute, Predicate Pred) const
 	{
 		checkf(&Attribute >= Attributes.GetData() && &Attribute <= Attributes.GetData()+Attributes.Num()
 			, TEXT("The attribute is not part of this Descriptor."));
 		if (Attribute.bIsADependencyForSomeoneElse)
 		{
-			ForEachDependencyImpl(Attribute.Name, (int32)(&Attribute - Attributes.GetData()), Pred);
+			ForEachDependentsOfImpl(Attribute.Name, (int32)(&Attribute - Attributes.GetData()), Pred);
 		}
 	}
 
@@ -185,7 +185,7 @@ private:
 	void SetAffectVisibility(FAttribute& Attribute, bool bUpdate);
 
 	template<typename Predicate>
-	void ForEachDependencyImpl(FName const& LookForName, int32 Index, Predicate& Pred) const
+	void ForEachDependentsOfImpl(FName const& LookForName, int32 Index, Predicate& Pred) const
 	{
 		++Index;
 		for (; Index < Attributes.Num(); ++Index)
@@ -196,7 +196,7 @@ private:
 				Pred(Index);
 				if (Other.bIsADependencyForSomeoneElse)
 				{
-					ForEachDependencyImpl(Other.Name, Index, Pred);
+					ForEachDependentsOfImpl(Other.Name, Index, Pred);
 				}
 			}
 		}
diff --git a/Engine/Source/Runtime/SlateCore/Public/Types/SlateAttributeMetaData.h b/Engine/Source/Runtime/SlateCore/Public/Types/SlateAttributeMetaData.h
--- a/Engine/Source/Runtime/SlateCore/Public/Types/SlateAttributeMetaData.h
+++ b/Engine/Source/Runtime/SlateCore/Public/Types/SlateAttributeMetaData.h
@@ -92,6 +92,10 @@ private:
 	void RegisterAttributeImpl(SWidget& OwningWidget, FSlateAttributeBase& Attribute, ESlateAttributeType AttributeType, TUniquePtr<ISlateAttributeGetter>&& Getter);
 	bool UnregisterAttributeImpl(const FSlateAttributeBase& Attribute);
 	void UpdateAttributesImpl(SWidget& OwningWidget, EInvalidationPermission InvaldiationStyle, int32 StartIndex, int32 IndexNum);
+	void SetNeedToResetFlag(int32 Index)
+	{
+		ResetFlag |= (Index < AffectVisibilityCounter) ? EResetFlags::NeedToReset_OnlyVisibility : EResetFlags::NeedToReset_ExceptVisibility;
+	}
 
 private:
 	int32 IndexOfAttribute(const FSlateAttributeBase& Attribute) const
@@ -102,7 +106,9 @@ private:
 
 	struct FGetterItem
 	{
-		FGetterItem() = default;
+		using FAttributeIndex = uint8;
+		static const FAttributeIndex InvalidAttributeIndex;
+
 		FGetterItem(const FGetterItem&) = delete;
 		FGetterItem(FGetterItem&&) = default;
 		FGetterItem& operator=(const FGetterItem&) = delete;
@@ -110,23 +116,23 @@ private:
 			: Attribute(InAttribute)
 			, Getter(MoveTemp(InGetter))
 			, SortOrder(InSortOrder)
-			, CachedAttributeDescriptorIndex(INDEX_NONE)
-			, CachedAttributeDependencyIndex(INDEX_NONE)
+			, CachedAttributeDescriptorIndex(InvalidAttributeIndex)
+			, CachedAttributeDependencyIndex(InvalidAttributeIndex)
 			, Flags(0)
 		{ }
-		FGetterItem(FSlateAttributeBase* InAttribute, uint32 InSortOrder, TUniquePtr<SlateAttributePrivate::ISlateAttributeGetter> && InGetter, int16 InAttributeDescriptorIndex)
+		FGetterItem(FSlateAttributeBase* InAttribute, uint32 InSortOrder, TUniquePtr<SlateAttributePrivate::ISlateAttributeGetter> && InGetter, FAttributeIndex InAttributeDescriptorIndex)
 			: Attribute(InAttribute)
 			, Getter(MoveTemp(InGetter))
 			, SortOrder(InSortOrder)
 			, CachedAttributeDescriptorIndex(InAttributeDescriptorIndex)
-			, CachedAttributeDependencyIndex(INDEX_NONE)
+			, CachedAttributeDependencyIndex(InvalidAttributeIndex)
 			, Flags(0)
 		{ }
 		FSlateAttributeBase* Attribute;
 		TUniquePtr<SlateAttributePrivate::ISlateAttributeGetter> Getter;
 		uint32 SortOrder;
-		int16 CachedAttributeDescriptorIndex;
-		int16 CachedAttributeDependencyIndex;
+		FAttributeIndex CachedAttributeDescriptorIndex;
+		FAttributeIndex CachedAttributeDependencyIndex;
 		union
 		{
 			struct
@@ -135,9 +141,10 @@ private:
 				int8 bUpdatedThisFrame : 1;
 				int8 bUpatedManually : 1;
 				int8 bIsADependencyForSomeoneElse : 1;
-				int8 bIsMemberType : 1;
-				int8 bIsManagedType : 1;
 				int8 bAffectVisibility : 1;
+				ESlateAttributeType AttributeType : 2;
+				int8 Unused0 : 1;
+				//~if more space is needed bUpatedManually&bUpdatedThisFrame could be merged
 			};
 			int8 Flags;
 		};
@@ -152,6 +159,7 @@ private:
 		/** If available, return the name of the attribute. */
 		FName GetAttributeName(const SWidget& OwningWidget) const;
 	};
+	static_assert(sizeof(FGetterItem) <= 32, "The size of FGetterItem is bigger than expected.");
 
 	TArray<FGetterItem, TInlineAllocator<4>> Attributes;
 	//~ There is a possibility that the widget has a CachedInvalidationReason and a parent become collapsed.
@@ -159,6 +167,16 @@ private:
 	//~1. The widget is collapsed indirectly, so we do not care if it's invalidated.
 	//~2. The parent widget will clear this widget PersistentState.
 	EInvalidateWidgetReason CachedInvalidationReason = EInvalidateWidgetReason::None;
-	bool bHasUpdatedManuallyFlagToReset = false;
 	uint8 AffectVisibilityCounter = 0;
+
+	enum class EResetFlags : uint8
+	{
+		None = 0,
+		NeedToReset_OnlyVisibility = (1 << 0),
+		NeedToReset_ExceptVisibility = (1 << 1),
+	};
+	FRIEND_ENUM_CLASS_FLAGS(EResetFlags);
+	EResetFlags ResetFlag = EResetFlags::None;
 };
+
+ENUM_CLASS_FLAGS(FSlateAttributeMetaData::EResetFlags);
diff --git a/Engine/Source/Runtime/SlateCore/Public/Widgets/Images/SImage.h b/Engine/Source/Runtime/SlateCore/Public/Widgets/Images/SImage.h
--- a/Engine/Source/Runtime/SlateCore/Public/Widgets/Images/SImage.h
+++ b/Engine/Source/Runtime/SlateCore/Public/Widgets/Images/SImage.h
@@ -83,13 +83,13 @@ protected:
 	// End SWidget overrides.
 
 	/** @return an attribute reference of Image */
-	TSlateAttributeRef<const FSlateBrush*> GetImageAttribute() const { return TSlateAttributeRef<FSlateBrush const*>(*this, ImageAttribute); }
+	TSlateAttributeRef<const FSlateBrush*> GetImageAttribute() const { return TSlateAttributeRef<FSlateBrush const*>(SharedThis(this), ImageAttribute); }
 
 	/** @return an attribute reference of ColorAndOpacity */
-	TSlateAttributeRef<FSlateColor> GetColorAndOpacityAttribute() const { return TSlateAttributeRef<FSlateColor>(*this, ColorAndOpacityAttribute); }
+	TSlateAttributeRef<FSlateColor> GetColorAndOpacityAttribute() const { return TSlateAttributeRef<FSlateColor>(SharedThis(this), ColorAndOpacityAttribute); }
 
 	/** @return an attribute reference of DesiredSizeOverride */
-	TSlateAttributeRef<TOptional<FVector2D>> GetDesiredSizeOverrideAttribute() const { return TSlateAttributeRef<TOptional<FVector2D>>(*this, DesiredSizeOverrideAttribute); }
+	TSlateAttributeRef<TOptional<FVector2D>> GetDesiredSizeOverrideAttribute() const { return TSlateAttributeRef<TOptional<FVector2D>>(SharedThis(this), DesiredSizeOverrideAttribute); }
 
 #if WITH_EDITORONLY_DATA
 	UE_DEPRECATED(5.0, "Direct access to Image is now deprecated. Use the setter or getter.")
diff --git a/Engine/Source/Runtime/SlateCore/Public/Widgets/SCompoundWidget.h b/Engine/Source/Runtime/SlateCore/Public/Widgets/SCompoundWidget.h
--- a/Engine/Source/Runtime/SlateCore/Public/Widgets/SCompoundWidget.h
+++ b/Engine/Source/Runtime/SlateCore/Public/Widgets/SCompoundWidget.h
@@ -89,13 +89,13 @@ protected:
 	// End SWidget overrides.
 
 	/** @return an attribute reference of ContentScale */
-	TSlateAttributeRef<FVector2D> GetContentScaleAttribute() const { return TSlateAttributeRef<FVector2D>(*this, ContentScaleAttribute); }
+	TSlateAttributeRef<FVector2D> GetContentScaleAttribute() const { return TSlateAttributeRef<FVector2D>(SharedThis(this), ContentScaleAttribute); }
 
 	/** @return an attribute reference of ColorAndOpacity */
-	TSlateAttributeRef<FLinearColor> GetColorAndOpacityAttribute() const { return TSlateAttributeRef<FLinearColor>{*this, ColorAndOpacityAttribute}; }
+	TSlateAttributeRef<FLinearColor> GetColorAndOpacityAttribute() const { return TSlateAttributeRef<FLinearColor>{SharedThis(this), ColorAndOpacityAttribute}; }
 
 	/** @return an attribute reference of ForegroundColor */
-	TSlateAttributeRef<FSlateColor> GetForegroundColorAttribute() const {	return TSlateAttributeRef<FSlateColor>{*this, ForegroundColorAttribute}; }
+	TSlateAttributeRef<FSlateColor> GetForegroundColorAttribute() const {	return TSlateAttributeRef<FSlateColor>{SharedThis(this), ForegroundColorAttribute}; }
 
 protected:
 
diff --git a/Engine/Source/Runtime/SlateCore/Public/Widgets/SWidget.h b/Engine/Source/Runtime/SlateCore/Public/Widgets/SWidget.h
--- a/Engine/Source/Runtime/SlateCore/Public/Widgets/SWidget.h
+++ b/Engine/Source/Runtime/SlateCore/Public/Widgets/SWidget.h
@@ -1671,13 +1671,13 @@ protected:
 	}
 
 	/** @return an attribute reference of EnabledStateAttribute */
-	TSlateAttributeRef<bool> GetEnabledStateAttribute() const { return TSlateAttributeRef<bool>(*this, EnabledStateAttribute); }
+	TSlateAttributeRef<bool> GetEnabledStateAttribute() const { return TSlateAttributeRef<bool>(SharedThis(this), EnabledStateAttribute); }
 	/** @return an attribute reference of VisibilityAttribute */
-	TSlateAttributeRef<EVisibility> GetVisibilityAttribute() const { return TSlateAttributeRef<EVisibility>(*this, VisibilityAttribute); }
+	TSlateAttributeRef<EVisibility> GetVisibilityAttribute() const { return TSlateAttributeRef<EVisibility>(SharedThis(this), VisibilityAttribute); }
 	/** @return an attribute reference of RenderTransformAttribute */
-	TSlateAttributeRef<TOptional<FSlateRenderTransform>> GetRenderTransformAttribute() const { return TSlateAttributeRef<TOptional<FSlateRenderTransform>>(*this, RenderTransformAttribute); }
+	TSlateAttributeRef<TOptional<FSlateRenderTransform>> GetRenderTransformAttribute() const { return TSlateAttributeRef<TOptional<FSlateRenderTransform>>(SharedThis(this), RenderTransformAttribute); }
 	/** @return an attribute reference of RenderTransformPivotAttribute */
-	TSlateAttributeRef<FVector2D> GetRenderTransformPivotAttribute() const { return TSlateAttributeRef<FVector2D>(*this, RenderTransformPivotAttribute); }
+	TSlateAttributeRef<FVector2D> GetRenderTransformPivotAttribute() const { return TSlateAttributeRef<FVector2D>(SharedThis(this), RenderTransformPivotAttribute); }
 
 protected:
 	/** Dtor ensures that active timer handles are UnRegistered with the SlateApplication. */
