Slate: Review fix from CL 16157339
#rnx

diff --git a/Engine/Source/Runtime/SlateCore/Private/Types/SlateAttributeMetaData.cpp b/Engine/Source/Runtime/SlateCore/Private/Types/SlateAttributeMetaData.cpp
--- a/Engine/Source/Runtime/SlateCore/Private/Types/SlateAttributeMetaData.cpp
+++ b/Engine/Source/Runtime/SlateCore/Private/Types/SlateAttributeMetaData.cpp
@@ -101,6 +101,7 @@ void FSlateAttributeMetaData::RegisterAttributeImpl(SWidget& OwningWidget, FSlat
 					const int32 FoundDependencyAttributeIndex = Descriptor.IndexOfMemberAttribute(FoundAttribute.Prerequisite);
 					if (FoundDependencyAttributeIndex != INDEX_NONE)
 					{
+						check(FoundDependencyAttributeIndex < std::numeric_limits<FGetterItem::FAttributeIndex>::max());
 						GetterItem.CachedAttributeDependencyIndex = (FGetterItem::FAttributeIndex)FoundDependencyAttributeIndex;
 					}
 				}
@@ -239,7 +240,7 @@ void FSlateAttributeMetaData::InvalidateWidget(SWidget& OwningWidget, const FSla
 			// The dependency attribute need to be updated in the update loop (note that it may not be registered yet)
 			if (GetterItem.bIsADependencyForSomeoneElse)
 			{
-				GetterItem.bUpatedManually = true;
+				GetterItem.bUpdatedManually = true;
 				AttributeMetaData->SetNeedToResetFlag(FoundIndex);
 			}
 		}
@@ -296,7 +297,7 @@ void FSlateAttributeMetaData::UpdateAllAttributes(SWidget& OwningWidget, EInvali
 		{
 			for (FGetterItem& GetterItem : AttributeMetaData->Attributes)
 			{
-				GetterItem.bUpatedManually = false;
+				GetterItem.bUpdatedManually = false;
 				GetterItem.bUpdatedThisFrame = false;
 			}
 			AttributeMetaData->ResetFlag = EResetFlags::None;
@@ -314,15 +315,15 @@ void FSlateAttributeMetaData::UpdateOnlyVisibilityAttributes(SWidget& OwningWidg
 			const int32 StartIndex = 0;
 			const int32 EndIndex = AttributeMetaData->AffectVisibilityCounter;
 			AttributeMetaData->UpdateAttributesImpl(OwningWidget, InvalidationStyle, StartIndex, EndIndex);
-			if (EnumHasAllFlags(AttributeMetaData->ResetFlag, EResetFlags::NeedToReset_ExceptVisibility))
+			if (EnumHasAllFlags(AttributeMetaData->ResetFlag, EResetFlags::NeedToReset_OnlyVisibility))
 			{
 				for (int32 Index = StartIndex; Index < EndIndex; ++Index)
 				{
 					FGetterItem& GetterItem = AttributeMetaData->Attributes[Index];
-					GetterItem.bUpatedManually = false;
+					GetterItem.bUpdatedManually = false;
 					GetterItem.bUpdatedThisFrame = false;
 				}
-				EnumRemoveFlags(AttributeMetaData->ResetFlag, EResetFlags::NeedToReset_ExceptVisibility);
+				EnumRemoveFlags(AttributeMetaData->ResetFlag, EResetFlags::NeedToReset_OnlyVisibility);
 			}
 		}
 	}
@@ -343,7 +344,7 @@ void FSlateAttributeMetaData::UpdateExceptVisibilityAttributes(SWidget& OwningWi
 				for (int32 Index = StartIndex; Index < EndIndex; ++Index)
 				{
 					FGetterItem& GetterItem = AttributeMetaData->Attributes[Index];
-					GetterItem.bUpatedManually = false;
+					GetterItem.bUpdatedManually = false;
 					GetterItem.bUpdatedThisFrame = false;
 				}
 				EnumRemoveFlags(AttributeMetaData->ResetFlag, EResetFlags::NeedToReset_ExceptVisibility);
@@ -393,7 +394,7 @@ void FSlateAttributeMetaData::UpdateAttributesImpl(SWidget& OwningWidget, EInval
 				if (OtherGetterItem.CachedAttributeDescriptorIndex == GetterItem.CachedAttributeDependencyIndex)
 				{
 					bFound = true;
-					bShouldUpdate = OtherGetterItem.bUpdatedThisFrame || OtherGetterItem.bUpatedManually;
+					bShouldUpdate = OtherGetterItem.bUpdatedThisFrame || OtherGetterItem.bUpdatedManually;
 					break;
 				}
 			}
@@ -453,7 +454,7 @@ void FSlateAttributeMetaData::UpdateAttribute(SWidget& OwningWidget, FSlateAttri
 				// The dependency attribute need to be updated in the update loop (note that it may not be registered yet)
 				if (GetterItem.bIsADependencyForSomeoneElse)
 				{
-					GetterItem.bUpatedManually = true;
+					GetterItem.bUpdatedManually = true;
 					AttributeMetaData->SetNeedToResetFlag(FoundIndex);
 				}
 			}
diff --git a/Engine/Source/Runtime/SlateCore/Public/Types/SlateAttributeMetaData.h b/Engine/Source/Runtime/SlateCore/Public/Types/SlateAttributeMetaData.h
--- a/Engine/Source/Runtime/SlateCore/Public/Types/SlateAttributeMetaData.h
+++ b/Engine/Source/Runtime/SlateCore/Public/Types/SlateAttributeMetaData.h
@@ -139,7 +139,7 @@ private:
 			{
 				int8 bUpdatedOnce : 1;
 				int8 bUpdatedThisFrame : 1;
-				int8 bUpatedManually : 1;
+				int8 bUpdatedManually : 1;
 				int8 bIsADependencyForSomeoneElse : 1;
 				int8 bAffectVisibility : 1;
 				ESlateAttributeType AttributeType : 2;
