Slate: Fix unitialized SlateAttribute. FLinearColor, FVector2D (and others) do not initialize their values in the default constructor. The value in memory at that time can be a nan (just by chance). They will not be assigned because we check if the value changed before doing any assignment.

diff --git a/Engine/Source/Runtime/Slate/Private/Widgets/Colors/SColorBlock.cpp b/Engine/Source/Runtime/Slate/Private/Widgets/Colors/SColorBlock.cpp
--- a/Engine/Source/Runtime/Slate/Private/Widgets/Colors/SColorBlock.cpp
+++ b/Engine/Source/Runtime/Slate/Private/Widgets/Colors/SColorBlock.cpp
@@ -19,14 +19,14 @@ void SColorBlock::PrivateRegisterAttributes(FSlateAttributeInitializer& Attribut
 }
 
 SColorBlock::SColorBlock()
-	: Color(*this)
-	, AlphaBackgroundBrush(*this)
-	, GradientCornerRadius(*this)
-	, ColorBlockSize(*this)
-	, AlphaDisplayMode(*this)
-	, ColorIsHSV(*this)
-	, ShowBackgroundForAlpha(*this)
-	, bUseSRGB(*this)
+	: Color(*this, FLinearColor::White)
+	, AlphaBackgroundBrush(*this, nullptr)
+	, GradientCornerRadius(*this, FVector4(0.0f))
+	, ColorBlockSize(*this, FVector2D(16, 16))
+	, AlphaDisplayMode(*this, EColorBlockAlphaDisplayMode::Combined)
+	, ColorIsHSV(*this, false)
+	, ShowBackgroundForAlpha(*this, false)
+	, bUseSRGB(*this, true)
 {
 }
 
diff --git a/Engine/Source/Runtime/Slate/Private/Widgets/Colors/SColorGradingWheel.cpp b/Engine/Source/Runtime/Slate/Private/Widgets/Colors/SColorGradingWheel.cpp
--- a/Engine/Source/Runtime/Slate/Private/Widgets/Colors/SColorGradingWheel.cpp
+++ b/Engine/Source/Runtime/Slate/Private/Widgets/Colors/SColorGradingWheel.cpp
@@ -13,7 +13,7 @@ void SColorGradingWheel::PrivateRegisterAttributes(FSlateAttributeInitializer& A
 }
 
 SColorGradingWheel::SColorGradingWheel()
-	: SelectedColorAttribute(*this)
+	: SelectedColorAttribute(*this, FLinearColor(ForceInit))
 	, DesiredWheelSizeAttribute(*this)
 	, ExponentDisplacementAttribute(*this)
 	, Union_IsAttributeSet(0)
diff --git a/Engine/Source/Runtime/Slate/Private/Widgets/Colors/SColorSpectrum.cpp b/Engine/Source/Runtime/Slate/Private/Widgets/Colors/SColorSpectrum.cpp
--- a/Engine/Source/Runtime/Slate/Private/Widgets/Colors/SColorSpectrum.cpp
+++ b/Engine/Source/Runtime/Slate/Private/Widgets/Colors/SColorSpectrum.cpp
@@ -6,7 +6,7 @@
 
 
 SColorSpectrum::SColorSpectrum()
-	: SelectedColor(*this)
+	: SelectedColor(*this, FLinearColor(ForceInit))
 {
 }
 
diff --git a/Engine/Source/Runtime/Slate/Private/Widgets/Colors/SColorWheel.cpp b/Engine/Source/Runtime/Slate/Private/Widgets/Colors/SColorWheel.cpp
--- a/Engine/Source/Runtime/Slate/Private/Widgets/Colors/SColorWheel.cpp
+++ b/Engine/Source/Runtime/Slate/Private/Widgets/Colors/SColorWheel.cpp
@@ -5,7 +5,7 @@
 #include "Styling/CoreStyle.h"
 
 SColorWheel::SColorWheel()
-	: SelectedColor(*this)
+	: SelectedColor(*this, FLinearColor(ForceInit))
 {}
 
 /* SColorWheel methods
diff --git a/Engine/Source/Runtime/Slate/Private/Widgets/Colors/SSimpleGradient.cpp b/Engine/Source/Runtime/Slate/Private/Widgets/Colors/SSimpleGradient.cpp
--- a/Engine/Source/Runtime/Slate/Private/Widgets/Colors/SSimpleGradient.cpp
+++ b/Engine/Source/Runtime/Slate/Private/Widgets/Colors/SSimpleGradient.cpp
@@ -8,8 +8,8 @@
 /* SSimpleGradient implementation
  *****************************************************************************/
 SSimpleGradient::SSimpleGradient()
-	: StartColor(*this)
-	, EndColor(*this)
+	: StartColor(*this, FLinearColor(0.0f, 0.0f, 0.0f))
+	, EndColor(*this, FLinearColor(1.0f, 1.0f, 1.0f))
 {
 	SetCanTick(false);
 }
diff --git a/Engine/Source/Runtime/Slate/Private/Widgets/Layout/SSpacer.cpp b/Engine/Source/Runtime/Slate/Private/Widgets/Layout/SSpacer.cpp
--- a/Engine/Source/Runtime/Slate/Private/Widgets/Layout/SSpacer.cpp
+++ b/Engine/Source/Runtime/Slate/Private/Widgets/Layout/SSpacer.cpp
@@ -9,7 +9,7 @@ void SSpacer::PrivateRegisterAttributes(FSlateAttributeInitializer& AttributeIni
 }
 
 SSpacer::SSpacer()
-	: SpacerSize(*this)
+	: SpacerSize(*this, FVector2D::ZeroVector)
 	, bIsSpacerSizeBound(false)
 {
 	SetCanTick(false);
diff --git a/Engine/Source/Runtime/Slate/Private/Widgets/SViewport.cpp b/Engine/Source/Runtime/Slate/Private/Widgets/SViewport.cpp
--- a/Engine/Source/Runtime/Slate/Private/Widgets/SViewport.cpp
+++ b/Engine/Source/Runtime/Slate/Private/Widgets/SViewport.cpp
@@ -12,22 +12,20 @@ DECLARE_CYCLE_STAT(TEXT("Game UI Tick"), STAT_ViewportTickTime, STATGROUP_Slate)
 DECLARE_CYCLE_STAT(TEXT("Game UI Paint"), STAT_ViewportPaintTime, STATGROUP_Slate);
 
 
-static const FVector2D DefaultViewportSize(320.0f, 240.0f);
-
 
 /* SViewport::FArguments interface
 *****************************************************************************/
 FVector2D SViewport::FArguments::GetDefaultViewportSize()
 {
-	return DefaultViewportSize;
+	return FVector2D(320.0f, 240.0f);
 }
 
 /* SViewport constructors
  *****************************************************************************/
 
 SViewport::SViewport()
-	: ShowDisabledEffect(*this)
-	, ViewportSize(*this)
+	: ShowDisabledEffect(*this, true)
+	, ViewportSize(*this, SViewport::FArguments::GetDefaultViewportSize())
 	, bRenderDirectlyToWindow(false)
 	, bEnableGammaCorrection(true)
 	, bReverseGammaCorrection(false)
diff --git a/Engine/Source/Runtime/Slate/Private/Widgets/Text/STextBlock.cpp b/Engine/Source/Runtime/Slate/Private/Widgets/Text/STextBlock.cpp
--- a/Engine/Source/Runtime/Slate/Private/Widgets/Text/STextBlock.cpp
+++ b/Engine/Source/Runtime/Slate/Private/Widgets/Text/STextBlock.cpp
@@ -56,19 +56,19 @@ STextBlock::STextBlock()
 	, Font(*this)
 	, StrikeBrush(*this)
 	, ColorAndOpacity(*this)
-	, ShadowOffset(*this)
-	, ShadowColorAndOpacity(*this)
-	, HighlightColor(*this)
+	, ShadowOffset(*this, FVector2D::ZeroVector)
+	, ShadowColorAndOpacity(*this, FLinearColor(ForceInit))
+	, HighlightColor(*this, FLinearColor(ForceInit))
 	, HighlightShape(*this)
 	, HighlightText(*this)
-	, WrapTextAt(*this)
-	, AutoWrapText(*this)
-	, WrappingPolicy(*this)
+	, WrapTextAt(*this, 0.0f)
+	, AutoWrapText(*this, false)
+	, WrappingPolicy(*this, ETextWrappingPolicy::DefaultWrapping)
 	, TransformPolicy(*this)
 	, Margin(*this)
 	, Justification(*this)
-	, LineHeightPercentage(*this)
-	, MinDesiredWidth(*this)
+	, LineHeightPercentage(*this, 1.0f)
+	, MinDesiredWidth(*this, 0.0f)
 	, Union_Flags(0)
 	, bSimpleTextMode(false)
 {
diff --git a/Engine/Source/Runtime/SlateCore/Public/Types/Attributes/SlateAttributeBase.h b/Engine/Source/Runtime/SlateCore/Public/Types/Attributes/SlateAttributeBase.h
--- a/Engine/Source/Runtime/SlateCore/Public/Types/Attributes/SlateAttributeBase.h
+++ b/Engine/Source/Runtime/SlateCore/Public/Types/Attributes/SlateAttributeBase.h
@@ -45,14 +45,42 @@ namespace SlateAttributePrivate
 #endif
 		}
 
+		void VerifyNan() const
+		{
+#if UE_SLATE_WITH_ATTRIBUTE_NAN_DIAGNOSTIC
+			if constexpr (std::is_same<FVector, InObjectType>::value)
+			{
+				ensureMsgf(!Value.ContainsNaN(), TEXT("Value contains a NaN. Initialize your FVector properly (see FVector::EForceInit)"));
+			}
+			else if constexpr (std::is_same<FVector2D, InObjectType>::value)
+			{
+				ensureMsgf(!Value.ContainsNaN(), TEXT("Value contains a NaN. Initialize your FVector2D properly (see FVector2D::EForceInit)"));
+			}
+			else if constexpr (std::is_same<FLinearColor, InObjectType>::value)
+			{
+				ensureMsgf(!FVector4(Value).ContainsNaN(), TEXT("Value contains a NaN. Initialize your FLinearColor properly (see FLinearColor::EForceInit)"));
+			}
+			else if constexpr (std::is_same<FRotator, InObjectType>::value)
+			{
+				ensureMsgf(!Value.ContainsNaN(), TEXT("Value contains a NaN. Initialize your FRotator properly (see FRotator::EForceInit)"));
+			}
+#endif
+		}
+
 	public:
 #if UE_SLATE_WITH_MEMBER_ATTRIBUTE_DEBUGGING
 		TSlateAttributeBase()
-			: Debug_OwningWidget(nullptr)
+			: Value()
+			, Debug_OwningWidget(nullptr)
 		{
+			VerifyNan();
 		}
 #else
-		TSlateAttributeBase() = default;
+		TSlateAttributeBase()
+			: Value()
+		{
+			VerifyNan();
+		}
 #endif
 
 		TSlateAttributeBase(const ObjectType& InValue)
@@ -61,6 +89,7 @@ namespace SlateAttributePrivate
 			, Debug_OwningWidget(nullptr)
 #endif
 		{
+			VerifyNan();
 		}
 
 		TSlateAttributeBase(ObjectType&& InValue)
@@ -69,6 +98,7 @@ namespace SlateAttributePrivate
 			, Debug_OwningWidget(nullptr)
 #endif
 		{
+			VerifyNan();
 		}
 
 		TSlateAttributeBase(SWidget& Widget)
@@ -77,6 +107,7 @@ namespace SlateAttributePrivate
 			, Debug_OwningWidget(&Widget)
 #endif
 		{
+			VerifyNan();
 		}
 
 		TSlateAttributeBase(SWidget& Widget, const ObjectType& InValue)
@@ -85,6 +116,7 @@ namespace SlateAttributePrivate
 			, Debug_OwningWidget(&Widget)
 #endif
 		{
+			VerifyNan();
 		}
 
 		TSlateAttributeBase(SWidget& Widget, ObjectType&& InValue)
@@ -93,6 +125,7 @@ namespace SlateAttributePrivate
 			, Debug_OwningWidget(&Widget)
 #endif
 		{
+			VerifyNan();
 		}
 
 		TSlateAttributeBase(SWidget& Widget, const FGetter& Getter, const ObjectType& InitialValue)
@@ -101,6 +134,7 @@ namespace SlateAttributePrivate
 			, Debug_OwningWidget(&Widget)
 #endif
 		{
+			VerifyNan();
 			if (Getter.IsBound())
 			{
 				ConstructWrapper(Widget, Getter);
@@ -113,6 +147,7 @@ namespace SlateAttributePrivate
 			, Debug_OwningWidget(&Widget)
 #endif
 		{
+			VerifyNan();
 			if (Getter.IsBound())
 			{
 				ConstructWrapper(Widget, Getter);
@@ -125,6 +160,7 @@ namespace SlateAttributePrivate
 			, Debug_OwningWidget(&Widget)
 #endif
 		{
+			VerifyNan();
 			if (Getter.IsBound())
 			{
 				ConstructWrapper(Widget, MoveTemp(Getter));
@@ -137,6 +173,7 @@ namespace SlateAttributePrivate
 			, Debug_OwningWidget(&Widget)
 #endif
 		{
+			VerifyNan();
 			if (Getter.IsBound())
 			{
 				ConstructWrapper(Widget, MoveTemp(Getter));
@@ -149,6 +186,7 @@ namespace SlateAttributePrivate
 			, Debug_OwningWidget(&Widget)
 #endif
 		{
+			VerifyNan();
 			if (Attribute.IsBound())
 			{
 				ConstructWrapper(Widget, Attribute.GetBinding());
@@ -161,6 +199,7 @@ namespace SlateAttributePrivate
 			, Debug_OwningWidget(&Widget)
 #endif
 		{
+			VerifyNan();
 			if (Attribute.IsBound())
 			{
 				ConstructWrapper(Widget, MoveTemp(Attribute.Steal().template Get<FGetter>()));
diff --git a/Engine/Source/Runtime/SlateCore/Public/Types/SlateAttribute.h b/Engine/Source/Runtime/SlateCore/Public/Types/SlateAttribute.h
--- a/Engine/Source/Runtime/SlateCore/Public/Types/SlateAttribute.h
+++ b/Engine/Source/Runtime/SlateCore/Public/Types/SlateAttribute.h
@@ -15,6 +15,9 @@
 #ifndef UE_SLATE_WITH_ATTRIBUTE_INITIALIZATION_ON_BIND
 	#define UE_SLATE_WITH_ATTRIBUTE_INITIALIZATION_ON_BIND 0
 #endif
+#ifndef UE_SLATE_WITH_ATTRIBUTE_NAN_DIAGNOSTIC
+	#define UE_SLATE_WITH_ATTRIBUTE_NAN_DIAGNOSTIC DO_ENSURE
+#endif
 
 /**
  * Use TSlateAttribute when it's a SWidget member.
