Slate: Add option to verify if a widget is already in the invalidaiton heap. Enabled with "Slate.InvalidationRoot.VerifyWidgetHeapContains"

diff --git a/Engine/Source/Runtime/SlateCore/Private/FastUpdate/SlateInvalidationWidgetHeap.cpp b/Engine/Source/Runtime/SlateCore/Private/FastUpdate/SlateInvalidationWidgetHeap.cpp
new file mode 100644
--- /dev/null
+++ b/Engine/Source/Runtime/SlateCore/Private/FastUpdate/SlateInvalidationWidgetHeap.cpp
@@ -0,0 +1,22 @@
+// Copyright Epic Games, Inc. All Rights Reserved.
+
+#include "FastUpdate/SlateInvalidationWidgetHeap.h"
+
+namespace UE
+{
+namespace Slate
+{
+namespace Private
+{
+#if UE_SLATE_WITH_INVALIDATIONWIDGETHEAP_DEBUGGING
+	bool GSlateInvalidationWidgetHeapVerifyWidgetContains = false;
+	static FAutoConsoleVariableRef CVarSlateInvalidationWidgetHeapVerifyWidgetContains(
+		TEXT("Slate.InvalidationRoot.VerifyWidgetHeapContains"),
+		GSlateInvalidationWidgetHeapVerifyWidgetContains,
+		TEXT("Verify that the widget is not already in the list before adding it.")
+	);
+#endif
+
+} // Private
+} // Slate
+} // UE
diff --git a/Engine/Source/Runtime/SlateCore/Private/FastUpdate/SlateInvalidationWidgetHeap.h b/Engine/Source/Runtime/SlateCore/Private/FastUpdate/SlateInvalidationWidgetHeap.h
--- a/Engine/Source/Runtime/SlateCore/Private/FastUpdate/SlateInvalidationWidgetHeap.h
+++ b/Engine/Source/Runtime/SlateCore/Private/FastUpdate/SlateInvalidationWidgetHeap.h
@@ -8,6 +8,9 @@
 #include "FastUpdate/SlateInvalidationWidgetList.h"
 #include "FastUpdate/SlateInvalidationWidgetSortOrder.h"
 
+#ifndef UE_SLATE_WITH_INVALIDATIONWIDGETHEAP_DEBUGGING
+	#define UE_SLATE_WITH_INVALIDATIONWIDGETHEAP_DEBUGGING DO_CHECK
+#endif
 
 namespace UE
 {
@@ -15,6 +18,10 @@ namespace Slate
 {
 namespace Private
 {
+#if UE_SLATE_WITH_INVALIDATIONWIDGETHEAP_DEBUGGING
+	extern bool GSlateInvalidationWidgetHeapVerifyWidgetContains;
+#endif
+
 	/**
 	* Ordered list of WidgetIndex. The order is based on the WidgetSortIndex.
 	*/
@@ -65,7 +72,7 @@ public:
 	void HeapPushUnique(FSlateInvalidationWidgetList::InvalidationWidgetType& InvalidationWidget)
 	{
 		check(InvalidationWidget.Index != FSlateInvalidationWidgetIndex::Invalid);
-		check(Contains_Debug(InvalidationWidget.Index) == InvalidationWidget.bContainedByWidgetPreHeap);
+		VerifyContainsFlag(InvalidationWidget);
 
 		if (!InvalidationWidget.bContainedByWidgetPreHeap)
 		{
@@ -177,6 +184,16 @@ public:
 	}
 
 private:
+	void VerifyContainsFlag(const FSlateInvalidationWidgetList::InvalidationWidgetType& InvalidationWidget)
+	{
+#if UE_SLATE_WITH_INVALIDATIONWIDGETHEAP_DEBUGGING
+		if (UE::Slate::Private::GSlateInvalidationWidgetHeapVerifyWidgetContains)
+		{
+			check(Contains_Debug(InvalidationWidget.Index) == InvalidationWidget.bContainedByWidgetPostHeap);
+		}
+#endif
+	}
+
 	FElementContainer Heap;
 	FSlateInvalidationWidgetList& OwnerList;
 };
@@ -211,8 +228,7 @@ public:
 	{
 		check(bIsHeap);
 		check(InvalidationWidget.Index != FSlateInvalidationWidgetIndex::Invalid);
-		check(Contains_Debug(InvalidationWidget.Index) == InvalidationWidget.bContainedByWidgetPostHeap
-			|| WidgetCannotBeAdded == InvalidationWidget.Index);
+		VerifyContainsFlag(InvalidationWidget);
 
 		if (!InvalidationWidget.bContainedByWidgetPostHeap)
 		{
@@ -226,8 +242,7 @@ public:
 	{
 		check(bIsHeap == false);
 		check(InvalidationWidget.Index != FSlateInvalidationWidgetIndex::Invalid);
-		check(Contains_Debug(InvalidationWidget.Index) == InvalidationWidget.bContainedByWidgetPostHeap
-			|| WidgetCannotBeAdded == InvalidationWidget.Index);
+		VerifyContainsFlag(InvalidationWidget);
 
 		if (!InvalidationWidget.bContainedByWidgetPostHeap)
 		{
@@ -369,6 +384,17 @@ public:
 	};
 
 private:
+	void VerifyContainsFlag(const FSlateInvalidationWidgetList::InvalidationWidgetType& InvalidationWidget)
+	{
+#if UE_SLATE_WITH_INVALIDATIONWIDGETHEAP_DEBUGGING
+		if (UE::Slate::Private::GSlateInvalidationWidgetHeapVerifyWidgetContains)
+		{
+			check(Contains_Debug(InvalidationWidget.Index) == InvalidationWidget.bContainedByWidgetPostHeap
+				|| WidgetCannotBeAdded == InvalidationWidget.Index);
+		}
+#endif
+	}
+
 	FElementContainer Heap;
 	FSlateInvalidationWidgetList& OwnerList;
 	FSlateInvalidationWidgetIndex WidgetCannotBeAdded;
