Slate: Add update flag for volatile prepass. These widgets always need a prepass in fastpath. Fix issue with collapsed, the desired size was not updated when a parent was also collapsed. Fix the wrong position of elements inside a Retainer when the Retainer moves. Keep an array of all widgets that always need to be updated (tick, timer, paint). That prevents notifying the proxy and the final list once the widget got painted and ensure that all elements are updated.

diff --git a/Engine/Source/Developer/SlateReflector/Private/Models/WidgetReflectorNode.cpp b/Engine/Source/Developer/SlateReflector/Private/Models/WidgetReflectorNode.cpp
--- a/Engine/Source/Developer/SlateReflector/Private/Models/WidgetReflectorNode.cpp
+++ b/Engine/Source/Developer/SlateReflector/Private/Models/WidgetReflectorNode.cpp
@@ -745,7 +745,11 @@ TSharedRef<FWidgetReflectorNodeBase> FWidgetReflectorNodeUtils::NewNodeTreeFrom(
 	TSharedRef<FWidgetReflectorNodeBase> NewNodeInstance = NewNode(InNodeType, InWidgetGeometry);
 
 	TSharedRef<SWidget> CurWidgetParent = InWidgetGeometry.Widget;
-	FChildren* Children = CurWidgetParent->Advanced_IsInvalidationRoot() ? CurWidgetParent->GetAllChildren() : CurWidgetParent->GetChildren();
+#if WITH_SLATE_DEBUGGING
+	FChildren* Children = CurWidgetParent->Debug_GetChildrenForReflector();
+#else
+	FChildren* Children = CurWidgetParent->GetChildren();
+#endif
 
 	auto BuildChild = [NewNodeInstance, CurWidgetParent, InNodeType](const TSharedRef<SWidget>& ChildWidget)
 	{
diff --git a/Engine/Source/Runtime/Slate/Private/Widgets/SInvalidationPanel.cpp b/Engine/Source/Runtime/Slate/Private/Widgets/SInvalidationPanel.cpp
--- a/Engine/Source/Runtime/Slate/Private/Widgets/SInvalidationPanel.cpp
+++ b/Engine/Source/Runtime/Slate/Private/Widgets/SInvalidationPanel.cpp
@@ -10,6 +10,7 @@
 #include "Framework/Application/SlateApplication.h"
 #include "Types/ReflectionMetadata.h"
 #include "Rendering/SlateObjectReferenceCollector.h"
+#include "Widgets/SNullWidget.h"
 
 DECLARE_CYCLE_STAT(TEXT("SInvalidationPanel::Paint"), STAT_SlateInvalidationPaint, STATGROUP_Slate);
 
@@ -46,6 +47,7 @@ SInvalidationPanel::SInvalidationPanel()
 	SetInvalidationRootWidget(*this);
 	SetInvalidationRootHittestGrid(HittestGrid.Get());
 	SetCanTick(false);
+	SetVolatilePrepass(GetCanCache());
 
 	LastIncomingColorAndOpacity = FLinearColor::White;
 
@@ -117,14 +119,15 @@ bool SInvalidationPanel::GetCanCache() const
 	bDebugFlags = bInvalidationPanelsEnabled;
 #endif
 
-	return bCanCache && !GSlateEnableGlobalInvalidation && !GetProxyHandle().HasValidInvalidationRootOwnership(this) && bDebugFlags;
+	return bCanCache && !GSlateEnableGlobalInvalidation && bDebugFlags;
 }
 
 void SInvalidationPanel::OnGlobalInvalidationToggled(bool bGlobalInvalidationEnabled)
 {
 	InvalidateRootChildOrder();
-
 	ClearAllFastPathData(true);
+	SetVolatilePrepass(GetCanCache());
+	Invalidate(EInvalidateWidgetReason::LayoutAndVolatility);
 }
 
 bool SInvalidationPanel::UpdateCachePrequisites(FSlateWindowElementList& OutDrawElements, const FGeometry& AllottedGeometry, const FSlateRect& MyCullingRect, int32 LayerId, const FWidgetStyle& InWidgetStyle) const
@@ -179,12 +182,8 @@ void SInvalidationPanel::SetCanCache(bool InCanCache)
 	if (bCanCache != InCanCache)
 	{
 		bCanCache = InCanCache;
-		// If the cache changed, invalidate the parent need to rebuild its list
-		//since InvalidationPanel cannot be nested in regular mode.
-		if (GetProxyHandle().IsValid(this) && !GSlateEnableGlobalInvalidation)
-		{
-			GetProxyHandle().GetInvalidationRootHandle().Advanced_GetInvalidationRootNoCheck()->InvalidateRootChildOrder(this);
-		}
+		SetVolatilePrepass(GetCanCache());
+		InvalidateRootChildOrder();
 	}
 }
 
@@ -200,10 +199,12 @@ FChildren* SInvalidationPanel::GetChildren()
 	}
 }
 
-FChildren* SInvalidationPanel::GetAllChildren()
+#if WITH_SLATE_DEBUGGING
+FChildren* SInvalidationPanel::Debug_GetChildrenForReflector()
 {
 	return SCompoundWidget::GetChildren();
 }
+#endif
 
 int32 SInvalidationPanel::OnPaint( const FPaintArgs& Args, const FGeometry& AllottedGeometry, const FSlateRect& MyCullingRect, FSlateWindowElementList& OutDrawElements, int32 LayerId, const FWidgetStyle& InWidgetStyle, bool bParentEnabled ) const
 {
@@ -239,7 +240,7 @@ int32 SInvalidationPanel::OnPaint( const FPaintArgs& Args, const FGeometry& Allo
 		if (bHittestCleared || bRequiresRecache)
 		{
 			// @todo: Overly aggressive?
-			MutableThis->InvalidateRootChildOrder();
+			MutableThis->InvalidateRootLayout(this);
 		}
 
 		// The root widget is our child.  We are not the root because we could be in a parent invalidation panel.  If we are nested in another invalidation panel, our OnPaint was called by that panel
@@ -283,14 +284,8 @@ bool SInvalidationPanel::CustomPrepass(float LayoutScaleMultiplier)
 
 	if (GetCanCache())
 	{
-		// The InvalidationRoot that own this retainer will call the ProcessInvalidation.
-		//ProcessInvalidation will only be called when the GlobalInvalidation is off and the Retainer is not inside another InvalidationRoot.
-		if (!GetProxyHandle().HasValidInvalidationRootOwnership(this))
-		{
-			ProcessInvalidation();
-		}
-
-		if (NeedsPrepass())
+		ProcessInvalidation();
+		if (NeedsSlowPath())
 		{
 			FChildren* Children = SCompoundWidget::GetChildren();
 			Prepass_ChildLoop(LayoutScaleMultiplier, Children);
@@ -315,7 +310,7 @@ const FSlateInvalidationRoot* SInvalidationPanel::Advanced_AsInvalidationRoot()
 
 TSharedRef<SWidget> SInvalidationPanel::GetRootWidget()
 {
-	return GetCanCache() ? SCompoundWidget::GetChildren()->GetChildAt(0) : EmptyChildSlot.GetChildAt(0);
+	return GetCanCache() ? SCompoundWidget::GetChildren()->GetChildAt(0) : SNullWidget::NullWidget;
 }
 
 int32 SInvalidationPanel::PaintSlowPath(const FSlateInvalidationContext& Context)
diff --git a/Engine/Source/Runtime/Slate/Public/Widgets/SInvalidationPanel.h b/Engine/Source/Runtime/Slate/Public/Widgets/SInvalidationPanel.h
--- a/Engine/Source/Runtime/Slate/Public/Widgets/SInvalidationPanel.h
+++ b/Engine/Source/Runtime/Slate/Public/Widgets/SInvalidationPanel.h
@@ -44,16 +44,19 @@ public:
 
 	void SetCanCache(bool InCanCache);
 
-	// SWidget overrides
+	//~ SWidget overrides
 	virtual int32 OnPaint( const FPaintArgs& Args, const FGeometry& AllottedGeometry, const FSlateRect& MyCullingRect, FSlateWindowElementList& OutDrawElements, int32 LayerId, const FWidgetStyle& InWidgetStyle, bool bParentEnabled ) const override;
 	virtual FChildren* GetChildren() override;
-	virtual FChildren* GetAllChildren() override;
-	// End SWidget
+#if WITH_SLATE_DEBUGGING
+	virtual FChildren* Debug_GetChildrenForReflector() override;
+#endif
+	//~ End SWidget
 
 	void SetContent(const TSharedRef< SWidget >& InContent);
 
 protected:
 	virtual bool CustomPrepass(float LayoutScaleMultiplier) override;
+	virtual bool ComputeVolatility() const { return GetCanCache() || SCompoundWidget::ComputeVolatility(); }
 	virtual bool Advanced_IsInvalidationRoot() const override;
 	virtual const FSlateInvalidationRoot* Advanced_AsInvalidationRoot() const override;
 	virtual TSharedRef<SWidget> GetRootWidget() override;
diff --git a/Engine/Source/Runtime/SlateCore/Private/FastUpdate/SlateInvalidationRoot.cpp b/Engine/Source/Runtime/SlateCore/Private/FastUpdate/SlateInvalidationRoot.cpp
--- a/Engine/Source/Runtime/SlateCore/Private/FastUpdate/SlateInvalidationRoot.cpp
+++ b/Engine/Source/Runtime/SlateCore/Private/FastUpdate/SlateInvalidationRoot.cpp
@@ -49,6 +49,14 @@ static FAutoConsoleVariableRef CVarSlateInvalidationRootDumpPreInvalidationList(
 );
 void LogPreInvalidationItem(const FSlateInvalidationWidgetList& FastWidgetPathList, FSlateInvalidationWidgetIndex WidgetIndex);
 
+bool GSlateInvalidationRootDumpPrepassInvalidationList = false;
+static FAutoConsoleVariableRef CVarSlateInvalidationRootDumpPrepassInvalidationList(
+	TEXT("Slate.InvalidationRoot.DumpPrepassInvalidationList"),
+	GSlateInvalidationRootDumpPrepassInvalidationList,
+	TEXT("Each frame, log the widgets that are processed in the prepass update phase.")
+);
+void LogPrepassInvalidationItem(const FSlateInvalidationWidgetList& FastWidgetPathList, FSlateInvalidationWidgetIndex WidgetIndex);
+
 bool GSlateInvalidationRootDumpPostInvalidationList = false;
 static FAutoConsoleVariableRef CVarSlateInvalidationRootDumpPostInvalidationList(
 	TEXT("Slate.InvalidationRoot.DumpPostInvalidationList"),
@@ -168,12 +176,11 @@ FSlateInvalidationRoot::FSlateInvalidationRoot()
 	: CachedElementData(new FSlateCachedElementData)
 	, InvalidationRootWidget(nullptr)
 	, RootHittestGrid(nullptr)
-	, FastPathGenerationNumber(INDEX_NONE)
 	, CachedMaxLayerId(0)
-	, bChildOrderInvalidated(false)
 	, bNeedsSlowPath(true)
 	, bNeedScreenPositionShift(false)
 	, bProcessingPreUpdate(false)
+	, bProcessingPrepassUpdate(false)
 	, bProcessingPostUpdate(false)
 	, bBuildingWidgetList(false)
 	, bProcessingChildOrderInvalidation(false)
@@ -187,6 +194,7 @@ FSlateInvalidationRoot::FSlateInvalidationRoot()
 	const FSlateInvalidationWidgetList::FArguments Arg = { GSlateInvalidationWidgetListMaxArrayElements, GSlateInvalidationWidgetListNumberElementLeftBeforeSplitting };
 	FastWidgetPathList = MakeUnique<FSlateInvalidationWidgetList>(InvalidationRootHandle, Arg);
 	WidgetsNeedingPreUpdate = MakeUnique<FSlateInvalidationWidgetPreHeap>(*FastWidgetPathList);
+	WidgetsNeedingPrepassUpdate = MakeUnique<FSlateInvalidationWidgetPrepassHeap>(*FastWidgetPathList);
 	WidgetsNeedingPostUpdate = MakeUnique<FSlateInvalidationWidgetPostHeap>(*FastWidgetPathList);
 
 #if WITH_SLATE_DEBUGGING
@@ -233,9 +241,6 @@ void FSlateInvalidationRoot::InvalidateRoot(const SWidget* Investigator)
 
 void FSlateInvalidationRoot::InvalidateRootChildOrder(const SWidget* Investigator)
 {
-	// Update the generation number. This will effectively invalidate all proxy handles
-	++FastPathGenerationNumber;
-
 	// Invalidate all proxy handles
 	FastWidgetPathList->Reset();
 	InvalidationRootWidget->Invalidate(EInvalidateWidgetReason::Prepass);
@@ -250,7 +255,7 @@ void FSlateInvalidationRoot::InvalidateRootChildOrder(const SWidget* Investigato
 void FSlateInvalidationRoot::InvalidateRootLayout(const SWidget* Investigator)
 {
 	InvalidationRootWidget->Invalidate(EInvalidateWidgetReason::Prepass);
-	bNeedsSlowPath = true;
+	bNeedsSlowPath = true; // with the loop before it should only do one slateprepass
 
 #if WITH_SLATE_DEBUGGING
 	FSlateDebugging::BroadcastInvalidationRootInvalidate(InvalidationRootWidget, Investigator, ESlateDebuggingInvalidateRootReason::Root);
@@ -265,27 +270,14 @@ void FSlateInvalidationRoot::InvalidateWidget(FWidgetProxy& Proxy, EInvalidateWi
 	if (!bNeedsSlowPath)
 	{
 		Proxy.CurrentInvalidateReason |= InvalidateReason;
-		if (EnumHasAnyFlags(InvalidateReason, EInvalidateWidgetReason::ChildOrder))
+		if (Slate::EInvalidateWidgetReason_HasPreUpdateFlag(InvalidateReason))
 		{
 			WidgetsNeedingPreUpdate->HeapPushUnique(Proxy);
-
-			if (!bChildOrderInvalidated)
-			{
-				bChildOrderInvalidated = true;
-				if (!InvalidationRootWidget->Advanced_IsWindow())
-				{
-					InvalidationRootWidget->MarkPrepassAsDirty();
-				}
-
-				if (!GSlateEnableGlobalInvalidation && !InvalidationRootWidget->Advanced_IsWindow())
-				{
-					InvalidationRootWidget->Invalidate(EInvalidateWidgetReason::Layout);
-				}
-			}
 		}
-		else if (Slate::EInvalidateWidgetReason_HasPreUpdateFlag(InvalidateReason))
+
+		if (!bProcessingPrepassUpdate && EnumHasAnyFlags(InvalidateReason, EInvalidateWidgetReason::Prepass))
 		{
-			WidgetsNeedingPreUpdate->HeapPushUnique(Proxy);
+			WidgetsNeedingPrepassUpdate->PushBackUnique(Proxy);
 		}
 
 		if (Slate::EInvalidateWidgetReason_HasPostUpdateFlag(InvalidateReason))
@@ -358,7 +350,6 @@ FSlateInvalidationResult FSlateInvalidationRoot::PaintInvalidationRoot(const FSl
 
 		GSlateIsOnFastUpdatePath = false;
 		bNeedsSlowPath = false;
-		bChildOrderInvalidated = false;
 
 		{
 			if (Context.bAllowFastPathUpdate)
@@ -434,24 +425,38 @@ bool FSlateInvalidationRoot::PaintFastPath(const FSlateInvalidationContext& Cont
 #endif
 
 		TGuardValue<bool> OnFastPathGuard(GSlateIsOnFastUpdatePath, true);
+		FSlateInvalidationWidgetList::FIndexRange PreviousPaintedWidgetRange;
 
 		// The update list is put in reverse order by ProcessInvalidation
 		for (int32 ListIndex = FinalUpdateList.Num() - 1; ListIndex >= 0; --ListIndex)
 		{
 			const FSlateInvalidationWidgetIndex MyIndex = FinalUpdateList[ListIndex];
+			if (PreviousPaintedWidgetRange.IsValid())
+			{
+				// It's already been processed by the previous draw
+				if (PreviousPaintedWidgetRange.Include(FSlateInvalidationWidgetSortOrder{*FastWidgetPathList , MyIndex}))
+				{
+					continue;
+				}
+			}
+
 			FWidgetProxy& WidgetProxy = (*FastWidgetPathList)[MyIndex];
+			SWidget* WidgetPtr = WidgetProxy.GetWidget();
 
 			// Check visibility, it may have been in the update list but a parent who was also in the update list already updated it.
-			SWidget* WidgetPtr = WidgetProxy.GetWidget();
-			if (!WidgetProxy.bUpdatedSinceLastInvalidate && WidgetProxy.Visibility.IsVisible() && WidgetPtr)
+			if (WidgetProxy.Visibility.IsVisible() && WidgetPtr)
 			{
-				bWidgetsNeededRepaint = bWidgetsNeededRepaint || WidgetPtr->HasAnyUpdateFlags(EWidgetUpdateFlags::NeedsRepaint | EWidgetUpdateFlags::NeedsVolatilePaint);
+				const bool bNeedPaint = WidgetPtr->HasAnyUpdateFlags(EWidgetUpdateFlags::NeedsRepaint | EWidgetUpdateFlags::NeedsVolatilePaint);
+				bWidgetsNeededRepaint = bWidgetsNeededRepaint || bNeedPaint;
+
+				if (bNeedPaint)
+				{
+					PreviousPaintedWidgetRange = FSlateInvalidationWidgetList::FIndexRange(*FastWidgetPathList, WidgetProxy.Index, WidgetProxy.LeafMostChildIndex);
+				}
 
 				const int32 NewLayerId = WidgetProxy.Update(*Context.PaintArgs, *Context.WindowElementList);
 				CachedMaxLayerId = FMath::Max(NewLayerId, CachedMaxLayerId);
 
-				WidgetProxy.MarkProxyUpdatedThisFrame(*WidgetsNeedingPostUpdate);
-
 				if (bNeedsSlowPath)
 				{
 					break;
@@ -490,6 +495,7 @@ void FSlateInvalidationRoot::BuildFastPathWidgetList(TSharedRef<SWidget> RootWid
 
 	// We do not care if update are requested. We need to redo all the data.
 	WidgetsNeedingPreUpdate->Reset(false);
+	WidgetsNeedingPrepassUpdate->Reset(false);
 	WidgetsNeedingPostUpdate->Reset(false);
 	FinalUpdateList.Reset();
 
@@ -519,21 +525,11 @@ void FSlateInvalidationRoot::ProcessPreUpdate()
 
 		// Add the root to the update list (to prepass and paint it)
 		check(RootWidget->GetProxyHandle().IsValid(&RootWidget.Get()));
+		WidgetsNeedingPostUpdate->Reset(true); // we can clear the post list, because all widgets will be updated
 		RootWidget->Invalidate(EInvalidateWidgetReason::Prepass);
 	}
 	else
 	{
-		// Put Widget waiting for update back in WidgetsNeedingPostUpdate to ensure index order
-		//and just in case Prepass need to be reexecuted.
-		{
-			for (FSlateInvalidationWidgetIndex WidgetIndex : FinalUpdateList)
-			{
-				FSlateInvalidationWidgetList::InvalidationWidgetType& InvalidationWidget = (*FastWidgetPathList)[WidgetIndex];
-				WidgetsNeedingPostUpdate->PushBackUnique(InvalidationWidget);
-			}
-			FinalUpdateList.Reset(WidgetsNeedingPostUpdate->Num());
-		}
-
 		{
 #if WITH_SLATE_DEBUGGING
 			if (GSlateInvalidationRootDumpPreInvalidationList)
@@ -551,16 +547,19 @@ void FSlateInvalidationRoot::ProcessPreUpdate()
 				FChildOriderInvalidationCallbackImpl(
 					const FSlateInvalidationWidgetList& InWidgetList
 					, FSlateInvalidationWidgetPreHeap& InPreUpdate
+					, FSlateInvalidationWidgetPrepassHeap& InPrepassUpdate
 					, FSlateInvalidationWidgetPostHeap& InPostUpdate
 					, FSlateInvalidationWidgetList::FWidgetAttributeIterator& InAttributeItt)
 					: WidgetList(InWidgetList)
 					, PreUpdate(InPreUpdate)
+					, PrepassUpdate(InPrepassUpdate)
 					, PostUpdate(InPostUpdate)
 					, AttributeItt(InAttributeItt)
 				{}
 				virtual ~FChildOriderInvalidationCallbackImpl() = default;
 				const FSlateInvalidationWidgetList& WidgetList;
 				FSlateInvalidationWidgetPreHeap& PreUpdate;
+				FSlateInvalidationWidgetPrepassHeap& PrepassUpdate;
 				FSlateInvalidationWidgetPostHeap& PostUpdate;
 				FSlateInvalidationWidgetList::FWidgetAttributeIterator& AttributeItt;
 				TArray<FSlateInvalidationWidgetPreHeap::FElement*> WidgetToResort;
@@ -571,6 +570,7 @@ void FSlateInvalidationRoot::ProcessPreUpdate()
 					//Also, their index will not be valid after this function.
 					PreUpdate.RemoveRange(Range);
 					PostUpdate.RemoveRange(Range);
+					PrepassUpdate.RemoveRange(Range);
 					AttributeItt.PreChildRemove(Range);
 				}
 				using FReIndexOperation = FSlateInvalidationWidgetList::IProcessChildOrderInvalidationCallback::FReIndexOperation;
@@ -588,6 +588,7 @@ void FSlateInvalidationRoot::ProcessPreUpdate()
 					};
 					PreUpdate.ForEachIndexes(ReIndexIfNeeded);
 					PostUpdate.ForEachIndexes(ReIndexIfNeeded);
+					PrepassUpdate.ForEachIndexes(ReIndexIfNeeded);
 					AttributeItt.ReIndexed(Operation);
 				}
 				using FReSortOperation = FSlateInvalidationWidgetList::IProcessChildOrderInvalidationCallback::FReSortOperation;
@@ -604,6 +605,7 @@ void FSlateInvalidationRoot::ProcessPreUpdate()
 					};
 					PreUpdate.ForEachIndexes(ReSortIfNeeded);
 					PostUpdate.ForEachIndexes(ReSortIfNeeded);
+					PrepassUpdate.ForEachIndexes(ReSortIfNeeded);
 				}
 				virtual void ProxiesPostResort()
 				{
@@ -619,7 +621,7 @@ void FSlateInvalidationRoot::ProcessPreUpdate()
 					AttributeItt.ProxiesBuilt(Range);
 				}
 
-			}  ChildOriderInvalidationCallback{ *FastWidgetPathList, *WidgetsNeedingPreUpdate, *WidgetsNeedingPostUpdate, AttributeItt };
+			}  ChildOrderInvalidationCallback{ *FastWidgetPathList, *WidgetsNeedingPreUpdate, *WidgetsNeedingPrepassUpdate, *WidgetsNeedingPostUpdate, AttributeItt };
 
 			while((AttributeItt.IsValid() || WidgetsNeedingPreUpdate->Num() > 0) && !bNeedsSlowPath)
 			{
@@ -716,14 +718,21 @@ void FSlateInvalidationRoot::ProcessPreUpdate()
 #if 0
 							FMemMark Mark(FMemStack::Get());
 							TArray<TTuple<FSlateInvalidationWidgetIndex, FSlateInvalidationWidgetSortOrder, TWeakPtr<SWidget>>, TMemStackAllocator<>> PreviousPreUpdate;
+							TArray<TTuple<FSlateInvalidationWidgetIndex, FSlateInvalidationWidgetSortOrder, TWeakPtr<SWidget>>, TMemStackAllocator<>> PreviousPrepassUpdate;
 							TArray<TTuple<FSlateInvalidationWidgetIndex, FSlateInvalidationWidgetSortOrder, TWeakPtr<SWidget>>, TMemStackAllocator<>> PreviousPostUpdate;
 							PreviousPreUpdate.Reserve(WidgetsNeedingPreUpdate->Num());
+							PreviousPrepassUpdate.Reserve(WidgetsNeedingPrepassUpdate->Num());
 							PreviousPostUpdate.Reserve(WidgetsNeedingPostUpdate->Num());
 							for (const auto& Element : WidgetsNeedingPreUpdate->GetRaw())
 							{
 								ensureAlwaysMsgf(FastWidgetPathList->IsValidIndex(Element.GetWidgetIndex()), TEXT("The element is invalid"));
 								PreviousPreUpdate.Emplace(Element.GetWidgetIndex(), Element.GetWidgetSortOrder(), (*FastWidgetPathList)[Element.GetWidgetIndex()].GetWidgetAsShared());
 							}
+							for (const auto& Element : WidgetsNeedingPrepassUpdate->GetRaw())
+							{
+								ensureAlwaysMsgf(FastWidgetPathList->IsValidIndex(Element.GetWidgetIndex()), TEXT("The element is invalid."));
+								PreviousPrepassUpdate.Emplace(Element.GetWidgetIndex(), Element.GetWidgetSortOrder(), (*FastWidgetPathList)[Element.GetWidgetIndex()].GetWidgetAsShared());
+							}
 							for (const auto& Element : WidgetsNeedingPostUpdate->GetRaw())
 							{
 								ensureAlwaysMsgf(FastWidgetPathList->IsValidIndex(Element.GetWidgetIndex()), TEXT("The element is invalid."));
@@ -732,7 +741,7 @@ void FSlateInvalidationRoot::ProcessPreUpdate()
 #endif
 
 							TGuardValue<bool> ProcessChildOrderInvalidationGuardValue(bProcessingChildOrderInvalidation, true);
-							FastWidgetPathList->ProcessChildOrderInvalidation(InvalidationWidget, ChildOriderInvalidationCallback);
+							FastWidgetPathList->ProcessChildOrderInvalidation(InvalidationWidget, ChildOrderInvalidationCallback);
 
 							// This widget may not be valid anymore (got removed because it doesn't fulfill the requirement anymore ie. NullWidget).
 
@@ -765,10 +774,57 @@ void FSlateInvalidationRoot::ProcessPreUpdate()
 #endif //UE_SLATE_WITH_INVALIDATIONWIDGETLIST_DEBUGGING
 }
 
-bool FSlateInvalidationRoot::ProcessPostUpdate()
+void FSlateInvalidationRoot::ProcessPrepassUpdate()
 {
-	WidgetsNeedingPostUpdate->Heapify();
+	TGuardValue<bool> Tmp(bProcessingPrepassUpdate, true);
 
+#if WITH_SLATE_DEBUGGING
+	if (GSlateInvalidationRootDumpPostInvalidationList)
+	{
+		UE_LOG(LogSlate, Log, TEXT("Dumping Prepass Invalidation List"));
+		UE_LOG(LogSlate, Log, TEXT("-------------------"));
+	}
+#endif
+
+	FSlateInvalidationWidgetList::FIndexRange PreviousInvalidationWidgetRange;
+
+	// It update forward (smallest index to biggest )
+	while (WidgetsNeedingPrepassUpdate->Num())
+	{
+		const FSlateInvalidationWidgetPrepassHeap::FElement WidgetElement = WidgetsNeedingPrepassUpdate->HeapPop();
+		if (PreviousInvalidationWidgetRange.IsValid())
+		{
+			// It's already been processed by the previous slate prepass
+			if (PreviousInvalidationWidgetRange.Include(WidgetElement.GetWidgetSortOrder()))
+			{
+				continue;
+			}
+		}
+		FWidgetProxy& WidgetProxy = (*FastWidgetPathList)[WidgetElement.GetWidgetIndex()];
+		PreviousInvalidationWidgetRange = FSlateInvalidationWidgetList::FIndexRange(*FastWidgetPathList, WidgetProxy.Index, WidgetProxy.LeafMostChildIndex);
+
+		// Widget could be null if it was removed and we are on the slow path
+		if (SWidget* WidgetPtr = WidgetProxy.GetWidget())
+		{
+#if WITH_SLATE_DEBUGGING
+			if (GSlateInvalidationRootDumpPrepassInvalidationList)
+			{
+				LogPrepassInvalidationItem(*FastWidgetPathList, WidgetElement.GetWidgetIndex());
+			}
+#endif
+
+			if (!WidgetProxy.Visibility.IsCollapsed() && WidgetPtr->HasAnyUpdateFlags(EWidgetUpdateFlags::NeedsVolatilePrepass))
+			{
+				WidgetPtr->MarkPrepassAsDirty();
+			}
+			WidgetProxy.ProcessLayoutInvalidation(*WidgetsNeedingPostUpdate, *FastWidgetPathList, *this);
+		}
+	}
+	WidgetsNeedingPrepassUpdate->Reset(true);
+}
+
+bool FSlateInvalidationRoot::ProcessPostUpdate()
+{
 #if UE_SLATE_WITH_INVALIDATIONWIDGETLIST_DEBUGGING
 	if (GSlateInvalidationRootVerifyWidgetsUpdateList)
 	{
@@ -793,10 +849,6 @@ bool FSlateInvalidationRoot::ProcessPostUpdate()
 		const FSlateInvalidationWidgetIndex WidgetIndex = WidgetsNeedingPostUpdate->HeapPop();
 		FWidgetProxy& WidgetProxy = (*FastWidgetPathList)[WidgetIndex];
 
-		// Reset each widgets paint state
-		// Must be done before actual painting because children can repaint 
-		WidgetProxy.bUpdatedSinceLastInvalidate = false;
-
 		// Widget could be null if it was removed and we are on the slow path
 		if (SWidget* WidgetPtr = WidgetProxy.GetWidget())
 		{
@@ -807,22 +859,8 @@ bool FSlateInvalidationRoot::ProcessPostUpdate()
 			}
 #endif
 
-			if (WidgetProxy.Visibility.IsVisible())
-			{
-				const bool bIsInvalidationRoot = WidgetPtr->Advanced_IsInvalidationRoot();
-				if (bIsInvalidationRoot && WidgetPtr != InvalidationRootWidget)
-				{
-					FSlateInvalidationRoot* InvalidationRoot = const_cast<FSlateInvalidationRoot*>(WidgetPtr->Advanced_AsInvalidationRoot());
-					check(InvalidationRoot);
-					// Prevent reentering call
-					FSlateInvalidationWidgetPostHeap::FScopeWidgetCannotBeAdded Guard{ *WidgetsNeedingPostUpdate, WidgetProxy };
-					InvalidationRoot->ProcessInvalidation();
-				}
-			}
-
 			bWidgetsNeedRepaint |= WidgetProxy.ProcessPostInvalidation(*WidgetsNeedingPostUpdate, *FastWidgetPathList, *this);
-
-			if (WidgetPtr->HasAnyUpdateFlags(EWidgetUpdateFlags::AnyUpdate))
+			if (WidgetPtr->HasAnyUpdateFlags(EWidgetUpdateFlags::AnyUpdate) && WidgetProxy.Visibility.IsVisible())
 			{
 				FinalUpdateList.Add(WidgetIndex);
 			}
@@ -850,6 +888,7 @@ bool FSlateInvalidationRoot::ProcessInvalidation()
 	if (!bNeedsSlowPath)
 	{
 		check(WidgetsNeedingPreUpdate);
+		check(WidgetsNeedingPrepassUpdate);
 		check(WidgetsNeedingPostUpdate);
 
 		ProcessPreUpdate();
@@ -860,11 +899,46 @@ bool FSlateInvalidationRoot::ProcessInvalidation()
 			ensureMsgf(FastWidgetPathList->VerifyProxiesWidget(), TEXT("We failed to verify that every WidgetProxy has a valid SWidget"));
 		}
 #endif //UE_SLATE_WITH_INVALIDATIONWIDGETLIST_DEBUGGING
+	}
+
+	if (!bNeedsSlowPath)
+	{
+		// Put all widgets in the VolatileUpdate list in the WidgetsNeedingPostUpdate
+		WidgetsNeedingPrepassUpdate->Heapify();
+		WidgetsNeedingPostUpdate->Heapify();
+		{
+			for (FSlateInvalidationWidgetList::FWidgetVolatileUpdateIterator Iterator = FastWidgetPathList->CreateWidgetVolatileUpdateIterator(true);
+				Iterator.IsValid();
+				Iterator.Advance())
+			{
+				FSlateInvalidationWidgetList::InvalidationWidgetType& InvalidationWidget = (*FastWidgetPathList)[Iterator.GetCurrentIndex()];
+				WidgetsNeedingPostUpdate->HeapPushUnique(InvalidationWidget);
+				if (InvalidationWidget.bIsVolatilePrepass)
+				{
+					WidgetsNeedingPrepassUpdate->HeapPushUnique(InvalidationWidget);
+				}
+			}
+		}
+	}
 
+	if (!bNeedsSlowPath)
+	{
+		ProcessPrepassUpdate();
+	}
+
+	if (!bNeedsSlowPath)
+	{
+		FinalUpdateList.Reset(WidgetsNeedingPostUpdate->Num());
 		bWidgetsNeedRepaint = ProcessPostUpdate();
 	}
-	else
+	
+	if (bNeedsSlowPath)
 	{
+		WidgetsNeedingPreUpdate->Reset(true);
+		WidgetsNeedingPrepassUpdate->Reset(true);
+		WidgetsNeedingPostUpdate->Reset(true);
+		FinalUpdateList.Reset();
+		CachedElementData->Empty();
 		bWidgetsNeedRepaint = true;
 	}
 
@@ -918,6 +992,7 @@ void FSlateInvalidationRoot::ClearAllFastPathData(bool bClearResourcesImmediatel
 #endif
 
 	WidgetsNeedingPreUpdate->Reset(false);
+	WidgetsNeedingPrepassUpdate->Reset(false);
 	WidgetsNeedingPostUpdate->Reset(false);
 	FastWidgetPathList->Empty();
 	CachedElementData->Empty();
@@ -1002,6 +1077,12 @@ void LogPreInvalidationItem(const FSlateInvalidationWidgetList& FastWidgetPathLi
 	}
 }
 
+void LogPrepassInvalidationItem(const FSlateInvalidationWidgetList& FastWidgetPathList, FSlateInvalidationWidgetIndex WidgetIndex)
+{
+	const FWidgetProxy& Proxy = FastWidgetPathList[WidgetIndex];
+	UE_LOG(LogSlate, Log, TEXT("  Prepass %s"), *FReflectionMetaData::GetWidgetDebugInfo(Proxy.GetWidget()));
+}
+
 void LogPostInvalidationItem(const FSlateInvalidationWidgetList& FastWidgetPathList, FSlateInvalidationWidgetIndex WidgetIndex)
 {
 	const FWidgetProxy& Proxy = FastWidgetPathList[WidgetIndex];
diff --git a/Engine/Source/Runtime/SlateCore/Private/FastUpdate/SlateInvalidationWidgetHeap.h b/Engine/Source/Runtime/SlateCore/Private/FastUpdate/SlateInvalidationWidgetHeap.h
--- a/Engine/Source/Runtime/SlateCore/Private/FastUpdate/SlateInvalidationWidgetHeap.h
+++ b/Engine/Source/Runtime/SlateCore/Private/FastUpdate/SlateInvalidationWidgetHeap.h
@@ -23,8 +23,8 @@ namespace Private
 #endif
 
 	/**
-	* Ordered list of WidgetIndex. The order is based on the WidgetSortIndex.
-	*/
+	 * Ordered list of WidgetIndex. The order is based on the WidgetSortIndex.
+	 */
 	struct FSlateInvalidationWidgetHeapElement
 	{
 	public:
@@ -101,14 +101,14 @@ public:
 	}
 	
 	/** Returns the biggest WidgetIndex from the list. */
-	inline FSlateInvalidationWidgetIndex HeapPeek() const
+	UE_NODISCARD inline FSlateInvalidationWidgetIndex HeapPeek() const
 	{
 		check(Heap.Num() > 0);
 		return Heap.HeapTop().GetWidgetIndex();
 	}
 	
 	/** Returns the biggest WidgetIndex from the list. */
-	inline const FElement& HeapPeekElement() const
+	UE_NODISCARD inline const FElement& HeapPeekElement() const
 	{
 		check(Heap.Num() > 0);
 		return Heap.HeapTop();
@@ -147,13 +147,13 @@ public:
 	}
 
 	/** Returns the number of elements in the list. */
-	inline int32 Num() const
+	UE_NODISCARD inline int32 Num() const
 	{
 		return Heap.Num();
 	}
 
 	/** Does it contains the widget index. */
-	bool Contains_Debug(const FSlateInvalidationWidgetIndex WidgetIndex) const
+	UE_NODISCARD bool Contains_Debug(const FSlateInvalidationWidgetIndex WidgetIndex) const
 	{
 		return Heap.ContainsByPredicate([WidgetIndex](const FElement& Element)
 			{
@@ -162,13 +162,13 @@ public:
 	}
 
 	/** @retuns true if the heap is heapified. */
-	inline bool IsValidHeap_Debug()
+	UE_NODISCARD inline bool IsValidHeap_Debug()
 	{
 		return Algo::IsHeap(Heap, SortPredicate());
 	}
 
 	/** Returns the raw list. */
-	inline const FElementContainer& GetRaw() const
+	UE_NODISCARD inline const FElementContainer& GetRaw() const
 	{
 		return Heap;
 	}
@@ -189,7 +189,7 @@ private:
 #if UE_SLATE_WITH_INVALIDATIONWIDGETHEAP_DEBUGGING
 		if (UE::Slate::Private::GSlateInvalidationWidgetHeapVerifyWidgetContains)
 		{
-			check(Contains_Debug(InvalidationWidget.Index) == InvalidationWidget.bContainedByWidgetPostHeap);
+			check(Contains_Debug(InvalidationWidget.Index) == InvalidationWidget.bContainedByWidgetPreHeap);
 		}
 #endif
 	}
@@ -265,7 +265,7 @@ public:
 	}
 
 	/** Returns and removes the biggest WidgetIndex from the list. */
-	FSlateInvalidationWidgetIndex HeapPop()
+	UE_NODISCARD FSlateInvalidationWidgetIndex HeapPop()
 	{
 		check(bIsHeap == true);
 		FSlateInvalidationWidgetIndex Result = Heap.HeapTop().GetWidgetIndex();
@@ -317,19 +317,19 @@ public:
 	}
 
 	/** Returns the number of elements in the list. */
-	inline int32 Num() const
+	UE_NODISCARD inline int32 Num() const
 	{
 		return Heap.Num();
 	}
 
 	/** Returns the number of elements in the list. */
-	inline bool IsHeap() const
+	UE_NODISCARD inline bool IsHeap() const
 	{
 		return bIsHeap;
 	}
 
 	/** Does it contains the widget index. */
-	bool Contains_Debug(const FSlateInvalidationWidgetIndex WidgetIndex) const
+	UE_NODISCARD bool Contains_Debug(const FSlateInvalidationWidgetIndex WidgetIndex) const
 	{
 		return Heap.ContainsByPredicate([WidgetIndex](const FElement& Element)
 			{
@@ -338,13 +338,13 @@ public:
 	}
 
 	/** @retuns true if the heap is heapified. */
-	inline bool IsValidHeap_Debug()
+	UE_NODISCARD inline bool IsValidHeap_Debug()
 	{
 		return Algo::IsHeap(Heap, SortPredicate());
 	}
 
 	/** Returns the raw list. */
-	inline const FElementContainer& GetRaw() const
+	UE_NODISCARD inline const FElementContainer& GetRaw() const
 	{
 		return Heap;
 	}
@@ -400,3 +400,167 @@ private:
 	FSlateInvalidationWidgetIndex WidgetCannotBeAdded;
 	bool bIsHeap;
 };
+
+/**
+ * Heap of widget, ordered by increasing sort order, of the widgets that need a SlatePrepass.
+ * Since a SlatePrepass of a parent widget will call the SlatePrepass of a child, only the parent will be in the list.
+ */
+class FSlateInvalidationWidgetPrepassHeap
+{
+public:
+	using FElement = UE::Slate::Private::FSlateInvalidationWidgetHeapElement;
+	struct FWidgetOrderGreater
+	{
+		FORCEINLINE bool operator()(const FElement& A, const FElement& B) const
+		{
+			return A.GetWidgetSortOrder() < B.GetWidgetSortOrder();
+		}
+	};
+	using SortPredicate = FWidgetOrderGreater;
+
+	static constexpr int32 NumberOfPreAllocatedElement = 32;
+	using FElementContainer = TArray<FElement, TInlineAllocator<NumberOfPreAllocatedElement>>;
+
+	FSlateInvalidationWidgetPrepassHeap(FSlateInvalidationWidgetList& InOwnerList)
+		: OwnerList(InOwnerList)
+		, bIsHeap(false)
+	{ }
+
+public:
+	/** Insert into the list at the proper order (see binary heap) but only if it's not already contains by the list. */
+	void HeapPushUnique(FSlateInvalidationWidgetList::InvalidationWidgetType& InvalidationWidget)
+	{
+		check(bIsHeap);
+		check(InvalidationWidget.Index != FSlateInvalidationWidgetIndex::Invalid);
+		VerifyContainsFlag(InvalidationWidget);
+
+		if (!InvalidationWidget.bContainedByWidgetPrepassList)
+		{
+			InvalidationWidget.bContainedByWidgetPrepassList = true;
+			Heap.HeapPush(FElement{ InvalidationWidget.Index, FSlateInvalidationWidgetSortOrder{OwnerList, InvalidationWidget.Index} }, SortPredicate());
+		}
+	}
+
+	/** Insert at the end of the list but only if it's not already contains by the list. */
+	void PushBackUnique(FSlateInvalidationWidgetList::InvalidationWidgetType& InvalidationWidget)
+	{
+		check(bIsHeap == false);
+		check(InvalidationWidget.Index != FSlateInvalidationWidgetIndex::Invalid);
+		VerifyContainsFlag(InvalidationWidget);
+
+		if (!InvalidationWidget.bContainedByWidgetPrepassList)
+		{
+			InvalidationWidget.bContainedByWidgetPrepassList = true;
+			Heap.Emplace(InvalidationWidget.Index, FSlateInvalidationWidgetSortOrder{ OwnerList, InvalidationWidget.Index });
+		}
+	}
+
+	/** Returns and removes the biggest WidgetIndex from the list. */
+	UE_NODISCARD FElement HeapPop()
+	{
+		check(bIsHeap == true);
+		FElement Result = Heap.HeapTop();
+		Heap.HeapPopDiscard(SortPredicate(), false);
+		OwnerList[Result.GetWidgetIndex()].bContainedByWidgetPrepassList = false;
+		return Result;
+	}
+
+	/** Remove range */
+	int32 RemoveRange(const FSlateInvalidationWidgetList::FIndexRange& Range)
+	{
+		check(bIsHeap == false);
+		int32 RemoveCount = 0;
+		for (int32 Index = Heap.Num() - 1; Index >= 0; --Index)
+		{
+			if (Range.Include(Heap[Index].GetWidgetSortOrder()))
+			{
+				const FSlateInvalidationWidgetIndex WidgetIndex = Heap[Index].GetWidgetIndex();
+				OwnerList[WidgetIndex].bContainedByWidgetPrepassList = false;
+				Heap.RemoveAtSwap(Index);
+				++RemoveCount;
+			}
+		}
+
+		return RemoveCount;
+	}
+
+	/** Empties the list, but doesn't change memory allocations. */
+	void Reset(bool bResetContained)
+	{
+		if (bResetContained)
+		{
+			for (const FElement& Element : Heap)
+			{
+				OwnerList[Element.GetWidgetIndex()].bContainedByWidgetPrepassList = false;
+			}
+		}
+
+		// Destroy the second allocator, if it exists.
+		Heap.Empty(NumberOfPreAllocatedElement);
+		bIsHeap = false;
+	}
+
+	void Heapify()
+	{
+		check(!bIsHeap);
+		Heap.Heapify(SortPredicate());
+		bIsHeap = true;
+	}
+
+	/** Returns the number of elements in the list. */
+	UE_NODISCARD inline int32 Num() const
+	{
+		return Heap.Num();
+	}
+
+	/** Returns the number of elements in the list. */
+	UE_NODISCARD inline bool IsHeap() const
+	{
+		return bIsHeap;
+	}
+
+	/** Does it contains the widget index. */
+	UE_NODISCARD bool Contains_Debug(const FSlateInvalidationWidgetIndex WidgetIndex) const
+	{
+		return Heap.ContainsByPredicate([WidgetIndex](const FElement& Element)
+			{
+				return Element.GetWidgetIndex() == WidgetIndex;
+			});
+	}
+
+	/** @retuns true if the heap is heapified. */
+	UE_NODISCARD inline bool IsValidHeap_Debug()
+	{
+		return Algo::IsHeap(Heap, SortPredicate());
+	}
+
+	/** Returns the raw list. */
+	UE_NODISCARD inline const FElementContainer& GetRaw() const
+	{
+		return Heap;
+	}
+
+	template<typename Predicate>
+	void ForEachIndexes(Predicate Pred)
+	{
+		for (FElement& Element : Heap)
+		{
+			Pred(Element);
+		}
+	}
+
+private:
+	void VerifyContainsFlag(const FSlateInvalidationWidgetList::InvalidationWidgetType& InvalidationWidget)
+	{
+#if UE_SLATE_WITH_INVALIDATIONWIDGETHEAP_DEBUGGING
+		if (UE::Slate::Private::GSlateInvalidationWidgetHeapVerifyWidgetContains)
+		{
+			check(Contains_Debug(InvalidationWidget.Index) == InvalidationWidget.bContainedByWidgetPrepassList);
+		}
+#endif
+	}
+
+	FElementContainer Heap;
+	FSlateInvalidationWidgetList& OwnerList;
+	bool bIsHeap;
+};
diff --git a/Engine/Source/Runtime/SlateCore/Private/FastUpdate/SlateInvalidationWidgetList.cpp b/Engine/Source/Runtime/SlateCore/Private/FastUpdate/SlateInvalidationWidgetList.cpp
--- a/Engine/Source/Runtime/SlateCore/Private/FastUpdate/SlateInvalidationWidgetList.cpp
+++ b/Engine/Source/Runtime/SlateCore/Private/FastUpdate/SlateInvalidationWidgetList.cpp
@@ -256,33 +256,139 @@ void FSlateInvalidationWidgetList::FWidgetAttributeIterator::Clear()
 
 
 /** */
-void FSlateInvalidationWidgetList::FArrayNode::RemoveElementIndexBiggerOrEqualThan(IndexType ElementIndex)
+FSlateInvalidationWidgetList::FWidgetVolatileUpdateIterator::FWidgetVolatileUpdateIterator(const FSlateInvalidationWidgetList& InWidgetList, bool bInSkipCollapsed)
+	: WidgetList(InWidgetList)
+	, CurrentWidgetIndex(FSlateInvalidationWidgetIndex::Invalid)
+	, AttributeIndex(0)
+	, bSkipCollapsed(bInSkipCollapsed)
 {
-	const int32 FoundIndex = ElementIndexList_WidgetWithRegisteredSlateAttribute.FindLowerBound(ElementIndex);
-	if (ElementIndexList_WidgetWithRegisteredSlateAttribute.IsValidIndex(FoundIndex))
+	AdvanceArray(WidgetList.FirstArrayIndex);
+	if (bSkipCollapsed)
 	{
-		ElementIndexList_WidgetWithRegisteredSlateAttribute.RemoveAt(FoundIndex, ElementIndexList_WidgetWithRegisteredSlateAttribute.Num() - FoundIndex);
+		SkipToNextExpend();
 	}
 }
 
-void FSlateInvalidationWidgetList::FArrayNode::RemoveElementIndexBetweenOrEqualThan(IndexType StartElementIndex, IndexType EndElementIndex)
+void FSlateInvalidationWidgetList::FWidgetVolatileUpdateIterator::Advance()
+{
+	Internal_Advance();
+	if (bSkipCollapsed)
+	{
+		SkipToNextExpend();
+	}
+}
+
+
+void FSlateInvalidationWidgetList::FWidgetVolatileUpdateIterator::Internal_Advance()
+{
+	++AttributeIndex;
+
+	const FArrayNode& ArrayNode = WidgetList.Data[CurrentWidgetIndex.ArrayIndex];
+	if (AttributeIndex >= ArrayNode.ElementIndexList_VolatileUpdateWidget.Num())
+	{
+		AttributeIndex = 0;
+		CurrentWidgetIndex = FSlateInvalidationWidgetIndex::Invalid;
+		AdvanceArray(ArrayNode.NextArrayIndex);
+	}
+	else
+	{
+		int32 ArrayIndex = CurrentWidgetIndex.ArrayIndex;
+		CurrentWidgetIndex = FSlateInvalidationWidgetIndex{ (IndexType)ArrayIndex, ArrayNode.ElementIndexList_VolatileUpdateWidget[AttributeIndex] };;
+	}
+}
+
+
+void FSlateInvalidationWidgetList::FWidgetVolatileUpdateIterator::SkipToNextExpend()
 {
-	const int32 StartFoundIndex = ElementIndexList_WidgetWithRegisteredSlateAttribute.FindLowerBound(StartElementIndex);
-	if (ElementIndexList_WidgetWithRegisteredSlateAttribute.IsValidIndex(StartFoundIndex))
+	while (CurrentWidgetIndex != FSlateInvalidationWidgetIndex::Invalid && WidgetList[CurrentWidgetIndex].Visibility.IsCollapsed())
 	{
-		const int32 EndFoundIndex = ElementIndexList_WidgetWithRegisteredSlateAttribute.FindUpperBound(EndElementIndex);
-		if (ElementIndexList_WidgetWithRegisteredSlateAttribute.IsValidIndex(EndFoundIndex))
+		const FSlateInvalidationWidgetIndex SeekToIndex = WidgetList.IncrementIndex(WidgetList[CurrentWidgetIndex].LeafMostChildIndex);
+		if (SeekToIndex != FSlateInvalidationWidgetIndex::Invalid)
 		{
-			ElementIndexList_WidgetWithRegisteredSlateAttribute.RemoveAt(StartFoundIndex, EndFoundIndex - StartFoundIndex);
+			Seek(SeekToIndex);
 		}
 		else
 		{
-			ElementIndexList_WidgetWithRegisteredSlateAttribute.RemoveAt(StartFoundIndex, ElementIndexList_WidgetWithRegisteredSlateAttribute.Num() - StartFoundIndex);
+			AttributeIndex = 0;
+			CurrentWidgetIndex = FSlateInvalidationWidgetIndex::Invalid;
+			break;
+		}
+	}
+}
+
+
+void FSlateInvalidationWidgetList::FWidgetVolatileUpdateIterator::Seek(FSlateInvalidationWidgetIndex SeekTo)
+{
+	check(SeekTo != FSlateInvalidationWidgetIndex::Invalid);
+
+	const FArrayNode& ArrayNode = WidgetList.Data[SeekTo.ArrayIndex];
+	AttributeIndex = ArrayNode.ElementIndexList_VolatileUpdateWidget.FindLowerBound(SeekTo.ElementIndex);
+	if (AttributeIndex == INDEX_NONE)
+	{
+		AttributeIndex = 0;
+		CurrentWidgetIndex = FSlateInvalidationWidgetIndex::Invalid;
+		AdvanceArray(ArrayNode.NextArrayIndex);
+	}
+	else
+	{
+		CurrentWidgetIndex = FSlateInvalidationWidgetIndex{ (IndexType)SeekTo.ArrayIndex, ArrayNode.ElementIndexList_VolatileUpdateWidget[AttributeIndex] };
+	}
+}
+
+
+void FSlateInvalidationWidgetList::FWidgetVolatileUpdateIterator::AdvanceArray(int32 ArrayIndex)
+{
+	while (ArrayIndex != INDEX_NONE)
+	{
+		const FArrayNode& NewArrayNode = WidgetList.Data[ArrayIndex];
+		if (NewArrayNode.ElementIndexList_VolatileUpdateWidget.Num() > 0)
+		{
+			CurrentWidgetIndex = FSlateInvalidationWidgetIndex{ (IndexType)ArrayIndex, NewArrayNode.ElementIndexList_VolatileUpdateWidget[0] };
+			break;
 		}
+		ArrayIndex = NewArrayNode.NextArrayIndex;
 	}
 }
 
 
+/** */
+void FSlateInvalidationWidgetList::FArrayNode::RemoveElementIndexBiggerOrEqualThan(IndexType ElementIndex)
+{
+	auto RemoveLowerBound = [ElementIndex](WidgetListType& ListType)
+	{
+		const int32 FoundIndex = ListType.FindLowerBound(ElementIndex);
+		if (ListType.IsValidIndex(FoundIndex))
+		{
+			ListType.RemoveAt(FoundIndex, ListType.Num() - FoundIndex);
+		}
+	};
+	RemoveLowerBound(ElementIndexList_WidgetWithRegisteredSlateAttribute);
+	RemoveLowerBound(ElementIndexList_VolatileUpdateWidget);
+}
+
+void FSlateInvalidationWidgetList::FArrayNode::RemoveElementIndexBetweenOrEqualThan(IndexType StartElementIndex, IndexType EndElementIndex)
+{
+	auto RemoveBound = [StartElementIndex, EndElementIndex](WidgetListType& ListType)
+	{
+		const int32 StartFoundIndex = ListType.FindLowerBound(StartElementIndex);
+		if (ListType.IsValidIndex(StartFoundIndex))
+		{
+			const int32 EndFoundIndex = ListType.FindUpperBound(EndElementIndex);
+			if (ListType.IsValidIndex(EndFoundIndex))
+			{
+				ListType.RemoveAt(StartFoundIndex, EndFoundIndex - StartFoundIndex);
+			}
+			else
+			{
+				ListType.RemoveAt(StartFoundIndex, ListType.Num() - StartFoundIndex);
+			}
+		}
+	};
+	RemoveBound(ElementIndexList_WidgetWithRegisteredSlateAttribute);
+	RemoveBound(ElementIndexList_VolatileUpdateWidget);
+}
+
+
 /** */
 FSlateInvalidationWidgetList::FSlateInvalidationWidgetList(FSlateInvalidationRootHandle InOwner, const FArguments& Args)
 	: Owner(InOwner)
@@ -334,6 +440,11 @@ FSlateInvalidationWidgetIndex FSlateInvalidationWidgetList::Internal_BuildWidget
 		}
 	}
 
+	if (ShouldBeAddedToVolatileUpdateList(Widget))
+	{		// The list is already sorted at this point. Add to the end.
+		Data[LastestIndex].ElementIndexList_VolatileUpdateWidget.AddUnsorted(NewIndex.ElementIndex);
+	}
+
 	const bool bIsInvalidationRoot = Widget->Advanced_IsInvalidationRoot();
 
 	{
@@ -343,6 +454,7 @@ FSlateInvalidationWidgetIndex FSlateInvalidationWidgetList::Internal_BuildWidget
 		WidgetProxy.LeafMostChildIndex = LeafMostChildIndex;
 		WidgetProxy.Visibility = NewVisibility;
 		WidgetProxy.bIsInvalidationRoot = bIsInvalidationRoot;
+		WidgetProxy.bIsVolatilePrepass = Widget->HasAnyUpdateFlags(EWidgetUpdateFlags::NeedsVolatilePrepass);
 #if UE_SLATE_WITH_INVALIDATIONWIDGETLIST_DEBUGGING
 		WidgetProxy.bDebug_AttributeUpdated = bUpdateSlateAttribute;
 #endif
@@ -353,7 +465,7 @@ FSlateInvalidationWidgetIndex FSlateInvalidationWidgetList::Internal_BuildWidget
 #endif
 	{
 		const FSlateInvalidationWidgetSortOrder SortIndex = { *this, NewIndex };
-		Widget->SetFastPathProxyHandle(FWidgetProxyHandle{ Owner, NewIndex, SortIndex, GenerationNumber }, NewVisibility, bParentVolatile);
+		Widget->SetFastPathProxyHandle(FWidgetProxyHandle{ Owner, NewIndex, SortIndex }, NewVisibility, bParentVolatile);
 	}
 
 	const bool bParentOrSelfVolatile = bParentVolatile || Widget->IsVolatile();
@@ -391,11 +503,6 @@ void FSlateInvalidationWidgetList::BuildWidgetList(TSharedRef<SWidget> InRoot)
 
 	Reset();
 	Root = InRoot;
-	FSlateInvalidationRoot*  InvalidationRoot = Owner.GetInvalidationRoot();
-	if (InvalidationRoot)
-	{
-		GenerationNumber = InvalidationRoot->GetFastPathGenerationNumber();
-	}
 	const bool bShouldAdd = ShouldBeAdded(InRoot);
 	if (bShouldAdd)
 	{
@@ -629,6 +736,24 @@ void FSlateInvalidationWidgetList::ProcessAttributeRegistrationInvalidation(cons
 	}
 }
 
+
+void FSlateInvalidationWidgetList::ProcessVolatileUpdateInvalidation(InvalidationWidgetType& InvalidationWidget)
+{
+	SWidget* WidgetPtr = InvalidationWidget.GetWidget();
+	check(WidgetPtr);
+
+	if (ShouldBeAddedToVolatileUpdateList(WidgetPtr))
+	{
+		Data[InvalidationWidget.Index.ArrayIndex].ElementIndexList_VolatileUpdateWidget.InsertUnique(InvalidationWidget.Index.ElementIndex);
+	}
+	else
+	{
+		Data[InvalidationWidget.Index.ArrayIndex].ElementIndexList_VolatileUpdateWidget.RemoveSingle(InvalidationWidget.Index.ElementIndex);
+	}
+	InvalidationWidget.bIsVolatilePrepass = WidgetPtr->HasAnyUpdateFlags(EWidgetUpdateFlags::NeedsVolatilePrepass);
+}
+
+
 namespace InvalidationList
 {
 	template<typename TSlateInvalidationWidgetList, typename Predicate>
@@ -746,7 +871,6 @@ void FSlateInvalidationWidgetList::Reset()
 	Root.Reset();
 	FirstArrayIndex = INDEX_NONE;
 	LastArrayIndex = INDEX_NONE;
-	GenerationNumber = INDEX_NONE;
 }
 
 
@@ -988,6 +1112,7 @@ void FSlateInvalidationWidgetList::RemoveDataNode(IndexType Index)
 	}
 	ArrayNode.ElementList.Empty();
 	ArrayNode.ElementIndexList_WidgetWithRegisteredSlateAttribute.Empty();
+	ArrayNode.ElementIndexList_VolatileUpdateWidget.Empty();
 	Data.RemoveAt(Index);
 
 	// No need to rebuild the order when we remove.
@@ -1074,7 +1199,7 @@ void FSlateInvalidationWidgetList::Internal_RemoveRangeFromSameParent(const FInd
 		|| (bRangeIsInSameElementArray && Data[Range.GetInclusiveMinWidgetIndex().ArrayIndex].StartIndex != Range.GetInclusiveMinWidgetIndex().ElementIndex);
 	if (bShouldCutArray)
 	{
-		_CutArray(Range.GetInclusiveMaxWidgetIndex());
+		Internal_CutArray(Range.GetInclusiveMaxWidgetIndex());
 	}
 
 	// Destroy/Remove the data that is not needed anymore
@@ -1196,7 +1321,7 @@ void FSlateInvalidationWidgetList::Internal_RemoveRangeFromSameParent(const FInd
 
 FSlateInvalidationWidgetList::FCutResult FSlateInvalidationWidgetList::CutArray(const FSlateInvalidationWidgetIndex WhereToCut)
 {
-	FCutResult CutResult = _CutArray(WhereToCut);
+	FCutResult CutResult = Internal_CutArray(WhereToCut);
 
 	// Remove the old data that is now moved to the new array
 	if (CutResult.OldElementIndexStart != INDEX_NONE)
@@ -1218,7 +1343,7 @@ FSlateInvalidationWidgetList::FCutResult FSlateInvalidationWidgetList::CutArray(
 }
 
 
-FSlateInvalidationWidgetList::FCutResult FSlateInvalidationWidgetList::_CutArray(const FSlateInvalidationWidgetIndex WhereToCut)
+FSlateInvalidationWidgetList::FCutResult FSlateInvalidationWidgetList::Internal_CutArray(const FSlateInvalidationWidgetIndex WhereToCut)
 {
 	FCutResult Result;
 	// Widget proxies that got moved/re-indexed by the operation.
@@ -1310,13 +1435,17 @@ FSlateInvalidationWidgetList::FCutResult FSlateInvalidationWidgetList::_CutArray
 			if (SWidget* Widget = NewInvalidationWidget.GetWidget())
 			{
 				FSlateInvalidationWidgetSortOrder SortIndex = { *this, NewWidgetIndex };
-				Widget->SetFastPathProxyHandle(FWidgetProxyHandle{ Owner, NewWidgetIndex, SortIndex, GenerationNumber });
+				Widget->SetFastPathProxyHandle(FWidgetProxyHandle{ Owner, NewWidgetIndex, SortIndex });
 
 				// Check for ElementIndexList
 				if (ShouldBeAddedToAttributeList(Widget))
 				{
 					Data[NewArrayIndex].ElementIndexList_WidgetWithRegisteredSlateAttribute.AddUnsorted(NewElementIndex);
 				}
+				if (ShouldBeAddedToVolatileUpdateList(Widget))
+				{
+					Data[NewArrayIndex].ElementIndexList_VolatileUpdateWidget.AddUnsorted(NewElementIndex);
+				}
 			}
 		}
 
diff --git a/Engine/Source/Runtime/SlateCore/Private/FastUpdate/SlateInvalidationWidgetList.h b/Engine/Source/Runtime/SlateCore/Private/FastUpdate/SlateInvalidationWidgetList.h
--- a/Engine/Source/Runtime/SlateCore/Private/FastUpdate/SlateInvalidationWidgetList.h
+++ b/Engine/Source/Runtime/SlateCore/Private/FastUpdate/SlateInvalidationWidgetList.h
@@ -8,6 +8,7 @@
 #include "FastUpdate/SlateInvalidationRootHandle.h"
 #include "FastUpdate/SlateInvalidationWidgetIndex.h"
 #include "FastUpdate/SlateInvalidationWidgetSortOrder.h"
+#include "FastUpdate/WidgetUpdateFlags.h"
 #include "Widgets/SNullWidget.h"
 #include "Widgets/SWidget.h"
 
@@ -97,9 +98,6 @@ public:
 	/** Get the root the widget list was built with. */
 	const TWeakPtr<SWidget> GetRoot() const { return Root; };
 
-	/** Get the Generation number the list was built on. */
-	int32 GetGenerationNumber() const { return GenerationNumber; }
-
 	/** */
 	struct IProcessChildOrderInvalidationCallback
 	{
@@ -139,12 +137,12 @@ public:
 	bool ProcessChildOrderInvalidation(const InvalidationWidgetType& InvalidationWidget, IProcessChildOrderInvalidationCallback& Callback);
 
 
-	/**
-	 * Test, then adds or removes from the registered attribute list.
-	 * @returns true if the Widget go added to the registered attribute list.
-	 */
+	/** Test, then adds or removes from the registered attribute list. */
 	void ProcessAttributeRegistrationInvalidation(const InvalidationWidgetType& InvalidationWidget);
 
+	/** Test, then adds or removes from the volatile update list. */
+	void ProcessVolatileUpdateInvalidation(InvalidationWidgetType& InvalidationWidget);
+
 	/** Performs an operation on all SWidget in the list. */
 	template<typename Predicate>
 	void ForEachWidget(Predicate Pred)
@@ -237,6 +235,41 @@ public:
 		return FWidgetAttributeIterator(*this);
 	}
 
+public:
+	/** Iterator that goes over all the widgets that needs to be updated every frame. */
+	struct FWidgetVolatileUpdateIterator
+	{
+	private:
+		const FSlateInvalidationWidgetList& WidgetList;
+		FSlateInvalidationWidgetIndex CurrentWidgetIndex;
+		int32 AttributeIndex;
+		bool bSkipCollapsed;
+
+	public:
+		FWidgetVolatileUpdateIterator(const FSlateInvalidationWidgetList& InWidgetList, bool bInSkipCollapsed);
+
+		/** Get the current widget index the iterator is pointing to. */
+		FSlateInvalidationWidgetIndex GetCurrentIndex() const { return CurrentWidgetIndex; }
+
+		/** Advance the iterator to the next valid widget index. */
+		void Advance();
+
+		/** Is the iterator pointing to a valid widget index. */
+		bool IsValid() const { return CurrentWidgetIndex != FSlateInvalidationWidgetIndex::Invalid; }
+
+	private:
+		void Internal_Advance();
+		void SkipToNextExpend();
+		void Seek(FSlateInvalidationWidgetIndex SeekTo);
+		void AdvanceArray(int32 ArrayIndex);
+	};
+
+	FWidgetVolatileUpdateIterator CreateWidgetVolatileUpdateIterator(bool bSkipCollapsed) const
+	{
+		return FWidgetVolatileUpdateIterator(*this, bSkipCollapsed);
+	}
+
+
 public:
 	/** Returns reference to element at give index. */
 	InvalidationWidgetType& operator[](const FSlateInvalidationWidgetIndex Index)
@@ -360,6 +393,18 @@ public:
 	{
 		return Widget->HasRegisteredSlateAttribute() && Widget->IsAttributesUpdatesEnabled();
 	}
+	static bool HasVolatileUpdateFlags(EWidgetUpdateFlags UpdateFlags)
+	{
+		return EnumHasAnyFlags(UpdateFlags, EWidgetUpdateFlags::NeedsTick | EWidgetUpdateFlags::NeedsActiveTimerUpdate | EWidgetUpdateFlags::NeedsVolatilePaint | EWidgetUpdateFlags::NeedsVolatilePrepass);
+	}
+	static bool ShouldBeAddedToVolatileUpdateList(const SWidget* Widget)
+	{
+		return Widget->HasAnyUpdateFlags(EWidgetUpdateFlags::NeedsTick | EWidgetUpdateFlags::NeedsActiveTimerUpdate | EWidgetUpdateFlags::NeedsVolatilePaint | EWidgetUpdateFlags::NeedsVolatilePrepass);
+	}
+	static bool ShouldBeAddedToVolatileUpdateList(const TSharedRef<SWidget>& Widget)
+	{
+		return ShouldBeAddedToVolatileUpdateList(&Widget.Get());
+	}
 
 private:
 	template <typename... ArgsType>
@@ -400,7 +445,7 @@ private:
 	using FFindChildrenElement = TPair<SWidget*, FSlateInvalidationWidgetIndex>;
 	void Internal_FindChildren(FSlateInvalidationWidgetIndex WidgetIndex, TArray<FFindChildrenElement, TMemStackAllocator<>>& Widgets) const;
 	void Internal_RemoveRangeFromSameParent(const FIndexRange Range);
-	FCutResult _CutArray(const FSlateInvalidationWidgetIndex WhereToCut);
+	FCutResult Internal_CutArray(const FSlateInvalidationWidgetIndex WhereToCut);
 
 private:
 	struct FArrayNode
@@ -412,6 +457,7 @@ private:
 		IndexType StartIndex = 0; // The array may start further in the ElementList as a result of a split.
 		ElementListType ElementList;
 		WidgetListType ElementIndexList_WidgetWithRegisteredSlateAttribute;
+		WidgetListType ElementIndexList_VolatileUpdateWidget;
 
 		void RemoveElementIndexBiggerOrEqualThan(IndexType ElementIndex);
 		void RemoveElementIndexBetweenOrEqualThan(IndexType StartElementIndex, IndexType EndElementIndex);
@@ -423,7 +469,6 @@ private:
 	TWeakPtr<SWidget> Root;
 	int32 FirstArrayIndex = INDEX_NONE;
 	int32 LastArrayIndex = INDEX_NONE;
-	int32 GenerationNumber = INDEX_NONE;
 	IProcessChildOrderInvalidationCallback* CurrentInvalidationCallback = nullptr;
 	const FArguments WidgetListConfig;
 };
\ No newline at end of file
diff --git a/Engine/Source/Runtime/SlateCore/Private/FastUpdate/WidgetProxy.cpp b/Engine/Source/Runtime/SlateCore/Private/FastUpdate/WidgetProxy.cpp
--- a/Engine/Source/Runtime/SlateCore/Private/FastUpdate/WidgetProxy.cpp
+++ b/Engine/Source/Runtime/SlateCore/Private/FastUpdate/WidgetProxy.cpp
@@ -25,7 +25,10 @@ FWidgetProxy::FWidgetProxy(TSharedRef<SWidget>& InWidget)
 	, LeafMostChildIndex(FSlateInvalidationWidgetIndex::Invalid)
 	, CurrentInvalidateReason(EInvalidateWidgetReason::None)
 	, Visibility()
-	, PrivateFlags(false)
+	, PrivateFlags(0)
+#if UE_SLATE_WITH_INVALIDATIONWIDGETLIST_DEBUGGING
+	, PrivateDebugFlags(0)
+#endif
 {
 #if UE_SLATE_WITH_WIDGETPROXY_WIDGETTYPE
 	WidgetName = GetWidget()->GetType();
@@ -80,6 +83,62 @@ int32 FWidgetProxy::Update(const FPaintArgs& PaintArgs, FSlateWindowElementList&
 	return OutgoingLayerId;
 }
 
+void FWidgetProxy::ProcessLayoutInvalidation(FSlateInvalidationWidgetPostHeap& UpdateList, FSlateInvalidationWidgetList& FastWidgetPathList, FSlateInvalidationRoot& Root)
+{
+	SWidget* WidgetPtr = GetWidget();
+	SCOPE_CYCLE_SWIDGET(WidgetPtr);
+	// When layout changes compute a new desired size for this widget
+	const FVector2D CurrentDesiredSize = WidgetPtr->GetDesiredSize();
+	FVector2D NewDesiredSize = FVector2D::ZeroVector;
+	if (!Visibility.IsCollapseDirectly())
+	{
+		if (WidgetPtr->NeedsPrepass())
+		{
+			WidgetPtr->SlatePrepass(WidgetPtr->PrepassLayoutScaleMultiplier.Get(1.0f));
+		}
+		else
+		{
+			WidgetPtr->CacheDesiredSize(WidgetPtr->PrepassLayoutScaleMultiplier.Get(1.0f));
+		}
+
+		NewDesiredSize = WidgetPtr->GetDesiredSize();
+	}
+
+	// Note even if volatile we need to recompute desired size. We do not need to invalidate parents though if they are volatile since they will naturally redraw this widget
+	if (!WidgetPtr->IsVolatileIndirectly() && Visibility.IsVisible())
+	{
+		// Set the value directly instead of calling AddUpdateFlags as an optimization
+		WidgetPtr->UpdateFlags |= EWidgetUpdateFlags::NeedsRepaint;
+	}
+
+	// If the desired size changed, invalidate the parent if it is visible
+	if (NewDesiredSize != CurrentDesiredSize || EnumHasAnyFlags(CurrentInvalidateReason, EInvalidateWidgetReason::Visibility | EInvalidateWidgetReason::RenderTransform))
+	{
+		if (ParentIndex != FSlateInvalidationWidgetIndex::Invalid)
+		{
+			FWidgetProxy& ParentProxy = FastWidgetPathList[ParentIndex];
+			if (ParentIndex == FastWidgetPathList.FirstIndex())
+			{
+				// root of the invalidation panel just invalidate the whole thing
+				Root.InvalidateRootLayout(WidgetPtr);
+			}
+			else if (!ParentProxy.Visibility.IsCollapseDirectly())
+			{
+				ParentProxy.CurrentInvalidateReason |= EInvalidateWidgetReason::Layout;
+#if WITH_SLATE_DEBUGGING
+				FSlateDebugging::BroadcastWidgetInvalidate(ParentProxy.GetWidget(), WidgetPtr, EInvalidateWidgetReason::Layout);
+#endif
+				UE_TRACE_SLATE_WIDGET_INVALIDATED(ParentProxy.GetWidget(), WidgetPtr, EInvalidateWidgetReason::Layout);
+				UpdateList.PushBackOrHeapUnique(ParentProxy);
+			}
+		}
+		else if (TSharedPtr<SWidget> ParentWidget = WidgetPtr->GetParentWidget())
+		{
+			ParentWidget->Invalidate(EInvalidateWidgetReason::Layout);
+		}
+	}
+}
+
 bool FWidgetProxy::ProcessPostInvalidation(FSlateInvalidationWidgetPostHeap& UpdateList, FSlateInvalidationWidgetList& FastWidgetPathList, FSlateInvalidationRoot& Root)
 {
 	bool bWidgetNeedsRepaint = false;
@@ -102,60 +161,11 @@ bool FWidgetProxy::ProcessPostInvalidation(FSlateInvalidationWidgetPostHeap& Upd
 		}
 		bWidgetNeedsRepaint = true;
 	}
-	else if (EnumHasAnyFlags(CurrentInvalidateReason, EInvalidateWidgetReason::RenderTransform | EInvalidateWidgetReason::Layout | EInvalidateWidgetReason::Visibility | EInvalidateWidgetReason::ChildOrder))
+	else if (EnumHasAnyFlags(CurrentInvalidateReason, EInvalidateWidgetReason::RenderTransform | EInvalidateWidgetReason::Layout | EInvalidateWidgetReason::Visibility | EInvalidateWidgetReason::ChildOrder)
+		|| WidgetPtr->HasAnyUpdateFlags(EWidgetUpdateFlags::NeedsVolatilePrepass))
 	{
-		SCOPE_CYCLE_SWIDGET(WidgetPtr);
-		// When layout changes compute a new desired size for this widget
-		FVector2D CurrentDesiredSize = WidgetPtr->GetDesiredSize();
-		FVector2D NewDesiredSize = FVector2D::ZeroVector;
-		if (!Visibility.IsCollapsed())
-		{
-			if (WidgetPtr->NeedsPrepass())
-			{
-				WidgetPtr->SlatePrepass(WidgetPtr->PrepassLayoutScaleMultiplier.Get(1.0f));
-			}
-			else
-			{
-				WidgetPtr->CacheDesiredSize(WidgetPtr->PrepassLayoutScaleMultiplier.Get(1.0f));
-			}
-
-			NewDesiredSize = WidgetPtr->GetDesiredSize();
-		}
-
-		// Note even if volatile we need to recompute desired size. We do not need to invalidate parents though if they are volatile since they will naturally redraw this widget
-		if (!WidgetPtr->IsVolatileIndirectly() && Visibility.IsVisible())
-		{
-			// Set the value directly instead of calling AddUpdateFlags as an optimization
-			WidgetPtr->UpdateFlags |= EWidgetUpdateFlags::NeedsRepaint;
-		}
-
-		// If the desired size changed, invalidate the parent if it is visible
-		if (NewDesiredSize != CurrentDesiredSize || EnumHasAnyFlags(CurrentInvalidateReason, EInvalidateWidgetReason::Visibility|EInvalidateWidgetReason::RenderTransform))
-		{
-			if (ParentIndex != FSlateInvalidationWidgetIndex::Invalid)
-			{
-				FWidgetProxy& ParentProxy = FastWidgetPathList[ParentIndex];
-				if (ParentIndex == FastWidgetPathList.FirstIndex())
-				{
-					// root of the invalidation panel just invalidate the whole thing
-					Root.InvalidateRootLayout(WidgetPtr);
-				}
-				else if (ParentProxy.Visibility.IsVisible())
-				{
-					ParentProxy.CurrentInvalidateReason |= EInvalidateWidgetReason::Layout;
-#if WITH_SLATE_DEBUGGING
-					FSlateDebugging::BroadcastWidgetInvalidate(ParentProxy.GetWidget(), WidgetPtr, EInvalidateWidgetReason::Layout);
-#endif
-					UE_TRACE_SLATE_WIDGET_INVALIDATED(ParentProxy.GetWidget(), WidgetPtr, EInvalidateWidgetReason::Layout);
-					UpdateList.HeapPushUnique(ParentProxy);
-				}
-			}
-			else if (TSharedPtr<SWidget> ParentWidget = WidgetPtr->GetParentWidget())
-			{
-				ParentWidget->Invalidate(EInvalidateWidgetReason::Layout);
-			}
-		}
-
+		ProcessLayoutInvalidation(UpdateList, FastWidgetPathList, Root);
+		
 		bWidgetNeedsRepaint = true;
 	}
 	else if (EnumHasAnyFlags(CurrentInvalidateReason, EInvalidateWidgetReason::Paint) && !WidgetPtr->IsVolatileIndirectly())
@@ -174,8 +184,6 @@ bool FWidgetProxy::ProcessPostInvalidation(FSlateInvalidationWidgetPostHeap& Upd
 
 void FWidgetProxy::MarkProxyUpdatedThisFrame(FSlateInvalidationWidgetPostHeap& PostUpdateList)
 {
-	bUpdatedSinceLastInvalidate = true;
-
 	SWidget* WidgetPtr = GetWidget();
 	if (WidgetPtr && Visibility.IsVisible() && WidgetPtr->HasAnyUpdateFlags(EWidgetUpdateFlags::AnyUpdate))
 	{
@@ -239,11 +247,10 @@ int32 FWidgetProxy::Repaint(const FPaintArgs& PaintArgs, FSlateWindowElementList
 	return NewLayerId;
 }
 
-FWidgetProxyHandle::FWidgetProxyHandle(const FSlateInvalidationRootHandle& InInvalidationRoot, FSlateInvalidationWidgetIndex InIndex, FSlateInvalidationWidgetSortOrder InSortIndex, int32 InGenerationNumber)
+FWidgetProxyHandle::FWidgetProxyHandle(const FSlateInvalidationRootHandle& InInvalidationRoot, FSlateInvalidationWidgetIndex InIndex, FSlateInvalidationWidgetSortOrder InSortIndex)
 	: InvalidationRootHandle(InInvalidationRoot)
 	, WidgetIndex(InIndex)
 	, WidgetSortOrder(InSortIndex)
-	, GenerationNumber(InGenerationNumber)
 {
 }
 
@@ -251,7 +258,6 @@ FWidgetProxyHandle::FWidgetProxyHandle(FSlateInvalidationWidgetIndex InIndex)
 	: InvalidationRootHandle()
 	, WidgetIndex(InIndex)
 	, WidgetSortOrder()
-	, GenerationNumber(INDEX_NONE)
 {
 }
 
@@ -259,16 +265,6 @@ bool FWidgetProxyHandle::IsValid(const SWidget* Widget) const
 {
 	FSlateInvalidationRoot* InvalidationRoot = InvalidationRootHandle.GetInvalidationRoot();
 	return InvalidationRoot
-		&& InvalidationRoot->GetFastPathGenerationNumber() == GenerationNumber
-		&& InvalidationRoot->GetFastPathWidgetList().IsValidIndex(WidgetIndex)
-		&& InvalidationRoot->GetFastPathWidgetList()[WidgetIndex].GetWidget() == Widget;
-}
-
-bool FWidgetProxyHandle::HasValidInvalidationRootOwnership(const SWidget* Widget) const
-{
-	FSlateInvalidationRoot* InvalidationRoot = InvalidationRootHandle.GetInvalidationRoot();
-	return InvalidationRoot
-		&& InvalidationRoot->GetFastPathWidgetList().GetGenerationNumber() == GenerationNumber
 		&& InvalidationRoot->GetFastPathWidgetList().IsValidIndex(WidgetIndex)
 		&& InvalidationRoot->GetFastPathWidgetList()[WidgetIndex].GetWidget() == Widget;
 }
@@ -284,33 +280,37 @@ FSlateInvalidationWidgetVisibility FWidgetProxyHandle::GetWidgetVisibility(const
 
 FWidgetProxy& FWidgetProxyHandle::GetProxy()
 {
-	return GetInvalidationRoot()->GetFastPathWidgetList()[WidgetIndex];
+	return GetInvalidationRoot_NoCheck()->GetFastPathWidgetList()[WidgetIndex];
 }
 
 const FWidgetProxy& FWidgetProxyHandle::GetProxy() const
 {
-	return GetInvalidationRoot()->GetFastPathWidgetList()[WidgetIndex];
+	return GetInvalidationRoot_NoCheck()->GetFastPathWidgetList()[WidgetIndex];
 }
 
-void FWidgetProxyHandle::MarkWidgetUpdatedThisFrame()
+void FWidgetProxyHandle::MarkWidgetDirty_NoCheck(EInvalidateWidgetReason InvalidateReason)
 {
-	GetInvalidationRoot()->GetFastPathWidgetList()[WidgetIndex].MarkProxyUpdatedThisFrame(*GetInvalidationRoot()->WidgetsNeedingPostUpdate);
+	FWidgetProxy& Proxy = GetInvalidationRoot_NoCheck()->GetFastPathWidgetList()[WidgetIndex];
+	GetInvalidationRoot_NoCheck()->InvalidateWidget(Proxy, InvalidateReason);
 }
 
-void FWidgetProxyHandle::MarkWidgetDirty(EInvalidateWidgetReason InvalidateReason)
+void FWidgetProxyHandle::MarkWidgetDirty_NoCheck(FWidgetProxy& Proxy)
 {
-	FWidgetProxy& Proxy = GetInvalidationRoot()->GetFastPathWidgetList()[WidgetIndex];
-	GetInvalidationRoot()->InvalidateWidget(Proxy, InvalidateReason);
+	GetInvalidationRoot_NoCheck()->InvalidateWidget(Proxy, Proxy.CurrentInvalidateReason);
 }
 
-void FWidgetProxyHandle::UpdateWidgetFlags(const SWidget* Widget, EWidgetUpdateFlags NewFlags)
+void FWidgetProxyHandle::UpdateWidgetFlags(const SWidget* Widget, EWidgetUpdateFlags PreviousFlags, EWidgetUpdateFlags NewFlags)
 {
-	if (Widget->HasAnyUpdateFlags(EWidgetUpdateFlags::AnyUpdate) && IsValid(Widget))
+	if (PreviousFlags != NewFlags)
 	{
-		FWidgetProxy& Proxy = GetInvalidationRoot()->GetFastPathWidgetList()[WidgetIndex];
-		if (Proxy.Visibility.IsVisible())
+		if ( FSlateInvalidationWidgetList::HasVolatileUpdateFlags(PreviousFlags) != FSlateInvalidationWidgetList::HasVolatileUpdateFlags(NewFlags)
+			|| EnumHasAnyFlags(PreviousFlags, EWidgetUpdateFlags::NeedsVolatilePrepass) != EnumHasAnyFlags(NewFlags, EWidgetUpdateFlags::NeedsVolatilePrepass))
 		{
-			GetInvalidationRoot()->WidgetsNeedingPostUpdate->PushBackOrHeapUnique(Proxy);
+			if (IsValid(Widget))
+			{
+				FWidgetProxy& Proxy = GetProxy();
+				GetInvalidationRoot_NoCheck()->GetFastPathWidgetList().ProcessVolatileUpdateInvalidation(Proxy);
+			}
 		}
 	}
 }
diff --git a/Engine/Source/Runtime/SlateCore/Private/Widgets/SWidget.cpp b/Engine/Source/Runtime/SlateCore/Private/Widgets/SWidget.cpp
--- a/Engine/Source/Runtime/SlateCore/Private/Widgets/SWidget.cpp
+++ b/Engine/Source/Runtime/SlateCore/Private/Widgets/SWidget.cpp
@@ -753,26 +753,35 @@ void SWidget::UpdateWidgetProxy(int32 NewLayerId, FSlateCachedElementsHandle& Ca
 	// Account for the case when the widget gets a new handle for some reason.  This should really never happen
 	if (PersistentState.CachedElementHandle.IsValid() && PersistentState.CachedElementHandle != CacheHandle)
 	{
-		ensureMsgf(!CacheHandle.IsValid(), TEXT("Widget was assigned a new cache handle.  This is not expected to happen."));
+		ensureMsgf(!CacheHandle.IsValid()
+			, TEXT("Widget '%s' was assigned a new cache handle. This is not expected to happen.")
+			, *FReflectionMetaData::GetWidgetPath(this));
 		PersistentState.CachedElementHandle.RemoveFromCache();
 	}
 	PersistentState.CachedElementHandle = CacheHandle;
 	PersistentState.OutgoingLayerId = NewLayerId;
 
+#if WITH_SLATE_DEBUGGING
 	if (FastPathProxyHandle.IsValid(this))
 	{
-		FWidgetProxy& MyProxy = FastPathProxyHandle.GetProxy();
-		MyProxy.Visibility = GetVisibility();
+		const FWidgetProxy& MyProxy = FastPathProxyHandle.GetProxy();
+		ensureMsgf(MyProxy.Visibility.IsVisibleDirectly() == GetVisibility().IsVisible()
+			, TEXT("The visibility of the widget '%s' changed during Paint")
+			, *FReflectionMetaData::GetWidgetPath(this));
 		if ((IsVolatile() && !IsVolatileIndirectly()) || (Advanced_IsInvalidationRoot() && !Advanced_IsWindow()))
 		{
-			AddUpdateFlags(EWidgetUpdateFlags::NeedsVolatilePaint);
+			ensureMsgf(HasAnyUpdateFlags(EWidgetUpdateFlags::NeedsVolatilePaint)
+				, TEXT("The votality of the widget '%s' changed during Paint")
+				, *FReflectionMetaData::GetWidgetPath(this));
 		}
 		else
 		{
-			RemoveUpdateFlags(EWidgetUpdateFlags::NeedsVolatilePaint);
+			ensureMsgf(!HasAnyUpdateFlags(EWidgetUpdateFlags::NeedsVolatilePaint)
+				, TEXT("The votality of the widget '%s' changed during Paint")
+				, *FReflectionMetaData::GetWidgetPath(this));
 		}
-		FastPathProxyHandle.MarkWidgetUpdatedThisFrame();
 	}
+#endif
 }
 
 void SWidget::SetFastPathSortOrder(const FSlateInvalidationWidgetSortOrder InSortOrder)
@@ -831,7 +840,7 @@ void SWidget::UpdateFastPathVisibility(FSlateInvalidationWidgetVisibility Parent
 	if (FastPathProxyHandle.IsValid(this))
 	{	
 		// Try and remove this from the current handles hit test grid.  If we are in a nested invalidation situation the hittest grid may have changed
-		HittestGridToRemoveFrom = FastPathProxyHandle.GetInvalidationRoot()->GetHittestGrid();
+		HittestGridToRemoveFrom = FastPathProxyHandle.GetInvalidationRoot_NoCheck()->GetHittestGrid();
 		FWidgetProxy& Proxy = FastPathProxyHandle.GetProxy();
 		Proxy.Visibility = NewVisibility;
 	}
@@ -859,7 +868,7 @@ void SWidget::UpdateFastPathWidgetRemoved(FHittestGrid* ParentHittestGrid)
 
 	if (FSlateInvalidationRoot* InvalidationRoot = FastPathProxyHandle.GetInvalidationRootHandle().GetInvalidationRoot())
 	{
-		HittestGridToRemoveFrom = FastPathProxyHandle.GetInvalidationRoot()->GetHittestGrid();
+		HittestGridToRemoveFrom = FastPathProxyHandle.GetInvalidationRoot_NoCheck()->GetHittestGrid();
 		InvalidationRoot->OnWidgetDestroyed(this);
 	}
 	FastPathProxyHandle = FWidgetProxyHandle();
@@ -1230,7 +1239,7 @@ void SWidget::Invalidate(EInvalidateWidgetReason InvalidateReason)
 		if (EnumHasAnyFlags(InvalidateReason, EInvalidateWidgetReason::Visibility))
 		{
 			SCOPED_NAMED_EVENT(SWidget_UpdateFastPathVisibility, FColor::Red);
-			UpdateFastPathVisibility(FastPathProxyHandle.GetProxy().Visibility.MimicAsParent(), FastPathProxyHandle.GetInvalidationRoot()->GetHittestGrid());
+			UpdateFastPathVisibility(FastPathProxyHandle.GetProxy().Visibility.MimicAsParent(), FastPathProxyHandle.GetInvalidationRoot_NoCheck()->GetHittestGrid());
 		}
 
 		if (bVolatilityChanged)
@@ -1244,7 +1253,7 @@ void SWidget::Invalidate(EInvalidateWidgetReason InvalidateReason)
 			ensure(!IsVolatile() || IsVolatileIndirectly() || EnumHasAnyFlags(UpdateFlags, EWidgetUpdateFlags::NeedsVolatilePaint));
 		}
 
-		FastPathProxyHandle.MarkWidgetDirty(InvalidateReason);
+		FastPathProxyHandle.MarkWidgetDirty_NoCheck(InvalidateReason);
 	}
 	else
 	{
@@ -1509,6 +1518,9 @@ int32 SWidget::Paint(const FPaintArgs& Args, const FGeometry& AllottedGeometry,
 
 	// Just repainted
 	MutableThis->RemoveUpdateFlags(EWidgetUpdateFlags::NeedsRepaint);
+#if WITH_SLATE_DEBUGGING
+	MutableThis->Debug_UpdateLastPaintFrame();
+#endif
 
 	// Detect children that should have been painted, but were skipped during the paint process.
 	// this will result in geometry being left on screen and not cleared, because it's visible, yet wasn't painted.
diff --git a/Engine/Source/Runtime/SlateCore/Public/FastUpdate/SlateInvalidationRoot.h b/Engine/Source/Runtime/SlateCore/Public/FastUpdate/SlateInvalidationRoot.h
--- a/Engine/Source/Runtime/SlateCore/Public/FastUpdate/SlateInvalidationRoot.h
+++ b/Engine/Source/Runtime/SlateCore/Public/FastUpdate/SlateInvalidationRoot.h
@@ -12,6 +12,7 @@
 struct FSlateCachedElementData;
 class FSlateInvalidationWidgetList;
 class FSlateInvalidationWidgetPreHeap;
+class FSlateInvalidationWidgetPrepassHeap;
 class FSlateInvalidationWidgetPostHeap;
 class FSlateWindowElementList;
 class FWidgetStyle;
@@ -100,8 +101,6 @@ public:
 	const FSlateCachedElementData& GetCachedElements() const { return *CachedElementData; }
 	/** @return the invalidation root as a widget. */
 	const SWidget* GetInvalidationRootWidget() const { return InvalidationRootWidget; }
-	/** @return the the generation number the widget proxy handle should have to be valid. */
-	int32 GetFastPathGenerationNumber() const { return FastPathGenerationNumber; }
 	/** @return the Handle of the InvalidationRoot. */
 	FSlateInvalidationRootHandle GetInvalidationRootHandle() const { return InvalidationRootHandle; }
 	/** @return the list of widgets that are controlled by the InvalidationRoot. */
@@ -151,6 +150,8 @@ private:
 
 	/** Update child order and slate attribute */
 	void ProcessPreUpdate();
+	/** Call slate prepass. */
+	void ProcessPrepassUpdate();
 	/** Update layout, and update the FinalUpdateList */
 	bool ProcessPostUpdate();
 
@@ -158,15 +159,18 @@ private:
 	/** List of all the Widget included by this SlateInvalidationRoot. */
 	TUniquePtr<FSlateInvalidationWidgetList> FastWidgetPathList;
 
+	/** Index of widgets that have the child order invalidated. They affect the widgets index/order. */
+	TUniquePtr<FSlateInvalidationWidgetPreHeap> WidgetsNeedingPreUpdate;
+
 	/**
-	 * Index of widgets that are dirty (child order)
-	 * and need to be updated before the Post list.
+	 * Index of widgets that have the Prepass invalidated.
+	 * They will be updated before the the other layout invalidations to reduce the number of SlatePrepass call.
 	 */
-	TUniquePtr<FSlateInvalidationWidgetPreHeap> WidgetsNeedingPreUpdate;
+	TUniquePtr<FSlateInvalidationWidgetPrepassHeap> WidgetsNeedingPrepassUpdate;
 
 	/**
-	 * Index to widgets that are dirty (volatile, or need some sort of per frame update such as a tick or timer)
-	 * and need to be updated after the Pre list.
+	 * Index of widgets that have invalidation (volatile, or need some sort of per frame update such as a tick or timer).
+	 * They will be added to the FinalUpdateList to be updated.
 	 */
 	TUniquePtr<FSlateInvalidationWidgetPostHeap> WidgetsNeedingPostUpdate;
 
@@ -179,22 +183,14 @@ private:
 
 	FHittestGrid* RootHittestGrid;
 
-	/**
-	 * The purpose of this number is as a unique Id for all widget proxy handles to validate themselves.
-	 * As widgets are added and removed, all their children are indirectly added or remove so it is necessary to invalidate all their handles.
-	 * Bumping the generation number is an efficient way to do this compared to iterating all the handles
-	 * The generation number is always incrementing
-	 */
-	int32 FastPathGenerationNumber;
-
 	int32 CachedMaxLayerId;
 
 	FSlateInvalidationRootHandle InvalidationRootHandle;
 
-	bool bChildOrderInvalidated;
 	bool bNeedsSlowPath;
 	bool bNeedScreenPositionShift;
 	bool bProcessingPreUpdate;
+	bool bProcessingPrepassUpdate;
 	bool bProcessingPostUpdate;
 	bool bBuildingWidgetList;
 	bool bProcessingChildOrderInvalidation;
diff --git a/Engine/Source/Runtime/SlateCore/Public/FastUpdate/WidgetProxy.h b/Engine/Source/Runtime/SlateCore/Public/FastUpdate/WidgetProxy.h
--- a/Engine/Source/Runtime/SlateCore/Public/FastUpdate/WidgetProxy.h
+++ b/Engine/Source/Runtime/SlateCore/Public/FastUpdate/WidgetProxy.h
@@ -57,10 +57,14 @@ public:
 	bool IsVisible() const { return bAncestorsVisible && bVisible; }
 	/** @returns true when all the widget ancestors are visible but the widget itself may not be visible. */
 	bool AreAncestorsVisible() const { return bAncestorsVisible; }
+	/** @returns true when the widget itself is visible (but may have invisible ancestors). */
+	bool IsVisibleDirectly() const { return bVisible; }
 	/** @returns true when at least one of the widget's ancestor is collapse or the widget itself is collapse. */
 	bool IsCollapsed() const { return bAncestorCollapse || bCollapse; }
 	/** @returns true when at least one of the widget's ancestor is collapse but the widget itself may not be collapse. */
 	bool IsCollapseIndirectly() const { return bAncestorCollapse; }
+	/** @returns true the widget itself is collapse. */
+	bool IsCollapseDirectly() const { return bCollapse; }
 
 	void SetVisibility(FSlateInvalidationWidgetVisibility ParentFlags, EVisibility InVisibility)
 	{
@@ -111,6 +115,7 @@ public:
 
 	int32 Update(const FPaintArgs& PaintArgs, FSlateWindowElementList& OutDrawElements);
 
+	void ProcessLayoutInvalidation(FSlateInvalidationWidgetPostHeap& UpdateList, FSlateInvalidationWidgetList& FastPathWidgetList, FSlateInvalidationRoot& Root);
 	bool ProcessPostInvalidation(FSlateInvalidationWidgetPostHeap& UpdateList, FSlateInvalidationWidgetList& FastPathWidgetList, FSlateInvalidationRoot& Root);
 
 	void MarkProxyUpdatedThisFrame(FSlateInvalidationWidgetPostHeap& UpdateList);
@@ -166,28 +171,35 @@ public:
 		struct
 		{
 		public:
-			/** Used to make sure we don't double process a widget that is invalidated.  (a widget can invalidate itself but an ancestor can end up painting that widget first thus rendering the child's own invalidate unnecessary */
-			uint8 bUpdatedSinceLastInvalidate : 1;
-			/** Is the widget already in a pending pre update list.  If it already is in an update list we don't bother adding it again */
+			/** Is the widget in the pre update list. */
 			uint8 bContainedByWidgetPreHeap : 1;
-			/** Is the widget already in a pending post update list.  If it already is in an update list we don't bother adding it again */
+			/** Is the widget in the post update list. */
 			uint8 bContainedByWidgetPostHeap : 1;
-			/** Is the widget an Invalidation Root. Cached value of SWidget::Advanced_IsInvalidationRoot */
-			uint8 bContainedByWidgetVolatileList : 1;
+			/** Is the widget in a pending prepass list. */
+			uint8 bContainedByWidgetPrepassList : 1;
 			/** Is the widget an Invalidation Root. Cached value of SWidget::Advanced_IsInvalidationRoot */
 			uint8 bIsInvalidationRoot : 1;
+			/** Is the widget has volatile prepass flag. */
+			uint8 bIsVolatilePrepass : 1;
+		};
+		uint8 PrivateFlags;
+	};
 
-		public:
 #if UE_SLATE_WITH_INVALIDATIONWIDGETLIST_DEBUGGING
+	union
+	{
+		struct
+		{
+		public:
 			/** Use with "Slate.InvalidationRoot.VerifyWidgetVisibility". Cached the last FastPathVisible value to find widgets that do not call Invalidate properly. */
 			uint8 bDebug_LastFrameVisible : 1;
 			uint8 bDebug_LastFrameVisibleSet : 1;
 			/** Use with "Slate.InvalidationRoot.VerifyWidgetAttribute". */
 			uint8 bDebug_AttributeUpdated : 1;
-#endif
 		};
-		uint8 PrivateFlags;
+		uint8 PrivateDebugFlags;
 	};
+#endif
 };
 
 #if !UE_SLATE_WITH_WIDGETPROXY_WIDGETTYPE
@@ -242,15 +254,10 @@ class FWidgetProxyHandle
 public:
 	FWidgetProxyHandle()
 		: WidgetIndex(FSlateInvalidationWidgetIndex::Invalid)
-		, GenerationNumber(INDEX_NONE)
 	{}
 
 	/** @returns true if it has a valid InvalidationRoot and Index. */
 	SLATECORE_API bool IsValid(const SWidget* Widget) const;
-	/**
-	 * @returns true if it has a valid InvalidationRoot owner
-	 * but it could be consider invalid because the InvalidationRoot needs to be rebuilt. */
-	SLATECORE_API bool HasValidInvalidationRootOwnership(const SWidget* Widget) const;
 
 	FSlateInvalidationRootHandle GetInvalidationRootHandle() const { return InvalidationRootHandle; }
 
@@ -263,21 +270,14 @@ public:
 	const FWidgetProxy& GetProxy() const;
 
 private:
-	bool HasValidIndexAndInvalidationRootHandle() const;
-	FSlateInvalidationRoot* GetInvalidationRoot() const { return InvalidationRootHandle.Advanced_GetInvalidationRootNoCheck(); }
-
-	/**
-	 * Marks the widget as updated this frame
-	 * Note: If the widget still has update flags (e.g it ticks or is volatile or something during update added new flags)
-	 * it will remain in the update list
-	 */
-	void MarkWidgetUpdatedThisFrame();
+	FSlateInvalidationRoot* GetInvalidationRoot_NoCheck() const { return InvalidationRootHandle.Advanced_GetInvalidationRootNoCheck(); }
 	
-	void MarkWidgetDirty(EInvalidateWidgetReason InvalidateReason);
-	SLATECORE_API void UpdateWidgetFlags(const SWidget* Widget, EWidgetUpdateFlags NewFlags);
+	void MarkWidgetDirty_NoCheck(FWidgetProxy& Proxy);
+	void MarkWidgetDirty_NoCheck(EInvalidateWidgetReason InvalidateReason);
+	SLATECORE_API void UpdateWidgetFlags(const SWidget* Widget, EWidgetUpdateFlags Previous, EWidgetUpdateFlags NewFlags);
 
 private:
-	FWidgetProxyHandle(const FSlateInvalidationRootHandle& InInvalidationRoot, FSlateInvalidationWidgetIndex InIndex, FSlateInvalidationWidgetSortOrder InSortIndex, int32 InGenerationNumber);
+	FWidgetProxyHandle(const FSlateInvalidationRootHandle& InInvalidationRoot, FSlateInvalidationWidgetIndex InIndex, FSlateInvalidationWidgetSortOrder InSortIndex);
 	FWidgetProxyHandle(FSlateInvalidationWidgetIndex InIndex);
 
 private:
@@ -287,6 +287,4 @@ private:
 	FSlateInvalidationWidgetIndex WidgetIndex;
 	/** Order of the widget in the fast path list. */
 	FSlateInvalidationWidgetSortOrder WidgetSortOrder;
-	/** This serves as an efficient way to test for validity which does not require invalidating all handles directly. */
-	int32 GenerationNumber;
 };
diff --git a/Engine/Source/Runtime/SlateCore/Public/FastUpdate/WidgetUpdateFlags.h b/Engine/Source/Runtime/SlateCore/Public/FastUpdate/WidgetUpdateFlags.h
--- a/Engine/Source/Runtime/SlateCore/Public/FastUpdate/WidgetUpdateFlags.h
+++ b/Engine/Source/Runtime/SlateCore/Public/FastUpdate/WidgetUpdateFlags.h
@@ -20,6 +20,9 @@ enum class EWidgetUpdateFlags : uint8
 
 	/** Needs repaint because the widget is volatile */
 	NeedsVolatilePaint = 1 << 6,
+
+	/** Needs slate prepass because the widget is volatile */
+	NeedsVolatilePrepass = 1 << 7,
 	
 	AnyUpdate = 0xff,
 };
diff --git a/Engine/Source/Runtime/SlateCore/Public/Widgets/SWidget.h b/Engine/Source/Runtime/SlateCore/Public/Widgets/SWidget.h
--- a/Engine/Source/Runtime/SlateCore/Public/Widgets/SWidget.h
+++ b/Engine/Source/Runtime/SlateCore/Public/Widgets/SWidget.h
@@ -738,6 +738,8 @@ public:
 	}
 
 protected:
+	void SetVolatilePrepass(bool bVolatile) { bVolatile ? AddUpdateFlags(EWidgetUpdateFlags::NeedsVolatilePrepass) : RemoveUpdateFlags(EWidgetUpdateFlags::NeedsVolatilePrepass); }
+
 	virtual bool CustomPrepass(float LayoutScaleMultiplier) { return false; }
 
 	/**
@@ -786,21 +788,16 @@ private:
 
 	void AddUpdateFlags(EWidgetUpdateFlags FlagsToAdd)
 	{
+		EWidgetUpdateFlags Previous = UpdateFlags;
 		UpdateFlags |= FlagsToAdd;
-		FastPathProxyHandle.UpdateWidgetFlags(this, UpdateFlags);
+		FastPathProxyHandle.UpdateWidgetFlags(this, Previous, UpdateFlags);
 	}
 
 	void RemoveUpdateFlags(EWidgetUpdateFlags FlagsToRemove)
 	{
+		EWidgetUpdateFlags Previous = UpdateFlags;
 		UpdateFlags &= (~FlagsToRemove);
-		FastPathProxyHandle.UpdateWidgetFlags(this, UpdateFlags);
-
-#if WITH_SLATE_DEBUGGING
-		if (EnumHasAnyFlags(FlagsToRemove, EWidgetUpdateFlags::NeedsRepaint))
-		{
-			Debug_UpdateLastPaintFrame();
-		}
-#endif
+		FastPathProxyHandle.UpdateWidgetFlags(this, Previous, UpdateFlags);
 	}
 
 	void UpdateWidgetProxy(int32 NewLayerId, FSlateCachedElementsHandle& CacheHandle);
@@ -860,14 +857,28 @@ public:
 	void ArrangeChildren(const FGeometry& AllottedGeometry, FArrangedChildren& ArrangedChildren, bool bUpdateAttributes = false) const;
 
 	/**
-	 * Every widget that has children must implement this method. This allows for iteration over the Widget's
-	 * children regardless of how they are actually stored.
+	 * Returns the useful children (if any) of this widget. Some widget type may hide widget if they are needed by the system.
+	 * Allows for iteration over the Widget's children regardless of how they are actually stored.
+	 * @note Should be renamed to GetVisibleChildren (not ALL children will be returned in all cases).
 	 */
-	 // @todo Slate: Consider renaming to GetVisibleChildren  (not ALL children will be returned in all cases)
 	virtual FChildren* GetChildren() = 0;
 
+	/**
+	 * Returns the children (if any) of this widget that are used by the invalidation system.
+	 * This is used by the FastPath system to generate the correct information for every widget.
+	 * @note Prefer GetChildren. Some widget may hide widget from you.
+	 * @note Should be name GetFastPathChildren or GetInvalidationChildren
+	 */
 	virtual FChildren* GetAllChildren() { return GetChildren(); }
 
+#if WITH_SLATE_DEBUGGING
+	/**
+	 * Returns all Widgets, including widget hidden from the invalidation system.
+	 * This is used by the WidgetReflector.
+	 */
+	virtual FChildren* Debug_GetChildrenForReflector() { return GetAllChildren(); }
+#endif
+
 	/**
 	 * Checks to see if this widget supports keyboard focus.  Override this in derived classes.
 	 *
diff --git a/Engine/Source/Runtime/UMG/Private/Slate/SRetainerWidget.cpp b/Engine/Source/Runtime/UMG/Private/Slate/SRetainerWidget.cpp
--- a/Engine/Source/Runtime/UMG/Private/Slate/SRetainerWidget.cpp
+++ b/Engine/Source/Runtime/UMG/Private/Slate/SRetainerWidget.cpp
@@ -8,6 +8,7 @@
 #include "Engine/TextureRenderTarget2D.h"
 #include "Engine/World.h"
 #include "UMGPrivate.h"
+#include "Widgets/SNullWidget.h"
 
 DECLARE_CYCLE_STAT(TEXT("Retainer Widget Tick"), STAT_SlateRetainerWidgetTick, STATGROUP_Slate);
 DECLARE_CYCLE_STAT(TEXT("Retainer Widget Paint"), STAT_SlateRetainerWidgetPaint, STATGROUP_Slate);
@@ -86,6 +87,10 @@ TFrameValue<int32> SRetainerWidget::Shared_RetainerWorkThisFrame(0);
 
 SRetainerWidget::SRetainerWidget()
 	: EmptyChildSlot(this)
+	, PreviousRenderSize(FIntPoint::NoneValue)
+	, PreviousClipRectSize(FIntPoint::NoneValue)
+	, PreviousColorAndOpacity(FColor::Transparent)
+	, LastIncomingLayerId(0)
 	, VirtualWindow(SNew(SVirtualWindow))
 	, HittestGrid(MakeShared<FHittestGrid>())
 	, RenderingResources(new FRetainerWidgetRenderingResources)
@@ -183,6 +188,7 @@ void SRetainerWidget::Construct(const FArguments& InArgs)
 
 	bEnableRetainedRenderingDesire = true;
 	bEnableRetainedRendering = false;
+	SetVolatilePrepass(bEnableRetainedRendering);
 
 	RefreshRenderingMode();
 	bRenderRequested = true;
@@ -271,6 +277,7 @@ void SRetainerWidget::RefreshRenderingMode()
 	if ( bEnableRetainedRendering != bShouldBeRenderingOffscreen )
 	{
 		bEnableRetainedRendering = bShouldBeRenderingOffscreen;
+		SetVolatilePrepass(bEnableRetainedRendering);
 		InvalidateRootChildOrder();
 	}
 }
@@ -333,10 +340,12 @@ FChildren* SRetainerWidget::GetChildren()
 	}
 }
 
-FChildren* SRetainerWidget::GetAllChildren()
+#if WITH_SLATE_DEBUGGING
+FChildren* SRetainerWidget::Debug_GetChildrenForReflector()
 {
 	return SCompoundWidget::GetChildren();
 }
+#endif
 
 void SRetainerWidget::SetRenderingPhase(int32 InPhase, int32 InPhaseCount)
 {
@@ -350,54 +359,63 @@ void SRetainerWidget::RequestRender()
 	InvalidateRootChildOrder();
 }
 
-bool SRetainerWidget::PaintRetainedContent(const FSlateInvalidationContext& Context, const FGeometry& AllottedGeometry)
-{
-	EPaintRetainedContentResult Result = PaintRetainedContentImpl(Context, AllottedGeometry);
-	return Result == EPaintRetainedContentResult::Painted;
-}
-
-SRetainerWidget::EPaintRetainedContentResult SRetainerWidget::PaintRetainedContentImpl(const FSlateInvalidationContext& Context, const FGeometry& AllottedGeometry)
+SRetainerWidget::EPaintRetainedContentResult SRetainerWidget::PaintRetainedContentImpl(const FSlateInvalidationContext& Context, const FGeometry& AllottedGeometry, int32 LayerId)
 {
 	if (RenderOnPhase)
 	{
 		if (LastTickedFrame != GFrameCounter && (GFrameCounter % PhaseCount) == Phase)
 		{
 			// If doing some phase based invalidation, just redraw everything again
-			RequestRender();
+			bRenderRequested = true;
+			InvalidateRootLayout();
 		}
 	}
 
 	const FPaintGeometry PaintGeometry = AllottedGeometry.ToPaintGeometry();
-	const FVector2D RenderSize = PaintGeometry.GetLocalSize() * PaintGeometry.GetAccumulatedRenderTransform().GetMatrix().GetScale().GetVector();
+	FSlateRenderTransform AccumulatedRenderTransform = PaintGeometry.GetAccumulatedRenderTransform();
+	const FVector2D RenderSize = PaintGeometry.GetLocalSize() * AccumulatedRenderTransform.GetMatrix().GetScale().GetVector();
+	const FIntPoint RoundedRenderSize = RenderSize.IntPoint();
 
 	if (RenderOnInvalidation)
 	{
 		// the invalidation root will take care of whether or not we actually rendered
 		bRenderRequested = true;
 
-		// Aggressively repaint when a base state changes.
-		const FVector2D ClipRectSize = Context.CullingRect.GetSize().RoundToVector();
+		const FIntPoint ClipRectSize = Context.CullingRect.GetSize().IntPoint();
 		const TOptional<FSlateClippingState> ClippingState = Context.WindowElementList->GetClippingState();
-		const FLinearColor ColorAndOpacityTint = Context.WidgetStyle.GetColorAndOpacityTint();
-		if (RenderSize != PreviousRenderSize
+		const FColor ColorAndOpacityTint = Context.WidgetStyle.GetColorAndOpacityTint().ToFColor(false);
+
+
+		// Aggressively re-layout when a base state changes.
+		if (RoundedRenderSize != PreviousRenderSize
 			|| AllottedGeometry != PreviousAllottedGeometry
+			|| AllottedGeometry.GetAccumulatedRenderTransform() != PreviousAllottedGeometry.GetAccumulatedRenderTransform()
 			|| ClipRectSize != PreviousClipRectSize
-			|| ClippingState != PreviousClippingState
-			|| ColorAndOpacityTint != PreviousColorAndOpacity)
+			|| ClippingState != PreviousClippingState)
 		{
-			PreviousRenderSize = RenderSize;
+			PreviousRenderSize = RoundedRenderSize;
 			PreviousAllottedGeometry = AllottedGeometry;
 			PreviousClipRectSize = ClipRectSize;
 			PreviousClippingState = ClippingState;
+
+			InvalidateRootLayout();
+		}
+
+		// Aggressively repaint when a base state changes.
+		if (LayerId != LastIncomingLayerId
+			|| ColorAndOpacityTint != PreviousColorAndOpacity)
+		{
+			LastIncomingLayerId = LayerId;
 			PreviousColorAndOpacity = ColorAndOpacityTint;
 
-			RequestRender();
+			GetRootWidget()->Invalidate(EInvalidateWidgetReason::Paint);
 		}
 	}
-	else if (RenderSize != PreviousRenderSize)
+	else if (RoundedRenderSize != PreviousRenderSize)
 	{
-		RequestRender();
-		PreviousRenderSize = RenderSize;
+		bRenderRequested = true;
+		InvalidateRootLayout();
+		PreviousRenderSize = RoundedRenderSize;
 	}
 
 	if (Shared_MaxRetainerWorkPerFrame > 0)
@@ -538,7 +556,7 @@ int32 SRetainerWidget::OnPaint(const FPaintArgs& Args, const FGeometry& Allotted
 		Context.IncomingLayerId = LayerId;
 		Context.CullingRect = MyCullingRect;
 
-		EPaintRetainedContentResult PaintResult = MutableThis->PaintRetainedContentImpl(Context, AllottedGeometry);
+		EPaintRetainedContentResult PaintResult = MutableThis->PaintRetainedContentImpl(Context, AllottedGeometry, LayerId);
 
 #if WITH_SLATE_DEBUGGING
 		if (PaintResult == EPaintRetainedContentResult::NotPainted
@@ -624,14 +642,8 @@ bool SRetainerWidget::CustomPrepass(float LayoutScaleMultiplier)
 {
 	if (bEnableRetainedRendering)
 	{
-		// The InvalidationRoot that own this retainer will call the ProcessInvalidation.
-		//ProcessInvalidation will only be called when the GlobalInvalidation is off and the Retainer is not inside another InvalidationRoot.
-		if (!GetProxyHandle().HasValidInvalidationRootOwnership(this))
-		{
-			ProcessInvalidation();
-		}
-
-		if (NeedsPrepass())
+		ProcessInvalidation();
+		if (NeedsSlowPath())
 		{
 			FChildren* Children = SCompoundWidget::GetChildren();
 			Prepass_ChildLoop(LayoutScaleMultiplier, Children);
@@ -646,7 +658,7 @@ bool SRetainerWidget::CustomPrepass(float LayoutScaleMultiplier)
 
 TSharedRef<SWidget> SRetainerWidget::GetRootWidget()
 {
-	return bEnableRetainedRendering ? SCompoundWidget::GetChildren()->GetChildAt(0) : EmptyChildSlot.GetChildAt(0);
+	return bEnableRetainedRendering ? SCompoundWidget::GetChildren()->GetChildAt(0) : SNullWidget::NullWidget;
 }
 
 int32 SRetainerWidget::PaintSlowPath(const FSlateInvalidationContext& Context)
diff --git a/Engine/Source/Runtime/UMG/Public/Slate/SRetainerWidget.h b/Engine/Source/Runtime/UMG/Public/Slate/SRetainerWidget.h
--- a/Engine/Source/Runtime/UMG/Public/Slate/SRetainerWidget.h
+++ b/Engine/Source/Runtime/UMG/Public/Slate/SRetainerWidget.h
@@ -79,12 +79,11 @@ public:
 
 	void SetTextureParameter(FName TextureParameter);
 
-	/** SWidget interface */
+	//~ SWidget interface
 	virtual FChildren* GetChildren() override;
-	virtual FChildren* GetAllChildren() override;
-
-	/** FInvalidationRoot interface */
-	virtual bool PaintRetainedContent(const FSlateInvalidationContext& Context, const FGeometry& AllottedGeometry);
+#if WITH_SLATE_DEBUGGING
+	virtual FChildren* Debug_GetChildrenForReflector() override;
+#endif
 
 	void SetWorld(UWorld* World);
 
@@ -92,6 +91,7 @@ protected:
 	/** SCompoundWidget interface */
 	virtual int32 OnPaint(const FPaintArgs& Args, const FGeometry& AllottedGeometry, const FSlateRect& MyCullingRect, FSlateWindowElementList& OutDrawElements, int32 LayerId, const FWidgetStyle& InWidgetStyle, bool bParentEnabled) const override;
 	virtual FVector2D ComputeDesiredSize(float Scale) const override;
+	virtual bool ComputeVolatility() const { return bEnableRetainedRendering || SCompoundWidget::ComputeVolatility(); }
 	virtual bool Advanced_IsInvalidationRoot() const { return bEnableRetainedRendering; }
 	virtual const FSlateInvalidationRoot* Advanced_AsInvalidationRoot() const override { return bEnableRetainedRendering ? this : nullptr; }
 	virtual bool CustomPrepass(float LayoutScaleMultiplier) override;
@@ -108,7 +108,7 @@ protected:
 		TextureSizeTooBig,
 		TextureSizeZero,
 	};
-	EPaintRetainedContentResult PaintRetainedContentImpl(const FSlateInvalidationContext& Context, const FGeometry& AllottedGeometry);
+	EPaintRetainedContentResult PaintRetainedContentImpl(const FSlateInvalidationContext& Context, const FGeometry& AllottedGeometry, int32 LayerId);
 	//~ End FSlateInvalidationRoot interface
 
 	void RefreshRenderingMode();
@@ -127,11 +127,12 @@ private:
 
 	mutable FSlateBrush SurfaceBrush;
 
-	FVector2D PreviousRenderSize;
+	FIntPoint PreviousRenderSize;
 	FGeometry PreviousAllottedGeometry;
-	FVector2D PreviousClipRectSize;
+	FIntPoint PreviousClipRectSize;
 	TOptional<FSlateClippingState> PreviousClippingState;
-	FLinearColor PreviousColorAndOpacity;
+	FColor PreviousColorAndOpacity;
+	int32 LastIncomingLayerId;
 
 	void UpdateWidgetRenderer();
 
