Fixed errors re format string / parameter list for various UE_LOGs:
  - number of supplied parameters less or more than the number of specifiers used in format string
  - invalid format string specifiers used
Also changed text of two UE_LOG error messages to match the other ones.

diff --git a/Engine/Source/Runtime/SlateCore/Private/FastUpdate/SlateInvalidationRoot.cpp b/Engine/Source/Runtime/SlateCore/Private/FastUpdate/SlateInvalidationRoot.cpp
--- a/Engine/Source/Runtime/SlateCore/Private/FastUpdate/SlateInvalidationRoot.cpp
+++ b/Engine/Source/Runtime/SlateCore/Private/FastUpdate/SlateInvalidationRoot.cpp
@@ -1148,7 +1148,7 @@ void VerifyHittest(SWidget* InvalidationRootWidget, FSlateInvalidationWidgetList
 
 	UE_SLATE_LOG_ERROR_IF_FALSE(WidgetList.VerifySortOrder()
 		, GSlateInvalidationRootVerifyHittestGrid
-		, TEXT("The array's sort order for InvalidationRoot '%d' is not respected.")
+		, TEXT("The array's sort order for InvalidationRoot '%s' is not respected.")
 		, *FReflectionMetaData::GetWidgetPath(InvalidationRootWidget));
 
 	TArray<FHittestGrid::FWidgetSortData> WeakHittestGridSortDatas = HittestGrid->GetAllWidgetSortDatas();
@@ -1237,13 +1237,13 @@ void VerifyWidgetVisibility(FSlateInvalidationWidgetList& WidgetList)
 						// Confirm that we have the correct parent
 						UE_SLATE_LOG_ERROR_IF_FALSE(WidgetList.IsValidIndex(InvalidationWidget.ParentIndex)
 							, GSlateInvalidationRootVerifyWidgetVisibility
-							, TEXT("The widget's parent index '%s' is invalid")
+							, TEXT("Widget '%s' Parent index is invalid.")
 							, *FReflectionMetaData::GetWidgetDebugInfo(Widget));
 
 						const FSlateInvalidationWidgetList::InvalidationWidgetType& ParentInvalidationWidget = WidgetList[InvalidationWidget.ParentIndex];
 						UE_SLATE_LOG_ERROR_IF_FALSE(ParentWidget.Get() == ParentInvalidationWidget.GetWidget()
 							, GSlateInvalidationRootVerifyWidgetVisibility
-							, TEXT("The widget's parent '%s' is not '%s'")
+							, TEXT("Widget '%s' Parent is not '%s'.")
 							, *FReflectionMetaData::GetWidgetDebugInfo(Widget)
 							, *FReflectionMetaData::GetWidgetDebugInfo(ParentWidget.Get()));
 
@@ -1254,7 +1254,7 @@ void VerifyWidgetVisibility(FSlateInvalidationWidgetList& WidgetList)
 					{
 						UE_SLATE_LOG_ERROR_IF_FALSE(ParentWidget == nullptr || ParentWidget->Advanced_IsInvalidationRoot()
 							, GSlateInvalidationRootVerifyWidgetVisibility
-							, TEXT("The widget's parent is valid and is not an invalidation root.")
+							, TEXT("Widget '%s' Parent is valid and is not an invalidation root.")
 							, *FReflectionMetaData::GetWidgetDebugInfo(Widget));
 					}
 
diff --git a/Engine/Source/Runtime/SlateCore/Private/Input/HittestGrid.cpp b/Engine/Source/Runtime/SlateCore/Private/Input/HittestGrid.cpp
--- a/Engine/Source/Runtime/SlateCore/Private/Input/HittestGrid.cpp
+++ b/Engine/Source/Runtime/SlateCore/Private/Input/HittestGrid.cpp
@@ -994,7 +994,7 @@ void FHittestGrid::LogGrid() const
 	{
 		const FWidgetData& CurWidgetData = *It;
 		const TSharedPtr<SWidget> CachedWidget = CurWidgetData.GetWidget();
-		UE_LOG(LogHittestDebug, Warning, TEXT("  [%d][%d][%d] => %s @ %s"),
+		UE_LOG(LogHittestDebug, Warning, TEXT("  [%d][%lld] => %s @ %s"),
 			It.GetIndex(),
 			CurWidgetData.PrimarySort,
 			CachedWidget.IsValid() ? *CachedWidget->ToString() : TEXT("Invalid Widget"),
