Slate: Add LexToString to EInvalidateWidgetReason. Update SlateTraceFlags and ConsoleSlateDebuggerInvalidate with the newest EInvalidateWidgetReason. Add ConsoleSlateDebuggerBreak to break the application when a widget get invalidated or painted.

diff --git a/Engine/Plugins/Slate/SlateInsights/Source/SlateInsights/Private/SSlateTraceFlags.cpp b/Engine/Plugins/Slate/SlateInsights/Source/SlateInsights/Private/SSlateTraceFlags.cpp
--- a/Engine/Plugins/Slate/SlateInsights/Source/SlateInsights/Private/SSlateTraceFlags.cpp
+++ b/Engine/Plugins/Slate/SlateInsights/Source/SlateInsights/Private/SSlateTraceFlags.cpp
@@ -77,21 +77,39 @@ namespace Private
 void SSlateTraceWidgetUpdateFlags::Construct(const FArguments& InArgs)
 {
 	UpdateFlagsValue = InArgs._UpdateFlags;
-	SetToolTipText(LOCTEXT("UpdateFlagsTooltip", "U : Tick : The widget was updated/ticked.\nT : Active Timer Update : The widget had an active timer.\nP : Repaint : The widget was dirty and was repainted.\nV : Volatile Paint : The widget was volatile and was repainted."));
+	SetToolTipText(LOCTEXT("UpdateFlagsTooltip",
+		"T : Tick : The widget was updated/ticked.\n"
+		"T : Active Timer Update : The widget had an active timer.\n"
+		"P : Repaint : The widget was dirty and was repainted.\n"
+		"V : Volatile Paint : The widget was volatile and was repainted."
+		"P : Volatile Prepass : The widget was volatile and did a SlatePrepass."));
 }
 
 void SSlateTraceInvalidateWidgetReasonFlags::Construct(const FArguments& InArgs)
 {
 	Reason = InArgs._Reason;
-	SetToolTipText(LOCTEXT("InvalidateWidgetReasonFlagsTooltip", "L : Layout : The widget desired size changed.\nP : Paint : The widget needs repainting but nothing affecting its size.\nU : Volatile : The widget volatility changed.\nC : Child Order : A child was added or removed (this implies layout).\nR : Render Transform : The widget render transform changed.\nV : Visibility : The widget visibility changed (this implies layout)."));
+	SetToolTipText(LOCTEXT("InvalidateWidgetReasonFlagsTooltip", 
+		"L : Layout : The widget desired size changed.\n"
+		"P : Paint : The widget needs repainting but nothing affecting its size.\n"
+		"V : Volatile : The widget volatility changed.\n"
+		"C : Child Order : A child was added or removed (this implies layout).\n"
+		"R : Render Transform : The widget render transform changed.\n"
+		"V : Visibility : The widget visibility changed (this implies layout)."
+		"A : Attribute Registration : Attributes got bound or unbound."
+		"P : Prepass : Re-cache desired size of all of this widget's children recursively."));
 }
 
 int32 SSlateTraceWidgetUpdateFlags::OnPaint(const FPaintArgs& Args, const FGeometry& AllottedGeometry, const FSlateRect& MyCullingRect,
 	FSlateWindowElementList& OutDrawElements, int32 LayerId, const FWidgetStyle& InWidgetStyle, bool bParentEnabled) const
 {
-	const EWidgetUpdateFlags AllUpdateFlags[] = { EWidgetUpdateFlags::NeedsTick, EWidgetUpdateFlags::NeedsActiveTimerUpdate,
-		EWidgetUpdateFlags::NeedsRepaint, EWidgetUpdateFlags::NeedsVolatilePaint };
-	const FString Text = TEXT("UTPV");
+	const EWidgetUpdateFlags AllUpdateFlags[] = {
+		EWidgetUpdateFlags::NeedsTick,
+		EWidgetUpdateFlags::NeedsActiveTimerUpdate,
+		EWidgetUpdateFlags::NeedsRepaint,
+		EWidgetUpdateFlags::NeedsVolatilePaint,
+		EWidgetUpdateFlags::NeedsVolatilePrepass,
+	};
+	const FString Text = TEXT("TTPVP");
 
 	Private::Paint<EWidgetUpdateFlags>(AllUpdateFlags, Text, UpdateFlagsValue, AllottedGeometry, OutDrawElements, LayerId, InWidgetStyle);
 
@@ -102,10 +120,17 @@ int32 SSlateTraceWidgetUpdateFlags::OnPaint(const FPaintArgs& Args, const FGeome
 int32 SSlateTraceInvalidateWidgetReasonFlags::OnPaint(const FPaintArgs& Args, const FGeometry& AllottedGeometry, const FSlateRect& MyCullingRect,
 	FSlateWindowElementList& OutDrawElements, int32 LayerId, const FWidgetStyle& InWidgetStyle, bool bParentEnabled) const
 {
-	const EInvalidateWidgetReason AllReasons[] = { EInvalidateWidgetReason::Layout, EInvalidateWidgetReason::Paint,
-		EInvalidateWidgetReason::Volatility, EInvalidateWidgetReason::ChildOrder,
-		EInvalidateWidgetReason::RenderTransform, EInvalidateWidgetReason::Visibility };
-	const FString Text = TEXT("LPUCRV");
+	const EInvalidateWidgetReason AllReasons[] = {
+		EInvalidateWidgetReason::Layout,
+		EInvalidateWidgetReason::Paint,
+		EInvalidateWidgetReason::Volatility,
+		EInvalidateWidgetReason::ChildOrder,
+		EInvalidateWidgetReason::RenderTransform,
+		EInvalidateWidgetReason::Visibility,
+		EInvalidateWidgetReason::AttributeRegistration,
+		EInvalidateWidgetReason::Prepass,
+	};
+	const FString Text = TEXT("LPVCRVAP");
 
 	Private::Paint<EInvalidateWidgetReason>(AllReasons, Text, Reason, AllottedGeometry, OutDrawElements, LayerId, InWidgetStyle);
 
diff --git a/Engine/Source/Runtime/SlateCore/Private/Debugging/ConsoleSlateDebuggerBreak.cpp b/Engine/Source/Runtime/SlateCore/Private/Debugging/ConsoleSlateDebuggerBreak.cpp
new file mode 100644
--- /dev/null
+++ b/Engine/Source/Runtime/SlateCore/Private/Debugging/ConsoleSlateDebuggerBreak.cpp
@@ -0,0 +1,324 @@
+// Copyright Epic Games, Inc. All Rights Reserved.
+
+#include "Debugging/ConsoleSlateDebuggerBreak.h"
+#include "Debugging/ConsoleSlateDebugger.h"
+
+#if WITH_SLATE_DEBUGGING
+
+#include "Application/SlateApplicationBase.h"
+#include "CoreGlobals.h"
+#include "Debugging/SlateDebugging.h"
+#include "HAL/PlatformMisc.h"
+#include "Misc/App.h"
+#include "Misc/ConfigCacheIni.h"
+#include "Misc/StringBuilder.h"
+#include "Types/ReflectionMetadata.h"
+#include "Widgets/SWidget.h"
+
+#define LOCTEXT_NAMESPACE "ConsoleSlateDebuggerBreak"
+
+
+FConsoleSlateDebuggerBreak::FConsoleSlateDebuggerBreak()
+	: WidgetInvalidationCommand(TEXT("SlateDebugger.Break.OnWidgetInvalidation"),
+		TEXT("Break when the widget get invalidated (must be attached to a debugger).\n")
+		TEXT("Usage: [WidgetPtr=0x1234567]|[WidgetId=12345] [Reason=Paint|Volatility|ChildOrder|RenderTransform|Visibility|AttributeRegistration|Prepass|All|]"),
+		FConsoleCommandWithArgsDelegate::CreateRaw(this, &FConsoleSlateDebuggerBreak::HandleBreakOnWidgetInvalidation))
+	, WidgetBeginPaintCommand(TEXT("SlateDebugger.Break.OnWidgetBeginPaint"),
+		TEXT("Break before the widget get painted (must be attached to a debugger).\n")
+		TEXT("Usage: [WidgetPtr=0x1234567]|[WidgetId=12345]"),
+		FConsoleCommandWithArgsDelegate::CreateRaw(this, &FConsoleSlateDebuggerBreak::HandleBeginWidgetPaint))
+	, WidgetEndPaintCommand(TEXT("SlateDebugger.Break.OnWidgetEndPaint"),
+		TEXT("Break after the widget got painted (must be attached to a debugger).\n")
+		TEXT("Usage: [WidgetPtr=0x1234567]|[WidgetId=12345]"),
+		FConsoleCommandWithArgsDelegate::CreateRaw(this, &FConsoleSlateDebuggerBreak::HandleEndWidgetPaint))
+	, RemoveAllCommand(TEXT("SlateDebugger.Break.RemoveAll"),
+		TEXT("Remove all request to break.\n"),
+		FConsoleCommandDelegate::CreateRaw(this, &FConsoleSlateDebuggerBreak::RemoveAll))
+{
+}
+
+
+FConsoleSlateDebuggerBreak::~FConsoleSlateDebuggerBreak()
+{
+	RemoveAll();
+}
+
+
+void FConsoleSlateDebuggerBreak::AddInvalidation(const SWidget& Widget, EInvalidateWidgetReason Reason)
+{
+	const FConsoleSlateDebuggerUtility::TSWidgetId WidgetId = FConsoleSlateDebuggerUtility::GetId(Widget);
+	AddInvalidation( {WidgetId, Reason} );
+}
+
+
+void FConsoleSlateDebuggerBreak::AddInvalidation(FInvalidationElement Element)
+{
+	if (Element.WidgetId == 0 || Element.Reason == EInvalidateWidgetReason::None)
+	{
+		return;
+	}
+
+	if (!FPlatformMisc::IsDebuggerPresent())
+	{
+		UE_LOG(LogSlateDebugger, Warning, TEXT("The debugger is not attached"));
+	}
+
+
+	if (BreakOnInvalidationElements.Num() == 0)
+	{
+		FSlateDebugging::WidgetInvalidateEvent.AddRaw(this, &FConsoleSlateDebuggerBreak::HandleWidgetInvalidated);
+	}
+
+	if (FInvalidationElement* FoundItem = BreakOnInvalidationElements.FindByPredicate([Element](const FInvalidationElement& Other)
+		{
+			return (Other.WidgetId == Element.WidgetId);
+		}))
+	{
+		*FoundItem = MoveTemp(Element);
+	}
+	else
+	{
+		BreakOnInvalidationElements.Add(MoveTemp(Element));
+	}
+}
+
+
+void FConsoleSlateDebuggerBreak::RemoveInvalidation(const SWidget& Widget)
+{
+	const FConsoleSlateDebuggerUtility::TSWidgetId WidgetId = FConsoleSlateDebuggerUtility::GetId(Widget);
+	BreakOnInvalidationElements.RemoveAllSwap([WidgetId](const FInvalidationElement& Element)
+		{
+			return (Element.WidgetId == WidgetId);
+		});
+	if (BreakOnInvalidationElements.Num() == 0)
+	{
+		FSlateDebugging::WidgetInvalidateEvent.RemoveAll(this);
+	}
+}
+
+
+void FConsoleSlateDebuggerBreak::AddBeginPaint(const SWidget& Widget)
+{
+	AddBeginPaint(FConsoleSlateDebuggerUtility::GetId(Widget));
+}
+
+
+void FConsoleSlateDebuggerBreak::AddBeginPaint(FConsoleSlateDebuggerUtility::TSWidgetId WidgetId)
+{
+	if (BreakOnBeginPaintElements.Num() == 0)
+	{
+		FSlateDebugging::BeginWidgetPaint.AddRaw(this, &FConsoleSlateDebuggerBreak::HandleBeginWidgetPaint);
+	}
+	BreakOnBeginPaintElements.AddUnique(WidgetId);
+}
+
+
+void FConsoleSlateDebuggerBreak::RemoveBeginPaint(const SWidget& Widget)
+{
+	BreakOnBeginPaintElements.RemoveSingleSwap(FConsoleSlateDebuggerUtility::GetId(Widget));
+	if (BreakOnBeginPaintElements.Num() == 0)
+	{
+		FSlateDebugging::BeginWidgetPaint.RemoveAll(this);
+	}
+}
+
+
+void FConsoleSlateDebuggerBreak::AddEndPaint(const SWidget& Widget)
+{
+	AddEndPaint(FConsoleSlateDebuggerUtility::GetId(Widget));
+}
+
+
+void FConsoleSlateDebuggerBreak::AddEndPaint(FConsoleSlateDebuggerUtility::TSWidgetId WidgetId)
+{
+	if (BreakOnEndPaintElements.Num() == 0)
+	{
+		FSlateDebugging::EndWidgetPaint.AddRaw(this, &FConsoleSlateDebuggerBreak::HandleEndWidgetPaint);
+	}
+	BreakOnEndPaintElements.AddUnique(WidgetId);
+}
+
+
+void FConsoleSlateDebuggerBreak::RemoveEndPaint(const SWidget& Widget)
+{
+	BreakOnEndPaintElements.RemoveSingleSwap(FConsoleSlateDebuggerUtility::GetId(Widget));
+	if (BreakOnEndPaintElements.Num() == 0)
+	{
+		FSlateDebugging::EndWidgetPaint.RemoveAll(this);
+	}
+}
+
+
+void FConsoleSlateDebuggerBreak::RemoveAll()
+{
+	FSlateDebugging::WidgetInvalidateEvent.RemoveAll(this);
+	BreakOnInvalidationElements.Empty();
+	FSlateDebugging::BeginWidgetPaint.RemoveAll(this);
+	BreakOnBeginPaintElements.Empty();
+	FSlateDebugging::EndWidgetPaint.RemoveAll(this);
+	BreakOnEndPaintElements.Empty();
+}
+
+
+bool FConsoleSlateDebuggerBreak::ParseWidgetArgs(FConsoleSlateDebuggerBreak::FInvalidationElement& Out, const TArray<FString>& Params, bool bAllowInvalidationReason)
+{
+	const TCHAR* ReasonStr = TEXT("Reason=");
+	const TCHAR* WidgetPtrStr = TEXT("WidgetPtr=");
+	const TCHAR* WidgetIdStr = TEXT("WidgetId=");
+	FStringView ReasonView{ ReasonStr };
+	FStringView WidgetPtrView{ WidgetPtrStr };
+	FStringView WidgetIdView{ WidgetIdStr };
+	for (const FString& Param : Params)
+	{
+		FStringView ParamView{Param};
+		if (ParamView.StartsWith(ReasonView, ESearchCase::IgnoreCase))
+		{
+			if (!bAllowInvalidationReason)
+			{
+				return false;
+			}
+			ParamView.RightChopInline(ReasonView.Len());
+			if (!LexTryParseString(Out.Reason, ParamView.GetData()))
+			{
+				return false;
+			}
+		}
+		else if (ParamView.StartsWith(WidgetPtrView, ESearchCase::IgnoreCase))
+		{
+			ParamView.RightChopInline(WidgetPtrView.Len());
+			FConsoleSlateDebuggerUtility::TSWidgetId WidgetNumber;
+			if (LexTryParseString(WidgetNumber, ParamView.GetData()) && WidgetNumber != 0)
+			{
+				Out.WidgetId = FConsoleSlateDebuggerUtility::GetId((SWidget*)WidgetNumber);
+			}
+			else
+			{
+				return false;
+			}
+		}
+		else if (ParamView.StartsWith(WidgetIdView, ESearchCase::IgnoreCase))
+		{
+			ParamView.RightChopInline(WidgetIdView.Len());
+			FConsoleSlateDebuggerUtility::TSWidgetId WidgetNumber;
+			if (LexTryParseString(WidgetNumber, ParamView.GetData()) && WidgetNumber != 0)
+			{
+				Out.WidgetId = FConsoleSlateDebuggerUtility::GetId((SWidget*)WidgetNumber);
+			}
+			else
+			{
+				return false;
+			}
+		}
+		else
+		{
+			return false;
+		}
+	}
+	return Params.Num() != 0;
+}
+
+void FConsoleSlateDebuggerBreak::HandleBreakOnWidgetInvalidation(const TArray<FString>& Params)
+{
+	const TCHAR* UsageMessage = TEXT("Usage: [WidgetPtr=0x1234567]|[WidgetId=12345] [Reason=Paint|Volatility|ChildOrder|RenderTransform|Visibility|AttributeRegistration|Prepass|Any|]");
+
+	FInvalidationElement InvalidationElement;
+	if (ParseWidgetArgs(InvalidationElement, Params, true))
+	{
+		AddInvalidation(InvalidationElement);
+	}
+	else
+	{
+		UE_LOG(LogSlateDebugger, Log, TEXT("%s"), UsageMessage);
+	}
+}
+
+
+void FConsoleSlateDebuggerBreak::HandleBeginWidgetPaint(const TArray<FString>& Params)
+{
+	const TCHAR* UsageMessage = TEXT("Usage: [WidgetPtr=0x1234567]|[WidgetId=12345]");
+
+	FInvalidationElement InvalidationElement;
+	if (ParseWidgetArgs(InvalidationElement, Params, false))
+	{
+		AddBeginPaint(InvalidationElement.WidgetId);
+	}
+	else
+	{
+		UE_LOG(LogSlateDebugger, Log, TEXT("%s"), UsageMessage);
+	}
+}
+
+
+void FConsoleSlateDebuggerBreak::HandleEndWidgetPaint(const TArray<FString>& Params)
+{
+	const TCHAR* UsageMessage = TEXT("Usage: [WidgetPtr=0x1234567]|[WidgetId=12345]");
+
+	FInvalidationElement InvalidationElement;
+	if (ParseWidgetArgs(InvalidationElement, Params, false))
+	{
+		AddEndPaint(InvalidationElement.WidgetId);
+	}
+	else
+	{
+		UE_LOG(LogSlateDebugger, Log, TEXT("%s"), UsageMessage);
+	}
+}
+
+
+void FConsoleSlateDebuggerBreak::HandleWidgetInvalidated(const FSlateDebuggingInvalidateArgs& Args)
+{
+	if (Args.WidgetInvalidated == nullptr)
+	{
+		return;
+	}
+
+	const FConsoleSlateDebuggerUtility::TSWidgetId WidgetInvalidatedId = FConsoleSlateDebuggerUtility::GetId(Args.WidgetInvalidated);
+	if (FInvalidationElement* FoundItem = BreakOnInvalidationElements.FindByPredicate([WidgetInvalidatedId](const FInvalidationElement& Element)
+		{
+			return (Element.WidgetId == WidgetInvalidatedId);
+		}))
+	{
+		if (EnumHasAnyFlags(Args.InvalidateWidgetReason, FoundItem->Reason))
+		{
+			UE_LOG(LogSlateDebugger, Log, TEXT("Widget '%s' was invalidated."), *FReflectionMetaData::GetWidgetPath(Args.WidgetInvalidated));
+			_DebugBreakAndPromptForRemote();
+		}
+	}
+}
+
+
+void FConsoleSlateDebuggerBreak::HandleBeginWidgetPaint(const SWidget* Widget, const FPaintArgs& /*Args*/, const FGeometry& /*AllottedGeometry*/, const FSlateRect& /*MyCullingRect*/, const FSlateWindowElementList& /*OutDrawElements*/, int32 /*LayerId*/)
+{
+	if (Widget == nullptr)
+	{
+		return;
+	}
+
+	const FConsoleSlateDebuggerUtility::TSWidgetId WidgetId = FConsoleSlateDebuggerUtility::GetId(Widget);
+	if (BreakOnBeginPaintElements.Contains(WidgetId))
+	{
+		UE_LOG(LogSlateDebugger, Log, TEXT("Widget '%s' was invalidated."), *FReflectionMetaData::GetWidgetPath(Widget));
+		_DebugBreakAndPromptForRemote();
+	}
+}
+
+
+void FConsoleSlateDebuggerBreak::HandleEndWidgetPaint(const SWidget* Widget, const FSlateWindowElementList& /*OutDrawElements*/, int32 /*LayerId*/)
+{
+	if (Widget == nullptr)
+	{
+		return;
+	}
+
+	const FConsoleSlateDebuggerUtility::TSWidgetId WidgetId = FConsoleSlateDebuggerUtility::GetId(Widget);
+	if (BreakOnEndPaintElements.Contains(WidgetId))
+	{
+		UE_LOG(LogSlateDebugger, Log, TEXT("Widget '%s' was invalidated."), *FReflectionMetaData::GetWidgetPath(Widget));
+		_DebugBreakAndPromptForRemote();
+	}
+}
+
+
+#undef LOCTEXT_NAMESPACE
+
+#endif // WITH_SLATE_DEBUGGING
\ No newline at end of file
diff --git a/Engine/Source/Runtime/SlateCore/Private/Debugging/ConsoleSlateDebuggerBreak.h b/Engine/Source/Runtime/SlateCore/Private/Debugging/ConsoleSlateDebuggerBreak.h
new file mode 100644
--- /dev/null
+++ b/Engine/Source/Runtime/SlateCore/Private/Debugging/ConsoleSlateDebuggerBreak.h
@@ -0,0 +1,75 @@
+// Copyright Epic Games, Inc. All Rights Reserved.
+
+#pragma once
+
+
+#include "Debugging/SlateDebugging.h"
+
+#if WITH_SLATE_DEBUGGING
+
+#include "CoreMinimal.h"
+#include "Debugging/ConsoleSlateDebuggerUtility.h"
+#include "Delegates/Delegate.h"
+#include "HAL/IConsoleManager.h"
+
+class SWidget;
+
+/**
+ * Allows debugging events from the console.
+ * Basics:
+ *   SlateDebugger.Break.OnWidgetInvalidation [WidgetPtr=0x123456789] [WidgetId=123456] [Reason=All]
+ *   SlateDebugger.Break.OnWidgetPaint [WidgetPointer=0x123456789] [WidgetId=123456]
+ */
+class FConsoleSlateDebuggerBreak
+{
+public:
+	FConsoleSlateDebuggerBreak();
+	~FConsoleSlateDebuggerBreak();
+
+	void AddInvalidation(const SWidget& Widget, EInvalidateWidgetReason Reason);
+	void RemoveInvalidation(const SWidget& Widget);
+
+	void AddBeginPaint(const SWidget& Widget);
+	void RemoveBeginPaint(const SWidget& Widget);
+
+	void AddEndPaint(const SWidget& Widget);
+	void RemoveEndPaint(const SWidget& Widget);
+
+	void RemoveAll();
+
+private:
+	struct FInvalidationElement
+	{
+		FConsoleSlateDebuggerUtility::TSWidgetId WidgetId = FConsoleSlateDebuggerUtility::InvalidWidgetId;
+		EInvalidateWidgetReason Reason = (EInvalidateWidgetReason)0xFF;
+	};
+	void AddInvalidation(FInvalidationElement Element);
+	void AddBeginPaint(FConsoleSlateDebuggerUtility::TSWidgetId Element);
+	void AddEndPaint(FConsoleSlateDebuggerUtility::TSWidgetId Element);
+
+private:
+	void HandleWidgetInvalidated(const FSlateDebuggingInvalidateArgs& /*Args*/);
+	void HandleBeginWidgetPaint(const SWidget* /*Widget*/, const FPaintArgs& /*Args*/, const FGeometry& /*AllottedGeometry*/, const FSlateRect& /*MyCullingRect*/, const FSlateWindowElementList& /*OutDrawElements*/, int32 /*LayerId*/);
+	void HandleEndWidgetPaint(const SWidget* /*Widget*/, const FSlateWindowElementList& /*OutDrawElements*/, int32 /*LayerId*/);
+
+private:
+	TArray<FInvalidationElement> BreakOnInvalidationElements;
+	TArray<FConsoleSlateDebuggerUtility::TSWidgetId> BreakOnBeginPaintElements;
+	TArray<FConsoleSlateDebuggerUtility::TSWidgetId> BreakOnEndPaintElements;
+
+	//~ Console objects
+private:
+	void HandleBreakOnWidgetInvalidation(const TArray<FString>& Params);
+	void HandleBeginWidgetPaint(const TArray<FString>& Params);
+	void HandleEndWidgetPaint(const TArray<FString>& Params);
+
+	bool ParseWidgetArgs(FInvalidationElement& Out, const TArray<FString>& Params, bool bAllowInvalidationReason);
+
+private:
+	FAutoConsoleCommand WidgetInvalidationCommand;
+	FAutoConsoleCommand WidgetBeginPaintCommand;
+	FAutoConsoleCommand WidgetEndPaintCommand;
+	FAutoConsoleCommand RemoveAllCommand;
+};
+
+#endif //WITH_SLATE_DEBUGGING
\ No newline at end of file
diff --git a/Engine/Source/Runtime/SlateCore/Private/Debugging/ConsoleSlateDebuggerInvalidate.cpp b/Engine/Source/Runtime/SlateCore/Private/Debugging/ConsoleSlateDebuggerInvalidate.cpp
--- a/Engine/Source/Runtime/SlateCore/Private/Debugging/ConsoleSlateDebuggerInvalidate.cpp
+++ b/Engine/Source/Runtime/SlateCore/Private/Debugging/ConsoleSlateDebuggerInvalidate.cpp
@@ -29,6 +29,7 @@ FConsoleSlateDebuggerInvalidate::FConsoleSlateDebuggerInvalidate()
 	, DrawRootRootColor(FColorList::Red)
 	, DrawRootChildOrderColor(FColorList::Blue)
 	, DrawRootScreenPositionColor(FColorList::Green)
+	, DrawWidgetPrepassColor(FColorList::Magenta)
 	, DrawWidgetLayoutColor(FColorList::Magenta)
 	, DrawWidgetPaintColor(FColorList::Yellow)
 	, DrawWidgetVolatilityColor(FColorList::Grey)
@@ -64,11 +65,11 @@ FConsoleSlateDebuggerInvalidate::FConsoleSlateDebuggerInvalidate()
 		FConsoleCommandDelegate::CreateRaw(this, &FConsoleSlateDebuggerInvalidate::ToggleLogInvalidatedWidget))
 	, SetInvalidateWidgetReasonFilterCommand(
 		TEXT("SlateDebugger.Invalidate.SetInvalidateWidgetReasonFilter"),
-		TEXT("Enable Invalidate Widget Reason filters. Usage: SetInvalidateWidgetReasonFilter [None] [Layout] [Paint] [Volatility] [ChildOrder] [RenderTransform] [Visibility] [Any]"),
+		TEXT("Enable Invalidate Widget Reason filters. Usage: SetInvalidateWidgetReasonFilter None|Layout|Paint|Volatility|ChildOrder|RenderTransform|Visibility|Any"),
 		FConsoleCommandWithArgsDelegate::CreateRaw(this, &FConsoleSlateDebuggerInvalidate::HandleSetInvalidateWidgetReasonFilter))
 	, SetInvalidateRootReasonFilterCommand(
 		TEXT("SlateDebugger.Invalidate.SetInvalidateRootReasonFilter"),
-		TEXT("Enable Invalidate Root Reason filters. Usage: SetInvalidateRootReasonFilter [None] [ChildOrder] [Root] [ScreenPosition] [Any]"),
+		TEXT("Enable Invalidate Root Reason filters. Usage: SetInvalidateRootReasonFilter None|ChildOrder|Root|ScreenPosition|Any"),
 		FConsoleCommandWithArgsDelegate::CreateRaw(this, &FConsoleSlateDebuggerInvalidate::HandleSetInvalidateRootReasonFilter))
 {
 	LoadConfig();
@@ -97,6 +98,7 @@ void FConsoleSlateDebuggerInvalidate::LoadConfig()
 	GetColor(TEXT("DrawRootRootColor"), DrawRootRootColor);
 	GetColor(TEXT("DrawRootChildOrderColor"), DrawRootChildOrderColor);
 	GetColor(TEXT("DrawRootScreenPositionColor"), DrawRootScreenPositionColor);
+	GetColor(TEXT("DrawWidgetPrepassColor"), DrawWidgetPrepassColor);
 	GetColor(TEXT("DrawWidgetLayoutColor"), DrawWidgetLayoutColor);
 	GetColor(TEXT("DrawWidgetPaintColor"), DrawWidgetPaintColor);
 	GetColor(TEXT("DrawWidgetVolatilityColor"), DrawWidgetVolatilityColor);
@@ -122,6 +124,7 @@ void FConsoleSlateDebuggerInvalidate::SaveConfig()
 	SetColor(TEXT("DrawRootRootColor"), DrawRootRootColor);
 	SetColor(TEXT("DrawRootChildOrderColor"), DrawRootChildOrderColor);
 	SetColor(TEXT("DrawRootScreenPositionColor"), DrawRootScreenPositionColor);
+	SetColor(TEXT("DrawWidgetPrepassColor"), DrawWidgetPrepassColor);
 	SetColor(TEXT("DrawWidgetLayoutColor"), DrawWidgetLayoutColor);
 	SetColor(TEXT("DrawWidgetPaintColor"), DrawWidgetPaintColor);
 	SetColor(TEXT("DrawWidgetVolatilityColor"), DrawWidgetVolatilityColor);
@@ -192,101 +195,24 @@ void FConsoleSlateDebuggerInvalidate::ToggleLogInvalidatedWidget()
 	SaveConfig();
 }
 
-namespace ConsoleSlateDebuggerInvalidate
-{
-	template<class EnumType, int32 BuilderSize>
-	bool CheckAndAddToMessageBuilder(TStringBuilder<BuilderSize>& MessageBuilder, EnumType Filter, EnumType Reason, const TCHAR* Message, bool bFirstFlag)
-	{
-		if (EnumHasAnyFlags(Filter, Reason))
-		{
-			if (!bFirstFlag)
-			{
-				MessageBuilder << TEXT("|");
-			}
-			bFirstFlag = false;
-			MessageBuilder << Message;
-		}
-		return bFirstFlag;
-	}
-
-	template<class EnumType>
-	bool TestAndSetEnum(const FString& Param, EnumType& NewFilter, EnumType Reason, const TCHAR* ReasonParam)
-	{
-		if (Param == ReasonParam)
-		{
-			NewFilter |= Reason;
-			return true;
-		}
-		return false;
-	}
-
-	template<int32 BuilderSize>
-	void BuildEnumMessage(TStringBuilder<BuilderSize>& MessageBuilder, EInvalidateWidgetReason InvalidateWidgetReasonFilter)
-	{
-		bool bFirstFlag = true;
-		bFirstFlag = CheckAndAddToMessageBuilder(MessageBuilder, InvalidateWidgetReasonFilter, EInvalidateWidgetReason::Layout, TEXT("Layout"), bFirstFlag);
-		bFirstFlag = CheckAndAddToMessageBuilder(MessageBuilder, InvalidateWidgetReasonFilter, EInvalidateWidgetReason::Paint, TEXT("Paint"), bFirstFlag);
-		bFirstFlag = CheckAndAddToMessageBuilder(MessageBuilder, InvalidateWidgetReasonFilter, EInvalidateWidgetReason::Volatility, TEXT("Volatility"), bFirstFlag);
-		bFirstFlag = CheckAndAddToMessageBuilder(MessageBuilder, InvalidateWidgetReasonFilter, EInvalidateWidgetReason::ChildOrder, TEXT("ChildOrder"), bFirstFlag);
-		bFirstFlag = CheckAndAddToMessageBuilder(MessageBuilder, InvalidateWidgetReasonFilter, EInvalidateWidgetReason::RenderTransform, TEXT("RenderTransform"), bFirstFlag);
-		bFirstFlag = CheckAndAddToMessageBuilder(MessageBuilder, InvalidateWidgetReasonFilter, EInvalidateWidgetReason::Visibility, TEXT("Visibility"), bFirstFlag);
-
-		if (bFirstFlag)
-		{
-			MessageBuilder << TEXT("None");
-		}
-	}
-
-	template<int32 BuilderSize>
-	void BuildEnumMessage(TStringBuilder<BuilderSize>& MessageBuilder, ESlateDebuggingInvalidateRootReason InvalidateRootReasonFilter)
-	{
-		bool bFirstFlag = true;
-		bFirstFlag = ConsoleSlateDebuggerInvalidate::CheckAndAddToMessageBuilder(MessageBuilder, InvalidateRootReasonFilter, ESlateDebuggingInvalidateRootReason::ChildOrder, TEXT("ChildOrder"), bFirstFlag);
-		bFirstFlag = ConsoleSlateDebuggerInvalidate::CheckAndAddToMessageBuilder(MessageBuilder, InvalidateRootReasonFilter, ESlateDebuggingInvalidateRootReason::Root, TEXT("Root"), bFirstFlag);
-		bFirstFlag = ConsoleSlateDebuggerInvalidate::CheckAndAddToMessageBuilder(MessageBuilder, InvalidateRootReasonFilter, ESlateDebuggingInvalidateRootReason::ScreenPosition, TEXT("ScreenPosition"), bFirstFlag);
-
-		if (bFirstFlag)
-		{
-			MessageBuilder << TEXT("None");
-		}
-	}
-}
-
 void FConsoleSlateDebuggerInvalidate::HandleSetInvalidateWidgetReasonFilter(const TArray<FString>& Params)
 {
-	const TCHAR* UsageMessage = TEXT("Usage: SetInvalidateWidgetReasonFilter [None] [Layout] [Paint] [Volatility] [ChildOrder] [RenderTransform] [Visibility] [Any]");
+	const TCHAR* UsageMessage = TEXT("Usage: SetInvalidateWidgetReasonFilter None|Layout|Paint|Volatility|ChildOrder|RenderTransform|Visibility|Any");
 	if (Params.Num() == 0)
 	{
 		UE_LOG(LogSlateDebugger, Log, TEXT("%s"), UsageMessage);
 
-		TStringBuilder<128> MessageBuilder;
+		TStringBuilder<512> MessageBuilder;
 		MessageBuilder << TEXT("Current Invalidate Widget Reason set: ");
-		ConsoleSlateDebuggerInvalidate::BuildEnumMessage(MessageBuilder, InvalidateWidgetReasonFilter);
+		MessageBuilder << LexToString(InvalidateWidgetReasonFilter);
 		UE_LOG(LogSlateDebugger, Log, TEXT("%s"), MessageBuilder.GetData());
 	}
 	else
 	{
-		EInvalidateWidgetReason NewInvalidateWidgetReasonFilter = EInvalidateWidgetReason::None;
-		bool bHasValidFlags = true;
-		for (const FString& Param : Params)
-		{
-			if (ConsoleSlateDebuggerInvalidate::TestAndSetEnum(Param, NewInvalidateWidgetReasonFilter, EInvalidateWidgetReason::Layout, TEXT("Layout"))) {}
-			else if (ConsoleSlateDebuggerInvalidate::TestAndSetEnum(Param, NewInvalidateWidgetReasonFilter, EInvalidateWidgetReason::Paint, TEXT("Paint"))) {}
-			else if (ConsoleSlateDebuggerInvalidate::TestAndSetEnum(Param, NewInvalidateWidgetReasonFilter, EInvalidateWidgetReason::Volatility, TEXT("Volatility"))) {}
-			else if (ConsoleSlateDebuggerInvalidate::TestAndSetEnum(Param, NewInvalidateWidgetReasonFilter, EInvalidateWidgetReason::ChildOrder, TEXT("ChildOrder"))) {}
-			else if (ConsoleSlateDebuggerInvalidate::TestAndSetEnum(Param, NewInvalidateWidgetReasonFilter, EInvalidateWidgetReason::RenderTransform, TEXT("RenderTransform"))) {}
-			else if (ConsoleSlateDebuggerInvalidate::TestAndSetEnum(Param, NewInvalidateWidgetReasonFilter, EInvalidateWidgetReason::Visibility, TEXT("Visibility"))) {}
-			else if (ConsoleSlateDebuggerInvalidate::TestAndSetEnum(Param, NewInvalidateWidgetReasonFilter, static_cast<EInvalidateWidgetReason>(0xFF), TEXT("Any"))) {}
-			else if (ConsoleSlateDebuggerInvalidate::TestAndSetEnum(Param, NewInvalidateWidgetReasonFilter, EInvalidateWidgetReason::None, TEXT("None"))) {}
-			else
-			{
-				bHasValidFlags = false;
-				UE_LOG(LogSlateDebugger, Warning, TEXT("Param '%s' is invalid."), *Param);
-				break;
-			}
-		}
-
-		if (!bHasValidFlags)
+		TStringBuilder<512> MessageBuilder;
+		MessageBuilder.Join(Params, TEXT('|'));
+		EInvalidateWidgetReason NewInvalidateWidgetReasonFilter;
+		if (!LexTryParseString(NewInvalidateWidgetReasonFilter, MessageBuilder.ToString()))
 		{
 			UE_LOG(LogSlateDebugger, Log, TEXT("%s"), UsageMessage);
 		}
@@ -300,42 +226,28 @@ void FConsoleSlateDebuggerInvalidate::HandleSetInvalidateWidgetReasonFilter(cons
 
 void FConsoleSlateDebuggerInvalidate::HandleSetInvalidateRootReasonFilter(const TArray<FString>& Params)
 {
-	const TCHAR* UsageMessage = TEXT("Usage: SetInvalidateRootReasonFilter [None] [ChildOrder] [Root] [ScreenPosition] [Any]");
+	const TCHAR* UsageMessage = TEXT("Usage: SetInvalidateRootReasonFilter None|ChildOrder|Root|ScreenPosition|Any");
 	if (Params.Num() == 0)
 	{
 		UE_LOG(LogSlateDebugger, Log, TEXT("%s"), UsageMessage);
 
-		TStringBuilder<128> MessageBuilder;
+		TStringBuilder<512> MessageBuilder;
 		MessageBuilder << TEXT("Current Invalidate Root Reason set: ");
-		ConsoleSlateDebuggerInvalidate::BuildEnumMessage(MessageBuilder, InvalidateRootReasonFilter);
+		MessageBuilder << LexToString(InvalidateRootReasonFilter);
 		UE_LOG(LogSlateDebugger, Log, TEXT("%s"), MessageBuilder.GetData());
 	}
 	else
 	{
-		ESlateDebuggingInvalidateRootReason NewInvalidateRoottReasonFilter = ESlateDebuggingInvalidateRootReason::None;
-		bool bHasValidFlags = true;
-		for (const FString& Param : Params)
-		{
-			if (ConsoleSlateDebuggerInvalidate::TestAndSetEnum(Param, NewInvalidateRoottReasonFilter, ESlateDebuggingInvalidateRootReason::ChildOrder, TEXT("ChildOrder"))) {}
-			else if (ConsoleSlateDebuggerInvalidate::TestAndSetEnum(Param, NewInvalidateRoottReasonFilter, ESlateDebuggingInvalidateRootReason::Root, TEXT("Root"))) {}
-			else if (ConsoleSlateDebuggerInvalidate::TestAndSetEnum(Param, NewInvalidateRoottReasonFilter, ESlateDebuggingInvalidateRootReason::ScreenPosition, TEXT("ScreenPosition"))) {}
-			else if (ConsoleSlateDebuggerInvalidate::TestAndSetEnum(Param, NewInvalidateRoottReasonFilter, static_cast<ESlateDebuggingInvalidateRootReason>(0xFF), TEXT("Any"))) {}
-			else if (ConsoleSlateDebuggerInvalidate::TestAndSetEnum(Param, NewInvalidateRoottReasonFilter, ESlateDebuggingInvalidateRootReason::None, TEXT("None"))) {}
-			else
-			{
-				bHasValidFlags = false;
-				UE_LOG(LogSlateDebugger, Warning, TEXT("Param '%s' is invalid."), *Param);
-				break;
-			}
-		}
-
-		if (!bHasValidFlags)
+		TStringBuilder<512> MessageBuilder;
+		MessageBuilder.Join(Params, TEXT('|'));
+		ESlateDebuggingInvalidateRootReason NewInvalidateRootReasonFilter;
+		if (!LexTryParseString(NewInvalidateRootReasonFilter, MessageBuilder.ToString()))
 		{
 			UE_LOG(LogSlateDebugger, Log, TEXT("%s"), UsageMessage);
 		}
 		else
 		{
-			InvalidateRootReasonFilter = NewInvalidateRoottReasonFilter;
+			InvalidateRootReasonFilter = NewInvalidateRootReasonFilter;
 			SaveConfig();
 		}
 	}
@@ -359,17 +271,17 @@ int32 FConsoleSlateDebuggerInvalidate::GetInvalidationPriority(EInvalidateWidget
 		return 50;
 	}
 
-	if (EnumHasAnyFlags(InvalidationInfo, EInvalidateWidgetReason::Layout | EInvalidateWidgetReason::ChildOrder | EInvalidateWidgetReason::Visibility | EInvalidateWidgetReason::RenderTransform))
+	if (EnumHasAnyFlags(InvalidationInfo, EInvalidateWidgetReason::Prepass | EInvalidateWidgetReason::Layout | EInvalidateWidgetReason::ChildOrder | EInvalidateWidgetReason::Visibility | EInvalidateWidgetReason::RenderTransform))
 	{
 		return 40;
 	}
-	else if (EnumHasAnyFlags(InvalidationInfo, EInvalidateWidgetReason::Paint))
+	else if (EnumHasAnyFlags(InvalidationInfo, EInvalidateWidgetReason::Volatility))
 	{
-		return 20;
+		return 30;
 	}
-	else if (EnumHasAnyFlags(InvalidationInfo, EInvalidateWidgetReason::Volatility))
+	else if (EnumHasAnyFlags(InvalidationInfo, EInvalidateWidgetReason::Paint))
 	{
-		return 10;
+		return 20;
 	}
 	return 0;
 }
@@ -389,7 +301,11 @@ const FLinearColor& FConsoleSlateDebuggerInvalidate::GetColor(const FInvalidatio
 		return DrawRootScreenPositionColor;
 	}
 
-	if (EnumHasAnyFlags(InvalidationInfo.WidgetReason, EInvalidateWidgetReason::Layout))
+	if (EnumHasAnyFlags(InvalidationInfo.WidgetReason, EInvalidateWidgetReason::Prepass))
+	{
+		return DrawWidgetPrepassColor;
+	}
+	else if (EnumHasAnyFlags(InvalidationInfo.WidgetReason, EInvalidateWidgetReason::Layout))
 	{
 		return DrawWidgetLayoutColor;
 	}
@@ -594,9 +510,9 @@ void FConsoleSlateDebuggerInvalidate::ProcessFrameList()
 			MessageBuilder << TEXT("' Invalidated: '");
 			MessageBuilder << FrameInvalidationInfo.WidgetInvalidatedName;
 			MessageBuilder << TEXT("' Root Reason: '");
-			ConsoleSlateDebuggerInvalidate::BuildEnumMessage(MessageBuilder, FrameInvalidationInfo.InvalidationRootReason);
+			MessageBuilder << LexToString(FrameInvalidationInfo.InvalidationRootReason);
 			MessageBuilder << TEXT("' Widget Reason: '");
-			ConsoleSlateDebuggerInvalidate::BuildEnumMessage(MessageBuilder, FrameInvalidationInfo.WidgetReason);
+			MessageBuilder << LexToString(FrameInvalidationInfo.WidgetReason);
 			MessageBuilder << TEXT("'");
 
 			UE_LOG(LogSlateDebugger, Log, TEXT("%s"), MessageBuilder.GetData());
diff --git a/Engine/Source/Runtime/SlateCore/Private/Debugging/ConsoleSlateDebuggerInvalidate.h b/Engine/Source/Runtime/SlateCore/Private/Debugging/ConsoleSlateDebuggerInvalidate.h
--- a/Engine/Source/Runtime/SlateCore/Private/Debugging/ConsoleSlateDebuggerInvalidate.h
+++ b/Engine/Source/Runtime/SlateCore/Private/Debugging/ConsoleSlateDebuggerInvalidate.h
@@ -96,6 +96,7 @@ private:
 	FLinearColor DrawRootRootColor;
 	FLinearColor DrawRootChildOrderColor;
 	FLinearColor DrawRootScreenPositionColor;
+	FLinearColor DrawWidgetPrepassColor;
 	FLinearColor DrawWidgetLayoutColor;
 	FLinearColor DrawWidgetPaintColor;
 	FLinearColor DrawWidgetVolatilityColor;
diff --git a/Engine/Source/Runtime/SlateCore/Private/Debugging/ConsoleSlateDebuggerUtility.cpp b/Engine/Source/Runtime/SlateCore/Private/Debugging/ConsoleSlateDebuggerUtility.cpp
--- a/Engine/Source/Runtime/SlateCore/Private/Debugging/ConsoleSlateDebuggerUtility.cpp
+++ b/Engine/Source/Runtime/SlateCore/Private/Debugging/ConsoleSlateDebuggerUtility.cpp
@@ -9,6 +9,15 @@
 
 DEFINE_LOG_CATEGORY(LogSlateDebugger);
 
+FConsoleSlateDebuggerUtility::TSWidgetId FConsoleSlateDebuggerUtility::GetId(const SWidget& Widget)
+{
+#if UE_SLATE_WITH_WIDGET_UNIQUE_IDENTIFIER
+	return Widget.GetId();
+#else
+	return reinterpret_cast<TSWidgetId>(&Widget);
+#endif
+}
+
 FConsoleSlateDebuggerUtility::TSWidgetId FConsoleSlateDebuggerUtility::GetId(const SWidget* Widget)
 {
 // We prefer to use the widget id but if it's not available, then use the widget address.
@@ -22,6 +31,15 @@ FConsoleSlateDebuggerUtility::TSWidgetId FConsoleSlateDebuggerUtility::GetId(con
 #endif
 }
 
+FConsoleSlateDebuggerUtility::TSWindowId FConsoleSlateDebuggerUtility::GetId(const SWindow& Widget)
+{
+#if UE_SLATE_WITH_WIDGET_UNIQUE_IDENTIFIER
+	return Widget.GetId();
+#else
+	return reinterpret_cast<TSWindowId>(&Widget);
+#endif
+}
+
 FConsoleSlateDebuggerUtility::TSWindowId FConsoleSlateDebuggerUtility::GetId(const SWindow* Widget)
 {
 #if UE_SLATE_WITH_WIDGET_UNIQUE_IDENTIFIER
@@ -31,6 +49,11 @@ FConsoleSlateDebuggerUtility::TSWindowId FConsoleSlateDebuggerUtility::GetId(con
 #endif
 }
 
+FConsoleSlateDebuggerUtility::TSWindowId FConsoleSlateDebuggerUtility::FindWindowId(const SWidget& Widget)
+{
+	return FindWindowId(&Widget);
+}
+
 FConsoleSlateDebuggerUtility::TSWindowId FConsoleSlateDebuggerUtility::FindWindowId(const SWidget* Widget)
 {
 	while (Widget)
diff --git a/Engine/Source/Runtime/SlateCore/Private/Debugging/ConsoleSlateDebuggerUtility.h b/Engine/Source/Runtime/SlateCore/Private/Debugging/ConsoleSlateDebuggerUtility.h
--- a/Engine/Source/Runtime/SlateCore/Private/Debugging/ConsoleSlateDebuggerUtility.h
+++ b/Engine/Source/Runtime/SlateCore/Private/Debugging/ConsoleSlateDebuggerUtility.h
@@ -30,8 +30,11 @@ struct FConsoleSlateDebuggerUtility
 	static const TSWidgetId InvalidWidgetId = 0;
 	static const TSWindowId InvalidWindowId = 0;
 
+	static TSWidgetId GetId(const SWidget& Widget);
 	static TSWidgetId GetId(const SWidget* Widget);
+	static TSWindowId GetId(const SWindow& Widget);
 	static TSWindowId GetId(const SWindow* Widget);
+	static TSWindowId FindWindowId(const SWidget& Widget);
 	static TSWindowId FindWindowId(const SWidget* Widget);
 };
 
diff --git a/Engine/Source/Runtime/SlateCore/Private/Debugging/SlateDebugging.cpp b/Engine/Source/Runtime/SlateCore/Private/Debugging/SlateDebugging.cpp
--- a/Engine/Source/Runtime/SlateCore/Private/Debugging/SlateDebugging.cpp
+++ b/Engine/Source/Runtime/SlateCore/Private/Debugging/SlateDebugging.cpp
@@ -241,6 +241,76 @@ FText FSlateDebuggingCursorQueryEventArgs::ToText() const
 	return EventText;
 }
 
+FString LexToString(ESlateDebuggingInvalidateRootReason InValue)
+{
+	if (InValue == ESlateDebuggingInvalidateRootReason::None)
+	{
+		return TEXT("None");
+	}
+	if (InValue == (ESlateDebuggingInvalidateRootReason)0xFF)
+	{
+		return TEXT("All");
+	}
+
+	TStringBuilder<512> Result;
+#define ENUM_CASE_TO_STRING(Enum) if (EnumHasAnyFlags(InValue, ESlateDebuggingInvalidateRootReason::Enum)) { if (Result.Len() != 0) { Result.Append(TEXT('|')); } Result.Append(TEXT(#Enum)); }
+	ENUM_CASE_TO_STRING(ChildOrder);
+	ENUM_CASE_TO_STRING(Root);
+	ENUM_CASE_TO_STRING(ScreenPosition);
+#undef ENUM_CASE_TO_STRING
+
+	return Result.ToString();
+}
+
+bool LexTryParseString(ESlateDebuggingInvalidateRootReason& OutValue, const TCHAR* Buffer)
+{
+	bool bResult = false;
+	ESlateDebuggingInvalidateRootReason Value = ESlateDebuggingInvalidateRootReason::None;
+	auto ParseResult = [&bResult, &Value](FStringView& SubString)
+	{
+		SubString.TrimStartAndEndInline();
+
+		if (SubString.Equals(TEXT("All"), ESearchCase::IgnoreCase)) { Value = (ESlateDebuggingInvalidateRootReason)0xFF; return; }
+		if (SubString.Equals(TEXT("Any"), ESearchCase::IgnoreCase)) { Value = (ESlateDebuggingInvalidateRootReason)0xFF; return; }
+
+#define ENUM_CASE_FROM_STRING(Enum) if (SubString.Equals(TEXT(#Enum), ESearchCase::IgnoreCase)) { Value |= ESlateDebuggingInvalidateRootReason::Enum; return; }
+		ENUM_CASE_FROM_STRING(None)
+		ENUM_CASE_FROM_STRING(ChildOrder)
+		ENUM_CASE_FROM_STRING(Root)
+		ENUM_CASE_FROM_STRING(ScreenPosition)
+		bResult = false;
+#undef ENUM_CASE_FROM_STRING
+	};
+
+	if (Buffer && *Buffer)
+	{
+		bResult = true;
+		while (const TCHAR* At = FCString::Strchr(Buffer, TEXT('|')))
+		{
+			FStringView SubString{ Buffer, UE_PTRDIFF_TO_INT32(At - Buffer) };
+			ParseResult(SubString);
+			Buffer = At + 1;
+		}
+		if (*Buffer)
+		{
+			FStringView SubString{ Buffer };
+			ParseResult(SubString);
+		}
+	}
+
+	if (bResult)
+	{
+		OutValue = Value;
+	}
+	return bResult;
+}
+
+void LexFromString(ESlateDebuggingInvalidateRootReason& OutValue, const TCHAR* Buffer)
+{
+	OutValue = ESlateDebuggingInvalidateRootReason::None;
+	LexTryParseString(OutValue, Buffer);
+}
+
 FSlateDebuggingInvalidateArgs::FSlateDebuggingInvalidateArgs(
 	const SWidget* InWidgetInvalidated,
 	const SWidget* InWidgetInvalidateInvestigator,
diff --git a/Engine/Source/Runtime/SlateCore/Private/SlateCoreModule.cpp b/Engine/Source/Runtime/SlateCore/Private/SlateCoreModule.cpp
--- a/Engine/Source/Runtime/SlateCore/Private/SlateCoreModule.cpp
+++ b/Engine/Source/Runtime/SlateCore/Private/SlateCoreModule.cpp
@@ -6,6 +6,7 @@
 #include "Debugging/ConsoleSlateDebuggerInvalidationRoot.h"
 #include "Debugging/ConsoleSlateDebuggerPaint.h"
 #include "Debugging/ConsoleSlateDebuggerUpdate.h"
+#include "Debugging/ConsoleSlateDebuggerBreak.h"
 #include "Modules/ModuleInterface.h"
 #include "Modules/ModuleManager.h"
 #include "SlateGlobals.h"
@@ -33,6 +34,7 @@ public:
 		SlateDebuggerInvalidationRoot = MakeUnique<FConsoleSlateDebuggerInvalidationRoot>();
 		SlateDebuggerPaint = MakeUnique<FConsoleSlateDebuggerPaint>();
 		SlateDebuggerUpdate = MakeUnique<FConsoleSlateDebuggerUpdate>();
+		SlateDebuggerBreak = MakeUnique<FConsoleSlateDebuggerBreak>();
 #endif
 	}
 
@@ -43,6 +45,7 @@ private:
 	TUniquePtr<FConsoleSlateDebuggerInvalidationRoot> SlateDebuggerInvalidationRoot;
 	TUniquePtr<FConsoleSlateDebuggerPaint> SlateDebuggerPaint;
 	TUniquePtr<FConsoleSlateDebuggerUpdate> SlateDebuggerUpdate;
+	TUniquePtr<FConsoleSlateDebuggerBreak> SlateDebuggerBreak;
 #endif
 };
 
diff --git a/Engine/Source/Runtime/SlateCore/Private/Test/SlateInvalidationWidgetReasonText.cpp b/Engine/Source/Runtime/SlateCore/Private/Test/SlateInvalidationWidgetReasonText.cpp
new file mode 100644
--- /dev/null
+++ b/Engine/Source/Runtime/SlateCore/Private/Test/SlateInvalidationWidgetReasonText.cpp
@@ -0,0 +1,51 @@
+// Copyright Epic Games, Inc. All Rights Reserved.
+
+#include "CoreMinimal.h"
+#include "Misc/AutomationTest.h"
+
+#include "Widgets/InvalidateWidgetReason.h"
+
+#if WITH_AUTOMATION_WORKER
+
+#define LOCTEXT_NAMESPACE "Slate.InvalidationWidgetReason.Lex"
+
+IMPLEMENT_SIMPLE_AUTOMATION_TEST(FSlateInvalidationWidgetReasonLexTest, "Slate.InvalidationWidgetReason.x", EAutomationTestFlags::ApplicationContextMask | EAutomationTestFlags::EngineFilter)
+
+bool FSlateInvalidationWidgetReasonLexTest::RunTest(const FString& Parameters)
+{
+	EInvalidateWidgetReason Value = (EInvalidateWidgetReason)0xFF;
+	AddErrorIfFalse(LexToString(Value) == TEXT("All"), TEXT("'All' was not converted properly."));
+	Value = EInvalidateWidgetReason::None;
+	AddErrorIfFalse(LexToString(Value) == TEXT("None"), TEXT("'None' was not converted properly."));
+
+	auto TestLextToStringTRyParseString = [this](EInvalidateWidgetReason Reason, const TCHAR* ReasonStr)
+		{
+			EInvalidateWidgetReason ParsedValue = EInvalidateWidgetReason::None;
+			FString StringValue = LexToString(Reason);
+			AddErrorIfFalse(LexTryParseString(ParsedValue, *StringValue), FString::Printf(TEXT("Was not able to parse '%s'."), ReasonStr));
+			AddErrorIfFalse(Reason == ParsedValue, FString::Printf(TEXT("The parsed value and the converted value for '%s' do not match."), ReasonStr));
+		};
+
+	TestLextToStringTRyParseString(EInvalidateWidgetReason::None, TEXT("None"));
+	TestLextToStringTRyParseString(EInvalidateWidgetReason::Layout, TEXT("Layout"));
+	TestLextToStringTRyParseString(EInvalidateWidgetReason::Paint, TEXT("Paint"));
+	TestLextToStringTRyParseString(EInvalidateWidgetReason::Volatility, TEXT("Volatility"));
+	TestLextToStringTRyParseString(EInvalidateWidgetReason::ChildOrder, TEXT("ChildOrder"));
+	TestLextToStringTRyParseString(EInvalidateWidgetReason::RenderTransform, TEXT("RenderTransform"));
+	TestLextToStringTRyParseString(EInvalidateWidgetReason::Visibility, TEXT("Visibility"));
+	TestLextToStringTRyParseString(EInvalidateWidgetReason::AttributeRegistration, TEXT("AttributeRegistration"));
+	TestLextToStringTRyParseString(EInvalidateWidgetReason::Prepass, TEXT("Prepass"));
+
+	TestLextToStringTRyParseString(EInvalidateWidgetReason::LayoutAndVolatility, TEXT("LayoutAndVolatility"));
+	TestLextToStringTRyParseString(EInvalidateWidgetReason::PaintAndVolatility, TEXT("PaintAndVolatility"));
+	TestLextToStringTRyParseString(EInvalidateWidgetReason::PaintAndVolatility|EInvalidateWidgetReason::Visibility| EInvalidateWidgetReason::None, TEXT("PaintAndVolatility With None"));
+	TestLextToStringTRyParseString(EInvalidateWidgetReason::PaintAndVolatility|EInvalidateWidgetReason::Visibility| (EInvalidateWidgetReason)0xFF, TEXT("PaintAndVolatility With All"));
+
+	AddErrorIfFalse(LexTryParseString(Value, TEXT(" Layout | Volatility ")), TEXT("'Layout | Volatility' With spaces"));
+	AddErrorIfFalse(Value == EInvalidateWidgetReason::LayoutAndVolatility, TEXT("'Layout | Volatility' With spaces"));
+
+
+	return true;
+}
+
+#endif //WITH_AUTOMATION_WORKER
diff --git a/Engine/Source/Runtime/SlateCore/Private/Widgets/InvalidateWidgetReason.cpp b/Engine/Source/Runtime/SlateCore/Private/Widgets/InvalidateWidgetReason.cpp
new file mode 100644
--- /dev/null
+++ b/Engine/Source/Runtime/SlateCore/Private/Widgets/InvalidateWidgetReason.cpp
@@ -0,0 +1,84 @@
+// Copyright Epic Games, Inc. All Rights Reserved.
+
+#include "Widgets/InvalidateWidgetReason.h"
+#include "Containers/StringFwd.h"
+
+bool LexTryParseString(EInvalidateWidgetReason& OutValue, const TCHAR* Buffer)
+{
+	bool bResult = false;
+	EInvalidateWidgetReason Value = EInvalidateWidgetReason::None;
+	auto ParseResult = [&bResult, &Value](FStringView& SubString)
+	{
+		SubString.TrimStartAndEndInline();
+
+		if (SubString.Equals(TEXT("All"), ESearchCase::IgnoreCase)) { Value = (EInvalidateWidgetReason)0xFF; return; }
+		if (SubString.Equals(TEXT("Any"), ESearchCase::IgnoreCase)) { Value = (EInvalidateWidgetReason)0xFF; return; }
+
+#define ENUM_CASE_FROM_STRING(Enum) if (SubString.Equals(TEXT(#Enum), ESearchCase::IgnoreCase)) { Value |= EInvalidateWidgetReason::Enum; return; }
+		ENUM_CASE_FROM_STRING(None)
+		ENUM_CASE_FROM_STRING(Layout)
+		ENUM_CASE_FROM_STRING(Paint)
+		ENUM_CASE_FROM_STRING(Volatility)
+		ENUM_CASE_FROM_STRING(ChildOrder)
+		ENUM_CASE_FROM_STRING(RenderTransform)
+		ENUM_CASE_FROM_STRING(Visibility)
+		ENUM_CASE_FROM_STRING(AttributeRegistration)
+		ENUM_CASE_FROM_STRING(Prepass)
+		bResult = false;
+#undef ENUM_CASE_FROM_STRING
+	};
+
+	if (Buffer && *Buffer)
+	{
+		bResult = true;
+		while (const TCHAR* At = FCString::Strchr(Buffer, TEXT('|')))
+		{
+			FStringView SubString{ Buffer, UE_PTRDIFF_TO_INT32(At - Buffer) };
+			ParseResult(SubString);
+			Buffer = At + 1;
+		}
+		if (*Buffer)
+		{
+			FStringView SubString{ Buffer };
+			ParseResult(SubString);
+		}
+	}
+
+	if (bResult)
+	{
+		OutValue = Value;
+	}
+	return bResult;
+}
+
+void LexFromString(EInvalidateWidgetReason& OutValue, const TCHAR* Buffer)
+{
+	OutValue = EInvalidateWidgetReason::None;
+	LexTryParseString(OutValue, Buffer);
+}
+
+FString LexToString(EInvalidateWidgetReason InValue)
+{
+	if (InValue == EInvalidateWidgetReason::None)
+	{
+		return TEXT("None");
+	}
+	if (InValue == (EInvalidateWidgetReason)0xFF)
+	{
+		return TEXT("All");
+	}
+
+	TStringBuilder<512> Result;
+	#define ENUM_CASE_TO_STRING(Enum) if (EnumHasAnyFlags(InValue, EInvalidateWidgetReason::Enum)) { if (Result.Len() != 0) { Result.Append(TEXT('|')); } Result.Append(TEXT(#Enum)); }
+	ENUM_CASE_TO_STRING(Layout);
+	ENUM_CASE_TO_STRING(Paint);
+	ENUM_CASE_TO_STRING(Volatility);
+	ENUM_CASE_TO_STRING(ChildOrder);
+	ENUM_CASE_TO_STRING(RenderTransform);
+	ENUM_CASE_TO_STRING(Visibility);
+	ENUM_CASE_TO_STRING(AttributeRegistration);
+	ENUM_CASE_TO_STRING(Prepass);
+	#undef ENUM_CASE_TO_STRING
+
+	return Result.ToString();
+}
diff --git a/Engine/Source/Runtime/SlateCore/Public/Debugging/SlateDebugging.h b/Engine/Source/Runtime/SlateCore/Public/Debugging/SlateDebugging.h
--- a/Engine/Source/Runtime/SlateCore/Public/Debugging/SlateDebugging.h
+++ b/Engine/Source/Runtime/SlateCore/Public/Debugging/SlateDebugging.h
@@ -207,6 +207,10 @@ enum class ESlateDebuggingInvalidateRootReason
 
 ENUM_CLASS_FLAGS(ESlateDebuggingInvalidateRootReason)
 
+SLATECORE_API FString LexToString(ESlateDebuggingInvalidateRootReason Reason);
+SLATECORE_API bool LexTryParseString(ESlateDebuggingInvalidateRootReason& OutMode, const TCHAR* InBuffer);
+SLATECORE_API void LexFromString(ESlateDebuggingInvalidateRootReason& OutMode, const TCHAR* InBuffer);
+
 struct SLATECORE_API FSlateDebuggingInvalidateArgs
 {
 	FSlateDebuggingInvalidateArgs(
diff --git a/Engine/Source/Runtime/SlateCore/Public/Widgets/InvalidateWidgetReason.h b/Engine/Source/Runtime/SlateCore/Public/Widgets/InvalidateWidgetReason.h
--- a/Engine/Source/Runtime/SlateCore/Public/Widgets/InvalidateWidgetReason.h
+++ b/Engine/Source/Runtime/SlateCore/Public/Widgets/InvalidateWidgetReason.h
@@ -73,5 +73,9 @@ enum class EInvalidateWidgetReason : uint8
 
 ENUM_CLASS_FLAGS(EInvalidateWidgetReason)
 
+SLATECORE_API FString LexToString(EInvalidateWidgetReason Reason);
+SLATECORE_API bool LexTryParseString(EInvalidateWidgetReason& OutMode, const TCHAR* InBuffer);
+SLATECORE_API void LexFromString(EInvalidateWidgetReason& OutMode, const TCHAR* InBuffer);
+
 // This typedefed because EInvalidateWidget will be deprecated soon
 typedef EInvalidateWidgetReason EInvalidateWidget;
