Slate: Preperation work for SlateAttribute in FSlot.
Force a valid SWidget pointer for FChildren and for FNoChildren. Deprecate the TSupportsOneChildMixin because it was miss used at some places and cause FChildren to be used as FSlot. Prepare the TPaddingWidgetSlotMixin and TAlignmentWidgetSlotMixin to support invalidation. Fix slot instance that were not using the mixin (easier when we will support SlateAttribute).
#prefligh 609e94f6cbf9a4000115d883

diff --git a/Engine/Source/Developer/SlateReflector/Private/Widgets/SAtlasVisualizer.cpp b/Engine/Source/Developer/SlateReflector/Private/Widgets/SAtlasVisualizer.cpp
--- a/Engine/Source/Developer/SlateReflector/Private/Widgets/SAtlasVisualizer.cpp
+++ b/Engine/Source/Developer/SlateReflector/Private/Widgets/SAtlasVisualizer.cpp
@@ -16,14 +16,7 @@
 class SAtlasVisualizerPanel : public IScrollableZoomable, public SPanel
 {
 public:
-	class FAtlasVisualizerPanelSlot : public TSupportsOneChildMixin<FAtlasVisualizerPanelSlot>
-	{
-	public:
-		FAtlasVisualizerPanelSlot(SWidget* InOwner)
-			: TSupportsOneChildMixin<FAtlasVisualizerPanelSlot>(InOwner)
-		{
-		}
-	};
+	using FAtlasVisualizerPanelSlot = FSingleWidgetChildrenWithSlot;
 
 	SLATE_BEGIN_ARGS(SAtlasVisualizerPanel)
 		{
diff --git a/Engine/Source/Editor/MaterialEditor/Private/SMaterialEditorViewport.cpp b/Engine/Source/Editor/MaterialEditor/Private/SMaterialEditorViewport.cpp
--- a/Engine/Source/Editor/MaterialEditor/Private/SMaterialEditorViewport.cpp
+++ b/Engine/Source/Editor/MaterialEditor/Private/SMaterialEditorViewport.cpp
@@ -775,15 +775,7 @@ void SMaterialEditor3DPreviewViewport::OnPropertyChanged(UObject* ObjectBeingMod
 class SMaterialEditorUIPreviewZoomer : public SPanel, public IScrollableZoomable, public FGCObject
 {
 public:
-
-	class FMaterialPreviewPanelSlot : public TSupportsOneChildMixin<FMaterialPreviewPanelSlot>
-	{
-	public:
-		FMaterialPreviewPanelSlot(SWidget* InOwner)
-			: TSupportsOneChildMixin<FMaterialPreviewPanelSlot>(InOwner)
-		{
-		}
-	};
+	using FMaterialPreviewPanelSlot = FSingleWidgetChildrenWithSlot;
 
 	SLATE_BEGIN_ARGS(SMaterialEditorUIPreviewZoomer)
 		: _OnContextMenuRequested()
diff --git a/Engine/Source/Editor/WorldPartitionEditor/Public/WorldPartition/SWorldPartitionEditorGrid.h b/Engine/Source/Editor/WorldPartitionEditor/Public/WorldPartition/SWorldPartitionEditorGrid.h
--- a/Engine/Source/Editor/WorldPartitionEditor/Public/WorldPartition/SWorldPartitionEditorGrid.h
+++ b/Engine/Source/Editor/WorldPartitionEditor/Public/WorldPartition/SWorldPartitionEditorGrid.h
@@ -53,8 +53,7 @@ public:
 
 	virtual FChildren* GetChildren() override
 	{
-		static FNoChildren NoChildrenInstance;
-		return &NoChildrenInstance;
+		return &FNoChildren::NoChildrenInstance;
 	}
 	
 	bool GetPlayerView(FVector& Location, FRotator& Rotation) const;
diff --git a/Engine/Source/Editor/WorldPartitionEditor/Public/WorldPartition/SWorldPartitionEditorGrid2D.h b/Engine/Source/Editor/WorldPartitionEditor/Public/WorldPartition/SWorldPartitionEditorGrid2D.h
--- a/Engine/Source/Editor/WorldPartitionEditor/Public/WorldPartition/SWorldPartitionEditorGrid2D.h
+++ b/Engine/Source/Editor/WorldPartitionEditor/Public/WorldPartition/SWorldPartitionEditorGrid2D.h
@@ -61,7 +61,7 @@ protected:
 
 	const TSharedRef<FUICommandList> CommandList;
 
-	FSimpleSlot ChildSlot;
+	FSingleWidgetChildrenWithBasicLayoutSlot ChildSlot;
 
 	FChildren* GetChildren()
 	{
diff --git a/Engine/Source/Runtime/Engine/Private/GameViewportClient.cpp b/Engine/Source/Runtime/Engine/Private/GameViewportClient.cpp
--- a/Engine/Source/Runtime/Engine/Private/GameViewportClient.cpp
+++ b/Engine/Source/Runtime/Engine/Private/GameViewportClient.cpp
@@ -2736,7 +2736,6 @@ void UGameViewportClient::AddViewportWidgetContent( TSharedRef<SWidget> Viewport
 	TSharedPtr< SOverlay > PinnedViewportOverlayWidget( ViewportOverlayWidget.Pin() );
 	if( ensure( PinnedViewportOverlayWidget.IsValid() ) )
 	{
-		// NOTE: Returns FSimpleSlot but we're ignoring here.  Could be used for alignment though.
 		PinnedViewportOverlayWidget->AddSlot( ZOrder )
 			[
 				ViewportContent
diff --git a/Engine/Source/Runtime/Engine/Private/Slate/SGameLayerManager.cpp b/Engine/Source/Runtime/Engine/Private/Slate/SGameLayerManager.cpp
--- a/Engine/Source/Runtime/Engine/Private/Slate/SGameLayerManager.cpp
+++ b/Engine/Source/Runtime/Engine/Private/Slate/SGameLayerManager.cpp
@@ -184,7 +184,6 @@ void SGameLayerManager::AddWidgetForPlayer(ULocalPlayer* Player, TSharedRef<SWid
 {
 	TSharedPtr<FPlayerLayer> PlayerLayer = FindOrCreatePlayerLayer(Player);
 	
-	// NOTE: Returns FSimpleSlot but we're ignoring here.  Could be used for alignment though.
 	PlayerLayer->Widget->AddSlot(ZOrder)
 	[
 		ViewportContent
diff --git a/Engine/Source/Runtime/Slate/Private/Framework/MultiBox/SUniformToolbarPanel.h b/Engine/Source/Runtime/Slate/Private/Framework/MultiBox/SUniformToolbarPanel.h
--- a/Engine/Source/Runtime/Slate/Private/Framework/MultiBox/SUniformToolbarPanel.h
+++ b/Engine/Source/Runtime/Slate/Private/Framework/MultiBox/SUniformToolbarPanel.h
@@ -25,11 +25,11 @@ class SLATE_API SUniformToolbarPanel : public SPanel
 {
 public:
 	/** Stores the per-child info for this panel type */
-	struct FSlot : public TSlotBase<FSlot>, public TSupportsContentAlignmentMixin<FSlot>
+	struct FSlot : public TSlotBase<FSlot>, public TAlignmentWidgetSlotMixin<FSlot>
 	{
 		FSlot()
 		: TSlotBase<FSlot>()
-		, TSupportsContentAlignmentMixin<FSlot>(HAlign_Fill, VAlign_Fill)
+		, TAlignmentWidgetSlotMixin<FSlot>(HAlign_Fill, VAlign_Fill)
 		{
 		}
 
diff --git a/Engine/Source/Runtime/Slate/Private/Widgets/Input/SMenuAnchor.cpp b/Engine/Source/Runtime/Slate/Private/Widgets/Input/SMenuAnchor.cpp
--- a/Engine/Source/Runtime/Slate/Private/Widgets/Input/SMenuAnchor.cpp
+++ b/Engine/Source/Runtime/Slate/Private/Widgets/Input/SMenuAnchor.cpp
@@ -83,8 +83,8 @@ SMenuAnchor::FPopupPlacement::FPopupPlacement(const FGeometry& PlacementGeometry
  */
 void SMenuAnchor::Construct( const FArguments& InArgs )
 {
-	Children.Add( new FSimpleSlot(this) );
-	Children.Add( new FSimpleSlot(this) );
+	Children.Add(new FBasicLayoutWidgetSlot());
+	Children.Add(new FBasicLayoutWidgetSlot());
 	
 
 	Children[0]
diff --git a/Engine/Source/Runtime/Slate/Private/Widgets/SInvalidationPanel.cpp b/Engine/Source/Runtime/Slate/Private/Widgets/SInvalidationPanel.cpp
--- a/Engine/Source/Runtime/Slate/Private/Widgets/SInvalidationPanel.cpp
+++ b/Engine/Source/Runtime/Slate/Private/Widgets/SInvalidationPanel.cpp
@@ -37,8 +37,7 @@ FAutoConsoleVariableRef CVarAlwaysInvalidate(
 
 
 SInvalidationPanel::SInvalidationPanel()
-	: EmptyChildSlot(this)
-	, HittestGrid(MakeShared<FHittestGrid>())
+	: HittestGrid(MakeShared<FHittestGrid>())
 	, bCanCache(true)
 	, bPaintedSinceLastPrepass(true)
 	, bWasCachable(false)
@@ -191,7 +190,7 @@ FChildren* SInvalidationPanel::GetChildren()
 {
 	if (GetCanCache())
 	{
-		return &EmptyChildSlot;
+		return &FNoChildren::NoChildrenInstance;
 	}
 	else
 	{
diff --git a/Engine/Source/Runtime/Slate/Private/Widgets/Views/SListPanel.cpp b/Engine/Source/Runtime/Slate/Private/Widgets/Views/SListPanel.cpp
--- a/Engine/Source/Runtime/Slate/Private/Widgets/Views/SListPanel.cpp
+++ b/Engine/Source/Runtime/Slate/Private/Widgets/Views/SListPanel.cpp
@@ -4,8 +4,6 @@
 #include "Layout/ArrangedChildren.h"
 
 
-FNoChildren SListPanel::NoChildren = FNoChildren();
-
 // Used to subtract a tiny amount from the available dimension to avoid floating point precision problems when arranging children
 static const float FloatingPointPrecisionOffset = 0.001f;
 
@@ -191,7 +189,7 @@ FChildren* SListPanel::GetChildren()
 		// When a refresh is pending it is unsafe to cache the desired sizes of our children because
 		// they may be representing unsound data structures. Any delegates/attributes accessing unsound
 		// data will cause a crash.
-		return &NoChildren;
+		return &FNoChildren::NoChildrenInstance;
 	}
 	else
 	{
diff --git a/Engine/Source/Runtime/Slate/Private/Widgets/Views/SListPanel.h b/Engine/Source/Runtime/Slate/Private/Widgets/Views/SListPanel.h
--- a/Engine/Source/Runtime/Slate/Private/Widgets/Views/SListPanel.h
+++ b/Engine/Source/Runtime/Slate/Private/Widgets/Views/SListPanel.h
@@ -155,10 +155,4 @@ protected:
 
 	/** The preferred number of lines that this widget should have orthogonal to the scroll axis. Only relevant for tile views. */
 	int32 PreferredNumLines = 1;
-
-
-private:
-	
-	/** Used to pretend that the panel has no children when asked to cache desired size. */
-	static FNoChildren NoChildren;
 };
diff --git a/Engine/Source/Runtime/Slate/Public/Widgets/Input/SMenuAnchor.h b/Engine/Source/Runtime/Slate/Public/Widgets/Input/SMenuAnchor.h
--- a/Engine/Source/Runtime/Slate/Public/Widgets/Input/SMenuAnchor.h
+++ b/Engine/Source/Runtime/Slate/Public/Widgets/Input/SMenuAnchor.h
@@ -234,5 +234,5 @@ protected:
 	FVector2D ScreenPopupPosition;
 
 	/** The currently arranged children in the menu anchor.  Changes as the opened/closed state of the widget changes. */
-	TPanelChildren<FSimpleSlot> Children;
+	TPanelChildren<FBasicLayoutWidgetSlot> Children;
 };
diff --git a/Engine/Source/Runtime/Slate/Public/Widgets/Input/SSegmentedControl.h b/Engine/Source/Runtime/Slate/Public/Widgets/Input/SSegmentedControl.h
--- a/Engine/Source/Runtime/Slate/Public/Widgets/Input/SSegmentedControl.h
+++ b/Engine/Source/Runtime/Slate/Public/Widgets/Input/SSegmentedControl.h
@@ -29,11 +29,11 @@ public:
 
 	/** Stores the per-child info for this panel type */
 	template< typename SlotOptionType >
-	struct FSlot : public TSlotBase<FSlot<SlotOptionType>>, TSupportsContentAlignmentMixin<FSlot<SlotOptionType>>
+	struct FSlot : public TSlotBase<FSlot<SlotOptionType>>, TAlignmentWidgetSlotMixin<FSlot<SlotOptionType>>
 	{
 		FSlot(const SlotOptionType& InValue) 
 		: TSlotBase<FSlot<SlotOptionType>>()
-		, TSupportsContentAlignmentMixin<FSlot<SlotOptionType>>(HAlign_Center, VAlign_Fill)
+		, TAlignmentWidgetSlotMixin<FSlot<SlotOptionType>>(HAlign_Center, VAlign_Fill)
 		, _Text()
 		, _Tooltip()
 		, _Icon(nullptr)
diff --git a/Engine/Source/Runtime/Slate/Public/Widgets/Layout/SBox.h b/Engine/Source/Runtime/Slate/Public/Widgets/Layout/SBox.h
--- a/Engine/Source/Runtime/Slate/Public/Widgets/Layout/SBox.h
+++ b/Engine/Source/Runtime/Slate/Public/Widgets/Layout/SBox.h
@@ -24,11 +24,10 @@ class SLATE_API SBox : public SPanel
 	SLATE_DECLARE_WIDGET(SBox, SPanel)
 
 public:
-
-	class UE_DEPRECATED(5.0, "FBoxSlot is deprecated. Use FSimpleSlot or FOneSimpleMemberChild")
-	FBoxSlot : public FSimpleSlot
+	class UE_DEPRECATED(5.0, "FBoxSlot is deprecated. Use FSingleWidgetChildrenWithBasicLayoutSlot or FOneSimpleMemberChild")
+	FBoxSlot : public FSingleWidgetChildrenWithBasicLayoutSlot
 	{
-		using FSimpleSlot::FSimpleSlot;
+		using FSingleWidgetChildrenWithBasicLayoutSlot::FSingleWidgetChildrenWithBasicLayoutSlot;
 	};
 
 	SLATE_BEGIN_ARGS(SBox)
diff --git a/Engine/Source/Runtime/Slate/Public/Widgets/Layout/SGridPanel.h b/Engine/Source/Runtime/Slate/Public/Widgets/Layout/SGridPanel.h
--- a/Engine/Source/Runtime/Slate/Public/Widgets/Layout/SGridPanel.h
+++ b/Engine/Source/Runtime/Slate/Public/Widgets/Layout/SGridPanel.h
@@ -32,13 +32,12 @@ public:
 		int TheLayer;
 	};
 
-	class FSlot : public TSlotBase<FSlot>, public TSupportsContentAlignmentMixin<FSlot>, public TSupportsContentPaddingMixin<FSlot>
+	class FSlot : public TBasicLayoutWidgetSlot<FSlot>
 	{
 		public:
 			/** Default values for a slot. */
 			FSlot( int32 Column, int32 Row, int32 InLayer )
-				: TSlotBase<FSlot>()
-				, TSupportsContentAlignmentMixin<FSlot>(HAlign_Fill, VAlign_Fill)
+				: TBasicLayoutWidgetSlot<FSlot>(HAlign_Fill, VAlign_Fill)
 				, ColumnParam( Column )
 				, ColumnSpanParam( 1 )
 				, RowParam( Row )
diff --git a/Engine/Source/Runtime/Slate/Public/Widgets/Layout/SResponsiveGridPanel.h b/Engine/Source/Runtime/Slate/Public/Widgets/Layout/SResponsiveGridPanel.h
--- a/Engine/Source/Runtime/Slate/Public/Widgets/Layout/SResponsiveGridPanel.h
+++ b/Engine/Source/Runtime/Slate/Public/Widgets/Layout/SResponsiveGridPanel.h
@@ -38,7 +38,7 @@ class SLATE_API SResponsiveGridPanel
 {
 public:
 
-	class FSlot : public TSlotBase<FSlot>, public TSupportsContentAlignmentMixin<FSlot>, public TSupportsContentPaddingMixin < FSlot >
+	class FSlot : public TBasicLayoutWidgetSlot< FSlot >
 	{
 	public:
 
@@ -52,8 +52,7 @@ public:
 	public:
 		/** Default values for a slot. */
 		FSlot(int32 Row)
-			: TSlotBase<FSlot>()
-			, TSupportsContentAlignmentMixin<FSlot>(HAlign_Fill, VAlign_Fill)
+			: TBasicLayoutWidgetSlot<FSlot>(HAlign_Fill, VAlign_Fill)
 			, RowParam(FMath::Max(0, Row))
 			, ColumnLayouts()
 		{ }
diff --git a/Engine/Source/Runtime/Slate/Public/Widgets/Layout/SScissorRectBox.h b/Engine/Source/Runtime/Slate/Public/Widgets/Layout/SScissorRectBox.h
--- a/Engine/Source/Runtime/Slate/Public/Widgets/Layout/SScissorRectBox.h
+++ b/Engine/Source/Runtime/Slate/Public/Widgets/Layout/SScissorRectBox.h
@@ -31,13 +31,7 @@ class SLATE_API SScissorRectBox : public SPanel
 {
 
 public:
-	class FScissorRectSlot : public TSupportsOneChildMixin<FScissorRectSlot>
-	{
-	public:
-		FScissorRectSlot(SWidget* InOwner)
-			: TSupportsOneChildMixin<FScissorRectSlot>(InOwner)
-		{}
-	};
+	using FScissorRectSlot = FSingleWidgetChildrenWithSlot;
 
 	SLATE_BEGIN_ARGS(SScissorRectBox)
 	{
diff --git a/Engine/Source/Runtime/Slate/Public/Widgets/Layout/SScrollBox.h b/Engine/Source/Runtime/Slate/Public/Widgets/Layout/SScrollBox.h
--- a/Engine/Source/Runtime/Slate/Public/Widgets/Layout/SScrollBox.h
+++ b/Engine/Source/Runtime/Slate/Public/Widgets/Layout/SScrollBox.h
@@ -69,34 +69,7 @@ class SLATE_API SScrollBox : public SCompoundWidget
 {
 public:
 	/** A Slot that provides layout options for the contents of a scrollable box. */
-	class SLATE_API FSlot : public TSlotBase<FSlot>, public TSupportsContentPaddingMixin<FSlot>
-	{
-	public:
-		FSlot()
-			: TSlotBase<FSlot>()
-			// Set both vertical and horizontal alignment to fill by default.  During layout, the
-			// alignment direction parallel to the scroll direction is assumed to be top, or left
-			// since that is how the items are stacked.
-			, HAlignment(HAlign_Fill)
-			, VAlignment(VAlign_Fill)
-		{
-		}
-
-		FSlot& HAlign( EHorizontalAlignment InHAlignment )
-		{
-			HAlignment = InHAlignment;
-			return *this;
-		}
-
-		FSlot& VAlign(EVerticalAlignment InVAlignment)
-		{
-			VAlignment = InVAlignment;
-			return *this;
-		}
-		
-		EHorizontalAlignment HAlignment;
-		EVerticalAlignment VAlignment;
-	};
+	using FSlot = FBasicLayoutWidgetSlot;
 
 	SLATE_BEGIN_ARGS(SScrollBox)
 		: _Style( &FAppStyle::Get().GetWidgetStyle<FScrollBoxStyle>("ScrollBox") )
diff --git a/Engine/Source/Runtime/Slate/Public/Widgets/Layout/SUniformGridPanel.h b/Engine/Source/Runtime/Slate/Public/Widgets/Layout/SUniformGridPanel.h
--- a/Engine/Source/Runtime/Slate/Public/Widgets/Layout/SUniformGridPanel.h
+++ b/Engine/Source/Runtime/Slate/Public/Widgets/Layout/SUniformGridPanel.h
@@ -19,11 +19,11 @@ class SLATE_API SUniformGridPanel : public SPanel
 {
 public:
 	/** Stores the per-child info for this panel type */
-	struct FSlot : public TSlotBase<FSlot>, public TSupportsContentAlignmentMixin<FSlot>
+	struct FSlot : public TSlotBase<FSlot>, public TAlignmentWidgetSlotMixin<FSlot>
 	{
 		FSlot( int32 InColumn, int32 InRow )
 		: TSlotBase<FSlot>()
-		, TSupportsContentAlignmentMixin<FSlot>(HAlign_Fill, VAlign_Fill)
+		, TAlignmentWidgetSlotMixin<FSlot>(HAlign_Fill, VAlign_Fill)
 		, Column( InColumn )
 		, Row( InRow )
 		{
diff --git a/Engine/Source/Runtime/Slate/Public/Widgets/Layout/SUniformWrapPanel.h b/Engine/Source/Runtime/Slate/Public/Widgets/Layout/SUniformWrapPanel.h
--- a/Engine/Source/Runtime/Slate/Public/Widgets/Layout/SUniformWrapPanel.h
+++ b/Engine/Source/Runtime/Slate/Public/Widgets/Layout/SUniformWrapPanel.h
@@ -19,11 +19,11 @@ class SLATE_API SUniformWrapPanel : public SPanel
 {
 public:
 	/** Stores the per-child info for this panel type */
-	struct FSlot : public TSlotBase<FSlot>, public TSupportsContentAlignmentMixin<FSlot>
+	struct FSlot : public TSlotBase<FSlot>, public TAlignmentWidgetSlotMixin<FSlot>
 	{
 		FSlot()
 		: TSlotBase<FSlot>()
-		, TSupportsContentAlignmentMixin<FSlot>(HAlign_Fill, VAlign_Fill)
+		, TAlignmentWidgetSlotMixin<FSlot>(HAlign_Fill, VAlign_Fill)
 		{
 
 		}
diff --git a/Engine/Source/Runtime/Slate/Public/Widgets/Layout/SWidgetSwitcher.h b/Engine/Source/Runtime/Slate/Public/Widgets/Layout/SWidgetSwitcher.h
--- a/Engine/Source/Runtime/Slate/Public/Widgets/Layout/SWidgetSwitcher.h
+++ b/Engine/Source/Runtime/Slate/Public/Widgets/Layout/SWidgetSwitcher.h
@@ -23,18 +23,7 @@ class SLATE_API SWidgetSwitcher
 {
 public:
 
-	class FSlot
-		: public TSlotBase<FSlot>
-		, public TSupportsContentAlignmentMixin<FSlot>
-		, public TSupportsContentPaddingMixin<FSlot>
-	{
-	public:
-
-		FSlot()
-			: TSlotBase<FSlot>()
-			,TSupportsContentAlignmentMixin<FSlot>( HAlign_Fill, VAlign_Fill )
-		{ }
-	};
+	using FSlot = FBasicLayoutWidgetSlot;
 
 	SLATE_BEGIN_ARGS(SWidgetSwitcher)
 		: _WidgetIndex(0)
diff --git a/Engine/Source/Runtime/Slate/Public/Widgets/Layout/SWindowTitleBarArea.h b/Engine/Source/Runtime/Slate/Public/Widgets/Layout/SWindowTitleBarArea.h
--- a/Engine/Source/Runtime/Slate/Public/Widgets/Layout/SWindowTitleBarArea.h
+++ b/Engine/Source/Runtime/Slate/Public/Widgets/Layout/SWindowTitleBarArea.h
@@ -25,14 +25,7 @@ class SLATE_API SWindowTitleBarArea
 {
 public:
 	
-	class SLATE_API FSlot : public TSupportsOneChildMixin<FSlot>, public TSupportsContentAlignmentMixin<FSlot>, public TSupportsContentPaddingMixin<FSlot>
-	{
-		public:
-			FSlot(SWidget* InOwner)
-			: TSupportsOneChildMixin<FSlot>(InOwner)
-			, TSupportsContentAlignmentMixin<FSlot>(HAlign_Fill, VAlign_Fill)
-			{ }
-	};
+	using FSlot = FSingleWidgetChildrenWithBasicLayoutSlot;
 
 	SLATE_BEGIN_ARGS(SWindowTitleBarArea)
 		: _HAlign( HAlign_Fill )
diff --git a/Engine/Source/Runtime/Slate/Public/Widgets/Layout/SWrapBox.h b/Engine/Source/Runtime/Slate/Public/Widgets/Layout/SWrapBox.h
--- a/Engine/Source/Runtime/Slate/Public/Widgets/Layout/SWrapBox.h
+++ b/Engine/Source/Runtime/Slate/Public/Widgets/Layout/SWrapBox.h
@@ -52,12 +52,11 @@ class SLATE_API SWrapBox : public SPanel
 public:
 
 	/** A slot that support alignment of content and padding */
-	class FSlot : public TSlotBase<FSlot>, public TSupportsContentAlignmentMixin<FSlot>, public TSupportsContentPaddingMixin<FSlot>
+	class FSlot : public TBasicLayoutWidgetSlot<FSlot>
 	{
 	public:
 		FSlot()
-			: TSlotBase<FSlot>()
-			, TSupportsContentAlignmentMixin<FSlot>(HAlign_Fill, VAlign_Fill)
+			: TBasicLayoutWidgetSlot<FSlot>(HAlign_Fill, VAlign_Fill)
 			, SlotFillLineWhenWidthLessThan()
 			, SlotFillLineWhenSizeLessThan()
 			, bSlotFillEmptySpace(false)
diff --git a/Engine/Source/Runtime/Slate/Public/Widgets/SInvalidationPanel.h b/Engine/Source/Runtime/Slate/Public/Widgets/SInvalidationPanel.h
--- a/Engine/Source/Runtime/Slate/Public/Widgets/SInvalidationPanel.h
+++ b/Engine/Source/Runtime/Slate/Public/Widgets/SInvalidationPanel.h
@@ -67,8 +67,6 @@ private:
 	bool UpdateCachePrequisites(FSlateWindowElementList& OutDrawElements, const FGeometry& AllottedGeometry, const FSlateRect& MyCullingRect, int32 LayerId, const FWidgetStyle& InWidgetStyle) const;
 
 private:
-	FSimpleSlot EmptyChildSlot;
-
 	mutable TSharedRef<FHittestGrid> HittestGrid;
 
 	mutable TOptional<FSlateClippingState> LastClippingState;
diff --git a/Engine/Source/Runtime/Slate/Public/Widgets/Views/STableRow.h b/Engine/Source/Runtime/Slate/Public/Widgets/Views/STableRow.h
--- a/Engine/Source/Runtime/Slate/Public/Widgets/Views/STableRow.h
+++ b/Engine/Source/Runtime/Slate/Public/Widgets/Views/STableRow.h
@@ -239,7 +239,7 @@ public:
 				InContent
 			];
 
-			InnerContentSlot = InnerContentSlotNativePtr;
+			InnerContentSlot = &InnerContentSlotNativePtr->AsSlot();
 		}
 		else
 		{
diff --git a/Engine/Source/Runtime/SlateCore/Private/Children.cpp b/Engine/Source/Runtime/SlateCore/Private/Children.cpp
--- a/Engine/Source/Runtime/SlateCore/Private/Children.cpp
+++ b/Engine/Source/Runtime/SlateCore/Private/Children.cpp
@@ -2,3 +2,13 @@
 
 #include "Layout/Children.h"
 
+#include "Widgets/SNullWidget.h"
+
+FNoChildren FNoChildren::NoChildrenInstance(&SNullWidget::NullWidget.Get());
+
+FNoChildren::FNoChildren()
+	: FChildren(&SNullWidget::NullWidget.Get())
+{
+
+}
+
diff --git a/Engine/Source/Runtime/SlateCore/Private/Widgets/SLeafWidget.cpp b/Engine/Source/Runtime/SlateCore/Private/Widgets/SLeafWidget.cpp
--- a/Engine/Source/Runtime/SlateCore/Private/Widgets/SLeafWidget.cpp
+++ b/Engine/Source/Runtime/SlateCore/Private/Widgets/SLeafWidget.cpp
@@ -3,12 +3,6 @@
 #include "Widgets/SLeafWidget.h"
 
 
-/* Static initialization
- *****************************************************************************/
-
-FNoChildren SLeafWidget::NoChildrenInstance;
-
-
 /* SLeafWidget interface
  *****************************************************************************/
 
@@ -19,7 +13,7 @@ void SLeafWidget::SetVisibility( TAttribute<EVisibility> InVisibility )
 
 FChildren* SLeafWidget::GetChildren( )
 {
-	return &SLeafWidget::NoChildrenInstance;
+	return &FNoChildren::NoChildrenInstance;
 }
 
 
diff --git a/Engine/Source/Runtime/SlateCore/Private/Widgets/SNullWidget.cpp b/Engine/Source/Runtime/SlateCore/Private/Widgets/SNullWidget.cpp
--- a/Engine/Source/Runtime/SlateCore/Private/Widgets/SNullWidget.cpp
+++ b/Engine/Source/Runtime/SlateCore/Private/Widgets/SNullWidget.cpp
@@ -8,8 +8,6 @@
 #include "Widgets/DeclarativeSyntaxSupport.h"
 
 
-FNoChildren NullWidgetNoChildren;
-
 class SLATECORE_API SNullWidgetContent
 	: public SWidget
 {
@@ -48,7 +46,7 @@ public:
 
 	virtual FChildren* GetChildren( ) override final
 	{
-		return &NullWidgetNoChildren;
+		return &FNoChildren::NoChildrenInstance;
 	}
 
 	virtual void OnArrangeChildren( const FGeometry& AllottedGeometry, FArrangedChildren& ArrangedChildren ) const override final
diff --git a/Engine/Source/Runtime/SlateCore/Public/Layout/BasicLayoutWidgetSlot.h b/Engine/Source/Runtime/SlateCore/Public/Layout/BasicLayoutWidgetSlot.h
new file mode 100644
--- /dev/null
+++ b/Engine/Source/Runtime/SlateCore/Public/Layout/BasicLayoutWidgetSlot.h
@@ -0,0 +1,124 @@
+// Copyright Epic Games, Inc. All Rights Reserved.
+
+#pragma once
+
+#include "CoreMinimal.h"
+#include "SlotBase.h"
+#include "Layout/FlowDirection.h"
+#include "Layout/Margin.h"
+#include "Misc/Optional.h"
+
+/** Mixin to add the alignment functionality to a base slot. */
+template <typename MixedIntoType>
+class TAlignmentWidgetSlotMixin
+{
+public:
+	TAlignmentWidgetSlotMixin()
+		: HAlignment(HAlign_Fill)
+		, VAlignment(VAlign_Fill)
+	{}
+
+	TAlignmentWidgetSlotMixin(const EHorizontalAlignment InHAlign, const EVerticalAlignment InVAlign)
+		: HAlignment(InHAlign)
+		, VAlignment(InVAlign)
+	{}
+
+public:
+	MixedIntoType& HAlign(EHorizontalAlignment InHAlignment)
+	{
+		HAlignment = InHAlignment;
+		return *(static_cast<MixedIntoType*>(this));
+	}
+
+	MixedIntoType& VAlign(EVerticalAlignment InVAlignment)
+	{
+		VAlignment = InVAlignment;
+		return *(static_cast<MixedIntoType*>(this));
+	}
+
+public:
+	/** Horizontal positioning of child within the allocated slot */
+	EHorizontalAlignment HAlignment;
+	/** Vertical positioning of child within the allocated slot */
+	EVerticalAlignment VAlignment;
+};
+
+
+PRAGMA_DISABLE_DEPRECATION_WARNINGS
+/** Mixin to add the padding functionality to a base slot. */
+template <typename MixedIntoType>
+class TPaddingWidgetSlotMixin
+{
+public:
+	TPaddingWidgetSlotMixin() = default;
+	TPaddingWidgetSlotMixin(const FMargin& Margin)
+		: SlotPadding(Margin)
+	{}
+
+public:
+	MixedIntoType& Padding(TAttribute<FMargin> InPadding)
+	{
+		SlotPadding = MoveTemp(InPadding);
+		return *(static_cast<MixedIntoType*>(this));
+	}
+
+	MixedIntoType& Padding(float Uniform)
+	{
+		SlotPadding = FMargin(Uniform);
+		return *(static_cast<MixedIntoType*>(this));
+	}
+
+	MixedIntoType& Padding(float Horizontal, float Vertical)
+	{
+		SlotPadding = FMargin(Horizontal, Vertical);
+		return *(static_cast<MixedIntoType*>(this));
+	}
+
+	MixedIntoType& Padding(float Left, float Top, float Right, float Bottom)
+	{
+		SlotPadding = FMargin(Left, Top, Right, Bottom);
+		return *(static_cast<MixedIntoType*>(this));
+	}
+
+public:
+	void SetPadding(TAttribute<FMargin> InPadding)
+	{
+		SlotPadding = MoveTemp(InPadding);
+	}
+
+	const FMargin& GetPadding() const { return SlotPadding.Get(); }
+
+public:
+	UE_DEPRECATED(5.0, "Direct access to SlotPadding is now deprecated. Use the setter or getter.")
+	TAttribute<FMargin> SlotPadding;
+};
+PRAGMA_ENABLE_DEPRECATION_WARNINGS
+
+
+/** A templated basic slot that can be used by layout. */
+template <typename SlotType>
+class TBasicLayoutWidgetSlot : public TSlotBase<SlotType>
+	, public TPaddingWidgetSlotMixin<SlotType>
+	, public TAlignmentWidgetSlotMixin<SlotType>
+{
+public:
+	TBasicLayoutWidgetSlot()
+		: TSlotBase<SlotType>()
+		, TPaddingWidgetSlotMixin<SlotType>()
+		, TAlignmentWidgetSlotMixin<SlotType>()
+	{}
+
+	TBasicLayoutWidgetSlot(const EHorizontalAlignment InHAlign, const EVerticalAlignment InVAlign)
+		: TSlotBase<SlotType>()
+		, TPaddingWidgetSlotMixin<SlotType>()
+		, TAlignmentWidgetSlotMixin<SlotType>(InHAlign, InVAlign)
+	{}
+};
+
+
+/** The basic slot that can be used by layout. */
+class FBasicLayoutWidgetSlot : public TBasicLayoutWidgetSlot<FBasicLayoutWidgetSlot>
+{
+public:
+	using TBasicLayoutWidgetSlot<FBasicLayoutWidgetSlot>::TBasicLayoutWidgetSlot;
+};
diff --git a/Engine/Source/Runtime/SlateCore/Public/Layout/Children.h b/Engine/Source/Runtime/SlateCore/Public/Layout/Children.h
--- a/Engine/Source/Runtime/SlateCore/Public/Layout/Children.h
+++ b/Engine/Source/Runtime/SlateCore/Public/Layout/Children.h
@@ -3,12 +3,10 @@
 #pragma once
 
 #include "CoreMinimal.h"
-#include "Misc/Attribute.h"
-#include "Layout/Margin.h"
 #include "Widgets/SNullWidget.h"
 #include "SlotBase.h"
+#include "Layout/BasicLayoutWidgetSlot.h"
 #include "Widgets/SWidget.h"
-#include "FlowDirection.h"
 
 /**
  * FChildren is an interface that must be implemented by all child containers.
@@ -16,7 +14,6 @@
  * the underlying Widget happens to store its children.
  * 
  * FChildren is intended to be returned by the GetChildren() method.
- * 
  */
 PRAGMA_DISABLE_DEPRECATION_WARNINGS
 class SLATECORE_API FChildren
@@ -25,8 +22,15 @@ public:
 	FChildren(SWidget* InOwner)
 		: Owner(InOwner)
 	{
+		check(InOwner);
 	}
 
+	FChildren(std::nullptr_t) = delete;
+
+	//~ Prevents allocation of FChilren. It created a confusion between FSlot and FChildren
+	void* operator new (size_t) = delete;
+	void* operator new[](size_t) = delete;
+
 	/** @return the number of children */
 	virtual int32 Num() const = 0;
 	/** @return pointer to the Widget at the specified Index. */
@@ -34,6 +38,8 @@ public:
 	/** @return const pointer to the Widget at the specified Index. */
 	virtual TSharedRef<const SWidget> GetChildAt( int32 Index ) const = 0;
 
+	SWidget& GetOwner() const { return *Owner; }
+
 protected:
 	friend class SWidget;
 	friend class FCombinedChildren;
@@ -43,9 +49,6 @@ protected:
 protected:
 	virtual ~FChildren(){}
 
-	SWidget* GetOwner() { return Owner; }
-	const SWidget* GetOwner() const { return Owner; }
-
 protected:
 	UE_DEPRECATED(5.0, "Direct access to Owner is now deprecated. Use the getter.")
 	SWidget* Owner;
@@ -54,17 +57,14 @@ PRAGMA_ENABLE_DEPRECATION_WARNINGS
 
 
 /**
- * Occasionally you may need to keep multiple descrete sets of children with differing slot requirements.
- * This datastructure can be used to link multiple FChildren under a single accessor so you can always return
+ * Occasionally you may need to keep multiple discrete sets of children with differing slot requirements.
+ * This data structure can be used to link multiple FChildren under a single accessor so you can always return
  * all children from GetChildren, but internally manage them in their own child lists.
  */
 class FCombinedChildren : public FChildren
 {
 public:
-	FCombinedChildren(SWidget* InOwner)
-		: FChildren(InOwner)
-	{
-	}
+	using FChildren::FChildren;
 
 	void AddChildren(FChildren& InLinkedChildren)
 	{
@@ -118,14 +118,21 @@ protected:
 
 /**
  * Widgets with no Children can return an instance of FNoChildren.
- * For convenience a shared instance SWidget::NoChildrenInstance can be used.
+ * For convenience a shared instance FNoChildren::NoChildrenInstance can be used.
  */
 class SLATECORE_API FNoChildren : public FChildren
 {
 public:
-	FNoChildren()
-		: FChildren(nullptr)
-	{}
+	static FNoChildren NoChildrenInstance;
+
+public:
+	UE_DEPRECATED(5.0, "FNoChildren take a valid reference to a SWidget")
+	FNoChildren();
+
+	FNoChildren(SWidget* InOwner)
+		: FChildren(InOwner)
+	{
+	}
 
 	virtual int32 Num() const override { return 0; }
 	
@@ -155,11 +162,13 @@ private:
 	}
 };
 
+
 /**
- * Widgets that will only have one child can return an instance of FOneChild.
+ * Widgets that will only have one child.
  */
 template <typename MixedIntoType>
-class TSupportsOneChildMixin : public FChildren, public TSlotBase<MixedIntoType>
+class UE_DEPRECATED(5.0, "TSupportsOneChildMixin is deprecated because it got confused between FSlot and FChildren. Use FSingleWidgetChildren.")
+TSupportsOneChildMixin : public FChildren, public TSlotBase<MixedIntoType>
 {
 public:
 	TSupportsOneChildMixin(SWidget* InOwner)
@@ -169,6 +178,8 @@ public:
 		this->RawParentPtr = InOwner;
 	}
 
+	TSupportsOneChildMixin(std::nullptr_t) = delete;
+
 	virtual int32 Num() const override { return 1; }
 
 	virtual TSharedRef<SWidget> GetChildAt( int32 ChildIndex ) override
@@ -187,6 +198,7 @@ private:
 	virtual const FSlotBase& GetSlotAt(int32 ChildIndex) const override { check(ChildIndex == 0); return *this; }
 };
 
+
 /**
  * For widgets that do not own their content, but are responsible for presenting someone else's content.
  * e.g. Tooltips are just presented by the owner window; not actually owned by it. They can go away at any time
@@ -196,11 +208,7 @@ template <typename ChildType>
 class TWeakChild : public FChildren
 {
 public:
-	TWeakChild(SWidget* InOwner)
-		: FChildren(InOwner)
-		, WidgetPtr()
-	{
-	}
+	using FChildren::FChildren;
 
 	virtual int32 Num() const override { return WidgetPtr.IsValid() ? 1 : 0 ; }
 
@@ -223,14 +231,11 @@ public:
 	void AttachWidget(const TSharedPtr<SWidget>& InWidget)
 	{
 		WidgetPtr = InWidget;
-		if (GetOwner()) 
-		{ 
-			GetOwner()->Invalidate(EInvalidateWidgetReason::ChildOrder);
+		GetOwner().Invalidate(EInvalidateWidgetReason::ChildOrder);
 
-			if (InWidget.IsValid() && InWidget != SNullWidget::NullWidget)
-			{
-				InWidget->AssignParentWidget(GetOwner()->AsShared());
-			}
+		if (InWidget.IsValid() && InWidget != SNullWidget::NullWidget)
+		{
+			InWidget->AssignParentWidget(GetOwner().AsShared());
 		}
 	}
 
@@ -242,7 +247,7 @@ public:
 
 			if (Widget != SNullWidget::NullWidget)
 			{
-				Widget->ConditionallyDetatchParentWidget(GetOwner());
+				Widget->ConditionallyDetatchParentWidget(&GetOwner());
 			}
 
 			WidgetPtr.Reset();
@@ -268,86 +273,119 @@ private:
 	TWeakPtr<ChildType> WidgetPtr;
 };
 
+
 template <typename MixedIntoType>
-class TSupportsContentAlignmentMixin
+class UE_DEPRECATED(5.0, "Renamed TSupportsContentPaddingMixin to TAlignmentWidgetSlotMixin to differenciate from FSlot and FChildren.")
+TSupportsContentAlignmentMixin : public TAlignmentWidgetSlotMixin<MixedIntoType>
+{
+	using TAlignmentWidgetSlotMixin<MixedIntoType>::TAlignmentWidgetSlotMixin;
+};
+
+template <typename MixedIntoType>
+class UE_DEPRECATED(5.0, "Renamed TSupportsContentPaddingMixin to TPaddingWidgetSlotMixin to differenciate from FSlot and FChildren.")
+TSupportsContentPaddingMixin : public TPaddingWidgetSlotMixin<MixedIntoType>
+{
+	using TPaddingWidgetSlotMixin<MixedIntoType>::TPaddingWidgetSlotMixin;
+};
+
+
+/** A FChildren that has only one child and can take a templated slot. */
+template<typename SlotType>
+class TSingleWidgetChildrenWithSlot : public FChildren, protected TSlotBase<SlotType>
 {
 public:
-	TSupportsContentAlignmentMixin(const EHorizontalAlignment InHAlign, const EVerticalAlignment InVAlign)
-	: HAlignment( InHAlign )
-	, VAlignment( InVAlign )
+	TSingleWidgetChildrenWithSlot(SWidget* InOwner)
+		: FChildren(InOwner)
+		, TSlotBase<SlotType>()
 	{
-		
+		this->RawParentPtr = InOwner;
 	}
 
-	MixedIntoType& HAlign( EHorizontalAlignment InHAlignment )
+	TSingleWidgetChildrenWithSlot(std::nullptr_t) = delete;
+
+public:
+	virtual int32 Num() const override { return 1; }
+
+	virtual TSharedRef<SWidget> GetChildAt(int32 ChildIndex) override
 	{
-		HAlignment = InHAlignment;
-		return *(static_cast<MixedIntoType*>(this));
+		check(ChildIndex == 0);
+		return this->GetWidget();
 	}
 
-	MixedIntoType& VAlign( EVerticalAlignment InVAlignment )
+	virtual TSharedRef<const SWidget> GetChildAt(int32 ChildIndex) const override
 	{
-		VAlignment = InVAlignment;
-		return *(static_cast<MixedIntoType*>(this));
+		check(ChildIndex == 0);
+		return this->GetWidget();
 	}
-	
-	EHorizontalAlignment HAlignment;
-	EVerticalAlignment VAlignment;
-};
 
-PRAGMA_DISABLE_DEPRECATION_WARNINGS
-template <typename MixedIntoType>
-class TSupportsContentPaddingMixin
-{
 public:
-	MixedIntoType& Padding( TAttribute<FMargin> InPadding )
-	{
-		SlotPadding = MoveTemp(InPadding);
-		return *(static_cast<MixedIntoType*>(this));
-	}
+	TSlotBase<SlotType>& AsSlot() { return *this; }
+	const TSlotBase<SlotType>& AsSlot() const { return *this; }
 
-	MixedIntoType& Padding( float Uniform )
+	using TSlotBase<SlotType>::AttachWidget;
+	using TSlotBase<SlotType>::DetachWidget;
+	using TSlotBase<SlotType>::GetWidget;
+	SlotType& operator[](const TSharedRef<SWidget>& InChildWidget)
 	{
-		SlotPadding = FMargin(Uniform);
-		return *(static_cast<MixedIntoType*>(this));
+		this->AttachWidget(InChildWidget);
+		return static_cast<SlotType&>(*this);
 	}
 
-	MixedIntoType& Padding( float Horizontal, float Vertical )
+	SlotType& Expose(SlotType*& OutVarToInit)
 	{
-		SlotPadding = FMargin(Horizontal, Vertical);
-		return *(static_cast<MixedIntoType*>(this));
+		OutVarToInit = static_cast<SlotType*>(this);
+		return static_cast<SlotType&>(*this);
 	}
 
-	MixedIntoType& Padding( float Left, float Top, float Right, float Bottom )
+private:
+	virtual const FSlotBase& GetSlotAt(int32 ChildIndex) const override
 	{
-		SlotPadding = FMargin(Left, Top, Right, Bottom);
-		return *(static_cast<MixedIntoType*>(this));
+		check(ChildIndex == 0);
+		return *this;
 	}
+};
 
-	void SetPadding(TAttribute<FMargin> InPadding)
-	{
-		SlotPadding = MoveTemp(InPadding);
-	}
-	const FMargin& GetPadding() const { return SlotPadding.Get(); }
 
+/** A FChildren that has only one child. */
+class FSingleWidgetChildrenWithSlot : public TSingleWidgetChildrenWithSlot<FSingleWidgetChildrenWithSlot>
+{
 public:
-	UE_DEPRECATED(5.0, "Direct access to SlotPadding is now deprecated. Use the getter.")
-	TAttribute< FMargin > SlotPadding;
+	using TSingleWidgetChildrenWithSlot<FSingleWidgetChildrenWithSlot>::TSingleWidgetChildrenWithSlot;
 };
-PRAGMA_ENABLE_DEPRECATION_WARNINGS
 
-/** A slot that support alignment of content and padding */
-class SLATECORE_API FSimpleSlot : public TSupportsOneChildMixin<FSimpleSlot>, public TSupportsContentAlignmentMixin<FSimpleSlot>, public TSupportsContentPaddingMixin<FSimpleSlot>
+
+/** A FChildren that has only one child and support alignment and padding. */
+class FSingleWidgetChildrenWithBasicLayoutSlot : public TSingleWidgetChildrenWithSlot<FSingleWidgetChildrenWithBasicLayoutSlot>
+	, public TPaddingWidgetSlotMixin<FSingleWidgetChildrenWithBasicLayoutSlot>
+	, public TAlignmentWidgetSlotMixin<FSingleWidgetChildrenWithBasicLayoutSlot>
 {
 public:
-	FSimpleSlot(SWidget* InParent)
-	: TSupportsOneChildMixin<FSimpleSlot>(InParent)
-	, TSupportsContentAlignmentMixin<FSimpleSlot>(HAlign_Fill, VAlign_Fill)
+	FSingleWidgetChildrenWithBasicLayoutSlot(SWidget* InOwner)
+		: TSingleWidgetChildrenWithSlot<FSingleWidgetChildrenWithBasicLayoutSlot>(InOwner)
+		, TPaddingWidgetSlotMixin<FSingleWidgetChildrenWithBasicLayoutSlot>()
+		, TAlignmentWidgetSlotMixin<FSingleWidgetChildrenWithBasicLayoutSlot>(HAlign_Fill, VAlign_Fill)
+	{
+	}
+
+	FSingleWidgetChildrenWithBasicLayoutSlot(SWidget* InOwner, const EHorizontalAlignment InHAlign, const EVerticalAlignment InVAlign)
+		: TSingleWidgetChildrenWithSlot<FSingleWidgetChildrenWithBasicLayoutSlot>(InOwner)
+		, TPaddingWidgetSlotMixin<FSingleWidgetChildrenWithBasicLayoutSlot>()
+		, TAlignmentWidgetSlotMixin<FSingleWidgetChildrenWithBasicLayoutSlot>(InHAlign, InVAlign)
 	{
 	}
+
+	FSingleWidgetChildrenWithBasicLayoutSlot(std::nullptr_t) = delete;
+	FSingleWidgetChildrenWithBasicLayoutSlot(std::nullptr_t, const EHorizontalAlignment InHAlign, const EVerticalAlignment InVAlign) = delete;
 };
 
 
+/** A slot that support alignment of content and padding */
+class UE_DEPRECATED(5.0, "FSimpleSlot is deprecated because it got confused from FChildren with FSlot. Use FSingleWidgetChildrenWithSimpleSlot.")
+FSimpleSlot : public FSingleWidgetChildrenWithBasicLayoutSlot
+{
+public:
+	using FSingleWidgetChildrenWithBasicLayoutSlot::FSingleWidgetChildrenWithBasicLayoutSlot;
+};
 
 
 /**
@@ -367,13 +405,7 @@ private:
 	}
 
 public:
-	TPanelChildren(SWidget* InOwner)
-		: FChildren(InOwner)
-	{
-	}
-
-	//~ TPanelChildren should have a valid pointer
-	TPanelChildren(std::nullptr_t InOwner) = delete;
+	using FChildren::FChildren;
 	
 	virtual int32 Num() const override
 	{
@@ -393,11 +425,8 @@ public:
 	int32 Add( SlotType* Slot )
 	{
 		int32 Index = Children.Add(TUniquePtr<SlotType>(Slot));
-
-		if (GetOwner())
-		{
-			Slot->AttachWidgetParent(GetOwner());
-		}
+		check(Slot);
+		Slot->AttachWidgetParent(&GetOwner());
 
 		return Index;
 	}
@@ -449,13 +478,9 @@ public:
 
 	void Insert(SlotType* Slot, int32 Index)
 	{
+		check(Slot);
 		Children.Insert(TUniquePtr<SlotType>(Slot), Index);
-
-		// Don't do parent manipulation if this panel has no owner.
-		if (GetOwner())
-		{
-			Slot->AttachWidgetParent(GetOwner());
-		}
+		Slot->AttachWidgetParent(&GetOwner());
 	}
 
 	void Move(int32 IndexToMove, int32 IndexToDestination)
@@ -468,10 +493,7 @@ public:
 			Children.Insert(MoveTemp(SlotToMove), IndexToDestination);
 		}
 
-		if (GetOwner())
-		{
-			GetOwner()->Invalidate(EInvalidateWidgetReason::ChildOrder);
-		}
+		GetOwner().Invalidate(EInvalidateWidgetReason::ChildOrder);
 	}
 
 	void Reserve( int32 NumToReserve )
@@ -494,10 +516,7 @@ public:
 		{
 			return Predicate(*One, *Two);
 		});
-		if (GetOwner())
-		{
-			GetOwner()->Invalidate(EInvalidateWidgetReason::ChildOrder);
-		}
+		GetOwner().Invalidate(EInvalidateWidgetReason::ChildOrder);
 	}
 
 	template <class PREDICATE_CLASS>
@@ -507,19 +526,13 @@ public:
 		{
 			return Predicate(*One, *Two);
 		});
-		if (GetOwner())
-		{
-			GetOwner()->Invalidate(EInvalidateWidgetReason::ChildOrder);
-		}
+		GetOwner().Invalidate(EInvalidateWidgetReason::ChildOrder);
 	}
 
 	void Swap( int32 IndexA, int32 IndexB )
 	{
 		Children.Swap(IndexA, IndexB);
-		if (GetOwner())
-		{
-			GetOwner()->Invalidate(EInvalidateWidgetReason::ChildOrder);
-		}
+		GetOwner().Invalidate(EInvalidateWidgetReason::ChildOrder);
 	}
 };
 
@@ -665,6 +678,8 @@ public:
 	{
 	}
 
+	TSlotlessChildren(std::nullptr_t, bool InbChangesInvalidatePrepass = true) = delete;
+
 	virtual int32 Num() const override
 	{
 		return Children.Num();
@@ -682,19 +697,16 @@ public:
 
 	int32 Add( const TSharedRef<ChildType>& Child )
 	{
-		if (GetOwner() && bChangesInvalidatePrepass)
+		if (bChangesInvalidatePrepass)
 		{ 
-			GetOwner()->Invalidate(EInvalidateWidgetReason::ChildOrder);
+			GetOwner().Invalidate(EInvalidateWidgetReason::ChildOrder);
 		}
 
 		int32 Index = Children.Add(Child);
 
-		if (GetOwner())
+		if (Child != SNullWidget::NullWidget)
 		{
-			if (Child != SNullWidget::NullWidget)
-			{
-				Child->AssignParentWidget(GetOwner()->AsShared());
-			}
+			Child->AssignParentWidget(GetOwner().AsShared());
 		}
 
 		return Index;
@@ -707,7 +719,7 @@ public:
 			TSharedRef<SWidget> Child = GetChildAt(ChildIndex);
 			if (Child != SNullWidget::NullWidget)
 			{
-				Child->ConditionallyDetatchParentWidget(GetOwner());
+				Child->ConditionallyDetatchParentWidget(&GetOwner());
 			}
 		}
 
@@ -731,7 +743,7 @@ public:
 			TSharedRef<SWidget> Child = GetChildAt(ChildIndex);
 			if (Child != SNullWidget::NullWidget)
 			{
-				Child->ConditionallyDetatchParentWidget(GetOwner());
+				Child->ConditionallyDetatchParentWidget(&GetOwner());
 			}
 		}
 
@@ -750,19 +762,15 @@ public:
 
 	void Insert(const TSharedRef<ChildType>& Child, int32 Index)
 	{
-		if (GetOwner() && bChangesInvalidatePrepass) 
+		if (bChangesInvalidatePrepass) 
 		{
-			GetOwner()->Invalidate(EInvalidateWidgetReason::ChildOrder);
+			GetOwner().Invalidate(EInvalidateWidgetReason::ChildOrder);
 		}
 
 		Children.Insert(Child, Index);
-
-		if (GetOwner())
+		if (Child != SNullWidget::NullWidget)
 		{
-			if (Child != SNullWidget::NullWidget)
-			{
-				Child->AssignParentWidget(GetOwner()->AsShared());
-			}
+			Child->AssignParentWidget(GetOwner().AsShared());
 		}
 	}
 
@@ -770,7 +778,7 @@ public:
 	{
 		if (Child != SNullWidget::NullWidget)
 		{
-			Child->ConditionallyDetatchParentWidget(GetOwner());
+			Child->ConditionallyDetatchParentWidget(&GetOwner());
 		}
 
 		return Children.Remove( Child );
@@ -781,7 +789,7 @@ public:
 		TSharedRef<SWidget> Child = GetChildAt(Index);
 		if (Child != SNullWidget::NullWidget)
 		{
-			Child->ConditionallyDetatchParentWidget(GetOwner());
+			Child->ConditionallyDetatchParentWidget(&GetOwner());
 		}
 
 		// Note: Child above ensures the instance we're removing from the array won't run its destructor until after it's fully removed from the array
@@ -805,18 +813,18 @@ public:
 	void Sort( const PREDICATE_CLASS& Predicate )
 	{
 		Children.Sort( Predicate );
-		if (GetOwner() && bChangesInvalidatePrepass)
+		if (bChangesInvalidatePrepass)
 		{
-			GetOwner()->Invalidate(EInvalidateWidgetReason::ChildOrder);
+			GetOwner().Invalidate(EInvalidateWidgetReason::ChildOrder);
 		}
 	}
 
 	void Swap( int32 IndexA, int32 IndexB )
 	{
 		Children.Swap(IndexA, IndexB);
-		if (GetOwner() && bChangesInvalidatePrepass)
+		if (bChangesInvalidatePrepass)
 		{
-			GetOwner()->Invalidate(EInvalidateWidgetReason::ChildOrder);
+			GetOwner().Invalidate(EInvalidateWidgetReason::ChildOrder);
 		}
 	}
 
@@ -826,13 +834,14 @@ private:
 
 
 /** A FChildren that support only one slot and alignment of content and padding */
-class FOneSimpleMemberChild : public TSupportsOneChildMixin<FOneSimpleMemberChild>, public TSupportsContentAlignmentMixin<FOneSimpleMemberChild>
+class FOneSimpleMemberChild : public TSingleWidgetChildrenWithSlot<FOneSimpleMemberChild>
+	, public TAlignmentWidgetSlotMixin<FOneSimpleMemberChild>
 {
 public:
 	template<typename WidgetType, typename V = typename std::enable_if<std::is_base_of<SWidget, WidgetType>::value>::type>
 	FOneSimpleMemberChild(WidgetType& InParent)
-		: TSupportsOneChildMixin<FOneSimpleMemberChild>(&InParent)
-		, TSupportsContentAlignmentMixin<FOneSimpleMemberChild>(HAlign_Fill, VAlign_Fill)
+		: TSingleWidgetChildrenWithSlot<FOneSimpleMemberChild>(&InParent)
+		, TAlignmentWidgetSlotMixin<FOneSimpleMemberChild>(HAlign_Fill, VAlign_Fill)
 		, SlotPaddingAttribute(InParent)
 	{
 	}
@@ -843,31 +852,31 @@ public:
 
 	FOneSimpleMemberChild& Padding(TAttribute<FMargin> InPadding)
 	{
-		SlotPaddingAttribute.Assign(*FChildren::GetOwner(), MoveTemp(InPadding));
+		SlotPaddingAttribute.Assign(FChildren::GetOwner(), MoveTemp(InPadding));
 		return *this;
 	}
 
 	FOneSimpleMemberChild& Padding(float Uniform)
 	{
-		SlotPaddingAttribute.Set(*FChildren::GetOwner(), FMargin(Uniform));
+		SlotPaddingAttribute.Set(FChildren::GetOwner(), FMargin(Uniform));
 		return *this;
 	}
 
 	FOneSimpleMemberChild& Padding(float Horizontal, float Vertical)
 	{
-		SlotPaddingAttribute.Set(*FChildren::GetOwner(), FMargin(Horizontal, Vertical));
+		SlotPaddingAttribute.Set(FChildren::GetOwner(), FMargin(Horizontal, Vertical));
 		return *this;
 	}
 
 	FOneSimpleMemberChild& Padding(float Left, float Top, float Right, float Bottom)
 	{
-		SlotPaddingAttribute.Set(*FChildren::GetOwner(), FMargin(Left, Top, Right, Bottom));
+		SlotPaddingAttribute.Set(FChildren::GetOwner(), FMargin(Left, Top, Right, Bottom));
 		return *this;
 	}
 
 	void SetPadding(TAttribute<FMargin> InPadding)
 	{
-		SlotPaddingAttribute.Assign(*FChildren::GetOwner(), MoveTemp(InPadding));
+		SlotPaddingAttribute.Assign(FChildren::GetOwner(), MoveTemp(InPadding));
 	}
 	const FMargin& GetPadding() const { return SlotPaddingAttribute.Get(); }
 
@@ -881,7 +890,11 @@ public:
 	using SlotPaddingAttributeRefType = SlateAttributePrivate::TSlateMemberAttributeRef<SlotPaddingAttributeType>;
 
 	template<typename WidgetType, typename V = typename std::enable_if<std::is_base_of<SWidget, WidgetType>::value>::type>
-	SlotPaddingAttributeRefType GetSlotPaddingAttribute() const { return SlotPaddingAttributeRefType(*(static_cast<const WidgetType*>(FChildren::GetOwner())), SlotPaddingAttribute); }
+	SlotPaddingAttributeRefType GetSlotPaddingAttribute() const
+	{
+		WidgetType&  Widget = static_cast<WidgetType&>(FChildren::GetOwner());
+		return SlotPaddingAttributeRefType(Widget.template SharedThis<WidgetType>(&Widget), SlotPaddingAttribute);
+	}
 
 protected:
 	SlotPaddingAttributeType SlotPaddingAttribute;
@@ -897,7 +910,17 @@ public:
 		: FChildren(InOwner)
 		, AllChildren(InAllChildren)
 		, WidgetIndex(InWidgetIndex)
-	{ }
+	{
+		check(InAllChildren);
+		check(WidgetIndex);
+	}
+
+	TOneDynamicChild(SWidget* InOwner, TPanelChildren<SlotType>* InAllChildren, std::nullptr_t) = delete;
+	TOneDynamicChild(SWidget* InOwner, std::nullptr_t, const TAttribute<int32>* InWidgetIndex) = delete;
+	TOneDynamicChild(SWidget* InOwner, std::nullptr_t, std::nullptr_t) = delete;
+	TOneDynamicChild(std::nullptr_t, TPanelChildren<SlotType>* InAllChildren, std::nullptr_t) = delete;
+	TOneDynamicChild(std::nullptr_t, std::nullptr_t, const TAttribute<int32>* InWidgetIndex) = delete;
+	TOneDynamicChild(std::nullptr_t, std::nullptr_t, std::nullptr_t) = delete;
 
 	virtual int32 Num() const override { return AllChildren->Num() > 0 ? 1 : 0; }
 
diff --git a/Engine/Source/Runtime/SlateCore/Public/SlotBase.h b/Engine/Source/Runtime/SlateCore/Public/SlotBase.h
--- a/Engine/Source/Runtime/SlateCore/Public/SlotBase.h
+++ b/Engine/Source/Runtime/SlateCore/Public/SlotBase.h
@@ -8,6 +8,7 @@
 
 class SWidget;
 
+/** Slot are a container of a SWidget used by the FChildren. */
 class SLATECORE_API FSlotBase
 {
 public:
@@ -104,14 +105,13 @@ private:
 };
 
 
+/** A slot that can be used by the declarative syntax. */
 template<typename SlotType>
 class TSlotBase : public FSlotBase
 {
 public:
 
-	TSlotBase()
-	: FSlotBase()
-	{}
+	TSlotBase() = default;
 
 	TSlotBase( const TSharedRef<SWidget>& InWidget )
 	: FSlotBase( InWidget )
@@ -120,12 +120,12 @@ public:
 	SlotType& operator[]( const TSharedRef<SWidget>& InChildWidget )
 	{
 		this->AttachWidget(InChildWidget);
-		return (SlotType&)(*this);
+		return static_cast<SlotType&>(*this);
 	}
 
 	SlotType& Expose( SlotType*& OutVarToInit )
 	{
-		OutVarToInit = (SlotType*)this;
-		return (SlotType&)(*this);
+		OutVarToInit = static_cast<SlotType*>(this);
+		return static_cast<SlotType&>(*this);
 	}
 };
diff --git a/Engine/Source/Runtime/SlateCore/Public/Widgets/SLeafWidget.h b/Engine/Source/Runtime/SlateCore/Public/Widgets/SLeafWidget.h
--- a/Engine/Source/Runtime/SlateCore/Public/Widgets/SLeafWidget.h
+++ b/Engine/Source/Runtime/SlateCore/Public/Widgets/SLeafWidget.h
@@ -62,9 +62,4 @@ private:
 	virtual FChildren* GetChildren() override;
 
 	virtual void OnArrangeChildren( const FGeometry& AllottedGeometry, FArrangedChildren& ArrangedChildren ) const override;
-
-private:
-
-	// Shared instance of FNoChildren for all widgets with no children.
-	static FNoChildren NoChildrenInstance;
 };
diff --git a/Engine/Source/Runtime/UMG/Private/Slate/SRetainerWidget.cpp b/Engine/Source/Runtime/UMG/Private/Slate/SRetainerWidget.cpp
--- a/Engine/Source/Runtime/UMG/Private/Slate/SRetainerWidget.cpp
+++ b/Engine/Source/Runtime/UMG/Private/Slate/SRetainerWidget.cpp
@@ -86,8 +86,7 @@ TFrameValue<int32> SRetainerWidget::Shared_RetainerWorkThisFrame(0);
 
 
 SRetainerWidget::SRetainerWidget()
-	: EmptyChildSlot(this)
-	, PreviousRenderSize(FIntPoint::NoneValue)
+	: PreviousRenderSize(FIntPoint::NoneValue)
 	, PreviousClipRectSize(FIntPoint::NoneValue)
 	, PreviousColorAndOpacity(FColor::Transparent)
 	, LastIncomingLayerId(0)
@@ -332,7 +331,7 @@ FChildren* SRetainerWidget::GetChildren()
 {
 	if (bEnableRetainedRendering)
 	{
-		return &EmptyChildSlot;
+		return &FNoChildren::NoChildrenInstance;
 	}
 	else
 	{
diff --git a/Engine/Source/Runtime/UMG/Public/Slate/SRetainerWidget.h b/Engine/Source/Runtime/UMG/Public/Slate/SRetainerWidget.h
--- a/Engine/Source/Runtime/UMG/Public/Slate/SRetainerWidget.h
+++ b/Engine/Source/Runtime/UMG/Public/Slate/SRetainerWidget.h
@@ -123,7 +123,6 @@ private:
 	static void OnRetainerModeCVarChanged( IConsoleVariable* CVar );
 	static FOnRetainedModeChanged OnRetainerModeChangedDelegate;
 #endif
-	FSimpleSlot EmptyChildSlot;
 
 	mutable FSlateBrush SurfaceBrush;
 
