Temp fix for cooker on android.
#rnx

diff --git a/Engine/Source/Runtime/SlateCore/Public/Layout/Children.h b/Engine/Source/Runtime/SlateCore/Public/Layout/Children.h
--- a/Engine/Source/Runtime/SlateCore/Public/Layout/Children.h
+++ b/Engine/Source/Runtime/SlateCore/Public/Layout/Children.h
@@ -22,7 +22,7 @@ public:
 	FChildren(SWidget* InOwner)
 		: Owner(InOwner)
 	{
-		check(InOwner);
+		//check(InOwner);
 	}
 
 	FChildren(std::nullptr_t) = delete;
