Slate: Force the initialization order for SNullWidget and FNoChildren.

diff --git a/Engine/Source/Runtime/SlateCore/Private/Children.cpp b/Engine/Source/Runtime/SlateCore/Private/Children.cpp
--- a/Engine/Source/Runtime/SlateCore/Private/Children.cpp
+++ b/Engine/Source/Runtime/SlateCore/Private/Children.cpp
@@ -4,11 +4,12 @@
 
 #include "Widgets/SNullWidget.h"
 
-FNoChildren FNoChildren::NoChildrenInstance(&SNullWidget::NullWidget.Get());
 
+PRAGMA_DISABLE_DEPRECATION_WARNINGS
 FNoChildren::FNoChildren()
 	: FChildren(&SNullWidget::NullWidget.Get())
 {
 
 }
+PRAGMA_ENABLE_DEPRECATION_WARNINGS
 
diff --git a/Engine/Source/Runtime/SlateCore/Private/Widgets/SNullWidget.cpp b/Engine/Source/Runtime/SlateCore/Private/Widgets/SNullWidget.cpp
--- a/Engine/Source/Runtime/SlateCore/Private/Widgets/SNullWidget.cpp
+++ b/Engine/Source/Runtime/SlateCore/Private/Widgets/SNullWidget.cpp
@@ -68,5 +68,14 @@ protected:
 	// End of SWidget interface
 };
 
+namespace NullWidgetPrivate
+{
+	TSharedRef<SWidget> Construct()
+	{
+		static TSharedRef<SWidget> Result = SNew(SNullWidgetContent).Visibility(EVisibility::Hidden);
+		return Result;
+	}
+}
+SLATECORE_API TSharedRef<SWidget> SNullWidget::NullWidget = NullWidgetPrivate::Construct();
 
-SLATECORE_API TSharedRef<SWidget> SNullWidget::NullWidget = SNew(SNullWidgetContent).Visibility(EVisibility::Hidden);
+FNoChildren FNoChildren::NoChildrenInstance(&NullWidgetPrivate::Construct().Get());
diff --git a/Engine/Source/Runtime/SlateCore/Public/Layout/Children.h b/Engine/Source/Runtime/SlateCore/Public/Layout/Children.h
--- a/Engine/Source/Runtime/SlateCore/Public/Layout/Children.h
+++ b/Engine/Source/Runtime/SlateCore/Public/Layout/Children.h
@@ -22,7 +22,7 @@ public:
 	FChildren(SWidget* InOwner)
 		: Owner(InOwner)
 	{
-		//check(InOwner);
+		check(InOwner);
 	}
 
 	FChildren(std::nullptr_t) = delete;
