Reskin UMG Default Styles for Starship

diff --git a/Engine/Source/Editor/UMGEditor/Public/WidgetBlueprint.h b/Engine/Source/Editor/UMGEditor/Public/WidgetBlueprint.h
--- a/Engine/Source/Editor/UMGEditor/Public/WidgetBlueprint.h
+++ b/Engine/Source/Editor/UMGEditor/Public/WidgetBlueprint.h
@@ -307,9 +307,6 @@ public:
 
 	bool ArePropertyBindingsAllowed() const;
 
-	/** Does the editor support widget from an editor package. */
-	virtual bool AllowEditorWidget() const { return false; }
-
 protected:
 #if WITH_EDITOR
 	virtual void LoadModulesRequiredForCompilation() override;
diff --git a/Engine/Source/Editor/UnrealEd/Public/BaseWidgetBlueprint.h b/Engine/Source/Editor/UnrealEd/Public/BaseWidgetBlueprint.h
--- a/Engine/Source/Editor/UnrealEd/Public/BaseWidgetBlueprint.h
+++ b/Engine/Source/Editor/UnrealEd/Public/BaseWidgetBlueprint.h
@@ -2,11 +2,11 @@
 
 #pragma once
 
-#include "Engine/Blueprint.h"
+#include "Blueprint/UserWidgetBlueprint.h"
 #include "BaseWidgetBlueprint.generated.h"
 
 UCLASS(Abstract)
-class UNREALED_API UBaseWidgetBlueprint : public UBlueprint
+class UNREALED_API UBaseWidgetBlueprint : public UUserWidgetBlueprint
 {
 	GENERATED_UCLASS_BODY()
 
diff --git a/Engine/Source/Runtime/Slate/Private/Framework/Application/SlateApplication.cpp b/Engine/Source/Runtime/Slate/Private/Framework/Application/SlateApplication.cpp
--- a/Engine/Source/Runtime/Slate/Private/Framework/Application/SlateApplication.cpp
+++ b/Engine/Source/Runtime/Slate/Private/Framework/Application/SlateApplication.cpp
@@ -43,6 +43,7 @@
 #include "ProfilingDebugging/CsvProfiler.h"
 #include "Trace/SlateTrace.h"
 #include "Styling/StarshipCoreStyle.h"
+#include "Styling/UMGCoreStyle.h"
 
 #ifndef SLATE_HAS_WIDGET_REFLECTOR
 	#define SLATE_HAS_WIDGET_REFLECTOR !(UE_BUILD_TEST || UE_BUILD_SHIPPING) && PLATFORM_DESKTOP
@@ -6749,6 +6750,7 @@ void FSlateApplication::InitializeCoreStyle()
 		FAppStyle::SetAppStyleSet(FCoreStyle::GetCoreStyle());
 	}
 
+	FUMGCoreStyle::ResetToDefault();
 }
 
 void FSlateApplication::SetWidgetReflector(const TSharedRef<IWidgetReflector>& WidgetReflector)
diff --git a/Engine/Source/Runtime/SlateCore/Private/Styling/StarshipCoreStyle.cpp b/Engine/Source/Runtime/SlateCore/Private/Styling/StarshipCoreStyle.cpp
--- a/Engine/Source/Runtime/SlateCore/Private/Styling/StarshipCoreStyle.cpp
+++ b/Engine/Source/Runtime/SlateCore/Private/Styling/StarshipCoreStyle.cpp
@@ -10,11 +10,12 @@
 #include "Styling/SlateStyleRegistry.h"
 #include "Styling/SlateStyle.h"
 #include "Styling/SlateTypes.h"
-#include "Fonts/LegacySlateFontInfoCache.h"
+#include "Fonts/CompositeFont.h"
 #include "Styling/SlateStyleMacros.h"
 #include "Styling/ToolBarStyle.h"
 #include "Styling/SegmentedControlStyle.h"
 #include "Styling/StyleColors.h"
+#include "Styling/CoreStyle.h"
 
 // This is to fix the issue that SlateStyleMacros like IMAGE_BRUSH look for RootToContentDir but Style->RootToContentDir is how the core style is set up
 #define RootToContentDir Style->RootToContentDir
@@ -26,7 +27,7 @@ TSharedPtr< ISlateStyle > FStarshipCoreStyle::Instance = nullptr;
 
 using namespace CoreStyleConstants;
 
-#define FONT(...) FSlateFontInfo(FLegacySlateFontInfoCache::Get().GetDefaultFont(), __VA_ARGS__)
+#define FONT(...) FSlateFontInfo(FCoreStyle::GetDefaultFont(), __VA_ARGS__)
 
 TUniquePtr<FStyleFonts> FStyleFonts::Instance = nullptr;
 
@@ -45,7 +46,7 @@ FStyleFonts::FStyleFonts()
 /* FStarshipCoreStyle helper class
  *****************************************************************************/
 
-class FStyle
+class FStarshipCoreStyle::FStyle
 	: public FSlateStyleSet
 {
 public:
@@ -83,8 +84,8 @@ public:
 
 TSharedRef<const FCompositeFont> FStarshipCoreStyle::GetDefaultFont()
 {
-	// FStarshipCoreStyle::GetDefaultFont is an alias so that the default font from FLegacySlateFontInfoCache can be acccessed outside of SlateCore (FLegacySlateFontInfoCache is private)
-	return FLegacySlateFontInfoCache::Get().GetDefaultFont();
+	// FStarshipCoreStyle currently uses same default font as core
+	return FCoreStyle::GetDefaultFont();
 }
 
 
@@ -145,7 +146,6 @@ void FStarshipCoreStyle::SetFocusBrush(FSlateBrush* NewBrush)
 	FSlateStyleRegistry::RegisterSlateStyle(Style.Get());
 }
 
-#define DEFAULT_FONT(...) FStarshipCoreStyle::GetDefaultFontStyle(__VA_ARGS__)
 
 TSharedRef<ISlateStyle> FStarshipCoreStyle::Create()
 {
@@ -2114,7 +2114,5 @@ void FStarshipCoreStyle::SetupMultiboxStyles(TSharedRef<FStyle>& Style)
 
 }
 
-#undef IMAGE_BRUSH
-#undef BOX_BRUSH
-#undef BORDER_BRUSH
-#undef DEFAULT_FONT
+#undef RootToContentDir
+#undef FONT
\ No newline at end of file
diff --git a/Engine/Source/Runtime/SlateCore/Private/Styling/UMGCoreStyle.cpp b/Engine/Source/Runtime/SlateCore/Private/Styling/UMGCoreStyle.cpp
new file mode 100644
--- /dev/null
+++ b/Engine/Source/Runtime/SlateCore/Private/Styling/UMGCoreStyle.cpp
@@ -0,0 +1,1382 @@
+// Copyright Epic Games, Inc. All Rights Reserved.
+
+#include "Styling/UMGCoreStyle.h"
+#include "SlateGlobals.h"
+#include "Brushes/SlateBorderBrush.h"
+#include "Brushes/SlateBoxBrush.h"
+#include "Brushes/SlateColorBrush.h"
+#include "Brushes/SlateImageBrush.h"
+#include "Brushes/SlateRoundedBoxBrush.h"
+#include "Styling/SlateStyleRegistry.h"
+#include "Styling/SlateStyle.h"
+#include "Styling/SlateTypes.h"
+#include "Fonts/CompositeFont.h"
+#include "Styling/SlateStyleMacros.h"
+#include "Styling/ToolBarStyle.h"
+#include "Styling/SegmentedControlStyle.h"
+#include "Styling/StyleColors.h"
+#include "Styling/CoreStyle.h"
+
+// This is to fix the issue that SlateStyleMacros like IMAGE_BRUSH look for RootToContentDir but Style->RootToContentDir is how the core style is set up
+#define RootToContentDir Style->RootToContentDir
+
+TSharedPtr< ISlateStyle > FUMGCoreStyle::Instance = nullptr;
+
+/* FUMGCoreStyle helper class
+ *****************************************************************************/
+
+struct FUMGColor
+{
+	inline static const FSlateColor Transparent = FStyleColors::Transparent;
+	inline static const FSlateColor Black = FLinearColor(0.000000, 0.000000, 0.000000);
+	inline static const FSlateColor Title = FLinearColor(0.849999, 0.849999, 0.849999);
+	inline static const FSlateColor WindowBorder = FLinearColor(0.577777, 0.577777, 0.577777);
+	inline static const FSlateColor Foldout = FLinearColor(0.577777, 0.577777, 0.577777);
+	inline static const FSlateColor Input = FLinearColor(0.577777, 0.577777, 0.577777);
+	inline static const FSlateColor InputOutline = FLinearColor(0.695111, 0.695111, 0.695111);
+	inline static const FSlateColor Recessed = FLinearColor(0.203300, 0.203300, 0.203300);
+	inline static const FSlateColor Background = FLinearColor(0.849999, 0.849999, 0.849999);
+	inline static const FSlateColor Panel = FLinearColor(0.276422, 0.276422, 0.276422);
+	inline static const FSlateColor Header = FLinearColor(0.384266, 0.384266, 0.384266);
+	inline static const FSlateColor Dropdown = FLinearColor(0.495466, 0.495466, 0.495466);
+	inline static const FSlateColor DropdownOutline = FLinearColor(0.822722, 0.822722, 0.822722);
+	inline static const FSlateColor Hover = FLinearColor(0.724268, 0.724268, 0.724268);
+	inline static const FSlateColor Hover2 = FLinearColor(0.215861, 0.215861, 0.215861);
+	inline static const FSlateColor White = FLinearColor(0.000000, 1.000000, 1.000000);
+	inline static const FSlateColor White25 = FLinearColor(0.000000, 1.000000, 1.000000);
+	inline static const FSlateColor Highlight = FLinearColor(0.000000, 0.187821, 0.679542);
+
+	inline static const FSlateColor Primary = FLinearColor(0.019382, 0.496933, 1.000000);
+	inline static const FSlateColor PrimaryHover = FLinearColor(0.158961, 0.644480, 1.000000);
+	inline static const FSlateColor PrimaryPress = FLinearColor(0.009721, 0.250158, 0.502886);
+	inline static const FSlateColor Secondary = FLinearColor(0.495466, 0.495466, 0.495466);
+
+	inline static const FSlateColor Foreground = FLinearColor(0.527115, 0.527115, 0.527115);
+	inline static const FSlateColor ForegroundHover = FLinearColor(0.000000, 1.000000, 1.000000);
+	inline static const FSlateColor ForegroundInverted = FLinearColor(0.472885, 0.472885, 0.472885);
+	inline static const FSlateColor ForegroundHeader = FLinearColor(0.577580, 0.577580, 0.577580);
+
+	inline static const FSlateColor Select = FLinearColor(0.019382, 0.496933, 1.000000);
+	inline static const FSlateColor SelectInactive = FLinearColor(0.318547, 0.450786, 0.520996);
+	inline static const FSlateColor SelectParent = FLinearColor(0.025187, 0.031896, 0.042311);
+	inline static const FSlateColor SelectHover = FLinearColor(0.276422, 0.276422, 0.276422);
+
+	inline static const FSlateColor Notifications = FStyleColors::Notifications;
+
+	inline static const FSlateColor AccentBlue = FStyleColors::AccentBlue;
+	inline static const FSlateColor AccentPurple = FStyleColors::AccentPurple;
+	inline static const FSlateColor AccentPink = FStyleColors::AccentPink;
+	inline static const FSlateColor AccentRed = FStyleColors::AccentRed;
+	inline static const FSlateColor AccentOrange = FStyleColors::AccentOrange;
+	inline static const FSlateColor AccentYellow = FStyleColors::AccentYellow;
+	inline static const FSlateColor AccentGreen = FStyleColors::AccentGreen;
+	inline static const FSlateColor AccentBrown = FStyleColors::AccentBrown;
+	inline static const FSlateColor AccentBlack = FStyleColors::AccentBlack;
+	inline static const FSlateColor AccentGray = FStyleColors::AccentGray;
+	inline static const FSlateColor AccentWhite = FStyleColors::AccentWhite;
+	inline static const FSlateColor AccentFolder = FStyleColors::AccentFolder;
+
+	inline static const FSlateColor Warning = FStyleColors::Warning;
+	inline static const FSlateColor Error = FStyleColors::Error;
+	inline static const FSlateColor Success = FStyleColors::Success;
+};
+
+class FUMGStyleSet
+	: public FSlateStyleSet
+{
+public:
+	FUMGStyleSet(const FName& InStyleSetName)
+		: FSlateStyleSet(InStyleSetName)
+
+		// These are the colors that are updated by the user style customizations
+		, SelectorColor_LinearRef(MakeShared<FLinearColor>(0.701f, 0.225f, 0.003f))
+		, SelectionColor_LinearRef(MakeShared<FLinearColor>(COLOR("18A0FBFF")))
+		, SelectionColor_Inactive_LinearRef(MakeShared<FLinearColor>(0.25f, 0.25f, 0.25f))
+		, SelectionColor_Pressed_LinearRef(MakeShared<FLinearColor>(0.701f, 0.225f, 0.003f))
+		, HighlightColor_LinearRef(MakeShared<FLinearColor>(0.068f, 0.068f, 0.068f))
+	{
+	}
+
+	static void SetColor(const TSharedRef<FLinearColor>& Source, const FLinearColor& Value)
+	{
+		Source->R = Value.R;
+		Source->G = Value.G;
+		Source->B = Value.B;
+		Source->A = Value.A;
+	}
+
+	// These are the colors that are updated by the user style customizations
+	const TSharedRef<FLinearColor> SelectorColor_LinearRef;
+	const TSharedRef<FLinearColor> SelectionColor_LinearRef;
+	const TSharedRef<FLinearColor> SelectionColor_Inactive_LinearRef;
+	const TSharedRef<FLinearColor> SelectionColor_Pressed_LinearRef;
+	const TSharedRef<FLinearColor> HighlightColor_LinearRef;
+};
+
+
+/* FUMGCoreStyle static functions
+ *****************************************************************************/
+
+TSharedRef<const FCompositeFont> FUMGCoreStyle::GetDefaultFont()
+{
+	// FUMGCoreStyle currently uses same default font as core
+	return FCoreStyle::GetDefaultFont();
+}
+
+
+FSlateFontInfo FUMGCoreStyle::GetDefaultFontStyle(const FName InTypefaceFontName, const int32 InSize, const FFontOutlineSettings& InOutlineSettings)
+{
+	return FSlateFontInfo(GetDefaultFont(), InSize, InTypefaceFontName, InOutlineSettings);
+}
+
+
+void FUMGCoreStyle::ResetToDefault()
+{
+	SetStyle(FUMGCoreStyle::Create());
+}
+
+
+void FUMGCoreStyle::SetSelectorColor(const FLinearColor& NewColor)
+{
+	TSharedPtr<FUMGStyleSet> Style = StaticCastSharedPtr<FUMGStyleSet>(Instance);
+	check(Style.IsValid());
+
+	FUMGStyleSet::SetColor(Style->SelectorColor_LinearRef, NewColor);
+}
+
+
+void FUMGCoreStyle::SetSelectionColor(const FLinearColor& NewColor)
+{
+	TSharedPtr<FUMGStyleSet> Style = StaticCastSharedPtr<FUMGStyleSet>(Instance);
+	check(Style.IsValid());
+
+	FUMGStyleSet::SetColor(Style->SelectionColor_LinearRef, NewColor);
+}
+
+
+void FUMGCoreStyle::SetInactiveSelectionColor(const FLinearColor& NewColor)
+{
+	TSharedPtr<FUMGStyleSet> Style = StaticCastSharedPtr<FUMGStyleSet>(Instance);
+	check(Style.IsValid());
+
+	FUMGStyleSet::SetColor(Style->SelectionColor_Inactive_LinearRef, NewColor);
+}
+
+
+void FUMGCoreStyle::SetPressedSelectionColor(const FLinearColor& NewColor)
+{
+	TSharedPtr<FUMGStyleSet> Style = StaticCastSharedPtr<FUMGStyleSet>(Instance);
+	check(Style.IsValid());
+
+	FUMGStyleSet::SetColor(Style->SelectionColor_Pressed_LinearRef, NewColor);
+}
+
+void FUMGCoreStyle::SetFocusBrush(FSlateBrush* NewBrush)
+{
+	TSharedRef<FUMGStyleSet> Style = StaticCastSharedRef<FUMGStyleSet>(Instance.ToSharedRef());
+	FSlateStyleRegistry::UnRegisterSlateStyle(Style.Get());
+
+	Style->Set("FocusRectangle", NewBrush);
+
+	FSlateStyleRegistry::RegisterSlateStyle(Style.Get());
+}
+
+TSharedRef<ISlateStyle> FUMGCoreStyle::Create()
+{
+	TSharedRef<FUMGStyleSet> Style = MakeShareable(new FUMGStyleSet("UMGCoreStyle"));
+	Style->SetContentRoot(FPaths::EngineContentDir() / TEXT("Slate"));
+
+	const FString CanaryPath = RootToContentDir(TEXT("Checkerboard"), TEXT(".png"));
+
+	if (!FPaths::FileExists(CanaryPath))
+	{
+		// Checkerboard is the default brush so we check for that. No slate fonts are required as those will fall back properly
+		UE_LOG(LogSlate, Warning, TEXT("FUMGCoreStyle assets not detected, skipping FUMGCoreStyle initialization"));
+
+		return Style;
+	}
+
+	// These are the Slate colors which reference the dynamic colors in FSlateCoreStyle; these are the colors to put into the style
+	const FSlateColor DefaultForeground(FUMGColor::Black);
+	const FSlateColor InvertedForeground(FUMGColor::ForegroundInverted);
+	const FSlateColor SelectorColor(Style->SelectorColor_LinearRef);
+	const FSlateColor SelectionColor(Style->SelectionColor_LinearRef);
+	const FSlateColor SelectionColor_Inactive(Style->SelectionColor_Inactive_LinearRef);
+	const FSlateColor SelectionColor_Pressed(Style->SelectionColor_Pressed_LinearRef);
+	const FSlateColor HighlightColor(FUMGColor::Highlight);
+
+	Style->Set("AppCoreStyleConstants::Icon", new IMAGE_BRUSH("CoreStyleConstants::Icons/DefaultAppCoreStyleConstants::Icon", CoreStyleConstants::Icon24x24));
+	Style->Set("AppCoreStyleConstants::Icon.Small", new IMAGE_BRUSH("CoreStyleConstants::Icons/DefaultAppCoreStyleConstants::Icon", CoreStyleConstants::Icon24x24));
+
+	Style->Set("AppCoreStyleConstants::IconPadding", FMargin(4, 4, 0, 0));
+	Style->Set("AppCoreStyleConstants::IconPadding.Small", FMargin(4, 4, 0, 0));
+
+	Style->Set("NormalFont", DEFAULT_FONT("Regular", RegularTextSize));
+
+	Style->Set("SmallFont", DEFAULT_FONT("Regular", SmallTextSize));
+
+	FSlateBrush* DefaultTextUnderlineBrush = new IMAGE_BRUSH("Old/White", CoreStyleConstants::Icon8x8, FLinearColor::White, ESlateBrushTileType::Both);
+	Style->Set("DefaultTextUnderline", DefaultTextUnderlineBrush);
+
+	// Normal Text
+	const FTextBlockStyle NormalText = FTextBlockStyle()
+		.SetFont(DEFAULT_FONT("Regular", RegularTextSize))
+		.SetColorAndOpacity(FSlateColor::UseForeground())
+		.SetShadowOffset(FVector2D::ZeroVector)
+		.SetShadowColorAndOpacity(FLinearColor::Black)
+		.SetHighlightColor(FLinearColor(0.02f, 0.3f, 0.0f))
+		.SetHighlightShape(BOX_BRUSH("Common/TextBlockHighlightShape", FMargin(3.f /8.f)));
+
+	const FTextBlockStyle NormalUnderlinedText = FTextBlockStyle(NormalText)
+		.SetUnderlineBrush(*DefaultTextUnderlineBrush);
+
+	// Monospaced Text
+	const FTextBlockStyle MonospacedText = FTextBlockStyle()
+		.SetFont(DEFAULT_FONT("Mono", 10))
+		.SetColorAndOpacity(FSlateColor::UseForeground())
+		.SetShadowOffset(FVector2D::ZeroVector)
+		.SetShadowColorAndOpacity(FLinearColor::Black)
+		.SetHighlightColor(FLinearColor(0.02f, 0.3f, 0.0f))
+		.SetHighlightShape(BOX_BRUSH("Common/TextBlockHighlightShape", FMargin(3.f/8.f))
+		);
+
+	const FTextBlockStyle MonospacedUnderlinedText = FTextBlockStyle(MonospacedText)
+		.SetUnderlineBrush(*DefaultTextUnderlineBrush);
+
+	Style->Set("MonospacedText", MonospacedText);
+	Style->Set("MonospacedUnderlinedText", MonospacedUnderlinedText);
+
+	// Small Text
+	const FTextBlockStyle SmallText = FTextBlockStyle(NormalText)
+		.SetFont(DEFAULT_FONT("Regular", SmallTextSize));
+
+	const FTextBlockStyle SmallUnderlinedText = FTextBlockStyle(SmallText)
+		.SetUnderlineBrush(*DefaultTextUnderlineBrush);
+
+	// Embossed Text
+	Style->Set("EmbossedText", FTextBlockStyle(NormalText)
+		.SetFont(DEFAULT_FONT("Regular", 24))
+		.SetColorAndOpacity(FLinearColor::Black )
+		.SetShadowOffset( FVector2D(0.0f, 1.0f))
+		.SetShadowColorAndOpacity(FLinearColor(0.8f, 0.8f, 0.8f, 0.5))
+		);
+
+	// Common brushes
+	FSlateBrush* GenericWhiteBox = new IMAGE_BRUSH("Old/White", CoreStyleConstants::Icon16x16);
+	{
+		Style->Set("Checkerboard", new IMAGE_BRUSH("Checkerboard", CoreStyleConstants::Icon16x16, FLinearColor::White, ESlateBrushTileType::Both));
+
+		Style->Set("GenericWhiteBox", GenericWhiteBox);
+
+		Style->Set("BlackBrush", new FSlateColorBrush(FLinearColor::Black));
+		Style->Set("WhiteBrush", new FSlateColorBrush(FLinearColor::White));
+
+		Style->Set("BoxShadow", new BOX_BRUSH( "Common/BoxShadow" , FMargin( 5.0f / 64.0f)));
+
+		Style->Set("FocusRectangle", new BORDER_BRUSH( "Old/DashedBorder", FMargin(6.0f / 32.0f), FLinearColor(1, 1, 1, 0.5)));
+	}
+
+	// Important colors
+	{
+		Style->Set("DefaultForeground", DefaultForeground);
+		Style->Set("InvertedForeground", InvertedForeground);
+
+		Style->Set("SelectorColor", SelectorColor);
+		Style->Set("SelectionColor", SelectionColor);
+		Style->Set("SelectionColor_Inactive", SelectionColor_Inactive);
+		Style->Set("SelectionColor_Pressed", SelectionColor_Pressed);
+	}
+
+	// Invisible buttons, borders, etc.
+	const FButtonStyle NoBorder = FButtonStyle()
+		.SetNormal(FSlateNoResource())
+		.SetHovered(FSlateNoResource())
+		.SetPressed(FSlateNoResource())
+		.SetNormalPadding(FMargin(0.0f, 0.0f, 0.0f, 1.0f))
+		.SetPressedPadding(FMargin(0.0f, 1.0f, 0.0f, 0.0f));
+
+	// Convenient transparent/invisible elements
+	{
+		Style->Set("NoBrush", new FSlateNoResource());
+		Style->Set("NoBorder", new FSlateNoResource());
+		Style->Set("NoBorder.Normal", new FSlateNoResource());
+		Style->Set("NoBorder.Hovered", new FSlateNoResource());
+		Style->Set("NoBorder.Pressed", new FSlateNoResource());
+		Style->Set("NoBorder", NoBorder);
+	}
+
+	
+	// Demo Recording
+	{
+		Style->Set("DemoRecording.CursorPing", new IMAGE_BRUSH("Common/CursorPing", FVector2D(31,31)));
+	}
+
+	// Error Reporting
+	{
+		Style->Set("ErrorReporting.Box", new BOX_BRUSH( "Common/TextBlockHighlightShape", FMargin(3.f / 8.f)));
+		Style->Set("ErrorReporting.EmptyBox", new BOX_BRUSH( "Common/TextBlockHighlightShape_Empty", FMargin(3.f / 8.f)));
+		Style->Set("ErrorReporting.BackgroundColor", FLinearColor(0.35f, 0.0f, 0.0f));
+		Style->Set("ErrorReporting.WarningBackgroundColor", FLinearColor(0.828f, 0.364f, 0.003f));
+		Style->Set("ErrorReporting.ForegroundColor", FLinearColor::White);
+	}
+
+	// Cursor CoreStyleConstants::Icons
+	{
+		Style->Set("SoftwareCursor_Grab", new IMAGE_BRUSH( "CoreStyleConstants::Icons/cursor_grab", CoreStyleConstants::Icon16x16));
+		Style->Set("SoftwareCursor_CardinalCross", new IMAGE_BRUSH( "CoreStyleConstants::Icons/cursor_cardinal_cross", CoreStyleConstants::Icon24x24));
+	}
+
+	// Common CoreStyleConstants::Icons
+	{
+		Style->Set("TrashCan", new IMAGE_BRUSH( "CoreStyleConstants::Icons/TrashCan", FVector2D(64, 64)));
+		Style->Set("TrashCan_Small", new IMAGE_BRUSH( "CoreStyleConstants::Icons/TrashCan_Small", FVector2D(18, 18)));
+	}
+
+	// Common CoreStyleConstants::Icons
+	{
+		Style->Set( "CoreStyleConstants::Icons.Cross", new IMAGE_BRUSH( "CoreStyleConstants::Icons/Cross_12x", CoreStyleConstants::Icon12x12 ) );
+		Style->Set( "CoreStyleConstants::Icons.Denied", new IMAGE_BRUSH( "CoreStyleConstants::Icons/denied_16x", CoreStyleConstants::Icon16x16 ) );
+		Style->Set( "CoreStyleConstants::Icons.Error", new IMAGE_BRUSH( "CoreStyleConstants::Icons/CoreStyleConstants::Icon_error_16x", CoreStyleConstants::Icon16x16) );
+		Style->Set( "CoreStyleConstants::Icons.Help", new IMAGE_BRUSH( "CoreStyleConstants::Icons/CoreStyleConstants::Icon_help_16x", CoreStyleConstants::Icon16x16) );
+		Style->Set( "CoreStyleConstants::Icons.Info", new IMAGE_BRUSH( "CoreStyleConstants::Icons/CoreStyleConstants::Icon_info_16x", CoreStyleConstants::Icon16x16) );
+		Style->Set( "CoreStyleConstants::Icons.Warning", new IMAGE_BRUSH( "CoreStyleConstants::Icons/CoreStyleConstants::Icon_warning_16x", CoreStyleConstants::Icon16x16) );
+		Style->Set( "CoreStyleConstants::Icons.Download", new IMAGE_BRUSH( "CoreStyleConstants::Icons/CoreStyleConstants::Icon_Downloads_16x", CoreStyleConstants::Icon16x16) );
+	}
+
+	// Tool panels
+	{
+		Style->Set( "ToolPanel.GroupBorder", new BOX_BRUSH( "Common/GroupBorder", FMargin(4.0f/16.0f) ) );
+		Style->Set( "Debug.Border", new BOX_BRUSH( "Common/DebugBorder", 4.0f/16.0f) );
+	}
+
+	// Popup text
+	{
+		Style->Set( "PopupText.Background", new BOX_BRUSH( "Old/Menu_Background", FMargin(8.0f/64.0f) ) );
+	}
+
+	// Generic command CoreStyleConstants::Icons
+	{
+		Style->Set( "GenericCommands.Undo", new IMAGE_BRUSH( "CoreStyleConstants::Icons/CoreStyleConstants::Icon_undo_16px", CoreStyleConstants::Icon16x16 ) );
+		Style->Set( "GenericCommands.Redo", new IMAGE_BRUSH( "CoreStyleConstants::Icons/CoreStyleConstants::Icon_redo_16px", CoreStyleConstants::Icon16x16 ) );
+
+		Style->Set( "GenericCommands.Copy", new IMAGE_BRUSH("CoreStyleConstants::Icons/Edit/CoreStyleConstants::Icon_Edit_Copy_16x", CoreStyleConstants::Icon16x16) );
+		Style->Set( "GenericCommands.Cut", new IMAGE_BRUSH("CoreStyleConstants::Icons/Edit/CoreStyleConstants::Icon_Edit_Cut_16x", CoreStyleConstants::Icon16x16) );
+		Style->Set( "GenericCommands.Delete", new IMAGE_BRUSH("CoreStyleConstants::Icons/Edit/CoreStyleConstants::Icon_Edit_Delete_16x", CoreStyleConstants::Icon16x16) );
+		Style->Set( "GenericCommands.Paste", new IMAGE_BRUSH("CoreStyleConstants::Icons/Edit/CoreStyleConstants::Icon_Edit_Paste_16x", CoreStyleConstants::Icon16x16) );
+		Style->Set( "GenericCommands.Duplicate", new IMAGE_BRUSH("CoreStyleConstants::Icons/Edit/CoreStyleConstants::Icon_Edit_Duplicate_16x", CoreStyleConstants::Icon16x16) );
+		Style->Set( "GenericCommands.Rename", new IMAGE_BRUSH( "CoreStyleConstants::Icons/Edit/CoreStyleConstants::Icon_Edit_Rename_16x", CoreStyleConstants::Icon16x16 ) );
+	}
+
+	// SVerticalBox Drag& Drop CoreStyleConstants::Icon
+	Style->Set("VerticalBoxDragIndicator", new IMAGE_BRUSH("Common/VerticalBoxDragIndicator", FVector2D(6, 45)));
+	Style->Set("VerticalBoxDragIndicatorShort", new IMAGE_BRUSH("Common/VerticalBoxDragIndicatorShort", FVector2D(6, 15)));
+	// SScrollBar defaults...
+	const FScrollBarStyle ScrollBar = FScrollBarStyle()
+		.SetVerticalTopSlotImage(IMAGE_BRUSH("Common/Scrollbar_Background_Vertical", FVector2D(8, 8)))
+		.SetVerticalBottomSlotImage(IMAGE_BRUSH("Common/Scrollbar_Background_Vertical", FVector2D(8, 8)))
+		.SetHorizontalTopSlotImage(IMAGE_BRUSH("Common/Scrollbar_Background_Horizontal", FVector2D(8, 8)))
+		.SetHorizontalBottomSlotImage(IMAGE_BRUSH("Common/Scrollbar_Background_Horizontal", FVector2D(8, 8)))
+		.SetNormalThumbImage( BOX_BRUSH( "Common/Scrollbar_Thumb", FMargin(4.f/16.f) ) )
+		.SetDraggedThumbImage( BOX_BRUSH( "Common/Scrollbar_Thumb", FMargin(4.f/16.f) ) )
+		.SetHoveredThumbImage( BOX_BRUSH( "Common/Scrollbar_Thumb", FMargin(4.f/16.f) ) );
+	{
+		Style->Set( "Scrollbar", ScrollBar );
+	}
+
+	// SButton defaults...
+	const FButtonStyle Button = FButtonStyle()
+		.SetNormal(FSlateRoundedBoxBrush(FUMGColor::Secondary, 4.0f, FUMGColor::InputOutline, CoreStyleConstants::InputFocusThickness))
+		.SetHovered(FSlateRoundedBoxBrush(FUMGColor::Hover, 4.0f, FUMGColor::Hover, CoreStyleConstants::InputFocusThickness))
+		.SetPressed(FSlateRoundedBoxBrush(FUMGColor::Header, 4.0f, FUMGColor::Hover, CoreStyleConstants::InputFocusThickness))
+		.SetNormalForeground(FUMGColor::ForegroundHover)
+		.SetHoveredForeground(FUMGColor::ForegroundHover)
+		.SetPressedForeground(FUMGColor::ForegroundHover)
+		.SetDisabledForeground(FUMGColor::Foreground)
+		.SetNormalPadding(CoreStyleConstants::ButtonMargins)
+		.SetPressedPadding(CoreStyleConstants::ButtonMargins);
+	{
+		Style->Set( "Button", Button );
+
+		Style->Set( "InvertedForeground", InvertedForeground );
+	}
+
+	// SComboButton and SComboBox defaults...
+	{
+		FComboButtonStyle ComboButton = FComboButtonStyle()
+			.SetButtonStyle(Button)
+			.SetDownArrowImage(IMAGE_BRUSH("Common/ComboArrow", CoreStyleConstants::Icon8x8))
+			.SetMenuBorderBrush(BOX_BRUSH("Old/Menu_Background", FMargin(8.0f/64.0f)))
+			.SetMenuBorderPadding(FMargin(0.0f));
+		Style->Set( "ComboButton", ComboButton );
+
+		ComboButton.SetMenuBorderPadding(FMargin(1.0));
+
+		FComboBoxStyle ComboBox = FComboBoxStyle()
+			.SetComboButtonStyle(ComboButton);
+		Style->Set( "ComboBox", ComboBox );
+	}
+
+	// SMessageLogListing
+	{
+		FComboButtonStyle MessageLogListingComboButton = FComboButtonStyle()
+			.SetButtonStyle(NoBorder)
+			.SetDownArrowImage(IMAGE_BRUSH("Common/ComboArrow", CoreStyleConstants::Icon8x8))
+			.SetMenuBorderBrush(FSlateNoResource())
+			.SetMenuBorderPadding(FMargin(0.0f));
+		Style->Set("MessageLogListingComboButton", MessageLogListingComboButton);
+	}
+
+	// SEditableComboBox defaults...
+	{
+		Style->Set( "EditableComboBox.Add", new IMAGE_BRUSH( "CoreStyleConstants::Icons/PlusSymbol_12x", CoreStyleConstants::Icon12x12 ) );
+		Style->Set( "EditableComboBox.Delete", new IMAGE_BRUSH( "CoreStyleConstants::Icons/Cross_12x", CoreStyleConstants::Icon12x12 ) );
+		Style->Set( "EditableComboBox.Rename", new IMAGE_BRUSH( "CoreStyleConstants::Icons/ellipsis_12x", CoreStyleConstants::Icon12x12 ) );
+		Style->Set( "EditableComboBox.Accept", new IMAGE_BRUSH( "Common/Check", CoreStyleConstants::Icon16x16 ) );
+	}
+
+	// SCheckBox defaults...
+	{
+		const FCheckBoxStyle BasicCheckBoxStyle = FCheckBoxStyle()
+			.SetCheckBoxType(ESlateCheckBoxType::CheckBox)
+			.SetUncheckedImage( IMAGE_BRUSH( "Common/CheckBox", CoreStyleConstants::Icon16x16 ) )
+			.SetUncheckedHoveredImage( IMAGE_BRUSH( "Common/CheckBox", CoreStyleConstants::Icon16x16 ) )
+			.SetUncheckedPressedImage( IMAGE_BRUSH( "Common/CheckBox_Hovered", CoreStyleConstants::Icon16x16, FLinearColor( 0.5f, 0.5f, 0.5f ) ) )
+			.SetCheckedImage( IMAGE_BRUSH( "Common/CheckBox_Checked_Hovered", CoreStyleConstants::Icon16x16 ) )
+			.SetCheckedHoveredImage( IMAGE_BRUSH( "Common/CheckBox_Checked_Hovered", CoreStyleConstants::Icon16x16, FLinearColor( 0.5f, 0.5f, 0.5f ) ) )
+			.SetCheckedPressedImage( IMAGE_BRUSH( "Common/CheckBox_Checked", CoreStyleConstants::Icon16x16 ) )
+			.SetUndeterminedImage( IMAGE_BRUSH( "Common/CheckBox_Undetermined", CoreStyleConstants::Icon16x16 ) )
+			.SetUndeterminedHoveredImage( IMAGE_BRUSH( "Common/CheckBox_Undetermined_Hovered", CoreStyleConstants::Icon16x16 ) )
+			.SetUndeterminedPressedImage( IMAGE_BRUSH( "Common/CheckBox_Undetermined_Hovered", CoreStyleConstants::Icon16x16, FLinearColor( 0.5f, 0.5f, 0.5f ) ) );
+		Style->Set( "Checkbox", BasicCheckBoxStyle );
+
+		/* Set images for various transparent SCheckBox states ... */
+		const FCheckBoxStyle BasicTransparentCheckBoxStyle = FCheckBoxStyle()
+			.SetCheckBoxType(ESlateCheckBoxType::ToggleButton)
+			.SetUncheckedImage( FSlateNoResource() )
+			.SetUncheckedHoveredImage( FSlateNoResource() )
+			.SetUncheckedPressedImage( FSlateNoResource() )
+			.SetCheckedImage( FSlateNoResource() )
+			.SetCheckedHoveredImage( FSlateNoResource() )
+			.SetCheckedPressedImage( FSlateNoResource() )
+			.SetUndeterminedImage( FSlateNoResource() )
+			.SetUndeterminedHoveredImage( FSlateNoResource() )
+			.SetUndeterminedPressedImage( FSlateNoResource() );
+		Style->Set( "TransparentCheckBox", BasicTransparentCheckBoxStyle );
+
+		/* Default Style for a toggleable button */
+		const FCheckBoxStyle ToggleButtonStyle = FCheckBoxStyle()
+			.SetCheckBoxType(ESlateCheckBoxType::ToggleButton)
+			.SetUncheckedImage( FSlateNoResource() )
+			.SetUncheckedHoveredImage( BOX_BRUSH("Common/RoundedSelection_16x", 4.0f/16.0f, SelectionColor ) )
+			.SetUncheckedPressedImage( BOX_BRUSH("Common/RoundedSelection_16x", 4.0f/16.0f, SelectionColor_Pressed ) )
+			.SetCheckedImage( BOX_BRUSH("Common/RoundedSelection_16x",  4.0f/16.0f, SelectionColor_Pressed ) )
+			.SetCheckedHoveredImage( BOX_BRUSH("Common/RoundedSelection_16x",  4.0f/16.0f, SelectionColor ) )
+			.SetCheckedPressedImage( BOX_BRUSH("Common/RoundedSelection_16x",  4.0f/16.0f, SelectionColor_Pressed ) );
+		Style->Set( "ToggleButtonCheckbox", ToggleButtonStyle );
+
+		/* Style for a toggleable button that mimics the coloring and look of a Table Row */
+		const FCheckBoxStyle ToggleButtonRowStyle = FCheckBoxStyle()
+			.SetCheckBoxType(ESlateCheckBoxType::ToggleButton)
+			.SetUncheckedImage(FSlateNoResource())
+			.SetUncheckedHoveredImage(IMAGE_BRUSH("Common/Selection", CoreStyleConstants::Icon8x8, SelectionColor_Inactive))
+			.SetUncheckedPressedImage(IMAGE_BRUSH("Common/Selection", CoreStyleConstants::Icon8x8, SelectionColor_Inactive))
+			.SetCheckedImage(IMAGE_BRUSH("Common/Selection", CoreStyleConstants::Icon8x8, SelectionColor))
+			.SetCheckedHoveredImage(IMAGE_BRUSH("Common/Selection", CoreStyleConstants::Icon8x8, SelectionColor))
+			.SetCheckedPressedImage(BOX_BRUSH("Common/Selector", 4.0f / 16.0f, SelectorColor));
+		Style->Set("ToggleButtonRowStyle", ToggleButtonRowStyle);
+
+		/* A radio button is actually just a SCheckBox box with different images */
+		/* Set images for various radio button (SCheckBox) states ... */
+		const FCheckBoxStyle BasicRadioButtonStyle = FCheckBoxStyle()
+			.SetUncheckedImage( IMAGE_BRUSH( "Common/RadioButton_Unselected_16x", CoreStyleConstants::Icon16x16 ) )
+			.SetUncheckedHoveredImage( IMAGE_BRUSH( "Common/RadioButton_Unselected_16x", CoreStyleConstants::Icon16x16 ) )
+			.SetUncheckedPressedImage( IMAGE_BRUSH( "Common/RadioButton_Unselected_16x", CoreStyleConstants::Icon16x16 ) )
+			.SetCheckedImage( IMAGE_BRUSH( "Common/RadioButton_Selected_16x", CoreStyleConstants::Icon16x16 ) )
+			.SetCheckedHoveredImage( IMAGE_BRUSH( "Common/RadioButton_Selected_16x", CoreStyleConstants::Icon16x16, SelectionColor ) )
+			.SetCheckedPressedImage( IMAGE_BRUSH( "Common/RadioButton_Unselected_16x", CoreStyleConstants::Icon16x16, SelectionColor_Pressed ) )
+			.SetUndeterminedImage( IMAGE_BRUSH( "Common/RadioButton_Unselected_16x", CoreStyleConstants::Icon16x16 ) )
+			.SetUndeterminedHoveredImage( IMAGE_BRUSH( "Common/RadioButton_Unselected_16x", CoreStyleConstants::Icon16x16, SelectionColor ) )
+			.SetUndeterminedPressedImage( IMAGE_BRUSH( "Common/RadioButton_Unselected_16x", CoreStyleConstants::Icon16x16, SelectionColor_Pressed ) );
+		Style->Set( "RadioButton", BasicRadioButtonStyle );
+	}
+
+	// SEditableText defaults...
+	{
+		FSlateBrush* SelectionBackground = new BOX_BRUSH( "Common/EditableTextSelectionBackground",  FMargin(4.f/16.f) );
+		FSlateBrush* SelectionTarget = new BOX_BRUSH( "Old/DashedBorder", FMargin(6.0f/32.0f), FLinearColor( 0.0f, 0.0f, 0.0f, 0.75f ) );
+		FSlateBrush* CompositionBackground = new BORDER_BRUSH( "Old/HyperlinkDotted",  FMargin(0,0,0,3/16.0f) );
+
+		const FEditableTextStyle NormalEditableTextStyle = FEditableTextStyle()
+			.SetBackgroundImageSelected( *SelectionBackground )
+			.SetBackgroundImageComposing( *CompositionBackground )
+			.SetCaretImage( *GenericWhiteBox );
+		Style->Set( "NormalEditableText", NormalEditableTextStyle );
+
+		Style->Set( "EditableText.SelectionBackground", SelectionBackground );
+		Style->Set( "EditableText.SelectionTarget", SelectionTarget );
+		Style->Set( "EditableText.CompositionBackground", CompositionBackground );
+	}
+
+	// SEditableTextBox defaults...
+	const FEditableTextBoxStyle NormalEditableTextBoxStyle = FEditableTextBoxStyle()
+		.SetBackgroundImageNormal( BOX_BRUSH( "Common/TextBox", FMargin(4.0f/16.0f) ) )
+		.SetBackgroundImageHovered( BOX_BRUSH( "Common/TextBox_Hovered", FMargin(4.0f/16.0f) ) )
+		.SetBackgroundImageFocused( BOX_BRUSH( "Common/TextBox_Hovered", FMargin(4.0f/16.0f) ) )
+		.SetBackgroundImageReadOnly( BOX_BRUSH( "Common/TextBox_ReadOnly", FMargin(4.0f/16.0f) ) )
+		.SetScrollBarStyle( ScrollBar );
+	{
+		Style->Set( "NormalEditableTextBox", NormalEditableTextBoxStyle );
+		// "NormalFont".
+	}
+
+	const FEditableTextBoxStyle DarkEditableTextBoxStyle = FEditableTextBoxStyle()
+		.SetBackgroundImageNormal(BOX_BRUSH("Common/TextBox_Dark", FMargin(4.0f / 16.0f)))
+		.SetBackgroundImageHovered(BOX_BRUSH("Common/TextBox_Hovered_Dark", FMargin(4.0f / 16.0f)))
+		.SetBackgroundImageFocused(BOX_BRUSH("Common/TextBox_Hovered_Dark", FMargin(4.0f / 16.0f)))
+		.SetBackgroundImageReadOnly(BOX_BRUSH("Common/TextBox_ReadOnly", FMargin(4.0f / 16.0f)))
+		.SetScrollBarStyle(ScrollBar);
+	{
+		Style->Set("DarkEditableTextBox", DarkEditableTextBoxStyle);
+		// "NormalFont".
+	}
+
+	// SSpinBox defaults...
+	{
+		// "NormalFont".
+
+		// "InvertedForeground".
+	}
+
+	// SSuggestionTextBox defaults...
+	{
+		// "NormalFont".
+
+		// "InvertedForeground".
+	}
+
+	// STextBlock defaults...
+	{
+		Style->Set("NormalText", NormalText);
+		Style->Set("NormalUnderlinedText", NormalUnderlinedText);
+
+		Style->Set("SmallText", SmallText);
+		Style->Set("SmallUnderlinedText", SmallUnderlinedText);
+	}
+
+	// SInlineEditableTextBlock
+	{
+		FTextBlockStyle InlineEditableTextBlockReadOnly = FTextBlockStyle(NormalText)
+			.SetColorAndOpacity( FSlateColor::UseForeground() )
+			.SetShadowOffset( FVector2D::ZeroVector )
+			.SetShadowColorAndOpacity( FLinearColor::Black );
+
+		FTextBlockStyle InlineEditableTextBlockSmallReadOnly = FTextBlockStyle(InlineEditableTextBlockReadOnly)
+			.SetFont(SmallText.Font);
+
+		FEditableTextBoxStyle InlineEditableTextBlockEditable = FEditableTextBoxStyle()
+			.SetFont(NormalText.Font)
+			.SetBackgroundImageNormal( BOX_BRUSH( "Common/TextBox", FMargin(4.0f/16.0f) ) )
+			.SetBackgroundImageHovered( BOX_BRUSH( "Common/TextBox_Hovered", FMargin(4.0f/16.0f) ) )
+			.SetBackgroundImageFocused( BOX_BRUSH( "Common/TextBox_Hovered", FMargin(4.0f/16.0f) ) )
+			.SetBackgroundImageReadOnly( BOX_BRUSH( "Common/TextBox_ReadOnly", FMargin(4.0f/16.0f) ) )
+			.SetScrollBarStyle( ScrollBar );
+
+		FEditableTextBoxStyle InlineEditableTextBlockSmallEditable = FEditableTextBoxStyle(InlineEditableTextBlockEditable)
+			.SetFont(SmallText.Font);
+
+		FInlineEditableTextBlockStyle InlineEditableTextBlockStyle = FInlineEditableTextBlockStyle()
+			.SetTextStyle(InlineEditableTextBlockReadOnly)
+			.SetEditableTextBoxStyle(InlineEditableTextBlockEditable);
+
+		Style->Set( "InlineEditableTextBlockStyle", InlineEditableTextBlockStyle );
+
+		FInlineEditableTextBlockStyle InlineEditableTextBlockSmallStyle = FInlineEditableTextBlockStyle()
+			.SetTextStyle(InlineEditableTextBlockSmallReadOnly)
+			.SetEditableTextBoxStyle(InlineEditableTextBlockSmallEditable);
+
+		Style->Set("InlineEditableTextBlockSmallStyle", InlineEditableTextBlockSmallStyle);
+	}
+
+	// SSuggestionTextBox defaults...
+	{
+		Style->Set( "SuggestionTextBox.Background",	new BOX_BRUSH( "Old/Menu_Background", FMargin(8.0f/64.0f) ) );
+		Style->Set( "SuggestionTextBox.Text", FTextBlockStyle()
+			.SetFont( DEFAULT_FONT( "Regular", 9 ) )
+			.SetColorAndOpacity( FLinearColor(FColor(0xffaaaaaa)) )
+			);
+	}
+
+	// SToolTip defaults...
+	{
+		Style->Set("ToolTip.Font", DEFAULT_FONT("Regular", 8));
+		Style->Set("ToolTip.Background", new BOX_BRUSH("Old/ToolTip_Background", FMargin(8.0f / 64.0f)));
+
+		Style->Set("ToolTip.LargerFont", DEFAULT_FONT("Regular", 9));
+		Style->Set("ToolTip.BrightBackground", new BOX_BRUSH("Old/ToolTip_BrightBackground", FMargin(8.0f / 64.0f)));
+	}
+
+	// SBorder defaults...
+	{
+		Style->Set( "Border", new BORDER_BRUSH( "Old/Border", 4.0f/16.0f ) );
+	}
+
+	// SHyperlink defaults...
+	{
+		FButtonStyle HyperlinkButton = FButtonStyle()
+			.SetNormal(BORDER_BRUSH("Old/HyperlinkDotted", FMargin(0,0,0,3/16.0f) ) )
+			.SetPressed(FSlateNoResource() )
+			.SetHovered(BORDER_BRUSH("Old/HyperlinkUnderline", FMargin(0,0,0,3/16.0f) ) );
+
+		FHyperlinkStyle Hyperlink = FHyperlinkStyle()
+			.SetUnderlineStyle(HyperlinkButton)
+			.SetTextStyle(NormalText)
+			.SetPadding(FMargin(0.0f));
+		Style->Set("Hyperlink", Hyperlink);
+	}
+
+	// SProgressBar defaults...
+	{
+		Style->Set( "ProgressBar", FProgressBarStyle()
+			.SetBackgroundImage( BOX_BRUSH( "Common/ProgressBar_Background", FMargin(5.f/12.f) ) )
+			.SetFillImage( BOX_BRUSH( "Common/ProgressBar_Fill", FMargin(5.f/12.f), FLinearColor( 1.0f, 0.22f, 0.0f )  ) )
+			.SetMarqueeImage( IMAGE_BRUSH( "Common/ProgressBar_Marquee", FVector2D(20,12), FLinearColor::White, ESlateBrushTileType::Horizontal ) )
+			);
+	}
+
+	// SThrobber, SCircularThrobber defaults...
+	{
+		Style->Set( "Throbber.Chunk", new IMAGE_BRUSH( "Common/Throbber_Piece", FVector2D(16,16) ) );
+		Style->Set( "Throbber.CircleChunk", new IMAGE_BRUSH( "Common/Throbber_Piece", FVector2D(8,8) ) );
+	}
+
+	// SExpandableArea defaults...
+	{
+		Style->Set( "ExpandableArea", FExpandableAreaStyle()
+			.SetCollapsedImage( IMAGE_BRUSH( "Common/TreeArrow_Collapsed", CoreStyleConstants::Icon10x10, DefaultForeground ) )
+			.SetExpandedImage( IMAGE_BRUSH( "Common/TreeArrow_Expanded", CoreStyleConstants::Icon10x10, DefaultForeground ) )
+			);
+		Style->Set( "ExpandableArea.TitleFont", DEFAULT_FONT( "Bold", 8 ) );
+		Style->Set( "ExpandableArea.Border", new BOX_BRUSH( "Common/GroupBorder", FMargin(4.0f/16.0f) ) );
+	}
+
+	// SSearchBox defaults...
+	{
+		const FEditableTextBoxStyle SpecialEditableTextBoxStyle = FEditableTextBoxStyle()
+			.SetBackgroundImageNormal( BOX_BRUSH( "Common/TextBox_Special", FMargin(8.0f/32.0f) ) )
+			.SetBackgroundImageHovered( BOX_BRUSH( "Common/TextBox_Special_Hovered", FMargin(8.0f/32.0f) ) )
+			.SetBackgroundImageFocused( BOX_BRUSH( "Common/TextBox_Special_Hovered", FMargin(8.0f/32.0f) ) )
+			.SetBackgroundImageReadOnly( BOX_BRUSH( "Common/TextBox_ReadOnly", FMargin(4.0f/16.0f) ) )
+			.SetScrollBarStyle( ScrollBar );
+
+		Style->Set( "SearchBox", FSearchBoxStyle()
+			.SetTextBoxStyle( SpecialEditableTextBoxStyle )
+			.SetUpArrowImage( IMAGE_BRUSH( "Common/UpArrow", CoreStyleConstants::Icon8x8 ) )
+			.SetDownArrowImage( IMAGE_BRUSH( "Common/DownArrow", CoreStyleConstants::Icon8x8 ) )
+			.SetGlassImage( IMAGE_BRUSH( "Common/SearchGlass", CoreStyleConstants::Icon16x16 ) )
+			.SetClearImage( IMAGE_BRUSH( "Common/X", CoreStyleConstants::Icon16x16 ) )
+			);
+	}
+
+	// SSlider and SVolumeControl defaults...
+	{
+		FSliderStyle SliderStyle = FSliderStyle()
+			.SetNormalBarImage(FSlateColorBrush(FColor::White))
+			.SetHoveredBarImage(FSlateColorBrush(FColor::White))
+			.SetDisabledBarImage(FSlateColorBrush(FLinearColor::Gray))
+			.SetNormalThumbImage( IMAGE_BRUSH( "Common/Button", FVector2D(8.0f, 14.0f) ) )
+			.SetHoveredThumbImage(IMAGE_BRUSH("Common/Button", FVector2D(8.0f, 14.0f)))
+			.SetDisabledThumbImage( IMAGE_BRUSH( "Common/Button_Disabled", FVector2D(8.0f, 14.0f) ) )
+			.SetBarThickness(2.0f);
+		Style->Set( "Slider", SliderStyle );
+
+		Style->Set( "VolumeControl", FVolumeControlStyle()
+			.SetSliderStyle( SliderStyle )
+			.SetHighVolumeImage( IMAGE_BRUSH( "Common/VolumeControl_High", CoreStyleConstants::Icon16x16 ) )
+			.SetMidVolumeImage( IMAGE_BRUSH( "Common/VolumeControl_Mid", CoreStyleConstants::Icon16x16 ) )
+			.SetLowVolumeImage( IMAGE_BRUSH( "Common/VolumeControl_Low", CoreStyleConstants::Icon16x16 ) )
+			.SetNoVolumeImage( IMAGE_BRUSH( "Common/VolumeControl_Off", CoreStyleConstants::Icon16x16 ) )
+			.SetMutedImage( IMAGE_BRUSH( "Common/VolumeControl_Muted", CoreStyleConstants::Icon16x16 ) )
+			);
+	}
+
+	// SSpinBox defaults...
+	{
+		Style->Set( "SpinBox", FSpinBoxStyle()
+			.SetBackgroundBrush( BOX_BRUSH( "Common/Spinbox", FMargin(4.0f/16.0f) ) )
+			.SetHoveredBackgroundBrush( BOX_BRUSH( "Common/Spinbox_Hovered", FMargin(4.0f/16.0f) ) )
+			.SetActiveFillBrush( BOX_BRUSH( "Common/Spinbox_Fill", FMargin(4.0f/16.0f, 4.0f/16.0f, 8.0f/16.0f, 4.0f/16.0f) ) )
+			.SetInactiveFillBrush( BOX_BRUSH( "Common/Spinbox_Fill_Hovered", FMargin(4.0f/16.0f) ) )
+			.SetArrowsImage( IMAGE_BRUSH( "Common/SpinArrows", CoreStyleConstants::Icon12x12 ) )
+			.SetForegroundColor( InvertedForeground )
+			);
+	}
+
+	// SNumericEntryBox defaults...
+	{
+		Style->Set( "NumericEntrySpinBox", FSpinBoxStyle()
+			.SetBackgroundBrush( FSlateNoResource() )
+			.SetHoveredBackgroundBrush( FSlateNoResource() )
+			.SetActiveFillBrush( BOX_BRUSH( "Common/Spinbox_Fill_Hovered", FMargin(4.0f/16.0f) ) )
+			.SetInactiveFillBrush( BOX_BRUSH( "Common/Spinbox_Fill", FMargin(4.0f/16.0f, 4.0f/16.0f, 8.0f/16.0f, 4.0f/16.0f) ) )
+			.SetArrowsImage( IMAGE_BRUSH( "Common/SpinArrows", CoreStyleConstants::Icon12x12 ) )
+			.SetTextPadding( FMargin(0.0f) )
+			.SetForegroundColor( InvertedForeground )
+			);
+
+		Style->Set("NumericEntrySpinBox_Dark", FSpinBoxStyle()
+			.SetBackgroundBrush(FSlateNoResource())
+			.SetHoveredBackgroundBrush(FSlateNoResource())
+			.SetActiveFillBrush(BOX_BRUSH("Common/Spinbox_Fill_Hovered_Dark", FMargin(4.0f / 16.0f)))
+			.SetInactiveFillBrush(BOX_BRUSH("Common/Spinbox_Fill_Dark", FMargin(4.0f / 16.0f, 4.0f / 16.0f, 8.0f / 16.0f, 4.0f / 16.0f)))
+			.SetArrowsImage(IMAGE_BRUSH("Common/SpinArrows", CoreStyleConstants::Icon12x12))
+			.SetTextPadding(FMargin(0.0f))
+			.SetForegroundColor(InvertedForeground)
+			);
+
+		Style->Set( "NumericEntrySpinBox.Decorator", new BOX_BRUSH( "Common/TextBoxLabelBorder", FMargin(5.0f/16.0f) ) );
+
+		Style->Set( "NumericEntrySpinBox.NarrowDecorator", new BOX_BRUSH( "Common/TextBoxLabelBorder", FMargin(2.0f/16.0f, 4.0f/16.0f, 2.0f/16.0f, 4.0f/16.0f) ) );
+	}
+
+	// SColorPicker defaults...
+	{
+		Style->Set( "ColorPicker.Border", new BOX_BRUSH( "Common/GroupBorder", FMargin(4.0f/16.0f) ) );
+		Style->Set( "ColorPicker.AlphaBackground", new IMAGE_BRUSH( "Common/Checker", CoreStyleConstants::Icon16x16, FLinearColor::White, ESlateBrushTileType::Both ) );
+		Style->Set( "ColorPicker.EyeDropper", new IMAGE_BRUSH( "CoreStyleConstants::Icons/eyedropper_16px", CoreStyleConstants::Icon16x16) );
+		Style->Set( "ColorPicker.Font", DEFAULT_FONT( "Regular", 10 ) );
+		Style->Set( "ColorPicker.Mode", new IMAGE_BRUSH( "Common/ColorPicker_Mode_16x", CoreStyleConstants::Icon16x16) );
+		Style->Set( "ColorPicker.Separator", new IMAGE_BRUSH( "Common/ColorPicker_Separator", FVector2D(2.0f, 2.0f) ) );
+		Style->Set( "ColorPicker.Selector", new IMAGE_BRUSH( "Common/Circle", FVector2D(8, 8) ) );
+		Style->Set( "ColorPicker.Slider", FSliderStyle()
+			.SetDisabledThumbImage( IMAGE_BRUSH( "Common/ColorPicker_SliderHandle", FVector2D(8.0f, 32.0f)) )
+			.SetNormalThumbImage( IMAGE_BRUSH( "Common/ColorPicker_SliderHandle", FVector2D(8.0f, 32.0f)) )
+			.SetHoveredThumbImage( IMAGE_BRUSH("Common/ColorPicker_SliderHandle", FVector2D(8.0f, 32.0f)) )
+			);
+	}
+
+	// SColorSpectrum defaults...
+	{
+		Style->Set( "ColorSpectrum.Spectrum", new IMAGE_BRUSH( "Common/ColorSpectrum", FVector2D(256, 256) ) );
+		Style->Set( "ColorSpectrum.Selector", new IMAGE_BRUSH( "Common/Circle", FVector2D(8, 8) ) );
+	}
+
+	// SColorThemes defaults...
+	{
+		Style->Set( "ColorThemes.DeleteButton", new IMAGE_BRUSH( "Common/X", CoreStyleConstants::Icon16x16) );
+	}
+
+	// SColorWheel defaults...
+	{
+		Style->Set( "ColorWheel.HueValueCircle", new IMAGE_BRUSH( "Common/ColorWheel", FVector2D(192, 192) ) );
+		Style->Set( "ColorWheel.Selector", new IMAGE_BRUSH( "Common/Circle", FVector2D(8, 8) ) );
+	}
+
+	// SColorGradingWheel defaults...
+	{
+		Style->Set("ColorGradingWheel.HueValueCircle", new IMAGE_BRUSH("Common/ColorGradingWheel", FVector2D(192, 192)));
+		Style->Set("ColorGradingWheel.Selector", new IMAGE_BRUSH("Common/Circle", FVector2D(8, 8)));
+	}
+
+	// SSplitter
+	{
+		Style->Set( "Splitter", FSplitterStyle()
+			.SetHandleNormalBrush( FSlateNoResource() )
+			.SetHandleHighlightBrush( IMAGE_BRUSH( "Common/SplitterHandleHighlight", CoreStyleConstants::Icon8x8, FLinearColor::White ) )
+			);
+	}
+
+	// TableView defaults...
+	{
+		const FTableRowStyle DefaultTableRowStyle = FTableRowStyle()
+			.SetEvenRowBackgroundBrush(FSlateNoResource())
+			.SetEvenRowBackgroundHoveredBrush(IMAGE_BRUSH("Common/Selection", CoreStyleConstants::Icon8x8, FLinearColor(1.0f, 1.0f, 1.0f, 0.1f)))
+			.SetOddRowBackgroundBrush(FSlateNoResource())
+			.SetOddRowBackgroundHoveredBrush(IMAGE_BRUSH("Common/Selection", CoreStyleConstants::Icon8x8, FLinearColor(1.0f, 1.0f, 1.0f, 0.1f)))
+			.SetSelectorFocusedBrush(BORDER_BRUSH("Common/Selector", FMargin(4.f / 16.f), SelectorColor))
+			.SetActiveBrush(IMAGE_BRUSH("Common/Selection", CoreStyleConstants::Icon8x8, SelectionColor))
+			.SetActiveHoveredBrush(IMAGE_BRUSH("Common/Selection", CoreStyleConstants::Icon8x8, SelectionColor))
+			.SetInactiveBrush(IMAGE_BRUSH("Common/Selection", CoreStyleConstants::Icon8x8, SelectionColor_Inactive))
+			.SetInactiveHoveredBrush(IMAGE_BRUSH("Common/Selection", CoreStyleConstants::Icon8x8, SelectionColor_Inactive))
+			.SetActiveHighlightedBrush(IMAGE_BRUSH("Common/Selection", CoreStyleConstants::Icon8x8, HighlightColor))
+			.SetInactiveHighlightedBrush(IMAGE_BRUSH("Common/Selection", CoreStyleConstants::Icon8x8, FSlateColor(FLinearColor(.1f, .1f, .1f))))
+
+			.SetTextColor(DefaultForeground)
+			.SetSelectedTextColor(InvertedForeground)
+			.SetDropIndicator_Above(BOX_BRUSH("Common/DropZoneIndicator_Above", FMargin(10.0f / 16.0f, 10.0f / 16.0f, 0, 0), SelectionColor))
+			.SetDropIndicator_Onto(BOX_BRUSH("Common/DropZoneIndicator_Onto", FMargin(4.0f / 16.0f), SelectionColor))
+			.SetDropIndicator_Below(BOX_BRUSH("Common/DropZoneIndicator_Below", FMargin(10.0f / 16.0f, 0, 0, 10.0f / 16.0f), SelectionColor));
+		Style->Set("TableView.Row", DefaultTableRowStyle);
+
+		// Make this the default for the ComboBox rows also
+		Style->Set("ComboBox.Row", DefaultTableRowStyle);  
+
+		const FTableRowStyle DarkTableRowStyle = FTableRowStyle(DefaultTableRowStyle)
+			.SetEvenRowBackgroundBrush(IMAGE_BRUSH("Common/Selection", CoreStyleConstants::Icon8x8, FLinearColor(0.0f, 0.0f, 0.0f, 0.1f)))
+			.SetOddRowBackgroundBrush(IMAGE_BRUSH("Common/Selection", CoreStyleConstants::Icon8x8, FLinearColor(0.0f, 0.0f, 0.0f, 0.1f)));
+		Style->Set("TableView.DarkRow", DarkTableRowStyle);
+
+		Style->Set( "TreeArrow_Collapsed", new IMAGE_BRUSH( "Common/TreeArrow_Collapsed", CoreStyleConstants::Icon10x10, DefaultForeground ) );
+		Style->Set( "TreeArrow_Collapsed_Hovered", new IMAGE_BRUSH( "Common/TreeArrow_Collapsed_Hovered", CoreStyleConstants::Icon10x10, DefaultForeground ) );
+		Style->Set( "TreeArrow_Expanded", new IMAGE_BRUSH( "Common/TreeArrow_Expanded", CoreStyleConstants::Icon10x10, DefaultForeground ) );
+		Style->Set( "TreeArrow_Expanded_Hovered", new IMAGE_BRUSH( "Common/TreeArrow_Expanded_Hovered", CoreStyleConstants::Icon10x10, DefaultForeground ) );
+
+		const FTableColumnHeaderStyle TableColumnHeaderStyle = FTableColumnHeaderStyle()
+			.SetSortPrimaryAscendingImage(IMAGE_BRUSH("Common/SortUpArrow", CoreStyleConstants::Icon8x4))
+			.SetSortPrimaryDescendingImage(IMAGE_BRUSH("Common/SortDownArrow", CoreStyleConstants::Icon8x4))
+			.SetSortSecondaryAscendingImage(IMAGE_BRUSH("Common/SortUpArrows", CoreStyleConstants::Icon16x4))
+			.SetSortSecondaryDescendingImage(IMAGE_BRUSH("Common/SortDownArrows", CoreStyleConstants::Icon16x4))
+			.SetNormalBrush( BOX_BRUSH( "Common/ColumnHeader", 4.f/32.f ) )
+			.SetHoveredBrush( BOX_BRUSH( "Common/ColumnHeader_Hovered", 4.f/32.f ) )
+			.SetMenuDropdownImage( IMAGE_BRUSH( "Common/ColumnHeader_Arrow", CoreStyleConstants::Icon8x8 ) )
+			.SetMenuDropdownNormalBorderBrush( BOX_BRUSH( "Common/ColumnHeaderMenuButton_Normal", 4.f/32.f ) )
+			.SetMenuDropdownHoveredBorderBrush( BOX_BRUSH( "Common/ColumnHeaderMenuButton_Hovered", 4.f/32.f ) );
+		Style->Set( "TableView.Header.Column", TableColumnHeaderStyle );
+
+		const FTableColumnHeaderStyle TableLastColumnHeaderStyle = FTableColumnHeaderStyle()
+			.SetSortPrimaryAscendingImage(IMAGE_BRUSH("Common/SortUpArrow", CoreStyleConstants::Icon8x4))
+			.SetSortPrimaryDescendingImage(IMAGE_BRUSH("Common/SortDownArrow", CoreStyleConstants::Icon8x4))
+			.SetSortSecondaryAscendingImage(IMAGE_BRUSH("Common/SortUpArrows", CoreStyleConstants::Icon16x4))
+			.SetSortSecondaryDescendingImage(IMAGE_BRUSH("Common/SortDownArrows", CoreStyleConstants::Icon16x4))
+			.SetNormalBrush( FSlateNoResource() )
+			.SetHoveredBrush( BOX_BRUSH( "Common/LastColumnHeader_Hovered", 4.f/32.f ) )
+			.SetMenuDropdownImage( IMAGE_BRUSH( "Common/ColumnHeader_Arrow", CoreStyleConstants::Icon8x8 ) )
+			.SetMenuDropdownNormalBorderBrush( BOX_BRUSH( "Common/ColumnHeaderMenuButton_Normal", 4.f/32.f ) )
+			.SetMenuDropdownHoveredBorderBrush( BOX_BRUSH( "Common/ColumnHeaderMenuButton_Hovered", 4.f/32.f ) );
+
+		const FSplitterStyle TableHeaderSplitterStyle = FSplitterStyle()
+			.SetHandleNormalBrush( FSlateNoResource() )
+			.SetHandleHighlightBrush( IMAGE_BRUSH( "Common/HeaderSplitterGrip", CoreStyleConstants::Icon8x8 ) );
+
+		Style->Set( "TableView.Header", FHeaderRowStyle()
+			.SetColumnStyle( TableColumnHeaderStyle )
+			.SetLastColumnStyle( TableLastColumnHeaderStyle )
+			.SetColumnSplitterStyle( TableHeaderSplitterStyle )
+			.SetBackgroundBrush( BOX_BRUSH( "Common/TableViewHeader", 4.f/32.f ) )
+			.SetForegroundColor( DefaultForeground )
+			);
+	}
+
+	// MultiBox
+	{
+		Style->Set( "MultiBox.GenericToolBarCoreStyleConstants::Icon", new IMAGE_BRUSH( "CoreStyleConstants::Icons/CoreStyleConstants::Icon_generic_toolbar", CoreStyleConstants::Icon40x40 ) );
+		Style->Set( "MultiBox.GenericToolBarCoreStyleConstants::Icon.Small", new IMAGE_BRUSH( "CoreStyleConstants::Icons/CoreStyleConstants::Icon_generic_toolbar", CoreStyleConstants::Icon20x20 ) );
+
+		Style->Set( "MultiBox.DeleteButton", FButtonStyle() 
+			.SetNormal ( IMAGE_BRUSH( "/Docking/CloseApp_Hovered", CoreStyleConstants::Icon16x16 ) )
+			.SetPressed( IMAGE_BRUSH( "/Docking/CloseApp_Pressed", CoreStyleConstants::Icon16x16 ) )
+			.SetHovered( IMAGE_BRUSH( "/Docking/CloseApp_Hovered", CoreStyleConstants::Icon16x16 ) )
+			);
+
+		Style->Set( "MultiboxHookColor", FLinearColor(0.f, 1.f, 0.f, 1.f) );
+	}
+
+	// ToolBar
+	{
+		Style->Set( "ToolBar.Background", new BOX_BRUSH( "Common/GroupBorder", FMargin(4.0f/16.0f) ) );
+		Style->Set( "ToolBar.CoreStyleConstants::Icon", new IMAGE_BRUSH( "CoreStyleConstants::Icons/CoreStyleConstants::Icon_tab_toolbar_16px", CoreStyleConstants::Icon16x16 ) );
+		Style->Set( "ToolBar.Expand", new IMAGE_BRUSH( "CoreStyleConstants::Icons/toolbar_expand_16x", CoreStyleConstants::Icon16x16) );
+		Style->Set( "ToolBar.SubMenuIndicator", new IMAGE_BRUSH( "Common/SubmenuArrow", CoreStyleConstants::Icon8x8 ) );
+		Style->Set( "ToolBar.SToolBarComboButtonBlock.Padding", FMargin(4.0f,0.0f));
+		Style->Set( "ToolBar.SToolBarButtonBlock.Padding", FMargin(4.0f,0.0f));
+		Style->Set( "ToolBar.SToolBarCheckComboButtonBlock.Padding", FMargin(4.0f,0.0f));
+		Style->Set( "ToolBar.SToolBarButtonBlock.CheckBox.Padding", FMargin(4.0f,0.0f) );
+		Style->Set( "ToolBar.SToolBarComboButtonBlock.ComboButton.Color", DefaultForeground );
+
+		Style->Set( "ToolBar.Block.IndentedPadding", FMargin( 18.0f, 2.0f, 4.0f, 4.0f ) );
+		Style->Set( "ToolBar.Block.Padding", FMargin( 2.0f, 2.0f, 4.0f, 4.0f ) );
+
+		Style->Set( "ToolBar.Separator", new FSlateColorBrush( FLinearColor(FColor(48, 48, 48)) ) );
+		Style->Set( "ToolBar.Separator.Padding", FMargin( 8.f, 0.f, 8.f, 0.f) );
+
+		Style->Set( "ToolBar.Label", FTextBlockStyle(NormalText) .SetFont( DEFAULT_FONT( "Regular", 9 ) ) );
+		Style->Set( "ToolBar.EditableText", FEditableTextBoxStyle(NormalEditableTextBoxStyle) .SetFont( DEFAULT_FONT( "Regular", 9 ) ) );
+		Style->Set( "ToolBar.Keybinding", FTextBlockStyle(NormalText) .SetFont( DEFAULT_FONT( "Regular", 8 ) ) );
+
+		Style->Set( "ToolBar.Heading", FTextBlockStyle(NormalText)
+			.SetFont( DEFAULT_FONT( "Regular", 8 ) )
+			.SetColorAndOpacity( FLinearColor( 0.4f, 0.4, 0.4f, 1.0f ) ) );
+
+		/* Create style for "ToolBar.CheckBox" widget ... */
+		const FCheckBoxStyle ToolBarCheckBoxStyle = FCheckBoxStyle()
+			.SetUncheckedImage( IMAGE_BRUSH( "Common/SmallCheckBox", CoreStyleConstants::Icon14x14 ) )
+			.SetCheckedImage( IMAGE_BRUSH( "Common/SmallCheckBox_Checked", CoreStyleConstants::Icon14x14 ) )
+			.SetUncheckedHoveredImage( IMAGE_BRUSH( "Common/SmallCheckBox_Hovered", CoreStyleConstants::Icon14x14 ) )
+			.SetCheckedImage( IMAGE_BRUSH( "Common/SmallCheckBox_Checked_Hovered", CoreStyleConstants::Icon14x14 ) )
+			.SetUncheckedPressedImage(IMAGE_BRUSH( "Common/SmallCheckBox_Hovered", CoreStyleConstants::Icon14x14, FLinearColor( 0.5f, 0.5f, 0.5f ) ) )
+			.SetCheckedPressedImage( IMAGE_BRUSH( "Common/SmallCheckBox_Checked_Hovered", CoreStyleConstants::Icon14x14, FLinearColor( 0.5f, 0.5f, 0.5f ) ) )
+			.SetUndeterminedImage( IMAGE_BRUSH( "Common/CheckBox_Undetermined", CoreStyleConstants::Icon14x14 ) )
+			.SetUndeterminedHoveredImage( IMAGE_BRUSH( "Common/CheckBox_Undetermined_Hovered", CoreStyleConstants::Icon14x14 ) )
+			.SetUndeterminedPressedImage( IMAGE_BRUSH( "Common/CheckBox_Undetermined_Hovered", CoreStyleConstants::Icon14x14, FLinearColor( 0.5f, 0.5f, 0.5f ) ) );
+		/* ... and add new style */
+		Style->Set( "ToolBar.CheckBox", ToolBarCheckBoxStyle );
+
+		/* Read-only checkbox that appears next to a menu item */
+		/* Set images for various SCheckBox states associated with read-only toolbar check box items... */
+		const FCheckBoxStyle BasicToolBarCheckStyle = FCheckBoxStyle()
+			.SetUncheckedImage(IMAGE_BRUSH("CoreStyleConstants::Icons/Empty_14x", CoreStyleConstants::Icon14x14))
+			.SetUncheckedHoveredImage(IMAGE_BRUSH("CoreStyleConstants::Icons/Empty_14x", CoreStyleConstants::Icon14x14))
+			.SetUncheckedPressedImage(IMAGE_BRUSH("Common/SmallCheckBox_Hovered", CoreStyleConstants::Icon14x14))
+			.SetCheckedImage(IMAGE_BRUSH("Common/SmallCheck", CoreStyleConstants::Icon14x14))
+			.SetCheckedHoveredImage(IMAGE_BRUSH("Common/SmallCheck", CoreStyleConstants::Icon14x14))
+			.SetCheckedPressedImage(IMAGE_BRUSH("Common/SmallCheck", CoreStyleConstants::Icon14x14))
+			.SetUndeterminedImage(IMAGE_BRUSH("CoreStyleConstants::Icons/Empty_14x", CoreStyleConstants::Icon14x14))
+			.SetUndeterminedHoveredImage(FSlateNoResource())
+			.SetUndeterminedPressedImage(FSlateNoResource());
+		Style->Set("ToolBar.Check", BasicToolBarCheckStyle);
+
+		// This radio button is actually just a check box with different images
+		/* Create style for "ToolBar.RadioButton" widget ... */
+		const FCheckBoxStyle ToolbarRadioButtonCheckBoxStyle = FCheckBoxStyle()
+				.SetUncheckedImage( IMAGE_BRUSH( "Common/RadioButton_Unselected_16x", CoreStyleConstants::Icon16x16 ) )
+				.SetCheckedImage( IMAGE_BRUSH( "Common/RadioButton_Selected_16x", CoreStyleConstants::Icon16x16 ) )
+				.SetUncheckedHoveredImage( IMAGE_BRUSH( "Common/RadioButton_Unselected_16x", CoreStyleConstants::Icon16x16, SelectionColor ) )
+				.SetCheckedHoveredImage( IMAGE_BRUSH( "Common/RadioButton_Unselected_16x", CoreStyleConstants::Icon16x16, SelectionColor ) )
+				.SetUncheckedPressedImage( IMAGE_BRUSH( "Common/RadioButton_Unselected_16x", CoreStyleConstants::Icon16x16, SelectionColor_Pressed ) )
+				.SetCheckedPressedImage( IMAGE_BRUSH( "Common/RadioButton_Selected_16x", CoreStyleConstants::Icon16x16, SelectionColor_Pressed ) );
+		/* ... and add new style */
+		Style->Set( "ToolBar.RadioButton", ToolbarRadioButtonCheckBoxStyle );
+
+		/* Create style for "ToolBar.ToggleButton" widget ... */
+		const FCheckBoxStyle ToolBarToggleButtonCheckBoxStyle = FCheckBoxStyle()
+				.SetCheckBoxType( ESlateCheckBoxType::ToggleButton )
+				.SetUncheckedImage( FSlateNoResource() )
+				.SetUncheckedPressedImage( BOX_BRUSH( "Common/RoundedSelection_16x", 4.0f/16.0f, SelectionColor_Pressed ) )
+				.SetUncheckedHoveredImage( BOX_BRUSH( "Common/RoundedSelection_16x", 4.0f/16.0f, SelectionColor ) )
+				.SetCheckedImage( BOX_BRUSH( "Common/RoundedSelection_16x",  4.0f/16.0f, SelectionColor_Pressed ) )
+				.SetCheckedHoveredImage( BOX_BRUSH( "Common/RoundedSelection_16x",  4.0f/16.0f, SelectionColor_Pressed ) )
+				.SetCheckedPressedImage( BOX_BRUSH( "Common/RoundedSelection_16x",  4.0f/16.0f, SelectionColor ) );
+		/* ... and add new style */
+		Style->Set( "ToolBar.ToggleButton", ToolBarToggleButtonCheckBoxStyle );
+
+		Style->Set( "ToolBar.Button", FButtonStyle(Button)
+			.SetNormal ( FSlateNoResource() )
+			.SetPressed( BOX_BRUSH( "Common/RoundedSelection_16x", 4.0f/16.0f, SelectionColor_Pressed ) )
+			.SetHovered( BOX_BRUSH( "Common/RoundedSelection_16x", 4.0f/16.0f, SelectionColor ))
+			.SetNormalForeground(FSlateColor::UseForeground())
+			.SetPressedForeground(FSlateColor::UseForeground())
+			.SetHoveredForeground(FSlateColor::UseForeground())
+			.SetDisabledForeground(FSlateColor::UseForeground())
+		);
+
+	/*	Style->Set( "ToolBar.Button.Normal", new FSlateNoResource() );
+		Style->Set( "ToolBar.Button.Pressed", new BOX_BRUSH( "Common/RoundedSelection_16x", 4.0f/16.0f, SelectionColor_Pressed ) );
+		Style->Set( "ToolBar.Button.Hovered", new BOX_BRUSH( "Common/RoundedSelection_16x", 4.0f/16.0f, SelectionColor ) );
+*/
+
+		Style->Set( "ToolBar.Button.Checked", new BOX_BRUSH( "Common/RoundedSelection_16x",  4.0f/16.0f, SelectionColor_Pressed ) );
+		Style->Set( "ToolBar.Button.Checked_Hovered", new BOX_BRUSH( "Common/RoundedSelection_16x",  4.0f/16.0f, SelectionColor_Pressed ) );
+		Style->Set( "ToolBar.Button.Checked_Pressed", new BOX_BRUSH( "Common/RoundedSelection_16x",  4.0f/16.0f, SelectionColor ) );
+
+		Style->Set("ToolBar.SimpleComboButton", Style->GetWidgetStyle<FComboButtonStyle>("ComboButton"));
+		Style->Set("ToolBar.CoreStyleConstants::IconSize", CoreStyleConstants::Icon20x20);
+	}
+
+	// MenuBar
+	{
+		Style->Set( "Menu.Background",	new BOX_BRUSH( "Old/Menu_Background", FMargin(8.0f/64.0f) ) );
+		Style->Set( "Menu.CoreStyleConstants::Icon", new IMAGE_BRUSH( "CoreStyleConstants::Icons/CoreStyleConstants::Icon_tab_toolbar_16px", CoreStyleConstants::Icon16x16 ) );
+		Style->Set( "Menu.Expand", new IMAGE_BRUSH( "CoreStyleConstants::Icons/toolbar_expand_16x", CoreStyleConstants::Icon16x16) );
+		Style->Set( "Menu.SubMenuIndicator", new IMAGE_BRUSH( "Common/SubmenuArrow", CoreStyleConstants::Icon8x8 ) );
+		Style->Set( "Menu.SToolBarComboButtonBlock.Padding", FMargin(4.0f));
+		Style->Set( "Menu.SToolBarButtonBlock.Padding", FMargin(4.0f));
+		Style->Set( "Menu.SToolBarCheckComboButtonBlock.Padding", FMargin(4.0f));
+		Style->Set( "Menu.SToolBarButtonBlock.CheckBox.Padding", FMargin(0.0f) );
+		Style->Set( "Menu.SToolBarComboButtonBlock.ComboButton.Color", DefaultForeground );
+
+
+		Style->Set( "Menu.Block.IndentedPadding", FMargin( 18.0f, 2.0f, 4.0f, 4.0f ) );
+		Style->Set( "Menu.Block.Padding", FMargin( 2.0f, 2.0f, 4.0f, 4.0f ) );
+
+		Style->Set( "Menu.Separator", new BOX_BRUSH( "Old/Button", 4.0f/32.0f ) );
+		Style->Set( "Menu.Separator.Padding", FMargin( 0.5f ) );
+
+		Style->Set( "Menu.Label", FTextBlockStyle(NormalText) .SetFont( DEFAULT_FONT( "Regular", 9 ) ) );
+		Style->Set( "Menu.EditableText", FEditableTextBoxStyle(NormalEditableTextBoxStyle) .SetFont( DEFAULT_FONT( "Regular", 9 ) ) );
+		Style->Set( "Menu.Keybinding", FTextBlockStyle(NormalText) .SetFont( DEFAULT_FONT( "Regular", 8 ) ) );
+
+		Style->Set( "Menu.Heading", FTextBlockStyle(NormalText)
+			.SetFont( DEFAULT_FONT( "Regular", 8 ) )
+			.SetColorAndOpacity( FLinearColor( 0.4f, 0.4, 0.4f, 1.0f ) ) );
+
+		const FMargin MenuBlockPadding(10.0f, 3.0f, 5.0f, 2.0f);
+		Style->Set("Menu.Heading.Padding", MenuBlockPadding + FMargin(0, 10, 0, 0));
+
+		/* Set images for various SCheckBox states associated with menu check box items... */
+		const FCheckBoxStyle BasicMenuCheckBoxStyle = FCheckBoxStyle()
+			.SetUncheckedImage( IMAGE_BRUSH( "Common/SmallCheckBox", CoreStyleConstants::Icon14x14 ) )
+			.SetUncheckedHoveredImage( IMAGE_BRUSH( "Common/SmallCheckBox_Hovered", CoreStyleConstants::Icon14x14 ) )
+			.SetUncheckedPressedImage( IMAGE_BRUSH( "Common/SmallCheckBox_Hovered", CoreStyleConstants::Icon14x14, FLinearColor( 0.5f, 0.5f, 0.5f ) ) )
+			.SetCheckedImage( IMAGE_BRUSH( "Common/SmallCheckBox_Checked", CoreStyleConstants::Icon14x14 ) )
+			.SetCheckedHoveredImage( IMAGE_BRUSH( "Common/SmallCheckBox_Checked_Hovered", CoreStyleConstants::Icon14x14 ) )
+			.SetCheckedPressedImage( IMAGE_BRUSH( "Common/SmallCheckBox_Checked_Hovered", CoreStyleConstants::Icon14x14, FLinearColor( 0.5f, 0.5f, 0.5f ) ) )
+			.SetUndeterminedImage( IMAGE_BRUSH( "Common/CheckBox_Undetermined", CoreStyleConstants::Icon14x14 ) )
+			.SetUndeterminedHoveredImage( IMAGE_BRUSH( "Common/CheckBox_Undetermined_Hovered", CoreStyleConstants::Icon14x14 ) )
+			.SetUndeterminedPressedImage( IMAGE_BRUSH( "Common/CheckBox_Undetermined_Hovered", CoreStyleConstants::Icon14x14, FLinearColor( 0.5f, 0.5f, 0.5f ) ) );
+ 
+		/* ...and add the new style */
+		Style->Set( "Menu.CheckBox", BasicMenuCheckBoxStyle );
+						
+		/* Read-only checkbox that appears next to a menu item */
+		/* Set images for various SCheckBox states associated with read-only menu check box items... */
+		const FCheckBoxStyle BasicMenuCheckStyle = FCheckBoxStyle()
+			.SetUncheckedImage( IMAGE_BRUSH( "CoreStyleConstants::Icons/Empty_14x", CoreStyleConstants::Icon14x14 ) )
+			.SetUncheckedHoveredImage( IMAGE_BRUSH( "CoreStyleConstants::Icons/Empty_14x", CoreStyleConstants::Icon14x14 ) )
+			.SetUncheckedPressedImage( IMAGE_BRUSH( "Common/SmallCheckBox_Hovered", CoreStyleConstants::Icon14x14 ) )
+			.SetCheckedImage( IMAGE_BRUSH( "Common/SmallCheck", CoreStyleConstants::Icon14x14 ) )
+			.SetCheckedHoveredImage( IMAGE_BRUSH( "Common/SmallCheck", CoreStyleConstants::Icon14x14 ) )
+			.SetCheckedPressedImage( IMAGE_BRUSH( "Common/SmallCheck", CoreStyleConstants::Icon14x14 ) )
+			.SetUndeterminedImage( IMAGE_BRUSH( "CoreStyleConstants::Icons/Empty_14x", CoreStyleConstants::Icon14x14 ) )
+			.SetUndeterminedHoveredImage( FSlateNoResource() )
+			.SetUndeterminedPressedImage( FSlateNoResource() );
+
+		/* ...and add the new style */
+		Style->Set( "Menu.Check", BasicMenuCheckStyle );
+
+		/* This radio button is actually just a check box with different images */
+		/* Set images for various Menu radio button (SCheckBox) states... */
+		const FCheckBoxStyle BasicMenuRadioButtonStyle = FCheckBoxStyle()
+			.SetUncheckedImage( IMAGE_BRUSH( "Common/RadioButton_Unselected_16x", CoreStyleConstants::Icon16x16 ) )
+			.SetUncheckedHoveredImage( IMAGE_BRUSH( "Common/RadioButton_Unselected_16x", CoreStyleConstants::Icon16x16 ) )
+			.SetUncheckedPressedImage( IMAGE_BRUSH( "Common/RadioButton_Unselected_16x", CoreStyleConstants::Icon16x16 ) )
+			.SetCheckedImage( IMAGE_BRUSH( "Common/RadioButton_Selected_16x", CoreStyleConstants::Icon16x16 ) )
+			.SetCheckedHoveredImage( IMAGE_BRUSH( "Common/RadioButton_Selected_16x", CoreStyleConstants::Icon16x16, SelectionColor ) )
+			.SetCheckedPressedImage( IMAGE_BRUSH( "Common/RadioButton_Unselected_16x", CoreStyleConstants::Icon16x16, SelectionColor_Pressed ) )
+			.SetUndeterminedImage( IMAGE_BRUSH( "Common/RadioButton_Unselected_16x", CoreStyleConstants::Icon16x16 ) )
+			.SetUndeterminedHoveredImage( IMAGE_BRUSH( "Common/RadioButton_Unselected_16x", CoreStyleConstants::Icon16x16, SelectionColor ) )
+			.SetUndeterminedPressedImage( IMAGE_BRUSH( "Common/RadioButton_Unselected_16x", CoreStyleConstants::Icon16x16, SelectionColor_Pressed ) );
+
+		/* ...and set new style */
+		Style->Set( "Menu.RadioButton", BasicMenuRadioButtonStyle );
+
+		/* Create style for "Menu.ToggleButton" widget ... */
+		const FCheckBoxStyle MenuToggleButtonCheckBoxStyle = FCheckBoxStyle()
+			.SetCheckBoxType( ESlateCheckBoxType::ToggleButton )
+			.SetUncheckedImage( FSlateNoResource() )
+			.SetUncheckedPressedImage( BOX_BRUSH( "Common/RoundedSelection_16x", 4.0f/16.0f, SelectionColor_Pressed ) )
+			.SetUncheckedHoveredImage( BOX_BRUSH( "Common/RoundedSelection_16x", 4.0f/16.0f, SelectionColor ) )
+			.SetCheckedImage( BOX_BRUSH( "Common/RoundedSelection_16x",  4.0f/16.0f, SelectionColor_Pressed ) )
+			.SetCheckedHoveredImage( BOX_BRUSH( "Common/RoundedSelection_16x",  4.0f/16.0f, SelectionColor_Pressed ) )
+			.SetCheckedPressedImage( BOX_BRUSH( "Common/RoundedSelection_16x",  4.0f/16.0f, SelectionColor ) );
+		/* ... and add new style */
+		Style->Set( "Menu.ToggleButton", MenuToggleButtonCheckBoxStyle );
+
+
+		FButtonStyle MenuButton =
+			FButtonStyle(NoBorder)
+			.SetNormal(FSlateNoResource())
+			.SetPressed(BOX_BRUSH("Common/RoundedSelection_16x", 4.0f / 16.0f, SelectionColor_Pressed))
+			.SetHovered(BOX_BRUSH("Common/RoundedSelection_16x", 4.0f / 16.0f, SelectionColor))
+			.SetNormalPadding(FMargin(0, 1))
+			.SetPressedPadding(FMargin(0, 2, 0, 0));
+
+		Style->Set("Menu.Button", MenuButton);
+
+		Style->Set( "Menu.Button.Checked", new BOX_BRUSH( "Common/RoundedSelection_16x",  4.0f/16.0f, SelectionColor_Pressed ) );
+		Style->Set( "Menu.Button.Checked_Hovered", new BOX_BRUSH( "Common/RoundedSelection_16x",  4.0f/16.0f, SelectionColor_Pressed ) );
+		Style->Set( "Menu.Button.Checked_Pressed", new BOX_BRUSH( "Common/RoundedSelection_16x",  4.0f/16.0f, SelectionColor ) );
+
+		/* The style of a menu bar button when it has a sub menu open */
+		Style->Set( "Menu.Button.SubMenuOpen", new BORDER_BRUSH( "Common/Selection", FMargin(4.f/16.f), FLinearColor(0.10f, 0.10f, 0.10f) ) );
+
+		Style->Set("WindowMenuBar.Background", new BOX_BRUSH("Old/Menu_Background", FMargin(8.0f / 64.0f)));
+		Style->Set("WindowMenuBar.Label", FTextBlockStyle(NormalText).SetFont(DEFAULT_FONT("Regular", 9)));
+		Style->Set("WindowMenuBar.Expand", new IMAGE_BRUSH("CoreStyleConstants::Icons/toolbar_expand_16x", CoreStyleConstants::Icon16x16));
+		Style->Set("WindowMenuBar.Button", MenuButton);
+		Style->Set("WindowMenuBar.Button.SubMenuOpen", new BORDER_BRUSH("Common/Selection", FMargin(4.f / 16.f), FLinearColor(0.10f, 0.10f, 0.10f)));
+
+		Style->Set("WindowMenuBar.MenuBar.Padding", FMargin(12, 2));
+	}
+
+	// SExpandableButton defaults...
+	{
+		Style->Set( "ExpandableButton.Background", new BOX_BRUSH( "Common/Button", 8.0f/32.0f ) );
+		
+		// Extra padding on the right and bottom to account for image shadow
+		Style->Set( "ExpandableButton.Padding", FMargin(3.f, 3.f, 6.f, 6.f) );
+
+		Style->Set( "ExpandableButton.CloseButton", new IMAGE_BRUSH( "Common/ExpansionButton_CloseOverlay", CoreStyleConstants::Icon16x16) );
+	}
+
+	// SBreadcrumbTrail defaults...
+	{
+		Style->Set( "BreadcrumbTrail.Delimiter", new IMAGE_BRUSH( "Common/Delimiter", CoreStyleConstants::Icon16x16 ) );
+
+		Style->Set( "BreadcrumbButton", FButtonStyle()
+			.SetNormal ( FSlateNoResource() )
+			.SetPressed( BOX_BRUSH( "Common/RoundedSelection_16x", 4.0f/16.0f, SelectionColor_Pressed ) )
+			.SetHovered( BOX_BRUSH( "Common/RoundedSelection_16x", 4.0f/16.0f, SelectionColor ) )
+			.SetNormalPadding( FMargin(0,0) )
+			.SetPressedPadding( FMargin(0,0) )
+			);
+	}
+
+	// SWizard defaults
+	{
+		Style->Set("Wizard.PageTitle", FTextBlockStyle(NormalText)
+			.SetFont(DEFAULT_FONT("BoldCondensed", 28))
+			.SetShadowOffset(FVector2D(1, 1))
+			.SetShadowColorAndOpacity(FLinearColor(0, 0, 0, 0.9f))
+			);
+	}
+
+	// SNotificationList defaults...
+	{
+		Style->Set( "NotificationList.FontBold", DEFAULT_FONT( "Bold", 16 ) );
+		Style->Set( "NotificationList.FontLight", DEFAULT_FONT( "Light", 12 ) );
+		Style->Set( "NotificationList.ItemBackground", new BOX_BRUSH( "Old/Menu_Background", FMargin(8.0f/64.0f) ) );
+		Style->Set( "NotificationList.ItemBackground_Border", new BOX_BRUSH( "Old/Menu_Background_Inverted_Border_Bold", FMargin(8.0f/64.0f) ) );
+		Style->Set( "NotificationList.ItemBackground_Border_Transparent", new BOX_BRUSH("Old/Notification_Border_Flash", FMargin(8.0f/64.0f)));
+		Style->Set( "NotificationList.SuccessImage", new IMAGE_BRUSH( "CoreStyleConstants::Icons/notificationlist_success", CoreStyleConstants::Icon16x16 ) );
+		Style->Set( "NotificationList.FailImage", new IMAGE_BRUSH( "CoreStyleConstants::Icons/notificationlist_fail", CoreStyleConstants::Icon16x16 ) );
+		Style->Set( "NotificationList.DefaultMessage", new IMAGE_BRUSH( "Common/EventMessage_Default", CoreStyleConstants::Icon40x40 ) );
+	}
+
+	// SSeparator defaults...
+	{
+		Style->Set( "Separator", new BOX_BRUSH( "Common/Separator", 1/4.0f, FLinearColor(1,1,1,0.5f) ) );
+	}
+
+	// SHeader defaults...
+	{
+		Style->Set( "Header.Pre", new BOX_BRUSH( "Common/Separator", FMargin(1/4.0f,0,2/4.0f,0), FLinearColor(1,1,1,0.5f) ) );
+		Style->Set( "Header.Post", new BOX_BRUSH( "Common/Separator", FMargin(2/4.0f,0,1/4.0f,0), FLinearColor(1,1,1,0.5f) ) );
+	}
+
+	// SDockTab, SDockingTarget, SDockingTabStack defaults...
+	{
+		Style->Set( "Docking.Background", new BOX_BRUSH( "Old/Menu_Background", FMargin(8.0f/64.0f) ) );
+		Style->Set( "Docking.Border", new BOX_BRUSH( "Common/GroupBorder", FMargin(4.0f/16.0f) ) );
+
+		FTextBlockStyle DockingTabFont = FTextBlockStyle(NormalText)
+			.SetFont( DEFAULT_FONT( "Regular", 9 ) )
+			.SetColorAndOpacity( FLinearColor(0.72f, 0.72f, 0.72f, 1.f) )
+			.SetShadowOffset( FVector2D( 1,1 ) )
+			.SetShadowColorAndOpacity( FLinearColor::Black ) ;
+
+		Style->Set( "Docking.UnhideTabwellButton", FButtonStyle(Button)
+			.SetNormal ( IMAGE_BRUSH( "/Docking/ShowTabwellButton_Normal", FVector2D(10,10) ) )
+			.SetPressed( IMAGE_BRUSH( "/Docking/ShowTabwellButton_Pressed", FVector2D(10,10) ) )
+			.SetHovered( IMAGE_BRUSH( "/Docking/ShowTabwellButton_Hovered", FVector2D(10,10) ) )
+			.SetNormalPadding(0)
+			.SetPressedPadding(0)
+			);
+
+		// Flash using the selection color for consistency with the rest of the UI scheme
+		const FSlateColor& TabFlashColor = SelectionColor;
+
+		const FButtonStyle CloseButton = FButtonStyle()
+			.SetNormal ( IMAGE_BRUSH( "/Docking/CloseApp_Normal", CoreStyleConstants::Icon16x16 ) )
+			.SetPressed( IMAGE_BRUSH( "/Docking/CloseApp_Pressed", CoreStyleConstants::Icon16x16 ) )
+			.SetHovered( IMAGE_BRUSH( "/Docking/CloseApp_Hovered", CoreStyleConstants::Icon16x16 ) );
+
+
+		FLinearColor DockColor_Inactive(FColor(45, 45, 45));
+		FLinearColor DockColor_Hovered(FColor(54, 54, 54));
+		FLinearColor DockColor_Active(FColor(62, 62, 62));
+
+		// Panel Tab
+		Style->Set( "Docking.Tab", FDockTabStyle()
+			.SetCloseButtonStyle( CloseButton )
+			.SetNormalBrush( BOX_BRUSH( "/Docking/Tab_Shape", 2.f /8.0f, DockColor_Inactive ) )
+			.SetHoveredBrush( BOX_BRUSH( "/Docking/Tab_Shape", 2.f /8.0f, DockColor_Hovered ) )
+			.SetForegroundBrush( BOX_BRUSH( "/Docking/Tab_Shape", 2.f /8.0f, DockColor_Active ) )
+			.SetColorOverlayTabBrush( BOX_BRUSH( "/Docking/Tab_ColorOverlay", 4 / 16.0f))
+			.SetContentAreaBrush( FSlateColorBrush( DockColor_Active ) )
+			.SetTabWellBrush( FSlateNoResource() )
+			.SetTabPadding( FMargin(5, 2, 5, 2) )
+			.SetOverlapWidth( -1.0f )
+			.SetFlashColor( TabFlashColor )
+			.SetTabTextStyle(DockingTabFont)
+			);
+
+		// App Tab
+		Style->Set( "Docking.MajorTab", FDockTabStyle()
+			.SetCloseButtonStyle( CloseButton )
+			.SetNormalBrush( BOX_BRUSH( "/Docking/AppTab_Inactive", FMargin(24.0f/64.0f, 4/32.0f) ) )
+			.SetColorOverlayTabBrush(BOX_BRUSH("/Docking/AppTab_ColorOverlay", FMargin(24.0f / 64.0f, 4 / 32.0f)))
+			.SetForegroundBrush( BOX_BRUSH( "/Docking/AppTab_Foreground", FMargin(24.0f/64.0f, 4/32.0f) ) )
+			.SetHoveredBrush( BOX_BRUSH( "/Docking/AppTab_Hovered", FMargin(24.0f/64.0f, 4/32.0f) ) )
+			.SetContentAreaBrush( BOX_BRUSH( "/Docking/AppTabContentArea", FMargin(4/16.0f) ) )
+			.SetTabWellBrush( FSlateNoResource() )
+			.SetTabPadding( FMargin(17, 4, 15, 4) )
+			.SetOverlapWidth( 21.0f )
+			.SetFlashColor( TabFlashColor )
+			.SetTabTextStyle(DockingTabFont)
+			);
+
+		// Dock Cross
+		Style->Set( "Docking.Cross.DockLeft", new IMAGE_BRUSH( "/Docking/OuterDockingIndicator", FVector2D(6, 6), FLinearColor(1.0f, 0.35f, 0.0f, 0.25f) ) );
+		Style->Set( "Docking.Cross.DockLeft_Hovered", new IMAGE_BRUSH( "/Docking/OuterDockingIndicator", FVector2D(6, 6), FLinearColor(1.0f, 0.35f, 0.0f) ) );
+		Style->Set( "Docking.Cross.DockTop", new IMAGE_BRUSH( "/Docking/OuterDockingIndicator", FVector2D(6, 6), FLinearColor(1.0f, 0.35f, 0.0f, 0.25f) ) );
+		Style->Set( "Docking.Cross.DockTop_Hovered", new IMAGE_BRUSH( "/Docking/OuterDockingIndicator", FVector2D(6, 6), FLinearColor(1.0f, 0.35f, 0.0f) ) );
+		Style->Set( "Docking.Cross.DockRight", new IMAGE_BRUSH( "/Docking/OuterDockingIndicator", FVector2D(6, 6), FLinearColor(1.0f, 0.35f, 0.0f, 0.25f) ) );
+		Style->Set( "Docking.Cross.DockRight_Hovered", new IMAGE_BRUSH( "/Docking/OuterDockingIndicator", FVector2D(6, 6), FLinearColor(1.0f, 0.35f, 0.0f) ) );
+		Style->Set( "Docking.Cross.DockBottom", new IMAGE_BRUSH( "/Docking/OuterDockingIndicator", FVector2D(6, 6), FLinearColor(1.0f, 0.35f, 0.0f, 0.25f) ) );
+		Style->Set( "Docking.Cross.DockBottom_Hovered", new IMAGE_BRUSH( "/Docking/OuterDockingIndicator", FVector2D(6, 6), FLinearColor(1.0f, 0.35f, 0.0f) ) );
+		Style->Set( "Docking.Cross.DockCenter", new IMAGE_BRUSH( "/Docking/DockingIndicator_Center", CoreStyleConstants::Icon64x64,  FLinearColor(1.0f, 0.35f, 0.0f, 0.25f) ) );
+		Style->Set( "Docking.Cross.DockCenter_Hovered", new IMAGE_BRUSH( "/Docking/DockingIndicator_Center", CoreStyleConstants::Icon64x64,  FLinearColor(1.0f, 0.35f, 0.0f) ) );
+		
+		Style->Set( "Docking.Cross.BorderLeft",	new FSlateNoResource() );
+		Style->Set( "Docking.Cross.BorderTop", new FSlateNoResource() );
+		Style->Set( "Docking.Cross.BorderRight", new FSlateNoResource() );
+		Style->Set( "Docking.Cross.BorderBottom", new FSlateNoResource() );
+		Style->Set( "Docking.Cross.BorderCenter", new FSlateNoResource() );
+
+		Style->Set( "Docking.Cross.PreviewWindowTint", FLinearColor(1.0f, 0.75f, 0.5f) );
+		Style->Set( "Docking.Cross.Tint", FLinearColor::White );
+		Style->Set( "Docking.Cross.HoveredTint", FLinearColor::White );
+	}
+
+	// SScrollBox defaults...
+	{
+		Style->Set( "ScrollBox", FScrollBoxStyle()
+			.SetTopShadowBrush( BOX_BRUSH( "Common/ScrollBoxShadowTop", FVector2D(16, 8), FMargin(0.5, 1, 0.5, 0) ) )
+			.SetBottomShadowBrush(BOX_BRUSH("Common/ScrollBoxShadowBottom", FVector2D(16, 8), FMargin(0.5, 0, 0.5, 1)))
+			.SetLeftShadowBrush(BOX_BRUSH("Common/ScrollBoxShadowLeft", FVector2D(8, 16), FMargin(1, 0.5, 0, 0.5)))
+			.SetRightShadowBrush(BOX_BRUSH("Common/ScrollBoxShadowRight", FVector2D(8, 16), FMargin(0, 0.5, 1, 0.5)))
+			);
+	}
+
+	// SScrollBorder defaults...
+	{
+		Style->Set( "ScrollBorder", FScrollBorderStyle()
+			.SetTopShadowBrush(BOX_BRUSH("Common/ScrollBorderShadowTop", FVector2D(16, 8), FMargin(0.5, 1, 0.5, 0)))
+			.SetBottomShadowBrush(BOX_BRUSH("Common/ScrollBorderShadowBottom", FVector2D(16, 8), FMargin(0.5, 0, 0.5, 1)))
+			);
+	}
+
+	// SWindow defaults...
+	{
+#if !PLATFORM_MAC
+		const FButtonStyle MinimizeButtonStyle = FButtonStyle(Button)
+			.SetNormal (IMAGE_BRUSH("Common/Window/WindowButton_Minimize_Normal", FVector2D(27.0f, 18.0f)))
+			.SetHovered(IMAGE_BRUSH("Common/Window/WindowButton_Minimize_Hovered", FVector2D(27.0f, 18.0f)))
+			.SetPressed(IMAGE_BRUSH("Common/Window/WindowButton_Minimize_Pressed", FVector2D(27.0f, 18.0f)));
+
+		const FButtonStyle MaximizeButtonStyle = FButtonStyle(Button)
+			.SetNormal (IMAGE_BRUSH("Common/Window/WindowButton_Maximize_Normal", FVector2D(23.0f, 18.0f)))
+			.SetHovered(IMAGE_BRUSH("Common/Window/WindowButton_Maximize_Hovered", FVector2D(23.0f, 18.0f)))
+			.SetPressed(IMAGE_BRUSH("Common/Window/WindowButton_Maximize_Pressed", FVector2D(23.0f, 18.0f)));
+
+		const FButtonStyle RestoreButtonStyle = FButtonStyle(Button)
+			.SetNormal (IMAGE_BRUSH("Common/Window/WindowButton_Restore_Normal", FVector2D(23.0f, 18)))
+			.SetHovered(IMAGE_BRUSH("Common/Window/WindowButton_Restore_Hovered", FVector2D(23.0f, 18)))
+			.SetPressed(IMAGE_BRUSH("Common/Window/WindowButton_Restore_Pressed", FVector2D(23.0f, 18)));
+
+		const FButtonStyle CloseButtonStyle = FButtonStyle(Button)
+			.SetNormal (IMAGE_BRUSH("Common/Window/WindowButton_Close_Normal", FVector2D(44.0f, 18.0f)))
+			.SetHovered(IMAGE_BRUSH("Common/Window/WindowButton_Close_Hovered", FVector2D(44.0f, 18.0f)))
+			.SetPressed(IMAGE_BRUSH("Common/Window/WindowButton_Close_Pressed", FVector2D(44.0f, 18.0f)));
+#endif
+
+		const FTextBlockStyle TitleTextStyle = FTextBlockStyle(NormalText)
+			.SetFont(DEFAULT_FONT("Regular", 9))
+			.SetColorAndOpacity(FLinearColor::White)
+			.SetShadowOffset(FVector2D(1.0f, 1.0f))
+			.SetShadowColorAndOpacity(FLinearColor::Black);
+
+		Style->Set( "Window", FWindowStyle()
+#if !PLATFORM_MAC
+			.SetMinimizeButtonStyle(MinimizeButtonStyle)
+			.SetMaximizeButtonStyle(MaximizeButtonStyle)
+			.SetRestoreButtonStyle(RestoreButtonStyle)
+			.SetCloseButtonStyle(CloseButtonStyle)
+#endif
+			.SetTitleTextStyle(TitleTextStyle)
+			.SetActiveTitleBrush(IMAGE_BRUSH("Common/Window/WindowTitle", CoreStyleConstants::Icon32x32, FLinearColor(1.0f, 1.0f, 1.0f, 1.0f), ESlateBrushTileType::Horizontal))
+			.SetInactiveTitleBrush(IMAGE_BRUSH("Common/Window/WindowTitle_Inactive", CoreStyleConstants::Icon32x32, FLinearColor(1.0f, 1.0f, 1.0f, 1.0f), ESlateBrushTileType::Horizontal))
+			.SetFlashTitleBrush(IMAGE_BRUSH("Common/Window/WindowTitle_Flashing", CoreStyleConstants::Icon24x24, FLinearColor(1.0f, 1.0f, 1.0f, 1.0f), ESlateBrushTileType::Horizontal))
+			.SetOutlineBrush(BORDER_BRUSH("Common/Window/WindowOutline", FMargin(3.0f / 32.0f)))
+			.SetOutlineColor(FLinearColor(0.1f, 0.1f, 0.1f, 1.0f))
+			.SetBorderBrush(BOX_BRUSH("Common/Window/WindowBorder", 0.48f))
+			.SetBackgroundBrush(IMAGE_BRUSH( "Common/Window/WindowBackground", FVector2D(74.0f, 74.0f), FLinearColor::White, ESlateBrushTileType::Both))
+			.SetChildBackgroundBrush(IMAGE_BRUSH( "Common/NoiseBackground", FVector2D(64.0f, 64.0f), FLinearColor::White, ESlateBrushTileType::Both))
+			);
+	}
+
+	// STutorialWrapper defaults...
+	{
+		Style->Set("Tutorials.Border", new BOX_BRUSH("Tutorials/TutorialBorder", FVector2D(64.0f, 64.0f), FMargin(25.0f/ 64.0f)));
+		Style->Set("Tutorials.Shadow", new BOX_BRUSH("Tutorials/TutorialShadow", FVector2D(256.0f, 256.0f), FMargin(114.0f / 256.0f)));
+	}
+
+	// Standard Dialog Settings
+	{
+		Style->Set("StandardDialog.ContentPadding",FMargin(16.0f, 3.0f));
+		Style->Set("StandardDialog.SlotPadding",FMargin(8.0f, 0.0f, 0.0f, 0.0f));
+		Style->Set("StandardDialog.MinDesiredSlotWidth", 80.0f);
+		Style->Set("StandardDialog.MinDesiredSlotHeight", 0.0f);
+		Style->Set("StandardDialog.LargeFont", DEFAULT_FONT("Regular", 11));
+	}
+
+	// Widget Reflector Window
+	{
+		Style->Set("WidgetReflector.TabCoreStyleConstants::Icon", new IMAGE_BRUSH("CoreStyleConstants::Icons/CoreStyleConstants::Icon_tab_WidgetReflector_16x", CoreStyleConstants::Icon16x16));
+		Style->Set("WidgetReflector.CoreStyleConstants::Icon", new IMAGE_BRUSH("CoreStyleConstants::Icons/CoreStyleConstants::Icon_tab_WidgetReflector_40x", CoreStyleConstants::Icon40x40));
+		Style->Set("WidgetReflector.CoreStyleConstants::Icon.Small", new IMAGE_BRUSH("CoreStyleConstants::Icons/CoreStyleConstants::Icon_tab_WidgetReflector_40x", CoreStyleConstants::Icon20x20));
+		Style->Set("WidgetReflector.FocusableCheck", FCheckBoxStyle()
+			.SetUncheckedImage(IMAGE_BRUSH("CoreStyleConstants::Icons/Empty_14x", CoreStyleConstants::Icon14x14))
+			.SetUncheckedHoveredImage(IMAGE_BRUSH("CoreStyleConstants::Icons/Empty_14x", CoreStyleConstants::Icon14x14))
+			.SetUncheckedPressedImage(IMAGE_BRUSH("Common/SmallCheckBox_Hovered", CoreStyleConstants::Icon14x14))
+			.SetCheckedImage(IMAGE_BRUSH("Common/SmallCheck", CoreStyleConstants::Icon14x14))
+			.SetCheckedHoveredImage(IMAGE_BRUSH("Common/SmallCheck", CoreStyleConstants::Icon14x14))
+			.SetCheckedPressedImage(IMAGE_BRUSH("Common/SmallCheck", CoreStyleConstants::Icon14x14))
+			.SetUndeterminedImage(IMAGE_BRUSH("CoreStyleConstants::Icons/Empty_14x", CoreStyleConstants::Icon14x14))
+			.SetUndeterminedHoveredImage(FSlateNoResource())
+			.SetUndeterminedPressedImage(FSlateNoResource())
+			);
+	}
+
+	// Message Log
+	{
+		Style->Set("MessageLog", FTextBlockStyle(NormalText)
+			.SetFont(DEFAULT_FONT("Regular", 8))
+			.SetShadowOffset(FVector2D::ZeroVector)
+		);
+		Style->Set("MessageLog.Error", new IMAGE_BRUSH("MessageLog/Log_Error", CoreStyleConstants::Icon16x16));
+		Style->Set("MessageLog.Warning", new IMAGE_BRUSH("MessageLog/Log_Warning", CoreStyleConstants::Icon16x16));
+		Style->Set("MessageLog.Note", new IMAGE_BRUSH("MessageLog/Log_Note", CoreStyleConstants::Icon16x16));
+	}
+
+	// Wizard CoreStyleConstants::Icons
+	{
+		Style->Set("Wizard.BackCoreStyleConstants::Icon", new IMAGE_BRUSH("CoreStyleConstants::Icons/BackCoreStyleConstants::Icon", CoreStyleConstants::Icon8x8));
+		Style->Set("Wizard.NextCoreStyleConstants::Icon", new IMAGE_BRUSH("CoreStyleConstants::Icons/NextCoreStyleConstants::Icon", CoreStyleConstants::Icon8x8));
+	}
+
+	// Syntax highlighting
+	{
+		const FTextBlockStyle SmallMonospacedText = FTextBlockStyle(MonospacedText)
+			.SetFont(DEFAULT_FONT("Mono", 9));
+
+		Style->Set("SyntaxHighlight.Normal", SmallMonospacedText);
+		Style->Set("SyntaxHighlight.Node", FTextBlockStyle(SmallMonospacedText).SetColorAndOpacity(FLinearColor(FColor(0xff006ab4)))); // blue
+		Style->Set("SyntaxHighlight.NodeAttributeKey", FTextBlockStyle(SmallMonospacedText).SetColorAndOpacity(FLinearColor(FColor(0xffb40000)))); // red
+		Style->Set("SyntaxHighlight.NodeAttribueAssignment", FTextBlockStyle(SmallMonospacedText).SetColorAndOpacity(FLinearColor(FColor(0xffb2b400)))); // yellow
+		Style->Set("SyntaxHighlight.NodeAttributeValue", FTextBlockStyle(SmallMonospacedText).SetColorAndOpacity(FLinearColor(FColor(0xffb46100)))); // orange
+	}
+
+	return Style;
+}
+
+
+void FUMGCoreStyle::SetStyle(const TSharedRef< ISlateStyle >& NewStyle)
+{
+	if (Instance.IsValid())
+	{
+		FSlateStyleRegistry::UnRegisterSlateStyle(*Instance.Get());
+	}
+
+	Instance = NewStyle;
+
+	if (Instance.IsValid())
+	{
+		FSlateStyleRegistry::RegisterSlateStyle(*Instance.Get());
+	}
+	else
+	{
+		ResetToDefault();
+	}
+}
+
+#undef RootToContentDir
\ No newline at end of file
diff --git a/Engine/Source/Runtime/SlateCore/Public/Styling/CoreStyle.h b/Engine/Source/Runtime/SlateCore/Public/Styling/CoreStyle.h
--- a/Engine/Source/Runtime/SlateCore/Public/Styling/CoreStyle.h
+++ b/Engine/Source/Runtime/SlateCore/Public/Styling/CoreStyle.h
@@ -77,3 +77,36 @@ private:
 	/** Singleton instances of this style. */
 	static TSharedPtr< class ISlateStyle > Instance;
 };
+
+namespace CoreStyleConstants
+{
+	// Note, these sizes are in Slate Units.
+	// Slate Units do NOT have to map to pixels.
+	const FVector2D Icon5x16(5.0f, 16.0f);
+	const FVector2D Icon8x4(8.0f, 4.0f);
+	const FVector2D Icon16x4(16.0f, 4.0f);
+	const FVector2D Icon8x8(8.0f, 8.0f);
+	const FVector2D Icon10x10(10.0f, 10.0f);
+	const FVector2D Icon12x12(12.0f, 12.0f);
+	const FVector2D Icon12x16(12.0f, 16.0f);
+	const FVector2D Icon14x14(14.0f, 14.0f);
+	const FVector2D Icon16x16(16.0f, 16.0f);
+	const FVector2D Icon18x18(18.0f, 18.0f);
+	const FVector2D Icon20x20(20.0f, 20.0f);
+	const FVector2D Icon22x22(22.0f, 22.0f);
+	const FVector2D Icon24x24(24.0f, 24.0f);
+	const FVector2D Icon25x25(25.0f, 25.0f);
+	const FVector2D Icon32x32(32.0f, 32.0f);
+	const FVector2D Icon40x40(40.0f, 40.0f);
+	const FVector2D Icon64x64(64.0f, 64.0f);
+	const FVector2D Icon36x24(36.0f, 24.0f);
+	const FVector2D Icon128x128(128.0f, 128.0f);
+
+	// Common Margins
+	const FMargin DefaultMargins(8.f, 4.f);
+	// Buttons already have a built in (4., 2.) padding - adding to that a little
+	const FMargin ButtonMargins(20.f, 2.f, 20.f, 3.f);
+
+	const float InputFocusRadius = 4.f;
+	const float InputFocusThickness = 1.0f;
+}
diff --git a/Engine/Source/Runtime/SlateCore/Public/Styling/SlateStyleMacros.h b/Engine/Source/Runtime/SlateCore/Public/Styling/SlateStyleMacros.h
--- a/Engine/Source/Runtime/SlateCore/Public/Styling/SlateStyleMacros.h
+++ b/Engine/Source/Runtime/SlateCore/Public/Styling/SlateStyleMacros.h
@@ -18,4 +18,4 @@
 #define CORE_BOX_BRUSH_SVG( RelativePath, ... ) FSlateVectorBoxBrush(RootToCoreContentDir(RelativePath, TEXT(".svg")), __VA_ARGS__)
 #define CORE_BORDER_BRUSH_SVG( RelativePath, ... ) FSlateVectorBorderBrush(RootToCoreContentDir(RelativePath, TEXT(".svg")), __VA_ARGS__)
 
-#define DEFAULT_FONT(...) FStarshipCoreStyle::GetDefaultFontStyle(__VA_ARGS__)
\ No newline at end of file
+#define DEFAULT_FONT(...) FCoreStyle::GetDefaultFontStyle(__VA_ARGS__)
\ No newline at end of file
diff --git a/Engine/Source/Runtime/SlateCore/Public/Styling/StarshipCoreStyle.h b/Engine/Source/Runtime/SlateCore/Public/Styling/StarshipCoreStyle.h
--- a/Engine/Source/Runtime/SlateCore/Public/Styling/StarshipCoreStyle.h
+++ b/Engine/Source/Runtime/SlateCore/Public/Styling/StarshipCoreStyle.h
@@ -8,42 +8,6 @@
 
 struct FSlateDynamicImageBrush;
 
-class FStyle;
-
-namespace CoreStyleConstants
-{
-	// Note, these sizes are in Slate Units.
-	// Slate Units do NOT have to map to pixels.
-	const FVector2D Icon5x16(5.0f, 16.0f);
-	const FVector2D Icon8x4(8.0f, 4.0f);
-	const FVector2D Icon16x4(16.0f, 4.0f);
-	const FVector2D Icon8x8(8.0f, 8.0f);
-	const FVector2D Icon10x10(10.0f, 10.0f);
-	const FVector2D Icon12x12(12.0f, 12.0f);
-	const FVector2D Icon12x16(12.0f, 16.0f);
-	const FVector2D Icon14x14(14.0f, 14.0f);
-	const FVector2D Icon16x16(16.0f, 16.0f);
-	const FVector2D Icon18x18(18.0f, 18.0f);
-	const FVector2D Icon20x20(20.0f, 20.0f);
-	const FVector2D Icon22x22(22.0f, 22.0f);
-	const FVector2D Icon24x24(24.0f, 24.0f);
-	const FVector2D Icon25x25(25.0f, 25.0f);
-	const FVector2D Icon32x32(32.0f, 32.0f);
-	const FVector2D Icon40x40(40.0f, 40.0f);
-	const FVector2D Icon64x64(64.0f, 64.0f);
-	const FVector2D Icon36x24(36.0f, 24.0f);
-	const FVector2D Icon128x128(128.0f, 128.0f);
-
-	// Common Margins
-	const FMargin DefaultMargins(8.f, 4.f);
-	// Buttons already have a built in (4., 2.) padding - adding to that a little
-	const FMargin ButtonMargins(20.f, 2.f, 20.f, 3.f);
-
-	const float InputFocusRadius = 4.f;
-	const float InputFocusThickness = 1.0f;
-
-}
-
 /**
  * Core slate style
  */
@@ -102,6 +66,7 @@ public:
 
 	static bool IsInitialized() { return Instance.IsValid(); }
 private:
+	class FStyle;
 
 	static void SetStyle(const TSharedRef<class ISlateStyle>& NewStyle);
 	
diff --git a/Engine/Source/Runtime/SlateCore/Public/Styling/StyleColors.h b/Engine/Source/Runtime/SlateCore/Public/Styling/StyleColors.h
--- a/Engine/Source/Runtime/SlateCore/Public/Styling/StyleColors.h
+++ b/Engine/Source/Runtime/SlateCore/Public/Styling/StyleColors.h
@@ -13,7 +13,7 @@
 #define COLOR( HexValue ) FLinearColor::FromSRGBColor(FColor::FromHex(HexValue))
 
 /**
- * Note: If you add another color here, you should update the Dark.json theme file in Engine\Content\Slate\Themes for consistency
+ * Note: If you add another color here, you should update the Dark.json theme file in Engine\Content\Slate\Themes & FUMGColors in UMGCoreStyle.cpp for consistency
  */
 UENUM()
 enum class EStyleColor : uint8
diff --git a/Engine/Source/Runtime/SlateCore/Public/Styling/UMGCoreStyle.h b/Engine/Source/Runtime/SlateCore/Public/Styling/UMGCoreStyle.h
new file mode 100644
--- /dev/null
+++ b/Engine/Source/Runtime/SlateCore/Public/Styling/UMGCoreStyle.h
@@ -0,0 +1,50 @@
+// Copyright Epic Games, Inc. All Rights Reserved.
+
+#pragma once
+
+#include "CoreMinimal.h"
+#include "Styling/ISlateStyle.h"
+#include "Styling/AppStyle.h"
+
+/**
+ * Core slate style
+ */
+class SLATECORE_API FUMGCoreStyle
+{
+public:
+
+	static TSharedRef<class ISlateStyle> Create();
+
+	static const ISlateStyle& Get()
+	{
+		return *(Instance.Get());
+	}
+
+	/** Get the default font for Slate */
+	static TSharedRef<const FCompositeFont> GetDefaultFont();
+
+	/** Get a font style using the default for for Slate */
+	static FSlateFontInfo GetDefaultFontStyle(const FName InTypefaceFontName, const int32 InSize, const FFontOutlineSettings& InOutlineSettings = FFontOutlineSettings());
+
+	static void ResetToDefault( );
+
+	/** Used to override the default selection colors */
+	static void SetSelectorColor( const FLinearColor& NewColor );
+	static void SetSelectionColor( const FLinearColor& NewColor );
+	static void SetInactiveSelectionColor( const FLinearColor& NewColor );
+	static void SetPressedSelectionColor( const FLinearColor& NewColor );
+	static void SetFocusBrush(FSlateBrush* NewBrush);
+
+	static const int32 RegularTextSize = 10;
+	static const int32 SmallTextSize = 8;
+
+	static bool IsInitialized() { return Instance.IsValid(); }
+private:
+
+	static void SetStyle(const TSharedRef<class ISlateStyle>& NewStyle);
+
+private:
+
+	/** Singleton instances of this style. */
+	static TSharedPtr< class ISlateStyle > Instance;
+};
diff --git a/Engine/Source/Runtime/UMG/Private/Components/Button.cpp b/Engine/Source/Runtime/UMG/Private/Components/Button.cpp
--- a/Engine/Source/Runtime/UMG/Private/Components/Button.cpp
+++ b/Engine/Source/Runtime/UMG/Private/Components/Button.cpp
@@ -5,6 +5,8 @@
 #include "Widgets/DeclarativeSyntaxSupport.h"
 #include "Widgets/Input/SButton.h"
 #include "Components/ButtonSlot.h"
+#include "Styling/UMGCoreStyle.h"
+#include "Blueprint/WidgetTree.h"
 
 #define LOCTEXT_NAMESPACE "UMG"
 
@@ -13,20 +15,38 @@
 
 static FButtonStyle* DefaultButtonStyle = nullptr;
 
+#if WITH_EDITOR
+static FButtonStyle* EditorButtonStyle = nullptr;
+#endif 
+
 UButton::UButton(const FObjectInitializer& ObjectInitializer)
 	: Super(ObjectInitializer)
 {
 	if (DefaultButtonStyle == nullptr)
 	{
-		// HACK: THIS SHOULD NOT COME FROM CORESTYLE AND SHOULD INSTEAD BE DEFINED BY ENGINE TEXTURES/PROJECT SETTINGS
-		DefaultButtonStyle = new FButtonStyle(FCoreStyle::Get().GetWidgetStyle<FButtonStyle>("Button"));
+		DefaultButtonStyle = new FButtonStyle(FUMGCoreStyle::Get().GetWidgetStyle<FButtonStyle>("Button"));
 
-		// Unlink UMG default colors from the editor settings colors.
+		// Unlink UMG default colors.
 		DefaultButtonStyle->UnlinkColors();
 	}
 
 	WidgetStyle = *DefaultButtonStyle;
 
+#if WITH_EDITOR 
+	if (EditorButtonStyle == nullptr)
+	{
+		EditorButtonStyle = new FButtonStyle(FCoreStyle::Get().GetWidgetStyle<FButtonStyle>("Button"));
+
+		// Unlink UMG Editor colors from the editor settings colors.
+		EditorButtonStyle->UnlinkColors();
+	}
+
+	if (IsEditorWidget())
+	{
+		WidgetStyle = *EditorButtonStyle;
+	}
+#endif // WITH_EDITOR
+
 	ColorAndOpacity = FLinearColor::White;
 	BackgroundColor = FLinearColor::White;
 
diff --git a/Engine/Source/Runtime/UMG/Private/Components/CheckBox.cpp b/Engine/Source/Runtime/UMG/Private/Components/CheckBox.cpp
--- a/Engine/Source/Runtime/UMG/Private/Components/CheckBox.cpp
+++ b/Engine/Source/Runtime/UMG/Private/Components/CheckBox.cpp
@@ -5,6 +5,7 @@
 #include "Widgets/DeclarativeSyntaxSupport.h"
 #include "Widgets/Input/SCheckBox.h"
 #include "Slate/SlateBrushAsset.h"
+#include "Styling/UMGCoreStyle.h"
 
 #define LOCTEXT_NAMESPACE "UMG"
 
@@ -13,20 +14,38 @@
 
 static FCheckBoxStyle* DefaultCheckboxStyle = nullptr;
 
+#if WITH_EDITOR
+static FCheckBoxStyle* EditorCheckboxStyle = nullptr;
+#endif 
+
 UCheckBox::UCheckBox(const FObjectInitializer& ObjectInitializer)
 	: Super(ObjectInitializer)
 {
 	if (DefaultCheckboxStyle == nullptr)
 	{
-		// HACK: THIS SHOULD NOT COME FROM CORESTYLE AND SHOULD INSTEAD BE DEFINED BY ENGINE TEXTURES/PROJECT SETTINGS
-		DefaultCheckboxStyle = new FCheckBoxStyle(FCoreStyle::Get().GetWidgetStyle<FCheckBoxStyle>("Checkbox"));
+		DefaultCheckboxStyle = new FCheckBoxStyle(FUMGCoreStyle::Get().GetWidgetStyle<FCheckBoxStyle>("Checkbox"));
 
-		// Unlink UMG default colors from the editor settings colors.
+		// Unlink UMG default colors.
 		DefaultCheckboxStyle->UnlinkColors();
 	}
 
 	WidgetStyle = *DefaultCheckboxStyle;
 
+#if WITH_EDITOR 
+	if (EditorCheckboxStyle == nullptr)
+	{
+		EditorCheckboxStyle = new FCheckBoxStyle(FCoreStyle::Get().GetWidgetStyle<FCheckBoxStyle>("Checkbox"));
+
+		// Unlink UMG Editor colors from the editor settings colors.
+		EditorCheckboxStyle->UnlinkColors();
+	}
+
+	if (IsEditorWidget())
+	{
+		WidgetStyle = *EditorCheckboxStyle;
+	}
+#endif // WITH_EDITOR
+
 	CheckedState = ECheckBoxState::Unchecked;
 
 	HorizontalAlignment = HAlign_Fill;
diff --git a/Engine/Source/Runtime/UMG/Private/Components/CircularThrobber.cpp b/Engine/Source/Runtime/UMG/Private/Components/CircularThrobber.cpp
--- a/Engine/Source/Runtime/UMG/Private/Components/CircularThrobber.cpp
+++ b/Engine/Source/Runtime/UMG/Private/Components/CircularThrobber.cpp
@@ -5,6 +5,7 @@
 #include "Widgets/DeclarativeSyntaxSupport.h"
 #include "Components/CanvasPanelSlot.h"
 #include "Widgets/Images/SThrobber.h"
+#include "Styling/UMGCoreStyle.h"
 
 #define LOCTEXT_NAMESPACE "UMG"
 
@@ -13,21 +14,39 @@
 
 static FSlateBrush* DefaultCircularThrobberBrushStyle = nullptr;
 
+#if WITH_EDITOR
+static FSlateBrush* EditorCircularThrobberBrushStyle = nullptr;
+#endif 
+
 UCircularThrobber::UCircularThrobber(const FObjectInitializer& ObjectInitializer)
 	: Super(ObjectInitializer)
 	, bEnableRadius(true)
 {
 	if (DefaultCircularThrobberBrushStyle == nullptr)
 	{
-		// HACK: THIS SHOULD NOT COME FROM CORESTYLE AND SHOULD INSTEAD BE DEFINED BY ENGINE TEXTURES/PROJECT SETTINGS
-		DefaultCircularThrobberBrushStyle = new FSlateBrush(*FCoreStyle::Get().GetBrush("Throbber.CircleChunk"));
+		DefaultCircularThrobberBrushStyle = new FSlateBrush(*FUMGCoreStyle::Get().GetBrush("Throbber.CircleChunk"));
 
-		// Unlink UMG default colors from the editor settings colors.
+		// Unlink UMG default colors.
 		DefaultCircularThrobberBrushStyle->UnlinkColors();
 	}
 
 	Image = *DefaultCircularThrobberBrushStyle;
 
+#if WITH_EDITOR 
+	if (EditorCircularThrobberBrushStyle == nullptr)
+	{
+		EditorCircularThrobberBrushStyle = new FSlateBrush(*FCoreStyle::Get().GetBrush("Throbber.CircleChunk"));
+
+		// Unlink UMG Editor colors from the editor settings colors.
+		EditorCircularThrobberBrushStyle->UnlinkColors();
+	}
+
+	if (IsEditorWidget())
+	{
+		Image = *EditorCircularThrobberBrushStyle;
+	}
+#endif // WITH_EDITOR
+
 	NumberOfPieces = 6;
 	Period = 0.75f;
 	Radius = 16.f;
diff --git a/Engine/Source/Runtime/UMG/Private/Components/ComboBoxKey.cpp b/Engine/Source/Runtime/UMG/Private/Components/ComboBoxKey.cpp
--- a/Engine/Source/Runtime/UMG/Private/Components/ComboBoxKey.cpp
+++ b/Engine/Source/Runtime/UMG/Private/Components/ComboBoxKey.cpp
@@ -2,6 +2,7 @@
 
 #include "Components/ComboBoxKey.h"
 
+#include "Styling/UMGCoreStyle.h"
 #include "UObject/ConstructorHelpers.h"
 #include "Widgets/Layout/SBox.h"
 #include "Widgets/SNullWidget.h"
@@ -17,8 +18,14 @@ UComboBoxKey::UComboBoxKey()
 	// We don't want to try and load fonts on the server.
 	if (!IsRunningDedicatedServer())
 	{
-		WidgetStyle = FCoreStyle::Get().GetWidgetStyle<FComboBoxStyle>("ComboBox");
-		ItemStyle = FCoreStyle::Get().GetWidgetStyle<FTableRowStyle>("TableView.Row");
+#if WITH_EDITOR 
+		const ISlateStyle& SlateStyle = IsEditorWidget() ? FCoreStyle::Get() : FUMGCoreStyle::Get();
+		WidgetStyle = SlateStyle.GetWidgetStyle<FComboBoxStyle>("ComboBox");
+		ItemStyle = SlateStyle.GetWidgetStyle<FTableRowStyle>("TableView.Row");
+#else
+		WidgetStyle = FUMGCoreStyle::Get().GetWidgetStyle<FComboBoxStyle>("ComboBox");
+		ItemStyle = FUMGCoreStyle::Get().GetWidgetStyle<FTableRowStyle>("TableView.Row");
+#endif // WITH_EDITOR
 	}
 
 	ContentPadding = FMargin(4.0, 2.0);
diff --git a/Engine/Source/Runtime/UMG/Private/Components/ComboBoxString.cpp b/Engine/Source/Runtime/UMG/Private/Components/ComboBoxString.cpp
--- a/Engine/Source/Runtime/UMG/Private/Components/ComboBoxString.cpp
+++ b/Engine/Source/Runtime/UMG/Private/Components/ComboBoxString.cpp
@@ -5,6 +5,7 @@
 #include "UObject/EditorObjectVersion.h"
 #include "UObject/ConstructorHelpers.h"
 #include "Engine/Font.h"
+#include "Styling/UMGCoreStyle.h"
 
 #define LOCTEXT_NAMESPACE "UMG"
 
@@ -14,31 +15,58 @@
 static FComboBoxStyle* DefaultComboBoxStyle = nullptr;
 static FTableRowStyle* DefaultComboBoxRowStyle = nullptr;
 
+#if WITH_EDITOR
+static FComboBoxStyle* EditorComboBoxStyle = nullptr;
+static FTableRowStyle* EditorComboBoxRowStyle = nullptr;
+#endif 
+
 UComboBoxString::UComboBoxString(const FObjectInitializer& ObjectInitializer)
 	: Super(ObjectInitializer)
 {
 	if (DefaultComboBoxStyle == nullptr)
 	{
-		// HACK: THIS SHOULD NOT COME FROM CORESTYLE AND SHOULD INSTEAD BE DEFINED BY ENGINE TEXTURES/PROJECT SETTINGS
-		DefaultComboBoxStyle = new FComboBoxStyle(FCoreStyle::Get().GetWidgetStyle<FComboBoxStyle>("ComboBox"));
+		DefaultComboBoxStyle = new FComboBoxStyle(FUMGCoreStyle::Get().GetWidgetStyle<FComboBoxStyle>("ComboBox"));
 
-		// Unlink UMG default colors from the editor settings colors.
+		// Unlink UMG default colors.
 		DefaultComboBoxStyle->UnlinkColors();
 	}
 
 	if (DefaultComboBoxRowStyle == nullptr)
 	{
-		// HACK: THIS SHOULD NOT COME FROM CORESTYLE AND SHOULD INSTEAD BE DEFINED BY ENGINE TEXTURES/PROJECT SETTINGS
-		DefaultComboBoxRowStyle = new FTableRowStyle(FCoreStyle::Get().GetWidgetStyle<FTableRowStyle>("TableView.Row"));
+		DefaultComboBoxRowStyle = new FTableRowStyle(FUMGCoreStyle::Get().GetWidgetStyle<FTableRowStyle>("TableView.Row"));
 
-		// Unlink UMG default colors from the editor settings colors.
+		// Unlink UMG default colors.
 		DefaultComboBoxRowStyle->UnlinkColors();
 	}
 
 	WidgetStyle = *DefaultComboBoxStyle;
-	WidgetStyle.UnlinkColors();
-
 	ItemStyle = *DefaultComboBoxRowStyle;
+
+#if WITH_EDITOR 
+	if (EditorComboBoxStyle == nullptr)
+	{
+		EditorComboBoxStyle = new FComboBoxStyle(FCoreStyle::Get().GetWidgetStyle<FComboBoxStyle>("ComboBox"));
+
+		// Unlink UMG Editor colors from the editor settings colors.
+		EditorComboBoxStyle->UnlinkColors();
+	}
+
+	if (EditorComboBoxRowStyle == nullptr)
+	{
+		EditorComboBoxRowStyle = new FTableRowStyle(FCoreStyle::Get().GetWidgetStyle<FTableRowStyle>("TableView.Row"));
+
+		// Unlink UMG Editor colors from the editor settings colors.
+		EditorComboBoxRowStyle->UnlinkColors();
+	}
+
+	if (IsEditorWidget())
+	{
+		WidgetStyle = *EditorComboBoxStyle;
+		ItemStyle = *EditorComboBoxRowStyle;
+	}
+#endif // WITH_EDITOR
+
+	WidgetStyle.UnlinkColors();
 	ItemStyle.UnlinkColors();
 
 	ForegroundColor = FLinearColor::Black;
diff --git a/Engine/Source/Runtime/UMG/Private/Components/EditableText.cpp b/Engine/Source/Runtime/UMG/Private/Components/EditableText.cpp
--- a/Engine/Source/Runtime/UMG/Private/Components/EditableText.cpp
+++ b/Engine/Source/Runtime/UMG/Private/Components/EditableText.cpp
@@ -6,6 +6,7 @@
 #include "Widgets/DeclarativeSyntaxSupport.h"
 #include "Widgets/Input/SEditableText.h"
 #include "Slate/SlateBrushAsset.h"
+#include "Styling/UMGCoreStyle.h"
 
 #define LOCTEXT_NAMESPACE "UMG"
 
@@ -14,20 +15,38 @@
 
 static FEditableTextStyle* DefaultEditableTextStyle = nullptr;
 
+#if WITH_EDITOR
+static FEditableTextStyle* EditorEditableTextStyle = nullptr;
+#endif 
+
 UEditableText::UEditableText(const FObjectInitializer& ObjectInitializer)
 	: Super(ObjectInitializer)
 {
 	if (DefaultEditableTextStyle == nullptr)
 	{
-		// HACK: THIS SHOULD NOT COME FROM CORESTYLE AND SHOULD INSTEAD BE DEFINED BY ENGINE TEXTURES/PROJECT SETTINGS
-		DefaultEditableTextStyle = new FEditableTextStyle(FCoreStyle::Get().GetWidgetStyle<FEditableTextStyle>("NormalEditableText"));
+		DefaultEditableTextStyle = new FEditableTextStyle(FUMGCoreStyle::Get().GetWidgetStyle<FEditableTextStyle>("NormalEditableText"));
 
-		// Unlink UMG default colors from the editor settings colors.
+		// Unlink UMG default colors.
 		DefaultEditableTextStyle->UnlinkColors();
 	}
 
 	WidgetStyle = *DefaultEditableTextStyle;
 
+#if WITH_EDITOR 
+	if (EditorEditableTextStyle == nullptr)
+	{
+		EditorEditableTextStyle = new FEditableTextStyle(FCoreStyle::Get().GetWidgetStyle<FEditableTextStyle>("NormalEditableText"));
+
+		// Unlink UMG Editor colors from the editor settings colors.
+		EditorEditableTextStyle->UnlinkColors();
+	}
+
+	if (IsEditorWidget())
+	{
+		WidgetStyle = *EditorEditableTextStyle;
+	}
+#endif // WITH_EDITOR
+
 	ColorAndOpacity_DEPRECATED = FLinearColor::Black;
 
 	if (!IsRunningDedicatedServer())
diff --git a/Engine/Source/Runtime/UMG/Private/Components/EditableTextBox.cpp b/Engine/Source/Runtime/UMG/Private/Components/EditableTextBox.cpp
--- a/Engine/Source/Runtime/UMG/Private/Components/EditableTextBox.cpp
+++ b/Engine/Source/Runtime/UMG/Private/Components/EditableTextBox.cpp
@@ -5,6 +5,7 @@
 #include "Engine/Font.h"
 #include "Widgets/DeclarativeSyntaxSupport.h"
 #include "Widgets/Input/SEditableTextBox.h"
+#include "Styling/UMGCoreStyle.h"
 
 #define LOCTEXT_NAMESPACE "UMG"
 
@@ -13,6 +14,10 @@
 
 static FEditableTextBoxStyle* DefaultEditableTextBoxStyle = nullptr;
 
+#if WITH_EDITOR
+static FEditableTextBoxStyle* EditorEditableTextBoxStyle = nullptr;
+#endif 
+
 UEditableTextBox::UEditableTextBox(const FObjectInitializer& ObjectInitializer)
 	: Super(ObjectInitializer)
 {
@@ -40,15 +45,29 @@ UEditableTextBox::UEditableTextBox(const FObjectInitializer& ObjectInitializer)
 
 	if (DefaultEditableTextBoxStyle == nullptr)
 	{
-		// HACK: THIS SHOULD NOT COME FROM CORESTYLE AND SHOULD INSTEAD BE DEFINED BY ENGINE TEXTURES/PROJECT SETTINGS
-		DefaultEditableTextBoxStyle = new FEditableTextBoxStyle(FCoreStyle::Get().GetWidgetStyle<FEditableTextBoxStyle>("NormalEditableTextBox"));
+		DefaultEditableTextBoxStyle = new FEditableTextBoxStyle(FUMGCoreStyle::Get().GetWidgetStyle<FEditableTextBoxStyle>("NormalEditableTextBox"));
 
-		// Unlink UMG default colors from the editor settings colors.
+		// Unlink UMG default colors.
 		DefaultEditableTextBoxStyle->UnlinkColors();
 	}
 
 	WidgetStyle = *DefaultEditableTextBoxStyle;
 
+#if WITH_EDITOR 
+	if (EditorEditableTextBoxStyle == nullptr)
+	{
+		EditorEditableTextBoxStyle = new FEditableTextBoxStyle(FCoreStyle::Get().GetWidgetStyle<FEditableTextBoxStyle>("NormalEditableTextBox"));
+
+		// Unlink UMG Editor colors from the editor settings colors.
+		EditorEditableTextBoxStyle->UnlinkColors();
+	}
+
+	if (IsEditorWidget())
+	{
+		WidgetStyle = *EditorEditableTextBoxStyle;
+	}
+#endif // WITH_EDITOR
+
 #if WITH_EDITORONLY_DATA
 	AccessibleBehavior = ESlateAccessibleBehavior::Auto;
 	bCanChildrenBeAccessible = false;
diff --git a/Engine/Source/Runtime/UMG/Private/Components/ExpandableArea.cpp b/Engine/Source/Runtime/UMG/Private/Components/ExpandableArea.cpp
--- a/Engine/Source/Runtime/UMG/Private/Components/ExpandableArea.cpp
+++ b/Engine/Source/Runtime/UMG/Private/Components/ExpandableArea.cpp
@@ -3,6 +3,7 @@
 #include "Components/ExpandableArea.h"
 #include "Widgets/SNullWidget.h"
 #include "Widgets/DeclarativeSyntaxSupport.h"
+#include "Styling/UMGCoreStyle.h"
 
 #include "Widgets/Layout/SExpandableArea.h"
 
@@ -17,6 +18,11 @@ static const FName BodyName(TEXT("Body"));
 static FExpandableAreaStyle* DefaultExpandableAreaStyle = nullptr;
 static FSlateBrush* DefaultExpandableAreaBorderBrush = nullptr;
 
+#if WITH_EDITOR
+static FExpandableAreaStyle* EditorExpandableAreaStyle = nullptr;
+static FSlateBrush* EditorExpandableAreaBorderBrush = nullptr;
+#endif 
+
 UExpandableArea::UExpandableArea(const FObjectInitializer& ObjectInitializer)
 	: Super(ObjectInitializer)
 	, bIsExpanded(false)
@@ -25,25 +31,47 @@ UExpandableArea::UExpandableArea(const FObjectInitializer& ObjectInitializer)
 
 	if (DefaultExpandableAreaStyle == nullptr)
 	{
-		// HACK: THIS SHOULD NOT COME FROM CORESTYLE AND SHOULD INSTEAD BE DEFINED BY ENGINE TEXTURES/PROJECT SETTINGS
-		DefaultExpandableAreaStyle = new FExpandableAreaStyle(FCoreStyle::Get().GetWidgetStyle<FExpandableAreaStyle>("ExpandableArea"));
+		DefaultExpandableAreaStyle = new FExpandableAreaStyle(FUMGCoreStyle::Get().GetWidgetStyle<FExpandableAreaStyle>("ExpandableArea"));
 
-		// Unlink UMG default colors from the editor settings colors.
+		// Unlink UMG default colors.
 		DefaultExpandableAreaStyle->UnlinkColors();
 	}
 
 	if (DefaultExpandableAreaBorderBrush == nullptr)
 	{
-		// HACK: THIS SHOULD NOT COME FROM CORESTYLE AND SHOULD INSTEAD BE DEFINED BY ENGINE TEXTURES/PROJECT SETTINGS
-		DefaultExpandableAreaBorderBrush = new FSlateBrush(*FCoreStyle::Get().GetBrush("ExpandableArea.Border"));
+		DefaultExpandableAreaBorderBrush = new FSlateBrush(*FUMGCoreStyle::Get().GetBrush("ExpandableArea.Border"));
 
-		// Unlink UMG default colors from the editor settings colors.
+		// Unlink UMG default colors.
 		DefaultExpandableAreaBorderBrush->UnlinkColors();
 	}
 
 	Style = *DefaultExpandableAreaStyle;
 	BorderBrush = *DefaultExpandableAreaBorderBrush;
 
+#if WITH_EDITOR 
+	if (EditorExpandableAreaStyle == nullptr)
+	{
+		EditorExpandableAreaStyle = new FExpandableAreaStyle(FCoreStyle::Get().GetWidgetStyle<FExpandableAreaStyle>("ExpandableArea"));
+
+		// Unlink UMG Editor colors from the editor settings colors.
+		EditorExpandableAreaStyle->UnlinkColors();
+	}
+
+	if (EditorExpandableAreaBorderBrush == nullptr)
+	{
+		EditorExpandableAreaBorderBrush = new FSlateBrush(*FCoreStyle::Get().GetBrush("ExpandableArea.Border"));
+
+		// Unlink UMG Editor colors from the editor settings colors.
+		EditorExpandableAreaBorderBrush->UnlinkColors();
+	}
+
+	if (IsEditorWidget())
+	{
+		Style = *EditorExpandableAreaStyle;
+		BorderBrush = *EditorExpandableAreaBorderBrush;
+	}
+#endif // WITH_EDITOR
+
 	BorderColor = FLinearColor::White;
 	AreaPadding = FMargin(1);
 	HeaderPadding = FMargin(4.0f, 2.0f);
diff --git a/Engine/Source/Runtime/UMG/Private/Components/InputKeySelector.cpp b/Engine/Source/Runtime/UMG/Private/Components/InputKeySelector.cpp
--- a/Engine/Source/Runtime/UMG/Private/Components/InputKeySelector.cpp
+++ b/Engine/Source/Runtime/UMG/Private/Components/InputKeySelector.cpp
@@ -7,36 +7,64 @@
 #include "Widgets/DeclarativeSyntaxSupport.h"
 #include "Widgets/Input/SInputKeySelector.h"
 #include "Internationalization/Internationalization.h"
+#include "Styling/UMGCoreStyle.h"
 
 #define LOCTEXT_NAMESPACE "UMG"
 
 static FButtonStyle* DefaultInputKeySelectorButtonStyle = nullptr;
 static FTextBlockStyle* DefaultInputKeySelectorTextStyle = nullptr;
 
+#if WITH_EDITOR 
+static FButtonStyle* EditorInputKeySelectorButtonStyle = nullptr;
+static FTextBlockStyle* EditorInputKeySelectorTextStyle = nullptr;
+#endif
+
 UInputKeySelector::UInputKeySelector( const FObjectInitializer& ObjectInitializer )
 	: Super(ObjectInitializer)
 {
 	if (DefaultInputKeySelectorButtonStyle == nullptr)
 	{
-		// HACK: THIS SHOULD NOT COME FROM CORESTYLE AND SHOULD INSTEAD BE DEFINED BY ENGINE TEXTURES/PROJECT SETTINGS
-		DefaultInputKeySelectorButtonStyle = new FButtonStyle(FCoreStyle::Get().GetWidgetStyle<FButtonStyle>("Button"));
+		DefaultInputKeySelectorButtonStyle = new FButtonStyle(FUMGCoreStyle::Get().GetWidgetStyle<FButtonStyle>("Button"));
 
-		// Unlink UMG default colors from the editor settings colors.
+		// Unlink UMG default colors.
 		DefaultInputKeySelectorButtonStyle->UnlinkColors();
 	}
 
 	if (DefaultInputKeySelectorTextStyle == nullptr)
 	{
-		// HACK: THIS SHOULD NOT COME FROM CORESTYLE AND SHOULD INSTEAD BE DEFINED BY ENGINE TEXTURES/PROJECT SETTINGS
-		DefaultInputKeySelectorTextStyle = new FTextBlockStyle(FCoreStyle::Get().GetWidgetStyle<FTextBlockStyle>("NormalText"));
+		DefaultInputKeySelectorTextStyle = new FTextBlockStyle(FUMGCoreStyle::Get().GetWidgetStyle<FTextBlockStyle>("NormalText"));
 
-		// Unlink UMG default colors from the editor settings colors.
+		// Unlink UMG default colors.
 		DefaultInputKeySelectorTextStyle->UnlinkColors();
 	}
 
 	WidgetStyle = *DefaultInputKeySelectorButtonStyle;
 	TextStyle = *DefaultInputKeySelectorTextStyle;
 
+#if WITH_EDITOR 
+	if (EditorInputKeySelectorButtonStyle == nullptr)
+	{
+		EditorInputKeySelectorButtonStyle = new FButtonStyle(FCoreStyle::Get().GetWidgetStyle<FButtonStyle>("Button"));
+
+		// Unlink UMG Editor colors from the editor settings colors.
+		EditorInputKeySelectorButtonStyle->UnlinkColors();
+	}
+
+	if (EditorInputKeySelectorTextStyle == nullptr)
+	{
+		EditorInputKeySelectorTextStyle = new FTextBlockStyle(FCoreStyle::Get().GetWidgetStyle<FTextBlockStyle>("NormalText"));
+
+		// Unlink UMG Editor colors from the editor settings colors.
+		EditorInputKeySelectorTextStyle->UnlinkColors();
+	}
+
+	if (IsEditorWidget())
+	{
+		WidgetStyle = *EditorInputKeySelectorButtonStyle;
+		TextStyle = *EditorInputKeySelectorTextStyle;
+	}
+#endif // WITH_EDITOR
+
 	KeySelectionText = NSLOCTEXT("InputKeySelector", "DefaultKeySelectionText", "...");
 	NoKeySpecifiedText = NSLOCTEXT("InputKeySelector", "DefaultEmptyText", "Empty");
 	SelectedKey = FInputChord(EKeys::Invalid);
diff --git a/Engine/Source/Runtime/UMG/Private/Components/MultiLineEditableText.cpp b/Engine/Source/Runtime/UMG/Private/Components/MultiLineEditableText.cpp
--- a/Engine/Source/Runtime/UMG/Private/Components/MultiLineEditableText.cpp
+++ b/Engine/Source/Runtime/UMG/Private/Components/MultiLineEditableText.cpp
@@ -5,6 +5,7 @@
 #include "Engine/Font.h"
 #include "Widgets/DeclarativeSyntaxSupport.h"
 #include "Widgets/Text/SMultiLineEditableText.h"
+#include "Styling/UMGCoreStyle.h"
 
 #define LOCTEXT_NAMESPACE "UMG"
 
@@ -13,20 +14,38 @@
 
 static FTextBlockStyle* DefaultMultiLineEditableTextStyle = nullptr;
 
+#if WITH_EDITOR
+static FTextBlockStyle* EditorMultiLineEditableTextStyle = nullptr;
+#endif 
+
 UMultiLineEditableText::UMultiLineEditableText(const FObjectInitializer& ObjectInitializer)
 	: Super(ObjectInitializer)
 {
 	if (DefaultMultiLineEditableTextStyle == nullptr)
 	{
-		// HACK: THIS SHOULD NOT COME FROM CORESTYLE AND SHOULD INSTEAD BE DEFINED BY ENGINE TEXTURES/PROJECT SETTINGS
-		DefaultMultiLineEditableTextStyle = new FTextBlockStyle(FCoreStyle::Get().GetWidgetStyle<FTextBlockStyle>("NormalText"));
+		DefaultMultiLineEditableTextStyle = new FTextBlockStyle(FUMGCoreStyle::Get().GetWidgetStyle<FTextBlockStyle>("NormalText"));
 
-		// Unlink UMG default colors from the editor settings colors.
+		// Unlink UMG default colors.
 		DefaultMultiLineEditableTextStyle->UnlinkColors();
 	}
 
 	WidgetStyle = *DefaultMultiLineEditableTextStyle;
 	
+#if WITH_EDITOR 
+	if (EditorMultiLineEditableTextStyle == nullptr)
+	{
+		EditorMultiLineEditableTextStyle = new FTextBlockStyle(FCoreStyle::Get().GetWidgetStyle<FTextBlockStyle>("NormalText"));
+
+		// Unlink UMG Editor colors from the editor settings colors.
+		EditorMultiLineEditableTextStyle->UnlinkColors();
+	}
+
+	if (IsEditorWidget())
+	{
+		WidgetStyle = *EditorMultiLineEditableTextStyle;
+	}
+#endif // WITH_EDITOR
+
 	bIsReadOnly = false;
 	SelectAllTextWhenFocused = false;
 	ClearTextSelectionOnFocusLoss = true;
diff --git a/Engine/Source/Runtime/UMG/Private/Components/MultiLineEditableTextBox.cpp b/Engine/Source/Runtime/UMG/Private/Components/MultiLineEditableTextBox.cpp
--- a/Engine/Source/Runtime/UMG/Private/Components/MultiLineEditableTextBox.cpp
+++ b/Engine/Source/Runtime/UMG/Private/Components/MultiLineEditableTextBox.cpp
@@ -5,6 +5,7 @@
 #include "Engine/Font.h"
 #include "Widgets/DeclarativeSyntaxSupport.h"
 #include "Widgets/Input/SMultiLineEditableTextBox.h"
+#include "Styling/UMGCoreStyle.h"
 
 #define LOCTEXT_NAMESPACE "UMG"
 
@@ -14,6 +15,11 @@
 static FEditableTextBoxStyle* DefaultMultiLineEditableTextBoxStyle = nullptr;
 static FTextBlockStyle* DefaultMultiLineEditableTextBoxTextStyle = nullptr;
 
+#if WITH_EDITOR
+static FEditableTextBoxStyle* EditorMultiLineEditableTextBoxStyle = nullptr;
+static FTextBlockStyle* EditorMultiLineEditableTextBoxTextStyle = nullptr;
+#endif 
+
 UMultiLineEditableTextBox::UMultiLineEditableTextBox(const FObjectInitializer& ObjectInitializer)
 	: Super(ObjectInitializer)
 {
@@ -23,25 +29,47 @@ UMultiLineEditableTextBox::UMultiLineEditableTextBox(const FObjectInitializer& O
 
 	if (DefaultMultiLineEditableTextBoxStyle == nullptr)
 	{
-		// HACK: THIS SHOULD NOT COME FROM CORESTYLE AND SHOULD INSTEAD BE DEFINED BY ENGINE TEXTURES/PROJECT SETTINGS
-		DefaultMultiLineEditableTextBoxStyle = new FEditableTextBoxStyle(FCoreStyle::Get().GetWidgetStyle<FEditableTextBoxStyle>("NormalEditableTextBox"));
+		DefaultMultiLineEditableTextBoxStyle = new FEditableTextBoxStyle(FUMGCoreStyle::Get().GetWidgetStyle<FEditableTextBoxStyle>("NormalEditableTextBox"));
 
-		// Unlink UMG default colors from the editor settings colors.
+		// Unlink UMG default colors.
 		DefaultMultiLineEditableTextBoxStyle->UnlinkColors();
 	}
 
 	if (DefaultMultiLineEditableTextBoxTextStyle == nullptr)
 	{
-		// HACK: THIS SHOULD NOT COME FROM CORESTYLE AND SHOULD INSTEAD BE DEFINED BY ENGINE TEXTURES/PROJECT SETTINGS
-		DefaultMultiLineEditableTextBoxTextStyle = new FTextBlockStyle(FCoreStyle::Get().GetWidgetStyle<FTextBlockStyle>("NormalText"));
+		DefaultMultiLineEditableTextBoxTextStyle = new FTextBlockStyle(FUMGCoreStyle::Get().GetWidgetStyle<FTextBlockStyle>("NormalText"));
 
-		// Unlink UMG default colors from the editor settings colors.
+		// Unlink UMG default colors.
 		DefaultMultiLineEditableTextBoxTextStyle->UnlinkColors();
 	}
 	
 	WidgetStyle = *DefaultMultiLineEditableTextBoxStyle;
 	TextStyle = *DefaultMultiLineEditableTextBoxTextStyle;
 
+#if WITH_EDITOR 
+	if (EditorMultiLineEditableTextBoxStyle == nullptr)
+	{
+		EditorMultiLineEditableTextBoxStyle = new FEditableTextBoxStyle(FCoreStyle::Get().GetWidgetStyle<FEditableTextBoxStyle>("NormalEditableTextBox"));
+
+		// Unlink UMG Editor colors from the editor settings colors.
+		EditorMultiLineEditableTextBoxStyle->UnlinkColors();
+	}
+
+	if (EditorMultiLineEditableTextBoxTextStyle == nullptr)
+	{
+		EditorMultiLineEditableTextBoxTextStyle = new FTextBlockStyle(FCoreStyle::Get().GetWidgetStyle<FTextBlockStyle>("NormalText"));
+
+		// Unlink UMG Editor colors from the editor settings colors.
+		EditorMultiLineEditableTextBoxTextStyle->UnlinkColors();
+	}
+	
+	if (IsEditorWidget())
+	{
+		WidgetStyle = *EditorMultiLineEditableTextBoxStyle;
+		TextStyle = *EditorMultiLineEditableTextBoxTextStyle;
+	}
+#endif // WITH_EDITOR
+
 	bIsReadOnly = false;
 	AllowContextMenu = true;
 	VirtualKeyboardDismissAction = EVirtualKeyboardDismissAction::TextChangeOnDismiss;
diff --git a/Engine/Source/Runtime/UMG/Private/Components/ProgressBar.cpp b/Engine/Source/Runtime/UMG/Private/Components/ProgressBar.cpp
--- a/Engine/Source/Runtime/UMG/Private/Components/ProgressBar.cpp
+++ b/Engine/Source/Runtime/UMG/Private/Components/ProgressBar.cpp
@@ -2,6 +2,7 @@
 
 #include "Components/ProgressBar.h"
 #include "Slate/SlateBrushAsset.h"
+#include "Styling/UMGCoreStyle.h"
 
 #define LOCTEXT_NAMESPACE "UMG"
 
@@ -10,19 +11,38 @@
 
 static FProgressBarStyle* DefaultProgressBarStyle = nullptr;
 
+#if WITH_EDITOR
+static FProgressBarStyle* EditorProgressBarStyle = nullptr;
+#endif 
+
 UProgressBar::UProgressBar(const FObjectInitializer& ObjectInitializer)
 	: Super(ObjectInitializer)
 {
 	if (DefaultProgressBarStyle == nullptr)
 	{
-		// HACK: THIS SHOULD NOT COME FROM CORESTYLE AND SHOULD INSTEAD BE DEFINED BY ENGINE TEXTURES/PROJECT SETTINGS
-		DefaultProgressBarStyle = new FProgressBarStyle(FCoreStyle::Get().GetWidgetStyle<FProgressBarStyle>("ProgressBar"));
+		DefaultProgressBarStyle = new FProgressBarStyle(FUMGCoreStyle::Get().GetWidgetStyle<FProgressBarStyle>("ProgressBar"));
 
-		// Unlink UMG default colors from the editor settings colors.
+		// Unlink UMG default colors.
 		DefaultProgressBarStyle->UnlinkColors();
 	}
 
 	WidgetStyle = *DefaultProgressBarStyle;
+
+#if WITH_EDITOR 
+	if (EditorProgressBarStyle == nullptr)
+	{
+		EditorProgressBarStyle = new FProgressBarStyle(FCoreStyle::Get().GetWidgetStyle<FProgressBarStyle>("ProgressBar"));
+
+		// Unlink UMG Editor colors from the editor settings colors.
+		EditorProgressBarStyle->UnlinkColors();
+	}
+
+	if (IsEditorWidget())
+	{
+		WidgetStyle = *EditorProgressBarStyle;
+	}
+#endif // WITH_EDITOR
+
 	WidgetStyle.FillImage.TintColor = FLinearColor::White;
 
 	BarFillType = EProgressBarFillType::LeftToRight;
diff --git a/Engine/Source/Runtime/UMG/Private/Components/ScrollBar.cpp b/Engine/Source/Runtime/UMG/Private/Components/ScrollBar.cpp
--- a/Engine/Source/Runtime/UMG/Private/Components/ScrollBar.cpp
+++ b/Engine/Source/Runtime/UMG/Private/Components/ScrollBar.cpp
@@ -2,6 +2,7 @@
 
 #include "Components/ScrollBar.h"
 #include "UObject/EditorObjectVersion.h"
+#include "Styling/UMGCoreStyle.h"
 
 #define LOCTEXT_NAMESPACE "UMG"
 
@@ -10,6 +11,10 @@
 
 static FScrollBarStyle* DefaultScrollBarStyle = nullptr;
 
+#if WITH_EDITOR
+static FScrollBarStyle* EditorScrollBarStyle = nullptr;
+#endif 
+
 UScrollBar::UScrollBar(const FObjectInitializer& ObjectInitializer)
 	: Super(ObjectInitializer)
 {
@@ -23,14 +28,28 @@ UScrollBar::UScrollBar(const FObjectInitializer& ObjectInitializer)
 
 	if (DefaultScrollBarStyle == nullptr)
 	{
-		// HACK: THIS SHOULD NOT COME FROM CORESTYLE AND SHOULD INSTEAD BE DEFINED BY ENGINE TEXTURES/PROJECT SETTINGS
-		DefaultScrollBarStyle = new FScrollBarStyle(FCoreStyle::Get().GetWidgetStyle<FScrollBarStyle>("Scrollbar"));
+		DefaultScrollBarStyle = new FScrollBarStyle(FUMGCoreStyle::Get().GetWidgetStyle<FScrollBarStyle>("Scrollbar"));
 
-		// Unlink UMG default colors from the editor settings colors.
+		// Unlink UMG default colors.
 		DefaultScrollBarStyle->UnlinkColors();
 	}
 	
 	WidgetStyle = *DefaultScrollBarStyle;
+
+#if WITH_EDITOR 
+	if (EditorScrollBarStyle == nullptr)
+	{
+		EditorScrollBarStyle = new FScrollBarStyle(FCoreStyle::Get().GetWidgetStyle<FScrollBarStyle>("Scrollbar"));
+
+		// Unlink UMG Editor colors from the editor settings colors.
+		EditorScrollBarStyle->UnlinkColors();
+	}
+
+	if (IsEditorWidget())
+	{
+		WidgetStyle = *EditorScrollBarStyle;
+	}
+#endif // WITH_EDITOR
 }
 
 void UScrollBar::ReleaseSlateResources(bool bReleaseChildren)
diff --git a/Engine/Source/Runtime/UMG/Private/Components/ScrollBox.cpp b/Engine/Source/Runtime/UMG/Private/Components/ScrollBox.cpp
--- a/Engine/Source/Runtime/UMG/Private/Components/ScrollBox.cpp
+++ b/Engine/Source/Runtime/UMG/Private/Components/ScrollBox.cpp
@@ -4,6 +4,7 @@
 #include "Containers/Ticker.h"
 #include "Components/ScrollBoxSlot.h"
 #include "UObject/EditorObjectVersion.h"
+#include "Styling/UMGCoreStyle.h"
 
 #define LOCTEXT_NAMESPACE "UMG"
 
@@ -13,6 +14,11 @@
 static FScrollBoxStyle* DefaultScrollBoxStyle = nullptr;
 static FScrollBarStyle* DefaultScrollBoxBarStyle = nullptr;
 
+#if WITH_EDITOR
+static FScrollBoxStyle* EditorScrollBoxStyle = nullptr;
+static FScrollBarStyle* EditorScrollBoxBarStyle = nullptr;
+#endif 
+
 UScrollBox::UScrollBox(const FObjectInitializer& ObjectInitializer)
 	: Super(ObjectInitializer)
 	, Orientation(Orient_Vertical)
@@ -36,25 +42,47 @@ UScrollBox::UScrollBox(const FObjectInitializer& ObjectInitializer)
 
 	if (DefaultScrollBoxStyle == nullptr)
 	{
-		// HACK: THIS SHOULD NOT COME FROM CORESTYLE AND SHOULD INSTEAD BE DEFINED BY ENGINE TEXTURES/PROJECT SETTINGS
-		DefaultScrollBoxStyle = new FScrollBoxStyle(FCoreStyle::Get().GetWidgetStyle<FScrollBoxStyle>("ScrollBox"));
+		DefaultScrollBoxStyle = new FScrollBoxStyle(FUMGCoreStyle::Get().GetWidgetStyle<FScrollBoxStyle>("ScrollBox"));
 
-		// Unlink UMG default colors from the editor settings colors.
+		// Unlink UMG default colors.
 		DefaultScrollBoxStyle->UnlinkColors();
 	}
 
 	if (DefaultScrollBoxBarStyle == nullptr)
 	{
-		// HACK: THIS SHOULD NOT COME FROM CORESTYLE AND SHOULD INSTEAD BE DEFINED BY ENGINE TEXTURES/PROJECT SETTINGS
-		DefaultScrollBoxBarStyle = new FScrollBarStyle(FCoreStyle::Get().GetWidgetStyle<FScrollBarStyle>("ScrollBar"));
+		DefaultScrollBoxBarStyle = new FScrollBarStyle(FUMGCoreStyle::Get().GetWidgetStyle<FScrollBarStyle>("ScrollBar"));
 
-		// Unlink UMG default colors from the editor settings colors.
+		// Unlink UMG default colors.
 		DefaultScrollBoxBarStyle->UnlinkColors();
 	}
 	
 	WidgetStyle = *DefaultScrollBoxStyle;
 	WidgetBarStyle = *DefaultScrollBoxBarStyle;
 
+#if WITH_EDITOR 
+	if (EditorScrollBoxStyle == nullptr)
+	{
+		EditorScrollBoxStyle = new FScrollBoxStyle(FCoreStyle::Get().GetWidgetStyle<FScrollBoxStyle>("ScrollBox"));
+
+		// Unlink UMG Editor colors from the editor settings colors.
+		EditorScrollBoxStyle->UnlinkColors();
+	}
+
+	if (EditorScrollBoxBarStyle == nullptr)
+	{
+		EditorScrollBoxBarStyle = new FScrollBarStyle(FCoreStyle::Get().GetWidgetStyle<FScrollBarStyle>("ScrollBar"));
+
+		// Unlink UMG Editor colors from the editor settings colors.
+		EditorScrollBoxBarStyle->UnlinkColors();
+	}
+	
+	if (IsEditorWidget())
+	{
+		WidgetStyle = *EditorScrollBoxStyle;
+		WidgetBarStyle = *EditorScrollBoxBarStyle;
+	}
+#endif // WITH_EDITOR
+
 	bAllowRightClickDragScrolling = true;
 }
 
diff --git a/Engine/Source/Runtime/UMG/Private/Components/Slider.cpp b/Engine/Source/Runtime/UMG/Private/Components/Slider.cpp
--- a/Engine/Source/Runtime/UMG/Private/Components/Slider.cpp
+++ b/Engine/Source/Runtime/UMG/Private/Components/Slider.cpp
@@ -3,6 +3,7 @@
 #include "Components/Slider.h"
 #include "Widgets/DeclarativeSyntaxSupport.h"
 #include "Widgets/Input/SSlider.h"
+#include "Styling/UMGCoreStyle.h"
 
 #define LOCTEXT_NAMESPACE "UMG"
 
@@ -11,6 +12,10 @@
 
 static FSliderStyle* DefaultSliderStyle = nullptr;
 
+#if WITH_EDITOR
+static FSliderStyle* EditorSliderStyle = nullptr;
+#endif 
+
 USlider::USlider(const FObjectInitializer& ObjectInitializer)
 	: Super(ObjectInitializer)
 {
@@ -26,15 +31,30 @@ USlider::USlider(const FObjectInitializer& ObjectInitializer)
 
 	if (DefaultSliderStyle == nullptr)
 	{
-		// HACK: THIS SHOULD NOT COME FROM CORESTYLE AND SHOULD INSTEAD BE DEFINED BY ENGINE TEXTURES/PROJECT SETTINGS
-		DefaultSliderStyle = new FSliderStyle(FCoreStyle::Get().GetWidgetStyle<FSliderStyle>("Slider"));
+		DefaultSliderStyle = new FSliderStyle(FUMGCoreStyle::Get().GetWidgetStyle<FSliderStyle>("Slider"));
 
-		// Unlink UMG default colors from the editor settings colors.
+		// Unlink UMG default colors.
 		DefaultSliderStyle->UnlinkColors();
 	}
 
 	WidgetStyle = *DefaultSliderStyle;
 
+
+#if WITH_EDITOR 
+	if (EditorSliderStyle == nullptr)
+	{
+		EditorSliderStyle = new FSliderStyle(FCoreStyle::Get().GetWidgetStyle<FSliderStyle>("Slider"));
+
+		// Unlink UMG Editor colors from the editor settings colors.
+		EditorSliderStyle->UnlinkColors();
+	}
+	
+	if (IsEditorWidget())
+	{
+		WidgetStyle = *EditorSliderStyle;
+	}
+#endif // WITH_EDITOR
+
 #if WITH_EDITORONLY_DATA
 	AccessibleBehavior = ESlateAccessibleBehavior::Summary;
 	bCanChildrenBeAccessible = false;
diff --git a/Engine/Source/Runtime/UMG/Private/Components/SpinBox.cpp b/Engine/Source/Runtime/UMG/Private/Components/SpinBox.cpp
--- a/Engine/Source/Runtime/UMG/Private/Components/SpinBox.cpp
+++ b/Engine/Source/Runtime/UMG/Private/Components/SpinBox.cpp
@@ -3,6 +3,7 @@
 #include "Components/SpinBox.h"
 #include "UObject/ConstructorHelpers.h"
 #include "Engine/Font.h"
+#include "Styling/UMGCoreStyle.h"
 
 #define LOCTEXT_NAMESPACE "UMG"
 
@@ -11,6 +12,10 @@
 
 static FSpinBoxStyle* DefaultSpinBoxStyle = nullptr;
 
+#if WITH_EDITOR
+static FSpinBoxStyle* EditorSpinBoxStyle = nullptr;
+#endif 
+
 USpinBox::USpinBox(const FObjectInitializer& ObjectInitializer)
 	: Super(ObjectInitializer)
 {
@@ -37,14 +42,29 @@ USpinBox::USpinBox(const FObjectInitializer& ObjectInitializer)
 
 	if (DefaultSpinBoxStyle == nullptr)
 	{
-		// HACK: THIS SHOULD NOT COME FROM CORESTYLE AND SHOULD INSTEAD BE DEFINED BY ENGINE TEXTURES/PROJECT SETTINGS
-		DefaultSpinBoxStyle = new FSpinBoxStyle(FCoreStyle::Get().GetWidgetStyle<FSpinBoxStyle>("SpinBox"));
+		DefaultSpinBoxStyle = new FSpinBoxStyle(FUMGCoreStyle::Get().GetWidgetStyle<FSpinBoxStyle>("SpinBox"));
 
-		// Unlink UMG default colors from the editor settings colors.
+		// Unlink UMG default colors.
 		DefaultSpinBoxStyle->UnlinkColors();
 	}
 
 	WidgetStyle = *DefaultSpinBoxStyle;
+
+#if WITH_EDITOR 
+	if (EditorSpinBoxStyle == nullptr)
+	{
+		EditorSpinBoxStyle = new FSpinBoxStyle(FCoreStyle::Get().GetWidgetStyle<FSpinBoxStyle>("SpinBox"));
+
+		// Unlink UMG Editor colors from the editor settings colors.
+		EditorSpinBoxStyle->UnlinkColors();
+	}
+	
+	if (IsEditorWidget())
+	{
+		WidgetStyle = *EditorSpinBoxStyle;
+	}
+#endif // WITH_EDITOR
+
 }
 
 void USpinBox::ReleaseSlateResources(bool bReleaseChildren)
diff --git a/Engine/Source/Runtime/UMG/Private/Components/Throbber.cpp b/Engine/Source/Runtime/UMG/Private/Components/Throbber.cpp
--- a/Engine/Source/Runtime/UMG/Private/Components/Throbber.cpp
+++ b/Engine/Source/Runtime/UMG/Private/Components/Throbber.cpp
@@ -3,6 +3,7 @@
 #include "Components/Throbber.h"
 #include "SlateFwd.h"
 #include "Slate/SlateBrushAsset.h"
+#include "Styling/UMGCoreStyle.h"
 
 #define LOCTEXT_NAMESPACE "UMG"
 
@@ -11,6 +12,10 @@
 
 static FSlateBrush* DefaultThrobberBrush = nullptr;
 
+#if WITH_EDITOR
+static FSlateBrush* EditorThrobberBrush = nullptr;
+#endif 
+
 UThrobber::UThrobber(const FObjectInitializer& ObjectInitializer)
 	: Super(ObjectInitializer)
 {
@@ -22,14 +27,28 @@ UThrobber::UThrobber(const FObjectInitializer& ObjectInitializer)
 
 	if (DefaultThrobberBrush == nullptr)
 	{
-		// HACK: THIS SHOULD NOT COME FROM CORESTYLE AND SHOULD INSTEAD BE DEFINED BY ENGINE TEXTURES/PROJECT SETTINGS
-		DefaultThrobberBrush = new FSlateBrush(*FCoreStyle::Get().GetBrush("Throbber.Chunk"));
+		DefaultThrobberBrush = new FSlateBrush(*FUMGCoreStyle::Get().GetBrush("Throbber.Chunk"));
 
-		// Unlink UMG default colors from the editor settings colors.
+		// Unlink UMG default colors.
 		DefaultThrobberBrush->UnlinkColors();
 	}
 
 	Image = *DefaultThrobberBrush;
+
+#if WITH_EDITOR 
+	if (EditorThrobberBrush == nullptr)
+	{
+		EditorThrobberBrush = new FSlateBrush(*FCoreStyle::Get().GetBrush("Throbber.Chunk"));
+
+		// Unlink UMG Editor colors from the editor settings colors.
+		EditorThrobberBrush->UnlinkColors();
+	}
+	
+	if (IsEditorWidget())
+	{
+		Image = *EditorThrobberBrush;
+	}
+#endif // WITH_EDITOR
 }
 
 void UThrobber::ReleaseSlateResources(bool bReleaseChildren)
diff --git a/Engine/Source/Runtime/UMG/Private/Components/Widget.cpp b/Engine/Source/Runtime/UMG/Private/Components/Widget.cpp
--- a/Engine/Source/Runtime/UMG/Private/Components/Widget.cpp
+++ b/Engine/Source/Runtime/UMG/Private/Components/Widget.cpp
@@ -21,6 +21,7 @@
 #include "Binding/PropertyBinding.h"
 #include "Logging/MessageLog.h"
 #include "Blueprint/UserWidget.h"
+#include "Blueprint/UserWidgetBlueprint.h"
 #include "Slate/SObjectWidget.h"
 #include "Blueprint/WidgetTree.h"
 #include "UMGStyle.h"
@@ -1076,6 +1077,23 @@ EVisibility UWidget::GetVisibilityInDesigner() const
 	return bHiddenInDesigner ? EVisibility::Collapsed : EVisibility::Visible;
 }
 
+bool UWidget::IsEditorWidget() const
+{
+	if (UWidgetTree* WidgetTree = Cast<UWidgetTree>(GetOuter()))
+	{
+		//@TODO: DarenC - This is a bit dirty, can't find a cleaner alternative yet though.
+		bool bIsEditorWidgetPreview = WidgetTree->RootWidget && WidgetTree->RootWidget->WidgetGeneratedBy.IsValid();
+		UObject* WidgetBPObject = bIsEditorWidgetPreview ? WidgetTree->RootWidget->WidgetGeneratedBy.Get() : WidgetTree->GetOuter();
+
+		if (UUserWidgetBlueprint* WidgetBP = Cast<UUserWidgetBlueprint>(WidgetBPObject))
+		{
+			return WidgetBP->AllowEditorWidget();
+		}
+	}
+
+	return false;
+}
+
 bool UWidget::IsVisibleInDesigner() const
 {
 	if (bHiddenInDesigner)
@@ -1138,7 +1156,7 @@ void UWidget::DeselectByDesigner()
 
 #undef LOCTEXT_NAMESPACE
 #define LOCTEXT_NAMESPACE "UMG"
-#endif
+#endif // WITH_EDITOR
 
 void UWidget::PreSave(const class ITargetPlatform* TargetPlatform)
 {
diff --git a/Engine/Source/Runtime/UMG/Private/WidgetBlueprintLibrary.cpp b/Engine/Source/Runtime/UMG/Private/WidgetBlueprintLibrary.cpp
--- a/Engine/Source/Runtime/UMG/Private/WidgetBlueprintLibrary.cpp
+++ b/Engine/Source/Runtime/UMG/Private/WidgetBlueprintLibrary.cpp
@@ -13,6 +13,7 @@
 #include "Rendering/DrawElements.h"
 #include "Styling/SlateTypes.h"
 #include "Styling/CoreStyle.h"
+#include "Styling/UMGCoreStyle.h"
 #include "Framework/Application/SlateApplication.h"
 #include "Slate/UMGDragDropOp.h"
 #include "Slate/SlateBrushAsset.h"
@@ -204,7 +205,7 @@ void UWidgetBlueprintLibrary::DrawText(FPaintContext& Context, const FString& In
 	Context.MaxLayer++;
 
 	//TODO UMG Create a font asset usable as a UFont or as a slate font asset.
-	FSlateFontInfo FontInfo = FCoreStyle::Get().GetWidgetStyle<FTextBlockStyle>("NormalText").Font;
+	FSlateFontInfo FontInfo = FUMGCoreStyle::Get().GetWidgetStyle<FTextBlockStyle>("NormalText").Font;
 	
 	FSlateDrawElement::MakeText(
 		Context.OutDrawElements,
diff --git a/Engine/Source/Runtime/UMG/Public/Blueprint/UserWidgetBlueprint.h b/Engine/Source/Runtime/UMG/Public/Blueprint/UserWidgetBlueprint.h
new file mode 100644
--- /dev/null
+++ b/Engine/Source/Runtime/UMG/Public/Blueprint/UserWidgetBlueprint.h
@@ -0,0 +1,16 @@
+// Copyright Epic Games, Inc. All Rights Reserved.
+
+#pragma once
+
+#include "Engine/Blueprint.h"
+#include "UserWidgetBlueprint.generated.h"
+
+UCLASS(Abstract)
+class UMG_API UUserWidgetBlueprint : public UBlueprint
+{
+	GENERATED_BODY()
+
+public:
+	/** Does the editor support widget from an editor package. */
+	virtual bool AllowEditorWidget() const { return false; }
+};
diff --git a/Engine/Source/Runtime/UMG/Public/Components/Widget.h b/Engine/Source/Runtime/UMG/Public/Components/Widget.h
--- a/Engine/Source/Runtime/UMG/Public/Components/Widget.h
+++ b/Engine/Source/Runtime/UMG/Public/Components/Widget.h
@@ -986,6 +986,8 @@ protected:
 	// This is an implementation detail that allows us to show and hide the widget in the designer
 	// regardless of the actual visibility state set by the user.
 	EVisibility GetVisibilityInDesigner() const;
+
+	bool IsEditorWidget() const;
 #endif
 
 	virtual void OnBindingChanged(const FName& Property);
