Slate: Deprecate direct acces from the Alignment padding and the owner pointer of a FSlot.

diff --git a/Engine/Plugins/Experimental/ControlRig/Source/ControlRigEditor/Private/Graph/SControlRigGraphNode.cpp b/Engine/Plugins/Experimental/ControlRig/Source/ControlRigEditor/Private/Graph/SControlRigGraphNode.cpp
--- a/Engine/Plugins/Experimental/ControlRig/Source/ControlRigEditor/Private/Graph/SControlRigGraphNode.cpp
+++ b/Engine/Plugins/Experimental/ControlRig/Source/ControlRigEditor/Private/Graph/SControlRigGraphNode.cpp
@@ -288,13 +288,9 @@ void SControlRigGraphNode::Construct( const FArguments& InArgs )
 				const int32 SlotToAdjustIndex = bPinWidgetForExpanderLeft ? 0 : Slots->Num() - 1;
 				SBoxPanel::FSlot& Slot = (*Slots)[SlotToAdjustIndex];
 
-				FMargin Padding;
-				if(Slot.SlotPadding.IsSet() || Slot.SlotPadding.IsBound())
-				{
-					Padding = Slot.SlotPadding.Get();
-				}
+				FMargin Padding = Slot.GetPadding();
 				Padding = FMargin(RightIndentation + Padding.Left, Padding.Top, LeftIndentation + Padding.Right, Padding.Bottom);
-				Slot.SlotPadding = Padding; 
+				Slot.SetPadding(Padding);
 			}
 		}
 	}
diff --git a/Engine/Source/Editor/CurveEditor/Private/SCurveEditorViewContainer.cpp b/Engine/Source/Editor/CurveEditor/Private/SCurveEditorViewContainer.cpp
--- a/Engine/Source/Editor/CurveEditor/Private/SCurveEditorViewContainer.cpp
+++ b/Engine/Source/Editor/CurveEditor/Private/SCurveEditorViewContainer.cpp
@@ -46,7 +46,7 @@ FVector2D SCurveEditorViewContainer::ComputeDesiredSize(float) const
 		{
 			const FVector2D& ChildDesiredSize = Child.GetWidget()->GetDesiredSize();
 
-			FMargin SlotPadding = Child.SlotPadding.Get();
+			FMargin SlotPadding = Child.GetPadding();
 			MyDesiredSize.Y += SlotPadding.GetTotalSpaceAlong<Orient_Vertical>();
 
 			// For a vertical panel, we want to find the maximum desired width (including margin).
diff --git a/Engine/Source/Editor/GraphEditor/Public/SNodePanel.h b/Engine/Source/Editor/GraphEditor/Public/SNodePanel.h
--- a/Engine/Source/Editor/GraphEditor/Public/SNodePanel.h
+++ b/Engine/Source/Editor/GraphEditor/Public/SNodePanel.h
@@ -275,30 +275,17 @@ public:
 	public:
 
 		/** A slot that support alignment of content and padding and z-order */
-		class FNodeSlot : public TSlotBase<FNodeSlot>
+		class FNodeSlot : public TSlotBase<FNodeSlot>, public TAlignmentWidgetSlotMixin<FNodeSlot>
 		{
 		public:
 			FNodeSlot()
 				: TSlotBase<FNodeSlot>()
-				, HAlignment(HAlign_Fill)
-				, VAlignment(VAlign_Fill)
+				, TAlignmentWidgetSlotMixin<FNodeSlot>(HAlign_Fill, VAlign_Fill)
 				, SlotPadding(0.0f)
 				, Offset( FVector2D::ZeroVector )
 				, AllowScale( true )
 			{ }
 
-			FNodeSlot& HAlign( EHorizontalAlignment InHAlignment )
-			{
-				HAlignment = InHAlignment;
-				return *this;
-			}
-
-			FNodeSlot& VAlign( EVerticalAlignment InVAlignment )
-			{
-				VAlignment = InVAlignment;
-				return *this;
-			}
-
 			FNodeSlot& Padding( const TAttribute<FMargin> InPadding )
 			{
 				SlotPadding = InPadding;
@@ -327,8 +314,6 @@ public:
 
 			/** The child widget contained in this slot. */
 			ENodeZone::Type Zone;
-			EHorizontalAlignment HAlignment;
-			EVerticalAlignment VAlignment;
 			TAttribute<FMargin> SlotPadding;
 			TAttribute<FVector2D> Offset;
 			TAttribute<FVector2D> Size;
diff --git a/Engine/Source/Editor/PIEPreviewDeviceProfileSelector/Private/PIEPreviewWindow.cpp b/Engine/Source/Editor/PIEPreviewDeviceProfileSelector/Private/PIEPreviewWindow.cpp
--- a/Engine/Source/Editor/PIEPreviewDeviceProfileSelector/Private/PIEPreviewWindow.cpp
+++ b/Engine/Source/Editor/PIEPreviewDeviceProfileSelector/Private/PIEPreviewWindow.cpp
@@ -572,7 +572,7 @@ void SPIEPreviewWindow::UpdateWindow()
 	ReshapeWindow(FVector2D(PosX, PosY), FVector2D(Device->GetWindowWidth(), Device->GetWindowHeight()));
 
 	// offset the viewport widget into its correct location
-	ContentSlot->SlotPadding = Device->GetViewportMargin();
+	ContentSlot->SetPadding(Device->GetViewportMargin());
 
 	// bezel orientation depends on the window size so we need to call it after ReshapeWindow()
 	ComputeBezelOrientation();
diff --git a/Engine/Source/Editor/Persona/Private/AnimTimeline/SAnimTrackArea.cpp b/Engine/Source/Editor/Persona/Private/AnimTimeline/SAnimTrackArea.cpp
--- a/Engine/Source/Editor/Persona/Private/AnimTimeline/SAnimTrackArea.cpp
+++ b/Engine/Source/Editor/Persona/Private/AnimTimeline/SAnimTrackArea.cpp
@@ -14,11 +14,9 @@
 #include "AnimTimeSliderController.h"
 
 FAnimTrackAreaSlot::FAnimTrackAreaSlot(const TSharedPtr<SAnimTrack>& InSlotContent)
+	: TAlignmentWidgetSlotMixin<FAnimTrackAreaSlot>(HAlign_Fill, VAlign_Top)
 {
 	TrackWidget = InSlotContent;
-	
-	HAlignment = HAlign_Fill;
-	VAlignment = VAlign_Top;
 
 	AttachWidget(
 		SNew(SWeakWidget)
diff --git a/Engine/Source/Editor/Persona/Private/AnimTimeline/SAnimTrackArea.h b/Engine/Source/Editor/Persona/Private/AnimTimeline/SAnimTrackArea.h
--- a/Engine/Source/Editor/Persona/Private/AnimTimeline/SAnimTrackArea.h
+++ b/Engine/Source/Editor/Persona/Private/AnimTimeline/SAnimTrackArea.h
@@ -23,7 +23,7 @@ class FAnimTimeSliderController;
 /**
  * Structure representing a slot in the track area.
  */
-class FAnimTrackAreaSlot : public TSlotBase<FAnimTrackAreaSlot>
+class FAnimTrackAreaSlot : public TSlotBase<FAnimTrackAreaSlot>, public TAlignmentWidgetSlotMixin<FAnimTrackAreaSlot>
 {
 public:
 
@@ -33,10 +33,6 @@ public:
 	/** Get the vertical position of this slot inside its parent. */
 	float GetVerticalOffset() const;
 
-	/** Horizontal/Vertical alignment for the slot. */
-	EHorizontalAlignment HAlignment;
-	EVerticalAlignment VAlignment;
-
 	/** The track that we represent. */
 	TWeakPtr<SAnimTrack> TrackWidget;
 };
diff --git a/Engine/Source/Editor/Sequencer/Private/SSequencerTrackArea.cpp b/Engine/Source/Editor/Sequencer/Private/SSequencerTrackArea.cpp
--- a/Engine/Source/Editor/Sequencer/Private/SSequencerTrackArea.cpp
+++ b/Engine/Source/Editor/Sequencer/Private/SSequencerTrackArea.cpp
@@ -18,11 +18,9 @@
 #include "DisplayNodes/SequencerObjectBindingNode.h"
 
 FTrackAreaSlot::FTrackAreaSlot(const TSharedPtr<SSequencerTrackLane>& InSlotContent)
+	: TAlignmentWidgetSlotMixin<FTrackAreaSlot>(HAlign_Fill, VAlign_Top)
 {
 	TrackLane = InSlotContent;
-	
-	HAlignment = HAlign_Fill;
-	VAlignment = VAlign_Top;
 
 	this->AttachWidget(
 		SNew(SWeakWidget)
diff --git a/Engine/Source/Editor/Sequencer/Private/SSequencerTrackArea.h b/Engine/Source/Editor/Sequencer/Private/SSequencerTrackArea.h
--- a/Engine/Source/Editor/Sequencer/Private/SSequencerTrackArea.h
+++ b/Engine/Source/Editor/Sequencer/Private/SSequencerTrackArea.h
@@ -26,7 +26,7 @@ class SSequencerTreeView;
  * Structure representing a slot in the track area.
  */
 class FTrackAreaSlot
-	: public TSlotBase<FTrackAreaSlot>
+	: public TSlotBase<FTrackAreaSlot>, public TAlignmentWidgetSlotMixin<FTrackAreaSlot>
 {
 public:
 
@@ -36,10 +36,6 @@ public:
 	/** Get the vertical position of this slot inside its parent. */
 	float GetVerticalOffset() const;
 
-	/** Horizontal/Vertical alignment for the slot. */
-	EHorizontalAlignment HAlignment;
-	EVerticalAlignment VAlignment;
-
 	/** The track lane that we represent. */
 	TWeakPtr<SSequencerTrackLane> TrackLane;
 };
diff --git a/Engine/Source/Runtime/AppFramework/Private/Widgets/Testing/STestSuite.cpp b/Engine/Source/Runtime/AppFramework/Private/Widgets/Testing/STestSuite.cpp
--- a/Engine/Source/Runtime/AppFramework/Private/Widgets/Testing/STestSuite.cpp
+++ b/Engine/Source/Runtime/AppFramework/Private/Widgets/Testing/STestSuite.cpp
@@ -4365,7 +4365,7 @@ class SDPIScalingTest : public SCompoundWidget
 
 	ECheckBoxState IsFillChecked() const
 	{
-		const bool bIsFilling = (ScalerSlot->HAlignment == HAlign_Fill);
+		const bool bIsFilling = (ScalerSlot->GetHorizontalAlignment() == HAlign_Fill);
 		return (bIsFilling)
 			? ECheckBoxState::Checked
 			: ECheckBoxState::Unchecked;
@@ -4373,8 +4373,8 @@ class SDPIScalingTest : public SCompoundWidget
 
 	void OnFillChecked(ECheckBoxState InValue)
 	{
-		ScalerSlot->HAlign( (InValue == ECheckBoxState::Checked) ? HAlign_Fill : HAlign_Center );
-		ScalerSlot->VAlign( (InValue == ECheckBoxState::Checked) ? VAlign_Fill : VAlign_Center );
+		ScalerSlot->SetHorizontalAlignment( (InValue == ECheckBoxState::Checked) ? HAlign_Fill : HAlign_Center );
+		ScalerSlot->SetVerticalAlignment( (InValue == ECheckBoxState::Checked) ? VAlign_Fill : VAlign_Center );
 	}
 
 	float DPIScale;
diff --git a/Engine/Source/Runtime/Slate/Private/Framework/Docking/SDockingTabStack.cpp b/Engine/Source/Runtime/Slate/Private/Framework/Docking/SDockingTabStack.cpp
--- a/Engine/Source/Runtime/Slate/Private/Framework/Docking/SDockingTabStack.cpp
+++ b/Engine/Source/Runtime/Slate/Private/Framework/Docking/SDockingTabStack.cpp
@@ -848,7 +848,7 @@ void SDockingTabStack::ReserveSpaceForWindowChrome(EChromeElement Element, bool
 	#endif
 
 	bShowingTitleBarArea = true;
-	const FMargin CurrentPadding = TitleBarSlot->SlotPadding.Get();
+	const FMargin CurrentPadding = TitleBarSlot->GetPadding();
 	switch (Element)
 	{
 	case EChromeElement::Controls:
diff --git a/Engine/Source/Runtime/Slate/Private/Widgets/Layout/SBackgroundBlur.cpp b/Engine/Source/Runtime/Slate/Private/Widgets/Layout/SBackgroundBlur.cpp
--- a/Engine/Source/Runtime/Slate/Private/Widgets/Layout/SBackgroundBlur.cpp
+++ b/Engine/Source/Runtime/Slate/Private/Widgets/Layout/SBackgroundBlur.cpp
@@ -82,20 +82,12 @@ void SBackgroundBlur::SetLowQualityBackgroundBrush(const FSlateBrush* InBrush)
 
 void SBackgroundBlur::SetHAlign(EHorizontalAlignment HAlign)
 {
-	if (ChildSlot.HAlignment != HAlign)
-	{
-		ChildSlot.HAlignment = HAlign;
-		Invalidate(EInvalidateWidget::Layout);
-	}
+	ChildSlot.SetHorizontalAlignment(HAlign);
 }
 
 void SBackgroundBlur::SetVAlign(EVerticalAlignment VAlign)
 {
-	if (ChildSlot.VAlignment != VAlign)
-	{
-		ChildSlot.VAlignment = VAlign;
-		Invalidate(EInvalidateWidget::Layout);
-	}
+	ChildSlot.SetVerticalAlignment(VAlign);
 }
 
 void SBackgroundBlur::SetPadding(TAttribute<FMargin> InPadding)
diff --git a/Engine/Source/Runtime/Slate/Private/Widgets/Layout/SBorder.cpp b/Engine/Source/Runtime/Slate/Private/Widgets/Layout/SBorder.cpp
--- a/Engine/Source/Runtime/Slate/Private/Widgets/Layout/SBorder.cpp
+++ b/Engine/Source/Runtime/Slate/Private/Widgets/Layout/SBorder.cpp
@@ -146,25 +146,17 @@ void SBorder::SetDesiredSizeScale(TAttribute<FVector2D> InDesiredSizeScale)
 
 void SBorder::SetHAlign(EHorizontalAlignment HAlign)
 {
-	if (ChildSlot.HAlignment != HAlign)
-	{
-		ChildSlot.HAlignment = HAlign;
-		Invalidate(EInvalidateWidget::Layout);
-	}
+	ChildSlot.SetHorizontalAlignment(HAlign);
 }
 
 void SBorder::SetVAlign(EVerticalAlignment VAlign)
 {
-	if (ChildSlot.VAlignment != VAlign)
-	{
-		ChildSlot.VAlignment = VAlign;
-		Invalidate(EInvalidateWidget::Layout);
-	}
+	ChildSlot.SetVerticalAlignment(VAlign);
 }
 
-void SBorder::SetPadding(const TAttribute<FMargin>& InPadding)
+void SBorder::SetPadding(TAttribute<FMargin> InPadding)
 {
-	ChildSlot.SetPadding(InPadding);
+	ChildSlot.SetPadding(MoveTemp(InPadding));
 }
 
 void SBorder::SetShowEffectWhenDisabled(TAttribute<bool> InShowEffectWhenDisabled)
diff --git a/Engine/Source/Runtime/Slate/Private/Widgets/Layout/SBox.cpp b/Engine/Source/Runtime/Slate/Private/Widgets/Layout/SBox.cpp
--- a/Engine/Source/Runtime/Slate/Private/Widgets/Layout/SBox.cpp
+++ b/Engine/Source/Runtime/Slate/Private/Widgets/Layout/SBox.cpp
@@ -66,23 +66,15 @@ void SBox::SetContent(const TSharedRef< SWidget >& InContent)
 
 void SBox::SetHAlign(EHorizontalAlignment HAlign)
 {
-	if (ChildSlot.HAlignment != HAlign)
-	{
-		ChildSlot.HAlignment = HAlign;
-		Invalidate(EInvalidateWidget::Layout);
-	}
+	ChildSlot.SetHorizontalAlignment(HAlign);
 }
 
 void SBox::SetVAlign(EVerticalAlignment VAlign)
 {
-	if (ChildSlot.VAlignment != VAlign)
-	{
-		ChildSlot.VAlignment = VAlign;
-		Invalidate(EInvalidateWidget::Layout);
-	}
+	ChildSlot.SetVerticalAlignment(VAlign);
 }
 
-void SBox::SetPadding(const TAttribute<FMargin>& InPadding)
+void SBox::SetPadding(TAttribute<FMargin> InPadding)
 {
 	ChildSlot.SetPadding(InPadding);
 }
diff --git a/Engine/Source/Runtime/Slate/Private/Widgets/Layout/SScaleBox.cpp b/Engine/Source/Runtime/Slate/Private/Widgets/Layout/SScaleBox.cpp
--- a/Engine/Source/Runtime/Slate/Private/Widgets/Layout/SScaleBox.cpp
+++ b/Engine/Source/Runtime/Slate/Private/Widgets/Layout/SScaleBox.cpp
@@ -235,14 +235,14 @@ void SScaleBox::OnArrangeChildren(const FGeometry& AllottedGeometry, FArrangedCh
 
 			// If the layout horizontally is fill, then we need the desired size to be the whole size of the widget, 
 			// but scale the inverse of the scale we're applying.
-			if (ChildSlot.HAlignment == HAlign_Fill)
+			if (ChildSlot.GetHorizontalAlignment() == HAlign_Fill)
 			{
 				SlotWidgetDesiredSize.X = AreaSize.X / FinalScale;
 			}
 
 			// If the layout vertically is fill, then we need the desired size to be the whole size of the widget, 
 			// but scale the inverse of the scale we're applying.
-			if (ChildSlot.VAlignment == VAlign_Fill)
+			if (ChildSlot.GetVerticalAlignment() == VAlign_Fill)
 			{
 				SlotWidgetDesiredSize.Y = AreaSize.Y / FinalScale;
 			}
@@ -315,20 +315,12 @@ void SScaleBox::SetContent(TSharedRef<SWidget> InContent)
 
 void SScaleBox::SetHAlign(EHorizontalAlignment HAlign)
 {
-	if(ChildSlot.HAlignment != HAlign)
-	{
-		ChildSlot.HAlignment = HAlign;
-		Invalidate(EInvalidateWidgetReason::Layout);
-	}
+	ChildSlot.SetHorizontalAlignment(HAlign);
 }
 
 void SScaleBox::SetVAlign(EVerticalAlignment VAlign)
 {
-	if(ChildSlot.VAlignment != VAlign)
-	{
-		ChildSlot.VAlignment = VAlign;
-		Invalidate(EInvalidateWidgetReason::Layout);
-	}
+	ChildSlot.SetVerticalAlignment(VAlign);
 }
 
 void SScaleBox::SetStretchDirection(EStretchDirection::Type InStretchDirection)
diff --git a/Engine/Source/Runtime/Slate/Private/Widgets/Layout/SWindowTitleBarArea.cpp b/Engine/Source/Runtime/Slate/Private/Widgets/Layout/SWindowTitleBarArea.cpp
--- a/Engine/Source/Runtime/Slate/Private/Widgets/Layout/SWindowTitleBarArea.cpp
+++ b/Engine/Source/Runtime/Slate/Private/Widgets/Layout/SWindowTitleBarArea.cpp
@@ -181,17 +181,17 @@ void SWindowTitleBarArea::SetContent(const TSharedRef< SWidget >& InContent)
 
 void SWindowTitleBarArea::SetHAlign(EHorizontalAlignment HAlign)
 {
-	ChildSlot.HAlignment = HAlign;
+	ChildSlot.SetHorizontalAlignment(HAlign);
 }
 
 void SWindowTitleBarArea::SetVAlign(EVerticalAlignment VAlign)
 {
-	ChildSlot.VAlignment = VAlign;
+	ChildSlot.SetVerticalAlignment(VAlign);
 }
 
-void SWindowTitleBarArea::SetPadding(const TAttribute<FMargin>& InPadding)
+void SWindowTitleBarArea::SetPadding(TAttribute<FMargin> InPadding)
 {
-	ChildSlot.Padding(InPadding);
+	ChildSlot.SetPadding(MoveTemp(InPadding));
 }
 
 FVector2D SWindowTitleBarArea::ComputeDesiredSize( float ) const
diff --git a/Engine/Source/Runtime/Slate/Private/Widgets/SCanvas.cpp b/Engine/Source/Runtime/Slate/Private/Widgets/SCanvas.cpp
--- a/Engine/Source/Runtime/Slate/Private/Widgets/SCanvas.cpp
+++ b/Engine/Source/Runtime/Slate/Private/Widgets/SCanvas.cpp
@@ -43,7 +43,7 @@ void SCanvas::OnArrangeChildren( const FGeometry& AllottedGeometry, FArrangedChi
 			//Handle HAlignment
 			FVector2D Offset(0.0f, 0.0f);
 
-			switch (CurChild.HAlignment)
+			switch (CurChild.GetHorizontalAlignment())
 			{
 			case HAlign_Center:
 				Offset.X = -Size.X / 2.0f;
@@ -57,7 +57,7 @@ void SCanvas::OnArrangeChildren( const FGeometry& AllottedGeometry, FArrangedChi
 			}
 
 			//handle VAlignment
-			switch (CurChild.VAlignment)
+			switch (CurChild.GetVerticalAlignment())
 			{
 			case VAlign_Bottom:
 				Offset.Y = -Size.Y;
diff --git a/Engine/Source/Runtime/Slate/Public/Widgets/Input/SSegmentedControl.h b/Engine/Source/Runtime/Slate/Public/Widgets/Input/SSegmentedControl.h
--- a/Engine/Source/Runtime/Slate/Public/Widgets/Input/SSegmentedControl.h
+++ b/Engine/Source/Runtime/Slate/Public/Widgets/Input/SSegmentedControl.h
@@ -173,11 +173,11 @@ public:
 			int32 ColumnIndex = MaxSegmentsPerLine ? SlotIndex % MaxSegmentsPerLine : SlotIndex;
 			UniformBox->AddSlot(ColumnIndex, MaxSegmentsPerLine > 0 ? SlotIndex / MaxSegmentsPerLine : 0)
 			// Note HAlignment is applied at the check box level because if it were applied here it would make the slots look physically disconnected from each other 
-			.VAlign(ChildSlotPtr->VAlignment)
+			.VAlign(ChildSlotPtr->GetVerticalAlignment())
 			[
 				SNew(SCheckBox)
 				.Clipping(EWidgetClipping::ClipToBounds)
-				.HAlign(ChildSlotPtr->HAlignment)
+				.HAlign(ChildSlotPtr->GetHorizontalAlignment())
 				.ToolTipText(ChildSlotPtr->_Tooltip)
 				.Style(ColumnIndex == 0 ? &Style->FirstControlStyle : ColumnIndex == (NumSlots - 1) ? &Style->LastControlStyle : &Style->ControlStyle)
 				.IsChecked(this, &SSegmentedControl::IsCurrentValue, ChildValue)
diff --git a/Engine/Source/Runtime/Slate/Public/Widgets/Layout/SBorder.h b/Engine/Source/Runtime/Slate/Public/Widgets/Layout/SBorder.h
--- a/Engine/Source/Runtime/Slate/Public/Widgets/Layout/SBorder.h
+++ b/Engine/Source/Runtime/Slate/Public/Widgets/Layout/SBorder.h
@@ -118,7 +118,7 @@ public:
 	void SetVAlign(EVerticalAlignment VAlign);
 
 	/** See Padding attribute */
-	void SetPadding(const TAttribute<FMargin>& InPadding);
+	void SetPadding(TAttribute<FMargin> InPadding);
 
 	/** See ShowEffectWhenDisabled attribute */
 	void SetShowEffectWhenDisabled(TAttribute<bool> InShowEffectWhenDisabled);
diff --git a/Engine/Source/Runtime/Slate/Public/Widgets/Layout/SBox.h b/Engine/Source/Runtime/Slate/Public/Widgets/Layout/SBox.h
--- a/Engine/Source/Runtime/Slate/Public/Widgets/Layout/SBox.h
+++ b/Engine/Source/Runtime/Slate/Public/Widgets/Layout/SBox.h
@@ -101,7 +101,7 @@ public:
 	void SetVAlign(EVerticalAlignment VAlign);
 
 	/** See Padding attribute */
-	void SetPadding(const TAttribute<FMargin>& InPadding);
+	void SetPadding(TAttribute<FMargin> InPadding);
 
 	/** See WidthOverride attribute */
 	void SetWidthOverride(TAttribute<FOptionalSize> InWidthOverride);
diff --git a/Engine/Source/Runtime/Slate/Public/Widgets/Layout/SWindowTitleBarArea.h b/Engine/Source/Runtime/Slate/Public/Widgets/Layout/SWindowTitleBarArea.h
--- a/Engine/Source/Runtime/Slate/Public/Widgets/Layout/SWindowTitleBarArea.h
+++ b/Engine/Source/Runtime/Slate/Public/Widgets/Layout/SWindowTitleBarArea.h
@@ -74,7 +74,7 @@ public:
 	void SetVAlign(EVerticalAlignment VAlign);
 
 	/** See Padding attribute */
-	void SetPadding(const TAttribute<FMargin>& InPadding);
+	void SetPadding(TAttribute<FMargin> InPadding);
 
 	void SetGameWindow(TSharedPtr<SWindow> Window)
 	{
diff --git a/Engine/Source/Runtime/Slate/Public/Widgets/SCanvas.h b/Engine/Source/Runtime/Slate/Public/Widgets/SCanvas.h
--- a/Engine/Source/Runtime/Slate/Public/Widgets/SCanvas.h
+++ b/Engine/Source/Runtime/Slate/Public/Widgets/SCanvas.h
@@ -25,8 +25,31 @@ public:
 
 	/**
 	 * Canvas slots allow child widgets to be positioned and sized
+	 *
+	 * Horizontal Alignment 
+	 *  Given a top aligned slot, where '+' represents the 
+	 *  anchor point defined by PositionAttr.
+	 *  
+	 *   Left				Center				Right
+	 *	+ _ _ _ _            _ _ + _ _          _ _ _ _ +
+	 *	|		  |		   | 		   |	  |		    |
+	 *	| _ _ _ _ |        | _ _ _ _ _ |	  | _ _ _ _ |
+	 * 
+	 *  Note: FILL is NOT supported.
+	 *
+	 * Vertical Alignment 
+	 *   Given a left aligned slot, where '+' represents the 
+	 *   anchor point defined by PositionAttr.
+	 *  
+	 *   Top				Center			  Bottom
+	 *	+_ _ _ _ _          _ _ _ _ _		 _ _ _ _ _ 
+	 *	|		  |		   | 		 |		|		  |
+	 *	| _ _ _ _ |        +		 |		|		  |
+	 *					   | _ _ _ _ |		+ _ _ _ _ |
+	 * 
+	 *  Note: FILL is NOT supported.
 	 */
-	class FSlot : public TSlotBase<FSlot>
+	class FSlot : public TSlotBase<FSlot>, public TAlignmentWidgetSlotMixin<FSlot>
 	{
 	public:		
 		FSlot& Position( const TAttribute<FVector2D>& InPosition )
@@ -41,59 +64,18 @@ public:
 			return *this;
 		}
 
-		FSlot& HAlign( EHorizontalAlignment InHAlignment )
-		{
-			HAlignment = InHAlignment;
-			return *this;
-		}
-
-		FSlot& VAlign( EVerticalAlignment InVAlignment )
-		{
-			VAlignment = InVAlignment;
-			return *this;
-		}
-
 		/** Position */
 		TAttribute<FVector2D> PositionAttr;
 
 		/** Size */
 		TAttribute<FVector2D> SizeAttr;
 
-		/** Horizontal Alignment 
-		*  Given a top aligned slot, where '+' represents the 
-		*  anchor point defined by PositionAttr.
-		*  
-		*   Left				Center				Right
-			+ _ _ _ _            _ _ + _ _          _ _ _ _ +
-			|		  |		   | 		   |	  |		    |
-			| _ _ _ _ |        | _ _ _ _ _ |	  | _ _ _ _ |
-		* 
-		*  Note: FILL is NOT supported.
-		*/
-		EHorizontalAlignment HAlignment;
-
-		/** Vertical Alignment 
-		*   Given a left aligned slot, where '+' represents the 
-		*   anchor point defined by PositionAttr.
-		*  
-		*   Top					Center			  Bottom
-		*	+_ _ _ _ _          _ _ _ _ _		 _ _ _ _ _ 
-		*	|		  |		   | 		 |		|		  |
-		*	| _ _ _ _ |        +		 |		|		  |
-		*					   | _ _ _ _ |		+ _ _ _ _ |
-		* 
-		*  Note: FILL is NOT supported.
-		*/
-		EVerticalAlignment VAlignment;
-
-
 		/** Default values for a slot. */
 		FSlot()
 			: TSlotBase<FSlot>()
+			, TAlignmentWidgetSlotMixin<FSlot>(HAlign_Left, VAlign_Top)
 			, PositionAttr( FVector2D::ZeroVector )
 			, SizeAttr( FVector2D( 1.0f, 1.0f ) )
-			, HAlignment(HAlign_Left)
-			, VAlignment(VAlign_Top)
 		{ }
 	};
 
diff --git a/Engine/Source/Runtime/SlateCore/Private/SlotBase.cpp b/Engine/Source/Runtime/SlateCore/Private/SlotBase.cpp
--- a/Engine/Source/Runtime/SlateCore/Private/SlotBase.cpp
+++ b/Engine/Source/Runtime/SlateCore/Private/SlotBase.cpp
@@ -1,28 +1,64 @@
 // Copyright Epic Games, Inc. All Rights Reserved.
 
 #include "SlotBase.h"
+#include "Layout/Children.h"
 #include "Widgets/SWidget.h"
 #include "Widgets/SNullWidget.h"
 
+PRAGMA_DISABLE_DEPRECATION_WARNINGS
 FSlotBase::FSlotBase()
-	: RawParentPtr(nullptr)
-	, Widget( SNullWidget::NullWidget )
+	: Owner(nullptr)
+	, Widget(SNullWidget::NullWidget)
+#if WITH_EDITORONLY_DATA
+	, RawParentPtr(nullptr)
+#endif
+{
+
+}
+
+FSlotBase::FSlotBase(const FChildren& Children)
+	: Owner(&Children)
+	, Widget(SNullWidget::NullWidget)
+#if WITH_EDITORONLY_DATA
+	, RawParentPtr(nullptr)
+#endif
 {
 
 }
 
 FSlotBase::FSlotBase( const TSharedRef<SWidget>& InWidget )
-	: RawParentPtr(nullptr)
-	, Widget( InWidget )
+	: Owner(nullptr)
+	, Widget(InWidget)
+#if WITH_EDITORONLY_DATA
+	, RawParentPtr(nullptr)
+#endif
+{
+
+}
+PRAGMA_ENABLE_DEPRECATION_WARNINGS
+
+SWidget* FSlotBase::GetOwnerWidget() const
 {
-	
+	return GetOwner() ? &(GetOwner()->GetOwner()) : nullptr;
+}
+
+void FSlotBase::SetOwner(const FChildren& InChildren)
+{
+	if (Owner != &InChildren)
+	{
+		if (ensureMsgf(Owner == nullptr, TEXT("Slots should not be reassigned to different parents.")))
+		{
+			Owner = &InChildren;
+			AfterContentOrOwnerAssigned();
+		}
+	}
 }
 
 const TSharedPtr<SWidget> FSlotBase::DetachWidget()
 {
 	if (Widget != SNullWidget::NullWidget)
 	{
-		Widget->ConditionallyDetatchParentWidget(RawParentPtr);
+		Widget->ConditionallyDetatchParentWidget(GetOwnerWidget());
 
 		// Invalidate Prepass?
 
@@ -40,10 +76,9 @@ const TSharedPtr<SWidget> FSlotBase::DetachWidget()
 void FSlotBase::Invalidate(EInvalidateWidgetReason InvalidateReason)
 {
 	// If a slot invalidates it needs to invalidate the parent of widget of its content.
-	const TSharedPtr<SWidget>& ParentWidget = Widget->GetParentWidget();
-	if (ParentWidget.IsValid())
+	if (SWidget* OwnerWidget = GetOwnerWidget())
 	{
-		ParentWidget->Invalidate(InvalidateReason);
+		OwnerWidget->Invalidate(InvalidateReason);
 	}
 }
 
@@ -51,13 +86,13 @@ void FSlotBase::DetatchParentFromContent()
 {
 	if (Widget != SNullWidget::NullWidget)
 	{
-		Widget->ConditionallyDetatchParentWidget(RawParentPtr);
+		Widget->ConditionallyDetatchParentWidget(GetOwnerWidget());
 	}
 }
 
 void FSlotBase::AfterContentOrOwnerAssigned()
 {
-	if (RawParentPtr)
+	if (SWidget* OwnerWidget = GetOwnerWidget())
 	{
 		if (Widget != SNullWidget::NullWidget)
 		{
@@ -65,7 +100,7 @@ void FSlotBase::AfterContentOrOwnerAssigned()
 			// have made assumptions about being able to freely reparent widgets, while they're
 			// still connected to an existing hierarchy.
 			//ensure(!Widget->IsParentValid());
-			Widget->AssignParentWidget(RawParentPtr->AsShared());
+			Widget->AssignParentWidget(OwnerWidget->AsShared());
 		}
 	}
 }
diff --git a/Engine/Source/Runtime/SlateCore/Private/Widgets/SBoxPanel.cpp b/Engine/Source/Runtime/SlateCore/Private/Widgets/SBoxPanel.cpp
--- a/Engine/Source/Runtime/SlateCore/Private/Widgets/SBoxPanel.cpp
+++ b/Engine/Source/Runtime/SlateCore/Private/Widgets/SBoxPanel.cpp
@@ -55,7 +55,7 @@ static void ArrangeChildrenAlong(EFlowDirection InLayoutFlow, const TPanelChildr
 			{
 				bAnyChildVisible = true;
 				// All widgets contribute their margin to the fixed space requirement
-				FixedTotal += CurChild.SlotPadding.Get().GetTotalSpaceAlong<Orientation>();
+				FixedTotal += CurChild.GetPadding().GetTotalSpaceAlong<Orientation>();
 
 				if ( CurChild.SizeParam.SizeRule == FSizeParam::SizeRule_Stretch )
 				{
@@ -129,7 +129,7 @@ static void ArrangeChildrenAlong(EFlowDirection InLayoutFlow, const TPanelChildr
 				}
 			}
 
-			const FMargin SlotPadding(LayoutPaddingWithFlow(InLayoutFlow, CurChild.SlotPadding.Get()));
+			const FMargin SlotPadding(LayoutPaddingWithFlow(InLayoutFlow, CurChild.GetPadding()));
 
 			FVector2D SlotSize = (Orientation == Orient_Vertical)
 				? FVector2D(AllottedGeometry.GetLocalSize().X, ChildSize + SlotPadding.GetTotalSpaceAlong<Orient_Vertical>())
@@ -211,7 +211,7 @@ static FVector2D ComputeDesiredSizeForBox( const TPanelChildren<SBoxPanel::FSlot
 			{
 				// For a vertical panel, we want to find the maximum desired width (including margin).
 				// That will be the desired width of the whole panel.
-				MyDesiredSize.X = FMath::Max(MyDesiredSize.X, CurChildDesiredSize.X + CurChild.SlotPadding.Get().GetTotalSpaceAlong<Orient_Horizontal>());
+				MyDesiredSize.X = FMath::Max(MyDesiredSize.X, CurChildDesiredSize.X + CurChild.GetPadding().GetTotalSpaceAlong<Orient_Horizontal>());
 
 				// Clamp to the max size if it was specified
 				float FinalChildDesiredSize = CurChildDesiredSize.Y;
@@ -221,13 +221,13 @@ static FVector2D ComputeDesiredSizeForBox( const TPanelChildren<SBoxPanel::FSlot
 					FinalChildDesiredSize = FMath::Min( MaxSize, FinalChildDesiredSize );
 				}
 
-				MyDesiredSize.Y += FinalChildDesiredSize + CurChild.SlotPadding.Get().GetTotalSpaceAlong<Orient_Vertical>();				
+				MyDesiredSize.Y += FinalChildDesiredSize + CurChild.GetPadding().GetTotalSpaceAlong<Orient_Vertical>();
 			}
 			else
 			{
 				// A horizontal panel is just a sideways vertical panel: the axes are swapped.
 
-				MyDesiredSize.Y = FMath::Max(MyDesiredSize.Y, CurChildDesiredSize.Y + CurChild.SlotPadding.Get().GetTotalSpaceAlong<Orient_Vertical>());
+				MyDesiredSize.Y = FMath::Max(MyDesiredSize.Y, CurChildDesiredSize.Y + CurChild.GetPadding().GetTotalSpaceAlong<Orient_Vertical>());
 
 				// Clamp to the max size if it was specified
 				float FinalChildDesiredSize = CurChildDesiredSize.X;
@@ -237,7 +237,7 @@ static FVector2D ComputeDesiredSizeForBox( const TPanelChildren<SBoxPanel::FSlot
 					FinalChildDesiredSize = FMath::Min( MaxSize, FinalChildDesiredSize );
 				}
 
-				MyDesiredSize.X += FinalChildDesiredSize + CurChild.SlotPadding.Get().GetTotalSpaceAlong<Orient_Horizontal>();
+				MyDesiredSize.X += FinalChildDesiredSize + CurChild.GetPadding().GetTotalSpaceAlong<Orient_Horizontal>();
 			}
 		}
 	}
diff --git a/Engine/Source/Runtime/SlateCore/Private/Widgets/SOverlay.cpp b/Engine/Source/Runtime/SlateCore/Private/Widgets/SOverlay.cpp
--- a/Engine/Source/Runtime/SlateCore/Private/Widgets/SOverlay.cpp
+++ b/Engine/Source/Runtime/SlateCore/Private/Widgets/SOverlay.cpp
@@ -31,7 +31,7 @@ void SOverlay::OnArrangeChildren( const FGeometry& AllottedGeometry, FArrangedCh
 		const EVisibility ChildVisibility = CurChild.GetWidget()->GetVisibility();
 		if ( ArrangedChildren.Accepts(ChildVisibility) )
 		{
-			const FMargin SlotPadding(LayoutPaddingWithFlow(GSlateFlowDirection, CurChild.SlotPadding.Get()));
+			const FMargin SlotPadding(LayoutPaddingWithFlow(GSlateFlowDirection, CurChild.GetPadding()));
 			AlignmentArrangeResult XResult = AlignChild<Orient_Horizontal>(GSlateFlowDirection, AllottedGeometry.GetLocalSize().X, CurChild, SlotPadding);
 			AlignmentArrangeResult YResult = AlignChild<Orient_Vertical>(AllottedGeometry.GetLocalSize().Y, CurChild, SlotPadding);
 
@@ -53,7 +53,7 @@ FVector2D SOverlay::ComputeDesiredSize( float ) const
 		const EVisibility ChildVisibilty = CurSlot.GetWidget()->GetVisibility();
 		if ( ChildVisibilty != EVisibility::Collapsed )
 		{
-			FVector2D ChildDesiredSize = CurSlot.GetWidget()->GetDesiredSize() + CurSlot.SlotPadding.Get().GetDesiredSize();
+			FVector2D ChildDesiredSize = CurSlot.GetWidget()->GetDesiredSize() + CurSlot.GetPadding().GetDesiredSize();
 			MaxSize.X = FMath::Max( MaxSize.X, ChildDesiredSize.X );
 			MaxSize.Y = FMath::Max( MaxSize.Y, ChildDesiredSize.Y );
 		}
diff --git a/Engine/Source/Runtime/SlateCore/Public/Layout/BasicLayoutWidgetSlot.h b/Engine/Source/Runtime/SlateCore/Public/Layout/BasicLayoutWidgetSlot.h
--- a/Engine/Source/Runtime/SlateCore/Public/Layout/BasicLayoutWidgetSlot.h
+++ b/Engine/Source/Runtime/SlateCore/Public/Layout/BasicLayoutWidgetSlot.h
@@ -8,6 +8,7 @@
 #include "Layout/Margin.h"
 #include "Misc/Optional.h"
 
+PRAGMA_DISABLE_DEPRECATION_WARNINGS
 /** Mixin to add the alignment functionality to a base slot. */
 template <typename MixedIntoType>
 class TAlignmentWidgetSlotMixin
@@ -36,12 +37,37 @@ public:
 		return *(static_cast<MixedIntoType*>(this));
 	}
 
+public:
+	void SetHorizontalAlignment(EHorizontalAlignment Alignment)
+	{
+		if (HAlignment != Alignment)
+		{
+			HAlignment = Alignment;
+			static_cast<MixedIntoType*>(this)->Invalidate(EInvalidateWidgetReason::Layout);
+		}
+	}
+
+	void SetVerticalAlignment(EVerticalAlignment Alignment)
+	{
+		if (VAlignment != Alignment)
+		{
+			VAlignment = Alignment;
+			static_cast<MixedIntoType*>(this)->Invalidate(EInvalidateWidgetReason::Layout);
+		}
+	}
+
+	EHorizontalAlignment GetHorizontalAlignment() const { return HAlignment; }
+	EVerticalAlignment GetVerticalAlignment() const { return VAlignment; }
+
 public:
 	/** Horizontal positioning of child within the allocated slot */
+	UE_DEPRECATED(5.0, "Direct access to HAlignment is now deprecated. Use the getter.")
 	EHorizontalAlignment HAlignment;
 	/** Vertical positioning of child within the allocated slot */
+	UE_DEPRECATED(5.0, "Direct access to VAlignment is now deprecated. Use the getter.")
 	EVerticalAlignment VAlignment;
 };
+PRAGMA_ENABLE_DEPRECATION_WARNINGS
 
 
 PRAGMA_DISABLE_DEPRECATION_WARNINGS
@@ -84,6 +110,7 @@ public:
 	void SetPadding(TAttribute<FMargin> InPadding)
 	{
 		SlotPadding = MoveTemp(InPadding);
+		static_cast<MixedIntoType*>(this)->Invalidate(EInvalidateWidgetReason::LayoutAndVolatility);
 	}
 
 	const FMargin& GetPadding() const { return SlotPadding.Get(); }
@@ -108,11 +135,25 @@ public:
 		, TAlignmentWidgetSlotMixin<SlotType>()
 	{}
 
+	TBasicLayoutWidgetSlot(FChildren& InOwner)
+		: TSlotBase<SlotType>(InOwner)
+		, TPaddingWidgetSlotMixin<SlotType>()
+		, TAlignmentWidgetSlotMixin<SlotType>()
+	{
+	}
+
 	TBasicLayoutWidgetSlot(const EHorizontalAlignment InHAlign, const EVerticalAlignment InVAlign)
 		: TSlotBase<SlotType>()
 		, TPaddingWidgetSlotMixin<SlotType>()
 		, TAlignmentWidgetSlotMixin<SlotType>(InHAlign, InVAlign)
 	{}
+
+	TBasicLayoutWidgetSlot(FChildren& InOwner, const EHorizontalAlignment InHAlign, const EVerticalAlignment InVAlign)
+		: TSlotBase<SlotType>(InOwner)
+		, TPaddingWidgetSlotMixin<SlotType>()
+		, TAlignmentWidgetSlotMixin<SlotType>(InHAlign, InVAlign)
+	{
+	}
 };
 
 
diff --git a/Engine/Source/Runtime/SlateCore/Public/Layout/Children.h b/Engine/Source/Runtime/SlateCore/Public/Layout/Children.h
--- a/Engine/Source/Runtime/SlateCore/Public/Layout/Children.h
+++ b/Engine/Source/Runtime/SlateCore/Public/Layout/Children.h
@@ -22,7 +22,7 @@ public:
 	FChildren(SWidget* InOwner)
 		: Owner(InOwner)
 	{
-		//check(InOwner);
+		check(InOwner);
 	}
 
 	FChildren(std::nullptr_t) = delete;
@@ -173,9 +173,8 @@ TSupportsOneChildMixin : public FChildren, public TSlotBase<MixedIntoType>
 public:
 	TSupportsOneChildMixin(SWidget* InOwner)
 		: FChildren(InOwner)
-		, TSlotBase<MixedIntoType>()
+		, TSlotBase<MixedIntoType>(static_cast<const FChildren&>(*this))
 	{
-		this->RawParentPtr = InOwner;
 	}
 
 	TSupportsOneChildMixin(std::nullptr_t) = delete;
@@ -296,9 +295,8 @@ class TSingleWidgetChildrenWithSlot : public FChildren, protected TSlotBase<Slot
 public:
 	TSingleWidgetChildrenWithSlot(SWidget* InOwner)
 		: FChildren(InOwner)
-		, TSlotBase<SlotType>()
+		, TSlotBase<SlotType>(static_cast<const FChildren&>(*this))
 	{
-		this->RawParentPtr = InOwner;
 	}
 
 	TSingleWidgetChildrenWithSlot(std::nullptr_t) = delete;
@@ -322,9 +320,11 @@ public:
 	TSlotBase<SlotType>& AsSlot() { return *this; }
 	const TSlotBase<SlotType>& AsSlot() const { return *this; }
 
+	using TSlotBase<SlotType>::GetOwnerWidget;
 	using TSlotBase<SlotType>::AttachWidget;
 	using TSlotBase<SlotType>::DetachWidget;
 	using TSlotBase<SlotType>::GetWidget;
+	using TSlotBase<SlotType>::Invalidate;
 	SlotType& operator[](const TSharedRef<SWidget>& InChildWidget)
 	{
 		this->AttachWidget(InChildWidget);
@@ -426,7 +426,7 @@ public:
 	{
 		int32 Index = Children.Add(TUniquePtr<SlotType>(Slot));
 		check(Slot);
-		Slot->AttachWidgetParent(&GetOwner());
+		Slot->SetOwner(*this);
 
 		return Index;
 	}
@@ -480,7 +480,7 @@ public:
 	{
 		check(Slot);
 		Children.Insert(TUniquePtr<SlotType>(Slot), Index);
-		Slot->AttachWidgetParent(&GetOwner());
+		Slot->SetOwner(*this);
 	}
 
 	void Move(int32 IndexToMove, int32 IndexToDestination)
diff --git a/Engine/Source/Runtime/SlateCore/Public/Layout/LayoutUtils.h b/Engine/Source/Runtime/SlateCore/Public/Layout/LayoutUtils.h
--- a/Engine/Source/Runtime/SlateCore/Public/Layout/LayoutUtils.h
+++ b/Engine/Source/Runtime/SlateCore/Public/Layout/LayoutUtils.h
@@ -41,16 +41,16 @@ namespace ArrangeUtils
 			{
 			default:
 			case EFlowDirection::LeftToRight:
-				return static_cast<int32>(InSlot.HAlignment);
+				return static_cast<int32>(InSlot.GetHorizontalAlignment());
 			case EFlowDirection::RightToLeft:
-				switch (InSlot.HAlignment)
+				switch (InSlot.GetHorizontalAlignment())
 				{
 				case HAlign_Left:
 					return static_cast<int32>(HAlign_Right);
 				case HAlign_Right:
 					return static_cast<int32>(HAlign_Left);
 				default:
-					return static_cast<int32>(InSlot.HAlignment);
+					return static_cast<int32>(InSlot.GetHorizontalAlignment());
 				}
 			}
 		}
@@ -63,7 +63,7 @@ namespace ArrangeUtils
 		static int32 AsInt(EFlowDirection InFlowDirection, const SlotType& InSlot )
 		{
 			// InFlowDirection has no effect in vertical orientations.
-			return static_cast<int32>(InSlot.VAlignment);
+			return static_cast<int32>(InSlot.GetVerticalAlignment());
 		}
 	};
 }
diff --git a/Engine/Source/Runtime/SlateCore/Public/SlotBase.h b/Engine/Source/Runtime/SlateCore/Public/SlotBase.h
--- a/Engine/Source/Runtime/SlateCore/Public/SlotBase.h
+++ b/Engine/Source/Runtime/SlateCore/Public/SlotBase.h
@@ -7,28 +7,37 @@
 #include "Widgets/InvalidateWidgetReason.h"
 
 class SWidget;
+class FChildren;
 
 /** Slot are a container of a SWidget used by the FChildren. */
 class SLATECORE_API FSlotBase
 {
 public:
-
 	FSlotBase();
-
-	FSlotBase( const TSharedRef<SWidget>& InWidget );
+	FSlotBase(const FChildren& Children);
+	FSlotBase(const TSharedRef<SWidget>& InWidget);
+	FSlotBase& operator=(const FSlotBase&) = delete;
+	FSlotBase(const FSlotBase&) = delete;
 
 	virtual ~FSlotBase();
 
-	FORCEINLINE_DEBUGGABLE void AttachWidgetParent(SWidget* InParent)
-	{
-		if (RawParentPtr != InParent)
-		{
-			ensureMsgf(RawParentPtr == nullptr, TEXT("Slots should not be reassigned to different parents."));
+	UE_DEPRECATED(5.0, "AttachWidgetParent is not used anymore. Use get SetOwner.")
+	void AttachWidgetParent(SWidget* InParent) { }
 
-			RawParentPtr = InParent;
-			AfterContentOrOwnerAssigned();
-		}
-	}
+	/**
+	 * Access the FChildren that own the slot.
+	 * The owner can be invalid when the slot is not attached.
+	 */
+	const FChildren* GetOwner() const { return Owner; }
+
+	/**
+	 * Access the widget that own the slot.
+	 * The owner can be invalid when the slot is not attached.
+	 */
+	SWidget* GetOwnerWidget() const;
+
+	/** Set the owner of the slot. */
+	void SetOwner(const FChildren& Children);
 
 	FORCEINLINE_DEBUGGABLE void AttachWidget( const TSharedRef<SWidget>& InWidget )
 	{
@@ -45,10 +54,7 @@ public:
 	 * There will always be a widget in the slot; sometimes it is
 	 * the SNullWidget instance.
 	 */
-	FORCEINLINE_DEBUGGABLE const TSharedRef<SWidget>& GetWidget() const
-	{
-		return Widget;
-	}
+	FORCEINLINE_DEBUGGABLE const TSharedRef<SWidget>& GetWidget() const { return Widget; }
 
 	/**
 	 * Remove the widget from its current slot.
@@ -57,9 +63,9 @@ public:
 	 */
 	const TSharedPtr<SWidget> DetachWidget();
 
-protected:
 	void Invalidate(EInvalidateWidgetReason InvalidateReason);
 
+protected:
 	/**
 	 * Performs the attribute assignment and invalidates the widget minimally based on what actually changed.  So if the boundness of the attribute didn't change
 	 * volatility won't need to be recalculated.  Returns true if the value changed.
@@ -86,22 +92,22 @@ protected:
 		return false;
 	}
 
-protected:
-	/** The parent and owner of the slot. */
-	SWidget* RawParentPtr;
-
 private:
 	void DetatchParentFromContent();
 	void AfterContentOrOwnerAssigned();
 
 private:
-	// non-copyable
-	FSlotBase& operator=(const FSlotBase&);
-	FSlotBase(const FSlotBase&);
-
-private:
+	/** The children that own the slot. */
+	const FChildren* Owner;
 	/** The content widget of the slot. */
 	TSharedRef<SWidget> Widget;
+
+#if WITH_EDITORONLY_DATA
+protected:
+	/** The parent and owner of the slot. */
+	UE_DEPRECATED(5.0, "RawParentPtr is not used anymore. Use get Owner Widget.")
+	SWidget* RawParentPtr;
+#endif
 };
 
 
@@ -110,12 +116,7 @@ template<typename SlotType>
 class TSlotBase : public FSlotBase
 {
 public:
-
-	TSlotBase() = default;
-
-	TSlotBase( const TSharedRef<SWidget>& InWidget )
-	: FSlotBase( InWidget )
-	{}
+	using FSlotBase::FSlotBase;
 
 	SlotType& operator[]( const TSharedRef<SWidget>& InChildWidget )
 	{
diff --git a/Engine/Source/Runtime/SlateCore/Public/Widgets/SBoxPanel.h b/Engine/Source/Runtime/SlateCore/Public/Widgets/SBoxPanel.h
--- a/Engine/Source/Runtime/SlateCore/Public/Widgets/SBoxPanel.h
+++ b/Engine/Source/Runtime/SlateCore/Public/Widgets/SBoxPanel.h
@@ -30,17 +30,11 @@ public:
 	 * A BoxPanel contains one BoxPanel child and describes how that
 	 * child should be arranged on the screen.
 	 */
-	class FSlot : public TSlotBase<FSlot>
+	class FSlot : public TBasicLayoutWidgetSlot<FSlot>
 	{
 	public:		
 		/** Horizontal and Vertical Boxes inherit from FSlot */
 		virtual ~FSlot(){}
-		
-		/** Horizontal positioning of child within the allocated slot */
-		TEnumAsByte<EHorizontalAlignment> HAlignment;
-		
-		/** Vertical positioning of child within the allocated slot */
-		TEnumAsByte<EVerticalAlignment> VAlignment;
 
 		/**
 		* How much space this slot should occupy along panel's direction.
@@ -50,9 +44,6 @@ public:
 		*   peers' SizeRule_Auto requirements have been satisfied.
 		*/
 		FSizeParam SizeParam;
-
-		/** The padding to add around the child. */
-		TAttribute<FMargin> SlotPadding;
 		
 		/** The max size that this slot can be (0 if no max) */
 		TAttribute<float> MaxSize;
@@ -60,11 +51,8 @@ public:
 	protected:
 		/** Default values for a slot. */
 		FSlot()
-			: TSlotBase<FSlot>()
-			, HAlignment( HAlign_Fill )
-			, VAlignment( VAlign_Fill )
+			: TBasicLayoutWidgetSlot<FSlot>(HAlign_Fill, VAlign_Fill)
 			, SizeParam( FStretch(1) )
-			, SlotPadding( FMargin(0) )
 			, MaxSize( 0.0f )
 		{ }
 	};
@@ -139,39 +127,40 @@ public:
 			SizeParam = FStretch( StretchCoefficient );
 			return *this;
 		}
+
 		FSlot& Padding( float Uniform )
 		{
-			SlotPadding = FMargin(Uniform);
+			SBoxPanel::FSlot::Padding(Uniform);
 			return *this;
 		}
 
 		FSlot& Padding( float Horizontal, float Vertical )
 		{
-			SlotPadding = FMargin(Horizontal, Vertical);
+			SBoxPanel::FSlot::Padding(Horizontal, Vertical);
 			return *this;
 		}
 
 		FSlot& Padding( float Left, float Top, float Right, float Bottom )
 		{
-			SlotPadding = FMargin(Left, Top, Right, Bottom);
+			SBoxPanel::FSlot::Padding(Left, Top, Right, Bottom);
 			return *this;
 		}
 		
 		FSlot& Padding( TAttribute<FMargin> InPadding )
 		{
-			SlotPadding = InPadding;
+			SBoxPanel::FSlot::Padding(MoveTemp(InPadding));
 			return *this;
 		}
 
-		FSlot& HAlign( EHorizontalAlignment InHAlignment )
+		FSlot& HAlign(EHorizontalAlignment InHAlignment)
 		{
-			HAlignment = InHAlignment;
+			SBoxPanel::FSlot::HAlign(InHAlignment);
 			return *this;
 		}
 
-		FSlot& VAlign( EVerticalAlignment InVAlignment )
+		FSlot& VAlign(EVerticalAlignment InVAlignment)
 		{
-			VAlignment = InVAlignment;
+			SBoxPanel::FSlot::VAlign(InVAlignment);
 			return *this;
 		}
 
@@ -277,45 +266,39 @@ public:
 			return *this;
 		}
 
-		FSlot& Padding( float Uniform )
+		FSlot& Padding(float Uniform)
 		{
-			SlotPadding = FMargin(Uniform);
+			SBoxPanel::FSlot::Padding(Uniform);
 			return *this;
 		}
 
-		FSlot& Padding( float Horizontal, float Vertical )
+		FSlot& Padding(float Horizontal, float Vertical)
 		{
-			SlotPadding = FMargin(Horizontal, Vertical);
+			SBoxPanel::FSlot::Padding(Horizontal, Vertical);
 			return *this;
 		}
 
-		FSlot& Padding( float Left, float Top, float Right, float Bottom )
+		FSlot& Padding(float Left, float Top, float Right, float Bottom)
 		{
-			SlotPadding = FMargin(Left, Top, Right, Bottom);
+			SBoxPanel::FSlot::Padding(Left, Top, Right, Bottom);
 			return *this;
 		}
 
-		FSlot& Padding( const TAttribute<FMargin>::FGetter& InDelegate )
+		FSlot& Padding(TAttribute<FMargin> InPadding)
 		{
-			SlotPadding.Bind( InDelegate );
+			SBoxPanel::FSlot::Padding(MoveTemp(InPadding));
 			return *this;
 		}
 
 		FSlot& HAlign( EHorizontalAlignment InHAlignment )
 		{
-			HAlignment = InHAlignment;
+			SBoxPanel::FSlot::HAlign(InHAlignment);
 			return *this;
 		}
 
 		FSlot& VAlign( EVerticalAlignment InVAlignment )
 		{
-			VAlignment = InVAlignment;
-			return *this;
-		}
-
-		FSlot& Padding( TAttribute<FMargin> InPadding )
-		{
-			SlotPadding = InPadding;
+			SBoxPanel::FSlot::VAlign(InVAlignment);
 			return *this;
 		}
 
diff --git a/Engine/Source/Runtime/SlateCore/Public/Widgets/SOverlay.h b/Engine/Source/Runtime/SlateCore/Public/Widgets/SOverlay.h
--- a/Engine/Source/Runtime/SlateCore/Public/Widgets/SOverlay.h
+++ b/Engine/Source/Runtime/SlateCore/Public/Widgets/SOverlay.h
@@ -37,62 +37,18 @@ class SLATECORE_API SOverlay
 public:	
 
 	/** A slot that support alignment of content and padding and z-order */
-	class SLATECORE_API FOverlaySlot : public TSlotBase<FOverlaySlot>
+	class SLATECORE_API FOverlaySlot : public TBasicLayoutWidgetSlot<FOverlaySlot>
 	{
 	public:
 		FOverlaySlot()
-			: TSlotBase<FOverlaySlot>()
+			: TBasicLayoutWidgetSlot<FOverlaySlot>(HAlign_Fill, VAlign_Fill)
 			, ZOrder(0)
-			, HAlignment(HAlign_Fill)
-			, VAlignment(VAlign_Fill)
-			, SlotPadding(0.0f)
 		{ }
 
-		FOverlaySlot& HAlign( EHorizontalAlignment InHAlignment )
-		{
-			HAlignment = InHAlignment;
-			return *this;
-		}
-
-		FOverlaySlot& VAlign( EVerticalAlignment InVAlignment )
-		{
-			VAlignment = InVAlignment;
-			return *this;
-		}
-
-		FOverlaySlot& Padding(float Uniform)
-		{
-			SlotPadding = FMargin(Uniform);
-			return *this;
-		}
-
-		FOverlaySlot& Padding(float Horizontal, float Vertical)
-		{
-			SlotPadding = FMargin(Horizontal, Vertical);
-			return *this;
-		}
-
-		FOverlaySlot& Padding(float Left, float Top, float Right, float Bottom)
-		{
-			SlotPadding = FMargin(Left, Top, Right, Bottom);
-			return *this;
-		}
-
-		FOverlaySlot& Padding( const TAttribute<FMargin> InPadding )
-		{
-			SlotPadding = InPadding;
-			return *this;
-		}
-
 		/** Slots with larger ZOrder values will draw above slots with smaller ZOrder values.  Slots
 		with the same ZOrder will simply draw in the order they were added.  Currently this only
 		works for overlay slots that are added dynamically with AddWidget() and RemoveWidget() */
 		int32 ZOrder;
-
-		TEnumAsByte<EHorizontalAlignment> HAlignment;
-		TEnumAsByte<EVerticalAlignment> VAlignment;
-
-		TAttribute< FMargin > SlotPadding;
 	};
 
 
diff --git a/Engine/Source/Runtime/UMG/Private/Components/DynamicEntryBoxBase.cpp b/Engine/Source/Runtime/UMG/Private/Components/DynamicEntryBoxBase.cpp
--- a/Engine/Source/Runtime/UMG/Private/Components/DynamicEntryBoxBase.cpp
+++ b/Engine/Source/Runtime/UMG/Private/Components/DynamicEntryBoxBase.cpp
@@ -162,7 +162,7 @@ void UDynamicEntryBoxBase::SetEntrySpacing(const FVector2D& InEntrySpacing)
 					}
 				}
 				SOverlay::FOverlaySlot& OverlaySlot = (*OverlayChildren)[ChildIdx];
-				OverlaySlot.SlotPadding = Padding;
+				OverlaySlot.SetPadding(Padding);
 			}
 		}
 		else if (EntryBoxType == EDynamicBoxType::Horizontal || EntryBoxType == EDynamicBoxType::Vertical)
@@ -179,7 +179,7 @@ void UDynamicEntryBoxBase::SetEntrySpacing(const FVector2D& InEntrySpacing)
 				Padding.Left = bIsHBox && !bIsFirstChild ? EntrySpacing.X : 0.f;
 
 				SBoxPanel::FSlot& BoxSlot = (*BoxChildren)[ChildIdx];
-				BoxSlot.SlotPadding = Padding;
+				BoxSlot.SetPadding(Padding);
 			}
 		}
 	}
@@ -401,9 +401,9 @@ void UDynamicEntryBoxBase::AddEntryChild(UUserWidget& ChildWidget)
 			VAlign = EntrySpacing.Y >= 0.f ? EVerticalAlignment::VAlign_Top : EVerticalAlignment::VAlign_Bottom;
 		}
 		
-		OverlaySlot.HAlignment = HAlign;
-		OverlaySlot.VAlignment = VAlign;
-		OverlaySlot.SlotPadding = BuildEntryPadding(TargetSpacing);
+		OverlaySlot.SetHorizontalAlignment(HAlign);
+		OverlaySlot.SetVerticalAlignment(VAlign);
+		OverlaySlot.SetPadding(BuildEntryPadding(TargetSpacing));
 
 		NewSlot = &OverlaySlot;
 	}
@@ -413,14 +413,14 @@ void UDynamicEntryBoxBase::AddEntryChild(UUserWidget& ChildWidget)
 		const bool bIsFirstChild = MyPanelWidget->GetChildren()->Num() == 0;
 
 		SBoxPanel::FSlot& BoxPanelSlot = bIsHBox ? (SBoxPanel::FSlot&)StaticCastSharedPtr<SHorizontalBox>(MyPanelWidget)->AddSlot().MaxWidth(MaxElementSize) : (SBoxPanel::FSlot&)StaticCastSharedPtr<SVerticalBox>(MyPanelWidget)->AddSlot().MaxHeight(MaxElementSize);
-		BoxPanelSlot.HAlignment = EntryHorizontalAlignment;
-		BoxPanelSlot.VAlignment = EntryVerticalAlignment;
+		BoxPanelSlot.SetHorizontalAlignment(EntryHorizontalAlignment);
+		BoxPanelSlot.SetVerticalAlignment(EntryVerticalAlignment);
 		BoxPanelSlot.SizeParam = UWidget::ConvertSerializedSizeParamToRuntime(EntrySizeRule);
 
 		FMargin Padding;
 		Padding.Top = bIsHBox || bIsFirstChild ? 0.f : EntrySpacing.Y;
 		Padding.Left = bIsHBox && !bIsFirstChild ? EntrySpacing.X : 0.f;
-		BoxPanelSlot.SlotPadding = Padding;
+		BoxPanelSlot.SetPadding(Padding);
 
 		NewSlot = &BoxPanelSlot;
 	}
diff --git a/Engine/Source/Runtime/UMG/Private/Components/GridSlot.cpp b/Engine/Source/Runtime/UMG/Private/Components/GridSlot.cpp
--- a/Engine/Source/Runtime/UMG/Private/Components/GridSlot.cpp
+++ b/Engine/Source/Runtime/UMG/Private/Components/GridSlot.cpp
@@ -106,7 +106,7 @@ void UGridSlot::SetHorizontalAlignment(EHorizontalAlignment InHorizontalAlignmen
 	HorizontalAlignment = InHorizontalAlignment;
 	if ( Slot )
 	{
-		Slot->HAlignment = InHorizontalAlignment;
+		Slot->SetHorizontalAlignment(InHorizontalAlignment);
 	}
 }
 
@@ -115,7 +115,7 @@ void UGridSlot::SetVerticalAlignment(EVerticalAlignment InVerticalAlignment)
 	VerticalAlignment = InVerticalAlignment;
 	if ( Slot )
 	{
-		Slot->VAlignment = InVerticalAlignment;
+		Slot->SetVerticalAlignment(InVerticalAlignment);
 	}
 }
 
diff --git a/Engine/Source/Runtime/UMG/Private/Components/UniformGridSlot.cpp b/Engine/Source/Runtime/UMG/Private/Components/UniformGridSlot.cpp
--- a/Engine/Source/Runtime/UMG/Private/Components/UniformGridSlot.cpp
+++ b/Engine/Source/Runtime/UMG/Private/Components/UniformGridSlot.cpp
@@ -54,7 +54,7 @@ void UUniformGridSlot::SetHorizontalAlignment(EHorizontalAlignment InHorizontalA
 	HorizontalAlignment = InHorizontalAlignment;
 	if ( Slot )
 	{
-		Slot->HAlignment = InHorizontalAlignment;
+		Slot->SetHorizontalAlignment(InHorizontalAlignment);
 	}
 }
 
@@ -63,7 +63,7 @@ void UUniformGridSlot::SetVerticalAlignment(EVerticalAlignment InVerticalAlignme
 	VerticalAlignment = InVerticalAlignment;
 	if ( Slot )
 	{
-		Slot->VAlignment = InVerticalAlignment;
+		Slot->SetVerticalAlignment(InVerticalAlignment);
 	}
 }
 
