Slate: Rename StateAttribute h files to inl files.
#rnx

diff --git a/Engine/Source/Runtime/SlateCore/Public/Types/Attributes/SlateAttributeBase.h b/Engine/Source/Runtime/SlateCore/Public/Types/Attributes/SlateAttributeBase.h
deleted file mode 100644
diff --git a/Engine/Source/Runtime/SlateCore/Public/Types/Attributes/SlateAttributeBase.inl b/Engine/Source/Runtime/SlateCore/Public/Types/Attributes/SlateAttributeBase.inl
new file mode 100644
--- /dev/null
+++ b/Engine/Source/Runtime/SlateCore/Public/Types/Attributes/SlateAttributeBase.inl
@@ -0,0 +1,612 @@
+// Copyright Epic Games, Inc. All Rights Reserved.
+
+#pragma once
+
+/** */
+namespace SlateAttributePrivate
+{
+	/**
+	 * Attribute object
+	 * InObjectType - Type of the value to store
+	 * InInvalidationReasonPredicate - Predicate that returns the type of invalidation to do when the value changes (e.g layout or paint)
+	 *								The invalidation can be overridden per widget. (This use memory allocation. See FSlateAttributeMetadata.)
+	 * InComparePredicateType - Predicate to compare the cached value with the Getter.
+	 * bInIsExternal - The attribute life is not controlled by the SWidget.
+	 */
+	template<typename InObjectType, typename InInvalidationReasonPredicate, typename InComparePredicateType, ESlateAttributeType InAttributeType>
+	struct TSlateAttributeBase : public FSlateAttributeImpl
+	{
+	public:
+		template<typename AttributeMemberType>
+		friend struct TSlateMemberAttributeRef;
+		friend FSlateAttributeContainer;
+
+		using ObjectType = InObjectType;
+		using FInvalidationReasonPredicate = InInvalidationReasonPredicate;
+		using FGetter = typename TAttribute<ObjectType>::FGetter;
+		using FComparePredicate = InComparePredicateType;
+
+		static EInvalidateWidgetReason GetInvalidationReason(const SWidget& Widget) { return FInvalidationReasonPredicate::GetInvalidationReason(Widget); }
+		static bool IdenticalTo(const SWidget& Widget, const ObjectType& Lhs, const ObjectType& Rhs) { return FComparePredicate::IdenticalTo(Widget, Lhs, Rhs); }
+
+	private:
+		void UpdateNowOnBind(SWidget& Widget)
+		{
+#if UE_SLATE_WITH_ATTRIBUTE_INITIALIZATION_ON_BIND
+			ProtectedUpdateNow(Widget, InAttributeType);
+#endif
+		}
+
+		void VerifyOwningWidget(const SWidget& Widget) const
+		{
+#if UE_SLATE_WITH_MEMBER_ATTRIBUTE_DEBUGGING
+			checkf(&Widget == Debug_OwningWidget || Debug_OwningWidget == nullptr,
+				TEXT("The Owning Widget is not the same as used at construction. This will cause bad memory access."));
+#endif
+		}
+
+		void VerifyNan() const
+		{
+#if UE_SLATE_WITH_ATTRIBUTE_NAN_DIAGNOSTIC
+			if constexpr (std::is_same<FVector, InObjectType>::value)
+			{
+				ensureMsgf(!Value.ContainsNaN(), TEXT("Value contains a NaN. Initialize your FVector properly (see FVector::EForceInit)"));
+			}
+			else if constexpr (std::is_same<FVector2D, InObjectType>::value)
+			{
+				ensureMsgf(!Value.ContainsNaN(), TEXT("Value contains a NaN. Initialize your FVector2D properly (see FVector2D::EForceInit)"));
+			}
+			else if constexpr (std::is_same<FLinearColor, InObjectType>::value)
+			{
+				ensureMsgf(!FVector4(Value).ContainsNaN(), TEXT("Value contains a NaN. Initialize your FLinearColor properly (see FLinearColor::EForceInit)"));
+			}
+			else if constexpr (std::is_same<FRotator, InObjectType>::value)
+			{
+				ensureMsgf(!Value.ContainsNaN(), TEXT("Value contains a NaN. Initialize your FRotator properly (see FRotator::EForceInit)"));
+			}
+#endif
+		}
+
+	public:
+#if UE_SLATE_WITH_MEMBER_ATTRIBUTE_DEBUGGING
+		TSlateAttributeBase()
+			: Value()
+			, Debug_OwningWidget(nullptr)
+		{
+			VerifyNan();
+		}
+#else
+		TSlateAttributeBase()
+			: Value()
+		{
+			VerifyNan();
+		}
+#endif
+
+		TSlateAttributeBase(const ObjectType& InValue)
+			: Value(InValue)
+#if UE_SLATE_WITH_MEMBER_ATTRIBUTE_DEBUGGING
+			, Debug_OwningWidget(nullptr)
+#endif
+		{
+			VerifyNan();
+		}
+
+		TSlateAttributeBase(ObjectType&& InValue)
+			: Value(MoveTemp(InValue))
+#if UE_SLATE_WITH_MEMBER_ATTRIBUTE_DEBUGGING
+			, Debug_OwningWidget(nullptr)
+#endif
+		{
+			VerifyNan();
+		}
+
+		TSlateAttributeBase(SWidget& Widget)
+			: Value()
+#if UE_SLATE_WITH_MEMBER_ATTRIBUTE_DEBUGGING
+			, Debug_OwningWidget(&Widget)
+#endif
+		{
+			VerifyNan();
+		}
+
+		TSlateAttributeBase(SWidget& Widget, const ObjectType& InValue)
+			: Value(InValue)
+#if UE_SLATE_WITH_MEMBER_ATTRIBUTE_DEBUGGING
+			, Debug_OwningWidget(&Widget)
+#endif
+		{
+			VerifyNan();
+		}
+
+		TSlateAttributeBase(SWidget& Widget, ObjectType&& InValue)
+			: Value(MoveTemp(InValue))
+#if UE_SLATE_WITH_MEMBER_ATTRIBUTE_DEBUGGING
+			, Debug_OwningWidget(&Widget)
+#endif
+		{
+			VerifyNan();
+		}
+
+		TSlateAttributeBase(SWidget& Widget, const FGetter& Getter, const ObjectType& InitialValue)
+			: Value(InitialValue)
+#if UE_SLATE_WITH_MEMBER_ATTRIBUTE_DEBUGGING
+			, Debug_OwningWidget(&Widget)
+#endif
+		{
+			VerifyNan();
+			if (Getter.IsBound())
+			{
+				ConstructWrapper(Widget, Getter);
+			}
+		}
+
+		TSlateAttributeBase(SWidget& Widget, const FGetter& Getter, ObjectType&& InitialValue)
+			: Value(MoveTemp(InitialValue))
+#if UE_SLATE_WITH_MEMBER_ATTRIBUTE_DEBUGGING
+			, Debug_OwningWidget(&Widget)
+#endif
+		{
+			VerifyNan();
+			if (Getter.IsBound())
+			{
+				ConstructWrapper(Widget, Getter);
+			}
+		}
+
+		TSlateAttributeBase(SWidget& Widget, FGetter&& Getter, const ObjectType& InitialValue)
+			: Value(InitialValue)
+#if UE_SLATE_WITH_MEMBER_ATTRIBUTE_DEBUGGING
+			, Debug_OwningWidget(&Widget)
+#endif
+		{
+			VerifyNan();
+			if (Getter.IsBound())
+			{
+				ConstructWrapper(Widget, MoveTemp(Getter));
+			}
+		}
+
+		TSlateAttributeBase(SWidget& Widget, FGetter&& Getter, ObjectType&& InitialValue)
+			: Value(MoveTemp(InitialValue))
+#if UE_SLATE_WITH_MEMBER_ATTRIBUTE_DEBUGGING
+			, Debug_OwningWidget(&Widget)
+#endif
+		{
+			VerifyNan();
+			if (Getter.IsBound())
+			{
+				ConstructWrapper(Widget, MoveTemp(Getter));
+			}
+		}
+
+		TSlateAttributeBase(SWidget& Widget, const TAttribute<ObjectType>& Attribute, const ObjectType& InitialValue)
+			: Value((Attribute.IsSet() && !Attribute.IsBound()) ? Attribute.Get() : InitialValue)
+#if UE_SLATE_WITH_MEMBER_ATTRIBUTE_DEBUGGING
+			, Debug_OwningWidget(&Widget)
+#endif
+		{
+			VerifyNan();
+			if (Attribute.IsBound())
+			{
+				ConstructWrapper(Widget, Attribute.GetBinding());
+			}
+		}
+
+		TSlateAttributeBase(SWidget& Widget, TAttribute<ObjectType>&& Attribute, ObjectType&& InitialValue)
+			: Value((Attribute.IsSet() && !Attribute.IsBound()) ? MoveTemp(Attribute.Steal().template Get<ObjectType>()) : MoveTemp(InitialValue))
+#if UE_SLATE_WITH_MEMBER_ATTRIBUTE_DEBUGGING
+			, Debug_OwningWidget(&Widget)
+#endif
+		{
+			VerifyNan();
+			if (Attribute.IsBound())
+			{
+				ConstructWrapper(Widget, MoveTemp(Attribute.Steal().template Get<FGetter>()));
+			}
+		}
+
+	public:
+		/** @return the SlateAttribute cached value. If the SlateAttribute is bound, the value will be cached at the end of the every frame. */
+		UE_NODISCARD const ObjectType& Get() const
+		{
+			return Value;
+		}
+
+		/** Update the cached value and invalidate the widget if needed. */
+		void UpdateNow(SWidget& Widget)
+		{
+			VerifyOwningWidget(Widget);
+			ProtectedUpdateNow(Widget, InAttributeType);
+		}
+
+	public:
+		/**
+		 * Unbind the SlateAttribute and set its value. It may invalidate the Widget if the value is different.
+		 * @return true if the value is considered different and an invalidation occurred.
+		 */
+		bool Set(SWidget& Widget, const ObjectType& NewValue)
+		{
+			VerifyOwningWidget(Widget);
+			ProtectedUnregisterAttribute(Widget, InAttributeType);
+
+			const bool bIsIdentical = IdenticalTo(Widget, Value, NewValue);
+			if (!bIsIdentical)
+			{
+				Value = NewValue;
+				ProtectedInvalidateWidget(Widget, InAttributeType, GetInvalidationReason(Widget));
+			}
+			return !bIsIdentical;
+		}
+
+		/**
+		 * Unbind the SlateAttribute and set its value. It may invalidate the Widget if the value is different.
+		 * @return true if the value is considered different and an invalidation occurred.
+		 */
+		bool Set(SWidget& Widget, ObjectType&& NewValue)
+		{
+			VerifyOwningWidget(Widget);
+			ProtectedUnregisterAttribute(Widget, InAttributeType);
+
+			const bool bIsIdentical = IdenticalTo(Widget, Value, NewValue);
+			if (!bIsIdentical)
+			{
+				Value = MoveTemp(NewValue);
+				ProtectedInvalidateWidget(Widget, InAttributeType, GetInvalidationReason(Widget));
+			}
+			return !bIsIdentical;
+		}
+
+	public:
+		/**
+		 * Bind the SlateAttribute to the Getter function.
+		 * (If enabled) Update the value from the Getter. If the value is different, then invalidate the widget.
+		 * The SlateAttribute will now be updated every frame from the Getter.
+		 */
+		void Bind(SWidget& Widget, const FGetter& Getter)
+		{
+			VerifyOwningWidget(Widget);
+			if (Getter.IsBound())
+			{
+				AssignBinding(Widget, Getter);
+			}
+			else
+			{
+				ProtectedUnregisterAttribute(Widget, InAttributeType);
+			}
+		}
+
+		/**
+		 * Bind the SlateAttribute to the Getter function.
+		 * (If enabled) Update the value from the Getter. If the value is different, then invalidate the widget.
+		 * The SlateAttribute will now be updated every frame from the Getter.
+		 */
+		void Bind(SWidget& Widget, FGetter&& Getter)
+		{
+			VerifyOwningWidget(Widget);
+			if (Getter.IsBound())
+			{
+				AssignBinding(Widget, MoveTemp(Getter));
+			}
+			else
+			{
+				ProtectedUnregisterAttribute(Widget, InAttributeType);
+			}
+		}
+
+		/**
+		 * Bind the SlateAttribute to the newly create getter function.
+		 * (If enabled) Update the value from the getter. If the value is different, then invalidate the widget.
+		 * The SlateAttribute will now be updated every frame from the Getter.
+		 */
+		template<typename WidgetType, typename U = typename std::enable_if<std::is_base_of<SWidget, WidgetType>::value>::type>
+		void Bind(WidgetType& Widget, typename FGetter::template TConstMethodPtr<WidgetType> MethodPtr)
+		{
+			Bind(Widget, FGetter::CreateSP(&Widget, MethodPtr));
+		}
+
+		/**
+		 * Bind the SlateAttribute to the Attribute Getter function (if it exist).
+		 * (If enabled) Update the value from the getter. If the value is different, then invalidate the widget.
+		 * The SlateAttribute will now be updated every frame from the Getter.
+		 * Or
+		 * Set the SlateAttribute's value if the Attribute is not bound but is set.
+		 * This will Unbind any previously bound getter.
+		 * If the value is different, then invalidate the widget.
+		 * Or
+		 * Unbind the SlateAttribute if the Attribute is not bound and not set.
+		 * @see Set
+		 * @return true if the getter was assigned or if the value is considered different and an invalidation occurred.
+		 */
+		bool Assign(SWidget& Widget, const TAttribute<ObjectType>& OtherAttribute)
+		{
+			VerifyOwningWidget(Widget);
+			if (OtherAttribute.IsBound())
+			{
+				return AssignBinding(Widget, OtherAttribute.GetBinding());
+			}
+			else if (OtherAttribute.IsSet())
+			{
+				return Set(Widget, OtherAttribute.Get());
+			}
+			else
+			{
+				ProtectedUnregisterAttribute(Widget, InAttributeType);
+				return false;
+			}
+		}
+
+		bool Assign(SWidget& Widget, TAttribute<ObjectType>&& OtherAttribute)
+		{
+			VerifyOwningWidget(Widget);
+			if (OtherAttribute.IsBound())
+			{
+				return AssignBinding(Widget, MoveTemp(OtherAttribute.Steal().template Get<FGetter>()));
+			}
+			else if (OtherAttribute.IsSet())
+			{
+				return Set(Widget, MoveTemp(OtherAttribute.Steal().template Get<ObjectType>()));
+			}
+			else
+			{
+				ProtectedUnregisterAttribute(Widget, InAttributeType);
+				return false;
+			}
+		}
+
+		bool Assign(SWidget& Widget, const TAttribute<ObjectType>& OtherAttribute, const ObjectType& DefaultValue)
+		{
+			VerifyOwningWidget(Widget);
+			if (OtherAttribute.IsBound())
+			{
+				return AssignBinding(Widget, OtherAttribute.GetBinding());
+			}
+			else if (OtherAttribute.IsSet())
+			{
+				return Set(Widget, OtherAttribute.Get());
+			}
+			else
+			{
+				return Set(Widget, DefaultValue);
+			}
+		}
+
+		bool Assign(SWidget& Widget, TAttribute<ObjectType>&& OtherAttribute, const ObjectType& DefaultValue)
+		{
+			VerifyOwningWidget(Widget);
+			if (OtherAttribute.IsBound())
+			{
+				return AssignBinding(Widget, MoveTemp(OtherAttribute.Steal().template Get<FGetter>()));
+			}
+			else if (OtherAttribute.IsSet())
+			{
+				return Set(Widget, MoveTemp(OtherAttribute.Steal().template Get<ObjectType>()));
+			}
+			else
+			{
+				return Set(Widget, DefaultValue);
+			}
+		}
+
+		bool Assign(SWidget& Widget, const TAttribute<ObjectType>& OtherAttribute, ObjectType&& DefaultValue)
+		{
+			VerifyOwningWidget(Widget);
+			if (OtherAttribute.IsBound())
+			{
+				return AssignBinding(Widget, OtherAttribute.GetBinding());
+			}
+			else if (OtherAttribute.IsSet())
+			{
+				return Set(Widget, OtherAttribute.Get());
+			}
+			else
+			{
+				return Set(Widget, MoveTemp(DefaultValue));
+			}
+		}
+
+		bool Assign(SWidget& Widget, TAttribute<ObjectType>&& OtherAttribute, ObjectType&& DefaultValue)
+		{
+			VerifyOwningWidget(Widget);
+			if (OtherAttribute.IsBound())
+			{
+				return AssignBinding(Widget, MoveTemp(OtherAttribute.Steal().template Get<FGetter>()));
+			}
+			else if (OtherAttribute.IsSet())
+			{
+				return Set(Widget, MoveTemp(OtherAttribute.Steal().template Get<ObjectType>()));
+			}
+			else
+			{
+				return Set(Widget, MoveTemp(DefaultValue));
+			}
+		}
+
+		/**
+		 * Remove the Getter function.
+		 * The Slate Attribute will not be updated anymore and will keep its current cached value.
+		 */
+		void Unbind(SWidget& Widget)
+		{
+			VerifyOwningWidget(Widget);
+			ProtectedUnregisterAttribute(Widget, InAttributeType);
+		}
+
+	public:
+		/** Build a Attribute from this SlateAttribute. */
+		UE_NODISCARD TAttribute<ObjectType> ToAttribute(const SWidget& Widget) const
+		{
+			if (ISlateAttributeGetter* Delegate = ProtectedFindGetter(Widget, InAttributeType))
+			{
+				return TAttribute<ObjectType>::Create(static_cast<FSlateAttributeGetterWrapper<TSlateAttributeBase>*>(Delegate)->GetDelegate());
+			}
+			return TAttribute<ObjectType>(Get());
+		}
+
+		/** @return True if the SlateAttribute is bound to a getter function. */
+		UE_NODISCARD bool IsBound(const SWidget& Widget) const
+		{
+			VerifyOwningWidget(Widget);
+			return ProtectedIsBound(Widget, InAttributeType);
+		}
+
+		/** @return True if they have the same Getter or the same value. */
+		UE_NODISCARD bool IdenticalTo(const SWidget& Widget, const TSlateAttributeBase& Other) const
+		{
+			VerifyOwningWidget(Widget);
+			FDelegateHandle ThisDelegateHandle = ProtectedFindGetterHandle(Widget, InAttributeType);
+			FDelegateHandle OthersDelegateHandle = ProtectedFindGetterHandle(Widget, InAttributeType);
+			if (ThisDelegateHandle == OthersDelegateHandle)
+			{
+				if (ThisDelegateHandle.IsValid())
+				{
+					return true;
+				}
+				return IdenticalTo(Widget, Get(), Other.Get());
+			}
+			return false;
+		}
+
+		/** @return True if they have the same Getter or, if the Attribute is set, the same value. */
+		UE_NODISCARD bool IdenticalTo(const SWidget& Widget, const TAttribute<ObjectType>& Other) const
+		{
+			VerifyOwningWidget(Widget);
+			FDelegateHandle ThisDelegateHandle = ProtectedFindGetterHandle(Widget, InAttributeType);
+			if (Other.IsBound())
+			{
+				return Other.GetBinding().GetHandle() == ThisDelegateHandle;
+			}
+			return !ThisDelegateHandle.IsValid() && Other.IsSet() && IdenticalTo(Widget, Get(), Other.Get());
+		}
+
+	private:
+		void ConstructWrapper(SWidget& Widget, const FGetter& Getter)
+		{
+			TUniquePtr<ISlateAttributeGetter> Wrapper = MakeUniqueGetter(*this, Getter);
+			ProtectedRegisterAttribute(Widget, InAttributeType, MoveTemp(Wrapper));
+			UpdateNowOnBind(Widget);
+		}
+
+		void ConstructWrapper(SWidget& Widget, FGetter&& Getter)
+		{
+			TUniquePtr<ISlateAttributeGetter> Wrapper = MakeUniqueGetter(*this, MoveTemp(Getter));
+			ProtectedRegisterAttribute(Widget, InAttributeType, MoveTemp(Wrapper));
+			UpdateNowOnBind(Widget);
+		}
+
+		bool AssignBinding(SWidget& Widget, const FGetter& Getter)
+		{
+			const FDelegateHandle PreviousGetterHandle = ProtectedFindGetterHandle(Widget, InAttributeType);
+			if (PreviousGetterHandle != Getter.GetHandle())
+			{
+				ConstructWrapper(Widget, Getter);
+				return true;
+			}
+			return false;
+		}
+
+		bool AssignBinding(SWidget& Widget, FGetter&& Getter)
+		{
+			const FDelegateHandle PreviousGetterHandle = ProtectedFindGetterHandle(Widget, InAttributeType);
+			if (PreviousGetterHandle != Getter.GetHandle())
+			{
+				ConstructWrapper(Widget, MoveTemp(Getter));
+				return true;
+			}
+			return false;
+		}
+
+	private:
+		template<typename SlateAttributeType>
+		class FSlateAttributeGetterWrapper : public ISlateAttributeGetter
+		{
+		public:
+			using ObjectType = typename SlateAttributeType::ObjectType;
+			using FGetter = typename TAttribute<ObjectType>::FGetter;
+			using FComparePredicate = typename SlateAttributeType::FComparePredicate;
+
+			FSlateAttributeGetterWrapper() = delete;
+			FSlateAttributeGetterWrapper(const FSlateAttributeGetterWrapper&) = delete;
+			FSlateAttributeGetterWrapper& operator= (const FSlateAttributeGetterWrapper&) = delete;
+			virtual ~FSlateAttributeGetterWrapper() = default;
+
+		public:
+			FSlateAttributeGetterWrapper(SlateAttributeType& InOwningAttribute, const FGetter& InGetterDelegate)
+				: Getter(InGetterDelegate)
+				, Attribute(&InOwningAttribute)
+			{
+			}
+			FSlateAttributeGetterWrapper(SlateAttributeType& InOwningAttribute, FGetter&& InGetterDelegate)
+				: Getter(MoveTemp(InGetterDelegate))
+				, Attribute(&InOwningAttribute)
+			{
+			}
+
+			virtual FUpdateAttributeResult UpdateAttribute(const SWidget& Widget) override
+			{
+				ObjectType NewValue = Getter.Execute();
+
+				const bool bIsIdentical = Attribute->IdenticalTo(Widget, Attribute->Value, NewValue);
+				if (!bIsIdentical)
+				{
+					// Set the value on the widget
+					Attribute->Value = MoveTemp(NewValue);
+					return Attribute->GetInvalidationReason(Widget);
+				}
+				return FUpdateAttributeResult();
+			}
+
+			virtual const FSlateAttributeBase& GetAttribute() const override
+			{
+				return *Attribute;
+			}
+
+			virtual void SetAttribute(FSlateAttributeBase& InBase) override
+			{
+				Attribute = &static_cast<SlateAttributeType&>(InBase);
+			}
+
+			virtual FDelegateHandle GetDelegateHandle() const override
+			{
+				return Getter.GetHandle();
+			}
+
+			const FGetter& GetDelegate() const
+			{
+				return Getter;
+			}
+
+		private:
+			/** Getter function to fetch the new value of the SlateAttribute. */
+			FGetter Getter;
+			/** The SlateAttribute of the SWidget owning the value. */
+			SlateAttributeType* Attribute;
+		};
+
+	private:
+		static TUniquePtr<ISlateAttributeGetter> MakeUniqueGetter(TSlateAttributeBase& Attribute, const FGetter& Getter)
+		{
+			return MakeUnique<FSlateAttributeGetterWrapper<TSlateAttributeBase>>(Attribute, Getter);
+		}
+
+		static TUniquePtr<ISlateAttributeGetter> MakeUniqueGetter(TSlateAttributeBase& Attribute, FGetter&& Getter)
+		{
+			return MakeUnique<FSlateAttributeGetterWrapper<TSlateAttributeBase>>(Attribute, MoveTemp(Getter));
+		}
+
+	private:
+		ObjectType Value;
+
+	protected:
+#if UE_SLATE_WITH_MEMBER_ATTRIBUTE_DEBUGGING
+		void Debug_SetOwningWidget(SWidget* InWidget)
+		{
+			Debug_OwningWidget = InWidget;
+		}
+
+		SWidget* Debug_OwningWidget;
+#endif
+	};
+
+
+} // SlateAttributePrivate
diff --git a/Engine/Source/Runtime/SlateCore/Public/Types/Attributes/SlateAttributeDefinition.h b/Engine/Source/Runtime/SlateCore/Public/Types/Attributes/SlateAttributeDefinition.h
deleted file mode 100644
diff --git a/Engine/Source/Runtime/SlateCore/Public/Types/Attributes/SlateAttributeDefinition.inl b/Engine/Source/Runtime/SlateCore/Public/Types/Attributes/SlateAttributeDefinition.inl
new file mode 100644
--- /dev/null
+++ b/Engine/Source/Runtime/SlateCore/Public/Types/Attributes/SlateAttributeDefinition.inl
@@ -0,0 +1,89 @@
+// Copyright Epic Games, Inc. All Rights Reserved.
+
+#pragma once
+
+#include "CoreMinimal.h"
+#include "Misc/Attribute.h"
+#include "Templates/EqualTo.h"
+#include "Widgets/InvalidateWidgetReason.h"
+
+#include <type_traits>
+
+
+class SWidget;
+
+/** */
+namespace SlateAttributePrivate
+{
+	/** Predicate used to identify if the InvalidationWidgetReason is defined in the attribute descriptor. */
+	struct FSlateAttributeNoInvalidationReason
+	{
+		static constexpr EInvalidateWidgetReason GetInvalidationReason(const SWidget&) { return EInvalidateWidgetReason::None; }
+	};
+
+
+	/** */
+	enum class ESlateAttributeType : uint8
+	{
+		Member		= 0,	// Member of a SWidget (are not allowed to move).
+		Managed		= 1,	// External to the SWidget, global variable or member that can moved.
+		Unused0		= 2,
+		Unused1		= 3,
+		// We use the attribute type in a bit field in FSlateAttributeMetaData. Only 4 value are allowed.
+	};
+
+
+	class ISlateAttributeGetter;
+	template<typename ObjectType, typename InvalidationReasonPredicate, typename FComparePredicate, ESlateAttributeType AttributeType>
+	struct TSlateAttributeBase;
+	struct FSlateAttributeContainer;
+	template<typename AttributeMemberType>
+	struct TSlateMemberAttributeRef;
+
+
+	/** */
+	class ISlateAttributeGetter
+	{
+	public:
+		struct FUpdateAttributeResult
+		{
+			FUpdateAttributeResult(EInvalidateWidgetReason InInvalidationReason)
+				: InvalidationReason(InInvalidationReason)
+				, bInvalidationRequested(true)
+			{}
+			FUpdateAttributeResult()
+				: InvalidationReason(EInvalidateWidgetReason::None)
+				, bInvalidationRequested(false)
+			{}
+			EInvalidateWidgetReason InvalidationReason;
+			bool bInvalidationRequested;
+		};
+
+
+		virtual FUpdateAttributeResult UpdateAttribute(const SWidget& Widget) = 0;
+		virtual const FSlateAttributeBase& GetAttribute() const = 0;
+		virtual void SetAttribute(FSlateAttributeBase&) = 0;
+		virtual FDelegateHandle GetDelegateHandle() const = 0;
+		virtual ~ISlateAttributeGetter() = default;
+	};
+
+
+	/** */
+	struct SLATECORE_API FSlateAttributeImpl : public FSlateAttributeBase
+	{
+	protected:
+		bool ProtectedIsWidgetInDestructionPath(SWidget* Widget) const;
+		bool ProtectedIsImplemented(const SWidget& Widget) const;
+		void ProtectedUnregisterAttribute(SWidget& Widget, ESlateAttributeType AttributeType) const;
+		void ProtectedRegisterAttribute(SWidget& Widget, ESlateAttributeType AttributeType, TUniquePtr<ISlateAttributeGetter>&& Wrapper);
+		void ProtectedInvalidateWidget(SWidget& Widget, ESlateAttributeType AttributeType, EInvalidateWidgetReason InvalidationReason) const;
+		bool ProtectedIsBound(const SWidget& Widget, ESlateAttributeType AttributeType) const;
+		ISlateAttributeGetter* ProtectedFindGetter(const SWidget& Widget, ESlateAttributeType AttributeType) const;
+		FDelegateHandle ProtectedFindGetterHandle(const SWidget& Widget, ESlateAttributeType AttributeType) const;
+		bool ProtectedIsIdenticalTo(const SWidget& Widget, ESlateAttributeType AttributeType, const FSlateAttributeBase& Other, const bool bHasSameValue) const;
+		bool ProtectedIsIdenticalToAttribute(const SWidget& Widget, ESlateAttributeType AttributeType, const void* Other, const bool bHasSameValue) const;
+		void ProtectedUpdateNow(SWidget& Widget, ESlateAttributeType AttributeType);
+		void ProtectedMoveAttribute(SWidget& Widget, ESlateAttributeType AttributeType, const FSlateAttributeBase* Other);
+	};
+
+} // SlateAttributePrivate
diff --git a/Engine/Source/Runtime/SlateCore/Public/Types/Attributes/SlateAttributeManaged.h b/Engine/Source/Runtime/SlateCore/Public/Types/Attributes/SlateAttributeManaged.h
deleted file mode 100644
diff --git a/Engine/Source/Runtime/SlateCore/Public/Types/Attributes/SlateAttributeManaged.inl b/Engine/Source/Runtime/SlateCore/Public/Types/Attributes/SlateAttributeManaged.inl
new file mode 100644
--- /dev/null
+++ b/Engine/Source/Runtime/SlateCore/Public/Types/Attributes/SlateAttributeManaged.inl
@@ -0,0 +1,247 @@
+// Copyright Epic Games, Inc. All Rights Reserved.
+
+#pragma once
+
+/** */
+namespace SlateAttributePrivate
+{
+	/**
+	 *
+	 */
+	template<typename InObjectType, typename InInvalidationReasonPredicate, typename InComparePredicate>
+	struct TSlateManagedAttribute : protected TSlateAttributeBase<InObjectType, InInvalidationReasonPredicate, InComparePredicate, ESlateAttributeType::Managed>
+	{
+	private:
+		using Super = TSlateAttributeBase<InObjectType, InInvalidationReasonPredicate, InComparePredicate, ESlateAttributeType::Managed>;
+
+	public:
+		using ObjectType = typename Super::ObjectType;
+		using FInvalidationReasonPredicate = typename Super::FInvalidationReasonPredicate;
+		using FGetter = typename Super::FGetter;
+		using FComparePredicate = typename Super::FComparePredicate;
+		static const bool IsMemberType = false;
+
+		static EInvalidateWidgetReason GetInvalidationReason(const SWidget& Widget) { return Super::GetInvalidationReason(Widget); }
+		static bool IdenticalTo(const SWidget& Widget, const ObjectType& Lhs, const ObjectType& Rhs) { return Super::IdenticalTo(Widget, Lhs, Rhs); }
+
+	public:
+		TSlateManagedAttribute() = delete;
+		TSlateManagedAttribute(const TSlateManagedAttribute&) = delete;
+		TSlateManagedAttribute& operator=(const TSlateManagedAttribute&) = delete;
+
+		~TSlateManagedAttribute()
+		{
+			Unbind();
+		}
+
+		TSlateManagedAttribute(TSlateManagedAttribute&& Other)
+			: Super(MoveTemp(Other))
+			, ManagedWidget(MoveTemp(Other.ManagedWidget))
+		{
+			if (TSharedPtr<SWidget> Pin = ManagedWidget.Pin())
+			{
+				FSlateAttributeImpl::ProtectedMoveAttribute(*Pin.Get(), ESlateAttributeType::Managed, &Other);
+			}
+		}
+
+		TSlateManagedAttribute& operator=(TSlateManagedAttribute&& Other)
+		{
+			Super::operator=(MoveTemp(Other));
+			ManagedWidget = MoveTemp(Other.ManagedWidget);
+
+			if (TSharedPtr<SWidget> Pin = ManagedWidget.Pin())
+			{
+				FSlateAttributeImpl::ProtectedMoveAttribute(*Pin.Get(), ESlateAttributeType::Managed, &Other);
+			}
+			return *this;
+		}
+
+		explicit TSlateManagedAttribute(TSharedRef<SWidget> Widget)
+			: Super(Widget.Get())
+			, ManagedWidget(Widget)
+		{ }
+
+		explicit TSlateManagedAttribute(TSharedRef<SWidget> Widget, const ObjectType& InValue)
+			: Super(Widget.Get(), InValue)
+			, ManagedWidget(Widget)
+		{
+		}
+
+		explicit TSlateManagedAttribute(TSharedRef<SWidget> Widget, ObjectType&& InValue)
+			: Super(Widget.Get(), MoveTemp(InValue))
+			, ManagedWidget(Widget)
+		{
+		}
+
+		explicit TSlateManagedAttribute(TSharedRef<SWidget> Widget, const FGetter& Getter, const ObjectType& InitialValue)
+			: Super(Widget.Get(), Getter, InitialValue)
+			, ManagedWidget(Widget)
+		{
+		}
+
+		explicit TSlateManagedAttribute(TSharedRef<SWidget> Widget, const FGetter& Getter, ObjectType&& InitialValue)
+			: Super(Widget.Get(), Getter, MoveTemp(InitialValue))
+			, ManagedWidget(Widget)
+		{
+		}
+
+		explicit TSlateManagedAttribute(TSharedRef<SWidget> Widget, FGetter&& Getter, const ObjectType& InitialValue)
+			: Super(Widget.Get(), MoveTemp(Getter), InitialValue)
+			, ManagedWidget(Widget)
+		{
+		}
+
+		explicit TSlateManagedAttribute(TSharedRef<SWidget> Widget, FGetter&& Getter, ObjectType&& InitialValue)
+			: Super(Widget.Get(), MoveTemp(Getter), MoveTemp(InitialValue))
+			, ManagedWidget(Widget)
+		{
+		}
+
+		explicit TSlateManagedAttribute(TSharedRef<SWidget> Widget, const TAttribute<ObjectType>& Attribute, const ObjectType& InitialValue)
+			: Super(Widget.Get(), Attribute, InitialValue)
+			, ManagedWidget(Widget)
+		{
+		}
+
+		explicit TSlateManagedAttribute(TSharedRef<SWidget> Widget, TAttribute<ObjectType>&& Attribute, ObjectType&& InitialValue)
+			: Super(Widget.Get(), MoveTemp(Attribute), MoveTemp(InitialValue))
+			, ManagedWidget(Widget)
+		{
+		}
+
+	public:
+		const ObjectType& Get() const
+		{
+			return Super::Get();
+		}
+
+		void UpdateNow()
+		{
+			if (TSharedPtr<SWidget> Pin = ManagedWidget.Pin())
+			{
+				Super::UpdateNow(*Pin.Get());
+			}
+		}
+
+	public:
+		void Set(const ObjectType& NewValue)
+		{
+			if (TSharedPtr<SWidget> Pin = ManagedWidget.Pin())
+			{
+				Super::Set(*Pin.Get(), NewValue);
+			}
+		}
+
+		void Set(ObjectType&& NewValue)
+		{
+			if (TSharedPtr<SWidget> Pin = ManagedWidget.Pin())
+			{
+				Super::Set(*Pin.Get(), MoveTemp(NewValue));
+			}
+		}
+
+		void Bind(const FGetter& Getter)
+		{
+			if (TSharedPtr<SWidget> Pin = ManagedWidget.Pin())
+			{
+				Super::Bind(*Pin.Get(), Getter);
+			}
+		}
+
+		void Bind(FGetter&& Getter)
+		{
+			if (TSharedPtr<SWidget> Pin = ManagedWidget.Pin())
+			{
+				Super::Bind(*Pin.Get(), MoveTemp(Getter));
+			}
+		}
+
+		void Assign(const TAttribute<ObjectType>& OtherAttribute)
+		{
+			if (TSharedPtr<SWidget> Pin = ManagedWidget.Pin())
+			{
+				Super::Assign(*Pin.Get(), OtherAttribute);
+			}
+		}
+
+		void Assign(TAttribute<ObjectType>&& OtherAttribute)
+		{
+			if (TSharedPtr<SWidget> Pin = ManagedWidget.Pin())
+			{
+				Super::Assign(*Pin.Get(), MoveTemp(OtherAttribute));
+			}
+		}
+
+		void Assign(const TAttribute<ObjectType>& OtherAttribute, const ObjectType& DefaultValue)
+		{
+			if (TSharedPtr<SWidget> Pin = ManagedWidget.Pin())
+			{
+				Super::Assign(*Pin.Get(), OtherAttribute, DefaultValue);
+			}
+		}
+
+		void Assign(const TAttribute<ObjectType>& OtherAttribute, ObjectType&& DefaultValue)
+		{
+			if (TSharedPtr<SWidget> Pin = ManagedWidget.Pin())
+			{
+				Super::Assign(*Pin.Get(), OtherAttribute, MoveTemp(DefaultValue));
+			}
+		}
+
+		void Assign(TAttribute<ObjectType>&& OtherAttribute, const ObjectType& DefaultValue)
+		{
+			if (TSharedPtr<SWidget> Pin = ManagedWidget.Pin())
+			{
+				Super::Assign(*Pin.Get(), MoveTemp(OtherAttribute), DefaultValue);
+			}
+		}
+
+		void Assign(TAttribute<ObjectType>&& OtherAttribute, ObjectType&& DefaultValue)
+		{
+			if (TSharedPtr<SWidget> Pin = ManagedWidget.Pin())
+			{
+				Super::Assign(*Pin.Get(), MoveTemp(OtherAttribute), MoveTemp(DefaultValue));
+			}
+		}
+
+		void Unbind()
+		{
+			if (TSharedPtr<SWidget> Pin = ManagedWidget.Pin())
+			{
+				Super::Unbind(*Pin.Get());
+			}
+		}
+
+	public:
+		bool IsBound() const
+		{
+			if (TSharedPtr<SWidget> Pin = ManagedWidget.Pin())
+			{
+				return Super::IsBound(*Pin.Get());
+			}
+			return false;
+		}
+
+		bool IdenticalTo(const TSlateManagedAttribute& Other) const
+		{
+			if (TSharedPtr<SWidget> Pin = ManagedWidget.Pin())
+			{
+				return Super::IdenticalTo(*Pin.Get(), Other);
+			}
+			return false;
+		}
+
+		bool IdenticalTo(const TAttribute<ObjectType>& Other) const
+		{
+			if (TSharedPtr<SWidget> Pin = ManagedWidget.Pin())
+			{
+				return Super::IdenticalTo(*Pin.Get(), Other);
+			}
+			return false;
+		}
+
+	private:
+		TWeakPtr<SWidget> ManagedWidget;
+	};
+
+} // SlateAttributePrivate
diff --git a/Engine/Source/Runtime/SlateCore/Public/Types/Attributes/SlateAttributeMember.h b/Engine/Source/Runtime/SlateCore/Public/Types/Attributes/SlateAttributeMember.h
deleted file mode 100644
diff --git a/Engine/Source/Runtime/SlateCore/Public/Types/Attributes/SlateAttributeMember.inl b/Engine/Source/Runtime/SlateCore/Public/Types/Attributes/SlateAttributeMember.inl
new file mode 100644
--- /dev/null
+++ b/Engine/Source/Runtime/SlateCore/Public/Types/Attributes/SlateAttributeMember.inl
@@ -0,0 +1,235 @@
+// Copyright Epic Games, Inc. All Rights Reserved.
+
+#pragma once
+
+/** */
+namespace SlateAttributePrivate
+{
+	/**
+	 *
+	 */
+	template<typename InObjectType, typename InInvalidationReasonPredicate, typename InComparePredicate>
+	struct TSlateMemberAttribute : public TSlateAttributeBase<InObjectType, InInvalidationReasonPredicate, InComparePredicate, ESlateAttributeType::Member>
+	{
+	private:
+		using Super = TSlateAttributeBase<InObjectType, InInvalidationReasonPredicate, InComparePredicate, ESlateAttributeType::Member>;
+
+		template<typename WidgetType, typename U = typename std::enable_if<std::is_base_of<SWidget, WidgetType>::value>::type>
+		static void VerifyAttributeAddress(const WidgetType& Widget, const TSlateMemberAttribute& Self)
+		{
+			checkf((UPTRINT)&Self >= (UPTRINT)&Widget && (UPTRINT)&Self < (UPTRINT)&Widget + sizeof(WidgetType),
+				TEXT("Use TAttribute or TSlateManagedAttribute instead. See SlateAttribute.h for more info."));
+			ensureAlwaysMsgf((HasDefinedInvalidationReason || Self.ProtectedIsImplemented(Widget)),
+				TEXT("The TSlateAttribute could not be found in the SlateAttributeDescriptor.\n")
+				TEXT("Use the SLATE_DECLARE_WIDGET and add the attribute in PrivateRegisterAttributes,\n")
+				TEXT("Or use TSlateAttribute with a valid Invalidation Reason instead."));
+		}
+
+	public:
+		using FGetter = typename Super::FGetter;
+		using ObjectType = typename Super::ObjectType;
+		static const bool IsMemberType = true;
+		static constexpr bool HasDefinedInvalidationReason = !std::is_same<InInvalidationReasonPredicate, FSlateAttributeNoInvalidationReason>::value;
+
+	public:
+		//~ You can only register Attribute that are defined in a SWidget (member of the class).
+		//~ Use the constructor with Widget pointer.
+		//~ See documentation in SlateAttribute.h for more info.
+		TSlateMemberAttribute() = delete;
+		TSlateMemberAttribute(const TSlateMemberAttribute&) = delete;
+		TSlateMemberAttribute(TSlateMemberAttribute&&) = delete;
+		TSlateMemberAttribute& operator=(const TSlateMemberAttribute&) = delete;
+		TSlateMemberAttribute& operator=(TSlateMemberAttribute&&) = delete;
+		void* operator new(size_t) = delete;
+		void* operator new[](size_t) = delete;
+
+#if UE_SLATE_WITH_MEMBER_ATTRIBUTE_DEBUGGING
+		~TSlateMemberAttribute()
+		{
+			/**
+			 * The parent should now be destroyed and the shared pointer should be invalidate.
+			 * If you hit the check, that means the TSlateAttribute is not a variable member of SWidget.
+			 * It will introduced bad memory access.
+			 * See documentation in SlateAttribute.h.
+			*/
+			checkf(Super::ProtectedIsWidgetInDestructionPath(Super::Debug_OwningWidget), TEXT("The Owning widget should be invalid."));
+		}
+#endif
+
+		template<typename WidgetType, typename U = typename std::enable_if<std::is_base_of<SWidget, WidgetType>::value>::type>
+		explicit TSlateMemberAttribute(WidgetType& Widget)
+			: Super(Widget)
+		{
+			VerifyAttributeAddress(Widget, *this);
+		}
+
+		template<typename WidgetType, typename U = typename std::enable_if<std::is_base_of<SWidget, WidgetType>::value>::type>
+		explicit TSlateMemberAttribute(WidgetType& Widget, const ObjectType& InValue)
+			: Super(Widget, InValue)
+		{
+			VerifyAttributeAddress(Widget, *this);
+		}
+
+		template<typename WidgetType, typename U = typename std::enable_if<std::is_base_of<SWidget, WidgetType>::value>::type>
+		explicit TSlateMemberAttribute(WidgetType& Widget, ObjectType&& InValue)
+			: Super(Widget, MoveTemp(InValue))
+		{
+			VerifyAttributeAddress(Widget, *this);
+		}
+	};
+
+
+
+	enum ESlateMemberAttributeRefNoCheckParam { SlateMemberAttributeRefNoCheckParam };
+
+	/*
+	 * A reference to a SlateAttribute that can be returned and saved for later.
+	 */
+	template<typename AttributeMemberType>
+	struct TSlateMemberAttributeRef
+	{
+	public:
+		using SlateAttributeType = AttributeMemberType;
+		using ObjectType = typename AttributeMemberType::ObjectType;
+		using AttributeType = TAttribute<ObjectType>;
+
+	private:
+		template<typename WidgetType>
+		static void VerifyAttributeAddress(const WidgetType& InWidget, const AttributeMemberType& InAttribute)
+		{
+			checkf((UPTRINT)&InAttribute >= (UPTRINT)&InWidget && (UPTRINT)&InAttribute < (UPTRINT)&InWidget + sizeof(WidgetType),
+				TEXT("The attribute is not a member of the widget."));
+			InAttribute.VerifyOwningWidget(InWidget);
+		}
+
+	public:
+		/** Constructor */
+		template<typename WidgetType, typename V = typename std::enable_if<std::is_base_of<SWidget, WidgetType>::value>::type>
+		explicit TSlateMemberAttributeRef(const TSharedRef<WidgetType>& InOwner, const AttributeMemberType& InAttribute)
+			: Owner(InOwner)
+			, Attribute(InAttribute)
+		{
+			VerifyAttributeAddress(InOwner.Get(), InAttribute);
+		}
+
+		/** Constructor */
+		explicit TSlateMemberAttributeRef(ESlateMemberAttributeRefNoCheckParam, const TSharedRef<SWidget>& InOwner, const AttributeMemberType& InAttribute)
+			: Owner(InOwner)
+			, Attribute(InAttribute)
+		{
+		}
+
+	public:
+		/** @return if the reference is valid. A reference can be invalid if the SWidget is destroyed. */
+		UE_NODISCARD bool IsValid() const
+		{
+			return Owner.IsValid();
+		}
+
+		/** @return the SlateAttribute cached value; undefined when IsValid() returns false. */
+		UE_NODISCARD const ObjectType& Get() const
+		{
+			if (TSharedPtr<const SWidget> Pin = Owner.Pin())
+			{
+				return Attribute.Get();
+			}
+			checkf(false, TEXT("It is an error to call GetValue() on an unset TSlateMemberAttributeRef. Please either check IsValid() or use Get(DefaultValue) instead."));
+			static ObjectType Tmp;
+			return Tmp;
+		}
+
+		/** @return the SlateAttribute cached value or the DefaultValue if the reference is invalid. */
+		UE_NODISCARD const ObjectType& Get(const ObjectType& DefaultValue) const
+		{
+			if (TSharedPtr<const SWidget> Pin = Owner.Pin())
+			{
+				return Attribute.Get();
+			}
+			return DefaultValue;
+		}
+
+		/** Update the cached value and invalidate the widget if needed. */
+		void UpdateValue()
+		{
+			if (TSharedPtr<const SWidget> Pin = Owner.Pin())
+			{
+				const_cast<AttributeMemberType>(Attribute).UpdateNow(*Pin.Get());
+			}
+		}
+
+		/**
+		 * Assumes the reference is valid.
+		 * Shorthand for the boilerplace code: MyAttribute.UpdateValueNow(); MyAttribute.Get();
+		 */
+		UE_NODISCARD const ObjectType& UpdateAndGet()
+		{
+			if (TSharedPtr<const SWidget> Pin = Owner.Pin())
+			{
+				const_cast<AttributeMemberType>(Attribute)->UpdateNow(*Pin.Get());
+				return Attribute.Get();
+			}
+			checkf(false, TEXT("It is an error to call GetValue() on an unset TSlateMemberAttributeRef. Please either check IsValid() or use Get(DefaultValue) instead."));
+		}
+
+		/**
+		 * Shorthand for the boilerplace code: MyAttribute.UpdateValueNow(); MyAttribute.Get(DefaultValue);
+		 * @return the SlateAttribute cached value or the DefaultValue if the reference is invalid.
+		 */
+		UE_NODISCARD const ObjectType& UpdateAndGet(const ObjectType& DefaultValue)
+		{
+			if (TSharedPtr<const SWidget> Pin = Owner.Pin())
+			{
+				const_cast<AttributeMemberType>(Attribute)->UpdateNow(*Pin.Get());
+				return Attribute.Get();
+			}
+			return DefaultValue;
+		}
+
+		/** Build a Attribute from this SlateAttribute. */
+		UE_NODISCARD TAttribute<ObjectType> ToAttribute() const
+		{
+			if (TSharedPtr<const SWidget> Pin = Owner.Pin())
+			{
+				return Attribute.ToAttribute(*Pin.Get());
+			}
+			return TAttribute<ObjectType>();
+		}
+
+		/** @return True if the SlateAttribute is bound to a getter function. */
+		UE_NODISCARD bool IsBound() const
+		{
+			if (TSharedPtr<const SWidget> Pin = Owner.Pin())
+			{
+				return Attribute.IsBound(*Pin.Get());
+			}
+			return false;
+		}
+
+		/** @return True if they have the same Getter or the same value. */
+		UE_NODISCARD bool IdenticalTo(const TSlateMemberAttributeRef& Other) const
+		{
+			TSharedPtr<const SWidget> SelfPin = Owner.Pin();
+			TSharedPtr<const SWidget> OtherPin = Other.Owner.Pin();
+			if (SelfPin == OtherPin && SelfPin)
+			{
+				return Attribute.IdenticalTo(*SelfPin.Get(), Other.Attribute);
+			}
+			return SelfPin == OtherPin;
+		}
+
+		/** @return True if they have the same Getter or, if the Attribute is set, the same value. */
+		UE_NODISCARD bool IdenticalTo(const TAttribute<ObjectType>& Other) const
+		{
+			if (TSharedPtr<const SWidget> Pin = Owner.Pin())
+			{
+				return Attribute.IdenticalTo(*Pin.Get(), Other);
+			}
+			return !Other.IsSet(); // if the other is not set, then both are invalid.
+		}
+
+	private:
+		TWeakPtr<const SWidget> Owner;
+		const AttributeMemberType& Attribute;
+	 };
+
+} // SlateAttributePrivate
diff --git a/Engine/Source/Runtime/SlateCore/Public/Types/SlateAttribute.h b/Engine/Source/Runtime/SlateCore/Public/Types/SlateAttribute.h
--- a/Engine/Source/Runtime/SlateCore/Public/Types/SlateAttribute.h
+++ b/Engine/Source/Runtime/SlateCore/Public/Types/SlateAttribute.h
@@ -239,7 +239,7 @@ struct FSlateAttributeBase
 };
 
 
-#include "Types/Attributes/SlateAttributeDefinition.h"
-#include "Types/Attributes/SlateAttributeBase.h"
-#include "Types/Attributes/SlateAttributeManaged.h"
-#include "Types/Attributes/SlateAttributeMember.h"
+#include "Types/Attributes/SlateAttributeDefinition.inl"
+#include "Types/Attributes/SlateAttributeBase.inl"
+#include "Types/Attributes/SlateAttributeManaged.inl"
+#include "Types/Attributes/SlateAttributeMember.inl"
