Temp Disable UMG CoreStyle due to incorrect file paths.

diff --git a/Engine/Source/Runtime/SlateCore/Private/Styling/UMGCoreStyle.cpp b/Engine/Source/Runtime/SlateCore/Private/Styling/UMGCoreStyle.cpp
--- a/Engine/Source/Runtime/SlateCore/Private/Styling/UMGCoreStyle.cpp
+++ b/Engine/Source/Runtime/SlateCore/Private/Styling/UMGCoreStyle.cpp
@@ -196,6 +196,9 @@ TSharedRef<ISlateStyle> FUMGCoreStyle::Create()
 		return Style;
 	}
 
+	// UE-116081
+#if 0
+
 	// These are the Slate colors which reference the dynamic colors in FSlateCoreStyle; these are the colors to put into the style
 	const FSlateColor DefaultForeground(FUMGColor::Black);
 	const FSlateColor InvertedForeground(FUMGColor::ForegroundInverted);
@@ -205,11 +208,11 @@ TSharedRef<ISlateStyle> FUMGCoreStyle::Create()
 	const FSlateColor SelectionColor_Pressed(Style->SelectionColor_Pressed_LinearRef);
 	const FSlateColor HighlightColor(FUMGColor::Highlight);
 
-	Style->Set("AppCoreStyleConstants::Icon", new IMAGE_BRUSH("CoreStyleConstants::Icons/DefaultAppCoreStyleConstants::Icon", CoreStyleConstants::Icon24x24));
-	Style->Set("AppCoreStyleConstants::Icon.Small", new IMAGE_BRUSH("CoreStyleConstants::Icons/DefaultAppCoreStyleConstants::Icon", CoreStyleConstants::Icon24x24));
+	Style->Set("AppIcon", new IMAGE_BRUSH("Icons/DefaultAppIcon", CoreStyleConstants::Icon24x24));
+	Style->Set("AppIcon.Small", new IMAGE_BRUSH("Icons/DefaultAppIcon", CoreStyleConstants::Icon24x24));
 
-	Style->Set("AppCoreStyleConstants::IconPadding", FMargin(4, 4, 0, 0));
-	Style->Set("AppCoreStyleConstants::IconPadding.Small", FMargin(4, 4, 0, 0));
+	Style->Set("AppIconPadding", FMargin(4, 4, 0, 0));
+	Style->Set("AppIconPadding.Small", FMargin(4, 4, 0, 0));
 
 	Style->Set("NormalFont", DEFAULT_FONT("Regular", RegularTextSize));
 
@@ -320,27 +323,27 @@ TSharedRef<ISlateStyle> FUMGCoreStyle::Create()
 		Style->Set("ErrorReporting.ForegroundColor", FLinearColor::White);
 	}
 
-	// Cursor CoreStyleConstants::Icons
+	// Cursor Icons
 	{
-		Style->Set("SoftwareCursor_Grab", new IMAGE_BRUSH( "CoreStyleConstants::Icons/cursor_grab", CoreStyleConstants::Icon16x16));
-		Style->Set("SoftwareCursor_CardinalCross", new IMAGE_BRUSH( "CoreStyleConstants::Icons/cursor_cardinal_cross", CoreStyleConstants::Icon24x24));
+		Style->Set("SoftwareCursor_Grab", new IMAGE_BRUSH( "Icons/cursor_grab", CoreStyleConstants::Icon16x16));
+		Style->Set("SoftwareCursor_CardinalCross", new IMAGE_BRUSH( "Icons/cursor_cardinal_cross", CoreStyleConstants::Icon24x24));
 	}
 
-	// Common CoreStyleConstants::Icons
+	// Common Icons
 	{
-		Style->Set("TrashCan", new IMAGE_BRUSH( "CoreStyleConstants::Icons/TrashCan", FVector2D(64, 64)));
-		Style->Set("TrashCan_Small", new IMAGE_BRUSH( "CoreStyleConstants::Icons/TrashCan_Small", FVector2D(18, 18)));
+		Style->Set("TrashCan", new IMAGE_BRUSH( "Icons/TrashCan", FVector2D(64, 64)));
+		Style->Set("TrashCan_Small", new IMAGE_BRUSH( "Icons/TrashCan_Small", FVector2D(18, 18)));
 	}
 
-	// Common CoreStyleConstants::Icons
+	// Common Icons
 	{
-		Style->Set( "CoreStyleConstants::Icons.Cross", new IMAGE_BRUSH( "CoreStyleConstants::Icons/Cross_12x", CoreStyleConstants::Icon12x12 ) );
-		Style->Set( "CoreStyleConstants::Icons.Denied", new IMAGE_BRUSH( "CoreStyleConstants::Icons/denied_16x", CoreStyleConstants::Icon16x16 ) );
-		Style->Set( "CoreStyleConstants::Icons.Error", new IMAGE_BRUSH( "CoreStyleConstants::Icons/CoreStyleConstants::Icon_error_16x", CoreStyleConstants::Icon16x16) );
-		Style->Set( "CoreStyleConstants::Icons.Help", new IMAGE_BRUSH( "CoreStyleConstants::Icons/CoreStyleConstants::Icon_help_16x", CoreStyleConstants::Icon16x16) );
-		Style->Set( "CoreStyleConstants::Icons.Info", new IMAGE_BRUSH( "CoreStyleConstants::Icons/CoreStyleConstants::Icon_info_16x", CoreStyleConstants::Icon16x16) );
-		Style->Set( "CoreStyleConstants::Icons.Warning", new IMAGE_BRUSH( "CoreStyleConstants::Icons/CoreStyleConstants::Icon_warning_16x", CoreStyleConstants::Icon16x16) );
-		Style->Set( "CoreStyleConstants::Icons.Download", new IMAGE_BRUSH( "CoreStyleConstants::Icons/CoreStyleConstants::Icon_Downloads_16x", CoreStyleConstants::Icon16x16) );
+		Style->Set( "Icons.Cross", new IMAGE_BRUSH( "Icons/Cross_12x", CoreStyleConstants::Icon12x12 ) );
+		Style->Set( "Icons.Denied", new IMAGE_BRUSH( "Icons/denied_16x", CoreStyleConstants::Icon16x16 ) );
+		Style->Set( "Icons.Error", new IMAGE_BRUSH( "Icons/Icon_error_16x", CoreStyleConstants::Icon16x16) );
+		Style->Set( "Icons.Help", new IMAGE_BRUSH( "Icons/Icon_help_16x", CoreStyleConstants::Icon16x16) );
+		Style->Set( "Icons.Info", new IMAGE_BRUSH( "Icons/Icon_info_16x", CoreStyleConstants::Icon16x16) );
+		Style->Set( "Icons.Warning", new IMAGE_BRUSH( "Icons/Icon_warning_16x", CoreStyleConstants::Icon16x16) );
+		Style->Set( "Icons.Download", new IMAGE_BRUSH( "Icons/Icon_Downloads_16x", CoreStyleConstants::Icon16x16) );
 	}
 
 	// Tool panels
@@ -354,20 +357,20 @@ TSharedRef<ISlateStyle> FUMGCoreStyle::Create()
 		Style->Set( "PopupText.Background", new BOX_BRUSH( "Old/Menu_Background", FMargin(8.0f/64.0f) ) );
 	}
 
-	// Generic command CoreStyleConstants::Icons
+	// Generic command Icons
 	{
-		Style->Set( "GenericCommands.Undo", new IMAGE_BRUSH( "CoreStyleConstants::Icons/CoreStyleConstants::Icon_undo_16px", CoreStyleConstants::Icon16x16 ) );
-		Style->Set( "GenericCommands.Redo", new IMAGE_BRUSH( "CoreStyleConstants::Icons/CoreStyleConstants::Icon_redo_16px", CoreStyleConstants::Icon16x16 ) );
+		Style->Set( "GenericCommands.Undo", new IMAGE_BRUSH( "Icons/Icon_undo_16px", CoreStyleConstants::Icon16x16 ) );
+		Style->Set( "GenericCommands.Redo", new IMAGE_BRUSH( "Icons/Icon_redo_16px", CoreStyleConstants::Icon16x16 ) );
 
-		Style->Set( "GenericCommands.Copy", new IMAGE_BRUSH("CoreStyleConstants::Icons/Edit/CoreStyleConstants::Icon_Edit_Copy_16x", CoreStyleConstants::Icon16x16) );
-		Style->Set( "GenericCommands.Cut", new IMAGE_BRUSH("CoreStyleConstants::Icons/Edit/CoreStyleConstants::Icon_Edit_Cut_16x", CoreStyleConstants::Icon16x16) );
-		Style->Set( "GenericCommands.Delete", new IMAGE_BRUSH("CoreStyleConstants::Icons/Edit/CoreStyleConstants::Icon_Edit_Delete_16x", CoreStyleConstants::Icon16x16) );
-		Style->Set( "GenericCommands.Paste", new IMAGE_BRUSH("CoreStyleConstants::Icons/Edit/CoreStyleConstants::Icon_Edit_Paste_16x", CoreStyleConstants::Icon16x16) );
-		Style->Set( "GenericCommands.Duplicate", new IMAGE_BRUSH("CoreStyleConstants::Icons/Edit/CoreStyleConstants::Icon_Edit_Duplicate_16x", CoreStyleConstants::Icon16x16) );
-		Style->Set( "GenericCommands.Rename", new IMAGE_BRUSH( "CoreStyleConstants::Icons/Edit/CoreStyleConstants::Icon_Edit_Rename_16x", CoreStyleConstants::Icon16x16 ) );
+		Style->Set( "GenericCommands.Copy", new IMAGE_BRUSH("Icons/Edit/Icon_Edit_Copy_16x", CoreStyleConstants::Icon16x16) );
+		Style->Set( "GenericCommands.Cut", new IMAGE_BRUSH("Icons/Edit/Icon_Edit_Cut_16x", CoreStyleConstants::Icon16x16) );
+		Style->Set( "GenericCommands.Delete", new IMAGE_BRUSH("Icons/Edit/Icon_Edit_Delete_16x", CoreStyleConstants::Icon16x16) );
+		Style->Set( "GenericCommands.Paste", new IMAGE_BRUSH("Icons/Edit/Icon_Edit_Paste_16x", CoreStyleConstants::Icon16x16) );
+		Style->Set( "GenericCommands.Duplicate", new IMAGE_BRUSH("Icons/Edit/Icon_Edit_Duplicate_16x", CoreStyleConstants::Icon16x16) );
+		Style->Set( "GenericCommands.Rename", new IMAGE_BRUSH( "Icons/Edit/Icon_Edit_Rename_16x", CoreStyleConstants::Icon16x16 ) );
 	}
 
-	// SVerticalBox Drag& Drop CoreStyleConstants::Icon
+	// SVerticalBox Drag& Drop Icon
 	Style->Set("VerticalBoxDragIndicator", new IMAGE_BRUSH("Common/VerticalBoxDragIndicator", FVector2D(6, 45)));
 	Style->Set("VerticalBoxDragIndicatorShort", new IMAGE_BRUSH("Common/VerticalBoxDragIndicatorShort", FVector2D(6, 15)));
 	// SScrollBar defaults...
@@ -428,9 +431,9 @@ TSharedRef<ISlateStyle> FUMGCoreStyle::Create()
 
 	// SEditableComboBox defaults...
 	{
-		Style->Set( "EditableComboBox.Add", new IMAGE_BRUSH( "CoreStyleConstants::Icons/PlusSymbol_12x", CoreStyleConstants::Icon12x12 ) );
-		Style->Set( "EditableComboBox.Delete", new IMAGE_BRUSH( "CoreStyleConstants::Icons/Cross_12x", CoreStyleConstants::Icon12x12 ) );
-		Style->Set( "EditableComboBox.Rename", new IMAGE_BRUSH( "CoreStyleConstants::Icons/ellipsis_12x", CoreStyleConstants::Icon12x12 ) );
+		Style->Set( "EditableComboBox.Add", new IMAGE_BRUSH( "Icons/PlusSymbol_12x", CoreStyleConstants::Icon12x12 ) );
+		Style->Set( "EditableComboBox.Delete", new IMAGE_BRUSH( "Icons/Cross_12x", CoreStyleConstants::Icon12x12 ) );
+		Style->Set( "EditableComboBox.Rename", new IMAGE_BRUSH( "Icons/ellipsis_12x", CoreStyleConstants::Icon12x12 ) );
 		Style->Set( "EditableComboBox.Accept", new IMAGE_BRUSH( "Common/Check", CoreStyleConstants::Icon16x16 ) );
 	}
 
@@ -742,7 +745,7 @@ TSharedRef<ISlateStyle> FUMGCoreStyle::Create()
 	{
 		Style->Set( "ColorPicker.Border", new BOX_BRUSH( "Common/GroupBorder", FMargin(4.0f/16.0f) ) );
 		Style->Set( "ColorPicker.AlphaBackground", new IMAGE_BRUSH( "Common/Checker", CoreStyleConstants::Icon16x16, FLinearColor::White, ESlateBrushTileType::Both ) );
-		Style->Set( "ColorPicker.EyeDropper", new IMAGE_BRUSH( "CoreStyleConstants::Icons/eyedropper_16px", CoreStyleConstants::Icon16x16) );
+		Style->Set( "ColorPicker.EyeDropper", new IMAGE_BRUSH( "Icons/eyedropper_16px", CoreStyleConstants::Icon16x16) );
 		Style->Set( "ColorPicker.Font", DEFAULT_FONT( "Regular", 10 ) );
 		Style->Set( "ColorPicker.Mode", new IMAGE_BRUSH( "Common/ColorPicker_Mode_16x", CoreStyleConstants::Icon16x16) );
 		Style->Set( "ColorPicker.Separator", new IMAGE_BRUSH( "Common/ColorPicker_Separator", FVector2D(2.0f, 2.0f) ) );
@@ -858,8 +861,8 @@ TSharedRef<ISlateStyle> FUMGCoreStyle::Create()
 
 	// MultiBox
 	{
-		Style->Set( "MultiBox.GenericToolBarCoreStyleConstants::Icon", new IMAGE_BRUSH( "CoreStyleConstants::Icons/CoreStyleConstants::Icon_generic_toolbar", CoreStyleConstants::Icon40x40 ) );
-		Style->Set( "MultiBox.GenericToolBarCoreStyleConstants::Icon.Small", new IMAGE_BRUSH( "CoreStyleConstants::Icons/CoreStyleConstants::Icon_generic_toolbar", CoreStyleConstants::Icon20x20 ) );
+		Style->Set( "MultiBox.GenericToolBarIcon", new IMAGE_BRUSH( "Icons/Icon_generic_toolbar", CoreStyleConstants::Icon40x40 ) );
+		Style->Set( "MultiBox.GenericToolBarIcon.Small", new IMAGE_BRUSH( "Icons/Icon_generic_toolbar", CoreStyleConstants::Icon20x20 ) );
 
 		Style->Set( "MultiBox.DeleteButton", FButtonStyle() 
 			.SetNormal ( IMAGE_BRUSH( "/Docking/CloseApp_Hovered", CoreStyleConstants::Icon16x16 ) )
@@ -873,8 +876,8 @@ TSharedRef<ISlateStyle> FUMGCoreStyle::Create()
 	// ToolBar
 	{
 		Style->Set( "ToolBar.Background", new BOX_BRUSH( "Common/GroupBorder", FMargin(4.0f/16.0f) ) );
-		Style->Set( "ToolBar.CoreStyleConstants::Icon", new IMAGE_BRUSH( "CoreStyleConstants::Icons/CoreStyleConstants::Icon_tab_toolbar_16px", CoreStyleConstants::Icon16x16 ) );
-		Style->Set( "ToolBar.Expand", new IMAGE_BRUSH( "CoreStyleConstants::Icons/toolbar_expand_16x", CoreStyleConstants::Icon16x16) );
+		Style->Set( "ToolBar.Icon", new IMAGE_BRUSH( "Icons/Icon_tab_toolbar_16px", CoreStyleConstants::Icon16x16 ) );
+		Style->Set( "ToolBar.Expand", new IMAGE_BRUSH( "Icons/toolbar_expand_16x", CoreStyleConstants::Icon16x16) );
 		Style->Set( "ToolBar.SubMenuIndicator", new IMAGE_BRUSH( "Common/SubmenuArrow", CoreStyleConstants::Icon8x8 ) );
 		Style->Set( "ToolBar.SToolBarComboButtonBlock.Padding", FMargin(4.0f,0.0f));
 		Style->Set( "ToolBar.SToolBarButtonBlock.Padding", FMargin(4.0f,0.0f));
@@ -913,13 +916,13 @@ TSharedRef<ISlateStyle> FUMGCoreStyle::Create()
 		/* Read-only checkbox that appears next to a menu item */
 		/* Set images for various SCheckBox states associated with read-only toolbar check box items... */
 		const FCheckBoxStyle BasicToolBarCheckStyle = FCheckBoxStyle()
-			.SetUncheckedImage(IMAGE_BRUSH("CoreStyleConstants::Icons/Empty_14x", CoreStyleConstants::Icon14x14))
-			.SetUncheckedHoveredImage(IMAGE_BRUSH("CoreStyleConstants::Icons/Empty_14x", CoreStyleConstants::Icon14x14))
+			.SetUncheckedImage(IMAGE_BRUSH("Icons/Empty_14x", CoreStyleConstants::Icon14x14))
+			.SetUncheckedHoveredImage(IMAGE_BRUSH("Icons/Empty_14x", CoreStyleConstants::Icon14x14))
 			.SetUncheckedPressedImage(IMAGE_BRUSH("Common/SmallCheckBox_Hovered", CoreStyleConstants::Icon14x14))
 			.SetCheckedImage(IMAGE_BRUSH("Common/SmallCheck", CoreStyleConstants::Icon14x14))
 			.SetCheckedHoveredImage(IMAGE_BRUSH("Common/SmallCheck", CoreStyleConstants::Icon14x14))
 			.SetCheckedPressedImage(IMAGE_BRUSH("Common/SmallCheck", CoreStyleConstants::Icon14x14))
-			.SetUndeterminedImage(IMAGE_BRUSH("CoreStyleConstants::Icons/Empty_14x", CoreStyleConstants::Icon14x14))
+			.SetUndeterminedImage(IMAGE_BRUSH("Icons/Empty_14x", CoreStyleConstants::Icon14x14))
 			.SetUndeterminedHoveredImage(FSlateNoResource())
 			.SetUndeterminedPressedImage(FSlateNoResource());
 		Style->Set("ToolBar.Check", BasicToolBarCheckStyle);
@@ -968,14 +971,14 @@ TSharedRef<ISlateStyle> FUMGCoreStyle::Create()
 		Style->Set( "ToolBar.Button.Checked_Pressed", new BOX_BRUSH( "Common/RoundedSelection_16x",  4.0f/16.0f, SelectionColor ) );
 
 		Style->Set("ToolBar.SimpleComboButton", Style->GetWidgetStyle<FComboButtonStyle>("ComboButton"));
-		Style->Set("ToolBar.CoreStyleConstants::IconSize", CoreStyleConstants::Icon20x20);
+		Style->Set("ToolBar.IconSize", CoreStyleConstants::Icon20x20);
 	}
 
 	// MenuBar
 	{
 		Style->Set( "Menu.Background",	new BOX_BRUSH( "Old/Menu_Background", FMargin(8.0f/64.0f) ) );
-		Style->Set( "Menu.CoreStyleConstants::Icon", new IMAGE_BRUSH( "CoreStyleConstants::Icons/CoreStyleConstants::Icon_tab_toolbar_16px", CoreStyleConstants::Icon16x16 ) );
-		Style->Set( "Menu.Expand", new IMAGE_BRUSH( "CoreStyleConstants::Icons/toolbar_expand_16x", CoreStyleConstants::Icon16x16) );
+		Style->Set( "Menu.Icon", new IMAGE_BRUSH( "Icons/Icon_tab_toolbar_16px", CoreStyleConstants::Icon16x16 ) );
+		Style->Set( "Menu.Expand", new IMAGE_BRUSH( "Icons/toolbar_expand_16x", CoreStyleConstants::Icon16x16) );
 		Style->Set( "Menu.SubMenuIndicator", new IMAGE_BRUSH( "Common/SubmenuArrow", CoreStyleConstants::Icon8x8 ) );
 		Style->Set( "Menu.SToolBarComboButtonBlock.Padding", FMargin(4.0f));
 		Style->Set( "Menu.SToolBarButtonBlock.Padding", FMargin(4.0f));
@@ -1019,13 +1022,13 @@ TSharedRef<ISlateStyle> FUMGCoreStyle::Create()
 		/* Read-only checkbox that appears next to a menu item */
 		/* Set images for various SCheckBox states associated with read-only menu check box items... */
 		const FCheckBoxStyle BasicMenuCheckStyle = FCheckBoxStyle()
-			.SetUncheckedImage( IMAGE_BRUSH( "CoreStyleConstants::Icons/Empty_14x", CoreStyleConstants::Icon14x14 ) )
-			.SetUncheckedHoveredImage( IMAGE_BRUSH( "CoreStyleConstants::Icons/Empty_14x", CoreStyleConstants::Icon14x14 ) )
+			.SetUncheckedImage( IMAGE_BRUSH( "Icons/Empty_14x", CoreStyleConstants::Icon14x14 ) )
+			.SetUncheckedHoveredImage( IMAGE_BRUSH( "Icons/Empty_14x", CoreStyleConstants::Icon14x14 ) )
 			.SetUncheckedPressedImage( IMAGE_BRUSH( "Common/SmallCheckBox_Hovered", CoreStyleConstants::Icon14x14 ) )
 			.SetCheckedImage( IMAGE_BRUSH( "Common/SmallCheck", CoreStyleConstants::Icon14x14 ) )
 			.SetCheckedHoveredImage( IMAGE_BRUSH( "Common/SmallCheck", CoreStyleConstants::Icon14x14 ) )
 			.SetCheckedPressedImage( IMAGE_BRUSH( "Common/SmallCheck", CoreStyleConstants::Icon14x14 ) )
-			.SetUndeterminedImage( IMAGE_BRUSH( "CoreStyleConstants::Icons/Empty_14x", CoreStyleConstants::Icon14x14 ) )
+			.SetUndeterminedImage( IMAGE_BRUSH( "Icons/Empty_14x", CoreStyleConstants::Icon14x14 ) )
 			.SetUndeterminedHoveredImage( FSlateNoResource() )
 			.SetUndeterminedPressedImage( FSlateNoResource() );
 
@@ -1080,7 +1083,7 @@ TSharedRef<ISlateStyle> FUMGCoreStyle::Create()
 
 		Style->Set("WindowMenuBar.Background", new BOX_BRUSH("Old/Menu_Background", FMargin(8.0f / 64.0f)));
 		Style->Set("WindowMenuBar.Label", FTextBlockStyle(NormalText).SetFont(DEFAULT_FONT("Regular", 9)));
-		Style->Set("WindowMenuBar.Expand", new IMAGE_BRUSH("CoreStyleConstants::Icons/toolbar_expand_16x", CoreStyleConstants::Icon16x16));
+		Style->Set("WindowMenuBar.Expand", new IMAGE_BRUSH("Icons/toolbar_expand_16x", CoreStyleConstants::Icon16x16));
 		Style->Set("WindowMenuBar.Button", MenuButton);
 		Style->Set("WindowMenuBar.Button.SubMenuOpen", new BORDER_BRUSH("Common/Selection", FMargin(4.f / 16.f), FLinearColor(0.10f, 0.10f, 0.10f)));
 
@@ -1126,8 +1129,8 @@ TSharedRef<ISlateStyle> FUMGCoreStyle::Create()
 		Style->Set( "NotificationList.ItemBackground", new BOX_BRUSH( "Old/Menu_Background", FMargin(8.0f/64.0f) ) );
 		Style->Set( "NotificationList.ItemBackground_Border", new BOX_BRUSH( "Old/Menu_Background_Inverted_Border_Bold", FMargin(8.0f/64.0f) ) );
 		Style->Set( "NotificationList.ItemBackground_Border_Transparent", new BOX_BRUSH("Old/Notification_Border_Flash", FMargin(8.0f/64.0f)));
-		Style->Set( "NotificationList.SuccessImage", new IMAGE_BRUSH( "CoreStyleConstants::Icons/notificationlist_success", CoreStyleConstants::Icon16x16 ) );
-		Style->Set( "NotificationList.FailImage", new IMAGE_BRUSH( "CoreStyleConstants::Icons/notificationlist_fail", CoreStyleConstants::Icon16x16 ) );
+		Style->Set( "NotificationList.SuccessImage", new IMAGE_BRUSH( "Icons/notificationlist_success", CoreStyleConstants::Icon16x16 ) );
+		Style->Set( "NotificationList.FailImage", new IMAGE_BRUSH( "Icons/notificationlist_fail", CoreStyleConstants::Icon16x16 ) );
 		Style->Set( "NotificationList.DefaultMessage", new IMAGE_BRUSH( "Common/EventMessage_Default", CoreStyleConstants::Icon40x40 ) );
 	}
 
@@ -1311,17 +1314,17 @@ TSharedRef<ISlateStyle> FUMGCoreStyle::Create()
 
 	// Widget Reflector Window
 	{
-		Style->Set("WidgetReflector.TabCoreStyleConstants::Icon", new IMAGE_BRUSH("CoreStyleConstants::Icons/CoreStyleConstants::Icon_tab_WidgetReflector_16x", CoreStyleConstants::Icon16x16));
-		Style->Set("WidgetReflector.CoreStyleConstants::Icon", new IMAGE_BRUSH("CoreStyleConstants::Icons/CoreStyleConstants::Icon_tab_WidgetReflector_40x", CoreStyleConstants::Icon40x40));
-		Style->Set("WidgetReflector.CoreStyleConstants::Icon.Small", new IMAGE_BRUSH("CoreStyleConstants::Icons/CoreStyleConstants::Icon_tab_WidgetReflector_40x", CoreStyleConstants::Icon20x20));
+		Style->Set("WidgetReflector.TabIcon", new IMAGE_BRUSH("Icons/Icon_tab_WidgetReflector_16x", CoreStyleConstants::Icon16x16));
+		Style->Set("WidgetReflector.Icon", new IMAGE_BRUSH("Icons/Icon_tab_WidgetReflector_40x", CoreStyleConstants::Icon40x40));
+		Style->Set("WidgetReflector.Icon.Small", new IMAGE_BRUSH("Icons/Icon_tab_WidgetReflector_40x", CoreStyleConstants::Icon20x20));
 		Style->Set("WidgetReflector.FocusableCheck", FCheckBoxStyle()
-			.SetUncheckedImage(IMAGE_BRUSH("CoreStyleConstants::Icons/Empty_14x", CoreStyleConstants::Icon14x14))
-			.SetUncheckedHoveredImage(IMAGE_BRUSH("CoreStyleConstants::Icons/Empty_14x", CoreStyleConstants::Icon14x14))
+			.SetUncheckedImage(IMAGE_BRUSH("Icons/Empty_14x", CoreStyleConstants::Icon14x14))
+			.SetUncheckedHoveredImage(IMAGE_BRUSH("Icons/Empty_14x", CoreStyleConstants::Icon14x14))
 			.SetUncheckedPressedImage(IMAGE_BRUSH("Common/SmallCheckBox_Hovered", CoreStyleConstants::Icon14x14))
 			.SetCheckedImage(IMAGE_BRUSH("Common/SmallCheck", CoreStyleConstants::Icon14x14))
 			.SetCheckedHoveredImage(IMAGE_BRUSH("Common/SmallCheck", CoreStyleConstants::Icon14x14))
 			.SetCheckedPressedImage(IMAGE_BRUSH("Common/SmallCheck", CoreStyleConstants::Icon14x14))
-			.SetUndeterminedImage(IMAGE_BRUSH("CoreStyleConstants::Icons/Empty_14x", CoreStyleConstants::Icon14x14))
+			.SetUndeterminedImage(IMAGE_BRUSH("Icons/Empty_14x", CoreStyleConstants::Icon14x14))
 			.SetUndeterminedHoveredImage(FSlateNoResource())
 			.SetUndeterminedPressedImage(FSlateNoResource())
 			);
@@ -1340,8 +1343,8 @@ TSharedRef<ISlateStyle> FUMGCoreStyle::Create()
 
 	// Wizard CoreStyleConstants::Icons
 	{
-		Style->Set("Wizard.BackCoreStyleConstants::Icon", new IMAGE_BRUSH("CoreStyleConstants::Icons/BackCoreStyleConstants::Icon", CoreStyleConstants::Icon8x8));
-		Style->Set("Wizard.NextCoreStyleConstants::Icon", new IMAGE_BRUSH("CoreStyleConstants::Icons/NextCoreStyleConstants::Icon", CoreStyleConstants::Icon8x8));
+		Style->Set("Wizard.BackIcon", new IMAGE_BRUSH("Icons/BackIcon", CoreStyleConstants::Icon8x8));
+		Style->Set("Wizard.NextIcon", new IMAGE_BRUSH("Icons/NextIcon", CoreStyleConstants::Icon8x8));
 	}
 
 	// Syntax highlighting
@@ -1355,7 +1358,7 @@ TSharedRef<ISlateStyle> FUMGCoreStyle::Create()
 		Style->Set("SyntaxHighlight.NodeAttribueAssignment", FTextBlockStyle(SmallMonospacedText).SetColorAndOpacity(FLinearColor(FColor(0xffb2b400)))); // yellow
 		Style->Set("SyntaxHighlight.NodeAttributeValue", FTextBlockStyle(SmallMonospacedText).SetColorAndOpacity(FLinearColor(FColor(0xffb46100)))); // orange
 	}
-
+#endif // // UE-116081
 	return Style;
 }
 
diff --git a/Engine/Source/Runtime/SlateCore/Public/Styling/UMGCoreStyle.h b/Engine/Source/Runtime/SlateCore/Public/Styling/UMGCoreStyle.h
--- a/Engine/Source/Runtime/SlateCore/Public/Styling/UMGCoreStyle.h
+++ b/Engine/Source/Runtime/SlateCore/Public/Styling/UMGCoreStyle.h
@@ -17,7 +17,9 @@ public:
 
 	static const ISlateStyle& Get()
 	{
-		return *(Instance.Get());
+		// UE-116081
+		// return *(Instance.Get());
+		return FAppStyle::Get();
 	}
 
 	/** Get the default font for Slate */
