Fix UMGCoreStyle referencing nonexistant files.
Expose WidgetStyle for List/Tree Views.

diff --git a/Engine/Source/Runtime/SlateCore/Private/Styling/UMGCoreStyle.cpp b/Engine/Source/Runtime/SlateCore/Private/Styling/UMGCoreStyle.cpp
--- a/Engine/Source/Runtime/SlateCore/Private/Styling/UMGCoreStyle.cpp
+++ b/Engine/Source/Runtime/SlateCore/Private/Styling/UMGCoreStyle.cpp
@@ -22,6 +22,8 @@
 
 TSharedPtr< ISlateStyle > FUMGCoreStyle::Instance = nullptr;
 
+using namespace CoreStyleConstants;
+
 /* FUMGCoreStyle helper class
  *****************************************************************************/
 
@@ -196,9 +198,6 @@ TSharedRef<ISlateStyle> FUMGCoreStyle::Create()
 		return Style;
 	}
 
-	// UE-116081
-#if 0
-
 	// These are the Slate colors which reference the dynamic colors in FSlateCoreStyle; these are the colors to put into the style
 	const FSlateColor DefaultForeground(FUMGColor::Black);
 	const FSlateColor InvertedForeground(FUMGColor::ForegroundInverted);
@@ -208,8 +207,8 @@ TSharedRef<ISlateStyle> FUMGCoreStyle::Create()
 	const FSlateColor SelectionColor_Pressed(Style->SelectionColor_Pressed_LinearRef);
 	const FSlateColor HighlightColor(FUMGColor::Highlight);
 
-	Style->Set("AppIcon", new IMAGE_BRUSH("Icons/DefaultAppIcon", CoreStyleConstants::Icon24x24));
-	Style->Set("AppIcon.Small", new IMAGE_BRUSH("Icons/DefaultAppIcon", CoreStyleConstants::Icon24x24));
+	Style->Set("AppIcon", new IMAGE_BRUSH("Starship/Common/UELogo", Icon24x24, FStyleColors::White));
+	Style->Set("AppIcon.Small", new IMAGE_BRUSH_SVG("Starship/Common/unreal-small", Icon24x24, FStyleColors::Foreground));
 
 	Style->Set("AppIconPadding", FMargin(4, 4, 0, 0));
 	Style->Set("AppIconPadding.Small", FMargin(4, 4, 0, 0));
@@ -218,7 +217,7 @@ TSharedRef<ISlateStyle> FUMGCoreStyle::Create()
 
 	Style->Set("SmallFont", DEFAULT_FONT("Regular", SmallTextSize));
 
-	FSlateBrush* DefaultTextUnderlineBrush = new IMAGE_BRUSH("Old/White", CoreStyleConstants::Icon8x8, FLinearColor::White, ESlateBrushTileType::Both);
+	FSlateBrush* DefaultTextUnderlineBrush = new IMAGE_BRUSH("Old/White", Icon8x8, FLinearColor::White, ESlateBrushTileType::Both);
 	Style->Set("DefaultTextUnderline", DefaultTextUnderlineBrush);
 
 	// Normal Text
@@ -265,9 +264,9 @@ TSharedRef<ISlateStyle> FUMGCoreStyle::Create()
 		);
 
 	// Common brushes
-	FSlateBrush* GenericWhiteBox = new IMAGE_BRUSH("Old/White", CoreStyleConstants::Icon16x16);
+	FSlateBrush* GenericWhiteBox = new IMAGE_BRUSH("Old/White", Icon16x16);
 	{
-		Style->Set("Checkerboard", new IMAGE_BRUSH("Checkerboard", CoreStyleConstants::Icon16x16, FLinearColor::White, ESlateBrushTileType::Both));
+		Style->Set("Checkerboard", new IMAGE_BRUSH("Checkerboard", Icon16x16, FLinearColor::White, ESlateBrushTileType::Both));
 
 		Style->Set("GenericWhiteBox", GenericWhiteBox);
 
@@ -325,8 +324,8 @@ TSharedRef<ISlateStyle> FUMGCoreStyle::Create()
 
 	// Cursor Icons
 	{
-		Style->Set("SoftwareCursor_Grab", new IMAGE_BRUSH( "Icons/cursor_grab", CoreStyleConstants::Icon16x16));
-		Style->Set("SoftwareCursor_CardinalCross", new IMAGE_BRUSH( "Icons/cursor_cardinal_cross", CoreStyleConstants::Icon24x24));
+		Style->Set("SoftwareCursor_Grab", new IMAGE_BRUSH( "Icons/cursor_grab", Icon16x16));
+		Style->Set("SoftwareCursor_CardinalCross", new IMAGE_BRUSH( "Icons/cursor_cardinal_cross", Icon24x24));
 	}
 
 	// Common Icons
@@ -337,13 +336,13 @@ TSharedRef<ISlateStyle> FUMGCoreStyle::Create()
 
 	// Common Icons
 	{
-		Style->Set( "Icons.Cross", new IMAGE_BRUSH( "Icons/Cross_12x", CoreStyleConstants::Icon12x12 ) );
-		Style->Set( "Icons.Denied", new IMAGE_BRUSH( "Icons/denied_16x", CoreStyleConstants::Icon16x16 ) );
-		Style->Set( "Icons.Error", new IMAGE_BRUSH( "Icons/Icon_error_16x", CoreStyleConstants::Icon16x16) );
-		Style->Set( "Icons.Help", new IMAGE_BRUSH( "Icons/Icon_help_16x", CoreStyleConstants::Icon16x16) );
-		Style->Set( "Icons.Info", new IMAGE_BRUSH( "Icons/Icon_info_16x", CoreStyleConstants::Icon16x16) );
-		Style->Set( "Icons.Warning", new IMAGE_BRUSH( "Icons/Icon_warning_16x", CoreStyleConstants::Icon16x16) );
-		Style->Set( "Icons.Download", new IMAGE_BRUSH( "Icons/Icon_Downloads_16x", CoreStyleConstants::Icon16x16) );
+		Style->Set( "Icons.Cross", new IMAGE_BRUSH( "Icons/Cross_12x", Icon12x12 ) );
+		Style->Set( "Icons.Denied", new IMAGE_BRUSH( "Icons/denied_16x", Icon16x16 ) );
+		Style->Set( "Icons.Error", new IMAGE_BRUSH( "Icons/Icon_error_16x", Icon16x16) );
+		Style->Set( "Icons.Help", new IMAGE_BRUSH( "Icons/Icon_help_16x", Icon16x16) );
+		Style->Set( "Icons.Info", new IMAGE_BRUSH( "Icons/Icon_info_16x", Icon16x16) );
+		Style->Set( "Icons.Warning", new IMAGE_BRUSH( "Icons/Icon_warning_16x", Icon16x16) );
+		Style->Set( "Icons.Download", new IMAGE_BRUSH( "Icons/Icon_Downloads_16x", Icon16x16) );
 	}
 
 	// Tool panels
@@ -359,15 +358,15 @@ TSharedRef<ISlateStyle> FUMGCoreStyle::Create()
 
 	// Generic command Icons
 	{
-		Style->Set( "GenericCommands.Undo", new IMAGE_BRUSH( "Icons/Icon_undo_16px", CoreStyleConstants::Icon16x16 ) );
-		Style->Set( "GenericCommands.Redo", new IMAGE_BRUSH( "Icons/Icon_redo_16px", CoreStyleConstants::Icon16x16 ) );
+		Style->Set( "GenericCommands.Undo", new IMAGE_BRUSH( "Icons/Icon_undo_16px", Icon16x16 ) );
+		Style->Set( "GenericCommands.Redo", new IMAGE_BRUSH( "Icons/Icon_redo_16px", Icon16x16 ) );
 
-		Style->Set( "GenericCommands.Copy", new IMAGE_BRUSH("Icons/Edit/Icon_Edit_Copy_16x", CoreStyleConstants::Icon16x16) );
-		Style->Set( "GenericCommands.Cut", new IMAGE_BRUSH("Icons/Edit/Icon_Edit_Cut_16x", CoreStyleConstants::Icon16x16) );
-		Style->Set( "GenericCommands.Delete", new IMAGE_BRUSH("Icons/Edit/Icon_Edit_Delete_16x", CoreStyleConstants::Icon16x16) );
-		Style->Set( "GenericCommands.Paste", new IMAGE_BRUSH("Icons/Edit/Icon_Edit_Paste_16x", CoreStyleConstants::Icon16x16) );
-		Style->Set( "GenericCommands.Duplicate", new IMAGE_BRUSH("Icons/Edit/Icon_Edit_Duplicate_16x", CoreStyleConstants::Icon16x16) );
-		Style->Set( "GenericCommands.Rename", new IMAGE_BRUSH( "Icons/Edit/Icon_Edit_Rename_16x", CoreStyleConstants::Icon16x16 ) );
+		Style->Set( "GenericCommands.Copy", new IMAGE_BRUSH("Icons/Edit/Icon_Edit_Copy_16x", Icon16x16) );
+		Style->Set( "GenericCommands.Cut", new IMAGE_BRUSH("Icons/Edit/Icon_Edit_Cut_16x", Icon16x16) );
+		Style->Set( "GenericCommands.Delete", new IMAGE_BRUSH("Icons/Edit/Icon_Edit_Delete_16x", Icon16x16) );
+		Style->Set( "GenericCommands.Paste", new IMAGE_BRUSH("Icons/Edit/Icon_Edit_Paste_16x", Icon16x16) );
+		Style->Set( "GenericCommands.Duplicate", new IMAGE_BRUSH("Icons/Edit/Icon_Edit_Duplicate_16x", Icon16x16) );
+		Style->Set( "GenericCommands.Rename", new IMAGE_BRUSH( "Icons/Edit/Icon_Edit_Rename_16x", Icon16x16 ) );
 	}
 
 	// SVerticalBox Drag& Drop Icon
@@ -388,15 +387,15 @@ TSharedRef<ISlateStyle> FUMGCoreStyle::Create()
 
 	// SButton defaults...
 	const FButtonStyle Button = FButtonStyle()
-		.SetNormal(FSlateRoundedBoxBrush(FUMGColor::Secondary, 4.0f, FUMGColor::InputOutline, CoreStyleConstants::InputFocusThickness))
-		.SetHovered(FSlateRoundedBoxBrush(FUMGColor::Hover, 4.0f, FUMGColor::Hover, CoreStyleConstants::InputFocusThickness))
-		.SetPressed(FSlateRoundedBoxBrush(FUMGColor::Header, 4.0f, FUMGColor::Hover, CoreStyleConstants::InputFocusThickness))
+		.SetNormal(FSlateRoundedBoxBrush(FUMGColor::Secondary, 4.0f, FUMGColor::InputOutline, InputFocusThickness))
+		.SetHovered(FSlateRoundedBoxBrush(FUMGColor::Hover, 4.0f, FUMGColor::Hover, InputFocusThickness))
+		.SetPressed(FSlateRoundedBoxBrush(FUMGColor::Header, 4.0f, FUMGColor::Hover, InputFocusThickness))
 		.SetNormalForeground(FUMGColor::ForegroundHover)
 		.SetHoveredForeground(FUMGColor::ForegroundHover)
 		.SetPressedForeground(FUMGColor::ForegroundHover)
 		.SetDisabledForeground(FUMGColor::Foreground)
-		.SetNormalPadding(CoreStyleConstants::ButtonMargins)
-		.SetPressedPadding(CoreStyleConstants::ButtonMargins);
+		.SetNormalPadding(ButtonMargins)
+		.SetPressedPadding(ButtonMargins);
 	{
 		Style->Set( "Button", Button );
 
@@ -407,7 +406,7 @@ TSharedRef<ISlateStyle> FUMGCoreStyle::Create()
 	{
 		FComboButtonStyle ComboButton = FComboButtonStyle()
 			.SetButtonStyle(Button)
-			.SetDownArrowImage(IMAGE_BRUSH("Common/ComboArrow", CoreStyleConstants::Icon8x8))
+			.SetDownArrowImage(IMAGE_BRUSH("Common/ComboArrow", Icon8x8))
 			.SetMenuBorderBrush(BOX_BRUSH("Old/Menu_Background", FMargin(8.0f/64.0f)))
 			.SetMenuBorderPadding(FMargin(0.0f));
 		Style->Set( "ComboButton", ComboButton );
@@ -423,7 +422,7 @@ TSharedRef<ISlateStyle> FUMGCoreStyle::Create()
 	{
 		FComboButtonStyle MessageLogListingComboButton = FComboButtonStyle()
 			.SetButtonStyle(NoBorder)
-			.SetDownArrowImage(IMAGE_BRUSH("Common/ComboArrow", CoreStyleConstants::Icon8x8))
+			.SetDownArrowImage(IMAGE_BRUSH("Common/ComboArrow", Icon8x8))
 			.SetMenuBorderBrush(FSlateNoResource())
 			.SetMenuBorderPadding(FMargin(0.0f));
 		Style->Set("MessageLogListingComboButton", MessageLogListingComboButton);
@@ -431,25 +430,25 @@ TSharedRef<ISlateStyle> FUMGCoreStyle::Create()
 
 	// SEditableComboBox defaults...
 	{
-		Style->Set( "EditableComboBox.Add", new IMAGE_BRUSH( "Icons/PlusSymbol_12x", CoreStyleConstants::Icon12x12 ) );
-		Style->Set( "EditableComboBox.Delete", new IMAGE_BRUSH( "Icons/Cross_12x", CoreStyleConstants::Icon12x12 ) );
-		Style->Set( "EditableComboBox.Rename", new IMAGE_BRUSH( "Icons/ellipsis_12x", CoreStyleConstants::Icon12x12 ) );
-		Style->Set( "EditableComboBox.Accept", new IMAGE_BRUSH( "Common/Check", CoreStyleConstants::Icon16x16 ) );
+		Style->Set( "EditableComboBox.Add", new IMAGE_BRUSH( "Icons/PlusSymbol_12x", Icon12x12 ) );
+		Style->Set( "EditableComboBox.Delete", new IMAGE_BRUSH( "Icons/Cross_12x", Icon12x12 ) );
+		Style->Set( "EditableComboBox.Rename", new IMAGE_BRUSH( "Icons/ellipsis_12x", Icon12x12 ) );
+		Style->Set( "EditableComboBox.Accept", new IMAGE_BRUSH( "Common/Check", Icon16x16 ) );
 	}
 
 	// SCheckBox defaults...
 	{
 		const FCheckBoxStyle BasicCheckBoxStyle = FCheckBoxStyle()
 			.SetCheckBoxType(ESlateCheckBoxType::CheckBox)
-			.SetUncheckedImage( IMAGE_BRUSH( "Common/CheckBox", CoreStyleConstants::Icon16x16 ) )
-			.SetUncheckedHoveredImage( IMAGE_BRUSH( "Common/CheckBox", CoreStyleConstants::Icon16x16 ) )
-			.SetUncheckedPressedImage( IMAGE_BRUSH( "Common/CheckBox_Hovered", CoreStyleConstants::Icon16x16, FLinearColor( 0.5f, 0.5f, 0.5f ) ) )
-			.SetCheckedImage( IMAGE_BRUSH( "Common/CheckBox_Checked_Hovered", CoreStyleConstants::Icon16x16 ) )
-			.SetCheckedHoveredImage( IMAGE_BRUSH( "Common/CheckBox_Checked_Hovered", CoreStyleConstants::Icon16x16, FLinearColor( 0.5f, 0.5f, 0.5f ) ) )
-			.SetCheckedPressedImage( IMAGE_BRUSH( "Common/CheckBox_Checked", CoreStyleConstants::Icon16x16 ) )
-			.SetUndeterminedImage( IMAGE_BRUSH( "Common/CheckBox_Undetermined", CoreStyleConstants::Icon16x16 ) )
-			.SetUndeterminedHoveredImage( IMAGE_BRUSH( "Common/CheckBox_Undetermined_Hovered", CoreStyleConstants::Icon16x16 ) )
-			.SetUndeterminedPressedImage( IMAGE_BRUSH( "Common/CheckBox_Undetermined_Hovered", CoreStyleConstants::Icon16x16, FLinearColor( 0.5f, 0.5f, 0.5f ) ) );
+			.SetUncheckedImage( IMAGE_BRUSH( "Common/CheckBox", Icon16x16 ) )
+			.SetUncheckedHoveredImage( IMAGE_BRUSH( "Common/CheckBox", Icon16x16 ) )
+			.SetUncheckedPressedImage( IMAGE_BRUSH( "Common/CheckBox_Hovered", Icon16x16, FLinearColor( 0.5f, 0.5f, 0.5f ) ) )
+			.SetCheckedImage( IMAGE_BRUSH( "Common/CheckBox_Checked_Hovered", Icon16x16 ) )
+			.SetCheckedHoveredImage( IMAGE_BRUSH( "Common/CheckBox_Checked_Hovered", Icon16x16, FLinearColor( 0.5f, 0.5f, 0.5f ) ) )
+			.SetCheckedPressedImage( IMAGE_BRUSH( "Common/CheckBox_Checked", Icon16x16 ) )
+			.SetUndeterminedImage( IMAGE_BRUSH( "Common/CheckBox_Undetermined", Icon16x16 ) )
+			.SetUndeterminedHoveredImage( IMAGE_BRUSH( "Common/CheckBox_Undetermined_Hovered", Icon16x16 ) )
+			.SetUndeterminedPressedImage( IMAGE_BRUSH( "Common/CheckBox_Undetermined_Hovered", Icon16x16, FLinearColor( 0.5f, 0.5f, 0.5f ) ) );
 		Style->Set( "Checkbox", BasicCheckBoxStyle );
 
 		/* Set images for various transparent SCheckBox states ... */
@@ -481,25 +480,25 @@ TSharedRef<ISlateStyle> FUMGCoreStyle::Create()
 		const FCheckBoxStyle ToggleButtonRowStyle = FCheckBoxStyle()
 			.SetCheckBoxType(ESlateCheckBoxType::ToggleButton)
 			.SetUncheckedImage(FSlateNoResource())
-			.SetUncheckedHoveredImage(IMAGE_BRUSH("Common/Selection", CoreStyleConstants::Icon8x8, SelectionColor_Inactive))
-			.SetUncheckedPressedImage(IMAGE_BRUSH("Common/Selection", CoreStyleConstants::Icon8x8, SelectionColor_Inactive))
-			.SetCheckedImage(IMAGE_BRUSH("Common/Selection", CoreStyleConstants::Icon8x8, SelectionColor))
-			.SetCheckedHoveredImage(IMAGE_BRUSH("Common/Selection", CoreStyleConstants::Icon8x8, SelectionColor))
+			.SetUncheckedHoveredImage(IMAGE_BRUSH("Common/Selection", Icon8x8, SelectionColor_Inactive))
+			.SetUncheckedPressedImage(IMAGE_BRUSH("Common/Selection", Icon8x8, SelectionColor_Inactive))
+			.SetCheckedImage(IMAGE_BRUSH("Common/Selection", Icon8x8, SelectionColor))
+			.SetCheckedHoveredImage(IMAGE_BRUSH("Common/Selection", Icon8x8, SelectionColor))
 			.SetCheckedPressedImage(BOX_BRUSH("Common/Selector", 4.0f / 16.0f, SelectorColor));
 		Style->Set("ToggleButtonRowStyle", ToggleButtonRowStyle);
 
 		/* A radio button is actually just a SCheckBox box with different images */
 		/* Set images for various radio button (SCheckBox) states ... */
 		const FCheckBoxStyle BasicRadioButtonStyle = FCheckBoxStyle()
-			.SetUncheckedImage( IMAGE_BRUSH( "Common/RadioButton_Unselected_16x", CoreStyleConstants::Icon16x16 ) )
-			.SetUncheckedHoveredImage( IMAGE_BRUSH( "Common/RadioButton_Unselected_16x", CoreStyleConstants::Icon16x16 ) )
-			.SetUncheckedPressedImage( IMAGE_BRUSH( "Common/RadioButton_Unselected_16x", CoreStyleConstants::Icon16x16 ) )
-			.SetCheckedImage( IMAGE_BRUSH( "Common/RadioButton_Selected_16x", CoreStyleConstants::Icon16x16 ) )
-			.SetCheckedHoveredImage( IMAGE_BRUSH( "Common/RadioButton_Selected_16x", CoreStyleConstants::Icon16x16, SelectionColor ) )
-			.SetCheckedPressedImage( IMAGE_BRUSH( "Common/RadioButton_Unselected_16x", CoreStyleConstants::Icon16x16, SelectionColor_Pressed ) )
-			.SetUndeterminedImage( IMAGE_BRUSH( "Common/RadioButton_Unselected_16x", CoreStyleConstants::Icon16x16 ) )
-			.SetUndeterminedHoveredImage( IMAGE_BRUSH( "Common/RadioButton_Unselected_16x", CoreStyleConstants::Icon16x16, SelectionColor ) )
-			.SetUndeterminedPressedImage( IMAGE_BRUSH( "Common/RadioButton_Unselected_16x", CoreStyleConstants::Icon16x16, SelectionColor_Pressed ) );
+			.SetUncheckedImage( IMAGE_BRUSH( "Common/RadioButton_Unselected_16x", Icon16x16 ) )
+			.SetUncheckedHoveredImage( IMAGE_BRUSH( "Common/RadioButton_Unselected_16x", Icon16x16 ) )
+			.SetUncheckedPressedImage( IMAGE_BRUSH( "Common/RadioButton_Unselected_16x", Icon16x16 ) )
+			.SetCheckedImage( IMAGE_BRUSH( "Common/RadioButton_Selected_16x", Icon16x16 ) )
+			.SetCheckedHoveredImage( IMAGE_BRUSH( "Common/RadioButton_Selected_16x", Icon16x16, SelectionColor ) )
+			.SetCheckedPressedImage( IMAGE_BRUSH( "Common/RadioButton_Unselected_16x", Icon16x16, SelectionColor_Pressed ) )
+			.SetUndeterminedImage( IMAGE_BRUSH( "Common/RadioButton_Unselected_16x", Icon16x16 ) )
+			.SetUndeterminedHoveredImage( IMAGE_BRUSH( "Common/RadioButton_Unselected_16x", Icon16x16, SelectionColor ) )
+			.SetUndeterminedPressedImage( IMAGE_BRUSH( "Common/RadioButton_Unselected_16x", Icon16x16, SelectionColor_Pressed ) );
 		Style->Set( "RadioButton", BasicRadioButtonStyle );
 	}
 
@@ -655,8 +654,8 @@ TSharedRef<ISlateStyle> FUMGCoreStyle::Create()
 	// SExpandableArea defaults...
 	{
 		Style->Set( "ExpandableArea", FExpandableAreaStyle()
-			.SetCollapsedImage( IMAGE_BRUSH( "Common/TreeArrow_Collapsed", CoreStyleConstants::Icon10x10, DefaultForeground ) )
-			.SetExpandedImage( IMAGE_BRUSH( "Common/TreeArrow_Expanded", CoreStyleConstants::Icon10x10, DefaultForeground ) )
+			.SetCollapsedImage( IMAGE_BRUSH( "Common/TreeArrow_Collapsed", Icon10x10, DefaultForeground ) )
+			.SetExpandedImage( IMAGE_BRUSH( "Common/TreeArrow_Expanded", Icon10x10, DefaultForeground ) )
 			);
 		Style->Set( "ExpandableArea.TitleFont", DEFAULT_FONT( "Bold", 8 ) );
 		Style->Set( "ExpandableArea.Border", new BOX_BRUSH( "Common/GroupBorder", FMargin(4.0f/16.0f) ) );
@@ -673,10 +672,10 @@ TSharedRef<ISlateStyle> FUMGCoreStyle::Create()
 
 		Style->Set( "SearchBox", FSearchBoxStyle()
 			.SetTextBoxStyle( SpecialEditableTextBoxStyle )
-			.SetUpArrowImage( IMAGE_BRUSH( "Common/UpArrow", CoreStyleConstants::Icon8x8 ) )
-			.SetDownArrowImage( IMAGE_BRUSH( "Common/DownArrow", CoreStyleConstants::Icon8x8 ) )
-			.SetGlassImage( IMAGE_BRUSH( "Common/SearchGlass", CoreStyleConstants::Icon16x16 ) )
-			.SetClearImage( IMAGE_BRUSH( "Common/X", CoreStyleConstants::Icon16x16 ) )
+			.SetUpArrowImage( IMAGE_BRUSH( "Common/UpArrow", Icon8x8 ) )
+			.SetDownArrowImage( IMAGE_BRUSH( "Common/DownArrow", Icon8x8 ) )
+			.SetGlassImage( IMAGE_BRUSH( "Common/SearchGlass", Icon16x16 ) )
+			.SetClearImage( IMAGE_BRUSH( "Common/X", Icon16x16 ) )
 			);
 	}
 
@@ -694,11 +693,11 @@ TSharedRef<ISlateStyle> FUMGCoreStyle::Create()
 
 		Style->Set( "VolumeControl", FVolumeControlStyle()
 			.SetSliderStyle( SliderStyle )
-			.SetHighVolumeImage( IMAGE_BRUSH( "Common/VolumeControl_High", CoreStyleConstants::Icon16x16 ) )
-			.SetMidVolumeImage( IMAGE_BRUSH( "Common/VolumeControl_Mid", CoreStyleConstants::Icon16x16 ) )
-			.SetLowVolumeImage( IMAGE_BRUSH( "Common/VolumeControl_Low", CoreStyleConstants::Icon16x16 ) )
-			.SetNoVolumeImage( IMAGE_BRUSH( "Common/VolumeControl_Off", CoreStyleConstants::Icon16x16 ) )
-			.SetMutedImage( IMAGE_BRUSH( "Common/VolumeControl_Muted", CoreStyleConstants::Icon16x16 ) )
+			.SetHighVolumeImage( IMAGE_BRUSH( "Common/VolumeControl_High", Icon16x16 ) )
+			.SetMidVolumeImage( IMAGE_BRUSH( "Common/VolumeControl_Mid", Icon16x16 ) )
+			.SetLowVolumeImage( IMAGE_BRUSH( "Common/VolumeControl_Low", Icon16x16 ) )
+			.SetNoVolumeImage( IMAGE_BRUSH( "Common/VolumeControl_Off", Icon16x16 ) )
+			.SetMutedImage( IMAGE_BRUSH( "Common/VolumeControl_Muted", Icon16x16 ) )
 			);
 	}
 
@@ -709,7 +708,7 @@ TSharedRef<ISlateStyle> FUMGCoreStyle::Create()
 			.SetHoveredBackgroundBrush( BOX_BRUSH( "Common/Spinbox_Hovered", FMargin(4.0f/16.0f) ) )
 			.SetActiveFillBrush( BOX_BRUSH( "Common/Spinbox_Fill", FMargin(4.0f/16.0f, 4.0f/16.0f, 8.0f/16.0f, 4.0f/16.0f) ) )
 			.SetInactiveFillBrush( BOX_BRUSH( "Common/Spinbox_Fill_Hovered", FMargin(4.0f/16.0f) ) )
-			.SetArrowsImage( IMAGE_BRUSH( "Common/SpinArrows", CoreStyleConstants::Icon12x12 ) )
+			.SetArrowsImage( IMAGE_BRUSH( "Common/SpinArrows", Icon12x12 ) )
 			.SetForegroundColor( InvertedForeground )
 			);
 	}
@@ -721,7 +720,7 @@ TSharedRef<ISlateStyle> FUMGCoreStyle::Create()
 			.SetHoveredBackgroundBrush( FSlateNoResource() )
 			.SetActiveFillBrush( BOX_BRUSH( "Common/Spinbox_Fill_Hovered", FMargin(4.0f/16.0f) ) )
 			.SetInactiveFillBrush( BOX_BRUSH( "Common/Spinbox_Fill", FMargin(4.0f/16.0f, 4.0f/16.0f, 8.0f/16.0f, 4.0f/16.0f) ) )
-			.SetArrowsImage( IMAGE_BRUSH( "Common/SpinArrows", CoreStyleConstants::Icon12x12 ) )
+			.SetArrowsImage( IMAGE_BRUSH( "Common/SpinArrows", Icon12x12 ) )
 			.SetTextPadding( FMargin(0.0f) )
 			.SetForegroundColor( InvertedForeground )
 			);
@@ -731,7 +730,7 @@ TSharedRef<ISlateStyle> FUMGCoreStyle::Create()
 			.SetHoveredBackgroundBrush(FSlateNoResource())
 			.SetActiveFillBrush(BOX_BRUSH("Common/Spinbox_Fill_Hovered_Dark", FMargin(4.0f / 16.0f)))
 			.SetInactiveFillBrush(BOX_BRUSH("Common/Spinbox_Fill_Dark", FMargin(4.0f / 16.0f, 4.0f / 16.0f, 8.0f / 16.0f, 4.0f / 16.0f)))
-			.SetArrowsImage(IMAGE_BRUSH("Common/SpinArrows", CoreStyleConstants::Icon12x12))
+			.SetArrowsImage(IMAGE_BRUSH("Common/SpinArrows", Icon12x12))
 			.SetTextPadding(FMargin(0.0f))
 			.SetForegroundColor(InvertedForeground)
 			);
@@ -744,10 +743,10 @@ TSharedRef<ISlateStyle> FUMGCoreStyle::Create()
 	// SColorPicker defaults...
 	{
 		Style->Set( "ColorPicker.Border", new BOX_BRUSH( "Common/GroupBorder", FMargin(4.0f/16.0f) ) );
-		Style->Set( "ColorPicker.AlphaBackground", new IMAGE_BRUSH( "Common/Checker", CoreStyleConstants::Icon16x16, FLinearColor::White, ESlateBrushTileType::Both ) );
-		Style->Set( "ColorPicker.EyeDropper", new IMAGE_BRUSH( "Icons/eyedropper_16px", CoreStyleConstants::Icon16x16) );
+		Style->Set( "ColorPicker.AlphaBackground", new IMAGE_BRUSH( "Common/Checker", Icon16x16, FLinearColor::White, ESlateBrushTileType::Both ) );
+		Style->Set( "ColorPicker.EyeDropper", new IMAGE_BRUSH( "Icons/eyedropper_16px", Icon16x16) );
 		Style->Set( "ColorPicker.Font", DEFAULT_FONT( "Regular", 10 ) );
-		Style->Set( "ColorPicker.Mode", new IMAGE_BRUSH( "Common/ColorPicker_Mode_16x", CoreStyleConstants::Icon16x16) );
+		Style->Set( "ColorPicker.Mode", new IMAGE_BRUSH( "Common/ColorPicker_Mode_16x", Icon16x16) );
 		Style->Set( "ColorPicker.Separator", new IMAGE_BRUSH( "Common/ColorPicker_Separator", FVector2D(2.0f, 2.0f) ) );
 		Style->Set( "ColorPicker.Selector", new IMAGE_BRUSH( "Common/Circle", FVector2D(8, 8) ) );
 		Style->Set( "ColorPicker.Slider", FSliderStyle()
@@ -765,7 +764,7 @@ TSharedRef<ISlateStyle> FUMGCoreStyle::Create()
 
 	// SColorThemes defaults...
 	{
-		Style->Set( "ColorThemes.DeleteButton", new IMAGE_BRUSH( "Common/X", CoreStyleConstants::Icon16x16) );
+		Style->Set( "ColorThemes.DeleteButton", new IMAGE_BRUSH( "Common/X", Icon16x16) );
 	}
 
 	// SColorWheel defaults...
@@ -784,24 +783,35 @@ TSharedRef<ISlateStyle> FUMGCoreStyle::Create()
 	{
 		Style->Set( "Splitter", FSplitterStyle()
 			.SetHandleNormalBrush( FSlateNoResource() )
-			.SetHandleHighlightBrush( IMAGE_BRUSH( "Common/SplitterHandleHighlight", CoreStyleConstants::Icon8x8, FLinearColor::White ) )
+			.SetHandleHighlightBrush( IMAGE_BRUSH( "Common/SplitterHandleHighlight", Icon8x8, FLinearColor::White ) )
 			);
 	}
 
+	// Lists, Trees
+	{
+		const FTableViewStyle DefaultTreeViewStyle = FTableViewStyle()
+			.SetBackgroundBrush(FSlateNoResource());
+		Style->Set("ListView", DefaultTreeViewStyle);
+
+		const FTableViewStyle DefaultTableViewStyle = FTableViewStyle()
+			.SetBackgroundBrush(FSlateNoResource());
+		Style->Set("TreeView", DefaultTableViewStyle);
+	}
+
 	// TableView defaults...
 	{
 		const FTableRowStyle DefaultTableRowStyle = FTableRowStyle()
 			.SetEvenRowBackgroundBrush(FSlateNoResource())
-			.SetEvenRowBackgroundHoveredBrush(IMAGE_BRUSH("Common/Selection", CoreStyleConstants::Icon8x8, FLinearColor(1.0f, 1.0f, 1.0f, 0.1f)))
+			.SetEvenRowBackgroundHoveredBrush(IMAGE_BRUSH("Common/Selection", Icon8x8, FLinearColor(1.0f, 1.0f, 1.0f, 0.1f)))
 			.SetOddRowBackgroundBrush(FSlateNoResource())
-			.SetOddRowBackgroundHoveredBrush(IMAGE_BRUSH("Common/Selection", CoreStyleConstants::Icon8x8, FLinearColor(1.0f, 1.0f, 1.0f, 0.1f)))
+			.SetOddRowBackgroundHoveredBrush(IMAGE_BRUSH("Common/Selection", Icon8x8, FLinearColor(1.0f, 1.0f, 1.0f, 0.1f)))
 			.SetSelectorFocusedBrush(BORDER_BRUSH("Common/Selector", FMargin(4.f / 16.f), SelectorColor))
-			.SetActiveBrush(IMAGE_BRUSH("Common/Selection", CoreStyleConstants::Icon8x8, SelectionColor))
-			.SetActiveHoveredBrush(IMAGE_BRUSH("Common/Selection", CoreStyleConstants::Icon8x8, SelectionColor))
-			.SetInactiveBrush(IMAGE_BRUSH("Common/Selection", CoreStyleConstants::Icon8x8, SelectionColor_Inactive))
-			.SetInactiveHoveredBrush(IMAGE_BRUSH("Common/Selection", CoreStyleConstants::Icon8x8, SelectionColor_Inactive))
-			.SetActiveHighlightedBrush(IMAGE_BRUSH("Common/Selection", CoreStyleConstants::Icon8x8, HighlightColor))
-			.SetInactiveHighlightedBrush(IMAGE_BRUSH("Common/Selection", CoreStyleConstants::Icon8x8, FSlateColor(FLinearColor(.1f, .1f, .1f))))
+			.SetActiveBrush(IMAGE_BRUSH("Common/Selection", Icon8x8, SelectionColor))
+			.SetActiveHoveredBrush(IMAGE_BRUSH("Common/Selection", Icon8x8, SelectionColor))
+			.SetInactiveBrush(IMAGE_BRUSH("Common/Selection", Icon8x8, SelectionColor_Inactive))
+			.SetInactiveHoveredBrush(IMAGE_BRUSH("Common/Selection", Icon8x8, SelectionColor_Inactive))
+			.SetActiveHighlightedBrush(IMAGE_BRUSH("Common/Selection", Icon8x8, HighlightColor))
+			.SetInactiveHighlightedBrush(IMAGE_BRUSH("Common/Selection", Icon8x8, FSlateColor(FLinearColor(.1f, .1f, .1f))))
 
 			.SetTextColor(DefaultForeground)
 			.SetSelectedTextColor(InvertedForeground)
@@ -814,41 +824,41 @@ TSharedRef<ISlateStyle> FUMGCoreStyle::Create()
 		Style->Set("ComboBox.Row", DefaultTableRowStyle);  
 
 		const FTableRowStyle DarkTableRowStyle = FTableRowStyle(DefaultTableRowStyle)
-			.SetEvenRowBackgroundBrush(IMAGE_BRUSH("Common/Selection", CoreStyleConstants::Icon8x8, FLinearColor(0.0f, 0.0f, 0.0f, 0.1f)))
-			.SetOddRowBackgroundBrush(IMAGE_BRUSH("Common/Selection", CoreStyleConstants::Icon8x8, FLinearColor(0.0f, 0.0f, 0.0f, 0.1f)));
+			.SetEvenRowBackgroundBrush(IMAGE_BRUSH("Common/Selection", Icon8x8, FLinearColor(0.0f, 0.0f, 0.0f, 0.1f)))
+			.SetOddRowBackgroundBrush(IMAGE_BRUSH("Common/Selection", Icon8x8, FLinearColor(0.0f, 0.0f, 0.0f, 0.1f)));
 		Style->Set("TableView.DarkRow", DarkTableRowStyle);
 
-		Style->Set( "TreeArrow_Collapsed", new IMAGE_BRUSH( "Common/TreeArrow_Collapsed", CoreStyleConstants::Icon10x10, DefaultForeground ) );
-		Style->Set( "TreeArrow_Collapsed_Hovered", new IMAGE_BRUSH( "Common/TreeArrow_Collapsed_Hovered", CoreStyleConstants::Icon10x10, DefaultForeground ) );
-		Style->Set( "TreeArrow_Expanded", new IMAGE_BRUSH( "Common/TreeArrow_Expanded", CoreStyleConstants::Icon10x10, DefaultForeground ) );
-		Style->Set( "TreeArrow_Expanded_Hovered", new IMAGE_BRUSH( "Common/TreeArrow_Expanded_Hovered", CoreStyleConstants::Icon10x10, DefaultForeground ) );
+		Style->Set( "TreeArrow_Collapsed", new IMAGE_BRUSH( "Common/TreeArrow_Collapsed", Icon10x10, DefaultForeground ) );
+		Style->Set( "TreeArrow_Collapsed_Hovered", new IMAGE_BRUSH( "Common/TreeArrow_Collapsed_Hovered", Icon10x10, DefaultForeground ) );
+		Style->Set( "TreeArrow_Expanded", new IMAGE_BRUSH( "Common/TreeArrow_Expanded", Icon10x10, DefaultForeground ) );
+		Style->Set( "TreeArrow_Expanded_Hovered", new IMAGE_BRUSH( "Common/TreeArrow_Expanded_Hovered", Icon10x10, DefaultForeground ) );
 
 		const FTableColumnHeaderStyle TableColumnHeaderStyle = FTableColumnHeaderStyle()
-			.SetSortPrimaryAscendingImage(IMAGE_BRUSH("Common/SortUpArrow", CoreStyleConstants::Icon8x4))
-			.SetSortPrimaryDescendingImage(IMAGE_BRUSH("Common/SortDownArrow", CoreStyleConstants::Icon8x4))
-			.SetSortSecondaryAscendingImage(IMAGE_BRUSH("Common/SortUpArrows", CoreStyleConstants::Icon16x4))
-			.SetSortSecondaryDescendingImage(IMAGE_BRUSH("Common/SortDownArrows", CoreStyleConstants::Icon16x4))
+			.SetSortPrimaryAscendingImage(IMAGE_BRUSH("Common/SortUpArrow", Icon8x4))
+			.SetSortPrimaryDescendingImage(IMAGE_BRUSH("Common/SortDownArrow", Icon8x4))
+			.SetSortSecondaryAscendingImage(IMAGE_BRUSH("Common/SortUpArrows", Icon16x4))
+			.SetSortSecondaryDescendingImage(IMAGE_BRUSH("Common/SortDownArrows", Icon16x4))
 			.SetNormalBrush( BOX_BRUSH( "Common/ColumnHeader", 4.f/32.f ) )
 			.SetHoveredBrush( BOX_BRUSH( "Common/ColumnHeader_Hovered", 4.f/32.f ) )
-			.SetMenuDropdownImage( IMAGE_BRUSH( "Common/ColumnHeader_Arrow", CoreStyleConstants::Icon8x8 ) )
+			.SetMenuDropdownImage( IMAGE_BRUSH( "Common/ColumnHeader_Arrow", Icon8x8 ) )
 			.SetMenuDropdownNormalBorderBrush( BOX_BRUSH( "Common/ColumnHeaderMenuButton_Normal", 4.f/32.f ) )
 			.SetMenuDropdownHoveredBorderBrush( BOX_BRUSH( "Common/ColumnHeaderMenuButton_Hovered", 4.f/32.f ) );
 		Style->Set( "TableView.Header.Column", TableColumnHeaderStyle );
 
 		const FTableColumnHeaderStyle TableLastColumnHeaderStyle = FTableColumnHeaderStyle()
-			.SetSortPrimaryAscendingImage(IMAGE_BRUSH("Common/SortUpArrow", CoreStyleConstants::Icon8x4))
-			.SetSortPrimaryDescendingImage(IMAGE_BRUSH("Common/SortDownArrow", CoreStyleConstants::Icon8x4))
-			.SetSortSecondaryAscendingImage(IMAGE_BRUSH("Common/SortUpArrows", CoreStyleConstants::Icon16x4))
-			.SetSortSecondaryDescendingImage(IMAGE_BRUSH("Common/SortDownArrows", CoreStyleConstants::Icon16x4))
+			.SetSortPrimaryAscendingImage(IMAGE_BRUSH("Common/SortUpArrow", Icon8x4))
+			.SetSortPrimaryDescendingImage(IMAGE_BRUSH("Common/SortDownArrow", Icon8x4))
+			.SetSortSecondaryAscendingImage(IMAGE_BRUSH("Common/SortUpArrows", Icon16x4))
+			.SetSortSecondaryDescendingImage(IMAGE_BRUSH("Common/SortDownArrows", Icon16x4))
 			.SetNormalBrush( FSlateNoResource() )
 			.SetHoveredBrush( BOX_BRUSH( "Common/LastColumnHeader_Hovered", 4.f/32.f ) )
-			.SetMenuDropdownImage( IMAGE_BRUSH( "Common/ColumnHeader_Arrow", CoreStyleConstants::Icon8x8 ) )
+			.SetMenuDropdownImage( IMAGE_BRUSH( "Common/ColumnHeader_Arrow", Icon8x8 ) )
 			.SetMenuDropdownNormalBorderBrush( BOX_BRUSH( "Common/ColumnHeaderMenuButton_Normal", 4.f/32.f ) )
 			.SetMenuDropdownHoveredBorderBrush( BOX_BRUSH( "Common/ColumnHeaderMenuButton_Hovered", 4.f/32.f ) );
 
 		const FSplitterStyle TableHeaderSplitterStyle = FSplitterStyle()
 			.SetHandleNormalBrush( FSlateNoResource() )
-			.SetHandleHighlightBrush( IMAGE_BRUSH( "Common/HeaderSplitterGrip", CoreStyleConstants::Icon8x8 ) );
+			.SetHandleHighlightBrush( IMAGE_BRUSH( "Common/HeaderSplitterGrip", Icon8x8 ) );
 
 		Style->Set( "TableView.Header", FHeaderRowStyle()
 			.SetColumnStyle( TableColumnHeaderStyle )
@@ -861,13 +871,13 @@ TSharedRef<ISlateStyle> FUMGCoreStyle::Create()
 
 	// MultiBox
 	{
-		Style->Set( "MultiBox.GenericToolBarIcon", new IMAGE_BRUSH( "Icons/Icon_generic_toolbar", CoreStyleConstants::Icon40x40 ) );
-		Style->Set( "MultiBox.GenericToolBarIcon.Small", new IMAGE_BRUSH( "Icons/Icon_generic_toolbar", CoreStyleConstants::Icon20x20 ) );
+		Style->Set( "MultiBox.GenericToolBarIcon", new IMAGE_BRUSH( "Icons/Icon_generic_toolbar", Icon40x40 ) );
+		Style->Set( "MultiBox.GenericToolBarIcon.Small", new IMAGE_BRUSH( "Icons/Icon_generic_toolbar", Icon20x20 ) );
 
 		Style->Set( "MultiBox.DeleteButton", FButtonStyle() 
-			.SetNormal ( IMAGE_BRUSH( "/Docking/CloseApp_Hovered", CoreStyleConstants::Icon16x16 ) )
-			.SetPressed( IMAGE_BRUSH( "/Docking/CloseApp_Pressed", CoreStyleConstants::Icon16x16 ) )
-			.SetHovered( IMAGE_BRUSH( "/Docking/CloseApp_Hovered", CoreStyleConstants::Icon16x16 ) )
+			.SetNormal ( IMAGE_BRUSH( "/Docking/CloseApp_Hovered", Icon16x16 ) )
+			.SetPressed( IMAGE_BRUSH( "/Docking/CloseApp_Pressed", Icon16x16 ) )
+			.SetHovered( IMAGE_BRUSH( "/Docking/CloseApp_Hovered", Icon16x16 ) )
 			);
 
 		Style->Set( "MultiboxHookColor", FLinearColor(0.f, 1.f, 0.f, 1.f) );
@@ -876,9 +886,9 @@ TSharedRef<ISlateStyle> FUMGCoreStyle::Create()
 	// ToolBar
 	{
 		Style->Set( "ToolBar.Background", new BOX_BRUSH( "Common/GroupBorder", FMargin(4.0f/16.0f) ) );
-		Style->Set( "ToolBar.Icon", new IMAGE_BRUSH( "Icons/Icon_tab_toolbar_16px", CoreStyleConstants::Icon16x16 ) );
-		Style->Set( "ToolBar.Expand", new IMAGE_BRUSH( "Icons/toolbar_expand_16x", CoreStyleConstants::Icon16x16) );
-		Style->Set( "ToolBar.SubMenuIndicator", new IMAGE_BRUSH( "Common/SubmenuArrow", CoreStyleConstants::Icon8x8 ) );
+		Style->Set( "ToolBar.Icon", new IMAGE_BRUSH( "Icons/Icon_tab_toolbar_16px", Icon16x16 ) );
+		Style->Set( "ToolBar.Expand", new IMAGE_BRUSH( "Icons/toolbar_expand_16x", Icon16x16) );
+		Style->Set( "ToolBar.SubMenuIndicator", new IMAGE_BRUSH( "Common/SubmenuArrow", Icon8x8 ) );
 		Style->Set( "ToolBar.SToolBarComboButtonBlock.Padding", FMargin(4.0f,0.0f));
 		Style->Set( "ToolBar.SToolBarButtonBlock.Padding", FMargin(4.0f,0.0f));
 		Style->Set( "ToolBar.SToolBarCheckComboButtonBlock.Padding", FMargin(4.0f,0.0f));
@@ -901,28 +911,28 @@ TSharedRef<ISlateStyle> FUMGCoreStyle::Create()
 
 		/* Create style for "ToolBar.CheckBox" widget ... */
 		const FCheckBoxStyle ToolBarCheckBoxStyle = FCheckBoxStyle()
-			.SetUncheckedImage( IMAGE_BRUSH( "Common/SmallCheckBox", CoreStyleConstants::Icon14x14 ) )
-			.SetCheckedImage( IMAGE_BRUSH( "Common/SmallCheckBox_Checked", CoreStyleConstants::Icon14x14 ) )
-			.SetUncheckedHoveredImage( IMAGE_BRUSH( "Common/SmallCheckBox_Hovered", CoreStyleConstants::Icon14x14 ) )
-			.SetCheckedImage( IMAGE_BRUSH( "Common/SmallCheckBox_Checked_Hovered", CoreStyleConstants::Icon14x14 ) )
-			.SetUncheckedPressedImage(IMAGE_BRUSH( "Common/SmallCheckBox_Hovered", CoreStyleConstants::Icon14x14, FLinearColor( 0.5f, 0.5f, 0.5f ) ) )
-			.SetCheckedPressedImage( IMAGE_BRUSH( "Common/SmallCheckBox_Checked_Hovered", CoreStyleConstants::Icon14x14, FLinearColor( 0.5f, 0.5f, 0.5f ) ) )
-			.SetUndeterminedImage( IMAGE_BRUSH( "Common/CheckBox_Undetermined", CoreStyleConstants::Icon14x14 ) )
-			.SetUndeterminedHoveredImage( IMAGE_BRUSH( "Common/CheckBox_Undetermined_Hovered", CoreStyleConstants::Icon14x14 ) )
-			.SetUndeterminedPressedImage( IMAGE_BRUSH( "Common/CheckBox_Undetermined_Hovered", CoreStyleConstants::Icon14x14, FLinearColor( 0.5f, 0.5f, 0.5f ) ) );
+			.SetUncheckedImage( IMAGE_BRUSH( "Common/SmallCheckBox", Icon14x14 ) )
+			.SetCheckedImage( IMAGE_BRUSH( "Common/SmallCheckBox_Checked", Icon14x14 ) )
+			.SetUncheckedHoveredImage( IMAGE_BRUSH( "Common/SmallCheckBox_Hovered", Icon14x14 ) )
+			.SetCheckedImage( IMAGE_BRUSH( "Common/SmallCheckBox_Checked_Hovered", Icon14x14 ) )
+			.SetUncheckedPressedImage(IMAGE_BRUSH( "Common/SmallCheckBox_Hovered", Icon14x14, FLinearColor( 0.5f, 0.5f, 0.5f ) ) )
+			.SetCheckedPressedImage( IMAGE_BRUSH( "Common/SmallCheckBox_Checked_Hovered", Icon14x14, FLinearColor( 0.5f, 0.5f, 0.5f ) ) )
+			.SetUndeterminedImage( IMAGE_BRUSH( "Common/CheckBox_Undetermined", Icon14x14 ) )
+			.SetUndeterminedHoveredImage( IMAGE_BRUSH( "Common/CheckBox_Undetermined_Hovered", Icon14x14 ) )
+			.SetUndeterminedPressedImage( IMAGE_BRUSH( "Common/CheckBox_Undetermined_Hovered", Icon14x14, FLinearColor( 0.5f, 0.5f, 0.5f ) ) );
 		/* ... and add new style */
 		Style->Set( "ToolBar.CheckBox", ToolBarCheckBoxStyle );
 
 		/* Read-only checkbox that appears next to a menu item */
 		/* Set images for various SCheckBox states associated with read-only toolbar check box items... */
 		const FCheckBoxStyle BasicToolBarCheckStyle = FCheckBoxStyle()
-			.SetUncheckedImage(IMAGE_BRUSH("Icons/Empty_14x", CoreStyleConstants::Icon14x14))
-			.SetUncheckedHoveredImage(IMAGE_BRUSH("Icons/Empty_14x", CoreStyleConstants::Icon14x14))
-			.SetUncheckedPressedImage(IMAGE_BRUSH("Common/SmallCheckBox_Hovered", CoreStyleConstants::Icon14x14))
-			.SetCheckedImage(IMAGE_BRUSH("Common/SmallCheck", CoreStyleConstants::Icon14x14))
-			.SetCheckedHoveredImage(IMAGE_BRUSH("Common/SmallCheck", CoreStyleConstants::Icon14x14))
-			.SetCheckedPressedImage(IMAGE_BRUSH("Common/SmallCheck", CoreStyleConstants::Icon14x14))
-			.SetUndeterminedImage(IMAGE_BRUSH("Icons/Empty_14x", CoreStyleConstants::Icon14x14))
+			.SetUncheckedImage(IMAGE_BRUSH("Icons/Empty_14x", Icon14x14))
+			.SetUncheckedHoveredImage(IMAGE_BRUSH("Icons/Empty_14x", Icon14x14))
+			.SetUncheckedPressedImage(IMAGE_BRUSH("Common/SmallCheckBox_Hovered", Icon14x14))
+			.SetCheckedImage(IMAGE_BRUSH("Common/SmallCheck", Icon14x14))
+			.SetCheckedHoveredImage(IMAGE_BRUSH("Common/SmallCheck", Icon14x14))
+			.SetCheckedPressedImage(IMAGE_BRUSH("Common/SmallCheck", Icon14x14))
+			.SetUndeterminedImage(IMAGE_BRUSH("Icons/Empty_14x", Icon14x14))
 			.SetUndeterminedHoveredImage(FSlateNoResource())
 			.SetUndeterminedPressedImage(FSlateNoResource());
 		Style->Set("ToolBar.Check", BasicToolBarCheckStyle);
@@ -930,12 +940,12 @@ TSharedRef<ISlateStyle> FUMGCoreStyle::Create()
 		// This radio button is actually just a check box with different images
 		/* Create style for "ToolBar.RadioButton" widget ... */
 		const FCheckBoxStyle ToolbarRadioButtonCheckBoxStyle = FCheckBoxStyle()
-				.SetUncheckedImage( IMAGE_BRUSH( "Common/RadioButton_Unselected_16x", CoreStyleConstants::Icon16x16 ) )
-				.SetCheckedImage( IMAGE_BRUSH( "Common/RadioButton_Selected_16x", CoreStyleConstants::Icon16x16 ) )
-				.SetUncheckedHoveredImage( IMAGE_BRUSH( "Common/RadioButton_Unselected_16x", CoreStyleConstants::Icon16x16, SelectionColor ) )
-				.SetCheckedHoveredImage( IMAGE_BRUSH( "Common/RadioButton_Unselected_16x", CoreStyleConstants::Icon16x16, SelectionColor ) )
-				.SetUncheckedPressedImage( IMAGE_BRUSH( "Common/RadioButton_Unselected_16x", CoreStyleConstants::Icon16x16, SelectionColor_Pressed ) )
-				.SetCheckedPressedImage( IMAGE_BRUSH( "Common/RadioButton_Selected_16x", CoreStyleConstants::Icon16x16, SelectionColor_Pressed ) );
+				.SetUncheckedImage( IMAGE_BRUSH( "Common/RadioButton_Unselected_16x", Icon16x16 ) )
+				.SetCheckedImage( IMAGE_BRUSH( "Common/RadioButton_Selected_16x", Icon16x16 ) )
+				.SetUncheckedHoveredImage( IMAGE_BRUSH( "Common/RadioButton_Unselected_16x", Icon16x16, SelectionColor ) )
+				.SetCheckedHoveredImage( IMAGE_BRUSH( "Common/RadioButton_Unselected_16x", Icon16x16, SelectionColor ) )
+				.SetUncheckedPressedImage( IMAGE_BRUSH( "Common/RadioButton_Unselected_16x", Icon16x16, SelectionColor_Pressed ) )
+				.SetCheckedPressedImage( IMAGE_BRUSH( "Common/RadioButton_Selected_16x", Icon16x16, SelectionColor_Pressed ) );
 		/* ... and add new style */
 		Style->Set( "ToolBar.RadioButton", ToolbarRadioButtonCheckBoxStyle );
 
@@ -971,15 +981,15 @@ TSharedRef<ISlateStyle> FUMGCoreStyle::Create()
 		Style->Set( "ToolBar.Button.Checked_Pressed", new BOX_BRUSH( "Common/RoundedSelection_16x",  4.0f/16.0f, SelectionColor ) );
 
 		Style->Set("ToolBar.SimpleComboButton", Style->GetWidgetStyle<FComboButtonStyle>("ComboButton"));
-		Style->Set("ToolBar.IconSize", CoreStyleConstants::Icon20x20);
+		Style->Set("ToolBar.IconSize", Icon20x20);
 	}
 
 	// MenuBar
 	{
 		Style->Set( "Menu.Background",	new BOX_BRUSH( "Old/Menu_Background", FMargin(8.0f/64.0f) ) );
-		Style->Set( "Menu.Icon", new IMAGE_BRUSH( "Icons/Icon_tab_toolbar_16px", CoreStyleConstants::Icon16x16 ) );
-		Style->Set( "Menu.Expand", new IMAGE_BRUSH( "Icons/toolbar_expand_16x", CoreStyleConstants::Icon16x16) );
-		Style->Set( "Menu.SubMenuIndicator", new IMAGE_BRUSH( "Common/SubmenuArrow", CoreStyleConstants::Icon8x8 ) );
+		Style->Set( "Menu.Icon", new IMAGE_BRUSH( "Icons/Icon_tab_toolbar_16px", Icon16x16 ) );
+		Style->Set( "Menu.Expand", new IMAGE_BRUSH( "Icons/toolbar_expand_16x", Icon16x16) );
+		Style->Set( "Menu.SubMenuIndicator", new IMAGE_BRUSH( "Common/SubmenuArrow", Icon8x8 ) );
 		Style->Set( "Menu.SToolBarComboButtonBlock.Padding", FMargin(4.0f));
 		Style->Set( "Menu.SToolBarButtonBlock.Padding", FMargin(4.0f));
 		Style->Set( "Menu.SToolBarCheckComboButtonBlock.Padding", FMargin(4.0f));
@@ -1006,15 +1016,15 @@ TSharedRef<ISlateStyle> FUMGCoreStyle::Create()
 
 		/* Set images for various SCheckBox states associated with menu check box items... */
 		const FCheckBoxStyle BasicMenuCheckBoxStyle = FCheckBoxStyle()
-			.SetUncheckedImage( IMAGE_BRUSH( "Common/SmallCheckBox", CoreStyleConstants::Icon14x14 ) )
-			.SetUncheckedHoveredImage( IMAGE_BRUSH( "Common/SmallCheckBox_Hovered", CoreStyleConstants::Icon14x14 ) )
-			.SetUncheckedPressedImage( IMAGE_BRUSH( "Common/SmallCheckBox_Hovered", CoreStyleConstants::Icon14x14, FLinearColor( 0.5f, 0.5f, 0.5f ) ) )
-			.SetCheckedImage( IMAGE_BRUSH( "Common/SmallCheckBox_Checked", CoreStyleConstants::Icon14x14 ) )
-			.SetCheckedHoveredImage( IMAGE_BRUSH( "Common/SmallCheckBox_Checked_Hovered", CoreStyleConstants::Icon14x14 ) )
-			.SetCheckedPressedImage( IMAGE_BRUSH( "Common/SmallCheckBox_Checked_Hovered", CoreStyleConstants::Icon14x14, FLinearColor( 0.5f, 0.5f, 0.5f ) ) )
-			.SetUndeterminedImage( IMAGE_BRUSH( "Common/CheckBox_Undetermined", CoreStyleConstants::Icon14x14 ) )
-			.SetUndeterminedHoveredImage( IMAGE_BRUSH( "Common/CheckBox_Undetermined_Hovered", CoreStyleConstants::Icon14x14 ) )
-			.SetUndeterminedPressedImage( IMAGE_BRUSH( "Common/CheckBox_Undetermined_Hovered", CoreStyleConstants::Icon14x14, FLinearColor( 0.5f, 0.5f, 0.5f ) ) );
+			.SetUncheckedImage( IMAGE_BRUSH( "Common/SmallCheckBox", Icon14x14 ) )
+			.SetUncheckedHoveredImage( IMAGE_BRUSH( "Common/SmallCheckBox_Hovered", Icon14x14 ) )
+			.SetUncheckedPressedImage( IMAGE_BRUSH( "Common/SmallCheckBox_Hovered", Icon14x14, FLinearColor( 0.5f, 0.5f, 0.5f ) ) )
+			.SetCheckedImage( IMAGE_BRUSH( "Common/SmallCheckBox_Checked", Icon14x14 ) )
+			.SetCheckedHoveredImage( IMAGE_BRUSH( "Common/SmallCheckBox_Checked_Hovered", Icon14x14 ) )
+			.SetCheckedPressedImage( IMAGE_BRUSH( "Common/SmallCheckBox_Checked_Hovered", Icon14x14, FLinearColor( 0.5f, 0.5f, 0.5f ) ) )
+			.SetUndeterminedImage( IMAGE_BRUSH( "Common/CheckBox_Undetermined", Icon14x14 ) )
+			.SetUndeterminedHoveredImage( IMAGE_BRUSH( "Common/CheckBox_Undetermined_Hovered", Icon14x14 ) )
+			.SetUndeterminedPressedImage( IMAGE_BRUSH( "Common/CheckBox_Undetermined_Hovered", Icon14x14, FLinearColor( 0.5f, 0.5f, 0.5f ) ) );
  
 		/* ...and add the new style */
 		Style->Set( "Menu.CheckBox", BasicMenuCheckBoxStyle );
@@ -1022,13 +1032,13 @@ TSharedRef<ISlateStyle> FUMGCoreStyle::Create()
 		/* Read-only checkbox that appears next to a menu item */
 		/* Set images for various SCheckBox states associated with read-only menu check box items... */
 		const FCheckBoxStyle BasicMenuCheckStyle = FCheckBoxStyle()
-			.SetUncheckedImage( IMAGE_BRUSH( "Icons/Empty_14x", CoreStyleConstants::Icon14x14 ) )
-			.SetUncheckedHoveredImage( IMAGE_BRUSH( "Icons/Empty_14x", CoreStyleConstants::Icon14x14 ) )
-			.SetUncheckedPressedImage( IMAGE_BRUSH( "Common/SmallCheckBox_Hovered", CoreStyleConstants::Icon14x14 ) )
-			.SetCheckedImage( IMAGE_BRUSH( "Common/SmallCheck", CoreStyleConstants::Icon14x14 ) )
-			.SetCheckedHoveredImage( IMAGE_BRUSH( "Common/SmallCheck", CoreStyleConstants::Icon14x14 ) )
-			.SetCheckedPressedImage( IMAGE_BRUSH( "Common/SmallCheck", CoreStyleConstants::Icon14x14 ) )
-			.SetUndeterminedImage( IMAGE_BRUSH( "Icons/Empty_14x", CoreStyleConstants::Icon14x14 ) )
+			.SetUncheckedImage( IMAGE_BRUSH( "Icons/Empty_14x", Icon14x14 ) )
+			.SetUncheckedHoveredImage( IMAGE_BRUSH( "Icons/Empty_14x", Icon14x14 ) )
+			.SetUncheckedPressedImage( IMAGE_BRUSH( "Common/SmallCheckBox_Hovered", Icon14x14 ) )
+			.SetCheckedImage( IMAGE_BRUSH( "Common/SmallCheck", Icon14x14 ) )
+			.SetCheckedHoveredImage( IMAGE_BRUSH( "Common/SmallCheck", Icon14x14 ) )
+			.SetCheckedPressedImage( IMAGE_BRUSH( "Common/SmallCheck", Icon14x14 ) )
+			.SetUndeterminedImage( IMAGE_BRUSH( "Icons/Empty_14x", Icon14x14 ) )
 			.SetUndeterminedHoveredImage( FSlateNoResource() )
 			.SetUndeterminedPressedImage( FSlateNoResource() );
 
@@ -1038,15 +1048,15 @@ TSharedRef<ISlateStyle> FUMGCoreStyle::Create()
 		/* This radio button is actually just a check box with different images */
 		/* Set images for various Menu radio button (SCheckBox) states... */
 		const FCheckBoxStyle BasicMenuRadioButtonStyle = FCheckBoxStyle()
-			.SetUncheckedImage( IMAGE_BRUSH( "Common/RadioButton_Unselected_16x", CoreStyleConstants::Icon16x16 ) )
-			.SetUncheckedHoveredImage( IMAGE_BRUSH( "Common/RadioButton_Unselected_16x", CoreStyleConstants::Icon16x16 ) )
-			.SetUncheckedPressedImage( IMAGE_BRUSH( "Common/RadioButton_Unselected_16x", CoreStyleConstants::Icon16x16 ) )
-			.SetCheckedImage( IMAGE_BRUSH( "Common/RadioButton_Selected_16x", CoreStyleConstants::Icon16x16 ) )
-			.SetCheckedHoveredImage( IMAGE_BRUSH( "Common/RadioButton_Selected_16x", CoreStyleConstants::Icon16x16, SelectionColor ) )
-			.SetCheckedPressedImage( IMAGE_BRUSH( "Common/RadioButton_Unselected_16x", CoreStyleConstants::Icon16x16, SelectionColor_Pressed ) )
-			.SetUndeterminedImage( IMAGE_BRUSH( "Common/RadioButton_Unselected_16x", CoreStyleConstants::Icon16x16 ) )
-			.SetUndeterminedHoveredImage( IMAGE_BRUSH( "Common/RadioButton_Unselected_16x", CoreStyleConstants::Icon16x16, SelectionColor ) )
-			.SetUndeterminedPressedImage( IMAGE_BRUSH( "Common/RadioButton_Unselected_16x", CoreStyleConstants::Icon16x16, SelectionColor_Pressed ) );
+			.SetUncheckedImage( IMAGE_BRUSH( "Common/RadioButton_Unselected_16x", Icon16x16 ) )
+			.SetUncheckedHoveredImage( IMAGE_BRUSH( "Common/RadioButton_Unselected_16x", Icon16x16 ) )
+			.SetUncheckedPressedImage( IMAGE_BRUSH( "Common/RadioButton_Unselected_16x", Icon16x16 ) )
+			.SetCheckedImage( IMAGE_BRUSH( "Common/RadioButton_Selected_16x", Icon16x16 ) )
+			.SetCheckedHoveredImage( IMAGE_BRUSH( "Common/RadioButton_Selected_16x", Icon16x16, SelectionColor ) )
+			.SetCheckedPressedImage( IMAGE_BRUSH( "Common/RadioButton_Unselected_16x", Icon16x16, SelectionColor_Pressed ) )
+			.SetUndeterminedImage( IMAGE_BRUSH( "Common/RadioButton_Unselected_16x", Icon16x16 ) )
+			.SetUndeterminedHoveredImage( IMAGE_BRUSH( "Common/RadioButton_Unselected_16x", Icon16x16, SelectionColor ) )
+			.SetUndeterminedPressedImage( IMAGE_BRUSH( "Common/RadioButton_Unselected_16x", Icon16x16, SelectionColor_Pressed ) );
 
 		/* ...and set new style */
 		Style->Set( "Menu.RadioButton", BasicMenuRadioButtonStyle );
@@ -1083,7 +1093,7 @@ TSharedRef<ISlateStyle> FUMGCoreStyle::Create()
 
 		Style->Set("WindowMenuBar.Background", new BOX_BRUSH("Old/Menu_Background", FMargin(8.0f / 64.0f)));
 		Style->Set("WindowMenuBar.Label", FTextBlockStyle(NormalText).SetFont(DEFAULT_FONT("Regular", 9)));
-		Style->Set("WindowMenuBar.Expand", new IMAGE_BRUSH("Icons/toolbar_expand_16x", CoreStyleConstants::Icon16x16));
+		Style->Set("WindowMenuBar.Expand", new IMAGE_BRUSH("Icons/toolbar_expand_16x", Icon16x16));
 		Style->Set("WindowMenuBar.Button", MenuButton);
 		Style->Set("WindowMenuBar.Button.SubMenuOpen", new BORDER_BRUSH("Common/Selection", FMargin(4.f / 16.f), FLinearColor(0.10f, 0.10f, 0.10f)));
 
@@ -1097,12 +1107,12 @@ TSharedRef<ISlateStyle> FUMGCoreStyle::Create()
 		// Extra padding on the right and bottom to account for image shadow
 		Style->Set( "ExpandableButton.Padding", FMargin(3.f, 3.f, 6.f, 6.f) );
 
-		Style->Set( "ExpandableButton.CloseButton", new IMAGE_BRUSH( "Common/ExpansionButton_CloseOverlay", CoreStyleConstants::Icon16x16) );
+		Style->Set( "ExpandableButton.CloseButton", new IMAGE_BRUSH( "Common/ExpansionButton_CloseOverlay", Icon16x16) );
 	}
 
 	// SBreadcrumbTrail defaults...
 	{
-		Style->Set( "BreadcrumbTrail.Delimiter", new IMAGE_BRUSH( "Common/Delimiter", CoreStyleConstants::Icon16x16 ) );
+		Style->Set( "BreadcrumbTrail.Delimiter", new IMAGE_BRUSH( "Common/Delimiter", Icon16x16 ) );
 
 		Style->Set( "BreadcrumbButton", FButtonStyle()
 			.SetNormal ( FSlateNoResource() )
@@ -1129,9 +1139,9 @@ TSharedRef<ISlateStyle> FUMGCoreStyle::Create()
 		Style->Set( "NotificationList.ItemBackground", new BOX_BRUSH( "Old/Menu_Background", FMargin(8.0f/64.0f) ) );
 		Style->Set( "NotificationList.ItemBackground_Border", new BOX_BRUSH( "Old/Menu_Background_Inverted_Border_Bold", FMargin(8.0f/64.0f) ) );
 		Style->Set( "NotificationList.ItemBackground_Border_Transparent", new BOX_BRUSH("Old/Notification_Border_Flash", FMargin(8.0f/64.0f)));
-		Style->Set( "NotificationList.SuccessImage", new IMAGE_BRUSH( "Icons/notificationlist_success", CoreStyleConstants::Icon16x16 ) );
-		Style->Set( "NotificationList.FailImage", new IMAGE_BRUSH( "Icons/notificationlist_fail", CoreStyleConstants::Icon16x16 ) );
-		Style->Set( "NotificationList.DefaultMessage", new IMAGE_BRUSH( "Common/EventMessage_Default", CoreStyleConstants::Icon40x40 ) );
+		Style->Set( "NotificationList.SuccessImage", new IMAGE_BRUSH( "Icons/notificationlist_success", Icon16x16 ) );
+		Style->Set( "NotificationList.FailImage", new IMAGE_BRUSH( "Icons/notificationlist_fail", Icon16x16 ) );
+		Style->Set( "NotificationList.DefaultMessage", new IMAGE_BRUSH( "Common/EventMessage_Default", Icon40x40 ) );
 	}
 
 	// SSeparator defaults...
@@ -1168,9 +1178,9 @@ TSharedRef<ISlateStyle> FUMGCoreStyle::Create()
 		const FSlateColor& TabFlashColor = SelectionColor;
 
 		const FButtonStyle CloseButton = FButtonStyle()
-			.SetNormal ( IMAGE_BRUSH( "/Docking/CloseApp_Normal", CoreStyleConstants::Icon16x16 ) )
-			.SetPressed( IMAGE_BRUSH( "/Docking/CloseApp_Pressed", CoreStyleConstants::Icon16x16 ) )
-			.SetHovered( IMAGE_BRUSH( "/Docking/CloseApp_Hovered", CoreStyleConstants::Icon16x16 ) );
+			.SetNormal ( IMAGE_BRUSH( "/Docking/CloseApp_Normal", Icon16x16 ) )
+			.SetPressed( IMAGE_BRUSH( "/Docking/CloseApp_Pressed", Icon16x16 ) )
+			.SetHovered( IMAGE_BRUSH( "/Docking/CloseApp_Hovered", Icon16x16 ) );
 
 
 		FLinearColor DockColor_Inactive(FColor(45, 45, 45));
@@ -1216,8 +1226,8 @@ TSharedRef<ISlateStyle> FUMGCoreStyle::Create()
 		Style->Set( "Docking.Cross.DockRight_Hovered", new IMAGE_BRUSH( "/Docking/OuterDockingIndicator", FVector2D(6, 6), FLinearColor(1.0f, 0.35f, 0.0f) ) );
 		Style->Set( "Docking.Cross.DockBottom", new IMAGE_BRUSH( "/Docking/OuterDockingIndicator", FVector2D(6, 6), FLinearColor(1.0f, 0.35f, 0.0f, 0.25f) ) );
 		Style->Set( "Docking.Cross.DockBottom_Hovered", new IMAGE_BRUSH( "/Docking/OuterDockingIndicator", FVector2D(6, 6), FLinearColor(1.0f, 0.35f, 0.0f) ) );
-		Style->Set( "Docking.Cross.DockCenter", new IMAGE_BRUSH( "/Docking/DockingIndicator_Center", CoreStyleConstants::Icon64x64,  FLinearColor(1.0f, 0.35f, 0.0f, 0.25f) ) );
-		Style->Set( "Docking.Cross.DockCenter_Hovered", new IMAGE_BRUSH( "/Docking/DockingIndicator_Center", CoreStyleConstants::Icon64x64,  FLinearColor(1.0f, 0.35f, 0.0f) ) );
+		Style->Set( "Docking.Cross.DockCenter", new IMAGE_BRUSH( "/Docking/DockingIndicator_Center", Icon64x64,  FLinearColor(1.0f, 0.35f, 0.0f, 0.25f) ) );
+		Style->Set( "Docking.Cross.DockCenter_Hovered", new IMAGE_BRUSH( "/Docking/DockingIndicator_Center", Icon64x64,  FLinearColor(1.0f, 0.35f, 0.0f) ) );
 		
 		Style->Set( "Docking.Cross.BorderLeft",	new FSlateNoResource() );
 		Style->Set( "Docking.Cross.BorderTop", new FSlateNoResource() );
@@ -1286,9 +1296,9 @@ TSharedRef<ISlateStyle> FUMGCoreStyle::Create()
 			.SetCloseButtonStyle(CloseButtonStyle)
 #endif
 			.SetTitleTextStyle(TitleTextStyle)
-			.SetActiveTitleBrush(IMAGE_BRUSH("Common/Window/WindowTitle", CoreStyleConstants::Icon32x32, FLinearColor(1.0f, 1.0f, 1.0f, 1.0f), ESlateBrushTileType::Horizontal))
-			.SetInactiveTitleBrush(IMAGE_BRUSH("Common/Window/WindowTitle_Inactive", CoreStyleConstants::Icon32x32, FLinearColor(1.0f, 1.0f, 1.0f, 1.0f), ESlateBrushTileType::Horizontal))
-			.SetFlashTitleBrush(IMAGE_BRUSH("Common/Window/WindowTitle_Flashing", CoreStyleConstants::Icon24x24, FLinearColor(1.0f, 1.0f, 1.0f, 1.0f), ESlateBrushTileType::Horizontal))
+			.SetActiveTitleBrush(IMAGE_BRUSH("Common/Window/WindowTitle", Icon32x32, FLinearColor(1.0f, 1.0f, 1.0f, 1.0f), ESlateBrushTileType::Horizontal))
+			.SetInactiveTitleBrush(IMAGE_BRUSH("Common/Window/WindowTitle_Inactive", Icon32x32, FLinearColor(1.0f, 1.0f, 1.0f, 1.0f), ESlateBrushTileType::Horizontal))
+			.SetFlashTitleBrush(IMAGE_BRUSH("Common/Window/WindowTitle_Flashing", Icon24x24, FLinearColor(1.0f, 1.0f, 1.0f, 1.0f), ESlateBrushTileType::Horizontal))
 			.SetOutlineBrush(BORDER_BRUSH("Common/Window/WindowOutline", FMargin(3.0f / 32.0f)))
 			.SetOutlineColor(FLinearColor(0.1f, 0.1f, 0.1f, 1.0f))
 			.SetBorderBrush(BOX_BRUSH("Common/Window/WindowBorder", 0.48f))
@@ -1314,17 +1324,17 @@ TSharedRef<ISlateStyle> FUMGCoreStyle::Create()
 
 	// Widget Reflector Window
 	{
-		Style->Set("WidgetReflector.TabIcon", new IMAGE_BRUSH("Icons/Icon_tab_WidgetReflector_16x", CoreStyleConstants::Icon16x16));
-		Style->Set("WidgetReflector.Icon", new IMAGE_BRUSH("Icons/Icon_tab_WidgetReflector_40x", CoreStyleConstants::Icon40x40));
-		Style->Set("WidgetReflector.Icon.Small", new IMAGE_BRUSH("Icons/Icon_tab_WidgetReflector_40x", CoreStyleConstants::Icon20x20));
+		Style->Set("WidgetReflector.TabIcon", new IMAGE_BRUSH("Icons/Icon_tab_WidgetReflector_16x", Icon16x16));
+		Style->Set("WidgetReflector.Icon", new IMAGE_BRUSH("Icons/Icon_tab_WidgetReflector_40x", Icon40x40));
+		Style->Set("WidgetReflector.Icon.Small", new IMAGE_BRUSH("Icons/Icon_tab_WidgetReflector_40x", Icon20x20));
 		Style->Set("WidgetReflector.FocusableCheck", FCheckBoxStyle()
-			.SetUncheckedImage(IMAGE_BRUSH("Icons/Empty_14x", CoreStyleConstants::Icon14x14))
-			.SetUncheckedHoveredImage(IMAGE_BRUSH("Icons/Empty_14x", CoreStyleConstants::Icon14x14))
-			.SetUncheckedPressedImage(IMAGE_BRUSH("Common/SmallCheckBox_Hovered", CoreStyleConstants::Icon14x14))
-			.SetCheckedImage(IMAGE_BRUSH("Common/SmallCheck", CoreStyleConstants::Icon14x14))
-			.SetCheckedHoveredImage(IMAGE_BRUSH("Common/SmallCheck", CoreStyleConstants::Icon14x14))
-			.SetCheckedPressedImage(IMAGE_BRUSH("Common/SmallCheck", CoreStyleConstants::Icon14x14))
-			.SetUndeterminedImage(IMAGE_BRUSH("Icons/Empty_14x", CoreStyleConstants::Icon14x14))
+			.SetUncheckedImage(IMAGE_BRUSH("Icons/Empty_14x", Icon14x14))
+			.SetUncheckedHoveredImage(IMAGE_BRUSH("Icons/Empty_14x", Icon14x14))
+			.SetUncheckedPressedImage(IMAGE_BRUSH("Common/SmallCheckBox_Hovered", Icon14x14))
+			.SetCheckedImage(IMAGE_BRUSH("Common/SmallCheck", Icon14x14))
+			.SetCheckedHoveredImage(IMAGE_BRUSH("Common/SmallCheck", Icon14x14))
+			.SetCheckedPressedImage(IMAGE_BRUSH("Common/SmallCheck", Icon14x14))
+			.SetUndeterminedImage(IMAGE_BRUSH("Icons/Empty_14x", Icon14x14))
 			.SetUndeterminedHoveredImage(FSlateNoResource())
 			.SetUndeterminedPressedImage(FSlateNoResource())
 			);
@@ -1336,15 +1346,15 @@ TSharedRef<ISlateStyle> FUMGCoreStyle::Create()
 			.SetFont(DEFAULT_FONT("Regular", 8))
 			.SetShadowOffset(FVector2D::ZeroVector)
 		);
-		Style->Set("MessageLog.Error", new IMAGE_BRUSH("MessageLog/Log_Error", CoreStyleConstants::Icon16x16));
-		Style->Set("MessageLog.Warning", new IMAGE_BRUSH("MessageLog/Log_Warning", CoreStyleConstants::Icon16x16));
-		Style->Set("MessageLog.Note", new IMAGE_BRUSH("MessageLog/Log_Note", CoreStyleConstants::Icon16x16));
+		Style->Set("MessageLog.Error", new IMAGE_BRUSH("MessageLog/Log_Error", Icon16x16));
+		Style->Set("MessageLog.Warning", new IMAGE_BRUSH("MessageLog/Log_Warning", Icon16x16));
+		Style->Set("MessageLog.Note", new IMAGE_BRUSH("MessageLog/Log_Note", Icon16x16));
 	}
 
-	// Wizard CoreStyleConstants::Icons
+	// Wizard Icons
 	{
-		Style->Set("Wizard.BackIcon", new IMAGE_BRUSH("Icons/BackIcon", CoreStyleConstants::Icon8x8));
-		Style->Set("Wizard.NextIcon", new IMAGE_BRUSH("Icons/NextIcon", CoreStyleConstants::Icon8x8));
+		Style->Set("Wizard.BackIcon", new IMAGE_BRUSH("Icons/BackIcon", Icon8x8));
+		Style->Set("Wizard.NextIcon", new IMAGE_BRUSH("Icons/NextIcon", Icon8x8));
 	}
 
 	// Syntax highlighting
@@ -1358,7 +1368,7 @@ TSharedRef<ISlateStyle> FUMGCoreStyle::Create()
 		Style->Set("SyntaxHighlight.NodeAttribueAssignment", FTextBlockStyle(SmallMonospacedText).SetColorAndOpacity(FLinearColor(FColor(0xffb2b400)))); // yellow
 		Style->Set("SyntaxHighlight.NodeAttributeValue", FTextBlockStyle(SmallMonospacedText).SetColorAndOpacity(FLinearColor(FColor(0xffb46100)))); // orange
 	}
-#endif // // UE-116081
+
 	return Style;
 }
 
diff --git a/Engine/Source/Runtime/SlateCore/Public/Styling/UMGCoreStyle.h b/Engine/Source/Runtime/SlateCore/Public/Styling/UMGCoreStyle.h
--- a/Engine/Source/Runtime/SlateCore/Public/Styling/UMGCoreStyle.h
+++ b/Engine/Source/Runtime/SlateCore/Public/Styling/UMGCoreStyle.h
@@ -17,9 +17,7 @@ public:
 
 	static const ISlateStyle& Get()
 	{
-		// UE-116081
-		// return *(Instance.Get());
-		return FAppStyle::Get();
+		return *(Instance.Get());
 	}
 
 	/** Get the default font for Slate */
diff --git a/Engine/Source/Runtime/UMG/Private/Components/ListView.cpp b/Engine/Source/Runtime/UMG/Private/Components/ListView.cpp
--- a/Engine/Source/Runtime/UMG/Private/Components/ListView.cpp
+++ b/Engine/Source/Runtime/UMG/Private/Components/ListView.cpp
@@ -5,6 +5,7 @@
 #include "Engine/World.h"
 #include "TimerManager.h"
 #include "Blueprint/ListViewDesignerPreviewItem.h"
+#include "Styling/UMGCoreStyle.h"
 #include "UMGPrivate.h"
 
 #define LOCTEXT_NAMESPACE "UMG"
@@ -12,10 +13,40 @@
 /////////////////////////////////////////////////////
 // UListView
 
+static FTableViewStyle* DefaultListViewStyle = nullptr;
+
+#if WITH_EDITOR
+static FTableViewStyle* EditorListViewStyle = nullptr;
+#endif 
+
 UListView::UListView(const FObjectInitializer& ObjectInitializer)
 	: Super(ObjectInitializer)
 	, Orientation(EOrientation::Orient_Vertical)
 {
+	if (DefaultListViewStyle == nullptr)
+	{
+		DefaultListViewStyle = new FTableViewStyle(FUMGCoreStyle::Get().GetWidgetStyle<FTableViewStyle>("ListView"));
+
+		// Unlink UMG default colors.
+		DefaultListViewStyle->UnlinkColors();
+	}
+
+	WidgetStyle = *DefaultListViewStyle;
+
+#if WITH_EDITOR 
+	if (EditorListViewStyle == nullptr)
+	{
+		EditorListViewStyle = new FTableViewStyle(FAppStyle::Get().GetWidgetStyle<FTableViewStyle>("ListView"));
+
+		// Unlink UMG default colors.
+		EditorListViewStyle->UnlinkColors();
+	}
+
+	if (IsEditorWidget())
+	{
+		WidgetStyle = *EditorListViewStyle;
+	}
+#endif // WITH_EDITOR
 }
 
 void UListView::ReleaseSlateResources(bool bReleaseChildren)
diff --git a/Engine/Source/Runtime/UMG/Private/Components/TreeView.cpp b/Engine/Source/Runtime/UMG/Private/Components/TreeView.cpp
--- a/Engine/Source/Runtime/UMG/Private/Components/TreeView.cpp
+++ b/Engine/Source/Runtime/UMG/Private/Components/TreeView.cpp
@@ -1,13 +1,44 @@
 // Copyright Epic Games, Inc. All Rights Reserved.
 
 #include "Components/TreeView.h"
+#include "Styling/UMGCoreStyle.h"
 
 /////////////////////////////////////////////////////
 // UTreeView
 
+static FTableViewStyle* DefaultTreeViewStyle = nullptr;
+
+#if WITH_EDITOR
+static FTableViewStyle* EditorTreeViewStyle = nullptr;
+#endif 
+
 UTreeView::UTreeView(const FObjectInitializer& ObjectInitializer)
 	: Super(ObjectInitializer)
 {
+	if (DefaultTreeViewStyle == nullptr)
+	{
+		DefaultTreeViewStyle = new FTableViewStyle(FUMGCoreStyle::Get().GetWidgetStyle<FTableViewStyle>("TreeView"));
+
+		// Unlink UMG default colors.
+		DefaultTreeViewStyle->UnlinkColors();
+	}
+
+	WidgetStyle = *DefaultTreeViewStyle;
+
+#if WITH_EDITOR 
+	if (EditorTreeViewStyle == nullptr)
+	{
+		EditorTreeViewStyle = new FTableViewStyle(FAppStyle::Get().GetWidgetStyle<FTableViewStyle>("TreeView"));
+
+		// Unlink UMG default colors.
+		EditorTreeViewStyle->UnlinkColors();
+	}
+
+	if (IsEditorWidget())
+	{
+		WidgetStyle = *EditorTreeViewStyle;
+	}
+#endif // WITH_EDITOR
 }
 
 TSharedRef<STableViewBase> UTreeView::RebuildListWidget()
diff --git a/Engine/Source/Runtime/UMG/Public/Components/ListView.h b/Engine/Source/Runtime/UMG/Public/Components/ListView.h
--- a/Engine/Source/Runtime/UMG/Public/Components/ListView.h
+++ b/Engine/Source/Runtime/UMG/Public/Components/ListView.h
@@ -146,6 +146,7 @@ protected:
 		Args.ConsumeMouseWheel = ConsumeMouseWheel;
 		Args.bReturnFocusToSelection = bReturnFocusToSelection;
 		Args.Orientation = Orientation;
+		Args.ListViewStyle = &WidgetStyle;
 		MyListView = ITypedUMGListView<UObject*>::ConstructListView<ListViewT>(this, ListItems, Args);
 		
 		MyListView->SetOnEntryInitialized(SListView<UObject*>::FOnEntryInitialized::CreateUObject(this, &UListView::HandleOnEntryInitializedInternal));
@@ -154,6 +155,9 @@ protected:
 	}
 
 protected:
+	UPROPERTY(EditAnywhere, BlueprintReadOnly, Category = ListView, meta = (DisplayName = "Style"))
+	FTableViewStyle WidgetStyle;
+
 	/** 
 	 * The scroll & layout orientation of the list. ListView and TileView only. 
 	 * Vertical will scroll vertically and arrange tiles into rows.
diff --git a/Engine/Source/Runtime/UMG/Public/Components/ListViewBase.h b/Engine/Source/Runtime/UMG/Public/Components/ListViewBase.h
--- a/Engine/Source/Runtime/UMG/Public/Components/ListViewBase.h
+++ b/Engine/Source/Runtime/UMG/Public/Components/ListViewBase.h
@@ -218,6 +218,7 @@ protected:
 		EConsumeMouseWheel ConsumeMouseWheel = EConsumeMouseWheel::WhenScrollingPossible;
 		bool bReturnFocusToSelection = false;
 		EOrientation Orientation = Orient_Vertical;
+		const FTableViewStyle* ListViewStyle = &FAppStyle::Get().GetWidgetStyle<FTableViewStyle>("ListView");
 	};
 
 	template <template<typename> class ListViewT = SListView, typename UListViewBaseT>
@@ -235,6 +236,7 @@ protected:
 			.SelectionMode(Args.SelectionMode)
 			.ReturnFocusToSelection(Args.bReturnFocusToSelection)
 			.Orientation(Args.Orientation)
+			.ListViewStyle(Args.ListViewStyle)
 			.OnGenerateRow_UObject(Implementer, &UListViewBaseT::HandleGenerateRow)
 			.OnSelectionChanged_UObject(Implementer, &UListViewBaseT::HandleSelectionChanged)
 			.OnIsSelectableOrNavigable_UObject(Implementer, &UListViewBaseT::HandleIsSelectableOrNavigable)
@@ -287,6 +289,7 @@ protected:
 		bool bClearSelectionOnClick = false;
 		EConsumeMouseWheel ConsumeMouseWheel = EConsumeMouseWheel::WhenScrollingPossible;
 		bool bReturnFocusToSelection = false;
+		const FTableViewStyle* TreeViewStyle = &FAppStyle::Get().GetWidgetStyle<FTableViewStyle>("TreeView");
 	};
 
 	template <template<typename> class TreeViewT = STreeView, typename UListViewBaseT>
@@ -302,6 +305,7 @@ protected:
 			.ConsumeMouseWheel(Args.ConsumeMouseWheel)
 			.SelectionMode(Args.SelectionMode)
 			.ReturnFocusToSelection(Args.bReturnFocusToSelection)
+			.TreeViewStyle(Args.TreeViewStyle)
 			.OnGenerateRow_UObject(Implementer, &UListViewBaseT::HandleGenerateRow)
 			.OnSelectionChanged_UObject(Implementer, &UListViewBaseT::HandleSelectionChanged)
 			.OnIsSelectableOrNavigable_UObject(Implementer, &UListViewBaseT::HandleIsSelectableOrNavigable)
diff --git a/Engine/Source/Runtime/UMG/Public/Components/TreeView.h b/Engine/Source/Runtime/UMG/Public/Components/TreeView.h
--- a/Engine/Source/Runtime/UMG/Public/Components/TreeView.h
+++ b/Engine/Source/Runtime/UMG/Public/Components/TreeView.h
@@ -65,6 +65,7 @@ protected:
 		Args.SelectionMode = SelectionMode;
 		Args.ConsumeMouseWheel = ConsumeMouseWheel;
 		Args.bReturnFocusToSelection = bReturnFocusToSelection;
+		Args.TreeViewStyle = &WidgetStyle;
 
 		MyListView = MyTreeView = ITypedUMGListView<UObject*>::ConstructTreeView<TreeViewT>(this, ListItems, Args);
 		MyTreeView->SetOnEntryInitialized(SListView<UObject*>::FOnEntryInitialized::CreateUObject(this, &UTreeView::HandleOnEntryInitializedInternal));
