Fix static analysis memory usage UMG Core Style

diff --git a/Engine/Source/Runtime/SlateCore/Private/Styling/UMGCoreStyle.cpp b/Engine/Source/Runtime/SlateCore/Private/Styling/UMGCoreStyle.cpp
--- a/Engine/Source/Runtime/SlateCore/Private/Styling/UMGCoreStyle.cpp
+++ b/Engine/Source/Runtime/SlateCore/Private/Styling/UMGCoreStyle.cpp
@@ -188,7 +188,7 @@ TSharedRef<ISlateStyle> FUMGCoreStyle::Create()
 	TSharedRef<FUMGStyleSet> Style = MakeShareable(new FUMGStyleSet("UMGCoreStyle"));
 	Style->SetContentRoot(FPaths::EngineContentDir() / TEXT("Slate"));
 
-	const FString CanaryPath = RootToContentDir(TEXT("Checkerboard"), TEXT(".png"));
+	static const FString CanaryPath = RootToContentDir(TEXT("Checkerboard"), TEXT(".png"));
 
 	if (!FPaths::FileExists(CanaryPath))
 	{
@@ -199,13 +199,13 @@ TSharedRef<ISlateStyle> FUMGCoreStyle::Create()
 	}
 
 	// These are the Slate colors which reference the dynamic colors in FSlateCoreStyle; these are the colors to put into the style
-	const FSlateColor DefaultForeground(FUMGColor::Black);
-	const FSlateColor InvertedForeground(FUMGColor::ForegroundInverted);
-	const FSlateColor SelectorColor(Style->SelectorColor_LinearRef);
-	const FSlateColor SelectionColor(Style->SelectionColor_LinearRef);
-	const FSlateColor SelectionColor_Inactive(Style->SelectionColor_Inactive_LinearRef);
-	const FSlateColor SelectionColor_Pressed(Style->SelectionColor_Pressed_LinearRef);
-	const FSlateColor HighlightColor(FUMGColor::Highlight);
+	static const FSlateColor DefaultForeground(FUMGColor::Black);
+	static const FSlateColor InvertedForeground(FUMGColor::ForegroundInverted);
+	static const FSlateColor SelectorColor(Style->SelectorColor_LinearRef);
+	static const FSlateColor SelectionColor(Style->SelectionColor_LinearRef);
+	static const FSlateColor SelectionColor_Inactive(Style->SelectionColor_Inactive_LinearRef);
+	static const FSlateColor SelectionColor_Pressed(Style->SelectionColor_Pressed_LinearRef);
+	static const FSlateColor HighlightColor(FUMGColor::Highlight);
 
 	Style->Set("AppIcon", new IMAGE_BRUSH("Starship/Common/UELogo", Icon24x24, FStyleColors::White));
 	Style->Set("AppIcon.Small", new IMAGE_BRUSH_SVG("Starship/Common/unreal-small", Icon24x24, FStyleColors::Foreground));
@@ -221,7 +221,7 @@ TSharedRef<ISlateStyle> FUMGCoreStyle::Create()
 	Style->Set("DefaultTextUnderline", DefaultTextUnderlineBrush);
 
 	// Normal Text
-	const FTextBlockStyle NormalText = FTextBlockStyle()
+	static const FTextBlockStyle NormalText = FTextBlockStyle()
 		.SetFont(DEFAULT_FONT("Regular", RegularTextSize))
 		.SetColorAndOpacity(FSlateColor::UseForeground())
 		.SetShadowOffset(FVector2D::ZeroVector)
@@ -229,11 +229,11 @@ TSharedRef<ISlateStyle> FUMGCoreStyle::Create()
 		.SetHighlightColor(FLinearColor(0.02f, 0.3f, 0.0f))
 		.SetHighlightShape(BOX_BRUSH("Common/TextBlockHighlightShape", FMargin(3.f /8.f)));
 
-	const FTextBlockStyle NormalUnderlinedText = FTextBlockStyle(NormalText)
+	static const FTextBlockStyle NormalUnderlinedText = FTextBlockStyle(NormalText)
 		.SetUnderlineBrush(*DefaultTextUnderlineBrush);
 
 	// Monospaced Text
-	const FTextBlockStyle MonospacedText = FTextBlockStyle()
+	static const FTextBlockStyle MonospacedText = FTextBlockStyle()
 		.SetFont(DEFAULT_FONT("Mono", 10))
 		.SetColorAndOpacity(FSlateColor::UseForeground())
 		.SetShadowOffset(FVector2D::ZeroVector)
@@ -242,17 +242,17 @@ TSharedRef<ISlateStyle> FUMGCoreStyle::Create()
 		.SetHighlightShape(BOX_BRUSH("Common/TextBlockHighlightShape", FMargin(3.f/8.f))
 		);
 
-	const FTextBlockStyle MonospacedUnderlinedText = FTextBlockStyle(MonospacedText)
+	static const FTextBlockStyle MonospacedUnderlinedText = FTextBlockStyle(MonospacedText)
 		.SetUnderlineBrush(*DefaultTextUnderlineBrush);
 
 	Style->Set("MonospacedText", MonospacedText);
 	Style->Set("MonospacedUnderlinedText", MonospacedUnderlinedText);
 
 	// Small Text
-	const FTextBlockStyle SmallText = FTextBlockStyle(NormalText)
+	static const FTextBlockStyle SmallText = FTextBlockStyle(NormalText)
 		.SetFont(DEFAULT_FONT("Regular", SmallTextSize));
 
-	const FTextBlockStyle SmallUnderlinedText = FTextBlockStyle(SmallText)
+	static const FTextBlockStyle SmallUnderlinedText = FTextBlockStyle(SmallText)
 		.SetUnderlineBrush(*DefaultTextUnderlineBrush);
 
 	// Embossed Text
@@ -290,7 +290,7 @@ TSharedRef<ISlateStyle> FUMGCoreStyle::Create()
 	}
 
 	// Invisible buttons, borders, etc.
-	const FButtonStyle NoBorder = FButtonStyle()
+	static const FButtonStyle NoBorder = FButtonStyle()
 		.SetNormal(FSlateNoResource())
 		.SetHovered(FSlateNoResource())
 		.SetPressed(FSlateNoResource())
@@ -373,7 +373,7 @@ TSharedRef<ISlateStyle> FUMGCoreStyle::Create()
 	Style->Set("VerticalBoxDragIndicator", new IMAGE_BRUSH("Common/VerticalBoxDragIndicator", FVector2D(6, 45)));
 	Style->Set("VerticalBoxDragIndicatorShort", new IMAGE_BRUSH("Common/VerticalBoxDragIndicatorShort", FVector2D(6, 15)));
 	// SScrollBar defaults...
-	const FScrollBarStyle ScrollBar = FScrollBarStyle()
+	static const FScrollBarStyle ScrollBar = FScrollBarStyle()
 		.SetVerticalTopSlotImage(IMAGE_BRUSH("Common/Scrollbar_Background_Vertical", FVector2D(8, 8)))
 		.SetVerticalBottomSlotImage(IMAGE_BRUSH("Common/Scrollbar_Background_Vertical", FVector2D(8, 8)))
 		.SetHorizontalTopSlotImage(IMAGE_BRUSH("Common/Scrollbar_Background_Horizontal", FVector2D(8, 8)))
@@ -386,7 +386,7 @@ TSharedRef<ISlateStyle> FUMGCoreStyle::Create()
 	}
 
 	// SButton defaults...
-	const FButtonStyle Button = FButtonStyle()
+	static const FButtonStyle Button = FButtonStyle()
 		.SetNormal(FSlateRoundedBoxBrush(FUMGColor::Secondary, 4.0f, FUMGColor::InputOutline, InputFocusThickness))
 		.SetHovered(FSlateRoundedBoxBrush(FUMGColor::Hover, 4.0f, FUMGColor::Hover, InputFocusThickness))
 		.SetPressed(FSlateRoundedBoxBrush(FUMGColor::Header, 4.0f, FUMGColor::Hover, InputFocusThickness))
@@ -438,7 +438,7 @@ TSharedRef<ISlateStyle> FUMGCoreStyle::Create()
 
 	// SCheckBox defaults...
 	{
-		const FCheckBoxStyle BasicCheckBoxStyle = FCheckBoxStyle()
+		static const FCheckBoxStyle BasicCheckBoxStyle = FCheckBoxStyle()
 			.SetCheckBoxType(ESlateCheckBoxType::CheckBox)
 			.SetUncheckedImage( IMAGE_BRUSH( "Common/CheckBox", Icon16x16 ) )
 			.SetUncheckedHoveredImage( IMAGE_BRUSH( "Common/CheckBox", Icon16x16 ) )
@@ -452,7 +452,7 @@ TSharedRef<ISlateStyle> FUMGCoreStyle::Create()
 		Style->Set( "Checkbox", BasicCheckBoxStyle );
 
 		/* Set images for various transparent SCheckBox states ... */
-		const FCheckBoxStyle BasicTransparentCheckBoxStyle = FCheckBoxStyle()
+		static const FCheckBoxStyle BasicTransparentCheckBoxStyle = FCheckBoxStyle()
 			.SetCheckBoxType(ESlateCheckBoxType::ToggleButton)
 			.SetUncheckedImage( FSlateNoResource() )
 			.SetUncheckedHoveredImage( FSlateNoResource() )
@@ -466,7 +466,7 @@ TSharedRef<ISlateStyle> FUMGCoreStyle::Create()
 		Style->Set( "TransparentCheckBox", BasicTransparentCheckBoxStyle );
 
 		/* Default Style for a toggleable button */
-		const FCheckBoxStyle ToggleButtonStyle = FCheckBoxStyle()
+		static const FCheckBoxStyle ToggleButtonStyle = FCheckBoxStyle()
 			.SetCheckBoxType(ESlateCheckBoxType::ToggleButton)
 			.SetUncheckedImage( FSlateNoResource() )
 			.SetUncheckedHoveredImage( BOX_BRUSH("Common/RoundedSelection_16x", 4.0f/16.0f, SelectionColor ) )
@@ -477,7 +477,7 @@ TSharedRef<ISlateStyle> FUMGCoreStyle::Create()
 		Style->Set( "ToggleButtonCheckbox", ToggleButtonStyle );
 
 		/* Style for a toggleable button that mimics the coloring and look of a Table Row */
-		const FCheckBoxStyle ToggleButtonRowStyle = FCheckBoxStyle()
+		static const FCheckBoxStyle ToggleButtonRowStyle = FCheckBoxStyle()
 			.SetCheckBoxType(ESlateCheckBoxType::ToggleButton)
 			.SetUncheckedImage(FSlateNoResource())
 			.SetUncheckedHoveredImage(IMAGE_BRUSH("Common/Selection", Icon8x8, SelectionColor_Inactive))
@@ -489,7 +489,7 @@ TSharedRef<ISlateStyle> FUMGCoreStyle::Create()
 
 		/* A radio button is actually just a SCheckBox box with different images */
 		/* Set images for various radio button (SCheckBox) states ... */
-		const FCheckBoxStyle BasicRadioButtonStyle = FCheckBoxStyle()
+		static const FCheckBoxStyle BasicRadioButtonStyle = FCheckBoxStyle()
 			.SetUncheckedImage( IMAGE_BRUSH( "Common/RadioButton_Unselected_16x", Icon16x16 ) )
 			.SetUncheckedHoveredImage( IMAGE_BRUSH( "Common/RadioButton_Unselected_16x", Icon16x16 ) )
 			.SetUncheckedPressedImage( IMAGE_BRUSH( "Common/RadioButton_Unselected_16x", Icon16x16 ) )
@@ -508,7 +508,7 @@ TSharedRef<ISlateStyle> FUMGCoreStyle::Create()
 		FSlateBrush* SelectionTarget = new BOX_BRUSH( "Old/DashedBorder", FMargin(6.0f/32.0f), FLinearColor( 0.0f, 0.0f, 0.0f, 0.75f ) );
 		FSlateBrush* CompositionBackground = new BORDER_BRUSH( "Old/HyperlinkDotted",  FMargin(0,0,0,3/16.0f) );
 
-		const FEditableTextStyle NormalEditableTextStyle = FEditableTextStyle()
+		static const FEditableTextStyle NormalEditableTextStyle = FEditableTextStyle()
 			.SetBackgroundImageSelected( *SelectionBackground )
 			.SetBackgroundImageComposing( *CompositionBackground )
 			.SetCaretImage( *GenericWhiteBox );
@@ -520,7 +520,7 @@ TSharedRef<ISlateStyle> FUMGCoreStyle::Create()
 	}
 
 	// SEditableTextBox defaults...
-	const FEditableTextBoxStyle NormalEditableTextBoxStyle = FEditableTextBoxStyle()
+	static const FEditableTextBoxStyle NormalEditableTextBoxStyle = FEditableTextBoxStyle()
 		.SetBackgroundImageNormal( BOX_BRUSH( "Common/TextBox", FMargin(4.0f/16.0f) ) )
 		.SetBackgroundImageHovered( BOX_BRUSH( "Common/TextBox_Hovered", FMargin(4.0f/16.0f) ) )
 		.SetBackgroundImageFocused( BOX_BRUSH( "Common/TextBox_Hovered", FMargin(4.0f/16.0f) ) )
@@ -531,7 +531,7 @@ TSharedRef<ISlateStyle> FUMGCoreStyle::Create()
 		// "NormalFont".
 	}
 
-	const FEditableTextBoxStyle DarkEditableTextBoxStyle = FEditableTextBoxStyle()
+	static const FEditableTextBoxStyle DarkEditableTextBoxStyle = FEditableTextBoxStyle()
 		.SetBackgroundImageNormal(BOX_BRUSH("Common/TextBox_Dark", FMargin(4.0f / 16.0f)))
 		.SetBackgroundImageHovered(BOX_BRUSH("Common/TextBox_Hovered_Dark", FMargin(4.0f / 16.0f)))
 		.SetBackgroundImageFocused(BOX_BRUSH("Common/TextBox_Hovered_Dark", FMargin(4.0f / 16.0f)))
@@ -663,7 +663,7 @@ TSharedRef<ISlateStyle> FUMGCoreStyle::Create()
 
 	// SSearchBox defaults...
 	{
-		const FEditableTextBoxStyle SpecialEditableTextBoxStyle = FEditableTextBoxStyle()
+		static const FEditableTextBoxStyle SpecialEditableTextBoxStyle = FEditableTextBoxStyle()
 			.SetBackgroundImageNormal( BOX_BRUSH( "Common/TextBox_Special", FMargin(8.0f/32.0f) ) )
 			.SetBackgroundImageHovered( BOX_BRUSH( "Common/TextBox_Special_Hovered", FMargin(8.0f/32.0f) ) )
 			.SetBackgroundImageFocused( BOX_BRUSH( "Common/TextBox_Special_Hovered", FMargin(8.0f/32.0f) ) )
@@ -789,18 +789,18 @@ TSharedRef<ISlateStyle> FUMGCoreStyle::Create()
 
 	// Lists, Trees
 	{
-		const FTableViewStyle DefaultTreeViewStyle = FTableViewStyle()
+		static const FTableViewStyle DefaultTreeViewStyle = FTableViewStyle()
 			.SetBackgroundBrush(FSlateNoResource());
 		Style->Set("ListView", DefaultTreeViewStyle);
 
-		const FTableViewStyle DefaultTableViewStyle = FTableViewStyle()
+		static const FTableViewStyle DefaultTableViewStyle = FTableViewStyle()
 			.SetBackgroundBrush(FSlateNoResource());
 		Style->Set("TreeView", DefaultTableViewStyle);
 	}
 
 	// TableView defaults...
 	{
-		const FTableRowStyle DefaultTableRowStyle = FTableRowStyle()
+		static const FTableRowStyle DefaultTableRowStyle = FTableRowStyle()
 			.SetEvenRowBackgroundBrush(FSlateNoResource())
 			.SetEvenRowBackgroundHoveredBrush(IMAGE_BRUSH("Common/Selection", Icon8x8, FLinearColor(1.0f, 1.0f, 1.0f, 0.1f)))
 			.SetOddRowBackgroundBrush(FSlateNoResource())
@@ -823,7 +823,7 @@ TSharedRef<ISlateStyle> FUMGCoreStyle::Create()
 		// Make this the default for the ComboBox rows also
 		Style->Set("ComboBox.Row", DefaultTableRowStyle);  
 
-		const FTableRowStyle DarkTableRowStyle = FTableRowStyle(DefaultTableRowStyle)
+		static const FTableRowStyle DarkTableRowStyle = FTableRowStyle(DefaultTableRowStyle)
 			.SetEvenRowBackgroundBrush(IMAGE_BRUSH("Common/Selection", Icon8x8, FLinearColor(0.0f, 0.0f, 0.0f, 0.1f)))
 			.SetOddRowBackgroundBrush(IMAGE_BRUSH("Common/Selection", Icon8x8, FLinearColor(0.0f, 0.0f, 0.0f, 0.1f)));
 		Style->Set("TableView.DarkRow", DarkTableRowStyle);
@@ -833,7 +833,7 @@ TSharedRef<ISlateStyle> FUMGCoreStyle::Create()
 		Style->Set( "TreeArrow_Expanded", new IMAGE_BRUSH( "Common/TreeArrow_Expanded", Icon10x10, DefaultForeground ) );
 		Style->Set( "TreeArrow_Expanded_Hovered", new IMAGE_BRUSH( "Common/TreeArrow_Expanded_Hovered", Icon10x10, DefaultForeground ) );
 
-		const FTableColumnHeaderStyle TableColumnHeaderStyle = FTableColumnHeaderStyle()
+		static const FTableColumnHeaderStyle TableColumnHeaderStyle = FTableColumnHeaderStyle()
 			.SetSortPrimaryAscendingImage(IMAGE_BRUSH("Common/SortUpArrow", Icon8x4))
 			.SetSortPrimaryDescendingImage(IMAGE_BRUSH("Common/SortDownArrow", Icon8x4))
 			.SetSortSecondaryAscendingImage(IMAGE_BRUSH("Common/SortUpArrows", Icon16x4))
@@ -845,7 +845,7 @@ TSharedRef<ISlateStyle> FUMGCoreStyle::Create()
 			.SetMenuDropdownHoveredBorderBrush( BOX_BRUSH( "Common/ColumnHeaderMenuButton_Hovered", 4.f/32.f ) );
 		Style->Set( "TableView.Header.Column", TableColumnHeaderStyle );
 
-		const FTableColumnHeaderStyle TableLastColumnHeaderStyle = FTableColumnHeaderStyle()
+		static const FTableColumnHeaderStyle TableLastColumnHeaderStyle = FTableColumnHeaderStyle()
 			.SetSortPrimaryAscendingImage(IMAGE_BRUSH("Common/SortUpArrow", Icon8x4))
 			.SetSortPrimaryDescendingImage(IMAGE_BRUSH("Common/SortDownArrow", Icon8x4))
 			.SetSortSecondaryAscendingImage(IMAGE_BRUSH("Common/SortUpArrows", Icon16x4))
@@ -856,7 +856,7 @@ TSharedRef<ISlateStyle> FUMGCoreStyle::Create()
 			.SetMenuDropdownNormalBorderBrush( BOX_BRUSH( "Common/ColumnHeaderMenuButton_Normal", 4.f/32.f ) )
 			.SetMenuDropdownHoveredBorderBrush( BOX_BRUSH( "Common/ColumnHeaderMenuButton_Hovered", 4.f/32.f ) );
 
-		const FSplitterStyle TableHeaderSplitterStyle = FSplitterStyle()
+		static const FSplitterStyle TableHeaderSplitterStyle = FSplitterStyle()
 			.SetHandleNormalBrush( FSlateNoResource() )
 			.SetHandleHighlightBrush( IMAGE_BRUSH( "Common/HeaderSplitterGrip", Icon8x8 ) );
 
@@ -910,7 +910,7 @@ TSharedRef<ISlateStyle> FUMGCoreStyle::Create()
 			.SetColorAndOpacity( FLinearColor( 0.4f, 0.4, 0.4f, 1.0f ) ) );
 
 		/* Create style for "ToolBar.CheckBox" widget ... */
-		const FCheckBoxStyle ToolBarCheckBoxStyle = FCheckBoxStyle()
+		static const FCheckBoxStyle ToolBarCheckBoxStyle = FCheckBoxStyle()
 			.SetUncheckedImage( IMAGE_BRUSH( "Common/SmallCheckBox", Icon14x14 ) )
 			.SetCheckedImage( IMAGE_BRUSH( "Common/SmallCheckBox_Checked", Icon14x14 ) )
 			.SetUncheckedHoveredImage( IMAGE_BRUSH( "Common/SmallCheckBox_Hovered", Icon14x14 ) )
@@ -925,7 +925,7 @@ TSharedRef<ISlateStyle> FUMGCoreStyle::Create()
 
 		/* Read-only checkbox that appears next to a menu item */
 		/* Set images for various SCheckBox states associated with read-only toolbar check box items... */
-		const FCheckBoxStyle BasicToolBarCheckStyle = FCheckBoxStyle()
+		static const FCheckBoxStyle BasicToolBarCheckStyle = FCheckBoxStyle()
 			.SetUncheckedImage(IMAGE_BRUSH("Icons/Empty_14x", Icon14x14))
 			.SetUncheckedHoveredImage(IMAGE_BRUSH("Icons/Empty_14x", Icon14x14))
 			.SetUncheckedPressedImage(IMAGE_BRUSH("Common/SmallCheckBox_Hovered", Icon14x14))
@@ -939,7 +939,7 @@ TSharedRef<ISlateStyle> FUMGCoreStyle::Create()
 
 		// This radio button is actually just a check box with different images
 		/* Create style for "ToolBar.RadioButton" widget ... */
-		const FCheckBoxStyle ToolbarRadioButtonCheckBoxStyle = FCheckBoxStyle()
+		static const FCheckBoxStyle ToolbarRadioButtonCheckBoxStyle = FCheckBoxStyle()
 				.SetUncheckedImage( IMAGE_BRUSH( "Common/RadioButton_Unselected_16x", Icon16x16 ) )
 				.SetCheckedImage( IMAGE_BRUSH( "Common/RadioButton_Selected_16x", Icon16x16 ) )
 				.SetUncheckedHoveredImage( IMAGE_BRUSH( "Common/RadioButton_Unselected_16x", Icon16x16, SelectionColor ) )
@@ -950,7 +950,7 @@ TSharedRef<ISlateStyle> FUMGCoreStyle::Create()
 		Style->Set( "ToolBar.RadioButton", ToolbarRadioButtonCheckBoxStyle );
 
 		/* Create style for "ToolBar.ToggleButton" widget ... */
-		const FCheckBoxStyle ToolBarToggleButtonCheckBoxStyle = FCheckBoxStyle()
+		static const FCheckBoxStyle ToolBarToggleButtonCheckBoxStyle = FCheckBoxStyle()
 				.SetCheckBoxType( ESlateCheckBoxType::ToggleButton )
 				.SetUncheckedImage( FSlateNoResource() )
 				.SetUncheckedPressedImage( BOX_BRUSH( "Common/RoundedSelection_16x", 4.0f/16.0f, SelectionColor_Pressed ) )
@@ -1011,11 +1011,11 @@ TSharedRef<ISlateStyle> FUMGCoreStyle::Create()
 			.SetFont( DEFAULT_FONT( "Regular", 8 ) )
 			.SetColorAndOpacity( FLinearColor( 0.4f, 0.4, 0.4f, 1.0f ) ) );
 
-		const FMargin MenuBlockPadding(10.0f, 3.0f, 5.0f, 2.0f);
+		static const FMargin MenuBlockPadding(10.0f, 3.0f, 5.0f, 2.0f);
 		Style->Set("Menu.Heading.Padding", MenuBlockPadding + FMargin(0, 10, 0, 0));
 
 		/* Set images for various SCheckBox states associated with menu check box items... */
-		const FCheckBoxStyle BasicMenuCheckBoxStyle = FCheckBoxStyle()
+		static const FCheckBoxStyle BasicMenuCheckBoxStyle = FCheckBoxStyle()
 			.SetUncheckedImage( IMAGE_BRUSH( "Common/SmallCheckBox", Icon14x14 ) )
 			.SetUncheckedHoveredImage( IMAGE_BRUSH( "Common/SmallCheckBox_Hovered", Icon14x14 ) )
 			.SetUncheckedPressedImage( IMAGE_BRUSH( "Common/SmallCheckBox_Hovered", Icon14x14, FLinearColor( 0.5f, 0.5f, 0.5f ) ) )
@@ -1031,7 +1031,7 @@ TSharedRef<ISlateStyle> FUMGCoreStyle::Create()
 						
 		/* Read-only checkbox that appears next to a menu item */
 		/* Set images for various SCheckBox states associated with read-only menu check box items... */
-		const FCheckBoxStyle BasicMenuCheckStyle = FCheckBoxStyle()
+		static const FCheckBoxStyle BasicMenuCheckStyle = FCheckBoxStyle()
 			.SetUncheckedImage( IMAGE_BRUSH( "Icons/Empty_14x", Icon14x14 ) )
 			.SetUncheckedHoveredImage( IMAGE_BRUSH( "Icons/Empty_14x", Icon14x14 ) )
 			.SetUncheckedPressedImage( IMAGE_BRUSH( "Common/SmallCheckBox_Hovered", Icon14x14 ) )
@@ -1047,7 +1047,7 @@ TSharedRef<ISlateStyle> FUMGCoreStyle::Create()
 
 		/* This radio button is actually just a check box with different images */
 		/* Set images for various Menu radio button (SCheckBox) states... */
-		const FCheckBoxStyle BasicMenuRadioButtonStyle = FCheckBoxStyle()
+		static const FCheckBoxStyle BasicMenuRadioButtonStyle = FCheckBoxStyle()
 			.SetUncheckedImage( IMAGE_BRUSH( "Common/RadioButton_Unselected_16x", Icon16x16 ) )
 			.SetUncheckedHoveredImage( IMAGE_BRUSH( "Common/RadioButton_Unselected_16x", Icon16x16 ) )
 			.SetUncheckedPressedImage( IMAGE_BRUSH( "Common/RadioButton_Unselected_16x", Icon16x16 ) )
@@ -1062,7 +1062,7 @@ TSharedRef<ISlateStyle> FUMGCoreStyle::Create()
 		Style->Set( "Menu.RadioButton", BasicMenuRadioButtonStyle );
 
 		/* Create style for "Menu.ToggleButton" widget ... */
-		const FCheckBoxStyle MenuToggleButtonCheckBoxStyle = FCheckBoxStyle()
+		static const FCheckBoxStyle MenuToggleButtonCheckBoxStyle = FCheckBoxStyle()
 			.SetCheckBoxType( ESlateCheckBoxType::ToggleButton )
 			.SetUncheckedImage( FSlateNoResource() )
 			.SetUncheckedPressedImage( BOX_BRUSH( "Common/RoundedSelection_16x", 4.0f/16.0f, SelectionColor_Pressed ) )
@@ -1175,9 +1175,9 @@ TSharedRef<ISlateStyle> FUMGCoreStyle::Create()
 			);
 
 		// Flash using the selection color for consistency with the rest of the UI scheme
-		const FSlateColor& TabFlashColor = SelectionColor;
+		static const FSlateColor& TabFlashColor = SelectionColor;
 
-		const FButtonStyle CloseButton = FButtonStyle()
+		static const FButtonStyle CloseButton = FButtonStyle()
 			.SetNormal ( IMAGE_BRUSH( "/Docking/CloseApp_Normal", Icon16x16 ) )
 			.SetPressed( IMAGE_BRUSH( "/Docking/CloseApp_Pressed", Icon16x16 ) )
 			.SetHovered( IMAGE_BRUSH( "/Docking/CloseApp_Hovered", Icon16x16 ) );
@@ -1261,28 +1261,28 @@ TSharedRef<ISlateStyle> FUMGCoreStyle::Create()
 	// SWindow defaults...
 	{
 #if !PLATFORM_MAC
-		const FButtonStyle MinimizeButtonStyle = FButtonStyle(Button)
+		static const FButtonStyle MinimizeButtonStyle = FButtonStyle(Button)
 			.SetNormal (IMAGE_BRUSH("Common/Window/WindowButton_Minimize_Normal", FVector2D(27.0f, 18.0f)))
 			.SetHovered(IMAGE_BRUSH("Common/Window/WindowButton_Minimize_Hovered", FVector2D(27.0f, 18.0f)))
 			.SetPressed(IMAGE_BRUSH("Common/Window/WindowButton_Minimize_Pressed", FVector2D(27.0f, 18.0f)));
 
-		const FButtonStyle MaximizeButtonStyle = FButtonStyle(Button)
+		static const FButtonStyle MaximizeButtonStyle = FButtonStyle(Button)
 			.SetNormal (IMAGE_BRUSH("Common/Window/WindowButton_Maximize_Normal", FVector2D(23.0f, 18.0f)))
 			.SetHovered(IMAGE_BRUSH("Common/Window/WindowButton_Maximize_Hovered", FVector2D(23.0f, 18.0f)))
 			.SetPressed(IMAGE_BRUSH("Common/Window/WindowButton_Maximize_Pressed", FVector2D(23.0f, 18.0f)));
 
-		const FButtonStyle RestoreButtonStyle = FButtonStyle(Button)
+		static const FButtonStyle RestoreButtonStyle = FButtonStyle(Button)
 			.SetNormal (IMAGE_BRUSH("Common/Window/WindowButton_Restore_Normal", FVector2D(23.0f, 18)))
 			.SetHovered(IMAGE_BRUSH("Common/Window/WindowButton_Restore_Hovered", FVector2D(23.0f, 18)))
 			.SetPressed(IMAGE_BRUSH("Common/Window/WindowButton_Restore_Pressed", FVector2D(23.0f, 18)));
 
-		const FButtonStyle CloseButtonStyle = FButtonStyle(Button)
+		static const FButtonStyle CloseButtonStyle = FButtonStyle(Button)
 			.SetNormal (IMAGE_BRUSH("Common/Window/WindowButton_Close_Normal", FVector2D(44.0f, 18.0f)))
 			.SetHovered(IMAGE_BRUSH("Common/Window/WindowButton_Close_Hovered", FVector2D(44.0f, 18.0f)))
 			.SetPressed(IMAGE_BRUSH("Common/Window/WindowButton_Close_Pressed", FVector2D(44.0f, 18.0f)));
 #endif
 
-		const FTextBlockStyle TitleTextStyle = FTextBlockStyle(NormalText)
+		static const FTextBlockStyle TitleTextStyle = FTextBlockStyle(NormalText)
 			.SetFont(DEFAULT_FONT("Regular", 9))
 			.SetColorAndOpacity(FLinearColor::White)
 			.SetShadowOffset(FVector2D(1.0f, 1.0f))
@@ -1359,7 +1359,7 @@ TSharedRef<ISlateStyle> FUMGCoreStyle::Create()
 
 	// Syntax highlighting
 	{
-		const FTextBlockStyle SmallMonospacedText = FTextBlockStyle(MonospacedText)
+		static const FTextBlockStyle SmallMonospacedText = FTextBlockStyle(MonospacedText)
 			.SetFont(DEFAULT_FONT("Mono", 9));
 
 		Style->Set("SyntaxHighlight.Normal", SmallMonospacedText);
