Fix crash slate attribute calling delegate on garbage collected UObject.

diff --git a/Engine/Source/Runtime/SlateCore/Private/Types/SlateAttribute.cpp b/Engine/Source/Runtime/SlateCore/Private/Types/SlateAttribute.cpp
--- a/Engine/Source/Runtime/SlateCore/Private/Types/SlateAttribute.cpp
+++ b/Engine/Source/Runtime/SlateCore/Private/Types/SlateAttribute.cpp
@@ -7,12 +7,15 @@
 #include "Widgets/SWidget.h"
 #include "Debugging/WidgetList.h"
 
+PRAGMA_DISABLE_OPTIMIZATION
+
 namespace SlateAttributePrivate
 {
 	void TestAttributeAddress(const SWidget& OwningWidget, const FSlateAttributeImpl& Attribute, ESlateAttributeType AttributeType)
 	{
 #if STATS && DO_CHECK
-		if (AttributeType == ESlateAttributeType::Member)
+		//@TODO: DarenC - Using allocsize as proxy for 'IsConstructionCompleted' due to encapsulation.
+		if (AttributeType == ESlateAttributeType::Member && OwningWidget.GetAllocSize() > 0)
 		{
 			UPTRINT SlateAttributePtr = (UPTRINT)&Attribute;
 			UPTRINT WidgetPtr = (UPTRINT)&OwningWidget;
@@ -117,3 +120,5 @@ namespace SlateAttributePrivate
 		FSlateAttributeMetaData::MoveAttribute(OwningWidget, *this, AttributeType, Other);
 	}
 }
+
+PRAGMA_ENABLE_OPTIMIZATION
\ No newline at end of file
diff --git a/Engine/Source/Runtime/SlateCore/Public/Types/Attributes/SlateAttributeBase.inl b/Engine/Source/Runtime/SlateCore/Public/Types/Attributes/SlateAttributeBase.inl
--- a/Engine/Source/Runtime/SlateCore/Public/Types/Attributes/SlateAttributeBase.inl
+++ b/Engine/Source/Runtime/SlateCore/Public/Types/Attributes/SlateAttributeBase.inl
@@ -544,14 +544,17 @@ namespace SlateAttributePrivate
 
 			virtual FUpdateAttributeResult UpdateAttribute(const SWidget& Widget) override
 			{
-				ObjectType NewValue = Getter.Execute();
-
-				const bool bIsIdentical = Attribute->IdenticalTo(Widget, Attribute->Value, NewValue);
-				if (!bIsIdentical)
+				if (Getter.IsBound())
 				{
-					// Set the value on the widget
-					Attribute->Value = MoveTemp(NewValue);
-					return Attribute->GetInvalidationReason(Widget);
+					ObjectType NewValue = Getter.Execute();
+
+					const bool bIsIdentical = Attribute->IdenticalTo(Widget, Attribute->Value, NewValue);
+					if (!bIsIdentical)
+					{
+						// Set the value on the widget
+						Attribute->Value = MoveTemp(NewValue);
+						return Attribute->GetInvalidationReason(Widget);
+					}
 				}
 				return FUpdateAttributeResult();
 			}
