Remove optimization pragmas.

diff --git a/Engine/Source/Runtime/SlateCore/Private/Types/SlateAttribute.cpp b/Engine/Source/Runtime/SlateCore/Private/Types/SlateAttribute.cpp
--- a/Engine/Source/Runtime/SlateCore/Private/Types/SlateAttribute.cpp
+++ b/Engine/Source/Runtime/SlateCore/Private/Types/SlateAttribute.cpp
@@ -7,8 +7,6 @@
 #include "Widgets/SWidget.h"
 #include "Debugging/WidgetList.h"
 
-PRAGMA_DISABLE_OPTIMIZATION
-
 namespace SlateAttributePrivate
 {
 	void TestAttributeAddress(const SWidget& OwningWidget, const FSlateAttributeImpl& Attribute, ESlateAttributeType AttributeType)
@@ -120,5 +118,3 @@ namespace SlateAttributePrivate
 		FSlateAttributeMetaData::MoveAttribute(OwningWidget, *this, AttributeType, Other);
 	}
 }
-
-PRAGMA_ENABLE_OPTIMIZATION
\ No newline at end of file
