Slate: Add FSlotArguments as a way to create FSlot in the declarative way, similar to SWidget. It will be used to support TSlateAttribute on FSlot. TAttribute will be in the FSlotArguments and the FSlate will have a TSlateAttribute, similar to SWidget. The previous slot will be deprecated once everything is converted. Convert SOverlay to FSlotArguments.

diff --git a/Engine/Plugins/Experimental/CommonUI/Source/CommonUI/Private/CommonVisibilitySwitcherSlot.cpp b/Engine/Plugins/Experimental/CommonUI/Source/CommonUI/Private/CommonVisibilitySwitcherSlot.cpp
--- a/Engine/Plugins/Experimental/CommonUI/Source/CommonUI/Private/CommonVisibilitySwitcherSlot.cpp
+++ b/Engine/Plugins/Experimental/CommonUI/Source/CommonUI/Private/CommonVisibilitySwitcherSlot.cpp
@@ -14,7 +14,8 @@ UCommonVisibilitySwitcherSlot::UCommonVisibilitySwitcherSlot(const FObjectInitia
 
 void UCommonVisibilitySwitcherSlot::BuildSlot(TSharedRef<SOverlay> Overlay)
 {
-	Slot = &Overlay->AddSlot()
+	Overlay->AddSlot()
+		.Expose(Slot)
 		.Padding(Padding)
 		.HAlign(HorizontalAlignment)
 		.VAlign(VerticalAlignment)
diff --git a/Engine/Source/Editor/LevelEditor/Private/SLevelViewport.cpp b/Engine/Source/Editor/LevelEditor/Private/SLevelViewport.cpp
--- a/Engine/Source/Editor/LevelEditor/Private/SLevelViewport.cpp
+++ b/Engine/Source/Editor/LevelEditor/Private/SLevelViewport.cpp
@@ -4269,53 +4269,55 @@ void SLevelViewport::ShowMouseCaptureLabel(ELabelAnchorMode AnchorMode)
 	EVerticalAlignment VAlign = (EVerticalAlignment)((AnchorMode/3)+1);
 	EHorizontalAlignment HAlign = (EHorizontalAlignment)((AnchorMode%3)+1);
 	
-	SOverlay::FOverlaySlot& Slot = ViewportOverlay->AddSlot();
-	PIEOverlaySlotIndex = Slot.ZOrder;
+	{
+		SOverlay::FScopedWidgetSlotArguments ScopedSlotArgumnet = ViewportOverlay->AddSlot();
+		PIEOverlaySlotIndex = ScopedSlotArgumnet.GetSlot()->GetZOrder();
 
-	Slot.HAlign(HAlign)
-	.VAlign(VAlign)
-	[
-		SNew( SBorder )
-		.BorderImage( FEditorStyle::GetBrush("NoBorder") )
-		.Visibility(this, &SLevelViewport::GetMouseCaptureLabelVisibility)
-		.ColorAndOpacity( this, &SLevelViewport::GetMouseCaptureLabelColorAndOpacity )
-		.ForegroundColor( FLinearColor::White )
-		.Padding(15.0f)
+		ScopedSlotArgumnet.HAlign(HAlign)
+		.VAlign(VAlign)
 		[
-			SNew( SButton )
-			.ButtonStyle( FEditorStyle::Get(), "EditorViewportToolBar.MenuButton" )
-			.IsFocusable(false)
-			.ButtonColorAndOpacity( FSlateColor(FLinearColor::Black) )
+			SNew( SBorder )
+			.BorderImage( FEditorStyle::GetBrush("NoBorder") )
+			.Visibility(this, &SLevelViewport::GetMouseCaptureLabelVisibility)
+			.ColorAndOpacity( this, &SLevelViewport::GetMouseCaptureLabelColorAndOpacity )
 			.ForegroundColor( FLinearColor::White )
+			.Padding(15.0f)
 			[
-				SNew( SHorizontalBox )
-				+ SHorizontalBox::Slot()
-				.MaxWidth(32.f)
-				.VAlign(VAlign_Center)
-				.Padding(0.0f, 2.0f, 2.0f, 2.0f)
+				SNew( SButton )
+				.ButtonStyle( FEditorStyle::Get(), "EditorViewportToolBar.MenuButton" )
+				.IsFocusable(false)
+				.ButtonColorAndOpacity( FSlateColor(FLinearColor::Black) )
+				.ForegroundColor( FLinearColor::White )
 				[
-					SNew( SVerticalBox )
-					+ SVerticalBox::Slot()
-					.MaxHeight(16.f)
+					SNew( SHorizontalBox )
+					+ SHorizontalBox::Slot()
+					.MaxWidth(32.f)
+					.VAlign(VAlign_Center)
+					.Padding(0.0f, 2.0f, 2.0f, 2.0f)
 					[
-						SNew(SImage)
-						.Image(FEditorStyle::GetBrush("LevelViewport.CursorIcon"))
+						SNew( SVerticalBox )
+						+ SVerticalBox::Slot()
+						.MaxHeight(16.f)
+						[
+							SNew(SImage)
+							.Image(FEditorStyle::GetBrush("LevelViewport.CursorIcon"))
+						]
+					]
+					+ SHorizontalBox::Slot()
+					.AutoWidth()
+					.VAlign(VAlign_Center)
+					.HAlign(HAlign_Center)
+					.Padding(2.0f, 2.0f)
+					[
+						SNew(STextBlock)
+						.Text(this, &SLevelViewport::GetMouseCaptureLabelText)
+						.Font(FCoreStyle::GetDefaultFontStyle("Bold", 9))
+						.ColorAndOpacity(FLinearColor::White)
 					]
-				]
-				+ SHorizontalBox::Slot()
-				.AutoWidth()
-				.VAlign(VAlign_Center)
-				.HAlign(HAlign_Center)
-				.Padding(2.0f, 2.0f)
-				[
-					SNew(STextBlock)
-					.Text(this, &SLevelViewport::GetMouseCaptureLabelText)
-					.Font(FCoreStyle::GetDefaultFontStyle("Bold", 9))
-					.ColorAndOpacity(FLinearColor::White)
 				]
 			]
-		]
-	];
+		];
+	}
 }
 
 void SLevelViewport::HideMouseCaptureLabel()
diff --git a/Engine/Source/Editor/UnrealEd/Private/AssetEditorViewportLayout.cpp b/Engine/Source/Editor/UnrealEd/Private/AssetEditorViewportLayout.cpp
--- a/Engine/Source/Editor/UnrealEd/Private/AssetEditorViewportLayout.cpp
+++ b/Engine/Source/Editor/UnrealEd/Private/AssetEditorViewportLayout.cpp
@@ -74,7 +74,7 @@ void SAssetEditorViewportsOverlay::Tick( const FGeometry& AllottedGeometry, cons
 	CachedSize = AllottedGeometry.Size;
 }
 
-SOverlay::FOverlaySlot& SAssetEditorViewportsOverlay::AddSlot()
+SOverlay::FScopedWidgetSlotArguments SAssetEditorViewportsOverlay::AddSlot()
 {
 	return OverlayWidget->AddSlot();
 }
diff --git a/Engine/Source/Editor/UnrealEd/Public/AssetEditorViewportLayout.h b/Engine/Source/Editor/UnrealEd/Public/AssetEditorViewportLayout.h
--- a/Engine/Source/Editor/UnrealEd/Public/AssetEditorViewportLayout.h
+++ b/Engine/Source/Editor/UnrealEd/Public/AssetEditorViewportLayout.h
@@ -86,7 +86,7 @@ public:
 	virtual void Tick(const FGeometry& AllottedGeometry, const double InCurrentTime, const float InDeltaTime) override;
 
 	/** Wraps SOverlay::AddSlot() */
-	SOverlay::FOverlaySlot& AddSlot();
+	SOverlay::FScopedWidgetSlotArguments AddSlot();
 
 	/** Wraps SOverlay::RemoveSlot() */
 	void RemoveSlot();
diff --git a/Engine/Source/Runtime/Slate/Public/Widgets/Views/STableRow.h b/Engine/Source/Runtime/Slate/Public/Widgets/Views/STableRow.h
--- a/Engine/Source/Runtime/Slate/Public/Widgets/Views/STableRow.h
+++ b/Engine/Source/Runtime/Slate/Public/Widgets/Views/STableRow.h
@@ -1368,7 +1368,7 @@ protected:
 					TAttribute<float> WidthBinding;
 					WidthBinding.BindRaw(&Column, &SHeaderRow::FColumn::GetWidth);
 
-					SHorizontalBox::FSlot& NewSlot = Box->AddSlot()
+					Box->AddSlot()
 					.HAlign(Column.CellHAlignment)
 					.VAlign(Column.CellVAlignment)
 					.FillWidth(WidthBinding)
diff --git a/Engine/Source/Runtime/SlateCore/Private/Widgets/SOverlay.cpp b/Engine/Source/Runtime/SlateCore/Private/Widgets/SOverlay.cpp
--- a/Engine/Source/Runtime/SlateCore/Private/Widgets/SOverlay.cpp
+++ b/Engine/Source/Runtime/SlateCore/Private/Widgets/SOverlay.cpp
@@ -14,13 +14,18 @@ SOverlay::SOverlay()
 	bCanSupportFocus = false;
 }
 
+SOverlay::FOverlaySlot::FSlotArguments SOverlay::Slot()
+{
+	return FOverlaySlot::FSlotArguments(MakeUnique<FOverlaySlot>());
+}
+
 void SOverlay::Construct( const SOverlay::FArguments& InArgs )
 {
-	const int32 NumSlots = InArgs.Slots.Num();
-	for ( int32 SlotIndex = 0; SlotIndex < NumSlots; ++SlotIndex )
-	{
-		Children.Add( InArgs.Slots[SlotIndex] );
-	}
+	// Because InArgs is const&, we need to do some bad cast here. That would need to change in the future.
+	//The Slot has a unique_ptr. It can't be copied.
+	//Previously, the Children.Add(), was wrong if we added the same slot twice (it would create a lot of issues).
+	//Because of that, it doesn't matter if we steal the slot from the FArguments and it enforces that a slot cannot be added twice.
+	Children.AddSlots(MoveTemp(const_cast<TArray<FOverlaySlot::FSlotArguments>&>(InArgs._Slots)));
 }
 
 void SOverlay::OnArrangeChildren( const FGeometry& AllottedGeometry, FArrangedChildren& ArrangedChildren ) const
@@ -118,18 +123,16 @@ int32 SOverlay::OnPaint( const FPaintArgs& Args, const FGeometry& AllottedGeomet
 	return MaxLayerId;
 }
 
-SOverlay::FOverlaySlot& SOverlay::AddSlot( int32 ZOrder )
+SOverlay::FScopedWidgetSlotArguments SOverlay::AddSlot( int32 ZOrder )
 {
-	FOverlaySlot& NewSlot = *new FOverlaySlot();
+	int32 Index = INDEX_NONE;
 	if ( ZOrder == INDEX_NONE )
 	{
 		// No ZOrder was specified; just add to the end of the list.
 		// Use a ZOrder index one after the last elements.
 		ZOrder = (Children.Num() == 0)
 			? 0
-			: ( Children[ Children.Num()-1 ].ZOrder + 1 );
-
-		this->Children.Add( &NewSlot );
+			: ( Children[ Children.Num()-1 ].GetZOrder() + 1 );
 	}
 	else
 	{
@@ -139,7 +142,7 @@ SOverlay::FOverlaySlot& SOverlay::AddSlot( int32 ZOrder )
 		for( ; CurSlotIndex < Children.Num(); ++CurSlotIndex )
 		{
 			const FOverlaySlot& CurSlot = Children[ CurSlotIndex ];
-			if( ZOrder < CurSlot.ZOrder )
+			if( ZOrder < CurSlot.GetZOrder() )
 			{
 				// Insert before
 				bFoundSlot = true;
@@ -148,13 +151,12 @@ SOverlay::FOverlaySlot& SOverlay::AddSlot( int32 ZOrder )
 		}
 
 		// Add a slot at the desired location
-		this->Children.Insert( &NewSlot, CurSlotIndex );
+		Index = CurSlotIndex;
 	}
 
-	Invalidate(EInvalidateWidgetReason::Layout);
-
-	NewSlot.ZOrder = ZOrder;
-	return NewSlot;
+	FScopedWidgetSlotArguments Result {MakeUnique<FOverlaySlot>(), this->Children, Index};
+	Result.ZOrder(ZOrder);
+	return MoveTemp(Result);
 }
 
 void SOverlay::RemoveSlot( int32 ZOrder )
@@ -163,7 +165,7 @@ void SOverlay::RemoveSlot( int32 ZOrder )
 	{
 		for( int32 ChildIndex=0; ChildIndex < Children.Num(); ++ChildIndex )
 		{
-			if ( Children[ChildIndex].ZOrder == ZOrder )
+			if ( Children[ChildIndex].GetZOrder() == ZOrder )
 			{
 				Children.RemoveAt( ChildIndex );
 				return;
diff --git a/Engine/Source/Runtime/SlateCore/Private/Widgets/SWindow.cpp b/Engine/Source/Runtime/SlateCore/Private/Widgets/SWindow.cpp
--- a/Engine/Source/Runtime/SlateCore/Private/Widgets/SWindow.cpp
+++ b/Engine/Source/Runtime/SlateCore/Private/Widgets/SWindow.cpp
@@ -1293,7 +1293,7 @@ bool SWindow::HasOverlay() const
 	return WindowOverlay.IsValid();
 }
 
-SOverlay::FOverlaySlot& SWindow::AddOverlaySlot( const int32 ZOrder )
+SOverlay::FScopedWidgetSlotArguments SWindow::AddOverlaySlot( const int32 ZOrder )
 {
 	if(!WindowOverlay.IsValid())
 	{
diff --git a/Engine/Source/Runtime/SlateCore/Public/Layout/BasicLayoutWidgetSlot.h b/Engine/Source/Runtime/SlateCore/Public/Layout/BasicLayoutWidgetSlot.h
--- a/Engine/Source/Runtime/SlateCore/Public/Layout/BasicLayoutWidgetSlot.h
+++ b/Engine/Source/Runtime/SlateCore/Public/Layout/BasicLayoutWidgetSlot.h
@@ -7,6 +7,7 @@
 #include "Layout/FlowDirection.h"
 #include "Layout/Margin.h"
 #include "Misc/Optional.h"
+#include "Widgets/DeclarativeSyntaxSupport.h"
 
 PRAGMA_DISABLE_DEPRECATION_WARNINGS
 /** Mixin to add the alignment functionality to a base slot. */
@@ -23,14 +24,47 @@ public:
 		: HAlignment(InHAlign)
 		, VAlignment(InVAlign)
 	{}
+	
+public:
+	struct FSlotArgumentsMixin
+	{
+	private:
+		friend TAlignmentWidgetSlotMixin;
+
+	public:
+		typename MixedIntoType::FSlotArguments& HAlign(EHorizontalAlignment InHAlignment)
+		{
+			_HAlignment = InHAlignment;
+			return static_cast<typename MixedIntoType::FSlotArguments&>(*this);
+		}
+
+		typename MixedIntoType::FSlotArguments& VAlign(EVerticalAlignment InVAlignment)
+		{
+			_VAlignment = InVAlignment;
+			return static_cast<typename MixedIntoType::FSlotArguments&>(*this);
+		}
+
+	private:
+		TOptional<EHorizontalAlignment> _HAlignment;
+		TOptional<EVerticalAlignment> _VAlignment;
+	};
+
+protected:
+	void ConstructMixin(const FChildren& SlotOwner, FSlotArgumentsMixin&& InArgs)
+	{
+		HAlignment = InArgs._HAlignment.Get(HAlignment);
+		VAlignment = InArgs._VAlignment.Get(VAlignment);
+	}
 
 public:
+	// HAlign will be deprecated soon. Use SetVerticalAlignment, or if you are using ChildSlot, construct a new slot with FSlotArguments
 	MixedIntoType& HAlign(EHorizontalAlignment InHAlignment)
 	{
 		HAlignment = InHAlignment;
 		return *(static_cast<MixedIntoType*>(this));
 	}
 
+	// VAlign will be deprecated soon. Use SetVerticalAlignment, or if you are using ChildSlot, construct a new slot with FSlotArguments
 	MixedIntoType& VAlign(EVerticalAlignment InVAlignment)
 	{
 		VAlignment = InVAlignment;
@@ -47,6 +81,11 @@ public:
 		}
 	}
 
+	EHorizontalAlignment GetHorizontalAlignment() const
+	{
+		return HAlignment;
+	}
+
 	void SetVerticalAlignment(EVerticalAlignment Alignment)
 	{
 		if (VAlignment != Alignment)
@@ -56,8 +95,10 @@ public:
 		}
 	}
 
-	EHorizontalAlignment GetHorizontalAlignment() const { return HAlignment; }
-	EVerticalAlignment GetVerticalAlignment() const { return VAlignment; }
+	EVerticalAlignment GetVerticalAlignment() const
+	{
+		return VAlignment;
+	}
 
 public:
 	/** Horizontal positioning of child within the allocated slot */
@@ -82,24 +123,70 @@ public:
 	{}
 
 public:
+	struct FSlotArgumentsMixin
+	{
+	private:
+		friend TPaddingWidgetSlotMixin;
+		TAttribute<FMargin> _Padding;
+
+	public:
+		typename MixedIntoType::FSlotArguments& Padding(TAttribute<FMargin> InPadding)
+		{
+			_Padding = MoveTemp(InPadding);
+			return static_cast<typename MixedIntoType::FSlotArguments&>(*this);
+		}
+
+		typename MixedIntoType::FSlotArguments& Padding(float Uniform)
+		{
+			_Padding = FMargin(Uniform);
+			return static_cast<typename MixedIntoType::FSlotArguments&>(*this);
+		}
+
+		typename MixedIntoType::FSlotArguments& Padding(float Horizontal, float Vertical)
+		{
+			_Padding = FMargin(Horizontal, Vertical);
+			return static_cast<typename MixedIntoType::FSlotArguments&>(*this);
+		}
+
+		typename MixedIntoType::FSlotArguments& Padding(float Left, float Top, float Right, float Bottom)
+		{
+			_Padding = FMargin(Left, Top, Right, Bottom);
+			return static_cast<typename MixedIntoType::FSlotArguments&>(*this);
+		}
+	};
+
+protected:
+	void ConstructMixin(const FChildren& SlotOwner, FSlotArgumentsMixin&& InArgs)
+	{
+		if (InArgs._Padding.IsSet())
+		{
+			SlotPadding = MoveTemp(InArgs._Padding);
+		}
+	}
+
+public:
+	// Padding will be deprecated soon. Use SetPadding, or if you are using ChildSlot, construct a new slot with FSlotArguments
 	MixedIntoType& Padding(TAttribute<FMargin> InPadding)
 	{
 		SlotPadding = MoveTemp(InPadding);
 		return *(static_cast<MixedIntoType*>(this));
 	}
 
+	// Padding will be deprecated soon. Use SetPadding, or if you are using ChildSlot, construct a new slot with FSlotArguments
 	MixedIntoType& Padding(float Uniform)
 	{
 		SlotPadding = FMargin(Uniform);
 		return *(static_cast<MixedIntoType*>(this));
 	}
 
+	// Padding will be deprecated soon. Use SetPadding, or if you are using ChildSlot, construct a new slot with FSlotArguments
 	MixedIntoType& Padding(float Horizontal, float Vertical)
 	{
 		SlotPadding = FMargin(Horizontal, Vertical);
 		return *(static_cast<MixedIntoType*>(this));
 	}
 
+	// Padding will be deprecated soon. Use SetPadding, or if you are using ChildSlot, construct a new slot with FSlotArguments
 	MixedIntoType& Padding(float Left, float Top, float Right, float Bottom)
 	{
 		SlotPadding = FMargin(Left, Top, Right, Bottom);
@@ -154,6 +241,17 @@ public:
 		, TAlignmentWidgetSlotMixin<SlotType>(InHAlign, InVAlign)
 	{
 	}
+
+public:
+	SLATE_SLOT_BEGIN_ARGS_TwoMixins(TBasicLayoutWidgetSlot, TSlotBase<SlotType>, TPaddingWidgetSlotMixin<SlotType>, TAlignmentWidgetSlotMixin<SlotType>)
+	SLATE_SLOT_END_ARGS()
+
+	void Construct(const FChildren& SlotOwner, FSlotArguments&& InArgs)
+	{
+		TPaddingWidgetSlotMixin<SlotType>::ConstructMixin(SlotOwner, MoveTemp(InArgs));
+		TAlignmentWidgetSlotMixin<SlotType>::ConstructMixin(SlotOwner, MoveTemp(InArgs));
+		TSlotBase<SlotType>::Construct(SlotOwner, MoveTemp(InArgs));
+	}
 };
 
 
diff --git a/Engine/Source/Runtime/SlateCore/Public/Layout/Children.h b/Engine/Source/Runtime/SlateCore/Public/Layout/Children.h
--- a/Engine/Source/Runtime/SlateCore/Public/Layout/Children.h
+++ b/Engine/Source/Runtime/SlateCore/Public/Layout/Children.h
@@ -316,6 +316,28 @@ public:
 		return this->GetWidget();
 	}
 
+public:
+	struct FSlotArguments : protected TSlotBase<SlotType>::FSlotArguments
+	{
+		FSlotArguments()
+			: TSlotBase<SlotType>::FSlotArguments(TSlotBase<SlotType>::ConstructSlotIsFChildren)
+		{
+		}
+
+		friend TSingleWidgetChildrenWithSlot;
+
+		typename SlotType::FSlotArguments& operator[](const TSharedRef<SWidget>& InChildWidget)
+		{
+			TSlotBase<SlotType>::FSlotArguments::AttachWidget(InChildWidget);
+			return static_cast<typename SlotType::FSlotArguments&>(*this);
+		}
+	};
+
+	void Construct(FSlotArguments&& InArgs)
+	{
+		TSlotBase<SlotType>::Construct(*this, MoveTemp(InArgs));
+	}
+
 public:
 	TSlotBase<SlotType>& AsSlot() { return *this; }
 	const TSlotBase<SlotType>& AsSlot() const { return *this; }
@@ -376,6 +398,17 @@ public:
 
 	FSingleWidgetChildrenWithBasicLayoutSlot(std::nullptr_t) = delete;
 	FSingleWidgetChildrenWithBasicLayoutSlot(std::nullptr_t, const EHorizontalAlignment InHAlign, const EVerticalAlignment InVAlign) = delete;
+
+public:
+	SLATE_SLOT_BEGIN_ARGS_TwoMixins(FSingleWidgetChildrenWithBasicLayoutSlot, TSingleWidgetChildrenWithSlot<FSingleWidgetChildrenWithBasicLayoutSlot>, TPaddingWidgetSlotMixin<FSingleWidgetChildrenWithBasicLayoutSlot>, TAlignmentWidgetSlotMixin<FSingleWidgetChildrenWithBasicLayoutSlot>)
+	SLATE_SLOT_END_ARGS()
+
+	void Construct(FSlotArguments&& InArgs)
+	{
+		TSingleWidgetChildrenWithSlot<FSingleWidgetChildrenWithBasicLayoutSlot>::Construct(MoveTemp(InArgs));
+		TPaddingWidgetSlotMixin<FSingleWidgetChildrenWithBasicLayoutSlot>::ConstructMixin(*this, MoveTemp(InArgs));
+		TAlignmentWidgetSlotMixin<FSingleWidgetChildrenWithBasicLayoutSlot>::ConstructMixin(*this, MoveTemp(InArgs));
+	}
 };
 
 
@@ -431,6 +464,25 @@ public:
 		return Index;
 	}
 
+	int32 AddSlot(typename SlotType::FSlotArguments&& SlotArgument)
+	{
+		TUniquePtr<SlotType> NewSlot = SlotArgument.StealSlot();
+		check(NewSlot.Get());
+		int32 Result = Children.Add(MoveTemp(NewSlot));
+		Children[Result]->Construct(*this, MoveTemp(SlotArgument));
+		return Result;
+	}
+
+	void AddSlots(TArray<typename SlotType::FSlotArguments> SlotArguments)
+	{
+		Children.Reserve(Children.Num() + SlotArguments.Num());
+		for (typename SlotType::FSlotArguments& Arg : SlotArguments)
+		{
+			AddSlot(MoveTemp(Arg));
+		}
+		SlotArguments.Reset();
+	}
+
 	void RemoveAt( int32 Index )
 	{
 		// NOTE:
@@ -483,6 +535,14 @@ public:
 		Slot->SetOwner(*this);
 	}
 
+	void InsertSlot(typename SlotType::FSlotArguments&& SlotArgument, int32 Index)
+	{
+		TUniquePtr<SlotType> NewSlot = SlotArgument.StealSlot();
+		check(NewSlot.Get());
+		Children.Insert(MoveTemp(NewSlot), Index);
+		Children[Index]->Construct(*this, MoveTemp(SlotArgument));
+	}
+
 	void Move(int32 IndexToMove, int32 IndexToDestination)
 	{
 		// @todo this is going to cause a problem for draw ordering
@@ -534,8 +594,56 @@ public:
 		Children.Swap(IndexA, IndexB);
 		GetOwner().Invalidate(EInvalidateWidgetReason::ChildOrder);
 	}
-};
 
+public:
+	/** At the end of the scope a slot will be constructed and added to the FChildren. */
+	struct FScopedWidgetSlotArguments : public SlotType::FSlotArguments
+	{
+	public:
+		FScopedWidgetSlotArguments(TUniquePtr<SlotType> InSlot, TPanelChildren<SlotType>& InChildren, int32 InIndex)
+			: SlotType::FSlotArguments(MoveTemp(InSlot))
+			, Children(InChildren)
+			, Index(InIndex)
+		{
+		}
+		FScopedWidgetSlotArguments(TUniquePtr<SlotType> InSlot, TPanelChildren<SlotType>& InChildren, int32 InIndex, TFunction<void(const SlotType*, int32)> OnAdded)
+			: SlotType::FSlotArguments(MoveTemp(InSlot))
+			, Children(InChildren)
+			, Index(InIndex)
+			, Added(OnAdded)
+		{
+		}
+	
+		FScopedWidgetSlotArguments(const FScopedWidgetSlotArguments&) = delete;
+		FScopedWidgetSlotArguments& operator=(const FScopedWidgetSlotArguments&) = delete;
+		FScopedWidgetSlotArguments(FScopedWidgetSlotArguments&&) = default;
+		FScopedWidgetSlotArguments& operator=(FScopedWidgetSlotArguments&&) = default;
+	
+		virtual ~FScopedWidgetSlotArguments()
+		{
+			if (const SlotType* Slot = this->GetSlot())	// Is nullptr when the FScopedWidgetSlotArguments was moved-constructed.
+			{
+				if (Index == INDEX_NONE)
+				{
+					Index = Children.AddSlot(MoveTemp(*this));
+				}
+				else
+				{
+					Children.InsertSlot(MoveTemp(*this), Index);
+				}
+				if (Added)
+				{
+					Added(Slot, Index);
+				}
+			}
+		}
+
+	private:
+		TPanelChildren<SlotType>& Children;
+		int32 Index;
+		TFunction<void(const SlotType*, int32)> Added;
+	};
+};
 
 
 template<typename SlotType>
diff --git a/Engine/Source/Runtime/SlateCore/Public/SlotBase.h b/Engine/Source/Runtime/SlateCore/Public/SlotBase.h
--- a/Engine/Source/Runtime/SlateCore/Public/SlotBase.h
+++ b/Engine/Source/Runtime/SlateCore/Public/SlotBase.h
@@ -6,21 +6,25 @@
 #include "Misc/Attribute.h"
 #include "Widgets/InvalidateWidgetReason.h"
 
-class SWidget;
 class FChildren;
+class SWidget;
 
 /** Slot are a container of a SWidget used by the FChildren. */
 class SLATECORE_API FSlotBase
 {
 public:
 	FSlotBase();
-	FSlotBase(const FChildren& Children);
+	FSlotBase(const FChildren& InParent);
 	FSlotBase(const TSharedRef<SWidget>& InWidget);
 	FSlotBase& operator=(const FSlotBase&) = delete;
 	FSlotBase(const FSlotBase&) = delete;
 
 	virtual ~FSlotBase();
 
+public:
+	struct FSlotArguments {};
+
+public:
 	UE_DEPRECATED(5.0, "AttachWidgetParent is not used anymore. Use get SetOwner.")
 	void AttachWidgetParent(SWidget* InParent) { }
 
@@ -36,9 +40,13 @@ public:
 	 */
 	SWidget* GetOwnerWidget() const;
 
-	/** Set the owner of the slot. */
+	/**
+	 * Set the owner of the slot.
+	 * Slots cannot be reassigned to different parents.
+	 */
 	void SetOwner(const FChildren& Children);
 
+	/** Attach the child widget the slot now owns. */
 	FORCEINLINE_DEBUGGABLE void AttachWidget( const TSharedRef<SWidget>& InWidget )
 	{
 		// TODO: If we don't hold a reference here, ~SWidget() could called on the old widget before the assignment takes place
@@ -54,7 +62,10 @@ public:
 	 * There will always be a widget in the slot; sometimes it is
 	 * the SNullWidget instance.
 	 */
-	FORCEINLINE_DEBUGGABLE const TSharedRef<SWidget>& GetWidget() const { return Widget; }
+	FORCEINLINE_DEBUGGABLE const TSharedRef<SWidget>& GetWidget() const
+	{
+		return Widget;
+	}
 
 	/**
 	 * Remove the widget from its current slot.
@@ -63,6 +74,7 @@ public:
 	 */
 	const TSharedPtr<SWidget> DetachWidget();
 
+	/** Invalidate the widget's owner. */
 	void Invalidate(EInvalidateWidgetReason InvalidateReason);
 
 protected:
@@ -112,7 +124,7 @@ protected:
 
 
 /** A slot that can be used by the declarative syntax. */
-template<typename SlotType>
+template <typename SlotType>
 class TSlotBase : public FSlotBase
 {
 public:
@@ -124,9 +136,82 @@ public:
 		return static_cast<SlotType&>(*this);
 	}
 
-	SlotType& Expose( SlotType*& OutVarToInit )
+	SlotType& Expose(SlotType*& OutVarToInit)
 	{
 		OutVarToInit = static_cast<SlotType*>(this);
 		return static_cast<SlotType&>(*this);
 	}
+
+
+	/** Argument to indicate the Slot is also its owner. */
+	enum EConstructSlotIsFChildren { ConstructSlotIsFChildren };
+
+	/** Struct to construct a slot. */
+	struct FSlotArguments : public FSlotBase::FSlotArguments
+	{
+	public:
+		FSlotArguments(EConstructSlotIsFChildren) {}
+		FSlotArguments(TUniquePtr<SlotType> InSlot)
+			: Slot(MoveTemp(InSlot))
+		{
+			check(Slot.Get());
+		}
+		FSlotArguments(const FSlotArguments&) = delete;
+		FSlotArguments& operator=(const FSlotArguments&) = delete;
+		FSlotArguments(FSlotArguments&&) = default;
+		FSlotArguments& operator=(FSlotArguments&&) = default;
+
+
+	public:
+		/** Attach the child widget the slot will own. */
+		typename SlotType::FSlotArguments& operator[](const TSharedRef<SWidget>& InChildWidget)
+		{
+			ChildWidget = InChildWidget;
+			return Me();
+		}
+
+		/** Initialize OutVarToInit with the slot that is being constructed. */
+		typename SlotType::FSlotArguments& Expose(SlotType*& OutVarToInit)
+		{
+			OutVarToInit = Slot.Get();
+			return Me();
+		}
+
+		/** Attach the child widget the slot will own. */
+		void AttachWidget(const TSharedRef<SWidget>& InChildWidget)
+		{
+			ChildWidget = InChildWidget;
+		}
+
+		/** @return the child widget that will be owned by the slot. */
+		TSharedPtr<SWidget> GetAttachedWidget() const { return ChildWidget; }
+
+		/** @return the slot that is being constructed. */
+		SlotType* GetSlot() const { return Slot.Get(); }
+
+		/** Steal the slot that is being constructed from the FSlotArguments. */
+		TUniquePtr<SlotType> StealSlot()
+		{
+			return MoveTemp(Slot);
+		}
+
+		/** Used by the named argument pattern as a safe way to 'return *this' for call-chaining purposes. */
+		typename SlotType::FSlotArguments& Me()
+		{
+			return static_cast<typename SlotType::FSlotArguments&>(*this);
+		}
+
+	private:
+		TUniquePtr<SlotType> Slot;
+		TSharedPtr<SWidget> ChildWidget;
+	};
+
+	void Construct(const FChildren& SlotOwner, FSlotArguments&& InArgs)
+	{
+		if (InArgs.GetAttachedWidget())
+		{
+			AttachWidget(InArgs.GetAttachedWidget().ToSharedRef());
+		}
+		SetOwner(SlotOwner);
+	}
 };
diff --git a/Engine/Source/Runtime/SlateCore/Public/Widgets/DeclarativeSyntaxSupport.h b/Engine/Source/Runtime/SlateCore/Public/Widgets/DeclarativeSyntaxSupport.h
--- a/Engine/Source/Runtime/SlateCore/Public/Widgets/DeclarativeSyntaxSupport.h
+++ b/Engine/Source/Runtime/SlateCore/Public/Widgets/DeclarativeSyntaxSupport.h
@@ -37,6 +37,9 @@ template<typename WidgetType> struct TSlateBaseNamedArgs;
 #define SAssignNew( ExposeAs, WidgetType, ... ) \
 	MakeTDecl<WidgetType>( #WidgetType, __FILE__, __LINE__, RequiredArgs::MakeRequiredArgs(__VA_ARGS__) ) . Expose( ExposeAs ) <<= TYPENAME_OUTSIDE_TEMPLATE WidgetType::FArguments()
 
+#define SArgumentNew( InArgs, WidgetType, ... ) \
+	MakeTDecl<WidgetType>( #WidgetType, __FILE__, __LINE__, RequiredArgs::MakeRequiredArgs(__VA_ARGS__) ) <<= TYPENAME_OUTSIDE_TEMPLATE InArgs
+
 
 /**
  * Widget authors can use SLATE_BEGIN_ARGS and SLATE_END_ARS to add support
@@ -113,12 +116,12 @@ template<typename WidgetType> struct TSlateBaseNamedArgs;
 		TAttribute< AttrType > _##AttrName
 
 #define SLATE_PRIVATE_ATTRIBUTE_FUNCTION( AttrType, AttrName ) \
-		WidgetArgsType& AttrName( const TAttribute< AttrType >& InAttribute ) \
+		WidgetArgsType& AttrName( TAttribute< AttrType > InAttribute ) \
 		{ \
-			_##AttrName = InAttribute; \
+			_##AttrName = MoveTemp(InAttribute); \
 			return this->Me(); \
 		} \
-	\
+		\
 		/* Bind attribute with delegate to a global function
 		 * NOTE: We use a template here to avoid 'typename' issues when hosting attributes inside templated classes */ \
 		template< typename StaticFuncPtr > \
@@ -400,6 +403,57 @@ template<typename WidgetType> struct TSlateBaseNamedArgs;
 		}
 
 
+#define SLATE_PRIVATE_SLOT_BEGIN_ARGS( SlotType, SlotParentType ) \
+	{ \
+		using WidgetArgsType = SlotType::FSlotArguments; \
+		using SlotParentType::FSlotArguments::FSlotArguments;
+
+/**
+ * Use this macro between SLATE_BEGIN_ARGS and SLATE_END_ARGS
+ * in order to add support for slots with the construct pattern.
+ */
+#define SLATE_SLOT_ARGUMENT( SlotType, SlotName ) \
+		TArray<SlotType::FSlotArguments> _##SlotName; \
+		WidgetArgsType& operator + (typename SlotType::FSlotArguments& SlotToAdd) \
+		{ \
+			_##SlotName.Add( MoveTemp(SlotToAdd) ); \
+			return *this; \
+		} \
+		WidgetArgsType& operator + (typename SlotType::FSlotArguments&& SlotToAdd) \
+		{ \
+			_##SlotName.Add( MoveTemp(SlotToAdd) ); \
+			return *this; \
+		}
+
+#define SLATE_SLOT_BEGIN_ARGS( SlotType, SlotParentType ) \
+	public: \
+	struct FSlotArguments : public SlotParentType::FSlotArguments \
+	SLATE_PRIVATE_SLOT_BEGIN_ARGS( SlotType, SlotParentType )
+		
+#define SLATE_SLOT_BEGIN_ARGS_OneMixin( SlotType, SlotParentType, Mixin1 ) \
+	public: \
+	struct FSlotArguments : public SlotParentType::FSlotArguments, public Mixin1::FSlotArgumentsMixin \
+	SLATE_PRIVATE_SLOT_BEGIN_ARGS( SlotType, SlotParentType )
+
+#define SLATE_SLOT_BEGIN_ARGS_TwoMixins( SlotType, SlotParentType, Mixin1, Mixin2 ) \
+	public: \
+	struct FSlotArguments : public SlotParentType::FSlotArguments, public Mixin1::FSlotArgumentsMixin, public Mixin2::FSlotArgumentsMixin \
+	SLATE_PRIVATE_SLOT_BEGIN_ARGS( SlotType, SlotParentType )
+
+#define SLATE_SLOT_BEGIN_ARGS_ThreeMixins( SlotType, SlotParentType, Mixin1, Mixin2, Mixin3 ) \
+	public: \
+	struct FSlotArguments : public SlotParentType::FSlotArguments, public Mixin1::FSlotArgumentsMixin, public Mixin2::FSlotArgumentsMixin, public Mixin3::FSlotArgumentsMixin \
+	SLATE_PRIVATE_SLOT_BEGIN_ARGS( SlotType, SlotParentType )
+	
+#define SLATE_SLOT_BEGIN_ARGS_FourMixins( SlotType, SlotParentType, Mixin1, Mixin2, Mixin3, Mixin4 ) \
+	public: \
+	struct FSlotArguments : public SlotParentType::FSlotArguments, public Mixin1::FSlotArgumentsMixin, public Mixin2::FSlotArgumentsMixin, public Mixin3::FSlotArgumentsMixin, public Mixin4::FSlotArgumentsMixin \
+	SLATE_PRIVATE_SLOT_BEGIN_ARGS( SlotType, SlotParentType )
+
+#define SLATE_SLOT_END_ARGS() \
+	};
+
+
 /** A widget reference that is always a valid pointer; defaults to SNullWidget */
 struct TAlwaysValidWidget
 {
diff --git a/Engine/Source/Runtime/SlateCore/Public/Widgets/SOverlay.h b/Engine/Source/Runtime/SlateCore/Public/Widgets/SOverlay.h
--- a/Engine/Source/Runtime/SlateCore/Public/Widgets/SOverlay.h
+++ b/Engine/Source/Runtime/SlateCore/Public/Widgets/SOverlay.h
@@ -37,6 +37,7 @@ class SLATECORE_API SOverlay
 public:	
 
 	/** A slot that support alignment of content and padding and z-order */
+	PRAGMA_DISABLE_DEPRECATION_WARNINGS
 	class SLATECORE_API FOverlaySlot : public TBasicLayoutWidgetSlot<FOverlaySlot>
 	{
 	public:
@@ -45,20 +46,43 @@ public:
 			, ZOrder(0)
 		{ }
 
-		/** Slots with larger ZOrder values will draw above slots with smaller ZOrder values.  Slots
-		with the same ZOrder will simply draw in the order they were added.  Currently this only
-		works for overlay slots that are added dynamically with AddWidget() and RemoveWidget() */
+		SLATE_SLOT_BEGIN_ARGS(FOverlaySlot, TBasicLayoutWidgetSlot<FOverlaySlot>)
+			SLATE_ARGUMENT(TOptional<int32>, ZOrder)
+		SLATE_SLOT_END_ARGS()
+
+		void Construct(const FChildren& SlotOwner, FSlotArguments&& InArgs)
+		{
+			TBasicLayoutWidgetSlot<FOverlaySlot>::Construct(SlotOwner, MoveTemp(InArgs));
+			ZOrder = InArgs._ZOrder.Get(ZOrder);
+		}
+
+		int32 GetZOrder() const { return ZOrder; }
+
+		void SetZOrder(int32 InOrder)
+		{
+			if (ZOrder != InOrder)
+			{
+				ZOrder = InOrder;
+				FSlotBase::Invalidate(EInvalidateWidgetReason::Layout);
+			}
+		}
+
+	public:
+		/**
+		 * Slots with larger ZOrder values will draw above slots with smaller ZOrder values.  Slots
+		 * with the same ZOrder will simply draw in the order they were added.  Currently this only
+		 * works for overlay slots that are added dynamically with AddWidget() and RemoveWidget()
+		 */
+		UE_DEPRECATED(5.0, "Direct access to ZOrder is now deprecated. Use the getter or setter.")
 		int32 ZOrder;
 	};
-
+	PRAGMA_ENABLE_DEPRECATION_WARNINGS
 
 	SLATE_BEGIN_ARGS( SOverlay )
 	{
 		_Visibility = EVisibility::SelfHitTestInvisible;
 	}
-
-		SLATE_SUPPORTS_SLOT( SOverlay::FOverlaySlot )
-
+		SLATE_SLOT_ARGUMENT( SOverlay::FOverlaySlot, Slots )
 	SLATE_END_ARGS()
 
 	SOverlay();
@@ -80,8 +104,9 @@ public:
 	 */
 	bool RemoveSlot( TSharedRef< SWidget > Widget );
 
+	using FScopedWidgetSlotArguments = TPanelChildren<FOverlaySlot>::FScopedWidgetSlotArguments;
 	/** Adds a slot at the specified location (ignores Z-order) */
-	FOverlaySlot& AddSlot(int32 ZOrder=INDEX_NONE);
+	FScopedWidgetSlotArguments AddSlot(int32 ZOrder=INDEX_NONE);
 
 	/** Removes a slot at the specified location */
 	void RemoveSlot(int32 ZOrder=INDEX_NONE);
@@ -90,10 +115,7 @@ public:
 	void ClearChildren();
 
 	/** @return a new slot. Slots contain children for SOverlay */
-	static FOverlaySlot& Slot()
-	{
-		return *(new FOverlaySlot());
-	}
+	static FOverlaySlot::FSlotArguments Slot();
 
 	// SWidget interface
 	virtual void OnArrangeChildren( const FGeometry& AllottedGeometry, FArrangedChildren& ArrangedChildren ) const override;
diff --git a/Engine/Source/Runtime/SlateCore/Public/Widgets/SWindow.h b/Engine/Source/Runtime/SlateCore/Public/Widgets/SWindow.h
--- a/Engine/Source/Runtime/SlateCore/Public/Widgets/SWindow.h
+++ b/Engine/Source/Runtime/SlateCore/Public/Widgets/SWindow.h
@@ -530,7 +530,7 @@ public:
 	 * @param	InZOrder	Z-order to use for this widget
 	 * @return The added overlay slot so that it can be configured and populated
 	 */
-	SOverlay::FOverlaySlot& AddOverlaySlot( const int32 ZOrder = INDEX_NONE );
+	SOverlay::FScopedWidgetSlotArguments AddOverlaySlot( const int32 ZOrder = INDEX_NONE );
 
 	/**
 	 * Removes a widget that is being drawn over the entire window
diff --git a/Engine/Source/Runtime/UMG/Private/Components/CanvasPanelSlot.cpp b/Engine/Source/Runtime/UMG/Private/Components/CanvasPanelSlot.cpp
--- a/Engine/Source/Runtime/UMG/Private/Components/CanvasPanelSlot.cpp
+++ b/Engine/Source/Runtime/UMG/Private/Components/CanvasPanelSlot.cpp
@@ -26,7 +26,8 @@ void UCanvasPanelSlot::ReleaseSlateResources(bool bReleaseChildren)
 
 void UCanvasPanelSlot::BuildSlot(TSharedRef<SConstraintCanvas> Canvas)
 {
-	Slot = &Canvas->AddSlot()
+	Canvas->AddSlot()
+		.Expose(Slot)
 		[
 			Content == nullptr ? SNullWidget::NullWidget : Content->TakeWidget()
 		];
diff --git a/Engine/Source/Runtime/UMG/Private/Components/DynamicEntryBoxBase.cpp b/Engine/Source/Runtime/UMG/Private/Components/DynamicEntryBoxBase.cpp
--- a/Engine/Source/Runtime/UMG/Private/Components/DynamicEntryBoxBase.cpp
+++ b/Engine/Source/Runtime/UMG/Private/Components/DynamicEntryBoxBase.cpp
@@ -365,22 +365,26 @@ FMargin UDynamicEntryBoxBase::BuildEntryPadding(const FVector2D& DesiredSpacing)
 
 void UDynamicEntryBoxBase::AddEntryChild(UUserWidget& ChildWidget)
 {
-	FSlotBase* NewSlot = nullptr;
 	if (EntryBoxType == EDynamicBoxType::Wrap || EntryBoxType == EDynamicBoxType::VerticalWrap)
 	{
-		NewSlot = &StaticCastSharedPtr<SWrapBox>(MyPanelWidget)->AddSlot()
+		StaticCastSharedPtr<SWrapBox>(MyPanelWidget)->AddSlot()
 			.FillEmptySpace(false)
 			.HAlign(EntryHorizontalAlignment)
-			.VAlign(EntryVerticalAlignment);
+			.VAlign(EntryVerticalAlignment)
+			[
+				ChildWidget.TakeWidget()
+			];
 	}
 	else if (EntryBoxType == EDynamicBoxType::Radial)
 	{
-		NewSlot = &StaticCastSharedPtr<SRadialBox>(MyPanelWidget)->AddSlot();
+		StaticCastSharedPtr<SRadialBox>(MyPanelWidget)->AddSlot()
+			[
+				ChildWidget.TakeWidget()
+			];
 	}
 	else if (EntryBoxType == EDynamicBoxType::Overlay)
 	{
 		const int32 ChildIdx = MyPanelWidget->GetChildren()->Num();
-		SOverlay::FOverlaySlot& OverlaySlot = (SOverlay::FOverlaySlot&)StaticCastSharedPtr<SOverlay>(MyPanelWidget)->AddSlot();
 
 		EHorizontalAlignment HAlign = EntryHorizontalAlignment;
 		EVerticalAlignment VAlign = EntryVerticalAlignment;
@@ -401,33 +405,53 @@ void UDynamicEntryBoxBase::AddEntryChild(UUserWidget& ChildWidget)
 			VAlign = EntrySpacing.Y >= 0.f ? EVerticalAlignment::VAlign_Top : EVerticalAlignment::VAlign_Bottom;
 		}
 		
-		OverlaySlot.SetHorizontalAlignment(HAlign);
-		OverlaySlot.SetVerticalAlignment(VAlign);
-		OverlaySlot.SetPadding(BuildEntryPadding(TargetSpacing));
-
-		NewSlot = &OverlaySlot;
+		StaticCastSharedPtr<SOverlay>(MyPanelWidget)->AddSlot()
+			.HAlign(HAlign)
+			.VAlign(VAlign)
+			.Padding(BuildEntryPadding(TargetSpacing))
+			[
+				ChildWidget.TakeWidget()
+			];
 	}
 	else
 	{
 		const bool bIsHBox = EntryBoxType == EDynamicBoxType::Horizontal;
 		const bool bIsFirstChild = MyPanelWidget->GetChildren()->Num() == 0;
 
-		SBoxPanel::FSlot& BoxPanelSlot = bIsHBox ? (SBoxPanel::FSlot&)StaticCastSharedPtr<SHorizontalBox>(MyPanelWidget)->AddSlot().MaxWidth(MaxElementSize) : (SBoxPanel::FSlot&)StaticCastSharedPtr<SVerticalBox>(MyPanelWidget)->AddSlot().MaxHeight(MaxElementSize);
-		BoxPanelSlot.SetHorizontalAlignment(EntryHorizontalAlignment);
-		BoxPanelSlot.SetVerticalAlignment(EntryVerticalAlignment);
-		BoxPanelSlot.SizeParam = UWidget::ConvertSerializedSizeParamToRuntime(EntrySizeRule);
-
 		FMargin Padding;
 		Padding.Top = bIsHBox || bIsFirstChild ? 0.f : EntrySpacing.Y;
 		Padding.Left = bIsHBox && !bIsFirstChild ? EntrySpacing.X : 0.f;
-		BoxPanelSlot.SetPadding(Padding);
 
-		NewSlot = &BoxPanelSlot;
-	}
-
-	if (ensure(NewSlot))
-	{
-		NewSlot->AttachWidget(ChildWidget.TakeWidget());
+		if (bIsHBox)
+		{
+			SHorizontalBox::FSlot* BoxPanelSlot = nullptr;
+			StaticCastSharedPtr<SHorizontalBox>(MyPanelWidget)->AddSlot()
+				.Expose(BoxPanelSlot)
+				.MaxWidth(MaxElementSize)
+				.HAlign(EntryHorizontalAlignment)
+				.VAlign(EntryVerticalAlignment)
+				.Padding(Padding)
+				[
+					ChildWidget.TakeWidget()
+				];
+			// Todo, add to declaration when SBoxPanel because aware of FSlotArguments
+			BoxPanelSlot->SizeParam = UWidget::ConvertSerializedSizeParamToRuntime(EntrySizeRule);
+		}
+		else
+		{
+			SVerticalBox::FSlot* BoxPanelSlot = nullptr;
+			StaticCastSharedPtr<SVerticalBox>(MyPanelWidget)->AddSlot()
+				.Expose(BoxPanelSlot)
+				.MaxHeight(MaxElementSize)
+				.HAlign(EntryHorizontalAlignment)
+				.VAlign(EntryVerticalAlignment)
+				.Padding(Padding)
+				[
+					ChildWidget.TakeWidget()
+				];
+			// Todo, add to declaration when SBoxPanel because aware of FSlotArguments
+			BoxPanelSlot->SizeParam = UWidget::ConvertSerializedSizeParamToRuntime(EntrySizeRule);
+		}
 	}
 }
 
diff --git a/Engine/Source/Runtime/UMG/Private/Components/HorizontalBoxSlot.cpp b/Engine/Source/Runtime/UMG/Private/Components/HorizontalBoxSlot.cpp
--- a/Engine/Source/Runtime/UMG/Private/Components/HorizontalBoxSlot.cpp
+++ b/Engine/Source/Runtime/UMG/Private/Components/HorizontalBoxSlot.cpp
@@ -11,7 +11,7 @@ UHorizontalBoxSlot::UHorizontalBoxSlot(const FObjectInitializer& ObjectInitializ
 {
 	HorizontalAlignment = HAlign_Fill;
 	VerticalAlignment = VAlign_Fill;
-	Slot = NULL;
+	Slot = nullptr;
 	Size = FSlateChildSize(ESlateSizeRule::Automatic);
 }
 
@@ -19,12 +19,13 @@ void UHorizontalBoxSlot::ReleaseSlateResources(bool bReleaseChildren)
 {
 	Super::ReleaseSlateResources(bReleaseChildren);
 
-	Slot = NULL;
+	Slot = nullptr;
 }
 
 void UHorizontalBoxSlot::BuildSlot(TSharedRef<SHorizontalBox> HorizontalBox)
-{
-	Slot = &HorizontalBox->AddSlot()
+{ 
+	HorizontalBox->AddSlot()
+		.Expose(Slot)
 		.HAlign(HorizontalAlignment)
 		.VAlign(VerticalAlignment)
 		.Padding(Padding)
@@ -58,7 +59,7 @@ void UHorizontalBoxSlot::SetHorizontalAlignment(EHorizontalAlignment InHorizonta
 	HorizontalAlignment = InHorizontalAlignment;
 	if ( Slot )
 	{
-		Slot->HAlign(InHorizontalAlignment);
+		Slot->SetHorizontalAlignment(InHorizontalAlignment);
 	}
 }
 
@@ -67,7 +68,7 @@ void UHorizontalBoxSlot::SetVerticalAlignment(EVerticalAlignment InVerticalAlign
 	VerticalAlignment = InVerticalAlignment;
 	if ( Slot )
 	{
-		Slot->VAlign(InVerticalAlignment);
+		Slot->SetVerticalAlignment(InVerticalAlignment);
 	}
 }
 
diff --git a/Engine/Source/Runtime/UMG/Private/Components/OverlaySlot.cpp b/Engine/Source/Runtime/UMG/Private/Components/OverlaySlot.cpp
--- a/Engine/Source/Runtime/UMG/Private/Components/OverlaySlot.cpp
+++ b/Engine/Source/Runtime/UMG/Private/Components/OverlaySlot.cpp
@@ -11,24 +11,25 @@ UOverlaySlot::UOverlaySlot(const FObjectInitializer& ObjectInitializer)
 {
 	HorizontalAlignment = HAlign_Left;
 	VerticalAlignment = VAlign_Top;
-	Slot = NULL;
+	Slot = nullptr;
 }
 
 void UOverlaySlot::ReleaseSlateResources(bool bReleaseChildren)
 {
 	Super::ReleaseSlateResources(bReleaseChildren);
 
-	Slot = NULL;
+	Slot = nullptr;
 }
 
 void UOverlaySlot::BuildSlot(TSharedRef<SOverlay> Overlay)
 {
-	Slot = &Overlay->AddSlot()
+	Overlay->AddSlot()
+		.Expose(Slot)
 		.Padding(Padding)
 		.HAlign(HorizontalAlignment)
 		.VAlign(VerticalAlignment)
 		[
-			Content == NULL ? SNullWidget::NullWidget : Content->TakeWidget()
+			Content == nullptr ? SNullWidget::NullWidget : Content->TakeWidget()
 		];
 }
 
@@ -37,7 +38,7 @@ void UOverlaySlot::SetPadding(FMargin InPadding)
 	Padding = InPadding;
 	if ( Slot )
 	{
-		Slot->Padding(InPadding);
+		Slot->SetPadding(InPadding);
 	}
 }
 
@@ -46,7 +47,7 @@ void UOverlaySlot::SetHorizontalAlignment(EHorizontalAlignment InHorizontalAlign
 	HorizontalAlignment = InHorizontalAlignment;
 	if ( Slot )
 	{
-		Slot->HAlign(InHorizontalAlignment);
+		Slot->SetHorizontalAlignment(InHorizontalAlignment);
 	}
 }
 
@@ -55,7 +56,7 @@ void UOverlaySlot::SetVerticalAlignment(EVerticalAlignment InVerticalAlignment)
 	VerticalAlignment = InVerticalAlignment;
 	if ( Slot )
 	{
-		Slot->VAlign(InVerticalAlignment);
+		Slot->SetVerticalAlignment(InVerticalAlignment);
 	}
 }
 
diff --git a/Engine/Source/Runtime/UMG/Private/Components/UniformGridSlot.cpp b/Engine/Source/Runtime/UMG/Private/Components/UniformGridSlot.cpp
--- a/Engine/Source/Runtime/UMG/Private/Components/UniformGridSlot.cpp
+++ b/Engine/Source/Runtime/UMG/Private/Components/UniformGridSlot.cpp
@@ -8,7 +8,7 @@
 
 UUniformGridSlot::UUniformGridSlot(const FObjectInitializer& ObjectInitializer)
 	: Super(ObjectInitializer)
-	, Slot(NULL)
+	, Slot(nullptr)
 {
 	HorizontalAlignment = HAlign_Left;
 	VerticalAlignment = VAlign_Top;
@@ -18,16 +18,17 @@ void UUniformGridSlot::ReleaseSlateResources(bool bReleaseChildren)
 {
 	Super::ReleaseSlateResources(bReleaseChildren);
 
-	Slot = NULL;
+	Slot = nullptr;
 }
 
 void UUniformGridSlot::BuildSlot(TSharedRef<SUniformGridPanel> GridPanel)
 {
-	Slot = &GridPanel->AddSlot(Column, Row)
+	GridPanel->AddSlot(Column, Row)
+		.Expose(Slot)
 		.HAlign(HorizontalAlignment)
 		.VAlign(VerticalAlignment)
 		[
-			Content == NULL ? SNullWidget::NullWidget : Content->TakeWidget()
+			Content == nullptr ? SNullWidget::NullWidget : Content->TakeWidget()
 		];
 }
 
diff --git a/Engine/Source/Runtime/UMG/Private/Components/VerticalBoxSlot.cpp b/Engine/Source/Runtime/UMG/Private/Components/VerticalBoxSlot.cpp
--- a/Engine/Source/Runtime/UMG/Private/Components/VerticalBoxSlot.cpp
+++ b/Engine/Source/Runtime/UMG/Private/Components/VerticalBoxSlot.cpp
@@ -8,7 +8,7 @@
 
 UVerticalBoxSlot::UVerticalBoxSlot(const FObjectInitializer& ObjectInitializer)
 	: Super(ObjectInitializer)
-	, Slot(NULL)
+	, Slot(nullptr)
 {
 	HorizontalAlignment = HAlign_Fill;
 	VerticalAlignment = VAlign_Fill;
@@ -19,17 +19,18 @@ void UVerticalBoxSlot::ReleaseSlateResources(bool bReleaseChildren)
 {
 	Super::ReleaseSlateResources(bReleaseChildren);
 
-	Slot = NULL;
+	Slot = nullptr;
 }
 
 void UVerticalBoxSlot::BuildSlot(TSharedRef<SVerticalBox> VerticalBox)
 {
-	Slot = &VerticalBox->AddSlot()
+	VerticalBox->AddSlot()
+		.Expose(Slot)
 		.Padding(Padding)
 		.HAlign(HorizontalAlignment)
 		.VAlign(VerticalAlignment)
 		[
-			Content == NULL ? SNullWidget::NullWidget : Content->TakeWidget()
+			Content == nullptr ? SNullWidget::NullWidget : Content->TakeWidget()
 		];
 
 	Slot->SizeParam = UWidget::ConvertSerializedSizeParamToRuntime(Size);
@@ -40,7 +41,7 @@ void UVerticalBoxSlot::SetPadding(FMargin InPadding)
 	Padding = InPadding;
 	if ( Slot )
 	{
-		Slot->Padding(InPadding);
+		Slot->SetPadding(InPadding);
 	}
 }
 
@@ -58,7 +59,7 @@ void UVerticalBoxSlot::SetHorizontalAlignment(EHorizontalAlignment InHorizontalA
 	HorizontalAlignment = InHorizontalAlignment;
 	if ( Slot )
 	{
-		Slot->HAlign(InHorizontalAlignment);
+		Slot->SetHorizontalAlignment(InHorizontalAlignment);
 	}
 }
 
@@ -67,7 +68,7 @@ void UVerticalBoxSlot::SetVerticalAlignment(EVerticalAlignment InVerticalAlignme
 	VerticalAlignment = InVerticalAlignment;
 	if ( Slot )
 	{
-		Slot->VAlign(InVerticalAlignment);
+		Slot->SetVerticalAlignment(InVerticalAlignment);
 	}
 }
 
diff --git a/Engine/Source/Runtime/UMG/Private/Components/WidgetSwitcherSlot.cpp b/Engine/Source/Runtime/UMG/Private/Components/WidgetSwitcherSlot.cpp
--- a/Engine/Source/Runtime/UMG/Private/Components/WidgetSwitcherSlot.cpp
+++ b/Engine/Source/Runtime/UMG/Private/Components/WidgetSwitcherSlot.cpp
@@ -9,7 +9,7 @@
 
 UWidgetSwitcherSlot::UWidgetSwitcherSlot(const FObjectInitializer& ObjectInitializer)
 	: Super(ObjectInitializer)
-	, Slot(NULL)
+	, Slot(nullptr)
 {
 	HorizontalAlignment = HAlign_Fill;
 	VerticalAlignment = VAlign_Fill;
@@ -19,17 +19,18 @@ void UWidgetSwitcherSlot::ReleaseSlateResources(bool bReleaseChildren)
 {
 	Super::ReleaseSlateResources(bReleaseChildren);
 
-	Slot = NULL;
+	Slot = nullptr;
 }
 
 void UWidgetSwitcherSlot::BuildSlot(TSharedRef<SWidgetSwitcher> WidgetSwitcher)
 {
-	Slot = &WidgetSwitcher->AddSlot()
+	WidgetSwitcher->AddSlot()
+		.Expose(Slot)
 		.Padding(Padding)
 		.HAlign(HorizontalAlignment)
 		.VAlign(VerticalAlignment)
 		[
-			Content == NULL ? SNullWidget::NullWidget : Content->TakeWidget()
+			Content == nullptr ? SNullWidget::NullWidget : Content->TakeWidget()
 		];
 }
 
@@ -47,7 +48,7 @@ void UWidgetSwitcherSlot::SetPadding(FMargin InPadding)
 	Padding = InPadding;
 	if ( Slot )
 	{
-		Slot->Padding(InPadding);
+		Slot->SetPadding(InPadding);
 	}
 }
 
@@ -56,7 +57,7 @@ void UWidgetSwitcherSlot::SetHorizontalAlignment(EHorizontalAlignment InHorizont
 	HorizontalAlignment = InHorizontalAlignment;
 	if ( Slot )
 	{
-		Slot->HAlign(InHorizontalAlignment);
+		Slot->SetHorizontalAlignment(InHorizontalAlignment);
 	}
 }
 
@@ -65,7 +66,7 @@ void UWidgetSwitcherSlot::SetVerticalAlignment(EVerticalAlignment InVerticalAlig
 	VerticalAlignment = InVerticalAlignment;
 	if ( Slot )
 	{
-		Slot->VAlign(InVerticalAlignment);
+		Slot->SetVerticalAlignment(InVerticalAlignment);
 	}
 }
 
diff --git a/Engine/Source/Runtime/UMG/Private/Components/WrapBoxSlot.cpp b/Engine/Source/Runtime/UMG/Private/Components/WrapBoxSlot.cpp
--- a/Engine/Source/Runtime/UMG/Private/Components/WrapBoxSlot.cpp
+++ b/Engine/Source/Runtime/UMG/Private/Components/WrapBoxSlot.cpp
@@ -8,7 +8,7 @@
 
 UWrapBoxSlot::UWrapBoxSlot(const FObjectInitializer& ObjectInitializer)
 	: Super(ObjectInitializer)
-	, Slot(NULL)
+	, Slot(nullptr)
 {
 	HorizontalAlignment = HAlign_Fill;
 	VerticalAlignment = VAlign_Fill;
@@ -22,12 +22,13 @@ void UWrapBoxSlot::ReleaseSlateResources(bool bReleaseChildren)
 {
 	Super::ReleaseSlateResources(bReleaseChildren);
 
-	Slot = NULL;
+	Slot = nullptr;
 }
 
 void UWrapBoxSlot::BuildSlot(TSharedRef<SWrapBox> WrapBox)
 {
-	Slot = &WrapBox->AddSlot()
+	WrapBox->AddSlot()
+		.Expose(Slot)
 		.Padding(Padding)
 		.HAlign(HorizontalAlignment)
 		.VAlign(VerticalAlignment)
@@ -35,7 +36,7 @@ void UWrapBoxSlot::BuildSlot(TSharedRef<SWrapBox> WrapBox)
 		.FillLineWhenSizeLessThan(FillSpanWhenLessThan == 0 ? TOptional<float>() : TOptional<float>(FillSpanWhenLessThan))
 		.ForceNewLine(false)
 		[
-			Content == NULL ? SNullWidget::NullWidget : Content->TakeWidget()
+			Content == nullptr ? SNullWidget::NullWidget : Content->TakeWidget()
 		];
 }
 
@@ -44,7 +45,7 @@ void UWrapBoxSlot::SetPadding(FMargin InPadding)
 	Padding = InPadding;
 	if ( Slot )
 	{
-		Slot->Padding(InPadding);
+		Slot->SetPadding(InPadding);
 	}
 }
 
@@ -71,7 +72,7 @@ void UWrapBoxSlot::SetHorizontalAlignment(EHorizontalAlignment InHorizontalAlign
 	HorizontalAlignment = InHorizontalAlignment;
 	if ( Slot )
 	{
-		Slot->HAlign(InHorizontalAlignment);
+		Slot->SetHorizontalAlignment(InHorizontalAlignment);
 	}
 }
 
@@ -80,7 +81,7 @@ void UWrapBoxSlot::SetVerticalAlignment(EVerticalAlignment InVerticalAlignment)
 	VerticalAlignment = InVerticalAlignment;
 	if ( Slot )
 	{
-		Slot->VAlign(InVerticalAlignment);
+		Slot->SetVerticalAlignment(InVerticalAlignment);
 	}
 }
 
