Fix build with "error C4458: declaration of 'Slot' hides class member"
#rnx

diff --git a/Engine/Source/Runtime/SlateCore/Public/Layout/Children.h b/Engine/Source/Runtime/SlateCore/Public/Layout/Children.h
--- a/Engine/Source/Runtime/SlateCore/Public/Layout/Children.h
+++ b/Engine/Source/Runtime/SlateCore/Public/Layout/Children.h
@@ -621,7 +621,7 @@ public:
 	
 		virtual ~FScopedWidgetSlotArguments()
 		{
-			if (const SlotType* Slot = this->GetSlot())	// Is nullptr when the FScopedWidgetSlotArguments was moved-constructed.
+			if (const SlotType* SlotPtr = this->GetSlot())	// Is nullptr when the FScopedWidgetSlotArguments was moved-constructed.
 			{
 				if (Index == INDEX_NONE)
 				{
@@ -633,7 +633,7 @@ public:
 				}
 				if (Added)
 				{
-					Added(Slot, Index);
+					Added(SlotPtr, Index);
 				}
 			}
 		}
