PR #8052: Fix Tritanope DisplayName (Contributed by [REDACTED_PERSON])

diff --git a/Engine/Source/Runtime/SlateCore/Public/Rendering/RenderingCommon.h b/Engine/Source/Runtime/SlateCore/Public/Rendering/RenderingCommon.h
--- a/Engine/Source/Runtime/SlateCore/Public/Rendering/RenderingCommon.h
+++ b/Engine/Source/Runtime/SlateCore/Public/Rendering/RenderingCommon.h
@@ -148,7 +148,7 @@ enum class EColorVisionDeficiency : uint8
 	NormalVision UMETA(DisplayName="Normal Vision"),
 	Deuteranope UMETA(DisplayName="Deuteranope (green weak/blind) (7% of males, 0.4% of females)"),
 	Protanope UMETA(DisplayName="Protanope (red weak/blind) (2% of males, 0.01% of females)"),
-	Tritanope UMETA(DisplayName="Tritanope (blue weak/bind) (0.0003% of males)"),
+	Tritanope UMETA(DisplayName="Tritanope (blue weak/blind) (0.0003% of males)"),
 };
 
 
