Slate: Convert SConcertScrollBox, SWrapBox, SBoxPanel, SWindow FSlot's to the FSlotArgument syntax.

diff --git a/Engine/Plugins/Developer/Concert/ConcertSync/ConcertSyncClient/Source/ConcertSyncUI/Private/SConcertScrollBox.cpp b/Engine/Plugins/Developer/Concert/ConcertSync/ConcertSyncClient/Source/ConcertSyncUI/Private/SConcertScrollBox.cpp
--- a/Engine/Plugins/Developer/Concert/ConcertSync/ConcertSyncClient/Source/ConcertSyncUI/Private/SConcertScrollBox.cpp
+++ b/Engine/Plugins/Developer/Concert/ConcertSync/ConcertSyncClient/Source/ConcertSyncUI/Private/SConcertScrollBox.cpp
@@ -56,14 +56,14 @@ void SConcertScrollBox::Construct(const FArguments& InArgs)
 		];
 
 	// Forward the slots to the inner scroll box.
-	for (SConcertScrollBox::FSlot* InSlot : InArgs.Slots)
+	for (const SConcertScrollBox::FSlot::FSlotArguments& InSlot : InArgs._Slots)
 	{
-		if (InSlot)
+		if (InSlot.GetAttachedWidget())
 		{
 			ScrollBox->AddSlot()
-				[
-					InSlot->GetWidget()
-				];
+			[
+				InSlot.GetAttachedWidget().ToSharedRef()
+			];
 		}
 	}
 }
@@ -88,9 +88,9 @@ void SConcertScrollBox::Tick(const FGeometry& AllottedGeometry, const double InC
 }
 
 
-SConcertScrollBox::FSlot& SConcertScrollBox::Slot()
+SConcertScrollBox::FSlot::FSlotArguments SConcertScrollBox::Slot()
 {
-	return *(new SConcertScrollBox::FSlot());
+	return SConcertScrollBox::FSlot::FSlotArguments(MakeUnique<SConcertScrollBox::FSlot>());
 }
 
 void SConcertScrollBox::HandleUserScrolled(float)
diff --git a/Engine/Plugins/Developer/Concert/ConcertSync/ConcertSyncClient/Source/ConcertSyncUI/Public/SConcertScrollBox.h b/Engine/Plugins/Developer/Concert/ConcertSync/ConcertSyncClient/Source/ConcertSyncUI/Public/SConcertScrollBox.h
--- a/Engine/Plugins/Developer/Concert/ConcertSync/ConcertSyncClient/Source/ConcertSyncUI/Public/SConcertScrollBox.h
+++ b/Engine/Plugins/Developer/Concert/ConcertSync/ConcertSyncClient/Source/ConcertSyncUI/Public/SConcertScrollBox.h
@@ -18,14 +18,15 @@ public:
 
 	class FSlot : public TSlotBase<FSlot>
 	{
-		public:
-			FSlot()
-				: TSlotBase<FSlot>()
-			{}
+	public:
+		using TSlotBase<FSlot>::TSlotBase;
+
+		SLATE_SLOT_BEGIN_ARGS(FSlot, TSlotBase<FSlot>)
+		SLATE_SLOT_END_ARGS()
 	};
 
 	SLATE_BEGIN_ARGS(SConcertScrollBox) {}
-	SLATE_SUPPORTS_SLOT(FSlot)
+		SLATE_SLOT_ARGUMENT(FSlot, Slots)
 	SLATE_END_ARGS();
 
 	/**
@@ -43,7 +44,7 @@ public:
 	/**
 	 * @return a new slot. Slots contain children for SScrollBox 
 	 */
-	static FSlot& Slot();
+	static FSlot::FSlotArguments Slot();
 
 private:
 	
diff --git a/Engine/Plugins/Experimental/CommonUI/Source/CommonUI/Private/CommonWidgetCarouselNavBar.cpp b/Engine/Plugins/Experimental/CommonUI/Source/CommonUI/Private/CommonWidgetCarouselNavBar.cpp
--- a/Engine/Plugins/Experimental/CommonUI/Source/CommonUI/Private/CommonWidgetCarouselNavBar.cpp
+++ b/Engine/Plugins/Experimental/CommonUI/Source/CommonUI/Private/CommonWidgetCarouselNavBar.cpp
@@ -56,9 +56,11 @@ void UCommonWidgetCarouselNavBar::RebuildButtons()
 				ButtonGroup->AddWidget(ButtonUserWidget);
 
 				TSharedRef<SWidget> ButtonSWidget = ButtonUserWidget->TakeWidget();
-				SHorizontalBox::FSlot& NewSlot = MyContainer->AddSlot();
-				NewSlot.Padding(ButtonPadding);
-				NewSlot.AttachWidget(ButtonSWidget);
+				MyContainer->AddSlot()
+					.Padding(ButtonPadding)
+					[
+						ButtonSWidget
+					];
 			}
 		}
 
diff --git a/Engine/Plugins/Experimental/ControlRig/Source/ControlRigEditor/Private/Graph/SControlRigGraphNode.cpp b/Engine/Plugins/Experimental/ControlRig/Source/ControlRigEditor/Private/Graph/SControlRigGraphNode.cpp
--- a/Engine/Plugins/Experimental/ControlRig/Source/ControlRigEditor/Private/Graph/SControlRigGraphNode.cpp
+++ b/Engine/Plugins/Experimental/ControlRig/Source/ControlRigEditor/Private/Graph/SControlRigGraphNode.cpp
@@ -284,9 +284,8 @@ void SControlRigGraphNode::Construct( const FArguments& InArgs )
 			}
 			else
 			{
-				TPanelChildren<SBoxPanel::FSlot>* Slots = (TPanelChildren<SBoxPanel::FSlot>*)FullPinHorizontalRowWidget->GetChildren();
-				const int32 SlotToAdjustIndex = bPinWidgetForExpanderLeft ? 0 : Slots->Num() - 1;
-				SBoxPanel::FSlot& Slot = (*Slots)[SlotToAdjustIndex];
+				const int32 SlotToAdjustIndex = bPinWidgetForExpanderLeft ? 0 : FullPinHorizontalRowWidget->NumSlots() - 1;
+				SHorizontalBox::FSlot& Slot = FullPinHorizontalRowWidget->GetSlot(SlotToAdjustIndex);
 
 				FMargin Padding = Slot.GetPadding();
 				Padding = FMargin(RightIndentation + Padding.Left, Padding.Top, LeftIndentation + Padding.Right, Padding.Bottom);
@@ -337,7 +336,7 @@ void SControlRigGraphNode::Construct( const FArguments& InArgs )
 				if(PinInfo.bIsContainer)
 				{
 					// make sure to minimize the width of the label
-					FirstSlot->AutoWidth();
+					FirstSlot->SetAutoWidth();
 					
 					// add array plus button
 					SlotLayout->AddSlot()
diff --git a/Engine/Plugins/FX/Niagara/Source/NiagaraEditor/Private/Widgets/SNiagaraGraphParameterMapGetNode.cpp b/Engine/Plugins/FX/Niagara/Source/NiagaraEditor/Private/Widgets/SNiagaraGraphParameterMapGetNode.cpp
--- a/Engine/Plugins/FX/Niagara/Source/NiagaraEditor/Private/Widgets/SNiagaraGraphParameterMapGetNode.cpp
+++ b/Engine/Plugins/FX/Niagara/Source/NiagaraEditor/Private/Widgets/SNiagaraGraphParameterMapGetNode.cpp
@@ -88,9 +88,6 @@ void SNiagaraGraphParameterMapGetNode::CreatePinWidgets()
 	// Deferred pin adding to line up input/output pins by name.
 	for (int32 i = 0; i < OutputPins.Num() + 1; i++)
 	{
-		SVerticalBox::FSlot& Slot = PinContainerRoot->AddSlot();
-		Slot.AutoHeight();
-
 		// Get nodes have an unequal number of pins. 
 		TSharedPtr<SWidget> Widget;
 		if (i == 0)
@@ -149,8 +146,7 @@ void SNiagaraGraphParameterMapGetNode::CreatePinWidgets()
 				];
 		}
 
-		TSharedPtr<SBorder> Border;
-		SAssignNew(Border, SBorder)
+		TSharedRef<SBorder> Border = SNew(SBorder)
 			.HAlign(HAlign_Fill)
 			.VAlign(VAlign_Fill)
 			.Padding(FMargin(0, 3))
@@ -160,7 +156,12 @@ void SNiagaraGraphParameterMapGetNode::CreatePinWidgets()
 			];
 		Border->SetBorderImage(TAttribute<const FSlateBrush*>::Create(TAttribute<const FSlateBrush*>::FGetter::CreateRaw(this, &SNiagaraGraphParameterMapGetNode::GetBackgroundBrush, Widget)));
 
-		Slot.AttachWidget(Border.ToSharedRef());
+
+		PinContainerRoot->AddSlot()
+			.AutoHeight()
+			[
+				Border
+			];
 	}
 }
 
diff --git a/Engine/Plugins/NetcodeUnitTest/NetcodeUnitTest/Source/NetcodeUnitTest/Private/UI/SLogWidget.cpp b/Engine/Plugins/NetcodeUnitTest/NetcodeUnitTest/Source/NetcodeUnitTest/Private/UI/SLogWidget.cpp
--- a/Engine/Plugins/NetcodeUnitTest/NetcodeUnitTest/Source/NetcodeUnitTest/Private/UI/SLogWidget.cpp
+++ b/Engine/Plugins/NetcodeUnitTest/NetcodeUnitTest/Source/NetcodeUnitTest/Private/UI/SLogWidget.cpp
@@ -174,7 +174,7 @@ void SLogWidget::Construct(const FArguments& Args)
 
 	// Conditionally enable/disable some UI elements (e.g. Suspend/AutoClose), that aren't necessary for the status window
 	auto ConditionalSlot =
-		[] (bool bCondition, SHorizontalBox::FSlot& InSlot) -> SHorizontalBox::FSlot&
+		[] (bool bCondition, SHorizontalBox::FSlot::FSlotArguments& InSlot) -> SHorizontalBox::FSlot::FSlotArguments&
 		{
 			if (bCondition)
 			{
diff --git a/Engine/Source/Developer/Merge/Private/SMergeGraphView.cpp b/Engine/Source/Developer/Merge/Private/SMergeGraphView.cpp
--- a/Engine/Source/Developer/Merge/Private/SMergeGraphView.cpp
+++ b/Engine/Source/Developer/Merge/Private/SMergeGraphView.cpp
@@ -439,9 +439,9 @@ void SMergeGraphView::Construct(const FArguments InArgs
 				return DiffViewUtils::Identical();
 			};
 
-			const auto Box = [](bool bIsPresent, FLinearColor Color) -> SHorizontalBox::FSlot&
+			const auto Box = [](bool bIsPresent, FLinearColor Color) -> SHorizontalBox::FSlot::FSlotArguments
 			{
-				return SHorizontalBox::Slot()
+				return MoveTemp(SHorizontalBox::Slot()
 					.AutoWidth()
 					.HAlign(HAlign_Right)
 					.VAlign(VAlign_Center)
@@ -450,7 +450,7 @@ void SMergeGraphView::Construct(const FArguments InArgs
 						SNew(SImage)
 						.ColorAndOpacity(Color)
 						.Image(bIsPresent ? FEditorStyle::GetBrush("BlueprintDif.HasGraph") : FEditorStyle::GetBrush("BlueprintDif.MissingGraph"))
-					];
+					]);
 			};
 
 			FLinearColor RemoteColor = ComputeColor(InDifference->bAnyConflics, InDifference->bRemoteDifferences);
diff --git a/Engine/Source/Editor/CurveEditor/Private/SCurveEditorViewContainer.cpp b/Engine/Source/Editor/CurveEditor/Private/SCurveEditorViewContainer.cpp
--- a/Engine/Source/Editor/CurveEditor/Private/SCurveEditorViewContainer.cpp
+++ b/Engine/Source/Editor/CurveEditor/Private/SCurveEditorViewContainer.cpp
@@ -53,7 +53,7 @@ FVector2D SCurveEditorViewContainer::ComputeDesiredSize(float) const
 			// That will be the desired width of the whole panel.
 			MyDesiredSize.X = FMath::Max(MyDesiredSize.X, ChildDesiredSize.X + SlotPadding.GetTotalSpaceAlong<Orient_Horizontal>());
 
-			if (Child.SizeParam.SizeRule == FSizeParam::SizeRule_Stretch)
+			if (Child.GetSizeParam().SizeRule == FSizeParam::SizeRule_Stretch)
 			{
 				++NumStretchPanels;
 			}
@@ -327,7 +327,9 @@ void SCurveEditorViewContainer::AddView(TSharedRef<SCurveEditorView> ViewToAdd)
 	ViewToAdd->RelativeOrder = Views.Num();
 
 	Views.Add(ViewToAdd);
-	SVerticalBox::FSlot& Slot = AddSlot()
+	SVerticalBox::FSlot* SlotPointer = nullptr;
+	AddSlot()
+	.Expose(SlotPointer)
 	[
 		SNew(SBox)
 		.Padding(MakeAttributeSP(this, &SCurveEditorViewContainer::GetSlotPadding, InsertIndex))
@@ -339,7 +341,7 @@ void SCurveEditorViewContainer::AddView(TSharedRef<SCurveEditorView> ViewToAdd)
 
 	if (ViewToAdd->ShouldAutoSize())
 	{
-		Slot.AutoHeight();
+		SlotPointer->SetAutoHeight();
 	}
 }
 
diff --git a/Engine/Source/Editor/Kismet/Private/DiffUtils.cpp b/Engine/Source/Editor/Kismet/Private/DiffUtils.cpp
--- a/Engine/Source/Editor/Kismet/Private/DiffUtils.cpp
+++ b/Engine/Source/Editor/Kismet/Private/DiffUtils.cpp
@@ -1015,9 +1015,9 @@ FText DiffViewUtils::GetPanelLabel(const UBlueprint* Blueprint, const FRevisionI
 	}
 }
 
-SHorizontalBox::FSlot& DiffViewUtils::Box(bool bIsPresent, FLinearColor Color)
+SHorizontalBox::FSlot::FSlotArguments DiffViewUtils::Box(bool bIsPresent, FLinearColor Color)
 {
-	return SHorizontalBox::Slot()
+	return MoveTemp(SHorizontalBox::Slot()
 		.AutoWidth()
 		.HAlign(HAlign_Right)
 		.VAlign(VAlign_Center)
@@ -1026,6 +1026,6 @@ SHorizontalBox::FSlot& DiffViewUtils::Box(bool bIsPresent, FLinearColor Color)
 			SNew(SImage)
 			.ColorAndOpacity(Color)
 			.Image(bIsPresent ? FEditorStyle::GetBrush("BlueprintDif.HasGraph") : FEditorStyle::GetBrush("BlueprintDif.MissingGraph"))
-		];
+		]);
 };
 
diff --git a/Engine/Source/Editor/Kismet/Public/DiffUtils.h b/Engine/Source/Editor/Kismet/Public/DiffUtils.h
--- a/Engine/Source/Editor/Kismet/Public/DiffUtils.h
+++ b/Engine/Source/Editor/Kismet/Public/DiffUtils.h
@@ -315,5 +315,5 @@ namespace DiffViewUtils
 	KISMET_API FText SCSDiffMessage(const FSCSDiffEntry& Difference, FText ObjectName);
 	KISMET_API FText GetPanelLabel(const UBlueprint* Blueprint, const FRevisionInfo& Revision, FText Label);
 
-	KISMET_API SHorizontalBox::FSlot& Box(bool bIsPresent, FLinearColor Color);
+	KISMET_API SHorizontalBox::FSlot::FSlotArguments Box(bool bIsPresent, FLinearColor Color);
 }
diff --git a/Engine/Source/Editor/PixelInspector/Private/PixelInspectorDetailsCustomization.cpp b/Engine/Source/Editor/PixelInspector/Private/PixelInspectorDetailsCustomization.cpp
--- a/Engine/Source/Editor/PixelInspector/Private/PixelInspectorDetailsCustomization.cpp
+++ b/Engine/Source/Editor/PixelInspector/Private/PixelInspectorDetailsCustomization.cpp
@@ -29,12 +29,6 @@ TSharedRef<SHorizontalBox> FPixelInspectorDetailsCustomization::GetGridColorCont
 	TSharedRef<SHorizontalBox> HorizontalMainGrid = SNew(SHorizontalBox);
 	for (int32 ColumnIndex = 0; ColumnIndex < FinalColorContextGridSize; ++ColumnIndex)
 	{
-		SBoxPanel::FSlot &HorizontalSlot = HorizontalMainGrid->AddSlot()
-			.AutoWidth()
-			.Padding(2.0f, 2.0f)
-			.VAlign(VAlign_Center)
-			.HAlign(HAlign_Center);
-
 		TSharedRef<SVerticalBox> VerticalColumn = SNew(SVerticalBox);
 		for (int32 RowIndex = 0; RowIndex < FinalColorContextGridSize; ++RowIndex)
 		{
@@ -47,7 +41,14 @@ TSharedRef<SHorizontalBox> FPixelInspectorDetailsCustomization::GetGridColorCont
 					CreateColorCell(RowIndex, ColumnIndex, PixelInspectorView->FinalColorContext[ColumnIndex + RowIndex*FinalColorContextGridSize])
 				];
 		}
-		HorizontalSlot[VerticalColumn];
+		HorizontalMainGrid->AddSlot()
+			.AutoWidth()
+			.Padding(2.0f, 2.0f)
+			.VAlign(VAlign_Center)
+			.HAlign(HAlign_Center)
+			[
+				VerticalColumn
+			];
 	}
 	return HorizontalMainGrid;
 }
diff --git a/Engine/Source/Editor/PropertyEditor/Private/DetailPropertyRow.cpp b/Engine/Source/Editor/PropertyEditor/Private/DetailPropertyRow.cpp
--- a/Engine/Source/Editor/PropertyEditor/Private/DetailPropertyRow.cpp
+++ b/Engine/Source/Editor/PropertyEditor/Private/DetailPropertyRow.cpp
@@ -801,24 +801,26 @@ void FDetailPropertyRow::MakeNameOrKeyWidget( FDetailWidgetRow& Row, const TShar
 			.IsEnabled( IsEnabledAttrib );
 	}
 
-	SHorizontalBox::FSlot& Slot = NameHorizontalBox->AddSlot()
+	SHorizontalBox::FSlot* SlotPointer = nullptr;
+	NameHorizontalBox->AddSlot()
 	.HAlign(HAlign_Left)
+	.Expose(SlotPointer)
 	[
 		NameWidget.ToSharedRef()
 	];
 
 	if (bHasKeyNode)
 	{
-		Slot.Padding(0.0f, 0.0f, 2.0f, 0.0f);
+		SlotPointer->SetPadding(FMargin(0.0f, 0.0f, 2.0f, 0.0f));
 	}
 	else if (InCustomRow.IsValid())
 	{
 		//Allow custom name slot to fill all the area. If the user add a SHorizontalBox with left and right align slot
-		Slot.FillWidth(1.0f);
+		SlotPointer->SetFillWidth(1.0f);
 	}
 	else
 	{
-		Slot.AutoWidth();
+		SlotPointer->SetAutoWidth();
 	}
 
 	Row.NameContent()
diff --git a/Engine/Source/Runtime/AppFramework/Private/Framework/Testing/SStarshipGallery.cpp b/Engine/Source/Runtime/AppFramework/Private/Framework/Testing/SStarshipGallery.cpp
--- a/Engine/Source/Runtime/AppFramework/Private/Framework/Testing/SStarshipGallery.cpp
+++ b/Engine/Source/Runtime/AppFramework/Private/Framework/Testing/SStarshipGallery.cpp
@@ -585,11 +585,9 @@ public:
         TSharedPtr<SGridPanel> WidgetGrid = SNew(SGridPanel);
 
         int WidgetNum = 0, RowCount = 15, Cols = 3;
-        auto NextSlot = [WidgetNum, RowCount, Cols](TSharedPtr<SGridPanel> Grid, const FText& InLabel) mutable -> SHorizontalBox::FSlot&
-        {
-
-            SHorizontalBox::FSlot* ContentsSlot;    
-            TSharedRef<SHorizontalBox> HBox = SNew(SHorizontalBox)+SHorizontalBox::Slot().Expose(ContentsSlot).VAlign(VAlign_Center).HAlign(HAlign_Fill).AutoWidth();
+        auto NextSlot = [WidgetNum, RowCount, Cols](TSharedPtr<SGridPanel> Grid, const FText& InLabel) mutable -> SHorizontalBox::FScopedWidgetSlotArguments
+        { 
+            TSharedRef<SHorizontalBox> HBox = SNew(SHorizontalBox);
 
             // Checkbox to show disabled state
             Grid->AddSlot((WidgetNum / RowCount)*Cols, WidgetNum % RowCount)
@@ -626,7 +624,12 @@ public:
             ];
 
             ++WidgetNum;
-            return *ContentsSlot;
+
+			SHorizontalBox::FScopedWidgetSlotArguments NewSlot = HBox->AddSlot();
+			NewSlot.VAlign(VAlign_Center)
+				.HAlign(HAlign_Fill)
+				.AutoWidth();
+			return MoveTemp(NewSlot);
         };
 
 
@@ -637,7 +640,7 @@ public:
 
         	if (!InIconName.IsNone())
         	{
-        		HBox->AddSlot()	
+        		HBox->AddSlot()
 	        	.AutoWidth()
 	            .VAlign(VAlign_Center)
 	            [
diff --git a/Engine/Source/Runtime/Slate/Private/Framework/MultiBox/MultiBox.cpp b/Engine/Source/Runtime/Slate/Private/Framework/MultiBox/MultiBox.cpp
--- a/Engine/Source/Runtime/Slate/Private/Framework/MultiBox/MultiBox.cpp
+++ b/Engine/Source/Runtime/Slate/Private/Framework/MultiBox/MultiBox.cpp
@@ -687,28 +687,29 @@ void SMultiBoxWidget::AddBlockWidget(const FMultiBlock& Block, TSharedPtr<SHoriz
 
 			bool bOverride = Block.GetAlignmentOverrides(HAlign, VAlign, bAutoWidth);
 
-			SHorizontalBox::FSlot& Slot = HorizontalBox->AddSlot();
-		
-			if(bOverride)
 			{
-				if (bAutoWidth)
+				SHorizontalBox::FScopedWidgetSlotArguments NewSlot = HorizontalBox->AddSlot();
+
+				NewSlot.Padding(0.f)
+				[
+					FinalWidgetWithHook
+				];
+
+				if (bOverride)
 				{
-					Slot.AutoWidth();
-				}
+					if (bAutoWidth)
+					{
+						NewSlot.AutoWidth();
+					}
 
-				Slot.HAlign(HAlign);
-				Slot.VAlign(VAlign);
-			}
-			else
-			{
-				Slot.AutoWidth();
+					NewSlot.HAlign(HAlign)
+						.VAlign(VAlign);
+				}
+				else
+				{
+					NewSlot.AutoWidth();
+				}
 			}
-
-			Slot.Padding(0.f);
-			Slot
-			[
-				FinalWidgetWithHook
-			];
 		}
 		break;
 	case EMultiBoxType::VerticalToolBar:
diff --git a/Engine/Source/Runtime/Slate/Private/Widgets/Input/SCheckBox.cpp b/Engine/Source/Runtime/Slate/Private/Widgets/Input/SCheckBox.cpp
--- a/Engine/Source/Runtime/Slate/Private/Widgets/Input/SCheckBox.cpp
+++ b/Engine/Source/Runtime/Slate/Private/Widgets/Input/SCheckBox.cpp
@@ -549,7 +549,7 @@ void SCheckBox::BuildCheckBox(TSharedRef<SWidget> InContent)
 		];
 		if (bCheckBoxContentUsesAutoWidth)
 		{
-			ContentSlot->AutoWidth();
+			ContentSlot->SetAutoWidth();
 		}
 	}
 	else if (ensure(CheckBoxType == ESlateCheckBoxType::ToggleButton))
diff --git a/Engine/Source/Runtime/Slate/Private/Widgets/Layout/SHeader.cpp b/Engine/Source/Runtime/Slate/Private/Widgets/Layout/SHeader.cpp
--- a/Engine/Source/Runtime/Slate/Private/Widgets/Layout/SHeader.cpp
+++ b/Engine/Source/Runtime/Slate/Private/Widgets/Layout/SHeader.cpp
@@ -12,7 +12,7 @@ void SHeader::Construct( const FArguments& InArgs )
 
 	SHorizontalBox::Construct( SHorizontalBox::FArguments()
 
-	+ SHorizontalBox::Slot()
+		+ SHorizontalBox::Slot()
 		.Expose( FirstSlot )
 		.FillWidth(1)
 		.VAlign(VAlign_Center)
@@ -51,11 +51,11 @@ void SHeader::Construct( const FArguments& InArgs )
 			break;
 
 		case HAlign_Left:
-			FirstSlot->AutoWidth();
+			FirstSlot->SetAutoWidth();
 			break;
 
 		case HAlign_Right:
-			LastSlot->AutoWidth();
+			LastSlot->SetAutoWidth();
 			break;
 	};
 }
diff --git a/Engine/Source/Runtime/Slate/Private/Widgets/Layout/SWrapBox.cpp b/Engine/Source/Runtime/Slate/Private/Widgets/Layout/SWrapBox.cpp
--- a/Engine/Source/Runtime/Slate/Private/Widgets/Layout/SWrapBox.cpp
+++ b/Engine/Source/Runtime/Slate/Private/Widgets/Layout/SWrapBox.cpp
@@ -8,16 +8,9 @@ SWrapBox::SWrapBox()
 {
 }
 
-SWrapBox::FSlot& SWrapBox::Slot()
+SWrapBox::FScopedWidgetSlotArguments SWrapBox::AddSlot()
 {
-	return *( new SWrapBox::FSlot() );
-}
-
-SWrapBox::FSlot& SWrapBox::AddSlot()
-{
-	SWrapBox::FSlot* NewSlot = new SWrapBox::FSlot();
-	Slots.Add(NewSlot);
-	return *NewSlot;
+	return FScopedWidgetSlotArguments{ MakeUnique<FSlot>(), Slots, INDEX_NONE };
 }
 
 int32 SWrapBox::RemoveSlot( const TSharedRef<SWidget>& SlotWidget )
@@ -40,11 +33,8 @@ void SWrapBox::Construct( const FArguments& InArgs )
 	Orientation = InArgs._Orientation;
 	HAlign = InArgs._HAlign;
 
-	// Copy the children from the declaration to the widget
-	for ( int32 ChildIndex=0; ChildIndex < InArgs.Slots.Num(); ++ChildIndex )
-	{
-		Slots.Add( InArgs.Slots[ChildIndex] );
-	}
+	// Build the children from the declaration to the widget
+	Slots.AddSlots(MoveTemp(const_cast<TArray<FSlot::FSlotArguments>&>(InArgs._Slots)));
 }
 
 void SWrapBox::Tick( const FGeometry& AllottedGeometry, const double InCurrentTime, const float InDeltaTime )
@@ -172,7 +162,7 @@ void SWrapBox::FChildArranger::Arrange()
 		{
 			const float BottomBoundOfChild = ArrangementData.SlotOffset.Y + ArrangementData.SlotSize.Y;
 
-			if (Slot.bSlotForceNewLine)
+			if (Slot.GetForceNewLine())
 			{
 				// Begin a new line if the current one isn't empty, because the slot demanded a new line.
 				if (!IsFirstChildInCurrentLine())
@@ -180,12 +170,9 @@ void SWrapBox::FChildArranger::Arrange()
 					BeginNewLine();
 				}
 			}
-			PRAGMA_DISABLE_DEPRECATION_WARNINGS
 			// Rule: If required due to a wrapping height under specified threshold, start a new line and allocate all of it to this child.
-			else if ((Slot.SlotFillLineWhenSizeLessThan.IsSet() && WrapBox.PreferredSize.Get() < Slot.SlotFillLineWhenSizeLessThan.GetValue())
-				|| (!Slot.SlotFillLineWhenSizeLessThan.IsSet() && Slot.SlotFillLineWhenWidthLessThan.IsSet() && WrapBox.PreferredSize.Get() < Slot.SlotFillLineWhenWidthLessThan.GetValue()))
+			else if (Slot.GetFillLineWhenSizeLessThan().IsSet() && WrapBox.PreferredSize.Get() < Slot.GetFillLineWhenSizeLessThan().GetValue())
 			{
-			PRAGMA_ENABLE_DEPRECATION_WARNINGS
 				// Begin a new line if the current one isn't empty, because we demand a whole line to ourselves.
 				if (!IsFirstChildInCurrentLine())
 				{
@@ -215,7 +202,7 @@ void SWrapBox::FChildArranger::Arrange()
 		{
 			const float RightBoundOfChild = ArrangementData.SlotOffset.X + ArrangementData.SlotSize.X;
 
-			if (Slot.bSlotForceNewLine)
+			if (Slot.GetForceNewLine())
 			{
 				// Begin a new line if the current one isn't empty, because the slot demanded a new line.
 				if (!IsFirstChildInCurrentLine())
@@ -223,12 +210,9 @@ void SWrapBox::FChildArranger::Arrange()
 					BeginNewLine();
 				}
 			}
-			PRAGMA_DISABLE_DEPRECATION_WARNINGS
 			// Rule: If required due to a wrapping width under specified threshold, start a new line and allocate all of it to this child.
-			else if ((Slot.SlotFillLineWhenSizeLessThan.IsSet() && WrapBox.PreferredSize.Get() < Slot.SlotFillLineWhenSizeLessThan.GetValue())
-				|| (!Slot.SlotFillLineWhenSizeLessThan.IsSet() && Slot.SlotFillLineWhenWidthLessThan.IsSet() && WrapBox.PreferredSize.Get() < Slot.SlotFillLineWhenWidthLessThan.GetValue()))
+			else if (Slot.GetFillLineWhenSizeLessThan().IsSet() && WrapBox.PreferredSize.Get() < Slot.GetFillLineWhenSizeLessThan().GetValue())
 			{
-			PRAGMA_ENABLE_DEPRECATION_WARNINGS
 				// Begin a new line if the current one isn't empty, because we demand a whole line to ourselves.
 				if (!IsFirstChildInCurrentLine())
 				{
@@ -277,7 +261,7 @@ void SWrapBox::FChildArranger::FinalizeLine(int32 IndexOfLastChildInCurrentLine)
 
 	// If we're doing this on a line where we fill the empty space,
 	// the widget of the current line will just be the whole line.
-	if (WrapBox.Slots.IsValidIndex(IndexOfLastChildInCurrentLine) && WrapBox.Slots[IndexOfLastChildInCurrentLine].bSlotFillEmptySpace)
+	if (WrapBox.Slots.IsValidIndex(IndexOfLastChildInCurrentLine) && WrapBox.Slots[IndexOfLastChildInCurrentLine].GetFillEmptySpace())
 	{
 		WidthOfCurrentLine = WrapBox.PreferredSize.Get();
 	}
@@ -302,7 +286,7 @@ void SWrapBox::FChildArranger::FinalizeLine(int32 IndexOfLastChildInCurrentLine)
 		RelativeChildIndex++;
 
 		// Rule: The last uncollapsed child in a line may request to fill the remaining empty space in the line.
-		if (ChildIndex == IndexOfLastChildInCurrentLine && Slot.bSlotFillEmptySpace)
+		if (ChildIndex == IndexOfLastChildInCurrentLine && Slot.GetFillEmptySpace())
 		{
 			if (WrapBox.Orientation == EOrientation::Orient_Vertical)
 			{
diff --git a/Engine/Source/Runtime/Slate/Public/Widgets/Layout/SWrapBox.h b/Engine/Source/Runtime/Slate/Public/Widgets/Layout/SWrapBox.h
--- a/Engine/Source/Runtime/Slate/Public/Widgets/Layout/SWrapBox.h
+++ b/Engine/Source/Runtime/Slate/Public/Widgets/Layout/SWrapBox.h
@@ -52,54 +52,88 @@ class SLATE_API SWrapBox : public SPanel
 public:
 
 	/** A slot that support alignment of content and padding */
+	PRAGMA_DISABLE_DEPRECATION_WARNINGS
 	class FSlot : public TBasicLayoutWidgetSlot<FSlot>
 	{
 	public:
 		FSlot()
 			: TBasicLayoutWidgetSlot<FSlot>(HAlign_Fill, VAlign_Fill)
-			, SlotFillLineWhenWidthLessThan()
 			, SlotFillLineWhenSizeLessThan()
 			, bSlotFillEmptySpace(false)
 			, bSlotForceNewLine(false)
 		{
 		}
 
-		UE_DEPRECATED(4.26, "Deprecated, please use FillLineWhenSizeLessThan() instead")
-		FSlot& FillLineWhenWidthLessThan(TOptional<float> InFillLineWhenWidthLessThan)
+		SLATE_SLOT_BEGIN_ARGS(FSlot, TBasicLayoutWidgetSlot<FSlot>)
+			SLATE_ARGUMENT(TOptional<float>, FillLineWhenSizeLessThan)
+			SLATE_ARGUMENT(TOptional<bool>, FillEmptySpace)
+			SLATE_ARGUMENT(TOptional<bool>, ForceNewLine)
+		SLATE_SLOT_END_ARGS()
+
+		void Construct(const FChildren& SlotOwner, FSlotArguments&& InArgs)
 		{
-			PRAGMA_DISABLE_DEPRECATION_WARNINGS
-			SlotFillLineWhenWidthLessThan = InFillLineWhenWidthLessThan;
-			PRAGMA_ENABLE_DEPRECATION_WARNINGS
-			return *(static_cast<FSlot*>(this));
+			TBasicLayoutWidgetSlot<FSlot>::Construct(SlotOwner, MoveTemp(InArgs));
+			if (InArgs._FillLineWhenSizeLessThan.IsSet())
+			{
+				SlotFillLineWhenSizeLessThan = InArgs._FillLineWhenSizeLessThan;
+			}
+			bSlotFillEmptySpace = InArgs._FillEmptySpace.Get(bSlotFillEmptySpace);
+			bSlotForceNewLine = InArgs._ForceNewLine.Get(bSlotForceNewLine);
 		}
 
 		/** Dependently of the Orientation, if the total available horizontal or vertical space in the wrap panel drops below this threshold, this slot will attempt to fill an entire line. */
-		FSlot& FillLineWhenSizeLessThan(TOptional<float> InFillLineWhenSizeLessThan)
+		void SetFillLineWhenSizeLessThan(TOptional<float> InFillLineWhenSizeLessThan)
+		{
+			if (SlotFillLineWhenSizeLessThan != InFillLineWhenSizeLessThan)
+			{
+				SlotFillLineWhenSizeLessThan = InFillLineWhenSizeLessThan;
+				FSlotBase::Invalidate(EInvalidateWidgetReason::Layout);
+			}
+		}
+
+		TOptional<float> GetFillLineWhenSizeLessThan() const
 		{
-			SlotFillLineWhenSizeLessThan = InFillLineWhenSizeLessThan;
-			return *(static_cast<FSlot*>(this));
+			return SlotFillLineWhenSizeLessThan;
 		}
 
 		/** Should this slot fill the remaining space on the line? */
-		FSlot& FillEmptySpace(bool bInFillEmptySpace)
+		void SetFillEmptySpace(bool bInFillEmptySpace)
 		{
-			bSlotFillEmptySpace = bInFillEmptySpace;
-			return *(static_cast<FSlot*>(this));
+			if (bSlotFillEmptySpace != bInFillEmptySpace)
+			{
+				bSlotFillEmptySpace = bInFillEmptySpace;
+				FSlotBase::Invalidate(EInvalidateWidgetReason::Layout);
+			}
 		}
 
-		FSlot& ForceNewLine(bool bInForceNewLine)
+		bool GetFillEmptySpace() const
 		{
-			bSlotForceNewLine = bInForceNewLine;
-			return *(static_cast<FSlot*>(this));
+			return bSlotFillEmptySpace;
 		}
 
-		UE_DEPRECATED(4.26, "Deprecated, please use SlotFillLineWhenSizeLessThan instead")
-		TOptional<float> SlotFillLineWhenWidthLessThan;
+		void SetForceNewLine(bool bInForceNewLine)
+		{
+			if (bSlotForceNewLine != bInForceNewLine)
+			{
+				bSlotForceNewLine = bInForceNewLine;
+				FSlotBase::Invalidate(EInvalidateWidgetReason::Layout);
+			}
+		}
 
+		bool GetForceNewLine() const
+		{
+			return bSlotForceNewLine;
+		}
+
+	public:
+		UE_DEPRECATED(5.0, "Direct access to SlotFillLineWhenSizeLessThan is now deprecated. Use the getter or setter.")
 		TOptional<float> SlotFillLineWhenSizeLessThan;
+		UE_DEPRECATED(5.0, "Direct access to bSlotFillEmptySpace is now deprecated. Use the getter or setter.")
 		bool bSlotFillEmptySpace;
+		UE_DEPRECATED(5.0, "Direct access to bSlotForceNewLine is now deprecated. Use the getter or setter.")
 		bool bSlotForceNewLine;
 	};
+	PRAGMA_ENABLE_DEPRECATION_WARNINGS
 
 
 	SLATE_BEGIN_ARGS(SWrapBox)
@@ -115,7 +149,7 @@ public:
 		}
 
 		/** The slot supported by this panel */
-		SLATE_SUPPORTS_SLOT( FSlot )
+		SLATE_SLOT_ARGUMENT( FSlot, Slots )
 
 		/** The preferred width, if not set will fill the space */
 		SLATE_ATTRIBUTE( float, PreferredWidth )
@@ -141,9 +175,13 @@ public:
 
 	SWrapBox();
 
-	static FSlot& Slot();
+	static FSlot::FSlotArguments Slot()
+	{
+		return FSlot::FSlotArguments(MakeUnique<FSlot>());
+	}
 
-	FSlot& AddSlot();
+	using FScopedWidgetSlotArguments = TPanelChildren<FSlot>::FScopedWidgetSlotArguments;
+	FScopedWidgetSlotArguments AddSlot();
 
 	/** Removes a slot from this box panel which contains the specified SWidget
 	 *
diff --git a/Engine/Source/Runtime/Slate/Public/Widgets/Views/STableRow.h b/Engine/Source/Runtime/Slate/Public/Widgets/Views/STableRow.h
--- a/Engine/Source/Runtime/Slate/Public/Widgets/Views/STableRow.h
+++ b/Engine/Source/Runtime/Slate/Public/Widgets/Views/STableRow.h
@@ -228,18 +228,14 @@ public:
 
 		if ( InOwnerTableMode == ETableViewMode::List || InOwnerTableMode == ETableViewMode::Tile )
 		{
-			// -- Row is in a ListView or the user --
-			FOneSimpleMemberChild* InnerContentSlotNativePtr = nullptr;
-
 			// We just need to hold on to this row's content.
 			this->ChildSlot
-			.Expose( InnerContentSlotNativePtr )
 			.Padding( InPadding )
 			[
 				InContent
 			];
 
-			InnerContentSlot = &InnerContentSlotNativePtr->AsSlot();
+			InnerContentSlot = &ChildSlot.AsSlot();
 		}
 		else
 		{
diff --git a/Engine/Source/Runtime/SlateCore/Private/Widgets/SBoxPanel.cpp b/Engine/Source/Runtime/SlateCore/Private/Widgets/SBoxPanel.cpp
--- a/Engine/Source/Runtime/SlateCore/Private/Widgets/SBoxPanel.cpp
+++ b/Engine/Source/Runtime/SlateCore/Private/Widgets/SBoxPanel.cpp
@@ -3,37 +3,63 @@
 #include "Widgets/SBoxPanel.h"
 #include "Layout/LayoutUtils.h"
 
-/**
- * Construct this widget
- *
- * @param	InArgs	The declaration data for this widget
- */
 void SHorizontalBox::Construct( const SHorizontalBox::FArguments& InArgs )
 {
-	const int32 NumSlots = InArgs.Slots.Num();
-	for ( int32 SlotIndex = 0; SlotIndex < NumSlots; ++SlotIndex )
+	Children.Reserve(InArgs._Slots.Num());
+	for (const FSlot::FSlotArguments& Arg : InArgs._Slots)
 	{
-		Children.Add( InArgs.Slots[SlotIndex] );
+		// Because we want to override the AutoWidth, the base class doesn't exactly have the same parent.
+		//We are casting from parent to child to a different parent to prevent a reinterpret_cast
+		const FSlotBase::FSlotArguments& ChilSlotArgument = static_cast<const FSlotBase::FSlotArguments&>(Arg);
+		const SBoxPanel::FSlot::FSlotArguments& BoxSlotArgument = static_cast<const SBoxPanel::FSlot::FSlotArguments&>(ChilSlotArgument);
+		// Because InArgs is const&, we need to do some hacking here. That would need to changed in the future.
+		//The Slot has a unique_ptr, it cannot be copied. Anyway, previously, the Children.Add(), was wrong if we added the same slot twice.
+		//Because of that, it doesn't matter if we steal the slot from the FArguments.
+		Children.AddSlot(MoveTemp(const_cast<SBoxPanel::FSlot::FSlotArguments&>(BoxSlotArgument)));
 	}
 }
 
-
-/**
- * Construct this widget
- *
- * @param	InArgs	The declaration data for this widget
- */
 void SVerticalBox::Construct( const SVerticalBox::FArguments& InArgs )
 {
-	const int32 NumSlots = InArgs.Slots.Num();
-	for ( int32 SlotIndex = 0; SlotIndex < NumSlots; ++SlotIndex )
+	Children.Reserve(InArgs._Slots.Num());
+	for (const FSlot::FSlotArguments& Arg : InArgs._Slots)
 	{
-		Children.Add( InArgs.Slots[SlotIndex] );
+		const FSlotBase::FSlotArguments& ChilSlotArgument = static_cast<const FSlotBase::FSlotArguments&>(Arg);
+		const SBoxPanel::FSlot::FSlotArguments& BoxSlotArgument = static_cast<const SBoxPanel::FSlot::FSlotArguments&>(ChilSlotArgument);
+		Children.AddSlot(MoveTemp(const_cast<SBoxPanel::FSlot::FSlotArguments&>(BoxSlotArgument)));
 	}
 }
 
-template<EOrientation Orientation>
-static void ArrangeChildrenAlong(EFlowDirection InLayoutFlow, const TPanelChildren<SBoxPanel::FSlot>& Children, const FGeometry& AllottedGeometry, FArrangedChildren& ArrangedChildren )
+SHorizontalBox::FSlot& SHorizontalBox::GetSlot(int32 SlotIndex)
+{
+	check(this->IsValidSlotIndex(SlotIndex));
+	FSlotBase& BaseSlot = static_cast<FSlotBase&>(Children[SlotIndex]);
+	return static_cast<SHorizontalBox::FSlot&>(BaseSlot);
+}
+
+const SHorizontalBox::FSlot& SHorizontalBox::GetSlot(int32 SlotIndex) const
+{
+	check(this->IsValidSlotIndex(SlotIndex));
+	const FSlotBase& BaseSlot = static_cast<const FSlotBase&>(Children[SlotIndex]);
+	return static_cast<const SHorizontalBox::FSlot&>(BaseSlot);
+}
+
+SVerticalBox::FSlot& SVerticalBox::GetSlot(int32 SlotIndex)
+{
+	check(this->IsValidSlotIndex(SlotIndex));
+	FSlotBase& BaseSlot = static_cast<FSlotBase&>(Children[SlotIndex]);
+	return static_cast<SVerticalBox::FSlot&>(BaseSlot);
+}
+
+const SVerticalBox::FSlot& SVerticalBox::GetSlot(int32 SlotIndex) const
+{
+	check(this->IsValidSlotIndex(SlotIndex));
+	const FSlotBase& BaseSlot = static_cast<const FSlotBase&>(Children[SlotIndex]);
+	return static_cast<const SVerticalBox::FSlot&>(BaseSlot);
+}
+
+template<EOrientation Orientation, typename SlotType>
+static void ArrangeChildrenAlong(EFlowDirection InLayoutFlow, const TPanelChildren<SlotType>& Children, const FGeometry& AllottedGeometry, FArrangedChildren& ArrangedChildren )
 {
 	// Allotted space will be given to fixed-size children first.
 	// Remaining space will be proportionately divided between stretch children (SizeRule_Stretch)
@@ -49,18 +75,18 @@ static void ArrangeChildrenAlong(EFlowDirection InLayoutFlow, const TPanelChildr
 		// (SizeRule_Auto).
 		for( int32 ChildIndex=0; ChildIndex < Children.Num(); ++ChildIndex )
 		{
-			const SBoxPanel::FSlot& CurChild = Children[ChildIndex];
+			const SlotType& CurChild = Children[ChildIndex];
 
 			if ( CurChild.GetWidget()->GetVisibility() != EVisibility::Collapsed )
 			{
 				bAnyChildVisible = true;
 				// All widgets contribute their margin to the fixed space requirement
-				FixedTotal += CurChild.GetPadding().GetTotalSpaceAlong<Orientation>();
+				FixedTotal += CurChild.GetPadding().template GetTotalSpaceAlong<Orientation>();
 
-				if ( CurChild.SizeParam.SizeRule == FSizeParam::SizeRule_Stretch )
+				if ( CurChild.GetSizeParam().SizeRule == FSizeParam::SizeRule_Stretch )
 				{
 					// for stretch children we save sum up the stretch coefficients
-					StretchCoefficientTotal += CurChild.SizeParam.Value.Get();
+					StretchCoefficientTotal += CurChild.GetSizeParam().Value.Get();
 				}
 				else
 				{
@@ -72,7 +98,7 @@ static void ArrangeChildrenAlong(EFlowDirection InLayoutFlow, const TPanelChildr
 						: ChildDesiredSize.X;
 
 					// Clamp to the max size if it was specified
-					float MaxSize = CurChild.MaxSize.Get();
+					float MaxSize = CurChild.GetMaxSize();
 					FixedTotal += MaxSize > 0 ? FMath::Min( MaxSize, ChildSize ) : ChildSize;
 				}
 			}
@@ -92,9 +118,9 @@ static void ArrangeChildrenAlong(EFlowDirection InLayoutFlow, const TPanelChildr
 
 		// Now that we have the total fixed-space requirement and the total stretch coefficients we can
 		// arrange widgets top-to-bottom or left-to-right (depending on the orientation).
-		for (TPanelChildrenConstIterator<SBoxPanel::FSlot> It(Children, Orientation, InLayoutFlow); It; ++It)
+		for (TPanelChildrenConstIterator<SlotType> It(Children, Orientation, InLayoutFlow); It; ++It)
 		{
-			const SBoxPanel::FSlot& CurChild = *It;
+			const SlotType& CurChild = *It;
 			const EVisibility ChildVisibility = CurChild.GetWidget()->GetVisibility();
 
 			// Figure out the area allocated to the child in the direction of BoxPanel
@@ -103,12 +129,12 @@ static void ArrangeChildrenAlong(EFlowDirection InLayoutFlow, const TPanelChildr
 			if (ChildVisibility != EVisibility::Collapsed)
 			{
 				// The size of the widget depends on its size type
-				if (CurChild.SizeParam.SizeRule == FSizeParam::SizeRule_Stretch)
+				if (CurChild.GetSizeParam().SizeRule == FSizeParam::SizeRule_Stretch)
 				{
 					if (StretchCoefficientTotal > 0)
 					{
 						// Stretch widgets get a fraction of the space remaining after all the fixed-space requirements are met
-						ChildSize = NonFixedSpace * CurChild.SizeParam.Value.Get() / StretchCoefficientTotal;
+						ChildSize = NonFixedSpace * CurChild.GetSizeParam().Value.Get() / StretchCoefficientTotal;
 					}
 				}
 				else
@@ -122,7 +148,7 @@ static void ArrangeChildrenAlong(EFlowDirection InLayoutFlow, const TPanelChildr
 				}
 
 				// Clamp to the max size if it was specified
-				float MaxSize = CurChild.MaxSize.Get();
+				float MaxSize = CurChild.GetMaxSize();
 				if (MaxSize > 0)
 				{
 					ChildSize = FMath::Min(MaxSize, ChildSize);
@@ -132,8 +158,8 @@ static void ArrangeChildrenAlong(EFlowDirection InLayoutFlow, const TPanelChildr
 			const FMargin SlotPadding(LayoutPaddingWithFlow(InLayoutFlow, CurChild.GetPadding()));
 
 			FVector2D SlotSize = (Orientation == Orient_Vertical)
-				? FVector2D(AllottedGeometry.GetLocalSize().X, ChildSize + SlotPadding.GetTotalSpaceAlong<Orient_Vertical>())
-				: FVector2D(ChildSize + SlotPadding.GetTotalSpaceAlong<Orient_Horizontal>(), AllottedGeometry.GetLocalSize().Y);
+				? FVector2D(AllottedGeometry.GetLocalSize().X, ChildSize + SlotPadding.template GetTotalSpaceAlong<Orient_Vertical>())
+				: FVector2D(ChildSize + SlotPadding.template GetTotalSpaceAlong<Orient_Horizontal>(), AllottedGeometry.GetLocalSize().Y);
 
 			// Figure out the size and local position of the child within the slot			
 			AlignmentArrangeResult XAlignmentResult = AlignChild<Orient_Horizontal>(InLayoutFlow, SlotSize.X, CurChild, SlotPadding);
@@ -173,13 +199,13 @@ static void ArrangeChildrenAlong(EFlowDirection InLayoutFlow, const TPanelChildr
  */
 void SBoxPanel::OnArrangeChildren( const FGeometry& AllottedGeometry, FArrangedChildren& ArrangedChildren ) const
 {
-	if ( this->Orientation == Orient_Horizontal )
+	if ( this->Orientation == EOrientation::Orient_Horizontal )
 	{
-		ArrangeChildrenAlong<Orient_Horizontal>(GSlateFlowDirection, this->Children, AllottedGeometry, ArrangedChildren );
+		ArrangeChildrenAlong<EOrientation::Orient_Horizontal>(GSlateFlowDirection, this->Children, AllottedGeometry, ArrangedChildren );
 	}
 	else
 	{
-		ArrangeChildrenAlong<Orient_Vertical>(GSlateFlowDirection, this->Children, AllottedGeometry, ArrangedChildren );
+		ArrangeChildrenAlong<EOrientation::Orient_Vertical>(GSlateFlowDirection, this->Children, AllottedGeometry, ArrangedChildren );
 	}
 }
 
@@ -191,8 +217,8 @@ void SBoxPanel::OnArrangeChildren( const FGeometry& AllottedGeometry, FArrangedC
  *
  * @return The size desired by the children given an orientation.
  */
-template<EOrientation Orientation>
-static FVector2D ComputeDesiredSizeForBox( const TPanelChildren<SBoxPanel::FSlot>& Children ) 
+template<EOrientation Orientation, typename SlotType>
+static FVector2D ComputeDesiredSizeForBox( const TPanelChildren<SlotType>& Children )
 {
 	// The desired size of this panel is the total size desired by its children plus any margins specified in this panel.
 	// The layout along the panel's axis is describe dy the SizeParam, while the perpendicular layout is described by the
@@ -200,7 +226,7 @@ static FVector2D ComputeDesiredSizeForBox( const TPanelChildren<SBoxPanel::FSlot
 	FVector2D MyDesiredSize(0,0);
 	for( int32 ChildIndex=0; ChildIndex < Children.Num(); ++ChildIndex )
 	{
-		const SBoxPanel::FSlot& CurChild = Children[ChildIndex];
+		const SlotType& CurChild = Children[ChildIndex];
 		
 		if ( CurChild.GetWidget()->GetVisibility() != EVisibility::Collapsed )
 		{
@@ -211,33 +237,33 @@ static FVector2D ComputeDesiredSizeForBox( const TPanelChildren<SBoxPanel::FSlot
 			{
 				// For a vertical panel, we want to find the maximum desired width (including margin).
 				// That will be the desired width of the whole panel.
-				MyDesiredSize.X = FMath::Max(MyDesiredSize.X, CurChildDesiredSize.X + CurChild.GetPadding().GetTotalSpaceAlong<Orient_Horizontal>());
+				MyDesiredSize.X = FMath::Max(MyDesiredSize.X, CurChildDesiredSize.X + CurChild.GetPadding().template GetTotalSpaceAlong<Orient_Horizontal>());
 
 				// Clamp to the max size if it was specified
 				float FinalChildDesiredSize = CurChildDesiredSize.Y;
-				float MaxSize = CurChild.MaxSize.Get();
+				float MaxSize = CurChild.GetMaxSize();
 				if( MaxSize > 0 )
 				{
 					FinalChildDesiredSize = FMath::Min( MaxSize, FinalChildDesiredSize );
 				}
 
-				MyDesiredSize.Y += FinalChildDesiredSize + CurChild.GetPadding().GetTotalSpaceAlong<Orient_Vertical>();
+				MyDesiredSize.Y += FinalChildDesiredSize + CurChild.GetPadding().template GetTotalSpaceAlong<Orient_Vertical>();
 			}
 			else
 			{
 				// A horizontal panel is just a sideways vertical panel: the axes are swapped.
 
-				MyDesiredSize.Y = FMath::Max(MyDesiredSize.Y, CurChildDesiredSize.Y + CurChild.GetPadding().GetTotalSpaceAlong<Orient_Vertical>());
+				MyDesiredSize.Y = FMath::Max(MyDesiredSize.Y, CurChildDesiredSize.Y + CurChild.GetPadding().template GetTotalSpaceAlong<Orient_Vertical>());
 
 				// Clamp to the max size if it was specified
 				float FinalChildDesiredSize = CurChildDesiredSize.X;
-				float MaxSize = CurChild.MaxSize.Get();
+				float MaxSize = CurChild.GetMaxSize();
 				if( MaxSize > 0 )
 				{
 					FinalChildDesiredSize = FMath::Min( MaxSize, FinalChildDesiredSize );
 				}
 
-				MyDesiredSize.X += FinalChildDesiredSize + CurChild.GetPadding().GetTotalSpaceAlong<Orient_Horizontal>();
+				MyDesiredSize.X += FinalChildDesiredSize + CurChild.GetPadding().template GetTotalSpaceAlong<Orient_Horizontal>();
 			}
 		}
 	}
@@ -245,12 +271,6 @@ static FVector2D ComputeDesiredSizeForBox( const TPanelChildren<SBoxPanel::FSlot
 	return MyDesiredSize;
 }
 
-/**
- * A Panel's desired size in the space required to arrange of its children on the screen while respecting all of
- * the children's desired sizes and any layout-related options specified by the user. See StackPanel for an example.
- *
- * @return The desired size.
- */
 FVector2D SBoxPanel::ComputeDesiredSize( float ) const
 {
 	return (Orientation == Orient_Horizontal)
@@ -258,7 +278,6 @@ FVector2D SBoxPanel::ComputeDesiredSize( float ) const
 		: ComputeDesiredSizeForBox<Orient_Vertical>(this->Children);
 }
 
-/** @return  The children of a panel in a slot-agnostic way. */
 FChildren* SBoxPanel::GetChildren()
 {
 	return &Children;
@@ -274,11 +293,6 @@ void SBoxPanel::ClearChildren()
 	Children.Empty();
 }
 
-/**
- * A Box Panel's orientation cannot be changed once it is constructed..
- *
- * @param InOrientation   The orientation of the Box Panel
- */
 SBoxPanel::SBoxPanel( EOrientation InOrientation )
 	: Children(this)
 	, Orientation(InOrientation)
@@ -313,16 +327,14 @@ FReply SDragAndDropVerticalBox::OnMouseButtonDown(const FGeometry& MyGeometry, c
 
 FReply SDragAndDropVerticalBox::OnDragDetected(const FGeometry& MyGeometry, const FPointerEvent& MouseEvent)
 {
-	TPanelChildren<FSlot>* PanelChildren = (TPanelChildren<FSlot>*)GetChildren();
 	FArrangedChildren ArrangedChildren(EVisibility::Visible);
 	ArrangeChildren(MyGeometry, ArrangedChildren);
 
 	int32 NodeUnderPositionIndex = SWidget::FindChildUnderMouse(ArrangedChildren, MouseEvent);
 
-	if (PanelChildren->IsValidIndex(NodeUnderPositionIndex))
+	if (Children.IsValidIndex(NodeUnderPositionIndex))
 	{
-		SVerticalBox::FSlot* Slot = &(*PanelChildren)[NodeUnderPositionIndex];
-
+		SVerticalBox::FSlot* Slot = static_cast<SVerticalBox::FSlot*>(&static_cast<FSlotBase&>(Children[NodeUnderPositionIndex]));
 		if (OnDragDetected_Handler.IsBound())
 		{
 			return OnDragDetected_Handler.Execute(MyGeometry, MouseEvent, NodeUnderPositionIndex, Slot);
@@ -396,11 +408,9 @@ FReply SDragAndDropVerticalBox::OnDragOver(const FGeometry& MyGeometry, const FD
 				FVector2D LocalPointerPos = ArrangedChildren[DragOverSlotIndex].Geometry.AbsoluteToLocal(DragDropEvent.GetScreenSpacePosition());
 				EItemDropZone ItemHoverZone = ZoneFromPointerPosition(LocalPointerPos, ArrangedChildren[DragOverSlotIndex].Geometry, ArrangedChildren[DragOp->SlotIndexBeingDragged].Geometry);
 
-				TPanelChildren<FSlot>* PanelChildren = (TPanelChildren<FSlot>*)GetChildren();
-
-				if (PanelChildren->IsValidIndex(DragOverSlotIndex))
+				if (Children.IsValidIndex(DragOverSlotIndex))
 				{
-					SVerticalBox::FSlot* Slot = &(*PanelChildren)[DragOverSlotIndex];
+					SVerticalBox::FSlot* Slot = static_cast<SVerticalBox::FSlot*>(&static_cast<FSlotBase&>(Children[DragOverSlotIndex]));
 
 					ItemDropZone = OnCanAcceptDrop.Execute(DragDropEvent, ItemHoverZone, Slot);
 					CurrentDragOperationScreenSpaceLocation = DragDropEvent.GetScreenSpacePosition();
@@ -427,11 +437,9 @@ FReply SDragAndDropVerticalBox::OnDrop(const FGeometry& MyGeometry, const FDragD
 		if (DragDropEvent.GetOperationAs<FDragAndDropVerticalBoxOp>().IsValid())
 		{
 			int32 NodeUnderPositionIndex = SWidget::FindChildUnderPosition(ArrangedChildren, DragDropEvent.GetScreenSpacePosition());
-			TPanelChildren<FSlot>* PanelChildren = (TPanelChildren<FSlot>*)GetChildren();
-
-			if (PanelChildren->IsValidIndex(NodeUnderPositionIndex))
+			if (Children.IsValidIndex(NodeUnderPositionIndex))
 			{
-				SVerticalBox::FSlot* Slot = &(*PanelChildren)[NodeUnderPositionIndex];
+				SVerticalBox::FSlot* Slot = static_cast<SVerticalBox::FSlot*>(&static_cast<FSlotBase&>(Children[NodeUnderPositionIndex]));
 				TOptional<EItemDropZone> ReportedZone = ItemDropZone;
 
 				if (OnCanAcceptDrop.IsBound() && ItemDropZone.IsSet())
@@ -448,7 +456,7 @@ FReply SDragAndDropVerticalBox::OnDrop(const FGeometry& MyGeometry, const FDragD
 						TSharedPtr<FDragAndDropVerticalBoxOp> DragOp = DragDropEvent.GetOperationAs<FDragAndDropVerticalBoxOp>();
 
 						// Perform the slot changes
-						PanelChildren->Move(DragOp->SlotIndexBeingDragged, NodeUnderPositionIndex);
+						Children.Move(DragOp->SlotIndexBeingDragged, NodeUnderPositionIndex);
 					}
 				}
 			}
diff --git a/Engine/Source/Runtime/SlateCore/Private/Widgets/SWindow.cpp b/Engine/Source/Runtime/SlateCore/Private/Widgets/SWindow.cpp
--- a/Engine/Source/Runtime/SlateCore/Private/Widgets/SWindow.cpp
+++ b/Engine/Source/Runtime/SlateCore/Private/Widgets/SWindow.cpp
@@ -50,9 +50,7 @@ public:
 		{
 			_Visibility = EVisibility::SelfHitTestInvisible;
 		}
-
-		SLATE_SUPPORTS_SLOT( FPopupLayerSlot )
-
+		SLATE_SLOT_ARGUMENT( FPopupLayerSlot, Slots )
 	SLATE_END_ARGS()
 
 	SPopupLayer()
@@ -65,33 +63,20 @@ public:
 
 		OwnerWindow = InWindow;
 
-		const int32 NumSlots = InArgs.Slots.Num();
-		for ( int32 SlotIndex = 0; SlotIndex < NumSlots; ++SlotIndex )
-		{
-			Children.Add( InArgs.Slots[SlotIndex] );
-		}
+		Children.AddSlots(MoveTemp(const_cast<TArray<FPopupLayerSlot::FSlotArguments>&>(InArgs._Slots)));
 	}
 
 	/** Make a new ListPanel::Slot  */
-	FPopupLayerSlot& Slot()
+	static FPopupLayerSlot::FSlotArguments Slot()
 	{
-		return *(new FPopupLayerSlot());
+		return FPopupLayerSlot::FSlotArguments(MakeUnique<FPopupLayerSlot>());
 	}
 
+	using FScopedWidgetSlotArguments = TPanelChildren<FPopupLayerSlot>::FScopedWidgetSlotArguments;
 	/** Add a slot to the ListPanel */
-	FPopupLayerSlot& AddSlot(int32 InsertAtIndex = INDEX_NONE)
+	FScopedWidgetSlotArguments AddSlot(int32 InsertAtIndex = INDEX_NONE)
 	{
-		FPopupLayerSlot& NewSlot = *new FPopupLayerSlot();
-		if (InsertAtIndex == INDEX_NONE)
-		{
-			this->Children.Add( &NewSlot );
-		}
-		else
-		{
-			this->Children.Insert( &NewSlot, InsertAtIndex );
-		}
-
-		return NewSlot;
+		return { MakeUnique<FPopupLayerSlot>(), Children, InsertAtIndex };
 	}
 
 	void RemoveSlot(const TSharedRef<SWidget>& WidgetToRemove)
@@ -1325,7 +1310,7 @@ TSharedPtr<FPopupLayer> SWindow::OnVisualizePopup(const TSharedRef<SWidget>& Pop
 }
 
 /** Return a new slot in the popup layer. Assumes that the window has a popup layer. */
-FPopupLayerSlot& SWindow::AddPopupLayerSlot()
+SWindow::FScopedWidgetSlotArguments SWindow::AddPopupLayerSlot()
 {
 	ensure( PopupLayer.IsValid() );
 	return PopupLayer->AddSlot();
diff --git a/Engine/Source/Runtime/SlateCore/Public/Widgets/SBoxPanel.h b/Engine/Source/Runtime/SlateCore/Public/Widgets/SBoxPanel.h
--- a/Engine/Source/Runtime/SlateCore/Public/Widgets/SBoxPanel.h
+++ b/Engine/Source/Runtime/SlateCore/Public/Widgets/SBoxPanel.h
@@ -21,44 +21,135 @@ class FArrangedChildren;
 /**
  * A BoxPanel contains one child and describes how that child should be arranged on the screen.
  */
-class SLATECORE_API SBoxPanel
-	: public SPanel
+class SLATECORE_API SBoxPanel : public SPanel
 {
-public:
-
+protected:
 	/**
 	 * A BoxPanel contains one BoxPanel child and describes how that
 	 * child should be arranged on the screen.
 	 */
-	class FSlot : public TBasicLayoutWidgetSlot<FSlot>
+	template<typename SlotType>
+	class TSlot : public TBasicLayoutWidgetSlot<SlotType>
 	{
-	public:		
-		/** Horizontal and Vertical Boxes inherit from FSlot */
-		virtual ~FSlot(){}
+	public:
+		SLATE_SLOT_BEGIN_ARGS(TSlot, TBasicLayoutWidgetSlot<SlotType>)
+			SLATE_ARGUMENT(TOptional<FSizeParam>, SizeParam)
+			TAttribute<float> _MaxSize;
+		SLATE_SLOT_END_ARGS()
+
+	protected:
+		/** Default values for a slot. */
+		TSlot()
+			: TBasicLayoutWidgetSlot<SlotType>(HAlign_Fill, VAlign_Fill)
+			, SizeParam(FStretch(1.f))
+			, MaxSize(0.0f)
+		{ }
+
+	public:
+		void Construct(const FChildren& SlotOwner, FSlotArguments&& InArgs)
+		{
+			TBasicLayoutWidgetSlot<SlotType>::Construct(SlotOwner, MoveTemp(InArgs));
+			if (InArgs._SizeParam.IsSet())
+			{
+				SizeParam = MoveTemp(InArgs._SizeParam.GetValue());
+			}
+			if (InArgs._MaxSize.IsSet())
+			{
+				MaxSize = MoveTemp(InArgs._MaxSize);
+			}
+		}
+
+		/** Get the space param this slot should occupy along panel's direction. */
+		const FSizeParam& GetSizeParam() const { return SizeParam; }
+		/** Get the max size the slot can be.*/
+		float GetMaxSize() const { return MaxSize.Get(); }
+
+	public:
+		/** Set the size Param of the slot, It could be a FStretch or a FAuto. */
+		void SetSizeParam(FSizeParam InSizeParam)
+		{
+			SizeParam = MoveTemp(InSizeParam);
+			FSlotBase::Invalidate(EInvalidateWidget::Layout);
+		}
 
+		/** The widget's DesiredSize will be used as the space required. */
+		void SetSizeToAuto()
+		{
+			SetSizeParam(FAuto());
+		}
+
+		/** The available space will be distributed proportionately. */
+		void SetSizeToStretch(TAttribute<float> StretchCoefficient)
+		{
+			SetSizeParam(FStretch(MoveTemp(StretchCoefficient)));
+		}
+
+		/** Set the max size in SlateUnit this slot can be. */
+		void SetMaxSize(TAttribute<float> InMaxSize)
+		{
+			MaxSize = MoveTemp(InMaxSize);
+			FSlotBase::Invalidate(EInvalidateWidget::Layout);
+		}
+
+	private:
 		/**
-		* How much space this slot should occupy along panel's direction.
-		*   When SizeRule is SizeRule_Auto, the widget's DesiredSize will be used as the space required.
-		*   When SizeRule is SizeRule_Stretch, the available space will be distributed proportionately between
-		*   peer Widgets depending on the Value property. Available space is space remaining after all the
-		*   peers' SizeRule_Auto requirements have been satisfied.
-		*/
+		 * How much space this slot should occupy along panel's direction.
+		 * When SizeRule is SizeRule_Auto, the widget's DesiredSize will be used as the space required.
+		 * When SizeRule is SizeRule_Stretch, the available space will be distributed proportionately between
+		 * peer Widgets depending on the Value property. Available space is space remaining after all the
+		 * peers' SizeRule_Auto requirements have been satisfied.
+		 */
 		FSizeParam SizeParam;
-		
+
 		/** The max size that this slot can be (0 if no max) */
 		TAttribute<float> MaxSize;
-			
-	protected:
-		/** Default values for a slot. */
-		FSlot()
-			: TBasicLayoutWidgetSlot<FSlot>(HAlign_Fill, VAlign_Fill)
-			, SizeParam( FStretch(1) )
-			, MaxSize( 0.0f )
-		{ }
 	};
 
 public:
+	class FSlot : public TSlot<FSlot>
+	{
+	};
+
+protected:
+	template<typename SlotType>
+	struct FScopedWidgetSlotArguments : public SlotType::FSlotArguments
+	{
+	public:
+		FScopedWidgetSlotArguments(TUniquePtr<SlotType> InSlot, TPanelChildren<FSlot>& InChildren, int32 InIndex)
+			: SlotType::FSlotArguments(MoveTemp(InSlot))
+			, Children(InChildren)
+			, Index(InIndex)
+		{
+		}
+
+		FScopedWidgetSlotArguments() = delete;
+		FScopedWidgetSlotArguments(const FScopedWidgetSlotArguments&) = delete;
+		FScopedWidgetSlotArguments& operator=(const FScopedWidgetSlotArguments&) = delete;
+		FScopedWidgetSlotArguments(FScopedWidgetSlotArguments&&) = default;
+		FScopedWidgetSlotArguments& operator=(FScopedWidgetSlotArguments&&) = default;
+
+		virtual ~FScopedWidgetSlotArguments()
+		{
+			if (this->GetSlot())	// Is nullptr when the FScopedWidgetSlotArguments is moved-constructed.
+			{
+				SBoxPanel::FSlot::FSlotArguments* SelfAsBaseSlot = static_cast<SBoxPanel::FSlot::FSlotArguments*>(static_cast<FSlotBase::FSlotArguments*>(this));
+				if (Index == INDEX_NONE)
+				{
+					Children.AddSlot(MoveTemp(*SelfAsBaseSlot));
+				}
+				else
+				{
+					Children.InsertSlot(MoveTemp(*SelfAsBaseSlot), Index);
+				}
+			}
+		}
+
+	private:
+		TPanelChildren<FSlot>& Children;
+		int32 Index;
+	};
 
+public:
 	/** Removes a slot from this box panel which contains the specified SWidget
 	 *
 	 * @param SlotWidget The widget to match when searching through the slots
@@ -66,22 +157,22 @@ public:
 	 */
 	int32 RemoveSlot( const TSharedRef<SWidget>& SlotWidget );
 
-	/**
-	 * Removes all children from the box.
-	 */
+	/** Removes all children from the box. */
 	void ClearChildren();
 
-public:
+	/** @return the number of slots. */
+	int32 NumSlots() const { return Children.Num(); }
 
-	// Begin SWidget overrides.
+	/** @return if it's a valid index slot index. */
+	bool IsValidSlotIndex(int32 Index) const { return Children.IsValidIndex(Index); }
+
+public:
+	//~ Begin SWidget overrides.
 	virtual void OnArrangeChildren(const FGeometry& AllottedGeometry, FArrangedChildren& ArrangedChildren) const override;
 	virtual FChildren* GetChildren() override;
-	// End SWidget overrides.
-
 protected:
-	// Begin SWidget overrides.
 	virtual FVector2D ComputeDesiredSize(float) const override;
-	// End SWidget overrides.
+	//~ End SWidget overrides.
 
 	/**
 	 * A Box Panel's orientation cannot be changed once it is constructed..
@@ -102,126 +193,91 @@ protected:
 class SLATECORE_API SHorizontalBox : public SBoxPanel
 {
 public:
-	class FSlot : public SBoxPanel::FSlot
+	class FSlot : public SBoxPanel::TSlot<FSlot>
 	{
-		public:
-		FSlot()
-		: SBoxPanel::FSlot()
-		{
-		}
-
-		FSlot& AutoWidth()
-		{
-			SizeParam = FAuto();
-			return *this;
-		}
-
-		FSlot& MaxWidth( const TAttribute< float >& InMaxWidth )
-		{
-			MaxSize = InMaxWidth;
-			return *this;
-		}
-
-		FSlot& FillWidth( const TAttribute< float >& StretchCoefficient )
+	public:
+		SLATE_SLOT_BEGIN_ARGS(FSlot, SBoxPanel::TSlot<FSlot>)
+			/** The widget's DesiredSize will be used as the space required. */
+			FSlotArguments& AutoWidth()
+			{
+				_SizeParam = FAuto();
+				return Me();
+			}
+			/** The available space will be distributed proportionately. */
+			FSlotArguments& FillWidth(TAttribute<float> InStretchCoefficient)
+			{
+				_SizeParam = FStretch(MoveTemp(InStretchCoefficient));
+				return Me();
+			}
+			/** Set the max size in SlateUnit this slot can be. */
+			FSlotArguments& MaxWidth(TAttribute<float> InMaxWidth)
+			{
+				_MaxSize = MoveTemp(InMaxWidth);
+				return Me();
+			}
+		SLATE_SLOT_END_ARGS()
+
+		/** The widget's DesiredSize will be used as the space required. */
+		void SetAutoWidth()
 		{
-			SizeParam = FStretch( StretchCoefficient );
-			return *this;
+			SetSizeToAuto();
 		}
 
-		FSlot& Padding( float Uniform )
+		/** The available space will be distributed proportionately. */
+		void SetFillWidth(TAttribute<float> InStretchCoefficient)
 		{
-			SBoxPanel::FSlot::Padding(Uniform);
-			return *this;
+			SetSizeToStretch(MoveTemp(InStretchCoefficient));
 		}
 
-		FSlot& Padding( float Horizontal, float Vertical )
-		{
-			SBoxPanel::FSlot::Padding(Horizontal, Vertical);
-			return *this;
-		}
-
-		FSlot& Padding( float Left, float Top, float Right, float Bottom )
-		{
-			SBoxPanel::FSlot::Padding(Left, Top, Right, Bottom);
-			return *this;
-		}
-		
-		FSlot& Padding( TAttribute<FMargin> InPadding )
+		/** Set the max size in SlateUnit this slot can be. */
+		void SetMaxWidth(TAttribute<float> InMaxWidth)
 		{
-			SBoxPanel::FSlot::Padding(MoveTemp(InPadding));
-			return *this;
+			SetMaxSize(MoveTemp(InMaxWidth));
 		}
 
-		FSlot& HAlign(EHorizontalAlignment InHAlignment)
+		void Construct(const FChildren& SlotOwner, FSlotArguments&& InArgs)
 		{
-			SBoxPanel::FSlot::HAlign(InHAlignment);
-			return *this;
+			SBoxPanel::TSlot<FSlot>::Construct(SlotOwner, MoveTemp(InArgs));
 		}
 
-		FSlot& VAlign(EVerticalAlignment InVAlignment)
-		{
-			SBoxPanel::FSlot::VAlign(InVAlignment);
-			return *this;
-		}
+		UE_DEPRECATED(5.0, "Chained AutoWidth is deprecated. Use the FSlotArgument or SetAutoWidth")
+		FSlot& AutoWidth() { SetAutoWidth(); return *this; }
 
-		FSlot& operator[]( TSharedRef<SWidget> InWidget )
-		{
-			SBoxPanel::FSlot::operator[](InWidget);
-			return *this;
-		}
+		UE_DEPRECATED(5.0, "Chained FillWidth is deprecated. Use the FSlotArgument or SetFillWidth")
+		FSlot& FillWidth(TAttribute<float> InStretchCoefficient) { SetFillWidth(InStretchCoefficient); return *this; }
 
-		FSlot& Expose( FSlot*& OutVarToInit )
-		{
-			OutVarToInit = this;
-			return *this;
-		}
+		UE_DEPRECATED(5.0, "Chained MaxWidth is deprecated. Use the FSlotArgument or SetMaxWidth")
+		FSlot& MaxWidth(TAttribute<float> InMaxWidth) { SetMaxWidth(InMaxWidth); return *this; }
 	};
 
-	static FSlot& Slot()
+	static FSlot::FSlotArguments Slot()
 	{
-		return *(new FSlot());
+		return FSlot::FSlotArguments(MakeUnique<FSlot>());
 	}
 
 	SLATE_BEGIN_ARGS( SHorizontalBox )
 	{
 		_Visibility = EVisibility::SelfHitTestInvisible;
 	}
-
-		SLATE_SUPPORTS_SLOT(SHorizontalBox::FSlot)
-
+		SLATE_SLOT_ARGUMENT(SHorizontalBox::FSlot, Slots)
 	SLATE_END_ARGS()
 
-	FSlot& AddSlot()
+	using FScopedWidgetSlotArguments = SBoxPanel::FScopedWidgetSlotArguments<SHorizontalBox::FSlot>;
+	FScopedWidgetSlotArguments AddSlot()
 	{
-		SHorizontalBox::FSlot& NewSlot = *new SHorizontalBox::FSlot();
-		this->Children.Add( &NewSlot );
-
-		Invalidate(EInvalidateWidgetReason::Layout);
-
-		return NewSlot;
+		return InsertSlot(INDEX_NONE);
 	}
 
-	FSlot& InsertSlot(int32 Index = INDEX_NONE)
+	FScopedWidgetSlotArguments InsertSlot(int32 Index = INDEX_NONE)
 	{
-		if (Index == INDEX_NONE)
-		{
-			return AddSlot();
-		}
-		SHorizontalBox::FSlot& NewSlot = *new SHorizontalBox::FSlot();
-		this->Children.Insert(&NewSlot, Index);
-
-		Invalidate(EInvalidateWidgetReason::Layout);
-
-		return NewSlot;
+		return FScopedWidgetSlotArguments(MakeUnique<FSlot>(), this->Children, Index);
 	}
 
-	int32 NumSlots() const
-	{
-		return this->Children.Num();
-	}
+	FSlot& GetSlot(int32 SlotIndex);
+	const FSlot& GetSlot(int32 SlotIndex) const;
 
 	FORCENOINLINE SHorizontalBox()
-	: SBoxPanel( Orient_Horizontal )
+		: SBoxPanel( Orient_Horizontal )
 	{
 		SetCanTick(false);
 		bCanSupportFocus = false;
@@ -239,85 +295,66 @@ public:
 class SLATECORE_API SVerticalBox : public SBoxPanel
 {
 public:
-	class FSlot : public SBoxPanel::FSlot
+	class FSlot : public SBoxPanel::TSlot<FSlot>
 	{
-		public:
-
-		FSlot()
-		: SBoxPanel::FSlot()
-		{
-		}
-
-		FSlot& AutoHeight()
+	public:
+		SLATE_SLOT_BEGIN_ARGS(FSlot, SBoxPanel::TSlot<FSlot>)
+			/** The widget's DesiredSize will be used as the space required. */
+			FSlotArguments& AutoHeight()
+			{
+				_SizeParam = FAuto();
+				return Me();
+			}
+			/** The available space will be distributed proportionately. */
+			FSlotArguments& FillHeight(TAttribute<float> InStretchCoefficient)
+			{
+				_SizeParam = FStretch(MoveTemp(InStretchCoefficient));
+				return Me();
+			}
+			/** Set the max size in SlateUnit this slot can be. */
+			FSlotArguments& MaxHeight(TAttribute<float> InMaxHeight)
+			{
+				_MaxSize = MoveTemp(InMaxHeight);
+				return Me();
+			}
+		SLATE_SLOT_END_ARGS()
+
+		/** The widget's DesiredSize will be used as the space required. */
+		void SetAutoHeight()
 		{
-			SizeParam = FAuto();
-			return *this;
+			SetSizeToAuto();
 		}
 
-		FSlot& MaxHeight( const TAttribute< float >& InMaxHeight )
+		/** The available space will be distributed proportionately. */
+		void SetFillHeight(TAttribute<float> InStretchCoefficient)
 		{
-			MaxSize = InMaxHeight;
-			return *this;
+			SetSizeToStretch(MoveTemp(InStretchCoefficient));
 		}
 
-		FSlot& FillHeight( const TAttribute< float >& StretchCoefficient )
+		/** Set the max size in SlateUnit this slot can be. */
+		void SetMaxHeight(TAttribute<float> InMaxHeight)
 		{
-			SizeParam = FStretch( StretchCoefficient );
-			return *this;
+			SetMaxSize(MoveTemp(InMaxHeight));
 		}
 
-		FSlot& Padding(float Uniform)
+		void Construct(const FChildren& SlotOwner, FSlotArguments&& InArgs)
 		{
-			SBoxPanel::FSlot::Padding(Uniform);
-			return *this;
+			SBoxPanel::TSlot<FSlot>::Construct(SlotOwner, MoveTemp(InArgs));
 		}
 
-		FSlot& Padding(float Horizontal, float Vertical)
-		{
-			SBoxPanel::FSlot::Padding(Horizontal, Vertical);
-			return *this;
-		}
-
-		FSlot& Padding(float Left, float Top, float Right, float Bottom)
-		{
-			SBoxPanel::FSlot::Padding(Left, Top, Right, Bottom);
-			return *this;
-		}
-
-		FSlot& Padding(TAttribute<FMargin> InPadding)
-		{
-			SBoxPanel::FSlot::Padding(MoveTemp(InPadding));
-			return *this;
-		}
-
-		FSlot& HAlign( EHorizontalAlignment InHAlignment )
-		{
-			SBoxPanel::FSlot::HAlign(InHAlignment);
-			return *this;
-		}
-
-		FSlot& VAlign( EVerticalAlignment InVAlignment )
-		{
-			SBoxPanel::FSlot::VAlign(InVAlignment);
-			return *this;
-		}
+		UE_DEPRECATED(5.0, "Chained AutoHeight is deprecated. Use the FSlotArgument or SetAutoHeight")
+		FSlot& AutoHeight() { SetAutoHeight(); return *this; }
 
-		FSlot& operator[]( TSharedRef<SWidget> InWidget )
-		{
-			SBoxPanel::FSlot::operator[](InWidget);
-			return *this;
-		}
+		UE_DEPRECATED(5.0, "Chained FillWidth is deprecated. Use the FSlotArgument or SetFillWidth")
+		FSlot& FillHeight(TAttribute<float> InStretchCoefficient) { SetFillHeight(InStretchCoefficient); return *this; }
 
-		FSlot& Expose( FSlot*& OutVarToInit )
-		{
-			OutVarToInit = this;
-			return *this;
-		}
+		UE_DEPRECATED(5.0, "Chained MaxWidth is deprecated. Use the FSlotArgument or SetMaxHeight")
+		FSlot& MaxHeight(TAttribute<float> InMaxHeight) { SetMaxHeight(InMaxHeight); return *this; }
 	};
 
-	static FSlot& Slot()
+	static FSlot::FSlotArguments Slot()
 	{
-		return *(new FSlot());
+		return FSlot::FSlotArguments(MakeUnique<FSlot>());
 	}
 
 
@@ -325,42 +362,25 @@ public:
 	{
 		_Visibility = EVisibility::SelfHitTestInvisible;
 	}
-
-		SLATE_SUPPORTS_SLOT(SVerticalBox::FSlot)
-
+		SLATE_SLOT_ARGUMENT(SVerticalBox::FSlot, Slots)
 	SLATE_END_ARGS()
 
-	FSlot& AddSlot()
+	using FScopedWidgetSlotArguments = SBoxPanel::FScopedWidgetSlotArguments<SVerticalBox::FSlot>;
+	FScopedWidgetSlotArguments AddSlot()
 	{
-		SVerticalBox::FSlot& NewSlot = *new SVerticalBox::FSlot();
-		this->Children.Add( &NewSlot );
-
-		Invalidate(EInvalidateWidgetReason::Layout);
-
-		return NewSlot;
+		return InsertSlot(INDEX_NONE);
 	}
 
-	FSlot& InsertSlot(int32 Index = INDEX_NONE)
+	FScopedWidgetSlotArguments InsertSlot(int32 Index = INDEX_NONE)
 	{
-		if (Index == INDEX_NONE)
-		{
-			return AddSlot();
-		}
-		SVerticalBox::FSlot& NewSlot = *new SVerticalBox::FSlot();
-		this->Children.Insert(&NewSlot, Index);
-
-		Invalidate(EInvalidateWidgetReason::Layout);
-
-		return NewSlot;
+		return FScopedWidgetSlotArguments(MakeUnique<FSlot>(), this->Children, Index);
 	}
 
-	int32 NumSlots() const
-	{
-		return this->Children.Num();
-	}
+	FSlot& GetSlot(int32 SlotIndex);
+	const FSlot& GetSlot(int32 SlotIndex) const;
 
 	FORCENOINLINE SVerticalBox()
-	: SBoxPanel( Orient_Vertical )
+		: SBoxPanel( Orient_Vertical )
 	{
 		SetCanTick(false);
 		bCanSupportFocus = false;
diff --git a/Engine/Source/Runtime/SlateCore/Public/Widgets/SWindow.h b/Engine/Source/Runtime/SlateCore/Public/Widgets/SWindow.h
--- a/Engine/Source/Runtime/SlateCore/Public/Widgets/SWindow.h
+++ b/Engine/Source/Runtime/SlateCore/Public/Widgets/SWindow.h
@@ -121,6 +121,119 @@ private:
 	TSharedPtr<SOverlay> Overlay;
 };
 
+
+/**
+ * Popups, tooltips, drag and drop decorators all can be executed without creating a new window.
+ * This slot along with the SWindow::AddPopupLayerSlot() API enabled it.
+ */
+struct FPopupLayerSlot : public TSlotBase<FPopupLayerSlot>
+{
+public:
+	FPopupLayerSlot()
+		: TSlotBase<FPopupLayerSlot>()
+		, DesktopPosition_Attribute(FVector2D::ZeroVector)
+		, WidthOverride_Attribute()
+		, HeightOverride_Attribute()
+		, Scale_Attribute(1.0f)
+		, Clamp_Attribute(false)
+		, ClampBuffer_Attribute(FVector2D::ZeroVector)
+	{}
+
+	SLATE_SLOT_BEGIN_ARGS(FPopupLayerSlot, TSlotBase<FPopupLayerSlot>)
+		/** Pixel position in desktop space */
+		SLATE_ATTRIBUTE(FVector2D, DesktopPosition)
+		/** Width override in pixels */
+		SLATE_ATTRIBUTE(float, WidthOverride)
+		/** Width override in pixels */
+		SLATE_ATTRIBUTE(float, HeightOverride)
+		/** DPI scaling to be applied to the contents of this slot */
+		SLATE_ATTRIBUTE(float, Scale)
+		/** Should this slot be kept within the parent window */
+		SLATE_ATTRIBUTE(bool, ClampToWindow)
+		/** If this slot is kept within the parent window, how far from the edges should we clamp it */
+		SLATE_ATTRIBUTE(FVector2D, ClampBuffer)
+	SLATE_SLOT_END_ARGS()
+
+	void Construct(const FChildren& SlotOwner, FSlotArguments&& InArgs)
+	{
+		if (InArgs._DesktopPosition.IsSet())
+		{
+			SetDesktopPosition(MoveTemp(InArgs._DesktopPosition));
+		}
+		if (InArgs._WidthOverride.IsSet())
+		{
+			SetWidthOverride(MoveTemp(InArgs._WidthOverride));
+		}
+		if (InArgs._HeightOverride.IsSet())
+		{
+			SetHeightOverride(MoveTemp(InArgs._HeightOverride));
+		}
+		if (InArgs._Scale.IsSet())
+		{
+			SetScale(MoveTemp(InArgs._Scale));
+		}
+		if (InArgs._ClampToWindow.IsSet())
+		{
+			SetClampToWindow(MoveTemp(InArgs._ClampToWindow));
+		}
+		if (InArgs._ClampBuffer.IsSet())
+		{
+			SetClampBuffer(MoveTemp(InArgs._ClampBuffer));
+		}
+		TSlotBase<FPopupLayerSlot>::Construct(SlotOwner, MoveTemp(InArgs));
+	}
+
+	/** Pixel position in desktop space */
+	void SetDesktopPosition(TAttribute<FVector2D> InDesktopPosition)
+	{
+		DesktopPosition_Attribute = MoveTemp(InDesktopPosition);
+	}
+
+	/** Width override in pixels */
+	void SetWidthOverride(TAttribute<float> InWidthOverride)
+	{
+		WidthOverride_Attribute = MoveTemp(InWidthOverride);
+	}
+
+	/** Width override in pixels */
+	void SetHeightOverride(TAttribute<float> InHeightOverride)
+	{
+		HeightOverride_Attribute = MoveTemp(InHeightOverride);
+	}
+
+	/** DPI scaling to be applied to the contents of this slot */
+	void SetScale(TAttribute<float> InScale)
+	{
+		Scale_Attribute = MoveTemp(InScale);
+	}
+
+	/** Should this slot be kept within the parent window */
+	void SetClampToWindow(TAttribute<bool> InClamp_Attribute)
+	{
+		Clamp_Attribute = MoveTemp(InClamp_Attribute);
+	}
+
+	/** If this slot is kept within the parent window, how far from the edges should we clamp it */
+	void SetClampBuffer(TAttribute<FVector2D> InClampBuffer_Attribute)
+	{
+		ClampBuffer_Attribute = MoveTemp(InClampBuffer_Attribute);
+	}
+
+private:
+	/** SPopupLayer arranges FPopupLayerSlots, so it needs to know all about */
+	friend class SPopupLayer;
+	/** TPanelChildren need access to the Widget member */
+	friend class TPanelChildren<FPopupLayerSlot>;
+
+	TAttribute<FVector2D> DesktopPosition_Attribute;
+	TAttribute<float> WidthOverride_Attribute;
+	TAttribute<float> HeightOverride_Attribute;
+	TAttribute<float> Scale_Attribute;
+	TAttribute<bool> Clamp_Attribute;
+	TAttribute<FVector2D> ClampBuffer_Attribute;
+};
+
+
 /**
  * SWindow is a platform-agnostic representation of a top-level window.
  */
@@ -551,8 +664,9 @@ public:
 	 */
 	virtual TSharedPtr<FPopupLayer> OnVisualizePopup(const TSharedRef<SWidget>& PopupContent) override;
 
+	using FScopedWidgetSlotArguments = TPanelChildren<FPopupLayerSlot>::FScopedWidgetSlotArguments;
 	/** Return a new slot in the popup layer. Assumes that the window has a popup layer. */
-	struct FPopupLayerSlot& AddPopupLayerSlot();
+	FScopedWidgetSlotArguments AddPopupLayerSlot();
 
 	/** Counterpart to AddPopupLayerSlot */
 	void RemovePopupLayerSlot( const TSharedRef<SWidget>& WidgetToRemove );
@@ -1237,81 +1351,6 @@ private:
 	TWeakPtr<FActiveTimerHandle> ActiveTimerHandle;
 };
 
-
-/**
- * Popups, tooltips, drag and drop decorators all can be executed without creating a new window.
- * This slot along with the SWindow::AddPopupLayerSlot() API enabled it.
- */
-struct FPopupLayerSlot : public TSlotBase<FPopupLayerSlot>
-{
-public:
-	FPopupLayerSlot()
-	: TSlotBase<FPopupLayerSlot>()
-	, DesktopPosition_Attribute(FVector2D::ZeroVector)
-	, WidthOverride_Attribute()
-	, HeightOverride_Attribute()
-	, Scale_Attribute(1.0f)
-	, Clamp_Attribute(false)
-	, ClampBuffer_Attribute(FVector2D::ZeroVector)
-	{}
-
-	/** Pixel position in desktop space */
-	FPopupLayerSlot& DesktopPosition( const TAttribute<FVector2D>& InDesktopPosition )
-	{
-		DesktopPosition_Attribute = InDesktopPosition;
-		return *this;
-	}
-
-	/** Width override in pixels */
-	FPopupLayerSlot& WidthOverride( const TAttribute<float>& InWidthOverride )
-	{
-		WidthOverride_Attribute = InWidthOverride;
-		return *this;
-	}
-
-	/** Width override in pixels */
-	FPopupLayerSlot& HeightOverride( const TAttribute<float>& InHeightOverride )
-	{
-		HeightOverride_Attribute = InHeightOverride;
-		return *this;
-	}
-
-	/** DPI scaling to be applied to the contents of this slot */
-	FPopupLayerSlot& Scale( const TAttribute<float>& InScale )
-	{
-		Scale_Attribute = InScale;
-		return *this;
-	}
-
-	/** Should this slot be kept within the parent window */
-	FPopupLayerSlot& ClampToWindow( const TAttribute<bool>& InClamp_Attribute)
-	{
-		Clamp_Attribute = InClamp_Attribute;
-		return *this;
-	}
-
-	/** If this slot is kept within the parent window, how far from the edges should we clamp it */
-	FPopupLayerSlot& ClampBuffer( const TAttribute<FVector2D>& InClampBuffer_Attribute )
-	{
-		ClampBuffer_Attribute = InClampBuffer_Attribute;
-		return *this;
-	}
-
-private:
-	/** SPopupLayer arranges FPopupLayerSlots, so it needs to know all about */
-	friend class SPopupLayer;
-	/** TPanelChildren need access to the Widget member */
-	friend class TPanelChildren<FPopupLayerSlot>;
-
-	TAttribute<FVector2D> DesktopPosition_Attribute;
-	TAttribute<float> WidthOverride_Attribute;
-	TAttribute<float> HeightOverride_Attribute;
-	TAttribute<float> Scale_Attribute;
-	TAttribute<bool> Clamp_Attribute;
-	TAttribute<FVector2D> ClampBuffer_Attribute;
-};
-
-
 #if WITH_EDITOR
 
 /**
diff --git a/Engine/Source/Runtime/UMG/Private/Components/DynamicEntryBoxBase.cpp b/Engine/Source/Runtime/UMG/Private/Components/DynamicEntryBoxBase.cpp
--- a/Engine/Source/Runtime/UMG/Private/Components/DynamicEntryBoxBase.cpp
+++ b/Engine/Source/Runtime/UMG/Private/Components/DynamicEntryBoxBase.cpp
@@ -424,33 +424,27 @@ void UDynamicEntryBoxBase::AddEntryChild(UUserWidget& ChildWidget)
 
 		if (bIsHBox)
 		{
-			SHorizontalBox::FSlot* BoxPanelSlot = nullptr;
 			StaticCastSharedPtr<SHorizontalBox>(MyPanelWidget)->AddSlot()
-				.Expose(BoxPanelSlot)
 				.MaxWidth(MaxElementSize)
 				.HAlign(EntryHorizontalAlignment)
 				.VAlign(EntryVerticalAlignment)
+				.SizeParam(UWidget::ConvertSerializedSizeParamToRuntime(EntrySizeRule))
 				.Padding(Padding)
 				[
 					ChildWidget.TakeWidget()
 				];
-			// Todo, add to declaration when SBoxPanel because aware of FSlotArguments
-			BoxPanelSlot->SizeParam = UWidget::ConvertSerializedSizeParamToRuntime(EntrySizeRule);
 		}
 		else
 		{
-			SVerticalBox::FSlot* BoxPanelSlot = nullptr;
 			StaticCastSharedPtr<SVerticalBox>(MyPanelWidget)->AddSlot()
-				.Expose(BoxPanelSlot)
 				.MaxHeight(MaxElementSize)
 				.HAlign(EntryHorizontalAlignment)
 				.VAlign(EntryVerticalAlignment)
+				.SizeParam(UWidget::ConvertSerializedSizeParamToRuntime(EntrySizeRule))
 				.Padding(Padding)
 				[
 					ChildWidget.TakeWidget()
 				];
-			// Todo, add to declaration when SBoxPanel because aware of FSlotArguments
-			BoxPanelSlot->SizeParam = UWidget::ConvertSerializedSizeParamToRuntime(EntrySizeRule);
 		}
 	}
 }
diff --git a/Engine/Source/Runtime/UMG/Private/Components/HorizontalBoxSlot.cpp b/Engine/Source/Runtime/UMG/Private/Components/HorizontalBoxSlot.cpp
--- a/Engine/Source/Runtime/UMG/Private/Components/HorizontalBoxSlot.cpp
+++ b/Engine/Source/Runtime/UMG/Private/Components/HorizontalBoxSlot.cpp
@@ -29,11 +29,10 @@ void UHorizontalBoxSlot::BuildSlot(TSharedRef<SHorizontalBox> HorizontalBox)
 		.HAlign(HorizontalAlignment)
 		.VAlign(VerticalAlignment)
 		.Padding(Padding)
+		.SizeParam(UWidget::ConvertSerializedSizeParamToRuntime(Size))
 		[
 			Content == NULL ? SNullWidget::NullWidget : Content->TakeWidget()
 		];
-
-	Slot->SizeParam = UWidget::ConvertSerializedSizeParamToRuntime(Size);
 }
 
 void UHorizontalBoxSlot::SetPadding(FMargin InPadding)
@@ -50,7 +49,7 @@ void UHorizontalBoxSlot::SetSize(FSlateChildSize InSize)
 	Size = InSize;
 	if ( Slot )
 	{
-		Slot->SizeParam = UWidget::ConvertSerializedSizeParamToRuntime(InSize);
+		Slot->SetSizeParam(UWidget::ConvertSerializedSizeParamToRuntime(InSize));
 	}
 }
 
diff --git a/Engine/Source/Runtime/UMG/Private/Components/VerticalBoxSlot.cpp b/Engine/Source/Runtime/UMG/Private/Components/VerticalBoxSlot.cpp
--- a/Engine/Source/Runtime/UMG/Private/Components/VerticalBoxSlot.cpp
+++ b/Engine/Source/Runtime/UMG/Private/Components/VerticalBoxSlot.cpp
@@ -29,11 +29,11 @@ void UVerticalBoxSlot::BuildSlot(TSharedRef<SVerticalBox> VerticalBox)
 		.Padding(Padding)
 		.HAlign(HorizontalAlignment)
 		.VAlign(VerticalAlignment)
+		.SizeParam(UWidget::ConvertSerializedSizeParamToRuntime(Size))
+		.Expose(Slot)
 		[
 			Content == nullptr ? SNullWidget::NullWidget : Content->TakeWidget()
 		];
-
-	Slot->SizeParam = UWidget::ConvertSerializedSizeParamToRuntime(Size);
 }
 
 void UVerticalBoxSlot::SetPadding(FMargin InPadding)
@@ -50,7 +50,7 @@ void UVerticalBoxSlot::SetSize(FSlateChildSize InSize)
 	Size = InSize;
 	if ( Slot )
 	{
-		Slot->SizeParam = UWidget::ConvertSerializedSizeParamToRuntime(InSize);
+		Slot->SetSizeParam(UWidget::ConvertSerializedSizeParamToRuntime(InSize));
 	}
 }
 
diff --git a/Engine/Source/Runtime/UMG/Private/Components/WrapBoxSlot.cpp b/Engine/Source/Runtime/UMG/Private/Components/WrapBoxSlot.cpp
--- a/Engine/Source/Runtime/UMG/Private/Components/WrapBoxSlot.cpp
+++ b/Engine/Source/Runtime/UMG/Private/Components/WrapBoxSlot.cpp
@@ -35,6 +35,7 @@ void UWrapBoxSlot::BuildSlot(TSharedRef<SWrapBox> WrapBox)
 		.FillEmptySpace(bFillEmptySpace)
 		.FillLineWhenSizeLessThan(FillSpanWhenLessThan == 0 ? TOptional<float>() : TOptional<float>(FillSpanWhenLessThan))
 		.ForceNewLine(false)
+		.Expose(Slot)
 		[
 			Content == nullptr ? SNullWidget::NullWidget : Content->TakeWidget()
 		];
@@ -54,7 +55,7 @@ void UWrapBoxSlot::SetFillEmptySpace(bool InbFillEmptySpace)
 	bFillEmptySpace = InbFillEmptySpace;
 	if ( Slot )
 	{
-		Slot->FillEmptySpace(InbFillEmptySpace);
+		Slot->SetFillEmptySpace(InbFillEmptySpace);
 	}
 }
 
@@ -63,7 +64,7 @@ void UWrapBoxSlot::SetFillSpanWhenLessThan(float InFillSpanWhenLessThan)
 	FillSpanWhenLessThan = InFillSpanWhenLessThan;
 	if ( Slot )
 	{
-		Slot->FillLineWhenSizeLessThan(InFillSpanWhenLessThan == 0 ? TOptional<float>() : TOptional<float>(InFillSpanWhenLessThan));
+		Slot->SetFillLineWhenSizeLessThan(InFillSpanWhenLessThan == 0 ? TOptional<float>() : TOptional<float>(InFillSpanWhenLessThan));
 	}
 }
 
@@ -90,7 +91,7 @@ void UWrapBoxSlot::SetNewLine(bool InForceNewLine)
 	bForceNewLine = InForceNewLine;
 	if (Slot)
 	{
-		Slot->bSlotForceNewLine = InForceNewLine;
+		Slot->SetForceNewLine(InForceNewLine);
 	}
 }
 
