Disable UE_SLATE_WITH_MEMBER_ATTRIBUTE_DEBUGGING. It's too slow in debug.

diff --git a/Engine/Source/Runtime/SlateCore/Public/Types/SlateAttribute.h b/Engine/Source/Runtime/SlateCore/Public/Types/SlateAttribute.h
--- a/Engine/Source/Runtime/SlateCore/Public/Types/SlateAttribute.h
+++ b/Engine/Source/Runtime/SlateCore/Public/Types/SlateAttribute.h
@@ -10,7 +10,7 @@
 #include <type_traits>
 
 #ifndef UE_SLATE_WITH_MEMBER_ATTRIBUTE_DEBUGGING
-	#define UE_SLATE_WITH_MEMBER_ATTRIBUTE_DEBUGGING UE_BUILD_DEBUG
+	#define UE_SLATE_WITH_MEMBER_ATTRIBUTE_DEBUGGING 0
 #endif
 #ifndef UE_SLATE_WITH_ATTRIBUTE_INITIALIZATION_ON_BIND
 	#define UE_SLATE_WITH_ATTRIBUTE_INITIALIZATION_ON_BIND 0
