Add default size to default button style.

diff --git a/Engine/Source/Runtime/SlateCore/Private/Styling/UMGCoreStyle.cpp b/Engine/Source/Runtime/SlateCore/Private/Styling/UMGCoreStyle.cpp
--- a/Engine/Source/Runtime/SlateCore/Private/Styling/UMGCoreStyle.cpp
+++ b/Engine/Source/Runtime/SlateCore/Private/Styling/UMGCoreStyle.cpp
@@ -387,9 +387,9 @@ TSharedRef<ISlateStyle> FUMGCoreStyle::Create()
 
 	// SButton defaults...
 	static const FButtonStyle Button = FButtonStyle()
-		.SetNormal(FSlateRoundedBoxBrush(FUMGColor::Secondary, 4.0f, FUMGColor::InputOutline, InputFocusThickness))
-		.SetHovered(FSlateRoundedBoxBrush(FUMGColor::Hover, 4.0f, FUMGColor::Hover, InputFocusThickness))
-		.SetPressed(FSlateRoundedBoxBrush(FUMGColor::Header, 4.0f, FUMGColor::Hover, InputFocusThickness))
+		.SetNormal(FSlateRoundedBoxBrush(FUMGColor::Secondary, 4.0f, FUMGColor::InputOutline, InputFocusThickness, Icon32x32))
+		.SetHovered(FSlateRoundedBoxBrush(FUMGColor::Hover, 4.0f, FUMGColor::Hover, InputFocusThickness, Icon32x32))
+		.SetPressed(FSlateRoundedBoxBrush(FUMGColor::Header, 4.0f, FUMGColor::Hover, InputFocusThickness, Icon32x32))
 		.SetNormalForeground(FUMGColor::ForegroundHover)
 		.SetHoveredForeground(FUMGColor::ForegroundHover)
 		.SetPressedForeground(FUMGColor::ForegroundHover)
