Slate Events can now be deprecated when using SLATE_EVENT_DEPRECATED.

To use, you will need to provide
* the deprecation version and message, similar to UE_DEPRECATED
* the deprecated delegate class type
* the deprecated event name
* the upgraded delegate event name
* a conversion function that converts the data parameters from the deprecated event execute to the upgraded event execute

diff --git a/Engine/Source/Runtime/Core/Public/Delegates/DelegateSignatureImpl.inl b/Engine/Source/Runtime/Core/Public/Delegates/DelegateSignatureImpl.inl
--- a/Engine/Source/Runtime/Core/Public/Delegates/DelegateSignatureImpl.inl
+++ b/Engine/Source/Runtime/Core/Public/Delegates/DelegateSignatureImpl.inl
@@ -129,11 +129,7 @@ public:
 	template <typename UserClass, typename Var1Type, typename Var2Type, typename Var3Type, typename Var4Type> struct TUObjectMethodDelegate_FourVars_Const  : TBaseUObjectMethodDelegateInstance<true , UserClass, FuncType, UserPolicy, Var1Type, Var2Type, Var3Type, Var4Type> { typedef TBaseUObjectMethodDelegateInstance<true , UserClass, FuncType, UserPolicy, Var1Type, Var2Type, Var3Type, Var4Type> Super; TUObjectMethodDelegate_FourVars_Const (UserClass* InUserObject, typename Super::FMethodPtr InMethodPtr, Var1Type Var1, Var2Type Var2, Var3Type Var3, Var4Type Var4) : Super(InUserObject, InMethodPtr, Var1, Var2, Var3, Var4) {} };
 	
 	/** Declare the user's static function pointer delegate instance types. */
-	                                                                                      struct FStaticDelegate           : TBaseStaticDelegateInstance<FuncType, UserPolicy                                        > { typedef TBaseStaticDelegateInstance<FuncType, UserPolicy                                        > Super; FStaticDelegate          ( typename Super::FFuncPtr InFuncPtr                                                            ) : Super(InFuncPtr                        ) {} };
-	template <typename Var1Type                                                         > struct TStaticDelegate_OneVar    : TBaseStaticDelegateInstance<FuncType, UserPolicy, Var1Type                              > { typedef TBaseStaticDelegateInstance<FuncType, UserPolicy, Var1Type                              > Super; TStaticDelegate_OneVar   ( typename Super::FFuncPtr InFuncPtr, Var1Type Var1                                             ) : Super(InFuncPtr, Var1                  ) {} };
-	template <typename Var1Type, typename Var2Type                                      > struct TStaticDelegate_TwoVars   : TBaseStaticDelegateInstance<FuncType, UserPolicy, Var1Type, Var2Type                    > { typedef TBaseStaticDelegateInstance<FuncType, UserPolicy, Var1Type, Var2Type                    > Super; TStaticDelegate_TwoVars  ( typename Super::FFuncPtr InFuncPtr, Var1Type Var1, Var2Type Var2                              ) : Super(InFuncPtr, Var1, Var2            ) {} };
-	template <typename Var1Type, typename Var2Type, typename Var3Type                   > struct TStaticDelegate_ThreeVars : TBaseStaticDelegateInstance<FuncType, UserPolicy, Var1Type, Var2Type, Var3Type          > { typedef TBaseStaticDelegateInstance<FuncType, UserPolicy, Var1Type, Var2Type, Var3Type          > Super; TStaticDelegate_ThreeVars( typename Super::FFuncPtr InFuncPtr, Var1Type Var1, Var2Type Var2, Var3Type Var3               ) : Super(InFuncPtr, Var1, Var2, Var3      ) {} };
-	template <typename Var1Type, typename Var2Type, typename Var3Type, typename Var4Type> struct TStaticDelegate_FourVars  : TBaseStaticDelegateInstance<FuncType, UserPolicy, Var1Type, Var2Type, Var3Type, Var4Type> { typedef TBaseStaticDelegateInstance<FuncType, UserPolicy, Var1Type, Var2Type, Var3Type, Var4Type> Super; TStaticDelegate_FourVars ( typename Super::FFuncPtr InFuncPtr, Var1Type Var1, Var2Type Var2, Var3Type Var3, Var4Type Var4) : Super(InFuncPtr, Var1, Var2, Var3, Var4) {} };
+	template <typename... VarTypes > struct FStaticDelegate : TBaseStaticDelegateInstance<FuncType, UserPolicy, VarTypes... > { typedef TBaseStaticDelegateInstance<FuncType, UserPolicy, VarTypes... > Super; FStaticDelegate ( typename Super::FFuncPtr InFuncPtr, VarTypes... Vars ) : Super(InFuncPtr, Vars... ) {} };
 
 public:
 
diff --git a/Engine/Source/Runtime/Core/Public/Misc/Attribute.h b/Engine/Source/Runtime/Core/Public/Misc/Attribute.h
--- a/Engine/Source/Runtime/Core/Public/Misc/Attribute.h
+++ b/Engine/Source/Runtime/Core/Public/Misc/Attribute.h
@@ -110,10 +110,11 @@ public:
 	 *
 	 * @param  InFuncPtr Member function to bind.  The function's structure (return value, arguments, etc) must match IBoundAttributeDelegate's definition.
 	 */
-	UE_NODISCARD static TAttribute Create( typename FGetter::FStaticDelegate::FFuncPtr InFuncPtr )
+	template <typename... VarTypes >
+	UE_NODISCARD static TAttribute Create( typename FGetter::template FStaticDelegate< VarTypes... >::FuncType InFuncPtr, VarTypes... Vars )
 	{
 		const bool bExplicitConstructor = true;
-		return TAttribute( FGetter::CreateStatic( InFuncPtr ), bExplicitConstructor );
+		return TAttribute( FGetter::CreateStatic( InFuncPtr, Vars... ), bExplicitConstructor );
 	}
 
 	/**
@@ -282,10 +283,11 @@ public:
 	 *
 	 * @param  InFuncPtr	Function to bind.  The function's structure (return value, arguments, etc) must match IBoundAttributeDelegate's definition.
 	 */
-	void BindStatic( typename FGetter::FStaticDelegate::FFuncPtr InFuncPtr )
+	template < typename... VarTypes >
+	void BindStatic( typename FGetter::template FStaticDelegate< VarTypes... >::FFuncPtr InFuncPtr, VarTypes... Vars )
 	{
 		bIsSet = true;
-		Getter.BindStatic( InFuncPtr );
+		Getter.BindStatic( InFuncPtr, Vars... );
 	}
 	
 	/**
diff --git a/Engine/Source/Runtime/SlateCore/Public/Types/SlateAttributeDescriptor.h b/Engine/Source/Runtime/SlateCore/Public/Types/SlateAttributeDescriptor.h
--- a/Engine/Source/Runtime/SlateCore/Public/Types/SlateAttributeDescriptor.h
+++ b/Engine/Source/Runtime/SlateCore/Public/Types/SlateAttributeDescriptor.h
@@ -38,7 +38,7 @@ public:
 		{ }
 
 		template<typename... PayloadTypes>
-		explicit FInvalidateWidgetReasonAttribute(typename FGetter::FStaticDelegate::FFuncPtr InFuncPtr, PayloadTypes&&... InputPayload)
+		explicit FInvalidateWidgetReasonAttribute(typename FGetter::template FStaticDelegate<PayloadTypes...>::FFuncPtr InFuncPtr, PayloadTypes&&... InputPayload)
 			: Reason(EInvalidateWidgetReason::None)
 			, Getter(FGetter::CreateStatic(InFuncPtr, Forward<PayloadTypes>(InputPayload)...))
 		{
diff --git a/Engine/Source/Runtime/SlateCore/Public/Widgets/DeclarativeSyntaxSupport.h b/Engine/Source/Runtime/SlateCore/Public/Widgets/DeclarativeSyntaxSupport.h
--- a/Engine/Source/Runtime/SlateCore/Public/Widgets/DeclarativeSyntaxSupport.h
+++ b/Engine/Source/Runtime/SlateCore/Public/Widgets/DeclarativeSyntaxSupport.h
@@ -124,36 +124,13 @@ template<typename WidgetType> struct TSlateBaseNamedArgs;
 		\
 		/* Bind attribute with delegate to a global function
 		 * NOTE: We use a template here to avoid 'typename' issues when hosting attributes inside templated classes */ \
-		template< typename StaticFuncPtr > \
-		WidgetArgsType& AttrName##_Static( StaticFuncPtr InFunc )	\
+		template< typename... VarTypes > \
+		WidgetArgsType& AttrName##_Static( typename TAttribute< AttrType >::FGetter::template FStaticDelegate<VarTypes...>::FFuncPtr InFunc, VarTypes... Vars )	\
 		{ \
-			_##AttrName = TAttribute< AttrType >::Create( TAttribute< AttrType >::FGetter::CreateStatic( InFunc ) ); \
-			return this->Me(); \
-		} \
-		template< typename Var1Type > \
-		WidgetArgsType& AttrName##_Static( typename TAttribute< AttrType >::FGetter::template TStaticDelegate_OneVar< Var1Type >::FFuncPtr InFunc, Var1Type Var1 )	\
-		{ \
-			_##AttrName = TAttribute< AttrType >::Create( TAttribute< AttrType >::FGetter::CreateStatic( InFunc, Var1 ) ); \
-			return this->Me(); \
-		} \
-		template< typename Var1Type, typename Var2Type > \
-		WidgetArgsType& AttrName##_Static( typename TAttribute< AttrType >::FGetter::template TStaticDelegate_TwoVars< Var1Type, Var2Type >::FFuncPtr InFunc, Var1Type Var1, Var2Type Var2 )	\
-		{ \
-			_##AttrName = TAttribute< AttrType >::Create( TAttribute< AttrType >::FGetter::CreateStatic( InFunc, Var1, Var2 ) ); \
-			return this->Me(); \
-		} \
-		template< typename Var1Type, typename Var2Type, typename Var3Type > \
-		WidgetArgsType& AttrName##_Static( typename TAttribute< AttrType >::FGetter::template TStaticDelegate_ThreeVars< Var1Type, Var2Type, Var3Type >::FFuncPtr InFunc, Var1Type Var1, Var2Type Var2, Var3Type Var3 )	\
-		{ \
-			_##AttrName = TAttribute< AttrType >::Create( TAttribute< AttrType >::FGetter::CreateStatic( InFunc, Var1, Var2, Var3 ) ); \
-			return this->Me(); \
-		} \
-		template< typename Var1Type, typename Var2Type, typename Var3Type, typename Var4Type > \
-		WidgetArgsType& AttrName##_Static( typename TAttribute< AttrType >::FGetter::template TStaticDelegate_FourVars< Var1Type, Var2Type, Var3Type, Var4Type >::FFuncPtr InFunc, Var1Type Var1, Var2Type Var2, Var3Type Var3, Var4Type Var4 )	\
-		{ \
-			_##AttrName = TAttribute< AttrType >::Create( TAttribute< AttrType >::FGetter::CreateStatic( InFunc, Var1, Var2, Var3, Var4 ) ); \
+			_##AttrName = TAttribute< AttrType >::Create( TAttribute< AttrType >::FGetter::CreateStatic( InFunc, Vars... ) ); \
 			return this->Me(); \
 		} \
+	\
 	\
 		/* Bind attribute with delegate to a lambda
 		 * technically this works for any functor types, but lambdas are the primary use case */ \
@@ -164,132 +141,36 @@ template<typename WidgetType> struct TSlateBaseNamedArgs;
 		} \
 	\
 		/* Bind attribute with delegate to a raw C++ class method */ \
-		template< class UserClass >	\
-		WidgetArgsType& AttrName##_Raw( UserClass* InUserObject, typename TAttribute< AttrType >::FGetter::template TConstMethodPtr< UserClass > InFunc )	\
-		{ \
-			_##AttrName = TAttribute< AttrType >::Create( TAttribute< AttrType >::FGetter::CreateRaw( InUserObject, InFunc ) ); \
-			return this->Me(); \
-		} \
-		template< class UserClass, typename Var1Type >	\
-		WidgetArgsType& AttrName##_Raw( UserClass* InUserObject, typename TAttribute< AttrType >::FGetter::template TConstMethodPtr< UserClass, Var1Type  > InFunc, Var1Type Var1 )	\
-		{ \
-			_##AttrName = TAttribute< AttrType >::Create( TAttribute< AttrType >::FGetter::CreateRaw( InUserObject, InFunc, Var1 ) ); \
-			return this->Me(); \
-		} \
-		template< class UserClass, typename Var1Type, typename Var2Type >	\
-		WidgetArgsType& AttrName##_Raw( UserClass* InUserObject, typename TAttribute< AttrType >::FGetter::template TConstMethodPtr< UserClass, Var1Type, Var2Type > InFunc, Var1Type Var1, Var2Type Var2 )	\
-		{ \
-			_##AttrName = TAttribute< AttrType >::Create( TAttribute< AttrType >::FGetter::CreateRaw( InUserObject, InFunc, Var1, Var2 ) ); \
-			return this->Me(); \
-		} \
-		template< class UserClass, typename Var1Type, typename Var2Type, typename Var3Type >	\
-		WidgetArgsType& AttrName##_Raw( UserClass* InUserObject, typename TAttribute< AttrType >::FGetter::template TConstMethodPtr< UserClass, Var1Type, Var2Type, Var3Type > InFunc, Var1Type Var1, Var2Type Var2, Var3Type Var3 )	\
-		{ \
-			_##AttrName = TAttribute< AttrType >::Create( TAttribute< AttrType >::FGetter::CreateRaw( InUserObject, InFunc, Var1, Var2, Var3 ) ); \
-			return this->Me(); \
-		} \
-		template< class UserClass, typename Var1Type, typename Var2Type, typename Var3Type, typename Var4Type >	\
-		WidgetArgsType& AttrName##_Raw( UserClass* InUserObject, typename TAttribute< AttrType >::FGetter::template TConstMethodPtr< UserClass, Var1Type, Var2Type, Var3Type, Var4Type > InFunc, Var1Type Var1, Var2Type Var2, Var3Type Var3, Var4Type Var4 )	\
+		template< class UserClass, typename... VarTypes >	\
+		WidgetArgsType& AttrName##_Raw( UserClass* InUserObject, typename TAttribute< AttrType >::FGetter::template TConstMethodPtr< UserClass, VarTypes... > InFunc, VarTypes... Vars )	\
 		{ \
-			_##AttrName = TAttribute< AttrType >::Create( TAttribute< AttrType >::FGetter::CreateRaw( InUserObject, InFunc, Var1, Var2, Var3, Var4 ) ); \
+			_##AttrName = TAttribute< AttrType >::Create( TAttribute< AttrType >::FGetter::CreateRaw( InUserObject, InFunc, Vars... ) ); \
 			return this->Me(); \
 		} \
 	\
 		/* Bind attribute with delegate to a shared pointer-based class method.  Slate mostly uses shared pointers so we use an overload for this type of binding. */ \
-		template< class UserClass >	\
-		WidgetArgsType& AttrName( TSharedRef< UserClass > InUserObjectRef, typename TAttribute< AttrType >::FGetter::template TConstMethodPtr< UserClass > InFunc )	\
-		{ \
-			_##AttrName = TAttribute< AttrType >( InUserObjectRef, InFunc ); \
-			return this->Me(); \
-		} \
-		template< class UserClass, typename Var1Type >	\
-		WidgetArgsType& AttrName( TSharedRef< UserClass > InUserObjectRef, typename TAttribute< AttrType >::FGetter::template TConstMethodPtr< UserClass, Var1Type > InFunc, Var1Type Var1 )	\
-		{ \
-			_##AttrName = TAttribute< AttrType >::Create( TAttribute< AttrType >::FGetter::CreateSP( InUserObjectRef, InFunc, Var1 ) ); \
-			return this->Me(); \
-		} \
-		template< class UserClass, typename Var1Type, typename Var2Type >	\
-		WidgetArgsType& AttrName( TSharedRef< UserClass > InUserObjectRef, typename TAttribute< AttrType >::FGetter::template TConstMethodPtr< UserClass, Var1Type, Var2Type > InFunc, Var1Type Var1, Var2Type Var2 )	\
-		{ \
-			_##AttrName = TAttribute< AttrType >::Create( TAttribute< AttrType >::FGetter::CreateSP( InUserObjectRef, InFunc, Var1, Var2 ) ); \
-			return this->Me(); \
-		} \
-		template< class UserClass, typename Var1Type, typename Var2Type, typename Var3Type >	\
-		WidgetArgsType& AttrName( TSharedRef< UserClass > InUserObjectRef, typename TAttribute< AttrType >::FGetter::template TConstMethodPtr< UserClass, Var1Type, Var2Type, Var3Type > InFunc, Var1Type Var1, Var2Type Var2, Var3Type Var3 )	\
+		template< class UserClass, typename... VarTypes >	\
+		WidgetArgsType& AttrName( TSharedRef< UserClass > InUserObjectRef, typename TAttribute< AttrType >::FGetter::template TConstMethodPtr< UserClass, VarTypes... > InFunc, VarTypes... Vars )	\
 		{ \
-			_##AttrName = TAttribute< AttrType >::Create( TAttribute< AttrType >::FGetter::CreateSP( InUserObjectRef, InFunc, Var1, Var2, Var3 ) ); \
-			return this->Me(); \
-		} \
-		template< class UserClass, typename Var1Type, typename Var2Type, typename Var3Type, typename Var4Type >	\
-		WidgetArgsType& AttrName( TSharedRef< UserClass > InUserObjectRef, typename TAttribute< AttrType >::FGetter::template TConstMethodPtr< UserClass, Var1Type, Var2Type, Var3Type, Var4Type > InFunc, Var1Type Var1, Var2Type Var2, Var3Type Var3, Var4Type Var4 )	\
-		{ \
-			_##AttrName = TAttribute< AttrType >::Create( TAttribute< AttrType >::FGetter::CreateSP( InUserObjectRef, InFunc, Var1, Var2, Var3, Var4 ) ); \
+			_##AttrName = TAttribute< AttrType >::Create( TAttribute< AttrType >::FGetter::CreateSP( InUserObjectRef, InFunc, Vars... ) ); \
 			return this->Me(); \
 		} \
 	\
 		/* Bind attribute with delegate to a shared pointer-based class method.  Slate mostly uses shared pointers so we use an overload for this type of binding. */ \
-		template< class UserClass >	\
-		WidgetArgsType& AttrName( UserClass* InUserObject, typename TAttribute< AttrType >::FGetter::template TConstMethodPtr< UserClass > InFunc )	\
-		{ \
-			_##AttrName = TAttribute< AttrType >::Create( TAttribute< AttrType >::FGetter::CreateSP( InUserObject, InFunc ) ); \
-			return this->Me(); \
-		} \
-		template< class UserClass, typename Var1Type >	\
-		WidgetArgsType& AttrName( UserClass* InUserObject, typename TAttribute< AttrType >::FGetter::template TConstMethodPtr< UserClass, Var1Type > InFunc, Var1Type Var1 )	\
-		{ \
-			_##AttrName = TAttribute< AttrType >::Create( TAttribute< AttrType >::FGetter::CreateSP( InUserObject, InFunc, Var1 ) ); \
-			return this->Me(); \
-		} \
-		template< class UserClass, typename Var1Type, typename Var2Type >	\
-		WidgetArgsType& AttrName( UserClass* InUserObject, typename TAttribute< AttrType >::FGetter::template TConstMethodPtr< UserClass, Var1Type, Var2Type > InFunc, Var1Type Var1, Var2Type Var2 )	\
-		{ \
-			_##AttrName = TAttribute< AttrType >::Create( TAttribute< AttrType >::FGetter::CreateSP( InUserObject, InFunc, Var1, Var2 ) ); \
-			return this->Me(); \
-		} \
-		template< class UserClass, typename Var1Type, typename Var2Type, typename Var3Type >	\
-		WidgetArgsType& AttrName( UserClass* InUserObject, typename TAttribute< AttrType >::FGetter::template TConstMethodPtr< UserClass, Var1Type, Var2Type, Var3Type > InFunc, Var1Type Var1, Var2Type Var2, Var3Type Var3 )	\
+		template< class UserClass, typename... VarTypes >	\
+		WidgetArgsType& AttrName( UserClass* InUserObject, typename TAttribute< AttrType >::FGetter::template TConstMethodPtr< UserClass, VarTypes... > InFunc, VarTypes... Vars )	\
 		{ \
-			_##AttrName = TAttribute< AttrType >::Create( TAttribute< AttrType >::FGetter::CreateSP( InUserObject, InFunc, Var1, Var2, Var3 ) ); \
-			return this->Me(); \
-		} \
-		template< class UserClass, typename Var1Type, typename Var2Type, typename Var3Type, typename Var4Type >	\
-		WidgetArgsType& AttrName( UserClass* InUserObject, typename TAttribute< AttrType >::FGetter::template TConstMethodPtr< UserClass, Var1Type, Var2Type, Var3Type, Var4Type > InFunc, Var1Type Var1, Var2Type Var2, Var3Type Var3, Var4Type Var4 )	\
-		{ \
-			_##AttrName = TAttribute< AttrType >::Create( TAttribute< AttrType >::FGetter::CreateSP( InUserObject, InFunc, Var1, Var2, Var3, Var4 ) ); \
+			_##AttrName = TAttribute< AttrType >::Create( TAttribute< AttrType >::FGetter::CreateSP( InUserObject, InFunc, Vars... ) ); \
 			return this->Me(); \
 		} \
 	\
 		/* Bind attribute with delegate to a UObject-based class method */ \
-		template< class UserClass >	\
-		WidgetArgsType& AttrName##_UObject( UserClass* InUserObject, typename TAttribute< AttrType >::FGetter::template TConstMethodPtr< UserClass > InFunc )	\
-		{ \
-			_##AttrName = TAttribute< AttrType >::Create( TAttribute< AttrType >::FGetter::CreateUObject( InUserObject, InFunc ) ); \
-			return this->Me(); \
-		} \
-		template< class UserClass, typename Var1Type >	\
-		WidgetArgsType& AttrName##_UObject( UserClass* InUserObject, typename TAttribute< AttrType >::FGetter::template TConstMethodPtr< UserClass, Var1Type > InFunc, Var1Type Var1 )	\
-		{ \
-			_##AttrName = TAttribute< AttrType >::Create( TAttribute< AttrType >::FGetter::CreateUObject( InUserObject, InFunc, Var1 ) ); \
-			return this->Me(); \
-		} \
-		template< class UserClass, typename Var1Type, typename Var2Type  >	\
-		WidgetArgsType& AttrName##_UObject( UserClass* InUserObject, typename TAttribute< AttrType >::FGetter::template TConstMethodPtr< UserClass, Var1Type, Var2Type > InFunc, Var1Type Var1, Var2Type Var2 )	\
-		{ \
-			_##AttrName = TAttribute< AttrType >::Create( TAttribute< AttrType >::FGetter::CreateUObject( InUserObject, InFunc, Var1, Var2 ) ); \
-			return this->Me(); \
-		} \
-		template< class UserClass, typename Var1Type, typename Var2Type, typename Var3Type >	\
-		WidgetArgsType& AttrName##_UObject( UserClass* InUserObject, typename TAttribute< AttrType >::FGetter::template TConstMethodPtr< UserClass, Var1Type, Var2Type, Var3Type > InFunc, Var1Type Var1, Var2Type Var2, Var3Type Var3 )	\
+		template< class UserClass, typename... VarTypes >	\
+		WidgetArgsType& AttrName##_UObject( UserClass* InUserObject, typename TAttribute< AttrType >::FGetter::template TConstMethodPtr< UserClass, VarTypes... > InFunc, VarTypes... Vars )	\
 		{ \
-			_##AttrName = TAttribute< AttrType >::Create( TAttribute< AttrType >::FGetter::CreateUObject( InUserObject, InFunc, Var1, Var2, Var3 ) ); \
+			_##AttrName = TAttribute< AttrType >::Create( TAttribute< AttrType >::FGetter::CreateUObject( InUserObject, InFunc, Vars... ) ); \
 			return this->Me(); \
 		} \
-		template< class UserClass, typename Var1Type, typename Var2Type, typename Var3Type, typename Var4Type >	\
-		WidgetArgsType& AttrName##_UObject( UserClass* InUserObject, typename TAttribute< AttrType >::FGetter::template TConstMethodPtr< UserClass, Var1Type, Var2Type, Var3Type, Var4Type > InFunc, Var1Type Var1, Var2Type Var2, Var3Type Var3, Var4Type Var4 )	\
-		{ \
-			_##AttrName = TAttribute< AttrType >::Create( TAttribute< AttrType >::FGetter::CreateUObject( InUserObject, InFunc, Var1, Var2, Var3, Var4 ) ); \
-			return this->Me(); \
-		}
 
 #define SLATE_PRIVATE_ARGUMENT_VARIABLE( ArgType, ArgName ) \
 		ArgType _##ArgName
@@ -537,34 +418,10 @@ struct NamedSlotProperty
 		\
 		/* Set event delegate to a global function */ \
 		/* NOTE: We use a template here to avoid 'typename' issues when hosting attributes inside templated classes */ \
-		template< typename StaticFuncPtr > \
-		WidgetArgsType& EventName##_Static( StaticFuncPtr InFunc )	\
-		{ \
-			_##EventName = DelegateName::CreateStatic( InFunc ); \
-			return *this; \
-		} \
-		template< typename StaticFuncPtr, typename Var1Type > \
-		WidgetArgsType& EventName##_Static( StaticFuncPtr InFunc, Var1Type Var1 )	\
+		template< typename StaticFuncPtr, typename... VarTypes > \
+		WidgetArgsType& EventName##_Static( StaticFuncPtr InFunc, VarTypes... Vars )	\
 		{ \
-			_##EventName = DelegateName::CreateStatic( InFunc, Var1 ); \
-			return *this; \
-		} \
-		template< typename StaticFuncPtr, typename Var1Type, typename Var2Type > \
-		WidgetArgsType& EventName##_Static( StaticFuncPtr InFunc, Var1Type Var1, Var2Type Var2 )	\
-		{ \
-			_##EventName = DelegateName::CreateStatic( InFunc, Var1, Var2 ); \
-			return *this; \
-		} \
-		template< typename StaticFuncPtr, typename Var1Type, typename Var2Type, typename Var3Type > \
-		WidgetArgsType& EventName##_Static( StaticFuncPtr InFunc, Var1Type Var1, Var2Type Var2, Var3Type Var3 )	\
-		{ \
-			_##EventName = DelegateName::CreateStatic( InFunc, Var1, Var2, Var3 ); \
-			return *this; \
-		} \
-		template< typename StaticFuncPtr, typename Var1Type, typename Var2Type, typename Var3Type, typename Var4Type > \
-		WidgetArgsType& EventName##_Static( StaticFuncPtr InFunc, Var1Type Var1, Var2Type Var2, Var3Type Var3, Var4Type Var4 )	\
-		{ \
-			_##EventName = DelegateName::CreateStatic( InFunc, Var1, Var2, Var3, Var4 ); \
+			_##EventName = DelegateName::CreateStatic( InFunc, Vars... ); \
 			return *this; \
 		} \
 		\
@@ -578,255 +435,171 @@ struct NamedSlotProperty
 		} \
 		\
 		/* Set event delegate to a raw C++ class method */ \
-		template< class UserClass >	\
-		WidgetArgsType& EventName##_Raw( UserClass* InUserObject, typename DelegateName::template TMethodPtr< UserClass > InFunc )	\
-		{ \
-			_##EventName = DelegateName::CreateRaw( InUserObject, InFunc ); \
-			return *this; \
-		} \
-		template< class UserClass >	\
-		WidgetArgsType& EventName##_Raw( UserClass* InUserObject, typename DelegateName::template TConstMethodPtr< UserClass > InFunc )	\
-		{ \
-			_##EventName = DelegateName::CreateRaw( InUserObject, InFunc ); \
-			return *this; \
-		} \
-		template< class UserClass, typename Var1Type >	\
-		WidgetArgsType& EventName##_Raw( UserClass* InUserObject, typename DelegateName::template TMethodPtr< UserClass, Var1Type > InFunc, Var1Type Var1 )	\
-		{ \
-			_##EventName = DelegateName::CreateRaw( InUserObject, InFunc, Var1 ); \
-			return *this; \
-		} \
-		template< class UserClass, typename Var1Type >	\
-		WidgetArgsType& EventName##_Raw( UserClass* InUserObject, typename DelegateName::template TConstMethodPtr< UserClass, Var1Type > InFunc, Var1Type Var1 )	\
-		{ \
-			_##EventName = DelegateName::CreateRaw( InUserObject, InFunc, Var1 ); \
-			return *this; \
-		} \
-		template< class UserClass, typename Var1Type, typename Var2Type >	\
-		WidgetArgsType& EventName##_Raw( UserClass* InUserObject, typename DelegateName::template TMethodPtr< UserClass, Var1Type, Var2Type > InFunc, Var1Type Var1, Var2Type Var2 )	\
-		{ \
-			_##EventName = DelegateName::CreateRaw( InUserObject, InFunc, Var1, Var2 ); \
-			return *this; \
-		} \
-		template< class UserClass, typename Var1Type, typename Var2Type >	\
-		WidgetArgsType& EventName##_Raw( UserClass* InUserObject, typename DelegateName::template TConstMethodPtr< UserClass, Var1Type, Var2Type > InFunc, Var1Type Var1, Var2Type Var2 )	\
-		{ \
-			_##EventName = DelegateName::CreateRaw( InUserObject, InFunc, Var1, Var2 ); \
-			return *this; \
-		} \
-		template< class UserClass, typename Var1Type, typename Var2Type, typename Var3Type >	\
-		WidgetArgsType& EventName##_Raw( UserClass* InUserObject, typename DelegateName::template TMethodPtr< UserClass, Var1Type, Var2Type, Var3Type > InFunc, Var1Type Var1, Var2Type Var2, Var3Type Var3 )	\
+		template< class UserClass, typename... VarTypes >	\
+		WidgetArgsType& EventName##_Raw( UserClass* InUserObject, typename DelegateName::template TMethodPtr< UserClass, VarTypes... > InFunc, VarTypes... Vars )	\
 		{ \
-			_##EventName = DelegateName::CreateRaw( InUserObject, InFunc, Var1, Var2, Var3 ); \
+			_##EventName = DelegateName::CreateRaw( InUserObject, InFunc, Vars... ); \
 			return *this; \
 		} \
-		template< class UserClass, typename Var1Type, typename Var2Type, typename Var3Type >	\
-		WidgetArgsType& EventName##_Raw( UserClass* InUserObject, typename DelegateName::template TConstMethodPtr< UserClass, Var1Type, Var2Type, Var3Type > InFunc, Var1Type Var1, Var2Type Var2, Var3Type Var3 )	\
+		template< class UserClass, typename... VarTypes >	\
+		WidgetArgsType& EventName##_Raw( UserClass* InUserObject, typename DelegateName::template TConstMethodPtr< UserClass, VarTypes... > InFunc, VarTypes... Vars )	\
 		{ \
-			_##EventName = DelegateName::CreateRaw( InUserObject, InFunc, Var1, Var2, Var3 ); \
-			return *this; \
-		} \
-		template< class UserClass, typename Var1Type, typename Var2Type, typename Var3Type, typename Var4Type >	\
-		WidgetArgsType& EventName##_Raw( UserClass* InUserObject, typename DelegateName::template TMethodPtr< UserClass, Var1Type, Var2Type, Var3Type, Var4Type > InFunc, Var1Type Var1, Var2Type Var2, Var3Type Var3, Var4Type Var4 )	\
-		{ \
-			_##EventName = DelegateName::CreateRaw( InUserObject, InFunc, Var1, Var2, Var3, Var4 ); \
-			return *this; \
-		} \
-		template< class UserClass, typename Var1Type, typename Var2Type, typename Var3Type, typename Var4Type >	\
-		WidgetArgsType& EventName##_Raw( UserClass* InUserObject, typename DelegateName::template TConstMethodPtr< UserClass, Var1Type, Var2Type, Var3Type, Var4Type > InFunc, Var1Type Var1, Var2Type Var2, Var3Type Var3, Var4Type Var4 )	\
-		{ \
-			_##EventName = DelegateName::CreateRaw( InUserObject, InFunc, Var1, Var2, Var3, Var4 ); \
+			_##EventName = DelegateName::CreateRaw( InUserObject, InFunc, Vars... ); \
 			return *this; \
 		} \
 		\
 		/* Set event delegate to a shared pointer-based class method.  Slate mostly uses shared pointers so we use an overload for this type of binding. */ \
-		template< class UserClass >	\
-		WidgetArgsType& EventName( TSharedRef< UserClass > InUserObjectRef, typename DelegateName::template TMethodPtr< UserClass > InFunc )	\
-		{ \
-			_##EventName = DelegateName::CreateSP( InUserObjectRef, InFunc ); \
-			return *this; \
-		} \
-		template< class UserClass >	\
-		WidgetArgsType& EventName( TSharedRef< UserClass > InUserObjectRef, typename DelegateName::template TConstMethodPtr< UserClass > InFunc )	\
-		{ \
-			_##EventName = DelegateName::CreateSP( InUserObjectRef, InFunc ); \
-			return *this; \
-		} \
-		template< class UserClass, typename Var1Type >	\
-		WidgetArgsType& EventName( TSharedRef< UserClass > InUserObjectRef, typename DelegateName::template TMethodPtr< UserClass, Var1Type > InFunc, Var1Type Var1 )	\
-		{ \
-			_##EventName = DelegateName::CreateSP( InUserObjectRef, InFunc, Var1 ); \
-			return *this; \
-		} \
-		template< class UserClass, typename Var1Type >	\
-		WidgetArgsType& EventName( TSharedRef< UserClass > InUserObjectRef, typename DelegateName::template TConstMethodPtr< UserClass, Var1Type > InFunc, Var1Type Var1 )	\
-		{ \
-			_##EventName = DelegateName::CreateSP( InUserObjectRef, InFunc, Var1 ); \
-			return *this; \
-		} \
-		template< class UserClass, typename Var1Type, typename Var2Type >	\
-		WidgetArgsType& EventName( TSharedRef< UserClass > InUserObjectRef, typename DelegateName::template TMethodPtr< UserClass, Var1Type, Var2Type > InFunc, Var1Type Var1, Var2Type Var2 )	\
-		{ \
-			_##EventName = DelegateName::CreateSP( InUserObjectRef, InFunc, Var1, Var2 ); \
-			return *this; \
-		} \
-		template< class UserClass, typename Var1Type, typename Var2Type >	\
-		WidgetArgsType& EventName( TSharedRef< UserClass > InUserObjectRef, typename DelegateName::template TConstMethodPtr< UserClass, Var1Type, Var2Type > InFunc, Var1Type Var1, Var2Type Var2 )	\
-		{ \
-			_##EventName = DelegateName::CreateSP( InUserObjectRef, InFunc, Var1, Var2 ); \
-			return *this; \
-		} \
-		template< class UserClass, typename Var1Type, typename Var2Type, typename Var3Type >	\
-		WidgetArgsType& EventName( TSharedRef< UserClass > InUserObjectRef, typename DelegateName::template TMethodPtr< UserClass, Var1Type, Var2Type, Var3Type > InFunc, Var1Type Var1, Var2Type Var2, Var3Type Var3 )	\
-		{ \
-			_##EventName = DelegateName::CreateSP( InUserObjectRef, InFunc, Var1, Var2, Var3 ); \
-			return *this; \
-		} \
-		template< class UserClass, typename Var1Type, typename Var2Type, typename Var3Type >	\
-		WidgetArgsType& EventName( TSharedRef< UserClass > InUserObjectRef, typename DelegateName::template TConstMethodPtr< UserClass, Var1Type, Var2Type, Var3Type > InFunc, Var1Type Var1, Var2Type Var2, Var3Type Var3 )	\
-		{ \
-			_##EventName = DelegateName::CreateSP( InUserObjectRef, InFunc, Var1, Var2, Var3 ); \
-			return *this; \
-		} \
-		template< class UserClass, typename Var1Type, typename Var2Type, typename Var3Type, typename Var4Type >	\
-		WidgetArgsType& EventName( TSharedRef< UserClass > InUserObjectRef, typename DelegateName::template TMethodPtr< UserClass, Var1Type, Var2Type, Var3Type, Var4Type > InFunc, Var1Type Var1, Var2Type Var2, Var3Type Var3, Var4Type Var4 )	\
+		template< class UserClass, typename... VarTypes >	\
+		WidgetArgsType& EventName( TSharedRef< UserClass > InUserObjectRef, typename DelegateName::template TMethodPtr< UserClass, VarTypes... > InFunc, VarTypes... Vars )	\
 		{ \
-			_##EventName = DelegateName::CreateSP( InUserObjectRef, InFunc, Var1, Var2, Var3, Var4 ); \
+			_##EventName = DelegateName::CreateSP( InUserObjectRef, InFunc, Vars... ); \
 			return *this; \
 		} \
-		template< class UserClass, typename Var1Type, typename Var2Type, typename Var3Type, typename Var4Type >	\
-		WidgetArgsType& EventName( TSharedRef< UserClass > InUserObjectRef, typename DelegateName::template TConstMethodPtr< UserClass, Var1Type, Var2Type, Var3Type, Var4Type > InFunc, Var1Type Var1, Var2Type Var2, Var3Type Var3, Var4Type Var4 )	\
+		template< class UserClass, typename... VarTypes >	\
+		WidgetArgsType& EventName( TSharedRef< UserClass > InUserObjectRef, typename DelegateName::template TConstMethodPtr< UserClass, VarTypes... > InFunc, VarTypes... Vars )	\
 		{ \
-			_##EventName = DelegateName::CreateSP( InUserObjectRef, InFunc, Var1, Var2, Var3, Var4 ); \
+			_##EventName = DelegateName::CreateSP( InUserObjectRef, InFunc, Vars... ); \
 			return *this; \
 		} \
 		\
 		/* Set event delegate to a shared pointer-based class method.  Slate mostly uses shared pointers so we use an overload for this type of binding. */ \
-		template< class UserClass >	\
-		WidgetArgsType& EventName( UserClass* InUserObject, typename DelegateName::template TMethodPtr< UserClass > InFunc )	\
-		{ \
-			_##EventName = DelegateName::CreateSP( InUserObject, InFunc ); \
-			return *this; \
-		} \
-		template< class UserClass >	\
-		WidgetArgsType& EventName( UserClass* InUserObject, typename DelegateName::template TConstMethodPtr< UserClass > InFunc )	\
+		template< class UserClass, typename... VarTypes >	\
+		WidgetArgsType& EventName( UserClass* InUserObject, typename DelegateName::template TMethodPtr< UserClass, VarTypes... > InFunc, VarTypes... Vars )	\
 		{ \
-			_##EventName = DelegateName::CreateSP( InUserObject, InFunc ); \
+			_##EventName = DelegateName::CreateSP( InUserObject, InFunc, Vars... ); \
 			return *this; \
 		} \
-		template< class UserClass, typename Var1Type >	\
-		WidgetArgsType& EventName( UserClass* InUserObject, typename DelegateName::template TMethodPtr< UserClass, Var1Type > InFunc, Var1Type Var1 )	\
+		template< class UserClass, typename... VarTypes >	\
+		WidgetArgsType& EventName( UserClass* InUserObject, typename DelegateName::template TConstMethodPtr< UserClass, VarTypes... > InFunc, VarTypes... Vars )	\
 		{ \
-			_##EventName = DelegateName::CreateSP( InUserObject, InFunc, Var1 ); \
-			return *this; \
-		} \
-		template< class UserClass, typename Var1Type >	\
-		WidgetArgsType& EventName( UserClass* InUserObject, typename DelegateName::template TConstMethodPtr< UserClass, Var1Type > InFunc, Var1Type Var1 )	\
-		{ \
-			_##EventName = DelegateName::CreateSP( InUserObject, InFunc, Var1 ); \
-			return *this; \
-		} \
-		template< class UserClass, typename Var1Type, typename Var2Type >	\
-		WidgetArgsType& EventName( UserClass* InUserObject, typename DelegateName::template TMethodPtr< UserClass, Var1Type, Var2Type > InFunc, Var1Type Var1, Var2Type Var2 )	\
-		{ \
-			_##EventName = DelegateName::CreateSP( InUserObject, InFunc, Var1, Var2 ); \
-			return *this; \
-		} \
-		template< class UserClass, typename Var1Type, typename Var2Type >	\
-		WidgetArgsType& EventName( UserClass* InUserObject, typename DelegateName::template TConstMethodPtr< UserClass, Var1Type, Var2Type > InFunc, Var1Type Var1, Var2Type Var2 )	\
-		{ \
-			_##EventName = DelegateName::CreateSP( InUserObject, InFunc, Var1, Var2 ); \
-			return *this; \
-		} \
-		template< class UserClass, typename Var1Type, typename Var2Type, typename Var3Type >	\
-		WidgetArgsType& EventName( UserClass* InUserObject, typename DelegateName::template TMethodPtr< UserClass, Var1Type, Var2Type, Var3Type > InFunc, Var1Type Var1, Var2Type Var2, Var3Type Var3 )	\
-		{ \
-			_##EventName = DelegateName::CreateSP( InUserObject, InFunc, Var1, Var2, Var3 ); \
-			return *this; \
-		} \
-		template< class UserClass, typename Var1Type, typename Var2Type, typename Var3Type >	\
-		WidgetArgsType& EventName( UserClass* InUserObject, typename DelegateName::template TConstMethodPtr< UserClass, Var1Type, Var2Type, Var3Type > InFunc, Var1Type Var1, Var2Type Var2, Var3Type Var3 )	\
-		{ \
-			_##EventName = DelegateName::CreateSP( InUserObject, InFunc, Var1, Var2, Var3 ); \
-			return *this; \
-		} \
-		template< class UserClass, typename Var1Type, typename Var2Type, typename Var3Type, typename Var4Type >	\
-		WidgetArgsType& EventName( UserClass* InUserObject, typename DelegateName::template TMethodPtr< UserClass, Var1Type, Var2Type, Var3Type, Var4Type > InFunc, Var1Type Var1, Var2Type Var2, Var3Type Var3, Var4Type Var4 )	\
-		{ \
-			_##EventName = DelegateName::CreateSP( InUserObject, InFunc, Var1, Var2, Var3, Var4 ); \
-			return *this; \
-		} \
-		template< class UserClass, typename Var1Type, typename Var2Type, typename Var3Type, typename Var4Type >	\
-		WidgetArgsType& EventName( UserClass* InUserObject, typename DelegateName::template TConstMethodPtr< UserClass, Var1Type, Var2Type, Var3Type, Var4Type > InFunc, Var1Type Var1, Var2Type Var2, Var3Type Var3, Var4Type Var4 )	\
-		{ \
-			_##EventName = DelegateName::CreateSP( InUserObject, InFunc, Var1, Var2, Var3, Var4 ); \
+			_##EventName = DelegateName::CreateSP( InUserObject, InFunc, Vars... ); \
 			return *this; \
 		} \
 		\
 		/* Set event delegate to a UObject-based class method */ \
-		template< class UserClass >	\
-		WidgetArgsType& EventName##_UObject( UserClass* InUserObject, typename DelegateName::template TMethodPtr< UserClass > InFunc )	\
-		{ \
-			_##EventName = DelegateName::CreateUObject( InUserObject, InFunc ); \
-			return *this; \
-		} \
-		template< class UserClass >	\
-		WidgetArgsType& EventName##_UObject( UserClass* InUserObject, typename DelegateName::template TConstMethodPtr< UserClass > InFunc )	\
-		{ \
-			_##EventName = DelegateName::CreateUObject( InUserObject, InFunc ); \
-			return *this; \
-		} \
-		template< class UserClass, typename Var1Type >	\
-		WidgetArgsType& EventName##_UObject( UserClass* InUserObject, typename DelegateName::template TMethodPtr< UserClass, Var1Type > InFunc, Var1Type Var1 )	\
-		{ \
-			_##EventName = DelegateName::CreateUObject( InUserObject, InFunc, Var1 ); \
-			return *this; \
-		} \
-		template< class UserClass, typename Var1Type >	\
-		WidgetArgsType& EventName##_UObject( UserClass* InUserObject, typename DelegateName::template TConstMethodPtr< UserClass, Var1Type > InFunc, Var1Type Var1 )	\
-		{ \
-			_##EventName = DelegateName::CreateUObject( InUserObject, InFunc, Var1 ); \
-			return *this; \
-		} \
-		template< class UserClass, typename Var1Type, typename Var2Type >	\
-		WidgetArgsType& EventName##_UObject( UserClass* InUserObject, typename DelegateName::template TMethodPtr< UserClass, Var1Type, Var2Type > InFunc, Var1Type Var1, Var2Type Var2 )	\
-		{ \
-			_##EventName = DelegateName::CreateUObject( InUserObject, InFunc, Var1, Var2 ); \
-			return *this; \
-		} \
-		template< class UserClass, typename Var1Type, typename Var2Type >	\
-		WidgetArgsType& EventName##_UObject( UserClass* InUserObject, typename DelegateName::template TConstMethodPtr< UserClass, Var1Type, Var2Type > InFunc, Var1Type Var1, Var2Type Var2 )	\
+		template< class UserClass, typename... VarTypes >	\
+		WidgetArgsType& EventName##_UObject( UserClass* InUserObject, typename DelegateName::template TMethodPtr< UserClass, VarTypes... > InFunc, VarTypes... Vars )	\
 		{ \
-			_##EventName = DelegateName::CreateUObject( InUserObject, InFunc, Var1, Var2 ); \
+			_##EventName = DelegateName::CreateUObject( InUserObject, InFunc, Vars... ); \
 			return *this; \
 		} \
-		template< class UserClass, typename Var1Type, typename Var2Type, typename Var3Type >	\
-		WidgetArgsType& EventName##_UObject( UserClass* InUserObject, typename DelegateName::template TMethodPtr< UserClass, Var1Type, Var2Type, Var3Type > InFunc, Var1Type Var1, Var2Type Var2, Var3Type Var3 )	\
+		template< class UserClass, typename... VarTypes >	\
+		WidgetArgsType& EventName##_UObject( UserClass* InUserObject, typename DelegateName::template TConstMethodPtr< UserClass, VarTypes... > InFunc, VarTypes... Vars )	\
 		{ \
-			_##EventName = DelegateName::CreateUObject( InUserObject, InFunc, Var1, Var2, Var3 ); \
-			return *this; \
-		} \
-		template< class UserClass, typename Var1Type, typename Var2Type, typename Var3Type >	\
-		WidgetArgsType& EventName##_UObject( UserClass* InUserObject, typename DelegateName::template TConstMethodPtr< UserClass, Var1Type, Var2Type, Var3Type > InFunc, Var1Type Var1, Var2Type Var2, Var3Type Var3 )	\
-		{ \
-			_##EventName = DelegateName::CreateUObject( InUserObject, InFunc, Var1, Var2, Var3 ); \
-			return *this; \
-		} \
-		template< class UserClass, typename Var1Type, typename Var2Type, typename Var3Type, typename Var4Type >	\
-		WidgetArgsType& EventName##_UObject( UserClass* InUserObject, typename DelegateName::template TMethodPtr< UserClass, Var1Type, Var2Type, Var3Type, Var4Type > InFunc, Var1Type Var1, Var2Type Var2, Var3Type Var3, Var4Type Var4 )	\
-		{ \
-			_##EventName = DelegateName::CreateUObject( InUserObject, InFunc, Var1, Var2, Var3, Var4 ); \
-			return *this; \
-		} \
-		template< class UserClass, typename Var1Type, typename Var2Type, typename Var3Type, typename Var4Type >	\
-		WidgetArgsType& EventName##_UObject( UserClass* InUserObject, typename DelegateName::template TConstMethodPtr< UserClass, Var1Type, Var2Type, Var3Type, Var4Type > InFunc, Var1Type Var1, Var2Type Var2, Var3Type Var3, Var4Type Var4 )	\
-		{ \
-			_##EventName = DelegateName::CreateUObject( InUserObject, InFunc, Var1, Var2, Var3, Var4 ); \
+			_##EventName = DelegateName::CreateUObject( InUserObject, InFunc, Vars... ); \
 			return *this; \
 		} \
 		\
 		DelegateName _##EventName; \
 
+#define SLATE_EVENT_DEPRECATED( DeprecationVersion, DeprecationMessage, DelegateName, EventName, UpgradedEventName, UpgradeFuncName)	\
+		UE_DEPRECATED(DeprecationVersion, DeprecationMessage)	\
+		WidgetArgsType& EventName( const DelegateName& InDeprecatedDelegate ) \
+		{	\
+			if (InDeprecatedDelegate.IsBound())	\
+			{	\
+				_##UpgradedEventName = UpgradeFuncName(InDeprecatedDelegate);	\
+			}	\
+			return *this;	\
+		}	\
+		\
+		template <typename StaticFuncPtr, typename... VarTypes>	\
+		UE_DEPRECATED(DeprecationVersion, DeprecationMessage)	\
+		WidgetArgsType& EventName##_Static( StaticFuncPtr InFunc, VarTypes... Vars )	\
+		{	\
+			if (InFunc)	\
+			{	\
+				DelegateName DeprecatedDelegate = DelegateName::CreateStatic(InFunc, Vars...);	\
+				_##UpgradedEventName = UpgradeFuncName(DeprecatedDelegate);	\
+			}	\
+			return *this;	\
+		}	\
+		\
+		template <typename FunctorType, typename... VarTypes>	\
+		UE_DEPRECATED(DeprecationVersion, DeprecationMessage)	\
+		WidgetArgsType& EventName##_Lambda(FunctorType&& InFunctor, VarTypes... Vars) \
+		{	\
+			if (InFunctor)	\
+			{	\
+				DelegateName DeprecatedDelegate = DelegateName::CreateLambda(InFunctor, Vars...);	\
+				_##UpgradedEventName = UpgradeFuncName(DeprecatedDelegate);	\
+			}	\
+			return *this;	\
+		}	\
+		\
+		template <typename UserClass, typename... VarTypes>	\
+		UE_DEPRECATED(DeprecationVersion, DeprecationMessage)	\
+		WidgetArgsType& EventName##_Raw( UserClass* InUserObject, typename DelegateName::template TMethodPtr< UserClass, VarTypes... > InFunc, VarTypes... Vars )	\
+		{	\
+			DelegateName DeprecatedDelegate = DelegateName::CreateRaw(InUserObject, InFunc, Vars...);	\
+			_##UpgradedEventName = UpgradeFuncName(DeprecatedDelegate);	\
+			return *this;	\
+		}	\
+		\
+		template <typename UserClass, typename... VarTypes>	\
+		UE_DEPRECATED(DeprecationVersion, DeprecationMessage)	\
+		WidgetArgsType& EventName##_Raw( UserClass* InUserObject, typename DelegateName::template TConstMethodPtr< UserClass, VarTypes... > InFunc, VarTypes... Vars )	\
+		{	\
+			DelegateName DeprecatedDelegate = DelegateName::CreateRaw(InUserObject, InFunc, Vars...);	\
+			_##UpgradedEventName = UpgradeFuncName(DeprecatedDelegate);	\
+			return *this;	\
+		}	\
+		\
+		template <typename UserClass, typename... VarTypes>	\
+		UE_DEPRECATED(DeprecationVersion, DeprecationMessage)	\
+		WidgetArgsType& EventName( TSharedRef< UserClass > InUserObjectRef, typename DelegateName::template TMethodPtr< UserClass, VarTypes... > InFunc, VarTypes... Vars )	\
+		{	\
+			DelegateName DeprecatedDelegate = DelegateName::CreateSP(InUserObjectRef, InFunc, Vars...);	\
+			_##UpgradedEventName = UpgradeFuncName(DeprecatedDelegate);	\
+			return *this;	\
+		}	\
+		\
+		template <typename UserClass, typename... VarTypes>	\
+		UE_DEPRECATED(DeprecationVersion, DeprecationMessage)	\
+		WidgetArgsType& EventName( TSharedRef< UserClass > InUserObjectRef, typename DelegateName::template TConstMethodPtr< UserClass, VarTypes... > InFunc, VarTypes... Vars )	\
+		{	\
+			DelegateName DeprecatedDelegate = DelegateName::CreateSP(InUserObjectRef, InFunc, Vars...);	\
+			_##UpgradedEventName = UpgradeFuncName(DeprecatedDelegate);	\
+			return *this;	\
+		}	\
+		\
+		template <typename UserClass, typename... VarTypes>	\
+		UE_DEPRECATED(DeprecationVersion, DeprecationMessage)	\
+		WidgetArgsType& EventName( UserClass* InUserObject, typename DelegateName::template TMethodPtr< UserClass, VarTypes... > InFunc, VarTypes... Vars )	\
+		{	\
+			DelegateName DeprecatedDelegate = DelegateName::CreateSP(InUserObject, InFunc, Vars...);	\
+			_##UpgradedEventName = UpgradeFuncName(DeprecatedDelegate);	\
+			return *this;	\
+		}	\
+		\
+		template <typename UserClass, typename... VarTypes>	\
+		UE_DEPRECATED(DeprecationVersion, DeprecationMessage)	\
+		WidgetArgsType& EventName( UserClass* InUserObject, typename DelegateName::template TConstMethodPtr< UserClass, VarTypes... > InFunc, VarTypes... Vars )	\
+		{	\
+			DelegateName DeprecatedDelegate = DelegateName::CreateSP(InUserObject, InFunc, Vars...);	\
+			_##UpgradedEventName = UpgradeFuncName(DeprecatedDelegate);	\
+			return *this;	\
+		}	\
+		\
+		template <typename UserClass, typename... VarTypes>	\
+		UE_DEPRECATED(DeprecationVersion, DeprecationMessage)	\
+		WidgetArgsType& EventName##_UObject( UserClass* InUserObject, typename DelegateName::template TMethodPtr< UserClass, VarTypes... > InFunc, VarTypes... Vars )	\
+		{	\
+			DelegateName DeprecatedDelegate = DelegateName::CreateUObject(InUserObject, InFunc, Vars...);	\
+			_##UpgradedEventName = UpgradeFuncName(DeprecatedDelegate);	\
+			return *this;	\
+		}	\
+		\
+		template <typename UserClass, typename... VarTypes>	\
+		UE_DEPRECATED(DeprecationVersion, DeprecationMessage)	\
+		WidgetArgsType& EventName##_UObject( UserClass* InUserObject, typename DelegateName::template TConstMethodPtr< UserClass, VarTypes... > InFunc, VarTypes... Vars )	\
+		{	\
+			DelegateName DeprecatedDelegate = DelegateName::CreateUObject(InUserObject, InFunc, Vars...);	\
+			_##UpgradedEventName = UpgradeFuncName(DeprecatedDelegate);	\
+			return *this;	\
+		}	\
+
+
 
 /** Base class for named arguments. Provides settings necessary for all widgets. */
 struct FSlateBaseNamedArgs
