Slate: Remove some shared ptr copy. Fix slot access in FChildren and add a ForEach function.

diff --git a/Engine/Source/Runtime/SlateCore/Private/FastUpdate/SlateInvalidationRoot.cpp b/Engine/Source/Runtime/SlateCore/Private/FastUpdate/SlateInvalidationRoot.cpp
--- a/Engine/Source/Runtime/SlateCore/Private/FastUpdate/SlateInvalidationRoot.cpp
+++ b/Engine/Source/Runtime/SlateCore/Private/FastUpdate/SlateInvalidationRoot.cpp
@@ -483,13 +483,13 @@ void FSlateInvalidationRoot::AdjustWidgetsDesktopGeometry(FVector2D WindowToDesk
 {
 	FSlateLayoutTransform WindowToDesktop(WindowToDesktopTransform);
 
-	FastWidgetPathList->ForEachWidget([WindowToDesktopTransform, &WindowToDesktop](SWidget* Widget)
+	FastWidgetPathList->ForEachWidget([WindowToDesktopTransform, &WindowToDesktop](SWidget& Widget)
 		{
-			Widget->PersistentState.DesktopGeometry.AppendTransform(WindowToDesktop);
+			Widget.PersistentState.DesktopGeometry.AppendTransform(WindowToDesktop);
 		});
 }
 
-void FSlateInvalidationRoot::BuildFastPathWidgetList(TSharedRef<SWidget> RootWidget)
+void FSlateInvalidationRoot::BuildFastPathWidgetList(const TSharedRef<SWidget>& RootWidget)
 {
 	TGuardValue<bool> Tmp(bBuildingWidgetList, true);
 
@@ -699,7 +699,7 @@ void FSlateInvalidationRoot::ProcessPreUpdate()
 
 							// This element was removed or added, seek will assign the correct widget to be ticked next.
 							AttributeItt.Seek(InvalidationWidget.Index);
-							if (FastWidgetPathList->ShouldBeAddedToAttributeList(WidgetPtr))
+							if (FastWidgetPathList->ShouldBeAddedToAttributeList(*WidgetPtr))
 							{
 								// Do we still need to update this element, if not, then remove it from the update list.
 								if (!Slate::EInvalidateWidgetReason_HasPreUpdateFlag(InvalidationWidget.CurrentInvalidateReason))
@@ -958,12 +958,12 @@ bool FSlateInvalidationRoot::ProcessInvalidation()
 
 void FSlateInvalidationRoot::ClearAllFastPathData(bool bClearResourcesImmediately)
 {
-	FastWidgetPathList->ForEachWidget([bClearResourcesImmediately](SWidget* Widget)
+	FastWidgetPathList->ForEachWidget([bClearResourcesImmediately](SWidget& Widget)
 		{
-			Widget->PersistentState.CachedElementHandle = FSlateCachedElementsHandle::Invalid;
+			Widget.PersistentState.CachedElementHandle = FSlateCachedElementsHandle::Invalid;
 			if (bClearResourcesImmediately)
 			{
-				Widget->FastPathProxyHandle = FWidgetProxyHandle();
+				Widget.FastPathProxyHandle = FWidgetProxyHandle();
 			}
 		});
 	FastWidgetPathList->Reset();
@@ -1178,14 +1178,14 @@ void VerifyHittest(SWidget* InvalidationRootWidget, FSlateInvalidationWidgetList
 
 	FSlateInvalidationWidgetSortOrder PreviousSecondarySort;
 	const SWidget* LastWidget = nullptr;
-	WidgetList.ForEachWidget([&HittestGridSortDatas, &PreviousSecondarySort, &LastWidget](const SWidget* Widget)
+	WidgetList.ForEachWidget([&HittestGridSortDatas, &PreviousSecondarySort, &LastWidget](const SWidget& Widget)
 		{
-			if (Widget && Widget->GetVisibility().IsHitTestVisible())
+			if (Widget.GetVisibility().IsHitTestVisible())
 			{
 				// Is the widget in the hittestgrid
-				const int32 FoundHittestIndex = HittestGridSortDatas.IndexOfByPredicate([Widget](const FHittestWidgetSortData& HittestGrid)
+				const int32 FoundHittestIndex = HittestGridSortDatas.IndexOfByPredicate([&Widget](const FHittestWidgetSortData& HittestGrid)
 					{
-						return HittestGrid.Widget == Widget;
+						return HittestGrid.Widget == &Widget;
 					});
 				const bool bHasFoundWidget = HittestGridSortDatas.IsValidIndex(FoundHittestIndex);
 				if (!bHasFoundWidget)
@@ -1193,12 +1193,12 @@ void VerifyHittest(SWidget* InvalidationRootWidget, FSlateInvalidationWidgetList
 					return;
 				}
 
-				UE_SLATE_LOG_ERROR_IF_FALSE(Widget->GetProxyHandle().GetWidgetSortOrder() == HittestGridSortDatas[FoundHittestIndex].SecondarySort
+				UE_SLATE_LOG_ERROR_IF_FALSE(Widget.GetProxyHandle().GetWidgetSortOrder() == HittestGridSortDatas[FoundHittestIndex].SecondarySort
 					, GSlateInvalidationRootVerifyHittestGrid
 					, TEXT("The SecondarySort of widget '%s' doesn't match the SecondarySort inside the hittestgrid.")
 					, *FReflectionMetaData::GetWidgetPath(Widget));
 
-				LastWidget = Widget;
+				LastWidget = &Widget;
 				PreviousSecondarySort = HittestGridSortDatas[FoundHittestIndex].SecondarySort;
 
 				HittestGridSortDatas.RemoveAtSwap(FoundHittestIndex);
@@ -1296,21 +1296,22 @@ void VerifyWidgetVisibility(FSlateInvalidationWidgetList& WidgetList)
 void VerifyWidgetVolatile(FSlateInvalidationWidgetList& WidgetList, TArray<FSlateInvalidationWidgetIndex>& FinalUpdateList)
 {
 	SWidget* Root = WidgetList.GetRoot().Pin().Get();
-	WidgetList.ForEachWidget([Root, &FinalUpdateList](SWidget* Widget)
+	check(Root);
+	WidgetList.ForEachWidget([Root, &FinalUpdateList](SWidget& Widget)
 		{
-			if (Widget != Root && GSlateInvalidationRootVerifyWidgetVolatile)
+			if (&Widget != Root && GSlateInvalidationRootVerifyWidgetVolatile)
 			{
 				{
-					const bool bWasVolatile = Widget->IsVolatile();
-					Widget->CacheVolatility();
-					const bool bIsVolatile = Widget->IsVolatile();
+					const bool bWasVolatile = Widget.IsVolatile();
+					Widget.CacheVolatility();
+					const bool bIsVolatile = Widget.IsVolatile();
 					UE_SLATE_LOG_ERROR_IF_FALSE(bWasVolatile == bIsVolatile
 						, GSlateInvalidationRootVerifyWidgetVolatile
 						, TEXT("Widget '%s' volatily changed without an invalidation.")
 						, *FReflectionMetaData::GetWidgetDebugInfo(Widget));
 				}
 
-				const TSharedPtr<const SWidget> ParentWidget = Widget->GetParentWidget();
+				const TSharedPtr<const SWidget> ParentWidget = Widget.GetParentWidget();
 				UE_SLATE_LOG_ERROR_IF_FALSE(ParentWidget
 					, GSlateInvalidationRootVerifyWidgetVolatile
 					, TEXT("Parent widget of widget '%s' is invalid.")
@@ -1318,24 +1319,24 @@ void VerifyWidgetVolatile(FSlateInvalidationWidgetList& WidgetList, TArray<FSlat
 
 				{
 					const bool bShouldBeVolatileIndirectly = ParentWidget->IsVolatileIndirectly() || ParentWidget->IsVolatile();
-					UE_SLATE_LOG_ERROR_IF_FALSE(Widget->IsVolatileIndirectly() == bShouldBeVolatileIndirectly
+					UE_SLATE_LOG_ERROR_IF_FALSE(Widget.IsVolatileIndirectly() == bShouldBeVolatileIndirectly
 					, GSlateInvalidationRootVerifyWidgetVolatile
 					, TEXT("Widget '%s' should be set as %s.")
 					, *FReflectionMetaData::GetWidgetDebugInfo(Widget)
 					, (bShouldBeVolatileIndirectly ? TEXT("volatile indirectly") : TEXT("not volatile indirectly")));
 				}
 
-				if (Widget->IsVolatile() && !Widget->IsVolatileIndirectly())
+				if (Widget.IsVolatile() && !Widget.IsVolatileIndirectly())
 				{
-					UE_SLATE_LOG_ERROR_IF_FALSE(Widget->HasAnyUpdateFlags(EWidgetUpdateFlags::NeedsVolatilePaint)
+					UE_SLATE_LOG_ERROR_IF_FALSE(Widget.HasAnyUpdateFlags(EWidgetUpdateFlags::NeedsVolatilePaint)
 						, GSlateInvalidationRootVerifyWidgetVolatile
 						, TEXT("Widget '%s' is volatile but doesn't have the update flag NeedsVolatilePaint.")
 						, *FReflectionMetaData::GetWidgetDebugInfo(Widget));
 
-					if (Widget->GetProxyHandle().IsValid(Widget))
+					if (Widget.GetProxyHandle().IsValid(Widget))
 					{
-						const bool bIsVisible = Widget->GetProxyHandle().GetProxy().Visibility.IsVisible();
-						const bool bIsContains = FinalUpdateList.Contains(Widget->GetProxyHandle().GetWidgetIndex());
+						const bool bIsVisible = Widget.GetProxyHandle().GetProxy().Visibility.IsVisible();
+						const bool bIsContains = FinalUpdateList.Contains(Widget.GetProxyHandle().GetWidgetIndex());
 						UE_SLATE_LOG_ERROR_IF_FALSE(bIsContains || !bIsVisible
 							, GSlateInvalidationRootVerifyWidgetVolatile
 							, TEXT("Widget '%s' is volatile but is not in the update list.")
diff --git a/Engine/Source/Runtime/SlateCore/Private/FastUpdate/SlateInvalidationWidgetList.cpp b/Engine/Source/Runtime/SlateCore/Private/FastUpdate/SlateInvalidationWidgetList.cpp
--- a/Engine/Source/Runtime/SlateCore/Private/FastUpdate/SlateInvalidationWidgetList.cpp
+++ b/Engine/Source/Runtime/SlateCore/Private/FastUpdate/SlateInvalidationWidgetList.cpp
@@ -408,14 +408,14 @@ FSlateInvalidationWidgetList::FSlateInvalidationWidgetList(FSlateInvalidationRoo
 }
 
 
-FSlateInvalidationWidgetIndex FSlateInvalidationWidgetList::Internal_BuildWidgetList_Recursive(TSharedRef<SWidget>& Widget, FSlateInvalidationWidgetIndex ParentIndex, IndexType& LastestIndex, FSlateInvalidationWidgetVisibility ParentVisibility, bool bParentVolatile)
+FSlateInvalidationWidgetIndex FSlateInvalidationWidgetList::Internal_BuildWidgetList_Recursive(SWidget& Widget, FSlateInvalidationWidgetIndex ParentIndex, IndexType& LastestIndex, FSlateInvalidationWidgetVisibility ParentVisibility, bool bParentVolatile)
 {
 	const bool bIsEmpty = IsEmpty();
 	const FSlateInvalidationWidgetIndex NewIndex = EmplaceInsertAfter(LastestIndex, Widget);
 	LastestIndex = NewIndex.ArrayIndex;
 
 	FSlateInvalidationWidgetIndex LeafMostChildIndex = NewIndex;
-	FSlateInvalidationWidgetVisibility NewVisibility {ParentVisibility, Widget->GetVisibility()};
+	FSlateInvalidationWidgetVisibility NewVisibility {ParentVisibility, Widget.GetVisibility()};
 
 	bool bUpdateSlateAttribute = false;
 	if (ShouldBeAddedToAttributeList(Widget))
@@ -427,13 +427,13 @@ FSlateInvalidationWidgetIndex FSlateInvalidationWidgetList::Internal_BuildWidget
 		{
 			if (!NewVisibility.IsCollapseIndirectly())
 			{
-				FSlateAttributeMetaData::UpdateOnlyVisibilityAttributes(Widget.Get(), FSlateAttributeMetaData::EInvalidationPermission::DenyAndClearDelayedInvalidation);
-				NewVisibility.SetVisibility(ParentVisibility, Widget->GetVisibility());
+				FSlateAttributeMetaData::UpdateOnlyVisibilityAttributes(Widget, FSlateAttributeMetaData::EInvalidationPermission::DenyAndClearDelayedInvalidation);
+				NewVisibility.SetVisibility(ParentVisibility, Widget.GetVisibility());
 
 				if (!NewVisibility.IsCollapsed())
 				{
-					FSlateAttributeMetaData::UpdateExceptVisibilityAttributes(Widget.Get(), FSlateAttributeMetaData::EInvalidationPermission::DenyAndClearDelayedInvalidation);
-					NewVisibility.SetVisibility(ParentVisibility, Widget->GetVisibility());
+					FSlateAttributeMetaData::UpdateExceptVisibilityAttributes(Widget, FSlateAttributeMetaData::EInvalidationPermission::DenyAndClearDelayedInvalidation);
+					NewVisibility.SetVisibility(ParentVisibility, Widget.GetVisibility());
 					bUpdateSlateAttribute = true;
 				}
 			}
@@ -445,7 +445,7 @@ FSlateInvalidationWidgetIndex FSlateInvalidationWidgetList::Internal_BuildWidget
 		Data[LastestIndex].ElementIndexList_VolatileUpdateWidget.AddUnsorted(NewIndex.ElementIndex);
 	}
 
-	const bool bIsInvalidationRoot = Widget->Advanced_IsInvalidationRoot();
+	const bool bIsInvalidationRoot = Widget.Advanced_IsInvalidationRoot();
 
 	{
 		InvalidationWidgetType& WidgetProxy = (*this)[NewIndex];
@@ -454,7 +454,7 @@ FSlateInvalidationWidgetIndex FSlateInvalidationWidgetList::Internal_BuildWidget
 		WidgetProxy.LeafMostChildIndex = LeafMostChildIndex;
 		WidgetProxy.Visibility = NewVisibility;
 		WidgetProxy.bIsInvalidationRoot = bIsInvalidationRoot;
-		WidgetProxy.bIsVolatilePrepass = Widget->HasAnyUpdateFlags(EWidgetUpdateFlags::NeedsVolatilePrepass);
+		WidgetProxy.bIsVolatilePrepass = Widget.HasAnyUpdateFlags(EWidgetUpdateFlags::NeedsVolatilePrepass);
 #if UE_SLATE_WITH_INVALIDATIONWIDGETLIST_DEBUGGING
 		WidgetProxy.bDebug_AttributeUpdated = bUpdateSlateAttribute;
 #endif
@@ -465,29 +465,25 @@ FSlateInvalidationWidgetIndex FSlateInvalidationWidgetList::Internal_BuildWidget
 #endif
 	{
 		const FSlateInvalidationWidgetSortOrder SortIndex = { *this, NewIndex };
-		Widget->SetFastPathProxyHandle(FWidgetProxyHandle{ Owner, NewIndex, SortIndex }, NewVisibility, bParentVolatile);
+		Widget.SetFastPathProxyHandle(FWidgetProxyHandle{ Owner, NewIndex, SortIndex }, NewVisibility, bParentVolatile);
 	}
 
-	const bool bParentOrSelfVolatile = bParentVolatile || Widget->IsVolatile();
+	const bool bParentOrSelfVolatile = bParentVolatile || Widget.IsVolatile();
 
 	// N.B. The SInvalidationBox needs a valid Proxy to decide if he's a root or not.
 	//const bool bDoRecursion = ShouldDoRecursion(Widget);
 	const bool bDoRecursion = !bIsInvalidationRoot || bIsEmpty;
 	if (bDoRecursion)
 	{
-		FChildren* Children = Widget->GetAllChildren();
-		int32 NumChildren = Children->Num();
-		for (int32 Index = 0; Index < NumChildren; ++Index)
-		{
-			TSharedRef<SWidget> NewWidget = Children->GetChildAt(Index);
-
-			const bool bShouldAdd = ShouldBeAdded(NewWidget);
-			if (bShouldAdd)
+		Widget.GetAllChildren()->ForEachWidget([&, this](SWidget& NewWidget)
 			{
-				check(NewWidget->GetParentWidget() == Widget);
-				LeafMostChildIndex = Internal_BuildWidgetList_Recursive(NewWidget, NewIndex, LastestIndex, NewVisibility, bParentOrSelfVolatile);
-			}
-		}
+				const bool bShouldAdd = ShouldBeAdded(NewWidget);
+				if (bShouldAdd)
+				{
+					check(NewWidget.GetParentWidget().Get() == &Widget);
+					LeafMostChildIndex = Internal_BuildWidgetList_Recursive(NewWidget, NewIndex, LastestIndex, NewVisibility, bParentOrSelfVolatile);
+				}
+			});
 
 		InvalidationWidgetType& WidgetProxy = (*this)[NewIndex];
 		WidgetProxy.LeafMostChildIndex = LeafMostChildIndex;
@@ -497,7 +493,7 @@ FSlateInvalidationWidgetIndex FSlateInvalidationWidgetList::Internal_BuildWidget
 }
 
 
-void FSlateInvalidationWidgetList::BuildWidgetList(TSharedRef<SWidget> InRoot)
+void FSlateInvalidationWidgetList::BuildWidgetList(const TSharedRef<SWidget>& InRoot)
 {
 	SCOPED_NAMED_EVENT(Slate_InvalidationList_ProcessBuild, FColorList::Blue);
 
@@ -513,21 +509,21 @@ void FSlateInvalidationWidgetList::BuildWidgetList(TSharedRef<SWidget> InRoot)
 		}
 		const bool bParentVolatile = false;
 		IndexType LatestArrayIndex = FSlateInvalidationWidgetIndex::Invalid.ArrayIndex;
-		Internal_BuildWidgetList_Recursive(InRoot, FSlateInvalidationWidgetIndex::Invalid, LatestArrayIndex, ParentVisibility, bParentVolatile);
+		Internal_BuildWidgetList_Recursive(InRoot.Get(), FSlateInvalidationWidgetIndex::Invalid, LatestArrayIndex, ParentVisibility, bParentVolatile);
 	}
 }
 
 
-void FSlateInvalidationWidgetList::Internal_RebuildWidgetListTree(TSharedRef<SWidget> Widget, int32 ChildAtIndex)
+void FSlateInvalidationWidgetList::Internal_RebuildWidgetListTree(SWidget& Widget, int32 ChildAtIndex)
 {
 	const bool bShouldAddWidget = ShouldBeAdded(Widget);
-	FChildren* ParentChildren = Widget->GetAllChildren();
-	const FSlateInvalidationWidgetIndex WidgetIndex = Widget->GetProxyHandle().GetWidgetIndex();
+	FChildren* ParentChildren = Widget.GetAllChildren();
+	const FSlateInvalidationWidgetIndex WidgetIndex = Widget.GetProxyHandle().GetWidgetIndex();
 	if (bShouldAddWidget && ChildAtIndex < ParentChildren->Num() && WidgetIndex != FSlateInvalidationWidgetIndex::Invalid)
 	{
 		// Since we are going to add item, the array may get invalidated. Do not use after Internal_BuildWidgetList_Recursive
 		const InvalidationWidgetType& CurrentInvalidationWidget = (*this)[WidgetIndex];
-		ensure(CurrentInvalidationWidget.GetWidget() == &Widget.Get());
+		ensure(CurrentInvalidationWidget.GetWidget() == &Widget);
 		FSlateInvalidationWidgetIndex PreviousLeafIndex = CurrentInvalidationWidget.LeafMostChildIndex;
 		FSlateInvalidationWidgetIndex NewLeafIndex = PreviousLeafIndex;
 		IndexType LastestArrayIndex = PreviousLeafIndex.ArrayIndex;
@@ -535,18 +531,16 @@ void FSlateInvalidationWidgetList::Internal_RebuildWidgetListTree(TSharedRef<SWi
 		if (bDoRecursion)
 		{
 			const FSlateInvalidationWidgetVisibility CurrentVisibility = CurrentInvalidationWidget.Visibility;
-			const bool bParentVolatile = Widget->IsVolatileIndirectly() || Widget->IsVolatile();
-			int32 NumChildren = ParentChildren->Num();
-			for (int32 Index = ChildAtIndex; Index < NumChildren; ++Index)
-			{
-				TSharedRef<SWidget> ChildWidget = ParentChildren->GetChildAt(Index);
-				const bool bShouldAddChild = ShouldBeAdded(ChildWidget);
-				if (bShouldAddChild)
+			const bool bParentVolatile = Widget.IsVolatileIndirectly() || Widget.IsVolatile();
+			ParentChildren->ForEachWidget([&, this](SWidget& ChildWidget)
 				{
-					check(ChildWidget->GetParentWidget() == Widget);
-					NewLeafIndex = Internal_BuildWidgetList_Recursive(ChildWidget, WidgetIndex, LastestArrayIndex, CurrentVisibility, bParentVolatile);
-				}
-			}
+					const bool bShouldAddChild = ShouldBeAdded(ChildWidget);
+					if (bShouldAddChild)
+					{
+						check(ChildWidget.GetParentWidget().Get() == &Widget);
+						NewLeafIndex = Internal_BuildWidgetList_Recursive(ChildWidget, WidgetIndex, LastestArrayIndex, CurrentVisibility, bParentVolatile);
+					}
+				});
 		}
 
 		if (NewLeafIndex != PreviousLeafIndex)
@@ -585,7 +579,7 @@ bool FSlateInvalidationWidgetList::ProcessChildOrderInvalidation(const Invalidat
 		FIndexRange ChildToRemoveRange;
 
 		// Was added, but it should not be there anymore
-		if (!ShouldBeAdded(WidgetPtr))
+		if (!ShouldBeAdded(*WidgetPtr))
 		{
 			// This proxy will not be valid after this call
 			bIsInvalidationWidgetStillValid = false;
@@ -596,7 +590,7 @@ bool FSlateInvalidationWidgetList::ProcessChildOrderInvalidation(const Invalidat
 			//Result.ChildBuilt
 		}
 		// If it is not supposed to had child, but had some, them remove them
-		else if (!ShouldDoRecursion(WidgetPtr))
+		else if (!ShouldDoRecursion(*WidgetPtr))
 		{
 			if (InvalidationWidget.Index != InvalidationWidget.LeafMostChildIndex)
 			{
@@ -702,7 +696,7 @@ bool FSlateInvalidationWidgetList::ProcessChildOrderInvalidation(const Invalidat
 				SCOPED_NAMED_EVENT(Slate_InvalidationList_ProcessRebuild, FColorList::Blue);
 				FSlateInvalidationWidgetIndex PreviousWidgetIndex = WidgetPtr->GetProxyHandle().GetWidgetIndex();
 				FSlateInvalidationWidgetIndex PreviousLeafMostChildIndex = InvalidationWidget.LeafMostChildIndex;
-				Internal_RebuildWidgetListTree(WidgetPtr->AsShared(), StartChildIndex);
+				Internal_RebuildWidgetListTree(*WidgetPtr, StartChildIndex);
 
 
 				const InvalidationWidgetType& NewInvalidationWidget = (*this)[PreviousWidgetIndex];
@@ -726,7 +720,7 @@ void FSlateInvalidationWidgetList::ProcessAttributeRegistrationInvalidation(cons
 	SWidget* WidgetPtr = InvalidationWidget.GetWidget();
 	check(WidgetPtr);
 
-	if (ShouldBeAddedToAttributeList(WidgetPtr))
+	if (ShouldBeAddedToAttributeList(*WidgetPtr))
 	{
 		Data[InvalidationWidget.Index.ArrayIndex].ElementIndexList_WidgetWithRegisteredSlateAttribute.InsertUnique(InvalidationWidget.Index.ElementIndex);
 	}
@@ -742,7 +736,7 @@ void FSlateInvalidationWidgetList::ProcessVolatileUpdateInvalidation(Invalidatio
 	SWidget* WidgetPtr = InvalidationWidget.GetWidget();
 	check(WidgetPtr);
 
-	if (ShouldBeAddedToVolatileUpdateList(WidgetPtr))
+	if (ShouldBeAddedToVolatileUpdateList(*WidgetPtr))
 	{
 		Data[InvalidationWidget.Index.ArrayIndex].ElementIndexList_VolatileUpdateWidget.InsertUnique(InvalidationWidget.Index.ElementIndex);
 	}
@@ -1438,11 +1432,11 @@ FSlateInvalidationWidgetList::FCutResult FSlateInvalidationWidgetList::Internal_
 				Widget->SetFastPathProxyHandle(FWidgetProxyHandle{ Owner, NewWidgetIndex, SortIndex });
 
 				// Check for ElementIndexList
-				if (ShouldBeAddedToAttributeList(Widget))
+				if (ShouldBeAddedToAttributeList(*Widget))
 				{
 					Data[NewArrayIndex].ElementIndexList_WidgetWithRegisteredSlateAttribute.AddUnsorted(NewElementIndex);
 				}
-				if (ShouldBeAddedToVolatileUpdateList(Widget))
+				if (ShouldBeAddedToVolatileUpdateList(*Widget))
 				{
 					Data[NewArrayIndex].ElementIndexList_VolatileUpdateWidget.AddUnsorted(NewElementIndex);
 				}
@@ -1476,9 +1470,9 @@ FSlateInvalidationWidgetList::FCutResult FSlateInvalidationWidgetList::Internal_
 }
 
 #if UE_SLATE_WITH_INVALIDATIONWIDGETLIST_DEBUGGING
-FSlateInvalidationWidgetIndex FSlateInvalidationWidgetList::FindWidget(const TSharedRef<SWidget> WidgetToFind) const
+FSlateInvalidationWidgetIndex FSlateInvalidationWidgetList::FindWidget(const SWidget& WidgetToFind) const
 {
-	SWidget* WidgetToFindPtr = &WidgetToFind.Get();
+	const SWidget* WidgetToFindPtr = &WidgetToFind;
 	for (FSlateInvalidationWidgetIndex Index = FirstIndex(); Index != FSlateInvalidationWidgetIndex::Invalid; Index = IncrementIndex(Index))
 	{
 		if (WidgetToFindPtr == (*this)[Index].GetWidget())
@@ -1502,11 +1496,11 @@ void FSlateInvalidationWidgetList::RemoveWidget(const FSlateInvalidationWidgetIn
 }
 
 
-void FSlateInvalidationWidgetList::RemoveWidget(const TSharedRef<SWidget> WidgetToRemove)
+void FSlateInvalidationWidgetList::RemoveWidget(const SWidget& WidgetToRemove)
 {
-	if (ensure(WidgetToRemove->GetProxyHandle().GetInvalidationRootHandle().GetUniqueId() == Owner.GetUniqueId()))
+	if (ensure(WidgetToRemove.GetProxyHandle().GetInvalidationRootHandle().GetUniqueId() == Owner.GetUniqueId()))
 	{
-		FSlateInvalidationWidgetIndex WidgetIndex = WidgetToRemove->GetProxyHandle().GetWidgetIndex();
+		FSlateInvalidationWidgetIndex WidgetIndex = WidgetToRemove.GetProxyHandle().GetWidgetIndex();
 		if (WidgetIndex != FSlateInvalidationWidgetIndex::Invalid)
 		{
 			const InvalidationWidgetType& InvalidationWidget = (*this)[WidgetIndex];
@@ -1518,12 +1512,12 @@ void FSlateInvalidationWidgetList::RemoveWidget(const TSharedRef<SWidget> Widget
 }
 
 
-TArray<TSharedPtr<SWidget>> FSlateInvalidationWidgetList::FindChildren(const TSharedRef<SWidget> Widget) const
+TArray<TSharedPtr<SWidget>> FSlateInvalidationWidgetList::FindChildren(const SWidget& Widget) const
 {
 	TArray<TSharedPtr<SWidget>> Result;
-	if (ensure(Widget->GetProxyHandle().GetInvalidationRootHandle().GetUniqueId() == Owner.GetUniqueId()))
+	if (ensure(Widget.GetProxyHandle().GetInvalidationRootHandle().GetUniqueId() == Owner.GetUniqueId()))
 	{
-		FSlateInvalidationWidgetIndex WidgetIndex = Widget->GetProxyHandle().GetWidgetIndex();
+		FSlateInvalidationWidgetIndex WidgetIndex = Widget.GetProxyHandle().GetWidgetIndex();
 		if (WidgetIndex == FSlateInvalidationWidgetIndex::Invalid)
 		{
 			return Result;
@@ -1758,11 +1752,11 @@ bool FSlateInvalidationWidgetList::VerifyElementIndexList() const
 		TArray<const SWidget*> WidgetListWithForEachWidget;
 		WidgetListWithForEachWidget.Reset(WidgetListWithIterator.Num());
 		int32 Index = 0;
-		const_cast<FSlateInvalidationWidgetList*>(this)->ForEachWidget([&](const SWidget* Widget)
+		const_cast<FSlateInvalidationWidgetList*>(this)->ForEachWidget([&](const SWidget& Widget)
 			{
 				if (ShouldBeAddedToAttributeList(Widget))
 				{
-					WidgetListWithForEachWidget.Add(Widget);
+					WidgetListWithForEachWidget.Add(&Widget);
 				}
 			});
 
diff --git a/Engine/Source/Runtime/SlateCore/Private/FastUpdate/SlateInvalidationWidgetList.h b/Engine/Source/Runtime/SlateCore/Private/FastUpdate/SlateInvalidationWidgetList.h
--- a/Engine/Source/Runtime/SlateCore/Private/FastUpdate/SlateInvalidationWidgetList.h
+++ b/Engine/Source/Runtime/SlateCore/Private/FastUpdate/SlateInvalidationWidgetList.h
@@ -91,7 +91,7 @@ public:
 	FSlateInvalidationWidgetList(FSlateInvalidationRootHandle Owner, const FArguments& Args);
 
 	/** Build the widget list from the root widget. */
-	void BuildWidgetList(TSharedRef<SWidget> Root);
+	void BuildWidgetList(const TSharedRef<SWidget>& Root);
 
 	/** Get the root the widget list was built with. */
 	TWeakPtr<SWidget> GetRoot() { return Root; };
@@ -157,7 +157,7 @@ public:
 				SWidget* Widget = ElementList[ElementIndex].GetWidget();
 				if (Widget)
 				{
-					Pred(Widget);
+					Pred(*Widget);
 				}
 			}
 
@@ -334,15 +334,15 @@ public:
 public:
 #if UE_SLATE_WITH_INVALIDATIONWIDGETLIST_DEBUGGING
 	/** For testing purposes. Return the InvalidationWidgetIndex of the Widget within the InvalidationWidgetList. */
-	FSlateInvalidationWidgetIndex FindWidget(const TSharedRef<SWidget> Widget) const;
+	FSlateInvalidationWidgetIndex FindWidget(const SWidget& Widget) const;
 
 	/** For testing purposes. Use ProcessChildOrderInvalidation. */
 	void RemoveWidget(const FSlateInvalidationWidgetIndex ToRemove);
 	/** For testing purposes. Use ProcessChildOrderInvalidation. */
-	void RemoveWidget(const TSharedRef<SWidget> WidgetToRemove);
+	void RemoveWidget(const SWidget& WidgetToRemove);
 
 	/** For testing purposes. Use to test ProcessChildOrderInvalidation */
-	UE_NODISCARD TArray<TSharedPtr<SWidget>> FindChildren(const TSharedRef<SWidget> Widget) const;
+	UE_NODISCARD TArray<TSharedPtr<SWidget>> FindChildren(const SWidget& Widget) const;
 
 	/**
 	 * For testing purposes.
@@ -369,41 +369,41 @@ public:
 #endif
 
 public:
-	bool ShouldDoRecursion(const SWidget* Widget) const
+	bool ShouldDoRecursion(const SWidget& Widget) const
 	{
-		return !Widget->Advanced_IsInvalidationRoot() || IsEmpty();
+		return !Widget.Advanced_IsInvalidationRoot() || IsEmpty();
 	}
-	bool ShouldDoRecursion(const TSharedRef<SWidget>& Widget) const
+	bool ShouldDoRecursion(const TSharedRef<const SWidget>& Widget) const
 	{
-		return !Widget->Advanced_IsInvalidationRoot() || IsEmpty();
+		return ShouldDoRecursion(Widget.Get());
 	}
-	static bool ShouldBeAdded(const SWidget* Widget)
+	static bool ShouldBeAdded(const SWidget& Widget)
 	{
-		return Widget != &(SNullWidget::NullWidget.Get());
+		return &Widget != &(SNullWidget::NullWidget.Get());
 	}
-	static bool ShouldBeAdded(const TSharedRef<SWidget>& Widget)
+	static bool ShouldBeAdded(const TSharedRef<const SWidget>& Widget)
 	{
 		return Widget != SNullWidget::NullWidget;
 	}
-	static bool ShouldBeAddedToAttributeList(const SWidget* Widget)
+	static bool ShouldBeAddedToAttributeList(const SWidget& Widget)
 	{
-		return Widget->HasRegisteredSlateAttribute() && Widget->IsAttributesUpdatesEnabled();
+		return Widget.HasRegisteredSlateAttribute() && Widget.IsAttributesUpdatesEnabled();
 	}
-	static bool ShouldBeAddedToAttributeList(const TSharedRef<SWidget>& Widget)
+	static bool ShouldBeAddedToAttributeList(const TSharedRef<const SWidget>& Widget)
 	{
-		return Widget->HasRegisteredSlateAttribute() && Widget->IsAttributesUpdatesEnabled();
+		return ShouldBeAddedToAttributeList(Widget.Get());
 	}
 	static bool HasVolatileUpdateFlags(EWidgetUpdateFlags UpdateFlags)
 	{
 		return EnumHasAnyFlags(UpdateFlags, EWidgetUpdateFlags::NeedsTick | EWidgetUpdateFlags::NeedsActiveTimerUpdate | EWidgetUpdateFlags::NeedsVolatilePaint | EWidgetUpdateFlags::NeedsVolatilePrepass);
 	}
-	static bool ShouldBeAddedToVolatileUpdateList(const SWidget* Widget)
+	static bool ShouldBeAddedToVolatileUpdateList(const SWidget& Widget)
 	{
-		return Widget->HasAnyUpdateFlags(EWidgetUpdateFlags::NeedsTick | EWidgetUpdateFlags::NeedsActiveTimerUpdate | EWidgetUpdateFlags::NeedsVolatilePaint | EWidgetUpdateFlags::NeedsVolatilePrepass);
+		return Widget.HasAnyUpdateFlags(EWidgetUpdateFlags::NeedsTick | EWidgetUpdateFlags::NeedsActiveTimerUpdate | EWidgetUpdateFlags::NeedsVolatilePaint | EWidgetUpdateFlags::NeedsVolatilePrepass);
 	}
-	static bool ShouldBeAddedToVolatileUpdateList(const TSharedRef<SWidget>& Widget)
+	static bool ShouldBeAddedToVolatileUpdateList(const TSharedRef<const SWidget>& Widget)
 	{
-		return ShouldBeAddedToVolatileUpdateList(&Widget.Get());
+		return ShouldBeAddedToVolatileUpdateList(Widget.Get());
 	}
 
 private:
@@ -440,8 +440,8 @@ private:
 	FCutResult CutArray(const FSlateInvalidationWidgetIndex WhereToCut);
 
 private:
-	FSlateInvalidationWidgetIndex Internal_BuildWidgetList_Recursive(TSharedRef<SWidget>& Widget, FSlateInvalidationWidgetIndex ParentIndex, IndexType& LastestIndex, FSlateInvalidationWidgetVisibility ParentVisibility, bool bParentVolatile);
-	void Internal_RebuildWidgetListTree(TSharedRef<SWidget> Widget, int32 ChildAtIndex);
+	FSlateInvalidationWidgetIndex Internal_BuildWidgetList_Recursive(SWidget& Widget, FSlateInvalidationWidgetIndex ParentIndex, IndexType& LastestIndex, FSlateInvalidationWidgetVisibility ParentVisibility, bool bParentVolatile);
+	void Internal_RebuildWidgetListTree(SWidget& Widget, int32 ChildAtIndex);
 	using FFindChildrenElement = TPair<SWidget*, FSlateInvalidationWidgetIndex>;
 	void Internal_FindChildren(FSlateInvalidationWidgetIndex WidgetIndex, TArray<FFindChildrenElement, TMemStackAllocator<>>& Widgets) const;
 	void Internal_RemoveRangeFromSameParent(const FIndexRange Range);
diff --git a/Engine/Source/Runtime/SlateCore/Private/FastUpdate/WidgetProxy.cpp b/Engine/Source/Runtime/SlateCore/Private/FastUpdate/WidgetProxy.cpp
--- a/Engine/Source/Runtime/SlateCore/Private/FastUpdate/WidgetProxy.cpp
+++ b/Engine/Source/Runtime/SlateCore/Private/FastUpdate/WidgetProxy.cpp
@@ -14,11 +14,11 @@
 
 const FSlateWidgetPersistentState FSlateWidgetPersistentState::NoState;
 
-FWidgetProxy::FWidgetProxy(TSharedRef<SWidget>& InWidget)
+FWidgetProxy::FWidgetProxy(SWidget& InWidget)
 #if UE_SLATE_WITH_WIDGETPROXY_WEAKPTR
-	: Widget(InWidget)
+	: Widget(InWidget.AsShared())
 #else
-	: Widget(&InWidget.Get())
+	: Widget(&InWidget)
 #endif
 	, Index(FSlateInvalidationWidgetIndex::Invalid)
 	, ParentIndex(FSlateInvalidationWidgetIndex::Invalid)
@@ -261,12 +261,18 @@ FWidgetProxyHandle::FWidgetProxyHandle(FSlateInvalidationWidgetIndex InIndex)
 {
 }
 
-bool FWidgetProxyHandle::IsValid(const SWidget* Widget) const
+bool FWidgetProxyHandle::IsValid(const SWidget& Widget) const
 {
 	FSlateInvalidationRoot* InvalidationRoot = InvalidationRootHandle.GetInvalidationRoot();
 	return InvalidationRoot
 		&& InvalidationRoot->GetFastPathWidgetList().IsValidIndex(WidgetIndex)
-		&& InvalidationRoot->GetFastPathWidgetList()[WidgetIndex].GetWidget() == Widget;
+		&& InvalidationRoot->GetFastPathWidgetList()[WidgetIndex].GetWidget() == &Widget;
+}
+
+bool FWidgetProxyHandle::IsValid(const SWidget* Widget) const
+{
+	check(Widget);
+	return IsValid(*Widget);
 }
 
 FSlateInvalidationWidgetVisibility FWidgetProxyHandle::GetWidgetVisibility(const SWidget* Widget) const
diff --git a/Engine/Source/Runtime/SlateCore/Private/Test/SlateInvalidationWidgetListTest.cpp b/Engine/Source/Runtime/SlateCore/Private/Test/SlateInvalidationWidgetListTest.cpp
--- a/Engine/Source/Runtime/SlateCore/Private/Test/SlateInvalidationWidgetListTest.cpp
+++ b/Engine/Source/Runtime/SlateCore/Private/Test/SlateInvalidationWidgetListTest.cpp
@@ -195,35 +195,35 @@ bool FSlateInvalidationWidgetListTest::RunTest(const FString& Parameters)
 			AddErrorIfFalse(List.VerifyWidgetsIndex(), TEXT("The widget list integrity has failed."));
 
 			{
-				TArray<TSharedPtr<SWidget>> FoundChildren = List.FindChildren(WidgetA.ToSharedRef());
+				TArray<TSharedPtr<SWidget>> FoundChildren = List.FindChildren(WidgetA.ToSharedRef().Get());
 				if (ChildOfWidgetA != FoundChildren)
 				{
 					AddError(TEXT("Was not able to find the child of VerticalBox A."));
 				}
 			}
 			{
-				TArray<TSharedPtr<SWidget>> FoundChildren = List.FindChildren(WidgetB.ToSharedRef());
+				TArray<TSharedPtr<SWidget>> FoundChildren = List.FindChildren(WidgetB.ToSharedRef().Get());
 				if (ChildOfWidgetB != FoundChildren)
 				{
 					AddError(TEXT("Was not able to find the child of VerticalBox B."));
 				}
 			}
 			{
-				TArray<TSharedPtr<SWidget>> FoundChildren = List.FindChildren(WidgetC.ToSharedRef());
+				TArray<TSharedPtr<SWidget>> FoundChildren = List.FindChildren(WidgetC.ToSharedRef().Get());
 				if (ChildOfWidgetC != FoundChildren)
 				{
 					AddError(TEXT("Was not able to find the child of VerticalBox C."));
 				}
 			}
 			{
-				TArray<TSharedPtr<SWidget>> FoundChildren = List.FindChildren(WidgetD.ToSharedRef());
+				TArray<TSharedPtr<SWidget>> FoundChildren = List.FindChildren(WidgetD.ToSharedRef().Get());
 				if (ChildOfWidgetD != FoundChildren)
 				{
 					AddError(TEXT("Was not able to find the child of VerticalBox D."));
 				}
 			}
 			{
-				TArray<TSharedPtr<SWidget>> FoundChildren = List.FindChildren(WidgetH.ToSharedRef());
+				TArray<TSharedPtr<SWidget>> FoundChildren = List.FindChildren(WidgetH.ToSharedRef().Get());
 				if (ChildOfWidgetH != FoundChildren)
 				{
 					AddError(TEXT("Was not able to find the child of VerticalBox H."));
@@ -241,7 +241,7 @@ bool FSlateInvalidationWidgetListTest::RunTest(const FString& Parameters)
 
 			// Remove Second child of F
 			{
-				List.RemoveWidget(WidgetF->GetAllChildren()->GetChildAt(1));
+				List.RemoveWidget(WidgetF->GetAllChildren()->GetChildAt(1).Get());
 				WidgetF->RemoveSlot(WidgetF->GetAllChildren()->GetChildAt(1));
 				{
 					FSlateInvalidationWidgetList::FArguments ArgsWithoutAssign = Args;
@@ -258,7 +258,7 @@ bool FSlateInvalidationWidgetListTest::RunTest(const FString& Parameters)
 			// Remove C and F
 			{
 				{
-					const FSlateInvalidationWidgetIndex WidgetIndexF = List.FindWidget(WidgetF.ToSharedRef());
+					const FSlateInvalidationWidgetIndex WidgetIndexF = List.FindWidget(WidgetF.ToSharedRef().Get());
 					if (!List.IsValidIndex(WidgetIndexF) || List[WidgetIndexF].GetWidget() != WidgetF.Get())
 					{
 						AddError(TEXT("The index of F is not valid anymore."));
@@ -266,7 +266,7 @@ bool FSlateInvalidationWidgetListTest::RunTest(const FString& Parameters)
 					List.RemoveWidget(WidgetIndexF);
 				}
 				{
-					const FSlateInvalidationWidgetIndex WidgetIndexC = List.FindWidget(WidgetC.ToSharedRef());
+					const FSlateInvalidationWidgetIndex WidgetIndexC = List.FindWidget(WidgetC.ToSharedRef().Get());
 					if (!List.IsValidIndex(WidgetIndexC) || List[WidgetIndexC].GetWidget() != WidgetC.Get())
 					{
 						AddError(TEXT("The index of C is not valid anymore."));
@@ -291,7 +291,7 @@ bool FSlateInvalidationWidgetListTest::RunTest(const FString& Parameters)
 				int32 ToRemoveIndex = RootChildOrder->GetAllChildren()->Num() - 1;
 				TSharedRef<SWidget> RemovedWidget = RootChildOrder->GetAllChildren()->GetChildAt(ToRemoveIndex);
 				RootChildOrder->RemoveSlot(RemovedWidget);
-				List.RemoveWidget(List.FindWidget(RemovedWidget));
+				List.RemoveWidget(List.FindWidget(RemovedWidget.Get()));
 
 				{
 					FSlateInvalidationWidgetList::FArguments ArgsWithoutAssign = Args;
diff --git a/Engine/Source/Runtime/SlateCore/Private/Types/SlateAttributeMetaData.cpp b/Engine/Source/Runtime/SlateCore/Private/Types/SlateAttributeMetaData.cpp
--- a/Engine/Source/Runtime/SlateCore/Private/Types/SlateAttributeMetaData.cpp
+++ b/Engine/Source/Runtime/SlateCore/Private/Types/SlateAttributeMetaData.cpp
@@ -29,18 +29,18 @@ FSlateAttributeMetaData* FSlateAttributeMetaData::FindMetaData(const SWidget& Ow
 	if (OwningWidget.HasRegisteredSlateAttribute())
 	{
 		check(OwningWidget.MetaData.Num() > 0);
-		TSharedRef<ISlateMetaData> SlateMetaData = OwningWidget.MetaData[0];
+		const TSharedRef<ISlateMetaData>& SlateMetaData = OwningWidget.MetaData[0];
 		check(SlateMetaData->IsOfType<FSlateAttributeMetaData>());
-		return &(StaticCastSharedRef<FSlateAttributeMetaData>(SlateMetaData).Get());
+		return &(static_cast<FSlateAttributeMetaData&>(SlateMetaData.Get()));
 	}
 #if WITH_SLATE_DEBUGGING
 	else if (OwningWidget.MetaData.Num() > 0)
 	{
-		TSharedRef<ISlateMetaData> SlateMetaData = OwningWidget.MetaData[0];
+		const TSharedRef<ISlateMetaData>& SlateMetaData = OwningWidget.MetaData[0];
 		if (SlateMetaData->IsOfType<FSlateAttributeMetaData>())
 		{
 			ensureMsgf(false, TEXT("bHasRegisteredSlateAttribute should be set on the SWidget '%s'"), *FReflectionMetaData::GetWidgetDebugInfo(OwningWidget));
-			return &(StaticCastSharedRef<FSlateAttributeMetaData>(SlateMetaData).Get());
+			return &(static_cast<FSlateAttributeMetaData&>(SlateMetaData.Get()));
 		}
 	}
 #endif
@@ -356,17 +356,14 @@ void FSlateAttributeMetaData::UpdateExceptVisibilityAttributes(SWidget& OwningWi
 
 void FSlateAttributeMetaData::UpdateChildrenOnlyVisibilityAttributes(SWidget& OwningWidget, EInvalidationPermission InvalidationStyle, bool bRecursive)
 {
-	FChildren* Children = OwningWidget.GetChildren();
-	const int32 NumChildren = Children->Num();
-	for (int32 ChildIndex = 0; ChildIndex < NumChildren; ++ChildIndex)
-	{
-		const TSharedRef<SWidget>& Child = Children->GetChildAt(ChildIndex);
-		UpdateOnlyVisibilityAttributes(Child.Get(), InvalidationStyle);
-		if (bRecursive)
+	OwningWidget.GetChildren()->ForEachWidget([InvalidationStyle, bRecursive](SWidget& Child)
 		{
-			UpdateChildrenOnlyVisibilityAttributes(Child.Get(), InvalidationStyle, bRecursive);
-		}
-	}
+			UpdateOnlyVisibilityAttributes(Child, InvalidationStyle);
+			if (bRecursive)
+			{
+				UpdateChildrenOnlyVisibilityAttributes(Child, InvalidationStyle, bRecursive);
+			}
+		});
 }
 
 
diff --git a/Engine/Source/Runtime/SlateCore/Private/Widgets/SWidget.cpp b/Engine/Source/Runtime/SlateCore/Private/Widgets/SWidget.cpp
--- a/Engine/Source/Runtime/SlateCore/Private/Widgets/SWidget.cpp
+++ b/Engine/Source/Runtime/SlateCore/Private/Widgets/SWidget.cpp
@@ -853,13 +853,10 @@ void SWidget::UpdateFastPathVisibility(FSlateInvalidationWidgetVisibility Parent
 	PersistentState.CachedElementHandle.ClearCachedElements();
 
 	// Loop through children
-	FChildren* MyChildren = GetAllChildren();
-	const int32 NumChildren = MyChildren->Num();
-	for (int32 ChildIndex = 0; ChildIndex < NumChildren; ++ChildIndex)
-	{
-		const TSharedRef<SWidget>& Child = MyChildren->GetChildAt(ChildIndex);
-		Child->UpdateFastPathVisibility(NewVisibility, HittestGridToRemoveFrom);
-	}
+	GetAllChildren()->ForEachWidget([NewVisibility, HittestGridToRemoveFrom](SWidget& Child)
+		{
+			Child.UpdateFastPathVisibility(NewVisibility, HittestGridToRemoveFrom);
+		});
 }
 
 void SWidget::UpdateFastPathWidgetRemoved(FHittestGrid* ParentHittestGrid)
@@ -881,13 +878,10 @@ void SWidget::UpdateFastPathWidgetRemoved(FHittestGrid* ParentHittestGrid)
 	PersistentState.CachedElementHandle.RemoveFromCache();
 
 	// Loop through children
-	FChildren* MyChildren = GetAllChildren();
-	const int32 NumChildren = MyChildren->Num();
-	for (int32 ChildIndex = 0; ChildIndex < NumChildren; ++ChildIndex)
-	{
-		const TSharedRef<SWidget>& Child = MyChildren->GetChildAt(ChildIndex);
-		Child->UpdateFastPathWidgetRemoved(HittestGridToRemoveFrom);
-	}
+	GetAllChildren()->ForEachWidget([HittestGridToRemoveFrom](SWidget& Child)
+		{
+			Child.UpdateFastPathWidgetRemoved(HittestGridToRemoveFrom);
+		});
 }
 
 void SWidget::UpdateFastPathVolatility(bool bParentVolatile)
@@ -903,13 +897,11 @@ void SWidget::UpdateFastPathVolatility(bool bParentVolatile)
 		RemoveUpdateFlags(EWidgetUpdateFlags::NeedsVolatilePaint);
 	}
 
-	FChildren* MyChildren = GetAllChildren();
-	const int32 NumChildren = MyChildren->Num();
-	for (int32 ChildIndex = 0; ChildIndex < NumChildren; ++ChildIndex)
-	{
-		const TSharedRef<SWidget>& Child = MyChildren->GetChildAt(ChildIndex);
-		Child->UpdateFastPathVolatility(bParentVolatile || IsVolatile());
-	}
+	const bool bNewParentVolatility = bParentVolatile || IsVolatile();
+	GetAllChildren()->ForEachWidget([bNewParentVolatility](SWidget& Child)
+		{
+			Child.UpdateFastPathVolatility(bNewParentVolatility);
+		});
 }
 
 void SWidget::CacheDesiredSize(float InLayoutScaleMultiplier)
@@ -1367,7 +1359,7 @@ int32 SWidget::Paint(const FPaintArgs& Args, const FGeometry& AllottedGeometry,
 	const EWidgetUpdateFlags PreviousUpdateFlag = UpdateFlags;
 
 	// TODO, Maybe we should just make Paint non-const and keep OnPaint const.
-	TSharedRef<SWidget> MutableThis = ConstCastSharedRef<SWidget>(AsShared());
+	SWidget* MutableThis = const_cast<SWidget*>(this);
 
 	INC_DWORD_STAT(STAT_SlateNumPaintedWidgets);
 	UE_TRACE_SCOPED_SLATE_WIDGET_PAINT(this);
@@ -1451,12 +1443,6 @@ int32 SWidget::Paint(const FPaintArgs& Args, const FGeometry& AllottedGeometry,
 	FPaintArgs UpdatedArgs = Args.WithNewParent(this);
 	UpdatedArgs.SetInheritedHittestability(bOutgoingHittestability);
 
-
-#if 0
-	// test ensure that we are not the last thing holding this widget together
-	ensure(!MutableThis.IsUnique());
-#endif
-
 #if WITH_SLATE_DEBUGGING
 	if (FastPathProxyHandle.IsValid(this) && PersistentState.CachedElementHandle.IsValid())
 	{
@@ -1468,7 +1454,7 @@ int32 SWidget::Paint(const FPaintArgs& Args, const FGeometry& AllottedGeometry,
 
 	if (bOutgoingHittestability)
 	{
-		Args.GetHittestGrid().AddWidget(MutableThis, 0, LayerId, FastPathProxyHandle.GetWidgetSortOrder());
+		Args.GetHittestGrid().AddWidget(MutableThis->AsShared(), 0, LayerId, FastPathProxyHandle.GetWidgetSortOrder());
 	}
 
 	if (bClipToBounds)
@@ -1499,16 +1485,13 @@ int32 SWidget::Paint(const FPaintArgs& Args, const FGeometry& AllottedGeometry,
 		// of rules apply to those widgets.
 		if (!IsVolatile() && !IsVolatileIndirectly() && !Advanced_IsInvalidationRoot())
 		{
-			const FChildren* MyChildren = MutableThis->GetChildren();
-			const int32 NumChildren = MyChildren->Num();
-			for (int32 ChildIndex = 0; ChildIndex < MyChildren->Num(); ++ChildIndex)
-			{
-				TSharedRef<const SWidget> Child = MyChildren->GetChildAt(ChildIndex);
-				if (Child->GetVisibility().IsVisible())
+			MutableThis->GetChildren()->ForEachWidget([&DebugChildWidgetsToPaint](const SWidget& Child)
 				{
-					DebugChildWidgetsToPaint.Add(Child);
-				}
-			}
+					if (Child.GetVisibility().IsVisible())
+					{
+						DebugChildWidgetsToPaint.Add(Child.AsShared());
+					}
+				});
 		}
 	}
 #endif
@@ -1664,42 +1647,41 @@ void SWidget::Prepass_Internal(float InLayoutScaleMultiplier)
 
 void SWidget::Prepass_ChildLoop(float InLayoutScaleMultiplier, FChildren* MyChildren)
 {
-	const int32 NumChildren = MyChildren->Num();
-	for (int32 ChildIndex = 0; ChildIndex < NumChildren; ++ChildIndex)
-	{
-		const TSharedRef<SWidget>& Child = MyChildren->GetChildAt(ChildIndex);
-
-		const bool bUpdateAttributes = Child->HasRegisteredSlateAttribute() && Child->IsAttributesUpdatesEnabled() && !GSlateIsOnFastProcessInvalidation;
-		if (bUpdateAttributes)
-		{
-			FSlateAttributeMetaData::UpdateOnlyVisibilityAttributes(Child.Get(), FSlateAttributeMetaData::EInvalidationPermission::AllowInvalidationIfConstructed);
-		}
-
-		if (Child->GetVisibility() != EVisibility::Collapsed)
+	int32 ChildIndex = 0;
+	MyChildren->ForEachWidget([this, &ChildIndex, InLayoutScaleMultiplier](SWidget& Child)
 		{
+			const bool bUpdateAttributes = Child.HasRegisteredSlateAttribute() && Child.IsAttributesUpdatesEnabled() && !GSlateIsOnFastProcessInvalidation;
 			if (bUpdateAttributes)
 			{
-				FSlateAttributeMetaData::UpdateExceptVisibilityAttributes(Child.Get(), FSlateAttributeMetaData::EInvalidationPermission::AllowInvalidationIfConstructed);
+				FSlateAttributeMetaData::UpdateOnlyVisibilityAttributes(Child, FSlateAttributeMetaData::EInvalidationPermission::AllowInvalidationIfConstructed);
 			}
 
-			const float ChildLayoutScaleMultiplier = bHasRelativeLayoutScale
-				? InLayoutScaleMultiplier * GetRelativeLayoutScale(ChildIndex, InLayoutScaleMultiplier)
-				: InLayoutScaleMultiplier;
+			if (Child.GetVisibility() != EVisibility::Collapsed)
+			{
+				if (bUpdateAttributes)
+				{
+					FSlateAttributeMetaData::UpdateExceptVisibilityAttributes(Child, FSlateAttributeMetaData::EInvalidationPermission::AllowInvalidationIfConstructed);
+				}
 
-			// Recur: Descend down the widget tree.
-			Child->Prepass_Internal(ChildLayoutScaleMultiplier);
-		}
-		else
-		{
-			// If the child widget is collapsed, we need to store the new layout scale it will have when 
-			// it is finally visible and invalidate it's prepass so that it gets that when its visibility
-			// is finally invalidated.
-			Child->MarkPrepassAsDirty();
-			Child->PrepassLayoutScaleMultiplier = bHasRelativeLayoutScale
-				? InLayoutScaleMultiplier * GetRelativeLayoutScale(ChildIndex, InLayoutScaleMultiplier)
-				: InLayoutScaleMultiplier;
-		}
-	}
+				const float ChildLayoutScaleMultiplier = bHasRelativeLayoutScale
+					? InLayoutScaleMultiplier * GetRelativeLayoutScale(ChildIndex, InLayoutScaleMultiplier)
+					: InLayoutScaleMultiplier;
+
+				// Recur: Descend down the widget tree.
+				Child.Prepass_Internal(ChildLayoutScaleMultiplier);
+			}
+			else
+			{
+				// If the child widget is collapsed, we need to store the new layout scale it will have when 
+				// it is finally visible and invalidate it's prepass so that it gets that when its visibility
+				// is finally invalidated.
+				Child.MarkPrepassAsDirty();
+				Child.PrepassLayoutScaleMultiplier = bHasRelativeLayoutScale
+					? InLayoutScaleMultiplier * GetRelativeLayoutScale(ChildIndex, InLayoutScaleMultiplier)
+					: InLayoutScaleMultiplier;
+			}
+			++ChildIndex;
+		});
 }
 
 TSharedRef<FActiveTimerHandle> SWidget::RegisterActiveTimer(float TickPeriod, FWidgetActiveTimerDelegate TickFunction)
@@ -1810,8 +1792,8 @@ void SWidget::AddMetadataInternal(const TSharedRef<ISlateMetaData>& AddMe)
 	if (AddMe->IsOfType<FReflectionMetaData>())
 	{
 #if UE_WITH_SLATE_DEBUG_FIND_WIDGET_REFLECTION_METADATA
-		TSharedRef<FReflectionMetaData> Reflection = StaticCastSharedRef<FReflectionMetaData>(AddMe);
-		if (Reflection->Name == FindWidgetMetaData::WidgeName && Reflection->Asset.Get() && Reflection->Asset.Get()->GetFName() == FindWidgetMetaData::AssetName)
+		FReflectionMetaData& Reflection = static_cast<FReflectionMetaData&>(AddMe.Get());
+		if (Reflection.Name == FindWidgetMetaData::WidgeName && Reflection.Asset.Get() && Reflection.Asset.Get()->GetFName() == FindWidgetMetaData::AssetName)
 		{
 			FindWidgetMetaData::FoundWidget = this;
 		}
diff --git a/Engine/Source/Runtime/SlateCore/Public/FastUpdate/SlateInvalidationRoot.h b/Engine/Source/Runtime/SlateCore/Public/FastUpdate/SlateInvalidationRoot.h
--- a/Engine/Source/Runtime/SlateCore/Public/FastUpdate/SlateInvalidationRoot.h
+++ b/Engine/Source/Runtime/SlateCore/Public/FastUpdate/SlateInvalidationRoot.h
@@ -145,7 +145,7 @@ private:
 	/** Call when an invalidation occurred. */
 	void InvalidateWidget(FWidgetProxy& Proxy, EInvalidateWidgetReason InvalidateReason);
 
-	void BuildFastPathWidgetList(TSharedRef<SWidget> RootWidget);
+	void BuildFastPathWidgetList(const TSharedRef<SWidget>& RootWidget);
 	void AdjustWidgetsDesktopGeometry(FVector2D WindowToDesktopTransform);
 
 	/** Update child order and slate attribute */
diff --git a/Engine/Source/Runtime/SlateCore/Public/FastUpdate/WidgetProxy.h b/Engine/Source/Runtime/SlateCore/Public/FastUpdate/WidgetProxy.h
--- a/Engine/Source/Runtime/SlateCore/Public/FastUpdate/WidgetProxy.h
+++ b/Engine/Source/Runtime/SlateCore/Public/FastUpdate/WidgetProxy.h
@@ -111,7 +111,7 @@ static_assert(sizeof(FSlateInvalidationWidgetVisibility) == sizeof(uint8), "FSla
 class FWidgetProxy
 {
 public:
-	FWidgetProxy(TSharedRef<SWidget>& InWidget);
+	FWidgetProxy(SWidget& InWidget);
 
 	int32 Update(const FPaintArgs& PaintArgs, FSlateWindowElementList& OutDrawElements);
 
@@ -257,6 +257,7 @@ public:
 	{}
 
 	/** @returns true if it has a valid InvalidationRoot and Index. */
+	SLATECORE_API bool IsValid(const SWidget& Widget) const;
 	SLATECORE_API bool IsValid(const SWidget* Widget) const;
 
 	FSlateInvalidationRootHandle GetInvalidationRootHandle() const { return InvalidationRootHandle; }
diff --git a/Engine/Source/Runtime/SlateCore/Public/Layout/Children.h b/Engine/Source/Runtime/SlateCore/Public/Layout/Children.h
--- a/Engine/Source/Runtime/SlateCore/Public/Layout/Children.h
+++ b/Engine/Source/Runtime/SlateCore/Public/Layout/Children.h
@@ -38,13 +38,119 @@ public:
 	/** @return const pointer to the Widget at the specified Index. */
 	virtual TSharedRef<const SWidget> GetChildAt( int32 Index ) const = 0;
 
+	/** @return the SWidget that own the FChildren */
 	SWidget& GetOwner() const { return *Owner; }
 
+	/** Applies the predicate to all the widgets contained by the FChildren. */
+	template<typename Predicate>
+	void ForEachWidget(Predicate Pred)
+	{
+		int32 WidgetCount = Num();
+		for (int32 Index = 0; Index < WidgetCount; ++Index)
+		{
+			Pred(GetChildRefAt(Index).GetWidget());
+		}
+	}
+
+	/** Applies the predicate to all the widgets contained by the FChildren. */
+	template<typename Predicate>
+	void ForEachWidget(Predicate Pred) const
+	{
+		int32 WidgetCount = Num();
+		for (int32 Index = 0; Index < WidgetCount; ++Index)
+		{
+			Pred(GetChildRefAt(Index).GetWidget());
+		}
+	}
+
 protected:
 	friend class SWidget;
 	friend class FCombinedChildren;
+	template<typename T>
+	friend class TOneDynamicChild;
+
+	enum ECopyConstruct { CopyConstruct };
+	enum ERefConstruct { ReferenceConstruct };
+
+	struct FWidgetRef
+	{
+	private:
+		TOptional<TSharedRef<SWidget>> WidgetCopy;
+		const TSharedRef<SWidget>& WidgetReference;
+
+	public:
+		FWidgetRef(ECopyConstruct, TSharedRef<SWidget> InWidgetCopy)
+			: WidgetCopy(MoveTemp(InWidgetCopy))
+			, WidgetReference(WidgetCopy.GetValue())
+		{}
+		FWidgetRef(ERefConstruct, const TSharedRef<SWidget>& InWidgetRef)
+			: WidgetCopy()
+			, WidgetReference(InWidgetRef)
+		{}
+		FWidgetRef(const FWidgetRef& Other)
+			: WidgetCopy(Other.WidgetCopy)
+			, WidgetReference(WidgetCopy.IsSet() ? WidgetCopy.GetValue() : Other.WidgetReference)
+		{}
+		FWidgetRef(FWidgetRef&& Other)
+			: WidgetCopy(MoveTemp(Other.WidgetCopy))
+			, WidgetReference(WidgetCopy.IsSet() ? WidgetCopy.GetValue() : Other.WidgetReference)
+		{}
+		FWidgetRef& operator=(const FWidgetRef&) = delete;
+		FWidgetRef& operator=(FWidgetRef&&) = delete;
+		const TSharedRef<SWidget>& GetWidgetRef() const
+		{
+			return WidgetReference;
+		}
+		SWidget& GetWidget() const
+		{
+			return WidgetReference.Get();
+		}
+	};
+
+	struct FConstWidgetRef
+	{
+	private:
+		TOptional<TSharedRef<const SWidget>> WidgetCopy;
+		const TSharedRef<const SWidget>& WidgetReference;
+
+	public:
+		FConstWidgetRef(ECopyConstruct, TSharedRef<const SWidget> InWidgetCopy)
+			: WidgetCopy(MoveTemp(InWidgetCopy))
+			, WidgetReference(WidgetCopy.GetValue())
+		{}
+		FConstWidgetRef(ERefConstruct, const TSharedRef<const SWidget>& InWidgetRef)
+			: WidgetCopy()
+			, WidgetReference(InWidgetRef)
+		{}
+		FConstWidgetRef(const FConstWidgetRef& Other)
+			: WidgetCopy(Other.WidgetCopy)
+			, WidgetReference(WidgetCopy.IsSet() ? WidgetCopy.GetValue() : Other.WidgetReference)
+		{}
+		FConstWidgetRef(FConstWidgetRef&& Other)
+			: WidgetCopy(MoveTemp(Other.WidgetCopy))
+			, WidgetReference(WidgetCopy.IsSet() ? WidgetCopy.GetValue() : Other.WidgetReference)
+		{}
+		FConstWidgetRef& operator=(const FConstWidgetRef&) = delete;
+		FConstWidgetRef& operator=(FConstWidgetRef&&) = delete;
+		const TSharedRef<const SWidget>& GetWidgetRef() const
+		{
+			return WidgetReference;
+		}
+		const SWidget& GetWidget() const
+		{
+			return WidgetReference.Get();
+		}
+	};
+
+	/** @return the number of slot the children has. */
+	virtual int32 NumSlot() const { return Num(); }
 	/** @return the const reference to the slot at the specified Index */
 	virtual const FSlotBase& GetSlotAt(int32 ChildIndex) const = 0;
+	/** @return ref to the Widget at the specified Index. */
+	virtual FWidgetRef GetChildRefAt(int32 Index) = 0;
+	/** @return ref to the Widget at the specified Index. */
+	virtual FConstWidgetRef GetChildRefAt(int32 Index) const = 0;
+
 
 protected:
 	virtual ~FChildren(){}
@@ -85,21 +191,32 @@ public:
 
 	virtual TSharedRef<SWidget> GetChildAt(int32 Index) override
 	{
-		return GetSlotAt(Index).GetWidget();
+		return GetChildRefAt(Index).GetWidgetRef();
 	}
 
 	virtual TSharedRef<const SWidget> GetChildAt(int32 Index) const override
 	{
-		return GetSlotAt(Index).GetWidget();
+		return GetChildRefAt(Index).GetWidgetRef();
 	}
 
 protected:
+	virtual int32 NumSlot() const override
+	{
+		int32 TotalNum = 0;
+		for (const FChildren* Children : LinkedChildren)
+		{
+			TotalNum += Children->NumSlot();
+		}
+
+		return TotalNum;
+	}
+
 	virtual const FSlotBase& GetSlotAt(int32 ChildIndex) const override
 	{
 		int32 TotalNum = 0;
 		for (const FChildren* Children : LinkedChildren)
 		{
-			const int32 NewTotal = TotalNum + Children->Num();
+			const int32 NewTotal = TotalNum + Children->NumSlot();
 			if (NewTotal > ChildIndex)
 			{
 				return Children->GetSlotAt(ChildIndex - TotalNum);
@@ -108,7 +225,43 @@ protected:
 		}
 
 		// This result should never occur users should always access a valid index for child slots.
-		static FSlotBase NullSlot; check(false); return NullSlot;
+		check(false);
+		static FSlotBase NullSlot;
+		return NullSlot;
+	}
+
+	virtual FWidgetRef GetChildRefAt(int32 Index) override
+	{
+		int32 TotalNum = 0;
+		for (FChildren* Children : LinkedChildren)
+		{
+			const int32 NewTotal = TotalNum + Children->Num();
+			if (NewTotal > Index)
+			{
+				return Children->GetChildRefAt(Index - TotalNum);
+			}
+			TotalNum = NewTotal;
+		}
+		// This result should never occur users should always access a valid index for child slots.
+		check(false);
+		return FWidgetRef(ReferenceConstruct, SNullWidget::NullWidget);
+	}
+
+	virtual FConstWidgetRef GetChildRefAt(int32 Index) const override
+	{
+		int32 TotalNum = 0;
+		for (const FChildren* Children : LinkedChildren)
+		{
+			const int32 NewTotal = TotalNum + Children->Num();
+			if (NewTotal > Index)
+			{
+				return Children->GetChildRefAt(Index - TotalNum);
+			}
+			TotalNum = NewTotal;
+		}
+		// This result should never occur users should always access a valid index for child slots.
+		check(false);
+		return FConstWidgetRef(ReferenceConstruct, SNullWidget::NullWidget);
 	}
 
 protected:
@@ -160,6 +313,18 @@ private:
 		static FSlotBase NullSlot;
 		return NullSlot;
 	}
+
+	virtual FWidgetRef GetChildRefAt(int32 Index) override
+	{
+		check(false);
+		return FWidgetRef(ReferenceConstruct, SNullWidget::NullWidget);
+	}
+	
+	virtual FConstWidgetRef GetChildRefAt(int32 Index) const override
+	{
+		check(false);
+		return FConstWidgetRef(ReferenceConstruct, SNullWidget::NullWidget);
+	}
 };
 
 
@@ -194,7 +359,22 @@ public:
 	}
 
 private:
-	virtual const FSlotBase& GetSlotAt(int32 ChildIndex) const override { check(ChildIndex == 0); return *this; }
+	virtual const FSlotBase& GetSlotAt(int32 ChildIndex) const override
+	{
+		check(ChildIndex == 0);
+		return *this;
+	}
+
+	virtual FWidgetRef GetChildRefAt(int32 ChildIndex) override
+	{
+		check(ChildIndex == 0);
+		return FWidgetRef(ReferenceConstruct, FSlotBase::GetWidget());
+	}
+	virtual FConstWidgetRef GetChildRefAt(int32 ChildIndex) const override
+	{
+		check(ChildIndex == 0);
+		return FConstWidgetRef(ReferenceConstruct, FSlotBase::GetWidget());
+	}
 };
 
 
@@ -209,7 +389,10 @@ class TWeakChild : public FChildren
 public:
 	using FChildren::FChildren;
 
-	virtual int32 Num() const override { return WidgetPtr.IsValid() ? 1 : 0 ; }
+	virtual int32 Num() const override
+	{
+		return WidgetPtr.IsValid() ? 1 : 0 ;
+	}
 
 	virtual TSharedRef<SWidget> GetChildAt( int32 ChildIndex ) override
 	{
@@ -224,7 +407,30 @@ public:
 	}
 
 private:
-	virtual const FSlotBase& GetSlotAt(int32 ChildIndex) const override { static FSlotBase NullSlot; check(ChildIndex == 0); return NullSlot; }
+	virtual int32 NumSlot() const override
+	{
+		return 0;
+	}
+
+	virtual const FSlotBase& GetSlotAt(int32 ChildIndex) const override
+	{
+		static FSlotBase NullSlot;
+		check(ChildIndex == 0);
+		return NullSlot;
+	}
+
+	virtual FWidgetRef GetChildRefAt(int32 ChildIndex) override
+	{
+		check(ChildIndex == 0);
+		TSharedPtr<SWidget> Widget = WidgetPtr.Pin();
+		return (Widget.IsValid()) ? FWidgetRef(CopyConstruct, Widget.ToSharedRef()) : FWidgetRef(ReferenceConstruct, SNullWidget::NullWidget);
+	}
+	virtual FConstWidgetRef GetChildRefAt(int32 ChildIndex) const override
+	{
+		check(ChildIndex == 0);
+		TSharedPtr<SWidget> Widget = WidgetPtr.Pin();
+		return (Widget.IsValid()) ? FConstWidgetRef(CopyConstruct, Widget.ToSharedRef()) : FConstWidgetRef(ReferenceConstruct, SNullWidget::NullWidget);
+	}
 
 public:
 	void AttachWidget(const TSharedPtr<SWidget>& InWidget)
@@ -240,10 +446,8 @@ public:
 
 	void DetachWidget()
 	{
-		if (WidgetPtr.IsValid())
+		if (TSharedPtr<SWidget> Widget = WidgetPtr.Pin())
 		{
-			TSharedPtr<SWidget> Widget = WidgetPtr.Pin();
-
 			if (Widget != SNullWidget::NullWidget)
 			{
 				Widget->ConditionallyDetatchParentWidget(&GetOwner());
@@ -260,14 +464,6 @@ public:
 		return (Widget.IsValid()) ? Widget.ToSharedRef() : SNullWidget::NullWidget;
 	}
 
-private:
-	SWidget& GetWidgetRef() const
-	{
-		ensure(Num() > 0);
-		SWidget* Widget = WidgetPtr.Pin().Get();
-		return Widget ? *Widget : SNullWidget::NullWidget.Get();
-	}
-
 private:
 	TWeakPtr<ChildType> WidgetPtr;
 };
@@ -365,6 +561,17 @@ private:
 		check(ChildIndex == 0);
 		return *this;
 	}
+
+	virtual FWidgetRef GetChildRefAt(int32 ChildIndex) override
+	{
+		check(ChildIndex == 0);
+		return FWidgetRef(ReferenceConstruct, this->GetWidget());
+	}
+	virtual FConstWidgetRef GetChildRefAt(int32 ChildIndex) const override
+	{
+		check(ChildIndex == 0);
+		return FConstWidgetRef(ReferenceConstruct, this->GetWidget());
+	}
 };
 
 
@@ -432,11 +639,22 @@ class TPanelChildren : public FChildren
 private:
 	TArray<TUniquePtr<SlotType>> Children;
 
+protected:
 	virtual const FSlotBase& GetSlotAt(int32 ChildIndex) const override
 	{
 		return *Children[ChildIndex];
 	}
 
+	virtual FWidgetRef GetChildRefAt(int32 ChildIndex) override
+	{
+		return FWidgetRef(ReferenceConstruct, Children[ChildIndex]->GetWidget());
+	}
+
+	virtual FConstWidgetRef GetChildRefAt(int32 ChildIndex) const override
+	{
+		return FConstWidgetRef(ReferenceConstruct, Children[ChildIndex]->GetWidget());
+	}
+
 public:
 	using FChildren::FChildren;
 	
@@ -455,6 +673,7 @@ public:
 		return Children[Index]->GetWidget();
 	}
 
+public:
 	int32 Add( SlotType* Slot )
 	{
 		int32 Index = Children.Add(TUniquePtr<SlotType>(Slot));
@@ -772,13 +991,29 @@ class TSlotlessChildren : public FChildren
 private:
 	TArray<TSharedRef<ChildType>> Children;
 
+	virtual int32 NumSlot() const override
+	{
+		return 0;
+	}
+
 	virtual const FSlotBase& GetSlotAt(int32 ChildIndex) const override
 	{
 		// @todo slate : slotless children should be removed altogether; for now they return a fake slot.
 		static FSlotBase NullSlot;
+		check(false);
 		return NullSlot;
 	}
 
+	virtual FWidgetRef GetChildRefAt(int32 ChildIndex) override
+	{
+		return FWidgetRef(ReferenceConstruct, Children[ChildIndex]);
+	}
+	
+	virtual FConstWidgetRef GetChildRefAt(int32 ChildIndex) const override
+	{
+		return FConstWidgetRef(ReferenceConstruct, Children[ChildIndex]);
+	}
+
 public:
 	TSlotlessChildren(SWidget* InOwner, bool InbChangesInvalidatePrepass = true)
 		: FChildren(InOwner)
@@ -1034,7 +1269,8 @@ public:
 
 	virtual TSharedRef<SWidget> GetChildAt(int32 Index) override
 	{
-		check(Index == 0); return AllChildren->GetChildAt(WidgetIndex->Get());
+		check(Index == 0);
+		return AllChildren->GetChildAt(WidgetIndex->Get());
 	}
 
 	virtual TSharedRef<const SWidget> GetChildAt(int32 Index) const override
@@ -1044,8 +1280,22 @@ public:
 	}
 
 private:
+	virtual const FSlotBase& GetSlotAt(int32 ChildIndex) const override
+	{
+		return (*AllChildren)[ChildIndex];
+	}
 
-	virtual const FSlotBase& GetSlotAt(int32 ChildIndex) const override { return (*AllChildren)[ChildIndex]; }
+	virtual FWidgetRef GetChildRefAt(int32 ChildIndex) override
+	{
+		check(ChildIndex == 0);
+		return static_cast<FChildren*>(AllChildren)->GetChildRefAt(WidgetIndex->Get());
+	}
+	
+	virtual FConstWidgetRef GetChildRefAt(int32 ChildIndex) const override
+	{
+		check(ChildIndex == 0);
+		return static_cast<const FChildren*>(AllChildren)->GetChildRefAt(WidgetIndex->Get());
+	}
 
 	TPanelChildren<SlotType>* AllChildren;
 	const TAttribute<int32>* WidgetIndex;
diff --git a/Engine/Source/Runtime/SlateCore/Public/SlotBase.h b/Engine/Source/Runtime/SlateCore/Public/SlotBase.h
--- a/Engine/Source/Runtime/SlateCore/Public/SlotBase.h
+++ b/Engine/Source/Runtime/SlateCore/Public/SlotBase.h
@@ -117,7 +117,7 @@ private:
 #if WITH_EDITORONLY_DATA
 protected:
 	/** The parent and owner of the slot. */
-	UE_DEPRECATED(5.0, "RawParentPtr is not used anymore. Use get Owner Widget.")
+	UE_DEPRECATED(5.0, "RawParentPtr is not used anymore. Use GetOwnerWidget.")
 	SWidget* RawParentPtr;
 #endif
 };
@@ -184,7 +184,7 @@ public:
 		}
 
 		/** @return the child widget that will be owned by the slot. */
-		TSharedPtr<SWidget> GetAttachedWidget() const { return ChildWidget; }
+		const TSharedPtr<SWidget>& GetAttachedWidget() const { return ChildWidget; }
 
 		/** @return the slot that is being constructed. */
 		SlotType* GetSlot() const { return Slot.Get(); }
