Rename SWidget::IsConstructionCompleted to SWidget::IsConstructed

diff --git a/Engine/Source/Runtime/SlateCore/Private/Types/SlateAttribute.cpp b/Engine/Source/Runtime/SlateCore/Private/Types/SlateAttribute.cpp
--- a/Engine/Source/Runtime/SlateCore/Private/Types/SlateAttribute.cpp
+++ b/Engine/Source/Runtime/SlateCore/Private/Types/SlateAttribute.cpp
@@ -12,7 +12,7 @@ namespace SlateAttributePrivate
 	void TestAttributeAddress(const SWidget& OwningWidget, const FSlateAttributeImpl& Attribute, ESlateAttributeType AttributeType)
 	{
 #if STATS && DO_CHECK
-		//@TODO: DarenC - Using allocsize as proxy for 'IsConstructionCompleted' due to encapsulation.
+		//@TODO: DarenC - Using allocsize as proxy for 'IsConstructed' due to encapsulation.
 		if (AttributeType == ESlateAttributeType::Member && OwningWidget.GetAllocSize() > 0)
 		{
 			UPTRINT SlateAttributePtr = (UPTRINT)&Attribute;
diff --git a/Engine/Source/Runtime/SlateCore/Private/Types/SlateAttributeMetaData.cpp b/Engine/Source/Runtime/SlateCore/Private/Types/SlateAttributeMetaData.cpp
--- a/Engine/Source/Runtime/SlateCore/Private/Types/SlateAttributeMetaData.cpp
+++ b/Engine/Source/Runtime/SlateCore/Private/Types/SlateAttributeMetaData.cpp
@@ -60,7 +60,7 @@ void FSlateAttributeMetaData::RegisterAttribute(SWidget& OwningWidget, FSlateAtt
 		NewAttributeMetaData->RegisterAttributeImpl(OwningWidget, Attribute, AttributeType, MoveTemp(Wrapper));
 		OwningWidget.bHasRegisteredSlateAttribute = true;
 		OwningWidget.MetaData.Insert(NewAttributeMetaData, 0);
-		if (OwningWidget.IsConstructionCompleted() && OwningWidget.IsAttributesUpdatesEnabled())
+		if (OwningWidget.IsConstructed() && OwningWidget.IsAttributesUpdatesEnabled())
 		{
 			OwningWidget.Invalidate(EInvalidateWidgetReason::AttributeRegistration);
 		}
@@ -145,7 +145,7 @@ bool FSlateAttributeMetaData::UnregisterAttribute(SWidget& OwningWidget, const F
 			check(bResult); // if the num is 0 then we should have remove an item.
 			OwningWidget.bHasRegisteredSlateAttribute = false;
 			OwningWidget.MetaData.RemoveAtSwap(0);
-			if (OwningWidget.IsConstructionCompleted() && OwningWidget.IsAttributesUpdatesEnabled())
+			if (OwningWidget.IsConstructed() && OwningWidget.IsAttributesUpdatesEnabled())
 			{
 				OwningWidget.Invalidate(EInvalidateWidgetReason::AttributeRegistration);
 			}
@@ -221,7 +221,7 @@ void FSlateAttributeMetaData::InvalidateWidget(SWidget& OwningWidget, const FSla
 	//N.B. no needs to set the bUpatedManually in this case because
 	//	1. they are in construction, so they will all be called anyway
 	//	2. they are in WidgetList, so the SlateAttribute.Set will not be called
-	if (!OwningWidget.IsConstructionCompleted())
+	if (!OwningWidget.IsConstructed())
 	{
 		return;
 	}
@@ -369,7 +369,7 @@ void FSlateAttributeMetaData::UpdateChildrenOnlyVisibilityAttributes(SWidget& Ow
 
 void FSlateAttributeMetaData::UpdateAttributesImpl(SWidget& OwningWidget, EInvalidationPermission InvalidationStyle, int32 StartIndex, int32 IndexNum)
 {
-	const bool bInvalidateIfNeeded = (InvalidationStyle == EInvalidationPermission::AllowInvalidation) || (InvalidationStyle == EInvalidationPermission::AllowInvalidationIfConstructed && OwningWidget.IsConstructionCompleted());
+	const bool bInvalidateIfNeeded = (InvalidationStyle == EInvalidationPermission::AllowInvalidation) || (InvalidationStyle == EInvalidationPermission::AllowInvalidationIfConstructed && OwningWidget.IsConstructed());
 	const bool bAllowInvalidation = bInvalidateIfNeeded || InvalidationStyle == EInvalidationPermission::DelayInvalidation;
 	EInvalidateWidgetReason InvalidationReason = EInvalidateWidgetReason::None;
 	for (int32 Index = StartIndex; Index < IndexNum; ++Index)
@@ -441,7 +441,7 @@ void FSlateAttributeMetaData::UpdateAttribute(SWidget& OwningWidget, FSlateAttri
 			ISlateAttributeGetter::FUpdateAttributeResult Result = GetterItem.Getter->UpdateAttribute(OwningWidget);
 			if (Result.bInvalidationRequested)
 			{
-				if (OwningWidget.IsConstructionCompleted())
+				if (OwningWidget.IsConstructed())
 				{
 					EInvalidateWidgetReason Reason = GetterItem.GetInvalidationDetail(OwningWidget, Result.InvalidationReason);
 					OwningWidget.Invalidate(Reason | AttributeMetaData->CachedInvalidationReason);
diff --git a/Engine/Source/Runtime/SlateCore/Private/Widgets/SNullWidget.cpp b/Engine/Source/Runtime/SlateCore/Private/Widgets/SNullWidget.cpp
--- a/Engine/Source/Runtime/SlateCore/Private/Widgets/SNullWidget.cpp
+++ b/Engine/Source/Runtime/SlateCore/Private/Widgets/SNullWidget.cpp
@@ -33,7 +33,7 @@ public:
 private:
 	virtual void SetVisibility( TAttribute<EVisibility> InVisibility ) override final
 	{
-		ensureMsgf(!IsConstructionCompleted(), TEXT("Attempting to SetVisibility() on SNullWidget. Mutating SNullWidget is not allowed.") );
+		ensureMsgf(!IsConstructed(), TEXT("Attempting to SetVisibility() on SNullWidget. Mutating SNullWidget is not allowed.") );
 	}
 public:
 	
diff --git a/Engine/Source/Runtime/SlateCore/Private/Widgets/SWidget.cpp b/Engine/Source/Runtime/SlateCore/Private/Widgets/SWidget.cpp
--- a/Engine/Source/Runtime/SlateCore/Private/Widgets/SWidget.cpp
+++ b/Engine/Source/Runtime/SlateCore/Private/Widgets/SWidget.cpp
@@ -1204,7 +1204,7 @@ void SWidget::Invalidate(EInvalidateWidgetReason InvalidateReason)
 {
 	SLATE_CROSS_THREAD_CHECK();
 
-	if (InvalidateReason == EInvalidateWidgetReason::None || !IsConstructionCompleted())
+	if (InvalidateReason == EInvalidateWidgetReason::None || !IsConstructed())
 	{
 		return;
 	}
diff --git a/Engine/Source/Runtime/SlateCore/Public/Widgets/SWidget.h b/Engine/Source/Runtime/SlateCore/Public/Widgets/SWidget.h
--- a/Engine/Source/Runtime/SlateCore/Public/Widgets/SWidget.h
+++ b/Engine/Source/Runtime/SlateCore/Public/Widgets/SWidget.h
@@ -1525,8 +1525,8 @@ protected:
 	/** Construct a SWidget based on initial parameters. */
 	void SWidgetConstruct(const FSlateBaseNamedArgs& Args);
 
-	/** Is the widget construction is completed (did we called completed the Construct() function) */
-	bool IsConstructionCompleted() const { return bIsDeclarativeSyntaxConstructionCompleted; }
+	/** Is the widget construction completed (did we called and returned from the Construct() function) */
+	bool IsConstructed() const { return bIsDeclarativeSyntaxConstructionCompleted; }
 
 	/** 
 	 * Find the geometry of a descendant widget. This method assumes that WidgetsToFind are a descendants of this widget.
