Fix an issue with the deprecated event macro for _Lambda expansion.

diff --git a/Engine/Source/Runtime/SlateCore/Public/Widgets/DeclarativeSyntaxSupport.h b/Engine/Source/Runtime/SlateCore/Public/Widgets/DeclarativeSyntaxSupport.h
--- a/Engine/Source/Runtime/SlateCore/Public/Widgets/DeclarativeSyntaxSupport.h
+++ b/Engine/Source/Runtime/SlateCore/Public/Widgets/DeclarativeSyntaxSupport.h
@@ -519,11 +519,8 @@ struct NamedSlotProperty
 		UE_DEPRECATED(DeprecationVersion, DeprecationMessage)	\
 		WidgetArgsType& EventName##_Lambda(FunctorType&& InFunctor, VarTypes... Vars) \
 		{	\
-			if (InFunctor)	\
-			{	\
-				DelegateName DeprecatedDelegate = DelegateName::CreateLambda(InFunctor, Vars...);	\
-				_##UpgradedEventName = UpgradeFuncName(DeprecatedDelegate);	\
-			}	\
+			DelegateName DeprecatedDelegate = DelegateName::CreateLambda(Forward<FunctorType>(InFunctor), Vars...);	\
+			_##UpgradedEventName = UpgradeFuncName(DeprecatedDelegate);	\
 			return *this;	\
 		}	\
 		\
