Slate: Fix memory access write violation caused by 16276861 where WidgetsNeedingPrepassUpdate could be writing to an already reset FastWidgetPathList.
Reset all arrays in FSlateInvalidationRoot::InvalidateRootChildOrder (code by [REDACTED_PERSON])

#rnx

diff --git a/Engine/Source/Runtime/SlateCore/Private/FastUpdate/SlateInvalidationRoot.cpp b/Engine/Source/Runtime/SlateCore/Private/FastUpdate/SlateInvalidationRoot.cpp
--- a/Engine/Source/Runtime/SlateCore/Private/FastUpdate/SlateInvalidationRoot.cpp
+++ b/Engine/Source/Runtime/SlateCore/Private/FastUpdate/SlateInvalidationRoot.cpp
@@ -243,6 +243,9 @@ void FSlateInvalidationRoot::InvalidateRootChildOrder(const SWidget* Investigato
 {
 	// Invalidate all proxy handles
 	FastWidgetPathList->Reset();
+	WidgetsNeedingPreUpdate->Reset(false);
+	WidgetsNeedingPrepassUpdate->Reset(false);
+	WidgetsNeedingPostUpdate->Reset(false);
 	InvalidationRootWidget->Invalidate(EInvalidateWidgetReason::Prepass);
 	bNeedsSlowPath = true;
 
