Slate: Change SWidget::IsHovered and IsDirectlyHovered to not be virtual anymore. Use SWidget::IsHoverable or SWidget::SetHover. This allow us to reduce the number of attribute on SButton.

diff --git a/Engine/Plugins/Enterprise/DataprepEditor/Source/DataprepEditor/Private/Widgets/DataprepGraph/SDataprepGraphActionNode.cpp b/Engine/Plugins/Enterprise/DataprepEditor/Source/DataprepEditor/Private/Widgets/DataprepGraph/SDataprepGraphActionNode.cpp
--- a/Engine/Plugins/Enterprise/DataprepEditor/Source/DataprepEditor/Private/Widgets/DataprepGraph/SDataprepGraphActionNode.cpp
+++ b/Engine/Plugins/Enterprise/DataprepEditor/Source/DataprepEditor/Private/Widgets/DataprepGraph/SDataprepGraphActionNode.cpp
@@ -157,7 +157,6 @@ public:
 	{
 		ParentPtr = InParent;
 		StepIndex = InArgs._StepIndex.Get(InParent->GetDataprepAction()->GetStepsCount());
-		bIsHovered = false;
 
 		const float InterStepSpacing = 5.f;
 		const bool bBottomSlot = StepIndex == InParent->GetDataprepAction()->GetStepsCount();
diff --git a/Engine/Plugins/Experimental/CommonUI/Source/CommonUI/Private/CommonButtonTypes.cpp b/Engine/Plugins/Experimental/CommonUI/Source/CommonUI/Private/CommonButtonTypes.cpp
--- a/Engine/Plugins/Experimental/CommonUI/Source/CommonUI/Private/CommonButtonTypes.cpp
+++ b/Engine/Plugins/Experimental/CommonUI/Source/CommonUI/Private/CommonButtonTypes.cpp
@@ -55,6 +55,8 @@ void SCommonButton::OnMouseEnter(const FGeometry& MyGeometry, const FPointerEven
 {
 	if (!MouseEvent.IsTouchEvent())
 	{
+		bHovered = true;
+		SetHover(bHovered && bIsInteractionEnabled);
 		SButton::OnMouseEnter(MyGeometry, MouseEvent);
 	}
 }
@@ -70,6 +72,8 @@ void SCommonButton::OnMouseLeave(const FPointerEvent& MouseEvent)
 	}
 	else
 	{
+		bHovered = false;
+		SetHover(false);
 		SButton::OnMouseLeave(MouseEvent);
 	}
 }
@@ -112,11 +116,6 @@ FReply SCommonButton::OnKeyUp(const FGeometry& MyGeometry, const FKeyEvent& InKe
 	return SButton::OnKeyUp(MyGeometry, InKeyEvent);
 }
 
-bool SCommonButton::IsHovered() const
-{
-	return bIsInteractionEnabled ? SButton::IsHovered() : false;
-}
-
 bool SCommonButton::IsPressed() const
 {
 	return IsInteractable() ? SButton::IsPressed() : false;
@@ -143,10 +142,11 @@ void SCommonButton::SetIsInteractionEnabled(bool bInIsInteractionEnabled)
 
 	bIsInteractionEnabled = bInIsInteractionEnabled;
 
-	// If the hover state changed due to an interactability change, trigger external logic accordingly.
-	const bool bIsHoveredNow = IsHovered();
+	// If the hover state changed due to an intractability change, trigger external logic accordingly.
+	const bool bIsHoveredNow = bHovered && bInIsInteractionEnabled;
 	if (bWasHovered != bIsHoveredNow)
 	{
+		SetHover(bIsHoveredNow);
 		if (bIsHoveredNow)
 		{
 			OnHovered.ExecuteIfBound();
diff --git a/Engine/Plugins/Experimental/CommonUI/Source/CommonUI/Private/CommonButtonTypes.h b/Engine/Plugins/Experimental/CommonUI/Source/CommonUI/Private/CommonButtonTypes.h
--- a/Engine/Plugins/Experimental/CommonUI/Source/CommonUI/Private/CommonButtonTypes.h
+++ b/Engine/Plugins/Experimental/CommonUI/Source/CommonUI/Private/CommonButtonTypes.h
@@ -44,7 +44,7 @@ public:
 		SLATE_ARGUMENT(bool, IsInteractionEnabled)
 		SLATE_END_ARGS()
 
-		void Construct(const FArguments& InArgs)
+	void Construct(const FArguments& InArgs)
 	{
 		OnDoubleClicked = InArgs._OnDoubleClicked;
 
@@ -65,10 +65,13 @@ public:
 			]);
 
 		SetCanTick(false);
+		// Set the hover state to indicate that we want to override the default behavior
+		SetHover(false);
 
 		OnReceivedFocus = InArgs._OnReceivedFocus;
 		bIsButtonEnabled = InArgs._IsButtonEnabled;
 		bIsInteractionEnabled = InArgs._IsInteractionEnabled;
+		bHovered = false;
 	}
 
 	virtual FReply OnMouseButtonDown(const FGeometry& MyGeometry, const FPointerEvent& MouseEvent) override;
@@ -87,8 +90,6 @@ public:
 
 	virtual FReply OnKeyUp(const FGeometry& MyGeometry, const FKeyEvent& InKeyEvent) override;
 
-	virtual bool IsHovered() const override;
-
 	virtual bool IsPressed() const override;
 
 	void SetIsButtonEnabled(bool bInIsButtonEnabled);
@@ -115,4 +116,7 @@ private:
 
 	/** True if clicking is enabled, to allow for things like double click */
 	bool bIsInteractionEnabled;
+
+	/** True if mouse over the widget */
+	bool bHovered;
 };
\ No newline at end of file
diff --git a/Engine/Plugins/Runtime/nDisplay/Source/DisplayClusterConfigurator/Private/Views/TreeViews/SDisplayClusterConfiguratorTreeItemRow.cpp b/Engine/Plugins/Runtime/nDisplay/Source/DisplayClusterConfigurator/Private/Views/TreeViews/SDisplayClusterConfiguratorTreeItemRow.cpp
--- a/Engine/Plugins/Runtime/nDisplay/Source/DisplayClusterConfigurator/Private/Views/TreeViews/SDisplayClusterConfiguratorTreeItemRow.cpp
+++ b/Engine/Plugins/Runtime/nDisplay/Source/DisplayClusterConfigurator/Private/Views/TreeViews/SDisplayClusterConfiguratorTreeItemRow.cpp
@@ -19,6 +19,8 @@ void SDisplayClusterConfiguratorTreeItemRow::Construct(const FArguments& InArgs,
 	Item = InArgs._Item;
 	FilterText = InArgs._FilterText;
 
+	SetHover(TAttribute<bool>::CreateSP(this, &SDisplayClusterConfiguratorTreeItemRow::ShouldAppearHovered));
+
 	check(Item.IsValid());
 
 	SMultiColumnTableRow< TSharedPtr<IDisplayClusterConfiguratorTreeItem> >::Construct(FSuperRowType::FArguments()
@@ -56,7 +58,7 @@ void SDisplayClusterConfiguratorTreeItemRow::ConstructChildren(ETableViewMode::T
 		];
 }
 
-bool SDisplayClusterConfiguratorTreeItemRow::IsHovered() const
+bool SDisplayClusterConfiguratorTreeItemRow::ShouldAppearHovered() const
 {
 	return Item.Pin()->IsHovered();
 }
diff --git a/Engine/Plugins/Runtime/nDisplay/Source/DisplayClusterConfigurator/Private/Views/TreeViews/SDisplayClusterConfiguratorTreeItemRow.h b/Engine/Plugins/Runtime/nDisplay/Source/DisplayClusterConfigurator/Private/Views/TreeViews/SDisplayClusterConfiguratorTreeItemRow.h
--- a/Engine/Plugins/Runtime/nDisplay/Source/DisplayClusterConfigurator/Private/Views/TreeViews/SDisplayClusterConfiguratorTreeItemRow.h
+++ b/Engine/Plugins/Runtime/nDisplay/Source/DisplayClusterConfigurator/Private/Views/TreeViews/SDisplayClusterConfiguratorTreeItemRow.h
@@ -39,12 +39,13 @@ public:
 	//~ End SWidget interface
 
 	// Begin SWidget
-	virtual bool IsHovered() const override;
 	virtual void OnMouseEnter(const FGeometry& MyGeometry, const FPointerEvent& MouseEvent) override;
 	virtual void OnMouseLeave(const FPointerEvent& MouseEvent) override;
 	// End SWidget
 
 private:
+	bool ShouldAppearHovered() const;
+
 	FReply HandleDragDetected(const FGeometry& MyGeometry, const FPointerEvent& MouseEvent);
 
 	void HandleDragEnter(const FDragDropEvent& DragDropEvent);
diff --git a/Engine/Source/Editor/GraphEditor/Private/SGraphPin.cpp b/Engine/Source/Editor/GraphEditor/Private/SGraphPin.cpp
--- a/Engine/Source/Editor/GraphEditor/Private/SGraphPin.cpp
+++ b/Engine/Source/Editor/GraphEditor/Private/SGraphPin.cpp
@@ -568,7 +568,7 @@ FReply SGraphPin::OnMouseButtonUp( const FGeometry& MyGeometry, const FPointerEv
 
 void SGraphPin::OnMouseEnter( const FGeometry& MyGeometry, const FPointerEvent& MouseEvent )
 {
-	if (!bIsHovered && ensure(!bGraphDataInvalid))
+	if (!IsHovered() && ensure(!bGraphDataInvalid))
 	{
 		UEdGraphPin* MyPin = GetPinObj();
 		if (MyPin && !MyPin->IsPendingKill() && MyPin->GetOuter() && MyPin->GetOuter()->IsA(UEdGraphNode::StaticClass()))
diff --git a/Engine/Source/Editor/KismetWidgets/Private/SPinTypeSelector.cpp b/Engine/Source/Editor/KismetWidgets/Private/SPinTypeSelector.cpp
--- a/Engine/Source/Editor/KismetWidgets/Private/SPinTypeSelector.cpp
+++ b/Engine/Source/Editor/KismetWidgets/Private/SPinTypeSelector.cpp
@@ -130,6 +130,9 @@ public:
 public:
 	void Construct(const FArguments& InArgs, const TSharedRef<STableViewBase>& InOwnerTable, TWeakPtr<SMenuOwner> InMenuOwner)
 	{
+		bHovered = false;
+		SetHover(TAttribute<bool>::CreateSP(this, &SPinTypeRow::ShouldAppearHovered));
+
 		SComboRow<FPinTypeTreeItem>::Construct( SComboRow<FPinTypeTreeItem>::FArguments()
 			.ToolTip(InArgs._ToolTip)
 			[
@@ -143,12 +146,24 @@ public:
 			InOwnerTable);
 	}
 
-	// SWidget interface
-	virtual bool IsHovered() const override
+	virtual void OnMouseEnter(const FGeometry& MyGeometry, const FPointerEvent& MouseEvent) override
 	{
-		return SComboRow<FPinTypeTreeItem>::IsHovered() || SubMenuHandler.Pin()->ShouldSubMenuAppearHovered();
+		bHovered = true;
+		SComboRow<FPinTypeTreeItem>::OnMouseEnter(MyGeometry, MouseEvent);
+	}
+
+	virtual void OnMouseLeave(const FPointerEvent& MouseEvent) override
+	{
+		bHovered = false;
+		SComboRow<FPinTypeTreeItem>::OnMouseLeave(MouseEvent);
+	}
+	//~ End of SWidget interface
+
+	/** Retuns TRUE if it should appear as hovered. */
+	bool ShouldAppearHovered() const
+	{
+		return bHovered || SubMenuHandler.Pin()->ShouldSubMenuAppearHovered();
 	}
-	// End of SWidget interface
 
 	/** Returns TRUE if there is a Sub-Menu available to open */
 	bool HasSubMenu() const
@@ -171,6 +186,8 @@ public:
 private:
 	/** The Sub-MenuHandler which is managing the sub-menu content so that mousing over other rows will not close the sub-menus immediately */
 	TWeakPtr<SSubMenuHandler> SubMenuHandler;
+	/** Keep an internal IsHovered flag*/
+	bool bHovered;
 };
 
 static bool ContainerRequiresGetTypeHash(EPinContainerType InType)
diff --git a/Engine/Source/Editor/Persona/Private/AnimTimeline/SAnimOutlinerItem.cpp b/Engine/Source/Editor/Persona/Private/AnimTimeline/SAnimOutlinerItem.cpp
--- a/Engine/Source/Editor/Persona/Private/AnimTimeline/SAnimOutlinerItem.cpp
+++ b/Engine/Source/Editor/Persona/Private/AnimTimeline/SAnimOutlinerItem.cpp
@@ -23,6 +23,9 @@ void SAnimOutlinerItem::Construct(const FArguments& InArgs, const TSharedRef<STa
 	OnGenerateWidgetForColumn = InArgs._OnGenerateWidgetForColumn;
 	HighlightText = InArgs._HighlightText;
 
+	bHovered = false;
+	SetHover(TAttribute<bool>::CreateSP(this, &SAnimOutlinerItem::ShouldAppearHovered));
+
 	SMultiColumnTableRow::Construct(
 		SMultiColumnTableRow::FArguments()
 			.ShowSelection(true),
@@ -73,6 +76,7 @@ void SAnimOutlinerItem::AddTrackAreaReference(const TSharedPtr<SAnimTrack>& InTr
 
 void SAnimOutlinerItem::OnMouseEnter(const FGeometry& MyGeometry, const FPointerEvent& MouseEvent)
 {
+	bHovered = true;
 	SMultiColumnTableRow<TSharedRef<FAnimTimelineTrack>>::OnMouseEnter(MyGeometry, MouseEvent);
 
 	TSharedPtr<FAnimTimelineTrack> PinnedTrack = Track.Pin();
@@ -84,6 +88,7 @@ void SAnimOutlinerItem::OnMouseEnter(const FGeometry& MyGeometry, const FPointer
 
 void SAnimOutlinerItem::OnMouseLeave(const FPointerEvent& MouseEvent)
 {
+	bHovered = false;
 	SMultiColumnTableRow<TSharedRef<FAnimTimelineTrack>>::OnMouseLeave(MouseEvent);
 
 	TSharedPtr<FAnimTimelineTrack> PinnedTrack = Track.Pin();
@@ -93,13 +98,12 @@ void SAnimOutlinerItem::OnMouseLeave(const FPointerEvent& MouseEvent)
 	}
 }
 
-bool SAnimOutlinerItem::IsHovered() const
+bool SAnimOutlinerItem::ShouldAppearHovered() const
 {
-	TSharedPtr<FAnimTimelineTrack> PinnedTrack = Track.Pin();
-	if(PinnedTrack.IsValid())
+	if(TSharedPtr<FAnimTimelineTrack> PinnedTrack = Track.Pin())
 	{
-		return SMultiColumnTableRow<TSharedRef<FAnimTimelineTrack>>::IsHovered() || PinnedTrack->IsHovered();
+		return bHovered || PinnedTrack->IsHovered();
 	}
 
-	return SMultiColumnTableRow<TSharedRef<FAnimTimelineTrack>>::IsHovered();
+	return bHovered;
 }
\ No newline at end of file
diff --git a/Engine/Source/Editor/Persona/Private/AnimTimeline/SAnimOutlinerItem.h b/Engine/Source/Editor/Persona/Private/AnimTimeline/SAnimOutlinerItem.h
--- a/Engine/Source/Editor/Persona/Private/AnimTimeline/SAnimOutlinerItem.h
+++ b/Engine/Source/Editor/Persona/Private/AnimTimeline/SAnimOutlinerItem.h
@@ -36,7 +36,6 @@ public:
 	virtual FVector2D ComputeDesiredSize(float LayoutScaleMultiplier) const override;
 	virtual void OnMouseEnter(const FGeometry& MyGeometry, const FPointerEvent& MouseEvent) override;
 	virtual void OnMouseLeave(const FPointerEvent& MouseEvent) override;
-	virtual bool IsHovered() const override;
 
 	/** SMultiColumnTableRow interface */
 	virtual TSharedRef<SWidget> GenerateWidgetForColumn(const FName& ColumnId) override;
@@ -54,6 +53,9 @@ private:
 	/** Get the label text for this item */
 	FText GetLabelText() const;
 
+	/** Get the hover state for this item */
+	bool ShouldAppearHovered() const;
+
 private:
 	/** The track that we represent */
 	TWeakPtr<FAnimTimelineTrack> Track;
@@ -66,4 +68,7 @@ private:
 
 	/** Text to highlight when searching */
 	TAttribute<FText> HighlightText;
+	
+	/** Keep an internal IsHovered flag*/
+	bool bHovered;
 };
\ No newline at end of file
diff --git a/Engine/Source/Editor/PropertyEditor/Private/DetailCategoryGroupNode.cpp b/Engine/Source/Editor/PropertyEditor/Private/DetailCategoryGroupNode.cpp
--- a/Engine/Source/Editor/PropertyEditor/Private/DetailCategoryGroupNode.cpp
+++ b/Engine/Source/Editor/PropertyEditor/Private/DetailCategoryGroupNode.cpp
@@ -157,7 +157,7 @@ FSlateColor SDetailCategoryTableRow::GetInnerBackgroundColor() const
 
 FSlateColor SDetailCategoryTableRow::GetOuterBackgroundColor() const
 {
-	if (bIsHovered)
+	if (IsHovered())
 	{
 		return FAppStyle::Get().GetSlateColor("Colors.Header");
 	}
diff --git a/Engine/Source/Editor/Sequencer/Private/LevelEditorSequencerIntegration.cpp b/Engine/Source/Editor/Sequencer/Private/LevelEditorSequencerIntegration.cpp
--- a/Engine/Source/Editor/Sequencer/Private/LevelEditorSequencerIntegration.cpp
+++ b/Engine/Source/Editor/Sequencer/Private/LevelEditorSequencerIntegration.cpp
@@ -998,6 +998,8 @@ public:
 		FadeOutSequence = FCurveSequence(0.0f, FaderConstants::FadeTime);
 		FadeOutSequence.JumpToEnd();
 
+		SetHover(false);
+
 		SBorder::Construct(SBorder::FArguments()
 			.BorderImage(FCoreStyle::Get().GetBrush("NoBorder"))
 			.Padding(0.0f)
@@ -1013,7 +1015,7 @@ public:
 	{
 		FLinearColor Color = FLinearColor::White;
 	
-		if(FadeOutSequence.IsPlaying() || !bIsHovered)
+		if(FadeOutSequence.IsPlaying() || !IsHovered())
 		{
 			Color.A = FMath::Lerp(FaderConstants::HoveredOpacity, FaderConstants::NonHoveredOpacity, FadeOutSequence.GetLerp());
 		}
@@ -1029,7 +1031,7 @@ public:
 	{
 		if(!FSlateApplication::Get().IsUsingHighPrecisionMouseMovment())
 		{
-			bIsHovered = true;
+			SetHover(true);
 			if(FadeOutSequence.IsPlaying())
 			{
 				// Fade out is already playing so just force the fade in curve to the end so we don't have a "pop" 
@@ -1047,7 +1049,7 @@ public:
 	{
 		if(!FSlateApplication::Get().IsUsingHighPrecisionMouseMovment())
 		{
-			bIsHovered = false;
+			SetHover(false);
 			FadeOutSequence.Play(AsShared());
 		}
 	}
diff --git a/Engine/Source/Editor/UMGEditor/Private/Hierarchy/SHierarchyViewItem.cpp b/Engine/Source/Editor/UMGEditor/Private/Hierarchy/SHierarchyViewItem.cpp
--- a/Engine/Source/Editor/UMGEditor/Private/Hierarchy/SHierarchyViewItem.cpp
+++ b/Engine/Source/Editor/UMGEditor/Private/Hierarchy/SHierarchyViewItem.cpp
@@ -1309,9 +1309,12 @@ void FHierarchyWidget::OnEndEditing()
 
 void SHierarchyViewItem::Construct(const FArguments& InArgs, const TSharedRef< STableViewBase >& InOwnerTableView, TSharedPtr<FHierarchyModel> InModel)
 {
+	bHovered = false;
 	Model = InModel;
 	Model->RenameEvent.BindSP(this, &SHierarchyViewItem::OnRequestBeginRename);
 
+	SetHover(TAttribute<bool>::CreateSP(this, &SHierarchyViewItem::ShouldAppearHovered));
+
 	STableRow< TSharedPtr<FHierarchyModel> >::Construct(
 		STableRow< TSharedPtr<FHierarchyModel> >::FArguments()
 		.OnCanAcceptDrop(this, &SHierarchyViewItem::HandleCanAcceptDrop)
@@ -1439,6 +1442,7 @@ SHierarchyViewItem::~SHierarchyViewItem()
 
 void SHierarchyViewItem::OnMouseEnter(const FGeometry& MyGeometry, const FPointerEvent& MouseEvent)
 {
+	bHovered = true;
 	STableRow< TSharedPtr<FHierarchyModel> >::OnMouseEnter(MyGeometry, MouseEvent);
 
 	Model->OnMouseEnter();
@@ -1446,6 +1450,7 @@ void SHierarchyViewItem::OnMouseEnter(const FGeometry& MyGeometry, const FPointe
 
 void SHierarchyViewItem::OnMouseLeave(const FPointerEvent& MouseEvent)
 {
+	bHovered = false;
 	STableRow< TSharedPtr<FHierarchyModel> >::OnMouseLeave(MouseEvent);
 
 	Model->OnMouseLeave();
@@ -1505,9 +1510,9 @@ FText SHierarchyViewItem::GetItemText() const
 	return Model->GetText();
 }
 
-bool SHierarchyViewItem::IsHovered() const
+bool SHierarchyViewItem::ShouldAppearHovered() const
 {
-	return bIsHovered || Model->IsHovered();
+	return bHovered || Model->IsHovered();
 }
 
 void SHierarchyViewItem::HandleDragEnter(FDragDropEvent const& DragDropEvent)
diff --git a/Engine/Source/Editor/UMGEditor/Private/Hierarchy/SHierarchyViewItem.h b/Engine/Source/Editor/UMGEditor/Private/Hierarchy/SHierarchyViewItem.h
--- a/Engine/Source/Editor/UMGEditor/Private/Hierarchy/SHierarchyViewItem.h
+++ b/Engine/Source/Editor/UMGEditor/Private/Hierarchy/SHierarchyViewItem.h
@@ -350,7 +350,6 @@ public:
 	virtual ~SHierarchyViewItem();
 
 	// Begin SWidget
-	virtual bool IsHovered() const override;
 	virtual void OnMouseEnter(const FGeometry& MyGeometry, const FPointerEvent& MouseEvent) override;
 	virtual void OnMouseLeave(const FPointerEvent& MouseEvent) override;
 	// End SWidget
@@ -376,6 +375,7 @@ private:
 
 	bool IsReadOnly() const;
 	void OnRequestBeginRename();
+	bool ShouldAppearHovered() const;
 
 	/** Gets the font to use for the text item, bold for customized named items */
 	FSlateFontInfo GetItemFont() const;
@@ -405,4 +405,7 @@ private:
 
 	/** Text when we start editing. */
 	FText InitialText;
+
+	/** Keep an internal IsHovered flag*/
+	bool bHovered;
 };
diff --git a/Engine/Source/Runtime/Slate/Private/Framework/MultiBox/SMenuEntryBlock.cpp b/Engine/Source/Runtime/Slate/Private/Framework/MultiBox/SMenuEntryBlock.cpp
--- a/Engine/Source/Runtime/Slate/Private/Framework/MultiBox/SMenuEntryBlock.cpp
+++ b/Engine/Source/Runtime/Slate/Private/Framework/MultiBox/SMenuEntryBlock.cpp
@@ -397,8 +397,6 @@ public:
 
 		FReply Reply = SButton::OnMouseButtonUp(MyGeometry, MouseEvent);
 
-		bIsHovered = true;
-
 		return Reply;
 	}
 
@@ -651,8 +649,11 @@ public:
 
 	void Construct( const FArguments& InArgs )
 	{
+		bHovered = false;
 		ShouldAppearHovered = InArgs._ShouldAppearHovered;
 
+		SetHover(TAttribute<bool>::CreateSP(this, &SSubMenuButton::HandleShouldAppearHovered));
+
 		SButton::FArguments ButtonArgs;
 		ButtonArgs.Text(InArgs._Label);
 		ButtonArgs.ForegroundColor( this, &SSubMenuButton::InvertOnHover );
@@ -674,13 +675,25 @@ public:
 		SButton::Construct( ButtonArgs );
 	}
 
-	virtual bool IsHovered() const override
+	virtual void OnMouseEnter(const FGeometry& MyGeometry, const FPointerEvent& MouseEvent) override
 	{
-		// Submenu widgets which have been opened should remain as if hovered, even if the cursor is outside them
-		return SWidget::IsHovered() || ShouldAppearHovered.Get();
+		bHovered = true;
+		SButton::OnMouseEnter(MyGeometry, MouseEvent);
+	}
+
+	virtual void OnMouseLeave(const FPointerEvent& MouseEvent) override
+	{
+		bHovered = false;
+		SButton::OnMouseLeave(MouseEvent);
 	}
 
 private:
+	bool HandleShouldAppearHovered() const
+	{
+		// Submenu widgets which have been opened should remain as if hovered, even if the cursor is outside them
+		return bHovered || ShouldAppearHovered.Get();
+	}
+
 	FSlateColor InvertOnHover() const
 	{
 		if ( this->IsHovered() )
@@ -696,6 +709,8 @@ private:
 private:
 	/** Attribute to indicate if the sub-menu is open or not */
 	TAttribute<bool> ShouldAppearHovered;
+	/** Keep an internal IsHovered flag*/
+	bool bHovered;
 };
 
 
diff --git a/Engine/Source/Runtime/Slate/Private/Widgets/Input/SCheckBox.cpp b/Engine/Source/Runtime/Slate/Private/Widgets/Input/SCheckBox.cpp
--- a/Engine/Source/Runtime/Slate/Private/Widgets/Input/SCheckBox.cpp
+++ b/Engine/Source/Runtime/Slate/Private/Widgets/Input/SCheckBox.cpp
@@ -301,10 +301,10 @@ FSlateColor SCheckBox::GetForegroundColor() const
 		switch (State)
 		{
 		case ECheckBoxState::Unchecked:
-			return bIsPressed ? Style->PressedForeground : bIsHovered ? Style->HoveredForeground : Style->ForegroundColor;
+			return bIsPressed ? Style->PressedForeground : IsHovered() ? Style->HoveredForeground : Style->ForegroundColor;
 			break;
 		case ECheckBoxState::Checked:
-			return bIsPressed ? Style->CheckedPressedForeground : bIsHovered ? Style->CheckedHoveredForeground : Style->CheckedForeground;
+			return bIsPressed ? Style->CheckedPressedForeground : IsHovered() ? Style->CheckedHoveredForeground : Style->CheckedForeground;
 			break;
 		default:
 		case ECheckBoxState::Undetermined:
diff --git a/Engine/Source/Runtime/Slate/Public/Widgets/Input/SRichTextHyperlink.h b/Engine/Source/Runtime/Slate/Public/Widgets/Input/SRichTextHyperlink.h
--- a/Engine/Source/Runtime/Slate/Public/Widgets/Input/SRichTextHyperlink.h
+++ b/Engine/Source/Runtime/Slate/Public/Widgets/Input/SRichTextHyperlink.h
@@ -40,6 +40,8 @@ public:
 	{
 		ViewModel = InViewModel;
 
+		SetHover(TAttribute<bool>::CreateSP(this, &SRichTextHyperlink::ShouldAppearHovered));
+
 		SHyperlink::Construct(
 			SHyperlink::FArguments()
 			.Text( InArgs._Text )
@@ -79,16 +81,16 @@ public:
 		return Reply;
 	}
 
-	virtual bool IsHovered() const override
-	{
-		return ViewModel->IsHovered();
-	}
-
 	virtual bool IsPressed() const override
 	{
 		return ViewModel->IsPressed();
 	}
 
+private:
+	bool ShouldAppearHovered() const
+	{
+		return ViewModel->IsHovered();
+	}
 
 private:
 
diff --git a/Engine/Source/Runtime/SlateCore/Private/Widgets/SWidget.cpp b/Engine/Source/Runtime/SlateCore/Private/Widgets/SWidget.cpp
--- a/Engine/Source/Runtime/SlateCore/Private/Widgets/SWidget.cpp
+++ b/Engine/Source/Runtime/SlateCore/Private/Widgets/SWidget.cpp
@@ -174,15 +174,19 @@ void SWidget::PrivateRegisterAttributes(FSlateAttributeInitializer& AttributeIni
 	SLATE_ADD_MEMBER_ATTRIBUTE_DEFINITION_WITH_NAME(AttributeInitializer, "Visibility", VisibilityAttribute, EInvalidateWidgetReason::Visibility)
 		.AffectVisibility();
 	SLATE_ADD_MEMBER_ATTRIBUTE_DEFINITION_WITH_NAME(AttributeInitializer, "EnabledState", EnabledStateAttribute, EInvalidateWidgetReason::Paint);
+	SLATE_ADD_MEMBER_ATTRIBUTE_DEFINITION_WITH_NAME(AttributeInitializer, "Hovered", HoveredAttribute, EInvalidateWidgetReason::Paint);
 	SLATE_ADD_MEMBER_ATTRIBUTE_DEFINITION_WITH_NAME(AttributeInitializer, "RenderTransform", RenderTransformAttribute, EInvalidateWidgetReason::Layout | EInvalidateWidgetReason::RenderTransform);
 	SLATE_ADD_MEMBER_ATTRIBUTE_DEFINITION_WITH_NAME(AttributeInitializer, "RenderTransformPivot", RenderTransformPivotAttribute, EInvalidateWidgetReason::Layout | EInvalidateWidgetReason::RenderTransform);
 }
 
+PRAGMA_DISABLE_DEPRECATION_WARNINGS
 SWidget::SWidget()
-	: bIsHovered(false)
-	, bCanSupportFocus(true)
+	: bCanSupportFocus(true)
 	, bCanHaveChildren(true)
 	, bClippingProxy(false)
+#if WITH_EDITORONLY_DATA
+	, bIsHovered(false)
+#endif
 	, bToolTipForceFieldEnabled(false)
 	, bForceVolatile(false)
 	, bCachedVolatile(false)
@@ -191,6 +195,7 @@ SWidget::SWidget()
 	, bHasRegisteredSlateAttribute(false)
 	, bEnabledAttributesUpdate(true)
 	, bIsDeclarativeSyntaxConstructionCompleted(false)
+	, bIsHoveredAttributeSet(false)
 	, bHasCustomPrepass(false)
 	, bHasRelativeLayoutScale(false)
 	, bVolatilityAlwaysInvalidatesPrepass(false)
@@ -206,6 +211,7 @@ SWidget::SWidget()
 	, DesiredSize()
 	, VisibilityAttribute(*this, EVisibility::Visible)
 	, EnabledStateAttribute(*this, true)
+	, HoveredAttribute(*this, false)
 	, RenderTransformPivotAttribute(*this, FVector2D::ZeroVector)
 	, RenderTransformAttribute(*this)
 	, CullingBoundsExtension()
@@ -226,6 +232,7 @@ SWidget::SWidget()
 	UE_SLATE_DEBUG_WIDGETLIST_ADD_WIDGET(this);
 	UE_TRACE_SLATE_WIDGET_ADDED(this);
 }
+PRAGMA_ENABLE_DEPRECATION_WARNINGS
 
 PRAGMA_DISABLE_DEPRECATION_WARNINGS
 SWidget::~SWidget()
@@ -494,7 +501,10 @@ FReply SWidget::OnMouseButtonDoubleClick(const FGeometry& MyGeometry, const FPoi
 
 void SWidget::OnMouseEnter( const FGeometry& MyGeometry, const FPointerEvent& MouseEvent )
 {
-	bIsHovered = true;
+	if (!bIsHoveredAttributeSet)
+	{
+		HoveredAttribute.Set(*this, true);
+	}
 
 	if (TSharedPtr<FSlateMouseEventsMetaData> Data = GetMetaData<FSlateMouseEventsMetaData>())
 	{
@@ -508,7 +518,10 @@ void SWidget::OnMouseEnter( const FGeometry& MyGeometry, const FPointerEvent& Mo
 
 void SWidget::OnMouseLeave( const FPointerEvent& MouseEvent )
 {
-	bIsHovered = false;
+	if (!bIsHoveredAttributeSet)
+	{
+		HoveredAttribute.Set(*this, false);
+	}
 
 	if (TSharedPtr<FSlateMouseEventsMetaData> Data = GetMetaData<FSlateMouseEventsMetaData>())
 	{
diff --git a/Engine/Source/Runtime/SlateCore/Public/Widgets/SWidget.h b/Engine/Source/Runtime/SlateCore/Public/Widgets/SWidget.h
--- a/Engine/Source/Runtime/SlateCore/Public/Widgets/SWidget.h
+++ b/Engine/Source/Runtime/SlateCore/Public/Widgets/SWidget.h
@@ -993,14 +993,31 @@ public:
 		return bToolTipForceFieldEnabled;
 	}
 
-	/** @return True if this widget hovered */
-	virtual bool IsHovered() const
+	/**
+	 * @return True if this widget hovered
+	 * @note IsHovered used to be virtual. Use SetHover to assign an attribute if you need to override the default behavior.
+	 */
+	bool IsHovered() const
 	{
-		return bIsHovered;
+		return HoveredAttribute.Get();
 	}
 
 	/** @return True if this widget is directly hovered */
-	virtual bool IsDirectlyHovered() const;
+	bool IsDirectlyHovered() const;
+
+protected:
+	/**
+	 * Set the hover state.
+	 * Once set, the attribute that the ownership and SWidget code will not update the attribute value.
+	 * You can return the control to the SWidget code by setting an empty TAttribute.
+	 */
+	void SetHover(TAttribute<bool> InHovered)
+	{
+		bIsHoveredAttributeSet = InHovered.IsSet();
+		HoveredAttribute.Assign(*this, MoveTemp(InHovered));
+	}
+
+public:
 
 	/**
 	 * @return is this widget visible, hidden or collapsed.
@@ -1687,6 +1704,8 @@ protected:
 
 	/** @return an attribute reference of EnabledStateAttribute */
 	TSlateAttributeRef<bool> GetEnabledStateAttribute() const { return TSlateAttributeRef<bool>(SharedThis(this), EnabledStateAttribute); }
+	/** @return an attribute reference of HoveredAttribute */
+	TSlateAttributeRef<bool> GetHoveredAttribute() const { return TSlateAttributeRef<bool>(SharedThis(this), HoveredAttribute); }
 	/** @return an attribute reference of VisibilityAttribute */
 	TSlateAttributeRef<EVisibility> GetVisibilityAttribute() const { return TSlateAttributeRef<EVisibility>(SharedThis(this), VisibilityAttribute); }
 	/** @return an attribute reference of RenderTransformAttribute */
@@ -1703,9 +1722,6 @@ private:
 	mutable FWidgetProxyHandle FastPathProxyHandle;
 
 protected:
-	/** Is this widget hovered? */
-	uint8 bIsHovered : 1;
-
 	/** Can the widget ever support keyboard focus */
 	uint8 bCanSupportFocus : 1;
 
@@ -1723,6 +1739,12 @@ protected:
 	  */
 	uint8 bClippingProxy : 1;
 
+#if WITH_EDITORONLY_DATA
+	/** Is this widget hovered? */
+	UE_DEPRECATED(5.0, "Direct access to bIsHovered is now deprecated. Use the IsHovered getter.")
+	uint8 bIsHovered : 1;
+#endif
+
 private:
 	/**
 	 * Whether this widget is a "tool tip force field".  That is, tool-tips should never spawn over the area
@@ -1751,6 +1773,9 @@ private:
 	/** The SNew or SAssignedNew construction is completed. */
 	uint8 bIsDeclarativeSyntaxConstructionCompleted : 1;
 
+	/** Is the attribute IsHovered is set? */
+	uint8 bIsHoveredAttributeSet : 1;
+
 protected:
 	uint8 bHasCustomPrepass : 1;
 
@@ -1810,6 +1835,9 @@ private:
 	/** Whether or not this widget is enabled */
 	TSlateAttribute<bool> EnabledStateAttribute;
 
+	/** Whether or not this widget is hovered */
+	TSlateAttribute<bool> HoveredAttribute;
+
 	/** Render transform pivot of this widget (in normalized local space) */
 	TSlateAttribute<FVector2D> RenderTransformPivotAttribute;
 
