Slate: Enforce that only some InvalidateWidgetReason are supported by SlateAttribute.

diff --git a/Engine/Source/Runtime/SlateCore/Private/Test/SlateAttributeTest.cpp b/Engine/Source/Runtime/SlateCore/Private/Test/SlateAttributeTest.cpp
--- a/Engine/Source/Runtime/SlateCore/Private/Test/SlateAttributeTest.cpp
+++ b/Engine/Source/Runtime/SlateCore/Private/Test/SlateAttributeTest.cpp
@@ -14,6 +14,7 @@
 
 #define LOCTEXT_NAMESPACE "Slate.Attribute"
 
+//IMPLEMENT_SIMPLE_AUTOMATION_TEST(FSlateAttributeTest, "Slate.Attribute", EAutomationTestFlags::ApplicationContextMask | EAutomationTestFlags::SmokeFilter)
 IMPLEMENT_SIMPLE_AUTOMATION_TEST(FSlateAttributeTest, "Slate.Attribute", EAutomationTestFlags::ApplicationContextMask | EAutomationTestFlags::EngineFilter)
 
 namespace UE
@@ -106,11 +107,11 @@ void SAttributeLeftWidget_Parent::PrivateRegisterAttributes(FSlateAttributeIniti
 	// The update order is B, A, D, C
 	// C updates when D is invalidated, so D needs to be before C
 	// A updates after B, so B needs to be before A
-	SLATE_ADD_MEMBER_ATTRIBUTE_DEFINITION(AttributeInitializer, IntAttributeD, EInvalidateWidgetReason::ChildOrder);
-	SLATE_ADD_MEMBER_ATTRIBUTE_DEFINITION(AttributeInitializer, IntAttributeC, EInvalidateWidgetReason::ChildOrder)
+	SLATE_ADD_MEMBER_ATTRIBUTE_DEFINITION(AttributeInitializer, IntAttributeD, EInvalidateWidgetReason::Layout);
+	SLATE_ADD_MEMBER_ATTRIBUTE_DEFINITION(AttributeInitializer, IntAttributeC, EInvalidateWidgetReason::Layout)
 		.UpdateDependency(GET_MEMBER_NAME_CHECKED(PrivateThisType, IntAttributeD));
-	SLATE_ADD_MEMBER_ATTRIBUTE_DEFINITION(AttributeInitializer, IntAttributeB, EInvalidateWidgetReason::ChildOrder);
-	SLATE_ADD_MEMBER_ATTRIBUTE_DEFINITION(AttributeInitializer, IntAttributeA, EInvalidateWidgetReason::ChildOrder)
+	SLATE_ADD_MEMBER_ATTRIBUTE_DEFINITION(AttributeInitializer, IntAttributeB, EInvalidateWidgetReason::Layout);
+	SLATE_ADD_MEMBER_ATTRIBUTE_DEFINITION(AttributeInitializer, IntAttributeA, EInvalidateWidgetReason::Layout)
 		.UpdatePrerequisite(GET_MEMBER_NAME_CHECKED(PrivateThisType, IntAttributeB));
 
 	AttributeInitializer.OverrideInvalidationReason(GET_MEMBER_NAME_CHECKED(PrivateThisType, IntAttributeD), FSlateAttributeDescriptor::FInvalidateWidgetReasonAttribute{ EInvalidateWidgetReason::Paint});
@@ -137,7 +138,7 @@ public:
 
 	void Construct(const FArguments& InArgs) {}
 
-	TSlateAttribute<int32, EInvalidateWidgetReason::ChildOrder> IntAttributeH;
+	TSlateAttribute<int32, EInvalidateWidgetReason::Layout> IntAttributeH;
 	TSlateAttribute<int32> IntAttributeI;
 	TSlateAttribute<int32> IntAttributeJ;
 	TSlateAttribute<int32> IntAttributeK;
@@ -149,15 +150,15 @@ SLATE_IMPLEMENT_WIDGET(SAttributeLeftWidget_Child)
 void SAttributeLeftWidget_Child::PrivateRegisterAttributes(FSlateAttributeInitializer& AttributeInitializer)
 {
 	// The update order is M, B, A, I, J, D, C, L, H, K
-	//SLATE_ADD_MEMBER_ATTRIBUTE_DEFINITION(AttributeInitializer, IntAttributeH, EInvalidateWidgetReason::ChildOrder);
-	SLATE_ADD_MEMBER_ATTRIBUTE_DEFINITION(AttributeInitializer, IntAttributeJ, EInvalidateWidgetReason::ChildOrder)
+	//SLATE_ADD_MEMBER_ATTRIBUTE_DEFINITION(AttributeInitializer, IntAttributeH, EInvalidateWidgetReason::Layout);
+	SLATE_ADD_MEMBER_ATTRIBUTE_DEFINITION(AttributeInitializer, IntAttributeJ, EInvalidateWidgetReason::Layout)
 		.UpdateDependency("IntAttributeA");
-	SLATE_ADD_MEMBER_ATTRIBUTE_DEFINITION(AttributeInitializer, IntAttributeK, EInvalidateWidgetReason::ChildOrder);
-	SLATE_ADD_MEMBER_ATTRIBUTE_DEFINITION(AttributeInitializer, IntAttributeI, EInvalidateWidgetReason::ChildOrder)
+	SLATE_ADD_MEMBER_ATTRIBUTE_DEFINITION(AttributeInitializer, IntAttributeK, EInvalidateWidgetReason::Layout);
+	SLATE_ADD_MEMBER_ATTRIBUTE_DEFINITION(AttributeInitializer, IntAttributeI, EInvalidateWidgetReason::Layout)
 		.UpdatePrerequisite("IntAttributeB");
-	SLATE_ADD_MEMBER_ATTRIBUTE_DEFINITION(AttributeInitializer, IntAttributeL, EInvalidateWidgetReason::ChildOrder)
+	SLATE_ADD_MEMBER_ATTRIBUTE_DEFINITION(AttributeInitializer, IntAttributeL, EInvalidateWidgetReason::Layout)
 		.UpdatePrerequisite("IntAttributeC");
-	SLATE_ADD_MEMBER_ATTRIBUTE_DEFINITION(AttributeInitializer, IntAttributeM, EInvalidateWidgetReason::ChildOrder)
+	SLATE_ADD_MEMBER_ATTRIBUTE_DEFINITION(AttributeInitializer, IntAttributeM, EInvalidateWidgetReason::Layout)
 		.UpdatePrerequisite("Visibility")
 		.AffectVisibility();
 }
@@ -203,7 +204,7 @@ public:
 
 bool FSlateAttributeTest::RunTest(const FString& Parameters)
 {
-	const int NumberOfAttributeInSWidget = 4;
+	const int NumberOfAttributeInSWidget = SWidget::StaticWidgetClass().GetAttributeDescriptor().GetAttributeNum();
 	int32 OrderCounter = 0;
 	auto OrderLambda = [this, &OrderCounter]() -> int32
 	{
diff --git a/Engine/Source/Runtime/SlateCore/Private/Types/SlateAttributeMetaData.cpp b/Engine/Source/Runtime/SlateCore/Private/Types/SlateAttributeMetaData.cpp
--- a/Engine/Source/Runtime/SlateCore/Private/Types/SlateAttributeMetaData.cpp
+++ b/Engine/Source/Runtime/SlateCore/Private/Types/SlateAttributeMetaData.cpp
@@ -284,6 +284,10 @@ void FSlateAttributeMetaData::InvalidateWidget(SWidget& OwningWidget, const FSla
 		}
 	}
 
+#if WITH_SLATE_DEBUGGING
+	ensureAlwaysMsgf(FSlateAttributeBase::IsInvalidateWidgetReasonSupported(Reason), TEXT("%s is not an EInvalidateWidgetReason supported by SlateAttribute."), *LexToString(Reason));
+#endif
+
 	OwningWidget.Invalidate(Reason);
 }
 
@@ -414,6 +418,9 @@ void FSlateAttributeMetaData::UpdateAttributesImpl(SWidget& OwningWidget, EInval
 
 	if (bInvalidateIfNeeded)
 	{
+#if WITH_SLATE_DEBUGGING
+		ensureAlwaysMsgf(FSlateAttributeBase::IsInvalidateWidgetReasonSupported(InvalidationReason | CachedInvalidationReason), TEXT("%s is not an EInvalidateWidgetReason supported by SlateAttribute."), *LexToString(InvalidationReason | CachedInvalidationReason));
+#endif
 		OwningWidget.Invalidate(InvalidationReason | CachedInvalidationReason);
 		CachedInvalidationReason = EInvalidateWidgetReason::None;
 	}
@@ -444,6 +451,9 @@ void FSlateAttributeMetaData::UpdateAttribute(SWidget& OwningWidget, FSlateAttri
 				if (OwningWidget.IsConstructed())
 				{
 					EInvalidateWidgetReason Reason = GetterItem.GetInvalidationDetail(OwningWidget, Result.InvalidationReason);
+#if WITH_SLATE_DEBUGGING
+					ensureAlwaysMsgf(FSlateAttributeBase::IsInvalidateWidgetReasonSupported(Reason | AttributeMetaData->CachedInvalidationReason), TEXT("%s is not an EInvalidateWidgetReason supported by SlateAttribute."), *LexToString(Reason | AttributeMetaData->CachedInvalidationReason));
+#endif
 					OwningWidget.Invalidate(Reason | AttributeMetaData->CachedInvalidationReason);
 					AttributeMetaData->CachedInvalidationReason = EInvalidateWidgetReason::None;
 				}
diff --git a/Engine/Source/Runtime/SlateCore/Public/Types/SlateAttribute.h b/Engine/Source/Runtime/SlateCore/Public/Types/SlateAttribute.h
--- a/Engine/Source/Runtime/SlateCore/Public/Types/SlateAttribute.h
+++ b/Engine/Source/Runtime/SlateCore/Public/Types/SlateAttribute.h
@@ -164,6 +164,26 @@
 class SWidget;
 
 
+/** Base struct of all SlateAttribute type. */
+struct FSlateAttributeBase
+{
+	/**
+	 * Not all invalidation is supported by SlateAttribute.
+	 * ChildOrder: The update of SlateAttribute is done in the SlatePrepass. We can't add or remove children in SlatePrepass.
+	 * AttributeRegistration: In FastPath, the SlateAttribute are updated in a loop. The iterator can't be modified while we are looping.
+	 */
+	template<typename T>
+	constexpr static bool IsInvalidateWidgetReasonSupported(T Reason)
+	{
+		return false;
+	}
+	constexpr static bool IsInvalidateWidgetReasonSupported(EInvalidateWidgetReason Reason)
+	{
+		return (Reason & (EInvalidateWidgetReason::ChildOrder | EInvalidateWidgetReason::AttributeRegistration)) == EInvalidateWidgetReason::None;
+	}
+};
+
+
 /** Default predicate to compare of Object for SlateAttribute. */
 template<typename ComparePredicate = TEqualTo<>>
 struct TSlateAttributeComparePredicate
@@ -191,6 +211,7 @@ struct TSlateAttributeFTextComparePredicate
 template<EInvalidateWidgetReason InvalidationReason>
 struct TSlateAttributeInvalidationReason
 {
+	static_assert(FSlateAttributeBase::IsInvalidateWidgetReasonSupported(InvalidationReason), "The invalidation is not supported by the SlateAttribute.");
 	static constexpr EInvalidateWidgetReason GetInvalidationReason(const SWidget&) { return InvalidationReason; }
 };
 
@@ -232,13 +253,6 @@ struct FSlateDeprecatedTAttribute
 };
 
 
-/** Base struct of all SlateAttribute type. */
-struct FSlateAttributeBase
-{
-
-};
-
-
 #include "Types/Attributes/SlateAttributeDefinition.inl"
 #include "Types/Attributes/SlateAttributeBase.inl"
 #include "Types/Attributes/SlateAttributeManaged.inl"
diff --git a/Engine/Source/Runtime/SlateCore/Public/Types/SlateAttributeDescriptor.h b/Engine/Source/Runtime/SlateCore/Public/Types/SlateAttributeDescriptor.h
--- a/Engine/Source/Runtime/SlateCore/Public/Types/SlateAttributeDescriptor.h
+++ b/Engine/Source/Runtime/SlateCore/Public/Types/SlateAttributeDescriptor.h
@@ -35,7 +35,8 @@ public:
 		explicit FInvalidateWidgetReasonAttribute(EInvalidateWidgetReason InReason)
 			: Reason(InReason)
 			, Getter()
-		{ }
+		{
+		}
 
 		template<typename... PayloadTypes>
 		explicit FInvalidateWidgetReasonAttribute(typename FGetter::template FStaticDelegate<PayloadTypes...>::FFuncPtr InFuncPtr, PayloadTypes&&... InputPayload)
@@ -215,6 +216,7 @@ private:
 #define SLATE_ADD_MEMBER_ATTRIBUTE_DEFINITION_WITH_NAME(_Initializer, _Name, _Property, _Reason) \
 		static_assert(decltype(_Property)::IsMemberType, "The SlateProperty is not a TSlateAttribute. Do not use SLATE_ADD_MEMBER_ATTRIBUTE_DEFINITION"); \
 		static_assert(!decltype(_Property)::HasDefinedInvalidationReason, "When implementing the SLATE_DECLARE_WIDGET pattern, use TSlateAttribute without the invalidation reason."); \
+		static_assert(!std::is_same<decltype(_Reason), EInvalidateWidgetReason>::value || FSlateAttributeBase::IsInvalidateWidgetReasonSupported(_Reason), "The invalidation is not supported by the SlateAttribute."); \
 		_Initializer.AddMemberAttribute(_Name, STRUCT_OFFSET(PrivateThisType, _Property), FSlateAttributeDescriptor::FInvalidateWidgetReasonAttribute{_Reason})
 
 #define SLATE_ADD_MEMBER_ATTRIBUTE_DEFINITION(_Initializer, _Property, _Reason) \
diff --git a/Engine/Source/Runtime/SlateCore/Public/Widgets/SlateControlledConstruction.h b/Engine/Source/Runtime/SlateCore/Public/Widgets/SlateControlledConstruction.h
--- a/Engine/Source/Runtime/SlateCore/Public/Widgets/SlateControlledConstruction.h
+++ b/Engine/Source/Runtime/SlateCore/Public/Widgets/SlateControlledConstruction.h
@@ -21,6 +21,10 @@ namespace SharedPointerInternals
 		using ThisClass = WidgetType; \
 		using PrivateThisType = WidgetType; \
 		using PrivateParentType = ParentType; \
+		template<class WidgetType, bool bIsUserWidget> \
+		friend struct TWidgetAllocator; \
+		template <typename ObjectType> \
+		friend class SharedPointerInternals::TIntrusiveReferenceController; \
 		static const FSlateWidgetClassData& GetPrivateWidgetClass() \
 		{ \
 			static FSlateWidgetClassData WidgetClassDataInstance = FSlateWidgetClassData(TIdentity<ParentType>(), #WidgetType, &WidgetType::PrivateRegisterAttributes); \
