Slate: Split the Invalidation update loop for better performance tracking.

diff --git a/Engine/Source/Runtime/SlateCore/Private/FastUpdate/SlateInvalidationRoot.cpp b/Engine/Source/Runtime/SlateCore/Private/FastUpdate/SlateInvalidationRoot.cpp
--- a/Engine/Source/Runtime/SlateCore/Private/FastUpdate/SlateInvalidationRoot.cpp
+++ b/Engine/Source/Runtime/SlateCore/Private/FastUpdate/SlateInvalidationRoot.cpp
@@ -129,8 +129,8 @@ static FAutoConsoleVariableRef CVarSlateInvalidationRootVerifySlateAttributes(
 	GSlateInvalidationRootVerifySlateAttribute,
 	TEXT("Each frame, verify that the widgets that have registered attribute are correctly updated once and the list contains all the widgets.")
 );
-void VerifySlateAttribute_BeforeProcessPreUpdate(FSlateInvalidationWidgetList& FastWidgetPathList);
-void VerifySlateAttribute_AfterProcessPreUpdate(const FSlateInvalidationWidgetList& FastWidgetPathList);
+void VerifySlateAttribute_BeforeUpdate(FSlateInvalidationWidgetList& FastWidgetPathList);
+void VerifySlateAttribute_AfterUpdate(const FSlateInvalidationWidgetList& FastWidgetPathList);
 
 #endif //UE_SLATE_WITH_INVALIDATIONWIDGETLIST_DEBUGGING
 
@@ -180,6 +180,7 @@ FSlateInvalidationRoot::FSlateInvalidationRoot()
 	, bNeedsSlowPath(true)
 	, bNeedScreenPositionShift(false)
 	, bProcessingPreUpdate(false)
+	, bProcessingAttributeUpdate(false)
 	, bProcessingPrepassUpdate(false)
 	, bProcessingPostUpdate(false)
 	, bBuildingWidgetList(false)
@@ -515,7 +516,7 @@ void FSlateInvalidationRoot::ProcessPreUpdate()
 	}
 	if (GSlateInvalidationRootVerifySlateAttribute)
 	{
-		VerifySlateAttribute_BeforeProcessPreUpdate(*FastWidgetPathList);
+		VerifySlateAttribute_BeforeUpdate(*FastWidgetPathList);
 	}
 #endif
 
@@ -542,8 +543,6 @@ void FSlateInvalidationRoot::ProcessPreUpdate()
 			}
 #endif
 
-			FSlateInvalidationWidgetList::FWidgetAttributeIterator AttributeItt = FastWidgetPathList->CreateWidgetAttributeIterator();
-
 			/** */
 			struct FChildOriderInvalidationCallbackImpl : FSlateInvalidationWidgetList::IProcessChildOrderInvalidationCallback
 			{
@@ -551,21 +550,18 @@ void FSlateInvalidationRoot::ProcessPreUpdate()
 					const FSlateInvalidationWidgetList& InWidgetList
 					, FSlateInvalidationWidgetPreHeap& InPreUpdate
 					, FSlateInvalidationWidgetPrepassHeap& InPrepassUpdate
-					, FSlateInvalidationWidgetPostHeap& InPostUpdate
-					, FSlateInvalidationWidgetList::FWidgetAttributeIterator& InAttributeItt)
+					, FSlateInvalidationWidgetPostHeap& InPostUpdate)
 					: WidgetList(InWidgetList)
 					, PreUpdate(InPreUpdate)
 					, PrepassUpdate(InPrepassUpdate)
 					, PostUpdate(InPostUpdate)
-					, AttributeItt(InAttributeItt)
 				{}
 				virtual ~FChildOriderInvalidationCallbackImpl() = default;
 				const FSlateInvalidationWidgetList& WidgetList;
 				FSlateInvalidationWidgetPreHeap& PreUpdate;
 				FSlateInvalidationWidgetPrepassHeap& PrepassUpdate;
 				FSlateInvalidationWidgetPostHeap& PostUpdate;
-				FSlateInvalidationWidgetList::FWidgetAttributeIterator& AttributeItt;
-				TArray<FSlateInvalidationWidgetPreHeap::FElement*> WidgetToResort;
+				TArray<FSlateInvalidationWidgetPreHeap::FElement*, TMemStackAllocator<>> WidgetToResort;
 
 				virtual void PreChildRemove(const FSlateInvalidationWidgetList::FIndexRange& Range) override
 				{
@@ -574,7 +570,6 @@ void FSlateInvalidationRoot::ProcessPreUpdate()
 					PreUpdate.RemoveRange(Range);
 					PostUpdate.RemoveRange(Range);
 					PrepassUpdate.RemoveRange(Range);
-					AttributeItt.PreChildRemove(Range);
 				}
 				using FReIndexOperation = FSlateInvalidationWidgetList::IProcessChildOrderInvalidationCallback::FReIndexOperation;
 				virtual void ProxiesReIndexed(const FReIndexOperation& Operation) override
@@ -592,7 +587,6 @@ void FSlateInvalidationRoot::ProcessPreUpdate()
 					PreUpdate.ForEachIndexes(ReIndexIfNeeded);
 					PostUpdate.ForEachIndexes(ReIndexIfNeeded);
 					PrepassUpdate.ForEachIndexes(ReIndexIfNeeded);
-					AttributeItt.ReIndexed(Operation);
 				}
 				using FReSortOperation = FSlateInvalidationWidgetList::IProcessChildOrderInvalidationCallback::FReSortOperation;
 				virtual void ProxiesPreResort(const FReSortOperation& Operation) override
@@ -617,143 +611,68 @@ void FSlateInvalidationRoot::ProcessPreUpdate()
 						Element->GetWidgetSortOrderRef() = FSlateInvalidationWidgetSortOrder{ WidgetList, Element->GetWidgetIndex() };
 					}
 					WidgetToResort.Reset();
-					AttributeItt.PostResort();
-				}
-				virtual void ProxiesBuilt(const FSlateInvalidationWidgetList::FIndexRange& Range) override
-				{
-					AttributeItt.ProxiesBuilt(Range);
 				}
+			};
 
-			}  ChildOrderInvalidationCallback{ *FastWidgetPathList, *WidgetsNeedingPreUpdate, *WidgetsNeedingPrepassUpdate, *WidgetsNeedingPostUpdate, AttributeItt };
+			FMemMark Mark(FMemStack::Get());
+			FChildOriderInvalidationCallbackImpl ChildOrderInvalidationCallback{ *FastWidgetPathList, *WidgetsNeedingPreUpdate, *WidgetsNeedingPrepassUpdate, *WidgetsNeedingPostUpdate };
 
-			while((AttributeItt.IsValid() || WidgetsNeedingPreUpdate->Num() > 0) && !bNeedsSlowPath)
+			while(WidgetsNeedingPreUpdate->Num() > 0 && !bNeedsSlowPath)
 			{
-				FSlateInvalidationWidgetSortOrder AttributeSortOrder = (AttributeItt.IsValid()) ? AttributeItt.GetCurrentSortOrder() : FSlateInvalidationWidgetSortOrder::LimitMax();
-				FSlateInvalidationWidgetSortOrder NeedsUpdateSortOrder = (WidgetsNeedingPreUpdate->Num() > 0) ? WidgetsNeedingPreUpdate->HeapPeekElement().GetWidgetSortOrder() : FSlateInvalidationWidgetSortOrder::LimitMax();
+				// Process ChildOrder && AttributeRegistration invalidation.
 
-				if (AttributeSortOrder == FSlateInvalidationWidgetSortOrder::LimitMax() && NeedsUpdateSortOrder == FSlateInvalidationWidgetSortOrder::LimitMax())
+				const FSlateInvalidationWidgetIndex WidgetIndex = WidgetsNeedingPreUpdate->HeapPop();
+				FSlateInvalidationWidgetList::InvalidationWidgetType& InvalidationWidget = (*FastWidgetPathList)[WidgetIndex];
+				// It could have been destroyed
+				if (SWidget* WidgetPtr = InvalidationWidget.GetWidget())
 				{
-					checkf(false, TEXT("An element inside the lists has an invalid sort order. Something went wrong."));
-					WidgetsNeedingPreUpdate->Reset(true);
-					bNeedsSlowPath = true;
-					break;
-				}
-
-				// Process in order
-				//1.Invalidation AttributeRegistration of NeedsUpdate
-				//2.UpdateAttributes of AttributeSortOrder
-				//3.Invalidation ChildOrder of NeedsUpdateSortOrder
-
-				if (AttributeSortOrder <= NeedsUpdateSortOrder)
-				{
-					// Update Attributes
-					//Note the attribute may still be in the list and will get remove in next loop tick. UpdateCollapsedAttributes and UpdateExpandedAttributes won't do anything.
-					FSlateInvalidationWidgetList::InvalidationWidgetType& InvalidationWidget = (*FastWidgetPathList)[AttributeItt.GetCurrentIndex()];
-					if (SWidget* WidgetPtr = InvalidationWidget.GetWidget())
-					{
-						if (!InvalidationWidget.Visibility.IsCollapseIndirectly())
-						{
-							// if my parent is not collapse, then update my visible state
-							FSlateAttributeMetaData::UpdateOnlyVisibilityAttributes(*WidgetPtr, FSlateAttributeMetaData::EInvalidationPermission::AllowInvalidation);
-							if (!InvalidationWidget.Visibility.IsCollapsed())
-							{
-#if UE_SLATE_WITH_INVALIDATIONWIDGETLIST_DEBUGGING
-								ensureAlwaysMsgf(!GSlateInvalidationRootVerifySlateAttribute || InvalidationWidget.bDebug_AttributeUpdated == false, TEXT("Attribute should only be updated once per frame."));
-								InvalidationWidget.bDebug_AttributeUpdated = true;
-#endif
-								FSlateAttributeMetaData::UpdateExceptVisibilityAttributes(*WidgetPtr, FSlateAttributeMetaData::EInvalidationPermission::AllowInvalidation);
-								AttributeItt.Advance();
-							}
-							else
-							{
-								AttributeItt.AdvanceToNextSibling();
-							}
-						}
-						else
-						{
-							AttributeItt.AdvanceToNextParent();
-						}
-					}
-					else
+#if WITH_SLATE_DEBUGGING
+					if (GSlateInvalidationRootDumpPreInvalidationList)
 					{
-						AttributeItt.Advance();
+						LogPreInvalidationItem(*FastWidgetPathList, WidgetIndex);
 					}
-				}
-				else
-				{
-					// Process ChildOrder invalidation.
+#endif
 
-					const FSlateInvalidationWidgetIndex WidgetIndex = WidgetsNeedingPreUpdate->HeapPeek();
-					FSlateInvalidationWidgetList::InvalidationWidgetType& InvalidationWidget = (*FastWidgetPathList)[WidgetIndex];
-					// It could have been destroyed
-					if (SWidget* WidgetPtr = InvalidationWidget.GetWidget())
+					if (EnumHasAnyFlags(InvalidationWidget.CurrentInvalidateReason, EInvalidateWidgetReason::ChildOrder))
 					{
-#if WITH_SLATE_DEBUGGING
-						if (GSlateInvalidationRootDumpPreInvalidationList)
+// Uncomment to see to be able to compare the list before and after when debugging
+#if 0
+						FMemMark Mark(FMemStack::Get());
+						TArray<TTuple<FSlateInvalidationWidgetIndex, FSlateInvalidationWidgetSortOrder, TWeakPtr<SWidget>>, TMemStackAllocator<>> PreviousPreUpdate;
+						TArray<TTuple<FSlateInvalidationWidgetIndex, FSlateInvalidationWidgetSortOrder, TWeakPtr<SWidget>>, TMemStackAllocator<>> PreviousPrepassUpdate;
+						TArray<TTuple<FSlateInvalidationWidgetIndex, FSlateInvalidationWidgetSortOrder, TWeakPtr<SWidget>>, TMemStackAllocator<>> PreviousPostUpdate;
+						PreviousPreUpdate.Reserve(WidgetsNeedingPreUpdate->Num());
+						PreviousPrepassUpdate.Reserve(WidgetsNeedingPrepassUpdate->Num());
+						PreviousPostUpdate.Reserve(WidgetsNeedingPostUpdate->Num());
+						for (const auto& Element : WidgetsNeedingPreUpdate->GetRaw())
 						{
-							LogPreInvalidationItem(*FastWidgetPathList, WidgetIndex);
+							ensureAlwaysMsgf(FastWidgetPathList->IsValidIndex(Element.GetWidgetIndex()), TEXT("The element is invalid"));
+							PreviousPreUpdate.Emplace(Element.GetWidgetIndex(), Element.GetWidgetSortOrder(), (*FastWidgetPathList)[Element.GetWidgetIndex()].GetWidgetAsShared());
 						}
-#endif
-
-						if (EnumHasAnyFlags(InvalidationWidget.CurrentInvalidateReason, EInvalidateWidgetReason::AttributeRegistration))
+						for (const auto& Element : WidgetsNeedingPrepassUpdate->GetRaw())
 						{
-							FastWidgetPathList->ProcessAttributeRegistrationInvalidation(InvalidationWidget);
-							EnumRemoveFlags(InvalidationWidget.CurrentInvalidateReason, EInvalidateWidgetReason::AttributeRegistration);
-
-							// This element was removed or added, seek will assign the correct widget to be ticked next.
-							AttributeItt.Seek(InvalidationWidget.Index);
-							if (FastWidgetPathList->ShouldBeAddedToAttributeList(*WidgetPtr))
-							{
-								// Do we still need to update this element, if not, then remove it from the update list.
-								if (!Slate::EInvalidateWidgetReason_HasPreUpdateFlag(InvalidationWidget.CurrentInvalidateReason))
-								{
-									WidgetsNeedingPreUpdate->HeapPopDiscard();
-								}
-
-								// We should update the attribute of this proxy before doing the ChildOrder (if any).
-								continue;
-							}
+							ensureAlwaysMsgf(FastWidgetPathList->IsValidIndex(Element.GetWidgetIndex()), TEXT("The element is invalid."));
+							PreviousPrepassUpdate.Emplace(Element.GetWidgetIndex(), Element.GetWidgetSortOrder(), (*FastWidgetPathList)[Element.GetWidgetIndex()].GetWidgetAsShared());
 						}
-
-						if (EnumHasAnyFlags(InvalidationWidget.CurrentInvalidateReason, EInvalidateWidgetReason::ChildOrder))
+						for (const auto& Element : WidgetsNeedingPostUpdate->GetRaw())
 						{
-// Uncomment to see to be able to compare the list before and after when debugging
-#if 0
-							FMemMark Mark(FMemStack::Get());
-							TArray<TTuple<FSlateInvalidationWidgetIndex, FSlateInvalidationWidgetSortOrder, TWeakPtr<SWidget>>, TMemStackAllocator<>> PreviousPreUpdate;
-							TArray<TTuple<FSlateInvalidationWidgetIndex, FSlateInvalidationWidgetSortOrder, TWeakPtr<SWidget>>, TMemStackAllocator<>> PreviousPrepassUpdate;
-							TArray<TTuple<FSlateInvalidationWidgetIndex, FSlateInvalidationWidgetSortOrder, TWeakPtr<SWidget>>, TMemStackAllocator<>> PreviousPostUpdate;
-							PreviousPreUpdate.Reserve(WidgetsNeedingPreUpdate->Num());
-							PreviousPrepassUpdate.Reserve(WidgetsNeedingPrepassUpdate->Num());
-							PreviousPostUpdate.Reserve(WidgetsNeedingPostUpdate->Num());
-							for (const auto& Element : WidgetsNeedingPreUpdate->GetRaw())
-							{
-								ensureAlwaysMsgf(FastWidgetPathList->IsValidIndex(Element.GetWidgetIndex()), TEXT("The element is invalid"));
-								PreviousPreUpdate.Emplace(Element.GetWidgetIndex(), Element.GetWidgetSortOrder(), (*FastWidgetPathList)[Element.GetWidgetIndex()].GetWidgetAsShared());
-							}
-							for (const auto& Element : WidgetsNeedingPrepassUpdate->GetRaw())
-							{
-								ensureAlwaysMsgf(FastWidgetPathList->IsValidIndex(Element.GetWidgetIndex()), TEXT("The element is invalid."));
-								PreviousPrepassUpdate.Emplace(Element.GetWidgetIndex(), Element.GetWidgetSortOrder(), (*FastWidgetPathList)[Element.GetWidgetIndex()].GetWidgetAsShared());
-							}
-							for (const auto& Element : WidgetsNeedingPostUpdate->GetRaw())
-							{
-								ensureAlwaysMsgf(FastWidgetPathList->IsValidIndex(Element.GetWidgetIndex()), TEXT("The element is invalid."));
-								PreviousPostUpdate.Emplace(Element.GetWidgetIndex(), Element.GetWidgetSortOrder(), (*FastWidgetPathList)[Element.GetWidgetIndex()].GetWidgetAsShared());
-							}
+							ensureAlwaysMsgf(FastWidgetPathList->IsValidIndex(Element.GetWidgetIndex()), TEXT("The element is invalid."));
+							PreviousPostUpdate.Emplace(Element.GetWidgetIndex(), Element.GetWidgetSortOrder(), (*FastWidgetPathList)[Element.GetWidgetIndex()].GetWidgetAsShared());
+						}
 #endif
 
-							TGuardValue<bool> ProcessChildOrderInvalidationGuardValue(bProcessingChildOrderInvalidation, true);
-							FastWidgetPathList->ProcessChildOrderInvalidation(InvalidationWidget, ChildOrderInvalidationCallback);
+						TGuardValue<bool> ProcessChildOrderInvalidationGuardValue(bProcessingChildOrderInvalidation, true);
+						FastWidgetPathList->ProcessChildOrderInvalidation(InvalidationWidget, ChildOrderInvalidationCallback);
 
-							// This widget may not be valid anymore (got removed because it doesn't fulfill the requirement anymore ie. NullWidget).
+						// We need to keep it to run the layout calculation in FWidgetProxy::ProcessPostInvalidation
+						//EnumRemoveFlags(InvalidationWidget.CurrentInvalidateReason, EInvalidateWidgetReason::ChildOrder);
+					}
 
-							AttributeItt.FixCurrentWidgetIndex();
-							// We need to keep it to run the layout calculation in FWidgetProxy::ProcessPostInvalidation
-							//EnumRemoveFlags(InvalidationWidget.CurrentInvalidateReason, EInvalidateWidgetReason::ChildOrder);
-						}
+					if (EnumHasAnyFlags(InvalidationWidget.CurrentInvalidateReason, EInvalidateWidgetReason::AttributeRegistration))
+					{
+						FastWidgetPathList->ProcessAttributeRegistrationInvalidation(InvalidationWidget);
+						EnumRemoveFlags(InvalidationWidget.CurrentInvalidateReason, EInvalidateWidgetReason::AttributeRegistration);
 					}
-					WidgetsNeedingPreUpdate->HeapPopDiscard();
 				}
 			}
 		}
@@ -770,9 +689,53 @@ void FSlateInvalidationRoot::ProcessPreUpdate()
 	{
 		ensureMsgf(FastWidgetPathList->VerifyWidgetsIndex(), TEXT("We failed to verify that every widgets has the correct index."));
 	}
+#endif //UE_SLATE_WITH_INVALIDATIONWIDGETLIST_DEBUGGING
+}
+
+void FSlateInvalidationRoot::ProcessAttributeUpdate()
+{
+	TGuardValue<bool> Tmp(bProcessingAttributeUpdate, true);
+
+
+	FSlateInvalidationWidgetList::FWidgetAttributeIterator AttributeItt = FastWidgetPathList->CreateWidgetAttributeIterator();
+	while (AttributeItt.IsValid())
+	{
+		FSlateInvalidationWidgetList::InvalidationWidgetType& InvalidationWidget = (*FastWidgetPathList)[AttributeItt.GetCurrentIndex()];
+		SWidget* WidgetPtr = InvalidationWidget.GetWidget();
+		if (ensureAlways(WidgetPtr))
+		{
+			if (!InvalidationWidget.Visibility.IsCollapseIndirectly())
+			{
+				// if my parent is not collapse, then update my visible state
+				FSlateAttributeMetaData::UpdateOnlyVisibilityAttributes(*WidgetPtr, FSlateAttributeMetaData::EInvalidationPermission::AllowInvalidation);
+				if (!InvalidationWidget.Visibility.IsCollapsed())
+				{
+#if UE_SLATE_WITH_INVALIDATIONWIDGETLIST_DEBUGGING
+					ensureAlwaysMsgf(!GSlateInvalidationRootVerifySlateAttribute || InvalidationWidget.bDebug_AttributeUpdated == false, TEXT("Attribute should only be updated once per frame."));
+					InvalidationWidget.bDebug_AttributeUpdated = true;
+#endif
+					FSlateAttributeMetaData::UpdateExceptVisibilityAttributes(*WidgetPtr, FSlateAttributeMetaData::EInvalidationPermission::AllowInvalidation);
+					AttributeItt.Advance();
+				}
+				else
+				{
+					AttributeItt.AdvanceToNextSibling();
+				}
+			}
+			else
+			{
+				AttributeItt.AdvanceToNextParent();
+			}
+		}
+		else
+		{
+			AttributeItt.Advance();
+		}
+	}
+#if UE_SLATE_WITH_INVALIDATIONWIDGETLIST_DEBUGGING
 	if (GSlateInvalidationRootVerifySlateAttribute)
 	{
-		VerifySlateAttribute_AfterProcessPreUpdate(*FastWidgetPathList);
+		VerifySlateAttribute_AfterUpdate(*FastWidgetPathList);
 	}
 #endif //UE_SLATE_WITH_INVALIDATIONWIDGETLIST_DEBUGGING
 }
@@ -894,6 +857,7 @@ bool FSlateInvalidationRoot::ProcessInvalidation()
 		check(WidgetsNeedingPrepassUpdate);
 		check(WidgetsNeedingPostUpdate);
 
+		SCOPED_NAMED_EVENT(Slate_InvalidationProcessing_PreUpdate, FColor::Blue);
 		ProcessPreUpdate();
 
 #if UE_SLATE_WITH_INVALIDATIONWIDGETLIST_DEBUGGING
@@ -904,6 +868,13 @@ bool FSlateInvalidationRoot::ProcessInvalidation()
 #endif //UE_SLATE_WITH_INVALIDATIONWIDGETLIST_DEBUGGING
 	}
 
+	if (!bNeedsSlowPath)
+	{
+		SCOPED_NAMED_EVENT(Slate_InvalidationProcessing_AttributeUpdate, FColor::Blue);
+
+		ProcessAttributeUpdate();
+	}
+
 	if (!bNeedsSlowPath)
 	{
 		// Put all widgets in the VolatileUpdate list in the WidgetsNeedingPostUpdate
@@ -926,12 +897,15 @@ bool FSlateInvalidationRoot::ProcessInvalidation()
 
 	if (!bNeedsSlowPath)
 	{
+		SCOPED_NAMED_EVENT(Slate_InvalidationProcessing_PrepassUpdate, FColor::Blue);
 		ProcessPrepassUpdate();
 	}
 
 	if (!bNeedsSlowPath)
 	{
 		FinalUpdateList.Reset(WidgetsNeedingPostUpdate->Num());
+
+		SCOPED_NAMED_EVENT(Slate_InvalidationProcessing_PostUpdate, FColor::Blue);
 		bWidgetsNeedRepaint = ProcessPostUpdate();
 	}
 	
@@ -1491,7 +1465,7 @@ void VerifyWidgetsUpdateList_AfterProcessPostUpdate(const TSharedRef<SWidget>& R
 	}
 }
 
-void VerifySlateAttribute_BeforeProcessPreUpdate(FSlateInvalidationWidgetList& FastWidgetPathList)
+void VerifySlateAttribute_BeforeUpdate(FSlateInvalidationWidgetList& FastWidgetPathList)
 {
 	FastWidgetPathList.ForEachInvalidationWidget([](FSlateInvalidationWidgetList::InvalidationWidgetType& InvalidationWidget)
 		{
@@ -1499,7 +1473,7 @@ void VerifySlateAttribute_BeforeProcessPreUpdate(FSlateInvalidationWidgetList& F
 		});
 }
 
-void VerifySlateAttribute_AfterProcessPreUpdate(const FSlateInvalidationWidgetList& FastWidgetPathList)
+void VerifySlateAttribute_AfterUpdate(const FSlateInvalidationWidgetList& FastWidgetPathList)
 {
 	const bool bElementIndexListValid = FastWidgetPathList.VerifyElementIndexList();
 	UE_SLATE_LOG_ERROR_IF_FALSE(bElementIndexListValid
diff --git a/Engine/Source/Runtime/SlateCore/Private/FastUpdate/SlateInvalidationWidgetList.cpp b/Engine/Source/Runtime/SlateCore/Private/FastUpdate/SlateInvalidationWidgetList.cpp
--- a/Engine/Source/Runtime/SlateCore/Private/FastUpdate/SlateInvalidationWidgetList.cpp
+++ b/Engine/Source/Runtime/SlateCore/Private/FastUpdate/SlateInvalidationWidgetList.cpp
@@ -417,7 +417,6 @@ FSlateInvalidationWidgetIndex FSlateInvalidationWidgetList::Internal_BuildWidget
 	FSlateInvalidationWidgetIndex LeafMostChildIndex = NewIndex;
 	FSlateInvalidationWidgetVisibility NewVisibility {ParentVisibility, Widget.GetVisibility()};
 
-	bool bUpdateSlateAttribute = false;
 	if (ShouldBeAddedToAttributeList(Widget))
 	{
 		// The list is already sorted at this point. Add to the end.
@@ -429,13 +428,6 @@ FSlateInvalidationWidgetIndex FSlateInvalidationWidgetList::Internal_BuildWidget
 			{
 				FSlateAttributeMetaData::UpdateOnlyVisibilityAttributes(Widget, FSlateAttributeMetaData::EInvalidationPermission::DenyAndClearDelayedInvalidation);
 				NewVisibility.SetVisibility(ParentVisibility, Widget.GetVisibility());
-
-				if (!NewVisibility.IsCollapsed())
-				{
-					FSlateAttributeMetaData::UpdateExceptVisibilityAttributes(Widget, FSlateAttributeMetaData::EInvalidationPermission::DenyAndClearDelayedInvalidation);
-					NewVisibility.SetVisibility(ParentVisibility, Widget.GetVisibility());
-					bUpdateSlateAttribute = true;
-				}
 			}
 		}
 	}
@@ -455,9 +447,6 @@ FSlateInvalidationWidgetIndex FSlateInvalidationWidgetList::Internal_BuildWidget
 		WidgetProxy.Visibility = NewVisibility;
 		WidgetProxy.bIsInvalidationRoot = bIsInvalidationRoot;
 		WidgetProxy.bIsVolatilePrepass = Widget.HasAnyUpdateFlags(EWidgetUpdateFlags::NeedsVolatilePrepass);
-#if UE_SLATE_WITH_INVALIDATIONWIDGETLIST_DEBUGGING
-		WidgetProxy.bDebug_AttributeUpdated = bUpdateSlateAttribute;
-#endif
 	}
 
 #if UE_SLATE_WITH_INVALIDATIONWIDGETLIST_DEBUGGING
@@ -698,7 +687,7 @@ bool FSlateInvalidationWidgetList::ProcessChildOrderInvalidation(const Invalidat
 				FSlateInvalidationWidgetIndex PreviousLeafMostChildIndex = InvalidationWidget.LeafMostChildIndex;
 				Internal_RebuildWidgetListTree(*WidgetPtr, StartChildIndex);
 
-
+#if UE_SLATE_WITH_INVALIDATIONWIDGETLIST_CHILDORDERCHECK
 				const InvalidationWidgetType& NewInvalidationWidget = (*this)[PreviousWidgetIndex];
 				const FIndexRange BuiltRange = (NewInvalidationWidget.LeafMostChildIndex == PreviousLeafMostChildIndex)
 					? FIndexRange{ *this, PreviousLeafMostChildIndex, NewInvalidationWidget.LeafMostChildIndex }
@@ -707,6 +696,7 @@ bool FSlateInvalidationWidgetList::ProcessChildOrderInvalidation(const Invalidat
 				{
 					Callback.ProxiesBuilt(BuiltRange);
 				}
+#endif
 			}
 		}
 	}
diff --git a/Engine/Source/Runtime/SlateCore/Private/FastUpdate/SlateInvalidationWidgetList.h b/Engine/Source/Runtime/SlateCore/Private/FastUpdate/SlateInvalidationWidgetList.h
--- a/Engine/Source/Runtime/SlateCore/Private/FastUpdate/SlateInvalidationWidgetList.h
+++ b/Engine/Source/Runtime/SlateCore/Private/FastUpdate/SlateInvalidationWidgetList.h
@@ -13,6 +13,7 @@
 #include "Widgets/SWidget.h"
 
 #define UE_SLATE_WITH_INVALIDATIONWIDGETLIST_RANGECHECK UE_BUILD_DEBUG
+#define UE_SLATE_WITH_INVALIDATIONWIDGETLIST_CHILDORDERCHECK UE_BUILD_DEBUG || UE_BUILD_DEVELOPMENT
 
 
 class FSlateInvalidationWidgetList
@@ -126,8 +127,10 @@ public:
 		virtual void ProxiesPreResort(const FReSortOperation& Operation) {}
 		/** Widget proxies that got resorted by the operation. */
 		virtual void ProxiesPostResort() {}
+#if UE_SLATE_WITH_INVALIDATIONWIDGETLIST_CHILDORDERCHECK
 		/** Widget proxies built by the operation. */
 		virtual void ProxiesBuilt(const FIndexRange& Range) {}
+#endif
 	};
 
 	/**
diff --git a/Engine/Source/Runtime/SlateCore/Private/Test/SlateInvalidationWidgetListTest.cpp b/Engine/Source/Runtime/SlateCore/Private/Test/SlateInvalidationWidgetListTest.cpp
--- a/Engine/Source/Runtime/SlateCore/Private/Test/SlateInvalidationWidgetListTest.cpp
+++ b/Engine/Source/Runtime/SlateCore/Private/Test/SlateInvalidationWidgetListTest.cpp
@@ -325,7 +325,9 @@ bool FSlateInvalidationWidgetListTest::RunTest(const FString& Parameters)
 				virtual void ProxiesReIndexed(const FReIndexOperation& Operation) override { ++ReIndexedCount; }
 				using FReSortOperation = FSlateInvalidationWidgetList::IProcessChildOrderInvalidationCallback::FReSortOperation;
 				virtual void ProxiesPreResort(const FReSortOperation& Operation) override { ++ResortCount; }
+#if UE_SLATE_WITH_INVALIDATIONWIDGETLIST_CHILDORDERCHECK
 				virtual void ProxiesBuilt(const FSlateInvalidationWidgetList::FIndexRange& Range) override { ++BuiltCount; }
+#endif
 			};
 
 			auto TestRemoveWidget = [&](TSharedPtr<SVerticalBox>& Widget, const TCHAR* Message) -> bool
@@ -473,6 +475,7 @@ bool FSlateInvalidationWidgetListTest::RunTest(const FString& Parameters)
 					}
 				}
 
+#if UE_SLATE_WITH_INVALIDATIONWIDGETLIST_CHILDORDERCHECK
 				TArray<TTuple<FString, FSlateInvalidationWidgetSortOrder, bool>> ProxiesBuiltToCheck;
 				virtual void ProxiesBuilt(const FSlateInvalidationWidgetList::FIndexRange& Range) override
 				{
@@ -494,6 +497,7 @@ bool FSlateInvalidationWidgetListTest::RunTest(const FString& Parameters)
 						}
 					}
 				}
+#endif
 			};
 
 			FChildOrderInvalidationCallback_RemoveG CallbackG;
@@ -535,7 +539,7 @@ bool FSlateInvalidationWidgetListTest::RunTest(const FString& Parameters)
 			}
 
 			{
-
+#if UE_SLATE_WITH_INVALIDATIONWIDGETLIST_CHILDORDERCHECK
 				CallbackG.ProxiesBuiltToCheck.Emplace(TEXT("C"), FSlateInvalidationWidgetSortOrder{ List, WidgetIndexC }, false);
 				CallbackG.ProxiesBuiltToCheck.Emplace(TEXT("F"), FSlateInvalidationWidgetSortOrder{ List, WidgetIndexF }, false);
 				const FSlateInvalidationWidgetIndex NewWidgetIndexH = WidgetF->GetAllChildren()->GetChildAt(0)->GetProxyHandle().GetWidgetIndex();
@@ -548,6 +552,7 @@ bool FSlateInvalidationWidgetListTest::RunTest(const FString& Parameters)
 				CallbackG.ProxiesBuiltToCheck.Emplace(TEXT("J"), FSlateInvalidationWidgetSortOrder{ List, NewWidgetIndexJ }, false);
 				const FSlateInvalidationWidgetIndex NewWidgetIndex14 = List.IncrementIndex(NewWidgetIndexJ);
 				CallbackG.ProxiesBuiltToCheck.Emplace(TEXT("14"), FSlateInvalidationWidgetSortOrder{ List, NewWidgetIndex14 }, false);
+#endif
 			}
 
 			// Remove G
diff --git a/Engine/Source/Runtime/SlateCore/Public/FastUpdate/SlateInvalidationRoot.h b/Engine/Source/Runtime/SlateCore/Public/FastUpdate/SlateInvalidationRoot.h
--- a/Engine/Source/Runtime/SlateCore/Public/FastUpdate/SlateInvalidationRoot.h
+++ b/Engine/Source/Runtime/SlateCore/Public/FastUpdate/SlateInvalidationRoot.h
@@ -148,8 +148,10 @@ private:
 	void BuildFastPathWidgetList(const TSharedRef<SWidget>& RootWidget);
 	void AdjustWidgetsDesktopGeometry(FVector2D WindowToDesktopTransform);
 
-	/** Update child order and slate attribute */
+	/** Update child order and slate attribute registration */
 	void ProcessPreUpdate();
+	/** Slate attribute update */
+	void ProcessAttributeUpdate();
 	/** Call slate prepass. */
 	void ProcessPrepassUpdate();
 	/** Update layout, and update the FinalUpdateList */
@@ -190,6 +192,7 @@ private:
 	bool bNeedsSlowPath;
 	bool bNeedScreenPositionShift;
 	bool bProcessingPreUpdate;
+	bool bProcessingAttributeUpdate;
 	bool bProcessingPrepassUpdate;
 	bool bProcessingPostUpdate;
 	bool bBuildingWidgetList;
