Slate: Deprecate SLATE_SUPPORTS_SLOT. SLATE_SLOT_ARGUMENT should now be use. The new macro support FSlotArguments. That let us create slot with TAttribute that can be transform into SlateAttribute.
Deprecate TAlignmentWidgetSlotMixin old declartive function. They were not invalidating the widget.
Deprecate TPanelChildren.Add and Insert function. We now use TUniquePtr, the ownership of the Slot wouldn't exist outside of the life of the slot.

diff --git a/Engine/Source/Editor/DetailCustomizations/Private/SlateBrushCustomization.cpp b/Engine/Source/Editor/DetailCustomizations/Private/SlateBrushCustomization.cpp
--- a/Engine/Source/Editor/DetailCustomizations/Private/SlateBrushCustomization.cpp
+++ b/Engine/Source/Editor/DetailCustomizations/Private/SlateBrushCustomization.cpp
@@ -776,8 +776,8 @@ private:
 	 */
 	void UpdateOverlayAlignment()
 	{
-		OverlaySlot->HAlign( HorizontalAlignment );
-		OverlaySlot->VAlign( VerticalAlignment );
+		OverlaySlot->SetHorizontalAlignment( HorizontalAlignment );
+		OverlaySlot->SetVerticalAlignment( VerticalAlignment );
 	}
 
 	/**
diff --git a/Engine/Source/Editor/GraphEditor/Private/SGraphPanel.cpp b/Engine/Source/Editor/GraphEditor/Private/SGraphPanel.cpp
--- a/Engine/Source/Editor/GraphEditor/Private/SGraphPanel.cpp
+++ b/Engine/Source/Editor/GraphEditor/Private/SGraphPanel.cpp
@@ -212,7 +212,7 @@ int32 SGraphPanel::OnPaint( const FPaintArgs& Args, const FGeometry& AllottedGeo
 				// Draw the comments and information popups for this node, if it has any.
 				{
 					const SNodePanel::SNode::FNodeSlot* CommentSlot = ChildNode->GetSlot( ENodeZone::TopCenter );
-					float CommentBubbleY = CommentSlot ? -CommentSlot->Offset.Get().Y : 0.f;
+					float CommentBubbleY = CommentSlot ? -CommentSlot->GetSlotOffset().Y : 0.f;
 					Context.bSelected = bSelected;
 					TArray<FGraphInformationPopupInfo> Popups;
 
diff --git a/Engine/Source/Editor/GraphEditor/Public/SNodePanel.h b/Engine/Source/Editor/GraphEditor/Public/SNodePanel.h
--- a/Engine/Source/Editor/GraphEditor/Public/SNodePanel.h
+++ b/Engine/Source/Editor/GraphEditor/Public/SNodePanel.h
@@ -278,32 +278,57 @@ public:
 		class FNodeSlot : public TSlotBase<FNodeSlot>, public TAlignmentWidgetSlotMixin<FNodeSlot>
 		{
 		public:
+			friend SNode;
+
 			FNodeSlot()
+				: FNodeSlot(ENodeZone::TopLeft)
+			{ }
+
+			FNodeSlot(ENodeZone::Type InZone)
 				: TSlotBase<FNodeSlot>()
 				, TAlignmentWidgetSlotMixin<FNodeSlot>(HAlign_Fill, VAlign_Fill)
+				, Zone(InZone)
 				, SlotPadding(0.0f)
-				, Offset( FVector2D::ZeroVector )
-				, AllowScale( true )
+				, Offset(FVector2D::ZeroVector)
+				, AllowScale(true)
 			{ }
 
+			SLATE_SLOT_BEGIN_ARGS_OneMixin(FNodeSlot, TSlotBase<FNodeSlot>, TAlignmentWidgetSlotMixin<FNodeSlot>)
+				SLATE_ATTRIBUTE(FMargin, Padding)
+				SLATE_ATTRIBUTE(FVector2D, SlotOffset)
+				SLATE_ATTRIBUTE(FVector2D, SlotSize)
+				SLATE_ATTRIBUTE(bool, AllowScaling)
+			SLATE_SLOT_END_ARGS()
+
+			void Construct(const FChildren& SlotOwner, FSlotArguments&& InArgs)
+			{
+				TSlotBase<FNodeSlot>::Construct(SlotOwner, MoveTemp(InArgs));
+				TAlignmentWidgetSlotMixin<FNodeSlot>::ConstructMixin(SlotOwner, MoveTemp(InArgs));
+			}
+
+		public:
+			UE_DEPRECATED(5.0, "Padding is now deprecated. Use the FSlotArgument or the SetPadding function.")
 			FNodeSlot& Padding( const TAttribute<FMargin> InPadding )
 			{
 				SlotPadding = InPadding;
 				return *this;
 			}
 
+			UE_DEPRECATED(5.0, "SlotOffset is now deprecated. Use the FSlotArgument or the SetSlotOffset function.")
 			FNodeSlot& SlotOffset( const TAttribute<FVector2D> InOffset )
 			{
 				Offset = InOffset;
 				return *this;
 			}
 
+			UE_DEPRECATED(5.0, "SlotSize is now deprecated. Use the FSlotArgument or the SetSlotSize function.")
 			FNodeSlot& SlotSize( const TAttribute<FVector2D> InSize )
 			{
 				Size = InSize;
 				return *this;
 			}
 
+			UE_DEPRECATED(5.0, "AllowScaling is now deprecated. Use the FSlotArgument or the SetAllowScalingfunction.")
 			FNodeSlot& AllowScaling( const TAttribute<bool> InAllowScale )
 			{
 				AllowScale = InAllowScale;
@@ -311,7 +336,52 @@ public:
 			}
 
 		public:
+			ENodeZone::Type GetZoneType() const
+			{
+				return Zone;
+			}
 
+			void SetPadding(TAttribute<FMargin> InPadding)
+			{
+				SlotPadding = MoveTemp(InPadding);
+			}
+
+			FMargin GetPadding() const
+			{
+				return SlotPadding.Get();
+			}
+
+			void SetSlotOffset(TAttribute<FVector2D> InOffset)
+			{
+				Offset = MoveTemp(InOffset);
+			}
+
+			FVector2D GetSlotOffset() const
+			{
+				return Offset.Get();
+			}
+
+			void SetSlotSize(TAttribute<FVector2D> InSize)
+			{
+				Size = MoveTemp(InSize);
+			}
+
+			FVector2D GetSlotSize() const
+			{
+				return Size.Get();
+			}
+
+			void SetAllowScaling(TAttribute<bool> InAllowScaling)
+			{
+				AllowScale = MoveTemp(InAllowScaling);
+			}
+
+			bool GetAllowScaling() const
+			{
+				return AllowScale.Get();
+			}
+
+		private:
 			/** The child widget contained in this slot. */
 			ENodeZone::Type Zone;
 			TAttribute<FMargin> SlotPadding;
@@ -422,22 +492,23 @@ public:
 		}
 		// End of SPanel Interface
 
-		FNodeSlot& GetOrAddSlot( const ENodeZone::Type SlotId )
+
+		using FScopedWidgetSlotArguments = TPanelChildren<FNodeSlot>::FScopedWidgetSlotArguments;
+		FScopedWidgetSlotArguments GetOrAddSlot( const ENodeZone::Type SlotId )
 		{
 			// Return existing
+			int32 InsertIndex = INDEX_NONE;
 			for( int32 ChildIndex = 0; ChildIndex < Children.Num(); ++ChildIndex )
 			{
 				if( Children[ ChildIndex ].Zone == SlotId )
 				{
-					return Children[ ChildIndex ];
+					Children.RemoveAt(ChildIndex);
+					InsertIndex = ChildIndex;
 				}
 			}
-			// Add Zone
-			FNodeSlot& NewSlot = *new FNodeSlot();
-			NewSlot.Zone = SlotId;
-			Children.Add( &NewSlot );
 
-			return NewSlot;
+			// Add new
+			return FScopedWidgetSlotArguments{ MakeUnique<FNodeSlot>(SlotId), Children, InsertIndex };
 		}
 
 		FNodeSlot* GetSlot( const ENodeZone::Type SlotId )
diff --git a/Engine/Source/Editor/Persona/Private/AnimTimeline/SAnimTrackArea.cpp b/Engine/Source/Editor/Persona/Private/AnimTimeline/SAnimTrackArea.cpp
--- a/Engine/Source/Editor/Persona/Private/AnimTimeline/SAnimTrackArea.cpp
+++ b/Engine/Source/Editor/Persona/Private/AnimTimeline/SAnimTrackArea.cpp
@@ -51,7 +51,7 @@ void SAnimTrackArea::Empty()
 void SAnimTrackArea::AddTrackSlot(const TSharedRef<FAnimTimelineTrack>& InTrack, const TSharedPtr<SAnimTrack>& InSlot)
 {
 	TrackSlots.Add(InTrack, InSlot);
-	Children.Add(new FAnimTrackAreaSlot(InSlot));
+	Children.AddSlot(FAnimTrackAreaSlot::FSlotArguments(MakeUnique<FAnimTrackAreaSlot>(InSlot)));
 }
 
 TSharedPtr<SAnimTrack> SAnimTrackArea::FindTrackSlot(const TSharedRef<FAnimTimelineTrack>& InTrack)
diff --git a/Engine/Source/Editor/Sequencer/Private/SSequencerTrackArea.cpp b/Engine/Source/Editor/Sequencer/Private/SSequencerTrackArea.cpp
--- a/Engine/Source/Editor/Sequencer/Private/SSequencerTrackArea.cpp
+++ b/Engine/Source/Editor/Sequencer/Private/SSequencerTrackArea.cpp
@@ -71,7 +71,7 @@ void SSequencerTrackArea::Empty()
 void SSequencerTrackArea::AddTrackSlot(const TSharedRef<FSequencerDisplayNode>& InNode, const TSharedPtr<SSequencerTrackLane>& InSlot)
 {
 	TrackSlots.Add(InNode, InSlot);
-	Children.Add(new FTrackAreaSlot(InSlot));
+	Children.AddSlot(FTrackAreaSlot::FSlotArguments(MakeUnique<FTrackAreaSlot>(InSlot)));
 }
 
 
diff --git a/Engine/Source/Runtime/Slate/Private/Framework/Docking/SDockingTabStack.cpp b/Engine/Source/Runtime/Slate/Private/Framework/Docking/SDockingTabStack.cpp
--- a/Engine/Source/Runtime/Slate/Private/Framework/Docking/SDockingTabStack.cpp
+++ b/Engine/Source/Runtime/Slate/Private/Framework/Docking/SDockingTabStack.cpp
@@ -155,7 +155,7 @@ void SDockingTabStack::Construct( const FArguments& InArgs, const TSharedRef<FTa
 
 				+ SVerticalBox::Slot()
 				.Expose(TitleBarSlot)
-				.AutoHeight()				
+				.AutoHeight()
 
 				+ SVerticalBox::Slot()
 				.AutoHeight()
@@ -837,7 +837,7 @@ TSharedPtr<FTabManager::FLayoutNode> SDockingTabStack::GatherPersistentLayout()
 void SDockingTabStack::ClearReservedSpace()
 {
 	bShowingTitleBarArea = false;
-	(*TitleBarSlot).Padding(0.f);
+	TitleBarSlot->SetPadding(0.f);
 }
 
 void SDockingTabStack::ReserveSpaceForWindowChrome(EChromeElement Element, bool bIncludePaddingForMenuBar, bool bOnlyMinorTabs)
@@ -858,11 +858,11 @@ void SDockingTabStack::ReserveSpaceForWindowChrome(EChromeElement Element, bool
 	switch (Element)
 	{
 	case EChromeElement::Controls:
-		(*TitleBarSlot).Padding(CurrentPadding + ControlsPadding);
+		TitleBarSlot->SetPadding(CurrentPadding + ControlsPadding);
 		break;
 
 	case EChromeElement::Icon:
-		(*TitleBarSlot).Padding(CurrentPadding + IconPadding);
+		TitleBarSlot->SetPadding(CurrentPadding + IconPadding);
 		break;
 
 	default:
@@ -896,7 +896,7 @@ void SDockingTabStack::SetParentNode( TSharedRef<class SDockingSplitter> InParen
 	// place the title bar widgets into our content instead!
 	const TSharedPtr<SDockingArea>& DockArea = GetDockArea();
 
-	(*TitleBarSlot)[TitleBarContent.ToSharedRef()];
+	TitleBarSlot->AttachWidget(TitleBarContent.ToSharedRef());
 }
 
 
diff --git a/Engine/Source/Runtime/Slate/Private/Widgets/Input/SMenuAnchor.cpp b/Engine/Source/Runtime/Slate/Private/Widgets/Input/SMenuAnchor.cpp
--- a/Engine/Source/Runtime/Slate/Private/Widgets/Input/SMenuAnchor.cpp
+++ b/Engine/Source/Runtime/Slate/Private/Widgets/Input/SMenuAnchor.cpp
@@ -83,15 +83,12 @@ SMenuAnchor::FPopupPlacement::FPopupPlacement(const FGeometry& PlacementGeometry
  */
 void SMenuAnchor::Construct( const FArguments& InArgs )
 {
-	Children.Add(new FBasicLayoutWidgetSlot());
-	Children.Add(new FBasicLayoutWidgetSlot());
-	
-
-	Children[0]
+	Children.AddSlot(MoveTemp(FBasicLayoutWidgetSlot::FSlotArguments(MakeUnique<FBasicLayoutWidgetSlot>())
 		.Padding(InArgs._Padding)
 		[
 			InArgs._Content.Widget
-		];
+		]));
+	Children.AddSlot(FBasicLayoutWidgetSlot::FSlotArguments(MakeUnique<FBasicLayoutWidgetSlot>()));
 
 	MenuContent                            = InArgs._MenuContent;
 	WrappedContent                         = InArgs._MenuContent;
@@ -283,11 +280,8 @@ bool SMenuAnchor::IsOpenViaCreatedWindow() const
 
 void SMenuAnchor::SetContent(TSharedRef<SWidget> InContent)
 {
-	Children[0]
-	.Padding(0.f)
-	[
-		InContent
-	];
+	Children[0].SetPadding(0.f);
+	Children[0].AttachWidget(InContent);
 }
 
 void SMenuAnchor::SetMenuContent(TSharedRef<SWidget> InMenuContent)
diff --git a/Engine/Source/Runtime/Slate/Private/Widgets/Layout/SBox.cpp b/Engine/Source/Runtime/Slate/Private/Widgets/Layout/SBox.cpp
--- a/Engine/Source/Runtime/Slate/Private/Widgets/Layout/SBox.cpp
+++ b/Engine/Source/Runtime/Slate/Private/Widgets/Layout/SBox.cpp
@@ -20,7 +20,7 @@ void SBox::PrivateRegisterAttributes(FSlateAttributeInitializer& AttributeInitia
 }
 
 SBox::SBox()
-	: ChildSlot(*this)
+	: ChildSlot(this)
 	, WidthOverride(*this)
 	, HeightOverride(*this)
 	, MinDesiredWidth(*this)
diff --git a/Engine/Source/Runtime/Slate/Private/Widgets/Layout/SDPIScaler.cpp b/Engine/Source/Runtime/Slate/Private/Widgets/Layout/SDPIScaler.cpp
--- a/Engine/Source/Runtime/Slate/Private/Widgets/Layout/SDPIScaler.cpp
+++ b/Engine/Source/Runtime/Slate/Private/Widgets/Layout/SDPIScaler.cpp
@@ -11,7 +11,7 @@ void SDPIScaler::PrivateRegisterAttributes(FSlateAttributeInitializer& Attribute
 }
 
 SDPIScaler::SDPIScaler()
-	: ChildSlot(*this)
+	: ChildSlot(this)
 	, DPIScaleAttribute(*this, 1.f)
 {
 	SetCanTick(false);
diff --git a/Engine/Source/Runtime/Slate/Private/Widgets/Layout/SScrollBarTrack.cpp b/Engine/Source/Runtime/Slate/Private/Widgets/Layout/SScrollBarTrack.cpp
--- a/Engine/Source/Runtime/Slate/Private/Widgets/Layout/SScrollBarTrack.cpp
+++ b/Engine/Source/Runtime/Slate/Private/Widgets/Layout/SScrollBarTrack.cpp
@@ -12,26 +12,14 @@ void SScrollBarTrack::Construct(const FArguments& InArgs)
 	MinThumbSize = 35;
 	Orientation = InArgs._Orientation;
 
-	// This panel only supports 3 children!
-	for (int32 CurChild = 0; CurChild < NUM_SCROLLBAR_SLOTS; ++CurChild)
-	{
-		Children.Add(new FSlot());
-	}
-
-	Children[TOP_SLOT_INDEX]
-		[
-			InArgs._TopSlot.Widget
-		];
-
-	Children[BOTTOM_SLOT_INDEX]
-		[
-			InArgs._BottomSlot.Widget
-		];
-
-	Children[THUMB_SLOT_INDEX]
-		[
-			InArgs._ThumbSlot.Widget
-		];
+	
+	static_assert(NUM_SCROLLBAR_SLOTS == 3, "SScrollBarTrack::FSlot numbers changed");
+	static_assert(TOP_SLOT_INDEX == 0, "SScrollBarTrack::FSlot order has changed");
+	Children.AddSlot(MoveTemp(FSlot::FSlotArguments(MakeUnique<FSlot>())[InArgs._TopSlot.Widget]));
+	static_assert(BOTTOM_SLOT_INDEX == 1, "SScrollBarTrack::FSlot order has changed");
+	Children.AddSlot(MoveTemp(FSlot::FSlotArguments(MakeUnique<FSlot>())[InArgs._BottomSlot.Widget]));
+	static_assert(THUMB_SLOT_INDEX == 2, "SScrollBarTrack::FSlot order has changed");
+	Children.AddSlot(MoveTemp(FSlot::FSlotArguments(MakeUnique<FSlot>())[InArgs._ThumbSlot.Widget]));
 }
 
 SScrollBarTrack::FTrackSizeInfo SScrollBarTrack::GetTrackSizeInfo(const FGeometry& InTrackGeometry) const
diff --git a/Engine/Source/Runtime/Slate/Private/Widgets/Layout/SScrollBox.cpp b/Engine/Source/Runtime/Slate/Private/Widgets/Layout/SScrollBox.cpp
--- a/Engine/Source/Runtime/Slate/Private/Widgets/Layout/SScrollBox.cpp
+++ b/Engine/Source/Runtime/Slate/Private/Widgets/Layout/SScrollBox.cpp
@@ -15,6 +15,7 @@ SScrollBox::FSlot::FSlotArguments SScrollBox::Slot()
 	return FSlot::FSlotArguments(MakeUnique<FSlot>());
 }
 
+PRAGMA_DISABLE_DEPRECATION_WARNINGS
 void SScrollPanel::Construct(const FArguments& InArgs, const TArray<SScrollBox::FSlot*>& InSlots)
 {
 	PhysicalOffset = 0;
@@ -27,6 +28,7 @@ void SScrollPanel::Construct(const FArguments& InArgs, const TArray<SScrollBox::
 	BackPadScrolling = InArgs._BackPadScrolling;
 	FrontPadScrolling = InArgs._FrontPadScrolling;
 }
+PRAGMA_ENABLE_DEPRECATION_WARNINGS
 
 
 void SScrollPanel::Construct(const FArguments& InArgs, TArray<SScrollBox::FSlot::FSlotArguments> InSlots)
diff --git a/Engine/Source/Runtime/Slate/Private/Widgets/Views/SListPanel.cpp b/Engine/Source/Runtime/Slate/Private/Widgets/Views/SListPanel.cpp
--- a/Engine/Source/Runtime/Slate/Private/Widgets/Views/SListPanel.cpp
+++ b/Engine/Source/Runtime/Slate/Private/Widgets/Views/SListPanel.cpp
@@ -19,26 +19,17 @@ void SListPanel::Construct( const FArguments& InArgs )
 	NumDesiredItems = InArgs._NumDesiredItems;
 	ItemAlignment = InArgs._ItemAlignment;
 	Orientation = InArgs._ListOrientation;
+	Children.AddSlots(MoveTemp(const_cast<TArray<FSlot::FSlotArguments>&>(InArgs._Slots)));
 }
 
-SListPanel::FSlot& SListPanel::Slot()
+SListPanel::FSlot::FSlotArguments SListPanel::Slot()
 {
-	return *(new FSlot());
+	return FSlot::FSlotArguments(MakeUnique<FSlot>());
 }
 	
-SListPanel::FSlot& SListPanel::AddSlot(int32 InsertAtIndex)
+SListPanel::FScopedWidgetSlotArguments SListPanel::AddSlot(int32 InsertAtIndex)
 {
-	FSlot& NewSlot = SListPanel::Slot();
-	if (InsertAtIndex == INDEX_NONE)
-	{
-		this->Children.Add( &NewSlot );
-	}
-	else
-	{
-		this->Children.Insert( &NewSlot, InsertAtIndex );
-	}
-	
-	return NewSlot;
+	return FScopedWidgetSlotArguments{ MakeUnique<FSlot>(), Children, InsertAtIndex };
 }
 
 void SListPanel::OnArrangeChildren( const FGeometry& AllottedGeometry, FArrangedChildren& ArrangedChildren ) const
diff --git a/Engine/Source/Runtime/Slate/Private/Widgets/Views/SListPanel.h b/Engine/Source/Runtime/Slate/Private/Widgets/Views/SListPanel.h
--- a/Engine/Source/Runtime/Slate/Private/Widgets/Views/SListPanel.h
+++ b/Engine/Source/Runtime/Slate/Private/Widgets/Views/SListPanel.h
@@ -31,10 +31,11 @@ public:
 	};
 	
 	/** Make a new ListPanel::Slot  */
-	static FSlot& Slot();
+	static FSlot::FSlotArguments Slot();
 	
+	using FScopedWidgetSlotArguments = TPanelChildren<FSlot>::FScopedWidgetSlotArguments;
 	/** Add a slot to the ListPanel */
-	FSlot& AddSlot(int32 InsertAtIndex = INDEX_NONE);
+	FScopedWidgetSlotArguments AddSlot(int32 InsertAtIndex = INDEX_NONE);
 	
 	SLATE_BEGIN_ARGS( SListPanel )
 		: _ItemWidth(0)
@@ -47,6 +48,7 @@ public:
 			_Clipping = EWidgetClipping::ClipToBounds;
 		}
 	
+		SLATE_SLOT_ARGUMENT(FSlot, Slots)
 		SLATE_ATTRIBUTE( float, ItemWidth )
 		SLATE_ATTRIBUTE( float, ItemHeight )
 		SLATE_ATTRIBUTE( int32, NumDesiredItems )
diff --git a/Engine/Source/Runtime/Slate/Public/Widgets/Layout/SBox.h b/Engine/Source/Runtime/Slate/Public/Widgets/Layout/SBox.h
--- a/Engine/Source/Runtime/Slate/Public/Widgets/Layout/SBox.h
+++ b/Engine/Source/Runtime/Slate/Public/Widgets/Layout/SBox.h
@@ -134,10 +134,10 @@ protected:
 
 protected:
 
-	struct FBoxOneChildSlot : ::FOneSimpleMemberChild
+	struct FBoxOneChildSlot : ::TSingleWidgetChildrenWithBasicLayoutSlot<EInvalidateWidgetReason::None> // we want to add it to the Attribute descriptor
 	{
 		friend SBox;
-		using ::FOneSimpleMemberChild::FOneSimpleMemberChild;
+		using ::TSingleWidgetChildrenWithBasicLayoutSlot<EInvalidateWidgetReason::None>::TSingleWidgetChildrenWithBasicLayoutSlot;
 	};
 	FBoxOneChildSlot ChildSlot;
 
diff --git a/Engine/Source/Runtime/Slate/Public/Widgets/Layout/SDPIScaler.h b/Engine/Source/Runtime/Slate/Public/Widgets/Layout/SDPIScaler.h
--- a/Engine/Source/Runtime/Slate/Public/Widgets/Layout/SDPIScaler.h
+++ b/Engine/Source/Runtime/Slate/Public/Widgets/Layout/SDPIScaler.h
@@ -65,10 +65,10 @@ protected:
 	FSlateDeprecatedTAttribute<float> DPIScale;
 #endif
 
-	struct FDPIScalerOneChildSlot : ::FOneSimpleMemberChild
+	struct FDPIScalerOneChildSlot : ::TSingleWidgetChildrenWithBasicLayoutSlot<EInvalidateWidgetReason::None> // we want to add it to the Attribute descriptor
 	{
 		friend SDPIScaler;
-		using ::FOneSimpleMemberChild::FOneSimpleMemberChild;
+		using ::TSingleWidgetChildrenWithBasicLayoutSlot<EInvalidateWidgetReason::None>::TSingleWidgetChildrenWithBasicLayoutSlot;
 	};
 
 	/** The content being scaled. */
diff --git a/Engine/Source/Runtime/SlateCore/Private/Widgets/SCompoundWidget.cpp b/Engine/Source/Runtime/SlateCore/Private/Widgets/SCompoundWidget.cpp
--- a/Engine/Source/Runtime/SlateCore/Private/Widgets/SCompoundWidget.cpp
+++ b/Engine/Source/Runtime/SlateCore/Private/Widgets/SCompoundWidget.cpp
@@ -78,7 +78,7 @@ FSlateColor SCompoundWidget::GetForegroundColor() const
 }
 
 SCompoundWidget::SCompoundWidget()
-	: ChildSlot(*this)
+	: ChildSlot(this)
 	, ContentScaleAttribute( *this, FVector2D(1.0f,1.0f) )
 	, ColorAndOpacityAttribute( *this, FLinearColor::White )
 	, ForegroundColorAttribute( *this, FSlateColor::UseForeground() )
diff --git a/Engine/Source/Runtime/SlateCore/Public/Layout/BasicLayoutWidgetSlot.h b/Engine/Source/Runtime/SlateCore/Public/Layout/BasicLayoutWidgetSlot.h
--- a/Engine/Source/Runtime/SlateCore/Public/Layout/BasicLayoutWidgetSlot.h
+++ b/Engine/Source/Runtime/SlateCore/Public/Layout/BasicLayoutWidgetSlot.h
@@ -9,6 +9,8 @@
 #include "Misc/Optional.h"
 #include "Widgets/DeclarativeSyntaxSupport.h"
 
+#include <type_traits>
+
 PRAGMA_DISABLE_DEPRECATION_WARNINGS
 /** Mixin to add the alignment functionality to a base slot. */
 template <typename MixedIntoType>
@@ -57,14 +59,118 @@ protected:
 	}
 
 public:
-	// HAlign will be deprecated soon. Use SetVerticalAlignment, or if you are using ChildSlot, construct a new slot with FSlotArguments
+	UE_DEPRECATED(5.0, "HAlign is now deprecated. Use the FSlotArgument or the SetHorizontalAlignment function.")
+	MixedIntoType& HAlign(EHorizontalAlignment InHAlignment)
+	{
+		HAlignment = InHAlignment;
+		return *(static_cast<MixedIntoType*>(this));
+	}
+
+	UE_DEPRECATED(5.0, "VAlign is now deprecated. Use the FSlotArgument or the SetVerticalAlignment function.")
+	MixedIntoType& VAlign(EVerticalAlignment InVAlignment)
+	{
+		VAlignment = InVAlignment;
+		return *(static_cast<MixedIntoType*>(this));
+	}
+
+public:
+	void SetHorizontalAlignment(EHorizontalAlignment Alignment)
+	{
+		if (HAlignment != Alignment)
+		{
+			HAlignment = Alignment;
+			static_cast<MixedIntoType*>(this)->Invalidate(EInvalidateWidgetReason::Layout);
+		}
+	}
+
+	EHorizontalAlignment GetHorizontalAlignment() const
+	{
+		return HAlignment;
+	}
+
+	void SetVerticalAlignment(EVerticalAlignment Alignment)
+	{
+		if (VAlignment != Alignment)
+		{
+			VAlignment = Alignment;
+			static_cast<MixedIntoType*>(this)->Invalidate(EInvalidateWidgetReason::Layout);
+		}
+	}
+
+	EVerticalAlignment GetVerticalAlignment() const
+	{
+		return VAlignment;
+	}
+
+public:
+	/** Horizontal positioning of child within the allocated slot */
+	UE_DEPRECATED(5.0, "Direct access to HAlignment is now deprecated. Use the getter.")
+	EHorizontalAlignment HAlignment;
+	/** Vertical positioning of child within the allocated slot */
+	UE_DEPRECATED(5.0, "Direct access to VAlignment is now deprecated. Use the getter.")
+	EVerticalAlignment VAlignment;
+};
+PRAGMA_ENABLE_DEPRECATION_WARNINGS
+
+
+PRAGMA_DISABLE_DEPRECATION_WARNINGS
+/** Mixin to add the alignment functionality to a base slot that is also a single children. */
+template <typename MixedIntoType>
+class TAlignmentSingleWidgetSlotMixin
+{
+public:
+	template<typename WidgetType, typename V = typename std::enable_if<std::is_base_of<SWidget, WidgetType>::value>::type>
+	TAlignmentSingleWidgetSlotMixin(WidgetType& InParent)
+		: HAlignment(HAlign_Fill)
+		, VAlignment(VAlign_Fill)
+	{}
+
+	template<typename WidgetType, typename V = typename std::enable_if<std::is_base_of<SWidget, WidgetType>::value>::type>
+	TAlignmentSingleWidgetSlotMixin(WidgetType& InParent, const EHorizontalAlignment InHAlign, const EVerticalAlignment InVAlign)
+		: HAlignment(InHAlign)
+		, VAlignment(InVAlign)
+	{}
+
+public:
+	struct FSlotArgumentsMixin
+	{
+	private:
+		friend TAlignmentSingleWidgetSlotMixin;
+
+	public:
+		typename MixedIntoType::FSlotArguments& HAlign(EHorizontalAlignment InHAlignment)
+		{
+			_HAlignment = InHAlignment;
+			return static_cast<typename MixedIntoType::FSlotArguments&>(*this);
+		}
+
+		typename MixedIntoType::FSlotArguments& VAlign(EVerticalAlignment InVAlignment)
+		{
+			_VAlignment = InVAlignment;
+			return static_cast<typename MixedIntoType::FSlotArguments&>(*this);
+		}
+
+	private:
+		TOptional<EHorizontalAlignment> _HAlignment;
+		TOptional<EVerticalAlignment> _VAlignment;
+	};
+
+protected:
+	void ConstructMixin(FSlotArgumentsMixin&& InArgs)
+	{
+		HAlignment = InArgs._HAlignment.Get(HAlignment);
+		VAlignment = InArgs._VAlignment.Get(VAlignment);
+	}
+
+public:
+	// HAlign will be deprecated soon. Use SetVerticalAlignment or construct a new slot with FSlotArguments
 	MixedIntoType& HAlign(EHorizontalAlignment InHAlignment)
 	{
 		HAlignment = InHAlignment;
 		return *(static_cast<MixedIntoType*>(this));
 	}
 
-	// VAlign will be deprecated soon. Use SetVerticalAlignment, or if you are using ChildSlot, construct a new slot with FSlotArguments
+	// VAlign will be deprecated soon. Use SetVerticalAlignment or construct a new slot with FSlotArguments
 	MixedIntoType& VAlign(EVerticalAlignment InVAlignment)
 	{
 		VAlignment = InVAlignment;
@@ -165,31 +271,31 @@ protected:
 	}
 
 public:
-	// Padding will be deprecated soon. Use SetPadding, or if you are using ChildSlot, construct a new slot with FSlotArguments
+	UE_DEPRECATED(5.0, "Padding is now deprecated. Use the FSlotArgument or the SetPadding function.")
 	MixedIntoType& Padding(TAttribute<FMargin> InPadding)
 	{
-		SlotPadding = MoveTemp(InPadding);
+		SetPadding(MoveTemp(InPadding));
 		return *(static_cast<MixedIntoType*>(this));
 	}
 
-	// Padding will be deprecated soon. Use SetPadding, or if you are using ChildSlot, construct a new slot with FSlotArguments
+	UE_DEPRECATED(5.0, "Padding is now deprecated. Use the FSlotArgument or the SetPadding function.")
 	MixedIntoType& Padding(float Uniform)
 	{
-		SlotPadding = FMargin(Uniform);
+		SetPadding(FMargin(Uniform));
 		return *(static_cast<MixedIntoType*>(this));
 	}
 
-	// Padding will be deprecated soon. Use SetPadding, or if you are using ChildSlot, construct a new slot with FSlotArguments
+	UE_DEPRECATED(5.0, "Padding is now deprecated. Use the FSlotArgument or the SetPadding function.")
 	MixedIntoType& Padding(float Horizontal, float Vertical)
 	{
-		SlotPadding = FMargin(Horizontal, Vertical);
+		SetPadding(FMargin(Horizontal, Vertical));
 		return *(static_cast<MixedIntoType*>(this));
 	}
 
-	// Padding will be deprecated soon. Use SetPadding, or if you are using ChildSlot, construct a new slot with FSlotArguments
+	UE_DEPRECATED(5.0, "Padding is now deprecated. Use the FSlotArgument or the SetPadding function.")
 	MixedIntoType& Padding(float Left, float Top, float Right, float Bottom)
 	{
-		SlotPadding = FMargin(Left, Top, Right, Bottom);
+		SetPadding(FMargin(Left, Top, Right, Bottom));
 		return *(static_cast<MixedIntoType*>(this));
 	}
 
@@ -200,7 +306,10 @@ public:
 		static_cast<MixedIntoType*>(this)->Invalidate(EInvalidateWidgetReason::LayoutAndVolatility);
 	}
 
-	const FMargin& GetPadding() const { return SlotPadding.Get(); }
+	FMargin GetPadding() const
+	{
+		return SlotPadding.Get();
+	}
 
 public:
 	UE_DEPRECATED(5.0, "Direct access to SlotPadding is now deprecated. Use the setter or getter.")
@@ -209,6 +318,133 @@ public:
 PRAGMA_ENABLE_DEPRECATION_WARNINGS
 
 
+
+PRAGMA_DISABLE_DEPRECATION_WARNINGS
+/** Mixin to add the padding functionality to a base slot that is also a single children. */
+template <typename MixedIntoType, EInvalidateWidgetReason InPaddingInvalidationReason = EInvalidateWidgetReason::Layout>
+class TPaddingSingleWidgetSlotMixin
+{
+public:
+	template<typename WidgetType, typename V = typename std::enable_if<std::is_base_of<SWidget, WidgetType>::value>::type>
+	TPaddingSingleWidgetSlotMixin(WidgetType& InParent)
+		: SlotPaddingAttribute(InParent)
+	{}
+
+	template<typename WidgetType, typename V = typename std::enable_if<std::is_base_of<SWidget, WidgetType>::value>::type>
+	TPaddingSingleWidgetSlotMixin(WidgetType& InParent, const FMargin & Margin)
+		: SlotPaddingAttribute(InParent, Margin)
+	{}
+
+public:
+	struct FSlotArgumentsMixin
+	{
+	private:
+		friend TPaddingSingleWidgetSlotMixin;
+		TAttribute<FMargin> _Padding;
+
+	public:
+		typename MixedIntoType::FSlotArguments& Padding(TAttribute<FMargin> InPadding)
+		{
+			_Padding = MoveTemp(InPadding);
+			return static_cast<typename MixedIntoType::FSlotArguments&>(*this);
+		}
+
+		typename MixedIntoType::FSlotArguments& Padding(float Uniform)
+		{
+			_Padding = FMargin(Uniform);
+			return static_cast<typename MixedIntoType::FSlotArguments&>(*this);
+		}
+
+		typename MixedIntoType::FSlotArguments& Padding(float Horizontal, float Vertical)
+		{
+			_Padding = FMargin(Horizontal, Vertical);
+			return static_cast<typename MixedIntoType::FSlotArguments&>(*this);
+		}
+
+		typename MixedIntoType::FSlotArguments& Padding(float Left, float Top, float Right, float Bottom)
+		{
+			_Padding = FMargin(Left, Top, Right, Bottom);
+			return static_cast<typename MixedIntoType::FSlotArguments&>(*this);
+		}
+	};
+
+protected:
+	void ConstructMixin(FSlotArgumentsMixin&& InArgs)
+	{
+		if (InArgs._Padding.IsSet())
+		{
+			SWidget* OwnerWidget = static_cast<MixedIntoType*>(this)->GetOwnerWidget();
+			check(OwnerWidget);
+			SlotPaddingAttribute.Assign(*OwnerWidget, MoveTemp(InArgs._Padding));
+		}
+	}
+
+public:
+	// Padding will be deprecated soon. Use SetPadding or construct a new slot with FSlotArguments
+	MixedIntoType& Padding(TAttribute<FMargin> InPadding)
+	{
+		SetPadding(MoveTemp(InPadding));
+		return *(static_cast<MixedIntoType*>(this));
+	}
+
+	// Padding will be deprecated soon. Use SetPadding or construct a new slot with FSlotArguments
+	MixedIntoType& Padding(float Uniform)
+	{
+		SetPadding(FMargin(Uniform));
+		return *(static_cast<MixedIntoType*>(this));
+	}
+
+	// Padding will be deprecated soon. Use SetPadding or construct a new slot with FSlotArguments
+	MixedIntoType& Padding(float Horizontal, float Vertical)
+	{
+		SetPadding(FMargin(Horizontal, Vertical));
+		return *(static_cast<MixedIntoType*>(this));
+	}
+
+	// Padding will be deprecated soon. Use SetPadding or construct a new slot with FSlotArguments
+	MixedIntoType& Padding(float Left, float Top, float Right, float Bottom)
+	{
+		SetPadding(FMargin(Left, Top, Right, Bottom));
+		return *(static_cast<MixedIntoType*>(this));
+	}
+
+#if WITH_EDITORONLY_DATA
+	UE_DEPRECATED(5.0, "Direct access to SlotPadding is now deprecated. Use the setter or getter.")
+	FSlateDeprecatedTAttribute<FMargin> SlotPadding;
+#endif
+
+public:
+	void SetPadding(TAttribute<FMargin> InPadding)
+	{
+		SWidget* OwnerWidget = static_cast<MixedIntoType*>(this)->GetOwnerWidget();
+		check(OwnerWidget);
+		SlotPaddingAttribute.Assign(*OwnerWidget, MoveTemp(InPadding));
+	}
+
+	FMargin GetPadding() const
+	{
+		return SlotPaddingAttribute.Get();
+	}
+
+public:
+	using SlotPaddingInvalidationType = typename std::conditional<InPaddingInvalidationReason == EInvalidateWidgetReason::None, ::SlateAttributePrivate::FSlateAttributeNoInvalidationReason, TSlateAttributeInvalidationReason<InPaddingInvalidationReason>>::type;
+	using SlotPaddingAttributeType = SlateAttributePrivate::TSlateMemberAttribute<FMargin, SlotPaddingInvalidationType, TSlateAttributeComparePredicate<>>;
+	using SlotPaddingAttributeRefType = SlateAttributePrivate::TSlateMemberAttributeRef<SlotPaddingAttributeType>;
+
+	template<typename WidgetType, typename V = typename std::enable_if<std::is_base_of<SWidget, WidgetType>::value>::type>
+	SlotPaddingAttributeRefType GetSlotPaddingAttribute() const
+	{
+		WidgetType* Widget = static_cast<WidgetType*>(static_cast<MixedIntoType*>(this)->GetOwnerWidget());
+		check(Widget);
+		return SlotPaddingAttributeRefType(Widget->template SharedThis<WidgetType>(Widget), SlotPaddingAttribute);
+	}
+
+protected:
+	SlotPaddingAttributeType SlotPaddingAttribute;
+};
+PRAGMA_ENABLE_DEPRECATION_WARNINGS
+
+
 /** A templated basic slot that can be used by layout. */
 template <typename SlotType>
 class TBasicLayoutWidgetSlot : public TSlotBase<SlotType>
@@ -248,9 +484,9 @@ public:
 
 	void Construct(const FChildren& SlotOwner, FSlotArguments&& InArgs)
 	{
+		TSlotBase<SlotType>::Construct(SlotOwner, MoveTemp(InArgs));
 		TPaddingWidgetSlotMixin<SlotType>::ConstructMixin(SlotOwner, MoveTemp(InArgs));
 		TAlignmentWidgetSlotMixin<SlotType>::ConstructMixin(SlotOwner, MoveTemp(InArgs));
-		TSlotBase<SlotType>::Construct(SlotOwner, MoveTemp(InArgs));
 	}
 };
 
diff --git a/Engine/Source/Runtime/SlateCore/Public/Layout/Children.h b/Engine/Source/Runtime/SlateCore/Public/Layout/Children.h
--- a/Engine/Source/Runtime/SlateCore/Public/Layout/Children.h
+++ b/Engine/Source/Runtime/SlateCore/Public/Layout/Children.h
@@ -584,41 +584,54 @@ public:
 
 
 /** A FChildren that has only one child and support alignment and padding. */
-class FSingleWidgetChildrenWithBasicLayoutSlot : public TSingleWidgetChildrenWithSlot<FSingleWidgetChildrenWithBasicLayoutSlot>
-	, public TPaddingWidgetSlotMixin<FSingleWidgetChildrenWithBasicLayoutSlot>
-	, public TAlignmentWidgetSlotMixin<FSingleWidgetChildrenWithBasicLayoutSlot>
+template<EInvalidateWidgetReason InPaddingInvalidationReason = EInvalidateWidgetReason::Layout>
+class TSingleWidgetChildrenWithBasicLayoutSlot : public TSingleWidgetChildrenWithSlot<TSingleWidgetChildrenWithBasicLayoutSlot<InPaddingInvalidationReason>>
+	, public TPaddingSingleWidgetSlotMixin<TSingleWidgetChildrenWithBasicLayoutSlot<InPaddingInvalidationReason>, InPaddingInvalidationReason>
+	, public TAlignmentSingleWidgetSlotMixin<TSingleWidgetChildrenWithBasicLayoutSlot<InPaddingInvalidationReason>>
 {
+private:
+	using ParentType = TSingleWidgetChildrenWithSlot<TSingleWidgetChildrenWithBasicLayoutSlot<InPaddingInvalidationReason>>;
+	using PaddingMixinType = TPaddingSingleWidgetSlotMixin<TSingleWidgetChildrenWithBasicLayoutSlot<InPaddingInvalidationReason>, InPaddingInvalidationReason>;
+	using AlignmentMixinType = TAlignmentSingleWidgetSlotMixin<TSingleWidgetChildrenWithBasicLayoutSlot<InPaddingInvalidationReason>>;
 public:
-	FSingleWidgetChildrenWithBasicLayoutSlot(SWidget* InOwner)
-		: TSingleWidgetChildrenWithSlot<FSingleWidgetChildrenWithBasicLayoutSlot>(InOwner)
-		, TPaddingWidgetSlotMixin<FSingleWidgetChildrenWithBasicLayoutSlot>()
-		, TAlignmentWidgetSlotMixin<FSingleWidgetChildrenWithBasicLayoutSlot>(HAlign_Fill, VAlign_Fill)
+	template<typename WidgetType, typename V = typename std::enable_if<std::is_base_of<SWidget, WidgetType>::value>::type>
+	TSingleWidgetChildrenWithBasicLayoutSlot(WidgetType* InOwner)
+		: ParentType(InOwner)
+		, PaddingMixinType(*InOwner)
+		, AlignmentMixinType(*InOwner, HAlign_Fill, VAlign_Fill)
 	{
 	}
 
-	FSingleWidgetChildrenWithBasicLayoutSlot(SWidget* InOwner, const EHorizontalAlignment InHAlign, const EVerticalAlignment InVAlign)
-		: TSingleWidgetChildrenWithSlot<FSingleWidgetChildrenWithBasicLayoutSlot>(InOwner)
-		, TPaddingWidgetSlotMixin<FSingleWidgetChildrenWithBasicLayoutSlot>()
-		, TAlignmentWidgetSlotMixin<FSingleWidgetChildrenWithBasicLayoutSlot>(InHAlign, InVAlign)
+	template<typename WidgetType, typename V = typename std::enable_if<std::is_base_of<SWidget, WidgetType>::value>::type>
+	TSingleWidgetChildrenWithBasicLayoutSlot(WidgetType* InOwner, const EHorizontalAlignment InHAlign, const EVerticalAlignment InVAlign)
+		: ParentType(InOwner)
+		, PaddingMixinType(*InOwner)
+		, AlignmentMixinType(*InOwner, InHAlign, InVAlign)
 	{
 	}
 
-	FSingleWidgetChildrenWithBasicLayoutSlot(std::nullptr_t) = delete;
-	FSingleWidgetChildrenWithBasicLayoutSlot(std::nullptr_t, const EHorizontalAlignment InHAlign, const EVerticalAlignment InVAlign) = delete;
+	TSingleWidgetChildrenWithBasicLayoutSlot(std::nullptr_t) = delete;
+	TSingleWidgetChildrenWithBasicLayoutSlot(std::nullptr_t, const EHorizontalAlignment InHAlign, const EVerticalAlignment InVAlign) = delete;
 
 public:
-	SLATE_SLOT_BEGIN_ARGS_TwoMixins(FSingleWidgetChildrenWithBasicLayoutSlot, TSingleWidgetChildrenWithSlot<FSingleWidgetChildrenWithBasicLayoutSlot>, TPaddingWidgetSlotMixin<FSingleWidgetChildrenWithBasicLayoutSlot>, TAlignmentWidgetSlotMixin<FSingleWidgetChildrenWithBasicLayoutSlot>)
+	SLATE_SLOT_BEGIN_ARGS_TwoMixins(TSingleWidgetChildrenWithBasicLayoutSlot, ParentType, PaddingMixinType, AlignmentMixinType)
 	SLATE_SLOT_END_ARGS()
 
 	void Construct(FSlotArguments&& InArgs)
 	{
-		TSingleWidgetChildrenWithSlot<FSingleWidgetChildrenWithBasicLayoutSlot>::Construct(MoveTemp(InArgs));
-		TPaddingWidgetSlotMixin<FSingleWidgetChildrenWithBasicLayoutSlot>::ConstructMixin(*this, MoveTemp(InArgs));
-		TAlignmentWidgetSlotMixin<FSingleWidgetChildrenWithBasicLayoutSlot>::ConstructMixin(*this, MoveTemp(InArgs));
+		ParentType::Construct(MoveTemp(InArgs));
+		PaddingMixinType::ConstructMixin(MoveTemp(InArgs));
+		AlignmentMixinType::ConstructMixin( MoveTemp(InArgs));
 	}
 };
 
 
+class FSingleWidgetChildrenWithBasicLayoutSlot : public TSingleWidgetChildrenWithBasicLayoutSlot<>
+{
+	using TSingleWidgetChildrenWithBasicLayoutSlot<>::TSingleWidgetChildrenWithBasicLayoutSlot;
+};
+
+
 /** A slot that support alignment of content and padding */
 class UE_DEPRECATED(5.0, "FSimpleSlot is deprecated because it got confused from FChildren with FSlot. Use FSingleWidgetChildrenWithSimpleSlot.")
 FSimpleSlot : public FSingleWidgetChildrenWithBasicLayoutSlot
@@ -674,6 +687,7 @@ public:
 	}
 
 public:
+	UE_DEPRECATED(5.0, "Add a slot directly has been deprecated. use the FSlotArgument to create a new slot")
 	int32 Add( SlotType* Slot )
 	{
 		int32 Index = Children.Add(TUniquePtr<SlotType>(Slot));
@@ -746,6 +760,7 @@ public:
 		Children = MoveTemp(ChildrenCopy);
 	}
 
+	UE_DEPRECATED(5.0, "Insert a slot directly has been deprecated. use the FSlotArgument to create a new slot")
 	void Insert(SlotType* Slot, int32 Index)
 	{
 		check(Slot);
@@ -1175,74 +1190,6 @@ private:
 };
 
 
-/** A FChildren that support only one slot and alignment of content and padding */
-class FOneSimpleMemberChild : public TSingleWidgetChildrenWithSlot<FOneSimpleMemberChild>
-	, public TAlignmentWidgetSlotMixin<FOneSimpleMemberChild>
-{
-public:
-	template<typename WidgetType, typename V = typename std::enable_if<std::is_base_of<SWidget, WidgetType>::value>::type>
-	FOneSimpleMemberChild(WidgetType& InParent)
-		: TSingleWidgetChildrenWithSlot<FOneSimpleMemberChild>(&InParent)
-		, TAlignmentWidgetSlotMixin<FOneSimpleMemberChild>(HAlign_Fill, VAlign_Fill)
-		, SlotPaddingAttribute(InParent)
-	{
-	}
-
-	//~ TSlateAttribute cannot be copied
-	FOneSimpleMemberChild(const FOneSimpleMemberChild&) = delete;
-	FOneSimpleMemberChild& operator=(const FOneSimpleMemberChild&) = delete;
-
-	FOneSimpleMemberChild& Padding(TAttribute<FMargin> InPadding)
-	{
-		SlotPaddingAttribute.Assign(FChildren::GetOwner(), MoveTemp(InPadding));
-		return *this;
-	}
-
-	FOneSimpleMemberChild& Padding(float Uniform)
-	{
-		SlotPaddingAttribute.Set(FChildren::GetOwner(), FMargin(Uniform));
-		return *this;
-	}
-
-	FOneSimpleMemberChild& Padding(float Horizontal, float Vertical)
-	{
-		SlotPaddingAttribute.Set(FChildren::GetOwner(), FMargin(Horizontal, Vertical));
-		return *this;
-	}
-
-	FOneSimpleMemberChild& Padding(float Left, float Top, float Right, float Bottom)
-	{
-		SlotPaddingAttribute.Set(FChildren::GetOwner(), FMargin(Left, Top, Right, Bottom));
-		return *this;
-	}
-
-	void SetPadding(TAttribute<FMargin> InPadding)
-	{
-		SlotPaddingAttribute.Assign(FChildren::GetOwner(), MoveTemp(InPadding));
-	}
-	const FMargin& GetPadding() const { return SlotPaddingAttribute.Get(); }
-
-#if WITH_EDITORONLY_DATA
-	UE_DEPRECATED(5.0, "Direct access to SlotPadding is now deprecated. Use the setter or getter.")
-	FSlateDeprecatedTAttribute<FMargin> SlotPadding;
-#endif
-
-public:
-	using SlotPaddingAttributeType = SlateAttributePrivate::TSlateMemberAttribute<FMargin, ::SlateAttributePrivate::FSlateAttributeNoInvalidationReason, TSlateAttributeComparePredicate<>>;
-	using SlotPaddingAttributeRefType = SlateAttributePrivate::TSlateMemberAttributeRef<SlotPaddingAttributeType>;
-
-	template<typename WidgetType, typename V = typename std::enable_if<std::is_base_of<SWidget, WidgetType>::value>::type>
-	SlotPaddingAttributeRefType GetSlotPaddingAttribute() const
-	{
-		WidgetType&  Widget = static_cast<WidgetType&>(FChildren::GetOwner());
-		return SlotPaddingAttributeRefType(Widget.template SharedThis<WidgetType>(&Widget), SlotPaddingAttribute);
-	}
-
-protected:
-	SlotPaddingAttributeType SlotPaddingAttribute;
-};
-
-
 /** Required to implement GetChildren() in a way that can dynamically return the currently active child. */
 template<typename SlotType>
 class TOneDynamicChild : public FChildren
diff --git a/Engine/Source/Runtime/SlateCore/Public/SlotBase.h b/Engine/Source/Runtime/SlateCore/Public/SlotBase.h
--- a/Engine/Source/Runtime/SlateCore/Public/SlotBase.h
+++ b/Engine/Source/Runtime/SlateCore/Public/SlotBase.h
@@ -4,6 +4,7 @@
 
 #include "CoreMinimal.h"
 #include "Misc/Attribute.h"
+#include "Types/SlateAttribute.h"
 #include "Widgets/InvalidateWidgetReason.h"
 
 class FChildren;
diff --git a/Engine/Source/Runtime/SlateCore/Public/Widgets/DeclarativeSyntaxSupport.h b/Engine/Source/Runtime/SlateCore/Public/Widgets/DeclarativeSyntaxSupport.h
--- a/Engine/Source/Runtime/SlateCore/Public/Widgets/DeclarativeSyntaxSupport.h
+++ b/Engine/Source/Runtime/SlateCore/Public/Widgets/DeclarativeSyntaxSupport.h
@@ -263,6 +263,7 @@ template<typename WidgetType> struct TSlateBaseNamedArgs;
  * in order to add support for slots.
  */
 #define SLATE_SUPPORTS_SLOT( SlotType ) \
+		UE_DEPRECATED(5.0, "SLATE_SUPPORTS_SLOT is deprecated. Use SLATE_SLOT_ARGUMENT") \
 		TArray< SlotType* > Slots; \
 		WidgetArgsType& operator + (SlotType& SlotToAdd) \
 		{ \
diff --git a/Engine/Source/Runtime/SlateCore/Public/Widgets/SCompoundWidget.h b/Engine/Source/Runtime/SlateCore/Public/Widgets/SCompoundWidget.h
--- a/Engine/Source/Runtime/SlateCore/Public/Widgets/SCompoundWidget.h
+++ b/Engine/Source/Runtime/SlateCore/Public/Widgets/SCompoundWidget.h
@@ -102,10 +102,10 @@ protected:
 	/** Disallow public construction */
 	SCompoundWidget();
 
-	struct FCompoundWidgetOneChildSlot : ::FOneSimpleMemberChild
+	struct FCompoundWidgetOneChildSlot : ::TSingleWidgetChildrenWithBasicLayoutSlot<EInvalidateWidgetReason::None>
 	{
 		friend SCompoundWidget;
-		using ::FOneSimpleMemberChild::FOneSimpleMemberChild;
+		using ::TSingleWidgetChildrenWithBasicLayoutSlot<EInvalidateWidgetReason::None>::TSingleWidgetChildrenWithBasicLayoutSlot;
 	};
 
 	/** The slot that contains this widget's descendants.*/
diff --git a/Engine/Source/Runtime/UMG/Private/Components/HorizontalBoxSlot.cpp b/Engine/Source/Runtime/UMG/Private/Components/HorizontalBoxSlot.cpp
--- a/Engine/Source/Runtime/UMG/Private/Components/HorizontalBoxSlot.cpp
+++ b/Engine/Source/Runtime/UMG/Private/Components/HorizontalBoxSlot.cpp
@@ -40,7 +40,7 @@ void UHorizontalBoxSlot::SetPadding(FMargin InPadding)
 	Padding = InPadding;
 	if ( Slot )
 	{
-		Slot->Padding(InPadding);
+		Slot->SetPadding(InPadding);
 	}
 }
 
