Slate: Re-introduce SlateAttribute.OnValueChanged(). N.B. You are only allowed to invalidate attribute in the callback that doesn't affect the child order or the visibility of the widget.

diff --git a/Engine/Source/Runtime/SlateCore/Private/FastUpdate/SlateInvalidationRoot.cpp b/Engine/Source/Runtime/SlateCore/Private/FastUpdate/SlateInvalidationRoot.cpp
--- a/Engine/Source/Runtime/SlateCore/Private/FastUpdate/SlateInvalidationRoot.cpp
+++ b/Engine/Source/Runtime/SlateCore/Private/FastUpdate/SlateInvalidationRoot.cpp
@@ -271,6 +271,15 @@ void FSlateInvalidationRoot::InvalidateWidget(FWidgetProxy& Proxy, EInvalidateWi
 {
 	ensureMsgf(bProcessingChildOrderInvalidation == false, TEXT("A widget got invalidated while building the childorder."));
 
+	if (bProcessingAttributeUpdate)
+	{
+		if (EnumHasAnyFlags(InvalidateReason, EInvalidateWidgetReason::AttributeRegistration | EInvalidateWidgetReason::ChildOrder | EInvalidateWidgetReason::Visibility))
+		{
+			ensureMsgf(false, TEXT("An invalid invalidation occurred while processing the widget attributes. That may result in an infinit loop."));
+			return;
+		}
+	}
+
 	if (!bNeedsSlowPath)
 	{
 		Proxy.CurrentInvalidateReason |= InvalidateReason;
@@ -714,8 +723,19 @@ void FSlateInvalidationRoot::ProcessAttributeUpdate()
 					ensureAlwaysMsgf(!GSlateInvalidationRootVerifySlateAttribute || InvalidationWidget.bDebug_AttributeUpdated == false, TEXT("Attribute should only be updated once per frame."));
 					InvalidationWidget.bDebug_AttributeUpdated = true;
 #endif
+
+#if WITH_SLATE_DEBUGGING
+					FSlateInvalidationWidgetVisibility PreviousVisibility = InvalidationWidget.Visibility;
+					FSlateInvalidationWidgetIndex PreviousLeafMostChildIndex = InvalidationWidget.LeafMostChildIndex;
+#endif
+
 					FSlateAttributeMetaData::UpdateExceptVisibilityAttributes(*WidgetPtr, FSlateAttributeMetaData::EInvalidationPermission::AllowInvalidation);
 					AttributeItt.Advance();
+
+#if WITH_SLATE_DEBUGGING
+					ensureMsgf(PreviousVisibility == InvalidationWidget.Visibility, TEXT("The visibility of widget '%s' doens't match the previous visibility after the attribute update."), *FReflectionMetaData::GetWidgetDebugInfo(WidgetPtr));
+					ensureMsgf(PreviousLeafMostChildIndex == InvalidationWidget.LeafMostChildIndex, TEXT("The number of child of widget '%s' doens't match the previous count after the attribute update."), *FReflectionMetaData::GetWidgetDebugInfo(WidgetPtr));
+#endif
 				}
 				else
 				{
diff --git a/Engine/Source/Runtime/SlateCore/Private/Test/SlateAttributeTest.cpp b/Engine/Source/Runtime/SlateCore/Private/Test/SlateAttributeTest.cpp
--- a/Engine/Source/Runtime/SlateCore/Private/Test/SlateAttributeTest.cpp
+++ b/Engine/Source/Runtime/SlateCore/Private/Test/SlateAttributeTest.cpp
@@ -136,7 +136,7 @@ public:
 	{
 	}
 
-	void Construct(const FArguments& InArgs) {}
+	void Construct(const FArguments& InArgs){}
 
 	TSlateAttribute<int32, EInvalidateWidgetReason::Layout> IntAttributeH;
 	TSlateAttribute<int32> IntAttributeI;
@@ -178,7 +178,14 @@ public:
 	{
 	}
 
-	void Construct(const FArguments& InArgs) {}
+	using SWidget::IsConstructed;
+
+	void Construct(const FArguments& InArgs)
+	{
+		IntAttributeA.Assign(*this, 88);
+		IntAttributeB.Assign(*this, 88);
+		IntAttributeC.Assign(*this, 88);
+	}
 	virtual FVector2D ComputeDesiredSize(float) const override { return FVector2D{ 100, 100 }; }
 	virtual int32 OnPaint(const FPaintArgs& Args, const FGeometry& AllottedGeometry, const FSlateRect& MyCullingRect, FSlateWindowElementList& OutDrawElements, int32 LayerId, const FWidgetStyle& InWidgetStyle, bool bParentEnabled) const override
 	{
@@ -193,11 +200,93 @@ public:
 	TSlateAttribute<int32> IntAttributeA;
 	TSlateAttribute<int32> IntAttributeB;
 	TSlateAttribute<int32> IntAttributeC;
-	int32 CallbackValueA;
-	int32 CallbackValueB;
-	int32 CallbackValueC;
+	int32 CallbackValueA = 555;
+	int32 CallbackValueB = 555;
+	int32 CallbackValueC = 555;
+	int32 CallbackIsConstructedValueA = 0;
+	int32 CallbackIsConstructedValueB = 0;
+	int32 CallbackIsConstructedValueC = 0;
+};
+
+SLATE_IMPLEMENT_WIDGET(SAttributeLeftWidget_OnInvalidationParent)
+void SAttributeLeftWidget_OnInvalidationParent::PrivateRegisterAttributes(FSlateAttributeInitializer& AttributeInitializer)
+{
+	SLATE_ADD_MEMBER_ATTRIBUTE_DEFINITION(AttributeInitializer, IntAttributeA, EInvalidateWidgetReason::Layout)
+		.OnValueChanged(FSlateAttributeDescriptor::FAttributeValueChangedDelegate::CreateLambda([](SWidget& Widget)
+			{
+				SAttributeLeftWidget_OnInvalidationParent& Instance = static_cast<SAttributeLeftWidget_OnInvalidationParent&>(Widget);
+				Instance.CallbackValueA = Instance.IntAttributeA.Get();
+				Instance.CallbackIsConstructedValueA += static_cast<SAttributeLeftWidget_OnInvalidationParent&>(Widget).IsConstructed() ? 1 : 0;
+			}));
+	SLATE_ADD_MEMBER_ATTRIBUTE_DEFINITION(AttributeInitializer, IntAttributeB, EInvalidateWidgetReason::Layout)
+		.OnValueChanged(FSlateAttributeDescriptor::FAttributeValueChangedDelegate::CreateLambda([](SWidget& Widget)
+			{
+				SAttributeLeftWidget_OnInvalidationParent& Instance = static_cast<SAttributeLeftWidget_OnInvalidationParent&>(Widget);
+				Instance.CallbackValueB = Instance.IntAttributeB.Get();
+				Instance.CallbackIsConstructedValueB += static_cast<SAttributeLeftWidget_OnInvalidationParent&>(Widget).IsConstructed() ? 1 : 0;
+			}));
+	SLATE_ADD_MEMBER_ATTRIBUTE_DEFINITION(AttributeInitializer, IntAttributeC, EInvalidateWidgetReason::Layout)
+		.OnValueChanged(FSlateAttributeDescriptor::FAttributeValueChangedDelegate::CreateLambda([](SWidget& Widget)
+			{
+				SAttributeLeftWidget_OnInvalidationParent& Instance = static_cast<SAttributeLeftWidget_OnInvalidationParent&>(Widget);
+				Instance.CallbackValueC = Instance.IntAttributeC.Get();
+				Instance.CallbackIsConstructedValueC += static_cast<SAttributeLeftWidget_OnInvalidationParent&>(Widget).IsConstructed() ? 1 : 0;
+			}));
+}
+
+class SAttributeLeftWidget_OnInvalidationChild : public SAttributeLeftWidget_OnInvalidationParent
+{
+	SLATE_DECLARE_WIDGET(SAttributeLeftWidget_OnInvalidationChild, SAttributeLeftWidget_OnInvalidationParent)
+public:
+	SLATE_BEGIN_ARGS(SAttributeLeftWidget_OnInvalidationChild) {}
+	SLATE_END_ARGS()
+
+	void Construct(const FArguments& InArgs) {}
+	virtual FVector2D ComputeDesiredSize(float) const override { return FVector2D{ 100, 100 }; }
+	virtual int32 OnPaint(const FPaintArgs& Args, const FGeometry& AllottedGeometry, const FSlateRect& MyCullingRect, FSlateWindowElementList& OutDrawElements, int32 LayerId, const FWidgetStyle& InWidgetStyle, bool bParentEnabled) const override
+	{
+		return LayerId;
+	}
+
+	virtual void SetCallbackValue(int32 Value) override
+	{
+		Super::SetCallbackValue(Value);
+		OverrideCallbackValueA = Value;
+		OverrideCallbackValueB = Value;
+		OverrideCallbackValueC = Value;
+	}
+
+	int32 OverrideCallbackValueA;
+	int32 OverrideCallbackValueB;
+	int32 OverrideCallbackValueC;
 };
 
+SLATE_IMPLEMENT_WIDGET(SAttributeLeftWidget_OnInvalidationChild)
+void SAttributeLeftWidget_OnInvalidationChild::PrivateRegisterAttributes(FSlateAttributeInitializer& AttributeInitializer)
+{
+	AttributeInitializer.OverrideOnValueChanged("IntAttributeA"
+		, FSlateAttributeDescriptor::ECallbackOverrideType::ReplacePrevious
+		, FSlateAttributeDescriptor::FAttributeValueChangedDelegate::CreateLambda([](SWidget& Widget)
+			{
+				SAttributeLeftWidget_OnInvalidationChild& Instance = static_cast<SAttributeLeftWidget_OnInvalidationChild&>(Widget);
+				Instance.OverrideCallbackValueA = Instance.IntAttributeA.Get();
+			}));
+	AttributeInitializer.OverrideOnValueChanged("IntAttributeB"
+		, FSlateAttributeDescriptor::ECallbackOverrideType::ExecuteAfterPrevious
+		, FSlateAttributeDescriptor::FAttributeValueChangedDelegate::CreateLambda([](SWidget& Widget)
+			{
+				SAttributeLeftWidget_OnInvalidationChild& Instance = static_cast<SAttributeLeftWidget_OnInvalidationChild&>(Widget);
+				Instance.OverrideCallbackValueB = Instance.CallbackValueB;
+			}));
+	AttributeInitializer.OverrideOnValueChanged("IntAttributeC"
+		, FSlateAttributeDescriptor::ECallbackOverrideType::ExecuteBeforePrevious
+		, FSlateAttributeDescriptor::FAttributeValueChangedDelegate::CreateLambda([](SWidget& Widget)
+			{
+				SAttributeLeftWidget_OnInvalidationChild& Instance = static_cast<SAttributeLeftWidget_OnInvalidationChild&>(Widget);
+				Instance.OverrideCallbackValueC = Instance.CallbackValueC;
+			}));
+}
+
 }
 }
 }
@@ -787,6 +876,116 @@ bool FSlateAttributeTest::RunTest(const FString& Parameters)
 		}
 	}
 
+	// Test OnInvalidation
+	{
+		{
+			TSharedRef<UE::Slate::Private::SAttributeLeftWidget_OnInvalidationParent> WidgetParent
+				= SNew(UE::Slate::Private::SAttributeLeftWidget_OnInvalidationParent);
+
+
+			AddErrorIfFalse(WidgetParent->CallbackIsConstructedValueA == 0, TEXT("The callback for value A was triggered while the widget was constructed."));
+			AddErrorIfFalse(WidgetParent->CallbackIsConstructedValueB == 0, TEXT("The callback for value B was triggered while the widget was constructed."));
+			AddErrorIfFalse(WidgetParent->CallbackIsConstructedValueC == 0, TEXT("The callback for value C was triggered while the widget was constructed."));
+
+			WidgetParent->SetCallbackValue(10);
+			WidgetParent->MarkPrepassAsDirty();
+			WidgetParent->SlatePrepass(1.f);
+			AddErrorIfFalse(WidgetParent->CallbackValueA == 10, TEXT("The callback for value A was triggered."));
+			AddErrorIfFalse(WidgetParent->CallbackValueB == 10, TEXT("The callback for value B was triggered."));
+			AddErrorIfFalse(WidgetParent->CallbackValueC == 10, TEXT("The callback for value C was triggered."));
+
+			AddErrorIfFalse(WidgetParent->CallbackIsConstructedValueA == 0, TEXT("The callback for value A was triggered while the widget was constructed."));
+			AddErrorIfFalse(WidgetParent->CallbackIsConstructedValueB == 0, TEXT("The callback for value B was triggered while the widget was constructed."));
+			AddErrorIfFalse(WidgetParent->CallbackIsConstructedValueC == 0, TEXT("The callback for value C was triggered while the widget was constructed."));
+		}
+		{
+			TSharedRef<UE::Slate::Private::SAttributeLeftWidget_OnInvalidationChild> WidgetChild = SNew(UE::Slate::Private::SAttributeLeftWidget_OnInvalidationChild);
+
+			AddErrorIfFalse(WidgetChild->CallbackIsConstructedValueA == 0, TEXT("The callback for value A was triggered while the widget was constructed."));
+			AddErrorIfFalse(WidgetChild->CallbackIsConstructedValueB == 0, TEXT("The callback for value B was triggered while the widget was constructed."));
+			AddErrorIfFalse(WidgetChild->CallbackIsConstructedValueC == 0, TEXT("The callback for value C was triggered while the widget was constructed."));
+
+			WidgetChild->SetCallbackValue(10);
+			WidgetChild->MarkPrepassAsDirty();
+			WidgetChild->SlatePrepass(1.f);
+			AddErrorIfFalse(WidgetChild->CallbackValueA == 10, TEXT("The callback for value A was triggered."));
+			AddErrorIfFalse(WidgetChild->CallbackValueB == 10, TEXT("The callback for value B was triggered."));
+			AddErrorIfFalse(WidgetChild->CallbackValueC == 10, TEXT("The callback for value C was triggered."));
+			AddErrorIfFalse(WidgetChild->OverrideCallbackValueA == 10, TEXT("The callback for override value A was triggered."));
+			AddErrorIfFalse(WidgetChild->OverrideCallbackValueB == 10, TEXT("The callback for override value B was triggered."));
+			AddErrorIfFalse(WidgetChild->OverrideCallbackValueC == 10, TEXT("The callback for override value C was triggered."));
+
+			AddErrorIfFalse(WidgetChild->CallbackIsConstructedValueA == 0, TEXT("The callback for value A was triggered while the widget was constructed."));
+			AddErrorIfFalse(WidgetChild->CallbackIsConstructedValueB == 0, TEXT("The callback for value B was triggered while the widget was constructed."));
+			AddErrorIfFalse(WidgetChild->CallbackIsConstructedValueC == 0, TEXT("The callback for value C was triggered while the widget was constructed."));
+		}
+
+		{
+			TSharedRef<UE::Slate::Private::SAttributeLeftWidget_OnInvalidationParent> WidgetParent = SNew(UE::Slate::Private::SAttributeLeftWidget_OnInvalidationParent);
+
+			AddErrorIfFalse(WidgetParent->CallbackIsConstructedValueA == 0, TEXT("The callback for value A was triggered while the widget was constructed."));
+			AddErrorIfFalse(WidgetParent->CallbackIsConstructedValueB == 0, TEXT("The callback for value B was triggered while the widget was constructed."));
+			AddErrorIfFalse(WidgetParent->CallbackIsConstructedValueC == 0, TEXT("The callback for value C was triggered while the widget was constructed."));
+
+			WidgetParent->SetCallbackValue(10);
+			WidgetParent->SetAttributeA(4);
+			AddErrorIfFalse(WidgetParent->CallbackValueA == 4, TEXT("The callback for value A was not triggered."));
+			AddErrorIfFalse(WidgetParent->CallbackValueB == 10, TEXT("The callback for value B was triggered."));
+			AddErrorIfFalse(WidgetParent->CallbackValueC == 10, TEXT("The callback for value C was triggered."));
+			WidgetParent->SetAttributeB(5);
+			AddErrorIfFalse(WidgetParent->CallbackValueB == 5, TEXT("The callback for value B was not triggered."));
+			AddErrorIfFalse(WidgetParent->CallbackValueC == 10, TEXT("The callback for value C was triggered."));
+			WidgetParent->SetAttributeC(6);
+			AddErrorIfFalse(WidgetParent->CallbackValueC == 6, TEXT("The callback for value C was not triggered."));
+		}
+
+		{
+			TSharedRef<UE::Slate::Private::SAttributeLeftWidget_OnInvalidationChild> WidgetChild = SNew(UE::Slate::Private::SAttributeLeftWidget_OnInvalidationChild);
+			WidgetChild->SetCallbackValue(10);
+			WidgetChild->CallbackValueB = 15;
+			WidgetChild->SetAttributeA(4);
+			AddErrorIfFalse(WidgetChild->CallbackValueA == 10, TEXT("The callback for value A was triggered."));
+			AddErrorIfFalse(WidgetChild->OverrideCallbackValueA == 4, TEXT("The callback for override value A was not triggered."));
+			AddErrorIfFalse(WidgetChild->CallbackValueB == 15, TEXT("The callback for value B was triggered."));
+			AddErrorIfFalse(WidgetChild->OverrideCallbackValueB == 10, TEXT("The callback for override value B was triggered."));
+			AddErrorIfFalse(WidgetChild->CallbackValueC == 10, TEXT("The callback for value C was triggered."));
+			AddErrorIfFalse(WidgetChild->OverrideCallbackValueC == 10, TEXT("The callback for override value C was triggered."));
+			WidgetChild->SetAttributeB(5);
+			AddErrorIfFalse(WidgetChild->CallbackValueB == 5, TEXT("The callback for value B was not triggered."));
+			AddErrorIfFalse(WidgetChild->OverrideCallbackValueB == 15, TEXT("The callback for override value B was not triggered."));
+			AddErrorIfFalse(WidgetChild->CallbackValueC == 10, TEXT("The callback for value C was triggered."));
+			AddErrorIfFalse(WidgetChild->OverrideCallbackValueC == 10, TEXT("The callback for override value C was triggered."));
+			WidgetChild->SetAttributeC(6);
+			AddErrorIfFalse(WidgetChild->CallbackValueC == 6, TEXT("The callback for value C was not triggered."));
+			AddErrorIfFalse(WidgetChild->OverrideCallbackValueC == 6, TEXT("The callback for override value C was not triggered."));
+		}
+		{
+			TSharedRef<UE::Slate::Private::SAttributeLeftWidget_OnInvalidationChild> WidgetChild = SNew(UE::Slate::Private::SAttributeLeftWidget_OnInvalidationChild);
+			WidgetChild->SetCallbackValue(10);
+			WidgetChild->CallbackValueB = 15;
+			WidgetChild->SetAttributeA(MakeAttributeLambda([](){ return 4; }));
+			WidgetChild->SetAttributeB(MakeAttributeLambda([](){ return 5; }));
+			WidgetChild->SetAttributeC(MakeAttributeLambda([](){ return 6; }));
+
+			AddErrorIfFalse(WidgetChild->CallbackValueA == 10, TEXT("The callback for value A was triggered."));
+			AddErrorIfFalse(WidgetChild->OverrideCallbackValueA == 10, TEXT("The callback for override value A was triggered."));
+			AddErrorIfFalse(WidgetChild->CallbackValueB == 15, TEXT("The callback for value B was triggered."));
+			AddErrorIfFalse(WidgetChild->OverrideCallbackValueB == 10, TEXT("The callback for override value B was triggered."));
+			AddErrorIfFalse(WidgetChild->CallbackValueC == 10, TEXT("The callback for value C was triggered."));
+			AddErrorIfFalse(WidgetChild->OverrideCallbackValueC == 10, TEXT("The callback for override value C was triggered."));
+
+			WidgetChild->MarkPrepassAsDirty();
+			WidgetChild->SlatePrepass(1.f);
+
+			AddErrorIfFalse(WidgetChild->CallbackValueA == 10, TEXT("The callback for value A was triggered."));
+			AddErrorIfFalse(WidgetChild->OverrideCallbackValueA == 4, TEXT("The callback for override value A was not triggered."));
+			AddErrorIfFalse(WidgetChild->CallbackValueB == 5, TEXT("The callback for value B was not triggered."));
+			AddErrorIfFalse(WidgetChild->OverrideCallbackValueB == 15, TEXT("The callback for override value B was not triggered."));
+			AddErrorIfFalse(WidgetChild->CallbackValueC == 6, TEXT("The callback for value C was not triggered."));
+			AddErrorIfFalse(WidgetChild->OverrideCallbackValueC == 6, TEXT("The callback for override value C was not triggered."));
+		}
+	}
+
 	return true;
 }
 
diff --git a/Engine/Source/Runtime/SlateCore/Private/Types/SlateAttributeDescriptor.cpp b/Engine/Source/Runtime/SlateCore/Private/Types/SlateAttributeDescriptor.cpp
--- a/Engine/Source/Runtime/SlateCore/Private/Types/SlateAttributeDescriptor.cpp
+++ b/Engine/Source/Runtime/SlateCore/Private/Types/SlateAttributeDescriptor.cpp
@@ -44,6 +44,16 @@ FSlateAttributeDescriptor::FInitializer::FAttributeEntry& FSlateAttributeDescrip
 }
 
 
+FSlateAttributeDescriptor::FInitializer::FAttributeEntry& FSlateAttributeDescriptor::FInitializer::FAttributeEntry::OnValueChanged(FAttributeValueChangedDelegate InCallback)
+{
+	if (Descriptor.Attributes.IsValidIndex(AttributeIndex))
+	{
+		Descriptor.Attributes[AttributeIndex].OnValueChanged = MoveTemp(InCallback);
+	}
+	return *this;
+}
+
+
 /** */
 FSlateAttributeDescriptor::FInitializer::FInitializer(FSlateAttributeDescriptor& InDescriptor)
 	: Descriptor(InDescriptor)
@@ -228,6 +238,12 @@ void FSlateAttributeDescriptor::FInitializer::OverrideInvalidationReason(FName A
 }
 
 
+void FSlateAttributeDescriptor::FInitializer::OverrideOnValueChanged(FName AttributeName, ECallbackOverrideType OverrideType, FAttributeValueChangedDelegate Callback)
+{
+	Descriptor.OverrideOnValueChanged(AttributeName, OverrideType, MoveTemp(Callback));
+}
+
+
 void FSlateAttributeDescriptor::FInitializer::SetAffectVisibility(FName AttributeName, bool bAffectVisibility)
 {
 	FAttribute* Attribute = Descriptor.FindAttribute(AttributeName);
@@ -317,6 +333,48 @@ void FSlateAttributeDescriptor::OverrideInvalidationReason(FName AttributeName,
 }
 
 
+void FSlateAttributeDescriptor::OverrideOnValueChanged(FName AttributeName, ECallbackOverrideType OverrideType, FAttributeValueChangedDelegate Callback)
+{
+	check(!AttributeName.IsNone());
+
+	FAttribute* FoundAttribute = FindAttribute(AttributeName);
+	if (ensureAlwaysMsgf(FoundAttribute != nullptr, TEXT("The attribute 's' doesn't exist."), *AttributeName.ToString()))
+	{
+		switch(OverrideType)
+		{
+		case ECallbackOverrideType::ReplacePrevious:
+			FoundAttribute->OnValueChanged = MoveTemp(Callback);
+			break;
+		case ECallbackOverrideType::ExecuteAfterPrevious:
+		case ECallbackOverrideType::ExecuteBeforePrevious:
+			if (FoundAttribute->OnValueChanged.IsBound() && Callback.IsBound())
+			{
+				FAttributeValueChangedDelegate Previous = FoundAttribute->OnValueChanged;
+				FoundAttribute->OnValueChanged = FAttributeValueChangedDelegate::CreateLambda([Previous{MoveTemp(Previous)}, Callback{MoveTemp(Callback)}, OverrideType](SWidget& Widget)
+					{
+						if (OverrideType == ECallbackOverrideType::ExecuteBeforePrevious)
+						{
+							Previous.ExecuteIfBound(Widget);
+						}
+						Callback.ExecuteIfBound(Widget);
+						if (OverrideType == ECallbackOverrideType::ExecuteAfterPrevious)
+						{
+							Previous.ExecuteIfBound(Widget);
+						}
+					});
+			}
+			else if (Callback.IsBound())
+			{
+				FoundAttribute->OnValueChanged = MoveTemp(Callback);
+			}
+			break;
+		default:
+			check(false);
+		}
+	}
+}
+
+
 void FSlateAttributeDescriptor::SetPrerequisite(FSlateAttributeDescriptor::FAttribute& Attribute, FName Prerequisite, bool bSetAsDependency)
 {
 	if (Prerequisite.IsNone())
diff --git a/Engine/Source/Runtime/SlateCore/Private/Types/SlateAttributeMetaData.cpp b/Engine/Source/Runtime/SlateCore/Private/Types/SlateAttributeMetaData.cpp
--- a/Engine/Source/Runtime/SlateCore/Private/Types/SlateAttributeMetaData.cpp
+++ b/Engine/Source/Runtime/SlateCore/Private/Types/SlateAttributeMetaData.cpp
@@ -192,14 +192,14 @@ TArray<FName> FSlateAttributeMetaData::GetAttributeNames(const SWidget& OwningWi
 }
 
 
-EInvalidateWidgetReason FSlateAttributeMetaData::FGetterItem::GetInvalidationDetail(const SWidget& OwningWidget, EInvalidateWidgetReason Reason) const
+FSlateAttributeMetaData::FGetterItem::FInvalidationDetail FSlateAttributeMetaData::FGetterItem::GetInvalidationDetail(const SWidget& OwningWidget, EInvalidateWidgetReason Reason) const
 {
 	if (CachedAttributeDescriptorIndex != FGetterItem::InvalidAttributeIndex)
 	{
 		const FSlateAttributeDescriptor::FAttribute& DescriptorAttribute = OwningWidget.GetWidgetClass().GetAttributeDescriptor().GetAttributeAtIndex(CachedAttributeDescriptorIndex);
-		return DescriptorAttribute.InvalidationReason.Get(OwningWidget);
+		return FInvalidationDetail{&DescriptorAttribute.OnValueChanged, DescriptorAttribute.InvalidationReason.Get(OwningWidget)};
 	}
-	return Reason;
+	return FInvalidationDetail{nullptr, Reason};
 }
 
 
@@ -226,6 +226,8 @@ void FSlateAttributeMetaData::InvalidateWidget(SWidget& OwningWidget, const FSla
 		return;
 	}
 
+	const FSlateAttributeDescriptor::FAttributeValueChangedDelegate* OnValueChangedCallback = nullptr;
+
 	if (FSlateAttributeMetaData* AttributeMetaData = FSlateAttributeMetaData::FindMetaData(OwningWidget))
 	{
 		const int32 FoundIndex = AttributeMetaData->IndexOfAttribute(Attribute);
@@ -233,7 +235,9 @@ void FSlateAttributeMetaData::InvalidateWidget(SWidget& OwningWidget, const FSla
 		{
 			FGetterItem& GetterItem = AttributeMetaData->Attributes[FoundIndex];
 			{
-				Reason = GetterItem.GetInvalidationDetail(OwningWidget, Reason) | AttributeMetaData->CachedInvalidationReason;
+				const FGetterItem::FInvalidationDetail Detail = GetterItem.GetInvalidationDetail(OwningWidget, Reason);
+				OnValueChangedCallback = Detail.Get<0>();
+				Reason = Detail.Get<1>() | AttributeMetaData->CachedInvalidationReason;
 				AttributeMetaData->CachedInvalidationReason = EInvalidateWidgetReason::None;
 			}
 
@@ -251,6 +255,7 @@ void FSlateAttributeMetaData::InvalidateWidget(SWidget& OwningWidget, const FSla
 			const FSlateAttributeDescriptor::OffsetType Offset = Private::FindOffet(OwningWidget, Attribute);
 			if (FSlateAttributeDescriptor::FAttribute const* FoundAttribute = AttributeDescriptor.FindMemberAttribute(Offset))
 			{
+				OnValueChangedCallback = &FoundAttribute->OnValueChanged;
 				Reason = FoundAttribute->InvalidationReason.Get(OwningWidget) | AttributeMetaData->CachedInvalidationReason;
 				AttributeMetaData->CachedInvalidationReason = EInvalidateWidgetReason::None;
 
@@ -281,6 +286,7 @@ void FSlateAttributeMetaData::InvalidateWidget(SWidget& OwningWidget, const FSla
 		if (FSlateAttributeDescriptor::FAttribute const* FoundAttribute = OwningWidget.GetWidgetClass().GetAttributeDescriptor().FindMemberAttribute(Offset))
 		{
 			Reason = FoundAttribute->InvalidationReason.Get(OwningWidget);
+			OnValueChangedCallback = &FoundAttribute->OnValueChanged;
 		}
 	}
 
@@ -289,6 +295,10 @@ void FSlateAttributeMetaData::InvalidateWidget(SWidget& OwningWidget, const FSla
 #endif
 
 	OwningWidget.Invalidate(Reason);
+	if (OnValueChangedCallback)
+	{
+		OnValueChangedCallback->ExecuteIfBound(OwningWidget);
+	}
 }
 
 
@@ -373,8 +383,8 @@ void FSlateAttributeMetaData::UpdateChildrenOnlyVisibilityAttributes(SWidget& Ow
 
 void FSlateAttributeMetaData::UpdateAttributesImpl(SWidget& OwningWidget, EInvalidationPermission InvalidationStyle, int32 StartIndex, int32 IndexNum)
 {
-	const bool bInvalidateIfNeeded = (InvalidationStyle == EInvalidationPermission::AllowInvalidation) || (InvalidationStyle == EInvalidationPermission::AllowInvalidationIfConstructed && OwningWidget.IsConstructed());
-	const bool bAllowInvalidation = bInvalidateIfNeeded || InvalidationStyle == EInvalidationPermission::DelayInvalidation;
+	bool bInvalidateIfNeeded = (InvalidationStyle == EInvalidationPermission::AllowInvalidation) || (InvalidationStyle == EInvalidationPermission::AllowInvalidationIfConstructed && OwningWidget.IsConstructed());
+	bool bAllowInvalidation = bInvalidateIfNeeded || InvalidationStyle == EInvalidationPermission::DelayInvalidation;
 	EInvalidateWidgetReason InvalidationReason = EInvalidateWidgetReason::None;
 	for (int32 Index = StartIndex; Index < IndexNum; ++Index)
 	{
@@ -412,7 +422,12 @@ void FSlateAttributeMetaData::UpdateAttributesImpl(SWidget& OwningWidget, EInval
 		if (Result.bInvalidationRequested && bAllowInvalidation)
 		{
 			SetNeedToResetFlag(Index);
-			InvalidationReason |= GetterItem.GetInvalidationDetail(OwningWidget, Result.InvalidationReason);
+			const FGetterItem::FInvalidationDetail Detail = GetterItem.GetInvalidationDetail(OwningWidget, Result.InvalidationReason);
+			if (Detail.Get<0>())
+			{
+				Detail.Get<0>()->ExecuteIfBound(OwningWidget);
+			}
+			InvalidationReason |= Detail.Get<1>();
 		}
 	}
 
@@ -450,11 +465,16 @@ void FSlateAttributeMetaData::UpdateAttribute(SWidget& OwningWidget, FSlateAttri
 			{
 				if (OwningWidget.IsConstructed())
 				{
-					EInvalidateWidgetReason Reason = GetterItem.GetInvalidationDetail(OwningWidget, Result.InvalidationReason);
+					const FGetterItem::FInvalidationDetail Detail = GetterItem.GetInvalidationDetail(OwningWidget, Result.InvalidationReason);
+					EInvalidateWidgetReason Reason = Detail.Get<1>() | AttributeMetaData->CachedInvalidationReason;
 #if WITH_SLATE_DEBUGGING
-					ensureAlwaysMsgf(FSlateAttributeBase::IsInvalidateWidgetReasonSupported(Reason | AttributeMetaData->CachedInvalidationReason), TEXT("%s is not an EInvalidateWidgetReason supported by SlateAttribute."), *LexToString(Reason | AttributeMetaData->CachedInvalidationReason));
+					ensureAlwaysMsgf(FSlateAttributeBase::IsInvalidateWidgetReasonSupported(Reason), TEXT("%s is not an EInvalidateWidgetReason supported by SlateAttribute."), *LexToString(Reason));
 #endif
-					OwningWidget.Invalidate(Reason | AttributeMetaData->CachedInvalidationReason);
+					OwningWidget.Invalidate(Reason);
+					if (Detail.Get<0>())
+					{
+						Detail.Get<0>()->ExecuteIfBound(OwningWidget);
+					}
 					AttributeMetaData->CachedInvalidationReason = EInvalidateWidgetReason::None;
 				}
 
diff --git a/Engine/Source/Runtime/SlateCore/Private/Widgets/SWidget.cpp b/Engine/Source/Runtime/SlateCore/Private/Widgets/SWidget.cpp
--- a/Engine/Source/Runtime/SlateCore/Private/Widgets/SWidget.cpp
+++ b/Engine/Source/Runtime/SlateCore/Private/Widgets/SWidget.cpp
@@ -1678,7 +1678,17 @@ void SWidget::Prepass_ChildLoop(float InLayoutScaleMultiplier, FChildren* MyChil
 			{
 				if (bUpdateAttributes)
 				{
+#if WITH_SLATE_DEBUGGING
+					EVisibility PreviousVisibility = GetVisibility();
+					int32 PreviousAllChildrenNum = Child.GetAllChildren()->Num();
+#endif
+
 					FSlateAttributeMetaData::UpdateExceptVisibilityAttributes(Child, FSlateAttributeMetaData::EInvalidationPermission::AllowInvalidationIfConstructed);
+
+#if WITH_SLATE_DEBUGGING
+					ensureMsgf(PreviousVisibility == GetVisibility(), TEXT("The visibility of widget '%s' doens't match the previous visibility after the attribute update."), *FReflectionMetaData::GetWidgetDebugInfo(this));
+					ensureMsgf(PreviousAllChildrenNum == Child.GetAllChildren()->Num(), TEXT("The number of child of widget '%s' doens't match the previous count after the attribute update."), *FReflectionMetaData::GetWidgetDebugInfo(this));
+#endif
 				}
 
 				const float ChildLayoutScaleMultiplier = bHasRelativeLayoutScale
diff --git a/Engine/Source/Runtime/SlateCore/Public/Types/SlateAttributeDescriptor.h b/Engine/Source/Runtime/SlateCore/Public/Types/SlateAttributeDescriptor.h
--- a/Engine/Source/Runtime/SlateCore/Public/Types/SlateAttributeDescriptor.h
+++ b/Engine/Source/Runtime/SlateCore/Public/Types/SlateAttributeDescriptor.h
@@ -67,6 +67,20 @@ public:
 		FGetter Getter;
 	};
 
+	/** */
+	DECLARE_DELEGATE_OneParam(FAttributeValueChangedDelegate, SWidget& /*Widget*/);
+
+	/** */
+	enum class ECallbackOverrideType
+	{
+		/** Replace the callback that the base class defined. */
+		ReplacePrevious,
+		/** Execute the callback that the base class defined, then execute the new callback. */
+		ExecuteAfterPrevious,
+		/** Execute the new callback, then execute the callback that the base class defined. */
+		ExecuteBeforePrevious,
+	};
+
 public:
 	using OffsetType = uint32;
 
@@ -81,6 +95,7 @@ public:
 		FName Prerequisite;
 		uint32 SortOrder = 0;
 		FInvalidateWidgetReasonAttribute InvalidationReason = FInvalidateWidgetReasonAttribute(EInvalidateWidgetReason::None);
+		FAttributeValueChangedDelegate OnValueChanged;
 		bool bIsMemberAttribute = false;
 		bool bIsPrerequisiteAlsoADependency = false;
 		bool bIsADependencyForSomeoneElse = false;
@@ -126,6 +141,15 @@ public:
 			 */
 			FAttributeEntry& AffectVisibility();
 
+			/**
+			 * Notified when the attribute value changed.
+			 * It's preferable that you delay any action to the Tick or Paint function.
+			 * You are not allowed to make changes that would affect the SWidget ChildOrder or its Visibility.
+			 * It will not be called when the SWidget is in his construction phase.
+			 * @see SWidget::IsConstructed
+			 */
+			FAttributeEntry& OnValueChanged(FAttributeValueChangedDelegate Callback);
+
 		private:
 			FSlateAttributeDescriptor& Descriptor;
 			int32 AttributeIndex;
@@ -139,6 +163,9 @@ public:
 		/** Change the InvalidationReason of an attribute defined in a base class. */
 		void OverrideInvalidationReason(FName AttributeName, FInvalidateWidgetReasonAttribute&& Reason);
 
+		/** Change the FAttributeValueChangedDelegate of an attribute defined in a base class. */
+		void OverrideOnValueChanged(FName AttributeName, ECallbackOverrideType OverrideType, FAttributeValueChangedDelegate Callback);
+
 		/** Change the update type of an attribute defined in a base class. */
 		void SetAffectVisibility(FName AttributeName, bool bAffectVisibility);
 
@@ -182,6 +209,7 @@ private:
 
 	FInitializer::FAttributeEntry AddMemberAttribute(FName AttributeName, OffsetType Offset, FInvalidateWidgetReasonAttribute ReasonGetter);
 	void OverrideInvalidationReason(FName AttributeName, FInvalidateWidgetReasonAttribute ReasonGetter);
+	void OverrideOnValueChanged(FName AttributeName, ECallbackOverrideType OverrideType, FAttributeValueChangedDelegate Callback);
 	void SetPrerequisite(FAttribute& Attribute, FName Prerequisite, bool bSetAsDependency);
 	void SetAffectVisibility(FAttribute& Attribute, bool bUpdate);
 
diff --git a/Engine/Source/Runtime/SlateCore/Public/Types/SlateAttributeMetaData.h b/Engine/Source/Runtime/SlateCore/Public/Types/SlateAttributeMetaData.h
--- a/Engine/Source/Runtime/SlateCore/Public/Types/SlateAttributeMetaData.h
+++ b/Engine/Source/Runtime/SlateCore/Public/Types/SlateAttributeMetaData.h
@@ -154,7 +154,8 @@ private:
 			return SortOrder < Other.SortOrder;
 		}
 
-		EInvalidateWidgetReason GetInvalidationDetail(const SWidget&, EInvalidateWidgetReason Reason) const;
+		using FInvalidationDetail = TTuple<const FSlateAttributeDescriptor::FAttributeValueChangedDelegate*, EInvalidateWidgetReason>;
+		FInvalidationDetail GetInvalidationDetail(const SWidget&, EInvalidateWidgetReason Reason) const;
 
 		/** If available, return the name of the attribute. */
 		FName GetAttributeName(const SWidget& OwningWidget) const;
