Slate: Remove the SlateAttribute dependency feature in favor of OnValueChanged.

diff --git a/Engine/Source/Runtime/SlateCore/Private/Test/SlateAttributeTest.cpp b/Engine/Source/Runtime/SlateCore/Private/Test/SlateAttributeTest.cpp
--- a/Engine/Source/Runtime/SlateCore/Private/Test/SlateAttributeTest.cpp
+++ b/Engine/Source/Runtime/SlateCore/Private/Test/SlateAttributeTest.cpp
@@ -109,7 +109,7 @@ void SAttributeLeftWidget_Parent::PrivateRegisterAttributes(FSlateAttributeIniti
 	// A updates after B, so B needs to be before A
 	SLATE_ADD_MEMBER_ATTRIBUTE_DEFINITION(AttributeInitializer, IntAttributeD, EInvalidateWidgetReason::Layout);
 	SLATE_ADD_MEMBER_ATTRIBUTE_DEFINITION(AttributeInitializer, IntAttributeC, EInvalidateWidgetReason::Layout)
-		.UpdateDependency(GET_MEMBER_NAME_CHECKED(PrivateThisType, IntAttributeD));
+		.UpdatePrerequisite(GET_MEMBER_NAME_CHECKED(PrivateThisType, IntAttributeD));
 	SLATE_ADD_MEMBER_ATTRIBUTE_DEFINITION(AttributeInitializer, IntAttributeB, EInvalidateWidgetReason::Layout);
 	SLATE_ADD_MEMBER_ATTRIBUTE_DEFINITION(AttributeInitializer, IntAttributeA, EInvalidateWidgetReason::Layout)
 		.UpdatePrerequisite(GET_MEMBER_NAME_CHECKED(PrivateThisType, IntAttributeB));
@@ -152,13 +152,13 @@ void SAttributeLeftWidget_Child::PrivateRegisterAttributes(FSlateAttributeInitia
 	// The update order is M, B, A, I, J, D, C, L, H, K
 	//SLATE_ADD_MEMBER_ATTRIBUTE_DEFINITION(AttributeInitializer, IntAttributeH, EInvalidateWidgetReason::Layout);
 	SLATE_ADD_MEMBER_ATTRIBUTE_DEFINITION(AttributeInitializer, IntAttributeJ, EInvalidateWidgetReason::Layout)
-		.UpdateDependency("IntAttributeA");
+		.UpdatePrerequisite("IntAttributeA");
 	SLATE_ADD_MEMBER_ATTRIBUTE_DEFINITION(AttributeInitializer, IntAttributeK, EInvalidateWidgetReason::Layout);
 	SLATE_ADD_MEMBER_ATTRIBUTE_DEFINITION(AttributeInitializer, IntAttributeI, EInvalidateWidgetReason::Layout)
 		.UpdatePrerequisite("IntAttributeB");
 	SLATE_ADD_MEMBER_ATTRIBUTE_DEFINITION(AttributeInitializer, IntAttributeL, EInvalidateWidgetReason::Layout)
 		.UpdatePrerequisite("IntAttributeC");
-	SLATE_ADD_MEMBER_ATTRIBUTE_DEFINITION(AttributeInitializer, IntAttributeM, EInvalidateWidgetReason::Layout)
+	SLATE_ADD_MEMBER_ATTRIBUTE_DEFINITION(AttributeInitializer, IntAttributeM, EInvalidateWidgetReason::Visibility)
 		.UpdatePrerequisite("Visibility")
 		.AffectVisibility();
 }
@@ -300,13 +300,6 @@ bool FSlateAttributeTest::RunTest(const FString& Parameters)
 		++OrderCounter;
 		return OrderCounter;
 	};
-	bool bWasUpdate = false;
-	int32 ReturnValue = 0;
-	auto UpdateLambda = [this, &bWasUpdate, &ReturnValue]() -> int32
-	{
-		bWasUpdate = true;
-		return ReturnValue;
-	};
 
 	{
 		TSharedRef<UE::Slate::Private::SAttributeLeftWidget_Parent> WidgetParent = SNew(UE::Slate::Private::SAttributeLeftWidget_Parent);
@@ -317,13 +310,13 @@ bool FSlateAttributeTest::RunTest(const FString& Parameters)
 		FSlateAttributeDescriptor const& AttributeDescriptor = WidgetParent->GetWidgetClass().GetAttributeDescriptor();
 		AddErrorIfFalse(AttributeDescriptor.GetAttributeNum() == 4 + NumberOfAttributeInSWidget, TEXT("Invalid number of attributes"));
 
-		const int32 IndexA = AttributeDescriptor.IndexOfMemberAttribute("IntAttributeA");
-		const int32 IndexB = AttributeDescriptor.IndexOfMemberAttribute("IntAttributeB");
-		const int32 IndexC = AttributeDescriptor.IndexOfMemberAttribute("IntAttributeC");
-		const int32 IndexD = AttributeDescriptor.IndexOfMemberAttribute("IntAttributeD");
-		const int32 IndexI = AttributeDescriptor.IndexOfMemberAttribute("IntAttributeI");
-		const int32 IndexJ = AttributeDescriptor.IndexOfMemberAttribute("IntAttributeJ");
-		const int32 IndexK = AttributeDescriptor.IndexOfMemberAttribute("IntAttributeK");
+		const int32 IndexA = AttributeDescriptor.IndexOfAttribute("IntAttributeA");
+		const int32 IndexB = AttributeDescriptor.IndexOfAttribute("IntAttributeB");
+		const int32 IndexC = AttributeDescriptor.IndexOfAttribute("IntAttributeC");
+		const int32 IndexD = AttributeDescriptor.IndexOfAttribute("IntAttributeD");
+		const int32 IndexI = AttributeDescriptor.IndexOfAttribute("IntAttributeI");
+		const int32 IndexJ = AttributeDescriptor.IndexOfAttribute("IntAttributeJ");
+		const int32 IndexK = AttributeDescriptor.IndexOfAttribute("IntAttributeK");
 
 		AddErrorIfFalse(IndexA != INDEX_NONE, TEXT("Could not find the Attribute A"));
 		AddErrorIfFalse(IndexB != INDEX_NONE, TEXT("Could not find the Attribute B"));
@@ -342,9 +335,9 @@ bool FSlateAttributeTest::RunTest(const FString& Parameters)
 		AddErrorIfFalse(AttributeDescriptor.FindAttribute("IntAttributeK") == nullptr, TEXT("Was not supposed to find the Attribute K"));
 
 		//B, A, D, C
-		AddErrorIfFalse(AttributeDescriptor.GetAttributeAtIndex(IndexB).SortOrder < AttributeDescriptor.GetAttributeAtIndex(IndexA).SortOrder, TEXT("B should have a lower value than A"));
-		AddErrorIfFalse(AttributeDescriptor.GetAttributeAtIndex(IndexD).SortOrder < AttributeDescriptor.GetAttributeAtIndex(IndexC).SortOrder, TEXT("D should have a lower value than C"));
-		AddErrorIfFalse(AttributeDescriptor.GetAttributeAtIndex(IndexA).SortOrder < AttributeDescriptor.GetAttributeAtIndex(IndexD).SortOrder, TEXT("A should have a lower value than D"));
+		AddErrorIfFalse(AttributeDescriptor.GetAttributeAtIndex(IndexB).GetSortOrder() < AttributeDescriptor.GetAttributeAtIndex(IndexA).GetSortOrder(), TEXT("B should have a lower value than A"));
+		AddErrorIfFalse(AttributeDescriptor.GetAttributeAtIndex(IndexD).GetSortOrder() < AttributeDescriptor.GetAttributeAtIndex(IndexC).GetSortOrder(), TEXT("D should have a lower value than C"));
+		AddErrorIfFalse(AttributeDescriptor.GetAttributeAtIndex(IndexA).GetSortOrder() < AttributeDescriptor.GetAttributeAtIndex(IndexD).GetSortOrder(), TEXT("A should have a lower value than D"));
 
 		{
 			OrderCounter = 0;
@@ -358,8 +351,6 @@ bool FSlateAttributeTest::RunTest(const FString& Parameters)
 			AddErrorIfFalse(WidgetParent->IntAttributeD.Get() == 99, TEXT("D It is not the expected value."));
 
 			OrderCounter = 0;
-			bWasUpdate = false;
-			ReturnValue = 4;
 			WidgetParent->MarkPrepassAsDirty();
 			WidgetParent->SlatePrepass(1.f);
 			AddErrorIfFalse(WidgetParent->IntAttributeA.Get() == 2, TEXT("A It is not the expected value."));
@@ -368,59 +359,18 @@ bool FSlateAttributeTest::RunTest(const FString& Parameters)
 			AddErrorIfFalse(WidgetParent->IntAttributeD.Get() == 3, TEXT("D It is not the expected value."));
 		}
 
-		{
-			OrderCounter = 0;
-			bWasUpdate = false;
-			ReturnValue = 5;
-			WidgetParent->IntAttributeC.Assign(WidgetParent.Get(), MakeAttributeLambda(UpdateLambda));
-			AddErrorIfFalse(!bWasUpdate, TEXT("C should not have be updated."));
-			AddErrorIfFalse(WidgetParent->IntAttributeC.Get() == 4, TEXT("C It is not the expected value."));
-			WidgetParent->MarkPrepassAsDirty();
-			WidgetParent->SlatePrepass(1.f);
-			AddErrorIfFalse(bWasUpdate, TEXT("C should be updated."));
-			AddErrorIfFalse(WidgetParent->IntAttributeC.Get() == 5, TEXT("C It is not the expected value."));
-		}
-
-		{
-			OrderCounter = 0;
-			bWasUpdate = false;
-			ReturnValue = 10;	//10 show that C didn't change
-			WidgetParent->MarkPrepassAsDirty();
-			WidgetParent->SlatePrepass(1.f);
-			AddErrorIfFalse(WidgetParent->IntAttributeA.Get() == 2, TEXT("A It is not the expected value."));
-			AddErrorIfFalse(WidgetParent->IntAttributeB.Get() == 1, TEXT("B It is not the expected value."));
-			AddErrorIfFalse(!bWasUpdate, TEXT("C should not be updated."));
-			AddErrorIfFalse(WidgetParent->IntAttributeC.Get() == 5, TEXT("C It is not the expected value."));
-			AddErrorIfFalse(WidgetParent->IntAttributeD.Get() == 3, TEXT("D It is not the expected value."));
-		}
-
 		{
 			WidgetParent->IntAttributeD.Set(WidgetParent.Get(), 8);
-			AddErrorIfFalse(WidgetParent->IntAttributeC.Get() == 5, TEXT("C It is not the expected value."));
+			AddErrorIfFalse(WidgetParent->IntAttributeC.Get() == 4, TEXT("C It is not the expected value."));
 			AddErrorIfFalse(WidgetParent->IntAttributeD.Get() == 8, TEXT("D It is not the expected value."));
 
 			OrderCounter = 0;
-			bWasUpdate = false;
-			ReturnValue = 10;
 			WidgetParent->MarkPrepassAsDirty();
 			WidgetParent->SlatePrepass(1.f);
 			AddErrorIfFalse(WidgetParent->IntAttributeA.Get() == 2, TEXT("A It is not the expected value."));
 			AddErrorIfFalse(WidgetParent->IntAttributeB.Get() == 1, TEXT("B It is not the expected value."));
-			AddErrorIfFalse(bWasUpdate, TEXT("C should be updated becaude D was."));
-			AddErrorIfFalse(WidgetParent->IntAttributeC.Get() == 10, TEXT("C It is not the expected value."));
+			AddErrorIfFalse(WidgetParent->IntAttributeC.Get() == 3, TEXT("C It is not the expected value."));
 			AddErrorIfFalse(WidgetParent->IntAttributeD.Get() == 8, TEXT("D It is not the expected value."));
-#ifndef PVS_STUDIO // Build machine refuses to disable this warning
-			AddErrorIfFalse(OrderCounter == 2, TEXT("There is no D attribute anymore.")); //-V547
-#endif
-		}
-
-		{
-			OrderCounter = 0;
-			bWasUpdate = false;
-			ReturnValue = 10;
-			WidgetParent->MarkPrepassAsDirty();
-			WidgetParent->SlatePrepass(1.f);
-			AddErrorIfFalse(!bWasUpdate, TEXT("C should not be updated."));
 		}
 	}
 
@@ -433,15 +383,15 @@ bool FSlateAttributeTest::RunTest(const FString& Parameters)
 		FSlateAttributeDescriptor const& AttributeDescriptor = WidgetChild->GetWidgetClass().GetAttributeDescriptor();
 		AddErrorIfFalse(AttributeDescriptor.GetAttributeNum() == 9 + NumberOfAttributeInSWidget, TEXT("Invalid number of attributes")); // H is not counted
 
-		const int32 IndexA = AttributeDescriptor.IndexOfMemberAttribute("IntAttributeA");
-		const int32 IndexB = AttributeDescriptor.IndexOfMemberAttribute("IntAttributeB");
-		const int32 IndexC = AttributeDescriptor.IndexOfMemberAttribute("IntAttributeC");
-		const int32 IndexD = AttributeDescriptor.IndexOfMemberAttribute("IntAttributeD");
-		const int32 IndexI = AttributeDescriptor.IndexOfMemberAttribute("IntAttributeI");
-		const int32 IndexJ = AttributeDescriptor.IndexOfMemberAttribute("IntAttributeJ");
-		const int32 IndexK = AttributeDescriptor.IndexOfMemberAttribute("IntAttributeK");
-		const int32 IndexL = AttributeDescriptor.IndexOfMemberAttribute("IntAttributeL");
-		const int32 IndexM = AttributeDescriptor.IndexOfMemberAttribute("IntAttributeM");
+		const int32 IndexA = AttributeDescriptor.IndexOfAttribute("IntAttributeA");
+		const int32 IndexB = AttributeDescriptor.IndexOfAttribute("IntAttributeB");
+		const int32 IndexC = AttributeDescriptor.IndexOfAttribute("IntAttributeC");
+		const int32 IndexD = AttributeDescriptor.IndexOfAttribute("IntAttributeD");
+		const int32 IndexI = AttributeDescriptor.IndexOfAttribute("IntAttributeI");
+		const int32 IndexJ = AttributeDescriptor.IndexOfAttribute("IntAttributeJ");
+		const int32 IndexK = AttributeDescriptor.IndexOfAttribute("IntAttributeK");
+		const int32 IndexL = AttributeDescriptor.IndexOfAttribute("IntAttributeL");
+		const int32 IndexM = AttributeDescriptor.IndexOfAttribute("IntAttributeM");
 
 		AddErrorIfFalse(IndexA != INDEX_NONE, TEXT("Could not find the Attribute A"));
 		AddErrorIfFalse(IndexB != INDEX_NONE, TEXT("Could not find the Attribute B"));
@@ -464,14 +414,14 @@ bool FSlateAttributeTest::RunTest(const FString& Parameters)
 		AddErrorIfFalse(&AttributeDescriptor.GetAttributeAtIndex(IndexM) == AttributeDescriptor.FindAttribute("IntAttributeM"), TEXT("Index and Attribute should return the same value."));
 		AddErrorIfFalse(AttributeDescriptor.FindAttribute("IntAttributeH") == nullptr, TEXT("H exist but is not defined."));
 
-		AddErrorIfFalse(AttributeDescriptor.GetAttributeAtIndex(IndexM).SortOrder < AttributeDescriptor.GetAttributeAtIndex(IndexB).SortOrder, TEXT("M should have a lower value than B"));
-		AddErrorIfFalse(AttributeDescriptor.GetAttributeAtIndex(IndexB).SortOrder < AttributeDescriptor.GetAttributeAtIndex(IndexA).SortOrder, TEXT("B should have a lower value than A"));
-		AddErrorIfFalse(AttributeDescriptor.GetAttributeAtIndex(IndexA).SortOrder <= AttributeDescriptor.GetAttributeAtIndex(IndexI).SortOrder, TEXT("A should have a lower value than I"));
-		AddErrorIfFalse(AttributeDescriptor.GetAttributeAtIndex(IndexI).SortOrder < AttributeDescriptor.GetAttributeAtIndex(IndexJ).SortOrder, TEXT("I should have a lower value than J"));
-		AddErrorIfFalse(AttributeDescriptor.GetAttributeAtIndex(IndexJ).SortOrder < AttributeDescriptor.GetAttributeAtIndex(IndexD).SortOrder, TEXT("J should have a lower value than D"));
-		AddErrorIfFalse(AttributeDescriptor.GetAttributeAtIndex(IndexD).SortOrder < AttributeDescriptor.GetAttributeAtIndex(IndexC).SortOrder, TEXT("D should have a lower value than C"));
-		AddErrorIfFalse(AttributeDescriptor.GetAttributeAtIndex(IndexC).SortOrder < AttributeDescriptor.GetAttributeAtIndex(IndexL).SortOrder, TEXT("C should have a lower value than L"));
-		AddErrorIfFalse(AttributeDescriptor.GetAttributeAtIndex(IndexL).SortOrder < AttributeDescriptor.GetAttributeAtIndex(IndexK).SortOrder, TEXT("L should have a lower value than K"));
+		AddErrorIfFalse(AttributeDescriptor.GetAttributeAtIndex(IndexM).GetSortOrder() < AttributeDescriptor.GetAttributeAtIndex(IndexB).GetSortOrder(), TEXT("M should have a lower value than B"));
+		AddErrorIfFalse(AttributeDescriptor.GetAttributeAtIndex(IndexB).GetSortOrder() < AttributeDescriptor.GetAttributeAtIndex(IndexA).GetSortOrder(), TEXT("B should have a lower value than A"));
+		AddErrorIfFalse(AttributeDescriptor.GetAttributeAtIndex(IndexA).GetSortOrder() <= AttributeDescriptor.GetAttributeAtIndex(IndexI).GetSortOrder(), TEXT("A should have a lower value than I"));
+		AddErrorIfFalse(AttributeDescriptor.GetAttributeAtIndex(IndexI).GetSortOrder() < AttributeDescriptor.GetAttributeAtIndex(IndexJ).GetSortOrder(), TEXT("I should have a lower value than J"));
+		AddErrorIfFalse(AttributeDescriptor.GetAttributeAtIndex(IndexJ).GetSortOrder() < AttributeDescriptor.GetAttributeAtIndex(IndexD).GetSortOrder(), TEXT("J should have a lower value than D"));
+		AddErrorIfFalse(AttributeDescriptor.GetAttributeAtIndex(IndexD).GetSortOrder() < AttributeDescriptor.GetAttributeAtIndex(IndexC).GetSortOrder(), TEXT("D should have a lower value than C"));
+		AddErrorIfFalse(AttributeDescriptor.GetAttributeAtIndex(IndexC).GetSortOrder() < AttributeDescriptor.GetAttributeAtIndex(IndexL).GetSortOrder(), TEXT("C should have a lower value than L"));
+		AddErrorIfFalse(AttributeDescriptor.GetAttributeAtIndex(IndexL).GetSortOrder() < AttributeDescriptor.GetAttributeAtIndex(IndexK).GetSortOrder(), TEXT("L should have a lower value than K"));
 
 
 		{
@@ -499,8 +449,6 @@ bool FSlateAttributeTest::RunTest(const FString& Parameters)
 
 
 			OrderCounter = 0;
-			bWasUpdate = false;
-			ReturnValue = 4;
 			WidgetChild->MarkPrepassAsDirty();
 			WidgetChild->SlatePrepass(1.f);
 			AddErrorIfFalse(WidgetChild->IntAttributeA.Get() == 3 || WidgetChild->IntAttributeA.Get() == 4, TEXT("A It is not the expected value."));
@@ -514,60 +462,6 @@ bool FSlateAttributeTest::RunTest(const FString& Parameters)
 			AddErrorIfFalse(WidgetChild->IntAttributeL.Get() == 8, TEXT("L It is not the expected value."));
 			AddErrorIfFalse(WidgetChild->IntAttributeM.Get() == 1, TEXT("L It is not the expected value."));
 		}
-
-		{
-			OrderCounter = 0;
-			bWasUpdate = false;
-			ReturnValue = 4;
-			WidgetChild->MarkPrepassAsDirty();
-			WidgetChild->SlatePrepass(1.f);
-			AddErrorIfFalse(WidgetChild->IntAttributeA.Get() == 3 || WidgetChild->IntAttributeA.Get() == 4, TEXT("A It is not the expected value."));
-			AddErrorIfFalse(WidgetChild->IntAttributeB.Get() == 2, TEXT("B It is not the expected value."));
-			AddErrorIfFalse(WidgetChild->IntAttributeC.Get() == 6, TEXT("C It is not the expected value.")); // will get updated because D changes
-			AddErrorIfFalse(WidgetChild->IntAttributeD.Get() == 5, TEXT("D It is not the expected value."));
-			AddErrorIfFalse(WidgetChild->IntAttributeH.Get() == 8, TEXT("H It is not the expected value."));
-			AddErrorIfFalse(WidgetChild->IntAttributeI.Get() == 3 || WidgetChild->IntAttributeI.Get() == 4, TEXT("I It is not the expected value."));
-			AddErrorIfFalse(WidgetChild->IntAttributeJ.Get() == 5, TEXT("J It is not the expected value.")); // should not get updated
-			AddErrorIfFalse(WidgetChild->IntAttributeK.Get() == 9, TEXT("K It is not the expected value."));
-			AddErrorIfFalse(WidgetChild->IntAttributeL.Get() == 7, TEXT("L It is not the expected value."));
-			AddErrorIfFalse(WidgetChild->IntAttributeM.Get() == 1, TEXT("M It is not the expected value."));
-		}
-
-		// Check the ForEachDependency result
-		{
-			const FSlateAttributeDescriptor& ChildDescriptor = WidgetChild->GetWidgetClass().GetAttributeDescriptor();
-			auto DependencyTest = [this, ChildDescriptor](auto& AttributeInstance, TSharedRef<SWidget> Widget, int32 Expected, FStringView VariableName)
-			{
-				auto FindOffet = [](const SWidget& OwningWidget, const FSlateAttributeBase& Attribute)
-				{
-					UPTRINT Offset = (UPTRINT)(&Attribute) - (UPTRINT)(&OwningWidget);
-					ensure(Offset <= std::numeric_limits<FSlateAttributeDescriptor::OffsetType>::max());
-					return (FSlateAttributeDescriptor::OffsetType)(Offset);
-				};
-
-				int32 Count = 0;
-				const FSlateAttributeDescriptor::FAttribute* FoundAttribute = ChildDescriptor.FindMemberAttribute(FindOffet(Widget.Get(), AttributeInstance));
-				if (FoundAttribute == nullptr)
-				{
-					AddError(FString::Printf(TEXT("Could not find attribute '%s'"), VariableName.GetData()));
-				}
-				else
-				{
-					ChildDescriptor.ForEachDependentsOn(*FoundAttribute, [&Count](uint8 Index) {++Count;});
-					AddErrorIfFalse(Count == Expected, FString::Printf(TEXT("'%s' doesn't have the correct number of dependency (returned %d, expected %d).")
-						, VariableName.GetData(), Count, Expected));
-				}
-			};
-			DependencyTest(WidgetChild->IntAttributeA, WidgetChild, 1, TEXT("A"));//J
-			DependencyTest(WidgetChild->IntAttributeB, WidgetChild, 0, TEXT("B"));//AIJ, they are prerequisite, not dependency
-			DependencyTest(WidgetChild->IntAttributeC, WidgetChild, 0, TEXT("C"));//L
-			DependencyTest(WidgetChild->IntAttributeD, WidgetChild, 1, TEXT("D"));//CL
-			DependencyTest(WidgetChild->IntAttributeI, WidgetChild, 0, TEXT("I"));
-			DependencyTest(WidgetChild->IntAttributeJ, WidgetChild, 0, TEXT("J"));
-			DependencyTest(WidgetChild->IntAttributeK, WidgetChild, 0, TEXT("K"));
-			DependencyTest(WidgetChild->IntAttributeL, WidgetChild, 0, TEXT("L"));
-			DependencyTest(WidgetChild->IntAttributeM, WidgetChild, 0, TEXT("M"));
-		}
 	}
 
 
diff --git a/Engine/Source/Runtime/SlateCore/Private/Types/SlateAttributeDescriptor.cpp b/Engine/Source/Runtime/SlateCore/Private/Types/SlateAttributeDescriptor.cpp
--- a/Engine/Source/Runtime/SlateCore/Private/Types/SlateAttributeDescriptor.cpp
+++ b/Engine/Source/Runtime/SlateCore/Private/Types/SlateAttributeDescriptor.cpp
@@ -7,6 +7,20 @@
 #include <limits>
 
 
+/** */
+FSlateAttributeDescriptor::FAttribute::FAttribute(FName InName, OffsetType InOffset, FInvalidateWidgetReasonAttribute InReason)
+	: Name(InName)
+	, Offset(InOffset)
+	, Prerequisite()
+	, SortOrder(DefaultSortOrder(InOffset))
+	, InvalidationReason(MoveTemp(InReason))
+	, OnValueChanged()
+	, AttributeType(SlateAttributePrivate::ESlateAttributeType::Member)
+	, bAffectVisibility(false)
+{
+
+}
+
 /** */
 FSlateAttributeDescriptor::FInitializer::FAttributeEntry::FAttributeEntry(FSlateAttributeDescriptor& InDescriptor, int32 InAttributeIndex)
 	: Descriptor(InDescriptor)
@@ -18,17 +32,7 @@ FSlateAttributeDescriptor::FInitializer::FAttributeEntry& FSlateAttributeDescrip
 {
 	if (Descriptor.Attributes.IsValidIndex(AttributeIndex))
 	{
-		Descriptor.SetPrerequisite(Descriptor.Attributes[AttributeIndex], Prerequisite, false);
-	}
-	return *this;
-}
-
-
-FSlateAttributeDescriptor::FInitializer::FAttributeEntry& FSlateAttributeDescriptor::FInitializer::FAttributeEntry::UpdateDependency(FName Dependency)
-{
-	if (Descriptor.Attributes.IsValidIndex(AttributeIndex))
-	{
-		Descriptor.SetPrerequisite(Descriptor.Attributes[AttributeIndex], Dependency, true);
+		Descriptor.SetPrerequisite(Descriptor.Attributes[AttributeIndex], Prerequisite);
 	}
 	return *this;
 }
@@ -153,7 +157,6 @@ FSlateAttributeDescriptor::FInitializer::~FInitializer()
 			if (ensureAlwaysMsgf(Descriptor.Attributes.IsValidIndex(PrerequisiteIndex), TEXT("The Prerequisite '%s' doesn't exist"), *Prerequisite.ToString()))
 			{
 				Prerequisites.Emplace(Index, PrerequisiteIndex, -1);
-				Descriptor.Attributes[PrerequisiteIndex].bIsADependencyForSomeoneElse = true;
 				bHavePrerequisite = true;
 			}
 			else
@@ -195,7 +198,7 @@ FSlateAttributeDescriptor::FInitializer::~FInitializer()
 	}
 
 	// Confirm that the attributes marked as "AffectVisibility" are in front of the list.
-#if 0
+#if 1
 	{
 		bool bLookingForAffectVisibility = true;
 		for (const FAttribute& Attribute : Descriptor.Attributes)
@@ -269,36 +272,39 @@ FSlateAttributeDescriptor::FAttribute const* FSlateAttributeDescriptor::FindAttr
 }
 
 
-int32 FSlateAttributeDescriptor::IndexOfMemberAttribute(OffsetType AttributeOffset) const
+FSlateAttributeDescriptor::FAttribute* FSlateAttributeDescriptor::FindAttribute(FName AttributeName)
 {
-	int32 FoundIndex = Attributes.IndexOfByPredicate([AttributeOffset](const FAttribute& Other) { return Other.Offset == AttributeOffset; });
-	check(FoundIndex == INDEX_NONE || Attributes[FoundIndex].bIsMemberAttribute);
-	return FoundIndex;
+	return Attributes.FindByPredicate([AttributeName](const FAttribute& Other) { return Other.Name == AttributeName; });
 }
 
 
-int32 FSlateAttributeDescriptor::IndexOfMemberAttribute(FName AttributeName) const
+FSlateAttributeDescriptor::FAttribute const* FSlateAttributeDescriptor::FindMemberAttribute(OffsetType AttributeOffset) const
 {
-	int32 FoundIndex = Attributes.IndexOfByPredicate([AttributeName](const FAttribute& Other) { return Other.Name == AttributeName; });
-	if (ensure(FoundIndex == INDEX_NONE || Attributes[FoundIndex].bIsMemberAttribute))
-	{
-		return FoundIndex;
-	}
-	return INDEX_NONE;
+	FSlateAttributeDescriptor::FAttribute const* Result = Attributes.FindByPredicate(
+		[AttributeOffset](const FAttribute& Other)
+		{
+			return Other.Offset == AttributeOffset
+				&& Other.AttributeType == SlateAttributePrivate::ESlateAttributeType::Member;
+		});
+	return Result;
 }
 
 
-FSlateAttributeDescriptor::FAttribute const* FSlateAttributeDescriptor::FindMemberAttribute(OffsetType AttributeOffset) const
+int32 FSlateAttributeDescriptor::IndexOfAttribute(FName AttributeName) const
 {
-	FSlateAttributeDescriptor::FAttribute const* Result = Attributes.FindByPredicate([AttributeOffset](const FAttribute& Other) { return Other.Offset == AttributeOffset; });
-	check(Result == nullptr || Result->bIsMemberAttribute);
-	return Result;
+	return Attributes.IndexOfByPredicate([AttributeName](const FAttribute& Other) { return Other.Name == AttributeName; });
 }
 
 
-FSlateAttributeDescriptor::FAttribute* FSlateAttributeDescriptor::FindAttribute(FName AttributeName)
+int32 FSlateAttributeDescriptor::IndexOfMemberAttribute(OffsetType AttributeOffset) const
 {
-	return Attributes.FindByPredicate([AttributeName](const FAttribute& Other) { return Other.Name == AttributeName; });
+	int32 FoundIndex = Attributes.IndexOfByPredicate(
+		[AttributeOffset](const FAttribute& Other)
+		{
+			return Other.Offset == AttributeOffset
+				&& Other.AttributeType == SlateAttributePrivate::ESlateAttributeType::Member;
+		});
+	return FoundIndex;
 }
 
 
@@ -310,12 +316,7 @@ FSlateAttributeDescriptor::FInitializer::FAttributeEntry FSlateAttributeDescript
 	FAttribute const* FoundAttribute = FindAttribute(AttributeName);
 	if (ensureAlwaysMsgf(FoundAttribute == nullptr, TEXT("The attribute '%s' already exist. (Do you have the correct parrent class in SLATE_DECLARE_WIDGET)"), *AttributeName.ToString()))
 	{
-		NewIndex = Attributes.AddZeroed();
-		FAttribute& NewAttribute = Attributes[NewIndex];
-		NewAttribute.Name = AttributeName;
-		NewAttribute.Offset = Offset;
-		NewAttribute.InvalidationReason = MoveTemp(Reason);
-		NewAttribute.bIsMemberAttribute = true;
+		NewIndex = Attributes.Emplace(AttributeName, Offset, MoveTemp(Reason));
 	}
 	return FInitializer::FAttributeEntry(*this, NewIndex);
 }
@@ -375,12 +376,11 @@ void FSlateAttributeDescriptor::OverrideOnValueChanged(FName AttributeName, ECal
 }
 
 
-void FSlateAttributeDescriptor::SetPrerequisite(FSlateAttributeDescriptor::FAttribute& Attribute, FName Prerequisite, bool bSetAsDependency)
+void FSlateAttributeDescriptor::SetPrerequisite(FSlateAttributeDescriptor::FAttribute& Attribute, FName Prerequisite)
 {
 	if (Prerequisite.IsNone())
 	{
 		Attribute.Prerequisite = FName();
-		Attribute.bIsPrerequisiteAlsoADependency = false;
 	}
 	else
 	{
@@ -388,7 +388,6 @@ void FSlateAttributeDescriptor::SetPrerequisite(FSlateAttributeDescriptor::FAttr
 		if (ensureAlwaysMsgf(FoundPrerequisite, TEXT("The prerequisite '%s' doesn't exist for attribute '%s'"), *Prerequisite.ToString(), *Attribute.Name.ToString()))
 		{
 			Attribute.Prerequisite = Prerequisite;
-			Attribute.bIsPrerequisiteAlsoADependency = bSetAsDependency;
 
 			// Verify recursion
 			{
@@ -401,7 +400,6 @@ void FSlateAttributeDescriptor::SetPrerequisite(FSlateAttributeDescriptor::FAttr
 					{
 						ensureAlwaysMsgf(false, TEXT("The prerequsite '%s' would introduce an infinit loop with attribute '%s'."), *Prerequisite.ToString(), *Attribute.Name.ToString());
 						Attribute.Prerequisite = FName();
-						Attribute.bIsPrerequisiteAlsoADependency = false;
 						break;
 					}
 					Recursion.Add(CurrentAttribute->Name);
@@ -413,7 +411,6 @@ void FSlateAttributeDescriptor::SetPrerequisite(FSlateAttributeDescriptor::FAttr
 		else
 		{
 			Attribute.Prerequisite = FName();
-			Attribute.bIsPrerequisiteAlsoADependency = false;
 		}
 	}
 }
@@ -426,5 +423,11 @@ void FSlateAttributeDescriptor::SetAffectVisibility(FAttribute& Attribute, bool
 		checkf(bInAffectVisibility, TEXT("The Visibility attribute must be marked at 'Affect Visibility'"));
 	}
 
+	if (bInAffectVisibility)
+	{
+		ensureAlwaysMsgf(EnumHasAnyFlags(Attribute.InvalidationReason.Reason, EInvalidateWidgetReason::Visibility) || Attribute.InvalidationReason.IsBound()
+			, TEXT("The attribute '%s' affect the visibility but doesn't have a Visibility as it's InvalidateWidgetReason"), *Attribute.Name.ToString());
+	}
+
 	Attribute.bAffectVisibility = bInAffectVisibility;
 }
\ No newline at end of file
diff --git a/Engine/Source/Runtime/SlateCore/Private/Types/SlateAttributeMetaData.cpp b/Engine/Source/Runtime/SlateCore/Private/Types/SlateAttributeMetaData.cpp
--- a/Engine/Source/Runtime/SlateCore/Private/Types/SlateAttributeMetaData.cpp
+++ b/Engine/Source/Runtime/SlateCore/Private/Types/SlateAttributeMetaData.cpp
@@ -10,9 +10,6 @@
 #include <limits>
 
 
-const FSlateAttributeMetaData::FGetterItem::FAttributeIndex FSlateAttributeMetaData::FGetterItem::InvalidAttributeIndex = std::numeric_limits<FSlateAttributeMetaData::FGetterItem::FAttributeIndex>::max();
-
-
 namespace Private
 {
 	FSlateAttributeDescriptor::OffsetType FindOffet(const SWidget& OwningWidget, const FSlateAttributeBase& Attribute)
@@ -74,7 +71,6 @@ void FSlateAttributeMetaData::RegisterAttributeImpl(SWidget& OwningWidget, FSlat
 	if (FoundIndex != INDEX_NONE)
 	{
 		Attributes[FoundIndex].Getter = MoveTemp(Getter);
-		Attributes[FoundIndex].bUpdatedOnce = false;
 	}
 	else
 	{
@@ -82,32 +78,13 @@ void FSlateAttributeMetaData::RegisterAttributeImpl(SWidget& OwningWidget, FSlat
 		{
 			// MemberAttribute are optional for now but will be needed in the future
 			const FSlateAttributeDescriptor::OffsetType Offset = Private::FindOffet(OwningWidget, Attribute);
-			const FSlateAttributeDescriptor& Descriptor = OwningWidget.GetWidgetClass().GetAttributeDescriptor();
-			const int32 FoundMemberAttributeIndex = Descriptor.IndexOfMemberAttribute(Offset);
-
-			if (FoundMemberAttributeIndex != INDEX_NONE)
+			const FSlateAttributeDescriptor::FAttribute* FoundMemberAttribute = OwningWidget.GetWidgetClass().GetAttributeDescriptor().FindMemberAttribute(Offset);
+			if (FoundMemberAttribute)
 			{
-				FSlateAttributeDescriptor::FAttribute const& FoundAttribute = Descriptor.GetAttributeAtIndex(FoundMemberAttributeIndex);
-				check(FoundMemberAttributeIndex < std::numeric_limits<FGetterItem::FAttributeIndex>::max());
-
-				const int32 InsertLocation = Algo::LowerBoundBy(Attributes, FoundAttribute.SortOrder, [](const FGetterItem& Item) { return Item.SortOrder; }, TLess<>());
-				FGetterItem& GetterItem = Attributes.Insert_GetRef({ &Attribute, FoundAttribute.SortOrder, MoveTemp(Getter), (FGetterItem::FAttributeIndex)FoundMemberAttributeIndex }, InsertLocation);
+				const int32 InsertLocation = Algo::LowerBoundBy(Attributes, FoundMemberAttribute->GetSortOrder(), [](const FGetterItem& Item) { return Item.SortOrder; }, TLess<>());
+				FGetterItem& GetterItem = Attributes.Insert_GetRef({ &Attribute, FoundMemberAttribute->GetSortOrder(), MoveTemp(Getter), FoundMemberAttribute }, InsertLocation);
 				GetterItem.AttributeType = ESlateAttributeType::Member;
-
-				// Do I have dependency or am I a dependency
-				if (!FoundAttribute.Prerequisite.IsNone() && FoundAttribute.bIsPrerequisiteAlsoADependency) 
-				{
-					// I can only be updated if the prerequisite is updated
-					const int32 FoundDependencyAttributeIndex = Descriptor.IndexOfMemberAttribute(FoundAttribute.Prerequisite);
-					if (FoundDependencyAttributeIndex != INDEX_NONE)
-					{
-						check(FoundDependencyAttributeIndex < std::numeric_limits<FGetterItem::FAttributeIndex>::max());
-						GetterItem.CachedAttributeDependencyIndex = (FGetterItem::FAttributeIndex)FoundDependencyAttributeIndex;
-					}
-				}
-				GetterItem.bIsADependencyForSomeoneElse = FoundAttribute.bIsADependencyForSomeoneElse;
-				GetterItem.bAffectVisibility = FoundAttribute.bAffectVisibility;
-				if (GetterItem.bAffectVisibility)
+				if (FoundMemberAttribute->DoesAffectVisibility())
 				{
 					++AffectVisibilityCounter;
 				}
@@ -161,7 +138,7 @@ bool FSlateAttributeMetaData::UnregisterAttributeImpl(const FSlateAttributeBase&
 	const int32 FoundIndex = IndexOfAttribute(Attribute);
 	if (FoundIndex != INDEX_NONE)
 	{
-		if (Attributes[FoundIndex].bAffectVisibility)
+		if (Attributes[FoundIndex].CachedAttributeDescriptor && Attributes[FoundIndex].CachedAttributeDescriptor->DoesAffectVisibility())
 		{
 			check(AffectVisibilityCounter > 0);
 			--AffectVisibilityCounter;
@@ -192,25 +169,9 @@ TArray<FName> FSlateAttributeMetaData::GetAttributeNames(const SWidget& OwningWi
 }
 
 
-FSlateAttributeMetaData::FGetterItem::FInvalidationDetail FSlateAttributeMetaData::FGetterItem::GetInvalidationDetail(const SWidget& OwningWidget, EInvalidateWidgetReason Reason) const
-{
-	if (CachedAttributeDescriptorIndex != FGetterItem::InvalidAttributeIndex)
-	{
-		const FSlateAttributeDescriptor::FAttribute& DescriptorAttribute = OwningWidget.GetWidgetClass().GetAttributeDescriptor().GetAttributeAtIndex(CachedAttributeDescriptorIndex);
-		return FInvalidationDetail{&DescriptorAttribute.OnValueChanged, DescriptorAttribute.InvalidationReason.Get(OwningWidget)};
-	}
-	return FInvalidationDetail{nullptr, Reason};
-}
-
-
 FName FSlateAttributeMetaData::FGetterItem::GetAttributeName(const SWidget& OwningWidget) const
 {
-	if (CachedAttributeDescriptorIndex != FGetterItem::InvalidAttributeIndex)
-	{
-		const FSlateAttributeDescriptor::FAttribute& DescriptorAttribute = OwningWidget.GetWidgetClass().GetAttributeDescriptor().GetAttributeAtIndex(CachedAttributeDescriptorIndex);
-		return DescriptorAttribute.Name;
-	}
-	return FName();
+	return CachedAttributeDescriptor ? CachedAttributeDescriptor->GetName() : FName();
 }
 
 
@@ -226,6 +187,18 @@ void FSlateAttributeMetaData::InvalidateWidget(SWidget& OwningWidget, const FSla
 		return;
 	}
 
+	auto InvalidateForMember = [Reason](SWidget& OwningWidget, const FSlateAttributeBase& Attribute) -> EInvalidateWidgetReason
+	{
+		const FSlateAttributeDescriptor::OffsetType Offset = Private::FindOffet(OwningWidget, Attribute);
+		if (const FSlateAttributeDescriptor::FAttribute* FoundAttribute = OwningWidget.GetWidgetClass().GetAttributeDescriptor().FindMemberAttribute(Offset))
+		{
+			ensureMsgf(Reason == EInvalidateWidgetReason::None, TEXT("SWidget is using the AttributeDescriptor, the invalidation reason should only be defined in one place. See SLATE_ADD_MEMBER_ATTRIBUTE_DEFINITION"));
+			FoundAttribute->ExecuteOnValueChangedIfBound(OwningWidget);
+			return FoundAttribute->GetInvalidationReason(OwningWidget);
+		}
+		return Reason;
+	};
+
 	const FSlateAttributeDescriptor::FAttributeValueChangedDelegate* OnValueChangedCallback = nullptr;
 
 	if (FSlateAttributeMetaData* AttributeMetaData = FSlateAttributeMetaData::FindMetaData(OwningWidget))
@@ -234,60 +207,25 @@ void FSlateAttributeMetaData::InvalidateWidget(SWidget& OwningWidget, const FSla
 		if (FoundIndex != INDEX_NONE)
 		{
 			FGetterItem& GetterItem = AttributeMetaData->Attributes[FoundIndex];
+			if (GetterItem.CachedAttributeDescriptor)
 			{
-				const FGetterItem::FInvalidationDetail Detail = GetterItem.GetInvalidationDetail(OwningWidget, Reason);
-				OnValueChangedCallback = Detail.Get<0>();
-				Reason = Detail.Get<1>() | AttributeMetaData->CachedInvalidationReason;
-				AttributeMetaData->CachedInvalidationReason = EInvalidateWidgetReason::None;
-			}
-
-			// The dependency attribute need to be updated in the update loop (note that it may not be registered yet)
-			if (GetterItem.bIsADependencyForSomeoneElse)
-			{
-				GetterItem.bUpdatedManually = true;
-				AttributeMetaData->SetNeedToResetFlag(FoundIndex);
+				ensureMsgf(Reason == EInvalidateWidgetReason::None, TEXT("SWidget is using the AttributeDescriptor, the invalidation reason should only be defined in one place. See SLATE_ADD_MEMBER_ATTRIBUTE_DEFINITION"));
+				GetterItem.CachedAttributeDescriptor->ExecuteOnValueChangedIfBound(OwningWidget);
+				Reason = GetterItem.CachedAttributeDescriptor->GetInvalidationReason(OwningWidget);
 			}
 		}
 		// Not registered/bound but may be defined in the Descriptor
 		else if (AttributeType == ESlateAttributeType::Member)
 		{
-			FSlateAttributeDescriptor const& AttributeDescriptor = OwningWidget.GetWidgetClass().GetAttributeDescriptor();
-			const FSlateAttributeDescriptor::OffsetType Offset = Private::FindOffet(OwningWidget, Attribute);
-			if (FSlateAttributeDescriptor::FAttribute const* FoundAttribute = AttributeDescriptor.FindMemberAttribute(Offset))
-			{
-				OnValueChangedCallback = &FoundAttribute->OnValueChanged;
-				Reason = FoundAttribute->InvalidationReason.Get(OwningWidget) | AttributeMetaData->CachedInvalidationReason;
-				AttributeMetaData->CachedInvalidationReason = EInvalidateWidgetReason::None;
-
-				if (FoundAttribute->bIsADependencyForSomeoneElse)
-				{
-					// Find if that dependency is registered, if not it is ok because every attribute is updated at least once
-					// Set UpdatedOnce to false to force a new update.
-					AttributeDescriptor.ForEachDependentsOn(*FoundAttribute, [AttributeMetaData](int32 DependencyIndex)
-					{
-						FGetterItem* FoundOther = AttributeMetaData->Attributes.FindByPredicate([DependencyIndex](FGetterItem const& Other)
-						{
-							check(DependencyIndex != INDEX_NONE);
-							check(DependencyIndex < std::numeric_limits<FGetterItem::FAttributeIndex>::max());
-							return Other.CachedAttributeDescriptorIndex == (FGetterItem::FAttributeIndex)DependencyIndex;
-						});
-						if (FoundOther)
-						{
-							FoundOther->bUpdatedOnce = false;
-						}
-					});
-				}
-			}
+			Reason = InvalidateForMember(OwningWidget, Attribute);
 		}
+
+		Reason |= AttributeMetaData->CachedInvalidationReason;
+		AttributeMetaData->CachedInvalidationReason = EInvalidateWidgetReason::None;
 	}
 	else if (AttributeType == ESlateAttributeType::Member)
 	{
-		const FSlateAttributeDescriptor::OffsetType Offset = Private::FindOffet(OwningWidget, Attribute);
-		if (FSlateAttributeDescriptor::FAttribute const* FoundAttribute = OwningWidget.GetWidgetClass().GetAttributeDescriptor().FindMemberAttribute(Offset))
-		{
-			Reason = FoundAttribute->InvalidationReason.Get(OwningWidget);
-			OnValueChangedCallback = &FoundAttribute->OnValueChanged;
-		}
+		Reason = InvalidateForMember(OwningWidget, Attribute);
 	}
 
 #if WITH_SLATE_DEBUGGING
@@ -295,10 +233,6 @@ void FSlateAttributeMetaData::InvalidateWidget(SWidget& OwningWidget, const FSla
 #endif
 
 	OwningWidget.Invalidate(Reason);
-	if (OnValueChangedCallback)
-	{
-		OnValueChangedCallback->ExecuteIfBound(OwningWidget);
-	}
 }
 
 
@@ -307,15 +241,6 @@ void FSlateAttributeMetaData::UpdateAllAttributes(SWidget& OwningWidget, EInvali
 	if (FSlateAttributeMetaData* AttributeMetaData = FSlateAttributeMetaData::FindMetaData(OwningWidget))
 	{
 		AttributeMetaData->UpdateAttributesImpl(OwningWidget, InvalidationStyle, 0, AttributeMetaData->Attributes.Num());
-		if (AttributeMetaData->ResetFlag != EResetFlags::None)
-		{
-			for (FGetterItem& GetterItem : AttributeMetaData->Attributes)
-			{
-				GetterItem.bUpdatedManually = false;
-				GetterItem.bUpdatedThisFrame = false;
-			}
-			AttributeMetaData->ResetFlag = EResetFlags::None;
-		}
 	}
 }
 
@@ -329,16 +254,6 @@ void FSlateAttributeMetaData::UpdateOnlyVisibilityAttributes(SWidget& OwningWidg
 			const int32 StartIndex = 0;
 			const int32 EndIndex = AttributeMetaData->AffectVisibilityCounter;
 			AttributeMetaData->UpdateAttributesImpl(OwningWidget, InvalidationStyle, StartIndex, EndIndex);
-			if (EnumHasAllFlags(AttributeMetaData->ResetFlag, EResetFlags::NeedToReset_OnlyVisibility))
-			{
-				for (int32 Index = StartIndex; Index < EndIndex; ++Index)
-				{
-					FGetterItem& GetterItem = AttributeMetaData->Attributes[Index];
-					GetterItem.bUpdatedManually = false;
-					GetterItem.bUpdatedThisFrame = false;
-				}
-				EnumRemoveFlags(AttributeMetaData->ResetFlag, EResetFlags::NeedToReset_OnlyVisibility);
-			}
 		}
 	}
 }
@@ -353,16 +268,6 @@ void FSlateAttributeMetaData::UpdateExceptVisibilityAttributes(SWidget& OwningWi
 			const int32 StartIndex = AttributeMetaData->AffectVisibilityCounter;
 			const int32 EndIndex = AttributeMetaData->Attributes.Num();
 			AttributeMetaData->UpdateAttributesImpl(OwningWidget, InvalidationStyle, StartIndex, EndIndex);
-			if (EnumHasAllFlags(AttributeMetaData->ResetFlag, EResetFlags::NeedToReset_ExceptVisibility))
-			{
-				for (int32 Index = StartIndex; Index < EndIndex; ++Index)
-				{
-					FGetterItem& GetterItem = AttributeMetaData->Attributes[Index];
-					GetterItem.bUpdatedManually = false;
-					GetterItem.bUpdatedThisFrame = false;
-				}
-				EnumRemoveFlags(AttributeMetaData->ResetFlag, EResetFlags::NeedToReset_ExceptVisibility);
-			}
 		}
 	}
 }
@@ -390,52 +295,30 @@ void FSlateAttributeMetaData::UpdateAttributesImpl(SWidget& OwningWidget, EInval
 	{
 		FGetterItem& GetterItem = Attributes[Index];
 
-		// Update every attribute at least once.
-		//Check if it has a dependency and if it was updated this frame (it could be from an UpdateNow)
-		if (GetterItem.CachedAttributeDependencyIndex != FGetterItem::InvalidAttributeIndex && GetterItem.bUpdatedOnce)
+		ISlateAttributeGetter::FUpdateAttributeResult Result = GetterItem.Getter->UpdateAttribute(OwningWidget);
+		if (Result.bInvalidationRequested)
 		{
-			// Note that the dependency is maybe not registered and the attribute may have been invalidated manually
-
-			// Because the list is sorted, the dependency needs to be before this element.
-			bool bShouldUpdate = false;
-			bool bFound = false;
-			for (int32 OtherIndex = Index-1; OtherIndex >= 0; --OtherIndex)
+			if (GetterItem.CachedAttributeDescriptor)
 			{
-				const FGetterItem& OtherGetterItem = Attributes[OtherIndex];
-				if (OtherGetterItem.CachedAttributeDescriptorIndex == GetterItem.CachedAttributeDependencyIndex)
+				GetterItem.CachedAttributeDescriptor->ExecuteOnValueChangedIfBound(OwningWidget);
+				if (bAllowInvalidation)
 				{
-					bFound = true;
-					bShouldUpdate = OtherGetterItem.bUpdatedThisFrame || OtherGetterItem.bUpdatedManually;
-					break;
+					InvalidationReason |= GetterItem.CachedAttributeDescriptor->GetInvalidationReason(OwningWidget);
 				}
 			}
-
-			if (!bShouldUpdate)
+			else if (bAllowInvalidation)
 			{
-				continue;
+				InvalidationReason |= Result.InvalidationReason;
 			}
 		}
-
-		ISlateAttributeGetter::FUpdateAttributeResult Result = GetterItem.Getter->UpdateAttribute(OwningWidget);
-		GetterItem.bUpdatedOnce = true;
-		GetterItem.bUpdatedThisFrame = Result.bInvalidationRequested;
-		if (Result.bInvalidationRequested && bAllowInvalidation)
-		{
-			SetNeedToResetFlag(Index);
-			const FGetterItem::FInvalidationDetail Detail = GetterItem.GetInvalidationDetail(OwningWidget, Result.InvalidationReason);
-			if (Detail.Get<0>())
-			{
-				Detail.Get<0>()->ExecuteIfBound(OwningWidget);
-			}
-			InvalidationReason |= Detail.Get<1>();
-		}
 	}
 
-	if (bInvalidateIfNeeded)
-	{
 #if WITH_SLATE_DEBUGGING
-		ensureAlwaysMsgf(FSlateAttributeBase::IsInvalidateWidgetReasonSupported(InvalidationReason | CachedInvalidationReason), TEXT("%s is not an EInvalidateWidgetReason supported by SlateAttribute."), *LexToString(InvalidationReason | CachedInvalidationReason));
+	ensureAlwaysMsgf(FSlateAttributeBase::IsInvalidateWidgetReasonSupported(InvalidationReason), TEXT("'%s' is not an EInvalidateWidgetReason supported by SlateAttribute."), *LexToString(InvalidationReason));
 #endif
+
+	if (bInvalidateIfNeeded)
+	{
 		OwningWidget.Invalidate(InvalidationReason | CachedInvalidationReason);
 		CachedInvalidationReason = EInvalidateWidgetReason::None;
 	}
@@ -458,32 +341,22 @@ void FSlateAttributeMetaData::UpdateAttribute(SWidget& OwningWidget, FSlateAttri
 		if (FoundIndex != INDEX_NONE)
 		{
 			FGetterItem& GetterItem = AttributeMetaData->Attributes[FoundIndex];
-			GetterItem.bUpdatedOnce = true;
 			check(GetterItem.Getter.Get());
 			ISlateAttributeGetter::FUpdateAttributeResult Result = GetterItem.Getter->UpdateAttribute(OwningWidget);
-			if (Result.bInvalidationRequested)
+			if (Result.bInvalidationRequested && OwningWidget.IsConstructed())
 			{
-				if (OwningWidget.IsConstructed())
+				EInvalidateWidgetReason InvalidationReason = Result.InvalidationReason;
+				if (GetterItem.CachedAttributeDescriptor)
 				{
-					const FGetterItem::FInvalidationDetail Detail = GetterItem.GetInvalidationDetail(OwningWidget, Result.InvalidationReason);
-					EInvalidateWidgetReason Reason = Detail.Get<1>() | AttributeMetaData->CachedInvalidationReason;
-#if WITH_SLATE_DEBUGGING
-					ensureAlwaysMsgf(FSlateAttributeBase::IsInvalidateWidgetReasonSupported(Reason), TEXT("%s is not an EInvalidateWidgetReason supported by SlateAttribute."), *LexToString(Reason));
-#endif
-					OwningWidget.Invalidate(Reason);
-					if (Detail.Get<0>())
-					{
-						Detail.Get<0>()->ExecuteIfBound(OwningWidget);
-					}
-					AttributeMetaData->CachedInvalidationReason = EInvalidateWidgetReason::None;
+					GetterItem.CachedAttributeDescriptor->ExecuteOnValueChangedIfBound(OwningWidget);
+					InvalidationReason = GetterItem.CachedAttributeDescriptor->GetInvalidationReason(OwningWidget);
 				}
 
-				// The dependency attribute need to be updated in the update loop (note that it may not be registered yet)
-				if (GetterItem.bIsADependencyForSomeoneElse)
-				{
-					GetterItem.bUpdatedManually = true;
-					AttributeMetaData->SetNeedToResetFlag(FoundIndex);
-				}
+#if WITH_SLATE_DEBUGGING
+				ensureAlwaysMsgf(FSlateAttributeBase::IsInvalidateWidgetReasonSupported(InvalidationReason), TEXT("%s is not an EInvalidateWidgetReason supported by SlateAttribute."), *LexToString(InvalidationReason));
+#endif
+				OwningWidget.Invalidate(InvalidationReason | AttributeMetaData->CachedInvalidationReason);
+				AttributeMetaData->CachedInvalidationReason = EInvalidateWidgetReason::None;
 			}
 		}
 	}
diff --git a/Engine/Source/Runtime/SlateCore/Public/Types/SlateAttributeDescriptor.h b/Engine/Source/Runtime/SlateCore/Public/Types/SlateAttributeDescriptor.h
--- a/Engine/Source/Runtime/SlateCore/Public/Types/SlateAttributeDescriptor.h
+++ b/Engine/Source/Runtime/SlateCore/Public/Types/SlateAttributeDescriptor.h
@@ -8,7 +8,13 @@
 
 #include <type_traits>
 
+
 class FSlateWidgetClassData;
+namespace SlateAttributePrivate
+{
+	enum class ESlateAttributeType : uint8;
+}
+
 
 /**
  * Describes the static information about a Widget's type SlateAttributes.
@@ -23,6 +29,8 @@ public:
 	 */
 	struct FInvalidateWidgetReasonAttribute
 	{
+		friend FSlateAttributeDescriptor;
+
 		using Arg1Type = const class SWidget&;
 		DECLARE_DELEGATE_RetVal_OneParam(EInvalidateWidgetReason, FGetter, Arg1Type);
 
@@ -82,6 +90,11 @@ public:
 	};
 
 public:
+	/** */
+	struct FAttribute;
+	struct FInitializer;
+
+	/** */
 	using OffsetType = uint32;
 
 	/** The default sort order that define in which order attributes will be updated. */
@@ -90,16 +103,29 @@ public:
 	/** */
 	struct FAttribute
 	{
+		friend FSlateAttributeDescriptor;
+		friend FInitializer;
+
+	public:
+		FAttribute(FName Name, OffsetType Offset, FInvalidateWidgetReasonAttribute Reason);
+
+		FName GetName() const { return Name; }
+		uint32 GetSortOrder() const { return SortOrder; }
+		EInvalidateWidgetReason GetInvalidationReason(const SWidget& Widget) const { return InvalidationReason.Get(Widget); }
+		SlateAttributePrivate::ESlateAttributeType GetAttributeType() const { return AttributeType; }
+		bool DoesAffectVisibility() const { return bAffectVisibility; }
+
+		void ExecuteOnValueChangedIfBound(SWidget& Widget) const { OnValueChanged.ExecuteIfBound(Widget); }
+
+	private:
 		FName Name;
-		OffsetType Offset = 0;
+		OffsetType Offset;
 		FName Prerequisite;
-		uint32 SortOrder = 0;
-		FInvalidateWidgetReasonAttribute InvalidationReason = FInvalidateWidgetReasonAttribute(EInvalidateWidgetReason::None);
+		uint32 SortOrder;
+		FInvalidateWidgetReasonAttribute InvalidationReason;
 		FAttributeValueChangedDelegate OnValueChanged;
-		bool bIsMemberAttribute = false;
-		bool bIsPrerequisiteAlsoADependency = false;
-		bool bIsADependencyForSomeoneElse = false;
-		bool bAffectVisibility = false;
+		SlateAttributePrivate::ESlateAttributeType AttributeType;
+		bool bAffectVisibility;
 	};
 
 	/** Internal class to initialize the SlateAttributeDescriptor (Add attributes or modify existing attributes). */
@@ -126,14 +152,6 @@ public:
 			 */
 			FAttributeEntry& UpdatePrerequisite(FName Prerequisite);
 
-			/**
-			 * The property only needs to be updated when the dependency changes inside the update loop.
-			 * The property can still be set/updated manually.
-			 * If the dependency is updated manually, then the property will be updated in the next update loop.
-			 * It will implicitly set a prerequisite.
-			 */
-			FAttributeEntry& UpdateDependency(FName Dependency);
-
 			/**
 			 * The attribute affect the visibility of the widget.
 			 * We only update the attributes that can change the visibility of the widget when the widget is collapsed.
@@ -183,26 +201,14 @@ public:
 	/** @returns the Attribute with the corresponding name. */
 	const FAttribute* FindAttribute(FName AttributeName) const;
 
-	/** @returns the index of a SlateAttribute that have the corresponding memory offset. */
-	int32 IndexOfMemberAttribute(OffsetType AttributeOffset) const;
-
-	/** @returns the index of a SlateAttribute that have the corresponding memory offset. */
-	int32 IndexOfMemberAttribute(FName AttributeName) const;
-
 	/** @returns the Attribute of a SlateAttribute that have the corresponding memory offset. */
 	const FAttribute* FindMemberAttribute(OffsetType AttributeOffset) const;
 
-	/** Iterator over each dependency this attribute is responsible of. */
-	template<typename Predicate>
-	void ForEachDependentsOn(FAttribute const& Attribute, Predicate Pred) const
-	{
-		checkf(&Attribute >= Attributes.GetData() && &Attribute <= Attributes.GetData()+Attributes.Num()
-			, TEXT("The attribute is not part of this Descriptor."));
-		if (Attribute.bIsADependencyForSomeoneElse)
-		{
-			ForEachDependentsOfImpl(Attribute.Name, (int32)(&Attribute - Attributes.GetData()), Pred);
-		}
-	}
+	/** @returns the index of a SlateAttribute that have the corresponding memory offset. */
+	int32 IndexOfAttribute(FName AttributeName) const;
+
+	/** @returns the index of a SlateAttribute that have the corresponding memory offset. */
+	int32 IndexOfMemberAttribute(OffsetType AttributeOffset) const;
 
 private:
 	FAttribute* FindAttribute(FName AttributeName);
@@ -210,27 +216,9 @@ private:
 	FInitializer::FAttributeEntry AddMemberAttribute(FName AttributeName, OffsetType Offset, FInvalidateWidgetReasonAttribute ReasonGetter);
 	void OverrideInvalidationReason(FName AttributeName, FInvalidateWidgetReasonAttribute ReasonGetter);
 	void OverrideOnValueChanged(FName AttributeName, ECallbackOverrideType OverrideType, FAttributeValueChangedDelegate Callback);
-	void SetPrerequisite(FAttribute& Attribute, FName Prerequisite, bool bSetAsDependency);
+	void SetPrerequisite(FAttribute& Attribute, FName Prerequisite);
 	void SetAffectVisibility(FAttribute& Attribute, bool bUpdate);
 
-	template<typename Predicate>
-	void ForEachDependentsOfImpl(FName const& LookForName, int32 Index, Predicate& Pred) const
-	{
-		++Index;
-		for (; Index < Attributes.Num(); ++Index)
-		{
-			FAttribute const& Other = Attributes[Index];
-			if (Other.Prerequisite == LookForName && Other.bIsPrerequisiteAlsoADependency)
-			{
-				Pred(Index);
-				if (Other.bIsADependencyForSomeoneElse)
-				{
-					ForEachDependentsOfImpl(Other.Name, Index, Pred);
-				}
-			}
-		}
-	}
-
 private:
 	TArray<FAttribute> Attributes;
 };
diff --git a/Engine/Source/Runtime/SlateCore/Public/Types/SlateAttributeMetaData.h b/Engine/Source/Runtime/SlateCore/Public/Types/SlateAttributeMetaData.h
--- a/Engine/Source/Runtime/SlateCore/Public/Types/SlateAttributeMetaData.h
+++ b/Engine/Source/Runtime/SlateCore/Public/Types/SlateAttributeMetaData.h
@@ -92,10 +92,6 @@ private:
 	void RegisterAttributeImpl(SWidget& OwningWidget, FSlateAttributeBase& Attribute, ESlateAttributeType AttributeType, TUniquePtr<ISlateAttributeGetter>&& Getter);
 	bool UnregisterAttributeImpl(const FSlateAttributeBase& Attribute);
 	void UpdateAttributesImpl(SWidget& OwningWidget, EInvalidationPermission InvaldiationStyle, int32 StartIndex, int32 IndexNum);
-	void SetNeedToResetFlag(int32 Index)
-	{
-		ResetFlag |= (Index < AffectVisibilityCounter) ? EResetFlags::NeedToReset_OnlyVisibility : EResetFlags::NeedToReset_ExceptVisibility;
-	}
 
 private:
 	int32 IndexOfAttribute(const FSlateAttributeBase& Attribute) const
@@ -106,57 +102,32 @@ private:
 
 	struct FGetterItem
 	{
-		using FAttributeIndex = uint8;
-		static const FAttributeIndex InvalidAttributeIndex;
-
 		FGetterItem(const FGetterItem&) = delete;
 		FGetterItem(FGetterItem&&) = default;
 		FGetterItem& operator=(const FGetterItem&) = delete;
 		FGetterItem(FSlateAttributeBase* InAttribute, uint32 InSortOrder, TUniquePtr<SlateAttributePrivate::ISlateAttributeGetter>&& InGetter)
 			: Attribute(InAttribute)
 			, Getter(MoveTemp(InGetter))
+			, CachedAttributeDescriptor(nullptr)
 			, SortOrder(InSortOrder)
-			, CachedAttributeDescriptorIndex(InvalidAttributeIndex)
-			, CachedAttributeDependencyIndex(InvalidAttributeIndex)
-			, Flags(0)
 		{ }
-		FGetterItem(FSlateAttributeBase* InAttribute, uint32 InSortOrder, TUniquePtr<SlateAttributePrivate::ISlateAttributeGetter> && InGetter, FAttributeIndex InAttributeDescriptorIndex)
+		FGetterItem(FSlateAttributeBase* InAttribute, uint32 InSortOrder, TUniquePtr<SlateAttributePrivate::ISlateAttributeGetter> && InGetter, const FSlateAttributeDescriptor::FAttribute* InAttributeDescriptor)
 			: Attribute(InAttribute)
 			, Getter(MoveTemp(InGetter))
+			, CachedAttributeDescriptor(InAttributeDescriptor)
 			, SortOrder(InSortOrder)
-			, CachedAttributeDescriptorIndex(InAttributeDescriptorIndex)
-			, CachedAttributeDependencyIndex(InvalidAttributeIndex)
-			, Flags(0)
 		{ }
 		FSlateAttributeBase* Attribute;
 		TUniquePtr<SlateAttributePrivate::ISlateAttributeGetter> Getter;
+		const FSlateAttributeDescriptor::FAttribute* CachedAttributeDescriptor;
 		uint32 SortOrder;
-		FAttributeIndex CachedAttributeDescriptorIndex;
-		FAttributeIndex CachedAttributeDependencyIndex;
-		union
-		{
-			struct
-			{
-				int8 bUpdatedOnce : 1;
-				int8 bUpdatedThisFrame : 1;
-				int8 bUpdatedManually : 1;
-				int8 bIsADependencyForSomeoneElse : 1;
-				int8 bAffectVisibility : 1;
-				ESlateAttributeType AttributeType : 2;
-				int8 Unused0 : 1;
-				//~if more space is needed bUpatedManually&bUpdatedThisFrame could be merged
-			};
-			int8 Flags;
-		};
+		ESlateAttributeType AttributeType;
 
 		bool operator<(const FGetterItem& Other) const
 		{
 			return SortOrder < Other.SortOrder;
 		}
 
-		using FInvalidationDetail = TTuple<const FSlateAttributeDescriptor::FAttributeValueChangedDelegate*, EInvalidateWidgetReason>;
-		FInvalidationDetail GetInvalidationDetail(const SWidget&, EInvalidateWidgetReason Reason) const;
-
 		/** If available, return the name of the attribute. */
 		FName GetAttributeName(const SWidget& OwningWidget) const;
 	};
@@ -169,15 +140,5 @@ private:
 	//~2. The parent widget will clear this widget PersistentState.
 	EInvalidateWidgetReason CachedInvalidationReason = EInvalidateWidgetReason::None;
 	uint8 AffectVisibilityCounter = 0;
-
-	enum class EResetFlags : uint8
-	{
-		None = 0,
-		NeedToReset_OnlyVisibility = (1 << 0),
-		NeedToReset_ExceptVisibility = (1 << 1),
-	};
-	FRIEND_ENUM_CLASS_FLAGS(EResetFlags);
-	EResetFlags ResetFlag = EResetFlags::None;
 };
 
-ENUM_CLASS_FLAGS(FSlateAttributeMetaData::EResetFlags);
