Slate: Add TSlateSlotAttribute, an equivalent of TSlateMemberAttribute but for FSlot. TSlateMemberAttribute only supports SWidget. Used an interface to support more types in the future. TSlateSlotAttribute are updated  in the order they are declared (like TSlateMemberAttribute) and by index in the TPanelChildren array.

diff --git a/Engine/Source/Runtime/SlateCore/Private/Types/SlateAttribute.cpp b/Engine/Source/Runtime/SlateCore/Private/Types/SlateAttribute.cpp
--- a/Engine/Source/Runtime/SlateCore/Private/Types/SlateAttribute.cpp
+++ b/Engine/Source/Runtime/SlateCore/Private/Types/SlateAttribute.cpp
@@ -9,6 +9,9 @@
 
 namespace SlateAttributePrivate
 {
+	/*
+	 *
+	 */
 	void TestAttributeAddress(const SWidget& OwningWidget, const FSlateAttributeImpl& Attribute, ESlateAttributeType AttributeType)
 	{
 #if STATS && DO_CHECK
@@ -26,6 +29,24 @@ namespace SlateAttributePrivate
 	}
 
 
+	/*
+	 *
+	 */
+	void ISlateAttributeContainer::RemoveContainerWidget(SWidget& Widget)
+	{
+		FSlateAttributeMetaData::RemoveContainerWidget(Widget, *this);
+	}
+
+
+	void ISlateAttributeContainer::UpdateContainerSortOrder(SWidget& Widget)
+	{
+		FSlateAttributeMetaData::UpdateContainerSortOrder(Widget, *this);
+	}
+
+
+	/*
+	 *
+	 */
 	bool FSlateAttributeImpl::ProtectedIsWidgetInDestructionPath(SWidget* Widget) const
 	{
 #if UE_WITH_SLATE_DEBUG_WIDGETLIST
@@ -76,6 +97,12 @@ namespace SlateAttributePrivate
 	}
 
 
+	void FSlateAttributeImpl::ProtectedInvalidateWidget(ISlateAttributeContainer& Container, ESlateAttributeType AttributeType, EInvalidateWidgetReason InvalidationReason) const
+	{
+		FSlateAttributeMetaData::InvalidateWidget(Container, *this, AttributeType, InvalidationReason);
+	}
+
+
 	bool FSlateAttributeImpl::ProtectedIsBound(const SWidget& OwningWidget, ESlateAttributeType AttributeType) const
 	{
 		TestAttributeAddress(OwningWidget, *this, AttributeType);
@@ -84,12 +111,24 @@ namespace SlateAttributePrivate
 	}
 
 
+	bool FSlateAttributeImpl::ProtectedIsBound(const ISlateAttributeContainer& Container, ESlateAttributeType AttributeType) const
+	{
+		return ProtectedIsBound(Container.GetContainerWidget(), AttributeType);
+	}
+
+
 	ISlateAttributeGetter* FSlateAttributeImpl::ProtectedFindGetter(const SWidget& OwningWidget, ESlateAttributeType AttributeType) const
 	{
 		TestAttributeAddress(OwningWidget, *this, AttributeType);
 
 		return FSlateAttributeMetaData::GetAttributeGetter(OwningWidget, *this);
 	}
+
+
+	ISlateAttributeGetter* FSlateAttributeImpl::ProtectedFindGetter(const ISlateAttributeContainer& Container, ESlateAttributeType AttributeType) const
+	{
+		return ProtectedFindGetter(Container.GetContainerWidget(), AttributeType);
+	}
 	
 	
 	FDelegateHandle FSlateAttributeImpl::ProtectedFindGetterHandle(const SWidget& OwningWidget, ESlateAttributeType AttributeType) const
@@ -98,7 +137,13 @@ namespace SlateAttributePrivate
 
 		return FSlateAttributeMetaData::GetAttributeGetterHandle(OwningWidget, *this);
 	}
-	
+
+
+	FDelegateHandle FSlateAttributeImpl::ProtectedFindGetterHandle(const ISlateAttributeContainer& Container, ESlateAttributeType AttributeType) const
+	{
+		return ProtectedFindGetterHandle(Container.GetContainerWidget(), AttributeType);
+	}
+
 
 	void FSlateAttributeImpl::ProtectedUpdateNow(SWidget& OwningWidget, ESlateAttributeType AttributeType)
 	{
@@ -107,9 +152,16 @@ namespace SlateAttributePrivate
 		FSlateAttributeMetaData::UpdateAttribute(OwningWidget, *this);
 	}
 
+
+	void FSlateAttributeImpl::ProtectedUpdateNow(ISlateAttributeContainer& Container, ESlateAttributeType AttributeType)
+	{
+		ProtectedUpdateNow(Container.GetContainerWidget(), AttributeType);
+	}
+
+
 	void FSlateAttributeImpl::ProtectedMoveAttribute(SWidget& OwningWidget, ESlateAttributeType AttributeType, const FSlateAttributeBase* Other)
 	{
-		checkf(AttributeType != ESlateAttributeType::Member, TEXT("Member Attribute cannot be moved."));
+		checkf(AttributeType == ESlateAttributeType::Managed, TEXT("Only Managed Attribute cannot be moved."));
 
 		// Other can't be pass as ref because it's already moved.
 		//We only need it's pointer to make sure it's not also used in FSlateAttributeMetaData (as a void*).
@@ -117,4 +169,16 @@ namespace SlateAttributePrivate
 
 		FSlateAttributeMetaData::MoveAttribute(OwningWidget, *this, AttributeType, Other);
 	}
+
+
+	void FSlateAttributeImpl::ProtectedUnregisterAttribute(ISlateAttributeContainer& Container, ESlateAttributeType AttributeType) const
+	{
+		FSlateAttributeMetaData::UnregisterAttribute(Container, *this);
+	}
+
+
+	void FSlateAttributeImpl::ProtectedRegisterAttribute(ISlateAttributeContainer& Container, ESlateAttributeType AttributeType, TUniquePtr<ISlateAttributeGetter>&& Wrapper)
+	{
+		FSlateAttributeMetaData::RegisterAttribute(Container, *this, AttributeType, MoveTemp(Wrapper));
+	}
 }
diff --git a/Engine/Source/Runtime/SlateCore/Private/Types/SlateAttributeDescriptor.cpp b/Engine/Source/Runtime/SlateCore/Private/Types/SlateAttributeDescriptor.cpp
--- a/Engine/Source/Runtime/SlateCore/Private/Types/SlateAttributeDescriptor.cpp
+++ b/Engine/Source/Runtime/SlateCore/Private/Types/SlateAttributeDescriptor.cpp
@@ -12,27 +12,69 @@ FSlateAttributeDescriptor::FAttribute::FAttribute(FName InName, OffsetType InOff
 	: Name(InName)
 	, Offset(InOffset)
 	, Prerequisite()
+	, ContainerName()
 	, SortOrder(DefaultSortOrder(InOffset))
 	, InvalidationReason(MoveTemp(InReason))
 	, OnValueChanged()
 	, AttributeType(SlateAttributePrivate::ESlateAttributeType::Member)
 	, bAffectVisibility(false)
 {
+}
 
+FSlateAttributeDescriptor::FAttribute::FAttribute(FName InContainerName, FName InName, OffsetType InOffset, FInvalidateWidgetReasonAttribute InReason)
+	: Name(InName)
+	, Offset(InOffset)
+	, Prerequisite()
+	, ContainerName(InContainerName)
+	, SortOrder(DefaultSortOrder(InOffset))
+	, InvalidationReason(MoveTemp(InReason))
+	, OnValueChanged()
+	, AttributeType(SlateAttributePrivate::ESlateAttributeType::Contained)
+	, bAffectVisibility(false)
+{
 }
 
+
+/** */
+FSlateAttributeDescriptor::FContainerInitializer::FAttributeEntry::FAttributeEntry(FSlateAttributeDescriptor & InDescriptor, FName InContainerName, int32 InAttributeIndex)
+	: Descriptor(InDescriptor)
+	, ContainerName(InContainerName)
+	, AttributeIndex(InAttributeIndex)
+{}
+
+
+FSlateAttributeDescriptor::FContainerInitializer::FAttributeEntry& FSlateAttributeDescriptor::FContainerInitializer::FAttributeEntry::UpdatePrerequisite(FName Prerequisite)
+{
+	if (Descriptor.Attributes.IsValidIndex(AttributeIndex))
+	{
+		Descriptor.SetPrerequisite(ContainerName, Descriptor.Attributes[AttributeIndex], Prerequisite);
+	}
+	return *this;
+}
+
+
+FSlateAttributeDescriptor::FContainerInitializer::FAttributeEntry& FSlateAttributeDescriptor::FContainerInitializer::FAttributeEntry::OnValueChanged(FAttributeValueChangedDelegate InCallback)
+{
+	if (Descriptor.Attributes.IsValidIndex(AttributeIndex))
+	{
+		Descriptor.Attributes[AttributeIndex].OnValueChanged = MoveTemp(InCallback);
+	}
+	return *this;
+}
+
+
 /** */
 FSlateAttributeDescriptor::FInitializer::FAttributeEntry::FAttributeEntry(FSlateAttributeDescriptor& InDescriptor, int32 InAttributeIndex)
 	: Descriptor(InDescriptor)
 	, AttributeIndex(InAttributeIndex)
 {}
 
-	
+
 FSlateAttributeDescriptor::FInitializer::FAttributeEntry& FSlateAttributeDescriptor::FInitializer::FAttributeEntry::UpdatePrerequisite(FName Prerequisite)
 {
 	if (Descriptor.Attributes.IsValidIndex(AttributeIndex))
 	{
-		Descriptor.SetPrerequisite(Descriptor.Attributes[AttributeIndex], Prerequisite);
+		Descriptor.SetPrerequisite(FName(), Descriptor.Attributes[AttributeIndex], Prerequisite);
 	}
 	return *this;
 }
@@ -58,6 +100,64 @@ FSlateAttributeDescriptor::FInitializer::FAttributeEntry& FSlateAttributeDescrip
 }
 
 
+/** */
+FSlateAttributeDescriptor::FContainerInitializer::FContainerInitializer(FSlateAttributeDescriptor& InDescriptor, FName InContainerName)
+	: Descriptor(InDescriptor)
+	, ContainerName(InContainerName)
+{
+}
+
+FSlateAttributeDescriptor::FContainerInitializer::FContainerInitializer(FSlateAttributeDescriptor& InDescriptor, const FSlateAttributeDescriptor& ParentDescriptor, FName InContainerName)
+	: Descriptor(InDescriptor)
+	, ContainerName(InContainerName)
+{
+	for (const FAttribute& Att : ParentDescriptor.Attributes)
+	{
+		if (Att.ContainerName == InContainerName)
+		{
+			bool bNameAlreadyExist = InDescriptor.Attributes.ContainsByPredicate([AttName = Att.GetName()](const FAttribute& Other)
+				{
+					return Other.GetName() == AttName;
+				});
+			if (ensureMsgf(!bNameAlreadyExist, TEXT("The attribute '%s' already exists in the FSlateAttributeDescriptor."), *Att.GetName().ToString()))
+			{
+				InDescriptor.Attributes.Add(Att);
+			}
+		}
+	}
+}
+
+
+FSlateAttributeDescriptor::FContainerInitializer::FAttributeEntry FSlateAttributeDescriptor::FContainerInitializer::AddContainedAttribute(FName AttributeName, OffsetType Offset, const FInvalidateWidgetReasonAttribute& Reason)
+{
+	return Descriptor.AddContainedAttribute(ContainerName, AttributeName, Offset, Reason);
+}
+
+
+FSlateAttributeDescriptor::FContainerInitializer::FAttributeEntry FSlateAttributeDescriptor::FContainerInitializer::AddContainedAttribute(FName AttributeName, OffsetType Offset, FInvalidateWidgetReasonAttribute&& Reason)
+{
+	return Descriptor.AddContainedAttribute(ContainerName, AttributeName, Offset, MoveTemp(Reason));
+}
+
+
+void FSlateAttributeDescriptor::FContainerInitializer::OverrideInvalidationReason(FName AttributeName, const FInvalidateWidgetReasonAttribute& Reason)
+{
+	Descriptor.OverrideInvalidationReason(ContainerName, AttributeName, Reason);
+}
+
+
+void FSlateAttributeDescriptor::FContainerInitializer::OverrideInvalidationReason(FName AttributeName, FInvalidateWidgetReasonAttribute&& Reason)
+{
+	Descriptor.OverrideInvalidationReason(ContainerName, AttributeName, MoveTemp(Reason));
+}
+
+
+void FSlateAttributeDescriptor::FContainerInitializer::OverrideOnValueChanged(FName AttributeName, ECallbackOverrideType OverrideType, FAttributeValueChangedDelegate Callback)
+{
+	Descriptor.OverrideOnValueChanged(ContainerName, AttributeName, OverrideType, MoveTemp(Callback));
+}
+
+
 /** */
 FSlateAttributeDescriptor::FInitializer::FInitializer(FSlateAttributeDescriptor& InDescriptor)
 	: Descriptor(InDescriptor)
@@ -68,14 +168,12 @@ FSlateAttributeDescriptor::FInitializer::FInitializer(FSlateAttributeDescriptor&
 	: Descriptor(InDescriptor)
 {
 	InDescriptor.Attributes = ParentDescriptor.Attributes;
+	InDescriptor.Containers = ParentDescriptor.Containers;
 }
 
 
 FSlateAttributeDescriptor::FInitializer::~FInitializer()
 {
-	checkf(Descriptor.Attributes.Num() <= std::numeric_limits<uint8>::max()-1
-		, TEXT("There are too many attributes '%d'. The index is saved as an uint8 in FSlateAttributeMetaData."), Descriptor.Attributes.Num());
-
 	// Confirm that the Visibility attribute is marked as "bAffectVisibility"
 	{
 		const FAttribute* FoundVisibilityAttribute = Descriptor.FindAttribute("Visibility");
@@ -85,8 +183,9 @@ FSlateAttributeDescriptor::FInitializer::~FInitializer()
 
 
 	// Update the sort order for the item that have prerequisite.
-	//Because adding the attribute is not required at the moment
+	//Because adding the attribute is not required at the moment,
 	//try to not change the order in which they were added.
+	//Because of that, the algo can creates bad ordering.
 	//The AffectVisibility attributes must be in the front of the list.
 
 	struct FPrerequisiteSort
@@ -119,6 +218,19 @@ FSlateAttributeDescriptor::FInitializer::~FInitializer()
 			const FAttribute& AttributeA = Attrbutes[A.AttributeIndex];
 			const FAttribute& AttributeB = Attrbutes[B.AttributeIndex];
 
+			if (AttributeA.ContainerIndex != AttributeB.ContainerIndex)
+			{
+				if (AttributeA.ContainerName.IsNone())
+				{
+					return true;
+				}
+				else if (AttributeB.ContainerName.IsNone())
+				{
+					return false;
+				}
+				return AttributeA.ContainerIndex < AttributeB.ContainerIndex;
+			}
+
 			if (AttributeA.bAffectVisibility != AttributeB.bAffectVisibility)
 			{
 				return AttributeA.bAffectVisibility;
@@ -140,6 +252,16 @@ FSlateAttributeDescriptor::FInitializer::~FInitializer()
 		}
 	};
 
+	if (Descriptor.Containers.Num() > 0)
+	{
+		for (FContainer& Container : Descriptor.Containers)
+		{
+			Container.SortOrder = DefaultSortOrder(Container.Offset);
+		}
+		Descriptor.Containers.Sort([](const FContainer& A, const FContainer& B) { return A.SortOrder < B.SortOrder; });
+	}
+
+
 	TArray<FPrerequisiteSort, TInlineAllocator<32>> Prerequisites;
 	Prerequisites.Reserve(Descriptor.Attributes.Num());
 
@@ -147,7 +269,19 @@ FSlateAttributeDescriptor::FInitializer::~FInitializer()
 	for (int32 Index = 0; Index < Descriptor.Attributes.Num(); ++Index)
 	{
 		FAttribute& Attribute = Descriptor.Attributes[Index];
-		Attribute.SortOrder = DefaultSortOrder(Attribute.Offset);
+
+		if (!Attribute.ContainerName.IsNone())
+		{
+			const int32 FoundIndex = Descriptor.Containers.IndexOfByPredicate([ContainerName = Attribute.ContainerName](const FContainer& Other) { return Other.GetName() == ContainerName; });
+			check(FoundIndex != INDEX_NONE);
+			Attribute.ContainerIndex = (uint8)FoundIndex;
+			Attribute.SortOrder = Descriptor.Containers[FoundIndex].SortOrder + (Attribute.Offset * 10);
+		}
+		else
+		{
+			Attribute.ContainerIndex = std::numeric_limits<uint8>::max();
+			Attribute.SortOrder = DefaultSortOrder(Attribute.Offset);
+		}
 
 		if (!Attribute.Prerequisite.IsNone())
 		{
@@ -198,7 +332,7 @@ FSlateAttributeDescriptor::FInitializer::~FInitializer()
 	}
 
 	// Confirm that the attributes marked as "AffectVisibility" are in front of the list.
-#if 1
+#if 0
 	{
 		bool bLookingForAffectVisibility = true;
 		for (const FAttribute& Attribute : Descriptor.Attributes)
@@ -228,22 +362,26 @@ FSlateAttributeDescriptor::FInitializer::FAttributeEntry FSlateAttributeDescript
 	return Descriptor.AddMemberAttribute(AttributeName, Offset, MoveTemp(Reason));
 }
 
+FSlateAttributeDescriptor::FContainerInitializer FSlateAttributeDescriptor::FInitializer::AddContainer(FName ContainerName, OffsetType Offset)
+{
+	return Descriptor.AddContainer(ContainerName, Offset);
+}
 
 void FSlateAttributeDescriptor::FInitializer::OverrideInvalidationReason(FName AttributeName, const FInvalidateWidgetReasonAttribute& Reason)
 {
-	Descriptor.OverrideInvalidationReason(AttributeName, Reason);
+	Descriptor.OverrideInvalidationReason(FName(), AttributeName, Reason);
 }
 
 
 void FSlateAttributeDescriptor::FInitializer::OverrideInvalidationReason(FName AttributeName, FInvalidateWidgetReasonAttribute&& Reason)
 {
-	Descriptor.OverrideInvalidationReason(AttributeName, MoveTemp(Reason));
+	Descriptor.OverrideInvalidationReason(FName(), AttributeName, MoveTemp(Reason));
 }
 
 
 void FSlateAttributeDescriptor::FInitializer::OverrideOnValueChanged(FName AttributeName, ECallbackOverrideType OverrideType, FAttributeValueChangedDelegate Callback)
 {
-	Descriptor.OverrideOnValueChanged(AttributeName, OverrideType, MoveTemp(Callback));
+	Descriptor.OverrideOnValueChanged(FName(), AttributeName, OverrideType, MoveTemp(Callback));
 }
 
 
@@ -266,7 +404,13 @@ FSlateAttributeDescriptor::FAttribute const& FSlateAttributeDescriptor::GetAttri
 }
 
 
-FSlateAttributeDescriptor::FAttribute const* FSlateAttributeDescriptor::FindAttribute(FName AttributeName) const
+const FSlateAttributeDescriptor::FContainer* FSlateAttributeDescriptor::FindContainer(FName ContainerName) const
+{
+	return Containers.FindByPredicate([ContainerName](const FContainer& Other) { return Other.GetName() == ContainerName; });
+}
+
+
+const FSlateAttributeDescriptor::FAttribute* FSlateAttributeDescriptor::FindAttribute(FName AttributeName) const
 {
 	return Attributes.FindByPredicate([AttributeName](const FAttribute& Other) { return Other.Name == AttributeName; });
 }
@@ -278,9 +422,9 @@ FSlateAttributeDescriptor::FAttribute* FSlateAttributeDescriptor::FindAttribute(
 }
 
 
-FSlateAttributeDescriptor::FAttribute const* FSlateAttributeDescriptor::FindMemberAttribute(OffsetType AttributeOffset) const
+const FSlateAttributeDescriptor::FAttribute* FSlateAttributeDescriptor::FindMemberAttribute(OffsetType AttributeOffset) const
 {
-	FSlateAttributeDescriptor::FAttribute const* Result = Attributes.FindByPredicate(
+	const FSlateAttributeDescriptor::FAttribute* Result = Attributes.FindByPredicate(
 		[AttributeOffset](const FAttribute& Other)
 		{
 			return Other.Offset == AttributeOffset
@@ -290,6 +434,25 @@ FSlateAttributeDescriptor::FAttribute const* FSlateAttributeDescriptor::FindMemb
 }
 
 
+const FSlateAttributeDescriptor::FAttribute* FSlateAttributeDescriptor::FindContainedAttribute(FName ContainerName, OffsetType AttributeOffset) const
+{
+	const FSlateAttributeDescriptor::FAttribute* Result = Attributes.FindByPredicate(
+		[AttributeOffset, ContainerName](const FAttribute& Other)
+		{
+			return Other.Offset == AttributeOffset
+				&& Other.AttributeType == SlateAttributePrivate::ESlateAttributeType::Contained
+				&& Other.ContainerName == ContainerName;
+		});
+	return Result;
+}
+
+
+int32 FSlateAttributeDescriptor::IndexOfContainer(FName ContainerName) const
+{
+	return Containers.IndexOfByPredicate([ContainerName](const FContainer& Other) { return Other.GetName() == ContainerName; });
+}
+
+
 int32 FSlateAttributeDescriptor::IndexOfAttribute(FName AttributeName) const
 {
 	return Attributes.IndexOfByPredicate([AttributeName](const FAttribute& Other) { return Other.Name == AttributeName; });
@@ -308,6 +471,32 @@ int32 FSlateAttributeDescriptor::IndexOfMemberAttribute(OffsetType AttributeOffs
 }
 
 
+int32 FSlateAttributeDescriptor::IndexOfContainedAttribute(FName ContainerName, OffsetType AttributeOffset) const
+{
+	int32 FoundIndex = Attributes.IndexOfByPredicate(
+		[AttributeOffset, ContainerName](const FAttribute& Other)
+		{
+			return Other.Offset == AttributeOffset
+				&& Other.AttributeType == SlateAttributePrivate::ESlateAttributeType::Contained
+				&& Other.ContainerName == ContainerName;
+		});
+	return FoundIndex;
+}
+
+
+FSlateAttributeDescriptor::FContainerInitializer FSlateAttributeDescriptor::AddContainer(FName ContainerName, OffsetType Offset)
+{
+	check(!ContainerName.IsNone());
+
+	const FContainer* FoundAttribute = FindContainer(ContainerName);
+	if (ensureAlwaysMsgf(FoundAttribute == nullptr, TEXT("The container '%s' already exist. (Do you have the correct parrent class in SLATE_DECLARE_WIDGET)"), *ContainerName.ToString()))
+	{
+		Containers.Emplace(ContainerName, Offset);
+	}
+	return FContainerInitializer(*this, ContainerName);
+}
+
+
 FSlateAttributeDescriptor::FInitializer::FAttributeEntry FSlateAttributeDescriptor::AddMemberAttribute(FName AttributeName, OffsetType Offset, FInvalidateWidgetReasonAttribute Reason)
 {
 	check(!AttributeName.IsNone());
@@ -322,61 +511,82 @@ FSlateAttributeDescriptor::FInitializer::FAttributeEntry FSlateAttributeDescript
 }
 
 
-void FSlateAttributeDescriptor::OverrideInvalidationReason(FName AttributeName, FInvalidateWidgetReasonAttribute Reason)
+FSlateAttributeDescriptor::FContainerInitializer::FAttributeEntry FSlateAttributeDescriptor::AddContainedAttribute(FName ContainerName, FName AttributeName, OffsetType Offset, FInvalidateWidgetReasonAttribute Reason)
+{
+	check(!AttributeName.IsNone());
+	check(!ContainerName.IsNone());
+
+	int32 NewIndex = INDEX_NONE;
+	FAttribute const* FoundAttribute = FindAttribute(AttributeName);
+	if (ensureAlwaysMsgf(FoundAttribute == nullptr, TEXT("The attribute '%s' already exist. (Do you have the correct parrent class in SLATE_DECLARE_WIDGET)"), *AttributeName.ToString()))
+	{
+		NewIndex = Attributes.Emplace(ContainerName, AttributeName, Offset, MoveTemp(Reason));
+	}
+	return FContainerInitializer::FAttributeEntry(*this, ContainerName, NewIndex);
+}
+
+
+void FSlateAttributeDescriptor::OverrideInvalidationReason(FName ContainerName, FName AttributeName, FInvalidateWidgetReasonAttribute Reason)
 {
 	check(!AttributeName.IsNone());
 
 	FAttribute* FoundAttribute = FindAttribute(AttributeName);
 	if (ensureAlwaysMsgf(FoundAttribute != nullptr, TEXT("The attribute 's' doesn't exist."), *AttributeName.ToString()))
 	{
-		FoundAttribute->InvalidationReason = MoveTemp(Reason);
+		if (ensureAlwaysMsgf(FoundAttribute->ContainerName == ContainerName, TEXT("You are not allowed to modify an attribute from a different container.")))
+		{
+			FoundAttribute->InvalidationReason = MoveTemp(Reason);
+		}
 	}
 }
 
 
-void FSlateAttributeDescriptor::OverrideOnValueChanged(FName AttributeName, ECallbackOverrideType OverrideType, FAttributeValueChangedDelegate Callback)
+void FSlateAttributeDescriptor::OverrideOnValueChanged(FName ContainerName, FName AttributeName, ECallbackOverrideType OverrideType, FAttributeValueChangedDelegate Callback)
 {
 	check(!AttributeName.IsNone());
 
 	FAttribute* FoundAttribute = FindAttribute(AttributeName);
 	if (ensureAlwaysMsgf(FoundAttribute != nullptr, TEXT("The attribute 's' doesn't exist."), *AttributeName.ToString()))
 	{
-		switch(OverrideType)
+		if (ensureAlwaysMsgf(FoundAttribute->ContainerName == ContainerName, TEXT("You are not allowed to modify an attribute from a different container.")))
 		{
-		case ECallbackOverrideType::ReplacePrevious:
-			FoundAttribute->OnValueChanged = MoveTemp(Callback);
-			break;
-		case ECallbackOverrideType::ExecuteAfterPrevious:
-		case ECallbackOverrideType::ExecuteBeforePrevious:
-			if (FoundAttribute->OnValueChanged.IsBound() && Callback.IsBound())
-			{
-				FAttributeValueChangedDelegate Previous = FoundAttribute->OnValueChanged;
-				FoundAttribute->OnValueChanged = FAttributeValueChangedDelegate::CreateLambda([Previous{MoveTemp(Previous)}, Callback{MoveTemp(Callback)}, OverrideType](SWidget& Widget)
-					{
-						if (OverrideType == ECallbackOverrideType::ExecuteBeforePrevious)
-						{
-							Previous.ExecuteIfBound(Widget);
-						}
-						Callback.ExecuteIfBound(Widget);
-						if (OverrideType == ECallbackOverrideType::ExecuteAfterPrevious)
-						{
-							Previous.ExecuteIfBound(Widget);
-						}
-					});
-			}
-			else if (Callback.IsBound())
+			switch(OverrideType)
 			{
+			case ECallbackOverrideType::ReplacePrevious:
 				FoundAttribute->OnValueChanged = MoveTemp(Callback);
+				break;
+			case ECallbackOverrideType::ExecuteAfterPrevious:
+			case ECallbackOverrideType::ExecuteBeforePrevious:
+				if (FoundAttribute->OnValueChanged.IsBound() && Callback.IsBound())
+				{
+					FAttributeValueChangedDelegate Previous = FoundAttribute->OnValueChanged;
+					FoundAttribute->OnValueChanged = FAttributeValueChangedDelegate::CreateLambda([Previous{MoveTemp(Previous)}, Callback{MoveTemp(Callback)}, OverrideType](SWidget& Widget)
+						{
+							if (OverrideType == ECallbackOverrideType::ExecuteBeforePrevious)
+							{
+								Previous.ExecuteIfBound(Widget);
+							}
+							Callback.ExecuteIfBound(Widget);
+							if (OverrideType == ECallbackOverrideType::ExecuteAfterPrevious)
+							{
+								Previous.ExecuteIfBound(Widget);
+							}
+						});
+				}
+				else if (Callback.IsBound())
+				{
+					FoundAttribute->OnValueChanged = MoveTemp(Callback);
+				}
+				break;
+			default:
+				check(false);
 			}
-			break;
-		default:
-			check(false);
 		}
 	}
 }
 
 
-void FSlateAttributeDescriptor::SetPrerequisite(FSlateAttributeDescriptor::FAttribute& Attribute, FName Prerequisite)
+void FSlateAttributeDescriptor::SetPrerequisite(FName ContainerName, FSlateAttributeDescriptor::FAttribute& Attribute, FName Prerequisite)
 {
 	if (Prerequisite.IsNone())
 	{
@@ -384,27 +594,30 @@ void FSlateAttributeDescriptor::SetPrerequisite(FSlateAttributeDescriptor::FAttr
 	}
 	else
 	{
-		FAttribute const* FoundPrerequisite = FindAttribute(Prerequisite);
+		const FAttribute* FoundPrerequisite = FindAttribute(Prerequisite);
 		if (ensureAlwaysMsgf(FoundPrerequisite, TEXT("The prerequisite '%s' doesn't exist for attribute '%s'"), *Prerequisite.ToString(), *Attribute.Name.ToString()))
 		{
-			Attribute.Prerequisite = Prerequisite;
-
-			// Verify recursion
+			if (ensureAlwaysMsgf(FoundPrerequisite->ContainerName == ContainerName, TEXT("You are not allowed to set a prerequisite from a different container.")))
 			{
-				TArray<FName, TInlineAllocator<16>> Recursion;
-				Recursion.Reserve(Attributes.Num());
-				FAttribute const* CurrentAttribute = &Attribute;
-				while (!CurrentAttribute->Prerequisite.IsNone())
+				Attribute.Prerequisite = Prerequisite;
+
+				// Verify recursion
 				{
-					if (Recursion.Contains(CurrentAttribute->Name))
+					TArray<FName, TInlineAllocator<16>> Recursion;
+					Recursion.Reserve(Attributes.Num());
+					FAttribute const* CurrentAttribute = &Attribute;
+					while (!CurrentAttribute->Prerequisite.IsNone())
 					{
-						ensureAlwaysMsgf(false, TEXT("The prerequsite '%s' would introduce an infinit loop with attribute '%s'."), *Prerequisite.ToString(), *Attribute.Name.ToString());
-						Attribute.Prerequisite = FName();
-						break;
+						if (Recursion.Contains(CurrentAttribute->Name))
+						{
+							ensureAlwaysMsgf(false, TEXT("The prerequsite '%s' would introduce an infinit loop with attribute '%s'."), *Prerequisite.ToString(), *Attribute.Name.ToString());
+							Attribute.Prerequisite = FName();
+							break;
+						}
+						Recursion.Add(CurrentAttribute->Name);
+						CurrentAttribute = FindAttribute(CurrentAttribute->Prerequisite);
+						check(CurrentAttribute);
 					}
-					Recursion.Add(CurrentAttribute->Name);
-					CurrentAttribute = FindAttribute(CurrentAttribute->Prerequisite);
-					check(CurrentAttribute);
 				}
 			}
 		}
diff --git a/Engine/Source/Runtime/SlateCore/Private/Types/SlateAttributeMetaData.cpp b/Engine/Source/Runtime/SlateCore/Private/Types/SlateAttributeMetaData.cpp
--- a/Engine/Source/Runtime/SlateCore/Private/Types/SlateAttributeMetaData.cpp
+++ b/Engine/Source/Runtime/SlateCore/Private/Types/SlateAttributeMetaData.cpp
@@ -9,15 +9,66 @@
 
 #include <limits>
 
-
+namespace UE
+{
 namespace Private
 {
-	FSlateAttributeDescriptor::OffsetType FindOffet(const SWidget& OwningWidget, const FSlateAttributeBase& Attribute)
+	FSlateAttributeDescriptor::OffsetType FindMemberOffset(const SWidget& OwningWidget, const FSlateAttributeBase& Attribute)
 	{
 		UPTRINT Offset = (UPTRINT)(&Attribute) - (UPTRINT)(&OwningWidget);
 		ensure(Offset <= std::numeric_limits<FSlateAttributeDescriptor::OffsetType>::max());
 		return (FSlateAttributeDescriptor::OffsetType)(Offset);
 	}
+
+	FSlateAttributeDescriptor::OffsetType FindContainOffset(const SlateAttributePrivate::ISlateAttributeContainer& OwningContainer, const FSlateAttributeBase& Attribute)
+	{
+		UPTRINT Offset = (UPTRINT)(&Attribute) - (UPTRINT)(&OwningContainer);
+		ensure(Offset <= std::numeric_limits<FSlateAttributeDescriptor::OffsetType>::max());
+		return (FSlateAttributeDescriptor::OffsetType)(Offset);
+	}
+
+	EInvalidateWidgetReason InvalidateForMember(SWidget& OwningWidget, const FSlateAttributeBase& Attribute, EInvalidateWidgetReason Reason)
+	{
+		const FSlateAttributeDescriptor::OffsetType Offset = FindMemberOffset(OwningWidget, Attribute);
+		if (const FSlateAttributeDescriptor::FAttribute* FoundAttribute = OwningWidget.GetWidgetClass().GetAttributeDescriptor().FindMemberAttribute(Offset))
+		{
+			ensureMsgf(Reason == EInvalidateWidgetReason::None, TEXT("SWidget is using the AttributeDescriptor, the invalidation reason should only be defined in one place. See SLATE_ADD_MEMBER_ATTRIBUTE_DEFINITION"));
+			FoundAttribute->ExecuteOnValueChangedIfBound(OwningWidget);
+			return FoundAttribute->GetInvalidationReason(OwningWidget);
+		}
+		return Reason;
+	}
+
+	EInvalidateWidgetReason InvalidateForContained(const SlateAttributePrivate::ISlateAttributeContainer& OwningContainer, SWidget& OwningWidget, const FSlateAttributeBase& Attribute, EInvalidateWidgetReason Reason)
+	{
+		const FSlateAttributeDescriptor::OffsetType Offset = FindContainOffset(OwningContainer, Attribute);
+		if (const FSlateAttributeDescriptor::FAttribute* FoundAttribute = OwningWidget.GetWidgetClass().GetAttributeDescriptor().FindContainedAttribute(OwningContainer.GetContainerName(), Offset))
+		{
+			ensureMsgf(Reason == EInvalidateWidgetReason::None, TEXT("SWidget is using the AttributeDescriptor, the invalidation reason should only be defined in one place. See SLATE_ADD_MEMBER_ATTRIBUTE_DEFINITION"));
+			FoundAttribute->ExecuteOnValueChangedIfBound(OwningWidget);
+			return FoundAttribute->GetInvalidationReason(OwningWidget);
+		}
+		return Reason;
+	}
+
+	int32 FindAttributeContainerOffset(const SlateAttributePrivate::ISlateAttributeContainer* Container, FSlateAttributeBase* Attribute)
+	{
+		PTRINT AttributePtr = (PTRINT)Attribute;
+		PTRINT ContainerPtr = (PTRINT)Container;
+		return AttributePtr - ContainerPtr;
+	}
+}
+}
+
+
+const SlateAttributePrivate::ISlateAttributeContainer* FSlateAttributeMetaData::FGetterItem::GetAttributeContainer() const
+{
+	if (AttributeContainerOffset != 0)
+	{
+		PTRINT AttributePtr = (PTRINT)Attribute;
+		return reinterpret_cast<const SlateAttributePrivate::ISlateAttributeContainer*>(AttributePtr - AttributeContainerOffset);
+	}
+	return nullptr;
 }
 
 
@@ -47,66 +98,125 @@ FSlateAttributeMetaData* FSlateAttributeMetaData::FindMetaData(const SWidget& Ow
 
 void FSlateAttributeMetaData::RegisterAttribute(SWidget& OwningWidget, FSlateAttributeBase& Attribute, ESlateAttributeType AttributeType, TUniquePtr<ISlateAttributeGetter>&& Wrapper)
 {
+	auto ExecuteRegister = [&](FSlateAttributeMetaData& AttributeMetaData)
+	{
+		if (AttributeType == ESlateAttributeType::Member)
+		{
+			AttributeMetaData.RegisterMemberAttributeImpl(OwningWidget, Attribute, MoveTemp(Wrapper));
+		}
+		else
+		{
+			check(AttributeType == ESlateAttributeType::Managed);
+			AttributeMetaData.RegisterManagedAttributeImpl(OwningWidget, Attribute, MoveTemp(Wrapper));
+		}
+	};
+
 	if (FSlateAttributeMetaData* AttributeMetaData = FSlateAttributeMetaData::FindMetaData(OwningWidget))
 	{
-		AttributeMetaData->RegisterAttributeImpl(OwningWidget, Attribute, AttributeType, MoveTemp(Wrapper));
+		if (int32 FoundIndex = AttributeMetaData->IndexOfAttribute(Attribute); FoundIndex != INDEX_NONE)
+		{
+			AttributeMetaData->Attributes[FoundIndex].Getter = MoveTemp(Wrapper);
+		}
+		else
+		{
+			ExecuteRegister(*AttributeMetaData);
+		}
 	}
 	else
 	{
 		TSharedRef<FSlateAttributeMetaData> NewAttributeMetaData = MakeShared<FSlateAttributeMetaData>();
-		NewAttributeMetaData->RegisterAttributeImpl(OwningWidget, Attribute, AttributeType, MoveTemp(Wrapper));
-		OwningWidget.bHasRegisteredSlateAttribute = true;
-		OwningWidget.MetaData.Insert(NewAttributeMetaData, 0);
-		if (OwningWidget.IsConstructed() && OwningWidget.IsAttributesUpdatesEnabled())
+		ExecuteRegister(NewAttributeMetaData.Get());
+		if (NewAttributeMetaData->GetRegisteredAttributeCount())
 		{
-			OwningWidget.Invalidate(EInvalidateWidgetReason::AttributeRegistration);
+			OwningWidget.bHasRegisteredSlateAttribute = true;
+			OwningWidget.MetaData.Insert(NewAttributeMetaData, 0);
+			if (OwningWidget.IsConstructed() && OwningWidget.IsAttributesUpdatesEnabled())
+			{
+				OwningWidget.Invalidate(EInvalidateWidgetReason::AttributeRegistration);
+			}
 		}
 	}
 }
 
 
-void FSlateAttributeMetaData::RegisterAttributeImpl(SWidget& OwningWidget, FSlateAttributeBase& Attribute, ESlateAttributeType AttributeType, TUniquePtr<ISlateAttributeGetter>&& Getter)
+void FSlateAttributeMetaData::RegisterAttribute(ISlateAttributeContainer& OwningContainer, FSlateAttributeBase& Attribute, ESlateAttributeType AttributeType, TUniquePtr<ISlateAttributeGetter>&& Wrapper)
 {
-	const int32 FoundIndex = IndexOfAttribute(Attribute);
-	if (FoundIndex != INDEX_NONE)
+	SWidget& OwningWidget = OwningContainer.GetContainerWidget();
+	if (FSlateAttributeMetaData* AttributeMetaData = FSlateAttributeMetaData::FindMetaData(OwningWidget))
 	{
-		Attributes[FoundIndex].Getter = MoveTemp(Getter);
+		if (int32 FoundIndex = AttributeMetaData->IndexOfAttribute(Attribute); FoundIndex != INDEX_NONE)
+		{
+			AttributeMetaData->Attributes[FoundIndex].Getter = MoveTemp(Wrapper);
+		}
+		else
+		{
+			AttributeMetaData->RegisterContainAttributeImpl(OwningContainer, OwningWidget, Attribute, MoveTemp(Wrapper));
+		}
 	}
 	else
 	{
-		if (AttributeType == ESlateAttributeType::Member)
+		TSharedRef<FSlateAttributeMetaData> NewAttributeMetaData = MakeShared<FSlateAttributeMetaData>();
+		NewAttributeMetaData->RegisterContainAttributeImpl(OwningContainer, OwningWidget, Attribute, MoveTemp(Wrapper));
+		if (NewAttributeMetaData->GetRegisteredAttributeCount())
 		{
-			// MemberAttribute are optional for now but will be needed in the future
-			const FSlateAttributeDescriptor::OffsetType Offset = Private::FindOffet(OwningWidget, Attribute);
-			const FSlateAttributeDescriptor::FAttribute* FoundMemberAttribute = OwningWidget.GetWidgetClass().GetAttributeDescriptor().FindMemberAttribute(Offset);
-			if (FoundMemberAttribute)
-			{
-				const int32 InsertLocation = Algo::LowerBoundBy(Attributes, FoundMemberAttribute->GetSortOrder(), [](const FGetterItem& Item) { return Item.SortOrder; }, TLess<>());
-				FGetterItem& GetterItem = Attributes.Insert_GetRef({ &Attribute, FoundMemberAttribute->GetSortOrder(), MoveTemp(Getter), FoundMemberAttribute }, InsertLocation);
-				GetterItem.AttributeType = ESlateAttributeType::Member;
-				if (FoundMemberAttribute->DoesAffectVisibility())
-				{
-					++AffectVisibilityCounter;
-				}
-			}
-			else
+			OwningWidget.bHasRegisteredSlateAttribute = true;
+			OwningWidget.MetaData.Insert(NewAttributeMetaData, 0);
+			if (OwningWidget.IsConstructed() && OwningWidget.IsAttributesUpdatesEnabled())
 			{
-				const uint32 SortOrder = FSlateAttributeDescriptor::DefaultSortOrder(Offset);
-
-				const  int32 InsertLocation = Algo::LowerBoundBy(Attributes, SortOrder, [](const FGetterItem& Item) { return Item.SortOrder; }, TLess<>());
-				FGetterItem& GetterItem = Attributes.Insert_GetRef({&Attribute, SortOrder, MoveTemp(Getter)}, InsertLocation);
-				GetterItem.AttributeType = ESlateAttributeType::Member;
+				OwningWidget.Invalidate(EInvalidateWidgetReason::AttributeRegistration);
 			}
 		}
-		else if (AttributeType == ESlateAttributeType::Managed)
+	}
+}
+
+
+void FSlateAttributeMetaData::RegisterMemberAttributeImpl(SWidget& OwningWidget, FSlateAttributeBase& Attribute, TUniquePtr<ISlateAttributeGetter>&& Getter)
+{
+	// MemberAttribute are optional for now but will be needed in the future
+	const FSlateAttributeDescriptor::OffsetType AttributeOffset = UE::Private::FindMemberOffset(OwningWidget, Attribute);
+	const FSlateAttributeDescriptor::FAttribute* FoundMemberAttribute = OwningWidget.GetWidgetClass().GetAttributeDescriptor().FindMemberAttribute(AttributeOffset);
+	if (FoundMemberAttribute)
+	{
+		const int32 InsertLocation = Algo::LowerBoundBy(Attributes, FoundMemberAttribute->GetSortOrder(), [](const FGetterItem& Item) { return Item.SortOrder; }, TLess<>());
+		FGetterItem& GetterItem = Attributes.Insert_GetRef({ &Attribute, FoundMemberAttribute->GetSortOrder(), MoveTemp(Getter), FoundMemberAttribute }, InsertLocation);
+		if (FoundMemberAttribute->DoesAffectVisibility())
 		{
-			const uint32 ManagedSortOrder = std::numeric_limits<uint32>::max();
-			FGetterItem& GetterItem = Attributes.Emplace_GetRef(&Attribute, ManagedSortOrder, MoveTemp(Getter));
-			GetterItem.AttributeType = ESlateAttributeType::Managed;
+			++AffectVisibilityCounter;
 		}
-		else
+	}
+	else
+	{
+		const uint32 SortOrder = FSlateAttributeDescriptor::DefaultSortOrder(AttributeOffset);
+		const  int32 InsertLocation = Algo::LowerBoundBy(Attributes, SortOrder, [](const FGetterItem& Item) { return Item.SortOrder; }, TLess<>());
+		Attributes.Insert({ &Attribute, SortOrder, MoveTemp(Getter) }, InsertLocation);
+	}
+}
+
+
+void FSlateAttributeMetaData::RegisterManagedAttributeImpl(SWidget& OwningWidget, FSlateAttributeBase& Attribute, TUniquePtr<ISlateAttributeGetter>&& Getter)
+{
+	const uint32 ManagedSortOrder = std::numeric_limits<uint32>::max();
+	Attributes.Emplace(&Attribute, ManagedSortOrder, MoveTemp(Getter));
+}
+
+
+void FSlateAttributeMetaData::RegisterContainAttributeImpl(ISlateAttributeContainer& OwningContainer, SWidget& OwningWidget, FSlateAttributeBase& Attribute, TUniquePtr<ISlateAttributeGetter>&& Getter)
+{
+	const FSlateAttributeDescriptor& Descriptor = OwningWidget.GetWidgetClass().GetAttributeDescriptor();
+	const FName ContainerName = OwningContainer.GetContainerName();
+	const FSlateAttributeDescriptor::OffsetType AttributeOffset = UE::Private::FindContainOffset(OwningContainer, Attribute);
+	const FSlateAttributeDescriptor::FAttribute* FoundAttribute = Descriptor.FindContainedAttribute(ContainerName, AttributeOffset);
+	ensureAlwaysMsgf(FoundAttribute, TEXT("Attribute with offset '%d' is not registered. All Slot Attributes needs to be registered."));
+	if (FoundAttribute)
+	{
+		const FSlateAttributeDescriptor::FContainer* FoundCountainer = Descriptor.FindContainer(ContainerName);
+		ensureAlwaysMsgf(FoundCountainer, TEXT("The container '%s' doesn't not exist."), *ContainerName.ToString());
+		if (FoundCountainer)
 		{
-			ensureMsgf(false, TEXT("The SlateAttributeType is not supported"));
+			FGetterItem NewItem{ &Attribute, FoundAttribute->GetSortOrder(), MoveTemp(Getter), FoundAttribute };
+			NewItem.AttributeContainerOffset = UE::Private::FindAttributeContainerOffset(&OwningContainer, &Attribute);
+			const int32 InsertLocation = Algo::LowerBound(Attributes, NewItem, TLess<>());
+			Attributes.Insert(MoveTemp(NewItem), InsertLocation);
 		}
 	}
 }
@@ -116,23 +226,37 @@ bool FSlateAttributeMetaData::UnregisterAttribute(SWidget& OwningWidget, const F
 {
 	if (FSlateAttributeMetaData* AttributeMetaData = FSlateAttributeMetaData::FindMetaData(OwningWidget))
 	{
+		check(AttributeMetaData->Attributes.Num() != 0);
 		const bool bResult = AttributeMetaData->UnregisterAttributeImpl(Attribute);
-		if (AttributeMetaData->Attributes.Num() == 0)
-		{
-			check(bResult); // if the num is 0 then we should have remove an item.
-			OwningWidget.bHasRegisteredSlateAttribute = false;
-			OwningWidget.MetaData.RemoveAtSwap(0);
-			if (OwningWidget.IsConstructed() && OwningWidget.IsAttributesUpdatesEnabled())
-			{
-				OwningWidget.Invalidate(EInvalidateWidgetReason::AttributeRegistration);
-			}
-		}
+		AttributeMetaData->RemoveMetaDataIfNeeded(OwningWidget, bResult);
 		return bResult;
 	}
 	return false;
 }
 
 
+void FSlateAttributeMetaData::RemoveMetaDataIfNeeded(SWidget& OwningWidget, bool bRemoved) const
+{
+	if (Attributes.Num() == 0)
+	{
+		check(bRemoved); // if the num is 0 then we should have remove an item.
+		OwningWidget.bHasRegisteredSlateAttribute = false;
+		OwningWidget.MetaData.RemoveAtSwap(0);
+		if (OwningWidget.IsConstructed() && OwningWidget.IsAttributesUpdatesEnabled())
+		{
+			OwningWidget.Invalidate(EInvalidateWidgetReason::AttributeRegistration);
+		}
+	}
+}
+
+
+bool FSlateAttributeMetaData::UnregisterAttribute(ISlateAttributeContainer& OwningContainer, const FSlateAttributeBase& Attribute)
+{
+	SWidget& OwningWidget = OwningContainer.GetContainerWidget();
+	return UnregisterAttribute(OwningWidget, Attribute);
+}
+
+
 bool FSlateAttributeMetaData::UnregisterAttributeImpl(const FSlateAttributeBase& Attribute)
 {
 	const int32 FoundIndex = IndexOfAttribute(Attribute);
@@ -187,18 +311,6 @@ void FSlateAttributeMetaData::InvalidateWidget(SWidget& OwningWidget, const FSla
 		return;
 	}
 
-	auto InvalidateForMember = [Reason](SWidget& OwningWidget, const FSlateAttributeBase& Attribute) -> EInvalidateWidgetReason
-	{
-		const FSlateAttributeDescriptor::OffsetType Offset = Private::FindOffet(OwningWidget, Attribute);
-		if (const FSlateAttributeDescriptor::FAttribute* FoundAttribute = OwningWidget.GetWidgetClass().GetAttributeDescriptor().FindMemberAttribute(Offset))
-		{
-			ensureMsgf(Reason == EInvalidateWidgetReason::None, TEXT("SWidget is using the AttributeDescriptor, the invalidation reason should only be defined in one place. See SLATE_ADD_MEMBER_ATTRIBUTE_DEFINITION"));
-			FoundAttribute->ExecuteOnValueChangedIfBound(OwningWidget);
-			return FoundAttribute->GetInvalidationReason(OwningWidget);
-		}
-		return Reason;
-	};
-
 	const FSlateAttributeDescriptor::FAttributeValueChangedDelegate* OnValueChangedCallback = nullptr;
 
 	if (FSlateAttributeMetaData* AttributeMetaData = FSlateAttributeMetaData::FindMetaData(OwningWidget))
@@ -206,7 +318,7 @@ void FSlateAttributeMetaData::InvalidateWidget(SWidget& OwningWidget, const FSla
 		const int32 FoundIndex = AttributeMetaData->IndexOfAttribute(Attribute);
 		if (FoundIndex != INDEX_NONE)
 		{
-			FGetterItem& GetterItem = AttributeMetaData->Attributes[FoundIndex];
+			const FGetterItem& GetterItem = AttributeMetaData->Attributes[FoundIndex];
 			if (GetterItem.CachedAttributeDescriptor)
 			{
 				ensureMsgf(Reason == EInvalidateWidgetReason::None, TEXT("SWidget is using the AttributeDescriptor, the invalidation reason should only be defined in one place. See SLATE_ADD_MEMBER_ATTRIBUTE_DEFINITION"));
@@ -217,7 +329,7 @@ void FSlateAttributeMetaData::InvalidateWidget(SWidget& OwningWidget, const FSla
 		// Not registered/bound but may be defined in the Descriptor
 		else if (AttributeType == ESlateAttributeType::Member)
 		{
-			Reason = InvalidateForMember(OwningWidget, Attribute);
+			Reason = UE::Private::InvalidateForMember(OwningWidget, Attribute, Reason);
 		}
 
 		Reason |= AttributeMetaData->CachedInvalidationReason;
@@ -225,7 +337,7 @@ void FSlateAttributeMetaData::InvalidateWidget(SWidget& OwningWidget, const FSla
 	}
 	else if (AttributeType == ESlateAttributeType::Member)
 	{
-		Reason = InvalidateForMember(OwningWidget, Attribute);
+		Reason = UE::Private::InvalidateForMember(OwningWidget, Attribute, Reason);
 	}
 
 #if WITH_SLATE_DEBUGGING
@@ -236,6 +348,54 @@ void FSlateAttributeMetaData::InvalidateWidget(SWidget& OwningWidget, const FSla
 }
 
 
+void FSlateAttributeMetaData::InvalidateWidget(ISlateAttributeContainer& OwningContainer, const FSlateAttributeBase& Attribute, ESlateAttributeType AttributeType, EInvalidateWidgetReason Reason)
+{
+	check(AttributeType == SlateAttributePrivate::ESlateAttributeType::Contained);
+
+	SWidget& Widget = OwningContainer.GetContainerWidget();
+	if (!Widget.IsConstructed())
+	{
+		return;
+	}
+
+	const FSlateAttributeDescriptor::FAttributeValueChangedDelegate* OnValueChangedCallback = nullptr;
+
+	if (FSlateAttributeMetaData* AttributeMetaData = FSlateAttributeMetaData::FindMetaData(Widget))
+	{
+		const int32 FoundIndex = AttributeMetaData->IndexOfAttribute(Attribute);
+		if (FoundIndex != INDEX_NONE)
+		{
+			const FGetterItem& GetterItem = AttributeMetaData->Attributes[FoundIndex];
+			if (GetterItem.CachedAttributeDescriptor)
+			{
+				ensureMsgf(Reason == EInvalidateWidgetReason::None, TEXT("SWidget is using the AttributeDescriptor, the invalidation reason should only be defined in one place. See SLATE_ADD_MEMBER_ATTRIBUTE_DEFINITION"));
+				GetterItem.CachedAttributeDescriptor->ExecuteOnValueChangedIfBound(Widget);
+				Reason = GetterItem.CachedAttributeDescriptor->GetInvalidationReason(Widget);
+			}
+			// else, it's not in the descriptor
+		}
+		// Not registered/bound but may be defined in the Descriptor
+		else
+		{
+			Reason = UE::Private::InvalidateForContained(OwningContainer, Widget, Attribute, Reason);
+		}
+
+		Reason |= AttributeMetaData->CachedInvalidationReason;
+		AttributeMetaData->CachedInvalidationReason = EInvalidateWidgetReason::None;
+	}
+	else
+	{
+		Reason = UE::Private::InvalidateForContained(OwningContainer, Widget, Attribute, Reason);
+	}
+
+#if WITH_SLATE_DEBUGGING
+	ensureAlwaysMsgf(FSlateAttributeBase::IsInvalidateWidgetReasonSupported(Reason), TEXT("%s is not an EInvalidateWidgetReason supported by SlateAttribute."), *LexToString(Reason));
+#endif
+
+	Widget.Invalidate(Reason);
+}
+
+
 void FSlateAttributeMetaData::UpdateAllAttributes(SWidget& OwningWidget, EInvalidationPermission InvalidationStyle)
 {
 	if (FSlateAttributeMetaData* AttributeMetaData = FSlateAttributeMetaData::FindMetaData(OwningWidget))
@@ -415,3 +575,28 @@ void FSlateAttributeMetaData::MoveAttribute(const SWidget& OwningWidget, FSlateA
 		}
 	}
 }
+
+
+void FSlateAttributeMetaData::RemoveContainerWidget(SWidget& OwningWidget, ISlateAttributeContainer& Container)
+{
+	// NB. This code can be called when we remove a FSlot or if the SWidget is destroyed.
+	//If the SWidget is being destroyed, the FindMetaData return nullptr.
+	if (FSlateAttributeMetaData* AttributeMetaData = FSlateAttributeMetaData::FindMetaData(OwningWidget))
+	{
+		int32 RemoveCount = AttributeMetaData->Attributes.RemoveAll([ContainerPtr = &Container](const FGetterItem& Item)
+			{
+				return Item.GetAttributeContainer() == ContainerPtr;
+			});
+
+		AttributeMetaData->RemoveMetaDataIfNeeded(OwningWidget, RemoveCount>0);
+	}
+}
+
+
+void FSlateAttributeMetaData::UpdateContainerSortOrder(SWidget& OwningWidget, ISlateAttributeContainer& Container)
+{
+	if (FSlateAttributeMetaData* AttributeMetaData = FSlateAttributeMetaData::FindMetaData(OwningWidget))
+	{
+		AttributeMetaData->Attributes.Sort(TLess<>());
+	}
+}
diff --git a/Engine/Source/Runtime/SlateCore/Public/Layout/BasicLayoutWidgetSlot.h b/Engine/Source/Runtime/SlateCore/Public/Layout/BasicLayoutWidgetSlot.h
--- a/Engine/Source/Runtime/SlateCore/Public/Layout/BasicLayoutWidgetSlot.h
+++ b/Engine/Source/Runtime/SlateCore/Public/Layout/BasicLayoutWidgetSlot.h
@@ -3,7 +3,7 @@
 #pragma once
 
 #include "CoreMinimal.h"
-#include "SlotBase.h"
+#include "Layout/WidgetSlotWithAttributeSupport.h"
 #include "Layout/FlowDirection.h"
 #include "Layout/Margin.h"
 #include "Misc/Optional.h"
@@ -447,36 +447,34 @@ PRAGMA_ENABLE_DEPRECATION_WARNINGS
 
 /** A templated basic slot that can be used by layout. */
 template <typename SlotType>
-class TBasicLayoutWidgetSlot : public TSlotBase<SlotType>
+class TBasicLayoutWidgetSlot : public TWidgetSlotWithAttributeSupport<SlotType>
 	, public TPaddingWidgetSlotMixin<SlotType>
 	, public TAlignmentWidgetSlotMixin<SlotType>
 {
 public:
 	TBasicLayoutWidgetSlot()
-		: TSlotBase<SlotType>()
+		: TWidgetSlotWithAttributeSupport<SlotType>()
 		, TPaddingWidgetSlotMixin<SlotType>()
 		, TAlignmentWidgetSlotMixin<SlotType>()
 	{}
 
 	TBasicLayoutWidgetSlot(FChildren& InOwner)
-		: TSlotBase<SlotType>(InOwner)
+		: TWidgetSlotWithAttributeSupport<SlotType>(InOwner)
 		, TPaddingWidgetSlotMixin<SlotType>()
 		, TAlignmentWidgetSlotMixin<SlotType>()
-	{
-	}
+	{}
 
 	TBasicLayoutWidgetSlot(const EHorizontalAlignment InHAlign, const EVerticalAlignment InVAlign)
-		: TSlotBase<SlotType>()
+		: TWidgetSlotWithAttributeSupport<SlotType>()
 		, TPaddingWidgetSlotMixin<SlotType>()
 		, TAlignmentWidgetSlotMixin<SlotType>(InHAlign, InVAlign)
 	{}
 
 	TBasicLayoutWidgetSlot(FChildren& InOwner, const EHorizontalAlignment InHAlign, const EVerticalAlignment InVAlign)
-		: TSlotBase<SlotType>(InOwner)
+		: TWidgetSlotWithAttributeSupport<SlotType>(InOwner)
 		, TPaddingWidgetSlotMixin<SlotType>()
 		, TAlignmentWidgetSlotMixin<SlotType>(InHAlign, InVAlign)
-	{
-	}
+	{}
 
 public:
 	SLATE_SLOT_BEGIN_ARGS_TwoMixins(TBasicLayoutWidgetSlot, TSlotBase<SlotType>, TPaddingWidgetSlotMixin<SlotType>, TAlignmentWidgetSlotMixin<SlotType>)
@@ -484,7 +482,7 @@ public:
 
 	void Construct(const FChildren& SlotOwner, FSlotArguments&& InArgs)
 	{
-		TSlotBase<SlotType>::Construct(SlotOwner, MoveTemp(InArgs));
+		TWidgetSlotWithAttributeSupport<SlotType>::Construct(SlotOwner, MoveTemp(InArgs));
 		TPaddingWidgetSlotMixin<SlotType>::ConstructMixin(SlotOwner, MoveTemp(InArgs));
 		TAlignmentWidgetSlotMixin<SlotType>::ConstructMixin(SlotOwner, MoveTemp(InArgs));
 	}
diff --git a/Engine/Source/Runtime/SlateCore/Public/Layout/Children.h b/Engine/Source/Runtime/SlateCore/Public/Layout/Children.h
--- a/Engine/Source/Runtime/SlateCore/Public/Layout/Children.h
+++ b/Engine/Source/Runtime/SlateCore/Public/Layout/Children.h
@@ -5,162 +5,10 @@
 #include "CoreMinimal.h"
 #include "Widgets/SNullWidget.h"
 #include "SlotBase.h"
+#include "Layout/ChildrenBase.h"
 #include "Layout/BasicLayoutWidgetSlot.h"
 #include "Widgets/SWidget.h"
 
-/**
- * FChildren is an interface that must be implemented by all child containers.
- * It allows iteration over a list of any Widget's children regardless of how
- * the underlying Widget happens to store its children.
- * 
- * FChildren is intended to be returned by the GetChildren() method.
- */
-PRAGMA_DISABLE_DEPRECATION_WARNINGS
-class SLATECORE_API FChildren
-{
-public:
-	FChildren(SWidget* InOwner)
-		: Owner(InOwner)
-	{
-		check(InOwner);
-	}
-
-	FChildren(std::nullptr_t) = delete;
-
-	//~ Prevents allocation of FChilren. It created a confusion between FSlot and FChildren
-	void* operator new (size_t) = delete;
-	void* operator new[](size_t) = delete;
-
-	/** @return the number of children */
-	virtual int32 Num() const = 0;
-	/** @return pointer to the Widget at the specified Index. */
-	virtual TSharedRef<SWidget> GetChildAt( int32 Index ) = 0;
-	/** @return const pointer to the Widget at the specified Index. */
-	virtual TSharedRef<const SWidget> GetChildAt( int32 Index ) const = 0;
-
-	/** @return the SWidget that own the FChildren */
-	SWidget& GetOwner() const { return *Owner; }
-
-	/** Applies the predicate to all the widgets contained by the FChildren. */
-	template<typename Predicate>
-	void ForEachWidget(Predicate Pred)
-	{
-		int32 WidgetCount = Num();
-		for (int32 Index = 0; Index < WidgetCount; ++Index)
-		{
-			Pred(GetChildRefAt(Index).GetWidget());
-		}
-	}
-
-	/** Applies the predicate to all the widgets contained by the FChildren. */
-	template<typename Predicate>
-	void ForEachWidget(Predicate Pred) const
-	{
-		int32 WidgetCount = Num();
-		for (int32 Index = 0; Index < WidgetCount; ++Index)
-		{
-			Pred(GetChildRefAt(Index).GetWidget());
-		}
-	}
-
-protected:
-	friend class SWidget;
-	friend class FCombinedChildren;
-	template<typename T>
-	friend class TOneDynamicChild;
-
-	enum ECopyConstruct { CopyConstruct };
-	enum ERefConstruct { ReferenceConstruct };
-
-	struct FWidgetRef
-	{
-	private:
-		TOptional<TSharedRef<SWidget>> WidgetCopy;
-		const TSharedRef<SWidget>& WidgetReference;
-
-	public:
-		FWidgetRef(ECopyConstruct, TSharedRef<SWidget> InWidgetCopy)
-			: WidgetCopy(MoveTemp(InWidgetCopy))
-			, WidgetReference(WidgetCopy.GetValue())
-		{}
-		FWidgetRef(ERefConstruct, const TSharedRef<SWidget>& InWidgetRef)
-			: WidgetCopy()
-			, WidgetReference(InWidgetRef)
-		{}
-		FWidgetRef(const FWidgetRef& Other)
-			: WidgetCopy(Other.WidgetCopy)
-			, WidgetReference(WidgetCopy.IsSet() ? WidgetCopy.GetValue() : Other.WidgetReference)
-		{}
-		FWidgetRef(FWidgetRef&& Other)
-			: WidgetCopy(MoveTemp(Other.WidgetCopy))
-			, WidgetReference(WidgetCopy.IsSet() ? WidgetCopy.GetValue() : Other.WidgetReference)
-		{}
-		FWidgetRef& operator=(const FWidgetRef&) = delete;
-		FWidgetRef& operator=(FWidgetRef&&) = delete;
-		const TSharedRef<SWidget>& GetWidgetRef() const
-		{
-			return WidgetReference;
-		}
-		SWidget& GetWidget() const
-		{
-			return WidgetReference.Get();
-		}
-	};
-
-	struct FConstWidgetRef
-	{
-	private:
-		TOptional<TSharedRef<const SWidget>> WidgetCopy;
-		const TSharedRef<const SWidget>& WidgetReference;
-
-	public:
-		FConstWidgetRef(ECopyConstruct, TSharedRef<const SWidget> InWidgetCopy)
-			: WidgetCopy(MoveTemp(InWidgetCopy))
-			, WidgetReference(WidgetCopy.GetValue())
-		{}
-		FConstWidgetRef(ERefConstruct, const TSharedRef<const SWidget>& InWidgetRef)
-			: WidgetCopy()
-			, WidgetReference(InWidgetRef)
-		{}
-		FConstWidgetRef(const FConstWidgetRef& Other)
-			: WidgetCopy(Other.WidgetCopy)
-			, WidgetReference(WidgetCopy.IsSet() ? WidgetCopy.GetValue() : Other.WidgetReference)
-		{}
-		FConstWidgetRef(FConstWidgetRef&& Other)
-			: WidgetCopy(MoveTemp(Other.WidgetCopy))
-			, WidgetReference(WidgetCopy.IsSet() ? WidgetCopy.GetValue() : Other.WidgetReference)
-		{}
-		FConstWidgetRef& operator=(const FConstWidgetRef&) = delete;
-		FConstWidgetRef& operator=(FConstWidgetRef&&) = delete;
-		const TSharedRef<const SWidget>& GetWidgetRef() const
-		{
-			return WidgetReference;
-		}
-		const SWidget& GetWidget() const
-		{
-			return WidgetReference.Get();
-		}
-	};
-
-	/** @return the number of slot the children has. */
-	virtual int32 NumSlot() const { return Num(); }
-	/** @return the const reference to the slot at the specified Index */
-	virtual const FSlotBase& GetSlotAt(int32 ChildIndex) const = 0;
-	/** @return ref to the Widget at the specified Index. */
-	virtual FWidgetRef GetChildRefAt(int32 Index) = 0;
-	/** @return ref to the Widget at the specified Index. */
-	virtual FConstWidgetRef GetChildRefAt(int32 Index) const = 0;
-
-
-protected:
-	virtual ~FChildren(){}
-
-protected:
-	UE_DEPRECATED(5.0, "Direct access to Owner is now deprecated. Use the getter.")
-	SWidget* Owner;
-};
-PRAGMA_ENABLE_DEPRECATION_WARNINGS
-
 
 /**
  * Occasionally you may need to keep multiple discrete sets of children with differing slot requirements.
@@ -287,6 +135,11 @@ public:
 	{
 	}
 
+	FNoChildren(SWidget* InOwner, FName InName)
+		: FChildren(InOwner, InName)
+	{
+	}
+
 	virtual int32 Num() const override { return 0; }
 	
 	virtual TSharedRef<SWidget> GetChildAt( int32 ) override
@@ -495,6 +348,12 @@ public:
 	{
 	}
 
+	TSingleWidgetChildrenWithSlot(SWidget* InOwner, FName InName)
+		: FChildren(InOwner, InName)
+		, TSlotBase<SlotType>(static_cast<const FChildren&>(*this))
+	{
+	}
+
 	TSingleWidgetChildrenWithSlot(std::nullptr_t) = delete;
 
 public:
@@ -593,6 +452,7 @@ private:
 	using ParentType = TSingleWidgetChildrenWithSlot<TSingleWidgetChildrenWithBasicLayoutSlot<InPaddingInvalidationReason>>;
 	using PaddingMixinType = TPaddingSingleWidgetSlotMixin<TSingleWidgetChildrenWithBasicLayoutSlot<InPaddingInvalidationReason>, InPaddingInvalidationReason>;
 	using AlignmentMixinType = TAlignmentSingleWidgetSlotMixin<TSingleWidgetChildrenWithBasicLayoutSlot<InPaddingInvalidationReason>>;
+
 public:
 	template<typename WidgetType, typename V = typename std::enable_if<std::is_base_of<SWidget, WidgetType>::value>::type>
 	TSingleWidgetChildrenWithBasicLayoutSlot(WidgetType* InOwner)
@@ -651,6 +511,7 @@ class TPanelChildren : public FChildren
 {
 private:
 	TArray<TUniquePtr<SlotType>> Children;
+	static constexpr bool bSupportSlotWithSlateAttribute = std::is_base_of<TWidgetSlotWithAttributeSupport<SlotType>, SlotType>::value;
 
 protected:
 	virtual const FSlotBase& GetSlotAt(int32 ChildIndex) const override
@@ -686,6 +547,11 @@ public:
 		return Children[Index]->GetWidget();
 	}
 
+	virtual bool SupportSlotWithSlateAttribute() const override
+	{
+		return bSupportSlotWithSlateAttribute;
+	}
+
 public:
 	UE_DEPRECATED(5.0, "Add a slot directly has been deprecated. use the FSlotArgument to create a new slot")
 	int32 Add( SlotType* Slot )
@@ -778,12 +644,15 @@ public:
 
 	void Move(int32 IndexToMove, int32 IndexToDestination)
 	{
-		// @todo this is going to cause a problem for draw ordering
-
 		{
 			TUniquePtr<SlotType> SlotToMove = MoveTemp(Children[IndexToMove]);
 			Children.RemoveAt(IndexToMove);
 			Children.Insert(MoveTemp(SlotToMove), IndexToDestination);
+			if constexpr (bSupportSlotWithSlateAttribute)
+			{
+				check(Children.Num() > 0);
+				Children[0]->RequestSortAttribute();
+			}
 		}
 
 		GetOwner().Invalidate(EInvalidateWidgetReason::ChildOrder);
@@ -799,32 +668,63 @@ public:
 		return Children.IsValidIndex( Index );
 	}
 
-	const SlotType& operator[](int32 Index) const { return *Children[Index]; }
-	SlotType& operator[](int32 Index) { return *Children[Index]; }
+	const SlotType& operator[](int32 Index) const
+	{
+		return *Children[Index];
+	}
+	SlotType& operator[](int32 Index)
+	{
+		return *Children[Index];
+	}
 
 	template <class PREDICATE_CLASS>
 	void Sort( const PREDICATE_CLASS& Predicate )
 	{
-		Children.Sort([&Predicate](const TUniquePtr<SlotType>& One, const TUniquePtr<SlotType>& Two)
+		if (Children.Num() > 0)
 		{
-			return Predicate(*One, *Two);
-		});
-		GetOwner().Invalidate(EInvalidateWidgetReason::ChildOrder);
+			Children.Sort([&Predicate](const TUniquePtr<SlotType>& One, const TUniquePtr<SlotType>& Two)
+			{
+				return Predicate(*One, *Two);
+			});
+
+			if constexpr (bSupportSlotWithSlateAttribute)
+			{
+				Children[0]->RequestSortAttribute();
+			}
+
+			GetOwner().Invalidate(EInvalidateWidgetReason::ChildOrder);
+		}
 	}
 
 	template <class PREDICATE_CLASS>
 	void StableSort(const PREDICATE_CLASS& Predicate)
 	{
-		Children.StableSort([&Predicate](const TUniquePtr<SlotType>& One, const TUniquePtr<SlotType>& Two)
+		if (Children.Num() > 0)
 		{
-			return Predicate(*One, *Two);
-		});
-		GetOwner().Invalidate(EInvalidateWidgetReason::ChildOrder);
+			Children.StableSort([&Predicate](const TUniquePtr<SlotType>& One, const TUniquePtr<SlotType>& Two)
+			{
+				return Predicate(*One, *Two);
+			});
+
+			if constexpr (bSupportSlotWithSlateAttribute)
+			{
+				Children[0]->RequestSortAttribute();
+			}
+
+			GetOwner().Invalidate(EInvalidateWidgetReason::ChildOrder);
+		}
 	}
 
 	void Swap( int32 IndexA, int32 IndexB )
 	{
 		Children.Swap(IndexA, IndexB);
+
+		if constexpr (bSupportSlotWithSlateAttribute)
+		{
+			check(Children.Num() > 0);
+			Children[0]->RequestSortAttribute();
+		}
+
 		GetOwner().Invalidate(EInvalidateWidgetReason::ChildOrder);
 	}
 
diff --git a/Engine/Source/Runtime/SlateCore/Public/Layout/ChildrenBase.h b/Engine/Source/Runtime/SlateCore/Public/Layout/ChildrenBase.h
new file mode 100644
--- /dev/null
+++ b/Engine/Source/Runtime/SlateCore/Public/Layout/ChildrenBase.h
@@ -0,0 +1,189 @@
+// Copyright Epic Games, Inc. All Rights Reserved.
+
+#pragma once
+
+#include "CoreMinimal.h"
+#include "SlotBase.h"
+
+class SWidget;
+class FSlotBase;
+
+/**
+ * FChildren is an interface that must be implemented by all child containers.
+ * It allows iteration over a list of any Widget's children regardless of how
+ * the underlying Widget happens to store its children.
+ * 
+ * FChildren is intended to be returned by the GetChildren() method.
+ */
+PRAGMA_DISABLE_DEPRECATION_WARNINGS
+class SLATECORE_API FChildren
+{
+public:
+	FChildren(SWidget* InOwner)
+		: Owner(InOwner)
+		, Name("Children")
+	{
+		check(InOwner);
+	}
+	FChildren(SWidget* InOwner, FName InName)
+		: Owner(InOwner)
+		, Name(InName)
+	{
+		check(InOwner);
+		check(!Name.IsNone());
+	}
+
+	FChildren(std::nullptr_t) = delete;
+
+	//~ Prevents allocation of FChilren. It creates confusion between FSlot and FChildren
+	void* operator new (size_t) = delete;
+	void* operator new[](size_t) = delete;
+
+	/** @return the number of children */
+	virtual int32 Num() const = 0;
+	/** @return pointer to the Widget at the specified Index. */
+	virtual TSharedRef<SWidget> GetChildAt( int32 Index ) = 0;
+	/** @return const pointer to the Widget at the specified Index. */
+	virtual TSharedRef<const SWidget> GetChildAt( int32 Index ) const = 0;
+
+	/** @return the SWidget that own the FChildren */
+	SWidget& GetOwner() const { return *Owner; }
+
+	/** Applies the predicate to all the widgets contained by the FChildren. */
+	template<typename Predicate>
+	void ForEachWidget(Predicate Pred)
+	{
+		int32 WidgetCount = Num();
+		for (int32 Index = 0; Index < WidgetCount; ++Index)
+		{
+			Pred(GetChildRefAt(Index).GetWidget());
+		}
+	}
+
+	/** Applies the predicate to all the widgets contained by the FChildren. */
+	template<typename Predicate>
+	void ForEachWidget(Predicate Pred) const
+	{
+		int32 WidgetCount = Num();
+		for (int32 Index = 0; Index < WidgetCount; ++Index)
+		{
+			Pred(GetChildRefAt(Index).GetWidget());
+		}
+	}
+
+	/** @return the number of slot the children has. */
+	virtual int32 NumSlot() const
+	{
+		return Num();
+	}
+
+	/** @return the const reference to the slot at the specified Index */
+	virtual const FSlotBase& GetSlotAt(int32 ChildIndex) const = 0;
+
+	/** @returns if the Children can support slot with SlateSlotAttribute. */
+	virtual bool SupportSlotWithSlateAttribute() const
+	{
+		return false;
+	}
+
+	/** Option to give a name to Children to SlotAttribute purposes or for debugging. */
+	FName GetName() const
+	{
+		return Name;
+	}
+
+protected:
+	friend class FCombinedChildren;
+	template<typename T>
+	friend class TOneDynamicChild;
+
+	enum ECopyConstruct { CopyConstruct };
+	enum ERefConstruct { ReferenceConstruct };
+
+	struct FWidgetRef
+	{
+	private:
+		TOptional<TSharedRef<SWidget>> WidgetCopy;
+		const TSharedRef<SWidget>& WidgetReference;
+
+	public:
+		FWidgetRef(ECopyConstruct, TSharedRef<SWidget> InWidgetCopy)
+			: WidgetCopy(MoveTemp(InWidgetCopy))
+			, WidgetReference(WidgetCopy.GetValue())
+		{}
+		FWidgetRef(ERefConstruct, const TSharedRef<SWidget>& InWidgetRef)
+			: WidgetCopy()
+			, WidgetReference(InWidgetRef)
+		{}
+		FWidgetRef(const FWidgetRef& Other)
+			: WidgetCopy(Other.WidgetCopy)
+			, WidgetReference(WidgetCopy.IsSet() ? WidgetCopy.GetValue() : Other.WidgetReference)
+		{}
+		FWidgetRef(FWidgetRef&& Other)
+			: WidgetCopy(MoveTemp(Other.WidgetCopy))
+			, WidgetReference(WidgetCopy.IsSet() ? WidgetCopy.GetValue() : Other.WidgetReference)
+		{}
+		FWidgetRef& operator=(const FWidgetRef&) = delete;
+		FWidgetRef& operator=(FWidgetRef&&) = delete;
+		const TSharedRef<SWidget>& GetWidgetRef() const
+		{
+			return WidgetReference;
+		}
+		SWidget& GetWidget() const
+		{
+			return WidgetReference.Get();
+		}
+	};
+
+	struct FConstWidgetRef
+	{
+	private:
+		TOptional<TSharedRef<const SWidget>> WidgetCopy;
+		const TSharedRef<const SWidget>& WidgetReference;
+
+	public:
+		FConstWidgetRef(ECopyConstruct, TSharedRef<const SWidget> InWidgetCopy)
+			: WidgetCopy(MoveTemp(InWidgetCopy))
+			, WidgetReference(WidgetCopy.GetValue())
+		{}
+		FConstWidgetRef(ERefConstruct, const TSharedRef<const SWidget>& InWidgetRef)
+			: WidgetCopy()
+			, WidgetReference(InWidgetRef)
+		{}
+		FConstWidgetRef(const FConstWidgetRef& Other)
+			: WidgetCopy(Other.WidgetCopy)
+			, WidgetReference(WidgetCopy.IsSet() ? WidgetCopy.GetValue() : Other.WidgetReference)
+		{}
+		FConstWidgetRef(FConstWidgetRef&& Other)
+			: WidgetCopy(MoveTemp(Other.WidgetCopy))
+			, WidgetReference(WidgetCopy.IsSet() ? WidgetCopy.GetValue() : Other.WidgetReference)
+		{}
+		FConstWidgetRef& operator=(const FConstWidgetRef&) = delete;
+		FConstWidgetRef& operator=(FConstWidgetRef&&) = delete;
+		const TSharedRef<const SWidget>& GetWidgetRef() const
+		{
+			return WidgetReference;
+		}
+		const SWidget& GetWidget() const
+		{
+			return WidgetReference.Get();
+		}
+	};
+
+	/** @return ref to the Widget at the specified Index. */
+	virtual FWidgetRef GetChildRefAt(int32 Index) = 0;
+	/** @return ref to the Widget at the specified Index. */
+	virtual FConstWidgetRef GetChildRefAt(int32 Index) const = 0;
+
+
+protected:
+	virtual ~FChildren(){}
+
+protected:
+	UE_DEPRECATED(5.0, "Direct access to Owner is now deprecated. Use the getter.")
+	SWidget* Owner;
+
+private:
+	FName Name;
+};
+PRAGMA_ENABLE_DEPRECATION_WARNINGS
diff --git a/Engine/Source/Runtime/SlateCore/Public/Layout/WidgetSlotWithAttributeSupport.h b/Engine/Source/Runtime/SlateCore/Public/Layout/WidgetSlotWithAttributeSupport.h
new file mode 100644
--- /dev/null
+++ b/Engine/Source/Runtime/SlateCore/Public/Layout/WidgetSlotWithAttributeSupport.h
@@ -0,0 +1,95 @@
+// Copyright Epic Games, Inc. All Rights Reserved.
+
+#pragma once
+
+#include "CoreMinimal.h"
+#include "SlotBase.h"
+#include "Layout/ChildrenBase.h"
+#include "Types/SlateAttribute.h"
+#include "Types/SlateAttributeDescriptor.h"
+
+#include <limits>
+
+
+/**
+ * A base slot that supports TSlateContainedAttribute.
+ * The FChildren that own the FSlot also needs to support SlateContainedAttribute.
+ * @see FChildren::SupportSlotWithSlateAttribute
+ */
+template <typename SlotType>
+class TWidgetSlotWithAttributeSupport : public TSlotBase<SlotType>, public SlateAttributePrivate::ISlateAttributeContainer
+{
+private:
+	using Super = TSlotBase<SlotType>;
+protected:
+	using TSlotBase<SlotType>::TSlotBase;
+
+	/**
+	 * A SlateAttribute that is member variable of a FSlot.
+	 * @usage: TSlateSlotAttribute<int32> MyAttribute1; TSlateSlotAttribute<int32, TSlateAttributeComparePredicate<>> MyAttribute2;
+	 */
+	template<typename InObjectType, typename InComparePredicate = TSlateAttributeComparePredicate<>>
+	struct TSlateSlotAttribute : public ::SlateAttributePrivate::TSlateContainedAttribute<InObjectType, ::SlateAttributePrivate::FSlateAttributeNoInvalidationReason, InComparePredicate>
+	{
+		using ::SlateAttributePrivate::TSlateContainedAttribute<InObjectType, ::SlateAttributePrivate::FSlateAttributeNoInvalidationReason, InComparePredicate>::TSlateContainedAttribute;
+	};
+
+public:
+	virtual ~TWidgetSlotWithAttributeSupport()
+	{
+		if (SWidget* Widget = Super::GetOwnerWidget())
+		{
+			RemoveContainerWidget(*Widget);
+		}
+	}
+
+	void Construct(const FChildren& SlotOwner, typename Super::FSlotArguments&& InArgs)
+	{
+		checkf(SlotOwner.SupportSlotWithSlateAttribute(), TEXT("The FChildren '%s' does not support SlateSlotAttribute but use a Slot type that does."), *SlotOwner.GetName().ToString());
+		checkf(SlotOwner.GetName().IsValid(), TEXT("The FChildren '%s' does not have a valid name."), *SlotOwner.GetName().ToString());
+		Super::Construct(SlotOwner, MoveTemp(InArgs));
+	}
+
+	static void RegisterAttributes(FSlateWidgetSlotAttributeInitializer& AttributeInitializer)
+	{
+	}
+
+	void RequestSortAttribute()
+	{
+		SWidget* Widget = Super::GetOwnerWidget();
+		ensureAlwaysMsgf(Widget, TEXT("FSlot needs to be constructed before we modify the FChildren."));
+		if (Widget)
+		{
+			UpdateContainerSortOrder(*Widget);
+		}
+	}
+
+
+	//~ Begin ISlateAttributeContainer interface
+	virtual SWidget& GetContainerWidget() const override
+	{
+		SWidget* Widget = TSlotBase<SlotType>::GetOwnerWidget();
+		checkf(Widget, TEXT("Slot Attributes has to be registered after the FSlot is constructed."));
+		return *Widget;
+	}
+
+	virtual FName GetContainerName() const override
+	{
+		return FSlotBase::GetOwner()->GetName();
+	}
+
+	virtual uint32 GetContainerSortOrder() const override
+	{
+		check(FSlotBase::GetOwner());
+		const int32 NumSlot = FSlotBase::GetOwner()->NumSlot();
+		for (int32 Index = 0; Index < NumSlot; ++Index)
+		{
+			if (&FSlotBase::GetOwner()->GetSlotAt(Index) == this)
+			{
+				return Index;
+			}
+		}
+		return std::numeric_limits<uint32>::max();
+	}
+	//~ End ISlateAttributeContainer interface
+};
diff --git a/Engine/Source/Runtime/SlateCore/Public/SlotBase.h b/Engine/Source/Runtime/SlateCore/Public/SlotBase.h
--- a/Engine/Source/Runtime/SlateCore/Public/SlotBase.h
+++ b/Engine/Source/Runtime/SlateCore/Public/SlotBase.h
@@ -4,7 +4,6 @@
 
 #include "CoreMinimal.h"
 #include "Misc/Attribute.h"
-#include "Types/SlateAttribute.h"
 #include "Widgets/InvalidateWidgetReason.h"
 
 class FChildren;
diff --git a/Engine/Source/Runtime/SlateCore/Public/Types/Attributes/SlateAttributeBase.inl b/Engine/Source/Runtime/SlateCore/Public/Types/Attributes/SlateAttributeBase.inl
--- a/Engine/Source/Runtime/SlateCore/Public/Types/Attributes/SlateAttributeBase.inl
+++ b/Engine/Source/Runtime/SlateCore/Public/Types/Attributes/SlateAttributeBase.inl
@@ -2,6 +2,8 @@
 
 #pragma once
 
+#include <type_traits>
+
 /** */
 namespace SlateAttributePrivate
 {
@@ -13,24 +15,30 @@ namespace SlateAttributePrivate
 	 * InComparePredicateType - Predicate to compare the cached value with the Getter.
 	 * bInIsExternal - The attribute life is not controlled by the SWidget.
 	 */
-	template<typename InObjectType, typename InInvalidationReasonPredicate, typename InComparePredicateType, ESlateAttributeType InAttributeType>
+	template<typename ContainerType, typename InObjectType, typename InInvalidationReasonPredicate, typename InComparePredicateType, ESlateAttributeType InAttributeType>
 	struct TSlateAttributeBase : public FSlateAttributeImpl
 	{
 	public:
 		template<typename AttributeMemberType>
 		friend struct TSlateMemberAttributeRef;
-		friend FSlateAttributeContainer;
 
 		using ObjectType = InObjectType;
 		using FInvalidationReasonPredicate = InInvalidationReasonPredicate;
 		using FGetter = typename TAttribute<ObjectType>::FGetter;
 		using FComparePredicate = InComparePredicateType;
 
+		static const ESlateAttributeType AttributeType = InAttributeType;
+		static constexpr bool HasDefinedInvalidationReason = !std::is_same<InInvalidationReasonPredicate, FSlateAttributeNoInvalidationReason>::value;
+
+		static_assert(std::is_same<ContainerType, SWidget>::value || std::is_same<ContainerType, ISlateAttributeContainer>::value, "The SlateAttribute container is not supported.");
+
 		static EInvalidateWidgetReason GetInvalidationReason(const SWidget& Widget) { return FInvalidationReasonPredicate::GetInvalidationReason(Widget); }
+		static EInvalidateWidgetReason GetInvalidationReason(const ISlateAttributeContainer& Container) { return FInvalidationReasonPredicate::GetInvalidationReason(Container.GetContainerWidget()); }
 		static bool IdenticalTo(const SWidget& Widget, const ObjectType& Lhs, const ObjectType& Rhs) { return FComparePredicate::IdenticalTo(Widget, Lhs, Rhs); }
+		static bool IdenticalTo(const ISlateAttributeContainer& Container, const ObjectType& Lhs, const ObjectType& Rhs) { return FComparePredicate::IdenticalTo(Container.GetContainerWidget(), Lhs, Rhs); }
 
 	private:
-		void UpdateNowOnBind(SWidget& Widget)
+		void UpdateNowOnBind(ContainerType& Widget)
 		{
 #if UE_SLATE_WITH_ATTRIBUTE_INITIALIZATION_ON_BIND
 			ProtectedUpdateNow(Widget, InAttributeType);
@@ -45,6 +53,10 @@ namespace SlateAttributePrivate
 #endif
 		}
 
+		void VerifyOwningWidget(const ISlateAttributeContainer& Widget) const
+		{
+		}
+
 		void VerifyNan() const
 		{
 #if UE_SLATE_WITH_ATTRIBUTE_NAN_DIAGNOSTIC
@@ -214,7 +226,7 @@ namespace SlateAttributePrivate
 		}
 
 		/** Update the cached value and invalidate the widget if needed. */
-		void UpdateNow(SWidget& Widget)
+		void UpdateNow(ContainerType& Widget)
 		{
 			VerifyOwningWidget(Widget);
 			ProtectedUpdateNow(Widget, InAttributeType);
@@ -225,7 +237,7 @@ namespace SlateAttributePrivate
 		 * Unbind the SlateAttribute and set its value. It may invalidate the Widget if the value is different.
 		 * @return true if the value is considered different and an invalidation occurred.
 		 */
-		bool Set(SWidget& Widget, const ObjectType& NewValue)
+		bool Set(ContainerType& Widget, const ObjectType& NewValue)
 		{
 			VerifyOwningWidget(Widget);
 			ProtectedUnregisterAttribute(Widget, InAttributeType);
@@ -243,7 +255,7 @@ namespace SlateAttributePrivate
 		 * Unbind the SlateAttribute and set its value. It may invalidate the Widget if the value is different.
 		 * @return true if the value is considered different and an invalidation occurred.
 		 */
-		bool Set(SWidget& Widget, ObjectType&& NewValue)
+		bool Set(ContainerType& Widget, ObjectType&& NewValue)
 		{
 			VerifyOwningWidget(Widget);
 			ProtectedUnregisterAttribute(Widget, InAttributeType);
@@ -263,7 +275,7 @@ namespace SlateAttributePrivate
 		 * (If enabled) Update the value from the Getter. If the value is different, then invalidate the widget.
 		 * The SlateAttribute will now be updated every frame from the Getter.
 		 */
-		void Bind(SWidget& Widget, const FGetter& Getter)
+		void Bind(ContainerType& Widget, const FGetter& Getter)
 		{
 			VerifyOwningWidget(Widget);
 			if (Getter.IsBound())
@@ -281,7 +293,7 @@ namespace SlateAttributePrivate
 		 * (If enabled) Update the value from the Getter. If the value is different, then invalidate the widget.
 		 * The SlateAttribute will now be updated every frame from the Getter.
 		 */
-		void Bind(SWidget& Widget, FGetter&& Getter)
+		void Bind(ContainerType& Widget, FGetter&& Getter)
 		{
 			VerifyOwningWidget(Widget);
 			if (Getter.IsBound())
@@ -299,7 +311,7 @@ namespace SlateAttributePrivate
 		 * (If enabled) Update the value from the getter. If the value is different, then invalidate the widget.
 		 * The SlateAttribute will now be updated every frame from the Getter.
 		 */
-		template<typename WidgetType, typename U = typename std::enable_if<std::is_base_of<SWidget, WidgetType>::value>::type>
+		template<typename WidgetType, typename U = typename std::enable_if<std::is_base_of<SWidget, WidgetType>::value && std::is_same<ContainerType, SWidget>::value>::type>
 		void Bind(WidgetType& Widget, typename FGetter::template TConstMethodPtr<WidgetType> MethodPtr)
 		{
 			Bind(Widget, FGetter::CreateSP(&Widget, MethodPtr));
@@ -318,7 +330,7 @@ namespace SlateAttributePrivate
 		 * @see Set
 		 * @return true if the getter was assigned or if the value is considered different and an invalidation occurred.
 		 */
-		bool Assign(SWidget& Widget, const TAttribute<ObjectType>& OtherAttribute)
+		bool Assign(ContainerType& Widget, const TAttribute<ObjectType>& OtherAttribute)
 		{
 			VerifyOwningWidget(Widget);
 			if (OtherAttribute.IsBound())
@@ -336,7 +348,7 @@ namespace SlateAttributePrivate
 			}
 		}
 
-		bool Assign(SWidget& Widget, TAttribute<ObjectType>&& OtherAttribute)
+		bool Assign(ContainerType& Widget, TAttribute<ObjectType>&& OtherAttribute)
 		{
 			VerifyOwningWidget(Widget);
 			if (OtherAttribute.IsBound())
@@ -354,7 +366,7 @@ namespace SlateAttributePrivate
 			}
 		}
 
-		bool Assign(SWidget& Widget, const TAttribute<ObjectType>& OtherAttribute, const ObjectType& DefaultValue)
+		bool Assign(ContainerType& Widget, const TAttribute<ObjectType>& OtherAttribute, const ObjectType& DefaultValue)
 		{
 			VerifyOwningWidget(Widget);
 			if (OtherAttribute.IsBound())
@@ -371,7 +383,7 @@ namespace SlateAttributePrivate
 			}
 		}
 
-		bool Assign(SWidget& Widget, TAttribute<ObjectType>&& OtherAttribute, const ObjectType& DefaultValue)
+		bool Assign(ContainerType& Widget, TAttribute<ObjectType>&& OtherAttribute, const ObjectType& DefaultValue)
 		{
 			VerifyOwningWidget(Widget);
 			if (OtherAttribute.IsBound())
@@ -388,7 +400,7 @@ namespace SlateAttributePrivate
 			}
 		}
 
-		bool Assign(SWidget& Widget, const TAttribute<ObjectType>& OtherAttribute, ObjectType&& DefaultValue)
+		bool Assign(ContainerType& Widget, const TAttribute<ObjectType>& OtherAttribute, ObjectType&& DefaultValue)
 		{
 			VerifyOwningWidget(Widget);
 			if (OtherAttribute.IsBound())
@@ -405,7 +417,7 @@ namespace SlateAttributePrivate
 			}
 		}
 
-		bool Assign(SWidget& Widget, TAttribute<ObjectType>&& OtherAttribute, ObjectType&& DefaultValue)
+		bool Assign(ContainerType& Widget, TAttribute<ObjectType>&& OtherAttribute, ObjectType&& DefaultValue)
 		{
 			VerifyOwningWidget(Widget);
 			if (OtherAttribute.IsBound())
@@ -426,7 +438,7 @@ namespace SlateAttributePrivate
 		 * Remove the Getter function.
 		 * The Slate Attribute will not be updated anymore and will keep its current cached value.
 		 */
-		void Unbind(SWidget& Widget)
+		void Unbind(ContainerType& Widget)
 		{
 			VerifyOwningWidget(Widget);
 			ProtectedUnregisterAttribute(Widget, InAttributeType);
@@ -434,7 +446,7 @@ namespace SlateAttributePrivate
 
 	public:
 		/** Build a Attribute from this SlateAttribute. */
-		UE_NODISCARD TAttribute<ObjectType> ToAttribute(const SWidget& Widget) const
+		UE_NODISCARD TAttribute<ObjectType> ToAttribute(const ContainerType& Widget) const
 		{
 			if (ISlateAttributeGetter* Delegate = ProtectedFindGetter(Widget, InAttributeType))
 			{
@@ -444,14 +456,14 @@ namespace SlateAttributePrivate
 		}
 
 		/** @return True if the SlateAttribute is bound to a getter function. */
-		UE_NODISCARD bool IsBound(const SWidget& Widget) const
+		UE_NODISCARD bool IsBound(const ContainerType& Widget) const
 		{
 			VerifyOwningWidget(Widget);
 			return ProtectedIsBound(Widget, InAttributeType);
 		}
 
 		/** @return True if they have the same Getter or the same value. */
-		UE_NODISCARD bool IdenticalTo(const SWidget& Widget, const TSlateAttributeBase& Other) const
+		UE_NODISCARD bool IdenticalTo(const ContainerType& Widget, const TSlateAttributeBase& Other) const
 		{
 			VerifyOwningWidget(Widget);
 			FDelegateHandle ThisDelegateHandle = ProtectedFindGetterHandle(Widget, InAttributeType);
@@ -468,7 +480,7 @@ namespace SlateAttributePrivate
 		}
 
 		/** @return True if they have the same Getter or, if the Attribute is set, the same value. */
-		UE_NODISCARD bool IdenticalTo(const SWidget& Widget, const TAttribute<ObjectType>& Other) const
+		UE_NODISCARD bool IdenticalTo(const ContainerType& Widget, const TAttribute<ObjectType>& Other) const
 		{
 			VerifyOwningWidget(Widget);
 			FDelegateHandle ThisDelegateHandle = ProtectedFindGetterHandle(Widget, InAttributeType);
@@ -480,21 +492,21 @@ namespace SlateAttributePrivate
 		}
 
 	private:
-		void ConstructWrapper(SWidget& Widget, const FGetter& Getter)
+		void ConstructWrapper(ContainerType& Widget, const FGetter& Getter)
 		{
 			TUniquePtr<ISlateAttributeGetter> Wrapper = MakeUniqueGetter(*this, Getter);
 			ProtectedRegisterAttribute(Widget, InAttributeType, MoveTemp(Wrapper));
 			UpdateNowOnBind(Widget);
 		}
 
-		void ConstructWrapper(SWidget& Widget, FGetter&& Getter)
+		void ConstructWrapper(ContainerType& Widget, FGetter&& Getter)
 		{
 			TUniquePtr<ISlateAttributeGetter> Wrapper = MakeUniqueGetter(*this, MoveTemp(Getter));
 			ProtectedRegisterAttribute(Widget, InAttributeType, MoveTemp(Wrapper));
 			UpdateNowOnBind(Widget);
 		}
 
-		bool AssignBinding(SWidget& Widget, const FGetter& Getter)
+		bool AssignBinding(ContainerType& Widget, const FGetter& Getter)
 		{
 			const FDelegateHandle PreviousGetterHandle = ProtectedFindGetterHandle(Widget, InAttributeType);
 			if (PreviousGetterHandle != Getter.GetHandle())
@@ -505,7 +517,7 @@ namespace SlateAttributePrivate
 			return false;
 		}
 
-		bool AssignBinding(SWidget& Widget, FGetter&& Getter)
+		bool AssignBinding(ContainerType& Widget, FGetter&& Getter)
 		{
 			const FDelegateHandle PreviousGetterHandle = ProtectedFindGetterHandle(Widget, InAttributeType);
 			if (PreviousGetterHandle != Getter.GetHandle())
diff --git a/Engine/Source/Runtime/SlateCore/Public/Types/Attributes/SlateAttributeContained.inl b/Engine/Source/Runtime/SlateCore/Public/Types/Attributes/SlateAttributeContained.inl
new file mode 100644
--- /dev/null
+++ b/Engine/Source/Runtime/SlateCore/Public/Types/Attributes/SlateAttributeContained.inl
@@ -0,0 +1,65 @@
+// Copyright Epic Games, Inc. All Rights Reserved.
+
+#pragma once
+
+#include <type_traits>
+
+
+class SWidget;
+
+
+/** */
+namespace SlateAttributePrivate
+{
+	/**
+	 *
+	 */
+	template<typename InObjectType, typename InInvalidationReasonPredicate, typename InComparePredicate>
+	struct TSlateContainedAttribute : public TSlateAttributeBase<ISlateAttributeContainer, InObjectType, InInvalidationReasonPredicate, InComparePredicate, ESlateAttributeType::Contained>
+	{
+	private:
+		using Super = TSlateAttributeBase<ISlateAttributeContainer, InObjectType, InInvalidationReasonPredicate, InComparePredicate, ESlateAttributeType::Contained>;
+
+		template<typename ContainerType, typename U = typename std::enable_if<std::is_base_of<ISlateAttributeContainer, ContainerType>::value>::type>
+		static void VerifyAttributeAddress(const ContainerType& Container, const TSlateContainedAttribute& Self)
+		{
+			checkf((UPTRINT)&Self >= (UPTRINT)&Container && (UPTRINT)&Self < (UPTRINT)&Container + sizeof(ContainerType),
+				TEXT("Use TAttribute or TSlateManagedAttribute instead. See SlateAttribute.h for more info."));
+			static_assert(!Super::HasDefinedInvalidationReason, "Define the Invalidation Reason in SlateAttributeDescriptor.");
+		}
+
+	public:
+		using FGetter = typename Super::FGetter;
+		using ObjectType = typename Super::ObjectType;
+
+	public:
+		//~ You can only register Attribute that are defined in a ISlateAttributeContainerOwner.
+		//~ Use the constructor with a ISlateAttributeContainerOwner pointer.
+		//~ See documentation in SlateAttribute.h for more info.
+		TSlateContainedAttribute() = delete;
+		TSlateContainedAttribute(const TSlateContainedAttribute&) = delete;
+		TSlateContainedAttribute(TSlateContainedAttribute&&) = delete;
+		TSlateContainedAttribute& operator=(const TSlateContainedAttribute&) = delete;
+		TSlateContainedAttribute& operator=(TSlateContainedAttribute&&) = delete;
+
+		template<typename ContainerType, typename U = typename std::enable_if<std::is_base_of<ISlateAttributeContainer, ContainerType>::value>::type>
+		explicit TSlateContainedAttribute(ContainerType& Container)
+		{
+			VerifyAttributeAddress(Container, *this);
+		}
+
+		template<typename ContainerType, typename U = typename std::enable_if<std::is_base_of<ISlateAttributeContainer, ContainerType>::value>::type>
+		explicit TSlateContainedAttribute(ContainerType& Container, const ObjectType& InValue)
+			: Super(InValue)
+		{
+			VerifyAttributeAddress(Container, *this);
+		}
+
+		template<typename ContainerType, typename U = typename std::enable_if<std::is_base_of<ISlateAttributeContainer, ContainerType>::value>::type>
+		explicit TSlateContainedAttribute(ContainerType& Container, ObjectType&& InValue)
+			: Super(MoveTemp(InValue))
+		{
+			VerifyAttributeAddress(Container, *this);
+		}
+	};
+} // SlateAttributePrivate
diff --git a/Engine/Source/Runtime/SlateCore/Public/Types/Attributes/SlateAttributeDefinition.inl b/Engine/Source/Runtime/SlateCore/Public/Types/Attributes/SlateAttributeDefinition.inl
--- a/Engine/Source/Runtime/SlateCore/Public/Types/Attributes/SlateAttributeDefinition.inl
+++ b/Engine/Source/Runtime/SlateCore/Public/Types/Attributes/SlateAttributeDefinition.inl
@@ -27,20 +27,31 @@ namespace SlateAttributePrivate
 	{
 		Member		= 0,	// Member of a SWidget (are not allowed to move).
 		Managed		= 1,	// External to the SWidget, global variable or member that can moved.
-		Unused0		= 2,
-		Unused1		= 3,
-		// We use the attribute type in a bit field in FSlateAttributeMetaData. Only 4 value are allowed.
+		Contained		= 2,	// Inside a FSlot or other container that use dynamic memory (always attached to one and only one SWidget).
 	};
 
 
 	class ISlateAttributeGetter;
-	template<typename ObjectType, typename InvalidationReasonPredicate, typename FComparePredicate, ESlateAttributeType AttributeType>
+	template<typename ContainerType, typename ObjectType, typename InvalidationReasonPredicate, typename FComparePredicate, ESlateAttributeType AttributeType>
 	struct TSlateAttributeBase;
-	struct FSlateAttributeContainer;
 	template<typename AttributeMemberType>
 	struct TSlateMemberAttributeRef;
 
 
+	/** Interface for structure that can be used to contain a SlateAttribute instead of a SWidget. */
+	class SLATECORE_API ISlateAttributeContainer
+	{
+	public:
+		virtual SWidget& GetContainerWidget() const = 0;
+		virtual FName GetContainerName() const = 0;
+		virtual uint32 GetContainerSortOrder() const = 0;
+
+	protected:
+		void RemoveContainerWidget(SWidget& Widget);
+		void UpdateContainerSortOrder(SWidget& Widget);
+	};
+
+
 	/** */
 	class ISlateAttributeGetter
 	{
@@ -72,18 +83,29 @@ namespace SlateAttributePrivate
 	struct SLATECORE_API FSlateAttributeImpl : public FSlateAttributeBase
 	{
 	protected:
-		bool ProtectedIsWidgetInDestructionPath(SWidget* Widget) const;
-		bool ProtectedIsImplemented(const SWidget& Widget) const;
 		void ProtectedUnregisterAttribute(SWidget& Widget, ESlateAttributeType AttributeType) const;
 		void ProtectedRegisterAttribute(SWidget& Widget, ESlateAttributeType AttributeType, TUniquePtr<ISlateAttributeGetter>&& Wrapper);
 		void ProtectedInvalidateWidget(SWidget& Widget, ESlateAttributeType AttributeType, EInvalidateWidgetReason InvalidationReason) const;
 		bool ProtectedIsBound(const SWidget& Widget, ESlateAttributeType AttributeType) const;
 		ISlateAttributeGetter* ProtectedFindGetter(const SWidget& Widget, ESlateAttributeType AttributeType) const;
 		FDelegateHandle ProtectedFindGetterHandle(const SWidget& Widget, ESlateAttributeType AttributeType) const;
-		bool ProtectedIsIdenticalTo(const SWidget& Widget, ESlateAttributeType AttributeType, const FSlateAttributeBase& Other, const bool bHasSameValue) const;
-		bool ProtectedIsIdenticalToAttribute(const SWidget& Widget, ESlateAttributeType AttributeType, const void* Other, const bool bHasSameValue) const;
 		void ProtectedUpdateNow(SWidget& Widget, ESlateAttributeType AttributeType);
+
+		//~ For Member
+		bool ProtectedIsWidgetInDestructionPath(SWidget* Widget) const;
+		bool ProtectedIsImplemented(const SWidget& Widget) const;
+
+		//~ For Manage
 		void ProtectedMoveAttribute(SWidget& Widget, ESlateAttributeType AttributeType, const FSlateAttributeBase* Other);
+
+		//~ For Contain
+		void ProtectedUnregisterAttribute(ISlateAttributeContainer& Widget, ESlateAttributeType AttributeType) const;
+		void ProtectedRegisterAttribute(ISlateAttributeContainer& Widget, ESlateAttributeType AttributeType, TUniquePtr<ISlateAttributeGetter>&& Wrapper);
+		void ProtectedInvalidateWidget(ISlateAttributeContainer& Widget, ESlateAttributeType AttributeType, EInvalidateWidgetReason InvalidationReason) const;
+		bool ProtectedIsBound(const ISlateAttributeContainer& Widget, ESlateAttributeType AttributeType) const;
+		ISlateAttributeGetter* ProtectedFindGetter(const ISlateAttributeContainer& Widget, ESlateAttributeType AttributeType) const;
+		FDelegateHandle ProtectedFindGetterHandle(const ISlateAttributeContainer& Widget, ESlateAttributeType AttributeType) const;
+		void ProtectedUpdateNow(ISlateAttributeContainer& Widget, ESlateAttributeType AttributeType);
 	};
 
 } // SlateAttributePrivate
diff --git a/Engine/Source/Runtime/SlateCore/Public/Types/Attributes/SlateAttributeManaged.inl b/Engine/Source/Runtime/SlateCore/Public/Types/Attributes/SlateAttributeManaged.inl
--- a/Engine/Source/Runtime/SlateCore/Public/Types/Attributes/SlateAttributeManaged.inl
+++ b/Engine/Source/Runtime/SlateCore/Public/Types/Attributes/SlateAttributeManaged.inl
@@ -9,17 +9,14 @@ namespace SlateAttributePrivate
 	 *
 	 */
 	template<typename InObjectType, typename InInvalidationReasonPredicate, typename InComparePredicate>
-	struct TSlateManagedAttribute : protected TSlateAttributeBase<InObjectType, InInvalidationReasonPredicate, InComparePredicate, ESlateAttributeType::Managed>
+	struct TSlateManagedAttribute : protected TSlateAttributeBase<SWidget, InObjectType, InInvalidationReasonPredicate, InComparePredicate, ESlateAttributeType::Managed>
 	{
 	private:
-		using Super = TSlateAttributeBase<InObjectType, InInvalidationReasonPredicate, InComparePredicate, ESlateAttributeType::Managed>;
+		using Super = TSlateAttributeBase<SWidget, InObjectType, InInvalidationReasonPredicate, InComparePredicate, ESlateAttributeType::Managed>;
 
 	public:
-		using ObjectType = typename Super::ObjectType;
-		using FInvalidationReasonPredicate = typename Super::FInvalidationReasonPredicate;
 		using FGetter = typename Super::FGetter;
-		using FComparePredicate = typename Super::FComparePredicate;
-		static const bool IsMemberType = false;
+		using ObjectType = typename Super::ObjectType;
 
 		static EInvalidateWidgetReason GetInvalidationReason(const SWidget& Widget) { return Super::GetInvalidationReason(Widget); }
 		static bool IdenticalTo(const SWidget& Widget, const ObjectType& Lhs, const ObjectType& Rhs) { return Super::IdenticalTo(Widget, Lhs, Rhs); }
diff --git a/Engine/Source/Runtime/SlateCore/Public/Types/Attributes/SlateAttributeMember.inl b/Engine/Source/Runtime/SlateCore/Public/Types/Attributes/SlateAttributeMember.inl
--- a/Engine/Source/Runtime/SlateCore/Public/Types/Attributes/SlateAttributeMember.inl
+++ b/Engine/Source/Runtime/SlateCore/Public/Types/Attributes/SlateAttributeMember.inl
@@ -9,17 +9,17 @@ namespace SlateAttributePrivate
 	 *
 	 */
 	template<typename InObjectType, typename InInvalidationReasonPredicate, typename InComparePredicate>
-	struct TSlateMemberAttribute : public TSlateAttributeBase<InObjectType, InInvalidationReasonPredicate, InComparePredicate, ESlateAttributeType::Member>
+	struct TSlateMemberAttribute : public TSlateAttributeBase<SWidget, InObjectType, InInvalidationReasonPredicate, InComparePredicate, ESlateAttributeType::Member>
 	{
 	private:
-		using Super = TSlateAttributeBase<InObjectType, InInvalidationReasonPredicate, InComparePredicate, ESlateAttributeType::Member>;
+		using Super = TSlateAttributeBase<SWidget, InObjectType, InInvalidationReasonPredicate, InComparePredicate, ESlateAttributeType::Member>;
 
 		template<typename WidgetType, typename U = typename std::enable_if<std::is_base_of<SWidget, WidgetType>::value>::type>
 		static void VerifyAttributeAddress(const WidgetType& Widget, const TSlateMemberAttribute& Self)
 		{
 			checkf((UPTRINT)&Self >= (UPTRINT)&Widget && (UPTRINT)&Self < (UPTRINT)&Widget + sizeof(WidgetType),
 				TEXT("Use TAttribute or TSlateManagedAttribute instead. See SlateAttribute.h for more info."));
-			ensureAlwaysMsgf((HasDefinedInvalidationReason || Self.ProtectedIsImplemented(Widget)),
+			ensureAlwaysMsgf((Super::HasDefinedInvalidationReason || Self.ProtectedIsImplemented(Widget)),
 				TEXT("The TSlateAttribute could not be found in the SlateAttributeDescriptor.\n")
 				TEXT("Use the SLATE_DECLARE_WIDGET and add the attribute in PrivateRegisterAttributes,\n")
 				TEXT("Or use TSlateAttribute with a valid Invalidation Reason instead."));
@@ -28,8 +28,6 @@ namespace SlateAttributePrivate
 	public:
 		using FGetter = typename Super::FGetter;
 		using ObjectType = typename Super::ObjectType;
-		static const bool IsMemberType = true;
-		static constexpr bool HasDefinedInvalidationReason = !std::is_same<InInvalidationReasonPredicate, FSlateAttributeNoInvalidationReason>::value;
 
 	public:
 		//~ You can only register Attribute that are defined in a SWidget (member of the class).
diff --git a/Engine/Source/Runtime/SlateCore/Public/Types/SlateAttribute.h b/Engine/Source/Runtime/SlateCore/Public/Types/SlateAttribute.h
--- a/Engine/Source/Runtime/SlateCore/Public/Types/SlateAttribute.h
+++ b/Engine/Source/Runtime/SlateCore/Public/Types/SlateAttribute.h
@@ -37,9 +37,9 @@
  *
  * TSlateAttribute makes the attribute system viable for invalidation and more performance-friendly while keeping the benefits of attributes intact.
  * TSlateAttributes are updated once per frame in the Prepass update phase. If the cached value of the TSlateAttribute changes, then it will invalidate the widget.
- * The TSlateAttributes are updated in the order the variables are defined in the SWidget definition (by default).
- * TSlateManagedAttribute are updated in a random order (after TSlateAttribute).
- * The update order of TSlateAttribute can be defined/override by setting a Prerequisite (see bellow for an example).
+ * The member attributes are updated in the order the variables are defined in the SWidget definition (by default).
+ * The managed attributes are updated in a random order (after TSlateAttribute).
+ * The update order of member attributes can be defined/override by setting a Prerequisite (see bellow for an example).
  * The invalidation reason can be a predicate (see bellow for an example).
  * The invalidation reason can be override per SWidget. Use override with precaution since it can break the invalidation of widget's parent.
  * The widget attributes are updated only if the widget is visible/not collapsed.
@@ -255,5 +255,7 @@ struct FSlateDeprecatedTAttribute
 
 #include "Types/Attributes/SlateAttributeDefinition.inl"
 #include "Types/Attributes/SlateAttributeBase.inl"
+
+#include "Types/Attributes/SlateAttributeContained.inl"
 #include "Types/Attributes/SlateAttributeManaged.inl"
 #include "Types/Attributes/SlateAttributeMember.inl"
diff --git a/Engine/Source/Runtime/SlateCore/Public/Types/SlateAttributeDescriptor.h b/Engine/Source/Runtime/SlateCore/Public/Types/SlateAttributeDescriptor.h
--- a/Engine/Source/Runtime/SlateCore/Public/Types/SlateAttributeDescriptor.h
+++ b/Engine/Source/Runtime/SlateCore/Public/Types/SlateAttributeDescriptor.h
@@ -90,44 +90,161 @@ public:
 	};
 
 public:
-	/** */
+	struct FContainer;
 	struct FAttribute;
 	struct FInitializer;
+	struct FContainerInitializer;
 
-	/** */
 	using OffsetType = uint32;
 
 	/** The default sort order that define in which order attributes will be updated. */
-	static uint32 DefaultSortOrder(OffsetType Offset) { return Offset * 100; }
+	static constexpr OffsetType DefaultSortOrder(OffsetType Offset) { return Offset * 100; }
+
+
+	/** */
+	struct FContainer
+	{
+		friend FInitializer;
+
+	public:
+		FContainer() = default;
+		FContainer(FName InName, OffsetType InOffset)
+			: Name(InName)
+			, Offset(InOffset)
+		{}
+
+		bool IsValid() const
+		{
+			return !Name.IsNone();
+		}
+
+		FName GetName() const
+		{
+			return Name;
+		}
+
+		uint32 GetSortOrder() const
+		{
+			return SortOrder;
+		}
+
+	private:
+		FName Name;
+		OffsetType Offset = 0;
+		uint32 SortOrder = 0;
+	};
+
 
 	/** */
 	struct FAttribute
 	{
 		friend FSlateAttributeDescriptor;
 		friend FInitializer;
+		friend FContainerInitializer;
 
 	public:
 		FAttribute(FName Name, OffsetType Offset, FInvalidateWidgetReasonAttribute Reason);
+		FAttribute(FName ContainerName, FName Name, OffsetType Offset, FInvalidateWidgetReasonAttribute Reason);
 
-		FName GetName() const { return Name; }
-		uint32 GetSortOrder() const { return SortOrder; }
-		EInvalidateWidgetReason GetInvalidationReason(const SWidget& Widget) const { return InvalidationReason.Get(Widget); }
-		SlateAttributePrivate::ESlateAttributeType GetAttributeType() const { return AttributeType; }
-		bool DoesAffectVisibility() const { return bAffectVisibility; }
+		FName GetName() const
+		{
+			return Name;
+		}
 
-		void ExecuteOnValueChangedIfBound(SWidget& Widget) const { OnValueChanged.ExecuteIfBound(Widget); }
+		uint32 GetSortOrder() const
+		{
+			return SortOrder;
+		}
+
+		EInvalidateWidgetReason GetInvalidationReason(const SWidget& Widget) const
+		{
+			return InvalidationReason.Get(Widget);
+		}
+
+		SlateAttributePrivate::ESlateAttributeType GetAttributeType() const
+		{
+			return AttributeType;
+		}
+
+		bool DoesAffectVisibility() const
+		{
+			return bAffectVisibility;
+		}
+
+		void ExecuteOnValueChangedIfBound(SWidget& Widget) const
+		{
+			OnValueChanged.ExecuteIfBound(Widget);
+		}
 
 	private:
 		FName Name;
 		OffsetType Offset;
 		FName Prerequisite;
+		FName ContainerName; // if the container IsNone, then the container is the SWidget itself.
 		uint32 SortOrder;
+		uint8 ContainerIndex;
 		FInvalidateWidgetReasonAttribute InvalidationReason;
 		FAttributeValueChangedDelegate OnValueChanged;
 		SlateAttributePrivate::ESlateAttributeType AttributeType;
 		bool bAffectVisibility;
 	};
 
+	/** Internal class to initialize the SlateAttributeDescriptor::FContainer attributes (Add attributes or modify existing attributes). */
+	struct SLATECORE_API FContainerInitializer
+	{
+	private:
+		friend FSlateAttributeDescriptor;
+		FContainerInitializer(FSlateAttributeDescriptor& InDescriptor, FName ContainerName);
+		FContainerInitializer(FSlateAttributeDescriptor& InDescriptor, const FSlateAttributeDescriptor& ParentDescriptor, FName ContainerName);
+
+	public:
+		FContainerInitializer() = delete;
+		FContainerInitializer(const FContainerInitializer&) = delete;
+		FContainerInitializer& operator= (const FContainerInitializer&) = delete;
+
+		struct SLATECORE_API FAttributeEntry
+		{
+			FAttributeEntry(FSlateAttributeDescriptor& Descriptor, FName ContainerName, int32 AttributeIndex);
+
+			/**
+			 * Update the attribute after the prerequisite.
+			 * The order is guaranteed but other attributes may be updated in between.
+			 * No order is guaranteed if the prerequisite or this property is updated manually.
+			 */
+			FAttributeEntry& UpdatePrerequisite(FName Prerequisite);
+
+			/**
+			 * Notified when the attribute value changed.
+			 * It's preferable that you delay any action to the Tick or Paint function.
+			 * You are not allowed to make changes that would affect the SWidget ChildOrder or its Visibility.
+			 * It will not be called when the SWidget is in his construction phase.
+			 * @see SWidget::IsConstructed
+			 */
+			FAttributeEntry& OnValueChanged(FAttributeValueChangedDelegate Callback);
+
+		private:
+			FSlateAttributeDescriptor& Descriptor;
+			FName ContainerName;
+			int32 AttributeIndex;
+		};
+
+		FAttributeEntry AddContainedAttribute(FName AttributeName, OffsetType Offset, const FInvalidateWidgetReasonAttribute& ReasonGetter);
+		FAttributeEntry AddContainedAttribute(FName AttributeName, OffsetType Offset, FInvalidateWidgetReasonAttribute&& ReasonGetter);
+
+	public:
+		/** Change the InvalidationReason of an attribute defined in a base class. */
+		void OverrideInvalidationReason(FName AttributeName, const FInvalidateWidgetReasonAttribute& Reason);
+		/** Change the InvalidationReason of an attribute defined in a base class. */
+		void OverrideInvalidationReason(FName AttributeName, FInvalidateWidgetReasonAttribute&& Reason);
+
+		/** Change the FAttributeValueChangedDelegate of an attribute defined in a base class. */
+		void OverrideOnValueChanged(FName AttributeName, ECallbackOverrideType OverrideType, FAttributeValueChangedDelegate Callback);
+
+	private:
+		FSlateAttributeDescriptor& Descriptor;
+		FName ContainerName;
+	};
+
 	/** Internal class to initialize the SlateAttributeDescriptor (Add attributes or modify existing attributes). */
 	struct SLATECORE_API FInitializer
 	{
@@ -176,6 +293,9 @@ public:
 		FAttributeEntry AddMemberAttribute(FName AttributeName, OffsetType Offset, const FInvalidateWidgetReasonAttribute& ReasonGetter);
 		FAttributeEntry AddMemberAttribute(FName AttributeName, OffsetType Offset, FInvalidateWidgetReasonAttribute&& ReasonGetter);
 
+		FContainerInitializer AddContainer(FName ContainerName, OffsetType Offset);
+
+	public:
 		/** Change the InvalidationReason of an attribute defined in a base class. */
 		void OverrideInvalidationReason(FName AttributeName, const FInvalidateWidgetReasonAttribute& Reason);
 		/** Change the InvalidationReason of an attribute defined in a base class. */
@@ -187,40 +307,57 @@ public:
 		/** Change the update type of an attribute defined in a base class. */
 		void SetAffectVisibility(FName AttributeName, bool bAffectVisibility);
 
-
 	private:
 		FSlateAttributeDescriptor& Descriptor;
 	};
 
 	/** @returns the number of Attributes registered. */
-	int32 GetAttributeNum() const { return Attributes.Num(); }
+	int32 GetAttributeNum() const
+	{
+		return Attributes.Num();
+	}
 
 	/** @returns the Attribute at the index previously found with IndexOfMemberAttribute */
 	const FAttribute& GetAttributeAtIndex(int32 Index) const;
 
+	/** @returns the Container with the corresponding name. */
+	const FContainer* FindContainer(FName ContainerName) const;
+
 	/** @returns the Attribute with the corresponding name. */
 	const FAttribute* FindAttribute(FName AttributeName) const;
 
 	/** @returns the Attribute of a SlateAttribute that have the corresponding memory offset. */
 	const FAttribute* FindMemberAttribute(OffsetType AttributeOffset) const;
 
+	/** @returns the Attribute of a SlateAttribute that have the corresponding memory offset. */
+	const FAttribute* FindContainedAttribute(FName ContainerName, OffsetType AttributeOffset) const;
+
+	/** @returns the index of the Container with the corresponding name. */
+	int32 IndexOfContainer(FName AttributeName) const;
+
 	/** @returns the index of a SlateAttribute that have the corresponding memory offset. */
 	int32 IndexOfAttribute(FName AttributeName) const;
 
 	/** @returns the index of a SlateAttribute that have the corresponding memory offset. */
 	int32 IndexOfMemberAttribute(OffsetType AttributeOffset) const;
 
+	/** @returns the index of a SlateAttribute that have the corresponding memory offset. */
+	int32 IndexOfContainedAttribute(FName ContainerName, OffsetType AttributeOffset) const;
+
 private:
 	FAttribute* FindAttribute(FName AttributeName);
 
+	FContainerInitializer AddContainer(FName AttributeName, OffsetType Offset);
 	FInitializer::FAttributeEntry AddMemberAttribute(FName AttributeName, OffsetType Offset, FInvalidateWidgetReasonAttribute ReasonGetter);
-	void OverrideInvalidationReason(FName AttributeName, FInvalidateWidgetReasonAttribute ReasonGetter);
-	void OverrideOnValueChanged(FName AttributeName, ECallbackOverrideType OverrideType, FAttributeValueChangedDelegate Callback);
-	void SetPrerequisite(FAttribute& Attribute, FName Prerequisite);
+	FContainerInitializer::FAttributeEntry AddContainedAttribute(FName ContainerName, FName AttributeName, OffsetType Offset, FInvalidateWidgetReasonAttribute ReasonGetter);
+	void OverrideInvalidationReason(FName ContainerName, FName AttributeName, FInvalidateWidgetReasonAttribute ReasonGetter);
+	void OverrideOnValueChanged(FName ContainerName, FName AttributeName, ECallbackOverrideType OverrideType, FAttributeValueChangedDelegate Callback);
+	void SetPrerequisite(FName ContainerName, FAttribute& Attribute, FName Prerequisite);
 	void SetAffectVisibility(FAttribute& Attribute, bool bUpdate);
 
 private:
 	TArray<FAttribute> Attributes;
+	TArray<FContainer, TInlineAllocator<1>> Containers;
 };
 
 /**
@@ -230,12 +367,29 @@ private:
  * @param _Reason The EInvalidationWidgetReason or a static function/lambda that takes a const SWidget& and that returns the invalidation reason.
  */
 #define SLATE_ADD_MEMBER_ATTRIBUTE_DEFINITION_WITH_NAME(_Initializer, _Name, _Property, _Reason) \
-		static_assert(decltype(_Property)::IsMemberType, "The SlateProperty is not a TSlateAttribute. Do not use SLATE_ADD_MEMBER_ATTRIBUTE_DEFINITION"); \
+		static_assert(decltype(_Property)::AttributeType == SlateAttributePrivate::ESlateAttributeType::Member, "The SlateProperty is not a TSlateAttribute. Do not use SLATE_ADD_MEMBER_ATTRIBUTE_DEFINITION"); \
 		static_assert(!decltype(_Property)::HasDefinedInvalidationReason, "When implementing the SLATE_DECLARE_WIDGET pattern, use TSlateAttribute without the invalidation reason."); \
-		static_assert(!std::is_same<decltype(_Reason), EInvalidateWidgetReason>::value || FSlateAttributeBase::IsInvalidateWidgetReasonSupported(_Reason), "The invalidation is not supported by the SlateAttribute."); \
+		static_assert(!std::is_same<decltype(_Reason), EInvalidateWidgetReason>::value || FSlateAttributeBase::IsInvalidateWidgetReasonSupported(_Reason), "The invalidation is not supported by the SlateAttribute system."); \
 		_Initializer.AddMemberAttribute(_Name, STRUCT_OFFSET(PrivateThisType, _Property), FSlateAttributeDescriptor::FInvalidateWidgetReasonAttribute{_Reason})
 
 #define SLATE_ADD_MEMBER_ATTRIBUTE_DEFINITION(_Initializer, _Property, _Reason) \
 	SLATE_ADD_MEMBER_ATTRIBUTE_DEFINITION_WITH_NAME(_Initializer, GET_MEMBER_NAME_CHECKED(PrivateThisType, _Property), _Property, _Reason)
 
+#define SLATE_ADD_PANELCHILDREN_DEFINITION_WITH_NAME(_Initializer, _Name, _Container) \
+		_Initializer.AddContainer(_Name, STRUCT_OFFSET(PrivateThisType, _Container))
+
+#define SLATE_ADD_PANELCHILDREN_DEFINITION(_Initializer, _Container) \
+		SLATE_ADD_PANELCHILDREN_DEFINITION_WITH_NAME(_Initializer, GET_MEMBER_NAME_CHECKED(PrivateThisType, _Container), _Container)
+
+#define SLATE_ADD_SLOT_ATTRIBUTE_DEFINITION_WITH_NAME(_SlotType, _Initializer, _Name, _Property, _Reason) \
+		static_assert(decltype(_Property)::AttributeType == SlateAttributePrivate::ESlateAttributeType::Contained, "The SlateProperty is not a TSlateAttribute. Do not use SLATE_ADD_CONTAINED_ATTRIBUTE_DEFINITION"); \
+		static_assert(!decltype(_Property)::HasDefinedInvalidationReason, "When implementing the SLATE_DECLARE_WIDGET pattern, use TSlateSlotAttribute without the invalidation reason."); \
+		static_assert(!std::is_same<decltype(_Reason), EInvalidateWidgetReason>::value || FSlateAttributeBase::IsInvalidateWidgetReasonSupported(_Reason), "The invalidation is not supported by the SlateAttribute system."); \
+		static_assert(STRUCT_OFFSET(_SlotType, _Property) > 0, "The offset cannot be calculated"); \
+		_Initializer.AddContainedAttribute(_Name, ((std::size_t)(&(((_SlotType*)(1))->_Property)) - (std::size_t)((SlateAttributePrivate::ISlateAttributeContainer*)((_SlotType*)(1)))), FSlateAttributeDescriptor::FInvalidateWidgetReasonAttribute{_Reason})
+
+#define SLATE_ADD_SLOT_ATTRIBUTE_DEFINITION(_SlotType, _Initializer, _Property, _Reason) \
+	SLATE_ADD_SLOT_ATTRIBUTE_DEFINITION_WITH_NAME(_SlotType, _Initializer, GET_MEMBER_NAME_CHECKED(_SlotType, _Property), _Property, _Reason)
+
 using FSlateAttributeInitializer = FSlateAttributeDescriptor::FInitializer;
+using FSlateWidgetSlotAttributeInitializer = FSlateAttributeDescriptor::FContainerInitializer;
diff --git a/Engine/Source/Runtime/SlateCore/Public/Types/SlateAttributeMetaData.h b/Engine/Source/Runtime/SlateCore/Public/Types/SlateAttributeMetaData.h
--- a/Engine/Source/Runtime/SlateCore/Public/Types/SlateAttributeMetaData.h
+++ b/Engine/Source/Runtime/SlateCore/Public/Types/SlateAttributeMetaData.h
@@ -9,6 +9,7 @@
 #include "Types/SlateAttributeDescriptor.h"
 #include "Widgets/InvalidateWidgetReason.h"
 
+#include <limits>
 
 class SWidget;
 
@@ -16,6 +17,7 @@ class SWidget;
 class SLATECORE_API FSlateAttributeMetaData : public ISlateMetaData
 {
 	friend SlateAttributePrivate::FSlateAttributeImpl;
+	friend SlateAttributePrivate::ISlateAttributeContainer;
 
 public:
 	SLATE_METADATA_TYPE(FSlateAttributeMetaData, ISlateMetaData);
@@ -77,20 +79,33 @@ public:
 	static TArray<FName> GetAttributeNames(const SWidget& OwningWidget);
 
 private:
+	using ISlateAttributeContainer = SlateAttributePrivate::ISlateAttributeContainer;
 	using ESlateAttributeType = SlateAttributePrivate::ESlateAttributeType;
 	using ISlateAttributeGetter = SlateAttributePrivate::ISlateAttributeGetter;
 	static void RegisterAttribute(SWidget& OwningWidget, FSlateAttributeBase& Attribute, ESlateAttributeType AttributeType, TUniquePtr<ISlateAttributeGetter>&& Wrapper);
+	static void RegisterAttribute(ISlateAttributeContainer& OwningContainer, FSlateAttributeBase& Attribute, ESlateAttributeType AttributeType, TUniquePtr<ISlateAttributeGetter>&& Wrapper);
 	static bool UnregisterAttribute(SWidget& OwningWidget, const FSlateAttributeBase& Attribute);
+	static bool UnregisterAttribute(ISlateAttributeContainer& OwningContainer, const FSlateAttributeBase& Attribute);
 	static void InvalidateWidget(SWidget& OwningWidget, const FSlateAttributeBase& Attribute, ESlateAttributeType AttributeType, EInvalidateWidgetReason Reason);
+	static void InvalidateWidget(ISlateAttributeContainer& OwningContainer, const FSlateAttributeBase& Attribute, ESlateAttributeType AttributeType, EInvalidateWidgetReason Reason);
 	static void UpdateAttribute(SWidget& OwningWidget, FSlateAttributeBase& Attribute);
 	static bool IsAttributeBound(const SWidget& OwningWidget, const FSlateAttributeBase& Attribute);
 	static SlateAttributePrivate::ISlateAttributeGetter* GetAttributeGetter(const SWidget& OwningWidget, const FSlateAttributeBase& Attribute);
 	static FDelegateHandle GetAttributeGetterHandle(const SWidget& OwningWidget, const FSlateAttributeBase& Attribute);
 	static void MoveAttribute(const SWidget& OwningWidget, FSlateAttributeBase& NewAttribute, ESlateAttributeType AttributeType, const FSlateAttributeBase* PreviousAttribute);
 
+	//~ For Container
+	static void RemoveContainerWidget(SWidget& Widget, ISlateAttributeContainer& Container);
+	static void UpdateContainerSortOrder(SWidget& Widget, ISlateAttributeContainer& Container);
+
+
 private:
-	void RegisterAttributeImpl(SWidget& OwningWidget, FSlateAttributeBase& Attribute, ESlateAttributeType AttributeType, TUniquePtr<ISlateAttributeGetter>&& Getter);
+	void RegisterMemberAttributeImpl(SWidget& OwningWidget, FSlateAttributeBase& Attribute, TUniquePtr<ISlateAttributeGetter>&& Getter);
+	void RegisterManagedAttributeImpl(SWidget& OwningWidget, FSlateAttributeBase& Attribute, TUniquePtr<ISlateAttributeGetter>&& Getter);
+	void RegisterContainAttributeImpl(ISlateAttributeContainer& OwningContainer, SWidget& OwningWidget, FSlateAttributeBase& Attribute, TUniquePtr<ISlateAttributeGetter>&& Getter);
 	bool UnregisterAttributeImpl(const FSlateAttributeBase& Attribute);
+	void RemoveMetaDataIfNeeded(SWidget& OwningWidget, bool bRemoved) const;
+	bool InvalidateWidgetAttribute(SWidget& OwningWidget, const FSlateAttributeBase& Attribute, EInvalidateWidgetReason Reason);
 	void UpdateAttributesImpl(SWidget& OwningWidget, EInvalidationPermission InvaldiationStyle, int32 StartIndex, int32 IndexNum);
 
 private:
@@ -110,26 +125,46 @@ private:
 			, Getter(MoveTemp(InGetter))
 			, CachedAttributeDescriptor(nullptr)
 			, SortOrder(InSortOrder)
+			, AttributeContainerOffset(0)
 		{ }
 		FGetterItem(FSlateAttributeBase* InAttribute, uint32 InSortOrder, TUniquePtr<SlateAttributePrivate::ISlateAttributeGetter> && InGetter, const FSlateAttributeDescriptor::FAttribute* InAttributeDescriptor)
 			: Attribute(InAttribute)
 			, Getter(MoveTemp(InGetter))
 			, CachedAttributeDescriptor(InAttributeDescriptor)
 			, SortOrder(InSortOrder)
+			, AttributeContainerOffset(0)
 		{ }
+
 		FSlateAttributeBase* Attribute;
 		TUniquePtr<SlateAttributePrivate::ISlateAttributeGetter> Getter;
 		const FSlateAttributeDescriptor::FAttribute* CachedAttributeDescriptor;
 		uint32 SortOrder;
-		ESlateAttributeType AttributeType;
+		int32 AttributeContainerOffset; // to fit inside a 32b structure, save the offset of the AttributeContainer;
 
 		bool operator<(const FGetterItem& Other) const
 		{
+			// if in a container, the smaller container wins first, then the property SortOrder.
+			if (AttributeContainerOffset && Other.AttributeContainerOffset)
+			{
+				const SlateAttributePrivate::ISlateAttributeContainer* Container = GetAttributeContainer();
+				const SlateAttributePrivate::ISlateAttributeContainer* OtherContainer = Other.GetAttributeContainer();
+				if (Container == OtherContainer)
+				{
+					return SortOrder < Other.SortOrder;
+				}
+				else
+				{
+					return Container->GetContainerSortOrder() < OtherContainer->GetContainerSortOrder();
+				}
+			}
 			return SortOrder < Other.SortOrder;
 		}
 
 		/** If available, return the name of the attribute. */
 		FName GetAttributeName(const SWidget& OwningWidget) const;
+
+		/** @return a pointer to the attribute container if the SlateAttribute type is Contained. */
+		const SlateAttributePrivate::ISlateAttributeContainer* GetAttributeContainer() const;
 	};
 	static_assert(sizeof(FGetterItem) <= 32, "The size of FGetterItem is bigger than expected.");
 
