Fix "hides class member" compilation warnings in WidgetSlotWithAttributeSupport.h

diff --git a/Engine/Source/Runtime/SlateCore/Public/Layout/WidgetSlotWithAttributeSupport.h b/Engine/Source/Runtime/SlateCore/Public/Layout/WidgetSlotWithAttributeSupport.h
--- a/Engine/Source/Runtime/SlateCore/Public/Layout/WidgetSlotWithAttributeSupport.h
+++ b/Engine/Source/Runtime/SlateCore/Public/Layout/WidgetSlotWithAttributeSupport.h
@@ -37,9 +37,9 @@ protected:
 public:
 	virtual ~TWidgetSlotWithAttributeSupport()
 	{
-		if (SWidget* Widget = Super::GetOwnerWidget())
+		if (SWidget* WidgetPtr = Super::GetOwnerWidget())
 		{
-			RemoveContainerWidget(*Widget);
+			RemoveContainerWidget(*WidgetPtr);
 		}
 	}
 
@@ -68,9 +68,9 @@ public:
 	//~ Begin ISlateAttributeContainer interface
 	virtual SWidget& GetContainerWidget() const override
 	{
-		SWidget* Widget = TSlotBase<SlotType>::GetOwnerWidget();
-		checkf(Widget, TEXT("Slot Attributes has to be registered after the FSlot is constructed."));
-		return *Widget;
+		SWidget* WidgetPtr = TSlotBase<SlotType>::GetOwnerWidget();
+		checkf(WidgetPtr, TEXT("Slot Attributes has to be registered after the FSlot is constructed."));
+		return *WidgetPtr;
 	}
 
 	virtual FName GetContainerName() const override
