Fix issues where if WITH_SLATE_DEBUGGING is force to 0 in a debug build it fails to compile

diff --git a/Engine/Source/Runtime/SlateCore/Public/FastUpdate/WidgetProxy.h b/Engine/Source/Runtime/SlateCore/Public/FastUpdate/WidgetProxy.h
--- a/Engine/Source/Runtime/SlateCore/Public/FastUpdate/WidgetProxy.h
+++ b/Engine/Source/Runtime/SlateCore/Public/FastUpdate/WidgetProxy.h
@@ -27,7 +27,7 @@ enum class EInvalidateWidgetReason : uint8;
 #define UE_SLATE_WITH_WIDGETPROXY_WEAKPTR 0
 #define UE_SLATE_VERIFY_WIDGETPROXY_WEAKPTR_STALE 0
 #define UE_SLATE_WITH_WIDGETPROXY_WIDGETTYPE 0
-#define UE_SLATE_WITH_INVALIDATIONWIDGETLIST_DEBUGGING !(UE_BUILD_SHIPPING || UE_BUILD_TEST)
+#define UE_SLATE_WITH_INVALIDATIONWIDGETLIST_DEBUGGING !(UE_BUILD_SHIPPING || UE_BUILD_TEST) && WITH_SLATE_DEBUGGING
 
 
 struct FSlateInvalidationWidgetVisibility
diff --git a/Engine/Source/Runtime/SlateCore/Public/Layout/Clipping.h b/Engine/Source/Runtime/SlateCore/Public/Layout/Clipping.h
--- a/Engine/Source/Runtime/SlateCore/Public/Layout/Clipping.h
+++ b/Engine/Source/Runtime/SlateCore/Public/Layout/Clipping.h
@@ -213,7 +213,7 @@ public:
 	/** Is a point inside the clipping state? */
 	bool IsPointInside(const FVector2D& Point) const;
 
-#if !(UE_BUILD_SHIPPING || UE_BUILD_TEST)
+#if WITH_SLATE_DEBUGGING
 	/** Set the state index that this clipping state originated from.  We just do this for debugging purposes. */
 	FORCEINLINE void SetDebuggingStateIndex(int32 InStateIndex) const
 	{
