Remove incorrectly nested deprecation pragma

diff --git a/Engine/Source/Runtime/SlateCore/Public/Layout/WidgetPath.h b/Engine/Source/Runtime/SlateCore/Public/Layout/WidgetPath.h
--- a/Engine/Source/Runtime/SlateCore/Public/Layout/WidgetPath.h
+++ b/Engine/Source/Runtime/SlateCore/Public/Layout/WidgetPath.h
@@ -53,12 +53,10 @@ public:
 	const TSharedPtr<const FVirtualPointerPosition>& GetCursorAt( int32 Index ) const;
 
 	/** Get the virtual representation of the mouse at each level in the widget path. */
-	PRAGMA_DISABLE_DEPRECATION_WARNINGS
 	TOptional<FVirtualPointerPosition> GetVirtualPointerPosition(int32 Index) const
 	{
 		return VirtualPointerPositions[Index];
 	}
-	PRAGMA_ENABLE_DEPRECATION_WARNINGS
 
 	/** @return true if the WidgetToFind is in this WidgetPath, false otherwise. */
 	UE_DEPRECATED(5.0, "ContainsWidget that takes a shared ptr is deprecated.")
