Fix shadowing error

diff --git a/Engine/Source/Runtime/SlateCore/Public/Layout/WidgetSlotWithAttributeSupport.h b/Engine/Source/Runtime/SlateCore/Public/Layout/WidgetSlotWithAttributeSupport.h
--- a/Engine/Source/Runtime/SlateCore/Public/Layout/WidgetSlotWithAttributeSupport.h
+++ b/Engine/Source/Runtime/SlateCore/Public/Layout/WidgetSlotWithAttributeSupport.h
@@ -56,11 +56,11 @@ public:
 
 	void RequestSortAttribute()
 	{
-		SWidget* Widget = Super::GetOwnerWidget();
-		ensureAlwaysMsgf(Widget, TEXT("FSlot needs to be constructed before we modify the FChildren."));
-		if (Widget)
+		SWidget* WidgetPtr = Super::GetOwnerWidget();
+		ensureAlwaysMsgf(WidgetPtr, TEXT("FSlot needs to be constructed before we modify the FChildren."));
+		if (WidgetPtr)
 		{
-			UpdateContainerSortOrder(*Widget);
+			UpdateContainerSortOrder(*WidgetPtr);
 		}
 	}
 
