"Fix" error C4003 not enough arguments for function-like macro invocation 'max'

diff --git a/Engine/Source/Runtime/SlateCore/Public/Layout/WidgetSlotWithAttributeSupport.h b/Engine/Source/Runtime/SlateCore/Public/Layout/WidgetSlotWithAttributeSupport.h
--- a/Engine/Source/Runtime/SlateCore/Public/Layout/WidgetSlotWithAttributeSupport.h
+++ b/Engine/Source/Runtime/SlateCore/Public/Layout/WidgetSlotWithAttributeSupport.h
@@ -89,7 +89,10 @@ public:
 				return Index;
 			}
 		}
+#pragma push_macro("UNDEF_MAX")
+#undef max
 		return std::numeric_limits<uint32>::max();
+#pragma pop_macro("UNDEF_MAX")
 	}
 	//~ End ISlateAttributeContainer interface
 };
