Fix shadow declaration type introduced with CL16841247

diff --git a/Engine/Source/Runtime/SlateCore/Public/Layout/BasicLayoutWidgetSlot.h b/Engine/Source/Runtime/SlateCore/Public/Layout/BasicLayoutWidgetSlot.h
--- a/Engine/Source/Runtime/SlateCore/Public/Layout/BasicLayoutWidgetSlot.h
+++ b/Engine/Source/Runtime/SlateCore/Public/Layout/BasicLayoutWidgetSlot.h
@@ -322,9 +322,9 @@ public:
 	FSlateDeprecatedTAttribute<FMargin> SlotPadding;
 
 private:
-	using SlotCompareType = TSlateAttributeComparePredicate<>;
-	using SlotType = ::SlateAttributePrivate::TSlateContainedAttribute<FMargin, ::SlateAttributePrivate::FSlateAttributeNoInvalidationReason, SlotCompareType>;
-	SlotType SlotPaddingAttribute;
+	using SlotPaddingCompareType = TSlateAttributeComparePredicate<>;
+	using SlotPaddingType = ::SlateAttributePrivate::TSlateContainedAttribute<FMargin, ::SlateAttributePrivate::FSlateAttributeNoInvalidationReason, SlotPaddingCompareType>;
+	SlotPaddingType SlotPaddingAttribute;
 };
 PRAGMA_ENABLE_DEPRECATION_WARNINGS
 
