Slate: Rename FSlateDeprecatedTAttribute to TSlateDeprecatedTAttribute to respect the coding standards.

diff --git a/Engine/Source/Runtime/Slate/Public/Widgets/Colors/SColorGradingWheel.h b/Engine/Source/Runtime/Slate/Public/Widgets/Colors/SColorGradingWheel.h
--- a/Engine/Source/Runtime/Slate/Public/Widgets/Colors/SColorGradingWheel.h
+++ b/Engine/Source/Runtime/Slate/Public/Widgets/Colors/SColorGradingWheel.h
@@ -121,11 +121,11 @@ protected:
 
 #if WITH_EDITORONLY_DATA
 	UE_DEPRECATED(5.0, "Direct access to SelectedColor is now deprecated. Use the setter or getter.")
-	FSlateDeprecatedTAttribute<FLinearColor> SelectedColor;
+	TSlateDeprecatedTAttribute<FLinearColor> SelectedColor;
 	UE_DEPRECATED(5.0, "Direct access to DesiredWheelSize is now deprecated. Use the setter or getter.")
-	FSlateDeprecatedTAttribute<int32> DesiredWheelSize;
+	TSlateDeprecatedTAttribute<int32> DesiredWheelSize;
 	UE_DEPRECATED(5.0, "Direct access to DesiredSizeOverride is now deprecated. Use the setter or getter.")
-	FSlateDeprecatedTAttribute<float> ExponentDisplacement;
+	TSlateDeprecatedTAttribute<float> ExponentDisplacement;
 #endif
 
 private:
diff --git a/Engine/Source/Runtime/Slate/Public/Widgets/Input/SButton.h b/Engine/Source/Runtime/Slate/Public/Widgets/Input/SButton.h
--- a/Engine/Source/Runtime/Slate/Public/Widgets/Input/SButton.h
+++ b/Engine/Source/Runtime/Slate/Public/Widgets/Input/SButton.h
@@ -271,7 +271,7 @@ private:
 protected:
 #if WITH_EDITORONLY_DATA
 	UE_DEPRECATED(5.0, "Direct access to ContentPadding is now deprecated. Use the setter or getter.")
-	FSlateDeprecatedTAttribute<FMargin> ContentPadding;
+	TSlateDeprecatedTAttribute<FMargin> ContentPadding;
 	/** Brush resource that represents a button */
 	UE_DEPRECATED(5.0, "NormalImage is now deprecated. Use the ButtonStyle.")
 	const FSlateBrush* NormalImage;
diff --git a/Engine/Source/Runtime/Slate/Public/Widgets/Layout/SBackgroundBlur.h b/Engine/Source/Runtime/Slate/Public/Widgets/Layout/SBackgroundBlur.h
--- a/Engine/Source/Runtime/Slate/Public/Widgets/Layout/SBackgroundBlur.h
+++ b/Engine/Source/Runtime/Slate/Public/Widgets/Layout/SBackgroundBlur.h
@@ -74,8 +74,8 @@ protected:
 	TSlateAttributeRef<TOptional<int32>> GetBlurRadiusAttribute() const { return TSlateAttributeRef<TOptional<int32>>{SharedThis(this), BlurRadiusAttribute}; }
 
 #if WITH_EDITORONLY_DATA
-	FSlateDeprecatedTAttribute<float> BlurStrength;
-	FSlateDeprecatedTAttribute<TOptional<int32>> BlurRadius;
+	TSlateDeprecatedTAttribute<float> BlurStrength;
+	TSlateDeprecatedTAttribute<TOptional<int32>> BlurRadius;
 #endif
 
 	bool bApplyAlphaToBlur;
diff --git a/Engine/Source/Runtime/Slate/Public/Widgets/Layout/SBorder.h b/Engine/Source/Runtime/Slate/Public/Widgets/Layout/SBorder.h
--- a/Engine/Source/Runtime/Slate/Public/Widgets/Layout/SBorder.h
+++ b/Engine/Source/Runtime/Slate/Public/Widgets/Layout/SBorder.h
@@ -149,11 +149,11 @@ protected:
 	UE_DEPRECATED(5.0, "Direct access to BorderImage is now deprecated. Use the setter or getter.")
 	FInvalidatableBrushAttribute BorderImage;
 	UE_DEPRECATED(5.0, "Direct access to BorderBackgroundColor is now deprecated. Use the setter or getter.")
-	FSlateDeprecatedTAttribute<FSlateColor> BorderBackgroundColor;
+	TSlateDeprecatedTAttribute<FSlateColor> BorderBackgroundColor;
 	UE_DEPRECATED(5.0, "Direct access to DesiredSizeScale is now deprecated. Use the setter or getter.")
-	FSlateDeprecatedTAttribute<FVector2D> DesiredSizeScale;
+	TSlateDeprecatedTAttribute<FVector2D> DesiredSizeScale;
 	UE_DEPRECATED(5.0, "Direct access to ShowDisabledEffect is now deprecated. Use the setter or getter.")
-	FSlateDeprecatedTAttribute<bool> ShowDisabledEffect;
+	TSlateDeprecatedTAttribute<bool> ShowDisabledEffect;
 #endif
 
  private:
diff --git a/Engine/Source/Runtime/Slate/Public/Widgets/Layout/SDPIScaler.h b/Engine/Source/Runtime/Slate/Public/Widgets/Layout/SDPIScaler.h
--- a/Engine/Source/Runtime/Slate/Public/Widgets/Layout/SDPIScaler.h
+++ b/Engine/Source/Runtime/Slate/Public/Widgets/Layout/SDPIScaler.h
@@ -62,7 +62,7 @@ protected:
 
 #if WITH_EDITORONLY_DATA
 	UE_DEPRECATED(5.0, "Direct access to DPIScale is now deprecated. Use the setter or getter.")
-	FSlateDeprecatedTAttribute<float> DPIScale;
+	TSlateDeprecatedTAttribute<float> DPIScale;
 #endif
 
 	struct FDPIScalerOneChildSlot : ::TSingleWidgetChildrenWithBasicLayoutSlot<EInvalidateWidgetReason::None> // we want to add it to the Attribute descriptor
diff --git a/Engine/Source/Runtime/Slate/Public/Widgets/SCanvas.h b/Engine/Source/Runtime/Slate/Public/Widgets/SCanvas.h
--- a/Engine/Source/Runtime/Slate/Public/Widgets/SCanvas.h
+++ b/Engine/Source/Runtime/Slate/Public/Widgets/SCanvas.h
@@ -79,9 +79,9 @@ public:
 	public:
 #if WITH_EDITORONLY_DATA
 		UE_DEPRECATED(5.0, "Direct access to PositionAttr is now deprecated. Use the getter or setter.")
-		FSlateDeprecatedTAttribute<FVector2D> PositionAttr;
+		TSlateDeprecatedTAttribute<FVector2D> PositionAttr;
 		UE_DEPRECATED(5.0, "Direct access to SizeAttr is now deprecated. Use the getter or setter.")
-		FSlateDeprecatedTAttribute<FVector2D> SizeAttr;
+		TSlateDeprecatedTAttribute<FVector2D> SizeAttr;
 #endif
 
 	public:
diff --git a/Engine/Source/Runtime/SlateCore/Public/Layout/BasicLayoutWidgetSlot.h b/Engine/Source/Runtime/SlateCore/Public/Layout/BasicLayoutWidgetSlot.h
--- a/Engine/Source/Runtime/SlateCore/Public/Layout/BasicLayoutWidgetSlot.h
+++ b/Engine/Source/Runtime/SlateCore/Public/Layout/BasicLayoutWidgetSlot.h
@@ -319,7 +319,7 @@ public:
 
 public:
 	UE_DEPRECATED(5.0, "Direct access to SlotPadding is now deprecated. Use the setter or getter.")
-	FSlateDeprecatedTAttribute<FMargin> SlotPadding;
+	TSlateDeprecatedTAttribute<FMargin> SlotPadding;
 
 private:
 	using SlotPaddingCompareType = TSlateAttributeComparePredicate<>;
@@ -421,7 +421,7 @@ public:
 
 #if WITH_EDITORONLY_DATA
 	UE_DEPRECATED(5.0, "Direct access to SlotPadding is now deprecated. Use the setter or getter.")
-	FSlateDeprecatedTAttribute<FMargin> SlotPadding;
+	TSlateDeprecatedTAttribute<FMargin> SlotPadding;
 #endif
 
 public:
diff --git a/Engine/Source/Runtime/SlateCore/Public/Types/SlateAttribute.h b/Engine/Source/Runtime/SlateCore/Public/Types/SlateAttribute.h
--- a/Engine/Source/Runtime/SlateCore/Public/Types/SlateAttribute.h
+++ b/Engine/Source/Runtime/SlateCore/Public/Types/SlateAttribute.h
@@ -218,22 +218,22 @@ struct TSlateAttributeInvalidationReason
 
 /** A structure used to help the user identify deprecated TAttribute that are now TSlateAttribute. */
 template<typename ObjectType>
-struct FSlateDeprecatedTAttribute
+struct TSlateDeprecatedTAttribute
 {
-	FSlateDeprecatedTAttribute() = default;
+	TSlateDeprecatedTAttribute() = default;
 
 	using FGetter = typename TAttribute<ObjectType>::FGetter;
 
 	template<typename OtherType>
-	FSlateDeprecatedTAttribute(const OtherType& InInitialValue)	{ }
+	TSlateDeprecatedTAttribute(const OtherType& InInitialValue)	{ }
 
-	FSlateDeprecatedTAttribute(ObjectType&& InInitialValue)	{ }
+	TSlateDeprecatedTAttribute(ObjectType&& InInitialValue)	{ }
 
 	template<class SourceType>
-	FSlateDeprecatedTAttribute(TSharedRef<SourceType> InUserObject, typename FGetter::template TConstMethodPtr< SourceType > InMethodPtr) {}
+	TSlateDeprecatedTAttribute(TSharedRef<SourceType> InUserObject, typename FGetter::template TConstMethodPtr< SourceType > InMethodPtr) {}
 
 	template< class SourceType >
-	FSlateDeprecatedTAttribute(SourceType* InUserObject, typename FGetter::template TConstMethodPtr< SourceType > InMethodPtr) { }
+	TSlateDeprecatedTAttribute(SourceType* InUserObject, typename FGetter::template TConstMethodPtr< SourceType > InMethodPtr) { }
 
 	bool IsSet() const { return false; }
 
diff --git a/Engine/Source/Runtime/SlateCore/Public/Widgets/Images/SImage.h b/Engine/Source/Runtime/SlateCore/Public/Widgets/Images/SImage.h
--- a/Engine/Source/Runtime/SlateCore/Public/Widgets/Images/SImage.h
+++ b/Engine/Source/Runtime/SlateCore/Public/Widgets/Images/SImage.h
@@ -95,9 +95,9 @@ protected:
 	UE_DEPRECATED(5.0, "Direct access to Image is now deprecated. Use the setter or getter.")
 	FInvalidatableBrushAttribute Image;
 	UE_DEPRECATED(5.0, "Direct access to ColorAndOpacity is now deprecated. Use the setter or getter.")
-	FSlateDeprecatedTAttribute<FSlateColor> ColorAndOpacity;
+	TSlateDeprecatedTAttribute<FSlateColor> ColorAndOpacity;
 	UE_DEPRECATED(5.0, "Direct access to DesiredSizeOverride is now deprecated. Use the setter or getter.")
-	FSlateDeprecatedTAttribute<TOptional<FVector2D>> DesiredSizeOverride;
+	TSlateDeprecatedTAttribute<TOptional<FVector2D>> DesiredSizeOverride;
 #endif
 
 private:
diff --git a/Engine/Source/Runtime/SlateCore/Public/Widgets/SCompoundWidget.h b/Engine/Source/Runtime/SlateCore/Public/Widgets/SCompoundWidget.h
--- a/Engine/Source/Runtime/SlateCore/Public/Widgets/SCompoundWidget.h
+++ b/Engine/Source/Runtime/SlateCore/Public/Widgets/SCompoundWidget.h
@@ -113,11 +113,11 @@ protected:
 
 #if WITH_EDITORONLY_DATA
 	UE_DEPRECATED(5.0, "Direct access to ContentScale is now deprecated. Use the setter or getter.")
-	FSlateDeprecatedTAttribute<FVector2D> ContentScale;
+	TSlateDeprecatedTAttribute<FVector2D> ContentScale;
 	UE_DEPRECATED(5.0, "Direct access to ColorAndOpacity is now deprecated. Use the setter or getter.")
-	FSlateDeprecatedTAttribute<FLinearColor> ColorAndOpacity;
+	TSlateDeprecatedTAttribute<FLinearColor> ColorAndOpacity;
 	UE_DEPRECATED(5.0, "Direct access to ForegroundColor is now deprecated. Use the setter or getter.")
-	FSlateDeprecatedTAttribute<FSlateColor> ForegroundColor;
+	TSlateDeprecatedTAttribute<FSlateColor> ForegroundColor;
 #endif
 
 private:
diff --git a/Engine/Source/Runtime/SlateCore/Public/Widgets/SWidget.h b/Engine/Source/Runtime/SlateCore/Public/Widgets/SWidget.h
--- a/Engine/Source/Runtime/SlateCore/Public/Widgets/SWidget.h
+++ b/Engine/Source/Runtime/SlateCore/Public/Widgets/SWidget.h
@@ -1865,13 +1865,13 @@ protected:
 #if WITH_EDITORONLY_DATA
 	/** Whether or not this widget is enabled */
 	UE_DEPRECATED(5.0, "Direct access to EnabledState is now deprecated. Use the setter or getter.")
-	FSlateDeprecatedTAttribute<bool> EnabledState;
+	TSlateDeprecatedTAttribute<bool> EnabledState;
 	/** Is this widget visible, hidden or collapsed */
 	UE_DEPRECATED(5.0, "Direct access to Visibility is now deprecated. Use the setter or getter.")
-	FSlateDeprecatedTAttribute<EVisibility> Visibility;
+	TSlateDeprecatedTAttribute<EVisibility> Visibility;
 	/** Render transform of this widget. TOptional<> to allow code to skip expensive overhead if there is no render transform applied. */
 	UE_DEPRECATED(5.0, "Direct access to RenderTransform is now deprecated. Use the setter or getter.")
-	FSlateDeprecatedTAttribute< TOptional<FSlateRenderTransform> > RenderTransform;
+	TSlateDeprecatedTAttribute< TOptional<FSlateRenderTransform> > RenderTransform;
 	/** Render transform pivot of this widget (in normalized local space) */
 	UE_DEPRECATED(5.0, "Direct access to RenderTransformPivot is now deprecated. Use the setter or getter.")
 	TAttribute<FVector2D> RenderTransformPivot;
