Slate: Add test to verify that SlateAttribute of types float and double and finite (not NaN).

diff --git a/Engine/Source/Runtime/SlateCore/Public/Types/Attributes/SlateAttributeBase.inl b/Engine/Source/Runtime/SlateCore/Public/Types/Attributes/SlateAttributeBase.inl
--- a/Engine/Source/Runtime/SlateCore/Public/Types/Attributes/SlateAttributeBase.inl
+++ b/Engine/Source/Runtime/SlateCore/Public/Types/Attributes/SlateAttributeBase.inl
@@ -60,14 +60,30 @@ namespace SlateAttributePrivate
 		void VerifyNan() const
 		{
 #if UE_SLATE_WITH_ATTRIBUTE_NAN_DIAGNOSTIC
-			if constexpr (std::is_same<FVector, InObjectType>::value)
+			if constexpr (std::is_same<float, InObjectType>::value)
 			{
-				ensureMsgf(!Value.ContainsNaN(), TEXT("Value contains a NaN. Initialize your FVector properly (see FVector::EForceInit)"));
+				ensureMsgf(FMath::IsFinite(Value), TEXT("Value contains a NaN. Initialize your float properly"));
+			}
+			else if constexpr (std::is_same<double, InObjectType>::value)
+			{
+				ensureMsgf(FMath::IsFinite(Value), TEXT("Value contains a NaN. Initialize your double properly"));
 			}
 			else if constexpr (std::is_same<FVector2D, InObjectType>::value)
 			{
 				ensureMsgf(!Value.ContainsNaN(), TEXT("Value contains a NaN. Initialize your FVector2D properly (see FVector2D::EForceInit)"));
 			}
+			else if constexpr (std::is_same<FVector3f, InObjectType>::value)
+			{
+				ensureMsgf(!Value.ContainsNaN(), TEXT("Value contains a NaN. Initialize your FVector3f properly (see FVector::EForceInit)"));
+			}
+			else if constexpr (std::is_same<FVector3d, InObjectType>::value)
+			{
+				ensureMsgf(!Value.ContainsNaN(), TEXT("Value contains a NaN. Initialize your FVector3d properly (see FVector::EForceInit)"));
+			}
+			else if constexpr (std::is_same<FVector4, InObjectType>::value)
+			{
+				ensureMsgf(!Value.ContainsNaN(), TEXT("Value contains a NaN. Initialize your FVector4 properly"));
+			}
 			else if constexpr (std::is_same<FLinearColor, InObjectType>::value)
 			{
 				ensureMsgf(!FVector4(Value).ContainsNaN(), TEXT("Value contains a NaN. Initialize your FLinearColor properly (see FLinearColor::EForceInit)"));
