Exclude Widget Reflector window from debugging #RB patrick.boutot #jira ue-118804

diff --git a/Engine/Source/Editor/UnrealEd/Private/PlayLevel.cpp b/Engine/Source/Editor/UnrealEd/Private/PlayLevel.cpp
--- a/Engine/Source/Editor/UnrealEd/Private/PlayLevel.cpp
+++ b/Engine/Source/Editor/UnrealEd/Private/PlayLevel.cpp
@@ -3094,6 +3094,7 @@ TSharedRef<SPIEViewport> UEditorEngine::GeneratePIEViewportWindow(const FRequest
 		FText ViewportName = GeneratePIEViewportWindowTitle(InNetMode, PreviewPlatform.GetEffectivePreviewFeatureLevel(), InSessionParams, InWorldContext.PIEInstance, InWorldContext.PIEFixedTickSeconds);
 		PieWindow = SNew(SWindow)
 			.Title(ViewportName)
+			.Tag("PIEWindow")
 			.ScreenPosition(FVector2D(WindowPosition.X, WindowPosition.Y))
 			.ClientSize(FVector2D(WindowSize.X, WindowSize.Y))	
 			.AutoCenter(bCenterNewWindowOverride ? EAutoCenter::PreferredWorkArea : EAutoCenter::None)
@@ -3105,7 +3106,6 @@ TSharedRef<SPIEViewport> UEditorEngine::GeneratePIEViewportWindow(const FRequest
 		PieWindow->SetAllowFastUpdate(true);
 	}
 
-
 	// Setup a delegate for switching to the play world on slate input events, drawing and ticking
 	FOnSwitchWorldHack OnWorldSwitch = FOnSwitchWorldHack::CreateUObject(this, &UEditorEngine::OnSwitchWorldForSlatePieWindow, InWorldContext.PIEInstance);
 	PieWindow->SetOnWorldSwitchHack(OnWorldSwitch);
diff --git a/Engine/Source/Runtime/SlateCore/Private/Debugging/ConsoleSlateDebuggerPaint.cpp b/Engine/Source/Runtime/SlateCore/Private/Debugging/ConsoleSlateDebuggerPaint.cpp
--- a/Engine/Source/Runtime/SlateCore/Private/Debugging/ConsoleSlateDebuggerPaint.cpp
+++ b/Engine/Source/Runtime/SlateCore/Private/Debugging/ConsoleSlateDebuggerPaint.cpp
@@ -26,11 +26,13 @@ FConsoleSlateDebuggerPaint::FConsoleSlateDebuggerPaint()
 	, bLogWidgetName(false)
 	, bLogWidgetNameOnce(false)
 	, bLogWarningIfWidgetIsPaintedMoreThanOnce(true)
+	, bDebugGameWindowOnly(true)
 	, DrawBoxColor(1.0f, 1.0f, 0.0f, 0.2f)
 	, DrawQuadColor(1.0f, 1.0f, 1.0f, 1.0f)
 	, DrawWidgetNameColor(FColorList::SpicyPink)
 	, MaxNumberOfWidgetInList(20)
 	, CacheDuration(2.0f)
+	, PIEWindowTag("PIEWindow")
 	, ShowPaintWidgetCommand(
 		TEXT("SlateDebugger.Paint.Start")
 		, TEXT("Start the painted widget debug tool. Use to show widget that have been painted this frame.")
@@ -60,6 +62,10 @@ FConsoleSlateDebuggerPaint::FConsoleSlateDebuggerPaint()
 		TEXT("SlateDebugger.Paint.LogWarningIfWidgetIsPaintedMoreThanOnce")
 		, bLogWarningIfWidgetIsPaintedMoreThanOnce
 		, TEXT("Option to log a warning if a widget is painted more than once in a single frame."))
+	, OnlyGameWindow(
+		TEXT("SlateDebugger.Paint.OnlyGameWindow"),
+		bDebugGameWindowOnly,
+		TEXT("Option to only the debug the game window"))
 {
 	GConfig->GetBool(TEXT("SlateDebugger.Paint"), TEXT("bDisplayWidgetsNameList"), bDisplayWidgetsNameList, *GEditorPerProjectIni);
 	GConfig->GetBool(TEXT("SlateDebugger.Paint"), TEXT("bUseWidgetPathAsName"), bUseWidgetPathAsName, *GEditorPerProjectIni);
@@ -179,6 +185,14 @@ void FConsoleSlateDebuggerPaint::HandleEndWidgetPaint(const SWidget* Widget, con
 	const FConsoleSlateDebuggerUtility::TSWidgetId WidgetId = FConsoleSlateDebuggerUtility::GetId(Widget);
 	const FConsoleSlateDebuggerUtility::TSWindowId WindowId = FConsoleSlateDebuggerUtility::GetId(OutDrawElements.GetPaintWindow());
 
+	// Exclude all windows but the game window
+	SWindow* WindowToDrawIn = OutDrawElements.GetPaintWindow();
+	if (bDebugGameWindowOnly && (WindowToDrawIn->GetType() != EWindowType::GameWindow && WindowToDrawIn->GetTag() != PIEWindowTag))
+	{
+		return;
+	}
+
+
 	FPaintInfo* FoundItem = PaintedWidgets.Find(WidgetId);
 	if (FoundItem == nullptr)
 	{
diff --git a/Engine/Source/Runtime/SlateCore/Private/Debugging/ConsoleSlateDebuggerPaint.h b/Engine/Source/Runtime/SlateCore/Private/Debugging/ConsoleSlateDebuggerPaint.h
--- a/Engine/Source/Runtime/SlateCore/Private/Debugging/ConsoleSlateDebuggerPaint.h
+++ b/Engine/Source/Runtime/SlateCore/Private/Debugging/ConsoleSlateDebuggerPaint.h
@@ -59,11 +59,13 @@ private:
 	bool bLogWidgetName;
 	bool bLogWidgetNameOnce;
 	bool bLogWarningIfWidgetIsPaintedMoreThanOnce;
+	bool bDebugGameWindowOnly;
 	FLinearColor DrawBoxColor;
 	FLinearColor DrawQuadColor;
 	FLinearColor DrawWidgetNameColor;
 	int32 MaxNumberOfWidgetInList;
 	float CacheDuration;
+	FName PIEWindowTag;
 
 	//~ Console objects
 	FAutoConsoleCommand ShowPaintWidgetCommand;
@@ -73,6 +75,7 @@ private:
 	FAutoConsoleCommand DisplayWidgetsNameListCommand;
 	FAutoConsoleVariableRef MaxNumberOfWidgetInListtRefCVar;
 	FAutoConsoleVariableRef LogWarningIfWidgetIsPaintedMoreThanOnceRefCVar;
+	FAutoConsoleVariableRef OnlyGameWindow;
 
 	struct FPaintInfo 
 	{
diff --git a/Engine/Source/Runtime/SlateCore/Private/Debugging/ConsoleSlateDebuggerUpdate.cpp b/Engine/Source/Runtime/SlateCore/Private/Debugging/ConsoleSlateDebuggerUpdate.cpp
--- a/Engine/Source/Runtime/SlateCore/Private/Debugging/ConsoleSlateDebuggerUpdate.cpp
+++ b/Engine/Source/Runtime/SlateCore/Private/Debugging/ConsoleSlateDebuggerUpdate.cpp
@@ -14,6 +14,7 @@
 #include "Styling/CoreStyle.h"
 #include "Types/ReflectionMetadata.h"
 #include "Widgets/SWidget.h"
+#include "Widgets/SWindow.h"
 
 #define LOCTEXT_NAMESPACE "ConsoleSlateDebuggerUpdate"
 
@@ -25,6 +26,7 @@ FConsoleSlateDebuggerUpdate::FConsoleSlateDebuggerUpdate()
 	, bDisplayUpdateFromPaint(false)
 	, bShowLegend(false)
 	, bShowQuad(false)
+	, bDebugGameWindowOnly(true)
 	, WidgetUpdateFlagsFilter(EWidgetUpdateFlags::AnyUpdate)
 	, DrawVolatilePaintColor(FColorList::Red)
 	, DrawRepaintColor(FColorList::Yellow)
@@ -34,6 +36,7 @@ FConsoleSlateDebuggerUpdate::FConsoleSlateDebuggerUpdate()
 	, MaxNumberOfWidgetInList(20)
 	, InvalidationRootIdFilter(-1)
 	, CacheDuration(2.0f)
+	, PIEWindowTag("PIEWindow")
 	, StartCommand(
 		TEXT("SlateDebugger.Update.Start"),
 		TEXT("Start the update widget debug tool. It shows when widgets are updated."),
@@ -67,6 +70,10 @@ FConsoleSlateDebuggerUpdate::FConsoleSlateDebuggerUpdate()
 		TEXT("SlateDebugger.Update.SetInvalidationRootIdFilter"),
 		InvalidationRootIdFilter,
 		TEXT("Option to show only the widgets that are part of an invalidation root."))
+	, OnlyGameWindow(
+		TEXT("SlateDebugger.Update.OnlyGameWindow"),
+		bDebugGameWindowOnly,
+		TEXT("Option to only the debug the game window"))
 {
 	GConfig->GetBool(TEXT("SlateDebugger.Update"), TEXT("bDisplayWidgetsNameList"), bDisplayWidgetsNameList, *GEditorPerProjectIni);
 	GConfig->GetBool(TEXT("SlateDebugger.Update"), TEXT("bUseWidgetPathAsName"), bUseWidgetPathAsName, *GEditorPerProjectIni);
@@ -351,8 +358,16 @@ void FConsoleSlateDebuggerUpdate::HandleWidgetUpdate(const FSlateDebuggingWidget
 
 void FConsoleSlateDebuggerUpdate::HandlePaintDebugInfo(const FPaintArgs& InArgs, const FGeometry& InAllottedGeometry, FSlateWindowElementList& InOutDrawElements, int32& InOutLayerId)
 {
+
 	++InOutLayerId;
 
+	// Exclude all windows but the game window
+	SWindow* WindowToDrawIn = InOutDrawElements.GetPaintWindow();
+	if (bDebugGameWindowOnly && (WindowToDrawIn->GetType() != EWindowType::GameWindow && WindowToDrawIn->GetTag() != PIEWindowTag))
+	{
+		return;
+	}
+
 	const FConsoleSlateDebuggerUtility::TSWindowId PaintWindow = FConsoleSlateDebuggerUtility::GetId(InOutDrawElements.GetPaintWindow());
 	int32 NumberOfWidget = 0;
 	TArray<const FString*, TInlineAllocator<100>> NamesToDisplay;
diff --git a/Engine/Source/Runtime/SlateCore/Private/Debugging/ConsoleSlateDebuggerUpdate.h b/Engine/Source/Runtime/SlateCore/Private/Debugging/ConsoleSlateDebuggerUpdate.h
--- a/Engine/Source/Runtime/SlateCore/Private/Debugging/ConsoleSlateDebuggerUpdate.h
+++ b/Engine/Source/Runtime/SlateCore/Private/Debugging/ConsoleSlateDebuggerUpdate.h
@@ -60,6 +60,7 @@ private:
 	bool bDisplayUpdateFromPaint;
 	bool bShowLegend;
 	bool bShowQuad;
+	bool bDebugGameWindowOnly;
 	EWidgetUpdateFlags WidgetUpdateFlagsFilter;
 	FLinearColor DrawVolatilePaintColor;
 	FLinearColor DrawRepaintColor;
@@ -69,6 +70,7 @@ private:
 	int32 MaxNumberOfWidgetInList;
 	int32 InvalidationRootIdFilter;
 	float CacheDuration;
+	FName PIEWindowTag;
 
 	//~ Console objects
 	FAutoConsoleCommand StartCommand;
@@ -79,6 +81,7 @@ private:
 	FAutoConsoleCommand ToogleDisplayUpdateFromPaintCommand;
 	FAutoConsoleCommand SetWidgetUpdateFlagsFilterCommand;
 	FAutoConsoleVariableRef InvalidationRootFilterRefCVar;
+	FAutoConsoleVariableRef OnlyGameWindow;
 
 	struct FWidgetInfo
 	{
