Temporary workaround: demoting assert to warning to unblock FastCook build tests.
#rnx

diff --git a/Engine/Source/Runtime/SlateCore/Private/FastUpdate/SlateInvalidationRoot.cpp b/Engine/Source/Runtime/SlateCore/Private/FastUpdate/SlateInvalidationRoot.cpp
--- a/Engine/Source/Runtime/SlateCore/Private/FastUpdate/SlateInvalidationRoot.cpp
+++ b/Engine/Source/Runtime/SlateCore/Private/FastUpdate/SlateInvalidationRoot.cpp
@@ -677,8 +677,12 @@ void FSlateInvalidationRoot::ProcessPreUpdate()
 						//EnumRemoveFlags(InvalidationWidget.CurrentInvalidateReason, EInvalidateWidgetReason::ChildOrder);
 
 						// The child may change and may become invalid but not the parent.
-						checkf(WidgetPtr->GetProxyHandle().IsValid(WidgetPtr) && WidgetPtr->GetProxyHandle().GetWidgetIndex() == InvalidationWidget.Index
-							, TEXT("The widget '%s' should be valid after a ProcessChildOrderInvalidation."), *FReflectionMetaData::GetWidgetDebugInfo(WidgetPtr));
+						UE_CLOG(!(WidgetPtr->GetProxyHandle().IsValid(WidgetPtr) && WidgetPtr->GetProxyHandle().GetWidgetIndex() == InvalidationWidget.Index),
+							LogSlate,
+							Warning,
+							TEXT("The widget '%s' should be valid after a ProcessChildOrderInvalidation."), *FReflectionMetaData::GetWidgetDebugInfo(WidgetPtr));
+						//checkf(WidgetPtr->GetProxyHandle().IsValid(WidgetPtr) && WidgetPtr->GetProxyHandle().GetWidgetIndex() == InvalidationWidget.Index
+						//	, TEXT("The widget '%s' should be valid after a ProcessChildOrderInvalidation."), *FReflectionMetaData::GetWidgetDebugInfo(WidgetPtr));
 					}
 
 					if (EnumHasAnyFlags(InvalidationWidget.CurrentInvalidateReason, EInvalidateWidgetReason::AttributeRegistration))
