Slate: Prevent access to an invalided widget after we process its child order.

diff --git a/Engine/Source/Runtime/SlateCore/Private/FastUpdate/SlateInvalidationRoot.cpp b/Engine/Source/Runtime/SlateCore/Private/FastUpdate/SlateInvalidationRoot.cpp
--- a/Engine/Source/Runtime/SlateCore/Private/FastUpdate/SlateInvalidationRoot.cpp
+++ b/Engine/Source/Runtime/SlateCore/Private/FastUpdate/SlateInvalidationRoot.cpp
@@ -642,6 +642,7 @@ void FSlateInvalidationRoot::ProcessPreUpdate()
 					}
 #endif
 
+					bool bIsInvalidationWidgetValid = true;
 					if (EnumHasAnyFlags(InvalidationWidget.CurrentInvalidateReason, EInvalidateWidgetReason::ChildOrder))
 					{
 // Uncomment to see to be able to compare the list before and after when debugging
@@ -671,21 +672,13 @@ void FSlateInvalidationRoot::ProcessPreUpdate()
 #endif
 
 						TGuardValue<bool> ProcessChildOrderInvalidationGuardValue(bProcessingChildOrderInvalidation, true);
-						FastWidgetPathList->ProcessChildOrderInvalidation(InvalidationWidget, ChildOrderInvalidationCallback);
+						bIsInvalidationWidgetValid = FastWidgetPathList->ProcessChildOrderInvalidation(InvalidationWidget, ChildOrderInvalidationCallback);
 
 						// We need to keep it to run the layout calculation in FWidgetProxy::ProcessPostInvalidation
 						//EnumRemoveFlags(InvalidationWidget.CurrentInvalidateReason, EInvalidateWidgetReason::ChildOrder);
-
-						// The child may change and may become invalid but not the parent.
-						UE_CLOG(!(WidgetPtr->GetProxyHandle().IsValid(WidgetPtr) && WidgetPtr->GetProxyHandle().GetWidgetIndex() == InvalidationWidget.Index),
-							LogSlate,
-							Warning,
-							TEXT("The widget '%s' should be valid after a ProcessChildOrderInvalidation."), *FReflectionMetaData::GetWidgetDebugInfo(WidgetPtr));
-						//checkf(WidgetPtr->GetProxyHandle().IsValid(WidgetPtr) && WidgetPtr->GetProxyHandle().GetWidgetIndex() == InvalidationWidget.Index
-						//	, TEXT("The widget '%s' should be valid after a ProcessChildOrderInvalidation."), *FReflectionMetaData::GetWidgetDebugInfo(WidgetPtr));
 					}
 
-					if (EnumHasAnyFlags(InvalidationWidget.CurrentInvalidateReason, EInvalidateWidgetReason::AttributeRegistration))
+					if (bIsInvalidationWidgetValid && EnumHasAnyFlags(InvalidationWidget.CurrentInvalidateReason, EInvalidateWidgetReason::AttributeRegistration))
 					{
 						FastWidgetPathList->ProcessAttributeRegistrationInvalidation(InvalidationWidget);
 						EnumRemoveFlags(InvalidationWidget.CurrentInvalidateReason, EInvalidateWidgetReason::AttributeRegistration);
