Slate text widgets now have a "Text Overflow" argument that dicates what happens to text when it is larger than its clip rect. The two overflow methods are Clip and Ellipsis. Clip is the default and existing behavior where text is clipped per pixel. Ellipsis replaces clipped text with a "..." string to provide a more polished look when text is clipped.
Supports left and right justification and left and right text flows.

The ellipsis uses the U+2026 character or 3 dots if the font does not have U+2026. Localizers can override the text used.

Unsupported:
- Simple, non-shaped text
- Non-axis aligned clip rects.
- Center justified text

diff --git a/Engine/Source/Runtime/Slate/Private/Framework/Text/SlateHyperlinkRun.cpp b/Engine/Source/Runtime/Slate/Private/Framework/Text/SlateHyperlinkRun.cpp
--- a/Engine/Source/Runtime/Slate/Private/Framework/Text/SlateHyperlinkRun.cpp
+++ b/Engine/Source/Runtime/Slate/Private/Framework/Text/SlateHyperlinkRun.cpp
@@ -106,15 +106,15 @@ void FSlateHyperlinkRun::OnNavigate()
 	NavigateDelegate.Execute( RunInfo.MetaData );
 }
 
-int32 FSlateHyperlinkRun::OnPaint( const FPaintArgs& Args, const FTextLayout::FLineView& Line, const TSharedRef< ILayoutBlock >& Block, const FTextBlockStyle& DefaultStyle, const FGeometry& AllottedGeometry, const FSlateRect& MyCullingRect, FSlateWindowElementList& OutDrawElements, int32 LayerId, const FWidgetStyle& InWidgetStyle, bool bParentEnabled ) const 
+int32 FSlateHyperlinkRun::OnPaint(const FPaintArgs& PaintArgs, const FTextArgs& TextArgs, const FGeometry& AllottedGeometry, const FSlateRect& MyCullingRect, FSlateWindowElementList& OutDrawElements, int32 LayerId, const FWidgetStyle& InWidgetStyle, bool bParentEnabled) const
 {
-	const TSharedRef< FWidgetLayoutBlock > WidgetBlock = StaticCastSharedRef< FWidgetLayoutBlock >( Block );
+	const TSharedRef< FWidgetLayoutBlock > WidgetBlock = StaticCastSharedRef< FWidgetLayoutBlock >(TextArgs.Block);
 
 	// The block size and offset values are pre-scaled, so we need to account for that when converting the block offsets into paint geometry
 	const float InverseScale = Inverse(AllottedGeometry.Scale);
 
-	const FGeometry WidgetGeometry = AllottedGeometry.MakeChild(TransformVector(InverseScale, Block->GetSize()), FSlateLayoutTransform(TransformPoint(InverseScale, Block->GetLocationOffset())));
-	return WidgetBlock->GetWidget()->Paint( Args, WidgetGeometry, MyCullingRect, OutDrawElements, LayerId, InWidgetStyle, bParentEnabled );
+	const FGeometry WidgetGeometry = AllottedGeometry.MakeChild(TransformVector(InverseScale, TextArgs.Block->GetSize()), FSlateLayoutTransform(TransformPoint(InverseScale, TextArgs.Block->GetLocationOffset())));
+	return WidgetBlock->GetWidget()->Paint(PaintArgs, WidgetGeometry, MyCullingRect, OutDrawElements, LayerId, InWidgetStyle, bParentEnabled);
 }
 
 const TArray< TSharedRef<SWidget> >& FSlateHyperlinkRun::GetChildren()
diff --git a/Engine/Source/Runtime/Slate/Private/Framework/Text/SlateImageRun.cpp b/Engine/Source/Runtime/Slate/Private/Framework/Text/SlateImageRun.cpp
--- a/Engine/Source/Runtime/Slate/Private/Framework/Text/SlateImageRun.cpp
+++ b/Engine/Source/Runtime/Slate/Private/Framework/Text/SlateImageRun.cpp
@@ -143,23 +143,23 @@ FVector2D FSlateImageRun::GetLocationAt( const TSharedRef< ILayoutBlock >& Block
 	return Block->GetLocationOffset();
 }
 
-int32 FSlateImageRun::OnPaint( const FPaintArgs& Args, const FTextLayout::FLineView& Line, const TSharedRef< ILayoutBlock >& Block, const FTextBlockStyle& DefaultStyle, const FGeometry& AllottedGeometry, const FSlateRect& MyCullingRect, FSlateWindowElementList& OutDrawElements, int32 LayerId, const FWidgetStyle& InWidgetStyle, bool bParentEnabled ) const 
+int32 FSlateImageRun::OnPaint(const FPaintArgs& PaintArgs, const FTextArgs& TextArgs, const FGeometry& AllottedGeometry, const FSlateRect& MyCullingRect, FSlateWindowElementList& OutDrawElements, int32 LayerId, const FWidgetStyle& InWidgetStyle, bool bParentEnabled) const
 {
 	// The block size and offset values are pre-scaled, so we need to account for that when converting the block offsets into paint geometry
 	const float InverseScale = Inverse(AllottedGeometry.Scale);
 
-	if ( Image->DrawAs != ESlateBrushDrawType::NoDrawType )
+	if (Image->DrawAs != ESlateBrushDrawType::NoDrawType)
 	{
-		const FLinearColor FinalColorAndOpacity( InWidgetStyle.GetColorAndOpacityTint() * Image->GetTint( InWidgetStyle ) );
+		const FLinearColor FinalColorAndOpacity(InWidgetStyle.GetColorAndOpacityTint() * Image->GetTint(InWidgetStyle));
 		const ESlateDrawEffect DrawEffects = bParentEnabled ? ESlateDrawEffect::None : ESlateDrawEffect::DisabledEffect;
 		FSlateDrawElement::MakeBox(
-			OutDrawElements, 
-			++LayerId, 
-			AllottedGeometry.ToPaintGeometry(TransformVector(InverseScale, Block->GetSize()), FSlateLayoutTransform(TransformPoint(InverseScale, Block->GetLocationOffset()))), 
-			Image, 
-			DrawEffects, 
+			OutDrawElements,
+			++LayerId,
+			AllottedGeometry.ToPaintGeometry(TransformVector(InverseScale, TextArgs.Block->GetSize()), FSlateLayoutTransform(TransformPoint(InverseScale, TextArgs.Block->GetLocationOffset()))),
+			Image,
+			DrawEffects,
 			FinalColorAndOpacity
-			);
+		);
 	}
 
 	return LayerId;
diff --git a/Engine/Source/Runtime/Slate/Private/Framework/Text/SlatePasswordRun.cpp b/Engine/Source/Runtime/Slate/Private/Framework/Text/SlatePasswordRun.cpp
--- a/Engine/Source/Runtime/Slate/Private/Framework/Text/SlatePasswordRun.cpp
+++ b/Engine/Source/Runtime/Slate/Private/Framework/Text/SlatePasswordRun.cpp
@@ -41,26 +41,26 @@ int8 FSlatePasswordRun::GetKerning(int32 CurrentIndex, float Scale, const FRunTe
 	return FSlateApplication::Get().GetRenderer()->GetFontMeasureService()->GetKerning(Style.Font, Scale, GetPasswordChar(), GetPasswordChar());
 }
 
-int32 FSlatePasswordRun::OnPaint(const FPaintArgs& Args, const FTextLayout::FLineView& Line, const TSharedRef< ILayoutBlock >& Block, const FTextBlockStyle& DefaultStyle, const FGeometry& AllottedGeometry, const FSlateRect& MyCullingRect, FSlateWindowElementList& OutDrawElements, int32 LayerId, const FWidgetStyle& InWidgetStyle, bool bParentEnabled) const
+int32 FSlatePasswordRun::OnPaint(const FPaintArgs& PaintArgs, const FTextArgs& TextArgs, const FGeometry& AllottedGeometry, const FSlateRect& MyCullingRect, FSlateWindowElementList& OutDrawElements, int32 LayerId, const FWidgetStyle& InWidgetStyle, bool bParentEnabled) const
 {
 	const ESlateDrawEffect DrawEffects = bParentEnabled ? ESlateDrawEffect::None : ESlateDrawEffect::DisabledEffect;
 
 	const bool ShouldDropShadow = Style.ShadowColorAndOpacity.A > 0.f && Style.ShadowOffset.SizeSquared() > 0.f;
-	const FVector2D BlockLocationOffset = Block->GetLocationOffset();
-	const FTextRange BlockRange = Block->GetTextRange();
-	
+	const FVector2D BlockLocationOffset = TextArgs.Block->GetLocationOffset();
+	const FTextRange BlockRange = TextArgs.Block->GetTextRange();
+
 	// The block size and offset values are pre-scaled, so we need to account for that when converting the block offsets into paint geometry
 	const float InverseScale = Inverse(AllottedGeometry.Scale);
 
 	// A negative shadow offset should be applied as a positive offset to the text to avoid clipping issues
 	const FVector2D DrawShadowOffset(
-		(Style.ShadowOffset.X > 0.0f) ? Style.ShadowOffset.X * AllottedGeometry.Scale : 0.0f, 
+		(Style.ShadowOffset.X > 0.0f) ? Style.ShadowOffset.X * AllottedGeometry.Scale : 0.0f,
 		(Style.ShadowOffset.Y > 0.0f) ? Style.ShadowOffset.Y * AllottedGeometry.Scale : 0.0f
-		);
+	);
 	const FVector2D DrawTextOffset(
-		(Style.ShadowOffset.X < 0.0f) ? -Style.ShadowOffset.X * AllottedGeometry.Scale : 0.0f, 
+		(Style.ShadowOffset.X < 0.0f) ? -Style.ShadowOffset.X * AllottedGeometry.Scale : 0.0f,
 		(Style.ShadowOffset.Y < 0.0f) ? -Style.ShadowOffset.Y * AllottedGeometry.Scale : 0.0f
-		);
+	);
 
 	const FString PasswordString = BuildPasswordString(BlockRange.Len());
 
@@ -70,28 +70,28 @@ int32 FSlatePasswordRun::OnPaint(const FPaintArgs& Args, const FTextLayout::FLin
 		FSlateDrawElement::MakeText(
 			OutDrawElements,
 			++LayerId,
-			AllottedGeometry.ToPaintGeometry(TransformVector(InverseScale, Block->GetSize()), FSlateLayoutTransform(TransformPoint(InverseScale, Block->GetLocationOffset() + DrawShadowOffset))),
+			AllottedGeometry.ToPaintGeometry(TransformVector(InverseScale, TextArgs.Block->GetSize()), FSlateLayoutTransform(TransformPoint(InverseScale, TextArgs.Block->GetLocationOffset() + DrawShadowOffset))),
 			PasswordString,
 			0,
 			PasswordString.Len(),
-			DefaultStyle.Font,
+			TextArgs.DefaultStyle.Font,
 			DrawEffects,
 			InWidgetStyle.GetColorAndOpacityTint() * Style.ShadowColorAndOpacity
-			);
+		);
 	}
 
 	// Draw the text itself
 	FSlateDrawElement::MakeText(
 		OutDrawElements,
 		++LayerId,
-		AllottedGeometry.ToPaintGeometry(TransformVector(InverseScale, Block->GetSize()), FSlateLayoutTransform(TransformPoint(InverseScale, Block->GetLocationOffset() + DrawTextOffset))),
+		AllottedGeometry.ToPaintGeometry(TransformVector(InverseScale, TextArgs.Block->GetSize()), FSlateLayoutTransform(TransformPoint(InverseScale, TextArgs.Block->GetLocationOffset() + DrawTextOffset))),
 		PasswordString,
 		0,
 		PasswordString.Len(),
-		DefaultStyle.Font,
+		TextArgs.DefaultStyle.Font,
 		DrawEffects,
 		InWidgetStyle.GetColorAndOpacityTint() * Style.ColorAndOpacity.GetColor(InWidgetStyle)
-		);
+	);
 
 	return LayerId;
 }
diff --git a/Engine/Source/Runtime/Slate/Private/Framework/Text/SlateTextHighlightRunRenderer.cpp b/Engine/Source/Runtime/Slate/Private/Framework/Text/SlateTextHighlightRunRenderer.cpp
--- a/Engine/Source/Runtime/Slate/Private/Framework/Text/SlateTextHighlightRunRenderer.cpp
+++ b/Engine/Source/Runtime/Slate/Private/Framework/Text/SlateTextHighlightRunRenderer.cpp
@@ -33,7 +33,10 @@ int32 FSlateTextHighlightRunRenderer::OnPaint( const FPaintArgs& Args, const FTe
 	FWidgetStyle WidgetStyle( InWidgetStyle );
 	WidgetStyle.SetForegroundColor( InvertedHighlightColor );
 
-	return Run->OnPaint( Args, Line, Block, DefaultStyle, AllottedGeometry, MyCullingRect, OutDrawElements, LayerId, WidgetStyle, bParentEnabled );
+	// When highlighting text we want the actual text so do not replace any of it with ellipsis
+	const ETextOverflowPolicy OverflowPolicy = ETextOverflowPolicy::Clip;
+
+	return Run->OnPaint( Args, FTextArgs(Line, Block, DefaultStyle, OverflowPolicy, ETextOverflowDirection::NoOverflow), AllottedGeometry, MyCullingRect, OutDrawElements, LayerId, WidgetStyle, bParentEnabled );
 }
 
 TSharedRef< FSlateTextHighlightRunRenderer > FSlateTextHighlightRunRenderer::Create()
diff --git a/Engine/Source/Runtime/Slate/Private/Framework/Text/SlateTextLayout.cpp b/Engine/Source/Runtime/Slate/Private/Framework/Text/SlateTextLayout.cpp
--- a/Engine/Source/Runtime/Slate/Private/Framework/Text/SlateTextLayout.cpp
+++ b/Engine/Source/Runtime/Slate/Private/Framework/Text/SlateTextLayout.cpp
@@ -67,14 +67,39 @@ int32 FSlateTextLayout::OnPaint( const FPaintArgs& Args, const FGeometry& Allott
 
 	int32 HighestLayerId = LayerId;
 
-	for (const FTextLayout::FLineView& LineView : LineViews)
+	auto IsLineViewVisible = [&AllottedGeometry, &InverseScale, &MyCullingRect](const FTextLayout::FLineView& LineView, bool bTestForVerticallyClippedLine = false)
 	{
 		// Is this line visible?  This checks if the culling rect, which represents the AABB around the last clipping rect, intersects the 
 		// line of text, this requires that we get the text line into render space.
 		// TODO perhaps save off this line view rect during text layout?
 		const FVector2D LocalLineOffset = LineView.Offset * InverseScale;
-		const FSlateRect LineViewRect(AllottedGeometry.GetRenderBoundingRect(FSlateRect(LocalLineOffset, LocalLineOffset + (LineView.Size * InverseScale))));
-		if ( !FSlateRect::DoRectanglesIntersect(LineViewRect, MyCullingRect))
+		FSlateRect LineViewRect(AllottedGeometry.GetRenderBoundingRect(FSlateRect(LocalLineOffset, LocalLineOffset + (LineView.Size * InverseScale))));
+
+		if (!bTestForVerticallyClippedLine)
+		{
+			return FSlateRect::DoRectanglesIntersect(LineViewRect, MyCullingRect);
+		}
+		else
+		{
+			// We only care about vertically clipped lines. Horizontal clipping will be replaced with an ellipsis
+			// So make the left and right side the same as the clip rect to avoid the test failing due to being larger horizontally
+
+			FSlateRect GeometryRect = AllottedGeometry.GetRenderBoundingRect();
+			LineViewRect.Left = GeometryRect.Left;
+			LineViewRect.Right = GeometryRect.Right;
+
+			return FSlateRect::IsRectangleContained(GeometryRect, LineViewRect);
+		}
+
+	};
+
+	for (int32 LineIndex = 0; LineIndex < LineViews.Num(); ++LineIndex)
+	{
+		const FTextLayout::FLineView& LineView = LineViews[LineIndex];
+
+		ETextOverflowPolicy OverflowPolicy = TextOverflowPolicyOverride.Get(DefaultTextStyle.OverflowPolicy);
+
+		if (!IsLineViewVisible(LineView))
 		{
 			continue;
 		}
@@ -86,6 +111,17 @@ int32 FSlateTextLayout::OnPaint( const FPaintArgs& Args, const FGeometry& Allott
 		const int32 TextLayer = BlockDebugLayer + 1;
 		int32 HighestBlockLayerId = TextLayer;
 
+		const ETextJustify::Type VisualJustification = CalculateLineViewVisualJustification(LineView);
+		ETextOverflowDirection OverflowDirection = VisualJustification == ETextJustify::Left ? ETextOverflowDirection::LeftToRight : (VisualJustification == ETextJustify::Right ? ETextOverflowDirection::RightToLeft : ETextOverflowDirection::NoOverflow);
+
+		bool bForceEllipsisDueToClippedLine = false;
+	
+		if (OverflowPolicy == ETextOverflowPolicy::Ellipsis && LineViews.Num() > 1)
+		{
+			// Force the ellipsis to be on when the next line in a multi line text layout is clipped. This forces an ellipsis on this line even when its not clipped to indicate that an entire line or more is clipped
+			bForceEllipsisDueToClippedLine = LineViews.IsValidIndex(LineIndex + 1) ? !IsLineViewVisible(LineViews[LineIndex + 1], true) : false;
+		}
+
 		// Render every block for this line
 		for (const TSharedRef< ILayoutBlock >& Block : LineView.Blocks)
 		{
@@ -105,6 +141,9 @@ int32 FSlateTextLayout::OnPaint( const FPaintArgs& Args, const FGeometry& Allott
 			}
 #endif
 
+			FTextArgs TextArgs(LineView, Block, DefaultTextStyle, OverflowPolicy, OverflowDirection);
+			TextArgs.bForceEllipsisDueToClippedLine = bForceEllipsisDueToClippedLine;
+
 			const TSharedRef< ISlateRun > Run = StaticCastSharedRef< ISlateRun >( Block->GetRun() );
 
 			int32 HighestRunLayerId = TextLayer;
@@ -115,7 +154,7 @@ int32 FSlateTextLayout::OnPaint( const FPaintArgs& Args, const FGeometry& Allott
 			}
 			else
 			{
-				HighestRunLayerId = Run->OnPaint( Args, LineView, Block, DefaultTextStyle, AllottedGeometry, MyCullingRect, OutDrawElements, TextLayer, InWidgetStyle, bParentEnabled );
+				HighestRunLayerId = Run->OnPaint( Args, TextArgs, AllottedGeometry, MyCullingRect, OutDrawElements, TextLayer, InWidgetStyle, bParentEnabled );
 			}
 
 			HighestBlockLayerId = FMath::Max( HighestBlockLayerId, HighestRunLayerId );
diff --git a/Engine/Source/Runtime/Slate/Private/Framework/Text/SlateTextRun.cpp b/Engine/Source/Runtime/Slate/Private/Framework/Text/SlateTextRun.cpp
--- a/Engine/Source/Runtime/Slate/Private/Framework/Text/SlateTextRun.cpp
+++ b/Engine/Source/Runtime/Slate/Private/Framework/Text/SlateTextRun.cpp
@@ -76,43 +76,56 @@ TSharedRef< ILayoutBlock > FSlateTextRun::CreateBlock( int32 BeginIndex, int32 E
 	return FDefaultLayoutBlock::Create( SharedThis( this ), FTextRange( BeginIndex, EndIndex ), Size, TextContext, Renderer );
 }
 
-int32 FSlateTextRun::OnPaint( const FPaintArgs& Args, const FTextLayout::FLineView& Line, const TSharedRef< ILayoutBlock >& Block, const FTextBlockStyle& DefaultStyle, const FGeometry& AllottedGeometry, const FSlateRect& MyCullingRect, FSlateWindowElementList& OutDrawElements, int32 LayerId, const FWidgetStyle& InWidgetStyle, bool bParentEnabled ) const 
+int32 FSlateTextRun::OnPaint(const FPaintArgs& PaintArgs, const FTextArgs& TextArgs, const FGeometry& AllottedGeometry, const FSlateRect& MyCullingRect, FSlateWindowElementList& OutDrawElements, int32 LayerId, const FWidgetStyle& InWidgetStyle, bool bParentEnabled) const
 {
 	const ESlateDrawEffect DrawEffects = bParentEnabled ? ESlateDrawEffect::None : ESlateDrawEffect::DisabledEffect;
 
+	const TSharedRef<ILayoutBlock>& Block = TextArgs.Block;
+	const FTextBlockStyle& DefaultStyle = TextArgs.DefaultStyle;
+	const FTextLayout::FLineView& Line = TextArgs.Line;
+
 	const bool ShouldDropShadow = Style.ShadowColorAndOpacity.A > 0.f && Style.ShadowOffset.SizeSquared() > 0.f;
 	const FVector2D BlockLocationOffset = Block->GetLocationOffset();
 	const FTextRange BlockRange = Block->GetTextRange();
 	const FLayoutBlockTextContext BlockTextContext = Block->GetTextContext();
-	
+
 	// The block size and offset values are pre-scaled, so we need to account for that when converting the block offsets into paint geometry
 	const float InverseScale = Inverse(AllottedGeometry.Scale);
 
 	// A negative shadow offset should be applied as a positive offset to the text to avoid clipping issues
 	const FVector2D DrawShadowOffset(
-		(Style.ShadowOffset.X > 0.0f) ? Style.ShadowOffset.X * AllottedGeometry.Scale : 0.0f, 
+		(Style.ShadowOffset.X > 0.0f) ? Style.ShadowOffset.X * AllottedGeometry.Scale : 0.0f,
 		(Style.ShadowOffset.Y > 0.0f) ? Style.ShadowOffset.Y * AllottedGeometry.Scale : 0.0f
-		);
+	);
 	const FVector2D DrawTextOffset(
-		(Style.ShadowOffset.X < 0.0f) ? -Style.ShadowOffset.X * AllottedGeometry.Scale : 0.0f, 
+		(Style.ShadowOffset.X < 0.0f) ? -Style.ShadowOffset.X * AllottedGeometry.Scale : 0.0f,
 		(Style.ShadowOffset.Y < 0.0f) ? -Style.ShadowOffset.Y * AllottedGeometry.Scale : 0.0f
-		);
+	);
 
 	// Make sure we have up-to-date shaped text to work with
 	// We use the full line view range (rather than the run range) so that text that spans runs will still be shaped correctly
 	FShapedGlyphSequenceRef ShapedText = ShapedTextCacheUtil::GetShapedTextSubSequence(
-		BlockTextContext.ShapedTextCache, 
-		FCachedShapedTextKey(Line.Range, AllottedGeometry.GetAccumulatedLayoutTransform().GetScale(), BlockTextContext, Style.Font), 
-		BlockRange, 
-		**Text,  
+		BlockTextContext.ShapedTextCache,
+		FCachedShapedTextKey(Line.Range, AllottedGeometry.GetAccumulatedLayoutTransform().GetScale(), BlockTextContext, Style.Font),
+		BlockRange,
+		**Text,
 		BlockTextContext.TextDirection
-		);
+	);
+
+	FTextOverflowArgs OverflowArgs;
+	if (TextArgs.OverflowPolicy == ETextOverflowPolicy::Ellipsis && TextArgs.OverflowDirection != ETextOverflowDirection::NoOverflow)
+	{
+		TSharedRef<FSlateFontCache> FontCache = FSlateApplication::Get().GetRenderer()->GetFontCache();
+		OverflowArgs.OverflowTextPtr = FontCache->GetOverflowEllipsisText(Style.Font, AllottedGeometry.GetAccumulatedLayoutTransform().GetScale());
+		OverflowArgs.OverflowDirection = TextArgs.OverflowDirection;
+		OverflowArgs.bForceEllipsisDueToClippedLine = TextArgs.bForceEllipsisDueToClippedLine;
+	}
 
 	// Draw the optional shadow
 	if (ShouldDropShadow)
 	{
 		FShapedGlyphSequenceRef ShadowShapedText = ShapedText;
-		if(Style.ShadowColorAndOpacity != Style.Font.OutlineSettings.OutlineColor)
+		if (Style.ShadowColorAndOpacity != Style.Font.OutlineSettings.OutlineColor)
 		{
 			// Copy font info for shadow to replace the outline color
 			FSlateFontInfo ShadowFontInfo = Style.Font;
@@ -140,8 +153,9 @@ int32 FSlateTextRun::OnPaint( const FPaintArgs& Args, const FTextLayout::FLineVi
 			ShadowShapedText,
 			DrawEffects,
 			InWidgetStyle.GetColorAndOpacityTint() * Style.ShadowColorAndOpacity,
-			InWidgetStyle.GetColorAndOpacityTint() * Style.Font.OutlineSettings.OutlineColor
-			);
+			InWidgetStyle.GetColorAndOpacityTint() * Style.Font.OutlineSettings.OutlineColor,
+			OverflowArgs
+		);
 	}
 
 	// Draw the text itself
@@ -152,8 +166,9 @@ int32 FSlateTextRun::OnPaint( const FPaintArgs& Args, const FTextLayout::FLineVi
 		ShapedText,
 		DrawEffects,
 		InWidgetStyle.GetColorAndOpacityTint() * Style.ColorAndOpacity.GetColor(InWidgetStyle),
-		InWidgetStyle.GetColorAndOpacityTint() * Style.Font.OutlineSettings.OutlineColor
-		);
+		InWidgetStyle.GetColorAndOpacityTint() * Style.Font.OutlineSettings.OutlineColor,
+		OverflowArgs
+	);
 
 	return LayerId;
 }
diff --git a/Engine/Source/Runtime/Slate/Private/Framework/Text/SlateWidgetRun.cpp b/Engine/Source/Runtime/Slate/Private/Framework/Text/SlateWidgetRun.cpp
--- a/Engine/Source/Runtime/Slate/Private/Framework/Text/SlateWidgetRun.cpp
+++ b/Engine/Source/Runtime/Slate/Private/Framework/Text/SlateWidgetRun.cpp
@@ -58,7 +58,7 @@ TSharedRef< ILayoutBlock > FSlateWidgetRun::CreateBlock( int32 StartIndex, int32
 	return FDefaultLayoutBlock::Create( SharedThis( this ), FTextRange( StartIndex, EndIndex ), Size, TextContext, Renderer );
 }
 
-int32 FSlateWidgetRun::OnPaint( const FPaintArgs& Args, const FTextLayout::FLineView& Line, const TSharedRef< ILayoutBlock >& Block, const FTextBlockStyle& DefaultStyle, const FGeometry& AllottedGeometry, const FSlateRect& MyCullingRect, FSlateWindowElementList& OutDrawElements, int32 LayerId, const FWidgetStyle& InWidgetStyle, bool bParentEnabled ) const 
+int32 FSlateWidgetRun::OnPaint(const FPaintArgs& PaintArgs, const FTextArgs& TextArgs, const FGeometry& AllottedGeometry, const FSlateRect& MyCullingRect, FSlateWindowElementList& OutDrawElements, int32 LayerId, const FWidgetStyle& InWidgetStyle, bool bParentEnabled) const
 {
 	// The block size and offset values are pre-scaled, so we need to account for that when converting the block offsets into paint geometry
 	const float InverseScale = Inverse(AllottedGeometry.Scale);
@@ -75,8 +75,8 @@ int32 FSlateWidgetRun::OnPaint( const FPaintArgs& Args, const FTextLayout::FLine
 		}
 	}
 
-	const FGeometry WidgetGeometry = AllottedGeometry.MakeChild(TransformVector(InverseScale, Block->GetSize()), FSlateLayoutTransform(TransformPoint(InverseScale, Block->GetLocationOffset())));
-	return Info.Widget->Paint( Args, WidgetGeometry, MyCullingRect, OutDrawElements, LayerId, InWidgetStyle, bParentEnabled );
+	const FGeometry WidgetGeometry = AllottedGeometry.MakeChild(TransformVector(InverseScale, TextArgs.Block->GetSize()), FSlateLayoutTransform(TransformPoint(InverseScale, TextArgs.Block->GetLocationOffset())));
+	return Info.Widget->Paint(PaintArgs, WidgetGeometry, MyCullingRect, OutDrawElements, LayerId, InWidgetStyle, bParentEnabled);
 }
 
 const TArray< TSharedRef<SWidget> >& FSlateWidgetRun::GetChildren()
diff --git a/Engine/Source/Runtime/Slate/Private/Framework/Text/TextLayout.cpp b/Engine/Source/Runtime/Slate/Private/Framework/Text/TextLayout.cpp
--- a/Engine/Source/Runtime/Slate/Private/Framework/Text/TextLayout.cpp
+++ b/Engine/Source/Runtime/Slate/Private/Framework/Text/TextLayout.cpp
@@ -1160,6 +1160,7 @@ FTextLayout::FTextLayout()
 	, GraphemeBreakIterator(FBreakIterator::CreateCharacterBoundaryIterator())
 	, WordBreakIterator(FBreakIterator::CreateWordBreakIterator())
 	, TextBiDiDetection(TextBiDi::CreateTextBiDi())
+	, TextOverflowPolicyOverride()
 {
 }
 
@@ -2560,6 +2561,17 @@ void FTextLayout::SetTextFlowDirection( const ETextFlowDirection InTextFlowDirec
 	DirtyAllLineModels(ELineModelDirtyState::WrappingInformation | ELineModelDirtyState::TextBaseDirection | ELineModelDirtyState::ShapingCache);
 }
 
+void FTextLayout::SetTextOverflowPolicy(const TOptional<ETextOverflowPolicy> InTextOverflowPolicy)
+{
+	if (TextOverflowPolicyOverride == InTextOverflowPolicy)
+	{
+		return;
+	}
+
+	TextOverflowPolicyOverride = InTextOverflowPolicy;
+	DirtyFlags |= ETextLayoutDirtyState::Layout;
+}
+
 void FTextLayout::SetJustification( ETextJustify::Type Value )
 {
 	if ( Justification == Value )
diff --git a/Engine/Source/Runtime/Slate/Private/Widgets/Input/SEditableText.cpp b/Engine/Source/Runtime/Slate/Private/Widgets/Input/SEditableText.cpp
--- a/Engine/Source/Runtime/Slate/Private/Widgets/Input/SEditableText.cpp
+++ b/Engine/Source/Runtime/Slate/Private/Widgets/Input/SEditableText.cpp
@@ -77,6 +77,7 @@ void SEditableText::Construct( const FArguments& InArgs )
 	EditableTextLayout->SetCompositionBrush(InArgs._BackgroundImageComposing.IsSet() ? InArgs._BackgroundImageComposing : &InArgs._Style->BackgroundImageComposing);
 	EditableTextLayout->SetDebugSourceInfo(TAttribute<FString>::Create(TAttribute<FString>::FGetter::CreateLambda([this]{ return FReflectionMetaData::GetWidgetDebugInfo(this); })));
 	EditableTextLayout->SetJustification(InArgs._Justification);
+	EditableTextLayout->SetOverflowPolicy(InArgs._OverflowPolicy);
 
 	// build context menu extender
 	MenuExtender = MakeShareable(new FExtender());
@@ -379,6 +380,11 @@ void SEditableText::SetTextFlowDirection(const TOptional<ETextFlowDirection>& In
 	EditableTextLayout->SetTextFlowDirection(InTextFlowDirection);
 }
 
+void SEditableText::SetOverflowPolicy(TOptional<ETextOverflowPolicy> InOverflowPolicy)
+{
+	EditableTextLayout->SetOverflowPolicy(InOverflowPolicy);
+}
+
 bool SEditableText::AnyTextSelected() const
 {
 	return EditableTextLayout->AnyTextSelected();
diff --git a/Engine/Source/Runtime/Slate/Private/Widgets/Input/SEditableTextBox.cpp b/Engine/Source/Runtime/Slate/Private/Widgets/Input/SEditableTextBox.cpp
--- a/Engine/Source/Runtime/Slate/Private/Widgets/Input/SEditableTextBox.cpp
+++ b/Engine/Source/Runtime/Slate/Private/Widgets/Input/SEditableTextBox.cpp
@@ -82,6 +82,7 @@ void SEditableTextBox::Construct( const FArguments& InArgs )
 					.VirtualKeyboardDismissAction( InArgs._VirtualKeyboardDismissAction )
 					.TextShapingMethod(InArgs._TextShapingMethod)
 					.TextFlowDirection( InArgs._TextFlowDirection )
+					.OverflowPolicy(InArgs._OverflowPolicy)
 				]
 			]
 		]
@@ -175,6 +176,11 @@ void SEditableTextBox::SetTextFlowDirection(const TOptional<ETextFlowDirection>&
 }
 
 
+void SEditableTextBox::SetOverflowPolicy(TOptional<ETextOverflowPolicy> InOverflowPolicy)
+{
+	EditableText->SetOverflowPolicy(InOverflowPolicy);
+}
+
 bool SEditableTextBox::AnyTextSelected() const
 {
 	return EditableText->AnyTextSelected();
diff --git a/Engine/Source/Runtime/Slate/Private/Widgets/Input/SMultiLineEditableTextBox.cpp b/Engine/Source/Runtime/Slate/Private/Widgets/Input/SMultiLineEditableTextBox.cpp
--- a/Engine/Source/Runtime/Slate/Private/Widgets/Input/SMultiLineEditableTextBox.cpp
+++ b/Engine/Source/Runtime/Slate/Private/Widgets/Input/SMultiLineEditableTextBox.cpp
@@ -133,6 +133,7 @@ void SMultiLineEditableTextBox::Construct( const FArguments& InArgs )
 					.TextShapingMethod(InArgs._TextShapingMethod)
 					.TextFlowDirection(InArgs._TextFlowDirection)
 					.AllowContextMenu(InArgs._AllowContextMenu)
+					.OverflowPolicy(InArgs._OverflowPolicy)
 				]
 
 				+SVerticalBox::Slot()
@@ -356,6 +357,11 @@ void SMultiLineEditableTextBox::SetJustification(const TAttribute<ETextJustify::
 	EditableText->SetJustification(InJustification);
 }
 
+void SMultiLineEditableTextBox::SetOverflowPolicy(TOptional<ETextOverflowPolicy> InOverflowPolicy)
+{
+	EditableText->SetOverflowPolicy(InOverflowPolicy);
+}
+
 void SMultiLineEditableTextBox::SetAllowContextMenu(const TAttribute< bool >& InAllowContextMenu)
 {
 	EditableText->SetAllowContextMenu(InAllowContextMenu);
diff --git a/Engine/Source/Runtime/Slate/Private/Widgets/Text/SInlineEditableTextBlock.cpp b/Engine/Source/Runtime/Slate/Private/Widgets/Text/SInlineEditableTextBlock.cpp
--- a/Engine/Source/Runtime/Slate/Private/Widgets/Text/SInlineEditableTextBlock.cpp
+++ b/Engine/Source/Runtime/Slate/Private/Widgets/Text/SInlineEditableTextBlock.cpp
@@ -43,6 +43,7 @@ void SInlineEditableTextBlock::Construct( const FArguments& InArgs )
 			.WrapTextAt( InArgs._WrapTextAt )
 			.Justification( InArgs._Justification )
 			.LineBreakPolicy( InArgs._LineBreakPolicy )
+			.OverflowPolicy(InArgs._OverflowPolicy)
 		]
 	];
 
@@ -53,7 +54,7 @@ void SInlineEditableTextBlock::Construct( const FArguments& InArgs )
 			.Text(InArgs._Text)
 			.Style(&InArgs._Style->EditableTextBoxStyle)
 			.Font(InArgs._Font)
-			.ToolTipText( InArgs._ToolTipText )
+			.ToolTipText(InArgs._ToolTipText)
 			.OnTextChanged(this, &SInlineEditableTextBlock::OnTextChanged)
 			.OnTextCommitted(this, &SInlineEditableTextBlock::OnTextBoxCommitted)
 			.WrapTextAt(InArgs._WrapTextAt)
@@ -61,7 +62,8 @@ void SInlineEditableTextBlock::Construct( const FArguments& InArgs )
 			.SelectAllTextWhenFocused(true)
 			.ClearKeyboardFocusOnCommit(true)
 			.RevertTextOnEscape(true)
-			.ModiferKeyForNewLine(InArgs._ModiferKeyForNewLine);
+			.ModiferKeyForNewLine(InArgs._ModiferKeyForNewLine)
+			.OverflowPolicy(InArgs._OverflowPolicy);
 	}
 	else
 #endif //WITH_FANCY_TEXT
@@ -70,11 +72,12 @@ void SInlineEditableTextBlock::Construct( const FArguments& InArgs )
 			.Text(InArgs._Text)
 			.Style(&InArgs._Style->EditableTextBoxStyle)
 			.Font(InArgs._Font)
-			.ToolTipText( InArgs._ToolTipText )
-			.OnTextChanged( this, &SInlineEditableTextBlock::OnTextChanged )
+			.ToolTipText(InArgs._ToolTipText)
+			.OnTextChanged(this, &SInlineEditableTextBlock::OnTextChanged)
 			.OnTextCommitted(this, &SInlineEditableTextBlock::OnTextBoxCommitted)
 			.SelectAllTextWhenFocused(true)
-			.ClearKeyboardFocusOnCommit(false);
+			.ClearKeyboardFocusOnCommit(false)
+			.OverflowPolicy(InArgs._OverflowPolicy);
 	}
 }
 
@@ -178,6 +181,11 @@ void SInlineEditableTextBlock::SetWrapTextAt( const TAttribute<float>& InWrapTex
 	TextBlock->SetWrapTextAt( InWrapTextAt );
 }
 
+void SInlineEditableTextBlock::SetOverflowPolicy(TOptional<ETextOverflowPolicy> InOverflowPolicy)
+{
+	TextBlock->SetOverflowPolicy(InOverflowPolicy);
+}
+
 FReply SInlineEditableTextBlock::OnMouseButtonDown( const FGeometry& MyGeometry, const FPointerEvent& MouseEvent )
 {
 	if( !MouseEvent.IsMouseButtonDown(EKeys::LeftMouseButton) || MouseEvent.IsControlDown() || MouseEvent.IsShiftDown())
diff --git a/Engine/Source/Runtime/Slate/Private/Widgets/Text/SMultiLineEditableText.cpp b/Engine/Source/Runtime/Slate/Private/Widgets/Text/SMultiLineEditableText.cpp
--- a/Engine/Source/Runtime/Slate/Private/Widgets/Text/SMultiLineEditableText.cpp
+++ b/Engine/Source/Runtime/Slate/Private/Widgets/Text/SMultiLineEditableText.cpp
@@ -89,6 +89,7 @@ void SMultiLineEditableText::Construct( const FArguments& InArgs )
 	EditableTextLayout->SetJustification(InArgs._Justification);
 	EditableTextLayout->SetLineHeightPercentage(InArgs._LineHeightPercentage);
 	EditableTextLayout->SetDebugSourceInfo(TAttribute<FString>::Create(TAttribute<FString>::FGetter::CreateLambda([this]{ return FReflectionMetaData::GetWidgetDebugInfo(this); })));
+	EditableTextLayout->SetOverflowPolicy(InArgs._OverflowPolicy);
 
 	// build context menu extender
 	MenuExtender = MakeShareable(new FExtender);
@@ -201,6 +202,11 @@ void SMultiLineEditableText::SetJustification(const TAttribute<ETextJustify::Typ
 	EditableTextLayout->SetJustification(InJustification);
 }
 
+void SMultiLineEditableText::SetOverflowPolicy(TOptional<ETextOverflowPolicy> InOverflowPolicy)
+{
+	EditableTextLayout->SetOverflowPolicy(InOverflowPolicy);
+}
+
 void SMultiLineEditableText::SetAllowContextMenu(const TAttribute< bool >& InAllowContextMenu)
 {
 	bAllowContextMenu = InAllowContextMenu;
diff --git a/Engine/Source/Runtime/Slate/Private/Widgets/Text/SRichTextBlock.cpp b/Engine/Source/Runtime/Slate/Private/Widgets/Text/SRichTextBlock.cpp
--- a/Engine/Source/Runtime/Slate/Private/Widgets/Text/SRichTextBlock.cpp
+++ b/Engine/Source/Runtime/Slate/Private/Widgets/Text/SRichTextBlock.cpp
@@ -183,6 +183,12 @@ void SRichTextBlock::SetDecoratorStyleSet(const ISlateStyle* NewDecoratorStyleSe
 	}
 }
 
+void SRichTextBlock::SetOverflowPolicy(TOptional<ETextOverflowPolicy> InOverflowPolicy)
+{
+	TextLayoutCache->SetTextOverflowPolicy(InOverflowPolicy);
+	Invalidate(EInvalidateWidget::Layout);
+}
+
 void SRichTextBlock::SetTextBlockScale(const float NewTextBlockScale)
 {
 	TextBlockScale = NewTextBlockScale;
diff --git a/Engine/Source/Runtime/Slate/Private/Widgets/Text/STextBlock.cpp b/Engine/Source/Runtime/Slate/Private/Widgets/Text/STextBlock.cpp
--- a/Engine/Source/Runtime/Slate/Private/Widgets/Text/STextBlock.cpp
+++ b/Engine/Source/Runtime/Slate/Private/Widgets/Text/STextBlock.cpp
@@ -121,6 +121,7 @@ void STextBlock::Construct( const FArguments& InArgs )
 		// We use a dummy style here (as it may not be safe to call the delegates used to compute the style), but the correct style is set by ComputeDesiredSize
 		TextLayoutCache = MakeUnique<FSlateTextBlockLayout>(this, FTextBlockStyle::GetDefault(), InArgs._TextShapingMethod, InArgs._TextFlowDirection, FCreateSlateTextLayout(), FPlainTextLayoutMarshaller::Create(), InArgs._LineBreakPolicy);
 		TextLayoutCache->SetDebugSourceInfo(TAttribute<FString>::Create(TAttribute<FString>::FGetter::CreateLambda([this] { return FReflectionMetaData::GetWidgetDebugInfo(this); })));
+		TextLayoutCache->SetTextOverflowPolicy(InArgs._OverflowPolicy.IsSet() ? InArgs._OverflowPolicy : TextStyle.OverflowPolicy);
 	}
 }
 
@@ -372,6 +373,12 @@ void STextBlock::SetTransformPolicy(TAttribute<ETextTransformPolicy> InTransform
 	TransformPolicy.Assign(*this, MoveTemp(InTransformPolicy));
 }
 
+void STextBlock::SetOverflowPolicy(TOptional<ETextOverflowPolicy> InOverflowPolicy)
+{
+	TextLayoutCache->SetTextOverflowPolicy(InOverflowPolicy);
+	InvalidateText(EInvalidateWidgetReason::Layout);
+}
+
 void STextBlock::SetShadowOffset(TAttribute<FVector2D> InShadowOffset)
 {
 	bIsAttributeShadowOffsetSet = InShadowOffset.IsSet();
diff --git a/Engine/Source/Runtime/Slate/Private/Widgets/Text/SlateEditableTextLayout.cpp b/Engine/Source/Runtime/Slate/Private/Widgets/Text/SlateEditableTextLayout.cpp
--- a/Engine/Source/Runtime/Slate/Private/Widgets/Text/SlateEditableTextLayout.cpp
+++ b/Engine/Source/Runtime/Slate/Private/Widgets/Text/SlateEditableTextLayout.cpp
@@ -493,6 +493,21 @@ void FSlateEditableTextLayout::SetLineHeightPercentage(const TAttribute<float>&
 	OwnerWidget->GetSlateWidget()->Invalidate(EInvalidateWidget::LayoutAndVolatility);
 }
 
+void FSlateEditableTextLayout::SetOverflowPolicy(TOptional<ETextOverflowPolicy> InOverflowPolicy)
+{
+	if(OverflowPolicyOverride != InOverflowPolicy)
+	{
+		OverflowPolicyOverride = InOverflowPolicy;
+		TextLayout->SetTextOverflowPolicy(OverflowPolicyOverride);
+		if (HintTextLayout.IsValid())
+		{
+			HintTextLayout->SetTextOverflowPolicy(OverflowPolicyOverride);
+		}
+
+		OwnerWidget->GetSlateWidget()->Invalidate(EInvalidateWidget::LayoutAndVolatility);
+	}
+}
+
 void FSlateEditableTextLayout::SetDebugSourceInfo(const TAttribute<FString>& InDebugSourceInfo)
 {
 	DebugSourceInfo = InDebugSourceInfo;
@@ -3239,11 +3254,15 @@ void FSlateEditableTextLayout::Tick(const FGeometry& AllottedGeometry, const dou
 		(OwnerWidget->GetSlateWidget()->HasAnyUserFocus().IsSet() || HasActiveContextMenu());
 	if (bShouldAppearFocused)
 	{
+		// When focused the user is editing or selecting text. Never allow ellipsis to replace text
+		TextLayout->SetTextOverflowPolicy(ETextOverflowPolicy::Clip);
 		// If we have focus then we don't allow the editable text itself to update, but we do still need to refresh the password and marshaller state
 		RefreshImpl(nullptr);
 	}
 	else
 	{
+		TextLayout->SetTextOverflowPolicy(OverflowPolicyOverride);
+
 		// We don't have focus, so we can perform a full refresh
 		Refresh();
 	}
diff --git a/Engine/Source/Runtime/Slate/Private/Widgets/Text/SlateTextBlockLayout.cpp b/Engine/Source/Runtime/Slate/Private/Widgets/Text/SlateTextBlockLayout.cpp
--- a/Engine/Source/Runtime/Slate/Private/Widgets/Text/SlateTextBlockLayout.cpp
+++ b/Engine/Source/Runtime/Slate/Private/Widgets/Text/SlateTextBlockLayout.cpp
@@ -257,6 +257,11 @@ void FSlateTextBlockLayout::SetTextFlowDirection(const TOptional<ETextFlowDirect
 	TextLayout->SetTextFlowDirection((InTextFlowDirection.IsSet()) ? InTextFlowDirection.GetValue() : GetDefaultTextFlowDirection());
 }
 
+void FSlateTextBlockLayout::SetTextOverflowPolicy(const TOptional<ETextOverflowPolicy> InTextOverflowPolicy)
+{
+	TextLayout->SetTextOverflowPolicy(InTextOverflowPolicy);
+}
+
 void FSlateTextBlockLayout::SetDebugSourceInfo(const TAttribute<FString>& InDebugSourceInfo)
 {
 	TextLayout->SetDebugSourceInfo(InDebugSourceInfo);
diff --git a/Engine/Source/Runtime/Slate/Public/Framework/Text/ISlateRun.h b/Engine/Source/Runtime/Slate/Public/Framework/Text/ISlateRun.h
--- a/Engine/Source/Runtime/Slate/Public/Framework/Text/ISlateRun.h
+++ b/Engine/Source/Runtime/Slate/Public/Framework/Text/ISlateRun.h
@@ -12,13 +12,32 @@ class FPaintArgs;
 class FSlateWindowElementList;
 struct FTextBlockStyle;
 
+struct FTextArgs
+{
+	FTextArgs(const FTextLayout::FLineView& InLine, const TSharedRef<ILayoutBlock>& InBlock, const FTextBlockStyle& InDefaultStyle, ETextOverflowPolicy InOverflowPolicy, ETextOverflowDirection InOverflowDirection)
+		: Line(InLine)
+		, Block(InBlock)
+		, DefaultStyle(InDefaultStyle)
+		, OverflowPolicy(InOverflowPolicy)
+		, OverflowDirection(InOverflowDirection)
+		, bForceEllipsisDueToClippedLine(false)
+	{}
+
+	const FTextLayout::FLineView& Line;
+	const TSharedRef< ILayoutBlock >& Block;
+	const FTextBlockStyle& DefaultStyle;
+	ETextOverflowPolicy OverflowPolicy;
+	ETextOverflowDirection OverflowDirection;
+	bool bForceEllipsisDueToClippedLine;
+
+};
+
 class SLATE_API ISlateRun : public IRun
 {
 public:
-
-	virtual int32 OnPaint( const FPaintArgs& Args, const FTextLayout::FLineView& Line, const TSharedRef< ILayoutBlock >& Block, const FTextBlockStyle& DefaultStyle, const FGeometry& AllottedGeometry, const FSlateRect& MyCullingRect, FSlateWindowElementList& OutDrawElements, int32 LayerId, const FWidgetStyle& InWidgetStyle, bool bParentEnabled ) const = 0;
+	virtual int32 OnPaint(const FPaintArgs& PaintArgs, const FTextArgs& TextArgs, const FGeometry& AllottedGeometry, const FSlateRect& MyCullingRect, FSlateWindowElementList& OutDrawElements, int32 LayerId, const FWidgetStyle& InWidgetStyle, bool bParentEnabled) const = 0;
 
 	virtual const TArray< TSharedRef<SWidget> >& GetChildren() = 0;
 
-	virtual void ArrangeChildren( const TSharedRef< ILayoutBlock >& Block, const FGeometry& AllottedGeometry, FArrangedChildren& ArrangedChildren ) const = 0;
+	virtual void ArrangeChildren(const TSharedRef< ILayoutBlock >& Block, const FGeometry& AllottedGeometry, FArrangedChildren& ArrangedChildren) const = 0;
 };
diff --git a/Engine/Source/Runtime/Slate/Public/Framework/Text/SlateHyperlinkRun.h b/Engine/Source/Runtime/Slate/Public/Framework/Text/SlateHyperlinkRun.h
--- a/Engine/Source/Runtime/Slate/Public/Framework/Text/SlateHyperlinkRun.h
+++ b/Engine/Source/Runtime/Slate/Public/Framework/Text/SlateHyperlinkRun.h
@@ -70,7 +70,7 @@ public:
 
 	virtual TSharedRef< ILayoutBlock > CreateBlock( int32 StartIndex, int32 EndIndex, FVector2D Size, const FLayoutBlockTextContext& TextContext, const TSharedPtr< IRunRenderer >& Renderer ) override;
 
-	virtual int32 OnPaint( const FPaintArgs& Args, const FTextLayout::FLineView& Line, const TSharedRef< ILayoutBlock >& Block, const FTextBlockStyle& DefaultStyle, const FGeometry& AllottedGeometry, const FSlateRect& MyCullingRect, FSlateWindowElementList& OutDrawElements, int32 LayerId, const FWidgetStyle& InWidgetStyle, bool bParentEnabled ) const override;
+	virtual int32 OnPaint(const FPaintArgs& PaintArgs, const FTextArgs& TextArgs, const FGeometry& AllottedGeometry, const FSlateRect& MyCullingRect, FSlateWindowElementList& OutDrawElements, int32 LayerId, const FWidgetStyle& InWidgetStyle, bool bParentEnabled) const override;
 
 	virtual const TArray< TSharedRef<SWidget> >& GetChildren() override;
 
diff --git a/Engine/Source/Runtime/Slate/Public/Framework/Text/SlateImageRun.h b/Engine/Source/Runtime/Slate/Public/Framework/Text/SlateImageRun.h
--- a/Engine/Source/Runtime/Slate/Public/Framework/Text/SlateImageRun.h
+++ b/Engine/Source/Runtime/Slate/Public/Framework/Text/SlateImageRun.h
@@ -43,7 +43,7 @@ public:
 
 	virtual TSharedRef< ILayoutBlock > CreateBlock( int32 StartIndex, int32 EndIndex, FVector2D Size, const FLayoutBlockTextContext& TextContext, const TSharedPtr< IRunRenderer >& Renderer ) override;
 
-	virtual int32 OnPaint( const FPaintArgs& Args, const FTextLayout::FLineView& Line, const TSharedRef< ILayoutBlock >& Block, const FTextBlockStyle& DefaultStyle, const FGeometry& AllottedGeometry, const FSlateRect& MyCullingRect, FSlateWindowElementList& OutDrawElements, int32 LayerId, const FWidgetStyle& InWidgetStyle, bool bParentEnabled ) const override;
+	virtual int32 OnPaint(const FPaintArgs& PaintArgs, const FTextArgs& TextArgs, const FGeometry& AllottedGeometry, const FSlateRect& MyCullingRect, FSlateWindowElementList& OutDrawElements, int32 LayerId, const FWidgetStyle& InWidgetStyle, bool bParentEnabled) const override;
 
 	virtual const TArray< TSharedRef<SWidget> >& GetChildren() override;
 
diff --git a/Engine/Source/Runtime/Slate/Public/Framework/Text/SlatePasswordRun.h b/Engine/Source/Runtime/Slate/Public/Framework/Text/SlatePasswordRun.h
--- a/Engine/Source/Runtime/Slate/Public/Framework/Text/SlatePasswordRun.h
+++ b/Engine/Source/Runtime/Slate/Public/Framework/Text/SlatePasswordRun.h
@@ -26,7 +26,7 @@ public:
 	virtual FVector2D Measure( int32 StartIndex, int32 EndIndex, float Scale, const FRunTextContext& TextContext ) const override;
 	virtual int8 GetKerning(int32 CurrentIndex, float Scale, const FRunTextContext& TextContext) const override;
 
-	virtual int32 OnPaint( const FPaintArgs& Args, const FTextLayout::FLineView& Line, const TSharedRef< ILayoutBlock >& Block, const FTextBlockStyle& DefaultStyle, const FGeometry& AllottedGeometry, const FSlateRect& MyCullingRect, FSlateWindowElementList& OutDrawElements, int32 LayerId, const FWidgetStyle& InWidgetStyle, bool bParentEnabled ) const override;
+	virtual int32 OnPaint(const FPaintArgs& PaintArgs, const FTextArgs& TextArgs, const FGeometry& AllottedGeometry, const FSlateRect& MyCullingRect, FSlateWindowElementList& OutDrawElements, int32 LayerId, const FWidgetStyle& InWidgetStyle, bool bParentEnabled) const override;
 
 	virtual int32 GetTextIndexAt( const TSharedRef< ILayoutBlock >& Block, const FVector2D& Location, float Scale, ETextHitPoint* const OutHitPoint = nullptr ) const override;
 
diff --git a/Engine/Source/Runtime/Slate/Public/Framework/Text/SlateTextRun.h b/Engine/Source/Runtime/Slate/Public/Framework/Text/SlateTextRun.h
--- a/Engine/Source/Runtime/Slate/Public/Framework/Text/SlateTextRun.h
+++ b/Engine/Source/Runtime/Slate/Public/Framework/Text/SlateTextRun.h
@@ -36,7 +36,7 @@ public:
 
 	virtual TSharedRef< ILayoutBlock > CreateBlock( int32 StartIndex, int32 EndIndex, FVector2D Size, const FLayoutBlockTextContext& TextContext, const TSharedPtr< IRunRenderer >& Renderer ) override;
 
-	virtual int32 OnPaint( const FPaintArgs& Args, const FTextLayout::FLineView& Line, const TSharedRef< ILayoutBlock >& Block, const FTextBlockStyle& DefaultStyle, const FGeometry& AllottedGeometry, const FSlateRect& MyCullingRect, FSlateWindowElementList& OutDrawElements, int32 LayerId, const FWidgetStyle& InWidgetStyle, bool bParentEnabled ) const override;
+	virtual int32 OnPaint(const FPaintArgs& PaintArgs, const FTextArgs& TextArgs, const FGeometry& AllottedGeometry, const FSlateRect& MyCullingRect, FSlateWindowElementList& OutDrawElements, int32 LayerId, const FWidgetStyle& InWidgetStyle, bool bParentEnabled) const override;
 
 	virtual const TArray< TSharedRef<SWidget> >& GetChildren() override;
 
diff --git a/Engine/Source/Runtime/Slate/Public/Framework/Text/SlateWidgetRun.h b/Engine/Source/Runtime/Slate/Public/Framework/Text/SlateWidgetRun.h
--- a/Engine/Source/Runtime/Slate/Public/Framework/Text/SlateWidgetRun.h
+++ b/Engine/Source/Runtime/Slate/Public/Framework/Text/SlateWidgetRun.h
@@ -57,7 +57,7 @@ public:
 	virtual FVector2D Measure( int32 StartIndex, int32 EndIndex, float Scale, const FRunTextContext& TextContext ) const override;
 	virtual int8 GetKerning(int32 CurrentIndex, float Scale, const FRunTextContext& TextContext) const override;
 	virtual TSharedRef< ILayoutBlock > CreateBlock( int32 StartIndex, int32 EndIndex, FVector2D Size, const FLayoutBlockTextContext& TextContext, const TSharedPtr< IRunRenderer >& Renderer ) override;
-	virtual int32 OnPaint( const FPaintArgs& Args, const FTextLayout::FLineView& Line, const TSharedRef< ILayoutBlock >& Block, const FTextBlockStyle& DefaultStyle, const FGeometry& AllottedGeometry, const FSlateRect& MyCullingRect, FSlateWindowElementList& OutDrawElements, int32 LayerId, const FWidgetStyle& InWidgetStyle, bool bParentEnabled ) const override;
+	virtual int32 OnPaint(const FPaintArgs& PaintArgs, const FTextArgs& TextArgs, const FGeometry& AllottedGeometry, const FSlateRect& MyCullingRect, FSlateWindowElementList& OutDrawElements, int32 LayerId, const FWidgetStyle& InWidgetStyle, bool bParentEnabled) const override;
 	virtual const TArray< TSharedRef<SWidget> >& GetChildren() override;
 	virtual void ArrangeChildren( const TSharedRef< ILayoutBlock >& Block, const FGeometry& AllottedGeometry, FArrangedChildren& ArrangedChildren ) const override;
 	virtual int32 GetTextIndexAt( const TSharedRef< ILayoutBlock >& Block, const FVector2D& Location, float Scale, ETextHitPoint* const OutHitPoint = nullptr ) const override;
diff --git a/Engine/Source/Runtime/Slate/Public/Framework/Text/TextLayout.h b/Engine/Source/Runtime/Slate/Public/Framework/Text/TextLayout.h
--- a/Engine/Source/Runtime/Slate/Public/Framework/Text/TextLayout.h
+++ b/Engine/Source/Runtime/Slate/Public/Framework/Text/TextLayout.h
@@ -411,6 +411,8 @@ public:
 	FORCEINLINE ETextFlowDirection GetTextFlowDirection() const { return TextFlowDirection; }
 	void SetTextFlowDirection( const ETextFlowDirection InTextFlowDirection );
 
+	void SetTextOverflowPolicy(const TOptional<ETextOverflowPolicy> InTextOverflowPolicy);
+
 	FORCEINLINE FMargin GetMargin() const { return Margin; }
 	void SetMargin( const FMargin& InMargin );
 
@@ -749,4 +751,8 @@ protected:
 
 	/** Information given to use by our an external source (typically our owner widget) to help identify who owns this text layout in the case of an error */
 	TAttribute<FString> DebugSourceInfo;
+
+	/** Override for the text overflow policy. If unset, the style is used */
+	TOptional<ETextOverflowPolicy> TextOverflowPolicyOverride;
+
 };
diff --git a/Engine/Source/Runtime/Slate/Public/Widgets/Input/SEditableText.h b/Engine/Source/Runtime/Slate/Public/Widgets/Input/SEditableText.h
--- a/Engine/Source/Runtime/Slate/Public/Widgets/Input/SEditableText.h
+++ b/Engine/Source/Runtime/Slate/Public/Widgets/Input/SEditableText.h
@@ -64,6 +64,7 @@ public:
 		, _VirtualKeyboardDismissAction(EVirtualKeyboardDismissAction::TextChangeOnDismiss)
 		, _TextShapingMethod()
 		, _TextFlowDirection()
+		, _OverflowPolicy()
 		{
 			_Clipping = EWidgetClipping::ClipToBounds;
 		}
@@ -171,7 +172,9 @@ public:
 
 		/** Which text flow direction should we use? (unset to use the default returned by GetDefaultTextFlowDirection) */
 		SLATE_ARGUMENT(TOptional<ETextFlowDirection>, TextFlowDirection)
-
+		
+		/** Determines what happens to text that is clipped and doesnt fit within the allotted area for this widget */
+		SLATE_ARGUMENT(TOptional<ETextOverflowPolicy>, OverflowPolicy)
 	SLATE_END_ARGS()
 
 	/** Constructor */
@@ -312,6 +315,9 @@ public:
 	/** See TextFlowDirection attribute */
 	void SetTextFlowDirection(const TOptional<ETextFlowDirection>& InTextFlowDirection);
 
+	/** Sets the overflow policy for this text block */
+	void SetOverflowPolicy(TOptional<ETextOverflowPolicy> InOverflowPolicy);
+
 	/** Query to see if any text is selected within the document */
 	bool AnyTextSelected() const;
 
diff --git a/Engine/Source/Runtime/Slate/Public/Widgets/Input/SEditableTextBox.h b/Engine/Source/Runtime/Slate/Public/Widgets/Input/SEditableTextBox.h
--- a/Engine/Source/Runtime/Slate/Public/Widgets/Input/SEditableTextBox.h
+++ b/Engine/Source/Runtime/Slate/Public/Widgets/Input/SEditableTextBox.h
@@ -56,6 +56,7 @@ public:
 		, _VirtualKeyboardOptions(FVirtualKeyboardOptions())
 		, _VirtualKeyboardTrigger(EVirtualKeyboardTrigger::OnFocusByPointer)
 		, _VirtualKeyboardDismissAction(EVirtualKeyboardDismissAction::TextChangeOnDismiss)
+		, _OverflowPolicy()
 		{
 		}
 
@@ -164,6 +165,8 @@ public:
 		/** Which text flow direction should we use? (unset to use the default returned by GetDefaultTextFlowDirection) */
 		SLATE_ARGUMENT(TOptional<ETextFlowDirection>, TextFlowDirection)
 
+		/** Determines what happens to text that is clipped and doesnt fit within the allotted area for this text box */
+		SLATE_ARGUMENT(TOptional<ETextOverflowPolicy>, OverflowPolicy)
 	SLATE_END_ARGS()
 
 	SEditableTextBox();
@@ -334,6 +337,9 @@ public:
 	/** See TextFlowDirection attribute */
 	void SetTextFlowDirection(const TOptional<ETextFlowDirection>& InTextFlowDirection);
 
+	/** Sets the overflow policy for this text block */
+	void SetOverflowPolicy(TOptional<ETextOverflowPolicy> InOverflowPolicy);
+
 	/** Query to see if any text is selected within the document */
 	bool AnyTextSelected() const;
 
diff --git a/Engine/Source/Runtime/Slate/Public/Widgets/Input/SMultiLineEditableTextBox.h b/Engine/Source/Runtime/Slate/Public/Widgets/Input/SMultiLineEditableTextBox.h
--- a/Engine/Source/Runtime/Slate/Public/Widgets/Input/SMultiLineEditableTextBox.h
+++ b/Engine/Source/Runtime/Slate/Public/Widgets/Input/SMultiLineEditableTextBox.h
@@ -73,6 +73,7 @@ public:
 		, _VirtualKeyboardDismissAction(EVirtualKeyboardDismissAction::TextChangeOnDismiss)
 		, _TextShapingMethod()
 		, _TextFlowDirection()
+		, _OverflowPolicy()
 		{}
 
 		/** The styling of the textbox */
@@ -236,6 +237,8 @@ public:
 		/** Which text flow direction should we use? (unset to use the default returned by GetDefaultTextFlowDirection) */
 		SLATE_ARGUMENT( TOptional<ETextFlowDirection>, TextFlowDirection )
 
+		/** Determines what happens to text that is clipped and doesn't fit within the allotted area for this widget */
+		SLATE_ARGUMENT(TOptional<ETextOverflowPolicy>, OverflowPolicy)
 	SLATE_END_ARGS()
 	
 	/**
@@ -360,6 +363,9 @@ public:
 	/** See Justification attribute */
 	void SetJustification(const TAttribute<ETextJustify::Type>& InJustification);
 
+	/** Sets the overflow policy for this text block */
+	void SetOverflowPolicy(TOptional<ETextOverflowPolicy> InOverflowPolicy);
+
 	/** See the AllowContextMenu attribute */
 	void SetAllowContextMenu(const TAttribute< bool >& InAllowContextMenu);
 
diff --git a/Engine/Source/Runtime/Slate/Public/Widgets/Text/SInlineEditableTextBlock.h b/Engine/Source/Runtime/Slate/Public/Widgets/Text/SInlineEditableTextBlock.h
--- a/Engine/Source/Runtime/Slate/Public/Widgets/Text/SInlineEditableTextBlock.h
+++ b/Engine/Source/Runtime/Slate/Public/Widgets/Text/SInlineEditableTextBlock.h
@@ -44,6 +44,7 @@ class SLATE_API SInlineEditableTextBlock: public SCompoundWidget
 		, _IsReadOnly(false)
 		, _MultiLine(false)
 		, _ModiferKeyForNewLine(EModifierKey::None)
+		, _OverflowPolicy()
 	{
 	}
 
@@ -103,6 +104,9 @@ class SLATE_API SInlineEditableTextBlock: public SCompoundWidget
 
 		/** Called whenever the text is changed programmatically or interactively by the user */
 		SLATE_EVENT( FOnVerifyTextChanged, OnVerifyTextChanged )
+
+		/** Determines what happens to text that is clipped and doesn't fit within the clip rect for this widget */
+		SLATE_ARGUMENT(TOptional<ETextOverflowPolicy>, OverflowPolicy)
 	SLATE_END_ARGS()
 
 	~SInlineEditableTextBlock();
@@ -149,6 +153,9 @@ class SLATE_API SInlineEditableTextBlock: public SCompoundWidget
 	/** Sets the wrap text at attribute.  See WrapTextAt attribute */
 	void SetWrapTextAt(const TAttribute<float>& InWrapTextAt);
 
+	/** Sets the overflow policy for this text block */
+	void SetOverflowPolicy(TOptional<ETextOverflowPolicy> InOverflowPolicy);
+
 protected:
 	/** Callback for the text box's OnTextChanged event */
 	void OnTextChanged(const FText& InText);
diff --git a/Engine/Source/Runtime/Slate/Public/Widgets/Text/SMultiLineEditableText.h b/Engine/Source/Runtime/Slate/Public/Widgets/Text/SMultiLineEditableText.h
--- a/Engine/Source/Runtime/Slate/Public/Widgets/Text/SMultiLineEditableText.h
+++ b/Engine/Source/Runtime/Slate/Public/Widgets/Text/SMultiLineEditableText.h
@@ -75,6 +75,7 @@ public:
 		, _VirtualKeyboardDismissAction(EVirtualKeyboardDismissAction::TextChangeOnDismiss)
 		, _TextShapingMethod()
 		, _TextFlowDirection()
+		, _OverflowPolicy()
 	{
 		_Clipping = EWidgetClipping::ClipToBounds;
 	}
@@ -203,6 +204,9 @@ public:
 		/** Which text flow direction should we use? (unset to use the default returned by GetDefaultTextFlowDirection) */
 		SLATE_ARGUMENT( TOptional<ETextFlowDirection>, TextFlowDirection )
 
+		/** Determines what happens to text that is clipped and doesnt fit within the clip rect for this widget */
+		SLATE_ARGUMENT(TOptional<ETextOverflowPolicy>, OverflowPolicy)
+
 	SLATE_END_ARGS()
 
 	SMultiLineEditableText();
@@ -279,6 +283,9 @@ public:
 	/** See Justification attribute */
 	void SetJustification(const TAttribute<ETextJustify::Type>& InJustification);
 
+	/** Sets the overflow policy for this text block */
+	void SetOverflowPolicy(TOptional<ETextOverflowPolicy> InOverflowPolicy);
+
 	/** See the AllowContextMenu attribute */
 	void SetAllowContextMenu(const TAttribute< bool >& InAllowContextMenu);
 
diff --git a/Engine/Source/Runtime/Slate/Public/Widgets/Text/SRichTextBlock.h b/Engine/Source/Runtime/Slate/Public/Widgets/Text/SRichTextBlock.h
--- a/Engine/Source/Runtime/Slate/Public/Widgets/Text/SRichTextBlock.h
+++ b/Engine/Source/Runtime/Slate/Public/Widgets/Text/SRichTextBlock.h
@@ -51,6 +51,7 @@ public:
 		, _TextFlowDirection()
 		, _Decorators()
 		, _Parser()
+		, _OverflowPolicy()
 		, _MinDesiredWidth()
 	{
 		_Clipping = EWidgetClipping::OnDemand;
@@ -108,6 +109,9 @@ public:
 		/** The parser used to resolve any markup used in the provided string. */
 		SLATE_ARGUMENT( TSharedPtr< class IRichTextMarkupParser >, Parser )
 
+		/** Determines what happens to text that is clipped and doesn't fit within the clip rect for this widget */
+		SLATE_ARGUMENT(TOptional<ETextOverflowPolicy>, OverflowPolicy)
+
 		/** Minimum width that this text block should be */
 		SLATE_ATTRIBUTE(float, MinDesiredWidth)
 
@@ -218,6 +222,9 @@ public:
 	/**  */
 	void SetDecoratorStyleSet(const ISlateStyle* NewDecoratorStyleSet);
 
+	/** Sets the overflow policy for this text block */
+	void SetOverflowPolicy(TOptional<ETextOverflowPolicy> InOverflowPolicy);
+
 	/**
 	 * Causes the text to reflow it's layout
 	 */
diff --git a/Engine/Source/Runtime/Slate/Public/Widgets/Text/STextBlock.h b/Engine/Source/Runtime/Slate/Public/Widgets/Text/STextBlock.h
--- a/Engine/Source/Runtime/Slate/Public/Widgets/Text/STextBlock.h
+++ b/Engine/Source/Runtime/Slate/Public/Widgets/Text/STextBlock.h
@@ -64,6 +64,7 @@ public:
 		, _TextShapingMethod()
 		, _TextFlowDirection()
 		, _LineBreakPolicy()
+		, _OverflowPolicy()
 		, _SimpleTextMode(false)
 		{
 			_Clipping = EWidgetClipping::OnDemand;
@@ -134,6 +135,9 @@ public:
 		/** The iterator to use to detect appropriate soft-wrapping points for lines (or null to use the default) */
 		SLATE_ARGUMENT( TSharedPtr<IBreakIterator>, LineBreakPolicy )
 
+		/** Determines what happens to text that is clipped and doesn't fit within the clip rect for this widget */
+		SLATE_ARGUMENT(TOptional<ETextOverflowPolicy>, OverflowPolicy)
+
 		/**
 		 * If this is enabled, text shaping, wrapping, justification are disabled in favor of much faster text layout and measurement.
 		 * This feature is suitable for numbers and text that changes often and impact performance.
@@ -215,6 +219,9 @@ public:
 	UE_DEPRECATED(5.0, "GetTransformPolicy is not accessible anymore since it's attribute value may not have been updated yet.")
 	ETextTransformPolicy GetTransformPolicy() const { return GetTransformPolicyImpl(); }
 
+	/** Sets the overflow policy for this text block */
+	void SetOverflowPolicy(TOptional<ETextOverflowPolicy> InOverflowPolicy);
+
 	/** See ShadowOffset attribute */
 	void SetShadowOffset(TAttribute<FVector2D> InShadowOffset);
 
diff --git a/Engine/Source/Runtime/Slate/Public/Widgets/Text/SlateEditableTextLayout.h b/Engine/Source/Runtime/Slate/Public/Widgets/Text/SlateEditableTextLayout.h
--- a/Engine/Source/Runtime/Slate/Public/Widgets/Text/SlateEditableTextLayout.h
+++ b/Engine/Source/Runtime/Slate/Public/Widgets/Text/SlateEditableTextLayout.h
@@ -105,6 +105,9 @@ public:
 	/** Set the amount to scale each lines height by */
 	void SetLineHeightPercentage(const TAttribute<float>& InLineHeightPercentage);
 
+	/** Set the text overflow policy that should be used to determine what happens to clipped text */
+	void SetOverflowPolicy(TOptional<ETextOverflowPolicy> InOverflowPolicy);
+
 	/** Set the information used to help identify who owns this text layout in the case of an error */
 	void SetDebugSourceInfo(const TAttribute<FString>& InDebugSourceInfo);
 
@@ -627,6 +630,9 @@ private:
 	/** How the text was committed by the virtual keyboard */
 	ETextCommit::Type VirtualKeyboardTextCommitType;
 
+	/** Override for for the overflow policy. If this is not set the text style setting is used */
+	TOptional<ETextOverflowPolicy> OverflowPolicyOverride;
+
 	/** The last known size of the widget from the previous OnPaint, used to recalculate wrapping */
 	FVector2D CachedSize;
 
diff --git a/Engine/Source/Runtime/Slate/Public/Widgets/Text/SlateTextBlockLayout.h b/Engine/Source/Runtime/Slate/Public/Widgets/Text/SlateTextBlockLayout.h
--- a/Engine/Source/Runtime/Slate/Public/Widgets/Text/SlateTextBlockLayout.h
+++ b/Engine/Source/Runtime/Slate/Public/Widgets/Text/SlateTextBlockLayout.h
@@ -157,6 +157,11 @@ public:
 	 */
 	void SetTextFlowDirection(const TOptional<ETextFlowDirection>& InTextFlowDirection);
 
+	/**
+	 * Set the text overflow policy that the internal text layout should use
+	 */
+	void SetTextOverflowPolicy(const TOptional<ETextOverflowPolicy> InTextOverflowPolicy);
+
 	/**
 	 * Set the information used to help identify who owns this text layout in the case of an error
 	 */
diff --git a/Engine/Source/Runtime/SlateCore/Private/Fonts/FontCache.cpp b/Engine/Source/Runtime/SlateCore/Private/Fonts/FontCache.cpp
--- a/Engine/Source/Runtime/SlateCore/Private/Fonts/FontCache.cpp
+++ b/Engine/Source/Runtime/SlateCore/Private/Fonts/FontCache.cpp
@@ -810,6 +810,7 @@ FSlateFontCache::FSlateFontCache( TSharedRef<ISlateFontAtlasFactory> InFontAtlas
 	, FontAtlasFactory( InFontAtlasFactory )
 	, bFlushRequested( false )
 	, OwningThread(InOwningThread)
+	, EllipsisText(NSLOCTEXT("FontCache", "TextOverflowIndicator", "\u2026"))
 {
 	FInternationalization::Get().OnCultureChanged().AddRaw(this, &FSlateFontCache::HandleCultureChanged);
 }
@@ -820,6 +821,7 @@ FSlateFontCache::~FSlateFontCache()
 
 	// Make sure things get destroyed in the correct order
 	FontToCharacterListCache.Empty();
+	FontToOverflowGlyphSequence.Empty();
 	ShapedGlyphToAtlasData.Empty();
 	TextShaper.Reset();
 	FontRenderer.Reset();
@@ -991,25 +993,25 @@ FCharacterList& FSlateFontCache::GetCharacterList( const FSlateFontInfo &InFontI
 	// Create a key for looking up each character
 	const FSlateFontKey FontKey( InFontInfo, InOutlineSettings, FontScale );
 
-	TSharedRef< class FCharacterList >* CachedCharacterList = FontToCharacterListCache.Find( FontKey );
+	TUniquePtr<FCharacterList>* CachedCharacterList = FontToCharacterListCache.Find( FontKey );
 
-	if( CachedCharacterList )
+	if(CachedCharacterList)
 	{
 #if WITH_EDITORONLY_DATA
 		// Clear out this entry if it's stale so that we make a new one
-		if( (*CachedCharacterList)->IsStale() )
+		if((*CachedCharacterList)->IsStale())
 		{
-			FontToCharacterListCache.Remove( FontKey );
+			FontToCharacterListCache.Remove(FontKey);
 			FlushData();
 		}
 		else
 #endif	// WITH_EDITORONLY_DATA
 		{
-			return CachedCharacterList->Get();
+			return **CachedCharacterList;
 		}
 	}
 
-	return FontToCharacterListCache.Add( FontKey, MakeShareable( new FCharacterList( FontKey, *this ) ) ).Get();
+	return *FontToCharacterListCache.Add(FontKey, MakeUnique<FCharacterList>(FontKey, *this));
 }
 
 FShapedGlyphFontAtlasData FSlateFontCache::GetShapedGlyphFontAtlasData( const FShapedGlyphEntry& InShapedGlyph, const FFontOutlineSettings& InOutlineSettings )
@@ -1039,7 +1041,6 @@ FShapedGlyphFontAtlasData FSlateFontCache::GetShapedGlyphFontAtlasData( const FS
 		return **FoundAtlasData;
 	}
 
-
 	{
 
 		QUICK_SCOPE_CYCLE_COUNTER(STAT_SlateFontCacheAddNewShapedEntry)
@@ -1058,6 +1059,22 @@ FShapedGlyphFontAtlasData FSlateFontCache::GetShapedGlyphFontAtlasData( const FS
 	}
 }
 
+FShapedGlyphSequenceRef FSlateFontCache::GetOverflowEllipsisText(const FSlateFontInfo& InFontInfo, const float InFontScale)
+{
+	FSlateFontKey Key(InFontInfo, InFontInfo.OutlineSettings, InFontScale);
+
+	FShapedGlyphSequenceRef* FoundSequence = FontToOverflowGlyphSequence.Find(Key);
+
+	if (FoundSequence)
+	{
+		return *FoundSequence;
+	}
+
+	FShapedGlyphSequenceRef NewSequence = ShapeBidirectionalText(EllipsisText.ToString(), InFontInfo, InFontScale, TextBiDi::ETextDirection::LeftToRight, GetDefaultTextShapingMethod());
+
+	return FontToOverflowGlyphSequence.Add(Key, NewSequence);
+}
+
 const FFontData& FSlateFontCache::GetDefaultFontData( const FSlateFontInfo& InFontInfo ) const
 {
 	return CompositeFontCache->GetDefaultFontData(InFontInfo);
@@ -1239,6 +1256,8 @@ void FSlateFontCache::FlushData()
 	}
 
 	FontToCharacterListCache.Empty();
+	FontToOverflowGlyphSequence.Empty();
+
 	ShapedGlyphToAtlasData.Empty();
 }
 
diff --git a/Engine/Source/Runtime/SlateCore/Private/Fonts/SlateTextShaper.cpp b/Engine/Source/Runtime/SlateCore/Private/Fonts/SlateTextShaper.cpp
--- a/Engine/Source/Runtime/SlateCore/Private/Fonts/SlateTextShaper.cpp
+++ b/Engine/Source/Runtime/SlateCore/Private/Fonts/SlateTextShaper.cpp
@@ -838,7 +838,7 @@ void FSlateTextShaper::PerformHarfBuzzTextShaping(const TCHAR* InText, const int
 
 bool FSlateTextShaper::InsertSubstituteGlyphs(const TCHAR* InText, const int32 InCharIndex, const TSharedRef<FShapedGlyphFaceData>& InShapedGlyphFaceData, const TSharedRef<FFreeTypeAdvanceCache>& AdvanceCache, TArray<FShapedGlyphEntry>& OutGlyphsToRender, const int16 InLetterSpacingScaled) const
 {
-	auto GetSpaceGlyphIndexAndAdvance = [this, &InShapedGlyphFaceData, &AdvanceCache](uint32& OutSpaceGlyphIndex, int16& OutSpaceXAdvance)
+	auto GetSpecifiedGlyphIndexAndAdvance = [this, &InShapedGlyphFaceData, &AdvanceCache](TCHAR Char, uint32& OutSpaceGlyphIndex, int16& OutSpaceXAdvance)
 	{
 		OutSpaceGlyphIndex = 0;
 		OutSpaceXAdvance = 0;
@@ -847,7 +847,7 @@ bool FSlateTextShaper::InsertSubstituteGlyphs(const TCHAR* InText, const int32 I
 			TSharedPtr<FFreeTypeFace> FTFace = InShapedGlyphFaceData->FontFace.Pin();
 			if (FTFace.IsValid())
 			{
-				OutSpaceGlyphIndex = FT_Get_Char_Index(FTFace->GetFace(), TEXT(' '));
+				OutSpaceGlyphIndex = FT_Get_Char_Index(FTFace->GetFace(), Char);
 
 				FT_Fixed CachedAdvanceData = 0;
 				if (AdvanceCache->FindOrCache(OutSpaceGlyphIndex, CachedAdvanceData))
@@ -884,7 +884,7 @@ bool FSlateTextShaper::InsertSubstituteGlyphs(const TCHAR* InText, const int32 I
 	{
 		uint32 SpaceGlyphIndex = 0;
 		int16 SpaceXAdvance = 0;
-		GetSpaceGlyphIndexAndAdvance(SpaceGlyphIndex, SpaceXAdvance);
+		GetSpecifiedGlyphIndexAndAdvance(TEXT(' '), SpaceGlyphIndex, SpaceXAdvance);
 
 		// We insert a spacer glyph with (up-to) the width of 4 space glyphs in-place of a tab character
 		const int16 NumSpacesToInsert = 4 - (OutGlyphsToRender.Num() % 4);
@@ -911,8 +911,8 @@ bool FSlateTextShaper::InsertSubstituteGlyphs(const TCHAR* InText, const int32 I
 	}
 
 	if (Char == TEXT('\u2009') ||	// Thin Space
-		Char == TEXT('\u202F')		// Narrow No-Break Space
-		)
+		Char == TEXT('\u202F') ||	// Narrow No-Break Space
+		Char == TEXT('\u2026'))		// Ellipsis character
 	{
 		// Not all fonts support these characters
 #if WITH_FREETYPE
@@ -930,24 +930,52 @@ bool FSlateTextShaper::InsertSubstituteGlyphs(const TCHAR* InText, const int32 I
 		}
 #endif // WITH_FREETYPE
 
-		// If it doesn't, then make these 2/3rd the width of a normal space
-		uint32 SpaceGlyphIndex = 0;
-		int16 SpaceXAdvance = 0;
-		GetSpaceGlyphIndexAndAdvance(SpaceGlyphIndex, SpaceXAdvance);
+		if (Char == TEXT('\u2026'))
+		{
+			// If the ellipsis character doesn't exist in the font replace with 3 dots
+			uint32 DotGlyphIndex = 0;
+			int16 DotXAdvance = 0;
+			GetSpecifiedGlyphIndexAndAdvance(TEXT('.'), DotGlyphIndex, DotXAdvance);
 
-		FShapedGlyphEntry& ShapedGlyphEntry = OutGlyphsToRender.AddDefaulted_GetRef();
-		ShapedGlyphEntry.FontFaceData = InShapedGlyphFaceData;
-		ShapedGlyphEntry.GlyphIndex = SpaceGlyphIndex;
-		ShapedGlyphEntry.SourceIndex = InCharIndex;
-		ShapedGlyphEntry.XAdvance = ((SpaceXAdvance + InLetterSpacingScaled) * 2) / 3;
-		ShapedGlyphEntry.YAdvance = 0;
-		ShapedGlyphEntry.XOffset = 0;
-		ShapedGlyphEntry.YOffset = 0;
-		ShapedGlyphEntry.Kerning = 0;
-		ShapedGlyphEntry.NumCharactersInGlyph = 1;
-		ShapedGlyphEntry.NumGraphemeClustersInGlyph = 1;
-		ShapedGlyphEntry.TextDirection = TextBiDi::ETextDirection::LeftToRight;
-		ShapedGlyphEntry.bIsVisible = false;
+			// Insert 3 dots
+			for (int32 Dot = 0; Dot < 3; ++Dot)
+			{
+				FShapedGlyphEntry& ShapedGlyphEntry = OutGlyphsToRender.AddDefaulted_GetRef();
+				ShapedGlyphEntry.FontFaceData = InShapedGlyphFaceData;
+				ShapedGlyphEntry.GlyphIndex = DotGlyphIndex;
+				ShapedGlyphEntry.SourceIndex = InCharIndex;
+				ShapedGlyphEntry.XAdvance = DotXAdvance + InLetterSpacingScaled;
+				ShapedGlyphEntry.YAdvance = 0;
+				ShapedGlyphEntry.XOffset = 0;
+				ShapedGlyphEntry.YOffset = 0;
+				ShapedGlyphEntry.Kerning = 0;
+				ShapedGlyphEntry.NumCharactersInGlyph = 1; // should this be 3 instead of a loop three times?
+				ShapedGlyphEntry.NumGraphemeClustersInGlyph = 1;
+				ShapedGlyphEntry.TextDirection = TextBiDi::ETextDirection::LeftToRight;
+				ShapedGlyphEntry.bIsVisible = true;
+			}
+		}
+		else
+		{
+			// If it doesn't, then make these 2/3rd the width of a normal space
+			uint32 SpaceGlyphIndex = 0;
+			int16 SpaceXAdvance = 0;
+			GetSpecifiedGlyphIndexAndAdvance(TEXT(' '), SpaceGlyphIndex, SpaceXAdvance);
+
+			FShapedGlyphEntry& ShapedGlyphEntry = OutGlyphsToRender.AddDefaulted_GetRef();
+			ShapedGlyphEntry.FontFaceData = InShapedGlyphFaceData;
+			ShapedGlyphEntry.GlyphIndex = SpaceGlyphIndex;
+			ShapedGlyphEntry.SourceIndex = InCharIndex;
+			ShapedGlyphEntry.XAdvance = ((SpaceXAdvance + InLetterSpacingScaled) * 2) / 3;
+			ShapedGlyphEntry.YAdvance = 0;
+			ShapedGlyphEntry.XOffset = 0;
+			ShapedGlyphEntry.YOffset = 0;
+			ShapedGlyphEntry.Kerning = 0;
+			ShapedGlyphEntry.NumCharactersInGlyph = 1;
+			ShapedGlyphEntry.NumGraphemeClustersInGlyph = 1;
+			ShapedGlyphEntry.TextDirection = TextBiDi::ETextDirection::LeftToRight;
+			ShapedGlyphEntry.bIsVisible = false;
+		}
 		return true;
 	}
 
diff --git a/Engine/Source/Runtime/SlateCore/Private/Rendering/DrawElements.cpp b/Engine/Source/Runtime/SlateCore/Private/Rendering/DrawElements.cpp
--- a/Engine/Source/Runtime/SlateCore/Private/Rendering/DrawElements.cpp
+++ b/Engine/Source/Runtime/SlateCore/Private/Rendering/DrawElements.cpp
@@ -406,7 +406,7 @@ namespace SlateDrawElement
 #endif
 }
 
-void FSlateDrawElement::MakeShapedText(FSlateWindowElementList& ElementList, uint32 InLayer, const FPaintGeometry& PaintGeometry, const FShapedGlyphSequenceRef& InShapedGlyphSequence, ESlateDrawEffect InDrawEffects, const FLinearColor& BaseTint, const FLinearColor& OutlineTint)
+void FSlateDrawElement::MakeShapedText(FSlateWindowElementList& ElementList, uint32 InLayer, const FPaintGeometry& PaintGeometry, const FShapedGlyphSequenceRef& InShapedGlyphSequence, ESlateDrawEffect InDrawEffects, const FLinearColor& BaseTint, const FLinearColor& OutlineTint, FTextOverflowArgs TextOverflowArgs)
 {
 	SCOPE_CYCLE_COUNTER(STAT_SlateDrawElementMakeTime)
 	PaintGeometry.CommitTransformsIfUsingLegacyConstructor();
@@ -438,6 +438,7 @@ void FSlateDrawElement::MakeShapedText(FSlateWindowElementList& ElementList, uin
 	FSlateShapedTextPayload& DataPayload = ElementList.CreatePayload<FSlateShapedTextPayload>(Element);
 	DataPayload.SetTint(BaseTint);
 	DataPayload.SetShapedText(ElementList, InShapedGlyphSequence, OutlineTint);
+	DataPayload.SetOverflowArgs(TextOverflowArgs);
 
 	Element.Init(ElementList, EElementType::ET_ShapedText, InLayer, PaintGeometry, InDrawEffects);
 }
diff --git a/Engine/Source/Runtime/SlateCore/Private/Rendering/ElementBatcher.cpp b/Engine/Source/Runtime/SlateCore/Private/Rendering/ElementBatcher.cpp
--- a/Engine/Source/Runtime/SlateCore/Private/Rendering/ElementBatcher.cpp
+++ b/Engine/Source/Runtime/SlateCore/Private/Rendering/ElementBatcher.cpp
@@ -1213,29 +1213,38 @@ void FSlateElementBatcher::AddShapedTextElement( const FSlateDrawElement& DrawEl
 {
 	const FSlateShapedTextPayload& DrawElementPayload = DrawElement.GetDataPayload<FSlateShapedTextPayload>();
 	const FShapedGlyphSequence* ShapedGlyphSequence = DrawElementPayload.GetShapedGlyphSequence().Get();
+
+	const FShapedGlyphSequence* OverflowGlyphSequence = DrawElementPayload.OverflowArgs.OverflowTextPtr.Get();
+
 	checkSlow(ShapedGlyphSequence);
 
 	const FFontOutlineSettings& OutlineSettings = ShapedGlyphSequence->GetFontOutlineSettings();
 
-	const TArray<FShapedGlyphEntry>& GlyphsToRender = ShapedGlyphSequence->GetGlyphsToRender();
-	ensure(GlyphsToRender.Num() > 0);
+	ensure(ShapedGlyphSequence->GetGlyphsToRender().Num() > 0);
 
 	const FColor BaseTint = PackVertexColor(DrawElementPayload.GetTint());
 
 	FSlateFontCache& FontCache = *RenderingPolicy->GetFontCache();
-	FSlateShaderResourceManager& ResourceManager = *RenderingPolicy->GetResourceManager();
+
 
 	const int16 TextBaseline = ShapedGlyphSequence->GetTextBaseline();
 	const uint16 MaxHeight = ShapedGlyphSequence->GetMaxTextHeight();
 
+	FShapedTextBuildContext BuildContext;
+
+	BuildContext.DrawElement = &DrawElement;
+	BuildContext.FontCache = &FontCache;
+	BuildContext.ShapedGlyphSequence = ShapedGlyphSequence;
+	BuildContext.OverflowGlyphSequence = OverflowGlyphSequence;
+	BuildContext.TextBaseline = TextBaseline;
+	BuildContext.MaxHeight = MaxHeight;
+
 	if (MaxHeight == 0)
 	{
 		// If the max text height is 0, we'll create NaN's further in the code, so avoid drawing text if this happens.
 		return;
 	}
 
-	ESlateDrawEffect InDrawEffects = DrawElement.GetDrawEffects();
-
 	const int32 Layer = DrawElement.GetLayer();
 
 	// extract the layout transform from the draw element
@@ -1245,8 +1254,9 @@ void FSlateElementBatcher::AddShapedTextElement( const FSlateDrawElement& DrawEl
 	// So we need to pull the layout scale out of the layout and render transform so we can apply them
 	// in local space with pre-scaled fonts.
 	const float FontScale = LayoutTransform.GetScale();
-	FSlateLayoutTransform InverseLayoutTransform = Inverse(Concatenate(Inverse(FontScale), LayoutTransform));
+
 	const FSlateRenderTransform RenderTransform = Concatenate(Inverse(FontScale), DrawElement.GetRenderTransform());
+	BuildContext.RenderTransform = &RenderTransform;
 
 	const UObject* BaseFontMaterial = ShapedGlyphSequence->GetFontMaterial();
 	const UObject* OutlineFontMaterial = OutlineSettings.OutlineMaterial;
@@ -1261,154 +1271,56 @@ void FSlateElementBatcher::AddShapedTextElement( const FSlateDrawElement& DrawEl
 		const float PosX = TopLeft.X+InHorizontalOffset;
 		float PosY = TopLeft.Y;
 
-		float LineX = PosX;
-		float LineY = PosY;
-
-		int32 FontTextureIndex = -1;
-		FSlateShaderResource* FontAtlasTexture = nullptr;
-		FSlateShaderResource* FontShaderResource = nullptr;
-		FColor FontTint = InTint;
-
-		FSlateRenderBatch* RenderBatch = nullptr;
-
-		uint32 VertexOffset = 0;
-		uint32 IndexOffset = 0;
-
-		float InvTextureSizeX = 0;
-		float InvTextureSizeY = 0;
+		BuildContext.FontMaterial = FontMaterial;
+		BuildContext.OutlineFontMaterial = OutlineFontMaterial;
+	
+		BuildContext.OutlineSettings = &InOutlineSettings;
+		BuildContext.StartLineX = PosX;
+		BuildContext.StartLineY = PosY;
+		BuildContext.LayerId = InLayer;
+		BuildContext.FontTint = InTint;
 
-		const bool bIsFontMaterial = FontMaterial != nullptr;
-		const bool bEnableOutline = InOutlineSettings.OutlineSize > 0;
+		BuildContext.bEnableOutline = InOutlineSettings.OutlineSize > 0;
 
 		// Optimize by culling
-		// Todo: this doesnt work with cached clipping
-		bool bEnableCulling = false;
-		float LocalClipBoundingBoxLeft = 0;
-		float LocalClipBoundingBoxRight = 0;
-		if (GlyphsToRender.Num() > 200)
+		// Todo: this doesn't work with cached clipping
+		BuildContext.bEnableCulling = false;
+		BuildContext.bForceEllipsis = DrawElementPayload.OverflowArgs.bForceEllipsisDueToClippedLine;
+		BuildContext.OverflowDirection = DrawElementPayload.OverflowArgs.OverflowDirection;
+
+		if (ShapedGlyphSequence->GetGlyphsToRender().Num() > 200 || (OverflowGlyphSequence && BuildContext.OverflowDirection != ETextOverflowDirection::NoOverflow))
 		{
 			const FSlateClippingState* ClippingState = ResolveClippingState(DrawElement);
 
 			if (ClippingState && ClippingState->ScissorRect.IsSet() && ClippingState->ScissorRect->IsAxisAligned() && RenderTransform.GetMatrix().IsIdentity())
 			{
-				bEnableCulling = true;
-				const FSlateRect LocalClipBoundingBox = TransformRect(RenderTransform.Inverse(), ClippingState->ScissorRect->GetBoundingBox());
-				LocalClipBoundingBoxLeft = LocalClipBoundingBox.Left;
-				LocalClipBoundingBoxRight = LocalClipBoundingBox.Right;
-			}
-		}
-
-		const int32 NumGlyphs = GlyphsToRender.Num();
-		for (int32 GlyphIndex = 0; GlyphIndex < NumGlyphs; ++GlyphIndex)
-		{
-			const FShapedGlyphEntry& GlyphToRender = GlyphsToRender[GlyphIndex];
-
-			if (GlyphToRender.bIsVisible)
-			{
-				const FShapedGlyphFontAtlasData GlyphAtlasData = FontCache.GetShapedGlyphFontAtlasData(GlyphToRender, InOutlineSettings);
-				 
-				if (GlyphAtlasData.Valid && (!bEnableOutline || GlyphAtlasData.SupportsOutline))
-				{
-					const float X = LineX + GlyphAtlasData.HorizontalOffset + GlyphToRender.XOffset;
-					// Note PosX,PosY is the upper left corner of the bounding box representing the string.  This computes the Y position of the baseline where text will sit
-
-					if (bEnableCulling)
-					{
-						if (X + GlyphAtlasData.USize < LocalClipBoundingBoxLeft)
-						{
-							LineX += GlyphToRender.XAdvance;
-							LineY += GlyphToRender.YAdvance;
-							continue;
-						}
-						else if (X > LocalClipBoundingBoxRight)
-						{
-							break;
-						}
-					}
-
-					if (FontAtlasTexture == nullptr || GlyphAtlasData.TextureIndex != FontTextureIndex)
-					{
-						// Font has a new texture for this glyph. Refresh the batch we use and the index we are currently using
-						FontTextureIndex = GlyphAtlasData.TextureIndex;
-
-						ISlateFontTexture* SlateFontTexture = FontCache.GetFontTexture(FontTextureIndex);
-						check(SlateFontTexture);
-
-						FontAtlasTexture = SlateFontTexture->GetSlateTexture();
-						check(FontAtlasTexture);
-
-						FontShaderResource = ResourceManager.GetFontShaderResource(FontTextureIndex, FontAtlasTexture, FontMaterial);
-						check(FontShaderResource);
-
-						const bool bIsGrayscale = SlateFontTexture->IsGrayscale();
-						FontTint = bIsGrayscale ? InTint : FColor::White;
-
-						RenderBatch = &CreateRenderBatch(InLayer, FShaderParams(), FontShaderResource, ESlateDrawPrimitive::TriangleList, bIsGrayscale ? ESlateShader::GrayscaleFont : ESlateShader::ColorFont, InDrawEffects, ESlateBatchDrawFlag::None, DrawElement);
+				// Non-render transformed box;
+				const FSlateRect ScissorRectBox = ClippingState->ScissorRect->GetBoundingBox();
 
-						// Reserve memory for the glyphs.  This isn't perfect as the text could contain spaces and we might not render the rest of the text in this batch but its better than resizing constantly
-						const int32 GlyphsLeft = NumGlyphs - GlyphIndex;
-						RenderBatch->ReserveVertices(GlyphsLeft*4);
-						RenderBatch->ReserveIndices(GlyphsLeft*6);
+				const FSlateRect LocalClipBoundingBox = TransformRect(RenderTransform.Inverse(), ScissorRectBox);
+				BuildContext.LocalClipBoundingBoxLeft = LocalClipBoundingBox.Left;
+				BuildContext.LocalClipBoundingBoxRight = LocalClipBoundingBox.Right - (BuildContext.bForceEllipsis ? OverflowGlyphSequence->GetMeasuredWidth() : 0);
 
-						InvTextureSizeX = 1.0f / FontAtlasTexture->GetWidth();
-						InvTextureSizeY = 1.0f / FontAtlasTexture->GetHeight();
-					}
-
-					const float BitmapRenderScale = GlyphToRender.GetBitmapRenderScale();
-					const float InvBitmapRenderScale = 1.0f / BitmapRenderScale;
-
-					const float Y = LineY - GlyphAtlasData.VerticalOffset + GlyphToRender.YOffset + ((MaxHeight + TextBaseline) * InvBitmapRenderScale);
-					const float U = GlyphAtlasData.StartU * InvTextureSizeX;
-					const float V = GlyphAtlasData.StartV * InvTextureSizeY;
-					const float SizeX = GlyphAtlasData.USize * BitmapRenderScale;
-					const float SizeY = GlyphAtlasData.VSize * BitmapRenderScale;
-					const float SizeU = GlyphAtlasData.USize * InvTextureSizeX;
-					const float SizeV = GlyphAtlasData.VSize * InvTextureSizeY;
-
-					{
-
-						const FVector2D UpperLeft(X, Y);
-						const FVector2D UpperRight(X + SizeX, Y);
-						const FVector2D LowerLeft(X, Y + SizeY);
-						const FVector2D LowerRight(X + SizeX, Y + SizeY);
-
-		
-						// The start index of these vertices in the index buffer
-						const uint32 IndexStart = RenderBatch->GetNumVertices();
-
-						float Ut = 0.0f, Vt = 0.0f, UtMax = 0.0f, VtMax = 0.0f;
-						if (bIsFontMaterial)
-						{
-							float DistAlpha = (float)GlyphIndex / NumGlyphs;
-							float DistAlphaNext = (float)(GlyphIndex + 1) / NumGlyphs;
-
-							// This creates a set of UV's that goes from 0-1, from left to right of the string in U and 0-1 baseline to baseline top to bottom in V
-							Ut = FMath::Lerp(0.0f, 1.0f, DistAlpha);
-							Vt = FMath::Lerp(0.0f, 1.0f, UpperLeft.Y / MaxHeight);
-
-							UtMax = FMath::Lerp(0.0f, 1.0f, DistAlphaNext);
-							VtMax = FMath::Lerp(0.0f, 1.0f, LowerLeft.Y / MaxHeight);
-						}
-
-						// Add four vertices to the list of verts to be added to the vertex buffer
-						RenderBatch->AddVertex(FSlateVertex::Make<Rounding>(RenderTransform, UpperLeft,								FVector4(U, V, Ut, Vt),							FVector2D(0.0f, 0.0f), FontTint ));
-						RenderBatch->AddVertex(FSlateVertex::Make<Rounding>(RenderTransform, FVector2D(LowerRight.X, UpperLeft.Y),	FVector4(U + SizeU, V, UtMax, Vt),				FVector2D(1.0f, 0.0f), FontTint ));
-						RenderBatch->AddVertex(FSlateVertex::Make<Rounding>(RenderTransform, FVector2D(UpperLeft.X, LowerRight.Y),	FVector4(U, V + SizeV, Ut, VtMax),				FVector2D(0.0f, 1.0f), FontTint ));
-						RenderBatch->AddVertex(FSlateVertex::Make<Rounding>(RenderTransform, LowerRight,								FVector4(U + SizeU, V + SizeV, UtMax, VtMax),	FVector2D(1.0f, 1.0f), FontTint ));
+				const float TransformOffset = RenderTransform.GetTranslation().X - DrawElement.GetPosition().X;
 
-						RenderBatch->AddIndex(IndexStart + 0);
-						RenderBatch->AddIndex(IndexStart + 1);
-						RenderBatch->AddIndex(IndexStart + 2);
-						RenderBatch->AddIndex(IndexStart + 1);
-						RenderBatch->AddIndex(IndexStart + 3);
-						RenderBatch->AddIndex(IndexStart + 2);
-					}
+				if (OverflowGlyphSequence && BuildContext.LocalClipBoundingBoxLeft <= 0 && BuildContext.LocalClipBoundingBoxRight > ShapedGlyphSequence->GetMeasuredWidth() + TransformOffset)
+				{
+					// Override overflow if the text is smaller than the clipping rect and wont be clipped
+					BuildContext.OverflowDirection = ETextOverflowDirection::NoOverflow;
+				}
+				else if(!OverflowGlyphSequence)
+				{
+					BuildContext.bEnableCulling = true;
 				}
 			}
-
-			LineX += GlyphToRender.XAdvance;
-			LineY += GlyphToRender.YAdvance;
+			else
+			{
+				// Overflow not supported on non-identity transforms
+				BuildContext.OverflowDirection = ETextOverflowDirection::NoOverflow;
+			}
 		}
+
+		BuildShapedTextSequence<Rounding>(BuildContext);
 	};
 
 	if (bOutlineFont)
@@ -2642,6 +2554,281 @@ const FSlateClippingState* FSlateElementBatcher::ResolveClippingState(const FSla
 	return nullptr;
 }
 
+template<ESlateVertexRounding Rounding>
+void FSlateElementBatcher::BuildShapedTextSequence(const FShapedTextBuildContext& Context)
+{
+	const FShapedGlyphSequence* GlyphSequenceToRender = Context.ShapedGlyphSequence;
+
+	FSlateShaderResourceManager& ResourceManager = *RenderingPolicy->GetResourceManager();
+
+	float InvTextureSizeX = 0;
+	float InvTextureSizeY = 0;
+
+	FSlateRenderBatch* RenderBatch = nullptr;
+
+	int32 FontTextureIndex = -1;
+	FSlateShaderResource* FontAtlasTexture = nullptr;
+	FSlateShaderResource* FontShaderResource = nullptr;
+
+	float LineX = Context.StartLineX;
+	float LineY = Context.StartLineY;
+
+	FSlateRenderTransform RenderTransform = *Context.RenderTransform;
+
+	FColor Tint = FColor::White;
+
+	ETextOverflowDirection OverflowDirection = Context.OverflowDirection;
+
+	float EllipsisLineX = 0;
+	float EllipsisLineY = 0;
+	bool bNeedEllipsis = false;
+	bool bChararcterWasClipped = false;
+
+	// For left to right overflow direction - Sum of total whitespace we're currently advancing through. Once a non-whitespace glyph is detected this will return to 0
+	// For right to left this value is unused. We just skip all leading whitespace
+	float PreviousWhitespaceAdvance = 0;
+
+	int32 NumGlyphs = GlyphSequenceToRender->GetGlyphsToRender().Num();
+	const TArray<FShapedGlyphEntry>& GlyphsToRender = GlyphSequenceToRender->GetGlyphsToRender();
+	for (int32 GlyphIndex = 0; GlyphIndex < NumGlyphs; ++GlyphIndex)
+	{
+		const FShapedGlyphEntry& GlyphToRender = GlyphsToRender[GlyphIndex];
+
+		const float BitmapRenderScale = GlyphToRender.GetBitmapRenderScale();
+		const float InvBitmapRenderScale = 1.0f / BitmapRenderScale;
+
+		float X = 0;
+		float SizeX = 0;
+		float Y = 0;
+		float U = 0;
+		float V = 0;
+		float SizeY = 0;
+		float SizeU = 0;
+		float SizeV = 0;
+
+		bool bCanRenderGlyph = GlyphToRender.bIsVisible;
+		if (bCanRenderGlyph)
+		{
+			// Get Sizing and atlas info
+			const FShapedGlyphFontAtlasData GlyphAtlasData = Context.FontCache->GetShapedGlyphFontAtlasData(GlyphToRender, *Context.OutlineSettings);
+			if (GlyphAtlasData.Valid && (!Context.bEnableOutline || GlyphAtlasData.SupportsOutline))
+			{
+				X = LineX + GlyphAtlasData.HorizontalOffset + GlyphToRender.XOffset;
+				// Note PosX,PosY is the upper left corner of the bounding box representing the string.  This computes the Y position of the baseline where text will sit
+
+				if (Context.bEnableCulling)
+				{
+					if (X + GlyphAtlasData.USize < Context.LocalClipBoundingBoxLeft)
+					{
+						LineX += GlyphToRender.XAdvance;
+						LineY += GlyphToRender.YAdvance;
+						continue;
+					}
+					else if (X > Context.LocalClipBoundingBoxRight)
+					{
+						break;
+					}
+				}
+
+				if (FontAtlasTexture == nullptr || GlyphAtlasData.TextureIndex != FontTextureIndex)
+				{
+					// Font has a new texture for this glyph. Refresh the batch we use and the index we are currently using
+					FontTextureIndex = GlyphAtlasData.TextureIndex;
+
+					ISlateFontTexture* SlateFontTexture = Context.FontCache->GetFontTexture(FontTextureIndex);
+					check(SlateFontTexture);
+
+					FontAtlasTexture = SlateFontTexture->GetSlateTexture();
+					check(FontAtlasTexture);
+
+					FontShaderResource = ResourceManager.GetFontShaderResource(FontTextureIndex, FontAtlasTexture, Context.FontMaterial);
+					check(FontShaderResource);
+
+					const bool bIsGrayscale = SlateFontTexture->IsGrayscale();
+					Tint = bIsGrayscale ? Context.FontTint : FColor::White;
+
+					RenderBatch = &CreateRenderBatch(Context.LayerId, FShaderParams(), FontShaderResource, ESlateDrawPrimitive::TriangleList, bIsGrayscale ? ESlateShader::GrayscaleFont : ESlateShader::ColorFont, Context.DrawElement->GetDrawEffects(), ESlateBatchDrawFlag::None, *Context.DrawElement);
+
+					// Reserve memory for the glyphs.  This isn't perfect as the text could contain spaces and we might not render the rest of the text in this batch but its better than resizing constantly
+					const int32 GlyphsLeft = NumGlyphs - GlyphIndex;
+					RenderBatch->ReserveVertices(GlyphsLeft * 4);
+					RenderBatch->ReserveIndices(GlyphsLeft * 6);
+
+					InvTextureSizeX = 1.0f / FontAtlasTexture->GetWidth();
+					InvTextureSizeY = 1.0f / FontAtlasTexture->GetHeight();
+				}
+
+
+				Y = LineY - GlyphAtlasData.VerticalOffset + GlyphToRender.YOffset + ((Context.MaxHeight + Context.TextBaseline) * InvBitmapRenderScale);
+				U = GlyphAtlasData.StartU * InvTextureSizeX;
+				V = GlyphAtlasData.StartV * InvTextureSizeY;
+				SizeX = GlyphAtlasData.USize * BitmapRenderScale;
+				SizeY = GlyphAtlasData.VSize * BitmapRenderScale;
+				SizeU = GlyphAtlasData.USize * InvTextureSizeX;
+				SizeV = GlyphAtlasData.VSize * InvTextureSizeY;
+			}
+			else
+			{
+				bCanRenderGlyph = false;
+			}
+
+		}
+		else
+		{
+			X = LineX;
+			SizeX = GlyphToRender.XAdvance;
+		}
+
+		// Overflow Detection
+		// First figure out the size of the glyph. If the glyph contains multiple characters we have to measure all of them and if clicked, omit them all. This is common in complex languages with lots of diacritics
+		float OverflowTestWidth = SizeX;
+		if (OverflowDirection != ETextOverflowDirection::NoOverflow && (GlyphToRender.NumGraphemeClustersInGlyph > 1 || GlyphToRender.NumCharactersInGlyph > 1))
+		{
+			const int32 StartIndex = GlyphIndex;
+			int32 EndIndex = GlyphIndex;
+			int32 NextIndex = GlyphIndex + 1;
+			const int32 SourceIndex = GlyphToRender.SourceIndex;
+			while (NextIndex < NumGlyphs && GlyphsToRender[NextIndex].SourceIndex == SourceIndex)
+			{
+				++EndIndex;
+				++NextIndex;
+			}
+			if (StartIndex < EndIndex)
+			{
+				OverflowTestWidth = GlyphSequenceToRender->GetMeasuredWidth(StartIndex, EndIndex).Get(SizeX);
+			}
+		}
+
+		// Left to right overflow - If the current pen position + the ellipsis cannot fit, we have reached the end of the possible area for drawing this text. 
+		if (OverflowDirection == ETextOverflowDirection::LeftToRight)
+		{
+			// If we are on the last glyph dont bother checking if the ellipsis can fit. If the last glyph can fit there is no need for ellipsis
+			float OverflowSequenceNeededSize = GlyphIndex < NumGlyphs - 1 ? Context.OverflowGlyphSequence->GetMeasuredWidth() : 0;
+			if(X + OverflowTestWidth + OverflowSequenceNeededSize >= Context.LocalClipBoundingBoxRight)
+			{
+				bNeedEllipsis = true;
+				// We subtract out any whitespace advance. This avoids the ellipsis from ever floating out in the middle of a block of whitespace.
+				// e.g without this something like "The quick brown		fox jumps over the lazy dog" could be clipped to "The quick brown	..." but we want it to be "The quick brown..."
+				EllipsisLineX = LineX - PreviousWhitespaceAdvance;
+				EllipsisLineY = LineY;
+				// No characters to render after the ellipsis on the right side
+				break;
+			}
+		}
+		else if(OverflowDirection == ETextOverflowDirection::RightToLeft)
+		{
+			bool bClipped = false;
+			// Right to left overflow
+			if (X < Context.LocalClipBoundingBoxLeft)
+			{
+				// This glyph is in the clipped region or is not visible so just advance. It cannot be shown
+				bClipped = true;
+				bChararcterWasClipped = true;
+			}
+			else if(bChararcterWasClipped)
+			{
+				// Can the ellipsis fit in the free spot by skipping the previous glyph(s)
+				float EllipsisWidth = Context.OverflowGlyphSequence->GetMeasuredWidth();
+				float AvailableX = X-Context.LocalClipBoundingBoxLeft;
+				if (AvailableX >= EllipsisWidth)
+				{
+					// The available area can fit the ellipsis. Mark that we need an ellipsis and stop checking for overflow. The rest of the text can be built normally
+					bNeedEllipsis = true;
+					EllipsisLineX = (LineX - EllipsisWidth);
+					EllipsisLineY = LineY;
+					OverflowDirection = ETextOverflowDirection::NoOverflow;
+				}
+				else
+				{
+					bClipped = true;
+					bChararcterWasClipped = true;
+				}
+
+			}
+			// If we just clipped a glyph omit all characters in said glyph. Otherwise floating diacritics would be visible above the ellipsis. This is common in complex languages.
+			if (bClipped && GlyphToRender.NumCharactersInGlyph > 1)
+			{
+				GlyphIndex += GlyphToRender.NumCharactersInGlyph - 1;
+				LineX += OverflowTestWidth;
+				continue;
+			}
+
+			bCanRenderGlyph = !bClipped;
+		}
+
+		if(bCanRenderGlyph)
+		{
+			const FVector2D UpperLeft(X, Y);
+			const FVector2D UpperRight(X + SizeX, Y);
+			const FVector2D LowerLeft(X, Y + SizeY);
+			const FVector2D LowerRight(X + SizeX, Y + SizeY);
+
+			// The start index of these vertices in the index buffer
+			const uint32 IndexStart = RenderBatch->GetNumVertices();
+
+			float Ut = 0.0f, Vt = 0.0f, UtMax = 0.0f, VtMax = 0.0f;
+			if (Context.FontMaterial)
+			{
+				float DistAlpha = (float)GlyphIndex / NumGlyphs;
+				float DistAlphaNext = (float)(GlyphIndex + 1) / NumGlyphs;
+
+				// This creates a set of UV's that goes from 0-1, from left to right of the string in U and 0-1 baseline to baseline top to bottom in V
+				Ut = FMath::Lerp(0.0f, 1.0f, DistAlpha);
+				Vt = FMath::Lerp(0.0f, 1.0f, UpperLeft.Y / Context.MaxHeight);
+
+				UtMax = FMath::Lerp(0.0f, 1.0f, DistAlphaNext);
+				VtMax = FMath::Lerp(0.0f, 1.0f, LowerLeft.Y / Context.MaxHeight);
+			}
+
+			// Add four vertices to the list of verts to be added to the vertex buffer
+			RenderBatch->AddVertex(FSlateVertex::Make<Rounding>(RenderTransform, UpperLeft, FVector4(U, V, Ut, Vt), FVector2D(0.0f, 0.0f), Tint));
+			RenderBatch->AddVertex(FSlateVertex::Make<Rounding>(RenderTransform, FVector2D(LowerRight.X, UpperLeft.Y), FVector4(U + SizeU, V, UtMax, Vt), FVector2D(1.0f, 0.0f), Tint));
+			RenderBatch->AddVertex(FSlateVertex::Make<Rounding>(RenderTransform, FVector2D(UpperLeft.X, LowerRight.Y), FVector4(U, V + SizeV, Ut, VtMax), FVector2D(0.0f, 1.0f), Tint));
+			RenderBatch->AddVertex(FSlateVertex::Make<Rounding>(RenderTransform, LowerRight, FVector4(U + SizeU, V + SizeV, UtMax, VtMax), FVector2D(1.0f, 1.0f), Tint));
+
+			RenderBatch->AddIndex(IndexStart + 0);
+			RenderBatch->AddIndex(IndexStart + 1);
+			RenderBatch->AddIndex(IndexStart + 2);
+			RenderBatch->AddIndex(IndexStart + 1);
+			RenderBatch->AddIndex(IndexStart + 3);
+			RenderBatch->AddIndex(IndexStart + 2);
+
+			// Reset whitespace advance to 0, this is a visible character
+			PreviousWhitespaceAdvance = 0;
+		}
+		else if(!GlyphToRender.bIsVisible)
+		{
+			// How much whitespace we are currently walking through
+			PreviousWhitespaceAdvance += GlyphToRender.XAdvance;
+		}
+
+		LineX += GlyphToRender.XAdvance;
+		LineY += GlyphToRender.YAdvance;
+	}
+
+	if (!bNeedEllipsis && Context.bForceEllipsis)
+	{
+		bNeedEllipsis = true;
+		EllipsisLineX = LineX;
+		EllipsisLineY = LineY;
+	}
+
+	if (bNeedEllipsis)
+	{
+		// Ellipsis can fit, place it at the current lineX
+		FShapedTextBuildContext EllipsisContext = Context;
+		EllipsisContext.bForceEllipsis = false;
+		EllipsisContext.ShapedGlyphSequence = Context.OverflowGlyphSequence;
+		EllipsisContext.OverflowGlyphSequence = nullptr;
+		EllipsisContext.bEnableCulling = false;
+		EllipsisContext.OverflowDirection = ETextOverflowDirection::NoOverflow;
+		EllipsisContext.StartLineX = EllipsisLineX;
+		EllipsisContext.StartLineY = EllipsisLineY;
+
+		BuildShapedTextSequence<Rounding>(EllipsisContext);
+	}
+}
+
 void FSlateElementBatcher::ResetBatches()
 {
 	bRequiresVsync = false;
diff --git a/Engine/Source/Runtime/SlateCore/Private/Styling/SlateTypes.cpp b/Engine/Source/Runtime/SlateCore/Private/Styling/SlateTypes.cpp
--- a/Engine/Source/Runtime/SlateCore/Private/Styling/SlateTypes.cpp
+++ b/Engine/Source/Runtime/SlateCore/Private/Styling/SlateTypes.cpp
@@ -77,6 +77,7 @@ FTextBlockStyle::FTextBlockStyle()
 	, SelectedBackgroundColor(FSlateColor::UseForeground())
 	, HighlightColor(FLinearColor::Black)
 	, TransformPolicy(ETextTransformPolicy::None)
+	, OverflowPolicy(ETextOverflowPolicy::Clip)
 {
 }
 
diff --git a/Engine/Source/Runtime/SlateCore/Public/Fonts/FontCache.h b/Engine/Source/Runtime/SlateCore/Public/Fonts/FontCache.h
--- a/Engine/Source/Runtime/SlateCore/Public/Fonts/FontCache.h
+++ b/Engine/Source/Runtime/SlateCore/Public/Fonts/FontCache.h
@@ -682,18 +682,12 @@ public:
 	 */
 	FShapedGlyphFontAtlasData GetShapedGlyphFontAtlasData( const FShapedGlyphEntry& InShapedGlyph, const FFontOutlineSettings& InOutlineSettings);
 
-	/** 
-	 * Add a new entries into a cache atlas
-	 *
-	 * @param InFontInfo	Information about the font being used for the characters
-	 * @param Characters	The characters to cache
-	 * @param FontScale		The font scale to use
-	 * @return true if the characters could be cached. false if the cache is full
+	/**
+	 * Gets the overflow glyph sequence for a given font. The overflow sequence is used to replace characters that are clipped
 	 */
-	bool AddNewEntry( const FShapedGlyphEntry& InShapedGlyph, const FFontOutlineSettings& InOutlineSettings, FShapedGlyphFontAtlasData& OutAtlasData );
-
-	bool AddNewEntry( const FCharacterRenderData InRenderData, uint8& OutTextureIndex, uint16& OutGlyphX, uint16& OutGlyphY, uint16& OutGlyphWidth, uint16& OutGlyphHeight );
+	FShapedGlyphSequenceRef GetOverflowEllipsisText(const FSlateFontInfo& InFontInfo, const float InFontScale);
 
+public:
 	/**
 	 * Flush the given object out of the cache
 	 */
@@ -868,6 +862,17 @@ private:
 	/** Called after the active culture has changed */
 	void HandleCultureChanged();
 
+	/**
+	 * Add a new entries into a cache atlas
+	 *
+	 * @param InFontInfo	Information about the font being used for the characters
+	 * @param Characters	The characters to cache
+	 * @param FontScale		The font scale to use
+	 * @return true if the characters could be cached. false if the cache is full
+	 */
+	bool AddNewEntry(const FShapedGlyphEntry& InShapedGlyph, const FFontOutlineSettings& InOutlineSettings, FShapedGlyphFontAtlasData& OutAtlasData);
+
+	bool AddNewEntry(const FCharacterRenderData InRenderData, uint8& OutTextureIndex, uint16& OutGlyphX, uint16& OutGlyphY, uint16& OutGlyphWidth, uint16& OutGlyphHeight);
 private:
 
 	/** FreeType library instance (owned by this font cache) */
@@ -886,7 +891,10 @@ private:
 	TUniquePtr<FSlateTextShaper> TextShaper;
 
 	/** Mapping Font keys to cached data */
-	TMap<FSlateFontKey, TSharedRef<class FCharacterList>, FDefaultSetAllocator, FSlateFontKeyFuncs<TSharedRef<class FCharacterList>>> FontToCharacterListCache;
+	TMap<FSlateFontKey, TUniquePtr<FCharacterList>, FDefaultSetAllocator, FSlateFontKeyFuncs<TUniquePtr<FCharacterList>>> FontToCharacterListCache;
+
+	/** Caches overflow text to display usually an ellipsis character for text elements that are clipped and request replacing clipped text with an ellipsis */
+	TMap<FSlateFontKey, FShapedGlyphSequenceRef, FDefaultSetAllocator, FSlateFontKeyFuncs<FShapedGlyphSequenceRef>> FontToOverflowGlyphSequence;
 
 	/** Mapping shaped glyphs to their cached atlas data */
 	TMap<FShapedGlyphEntryKey, TSharedRef<FShapedGlyphFontAtlasData>> ShapedGlyphToAtlasData;
@@ -919,4 +927,8 @@ private:
 	FOnReleaseFontResources OnReleaseResourcesDelegate;
 
 	ESlateTextureAtlasThreadId OwningThread;
+
+	/** Overflow text string to use to replace clipped characters */
+	FText EllipsisText;
+
 };
diff --git a/Engine/Source/Runtime/SlateCore/Public/Rendering/DrawElementPayloads.h b/Engine/Source/Runtime/SlateCore/Public/Rendering/DrawElementPayloads.h
--- a/Engine/Source/Runtime/SlateCore/Public/Rendering/DrawElementPayloads.h
+++ b/Engine/Source/Runtime/SlateCore/Public/Rendering/DrawElementPayloads.h
@@ -24,6 +24,15 @@ struct FSlateGradientStop
 };
 template <> struct TIsPODType<FSlateGradientStop> { enum { Value = true }; };
 
+enum class ETextOverflowDirection : uint8
+{
+	// No overflow
+	NoOverflow,
+	// Left justification overflow
+	LeftToRight,
+	// Right justification overflow
+	RightToLeft
+};
 
 struct FSlateDataPayload
 {
@@ -133,13 +142,37 @@ struct FSlateTextPayload : public FSlateDataPayload, public FSlateTintableElemen
 };
 
 
-struct FSlateShapedTextPayload : public FSlateDataPayload, public FSlateTintableElement
+
+struct FTextOverflowArgs
 {
+	FTextOverflowArgs(FShapedGlyphSequencePtr& InOverflowText, ETextOverflowDirection InOverflowDirection)
+		: OverflowTextPtr(InOverflowText)
+		, OverflowDirection(InOverflowDirection)
+		, bForceEllipsisDueToClippedLine(false)
+		
+	{}
+
+	FTextOverflowArgs()
+		: OverflowDirection(ETextOverflowDirection::NoOverflow)
+		, bForceEllipsisDueToClippedLine(false)
+	{}
+
+	/** Sequence that represents the ellipsis glyph */
+	FShapedGlyphSequencePtr OverflowTextPtr;
+	ETextOverflowDirection OverflowDirection;
+	bool bForceEllipsisDueToClippedLine;
+	
+};
+
+struct FSlateShapedTextPayload : public FSlateDataPayload, public FSlateTintableElement
+{ 
 	// Shaped text data
 	FShapedGlyphSequencePtr ShapedGlyphSequence;
 
 	FLinearColor OutlineTint;
 
+	FTextOverflowArgs OverflowArgs;
+
 	const FShapedGlyphSequencePtr& GetShapedGlyphSequence() const { return ShapedGlyphSequence; }
 	FLinearColor GetOutlineTint() const { return OutlineTint; }
 
@@ -149,12 +182,23 @@ struct FSlateShapedTextPayload : public FSlateDataPayload, public FSlateTintable
 		OutlineTint = InOutlineTint;
 	}
 
+	void SetOverflowArgs(const FTextOverflowArgs& InArgs)
+	{
+		OverflowArgs = InArgs;
+		check(InArgs.OverflowDirection == ETextOverflowDirection::NoOverflow || InArgs.OverflowTextPtr.IsValid());
+	}
+
 	virtual void AddReferencedObjects(FReferenceCollector& Collector)
 	{
 		if (ShapedGlyphSequence.IsValid())
 		{
 			const_cast<FShapedGlyphSequence*>(ShapedGlyphSequence.Get())->AddReferencedObjects(Collector);
 		}
+
+		if (OverflowArgs.OverflowTextPtr.IsValid())
+		{
+			const_cast<FShapedGlyphSequence*>(OverflowArgs.OverflowTextPtr.Get())->AddReferencedObjects(Collector);
+		}
 	}
 };
 
diff --git a/Engine/Source/Runtime/SlateCore/Public/Rendering/DrawElements.h b/Engine/Source/Runtime/SlateCore/Public/Rendering/DrawElements.h
--- a/Engine/Source/Runtime/SlateCore/Public/Rendering/DrawElements.h
+++ b/Engine/Source/Runtime/SlateCore/Public/Rendering/DrawElements.h
@@ -157,7 +157,7 @@ public:
 	 * @param InDrawEffects         Optional draw effects to apply
 	 * @param InTint                Color to tint the element
 	 */
-	SLATECORE_API static void MakeShapedText( FSlateWindowElementList& ElementList, uint32 InLayer, const FPaintGeometry& PaintGeometry, const FShapedGlyphSequenceRef& InShapedGlyphSequence, ESlateDrawEffect InDrawEffects, const FLinearColor& BaseTint, const FLinearColor& OutlineTint);
+	SLATECORE_API static void MakeShapedText( FSlateWindowElementList& ElementList, uint32 InLayer, const FPaintGeometry& PaintGeometry, const FShapedGlyphSequenceRef& InShapedGlyphSequence, ESlateDrawEffect InDrawEffects, const FLinearColor& BaseTint, const FLinearColor& OutlineTint, FTextOverflowArgs TextOverflowArgs = FTextOverflowArgs());
 
 	/**
 	 * Creates a gradient element
diff --git a/Engine/Source/Runtime/SlateCore/Public/Rendering/ElementBatcher.h b/Engine/Source/Runtime/SlateCore/Public/Rendering/ElementBatcher.h
--- a/Engine/Source/Runtime/SlateCore/Public/Rendering/ElementBatcher.h
+++ b/Engine/Source/Runtime/SlateCore/Public/Rendering/ElementBatcher.h
@@ -17,6 +17,7 @@ class FSlateWindowElementList;
 struct FShaderParams;
 struct FSlateCachedElementData;
 struct FSlateCachedElementList;
+enum class ETextOverflowDirection : uint8;
 
 DECLARE_CYCLE_STAT_EXTERN(TEXT("Add Elements Time"), STAT_SlateAddElements, STATGROUP_Slate, SLATECORE_API);
 
@@ -350,6 +351,33 @@ private:
 
 	const FSlateClippingState* ResolveClippingState(const FSlateDrawElement& DrawElement) const;
 
+	struct FShapedTextBuildContext
+	{
+		const class FShapedGlyphSequence* ShapedGlyphSequence;
+		const class FShapedGlyphSequence* OverflowGlyphSequence;
+		const UObject* FontMaterial;
+		const UObject* OutlineFontMaterial;
+		const struct FFontOutlineSettings* OutlineSettings;
+		const FSlateDrawElement* DrawElement;
+		class FSlateFontCache* FontCache;
+		const FSlateRenderTransform* RenderTransform;
+		float TextBaseline;
+		float MaxHeight;
+		float StartLineX;
+		float StartLineY;
+		float LocalClipBoundingBoxLeft = 0;
+		float LocalClipBoundingBoxRight = 0;
+		int32 LayerId;
+		FColor FontTint;
+		ETextOverflowDirection OverflowDirection;
+		bool bEnableOutline : 1;
+		bool bEnableCulling : 1;
+		bool bForceEllipsis : 1;
+		
+	};
+
+	template<ESlateVertexRounding Rounding>
+	void BuildShapedTextSequence(const FShapedTextBuildContext& Context);
 private:
 	/** Uncached Batch data currently being filled in */
 	FSlateBatchData* BatchData;
diff --git a/Engine/Source/Runtime/SlateCore/Public/Styling/SlateTypes.h b/Engine/Source/Runtime/SlateCore/Public/Styling/SlateTypes.h
--- a/Engine/Source/Runtime/SlateCore/Public/Styling/SlateTypes.h
+++ b/Engine/Source/Runtime/SlateCore/Public/Styling/SlateTypes.h
@@ -69,6 +69,20 @@ enum class ECheckBoxState : uint8
 	Undetermined
 };
 
+/**
+ * The different methods that can be used to determine what happens to text when it is longer than its allowed length
+ */
+UENUM(BlueprintType)
+enum class ETextOverflowPolicy : uint8
+{
+	/** Overflowing text will be clipped */
+	Clip = 0,
+
+	/** Overflowing text will be replaced with an ellipsis */
+	Ellipsis,
+};
+
+
 /**
  * Represents the appearance of an SCheckBox
  */
@@ -361,6 +375,11 @@ struct SLATECORE_API FTextBlockStyle : public FSlateWidgetStyle
 	ETextTransformPolicy TransformPolicy;
 	FTextBlockStyle& SetTransformPolicy( const ETextTransformPolicy& InTransformPolicy ){ TransformPolicy = InTransformPolicy; return *this; }
 
+	/** Determines what happens to text that is clipped and doesn't fit within the clip rect of a text widget */
+	UPROPERTY(EditAnywhere, BlueprintReadWrite, Category = Appearance, AdvancedDisplay)
+	ETextOverflowPolicy OverflowPolicy;
+	FTextBlockStyle& SetOverflowPolicy(const ETextOverflowPolicy& InOverflowPolicy) { OverflowPolicy = InOverflowPolicy; return *this; }
+
 	/**
 	 * Unlinks all colors in this style.
 	 * @see FSlateColor::Unlink
diff --git a/Engine/Source/Runtime/UMG/Private/Components/EditableText.cpp b/Engine/Source/Runtime/UMG/Private/Components/EditableText.cpp
--- a/Engine/Source/Runtime/UMG/Private/Components/EditableText.cpp
+++ b/Engine/Source/Runtime/UMG/Private/Components/EditableText.cpp
@@ -67,6 +67,7 @@ UEditableText::UEditableText(const FObjectInitializer& ObjectInitializer)
 	VirtualKeyboardTrigger = EVirtualKeyboardTrigger::OnFocusByPointer;
 	VirtualKeyboardDismissAction = EVirtualKeyboardDismissAction::TextChangeOnDismiss;
 	Clipping = EWidgetClipping::ClipToBounds;
+	OverflowPolicy = ETextOverflowPolicy::Clip;
 
 #if WITH_EDITORONLY_DATA
 	AccessibleBehavior = ESlateAccessibleBehavior::Auto;
@@ -83,21 +84,22 @@ void UEditableText::ReleaseSlateResources(bool bReleaseChildren)
 
 TSharedRef<SWidget> UEditableText::RebuildWidget()
 {
-	MyEditableText = SNew( SEditableText )
-		.Style( &WidgetStyle )
-		.MinDesiredWidth( MinimumDesiredWidth )
-		.IsCaretMovedWhenGainFocus( IsCaretMovedWhenGainFocus )
-		.SelectAllTextWhenFocused( SelectAllTextWhenFocused )
-		.RevertTextOnEscape( RevertTextOnEscape )
-		.ClearKeyboardFocusOnCommit( ClearKeyboardFocusOnCommit )
-		.SelectAllTextOnCommit( SelectAllTextOnCommit )
-		.OnTextChanged( BIND_UOBJECT_DELEGATE( FOnTextChanged, HandleOnTextChanged ) )
-		.OnTextCommitted( BIND_UOBJECT_DELEGATE( FOnTextCommitted, HandleOnTextCommitted ) )
-		.VirtualKeyboardType( EVirtualKeyboardType::AsKeyboardType( KeyboardType.GetValue() ) )
+	MyEditableText = SNew(SEditableText)
+		.Style(&WidgetStyle)
+		.MinDesiredWidth(MinimumDesiredWidth)
+		.IsCaretMovedWhenGainFocus(IsCaretMovedWhenGainFocus)
+		.SelectAllTextWhenFocused(SelectAllTextWhenFocused)
+		.RevertTextOnEscape(RevertTextOnEscape)
+		.ClearKeyboardFocusOnCommit(ClearKeyboardFocusOnCommit)
+		.SelectAllTextOnCommit(SelectAllTextOnCommit)
+		.OnTextChanged(BIND_UOBJECT_DELEGATE(FOnTextChanged, HandleOnTextChanged))
+		.OnTextCommitted(BIND_UOBJECT_DELEGATE(FOnTextCommitted, HandleOnTextCommitted))
+		.VirtualKeyboardType(EVirtualKeyboardType::AsKeyboardType(KeyboardType.GetValue()))
 		.VirtualKeyboardOptions(VirtualKeyboardOptions)
 		.VirtualKeyboardTrigger(VirtualKeyboardTrigger)
 		.VirtualKeyboardDismissAction(VirtualKeyboardDismissAction)
-		.Justification( Justification );
+		.Justification(Justification)
+		.OverflowPolicy(OverflowPolicy);
 	
 	return MyEditableText.ToSharedRef();
 }
@@ -116,6 +118,8 @@ void UEditableText::SynchronizeProperties()
 	MyEditableText->SetAllowContextMenu(AllowContextMenu);
 	MyEditableText->SetVirtualKeyboardDismissAction(VirtualKeyboardDismissAction);
 	MyEditableText->SetJustification(Justification);
+	MyEditableText->SetOverflowPolicy(OverflowPolicy);
+
 	// TODO UMG Complete making all properties settable on SEditableText
 
 	ShapedTextOptions.SynchronizeShapedTextProperties(*MyEditableText);
@@ -176,6 +180,15 @@ void UEditableText::SetJustification(ETextJustify::Type InJustification)
 	}
 }
 
+void UEditableText::SetTextOverflowPolicy(ETextOverflowPolicy InOverflowPolicy)
+{
+	OverflowPolicy = InOverflowPolicy;
+	if (MyEditableText.IsValid())
+	{
+		MyEditableText->SetOverflowPolicy(InOverflowPolicy);
+	}
+}
+
 void UEditableText::SetClearKeyboardFocusOnCommit(bool bInClearKeyboardFocusOnCommit)
 {
 	ClearKeyboardFocusOnCommit = bInClearKeyboardFocusOnCommit;
diff --git a/Engine/Source/Runtime/UMG/Private/Components/EditableTextBox.cpp b/Engine/Source/Runtime/UMG/Private/Components/EditableTextBox.cpp
--- a/Engine/Source/Runtime/UMG/Private/Components/EditableTextBox.cpp
+++ b/Engine/Source/Runtime/UMG/Private/Components/EditableTextBox.cpp
@@ -42,6 +42,7 @@ UEditableTextBox::UEditableTextBox(const FObjectInitializer& ObjectInitializer)
 	SelectAllTextOnCommit = false;
 	AllowContextMenu = true;
 	VirtualKeyboardDismissAction = EVirtualKeyboardDismissAction::TextChangeOnDismiss;
+	OverflowPolicy = ETextOverflowPolicy::Clip;
 
 	if (DefaultEditableTextBoxStyle == nullptr)
 	{
@@ -98,7 +99,8 @@ TSharedRef<SWidget> UEditableTextBox::RebuildWidget()
 		.VirtualKeyboardOptions(VirtualKeyboardOptions)
 		.VirtualKeyboardTrigger(VirtualKeyboardTrigger)
 		.VirtualKeyboardDismissAction(VirtualKeyboardDismissAction)
-		.Justification(Justification);
+		.Justification(Justification)
+		.OverflowPolicy(OverflowPolicy);
 
 	return MyEditableTextBlock.ToSharedRef();
 }
@@ -124,6 +126,7 @@ void UEditableTextBox::SynchronizeProperties()
 	MyEditableTextBlock->SetAllowContextMenu(AllowContextMenu);
 	MyEditableTextBlock->SetVirtualKeyboardDismissAction(VirtualKeyboardDismissAction);
 	MyEditableTextBlock->SetJustification(Justification);
+	MyEditableTextBlock->SetOverflowPolicy(OverflowPolicy);
 
 	ShapedTextOptions.SynchronizeShapedTextProperties(*MyEditableTextBlock);
 }
@@ -218,6 +221,15 @@ void UEditableTextBox::SetJustification(ETextJustify::Type InJustification)
 	}
 }
 
+void UEditableTextBox::SetTextOverflowPolicy(ETextOverflowPolicy InOverflowPolicy)
+{
+	OverflowPolicy = InOverflowPolicy;
+	if (MyEditableTextBlock.IsValid())
+	{
+		MyEditableTextBlock->SetOverflowPolicy(InOverflowPolicy);
+	}
+}
+
 void UEditableTextBox::HandleOnTextChanged(const FText& InText)
 {
 	Text = InText;
diff --git a/Engine/Source/Runtime/UMG/Private/Components/TextBlock.cpp b/Engine/Source/Runtime/UMG/Private/Components/TextBlock.cpp
--- a/Engine/Source/Runtime/UMG/Private/Components/TextBlock.cpp
+++ b/Engine/Source/Runtime/UMG/Private/Components/TextBlock.cpp
@@ -145,6 +145,16 @@ void UTextBlock::SetTextTransformPolicy(ETextTransformPolicy InTransformPolicy)
 	}
 }
 
+void UTextBlock::SetTextOverflowPolicy(ETextOverflowPolicy InOverflowPolicy)
+{
+	TextOverflowPolicy = InOverflowPolicy;
+
+	if (MyTextBlock.IsValid())
+	{
+		MyTextBlock->SetOverflowPolicy(TextOverflowPolicy);
+	}
+}
+
 UMaterialInstanceDynamic* UTextBlock::GetDynamicFontMaterial()
 {
 	if (Font.FontMaterial)
@@ -288,6 +298,8 @@ void UTextBlock::SynchronizeProperties()
 		MyTextBlock->SetShadowColorAndOpacity( ShadowColorAndOpacityBinding );
 		MyTextBlock->SetMinDesiredWidth( MinDesiredWidth );
 		MyTextBlock->SetTransformPolicy( TextTransformPolicy );
+		MyTextBlock->SetOverflowPolicy(TextOverflowPolicy);
+
 		Super::SynchronizeTextLayoutProperties( *MyTextBlock );
 	}
 }
diff --git a/Engine/Source/Runtime/UMG/Public/Components/EditableText.h b/Engine/Source/Runtime/UMG/Public/Components/EditableText.h
--- a/Engine/Source/Runtime/UMG/Public/Components/EditableText.h
+++ b/Engine/Source/Runtime/UMG/Public/Components/EditableText.h
@@ -132,6 +132,10 @@ public:
 	UPROPERTY(EditAnywhere, BlueprintReadWrite, BlueprintSetter=SetJustification, Category=Appearance)
 	TEnumAsByte<ETextJustify::Type> Justification;
 
+	/** Sets what should happen when text is clipped because the block does not have enough space */
+	UPROPERTY(EditAnywhere, BlueprintReadWrite, BlueprintSetter=SetTextOverflowPolicy, Category=Appearance)
+	ETextOverflowPolicy OverflowPolicy;
+
 	/** Controls how the text within this widget should be shaped. */
 	UPROPERTY(EditAnywhere, BlueprintReadOnly, Category=Localization, AdvancedDisplay, meta=(ShowOnlyInnerProperties))
 	FShapedTextOptions ShapedTextOptions;
@@ -173,6 +177,9 @@ public:
 	UFUNCTION(BlueprintSetter)
 	void SetJustification(ETextJustify::Type InJustification);
 
+	UFUNCTION(BlueprintSetter)
+	void SetTextOverflowPolicy(ETextOverflowPolicy InOverflowPolicy);
+
 	void SetClearKeyboardFocusOnCommit(bool bInClearKeyboardFocusOnCommit);
 
 public:
diff --git a/Engine/Source/Runtime/UMG/Public/Components/EditableTextBox.h b/Engine/Source/Runtime/UMG/Public/Components/EditableTextBox.h
--- a/Engine/Source/Runtime/UMG/Public/Components/EditableTextBox.h
+++ b/Engine/Source/Runtime/UMG/Public/Components/EditableTextBox.h
@@ -135,6 +135,11 @@ public:
 	UPROPERTY(EditAnywhere, BlueprintReadWrite, BlueprintSetter = SetJustification, Category=Appearance)
 	TEnumAsByte<ETextJustify::Type> Justification;
 
+	/** Sets what should happen when text is clipped because the block does not have enough space */
+	UPROPERTY(EditAnywhere, BlueprintReadWrite, BlueprintSetter = SetTextOverflowPolicy, Category = Appearance)
+	ETextOverflowPolicy OverflowPolicy;
+
+
 	/** Controls how the text within this widget should be shaped. */
 	UPROPERTY(EditAnywhere, BlueprintReadOnly, Category=Localization, AdvancedDisplay, meta=(ShowOnlyInnerProperties))
 	FShapedTextOptions ShapedTextOptions;
@@ -183,6 +188,9 @@ public:
 	UFUNCTION(BlueprintSetter)
 	void SetJustification(ETextJustify::Type InJustification);
 
+	UFUNCTION(BlueprintSetter)
+	void SetTextOverflowPolicy(ETextOverflowPolicy InOverflowPolicy);
+
 	UFUNCTION(BlueprintCallable, Category="Widget", meta=(DisplayName="SetForegroundColor (Text Box)"))
 	void SetForegroundColor(FLinearColor color);
 
diff --git a/Engine/Source/Runtime/UMG/Public/Components/TextBlock.h b/Engine/Source/Runtime/UMG/Public/Components/TextBlock.h
--- a/Engine/Source/Runtime/UMG/Public/Components/TextBlock.h
+++ b/Engine/Source/Runtime/UMG/Public/Components/TextBlock.h
@@ -97,6 +97,9 @@ public:
 	UFUNCTION(BlueprintCallable, Category = "Appearance")
 	void SetTextTransformPolicy(ETextTransformPolicy InTransformPolicy);
 
+	UFUNCTION(BlueprintCallable, Category = "Appearance")
+	void SetTextOverflowPolicy(ETextOverflowPolicy InOverflowPolicy);
+
 	/**  */
 	UFUNCTION(BlueprintCallable, Category="Appearance")
 	UMaterialInstanceDynamic* GetDynamicFontMaterial();
@@ -158,6 +161,10 @@ public:
 	UPROPERTY(EditAnywhere, BlueprintReadOnly, Category=Appearance, meta=(DisplayName="Transform Policy"))
 	ETextTransformPolicy TextTransformPolicy;
 
+	/** Sets what happens to text that is clipped and doesn't fit within the clip rect for this widget */
+	UPROPERTY(EditAnywhere, BlueprintReadWrite, Category = Appearance, AdvancedDisplay, meta = (DisplayName = "Overflow Policy"))
+	ETextOverflowPolicy TextOverflowPolicy;
+
 	/** 
 	 * Gets the widget text
 	 * @return The widget text
