Fix static analysis warnings

diff --git a/Engine/Source/Runtime/SlateCore/Private/Rendering/ElementBatcher.cpp b/Engine/Source/Runtime/SlateCore/Private/Rendering/ElementBatcher.cpp
--- a/Engine/Source/Runtime/SlateCore/Private/Rendering/ElementBatcher.cpp
+++ b/Engine/Source/Runtime/SlateCore/Private/Rendering/ElementBatcher.cpp
@@ -2756,7 +2756,7 @@ void FSlateElementBatcher::BuildShapedTextSequence(const FShapedTextBuildContext
 			bCanRenderGlyph = !bClipped;
 		}
 
-		if(bCanRenderGlyph)
+		if(bCanRenderGlyph && RenderBatch)
 		{
 			const FVector2D UpperLeft(X, Y);
 			const FVector2D UpperRight(X + SizeX, Y);
