Slate: Change the inner type of FChildren::FWidgetRef::WidgetReference to prevent implicit copy of the argument in unoptimized build.

diff --git a/Engine/Source/Runtime/SlateCore/Public/Layout/Children.h b/Engine/Source/Runtime/SlateCore/Public/Layout/Children.h
--- a/Engine/Source/Runtime/SlateCore/Public/Layout/Children.h
+++ b/Engine/Source/Runtime/SlateCore/Public/Layout/Children.h
@@ -39,12 +39,36 @@ public:
 
 	virtual TSharedRef<SWidget> GetChildAt(int32 Index) override
 	{
-		return GetChildRefAt(Index).GetWidgetRef();
+		int32 TotalNum = 0;
+		for (FChildren* Children : LinkedChildren)
+		{
+			const int32 NewTotal = TotalNum + Children->Num();
+			if (NewTotal > Index)
+			{
+				return Children->GetChildAt(Index - TotalNum);
+			}
+			TotalNum = NewTotal;
+		}
+		// This result should never occur users should always access a valid index for child slots.
+		check(false);
+		return SNullWidget::NullWidget;
 	}
 
 	virtual TSharedRef<const SWidget> GetChildAt(int32 Index) const override
 	{
-		return GetChildRefAt(Index).GetWidgetRef();
+		int32 TotalNum = 0;
+		for (const FChildren* Children : LinkedChildren)
+		{
+			const int32 NewTotal = TotalNum + Children->Num();
+			if (NewTotal > Index)
+			{
+				return Children->GetChildAt(Index - TotalNum);
+			}
+			TotalNum = NewTotal;
+		}
+		// This result should never occur users should always access a valid index for child slots.
+		check(false);
+		return SNullWidget::NullWidget;
 	}
 
 protected:
@@ -92,7 +116,7 @@ protected:
 		}
 		// This result should never occur users should always access a valid index for child slots.
 		check(false);
-		return FWidgetRef(ReferenceConstruct, SNullWidget::NullWidget);
+		return FWidgetRef(ReferenceConstruct, SNullWidget::NullWidget.Get());
 	}
 
 	virtual FConstWidgetRef GetChildRefAt(int32 Index) const override
@@ -109,7 +133,7 @@ protected:
 		}
 		// This result should never occur users should always access a valid index for child slots.
 		check(false);
-		return FConstWidgetRef(ReferenceConstruct, SNullWidget::NullWidget);
+		return FConstWidgetRef(ReferenceConstruct, SNullWidget::NullWidget.Get());
 	}
 
 protected:
@@ -170,13 +194,13 @@ private:
 	virtual FWidgetRef GetChildRefAt(int32 Index) override
 	{
 		check(false);
-		return FWidgetRef(ReferenceConstruct, SNullWidget::NullWidget);
+		return FWidgetRef(ReferenceConstruct, SNullWidget::NullWidget.Get());
 	}
 	
 	virtual FConstWidgetRef GetChildRefAt(int32 Index) const override
 	{
 		check(false);
-		return FConstWidgetRef(ReferenceConstruct, SNullWidget::NullWidget);
+		return FConstWidgetRef(ReferenceConstruct, SNullWidget::NullWidget.Get());
 	}
 };
 
@@ -221,12 +245,12 @@ private:
 	virtual FWidgetRef GetChildRefAt(int32 ChildIndex) override
 	{
 		check(ChildIndex == 0);
-		return FWidgetRef(ReferenceConstruct, FSlotBase::GetWidget());
+		return FWidgetRef(ReferenceConstruct, FSlotBase::GetWidget().Get());
 	}
 	virtual FConstWidgetRef GetChildRefAt(int32 ChildIndex) const override
 	{
 		check(ChildIndex == 0);
-		return FConstWidgetRef(ReferenceConstruct, FSlotBase::GetWidget());
+		return FConstWidgetRef(ReferenceConstruct, FSlotBase::GetWidget().Get());
 	}
 };
 
@@ -276,13 +300,13 @@ private:
 	{
 		check(ChildIndex == 0);
 		TSharedPtr<SWidget> Widget = WidgetPtr.Pin();
-		return (Widget.IsValid()) ? FWidgetRef(CopyConstruct, Widget.ToSharedRef()) : FWidgetRef(ReferenceConstruct, SNullWidget::NullWidget);
+		return (Widget.IsValid()) ? FWidgetRef(CopyConstruct, Widget.ToSharedRef()) : FWidgetRef(ReferenceConstruct, SNullWidget::NullWidget.Get());
 	}
 	virtual FConstWidgetRef GetChildRefAt(int32 ChildIndex) const override
 	{
 		check(ChildIndex == 0);
 		TSharedPtr<SWidget> Widget = WidgetPtr.Pin();
-		return (Widget.IsValid()) ? FConstWidgetRef(CopyConstruct, Widget.ToSharedRef()) : FConstWidgetRef(ReferenceConstruct, SNullWidget::NullWidget);
+		return (Widget.IsValid()) ? FConstWidgetRef(CopyConstruct, Widget.ToSharedRef()) : FConstWidgetRef(ReferenceConstruct, SNullWidget::NullWidget.Get());
 	}
 
 public:
@@ -424,12 +448,12 @@ private:
 	virtual FWidgetRef GetChildRefAt(int32 ChildIndex) override
 	{
 		check(ChildIndex == 0);
-		return FWidgetRef(ReferenceConstruct, this->GetWidget());
+		return FWidgetRef(ReferenceConstruct, this->GetWidget().Get());
 	}
 	virtual FConstWidgetRef GetChildRefAt(int32 ChildIndex) const override
 	{
 		check(ChildIndex == 0);
-		return FConstWidgetRef(ReferenceConstruct, this->GetWidget());
+		return FConstWidgetRef(ReferenceConstruct, this->GetWidget().Get());
 	}
 };
 
@@ -521,12 +545,12 @@ protected:
 
 	virtual FWidgetRef GetChildRefAt(int32 ChildIndex) override
 	{
-		return FWidgetRef(ReferenceConstruct, Children[ChildIndex]->GetWidget());
+		return FWidgetRef(ReferenceConstruct, Children[ChildIndex]->GetWidget().Get());
 	}
 
 	virtual FConstWidgetRef GetChildRefAt(int32 ChildIndex) const override
 	{
-		return FConstWidgetRef(ReferenceConstruct, Children[ChildIndex]->GetWidget());
+		return FConstWidgetRef(ReferenceConstruct, Children[ChildIndex]->GetWidget().Get());
 	}
 
 public:
@@ -920,12 +944,12 @@ private:
 
 	virtual FWidgetRef GetChildRefAt(int32 ChildIndex) override
 	{
-		return FWidgetRef(ReferenceConstruct, Children[ChildIndex]);
+		return FWidgetRef(ReferenceConstruct, Children[ChildIndex].Get());
 	}
 	
 	virtual FConstWidgetRef GetChildRefAt(int32 ChildIndex) const override
 	{
-		return FConstWidgetRef(ReferenceConstruct, Children[ChildIndex]);
+		return FConstWidgetRef(ReferenceConstruct, Children[ChildIndex].Get());
 	}
 
 public:
diff --git a/Engine/Source/Runtime/SlateCore/Public/Layout/ChildrenBase.h b/Engine/Source/Runtime/SlateCore/Public/Layout/ChildrenBase.h
--- a/Engine/Source/Runtime/SlateCore/Public/Layout/ChildrenBase.h
+++ b/Engine/Source/Runtime/SlateCore/Public/Layout/ChildrenBase.h
@@ -104,34 +104,30 @@ protected:
 	{
 	private:
 		TOptional<TSharedRef<SWidget>> WidgetCopy;
-		const TSharedRef<SWidget>& WidgetReference;
+		SWidget& WidgetReference;
 
 	public:
 		FWidgetRef(ECopyConstruct, TSharedRef<SWidget> InWidgetCopy)
 			: WidgetCopy(MoveTemp(InWidgetCopy))
-			, WidgetReference(WidgetCopy.GetValue())
+			, WidgetReference(WidgetCopy.GetValue().Get())
 		{}
-		FWidgetRef(ERefConstruct, const TSharedRef<SWidget>& InWidgetRef)
+		FWidgetRef(ERefConstruct, SWidget& InWidgetRef)
 			: WidgetCopy()
 			, WidgetReference(InWidgetRef)
 		{}
 		FWidgetRef(const FWidgetRef& Other)
 			: WidgetCopy(Other.WidgetCopy)
-			, WidgetReference(WidgetCopy.IsSet() ? WidgetCopy.GetValue() : Other.WidgetReference)
+			, WidgetReference(WidgetCopy.IsSet() ? WidgetCopy.GetValue().Get() : Other.WidgetReference)
 		{}
 		FWidgetRef(FWidgetRef&& Other)
 			: WidgetCopy(MoveTemp(Other.WidgetCopy))
-			, WidgetReference(WidgetCopy.IsSet() ? WidgetCopy.GetValue() : Other.WidgetReference)
+			, WidgetReference(WidgetCopy.IsSet() ? WidgetCopy.GetValue().Get() : Other.WidgetReference)
 		{}
 		FWidgetRef& operator=(const FWidgetRef&) = delete;
 		FWidgetRef& operator=(FWidgetRef&&) = delete;
-		const TSharedRef<SWidget>& GetWidgetRef() const
-		{
-			return WidgetReference;
-		}
 		SWidget& GetWidget() const
 		{
-			return WidgetReference.Get();
+			return WidgetReference;
 		}
 	};
 
@@ -139,34 +135,30 @@ protected:
 	{
 	private:
 		TOptional<TSharedRef<const SWidget>> WidgetCopy;
-		const TSharedRef<const SWidget>& WidgetReference;
+		const SWidget& WidgetReference;
 
 	public:
 		FConstWidgetRef(ECopyConstruct, TSharedRef<const SWidget> InWidgetCopy)
 			: WidgetCopy(MoveTemp(InWidgetCopy))
-			, WidgetReference(WidgetCopy.GetValue())
+			, WidgetReference(WidgetCopy.GetValue().Get())
 		{}
-		FConstWidgetRef(ERefConstruct, const TSharedRef<const SWidget>& InWidgetRef)
+		FConstWidgetRef(ERefConstruct, const SWidget& InWidgetRef)
 			: WidgetCopy()
 			, WidgetReference(InWidgetRef)
 		{}
 		FConstWidgetRef(const FConstWidgetRef& Other)
 			: WidgetCopy(Other.WidgetCopy)
-			, WidgetReference(WidgetCopy.IsSet() ? WidgetCopy.GetValue() : Other.WidgetReference)
+			, WidgetReference(WidgetCopy.IsSet() ? WidgetCopy.GetValue().Get() : Other.WidgetReference)
 		{}
 		FConstWidgetRef(FConstWidgetRef&& Other)
 			: WidgetCopy(MoveTemp(Other.WidgetCopy))
-			, WidgetReference(WidgetCopy.IsSet() ? WidgetCopy.GetValue() : Other.WidgetReference)
+			, WidgetReference(WidgetCopy.IsSet() ? WidgetCopy.GetValue().Get() : Other.WidgetReference)
 		{}
 		FConstWidgetRef& operator=(const FConstWidgetRef&) = delete;
 		FConstWidgetRef& operator=(FConstWidgetRef&&) = delete;
-		const TSharedRef<const SWidget>& GetWidgetRef() const
-		{
-			return WidgetReference;
-		}
 		const SWidget& GetWidget() const
 		{
-			return WidgetReference.Get();
+			return WidgetReference;
 		}
 	};
 
