Fixed STextBlock not updating when some FTextBlockStyle properties were changed

StrikeBrush, UnderlineBrush, TransformPolicy, and OverflowPolicy were being ignored

diff --git a/Engine/Source/Runtime/Slate/Private/Widgets/Text/SlateTextBlockLayout.cpp b/Engine/Source/Runtime/Slate/Private/Widgets/Text/SlateTextBlockLayout.cpp
--- a/Engine/Source/Runtime/Slate/Private/Widgets/Text/SlateTextBlockLayout.cpp
+++ b/Engine/Source/Runtime/Slate/Private/Widgets/Text/SlateTextBlockLayout.cpp
@@ -334,14 +334,7 @@ void FSlateTextBlockLayout::UpdateTextHighlights(const FText& InHighlightText)
 bool FSlateTextBlockLayout::IsStyleUpToDate(const FTextBlockStyle& NewStyle) const
 {
 	const FTextBlockStyle& CurrentStyle = TextLayout->GetDefaultTextStyle();
-
-	return (CurrentStyle.Font.IsIdenticalTo(NewStyle.Font)) 
-		&& (CurrentStyle.ColorAndOpacity == NewStyle.ColorAndOpacity)
-		&& (CurrentStyle.ShadowOffset == NewStyle.ShadowOffset)
-		&& (CurrentStyle.ShadowColorAndOpacity == NewStyle.ShadowColorAndOpacity)
-		&& (CurrentStyle.SelectedBackgroundColor == NewStyle.SelectedBackgroundColor)
-		&& (CurrentStyle.HighlightColor == NewStyle.HighlightColor)
-		&& (CurrentStyle.HighlightShape == NewStyle.HighlightShape);
+	return CurrentStyle.IsIdenticalTo(NewStyle);
 }
 
 float FSlateTextBlockLayout::CalculateWrappingWidth() const
diff --git a/Engine/Source/Runtime/SlateCore/Public/Styling/SlateTypes.h b/Engine/Source/Runtime/SlateCore/Public/Styling/SlateTypes.h
--- a/Engine/Source/Runtime/SlateCore/Public/Styling/SlateTypes.h
+++ b/Engine/Source/Runtime/SlateCore/Public/Styling/SlateTypes.h
@@ -380,6 +380,24 @@ struct SLATECORE_API FTextBlockStyle : public FSlateWidgetStyle
 	ETextOverflowPolicy OverflowPolicy;
 	FTextBlockStyle& SetOverflowPolicy(const ETextOverflowPolicy& InOverflowPolicy) { OverflowPolicy = InOverflowPolicy; return *this; }
 
+	/**
+	 * Checks to see whether this style is identical to another.
+	 */
+	bool IsIdenticalTo(const FTextBlockStyle& InOther) const
+	{
+		return Font.IsIdenticalTo(InOther.Font)
+			&& ColorAndOpacity == InOther.ColorAndOpacity
+			&& ShadowOffset == InOther.ShadowOffset
+			&& ShadowColorAndOpacity == InOther.ShadowColorAndOpacity
+			&& SelectedBackgroundColor == InOther.SelectedBackgroundColor
+			&& HighlightColor == InOther.HighlightColor
+			&& HighlightShape == InOther.HighlightShape
+			&& StrikeBrush == InOther.StrikeBrush
+			&& UnderlineBrush == InOther.UnderlineBrush
+			&& TransformPolicy == InOther.TransformPolicy
+			&& OverflowPolicy == InOther.OverflowPolicy;
+	}
+
 	/**
 	 * Unlinks all colors in this style.
 	 * @see FSlateColor::Unlink
