Major Tabs reskin - Removed vertical separaters, changed inter-tab spacing and added color backing to asset tab icons.

diff --git a/Engine/Source/Runtime/Slate/Private/Framework/Docking/SDockingTabWell.cpp b/Engine/Source/Runtime/Slate/Private/Framework/Docking/SDockingTabWell.cpp
--- a/Engine/Source/Runtime/Slate/Private/Framework/Docking/SDockingTabWell.cpp
+++ b/Engine/Source/Runtime/Slate/Private/Framework/Docking/SDockingTabWell.cpp
@@ -28,7 +28,7 @@ void SDockingTabWell::Construct( const FArguments& InArgs )
 	ChildBeingDraggedOffset = 0.0f;
 	TabGrabOffsetFraction = FVector2D::ZeroVector;
 		
-	SeparatorBrush = FAppStyle::Get().GetBrush("Docking.Tab.InactiveTabSeparator");
+	SeparatorBrush = nullptr; // No separater between tabs
 
 	// We need a valid parent here. TabPanels must exist in a SDockingNode
 	check( InArgs._ParentStackNode.Get().IsValid() );
diff --git a/Engine/Source/Runtime/Slate/Private/Widgets/Docking/SDockTab.cpp b/Engine/Source/Runtime/Slate/Private/Widgets/Docking/SDockTab.cpp
--- a/Engine/Source/Runtime/Slate/Private/Widgets/Docking/SDockTab.cpp
+++ b/Engine/Source/Runtime/Slate/Private/Widgets/Docking/SDockTab.cpp
@@ -619,7 +619,7 @@ void SDockTab::Construct( const FArguments& InArgs )
 				[
 					SNew(SBorder)
 					// Don't allow active tab overlay to absorb mouse clicks
-					.Padding(1.0f)
+					.Padding(this, &SDockTab::GetTabIconBorderPadding)
 					.Visibility(EVisibility::HitTestInvisible)
 					// Overlay for color-coded tab effect
 					.BorderImage(this, &SDockTab::GetColorOverlayImageBrush)
@@ -772,6 +772,11 @@ FMargin SDockTab::GetTabPadding() const
 	return NewPadding;
 }
 
+FMargin SDockTab::GetTabIconBorderPadding() const
+{
+	return FMargin(GetCurrentStyle().IconBorderPadding);
+}
+
 const FSlateBrush* SDockTab::GetColorOverlayImageBrush() const
 {
 	if (this->TabColorScale.Get().A > 0.0f)
@@ -788,7 +793,21 @@ EVisibility SDockTab::GetActiveTabIndicatorVisibility() const
 
 FSlateColor SDockTab::GetTabColor() const
 {
-	return TabColorScale.Get();
+	FLinearColor BaseColor = TabColorScale.Get();
+	
+	FSlateColor FinalColor;
+
+	// 50% Black Icon if the tab is active, slightly dimmed to 60% black otherwise
+	if (this->IsForeground() || this->IsHovered())
+	{
+		FinalColor = BaseColor.CopyWithNewOpacity(0.5);
+	}
+	else
+	{
+		FinalColor = BaseColor.CopyWithNewOpacity(0.4);
+	}
+
+	return FinalColor;
 }
 
 
diff --git a/Engine/Source/Runtime/Slate/Public/Widgets/Docking/SDockTab.h b/Engine/Source/Runtime/Slate/Public/Widgets/Docking/SDockTab.h
--- a/Engine/Source/Runtime/Slate/Public/Widgets/Docking/SDockTab.h
+++ b/Engine/Source/Runtime/Slate/Public/Widgets/Docking/SDockTab.h
@@ -365,6 +365,9 @@ protected:
 	/** @return the size the tab icon should be */
 	TOptional<FVector2D> GetTabIconSize() const;
 
+	/** @return the padding for the tab icon border */
+	FMargin GetTabIconBorderPadding() const;
+
 private:
 	/** Activates the tab in its tab well */
 	EActiveTimerReturnType TriggerActivateTab( double InCurrentTime, float InDeltaTime );
diff --git a/Engine/Source/Runtime/SlateCore/Private/Styling/SlateTypes.cpp b/Engine/Source/Runtime/SlateCore/Private/Styling/SlateTypes.cpp
--- a/Engine/Source/Runtime/SlateCore/Private/Styling/SlateTypes.cpp
+++ b/Engine/Source/Runtime/SlateCore/Private/Styling/SlateTypes.cpp
@@ -623,6 +623,7 @@ FDockTabStyle::FDockTabStyle()
 	, HoveredForegroundColor(FSlateColor::UseForeground())
 	, ActiveForegroundColor(FSlateColor::UseForeground())
 	, ForegroundForegroundColor(FSlateColor::UseForeground())
+	, IconBorderPadding(1.f)
 {
 }
 
diff --git a/Engine/Source/Runtime/SlateCore/Private/Styling/StarshipCoreStyle.cpp b/Engine/Source/Runtime/SlateCore/Private/Styling/StarshipCoreStyle.cpp
--- a/Engine/Source/Runtime/SlateCore/Private/Styling/StarshipCoreStyle.cpp
+++ b/Engine/Source/Runtime/SlateCore/Private/Styling/StarshipCoreStyle.cpp
@@ -1576,24 +1576,24 @@ void FStarshipCoreStyle::SetupDockingStyles(TSharedRef<FStyle>& Style)
 		FDockTabStyle()
 		.SetCloseButtonStyle(CloseButton)
 		.SetNormalBrush(FSlateNoResource())
-		.SetHoveredBrush(BOX_BRUSH("/Starship/Docking/DockTab_Hover", 4.0f / 20.0f, FStyleColors::Panel))
+		.SetHoveredBrush(BOX_BRUSH("/Starship/Docking/DockTab_Foreground", 4.0f / 20.0f, FStyleColors::Panel.GetSpecifiedColor().CopyWithNewOpacity(0.8))) // Panel + 20% Black
 		.SetForegroundBrush(BOX_BRUSH("/Starship/Docking/DockTab_Foreground", 4.0f / 20.0f, FStyleColors::Panel))
-
 		.SetColorOverlayTabBrush(FSlateNoResource())
 		.SetColorOverlayIconBrush(FSlateNoResource())
 		.SetContentAreaBrush(FSlateColorBrush(FStyleColors::Panel))
 		.SetTabWellBrush(FSlateColorBrush(FStyleColors::Background))
 		.SetFlashColor(TabFlashColor)
 
-		.SetTabPadding(FMargin(10.f, 3.f, 10.f, 4.f))
-		.SetOverlapWidth(0.0f)
+		.SetTabPadding(FMargin(4.f, 3.f, 10.f, 4.f))
+		.SetOverlapWidth(-2.0f) // Negative Overlap Width to add space between tabs
 
 		.SetNormalForegroundColor(FStyleColors::Foreground)
-		.SetActiveForegroundColor(FStyleColors::ForegroundHover)
-		.SetForegroundForegroundColor(FStyleColors::Foreground)
-		.SetHoveredForegroundColor(FStyleColors::ForegroundHover)
+		.SetActiveForegroundColor(FStyleColors::White)
+		.SetForegroundForegroundColor(FStyleColors::White)
+		.SetHoveredForegroundColor(FStyleColors::White)
 		.SetTabTextStyle(NormalText);
 
+
 	// Panel Tab
 	Style->Set("Docking.Tab", MinorTabStyle);
 
@@ -1601,23 +1601,24 @@ void FStarshipCoreStyle::SetupDockingStyles(TSharedRef<FStyle>& Style)
 	Style->Set("Docking.MajorTab", FDockTabStyle()
 		.SetCloseButtonStyle(CloseButton)
 		.SetNormalBrush(FSlateNoResource())
-		.SetHoveredBrush(BOX_BRUSH("/Starship/Docking/DockTab_Hover", 4.0f / 20.0f, FStyleColors::Panel))
+		.SetHoveredBrush(BOX_BRUSH("/Starship/Docking/DockTab_Foreground", 4.0f / 20.0f, FStyleColors::Panel.GetSpecifiedColor().CopyWithNewOpacity(0.8))) // Panel + 20% Black
 		.SetForegroundBrush(BOX_BRUSH("/Starship/Docking/DockTab_Foreground", 4.0f / 20.0f, FStyleColors::Panel))
 
 		.SetColorOverlayTabBrush(FSlateNoResource())
-		.SetColorOverlayIconBrush(FSlateNoResource())
+		.SetColorOverlayIconBrush(FSlateRoundedBoxBrush(FStyleColors::Foreground, 4.f))
 		.SetContentAreaBrush(FSlateColorBrush(FStyleColors::Background))
 		.SetTabWellBrush(FSlateColorBrush(FStyleColors::Background))
 
-		.SetTabPadding(FMargin(10.f, 7.f, 10.f, 8.f))
-		.SetOverlapWidth(0.f)
+		.SetTabPadding(FMargin(4.f, 7.f, 10.f, 8.f))
+		.SetOverlapWidth(-2.0f) // Negative Overlap Width to add space between tabs
 		.SetFlashColor(TabFlashColor)
 
 		.SetNormalForegroundColor(FStyleColors::Foreground)
-		.SetActiveForegroundColor(FStyleColors::ForegroundHover)
-		.SetForegroundForegroundColor(FStyleColors::Foreground)
-		.SetHoveredForegroundColor(FStyleColors::ForegroundHover)
+		.SetActiveForegroundColor(FStyleColors::White)
+		.SetForegroundForegroundColor(FStyleColors::White)
+		.SetHoveredForegroundColor(FStyleColors::White)
 		.SetTabTextStyle(NormalText)
+		.SetIconBorderPadding(4.f)
 	);
 
 	Style->Set("Docking.Tab.ContentAreaBrush", new FSlateNoResource());
diff --git a/Engine/Source/Runtime/SlateCore/Public/Styling/SlateTypes.h b/Engine/Source/Runtime/SlateCore/Public/Styling/SlateTypes.h
--- a/Engine/Source/Runtime/SlateCore/Public/Styling/SlateTypes.h
+++ b/Engine/Source/Runtime/SlateCore/Public/Styling/SlateTypes.h
@@ -1843,6 +1843,12 @@ struct SLATECORE_API FDockTabStyle : public FSlateWidgetStyle
 	UPROPERTY(EditAnywhere, AdvancedDisplay, Category=Appearance)
 	FSlateColor ForegroundForegroundColor;
 	FDockTabStyle& SetForegroundForegroundColor( const FSlateColor& InForegroundForegroundColor ){ ForegroundForegroundColor = InForegroundForegroundColor; return *this; }
+
+	/** The padding applied to the border around the tab icon */
+	UPROPERTY(EditAnywhere, AdvancedDisplay, Category = Appearance)
+	float IconBorderPadding;
+	FDockTabStyle& SetIconBorderPadding(const float InIconBorderPadding) { IconBorderPadding = InIconBorderPadding; return *this; }
+
 };
 
 
