Fixed minor tabs being incorrectly colored (similar to the major tabs reskin)

diff --git a/Engine/Source/Runtime/SlateCore/Private/Styling/StarshipCoreStyle.cpp b/Engine/Source/Runtime/SlateCore/Private/Styling/StarshipCoreStyle.cpp
--- a/Engine/Source/Runtime/SlateCore/Private/Styling/StarshipCoreStyle.cpp
+++ b/Engine/Source/Runtime/SlateCore/Private/Styling/StarshipCoreStyle.cpp
@@ -1588,9 +1588,9 @@ void FStarshipCoreStyle::SetupDockingStyles(TSharedRef<FStyle>& Style)
 		.SetOverlapWidth(-2.0f) // Negative Overlap Width to add space between tabs
 
 		.SetNormalForegroundColor(FStyleColors::Foreground)
-		.SetActiveForegroundColor(FStyleColors::White)
-		.SetForegroundForegroundColor(FStyleColors::White)
-		.SetHoveredForegroundColor(FStyleColors::White)
+		.SetActiveForegroundColor(FStyleColors::ForegroundHover)
+		.SetForegroundForegroundColor(FStyleColors::Foreground)
+		.SetHoveredForegroundColor(FStyleColors::ForegroundHover)
 		.SetTabTextStyle(NormalText);
 
 
