Remove Margin clamping

diff --git a/Engine/Source/Runtime/SlateCore/Private/Rendering/ElementBatcher.cpp b/Engine/Source/Runtime/SlateCore/Private/Rendering/ElementBatcher.cpp
--- a/Engine/Source/Runtime/SlateCore/Private/Rendering/ElementBatcher.cpp
+++ b/Engine/Source/Runtime/SlateCore/Private/Rendering/ElementBatcher.cpp
@@ -2247,10 +2247,10 @@ void FSlateElementBatcher::AddBorderElement( const FSlateDrawElement& DrawElemen
 	// Account for negative sizes
 	bool bIsFlippedX = TopLeft.X > BotRight.X;
 	bool bIsFlippedY = TopLeft.Y > BotRight.Y;
-	Margin.Left = FMath::Clamp(bIsFlippedX ? -Margin.Left : Margin.Left, -1.0f, 1.0f);
-	Margin.Top = FMath::Clamp(bIsFlippedY ? -Margin.Top : Margin.Top, -1.0f, 1.0f);
-	Margin.Right = FMath::Clamp(bIsFlippedX ? -Margin.Right : Margin.Right, -1.0f, 1.0f);
-	Margin.Bottom = FMath::Clamp(bIsFlippedY ? -Margin.Bottom : Margin.Bottom, -1.0f, 1.0f);
+	Margin.Left = bIsFlippedX ? -Margin.Left : Margin.Left;
+	Margin.Top = bIsFlippedY ? -Margin.Top : Margin.Top;
+	Margin.Right = bIsFlippedX ? -Margin.Right : Margin.Right;
+	Margin.Bottom = bIsFlippedY ? -Margin.Bottom : Margin.Bottom;
 
 	// Determine the margins for each quad
 	FVector2D TopLeftMargin(TextureSizeLocalSpace * FVector2D(Margin.Left, Margin.Top));
