Fixed an issue causing major tabs to be larger than intended due to padding issues.

diff --git a/Engine/Source/Runtime/SlateCore/Private/Styling/StarshipCoreStyle.cpp b/Engine/Source/Runtime/SlateCore/Private/Styling/StarshipCoreStyle.cpp
--- a/Engine/Source/Runtime/SlateCore/Private/Styling/StarshipCoreStyle.cpp
+++ b/Engine/Source/Runtime/SlateCore/Private/Styling/StarshipCoreStyle.cpp
@@ -1610,7 +1610,7 @@ void FStarshipCoreStyle::SetupDockingStyles(TSharedRef<FStyle>& Style)
 		.SetContentAreaBrush(FSlateColorBrush(FStyleColors::Background))
 		.SetTabWellBrush(FSlateColorBrush(FStyleColors::Background))
 
-		.SetTabPadding(FMargin(4.f, 7.f, 10.f, 8.f))
+		.SetTabPadding(FMargin(4.f, 4.f, 10.f, 4.f))
 		.SetOverlapWidth(-2.0f) // Negative Overlap Width to add space between tabs
 		.SetFlashColor(TabFlashColor)
 
