Makes the font size for default SToolTips consistent with SDocumentationTooltips.

diff --git a/Engine/Source/Runtime/SlateCore/Private/Styling/StarshipCoreStyle.cpp b/Engine/Source/Runtime/SlateCore/Private/Styling/StarshipCoreStyle.cpp
--- a/Engine/Source/Runtime/SlateCore/Private/Styling/StarshipCoreStyle.cpp
+++ b/Engine/Source/Runtime/SlateCore/Private/Styling/StarshipCoreStyle.cpp
@@ -445,7 +445,7 @@ TSharedRef<ISlateStyle> FStarshipCoreStyle::Create()
 
 	// SToolTip defaults...
 	{
-		Style->Set("ToolTip.Font", StyleFonts.Small);
+		Style->Set("ToolTip.Font", DEFAULT_FONT("Regular", 9));
 		Style->Set("ToolTip.Background", new FSlateRoundedBoxBrush(FStyleColors::Dropdown, 0.0f, FStyleColors::DropdownOutline, 1.0f));
 
 		Style->Set("ToolTip.LargerFont", StyleFonts.Normal);
