Stage one of editor style recolor

- Buttons, checkboxes, menus, toolbars, trees, combos, and text boxes all have consistent sizing and coloring
- Invalidated current dark theme

#pf https://horde.devtools.epicgames.com/job/612d40b9423a8f00013c5b92

diff --git a/Engine/Source/Editor/ContentBrowser/Private/SContentBrowser.cpp b/Engine/Source/Editor/ContentBrowser/Private/SContentBrowser.cpp
--- a/Engine/Source/Editor/ContentBrowser/Private/SContentBrowser.cpp
+++ b/Engine/Source/Editor/ContentBrowser/Private/SContentBrowser.cpp
@@ -202,7 +202,6 @@ void SContentBrowser::Construct( const FArguments& InArgs, const FName& InInstan
 	{
 		ViewOptions =
 			SNew(SComboButton)
-			.ContentPadding(0)
 			.ComboButtonStyle(&FAppStyle::Get().GetWidgetStyle<FComboButtonStyle>("SimpleComboButton"))
 			.OnGetMenuContent(AssetViewPtr.ToSharedRef(), &SAssetView::GetViewButtonContent)
 			.HasDownArrow(false)
@@ -344,6 +343,7 @@ void SContentBrowser::Construct( const FArguments& InArgs, const FName& InInstan
 				.AutoWidth()
 				.Padding(5.0f, 0.0f, 0.0f, 0.0f)
 				.HAlign(HAlign_Right)
+				.VAlign(VAlign_Center)
 				[
 					ViewOptions
 				]
@@ -646,7 +646,6 @@ TSharedRef<SWidget> SContentBrowser::CreateLockButton(const FContentBrowserConfi
 			SNew(SButton)
 			.ButtonStyle(FAppStyle::Get(), "SimpleButton")
 			.ToolTipText(LOCTEXT("LockToggleTooltip", "Toggle lock. If locked, this browser will ignore Find in Content Browser requests."))
-			.ContentPadding(FMargin(1, 0))
 			.OnClicked(this, &SContentBrowser::ToggleLockClicked)
 			.AddMetaData<FTagMetaData>(FTagMetaData(TEXT("ContentBrowserLock")))
 			.Visibility(this, &SContentBrowser::GetLockButtonVisibility)
@@ -717,7 +716,7 @@ TSharedRef<SWidget> SContentBrowser::CreateAssetView(const FContentBrowserConfig
 				.Padding(5, 0, 0, 0)
 				[
 					SNew(SComboButton)
-					.ComboButtonStyle(&FAppStyle::Get().GetWidgetStyle<FComboButtonStyle>("SimpleComboButton"))
+					.ComboButtonStyle(&FAppStyle::Get().GetWidgetStyle<FComboButtonStyle>("SimpleComboButtonWithIcon"))
 					.ForegroundColor(FSlateColor::UseStyle())
 					.ToolTipText(LOCTEXT("AddFilterToolTip", "Add an asset filter."))
 					.OnGetMenuContent(this, &SContentBrowser::MakeAddFilterMenu)
@@ -888,7 +887,6 @@ TSharedRef<SWidget> SContentBrowser::CreateDrawerDockButton(const FContentBrowse
 			SNew(SButton)
 			.ButtonStyle(FAppStyle::Get(), "SimpleButton")
 			.ToolTipText(LOCTEXT("DockInLayout_Tooltip", "Docks this content browser in the current layout, copying all settings from the drawer.\nThe drawer will still be usable as a temporary browser."))
-			.ContentPadding(FMargin(1, 0))
 			.OnClicked(this, &SContentBrowser::DockInLayoutClicked)
 			[
 				SNew(SHorizontalBox)
diff --git a/Engine/Source/Editor/ContentBrowser/Private/SourcesViewWidgets.cpp b/Engine/Source/Editor/ContentBrowser/Private/SourcesViewWidgets.cpp
--- a/Engine/Source/Editor/ContentBrowser/Private/SourcesViewWidgets.cpp
+++ b/Engine/Source/Editor/ContentBrowser/Private/SourcesViewWidgets.cpp
@@ -225,13 +225,6 @@ FSlateColor SAssetTreeItem::GetFolderColor() const
 			}
 		}
 
-		if(FoundColor == ContentBrowserUtils::GetDefaultColor() && IsSelected.IsBound() && IsSelected.Execute())
-		{
-			static const FSlateColor InvertedForeground = FAppStyle::Get().GetSlateColor("Colors.ForegroundInverted");
-
-			return InvertedForeground;
-		}
-
 		return FoundColor;
 	}
 	
diff --git a/Engine/Source/Editor/DetailCustomizations/Private/ColorStructCustomization.cpp b/Engine/Source/Editor/DetailCustomizations/Private/ColorStructCustomization.cpp
--- a/Engine/Source/Editor/DetailCustomizations/Private/ColorStructCustomization.cpp
+++ b/Engine/Source/Editor/DetailCustomizations/Private/ColorStructCustomization.cpp
@@ -102,7 +102,7 @@ TSharedRef<SWidget> FColorStructCustomization::CreateColorWidget(TWeakPtr<IPrope
 					.ShowBackgroundForAlpha(true)
 					.AlphaDisplayMode(bIgnoreAlpha ? EColorBlockAlphaDisplayMode::Ignore : EColorBlockAlphaDisplayMode::Separate)
 					.OnMouseButtonDown(this, &FColorStructCustomization::OnMouseButtonDownColorBlock)
-					.Size(FVector2D(70.0f, 22.0f))
+					.Size(FVector2D(70.0f, 20.0f))
 					.CornerRadius(FVector4(4.0f,4.0f,4.0f,4.0f))
 					.IsEnabled(this, &FColorStructCustomization::IsValueEnabled, StructWeakHandlePtr)
 				]
diff --git a/Engine/Source/Editor/EditorStyle/Private/StarshipStyle.cpp b/Engine/Source/Editor/EditorStyle/Private/StarshipStyle.cpp
--- a/Engine/Source/Editor/EditorStyle/Private/StarshipStyle.cpp
+++ b/Engine/Source/Editor/EditorStyle/Private/StarshipStyle.cpp
@@ -669,24 +669,6 @@ void FStarshipEditorStyle::FStyle::SetupGeneralStyles()
 		Set("HoverOnlyHyperlink", HoverOnlyHyperlink);
 	}
 
-	FButtonStyle RoundedButtonStyle = FButtonStyle()
-		.SetNormal(FSlateRoundedBoxBrush(FStyleColors::Dropdown, 4))
-		.SetHovered(FSlateRoundedBoxBrush(FStyleColors::Hover, 4))
-		.SetPressed(FSlateRoundedBoxBrush(FStyleColors::Hover, 4))
-		.SetNormalForeground(FStyleColors::Foreground)
-		.SetHoveredForeground(FStyleColors::ForegroundHover)
-		.SetPressedForeground(FStyleColors::ForegroundHover)
-		.SetDisabledForeground(FStyleColors::White25)
-		.SetNormalPadding(FMargin(6, 4))
-		.SetPressedPadding(FMargin(6, 4));
-
-	Set("RoundedButton", RoundedButtonStyle);
-
-	FComboButtonStyle AddNewComboButtonStyle = FComboButtonStyle(FStarshipCoreStyle::GetCoreStyle().GetWidgetStyle<FComboButtonStyle>("SimpleComboButton"))
-		.SetButtonStyle(RoundedButtonStyle) 
-		.SetDownArrowImage(FSlateNoResource());
-
-	Set("RoundedComboButton", AddNewComboButtonStyle);
 
 #endif // WITH_EDITOR || (IS_PROGRAM && WITH_UNREAL_DEVELOPER_TOOLS)
 
@@ -2287,6 +2269,21 @@ void FStarshipEditorStyle::FStyle::SetupViewportStyles()
 
 		FMargin ViewportMargin(4.f, 4.f);
 
+		const FCheckBoxStyle ViewportToggleButton = FCheckBoxStyle()
+			.SetCheckBoxType(ESlateCheckBoxType::ToggleButton)
+			.SetCheckedImage(FSlateNoResource())
+			.SetCheckedHoveredImage(FSlateNoResource())
+			.SetCheckedPressedImage(FSlateNoResource())
+			.SetUncheckedImage(FSlateNoResource())
+			.SetUncheckedHoveredImage(FSlateNoResource())
+			.SetUncheckedPressedImage(FSlateNoResource())
+			.SetForegroundColor(FStyleColors::Foreground)
+			.SetHoveredForegroundColor(FStyleColors::ForegroundHover)
+			.SetPressedForegroundColor(FStyleColors::ForegroundHover)
+			.SetCheckedForegroundColor(FStyleColors::PrimaryHover)
+			.SetCheckedHoveredForegroundColor(FStyleColors::PrimaryHover)
+			.SetPadding(0);
+
 		FSlateRoundedBoxBrush* ViewportGroupBrush = new FSlateRoundedBoxBrush(TransInput, 12.f, FStyleColors::Dropdown, 1.0);
 		Set("EditorViewportToolBar.Group", ViewportGroupBrush);
 
@@ -2302,7 +2299,7 @@ void FStarshipEditorStyle::FStyle::SetupViewportStyles()
 			.SetPressedPadding(ViewportMargin);
 		Set("EditorViewportToolBar.Button", ViewportMenuButton);
 
-		const FCheckBoxStyle ViewportMenuToggleLeftButtonStyle = FCheckBoxStyle(ViewportToolbarStyle.ToggleButton)
+		const FCheckBoxStyle ViewportMenuToggleLeftButtonStyle = FCheckBoxStyle(ViewportToggleButton)
 			.SetCheckBoxType(ESlateCheckBoxType::ToggleButton)
 			.SetUncheckedImage(		  BOX_BRUSH("Starship/EditorViewport/ToolBarLeftGroup", 12.f/25.f))
 			.SetUncheckedPressedImage(BOX_BRUSH("Starship/EditorViewport/ToolBarLeftGroup", 12.f/25.f))
@@ -2313,7 +2310,7 @@ void FStarshipEditorStyle::FStyle::SetupViewportStyles()
 			.SetPadding(ViewportMargin);
 		Set("EditorViewportToolBar.ToggleButton.Start", ViewportMenuToggleLeftButtonStyle);
 
-		const FCheckBoxStyle ViewportMenuToggleMiddleButtonStyle = FCheckBoxStyle(ViewportToolbarStyle.ToggleButton)
+		const FCheckBoxStyle ViewportMenuToggleMiddleButtonStyle = FCheckBoxStyle(ViewportToggleButton)
 			.SetCheckBoxType(ESlateCheckBoxType::ToggleButton)
 			.SetUncheckedImage(		  BOX_BRUSH("Starship/EditorViewport/ToolBarMiddleGroup", 12.f/25.f))
 			.SetUncheckedPressedImage(BOX_BRUSH("Starship/EditorViewport/ToolBarMiddleGroup", 12.f/25.f))
@@ -2324,7 +2321,7 @@ void FStarshipEditorStyle::FStyle::SetupViewportStyles()
 			.SetPadding(ViewportMargin);
 		Set("EditorViewportToolBar.ToggleButton.Middle", ViewportMenuToggleMiddleButtonStyle);
 
-		const FCheckBoxStyle ViewportMenuToggleRightButtonStyle = FCheckBoxStyle(ViewportToolbarStyle.ToggleButton)
+		const FCheckBoxStyle ViewportMenuToggleRightButtonStyle = FCheckBoxStyle(ViewportToggleButton)
 			.SetCheckBoxType(ESlateCheckBoxType::ToggleButton)
 			.SetUncheckedImage(		  BOX_BRUSH("Starship/EditorViewport/ToolBarRightGroup", 12.f/25.f))
 			.SetUncheckedPressedImage(BOX_BRUSH("Starship/EditorViewport/ToolBarRightGroup", 12.f/25.f))
@@ -2337,8 +2334,9 @@ void FStarshipEditorStyle::FStyle::SetupViewportStyles()
 
 		// We want a background-less version as the ComboMenu has its own unified background
 		const FToolBarStyle& SlimCoreToolBarStyle = FStarshipCoreStyle::GetCoreStyle().GetWidgetStyle<FToolBarStyle>("SlimToolBar");
+
 		Set("EditorViewportToolBar.ComboMenu.ButtonStyle", FButtonStyle(SlimCoreToolBarStyle.ButtonStyle).SetNormalPadding(0.0).SetPressedPadding(0.0));
-		Set("EditorViewportToolBar.ComboMenu.ToggleButton", FCheckBoxStyle(SlimCoreToolBarStyle.ToggleButton).SetPadding(0.0));
+		Set("EditorViewportToolBar.ComboMenu.ToggleButton", ViewportToggleButton);
 		Set("EditorViewportToolBar.ComboMenu.LabelStyle", SlimCoreToolBarStyle.LabelStyle);
 
 		FCheckBoxStyle MaximizeRestoreButton = FCheckBoxStyle(ViewportToolbarStyle.ToggleButton)
@@ -5390,18 +5388,14 @@ void FStarshipEditorStyle::FStyle::SetupPersonaStyle()
 
 		FToolBarStyle PlayToolbar = FStarshipCoreStyle::GetCoreStyle().GetWidgetStyle<FToolBarStyle>("AssetEditorToolbar");
 
-		const FButtonStyle PlayToolbarButton = FButtonStyle(PlayToolbar.ButtonStyle)
+		const FButtonStyle PlayToolbarButton = FButtonStyle(PlayToolbar.SettingsButtonStyle)
 			.SetNormalForeground(FStyleColors::AccentGreen)
 			.SetPressedForeground(GreenPress)
 			.SetHoveredForeground(GreenHover);
 
-		FComboButtonStyle PlayToolbarComboButton = FComboButtonStyle(FStarshipCoreStyle::GetCoreStyle().GetWidgetStyle<FComboButtonStyle>("ComboButton"))
-			.SetContentPadding(0)
-			.SetDownArrowImage(CORE_IMAGE_BRUSH_SVG("Starship/Common/ellipsis-vertical-narrow", FVector2D(6, 24)));
-		PlayToolbarComboButton.ButtonStyle = PlayToolbarButton;
 
-		PlayToolbar.SetButtonStyle(PlayToolbarButton);
-		PlayToolbar.SetSettingsComboButtonStyle(PlayToolbarComboButton);
+		PlayToolbar.SetSettingsButtonStyle(PlayToolbarButton);
+
 		PlayToolbar.SetShowLabels(true);
 
 		Set("PlayToolBar", PlayToolbar);
@@ -6330,6 +6324,7 @@ void FStarshipEditorStyle::FStyle::SetupContentBrowserStyle()
 
 		ContentBrowserToolBarStyle.SetSeparatorBrush(FSlateNoResource());
 		ContentBrowserToolBarStyle.SetSeparatorPadding(FMargin(4.0f, 0.0f));
+		ContentBrowserToolBarStyle.SetBackgroundPadding(FMargin(4.f, 2.f, 0.f, 2.f));
 	
 		Set("ContentBrowser.ToolBar", ContentBrowserToolBarStyle);
 	}
@@ -7799,6 +7794,8 @@ void FStarshipEditorStyle::FStyle::SetupStatusBarStyle()
 	StatusBarToolBarStyle.SetBackground(FSlateNoResource());
 	StatusBarToolBarStyle.SetLabelPadding(FMargin(5, 5, 0, 5));
 	StatusBarToolBarStyle.SetIconSize(Icon16x16);
+	StatusBarToolBarStyle.SetBackgroundPadding(FMargin(4.f, 4.f, 0.f, 4.f));
+
 	Set("StatusBarToolBar", StatusBarToolBarStyle);
 
 	Set("StatusBar.Message.MessageText",
diff --git a/Engine/Source/Editor/EditorWidgets/Private/SEditorHeaderButton.cpp b/Engine/Source/Editor/EditorWidgets/Private/SEditorHeaderButton.cpp
--- a/Engine/Source/Editor/EditorWidgets/Private/SEditorHeaderButton.cpp
+++ b/Engine/Source/Editor/EditorWidgets/Private/SEditorHeaderButton.cpp
@@ -4,6 +4,7 @@
 #include "SComponentClassCombo.h"
 #include "Widgets/Images/SImage.h"
 #include "Widgets/Input/SButton.h"
+#include "Styling/StyleColors.h"
 
 void SEditorHeaderButton::Construct(const FArguments& InArgs)
 {
@@ -19,7 +20,7 @@ void SEditorHeaderButton::Construct(const FArguments& InArgs)
 		[
 			SNew(SImage)
 			.Image(InArgs._Icon)
-			.ColorAndOpacity(FSlateColor::UseForeground())
+			.ColorAndOpacity(FStyleColors::AccentGreen)
 		]
 		+ SHorizontalBox::Slot()
 		.Padding(FMargin(3, 0, 0, 0))
@@ -27,7 +28,7 @@ void SEditorHeaderButton::Construct(const FArguments& InArgs)
 		.AutoWidth()
 		[
 			SNew(STextBlock)
-			.TextStyle(FAppStyle::Get(), "DialogButtonText")
+			.TextStyle(FAppStyle::Get(), "SmallButtonText")
 			.Text(InArgs._Text)
 			.Visibility_Lambda([Text]() { return Text.Get(FText::GetEmpty()).IsEmpty() ? EVisibility::Collapsed : EVisibility::Visible; })
 		];
@@ -37,13 +38,11 @@ void SEditorHeaderButton::Construct(const FArguments& InArgs)
 		ChildSlot
 		[
 			SAssignNew(Button, SButton)
-			.ButtonStyle(FAppStyle::Get(), "RoundedButton")
 			.ForegroundColor(FSlateColor::UseStyle())
 			.IsEnabled(InArgs._IsEnabled)
 			.ToolTipText(InArgs._ToolTipText)
 			.HAlign(HAlign_Center)
 			.VAlign(VAlign_Center)
-			.ContentPadding(0)
 			.OnClicked(InArgs._OnClicked)
 			[
 				ButtonContent
@@ -55,7 +54,9 @@ void SEditorHeaderButton::Construct(const FArguments& InArgs)
 		ChildSlot
 		[
 			SAssignNew(ComboButton, SComboButton)
-			.ComboButtonStyle(FAppStyle::Get(), "RoundedComboButton")
+			.HasDownArrow(false)
+			.ContentPadding(FMargin(2.0f, 3.0f))
+			.ButtonStyle(&FAppStyle::Get().GetWidgetStyle<FButtonStyle>("Button"))
 			.ForegroundColor(FSlateColor::UseStyle())
 			.IsEnabled(InArgs._IsEnabled)
 			.ToolTipText(InArgs._ToolTipText)
diff --git a/Engine/Source/Editor/EditorWidgets/Private/SSimpleButton.cpp b/Engine/Source/Editor/EditorWidgets/Private/SSimpleButton.cpp
new file mode 100644
--- /dev/null
+++ b/Engine/Source/Editor/EditorWidgets/Private/SSimpleButton.cpp
@@ -0,0 +1,38 @@
+// Copyright Epic Games, Inc. All Rights Reserved.
+
+#include "SSimpleButton.h"
+#include "Widgets/SBoxPanel.h"
+#include "Widgets/Images/SImage.h"
+#include "Widgets/Text/STextBlock.h"
+
+void SSimpleButton::Construct(const FArguments& InArgs)
+{
+	SButton::Construct(SButton::FArguments()
+		.ButtonStyle(&FAppStyle::Get().GetWidgetStyle<FButtonStyle>(InArgs._Text.IsSet() ? "SimpleButtonLabelAndIcon" : "SimpleButton"))
+		.OnClicked(InArgs._OnClicked)
+		.ForegroundColor(FSlateColor::UseStyle())
+		[
+			SNew(SHorizontalBox)
+			+SHorizontalBox::Slot()
+			.AutoWidth()
+			.VAlign(VAlign_Center)
+			[
+				SNew(SImage)
+				.ColorAndOpacity(FSlateColor::UseForeground())
+				.Image(InArgs._Icon)
+				.Visibility(InArgs._Icon.IsSet() ? EVisibility::HitTestInvisible : EVisibility::Collapsed)
+			]
+			+ SHorizontalBox::Slot()
+			.VAlign(VAlign_Center)
+			.Padding(FMargin(3, 0, 0, 0))
+			.AutoWidth()
+			[
+				SNew(STextBlock)
+				.TextStyle(&FAppStyle::Get().GetWidgetStyle<FTextBlockStyle>("SmallButtonText"))
+				.Text(InArgs._Text)
+				.Visibility(InArgs._Text.IsSet() ? EVisibility::HitTestInvisible : EVisibility::Collapsed)
+			]
+		]
+	);
+}
+
diff --git a/Engine/Source/Editor/EditorWidgets/Private/SSimpleComboButton.cpp b/Engine/Source/Editor/EditorWidgets/Private/SSimpleComboButton.cpp
new file mode 100644
--- /dev/null
+++ b/Engine/Source/Editor/EditorWidgets/Private/SSimpleComboButton.cpp
@@ -0,0 +1,47 @@
+// Copyright Epic Games, Inc. All Rights Reserved.
+
+#include "SSimpleComboButton.h"
+#include "Widgets/SBoxPanel.h"
+#include "Widgets/Images/SImage.h"
+#include "Widgets/Text/STextBlock.h"
+
+void SSimpleComboButton::Construct(const FArguments& InArgs)
+{
+
+	TAttribute<FText> Text = InArgs._Text;
+
+
+	TSharedRef<SHorizontalBox> ButtonContent = SNew(SHorizontalBox)
+		+ SHorizontalBox::Slot()
+		.HAlign(HAlign_Center)
+		.VAlign(VAlign_Center)
+		.AutoWidth()
+		[
+			SNew(SImage)
+			.Image(InArgs._Icon)
+			.ColorAndOpacity(FSlateColor::UseForeground())
+		]
+		+ SHorizontalBox::Slot()
+		.Padding(FMargin(3, 0, 0, 0))
+		.VAlign(VAlign_Center)
+		.AutoWidth()
+		[
+			SNew(STextBlock)
+			.TextStyle(FAppStyle::Get(), "SmallButtonText")
+			.Text(InArgs._Text)
+			.Visibility_Lambda([Text]() { return Text.Get(FText::GetEmpty()).IsEmpty() ? EVisibility::Collapsed : EVisibility::Visible; })
+		];
+
+	SComboButton::Construct(SComboButton::FArguments()
+		.HasDownArrow(false)
+		.ContentPadding(FMargin(2.0f, 2.0f))
+		.ComboButtonStyle(FAppStyle::Get(), "SimpleComboButton")
+		.ForegroundColor(FSlateColor::UseStyle())
+		.ButtonContent() [ButtonContent]
+		.MenuContent()[InArgs._MenuContent.Widget]
+		.OnGetMenuContent(InArgs._OnGetMenuContent)
+		.OnMenuOpenChanged(InArgs._OnMenuOpenChanged)
+		.OnComboBoxOpened(InArgs._OnComboBoxOpened)
+	);
+}
+
diff --git a/Engine/Source/Editor/EditorWidgets/Public/SSimpleButton.h b/Engine/Source/Editor/EditorWidgets/Public/SSimpleButton.h
new file mode 100644
--- /dev/null
+++ b/Engine/Source/Editor/EditorWidgets/Public/SSimpleButton.h
@@ -0,0 +1,28 @@
+// Copyright Epic Games, Inc. All Rights Reserved.
+
+#include "Framework/SlateDelegates.h"
+#include "Widgets/DeclarativeSyntaxSupport.h"
+#include "Widgets/Input/SButton.h"
+
+class EDITORWIDGETS_API SSimpleButton : public SButton
+{
+public:
+
+	SLATE_BEGIN_ARGS(SSimpleButton)
+	{}
+		/** The text to display in the button. */
+		SLATE_ATTRIBUTE(FText, Text)
+
+		SLATE_ATTRIBUTE(const FSlateBrush*, Icon)
+
+		/** The clicked handler. Note that if this is set, the button will behave as though it were just a button.
+		 * This means that OnGetMenuContent, OnComboBoxOpened and OnMenuOpenChanged will all be ignored, since there is no menu.
+		 */
+		SLATE_EVENT(FOnClicked, OnClicked)
+
+	SLATE_END_ARGS()
+
+	SSimpleButton() {}
+
+	void Construct(const FArguments& InArgs);
+};
\ No newline at end of file
diff --git a/Engine/Source/Editor/EditorWidgets/Public/SSimpleComboButton.h b/Engine/Source/Editor/EditorWidgets/Public/SSimpleComboButton.h
new file mode 100644
--- /dev/null
+++ b/Engine/Source/Editor/EditorWidgets/Public/SSimpleComboButton.h
@@ -0,0 +1,30 @@
+// Copyright Epic Games, Inc. All Rights Reserved.
+
+#include "Framework/SlateDelegates.h"
+#include "Widgets/DeclarativeSyntaxSupport.h"
+#include "Widgets/Input/SComboButton.h"
+
+class EDITORWIDGETS_API SSimpleComboButton : public SComboButton
+{
+public:
+
+	SLATE_BEGIN_ARGS(SSimpleComboButton)
+	{}
+		/** The text to display in the button. */
+		SLATE_ATTRIBUTE(FText, Text)
+
+		SLATE_ATTRIBUTE(const FSlateBrush*, Icon)
+	
+		/** The static menu content widget. */
+		SLATE_NAMED_SLOT(FArguments, MenuContent)
+
+		SLATE_EVENT(FOnGetContent, OnGetMenuContent)
+		SLATE_EVENT(FOnComboBoxOpened, OnComboBoxOpened)
+		SLATE_EVENT(FOnIsOpenChanged, OnMenuOpenChanged)
+
+	SLATE_END_ARGS()
+
+	SSimpleComboButton() {}
+
+	void Construct(const FArguments& InArgs);
+};
\ No newline at end of file
diff --git a/Engine/Source/Editor/KismetWidgets/Private/SPinTypeSelector.cpp b/Engine/Source/Editor/KismetWidgets/Private/SPinTypeSelector.cpp
--- a/Engine/Source/Editor/KismetWidgets/Private/SPinTypeSelector.cpp
+++ b/Engine/Source/Editor/KismetWidgets/Private/SPinTypeSelector.cpp
@@ -393,7 +393,7 @@ void SPinTypeSelector::Construct(const FArguments& InArgs, FGetPinTypeTree GetPi
 			.WidthOverride(SelectorType == ESelectorType::Full ? 125.f : FOptionalSize())
 			[
 				SAssignNew(TypeComboButton, SComboButton)
-				.ComboButtonStyle(FAppStyle::Get(), SelectorType == ESelectorType::Full ? "ComboButton" : "BlueprintEditor.CompactVariableTypeSelector")
+				.ComboButtonStyle(FAppStyle::Get(), SelectorType == ESelectorType::Full ? "ComboButton" : "ComboButton")
 				.OnGetMenuContent(this, &SPinTypeSelector::GetMenuContent, false)
 				.ContentPadding(0)
 				.ToolTipText(this, &SPinTypeSelector::GetToolTipForComboBoxType)
@@ -441,13 +441,7 @@ void SPinTypeSelector::Construct(const FArguments& InArgs, FGetPinTypeTree GetPi
 			HBox->AddSlot()
 			[
 				SNew(SBox)
-				.Visibility(
-					TAttribute<EVisibility>::Create(
-						TAttribute<EVisibility>::FGetter::CreateLambda(
-							[this]() {return this->TargetPinType.Get().IsMap() == true ? EVisibility::Visible : EVisibility::Collapsed; }
-						)
-					)
-				)
+				.Visibility_Lambda([this]() {return this->TargetPinType.Get().IsMap() == true ? EVisibility::Visible : EVisibility::Collapsed; })
 				[
 					SAssignNew( SecondaryTypeComboButton, SComboButton )
 					.OnGetMenuContent(this, &SPinTypeSelector::GetMenuContent, true )
@@ -488,6 +482,7 @@ void SPinTypeSelector::Construct(const FArguments& InArgs, FGetPinTypeTree GetPi
 		SNew(SWidgetSwitcher)
 		.WidgetIndex_Lambda([this](){return ReadOnly.Get() ? 1 : 0; })
 		+ SWidgetSwitcher::Slot()
+		.Padding(SelectorType == ESelectorType::Partial ? FMargin(-6.0f, 0.0f,0.0f,0.0f) : FMargin(0))
 		[
 			Widget.ToSharedRef()
 		]
@@ -498,7 +493,7 @@ void SPinTypeSelector::Construct(const FArguments& InArgs, FGetPinTypeTree GetPi
 			+ SHorizontalBox::Slot()
 			.VAlign(VAlign_Center)
 			.HAlign(HAlign_Left)
-			.Padding(FMargin(0.0f, 2.0f, 2.0f, 2.0f))
+			.Padding(FMargin(2.0f, 3.0f, 0.0f, 3.0f))
 			.AutoWidth()
 			[
 				SNew(SImage)
@@ -514,7 +509,7 @@ void SPinTypeSelector::Construct(const FArguments& InArgs, FGetPinTypeTree GetPi
 				SNew(STextBlock)
 				.Text(this, &SPinTypeSelector::GetTypeDescription)
 				.Font(InArgs._Font)
-				.ColorAndOpacity(FSlateColor::UseSubduedForeground())
+				.ColorAndOpacity(FSlateColor::UseForeground())
 			]
 		]	
 	];
diff --git a/Engine/Source/Editor/Persona/Private/TabSpawners.cpp b/Engine/Source/Editor/Persona/Private/TabSpawners.cpp
--- a/Engine/Source/Editor/Persona/Private/TabSpawners.cpp
+++ b/Engine/Source/Editor/Persona/Private/TabSpawners.cpp
@@ -464,47 +464,37 @@ TSharedRef<SWidget> FAnimBlueprintPreviewEditorSummoner::CreateTabBody(const FWo
 			[
 				SNew(SHorizontalBox)
 				+SHorizontalBox::Slot()
-				.Padding(FMargin( 0.f, 0.f, 2.f, 0.f ))
+				.AutoWidth()
+				.Padding(FMargin( 5.f, 0.f, 2.f, 0.f ))
 				[
-					SNew(SBorder)
-					.BorderImage(this, &FAnimBlueprintPreviewEditorSummoner::GetBorderBrushByMode, EAnimBlueprintEditorMode::PreviewMode)
-					.Padding(0)
+					SNew(SCheckBox)
+					.Style(FEditorStyle::Get(), "RadioButton")
+					.IsChecked(this, &FAnimBlueprintPreviewEditorSummoner::IsChecked, EAnimBlueprintEditorMode::PreviewMode)
+					.OnCheckStateChanged(const_cast<FAnimBlueprintPreviewEditorSummoner*>(this), &FAnimBlueprintPreviewEditorSummoner::OnCheckedChanged, EAnimBlueprintEditorMode::PreviewMode)
+					.ToolTip(IDocumentation::Get()->CreateToolTip(	LOCTEXT("AnimBlueprintPropertyEditorPreviewMode", "Switch to editing the preview instance properties"),
+																	NULL,
+																	TEXT("Shared/Editors/Persona"),
+																	TEXT("AnimBlueprintPropertyEditorPreviewMode")))
 					[
-						SNew(SCheckBox)
-						.Style(FEditorStyle::Get(), "RadioButton")
-						.IsChecked(this, &FAnimBlueprintPreviewEditorSummoner::IsChecked, EAnimBlueprintEditorMode::PreviewMode)
-						.OnCheckStateChanged(const_cast<FAnimBlueprintPreviewEditorSummoner*>(this), &FAnimBlueprintPreviewEditorSummoner::OnCheckedChanged, EAnimBlueprintEditorMode::PreviewMode)
-						.ToolTip(IDocumentation::Get()->CreateToolTip(	LOCTEXT("AnimBlueprintPropertyEditorPreviewMode", "Switch to editing the preview instance properties"),
-																		NULL,
-																		TEXT("Shared/Editors/Persona"),
-																		TEXT("AnimBlueprintPropertyEditorPreviewMode")))
-						[
-							SNew( STextBlock )
-							.Font( FCoreStyle::GetDefaultFontStyle("Bold", 9) )
-							.Text( LOCTEXT("AnimBlueprintDefaultsPreviewMode", "Edit Preview") )
-						]
+						SNew( STextBlock )
+						.Text( LOCTEXT("AnimBlueprintDefaultsPreviewMode", "Edit Preview") )
 					]
 				]
 				+SHorizontalBox::Slot()
+				.AutoWidth()
 				.Padding(FMargin( 2.f, 0.f, 0.f, 0.f ))
 				[
-					SNew(SBorder)
-					.BorderImage(this, &FAnimBlueprintPreviewEditorSummoner::GetBorderBrushByMode, EAnimBlueprintEditorMode::DefaultsMode)
-					.Padding(0)
+					SNew(SCheckBox)
+					.Style(FEditorStyle::Get(), "RadioButton")
+					.IsChecked(this, &FAnimBlueprintPreviewEditorSummoner::IsChecked, EAnimBlueprintEditorMode::DefaultsMode)
+					.OnCheckStateChanged(const_cast<FAnimBlueprintPreviewEditorSummoner*>(this), &FAnimBlueprintPreviewEditorSummoner::OnCheckedChanged, EAnimBlueprintEditorMode::DefaultsMode)
+					.ToolTip(IDocumentation::Get()->CreateToolTip(	LOCTEXT("AnimBlueprintPropertyEditorDefaultMode", "Switch to editing the class defaults"),
+																	NULL,
+																	TEXT("Shared/Editors/Persona"),
+																	TEXT("AnimBlueprintPropertyEditorDefaultMode")))
 					[
-						SNew(SCheckBox)
-						.Style(FEditorStyle::Get(), "RadioButton")
-						.IsChecked(this, &FAnimBlueprintPreviewEditorSummoner::IsChecked, EAnimBlueprintEditorMode::DefaultsMode)
-						.OnCheckStateChanged(const_cast<FAnimBlueprintPreviewEditorSummoner*>(this), &FAnimBlueprintPreviewEditorSummoner::OnCheckedChanged, EAnimBlueprintEditorMode::DefaultsMode)
-						.ToolTip(IDocumentation::Get()->CreateToolTip(	LOCTEXT("AnimBlueprintPropertyEditorDefaultMode", "Switch to editing the class defaults"),
-																		NULL,
-																		TEXT("Shared/Editors/Persona"),
-																		TEXT("AnimBlueprintPropertyEditorDefaultMode")))
-						[
-							SNew( STextBlock )
-							.Font( FCoreStyle::GetDefaultFontStyle("Bold", 9) )
-							.Text( LOCTEXT("AnimBlueprintDefaultsDefaultsMode", "Edit Defaults") )
-						]
+						SNew( STextBlock )
+						.Text( LOCTEXT("AnimBlueprintDefaultsDefaultsMode", "Edit Defaults") )
 					]
 				]
 			]
@@ -549,17 +539,6 @@ ECheckBoxState FAnimBlueprintPreviewEditorSummoner::IsChecked(EAnimBlueprintEdit
 	return CurrentMode == Mode ? ECheckBoxState::Checked: ECheckBoxState::Unchecked;
 }
 
-const FSlateBrush* FAnimBlueprintPreviewEditorSummoner::GetBorderBrushByMode(EAnimBlueprintEditorMode::Type Mode) const
-{
-	if(Mode == CurrentMode)
-	{
-		return FEditorStyle::GetBrush("ModeSelector.ToggleButton.Pressed");
-	}
-	else
-	{
-		return FEditorStyle::GetBrush("ModeSelector.ToggleButton.Normal");
-	}
-}
 
 void FAnimBlueprintPreviewEditorSummoner::OnCheckedChanged(ECheckBoxState NewType, EAnimBlueprintEditorMode::Type Mode)
 {
diff --git a/Engine/Source/Editor/Persona/Private/TabSpawners.h b/Engine/Source/Editor/Persona/Private/TabSpawners.h
--- a/Engine/Source/Editor/Persona/Private/TabSpawners.h
+++ b/Engine/Source/Editor/Persona/Private/TabSpawners.h
@@ -288,7 +288,6 @@ private:
 	/** Delegates to customize tab look based on selected mode */
 	EVisibility IsEditorVisible(EAnimBlueprintEditorMode::Type Mode) const;
 	ECheckBoxState IsChecked(EAnimBlueprintEditorMode::Type Mode) const;
-	const FSlateBrush* GetBorderBrushByMode(EAnimBlueprintEditorMode::Type Mode) const;
 
 	/** Handle changing of editor mode */
 	void OnCheckedChanged(ECheckBoxState NewType, EAnimBlueprintEditorMode::Type Mode);
diff --git a/Engine/Source/Editor/PropertyEditor/Private/MaterialList.cpp b/Engine/Source/Editor/PropertyEditor/Private/MaterialList.cpp
--- a/Engine/Source/Editor/PropertyEditor/Private/MaterialList.cpp
+++ b/Engine/Source/Editor/PropertyEditor/Private/MaterialList.cpp
@@ -175,7 +175,7 @@ public:
 							// Add a menu for displaying all textures 
 							SNew(SComboButton)
 							.OnGetMenuContent(this, &FMaterialItemView::OnGetTexturesMenuForMaterial)
-							.ComboButtonStyle(FAppStyle::Get(), "SimpleComboButton")
+							.ComboButtonStyle(FAppStyle::Get(), "SimpleComboButtonWithIcon")
 							.VAlign(VAlign_Center)
 							.IsEnabled( this, &FMaterialItemView::IsTexturesMenuEnabled )
 							.Visibility(bShowUsedTextures ? EVisibility::Visible : EVisibility::Hidden)
diff --git a/Engine/Source/Editor/PropertyEditor/Private/PropertyCustomizationHelpers.cpp b/Engine/Source/Editor/PropertyEditor/Private/PropertyCustomizationHelpers.cpp
--- a/Engine/Source/Editor/PropertyEditor/Private/PropertyCustomizationHelpers.cpp
+++ b/Engine/Source/Editor/PropertyEditor/Private/PropertyCustomizationHelpers.cpp
@@ -67,8 +67,8 @@ namespace PropertyCustomizationHelpers
 				SNew(SBox)
 				.HAlign(HAlign_Center)
 				.VAlign(VAlign_Center)
-				.WidthOverride(20)
-				.HeightOverride(20)
+				.WidthOverride(22)
+				.HeightOverride(22)
 				[
 					SNew(SButton)
 					.ButtonStyle( FAppStyle::Get(), "SimpleButton" )
diff --git a/Engine/Source/Editor/PropertyEditor/Private/UserInterface/PropertyEditor/SPropertyEditorBool.cpp b/Engine/Source/Editor/PropertyEditor/Private/UserInterface/PropertyEditor/SPropertyEditorBool.cpp
--- a/Engine/Source/Editor/PropertyEditor/Private/UserInterface/PropertyEditor/SPropertyEditorBool.cpp
+++ b/Engine/Source/Editor/PropertyEditor/Private/UserInterface/PropertyEditor/SPropertyEditorBool.cpp
@@ -6,12 +6,9 @@ void SPropertyEditorBool::Construct( const FArguments& InArgs, const TSharedRef<
 {
 	PropertyEditor = InPropertyEditor;
 
-	static const FName DefaultForegroundName("DefaultForeground");
-
 	CheckBox = SNew( SCheckBox )
 		.OnCheckStateChanged( this, &SPropertyEditorBool::OnCheckStateChanged )
 		.IsChecked( this, &SPropertyEditorBool::OnGetCheckState )
-		.ForegroundColor( FEditorStyle::GetSlateColor(DefaultForegroundName) )
 		.Padding(0.0f);
 
 	ChildSlot
diff --git a/Engine/Source/Editor/SceneOutliner/Private/SSceneOutliner.cpp b/Engine/Source/Editor/SceneOutliner/Private/SSceneOutliner.cpp
--- a/Engine/Source/Editor/SceneOutliner/Private/SSceneOutliner.cpp
+++ b/Engine/Source/Editor/SceneOutliner/Private/SSceneOutliner.cpp
@@ -154,7 +154,7 @@ void SSceneOutliner::Construct(const FArguments& InArgs, const FSceneOutlinerIni
 		.AutoWidth()
 		[
 			SAssignNew( ViewOptionsComboButton, SComboButton )
-			.ComboButtonStyle( FAppStyle::Get(), "SimpleComboButton" ) // Use the tool bar item style for this button
+			.ComboButtonStyle( FAppStyle::Get(), "SimpleComboButtonWithIcon" ) // Use the tool bar item style for this button
 			.OnGetMenuContent( this, &SSceneOutliner::GetViewButtonContent, Mode->ShowFilterOptions())
 			.HasDownArrow(false)
 			.ButtonContent()
diff --git a/Engine/Source/Editor/SceneOutliner/Private/SceneOutlinerActorInfoColumn.cpp b/Engine/Source/Editor/SceneOutliner/Private/SceneOutlinerActorInfoColumn.cpp
--- a/Engine/Source/Editor/SceneOutliner/Private/SceneOutlinerActorInfoColumn.cpp
+++ b/Engine/Source/Editor/SceneOutliner/Private/SceneOutlinerActorInfoColumn.cpp
@@ -21,6 +21,7 @@
 #include "WorldPartition/DataLayer/DataLayer.h"
 #include "WorldPartition/DataLayer/WorldDataLayers.h"
 #include "WorldPartition/WorldPartitionActorDesc.h"
+#include "Styling/StyleColors.h"
 
 #define LOCTEXT_NAMESPACE "SceneOutlinerActorInfoColumn"
 
@@ -119,7 +120,7 @@ const TSharedRef< SWidget > FTypeInfoColumn::ConstructRowWidget( FSceneOutlinerT
 		MainText
 	];
 
-	auto Hyperlink = ConstructClassHyperlink(*TreeItem);
+	TSharedPtr<SWidget> Hyperlink = ConstructClassHyperlink(*TreeItem);
 	if (Hyperlink.IsValid())
 	{
 		// If we got a hyperlink, disable hide default text, and show the hyperlink
@@ -134,9 +135,9 @@ const TSharedRef< SWidget > FTypeInfoColumn::ConstructRowWidget( FSceneOutlinerT
 			// Make sure that the hyperlink shows as black (by multiplying black * desired color) when selected so it is readable against the orange background even if blue/green/etc... normally
 			SNew(SBorder)
 			.BorderImage(FEditorStyle::GetBrush("NoBorder"))
-			.ColorAndOpacity_Static([](TWeakPtr<const STableRow<FSceneOutlinerTreeItemPtr>> WeakRow)->FLinearColor{
+			.ForegroundColor_Static([](TWeakPtr<const STableRow<FSceneOutlinerTreeItemPtr>> WeakRow)->FSlateColor{
 				auto TableRow = WeakRow.Pin();
-				return TableRow.IsValid() && TableRow->IsSelected() ? FLinearColor::Black : FLinearColor::White;
+				return TableRow.IsValid() && TableRow->IsSelected() ? FStyleColors::ForegroundHover : FSlateColor::UseStyle();
 			}, TWeakPtr<const STableRow<FSceneOutlinerTreeItemPtr>>(StaticCastSharedRef<const STableRow<FSceneOutlinerTreeItemPtr>>(Row.AsShared())))
 			[
 				Hyperlink.ToSharedRef()
diff --git a/Engine/Source/Editor/SkeletonEditor/Private/SSkeletonTree.cpp b/Engine/Source/Editor/SkeletonEditor/Private/SSkeletonTree.cpp
--- a/Engine/Source/Editor/SkeletonEditor/Private/SSkeletonTree.cpp
+++ b/Engine/Source/Editor/SkeletonEditor/Private/SSkeletonTree.cpp
@@ -66,6 +66,7 @@
 #include "Widgets/Views/STreeView.h"
 #include "IPinnedCommandList.h"
 #include "PersonaModule.h"
+#include "SEditorHeaderButton.h"
 
 #define LOCTEXT_NAMESPACE "SSkeletonTree"
 
@@ -218,34 +219,9 @@ void SSkeletonTree::Construct(const FArguments& InArgs, const TSharedRef<FEditab
 				.VAlign(VAlign_Center)
 				.Padding(FMargin(6.f, 0.0))
 				[
-					SAssignNew(FilterComboButton, SComboButton)
-					.ButtonStyle(&FAppStyle::Get().GetWidgetStyle<FButtonStyle>("SecondaryButton"))
-					.ForegroundColor(FSlateColor::UseStyle())
-					.HasDownArrow(false)
+					SNew(SEditorHeaderButton)
 					.OnGetMenuContent( this, &SSkeletonTree::CreateNewMenu )
-					// .AddMetaData<FTagMetaData>(TEXT("SkelTree.Bones"))
-					.ButtonContent()
-					[
-						SNew(SHorizontalBox)
-						+ SHorizontalBox::Slot()
-						.AutoWidth()
-						.HAlign(HAlign_Center)
-						.VAlign(VAlign_Center)
-						[
-							SNew(SImage)
-							.Image(FAppStyle::Get().GetBrush("Icons.Plus"))
-							.ColorAndOpacity(FSlateColor::UseForeground())
-						]
-						+ SHorizontalBox::Slot()
-						.Padding(FMargin(7.f, 0.f, 0.0, 0.0))
-						.VAlign(VAlign_Center)
-						.AutoWidth()
-						[
-							SNew(STextBlock)
-							.TextStyle(FAppStyle::Get(), "SmallButtonText")
-							.Text(LOCTEXT("NewButton", "ADD"))
-						]
-					]
+					.Icon(FAppStyle::Get().GetBrush("Icons.Plus"))
 				]
 
 				+SHorizontalBox::Slot()
diff --git a/Engine/Source/Editor/SkeletonEditor/SkeletonEditor.Build.cs b/Engine/Source/Editor/SkeletonEditor/SkeletonEditor.Build.cs
--- a/Engine/Source/Editor/SkeletonEditor/SkeletonEditor.Build.cs
+++ b/Engine/Source/Editor/SkeletonEditor/SkeletonEditor.Build.cs
@@ -12,21 +12,22 @@ public class SkeletonEditor : ModuleRules
 				"CoreUObject",
 				"Engine",
 				"ApplicationCore",
-                "InputCore",
+				"InputCore",
 				"Slate",
 				"SlateCore",
-                "EditorStyle",
+				"EditorStyle",
 				"EditorFramework",
-                "UnrealEd",
-                "Persona",
-                "AnimGraph",
-                "AnimGraphRuntime",
-                "ContentBrowser",
-                "AssetRegistry",
-                "BlueprintGraph",
-                "Kismet",
-                "PinnedCommandList",
-				"ToolMenus"
+				"UnrealEd",
+				"Persona",
+				"AnimGraph",
+				"AnimGraphRuntime",
+				"ContentBrowser",
+				"AssetRegistry",
+				"BlueprintGraph",
+				"Kismet",
+				"PinnedCommandList",
+				"ToolMenus",
+				"EditorWidgets",
             }
 		);
 
diff --git a/Engine/Source/Editor/StatusBar/Private/SStatusBar.cpp b/Engine/Source/Editor/StatusBar/Private/SStatusBar.cpp
--- a/Engine/Source/Editor/StatusBar/Private/SStatusBar.cpp
+++ b/Engine/Source/Editor/StatusBar/Private/SStatusBar.cpp
@@ -623,7 +623,6 @@ void SStatusBar::Construct(const FArguments& InArgs, FName InStatusBarName, cons
 					SNew(SBorder)
 					.Padding(0.0f)
 					.BorderImage(StatusBarBackground)
-					.VAlign(VAlign_Center)
 					[
 						MakeStatusBarToolBarWidget()
 					]
diff --git a/Engine/Source/Editor/SubobjectEditor/Private/SSubobjectInstanceEditor.cpp b/Engine/Source/Editor/SubobjectEditor/Private/SSubobjectInstanceEditor.cpp
--- a/Engine/Source/Editor/SubobjectEditor/Private/SSubobjectInstanceEditor.cpp
+++ b/Engine/Source/Editor/SubobjectEditor/Private/SSubobjectInstanceEditor.cpp
@@ -8,7 +8,6 @@
 #include "Editor/EditorEngine.h"
 #include "ISCSEditorUICustomization.h"	// #TODO_BH Rename this to subobject
 #include "SubobjectEditorExtensionContext.h"
-
 #include "Styling/SlateIconFinder.h"
 #include "SlateOptMacros.h"
 #include "Widgets/Input/SSearchBox.h"
@@ -17,6 +16,8 @@
 #include "Framework/MultiBox/MultiBoxBuilder.h"
 #include "SComponentClassCombo.h"
 #include "SEditorHeaderButton.h"
+#include "SSimpleButton.h"
+#include "SSimpleComboButton.h"
 #include "Editor/UnrealEdEngine.h"
 #include "Subsystems/PanelExtensionSubsystem.h"	// SExtensionPanel
 #include "Kismet2/ComponentEditorUtils.h"
@@ -91,10 +92,9 @@ void SSubobjectInstanceEditor::Construct(const FArguments& InArgs)
 
 	+ SHorizontalBox::Slot()
 	.VAlign(VAlign_Center)
-	.Padding(0.0f,0.0f,4.0f,0.0f)
 	.AutoWidth()
 	[
-		SNew(SEditorHeaderButton)
+		SNew(SSimpleButton)
 		.AddMetaData<FTagMetaData>(FTagMetaData(TEXT("Actor.ConvertToBlueprint")))
 		.Visibility(this, &SSubobjectInstanceEditor::GetPromoteToBlueprintButtonVisibility)
 		.OnClicked(this, &SSubobjectInstanceEditor::OnPromoteToBlueprintClicked)
@@ -109,7 +109,7 @@ void SSubobjectInstanceEditor::Construct(const FArguments& InArgs)
 	.VAlign(VAlign_Center)
 	.AutoWidth()
 	[
-		SNew(SEditorHeaderButton)
+		SNew(SSimpleComboButton)
 		.AddMetaData<FTagMetaData>(FTagMetaData(TEXT("Actor.EditBlueprint")))
 		.Visibility(this, &SSubobjectInstanceEditor::GetEditBlueprintButtonVisibility)
 		.ToolTipText(LOCTEXT("EditActorBlueprint_Tooltip", "Edit the Blueprint for this Actor"))
diff --git a/Engine/Source/Editor/UnrealEd/Private/WorkflowOrientedApp/SModeWidget.cpp b/Engine/Source/Editor/UnrealEd/Private/WorkflowOrientedApp/SModeWidget.cpp
--- a/Engine/Source/Editor/UnrealEd/Private/WorkflowOrientedApp/SModeWidget.cpp
+++ b/Engine/Source/Editor/UnrealEd/Private/WorkflowOrientedApp/SModeWidget.cpp
@@ -93,7 +93,7 @@ void SModeWidget::Construct(const FArguments& InArgs, const FText& InText, const
 	.VAlign(VAlign_Fill)
 	[
 		SNew(SCheckBox)
-		.Style(FAppStyle::Get(), "SegmentedCombo.ButtonOnly")
+		.Style(FAppStyle::Get(), "ToggleButtonCheckbox")
 		.IsChecked(this, &SModeWidget::GetModeCheckState)
 		.OnCheckStateChanged(this, &SModeWidget::OnModeTabClicked)
 		[
diff --git a/Engine/Source/Runtime/Slate/Private/Framework/MultiBox/SMenuEntryBlock.cpp b/Engine/Source/Runtime/Slate/Private/Framework/MultiBox/SMenuEntryBlock.cpp
--- a/Engine/Source/Runtime/Slate/Private/Framework/MultiBox/SMenuEntryBlock.cpp
+++ b/Engine/Source/Runtime/Slate/Private/Framework/MultiBox/SMenuEntryBlock.cpp
@@ -13,6 +13,7 @@
 #include "Widgets/SToolTip.h"
 #include "Widgets/Colors/SColorBlock.h"
 #include "Widgets/Input/SCheckBox.h"
+#include "Styling/StyleColors.h"
 
 
 FMenuEntryBlock::FMenuEntryBlock( const FName& InExtensionHook, const TSharedPtr< const FUICommandInfo > InCommand, TSharedPtr< const FUICommandList > InCommandList, const TAttribute<FText>& InLabelOverride, const TAttribute<FText>& InToolTipOverride, const FSlateIcon& InIconOverride, bool bInCloseSelfOnly, bool bInShouldCloseWindowAfterMenuSelection)
@@ -527,7 +528,6 @@ TSharedRef< SWidget > SMenuEntryBlock::BuildMenuEntryWidget( const FMenuEntryBui
 	else if (UserInterfaceType == EUserInterfaceActionType::RadioButton)
 	{
 		CheckBoxStyle = ISlateStyle::Join( StyleName, ".RadioButton" );
-		CheckBoxForegroundColor = TAttribute<FSlateColor>::Create( TAttribute<FSlateColor>::FGetter::CreateRaw( this, &SMenuEntryBlock::TintOnHover ) );
 	}
 
 
@@ -711,7 +711,7 @@ private:
 	{
 		if ( this->IsHovered() )
 		{
-			return FLinearColor::Black;
+			return FStyleColors::ForegroundHover;
 		}
 		else
 		{
diff --git a/Engine/Source/Runtime/Slate/Private/Framework/MultiBox/SToolBarButtonBlock.cpp b/Engine/Source/Runtime/Slate/Private/Framework/MultiBox/SToolBarButtonBlock.cpp
--- a/Engine/Source/Runtime/Slate/Private/Framework/MultiBox/SToolBarButtonBlock.cpp
+++ b/Engine/Source/Runtime/Slate/Private/Framework/MultiBox/SToolBarButtonBlock.cpp
@@ -198,6 +198,7 @@ void SToolBarButtonBlock::BuildMultiBlockWidget(const ISlateStyle* StyleSet, con
 	TSharedRef<SLayeredImage> IconWidget =
 		SNew(SLayeredImage)
 		.ColorAndOpacity(this, &SToolBarButtonBlock::GetIconForegroundColor)
+		.Visibility(EVisibility::HitTestInvisible)
 		.Image(this, &SToolBarButtonBlock::GetIconBrush);
 
 	IconWidget->AddLayer(TAttribute<const FSlateBrush*>(this, &SToolBarButtonBlock::GetOverlayIconBrush));
@@ -274,15 +275,19 @@ void SToolBarButtonBlock::BuildMultiBlockWidget(const ISlateStyle* StyleSet, con
 	if( UserInterfaceType == EUserInterfaceActionType::Button )
 	{
 		FName BlockStyle = EMultiBlockLocation::ToName(ISlateStyle::Join( StyleName, ".Button" ), BlockLocation);
-		const FButtonStyle& ButtonStyle = BlockLocation == EMultiBlockLocation::None ? ToolBarStyle.ButtonStyle : StyleSet->GetWidgetStyle<FButtonStyle>(BlockStyle);
+		const FButtonStyle* ToolbarButtonStyle = BlockLocation == EMultiBlockLocation::None ? &ToolBarStyle.ButtonStyle : &StyleSet->GetWidgetStyle<FButtonStyle>(BlockStyle);
+
+		if (OptionsBlockWidget.IsValid())
+		{
+			ToolbarButtonStyle = &ToolBarStyle.SettingsButtonStyle;
+		}
 
 		ChildSlot
 		[
 			// Create a button
 			SNew(SButton)
 			.ContentPadding(0.f)
-			// Use the tool bar item style for this button
-			.ButtonStyle(&ButtonStyle)
+			.ButtonStyle(ToolbarButtonStyle)
 			.IsEnabled(this, &SToolBarButtonBlock::IsEnabled)
 			.OnClicked(this, &SToolBarButtonBlock::OnClicked)
 			.ToolTip(FMultiBoxSettings::ToolTipConstructor.Execute(ActualToolTip, nullptr, Action.Pin()))
@@ -296,20 +301,24 @@ void SToolBarButtonBlock::BuildMultiBlockWidget(const ISlateStyle* StyleSet, con
 	{
 		FName BlockStyle = EMultiBlockLocation::ToName(ISlateStyle::Join( StyleName, ".ToggleButton" ), BlockLocation);
 	
-		const FCheckBoxStyle& CheckStyle = BlockLocation == EMultiBlockLocation::None ? ToolBarStyle.ToggleButton : StyleSet->GetWidgetStyle<FCheckBoxStyle>(BlockStyle);
+		const FCheckBoxStyle* CheckStyle = BlockLocation == EMultiBlockLocation::None ? &ToolBarStyle.ToggleButton : &StyleSet->GetWidgetStyle<FCheckBoxStyle>(BlockStyle);
+
+		if (OptionsBlockWidget.IsValid())
+		{
+			CheckStyle = &ToolBarStyle.SettingsToggleButton;
+		}
 
 		ChildSlot
 		[
 			// Create a check box
 			SNew(SCheckBox)
 			// Use the tool bar style for this check box
-			.Style(&CheckStyle)
+			.Style(CheckStyle)
 			.IsFocusable(bIsFocusable)
 			.ToolTip( FMultiBoxSettings::ToolTipConstructor.Execute( ActualToolTip, nullptr, Action.Pin()))		
 			.OnCheckStateChanged(this, &SToolBarButtonBlock::OnCheckStateChanged )
-			.IsChecked(this, &SToolBarButtonBlock::OnIsChecked)
+			.IsChecked(this, &SToolBarButtonBlock::GetCheckState)
 			.IsEnabled(this, &SToolBarButtonBlock::IsEnabled)
-			.Padding(ToolBarStyle.CheckBoxPadding)
 			[
 				ButtonContent
 			]
@@ -326,7 +335,7 @@ void SToolBarButtonBlock::BuildMultiBlockWidget(const ISlateStyle* StyleSet, con
 			.AutoWidth()
 			[
 				SAssignNew(ButtonBorder, SBorder)
-				.Padding(FMargin(10.0f, 0.0f,2.0f, 0.0f))
+				.Padding(0)
 				.BorderImage(this, &SToolBarButtonBlock::GetOptionsBlockLeftBrush)
 				.VAlign(VAlign_Center)
 				[
@@ -334,11 +343,11 @@ void SToolBarButtonBlock::BuildMultiBlockWidget(const ISlateStyle* StyleSet, con
 				]
 			]
 			+ SHorizontalBox::Slot()
-			.Padding(0.0f,0.0f)
+	
 			.AutoWidth()
 			[
 				SAssignNew(OptionsBorder, SBorder)
-				.Padding(FMargin(0.0f, 0.0f, 2.0f, 0.0f))
+				.Padding(0)
 				.BorderImage(this, &SToolBarButtonBlock::GetOptionsBlockRightBrush)
 				.VAlign(VAlign_Center)
 				[
@@ -349,6 +358,7 @@ void SToolBarButtonBlock::BuildMultiBlockWidget(const ISlateStyle* StyleSet, con
 	}
 	else
 	{
+		// Space between buttons. It does not make the buttons larger
 		ChildSlot.Padding(ToolBarStyle.ButtonPadding);
 	}
 
@@ -405,10 +415,10 @@ void SToolBarButtonBlock::OnCheckStateChanged( const ECheckBoxState NewCheckedSt
  *
  * @return ECheckBoxState::Checked if it should be checked, ECheckBoxState::Unchecked if not.
  */
-ECheckBoxState SToolBarButtonBlock::OnIsChecked() const
+ECheckBoxState SToolBarButtonBlock::GetCheckState() const
 {
-	TSharedPtr< const FUICommandList > ActionList = MultiBlock->GetActionList();
-	TSharedPtr< const FUICommandInfo > Action = MultiBlock->GetAction();
+	TSharedPtr<const FUICommandList> ActionList = MultiBlock->GetActionList();
+	TSharedPtr<const FUICommandInfo> Action = MultiBlock->GetAction();
 	const FUIAction& DirectActions = MultiBlock->GetDirectActions();
 
 	ECheckBoxState CheckState = ECheckBoxState::Unchecked;
@@ -474,7 +484,7 @@ EVisibility SToolBarButtonBlock::GetBlockVisibility() const
 
 EVisibility SToolBarButtonBlock::GetIconVisibility(bool bIsASmallIcon) const
 {
-	return ((bForceSmallIcons || FMultiBoxSettings::UseSmallToolBarIcons.Get()) ^ bIsASmallIcon) ? EVisibility::Collapsed : EVisibility::Visible;
+	return ((bForceSmallIcons || FMultiBoxSettings::UseSmallToolBarIcons.Get()) ^ bIsASmallIcon) ? EVisibility::Collapsed : EVisibility::HitTestInvisible;
 }
 
 const FSlateBrush* SToolBarButtonBlock::GetIconBrush() const
@@ -557,19 +567,23 @@ FSlateColor SToolBarButtonBlock::GetIconForegroundColor() const
 
 const FSlateBrush* SToolBarButtonBlock::GetOptionsBlockLeftBrush() const
 {
+	static const FName ToggledLeft("ToolbarSettingsRegion.LeftToggle");
+
 	if (ButtonBorder->IsHovered())
 	{
 		static const FName LeftHover("ToolbarSettingsRegion.LeftHover");
-		return FAppStyle::Get().GetBrush(LeftHover);
+		static const FName ToggledLeftHover("ToolbarSettingsRegion.LeftToggleHover");
+
+		return GetCheckState() == ECheckBoxState::Checked ? FAppStyle::Get().GetBrush(ToggledLeftHover) : FAppStyle::Get().GetBrush(LeftHover);
 	}
 	else if (OptionsBorder->IsHovered())
 	{
 		static const FName Left("ToolbarSettingsRegion.Left");
-		return FAppStyle::Get().GetBrush(Left);
+		return GetCheckState() == ECheckBoxState::Checked ? FAppStyle::Get().GetBrush(ToggledLeft) : FAppStyle::Get().GetBrush(Left);
 	}
 	else
 	{
-		return FStyleDefaults::GetNoBrush();
+		return GetCheckState() == ECheckBoxState::Checked ? FAppStyle::Get().GetBrush(ToggledLeft) : FStyleDefaults::GetNoBrush();
 	}
 
 }
@@ -581,7 +595,7 @@ const FSlateBrush* SToolBarButtonBlock::GetOptionsBlockRightBrush() const
 		static const FName RightHover("ToolbarSettingsRegion.RightHover");
 		return FAppStyle::Get().GetBrush(RightHover);
 	}
-	else if (ButtonBorder->IsHovered())
+	else if (ButtonBorder->IsHovered() || GetCheckState() == ECheckBoxState::Checked)
 	{
 		static const FName Right("ToolbarSettingsRegion.Right");
 		return FAppStyle::Get().GetBrush(Right);
diff --git a/Engine/Source/Runtime/Slate/Private/Framework/MultiBox/SToolBarComboButtonBlock.cpp b/Engine/Source/Runtime/Slate/Private/Framework/MultiBox/SToolBarComboButtonBlock.cpp
--- a/Engine/Source/Runtime/Slate/Private/Framework/MultiBox/SToolBarComboButtonBlock.cpp
+++ b/Engine/Source/Runtime/Slate/Private/Framework/MultiBox/SToolBarComboButtonBlock.cpp
@@ -179,23 +179,24 @@ void SToolBarComboButtonBlock::BuildMultiBlockWidget(const ISlateStyle* StyleSet
 		
 	EMultiBlockLocation::Type BlockLocation = GetMultiBlockLocation();
 	FName BlockStyle = EMultiBlockLocation::ToName(ISlateStyle::Join(StyleName, ".Button"), BlockLocation);
-	const FButtonStyle& ButtonStyle = BlockLocation == EMultiBlockLocation::None ? ToolBarStyle.ButtonStyle : StyleSet->GetWidgetStyle<FButtonStyle>(BlockStyle);
-
-	OpenForegroundColor = ButtonStyle.HoveredForeground;
+	const FButtonStyle* ButtonStyle = BlockLocation == EMultiBlockLocation::None ? &ToolBarStyle.ButtonStyle : &StyleSet->GetWidgetStyle<FButtonStyle>(BlockStyle);
 
 	const FComboButtonStyle* ComboStyle = &ToolBarStyle.ComboButtonStyle;
 	if (ToolBarComboButtonBlock->bSimpleComboBox)
 	{
 		ComboStyle = &ToolBarStyle.SettingsComboButton;
+		ButtonStyle = &ComboStyle->ButtonStyle;
 	}
 
+	OpenForegroundColor = ButtonStyle->HoveredForeground;
+
 	ChildSlot
 	[
 		SAssignNew(ComboButtonWidget, SComboButton)
 		.AddMetaData<FTagMetaData>(FTagMetaData(TagName))
 		.ContentPadding(0.f)
 		.ComboButtonStyle(ComboStyle)
-		.ButtonStyle(&ButtonStyle)
+		.ButtonStyle(ButtonStyle)
 		.ToolTipText(ToolBarComboButtonBlock->ToolTip)
 		.ForegroundColor(this, &SToolBarComboButtonBlock::OnGetForegroundColor)
 		// Route the content generator event
diff --git a/Engine/Source/Runtime/Slate/Private/Framework/MultiBox/SWidgetBlock.cpp b/Engine/Source/Runtime/Slate/Private/Framework/MultiBox/SWidgetBlock.cpp
--- a/Engine/Source/Runtime/Slate/Private/Framework/MultiBox/SWidgetBlock.cpp
+++ b/Engine/Source/Runtime/Slate/Private/Framework/MultiBox/SWidgetBlock.cpp
@@ -168,7 +168,7 @@ void SWidgetBlock::BuildMultiBlockWidget(const ISlateStyle* StyleSet, const FNam
 			]
 		]
 		+SHorizontalBox::Slot()
-		.VAlign(bHasLabel ? VAlign_Bottom : VAlign_Fill)
+		.VAlign(bHasLabel ? VAlign_Center : VAlign_Fill)
 		.FillWidth(1.f)
 		[
 			WidgetBlock->ContentWidget
diff --git a/Engine/Source/Runtime/Slate/Private/Widgets/Input/SComboButton.cpp b/Engine/Source/Runtime/Slate/Private/Widgets/Input/SComboButton.cpp
--- a/Engine/Source/Runtime/Slate/Private/Widgets/Input/SComboButton.cpp
+++ b/Engine/Source/Runtime/Slate/Private/Widgets/Input/SComboButton.cpp
@@ -59,7 +59,7 @@ void SComboButton::Construct( const FArguments& InArgs )
 
 				+SHorizontalBox::Slot()
 				.AutoWidth()
-				.HAlign( HAlign_Center )
+				.HAlign( HAlign_Right )
 				.VAlign(InArgs._HasDownArrow ? (EVerticalAlignment) InArgs._ComboButtonStyle->DownArrowAlign : VAlign_Center)
 				.Padding(InArgs._HasDownArrow ? InArgs._ComboButtonStyle->DownArrowPadding : FMargin(0))
 				[
diff --git a/Engine/Source/Runtime/Slate/Private/Widgets/Input/SSearchBox.cpp b/Engine/Source/Runtime/Slate/Private/Widgets/Input/SSearchBox.cpp
--- a/Engine/Source/Runtime/Slate/Private/Widgets/Input/SSearchBox.cpp
+++ b/Engine/Source/Runtime/Slate/Private/Widgets/Input/SSearchBox.cpp
@@ -125,19 +125,18 @@ void SSearchBox::Construct( const FArguments& InArgs )
 			SNew(SImage)
 			.Visibility(this, &SSearchBox::GetSearchGlassVisibility)
 			.Image(&InArgs._Style->GlassImage)
-			.ColorAndOpacity(FAppStyle::Get().GetSlateColor("Colors.Foreground"))
+			.ColorAndOpacity(FSlateColor::UseForeground())
 		];
 	// Add an X to clear the search whenever there is some text typed into it
 	Box->InsertSlot(SlotIndex++)
 	.AutoWidth()
-	.Padding(InArgs._Style->ImagePadding)
 	.HAlign(HAlign_Center)
 	.VAlign(VAlign_Center)
 	[
 		SNew(SButton)
 		.Visibility(this, &SSearchBox::GetXVisibility)
-		.ButtonStyle(FAppStyle::Get(), "SimpleButton")
-		.ContentPadding(0.f)
+		.ButtonStyle(FAppStyle::Get(), "HoverOnlyButton")
+		.ContentPadding(FMargin(2.f, 0.0f))
 		.HAlign(HAlign_Center)
 		.VAlign(VAlign_Center)
 		.OnClicked(this, &SSearchBox::OnClearSearch)
@@ -146,8 +145,8 @@ void SSearchBox::Construct( const FArguments& InArgs )
 		.IsFocusable(true)
 		[
 			SNew(SImage)
-			.Image( &InArgs._Style->ClearImage )
-			.ColorAndOpacity( FSlateColor::UseForeground() )
+			.Image(&InArgs._Style->ClearImage)
+			.ColorAndOpacity(FSlateColor::UseForeground())
 		]
 	];
 }
diff --git a/Engine/Source/Runtime/Slate/Public/Framework/MultiBox/SToolBarButtonBlock.h b/Engine/Source/Runtime/Slate/Public/Framework/MultiBox/SToolBarButtonBlock.h
--- a/Engine/Source/Runtime/Slate/Public/Framework/MultiBox/SToolBarButtonBlock.h
+++ b/Engine/Source/Runtime/Slate/Public/Framework/MultiBox/SToolBarButtonBlock.h
@@ -150,7 +150,7 @@ protected:
 	 *
 	 * @return ECheckBoxState::Checked if it should be checked, ECheckBoxState::Unchecked if not.
 	 */
-	ECheckBoxState OnIsChecked() const;
+	ECheckBoxState GetCheckState() const;
 
 	/**
 	 * Called by Slate to determine if this button is enabled
diff --git a/Engine/Source/Runtime/Slate/Public/Widgets/Input/SHyperlink.h b/Engine/Source/Runtime/Slate/Public/Widgets/Input/SHyperlink.h
--- a/Engine/Source/Runtime/Slate/Public/Widgets/Input/SHyperlink.h
+++ b/Engine/Source/Runtime/Slate/Public/Widgets/Input/SHyperlink.h
@@ -12,6 +12,7 @@
 #include "Styling/SlateTypes.h"
 #include "Styling/CoreStyle.h"
 #include "Widgets/Input/SButton.h"
+#include "Widgets/Text/STextBlock.h"
 #if WITH_ACCESSIBILITY
 #include "Widgets/Accessibility/SlateAccessibleWidgets.h"
 #endif
@@ -65,14 +66,17 @@ public:
 
 		SButton::Construct(
 			SButton::FArguments()
-			.Text( InArgs._Text )
-			.ContentPadding( Padding )
-			.ButtonStyle( UnderlineStyle )
-			.TextStyle( TextStyle )
+			.ContentPadding(Padding)
+			.ButtonStyle(UnderlineStyle)
 			.OnClicked(this, &SHyperlink::Hyperlink_OnClicked)
 			.ForegroundColor(FSlateColor::UseForeground())
-			.TextShapingMethod( InArgs._TextShapingMethod )
-			.TextFlowDirection( InArgs._TextFlowDirection )
+			.TextShapingMethod(InArgs._TextShapingMethod)
+			.TextFlowDirection(InArgs._TextFlowDirection)
+			[
+				SNew(STextBlock)
+				.TextStyle(TextStyle)
+				.Text(InArgs._Text)
+			]
 		);
 	}
 
diff --git a/Engine/Source/Runtime/Slate/Public/Widgets/Input/SNumericEntryBox.h b/Engine/Source/Runtime/Slate/Public/Widgets/Input/SNumericEntryBox.h
--- a/Engine/Source/Runtime/Slate/Public/Widgets/Input/SNumericEntryBox.h
+++ b/Engine/Source/Runtime/Slate/Public/Widgets/Input/SNumericEntryBox.h
@@ -192,7 +192,7 @@ public:
 			CachedValueString = Interface->ToString(CachedExternalValue.GetValue());
 		}
 
-		TAttribute<FMargin> TextMargin = InArgs._OverrideTextMargin.IsSet() ? InArgs._OverrideTextMargin : InArgs._EditableTextBoxStyle->Padding;
+
 		const bool bAllowSpin = InArgs._AllowSpin;
 		TSharedPtr<SWidget> FinalWidget;
 
@@ -201,7 +201,6 @@ public:
 			SAssignNew(SpinBox, SSpinBox<NumericType>)
 				.Style(InArgs._SpinBoxStyle)
 				.Font(InArgs._Font.IsSet() ? InArgs._Font : InArgs._EditableTextBoxStyle->Font)
-				.ContentPadding(TextMargin)
 				.Value(this, &SNumericEntryBox<NumericType>::OnGetValueForSpinBox)
 				.Delta(InArgs._Delta)
 				.ShiftMouseMovePixelPerDelta(InArgs._ShiftMouseMovePixelPerDelta)
@@ -251,6 +250,9 @@ public:
 				];
 		}
 
+
+		TAttribute<FMargin> TextMargin = InArgs._OverrideTextMargin.IsSet() ? InArgs._OverrideTextMargin : InArgs._EditableTextBoxStyle->Padding;
+
 		Overlay->AddSlot()
 			.HAlign(HAlign_Fill) 
 			.VAlign(VAlign_Center)
diff --git a/Engine/Source/Runtime/Slate/Public/Widgets/Input/SSegmentedControl.h b/Engine/Source/Runtime/Slate/Public/Widgets/Input/SSegmentedControl.h
--- a/Engine/Source/Runtime/Slate/Public/Widgets/Input/SSegmentedControl.h
+++ b/Engine/Source/Runtime/Slate/Public/Widgets/Input/SSegmentedControl.h
@@ -118,7 +118,7 @@ public:
 
 	SLATE_BEGIN_ARGS( SSegmentedControl<OptionType> )
 		: _Style(&FAppStyle::Get().GetWidgetStyle<FSegmentedControlStyle>("SegmentedControl"))
-		, _TextStyle(&FAppStyle::Get().GetWidgetStyle<FTextBlockStyle>("ButtonText"))
+		, _TextStyle(&FAppStyle::Get().GetWidgetStyle<FTextBlockStyle>("SmallButtonText"))
 		, _MaxSegmentsPerLine(0)
 	{}
 		/** Slot type supported by this panel */
diff --git a/Engine/Source/Runtime/Slate/Public/Widgets/Input/SSpinBox.h b/Engine/Source/Runtime/Slate/Public/Widgets/Input/SSpinBox.h
--- a/Engine/Source/Runtime/Slate/Public/Widgets/Input/SSpinBox.h
+++ b/Engine/Source/Runtime/Slate/Public/Widgets/Input/SSpinBox.h
@@ -60,10 +60,10 @@ class SSpinBox
 public:
 
 	/** Notification for numeric value change */
-	DECLARE_DELEGATE_OneParam( FOnValueChanged, NumericType );
+	DECLARE_DELEGATE_OneParam(FOnValueChanged, NumericType);
 
 	/** Notification for numeric value committed */
-	DECLARE_DELEGATE_TwoParams( FOnValueCommitted, NumericType, ETextCommit::Type);
+	DECLARE_DELEGATE_TwoParams(FOnValueCommitted, NumericType, ETextCommit::Type);
 
 	/** Notification when the max/min spinner values are changed (only apply if SupportDynamicSliderMaxValue or SupportDynamicSliderMinValue are true) */
 	DECLARE_DELEGATE_FourParams(FOnDynamicSliderMinMaxValueChanged, NumericType, TWeakPtr<SWidget>, bool, bool);
@@ -86,21 +86,21 @@ public:
 		, _ClearKeyboardFocusOnCommit(false)
 		, _SelectAllTextOnCommit(true)
 		, _MinDesiredWidth(0.0f)
-		{}
+	{}
 
-		/** The style used to draw this spinbox */
-		SLATE_STYLE_ARGUMENT( FSpinBoxStyle, Style )
+	/** The style used to draw this spinbox */
+	SLATE_STYLE_ARGUMENT(FSpinBoxStyle, Style)
 
 		/** The value to display */
-		SLATE_ATTRIBUTE( NumericType, Value )
+		SLATE_ATTRIBUTE(NumericType, Value)
 		/** The minimum value that can be entered into the text edit box */
-		SLATE_ATTRIBUTE( TOptional< NumericType >, MinValue )
+		SLATE_ATTRIBUTE(TOptional< NumericType >, MinValue)
 		/** The maximum value that can be entered into the text edit box */
-		SLATE_ATTRIBUTE( TOptional< NumericType >, MaxValue )
+		SLATE_ATTRIBUTE(TOptional< NumericType >, MaxValue)
 		/** The minimum value that can be specified by using the slider, defaults to MinValue */
-		SLATE_ATTRIBUTE( TOptional< NumericType >, MinSliderValue )
+		SLATE_ATTRIBUTE(TOptional< NumericType >, MinSliderValue)
 		/** The maximum value that can be specified by using the slider, defaults to MaxValue */
-		SLATE_ATTRIBUTE( TOptional< NumericType >, MaxSliderValue )
+		SLATE_ATTRIBUTE(TOptional< NumericType >, MaxSliderValue)
 		/** The minimum fractional digits the spin box displays, defaults to 1 */
 		SLATE_ATTRIBUTE(TOptional< int32 >, MinFractionalDigits)
 		/** The maximum fractional digits the spin box displays, defaults to 6 */
@@ -108,49 +108,49 @@ public:
 		/** Whether typed values should use delta snapping, defaults to false */
 		SLATE_ATTRIBUTE(bool, AlwaysUsesDeltaSnap)
 		/** Delta to increment the value as the slider moves.  If not specified will determine automatically */
-		SLATE_ATTRIBUTE( NumericType, Delta )
+		SLATE_ATTRIBUTE(NumericType, Delta)
 		/** How many pixel the mouse must move to change the value of the delta step */
-		SLATE_ATTRIBUTE( int32, ShiftMouseMovePixelPerDelta )
+		SLATE_ATTRIBUTE(int32, ShiftMouseMovePixelPerDelta)
 		/** If we're an unbounded spinbox, what value do we divide mouse movement by before multiplying by Delta. Requires Delta to be set. */
-		SLATE_ATTRIBUTE( int32, LinearDeltaSensitivity)
+		SLATE_ATTRIBUTE(int32, LinearDeltaSensitivity)
 		/** Tell us if we want to support dynamically changing of the max value using ctrl */
 		SLATE_ATTRIBUTE(bool, SupportDynamicSliderMaxValue)
 		/** Tell us if we want to support dynamically changing of the min value using ctrl */
 		SLATE_ATTRIBUTE(bool, SupportDynamicSliderMinValue)
 		/** Called right after the max slider value is changed (only relevant if SupportDynamicSliderMaxValue is true) */
-		SLATE_EVENT( FOnDynamicSliderMinMaxValueChanged, OnDynamicSliderMaxValueChanged )
+		SLATE_EVENT(FOnDynamicSliderMinMaxValueChanged, OnDynamicSliderMaxValueChanged)
 		/** Called right after the min slider value is changed (only relevant if SupportDynamicSliderMinValue is true) */
-		SLATE_EVENT( FOnDynamicSliderMinMaxValueChanged, OnDynamicSliderMinValueChanged )
+		SLATE_EVENT(FOnDynamicSliderMinMaxValueChanged, OnDynamicSliderMinValueChanged)
 		/** Use exponential scale for the slider */
-		SLATE_ATTRIBUTE( float, SliderExponent )
+		SLATE_ATTRIBUTE(float, SliderExponent)
 		/** When use exponential scale for the slider which is the neutral value */
-		SLATE_ATTRIBUTE( NumericType, SliderExponentNeutralValue )
+		SLATE_ATTRIBUTE(NumericType, SliderExponentNeutralValue)
 		/** Font used to display text in the slider */
-		SLATE_ATTRIBUTE( FSlateFontInfo, Font )
+		SLATE_ATTRIBUTE(FSlateFontInfo, Font)
 		/** Padding to add around this widget and its internal widgets */
-		SLATE_ATTRIBUTE( FMargin, ContentPadding )
+		SLATE_ATTRIBUTE(FMargin, ContentPadding)
 		/** Called when the value is changed by slider or typing */
-		SLATE_EVENT( FOnValueChanged, OnValueChanged )
+		SLATE_EVENT(FOnValueChanged, OnValueChanged)
 		/** Called when the value is committed (by pressing enter) */
-		SLATE_EVENT( FOnValueCommitted, OnValueCommitted )
+		SLATE_EVENT(FOnValueCommitted, OnValueCommitted)
 		/** Called right before the slider begins to move */
-		SLATE_EVENT( FSimpleDelegate, OnBeginSliderMovement )
+		SLATE_EVENT(FSimpleDelegate, OnBeginSliderMovement)
 		/** Called right after the slider handle is released by the user */
-		SLATE_EVENT( FOnValueChanged, OnEndSliderMovement )
+		SLATE_EVENT(FOnValueChanged, OnEndSliderMovement)
 		/** Whether to clear keyboard focus when pressing enter to commit changes */
-		SLATE_ATTRIBUTE( bool, ClearKeyboardFocusOnCommit )
+		SLATE_ATTRIBUTE(bool, ClearKeyboardFocusOnCommit)
 		/** Whether to select all text when pressing enter to commit changes */
-		SLATE_ATTRIBUTE( bool, SelectAllTextOnCommit )
+		SLATE_ATTRIBUTE(bool, SelectAllTextOnCommit)
 		/** Minimum width that a spin box should be */
-		SLATE_ATTRIBUTE( float, MinDesiredWidth )
+		SLATE_ATTRIBUTE(float, MinDesiredWidth)
 		/** How should the value be justified in the spinbox. */
-		SLATE_ATTRIBUTE( ETextJustify::Type, Justification )
+		SLATE_ATTRIBUTE(ETextJustify::Type, Justification)
 		/** Provide custom type conversion functionality to this spin box */
-		SLATE_ARGUMENT( TSharedPtr< INumericTypeInterface<NumericType> >, TypeInterface )
+		SLATE_ARGUMENT(TSharedPtr< INumericTypeInterface<NumericType> >, TypeInterface)
 		/** If refresh requests for the viewport should happen for all value changes **/
 		SLATE_ARGUMENT(bool, PreventThrottling)
 		/** Menu extender for the right-click context menu */
-		SLATE_EVENT( FMenuExtensionDelegate, ContextMenuExtender )
+		SLATE_EVENT(FMenuExtensionDelegate, ContextMenuExtender)
 
 	SLATE_END_ARGS()
 
@@ -163,14 +163,14 @@ public:
 	 *
 	 * @param InArgs   A declaration from which to construct the widget
 	 */
-	void Construct( const FArguments& InArgs )
+	void Construct(const FArguments& InArgs)
 	{
 		check(InArgs._Style);
 
 		Style = InArgs._Style;
 
 		SetForegroundColor(InArgs._Style->ForegroundColor);
-		Interface = InArgs._TypeInterface.IsValid() ? InArgs._TypeInterface : MakeShareable( new TDefaultNumericTypeInterface<NumericType> );
+		Interface = InArgs._TypeInterface.IsValid() ? InArgs._TypeInterface : MakeShareable(new TDefaultNumericTypeInterface<NumericType>);
 
 		ValueAttribute = InArgs._Value;
 		OnValueChanged = InArgs._OnValueChanged;
@@ -227,8 +227,12 @@ public:
 
 		BackgroundHoveredBrush = &InArgs._Style->HoveredBackgroundBrush;
 		BackgroundBrush = &InArgs._Style->BackgroundBrush;
+		BackgroundActiveBrush = InArgs._Style->ActiveBackgroundBrush.IsSet() ? &InArgs._Style->ActiveBackgroundBrush : BackgroundHoveredBrush;
+
 		ActiveFillBrush = &InArgs._Style->ActiveFillBrush;
+		HoveredFillBrush = InArgs._Style->HoveredFillBrush.IsSet() ? &InArgs._Style->HoveredFillBrush : ActiveFillBrush;
 		InactiveFillBrush = &InArgs._Style->InactiveFillBrush;
+
 		const FMargin& TextMargin = InArgs._Style->TextPadding;
 
 		bDragging = false;
@@ -237,68 +241,67 @@ public:
 		bIsTextChanging = false;
 
 		this->ChildSlot
-		.Padding( InArgs._ContentPadding )
+		.Padding(InArgs._ContentPadding)
 		[
 			SNew(SHorizontalBox)
 
 			+ SHorizontalBox::Slot()
 			.FillWidth(1.0f)
-			.Padding( TextMargin )
-			.HAlign(HAlign_Fill) 
+			.Padding(TextMargin)
+			.HAlign(HAlign_Fill)
 			.VAlign(VAlign_Center)
 			[
 				SAssignNew(TextBlock, STextBlock)
 				.Font(InArgs._Font)
-				.Text( this, &SSpinBox<NumericType>::GetValueAsText )
-				.MinDesiredWidth( this, &SSpinBox<NumericType>::GetTextMinDesiredWidth )
+				.Text(this, &SSpinBox<NumericType>::GetValueAsText)
+				.MinDesiredWidth(this, &SSpinBox<NumericType>::GetTextMinDesiredWidth)
 				.Justification(InArgs._Justification)
 			]
-
 			+ SHorizontalBox::Slot()
 			.FillWidth(1.0f)
-			.Padding( TextMargin )
-			.HAlign(HAlign_Fill) 
-			.VAlign(VAlign_Center) 
+			.Padding(TextMargin)
+			.HAlign(HAlign_Fill)
+			.VAlign(VAlign_Center)
 			[
 				SAssignNew(EditableText, SEditableText)
-				.Visibility( EVisibility::Collapsed )
+				.Visibility(EVisibility::Collapsed)
 				.Font(InArgs._Font)
-				.SelectAllTextWhenFocused( true )
-				.Text( this, &SSpinBox<NumericType>::GetValueAsText )
+				.SelectAllTextWhenFocused(true)
+				.Text(this, &SSpinBox<NumericType>::GetValueAsText)
 				.OnIsTypedCharValid(this, &SSpinBox<NumericType>::IsCharacterValid)
-				.OnTextChanged( this, &SSpinBox<NumericType>::TextField_OnTextChanged )
-				.OnTextCommitted( this, &SSpinBox<NumericType>::TextField_OnTextCommitted )
-				.ClearKeyboardFocusOnCommit( InArgs._ClearKeyboardFocusOnCommit )
-				.SelectAllTextOnCommit( InArgs._SelectAllTextOnCommit )
-				.MinDesiredWidth( this, &SSpinBox<NumericType>::GetTextMinDesiredWidth )
+				.OnTextChanged(this, &SSpinBox<NumericType>::TextField_OnTextChanged)
+				.OnTextCommitted(this, &SSpinBox<NumericType>::TextField_OnTextCommitted)
+				.ClearKeyboardFocusOnCommit(InArgs._ClearKeyboardFocusOnCommit)
+				.SelectAllTextOnCommit(InArgs._SelectAllTextOnCommit)
+				.MinDesiredWidth(this, &SSpinBox<NumericType>::GetTextMinDesiredWidth)
 				.VirtualKeyboardType(EKeyboardType::Keyboard_Number)
 				.Justification(InArgs._Justification)
 				.VirtualKeyboardTrigger(EVirtualKeyboardTrigger::OnAllFocusEvents)
-				.ContextMenuExtender( InArgs._ContextMenuExtender )
+				.ContextMenuExtender(InArgs._ContextMenuExtender)
 			]
-
-			+SHorizontalBox::Slot()
+			+ SHorizontalBox::Slot()
 			.AutoWidth()
-			.HAlign(HAlign_Fill) 
-			.VAlign(VAlign_Center) 
+			.HAlign(HAlign_Fill)
+			.VAlign(VAlign_Center)
 			[
 				SNew(SImage)
-				.Image( &InArgs._Style->ArrowsImage )
-				.ColorAndOpacity( FSlateColor::UseForeground())
+				.Image(&InArgs._Style->ArrowsImage)
+				.ColorAndOpacity(FSlateColor::UseForeground())
 			]
 		];
 	}
 
-	virtual int32 OnPaint( const FPaintArgs& Args, const FGeometry& AllottedGeometry, const FSlateRect& MyCullingRect, FSlateWindowElementList& OutDrawElements, int32 LayerId, const FWidgetStyle& InWidgetStyle, bool bParentEnabled ) const override
+	virtual int32 OnPaint(const FPaintArgs& Args, const FGeometry& AllottedGeometry, const FSlateRect& MyCullingRect, FSlateWindowElementList& OutDrawElements, int32 LayerId, const FWidgetStyle& InWidgetStyle, bool bParentEnabled) const override
 	{
-		const bool bActiveFeedback = IsHovered() || bDragging;
 
-		const FSlateBrush* BackgroundImage = bActiveFeedback ?
+		const bool bActiveFeedback = bDragging;
+
+		const FSlateBrush* BackgroundImage = bActiveFeedback ? BackgroundActiveBrush : IsHovered() ?
 			BackgroundHoveredBrush :
 			BackgroundBrush;
 
 		const FSlateBrush* FillImage = bActiveFeedback ?
-			ActiveFillBrush :
+			ActiveFillBrush : IsHovered() ? HoveredFillBrush : 
 			InactiveFillBrush;
 
 		const int32 BackgroundLayer = LayerId;
@@ -313,7 +316,7 @@ public:
 			BackgroundImage,
 			DrawEffects,
 			BackgroundImage->GetTint(InWidgetStyle) * InWidgetStyle.GetColorAndOpacityTint()
-			);
+		);
 
 		const int32 FilledLayer = BackgroundLayer + 1;
 
@@ -342,18 +345,18 @@ public:
 					FractionFilled = 1.0f - FMath::Pow(1.0f - FractionFilled, CachedSliderExponent);
 				}
 			}
-			const FVector2D FillSize( AllottedGeometry.GetLocalSize().X * FractionFilled, AllottedGeometry.GetLocalSize().Y );
+			const FVector2D FillSize(AllottedGeometry.GetLocalSize().X * FractionFilled, AllottedGeometry.GetLocalSize().Y);
 
-			if ( ! IsInTextMode() )
+			if (!IsInTextMode())
 			{
 				FSlateDrawElement::MakeBox(
 					OutDrawElements,
 					FilledLayer,
-					AllottedGeometry.ToPaintGeometry(FVector2D(0.0f, 0.0f), FillSize),
+					AllottedGeometry.ToPaintGeometry(Style->InsetPadding.GetTopLeft(), FillSize - FVector2D(Style->InsetPadding.GetTotalSpaceAlong<Orient_Horizontal>(), Style->InsetPadding.GetTotalSpaceAlong<Orient_Vertical>())),
 					FillImage,
 					DrawEffects,
 					FillImage->GetTint(InWidgetStyle) * InWidgetStyle.GetColorAndOpacityTint()
-					);
+				);
 			}
 		}
 
@@ -367,9 +370,9 @@ public:
 	 * @param MouseEvent Information about the input event
 	 * @return Whether the event was handled along with possible requests for the system to take action.
 	 */
-	virtual FReply OnMouseButtonDown( const FGeometry& MyGeometry, const FPointerEvent& MouseEvent ) override
+	virtual FReply OnMouseButtonDown(const FGeometry& MyGeometry, const FPointerEvent& MouseEvent) override
 	{
-		if ( MouseEvent.GetEffectingButton() == EKeys::LeftMouseButton && PointerDraggingSliderIndex == INDEX_NONE )
+		if (MouseEvent.GetEffectingButton() == EKeys::LeftMouseButton && PointerDraggingSliderIndex == INDEX_NONE)
 		{
 			DistanceDragged = 0.f;
 			PreDragValue = ValueAttribute.Get();
@@ -398,11 +401,11 @@ public:
 	 * @param MouseEvent Information about the input event
 	 * @return Whether the event was handled along with possible requests for the system to take action.
 	 */
-	virtual FReply OnMouseButtonUp( const FGeometry& MyGeometry, const FPointerEvent& MouseEvent ) override
+	virtual FReply OnMouseButtonUp(const FGeometry& MyGeometry, const FPointerEvent& MouseEvent) override
 	{
-		if ( MouseEvent.GetEffectingButton() == EKeys::LeftMouseButton && PointerDraggingSliderIndex == MouseEvent.GetPointerIndex() )
+		if (MouseEvent.GetEffectingButton() == EKeys::LeftMouseButton && PointerDraggingSliderIndex == MouseEvent.GetPointerIndex())
 		{
-			if ( !this->HasMouseCapture() )
+			if (!this->HasMouseCapture())
 			{
 				// Lost Capture - ensure reset
 				bDragging = false;
@@ -411,7 +414,7 @@ public:
 				return FReply::Unhandled();
 			}
 
-			if ( bDragging )
+			if (bDragging)
 			{
 				NumericType CurrentDelta = Delta.Get();
 				if (CurrentDelta != NumericType())
@@ -428,12 +431,12 @@ public:
 
 			FReply Reply = FReply::Handled().ReleaseMouseCapture();
 
-			if ( !MouseEvent.IsTouchEvent() )
+			if (!MouseEvent.IsTouchEvent())
 			{
 				Reply.SetMousePos(CachedMousePosition);
 			}
 
-			if ( DistanceDragged < FSlateApplication::Get().GetDragTriggerDistance() )
+			if (DistanceDragged < FSlateApplication::Get().GetDragTriggerDistance())
 			{
 				EnterTextMode();
 				Reply.SetUserFocus(EditableText.ToSharedRef(), EFocusCause::SetDirectly);
@@ -506,11 +509,11 @@ public:
 	 * @param MouseEvent Information about the input event
 	 * @return Whether the event was handled along with possible requests for the system to take action.
 	 */
-	virtual FReply OnMouseMove( const FGeometry& MyGeometry, const FPointerEvent& MouseEvent ) override
+	virtual FReply OnMouseMove(const FGeometry& MyGeometry, const FPointerEvent& MouseEvent) override
 	{
-		if ( PointerDraggingSliderIndex == MouseEvent.GetPointerIndex() )
+		if (PointerDraggingSliderIndex == MouseEvent.GetPointerIndex())
 		{
-			if(!this->HasMouseCapture())
+			if (!this->HasMouseCapture())
 			{
 				// Lost the mouse capture - ensure reset
 				bDragging = false;
@@ -641,11 +644,11 @@ public:
 		return FReply::Unhandled();
 	}
 
-	virtual FCursorReply OnCursorQuery( const FGeometry& MyGeometry, const FPointerEvent& CursorEvent ) const override
+	virtual FCursorReply OnCursorQuery(const FGeometry& MyGeometry, const FPointerEvent& CursorEvent) const override
 	{
-		return bDragging ? 
-			FCursorReply::Cursor( EMouseCursor::None ) :
-			FCursorReply::Cursor( EMouseCursor::ResizeLeftRight );
+		return bDragging ?
+			FCursorReply::Cursor(EMouseCursor::None) :
+			FCursorReply::Cursor(EMouseCursor::ResizeLeftRight);
 	}
 
 	virtual bool SupportsKeyboardFocus() const override
@@ -655,9 +658,9 @@ public:
 	}
 
 
-	virtual FReply OnFocusReceived( const FGeometry& MyGeometry, const FFocusEvent& InFocusEvent ) override
+	virtual FReply OnFocusReceived(const FGeometry& MyGeometry, const FFocusEvent& InFocusEvent) override
 	{
-		if ( !bDragging && (InFocusEvent.GetCause() == EFocusCause::Navigation || InFocusEvent.GetCause() == EFocusCause::SetDirectly) )
+		if (!bDragging && (InFocusEvent.GetCause() == EFocusCause::Navigation || InFocusEvent.GetCause() == EFocusCause::SetDirectly))
 		{
 			EnterTextMode();
 			return FReply::Handled().SetUserFocus(EditableText.ToSharedRef(), InFocusEvent.GetCause());
@@ -668,10 +671,10 @@ public:
 		}
 	}
 
-	virtual FReply OnKeyDown( const FGeometry& MyGeometry, const FKeyEvent& InKeyEvent ) override
+	virtual FReply OnKeyDown(const FGeometry& MyGeometry, const FKeyEvent& InKeyEvent) override
 	{
 		const FKey Key = InKeyEvent.GetKey();
-		if ( Key == EKeys::Escape && HasMouseCapture())
+		if (Key == EKeys::Escape && HasMouseCapture())
 		{
 			bDragging = false;
 			PointerDraggingSliderIndex = INDEX_NONE;
@@ -680,7 +683,7 @@ public:
 			NotifyValueCommitted(PreDragValue);
 			return FReply::Handled().ReleaseMouseCapture().SetMousePos(CachedMousePosition);
 		}
-		else if ( Key == EKeys::Up || Key == EKeys::Right )
+		else if (Key == EKeys::Up || Key == EKeys::Right)
 		{
 			const NumericType LocalValueAttribute = ValueAttribute.Get();
 			const NumericType LocalDelta = Delta.Get();
@@ -689,7 +692,7 @@ public:
 			ExitTextMode();
 			return FReply::Handled();
 		}
-		else if ( Key == EKeys::Down || Key == EKeys::Left )
+		else if (Key == EKeys::Down || Key == EKeys::Left)
 		{
 			const NumericType LocalValueAttribute = ValueAttribute.Get();
 			const NumericType LocalDelta = Delta.Get();
@@ -698,7 +701,7 @@ public:
 			ExitTextMode();
 			return FReply::Handled();
 		}
-		else if ( Key == EKeys::Enter )
+		else if (Key == EKeys::Enter)
 		{
 			InternalValue = (double)ValueAttribute.Get();
 			EnterTextMode();
@@ -721,16 +724,16 @@ public:
 
 	/** See the Value attribute */
 	NumericType GetValue() const { return ValueAttribute.Get(); }
-	void SetValue(const TAttribute<NumericType>& InValueAttribute) 
+	void SetValue(const TAttribute<NumericType>& InValueAttribute)
 	{
-		ValueAttribute = InValueAttribute; 
+		ValueAttribute = InValueAttribute;
 		const NumericType LocalValueAttribute = ValueAttribute.Get();
 		CommitValue(LocalValueAttribute, (double)LocalValueAttribute, ECommitMethod::CommittedViaCode, ETextCommit::Default);
 	}
 
 	/** See the MinValue attribute */
 	NumericType GetMinValue() const { return MinValue.Get().Get(std::numeric_limits<NumericType>::lowest()); }
-	void SetMinValue(const TAttribute<TOptional<NumericType>>& InMinValue) 
+	void SetMinValue(const TAttribute<TOptional<NumericType>>& InMinValue)
 	{
 		MinValue = InMinValue;
 		UpdateIsSpinRangeUnlimited();
@@ -738,9 +741,9 @@ public:
 
 	/** See the MaxValue attribute */
 	NumericType GetMaxValue() const { return MaxValue.Get().Get(std::numeric_limits<NumericType>::max()); }
-	void SetMaxValue(const TAttribute<TOptional<NumericType>>& InMaxValue) 
+	void SetMaxValue(const TAttribute<TOptional<NumericType>>& InMaxValue)
 	{
-		MaxValue = InMaxValue; 
+		MaxValue = InMaxValue;
 		UpdateIsSpinRangeUnlimited();
 	}
 
@@ -748,7 +751,7 @@ public:
 	bool IsMinSliderValueBound() const { return MinSliderValue.IsBound(); }
 
 	NumericType GetMinSliderValue() const { return MinSliderValue.Get().Get(std::numeric_limits<NumericType>::lowest()); }
-	void SetMinSliderValue(const TAttribute<TOptional<NumericType>>& InMinSliderValue) 
+	void SetMinSliderValue(const TAttribute<TOptional<NumericType>>& InMinSliderValue)
 	{
 		MinSliderValue = (InMinSliderValue.Get().IsSet()) ? InMinSliderValue : MinValue;
 		UpdateIsSpinRangeUnlimited();
@@ -758,7 +761,7 @@ public:
 	bool IsMaxSliderValueBound() const { return MaxSliderValue.IsBound(); }
 
 	NumericType GetMaxSliderValue() const { return MaxSliderValue.Get().Get(std::numeric_limits<NumericType>::max()); }
-	void SetMaxSliderValue(const TAttribute<TOptional<NumericType>>& InMaxSliderValue) 
+	void SetMaxSliderValue(const TAttribute<TOptional<NumericType>>& InMaxSliderValue)
 	{
 		MaxSliderValue = (InMaxSliderValue.Get().IsSet()) ? InMaxSliderValue : MaxValue;;
 		UpdateIsSpinRangeUnlimited();
@@ -784,7 +787,7 @@ public:
 
 	/** See the Delta attribute */
 	NumericType GetDelta() const { return Delta.Get(); }
-	void SetDelta(NumericType InDelta) { Delta.Set( InDelta ); }
+	void SetDelta(NumericType InDelta) { Delta.Set(InDelta); }
 
 	/** See the SliderExponent attribute */
 	float GetSliderExponent() const { return SliderExponent.Get(); }
@@ -798,15 +801,15 @@ protected:
 	/** Make the spinbox switch to keyboard-based input mode. */
 	void EnterTextMode()
 	{
-		TextBlock->SetVisibility( EVisibility::Collapsed );
-		EditableText->SetVisibility( EVisibility::Visible );
+		TextBlock->SetVisibility(EVisibility::Collapsed);
+		EditableText->SetVisibility(EVisibility::Visible);
 	}
 
 	/** Make the spinbox switch to mouse-based input mode. */
 	void ExitTextMode()
 	{
-		TextBlock->SetVisibility( EVisibility::Visible );
-		EditableText->SetVisibility( EVisibility::Collapsed );
+		TextBlock->SetVisibility(EVisibility::Visible);
+		EditableText->SetVisibility(EVisibility::Collapsed);
 	}
 
 	/** @return the value being observed by the spinbox as a string */
@@ -865,7 +868,7 @@ protected:
 	 * @param NewText		The value of text coming from the editable text field.
 	 * @param CommitInfo	Information about the source of the commit
 	 */
-	void TextField_OnTextCommitted( const FText& NewText, ETextCommit::Type CommitInfo )
+	void TextField_OnTextCommitted(const FText& NewText, ETextCommit::Type CommitInfo)
 	{
 		if (CommitInfo != ETextCommit::OnEnter)
 		{
@@ -1035,8 +1038,10 @@ private:
 	const FSpinBoxStyle* Style;
 
 	const FSlateBrush* BackgroundHoveredBrush;
+	const FSlateBrush* BackgroundActiveBrush;
 	const FSlateBrush* BackgroundBrush;
 	const FSlateBrush* ActiveFillBrush;
+	const FSlateBrush* HoveredFillBrush;
 	const FSlateBrush* InactiveFillBrush;
 
 	float DistanceDragged;
@@ -1081,7 +1086,7 @@ private:
 		constexpr bool bCanBeRepresentedInDouble = std::numeric_limits<double>::digits >= std::numeric_limits<NumericType>::digits;
 		if (bIsIntegral && !bCanBeRepresentedInDouble)
 		{
-			return (NumericType)FMath::Clamp<double>(FMath::FloorToDouble(ValueToRound + 0.5), -1.0*(double)(1ll << std::numeric_limits<double>::digits), (double)(1ll << std::numeric_limits<double>::digits));
+			return (NumericType)FMath::Clamp<double>(FMath::FloorToDouble(ValueToRound + 0.5), -1.0 * (double)(1ll << std::numeric_limits<double>::digits), (double)(1ll << std::numeric_limits<double>::digits));
 		}
 		else if (bIsIntegral)
 		{
diff --git a/Engine/Source/Runtime/SlateCore/Private/Styling/SlateBrush.cpp b/Engine/Source/Runtime/SlateCore/Private/Styling/SlateBrush.cpp
--- a/Engine/Source/Runtime/SlateCore/Private/Styling/SlateBrush.cpp
+++ b/Engine/Source/Runtime/SlateCore/Private/Styling/SlateBrush.cpp
@@ -20,6 +20,7 @@ FSlateBrush::FSlateBrush()
 	, ImageType(ESlateBrushImageType::NoImage)
 	, bIsDynamicallyLoaded(false)
 	, bHasUObject_DEPRECATED(false)
+	, bIsSet(true)
 {
 }
 
@@ -48,6 +49,7 @@ FSlateBrush::FSlateBrush( ESlateBrushDrawType::Type InDrawType,
 	, Mirroring( ESlateBrushMirrorType::NoMirror )
 	, ImageType( InImageType )
 	, bIsDynamicallyLoaded( bInDynamicallyLoaded )
+	, bIsSet(true)
 {
 	bHasUObject_DEPRECATED = (InObjectResource != nullptr) || InResourceName.ToString().StartsWith(FSlateBrush::UTextureIdentifier());
 
@@ -83,6 +85,7 @@ FSlateBrush::FSlateBrush( ESlateBrushDrawType::Type InDrawType,
 	, Mirroring( ESlateBrushMirrorType::NoMirror )
 	, ImageType( InImageType )
 	, bIsDynamicallyLoaded( bInDynamicallyLoaded )
+	, bIsSet(true)
 {
 	bHasUObject_DEPRECATED = (InObjectResource != nullptr) || InResourceName.ToString().StartsWith(FSlateBrush::UTextureIdentifier());
 
@@ -117,6 +120,7 @@ FSlateBrush::FSlateBrush( ESlateBrushDrawType::Type InDrawType,
 	, Mirroring( ESlateBrushMirrorType::NoMirror )
 	, ImageType(InImageType)
 	, bIsDynamicallyLoaded(bInDynamicallyLoaded)
+	, bIsSet(true)
 {
 	bHasUObject_DEPRECATED = (InObjectResource != nullptr) || InResourceName.ToString().StartsWith(FSlateBrush::UTextureIdentifier());
 
diff --git a/Engine/Source/Runtime/SlateCore/Private/Styling/SlateTypes.cpp b/Engine/Source/Runtime/SlateCore/Private/Styling/SlateTypes.cpp
--- a/Engine/Source/Runtime/SlateCore/Private/Styling/SlateTypes.cpp
+++ b/Engine/Source/Runtime/SlateCore/Private/Styling/SlateTypes.cpp
@@ -474,18 +474,23 @@ const FInlineTextImageStyle& FInlineTextImageStyle::GetDefault()
 }
 
 FSpinBoxStyle::FSpinBoxStyle()
-	: ForegroundColor(FSlateColor::UseForeground())
+	: ActiveBackgroundBrush(FSlateOptionalBrush())
+	, HoveredFillBrush(FSlateOptionalBrush())
+	, ForegroundColor(FSlateColor::UseForeground())
 	, TextPadding(FMargin(1.0f,2.0f))
+	, InsetPadding(FMargin(0))
 {
 }
 
 void FSpinBoxStyle::GetResources( TArray< const FSlateBrush* >& OutBrushes ) const
 {
-	OutBrushes.Add( &BackgroundBrush );
-	OutBrushes.Add( &HoveredBackgroundBrush );
-	OutBrushes.Add( &ActiveFillBrush );
-	OutBrushes.Add( &InactiveFillBrush );
-	OutBrushes.Add( &ArrowsImage );
+	OutBrushes.Add(&BackgroundBrush);
+	OutBrushes.Add(&HoveredBackgroundBrush);
+	OutBrushes.Add(&ActiveBackgroundBrush);
+	OutBrushes.Add(&ActiveFillBrush);
+	OutBrushes.Add(&HoveredFillBrush);
+	OutBrushes.Add(&InactiveFillBrush);
+	OutBrushes.Add(&ArrowsImage);
 }
 
 const FName FSpinBoxStyle::TypeName( TEXT("FSpinBoxStyle") );
diff --git a/Engine/Source/Runtime/SlateCore/Private/Styling/StarshipCoreStyle.cpp b/Engine/Source/Runtime/SlateCore/Private/Styling/StarshipCoreStyle.cpp
--- a/Engine/Source/Runtime/SlateCore/Private/Styling/StarshipCoreStyle.cpp
+++ b/Engine/Source/Runtime/SlateCore/Private/Styling/StarshipCoreStyle.cpp
@@ -421,9 +421,6 @@ TSharedRef<ISlateStyle> FStarshipCoreStyle::Create()
 	SetupComboButtonStyles(Style);
 	SetupCheckboxStyles(Style);
 
-	const FButtonStyle& Button = Style->GetWidgetStyle<FButtonStyle>("Button");
-	const FButtonStyle& SimpleButton = Style->GetWidgetStyle<FButtonStyle>("SimpleButton");
-
 	// SMessageLogListing
 	{
 		FComboButtonStyle MessageLogListingComboButton = FComboButtonStyle()
@@ -480,8 +477,8 @@ TSharedRef<ISlateStyle> FStarshipCoreStyle::Create()
 	{
 		Style->Set("ProgressBar", FProgressBarStyle()
 			.SetBackgroundImage(FSlateColorBrush(FStyleColors::Foldout))
-			.SetFillImage(IMAGE_BRUSH("Starship/CoreWidgets/ProgressBar/ProgressMarquee", FVector2D(20.f, 12.f), FStyleColors::Primary, ESlateBrushTileType::Horizontal))
-			.SetMarqueeImage(IMAGE_BRUSH("Starship/CoreWidgets/ProgressBar/ProgressMarquee", FVector2D(20.f, 12.f), FStyleColors::Primary, ESlateBrushTileType::Horizontal))
+			.SetFillImage(IMAGE_BRUSH("Starship/CoreWidgets/ProgressBar/ProgressMarquee", FVector2D(20.f, 12.f), FStyleColors::PrimaryHover, ESlateBrushTileType::Horizontal))
+			.SetMarqueeImage(IMAGE_BRUSH("Starship/CoreWidgets/ProgressBar/ProgressMarquee", FVector2D(20.f, 12.f), FStyleColors::PrimaryHover, ESlateBrushTileType::Horizontal))
 			.SetEnableFillAnimation(true)
 		);
 	}
@@ -508,7 +505,7 @@ TSharedRef<ISlateStyle> FStarshipCoreStyle::Create()
 		FSliderStyle SliderStyle = FSliderStyle()
 			.SetNormalBarImage(   FSlateRoundedBoxBrush(FStyleColors::Input, 2.0f, FStyleColors::Input, 1.0f))
 			.SetHoveredBarImage(  FSlateRoundedBoxBrush(FStyleColors::Input, 2.0f, FStyleColors::Input, 1.0f))
-			.SetNormalThumbImage(  FSlateRoundedBoxBrush(FStyleColors::Hover2, Icon8x8) )
+			.SetNormalThumbImage(  FSlateRoundedBoxBrush(FStyleColors::Foreground, Icon8x8) )
 			.SetHoveredThumbImage( FSlateRoundedBoxBrush(FStyleColors::ForegroundHover, Icon8x8) )
 			.SetBarThickness(4.0f);
 		Style->Set("Slider", SliderStyle);
@@ -523,43 +520,33 @@ TSharedRef<ISlateStyle> FStarshipCoreStyle::Create()
 		);
 	}
 
-	// SSpinBox defaults...
+	// SpinBox defaults...
 	{
-		Style->Set("SpinBox", FSpinBoxStyle()
+		// SSpinBox
+		FSpinBoxStyle DefaultSpinBox =
+			FSpinBoxStyle()
 			.SetBackgroundBrush(FSlateRoundedBoxBrush(FStyleColors::Input, InputFocusRadius, FStyleColors::InputOutline, InputFocusThickness))
 			.SetHoveredBackgroundBrush(FSlateRoundedBoxBrush(FStyleColors::Input, InputFocusRadius, FStyleColors::Hover, InputFocusThickness))
-
-			.SetActiveFillBrush(FSlateRoundedBoxBrush(FStyleColors::Hover, InputFocusRadius, FLinearColor::Transparent, InputFocusThickness))
-			.SetInactiveFillBrush(FSlateRoundedBoxBrush(FStyleColors::Secondary, InputFocusRadius, FLinearColor::Transparent, InputFocusThickness))
+			.SetActiveBackgroundBrush(FSlateRoundedBoxBrush(FStyleColors::Input, InputFocusRadius, FStyleColors::Primary, InputFocusThickness))
+			.SetActiveFillBrush(FSlateRoundedBoxBrush(FStyleColors::Primary, InputFocusRadius, FLinearColor::Transparent, InputFocusThickness))
+			.SetHoveredFillBrush(FSlateRoundedBoxBrush(FStyleColors::Hover, InputFocusRadius, FLinearColor::Transparent, InputFocusThickness))
+			.SetInactiveFillBrush(FSlateRoundedBoxBrush(FStyleColors::Dropdown, InputFocusRadius, FLinearColor::Transparent, InputFocusThickness))
 			.SetArrowsImage(FSlateNoResource())
 			.SetForegroundColor(FStyleColors::ForegroundHover)
-			.SetTextPadding(FMargin(10.f, 3.5f, 10.f, 4.f))
-		);
-	}
+			.SetTextPadding(FMargin(7.f, 2.5f, 7.f, 2.5f))
+			.SetInsetPadding(FMargin(1));
 
-	// SNumericEntryBox defaults...
-	{
-		Style->Set("NumericEntrySpinBox", FSpinBoxStyle()
-			.SetBackgroundBrush(FSlateNoResource())
-			.SetHoveredBackgroundBrush(FSlateNoResource())
-
-			.SetActiveFillBrush(FSlateRoundedBoxBrush(FStyleColors::Hover, InputFocusRadius, FLinearColor::Transparent, InputFocusThickness))
-			.SetInactiveFillBrush(FSlateRoundedBoxBrush(FStyleColors::Secondary, InputFocusRadius, FLinearColor::Transparent, InputFocusThickness))
+		Style->Set("SpinBox", DefaultSpinBox);
 
-			.SetArrowsImage(FSlateNoResource())
-			.SetTextPadding(FMargin(0.f))
-			.SetForegroundColor(FStyleColors::ForegroundHover)
-		);
+		// SNumericEntryBox
+		Style->Set("NumericEntrySpinBox", DefaultSpinBox);
 
 		Style->Set("NumericEntrySpinBox_Dark", FSpinBoxStyle()
 			.SetBackgroundBrush(FSlateRoundedBoxBrush(FStyleColors::Input, InputFocusRadius, FStyleColors::InputOutline, InputFocusThickness))
 			.SetHoveredBackgroundBrush(FSlateRoundedBoxBrush(FStyleColors::Input, InputFocusRadius, FStyleColors::Hover, InputFocusThickness))
-
 			.SetActiveFillBrush(FSlateRoundedBoxBrush(FStyleColors::Hover, InputFocusRadius, FLinearColor::Transparent, InputFocusThickness))
 			.SetInactiveFillBrush(FSlateRoundedBoxBrush(FStyleColors::Secondary, InputFocusRadius, FLinearColor::Transparent, InputFocusThickness))
-
 			.SetArrowsImage(FSlateNoResource())
-
 			.SetTextPadding(FMargin(0.f))
 			.SetForegroundColor(FStyleColors::ForegroundHover)
 		);
@@ -1048,7 +1035,7 @@ void FStarshipCoreStyle::SetupTextStyles(TSharedRef<FStyle>& Style)
 		.SetBackgroundImageFocused(FSlateRoundedBoxBrush(FStyleColors::Input, InputFocusRadius, FStyleColors::Primary, InputFocusThickness))
 		.SetBackgroundImageReadOnly(FSlateRoundedBoxBrush(FStyleColors::Header, InputFocusRadius, FStyleColors::InputOutline, InputFocusThickness))
 		.SetFont(StyleFonts.Normal)
-		.SetPadding(FMargin(12.f, 4.0f, 12.f, 5.0f)) // The padding should be 4 top, 5 bottom
+		.SetPadding(FMargin(8.0f, 3.0f, 8.0f, 4.0f)) 
 		.SetForegroundColor(FStyleColors::Foreground)
 		.SetBackgroundColor(FStyleColors::White)
 		.SetReadOnlyForegroundColor(FSlateColor::UseForeground())
@@ -1074,7 +1061,7 @@ void FStarshipCoreStyle::SetupTextStyles(TSharedRef<FStyle>& Style)
 			.SetDownArrowImage(IMAGE_BRUSH_SVG("Starship/Common/chevron-down", Icon16x16, FStyleColors::Foreground))
 			.SetGlassImage(IMAGE_BRUSH_SVG("Starship/Common/search", Icon16x16))
 			.SetClearImage(IMAGE_BRUSH_SVG("Starship/Common/close", Icon16x16))
-			.SetImagePadding(FMargin(3.f, 0.f, -2.f, 0.f))
+			.SetImagePadding(FMargin(3.f, 0.f, 0.f, 0.f))
 			.SetLeftAlignButtons(true)
 		);
 	}
@@ -1084,12 +1071,9 @@ void FStarshipCoreStyle::SetupTextStyles(TSharedRef<FStyle>& Style)
 		// Normal Editable Text 
 		FTextBlockStyle InlineEditableTextBlockReadOnly = FTextBlockStyle(NormalText);
 
-		FEditableTextBoxStyle InlineEditableTextBlockEditable = FEditableTextBoxStyle(NormalEditableTextBoxStyle)
-			.SetPadding(FMargin(6.f, 4.5f, 6.f, 4.5f));
-
 		FInlineEditableTextBlockStyle InlineEditableTextBlockStyle = FInlineEditableTextBlockStyle()
 			.SetTextStyle(InlineEditableTextBlockReadOnly)
-			.SetEditableTextBoxStyle(InlineEditableTextBlockEditable);
+			.SetEditableTextBoxStyle(NormalEditableTextBoxStyle);
 
 		Style->Set("InlineEditableTextBlockStyle", InlineEditableTextBlockStyle);
 
@@ -1097,7 +1081,7 @@ void FStarshipCoreStyle::SetupTextStyles(TSharedRef<FStyle>& Style)
 		FTextBlockStyle InlineEditableTextBlockSmallReadOnly = FTextBlockStyle(InlineEditableTextBlockReadOnly)
 			.SetFont(StyleFonts.Small);
 
-		FEditableTextBoxStyle InlineEditableTextBlockSmallEditable = FEditableTextBoxStyle(InlineEditableTextBlockEditable)
+		FEditableTextBoxStyle InlineEditableTextBlockSmallEditable = FEditableTextBoxStyle(NormalEditableTextBoxStyle)
 			.SetFont(StyleFonts.Small);
 
 		FInlineEditableTextBlockStyle InlineEditableTextBlockSmallStyle = FInlineEditableTextBlockStyle()
@@ -1117,52 +1101,63 @@ void FStarshipCoreStyle::SetupButtonStyles(TSharedRef<FStyle>& Style)
 {
 	// SButton defaults
 	const FButtonStyle PrimaryButton = FButtonStyle()
-		.SetNormal(FSlateRoundedBoxBrush(FStyleColors::Primary, 4.0f, FStyleColors::InputOutline, InputFocusThickness))
-		.SetHovered(FSlateRoundedBoxBrush(FStyleColors::PrimaryHover, 4.0f, FStyleColors::Hover, InputFocusThickness))
-		.SetPressed(FSlateRoundedBoxBrush(FStyleColors::PrimaryPress, 4.0f, FStyleColors::Hover, InputFocusThickness))
-		.SetNormalForeground(FStyleColors::Panel)
-		.SetHoveredForeground(FStyleColors::Panel)
-		.SetPressedForeground(FStyleColors::Panel)
-		.SetDisabledForeground(FStyleColors::Panel)
+		.SetNormal(FSlateRoundedBoxBrush(FStyleColors::Primary, 4.0f, FStyleColors::Input, InputFocusThickness))
+		.SetHovered(FSlateRoundedBoxBrush(FStyleColors::PrimaryHover, 4.0f, FStyleColors::Input, InputFocusThickness))
+		.SetPressed(FSlateRoundedBoxBrush(FStyleColors::PrimaryPress, 4.0f, FStyleColors::Input, InputFocusThickness))
+		.SetDisabled(FSlateRoundedBoxBrush(FStyleColors::Dropdown, 4.0f, FStyleColors::Recessed, InputFocusThickness))
+		.SetNormalForeground(FStyleColors::ForegroundHover)
+		.SetHoveredForeground(FStyleColors::ForegroundHover)
+		.SetPressedForeground(FStyleColors::ForegroundHover)
+		.SetDisabledForeground(FStyleColors::Foreground)
 		.SetNormalPadding(ButtonMargins)
-		.SetPressedPadding(ButtonMargins);
+		.SetPressedPadding(PressedButtonMargins);
 
 
 	const FButtonStyle Button = FButtonStyle()
-		.SetNormal(FSlateRoundedBoxBrush(FStyleColors::Secondary, 4.0f, FStyleColors::InputOutline, InputFocusThickness))
-		.SetHovered(FSlateRoundedBoxBrush(FStyleColors::Hover, 4.0f, FStyleColors::Hover, InputFocusThickness))
-		.SetPressed(FSlateRoundedBoxBrush(FStyleColors::Header, 4.0f, FStyleColors::Hover, InputFocusThickness))
+		.SetNormal(FSlateRoundedBoxBrush(FStyleColors::Secondary, 4.0f, FStyleColors::Input, InputFocusThickness))
+		.SetHovered(FSlateRoundedBoxBrush(FStyleColors::Hover, 4.0f, FStyleColors::Input, InputFocusThickness))
+		.SetPressed(FSlateRoundedBoxBrush(FStyleColors::Header, 4.0f, FStyleColors::Input, InputFocusThickness))
+		.SetDisabled(FSlateRoundedBoxBrush(FStyleColors::Dropdown, 4.0f, FStyleColors::Recessed, InputFocusThickness))
 		.SetNormalForeground(FStyleColors::ForegroundHover)
 		.SetHoveredForeground(FStyleColors::ForegroundHover)
 		.SetPressedForeground(FStyleColors::ForegroundHover)
 		.SetDisabledForeground(FStyleColors::Foreground)
 		.SetNormalPadding(ButtonMargins)
-		.SetPressedPadding(ButtonMargins);
+		.SetPressedPadding(PressedButtonMargins);
 
-	const FButtonStyle SimpleButton = FButtonStyle()
-		.SetNormal(FSlateNoResource())
-		.SetHovered(FSlateNoResource())
-		.SetPressed(FSlateNoResource())
-		.SetDisabled(FSlateNoResource())
-		.SetNormalForeground(FStyleColors::Foreground)
+	// @todo Recolor - Get rid of this, it is a copy of button
+	const FButtonStyle SecondaryButton = FButtonStyle()
+		.SetNormal(FSlateRoundedBoxBrush(FStyleColors::Secondary, 4.0f, FStyleColors::InputOutline, InputFocusThickness))
+		.SetHovered(FSlateRoundedBoxBrush(FStyleColors::Hover, 4.0f, FStyleColors::InputOutline, InputFocusThickness))
+		.SetPressed(FSlateRoundedBoxBrush(FStyleColors::Header, 4.0f, FStyleColors::InputOutline, InputFocusThickness))
+		.SetDisabled(FSlateRoundedBoxBrush(FStyleColors::Dropdown, 4.0f, FStyleColors::InputOutline, InputFocusThickness))
+		.SetNormalForeground(FStyleColors::ForegroundHover)
 		.SetHoveredForeground(FStyleColors::ForegroundHover)
 		.SetPressedForeground(FStyleColors::ForegroundHover)
 		.SetDisabledForeground(FStyleColors::Foreground)
-		.SetNormalPadding(FMargin(2.f, 2.f, 2.f, 2.f))
-		.SetPressedPadding(FMargin(2.f, 3.f, 2.f, 1.f));
+		.SetNormalPadding(FMargin(8.f, 4.5f, 8.f, 3.5f))
+		.SetPressedPadding(FMargin(8.f, 5.f, 6.f, 3.f));
 
 
-	const FButtonStyle SecondaryButton = FButtonStyle()
-		.SetNormal(FSlateRoundedBoxBrush(FStyleColors::Secondary, 4.0f))
-		.SetHovered(FSlateRoundedBoxBrush(FStyleColors::Hover, 4.0f))
-		.SetPressed(FSlateRoundedBoxBrush(FStyleColors::Header, 4.0f))
+	const FButtonStyle SimpleButton = FButtonStyle()
+		.SetNormal(FSlateNoResource())
+		.SetHovered(FSlateRoundedBoxBrush(FStyleColors::Dropdown, 4.0f))
+		.SetPressed(FSlateRoundedBoxBrush(FStyleColors::Dropdown, 4.0f))
 		.SetDisabled(FSlateNoResource())
-		.SetNormalForeground(FStyleColors::ForegroundHover)
+		.SetNormalForeground(FStyleColors::Foreground)
 		.SetHoveredForeground(FStyleColors::ForegroundHover)
 		.SetPressedForeground(FStyleColors::ForegroundHover)
 		.SetDisabledForeground(FStyleColors::Foreground)
-		.SetNormalPadding(FMargin(8.f, 4.5f, 8.f, 3.5f))
-		.SetPressedPadding(FMargin(8.f, 5.f, 6.f, 3.f));
+		.SetNormalPadding(ButtonMargins)
+		.SetPressedPadding(PressedButtonMargins);
+
+	const FButtonStyle SimpleButtonLabelAndIcon = FButtonStyle(SimpleButton)
+		.SetNormalPadding(FMargin(4.f, 1.f, 4.f, 1.f))
+		.SetPressedPadding(FMargin(4.f, 2.0f, 4.f, 0.0f));
+
+	FButtonStyle SimpleButtonIconOnly = FButtonStyle(SimpleButton)
+		.SetNormalPadding(FMargin(2.f, 1.f, 2.f, 1.f))
+		.SetPressedPadding(FMargin(2.f, 2.0f, 2.f, 0.0f));
 	{
 
 		const FTextBlockStyle& NormalText = Style->GetWidgetStyle<FTextBlockStyle>("NormalText");
@@ -1171,17 +1166,18 @@ void FStarshipCoreStyle::SetupButtonStyles(TSharedRef<FStyle>& Style)
 
 		Style->Set("PrimaryButton", PrimaryButton);
 		Style->Set("Button", Button);
-		Style->Set("SimpleButton", SimpleButton);
+		Style->Set("SimpleButton", SimpleButtonIconOnly);
+		Style->Set("SimpleButtonLabelAndIcon", SimpleButtonLabelAndIcon);
 		Style->Set("SecondaryButton", SecondaryButton);
 
-
-		Style->Set("DialogButtonText", FTextBlockStyle(NormalText).SetFont(FStyleFonts::Get().NormalBold).SetTransformPolicy(ETextTransformPolicy::ToUpper));
-
+		//@todo Recolor - This is reundant with button text now.
+		Style->Set("DialogButtonText", FTextBlockStyle(NormalText).SetFont(FStyleFonts::Get().NormalBold));
 
 		FSlateFontInfo AddFont = FStyleFonts::Get().SmallBold;
 		AddFont.LetterSpacing = 200;
 
-		Style->Set("SmallButtonText", FTextBlockStyle(NormalText).SetFont(AddFont).SetTransformPolicy(ETextTransformPolicy::ToUpper));
+		//@todo Recolor - Rename to simple button text
+		Style->Set("SmallButtonText", FTextBlockStyle(NormalText));
 
 		// A button that only responds to clicks and displays nothing
 		const FButtonStyle InvisibleButton = FButtonStyle()
@@ -1197,6 +1193,14 @@ void FStarshipCoreStyle::SetupButtonStyles(TSharedRef<FStyle>& Style)
 			.SetPressedPadding(FMargin(0.f));
 
 		Style->Set("InvisibleButton", InvisibleButton);
+
+		SimpleButtonIconOnly
+			.SetHovered(FSlateNoResource())
+			.SetPressed(FSlateNoResource())
+			.SetNormalPadding(FMargin(2.f, 0))
+			.SetPressedPadding(FMargin(2.0f, 0.f));
+
+		Style->Set("HoverOnlyButton", SimpleButtonIconOnly);
 	}
 }
 
@@ -1207,19 +1211,20 @@ void FStarshipCoreStyle::SetupComboButtonStyles(TSharedRef<FStyle>& Style)
 		.SetNormal(FSlateRoundedBoxBrush(FStyleColors::Input, InputFocusRadius, FStyleColors::InputOutline, InputFocusThickness))
 		.SetHovered(FSlateRoundedBoxBrush(FStyleColors::Input, InputFocusRadius, FStyleColors::Hover, InputFocusThickness))
 		.SetPressed(FSlateRoundedBoxBrush(FStyleColors::Input, InputFocusRadius, FStyleColors::Hover, InputFocusThickness))
+		.SetDisabled(FSlateRoundedBoxBrush(FStyleColors::Dropdown, InputFocusRadius))
 		.SetNormalForeground(FStyleColors::Foreground)
 		.SetHoveredForeground(FStyleColors::ForegroundHover)
 		.SetPressedForeground(FStyleColors::ForegroundHover)
-		.SetDisabledForeground(FStyleColors::White25)
-		.SetNormalPadding(FMargin(8.f, 2.f, 4.f, 2.f))
-		.SetPressedPadding(FMargin(8.f, 2.f, 4.f, 2.f));
+		.SetDisabledForeground(FStyleColors::ForegroundHover)
+		.SetNormalPadding(ComboButtonMargin)
+		.SetPressedPadding(PressedComboButtonMargin);
 
 
 	// SComboBox 
 	FComboButtonStyle ComboButton = FComboButtonStyle()
 		.SetButtonStyle(ComboButtonButton)
 		.SetContentPadding(0.f)
-		.SetDownArrowImage(IMAGE_BRUSH_SVG("Starship/CoreWidgets/ComboBox/wide-chevron-down", FVector2D(20.f, 16.f)))
+		.SetDownArrowImage(IMAGE_BRUSH_SVG("Starship/CoreWidgets/ComboBox/wide-chevron-down", FVector2D(16.f, 16.f)))
 		.SetMenuBorderBrush(FSlateRoundedBoxBrush(FStyleColors::Dropdown, 0.0f, FStyleColors::DropdownOutline, 1.0f))
 		.SetMenuBorderPadding(0.f);
 	Style->Set("ComboButton", ComboButton);
@@ -1231,20 +1236,20 @@ void FStarshipCoreStyle::SetupComboButtonStyles(TSharedRef<FStyle>& Style)
 		.SetComboButtonStyle(ComboButton);
 	Style->Set("ComboBox", ComboBox);
 
-	const FButtonStyle& SimpleButton = Style->GetWidgetStyle<FButtonStyle>("SimpleButton");
+	const FButtonStyle& SimpleButton = Style->GetWidgetStyle<FButtonStyle>("SimpleButtonLabelAndIcon");
 
 	// Simple Combo Box  (borderless)
 	FComboButtonStyle SimpleComboButton = FComboButtonStyle()
 		.SetButtonStyle(SimpleButton)
 		.SetContentPadding(0.f)
+		.SetDownArrowImage(IMAGE_BRUSH_SVG("Starship/CoreWidgets/ComboBox/wide-chevron-down", FVector2D(16.f, 16.f)))
 		.SetShadowOffset(FVector2D(0.0f, 0.0f))
-		.SetDownArrowImage(IMAGE_BRUSH_SVG("Starship/CoreWidgets/ComboBox/corner-dropdown", FVector2D(7.0f, 7.0f)))
-		.SetDownArrowPadding(FMargin(0.0f))
-		.SetDownArrowAlignment(EVerticalAlignment::VAlign_Bottom)
 		.SetMenuBorderBrush(FSlateRoundedBoxBrush(FStyleColors::Dropdown, 0.0f, FStyleColors::DropdownOutline, 1.0f))
 		.SetMenuBorderPadding(0.0f);
 	Style->Set("SimpleComboButton", SimpleComboButton);
 
+	SimpleComboButton.SetButtonStyle(Style->GetWidgetStyle<FButtonStyle>("SimpleButton"));
+	Style->Set("SimpleComboButtonWithIcon", SimpleComboButton);
 
 	FComboBoxStyle SimpleComboBox = FComboBoxStyle()
 		.SetContentPadding(0.f)
@@ -1252,7 +1257,6 @@ void FStarshipCoreStyle::SetupComboButtonStyles(TSharedRef<FStyle>& Style)
 		.SetComboButtonStyle(SimpleComboButton);
 	Style->Set("SimpleComboBox", SimpleComboBox);
 
-
 	const FTableRowStyle ComboBoxRow = FTableRowStyle()
 		.SetEvenRowBackgroundBrush(FSlateNoResource())
 		.SetEvenRowBackgroundHoveredBrush(FSlateColorBrush(FStyleColors::Hover))
@@ -1302,55 +1306,29 @@ void FStarshipCoreStyle::SetupCheckboxStyles(TSharedRef<FStyle>& Style)
 	const FCheckBoxStyle BasicCheckBoxStyle = FCheckBoxStyle()
 		.SetCheckBoxType(ESlateCheckBoxType::CheckBox)
 
-		.SetForegroundColor(FStyleColors::Foreground)
+		.SetForegroundColor(FStyleColors::ForegroundHover)
 		.SetHoveredForegroundColor(FStyleColors::ForegroundHover)
 		.SetPressedForegroundColor(FStyleColors::ForegroundHover)
-		.SetCheckedForegroundColor(FStyleColors::Foreground)
+		.SetCheckedForegroundColor(FStyleColors::ForegroundHover)
 		.SetCheckedHoveredForegroundColor(FStyleColors::ForegroundHover)
 		.SetCheckedPressedForegroundColor(FStyleColors::ForegroundHover)
-		.SetUndeterminedForegroundColor(FStyleColors::Foreground)
+		.SetUndeterminedForegroundColor(FStyleColors::ForegroundHover)
 
 		.SetUncheckedImage(FSlateNoResource())
 		.SetUncheckedHoveredImage(FSlateNoResource())
 		.SetUncheckedPressedImage(FSlateNoResource())
 		.SetCheckedImage(IMAGE_BRUSH_SVG("Starship/CoreWidgets/CheckBox/check", Icon16x16, FStyleColors::Primary))
 		.SetCheckedHoveredImage(IMAGE_BRUSH_SVG("Starship/CoreWidgets/CheckBox/check", Icon16x16, FStyleColors::PrimaryHover))
-		.SetCheckedPressedImage(IMAGE_BRUSH_SVG("Starship/CoreWidgets/CheckBox/check", Icon16x16, FStyleColors::Primary))
+		.SetCheckedPressedImage(IMAGE_BRUSH_SVG("Starship/CoreWidgets/CheckBox/check", Icon16x16, FStyleColors::PrimaryPress))
 		.SetUndeterminedImage(IMAGE_BRUSH_SVG("Starship/CoreWidgets/CheckBox/indeterminate", Icon16x16, FStyleColors::Primary))
 		.SetUndeterminedHoveredImage(IMAGE_BRUSH_SVG("Starship/CoreWidgets/CheckBox/indeterminate", Icon16x16, FStyleColors::PrimaryHover))
-		.SetUndeterminedPressedImage(IMAGE_BRUSH_SVG("Starship/CoreWidgets/CheckBox/indeterminate", Icon16x16, FStyleColors::Primary))
-		.SetBackgroundImage(FSlateRoundedBoxBrush(FStyleColors::Input, CheckboxCornerRadius, FStyleColors::InputOutline, CheckboxOutlineThickness, Icon18x18))
+		.SetUndeterminedPressedImage(IMAGE_BRUSH_SVG("Starship/CoreWidgets/CheckBox/indeterminate", Icon16x16, FStyleColors::PrimaryPress))
+		.SetBackgroundImage(FSlateRoundedBoxBrush(FStyleColors::Input, CheckboxCornerRadius, FStyleColors::Secondary, CheckboxOutlineThickness, Icon18x18))
 		.SetBackgroundHoveredImage(FSlateRoundedBoxBrush(FStyleColors::Input, CheckboxCornerRadius, FStyleColors::Hover, CheckboxOutlineThickness, Icon18x18))
-		.SetBackgroundPressedImage(FSlateRoundedBoxBrush(FStyleColors::Foldout, CheckboxCornerRadius, FStyleColors::Hover, CheckboxOutlineThickness, Icon18x18));
+		.SetBackgroundPressedImage(FSlateRoundedBoxBrush(FStyleColors::Input, CheckboxCornerRadius, FStyleColors::Hover, CheckboxOutlineThickness, Icon18x18));
 
 	Style->Set("Checkbox", BasicCheckBoxStyle);
 
-	const FCheckBoxStyle SimplifiedCheckBoxStyle = FCheckBoxStyle()
-		.SetCheckBoxType(ESlateCheckBoxType::CheckBox)
-
-		.SetForegroundColor(FLinearColor::White)
-		.SetHoveredForegroundColor(FLinearColor::White)
-		.SetPressedForegroundColor(FLinearColor::White)
-		.SetCheckedForegroundColor(FLinearColor::White)
-		.SetCheckedHoveredForegroundColor(FLinearColor::White)
-		.SetCheckedPressedForegroundColor(FLinearColor::White)
-		.SetUndeterminedForegroundColor(FLinearColor::White)
-
-		.SetUncheckedImage(FSlateNoResource())
-		.SetUncheckedHoveredImage(FSlateNoResource())
-		.SetUncheckedPressedImage(FSlateNoResource())
-		.SetCheckedImage(IMAGE_BRUSH_SVG("Starship/CoreWidgets/CheckBox/check", Icon16x16, FStyleColors::Foreground))
-		.SetCheckedHoveredImage(IMAGE_BRUSH_SVG("Starship/CoreWidgets/CheckBox/check", Icon16x16, FStyleColors::ForegroundHover))
-		.SetCheckedPressedImage(IMAGE_BRUSH_SVG("Starship/CoreWidgets/CheckBox/check", Icon16x16, FStyleColors::Foreground))
-		.SetUndeterminedImage(IMAGE_BRUSH_SVG("Starship/CoreWidgets/CheckBox/indeterminate", Icon16x16, FStyleColors::Foreground))
-		.SetUndeterminedHoveredImage(IMAGE_BRUSH_SVG("Starship/CoreWidgets/CheckBox/indeterminate", Icon16x16, FStyleColors::ForegroundHover))
-		.SetUndeterminedPressedImage(IMAGE_BRUSH_SVG("Starship/CoreWidgets/CheckBox/indeterminate", Icon16x16, FStyleColors::Foreground))
-		.SetBackgroundImage(FSlateRoundedBoxBrush(FStyleColors::Input, CheckboxCornerRadius, FStyleColors::InputOutline, CheckboxOutlineThickness, Icon18x18))
-		.SetBackgroundHoveredImage(FSlateRoundedBoxBrush(FStyleColors::Input, CheckboxCornerRadius, FStyleColors::Hover, CheckboxOutlineThickness, Icon18x18))
-		.SetBackgroundPressedImage(FSlateRoundedBoxBrush(FStyleColors::Foldout, CheckboxCornerRadius, FStyleColors::Hover, CheckboxOutlineThickness, Icon18x18));
-
-	Style->Set("SimplifiedCheckbox", SimplifiedCheckBoxStyle);
-
 	/* Set images for various transparent SCheckBox states ... */
 	const FCheckBoxStyle BasicTransparentCheckBoxStyle = FCheckBoxStyle()
 		.SetCheckBoxType(ESlateCheckBoxType::ToggleButton)
@@ -1364,25 +1342,30 @@ void FStarshipCoreStyle::SetupCheckboxStyles(TSharedRef<FStyle>& Style)
 		.SetUndeterminedHoveredImage(FSlateNoResource())
 		.SetUndeterminedPressedImage(FSlateNoResource())
 		.SetForegroundColor(FStyleColors::Foreground)
-		.SetHoveredForegroundColor(FStyleColors::ForegroundHover);
+		.SetHoveredForegroundColor(FStyleColors::ForegroundHover)
+		.SetPressedForegroundColor(FStyleColors::ForegroundHover)
+		.SetCheckedForegroundColor(FStyleColors::ForegroundHover)
+		.SetCheckedHoveredForegroundColor(FStyleColors::ForegroundHover)
+		.SetCheckedPressedForegroundColor(FStyleColors::ForegroundHover);
 
 	Style->Set("TransparentCheckBox", BasicTransparentCheckBoxStyle);
 
 	/* Default Style for a toggleable button */
 	const FCheckBoxStyle ToggleButtonStyle = FCheckBoxStyle()
 		.SetCheckBoxType(ESlateCheckBoxType::ToggleButton)
-		.SetCheckedImage(FSlateNoResource())
-		.SetCheckedHoveredImage(FSlateRoundedBoxBrush(FStyleColors::Hover, InputFocusRadius))
-		.SetCheckedPressedImage(FSlateRoundedBoxBrush(FStyleColors::Input, InputFocusRadius))
+		.SetCheckedImage(FSlateRoundedBoxBrush(FStyleColors::Primary, 4.0f))
+		.SetCheckedHoveredImage(FSlateRoundedBoxBrush(FStyleColors::PrimaryHover, 4.0f))
+		.SetCheckedPressedImage(FSlateRoundedBoxBrush(FStyleColors::Dropdown, 4.0f))
 		.SetUncheckedImage(FSlateNoResource())
-		.SetUncheckedHoveredImage(FSlateRoundedBoxBrush(FStyleColors::Hover, InputFocusRadius))
-		.SetUncheckedPressedImage(FSlateRoundedBoxBrush(FStyleColors::Input, InputFocusRadius))
+		.SetUncheckedHoveredImage(FSlateRoundedBoxBrush(FStyleColors::Dropdown, 4.0f))
+		.SetUncheckedPressedImage(FSlateRoundedBoxBrush(FStyleColors::Dropdown, 4.0f))
 		.SetForegroundColor(FStyleColors::Foreground)
 		.SetHoveredForegroundColor(FStyleColors::ForegroundHover)
 		.SetPressedForegroundColor(FStyleColors::ForegroundHover)
-		.SetCheckedForegroundColor(FStyleColors::Primary)
-		.SetCheckedHoveredForegroundColor(FStyleColors::PrimaryHover)
-		.SetPadding(DefaultMargins);
+		.SetCheckedForegroundColor(FStyleColors::ForegroundHover)
+		.SetCheckedHoveredForegroundColor(FStyleColors::ForegroundHover)
+		.SetCheckedPressedForegroundColor(FStyleColors::ForegroundHover)
+		.SetPadding(ToggleButtonMargins);
 
 	Style->Set("ToggleButtonCheckbox", ToggleButtonStyle);
 
@@ -1400,61 +1383,63 @@ void FStarshipCoreStyle::SetupCheckboxStyles(TSharedRef<FStyle>& Style)
 		.SetPressedForegroundColor(FStyleColors::ForegroundHover)
 		.SetCheckedForegroundColor(FStyleColors::Primary)
 		.SetCheckedHoveredForegroundColor(FStyleColors::PrimaryHover)
-		.SetPadding(DefaultMargins);
+		.SetPadding(ToggleButtonMargins);
 
+	//@todo Recolor - Get rid of this. ToggleButtonCheckbox is what we should use
 	Style->Set("ToggleButtonCheckboxAlt", ToggleButtonAltStyle);
 
+	FMargin SegmentedControlMargins(8, 3);
 	/* Style for a segmented box */
-	const FCheckBoxStyle SegmentedBoxLeft = FCheckBoxStyle()
+	FCheckBoxStyle SegmentedBoxLeft = FCheckBoxStyle()
 		.SetCheckBoxType(ESlateCheckBoxType::ToggleButton)
 
-		.SetUncheckedImage(BOX_BRUSH("/Starship/CoreWidgets/SegmentedBox/left", FVector2D(16.f, 16.f), FMargin(4.0f / 16.0f), FStyleColors::Secondary))
+		.SetUncheckedImage(BOX_BRUSH("/Starship/CoreWidgets/SegmentedBox/left", FVector2D(16.f, 16.f), FMargin(4.0f / 16.0f), FStyleColors::Input))
 		.SetUncheckedHoveredImage(BOX_BRUSH("/Starship/CoreWidgets/SegmentedBox/left", FVector2D(16.f, 16.f), FMargin(4.0f / 16.0f), FStyleColors::Hover))
 		.SetUncheckedPressedImage(BOX_BRUSH("/Starship/CoreWidgets/SegmentedBox/left", FVector2D(16.f, 16.f), FMargin(4.0f / 16.0f), FStyleColors::Secondary))
-		.SetCheckedImage(BOX_BRUSH("/Starship/CoreWidgets/SegmentedBox/left", FVector2D(16.f, 16.f), FMargin(4.0f / 16.0f), FStyleColors::Input))
-		.SetCheckedHoveredImage(BOX_BRUSH("/Starship/CoreWidgets/SegmentedBox/left", FVector2D(16.f, 16.f), FMargin(4.0f / 16.0f), FStyleColors::Input))
-		.SetCheckedPressedImage(BOX_BRUSH("/Starship/CoreWidgets/SegmentedBox/left", FVector2D(16.f, 16.f), FMargin(4.0f / 16.0f), FStyleColors::Input))
+		.SetCheckedImage(BOX_BRUSH("/Starship/CoreWidgets/SegmentedBox/left", FVector2D(16.f, 16.f), FMargin(4.0f / 16.0f), FStyleColors::Dropdown))
+		.SetCheckedHoveredImage(BOX_BRUSH("/Starship/CoreWidgets/SegmentedBox/left", FVector2D(16.f, 16.f), FMargin(4.0f / 16.0f), FStyleColors::Hover))
+		.SetCheckedPressedImage(BOX_BRUSH("/Starship/CoreWidgets/SegmentedBox/left", FVector2D(16.f, 16.f), FMargin(4.0f / 16.0f), FStyleColors::Header))
 		.SetForegroundColor(FStyleColors::Foreground)
 		.SetHoveredForegroundColor(FStyleColors::ForegroundHover)
 		.SetPressedForegroundColor(FStyleColors::ForegroundHover)
-		.SetCheckedForegroundColor(FStyleColors::Primary)
-		.SetCheckedHoveredForegroundColor(FStyleColors::Primary)
-		.SetCheckedPressedForegroundColor(FStyleColors::Primary)
-		.SetPadding(DefaultMargins);
+		.SetCheckedForegroundColor(FStyleColors::ForegroundHover)
+		.SetCheckedHoveredForegroundColor(FStyleColors::ForegroundHover)
+		.SetCheckedPressedForegroundColor(FStyleColors::ForegroundHover)
+		.SetPadding(SegmentedControlMargins);
 
 
 	const FCheckBoxStyle SegmentedBoxCenter = FCheckBoxStyle()
 		.SetCheckBoxType(ESlateCheckBoxType::ToggleButton)
-		.SetUncheckedImage(FSlateColorBrush(FStyleColors::Secondary))
+		.SetUncheckedImage(FSlateColorBrush(FStyleColors::Input))
 		.SetUncheckedHoveredImage(FSlateColorBrush(FStyleColors::Hover))
 		.SetUncheckedPressedImage(FSlateColorBrush(FStyleColors::Secondary))
-		.SetCheckedImage(FSlateColorBrush(FStyleColors::Input))
-		.SetCheckedHoveredImage(FSlateColorBrush(FStyleColors::Input))
-		.SetCheckedPressedImage(FSlateColorBrush(FStyleColors::Input))
+		.SetCheckedImage(FSlateColorBrush(FStyleColors::Dropdown))
+		.SetCheckedHoveredImage(FSlateColorBrush(FStyleColors::Hover))
+		.SetCheckedPressedImage(FSlateColorBrush(FStyleColors::Header))
 		.SetForegroundColor(FStyleColors::Foreground)
 		.SetHoveredForegroundColor(FStyleColors::ForegroundHover)
 		.SetPressedForegroundColor(FStyleColors::ForegroundHover)
-		.SetCheckedForegroundColor(FStyleColors::Primary)
-		.SetCheckedHoveredForegroundColor(FStyleColors::Primary)
-		.SetCheckedPressedForegroundColor(FStyleColors::Primary)
-		.SetPadding(DefaultMargins);
+		.SetCheckedForegroundColor(FStyleColors::ForegroundHover)
+		.SetCheckedHoveredForegroundColor(FStyleColors::ForegroundHover)
+		.SetCheckedPressedForegroundColor(FStyleColors::ForegroundHover)
+		.SetPadding(SegmentedControlMargins);
 
 
 	const FCheckBoxStyle SegmentedBoxRight = FCheckBoxStyle()
 		.SetCheckBoxType(ESlateCheckBoxType::ToggleButton)
-		.SetUncheckedImage(BOX_BRUSH("/Starship/CoreWidgets/SegmentedBox/right", FVector2D(16.f, 16.f), FMargin(4.0f / 16.0f), FStyleColors::Secondary))
+		.SetUncheckedImage(BOX_BRUSH("/Starship/CoreWidgets/SegmentedBox/right", FVector2D(16.f, 16.f), FMargin(4.0f / 16.0f), FStyleColors::Input))
 		.SetUncheckedHoveredImage(BOX_BRUSH("/Starship/CoreWidgets/SegmentedBox/right", FVector2D(16.f, 16.f), FMargin(4.0f / 16.0f), FStyleColors::Hover))
-		.SetUncheckedPressedImage(BOX_BRUSH("/Starship/CoreWidgets/SegmentedBox/right", FVector2D(16.f, 16.f), FMargin(4.0f / 16.0f), FStyleColors::Secondary))
-		.SetCheckedImage(BOX_BRUSH("/Starship/CoreWidgets/SegmentedBox/right", FVector2D(16.f, 16.f), FMargin(4.0f / 16.0f), FStyleColors::Input))
-		.SetCheckedHoveredImage(BOX_BRUSH("/Starship/CoreWidgets/SegmentedBox/right", FVector2D(16.f, 16.f), FMargin(4.0f / 16.0f), FStyleColors::Input))
-		.SetCheckedPressedImage(BOX_BRUSH("/Starship/CoreWidgets/SegmentedBox/right", FVector2D(16.f, 16.f), FMargin(4.0f / 16.0f), FStyleColors::Input))
+		.SetUncheckedPressedImage(BOX_BRUSH("/Starship/CoreWidgets/SegmentedBox/right", FVector2D(16.f, 16.f), FMargin(4.0f / 16.0f), FStyleColors::Header))
+		.SetCheckedImage(BOX_BRUSH("/Starship/CoreWidgets/SegmentedBox/right", FVector2D(16.f, 16.f), FMargin(4.0f / 16.0f), FStyleColors::Dropdown))
+		.SetCheckedHoveredImage(BOX_BRUSH("/Starship/CoreWidgets/SegmentedBox/right", FVector2D(16.f, 16.f), FMargin(4.0f / 16.0f), FStyleColors::Hover))
+		.SetCheckedPressedImage(BOX_BRUSH("/Starship/CoreWidgets/SegmentedBox/right", FVector2D(16.f, 16.f), FMargin(4.0f / 16.0f), FStyleColors::Header))
 		.SetForegroundColor(FStyleColors::Foreground)
 		.SetHoveredForegroundColor(FStyleColors::ForegroundHover)
 		.SetPressedForegroundColor(FStyleColors::ForegroundHover)
-		.SetCheckedForegroundColor(FStyleColors::Primary)
-		.SetCheckedHoveredForegroundColor(FStyleColors::Primary)
-		.SetCheckedPressedForegroundColor(FStyleColors::Primary)
-		.SetPadding(DefaultMargins);
+		.SetCheckedForegroundColor(FStyleColors::ForegroundHover)
+		.SetCheckedHoveredForegroundColor(FStyleColors::ForegroundHover)
+		.SetCheckedPressedForegroundColor(FStyleColors::ForegroundHover)
+		.SetPadding(SegmentedControlMargins);
 
 	Style->Set("SegmentedControl", FSegmentedControlStyle()
 		.SetControlStyle(SegmentedBoxCenter)
@@ -1462,6 +1447,10 @@ void FStarshipCoreStyle::SetupCheckboxStyles(TSharedRef<FStyle>& Style)
 		.SetLastControlStyle(SegmentedBoxRight)
 	);
 
+	SegmentedBoxLeft.SetUncheckedImage(BOX_BRUSH("/Starship/CoreWidgets/SegmentedBox/left", FVector2D(16.f, 16.f), FMargin(4.0f / 16.0f), FStyleColors::Dropdown));
+	SegmentedBoxLeft.SetCheckedImage(BOX_BRUSH("/Starship/CoreWidgets/SegmentedBox/left", FVector2D(16.f, 16.f), FMargin(4.0f / 16.0f), FStyleColors::Primary));
+	SegmentedBoxLeft.SetCheckedHoveredImage(BOX_BRUSH("/Starship/CoreWidgets/SegmentedBox/right", FVector2D(16.f, 16.f), FMargin(4.0f / 16.0f), FStyleColors::PrimaryHover));
+
 	Style->Set("SegmentedCombo.Left", SegmentedBoxLeft);
 
 	const FButtonStyle SegmentedComboRightButton = FButtonStyle()
@@ -1485,35 +1474,13 @@ void FStarshipCoreStyle::SetupCheckboxStyles(TSharedRef<FStyle>& Style)
 
 	// Button that looks/behaves like a SegmentedCombo for cases where SegmentedCombos and 
 	// buttons are placed adjacent
-	const FCheckBoxStyle SegmentedBoxLeftOnly = FCheckBoxStyle()
-		.SetCheckBoxType(ESlateCheckBoxType::ToggleButton)
-		.SetUncheckedImage(       FSlateRoundedBoxBrush(FStyleColors::Secondary, 4.f) )
+	const FCheckBoxStyle SegmentedBoxButtonOnly = FCheckBoxStyle(ToggleButtonStyle)
+		.SetUncheckedImage(FSlateRoundedBoxBrush(FStyleColors::Dropdown, 4.f))
 		.SetUncheckedHoveredImage(FSlateRoundedBoxBrush(FStyleColors::Hover, 4.f))
-		.SetUncheckedPressedImage(FSlateRoundedBoxBrush(FStyleColors::Secondary, 4.f))
-		.SetCheckedImage(         FSlateRoundedBoxBrush(FStyleColors::Input, 4.f))
-		.SetCheckedHoveredImage(  FSlateRoundedBoxBrush(FStyleColors::Input, 4.f))
-		.SetCheckedPressedImage(  FSlateRoundedBoxBrush(FStyleColors::Input, 4.f))
-		.SetForegroundColor(FStyleColors::Foreground)
-		.SetHoveredForegroundColor(FStyleColors::ForegroundHover)
-		.SetPressedForegroundColor(FStyleColors::ForegroundHover)
-		.SetCheckedForegroundColor(FStyleColors::Primary)
-		.SetCheckedHoveredForegroundColor(FStyleColors::Primary)
-		.SetCheckedPressedForegroundColor(FStyleColors::Primary)
-		.SetPadding(DefaultMargins);
-	Style->Set("SegmentedCombo.ButtonOnly", SegmentedBoxLeftOnly);
-
-	/* Style for a toggleable button that mimics the coloring and look of a Table Row */
-	/*
-	const FCheckBoxStyle ToggleButtonRowStyle = FCheckBoxStyle()
-		.SetCheckBoxType(ESlateCheckBoxType::ToggleButton)
-		.SetUncheckedImage(FSlateNoResource())
-		.SetUncheckedHoveredImage(IMAGE_BRUSH("Common/Selection", Icon8x8, SelectionColor_Inactive))
-		.SetUncheckedPressedImage(IMAGE_BRUSH("Common/Selection", Icon8x8, SelectionColor_Inactive))
-		.SetCheckedImage(IMAGE_BRUSH("Common/Selection", Icon8x8, SelectionColor))
-		.SetCheckedHoveredImage(IMAGE_BRUSH("Common/Selection", Icon8x8, SelectionColor))
-		.SetCheckedPressedImage(BOX_BRUSH("Common/Selector", 4.0f / 16.0f, SelectorColor));
-	Style->Set("ToggleButtonRowStyle", ToggleButtonRowStyle);
-	*/
+		.SetUncheckedPressedImage(FSlateRoundedBoxBrush(FStyleColors::Header, 4.f));
+
+	Style->Set("SegmentedCombo.ButtonOnly", SegmentedBoxButtonOnly);
+
 
 	/* A radio button is actually just a SCheckBox box with different images */
 	/* Set images for various radio button (SCheckBox) states ... */
@@ -1528,12 +1495,12 @@ void FStarshipCoreStyle::SetupCheckboxStyles(TSharedRef<FStyle>& Style)
 		.SetCheckedPressedForegroundColor(FLinearColor::White)
 		.SetUndeterminedForegroundColor(FLinearColor::White)
 
-		.SetUncheckedImage(IMAGE_BRUSH_SVG("/Starship/CoreWidgets/CheckBox/radio-off", Icon16x16, FStyleColors::White25))
-		.SetUncheckedHoveredImage(IMAGE_BRUSH_SVG("/Starship/CoreWidgets/CheckBox/radio-off", Icon16x16, FStyleColors::ForegroundHover))
-		.SetUncheckedPressedImage(IMAGE_BRUSH_SVG("/Starship/CoreWidgets/CheckBox/radio-off", Icon16x16, FStyleColors::ForegroundHover))
+		.SetUncheckedImage(IMAGE_BRUSH_SVG("/Starship/CoreWidgets/CheckBox/radio-off", Icon16x16, FStyleColors::Dropdown))
+		.SetUncheckedHoveredImage(IMAGE_BRUSH_SVG("/Starship/CoreWidgets/CheckBox/radio-off", Icon16x16, FStyleColors::Hover))
+		.SetUncheckedPressedImage(IMAGE_BRUSH_SVG("/Starship/CoreWidgets/CheckBox/radio-off", Icon16x16, FStyleColors::Dropdown))
 		.SetCheckedImage(IMAGE_BRUSH_SVG("/Starship/CoreWidgets/CheckBox/radio-on", Icon16x16, FStyleColors::Primary))
-		.SetCheckedHoveredImage(IMAGE_BRUSH_SVG("/Starship/CoreWidgets/CheckBox/radio-on", Icon16x16, FStyleColors::Primary))
-		.SetCheckedPressedImage(IMAGE_BRUSH_SVG("/Starship/CoreWidgets/CheckBox/radio-off", Icon16x16, FStyleColors::Primary))
+		.SetCheckedHoveredImage(IMAGE_BRUSH_SVG("/Starship/CoreWidgets/CheckBox/radio-on", Icon16x16, FStyleColors::PrimaryHover))
+		.SetCheckedPressedImage(IMAGE_BRUSH_SVG("/Starship/CoreWidgets/CheckBox/radio-on", Icon16x16, FStyleColors::Primary))
 		.SetUndeterminedImage(IMAGE_BRUSH_SVG("/Starship/CoreWidgets/CheckBox/radio-off", Icon16x16, FStyleColors::White25))
 		.SetUndeterminedHoveredImage(IMAGE_BRUSH_SVG("/Starship/CoreWidgets/CheckBox/radio-off", Icon16x16, FStyleColors::ForegroundHover))
 		.SetUndeterminedPressedImage(IMAGE_BRUSH_SVG("/Starship/CoreWidgets/CheckBox/radio-off", Icon16x16, FStyleColors::ForegroundHover))
@@ -1765,7 +1732,7 @@ void FStarshipCoreStyle::SetupTableViewStyles(TSharedRef<FStyle>& Style)
 		.SetInactiveHighlightedBrush(FSlateColorBrush(FStyleColors::SelectParent))// This is the parent highlight
 
 		.SetTextColor(FStyleColors::Foreground)
-		.SetSelectedTextColor(FStyleColors::ForegroundInverted)
+		.SetSelectedTextColor(FStyleColors::ForegroundHover)
 
 		.SetDropIndicator_Above(BOX_BRUSH("Common/DropZoneIndicator_Above", FMargin(10.0f / 16.0f, 10.0f / 16.0f, 0.f, 0.f), SelectionColor))
 		.SetDropIndicator_Onto(BOX_BRUSH("Common/DropZoneIndicator_Onto", FMargin(4.0f / 16.0f), SelectionColor))
@@ -1920,74 +1887,85 @@ void FStarshipCoreStyle::SetupMultiboxStyles(TSharedRef<FStyle>& Style)
 		Style->Set("ToolBar", NormalToolbarStyle);
 
 		// Slim Toolbar
+		const FButtonStyle& SimpleButtonStyle = Style->GetWidgetStyle<FButtonStyle>("SimpleButton");
 
-		FButtonStyle SlimToolBarButton = FButtonStyle(ToolbarButton)
-			.SetPressed(FSlateNoResource())
-			.SetHovered(FSlateNoResource())
+		FButtonStyle SlimToolBarButton = FButtonStyle(SimpleButtonStyle)
 			.SetNormalForeground(FStyleColors::Foreground)
 			.SetPressedForeground(FStyleColors::ForegroundHover)
 			.SetHoveredForeground(FStyleColors::ForegroundHover)
-			.SetDisabledForeground(FStyleColors::Foreground);
+			.SetDisabledForeground(FStyleColors::Foreground)
+			.SetNormalPadding(FMargin(8.f, 4.f, 8.f, 4.f))
+			.SetPressedPadding(FMargin(8.f, 5.f, 8.f, 3.f));
 
 
 		FComboButtonStyle SlimToolBarComboButton = FComboButtonStyle(Style->GetWidgetStyle<FComboButtonStyle>("SimpleComboButton"))
-			.SetButtonStyle(SlimToolBarButton)
-			.SetDownArrowPadding(FMargin(1.f, 0.f, 2.f, 5.f));
+			.SetDownArrowPadding(FMargin(1.f, 0.f, 2.f, 0.f));
 
 		FToolBarStyle SlimToolbarStyle =
 			FToolBarStyle()
 			.SetBackground(FSlateColorBrush(FStyleColors::Panel))
-			.SetBackgroundPadding(FMargin(0.f, 4.f, 0.f, 4.f))
+			.SetBackgroundPadding(FMargin(4.f, 6.f, 0.f, 6.f))
 			.SetExpandBrush(IMAGE_BRUSH("Icons/toolbar_expand_16x", Icon16x16))
-			.SetComboButtonPadding(FMargin(6.0f, 0.0f))
-			.SetButtonPadding(FMargin(4.0f, 4.0f))
+			.SetComboButtonPadding(FMargin(0.0f, 0.0f))
+			.SetButtonPadding(FMargin(4.0f, 0.0f))
 			.SetCheckBoxPadding(FMargin(8.0f, 0.0f))
 			.SetSeparatorBrush(FSlateColorBrush(FStyleColors::Recessed))
 			.SetSeparatorPadding(FMargin(8.f, -4.0f))
 			.SetLabelStyle(FTextBlockStyle(NormalText))
 			.SetComboButtonStyle(SlimToolBarComboButton)
-			.SetLabelPadding(FMargin(5.f, 0.f, 0.f, 0.f))
+			.SetLabelPadding(FMargin(4.f, 0.f, 0.f, 0.f))
 			.SetEditableTextStyle(FEditableTextBoxStyle(NormalEditableTextBoxStyle));
 
-		FCheckBoxStyle SlimToolBarToggleButtonCheckBoxStyle = FCheckBoxStyle()
-			.SetCheckBoxType(ESlateCheckBoxType::ToggleButton)
-			.SetUncheckedImage(FSlateNoResource())
-			.SetUncheckedPressedImage(FSlateNoResource())
-			.SetUncheckedHoveredImage(FSlateNoResource())
-			.SetCheckedImage(FSlateNoResource())
-			.SetCheckedHoveredImage(FSlateNoResource())
-			.SetCheckedPressedImage(FSlateNoResource())
-			.SetForegroundColor(FStyleColors::Foreground)
-			.SetPressedForegroundColor(FStyleColors::ForegroundHover)
-			.SetHoveredForegroundColor(FStyleColors::ForegroundHover)
-			.SetCheckedForegroundColor(FStyleColors::Primary)
-			.SetCheckedPressedForegroundColor(FStyleColors::PrimaryPress)
-			.SetCheckedHoveredForegroundColor(FStyleColors::PrimaryHover);
+		FCheckBoxStyle SlimToolBarToggleButtonCheckBoxStyle = FCheckBoxStyle(Style->GetWidgetStyle<FCheckBoxStyle>("ToggleButtonCheckbox"))
+			.SetPadding(FMargin(8.f, 4.f, 8.f, 4.f));
 
 		SlimToolbarStyle.SetToggleButtonStyle(SlimToolBarToggleButtonCheckBoxStyle);
 		SlimToolbarStyle.SetButtonStyle(SlimToolBarButton);
 
-		FComboButtonStyle SlimToolBarSettingsComboButton = FComboButtonStyle(Style->GetWidgetStyle<FComboButtonStyle>("ComboButton"))
+		FButtonStyle SlimToolBarButtonSettingsRegion = FButtonStyle(SimpleButtonStyle)
+			.SetPressed(FSlateNoResource())
+			.SetHovered(FSlateNoResource())
+			.SetNormal(FSlateNoResource())
+			.SetNormalForeground(FStyleColors::Foreground)
+			.SetPressedForeground(FStyleColors::ForegroundHover)
+			.SetHoveredForeground(FStyleColors::ForegroundHover)
+			.SetDisabledForeground(FStyleColors::Foreground)
+			.SetNormalPadding(FMargin(8.f, 4.f, 8.f, 4.f))
+			.SetPressedPadding(FMargin(8.f, 5.f, 8.f, 3.f));
+
+
+		FCheckBoxStyle SettingsToggleButtonStyle = FCheckBoxStyle(Style->GetWidgetStyle<FCheckBoxStyle>("TransparentCheckBox"))
+			.SetPadding(FMargin(8.f, 4.f, 8.f, 4.f));
+
+		SlimToolbarStyle.SetSettingsButtonStyle(SlimToolBarButtonSettingsRegion);
+		SlimToolbarStyle.SetSettingsToggleButtonStyle(SettingsToggleButtonStyle);
+		SlimToolbarStyle.SetIconSize(Icon20x20);
+
+		SlimToolBarButtonSettingsRegion
+			.SetNormalPadding(FMargin(4.f, 4.f, 4.f, 4.f))
+			.SetPressedPadding(FMargin(4.f, 5.f, 4.f, 3.f));
+
+		FComboButtonStyle SlimToolBarSettingsComboButton = FComboButtonStyle(SlimToolBarComboButton)
 			.SetContentPadding(0)
-			.SetButtonStyle(SlimToolBarButton)
-			.SetDownArrowImage(IMAGE_BRUSH_SVG("Starship/Common/ellipsis-vertical-narrow", FVector2D(6, 24)));
+			.SetButtonStyle(SlimToolBarButtonSettingsRegion)
+			.SetDownArrowImage(IMAGE_BRUSH_SVG("Starship/Common/ellipsis-vertical-narrow", FVector2D(6, 20)));
 
 		SlimToolbarStyle.SetSettingsComboButtonStyle(SlimToolBarSettingsComboButton);
-		SlimToolbarStyle.SetIconSize(Icon20x20);
 
-		Style->Set("ToolbarSettingsRegion.LeftHover", new FSlateRoundedBoxBrush(FStyleColors::Dropdown, FVector4(4.0f, 0.0f, 0.0f, 4.0f)));
-		Style->Set("ToolbarSettingsRegion.RightHover", new FSlateRoundedBoxBrush(FStyleColors::Dropdown, FVector4(0.0f, 4.0f, 4.0f, 0.0f)));
-		Style->Set("ToolbarSettingsRegion.Left", new FSlateRoundedBoxBrush(FStyleColors::Header, FVector4(4.0f, 0.0f, 0.0f, 4.0f), FStyleColors::Header, 1.0f));
-		Style->Set("ToolbarSettingsRegion.Right", new FSlateRoundedBoxBrush(FStyleColors::Header, FVector4(0.0f, 4.0f, 4.0f, 0.0f), FStyleColors::Header, 1.0f));
+
+		Style->Set("ToolbarSettingsRegion.LeftHover", new FSlateRoundedBoxBrush(FStyleColors::Hover, FVector4(4.0f, 0.0f, 0.0f, 4.0f)));
+		Style->Set("ToolbarSettingsRegion.RightHover", new FSlateRoundedBoxBrush(FStyleColors::Hover, FVector4(0.0f, 4.0f, 4.0f, 0.0f)));
+		Style->Set("ToolbarSettingsRegion.Left", new FSlateRoundedBoxBrush(FStyleColors::Dropdown, FVector4(4.0f, 0.0f, 0.0f, 4.0f), FStyleColors::Header, 1.0f));
+		Style->Set("ToolbarSettingsRegion.Right", new FSlateRoundedBoxBrush(FStyleColors::Dropdown, FVector4(0.0f, 4.0f, 4.0f, 0.0f), FStyleColors::Header, 1.0f));
+		Style->Set("ToolbarSettingsRegion.LeftToggleHover", new FSlateRoundedBoxBrush(FStyleColors::PrimaryHover, FVector4(4.0f, 0.0f, 0.0f, 4.0f)));
+		Style->Set("ToolbarSettingsRegion.LeftToggle", new FSlateRoundedBoxBrush(FStyleColors::Primary, FVector4(4.0f, 0.0f, 0.0f, 4.0f)));
+
 
 		Style->Set("SlimToolBar", SlimToolbarStyle);
 
-		SlimToolbarStyle.SetComboButtonPadding(FMargin(4.0f, 0.0f));
-		SlimToolbarStyle.SetButtonPadding(FMargin(4.0f, 4.0f));
-		SlimToolbarStyle.SetCheckBoxPadding(FMargin(8.0f, 0.0f));
+
 		SlimToolbarStyle.SetShowLabels(false);
-		SlimToolbarStyle.SetLabelPadding(FMargin(5, 0, 0, 0));
-	
+		
 		Style->Set("AssetEditorToolbar", SlimToolbarStyle);
 
 		// Callout Toolbar - Used to "call out" the toolbar button with text
@@ -2070,14 +2048,14 @@ void FStarshipCoreStyle::SetupMultiboxStyles(TSharedRef<FStyle>& Style)
 		/* This radio button is actually just a check box with different images */
 		/* Set images for various Menu radio button (SCheckBox) states... */
 		const FCheckBoxStyle BasicMenuRadioButtonStyle = FCheckBoxStyle()
-			.SetUncheckedImage(           FSlateRoundedBoxBrush(FStyleColors::Header,  Icon8x8))
-			.SetUncheckedHoveredImage(    FSlateRoundedBoxBrush(FStyleColors::Hover2, Icon8x8))
-			.SetUncheckedPressedImage(    FSlateRoundedBoxBrush(FStyleColors::White,   Icon8x8))
+			.SetUncheckedImage(           FSlateRoundedBoxBrush(FStyleColors::Input,  Icon8x8))
+			.SetUncheckedHoveredImage(    FSlateRoundedBoxBrush(FStyleColors::Input, Icon8x8))
+			.SetUncheckedPressedImage(    FSlateRoundedBoxBrush(FStyleColors::Hover2,   Icon8x8))
 			.SetCheckedImage(             FSlateRoundedBoxBrush(FStyleColors::White,   Icon8x8))
 			.SetCheckedHoveredImage(      FSlateRoundedBoxBrush(FStyleColors::White,   Icon8x8))
 			.SetCheckedPressedImage(      FSlateRoundedBoxBrush(FStyleColors::White,   Icon8x8))
-			.SetUndeterminedImage(        FSlateRoundedBoxBrush(FStyleColors::Header,  Icon8x8))
-			.SetUndeterminedHoveredImage( FSlateRoundedBoxBrush(FStyleColors::Hover2, Icon8x8))
+			.SetUndeterminedImage(        FSlateRoundedBoxBrush(FStyleColors::Input,  Icon8x8))
+			.SetUndeterminedHoveredImage( FSlateRoundedBoxBrush(FStyleColors::White, Icon8x8))
 			.SetUndeterminedPressedImage( FSlateRoundedBoxBrush(FStyleColors::White,   Icon8x8));
 
 		/* ...and set new style */
@@ -2101,7 +2079,7 @@ void FStarshipCoreStyle::SetupMultiboxStyles(TSharedRef<FStyle>& Style)
 			.SetNormal(FSlateNoResource())
 			.SetPressed(FSlateColorBrush(FStyleColors::Primary))
 			.SetHovered(FSlateColorBrush(FStyleColors::Primary))
-			.SetHoveredForeground(FStyleColors::Black)
+			.SetHoveredForeground(FStyleColors::ForegroundHover)
 			.SetNormalPadding(FMargin(0, 2))
 			.SetPressedPadding(FMargin(0, 3, 0, 1));
 
@@ -2117,7 +2095,7 @@ void FStarshipCoreStyle::SetupMultiboxStyles(TSharedRef<FStyle>& Style)
 			FButtonStyle(MenuButton)
 			.SetHovered(FSlateColorBrush(FStyleColors::Hover))
 			.SetHoveredForeground(FStyleColors::ForegroundHover)
-			.SetPressedForeground(FStyleColors::Black)
+			.SetPressedForeground(FStyleColors::ForegroundHover)
 			.SetNormalForeground(FStyleColors::Foreground);
 
 		// For menu bars we need to ignore the button style
diff --git a/Engine/Source/Runtime/SlateCore/Private/Styling/StyleColors.cpp b/Engine/Source/Runtime/SlateCore/Private/Styling/StyleColors.cpp
--- a/Engine/Source/Runtime/SlateCore/Private/Styling/StyleColors.cpp
+++ b/Engine/Source/Runtime/SlateCore/Private/Styling/StyleColors.cpp
@@ -80,7 +80,7 @@ void USlateThemeManager::InitalizeDefaults()
 	SetDefaultColor(EStyleColor::WindowBorder, COLOR("0F0F0FFF"));
 	SetDefaultColor(EStyleColor::Foldout, COLOR("0F0F0FFF"));
 	SetDefaultColor(EStyleColor::Input, COLOR("0F0F0FFF"));
-	SetDefaultColor(EStyleColor::InputOutline, COLOR("454545FF"));
+	SetDefaultColor(EStyleColor::InputOutline, COLOR("383838FF"));
 	SetDefaultColor(EStyleColor::Recessed, COLOR("#1A1A1AFF"));
 	SetDefaultColor(EStyleColor::Panel, COLOR("#242424FF"));
 	SetDefaultColor(EStyleColor::Header, COLOR("#2F2F2FFF"));
@@ -91,11 +91,11 @@ void USlateThemeManager::InitalizeDefaults()
 
 	SetDefaultColor(EStyleColor::White, COLOR("#FFFFFFFF"));
 	SetDefaultColor(EStyleColor::White25, COLOR("#FFFFFF40"));
-	SetDefaultColor(EStyleColor::Highlight, COLOR("#0078D7FF"));
+	SetDefaultColor(EStyleColor::Highlight, COLOR("#0070E0FF"));
 
-	SetDefaultColor(EStyleColor::Primary, COLOR("#26BBFFFF"));
-	SetDefaultColor(EStyleColor::PrimaryHover, COLOR("#6FD2FFFF"));
-	SetDefaultColor(EStyleColor::PrimaryPress, COLOR("#1989BCFF"));
+	SetDefaultColor(EStyleColor::Primary, COLOR("#0070E0FF"));
+	SetDefaultColor(EStyleColor::PrimaryHover, COLOR("#0E86FFFF"));
+	SetDefaultColor(EStyleColor::PrimaryPress, COLOR("#0050A0FF"));
 	SetDefaultColor(EStyleColor::Secondary, COLOR("#383838FF"));
 
 	SetDefaultColor(EStyleColor::Foreground, COLOR("#C0C0C0FF"));
@@ -104,7 +104,7 @@ void USlateThemeManager::InitalizeDefaults()
 	SetDefaultColor(EStyleColor::ForegroundHeader, COLOR("#C8C8C8FF"));
 
 	SetDefaultColor(EStyleColor::Select, GetDefaultColor(EStyleColor::Primary));
-	SetDefaultColor(EStyleColor::SelectInactive, COLOR("#99B3BFFF"));
+	SetDefaultColor(EStyleColor::SelectInactive, COLOR("#40576F"));
 	SetDefaultColor(EStyleColor::SelectParent, COLOR("#2C323AFF"));
 	SetDefaultColor(EStyleColor::SelectHover, GetDefaultColor(EStyleColor::Panel));
 
@@ -364,7 +364,8 @@ void USlateThemeManager::EnsureValidCurrentTheme()
 	DefaultDarkTheme.DisplayName = NSLOCTEXT("StyleColors", "DefaultDarkTheme", "Dark");
 	// If you change this you invalidate the default dark theme forcing the default theme to be reset and the existing dark theme to become a user theme
 	// In general you should not do this. You should instead update the default "dark.json" file 
-	DefaultDarkTheme.Id = FGuid(0xE431F929, 0x452C4686, 0xAB5C0FE1, 0x2910653A);
+
+	DefaultDarkTheme.Id = FGuid(0x13438026, 0x5FBB4A9C, 0xA00A1DC9, 0x770217B8);
 	DefaultDarkTheme.Filename = FPaths::EngineContentDir() / TEXT("Slate/Themes/Dark.json");
 
 	int32 ThemeIndex = Themes.AddUnique(DefaultDarkTheme);
diff --git a/Engine/Source/Runtime/SlateCore/Private/Styling/ToolBarStyle.cpp b/Engine/Source/Runtime/SlateCore/Private/Styling/ToolBarStyle.cpp
--- a/Engine/Source/Runtime/SlateCore/Private/Styling/ToolBarStyle.cpp
+++ b/Engine/Source/Runtime/SlateCore/Private/Styling/ToolBarStyle.cpp
@@ -42,6 +42,8 @@ void FToolBarStyle::GetResources(TArray<const FSlateBrush*>& OutBrushes) const
 	EditableTextStyle.GetResources(OutBrushes);
 	ToggleButton.GetResources(OutBrushes);
 	SettingsComboButton.GetResources(OutBrushes);
+	SettingsToggleButton.GetResources(OutBrushes);
+	SettingsButtonStyle.GetResources(OutBrushes);
 	ButtonStyle.GetResources(OutBrushes);
 }
 
diff --git a/Engine/Source/Runtime/SlateCore/Public/Brushes/SlateNoResource.h b/Engine/Source/Runtime/SlateCore/Public/Brushes/SlateNoResource.h
--- a/Engine/Source/Runtime/SlateCore/Public/Brushes/SlateNoResource.h
+++ b/Engine/Source/Runtime/SlateCore/Public/Brushes/SlateNoResource.h
@@ -21,3 +21,17 @@ struct SLATECORE_API FSlateNoResource
 		: FSlateBrush(ESlateBrushDrawType::NoDrawType, FName(NAME_None), FMargin(0), ESlateBrushTileType::NoTile, ESlateBrushImageType::NoImage, InImageSize)
 	{ }
 };
+
+/** 
+ * This represents an undefined brush. FSlateNoResource is a valid brush which means don't render anything.  FSlateOptionalBrush means let the widget decide what to do if the brush is unset. 
+ * Widgets that don't check for the brush being unspecified will draw nothing
+ */
+struct SLATECORE_API FSlateOptionalBrush
+	: public FSlateBrush
+{
+	FSlateOptionalBrush(const FVector2D& InImageSize = FVector2D::ZeroVector)
+		: FSlateBrush(ESlateBrushDrawType::NoDrawType, FName(NAME_None), FMargin(0), ESlateBrushTileType::NoTile, ESlateBrushImageType::NoImage, InImageSize)
+	{
+		bIsSet = false;
+	}
+};
\ No newline at end of file
diff --git a/Engine/Source/Runtime/SlateCore/Public/Styling/CoreStyle.h b/Engine/Source/Runtime/SlateCore/Public/Styling/CoreStyle.h
--- a/Engine/Source/Runtime/SlateCore/Public/Styling/CoreStyle.h
+++ b/Engine/Source/Runtime/SlateCore/Public/Styling/CoreStyle.h
@@ -105,7 +105,12 @@ namespace CoreStyleConstants
 	// Common Margins
 	const FMargin DefaultMargins(8.f, 4.f);
 	// Buttons already have a built in (4., 2.) padding - adding to that a little
-	const FMargin ButtonMargins(20.f, 2.f, 20.f, 3.f);
+	const FMargin ButtonMargins(12.f, 1.f, 12.f, 1.f);
+
+	const FMargin PressedButtonMargins(12.f, 2.0f, 12.f, 0.0f);
+	const FMargin ToggleButtonMargins(16.0f, 2.0f);
+	const FMargin ComboButtonMargin(8.f, 1.f, 8.f, 1.f);
+	const FMargin PressedComboButtonMargin(8.f, 2.f, 8.f, 0.f);
 
 	const float InputFocusRadius = 4.f;
 	const float InputFocusThickness = 1.0f;
diff --git a/Engine/Source/Runtime/SlateCore/Public/Styling/SlateBrush.h b/Engine/Source/Runtime/SlateCore/Public/Styling/SlateBrush.h
--- a/Engine/Source/Runtime/SlateCore/Public/Styling/SlateBrush.h
+++ b/Engine/Source/Runtime/SlateCore/Public/Styling/SlateBrush.h
@@ -404,6 +404,8 @@ public:
 		return ResourceHandle;
 	}
 
+	bool IsSet() const { return bIsSet; }
+
 #if WITH_EDITOR
 	void InvalidateResourceHandle()
 	{
@@ -466,6 +468,9 @@ protected:
 	UPROPERTY()
 	uint8 bHasUObject_DEPRECATED:1;
 
+	/** This is true for all constructed brushes except for optional brushes */
+	uint8 bIsSet : 1;
+
 	/** 
 	 * This constructor is protected; use one of the deriving classes instead.
 	 *
diff --git a/Engine/Source/Runtime/SlateCore/Public/Styling/SlateTypes.h b/Engine/Source/Runtime/SlateCore/Public/Styling/SlateTypes.h
--- a/Engine/Source/Runtime/SlateCore/Public/Styling/SlateTypes.h
+++ b/Engine/Source/Runtime/SlateCore/Public/Styling/SlateTypes.h
@@ -110,47 +110,47 @@ struct SLATECORE_API FCheckBoxStyle : public FSlateWidgetStyle
 	/* CheckBox appearance when the CheckBox is unchecked (normal) */
 	UPROPERTY( EditAnywhere, BlueprintReadWrite, Category = Appearance )
 	FSlateBrush UncheckedImage;
-	FCheckBoxStyle & SetUncheckedImage( const FSlateBrush& InUncheckedImage ){ UncheckedImage = InUncheckedImage; return *this; }
+	FCheckBoxStyle& SetUncheckedImage( const FSlateBrush& InUncheckedImage ){ UncheckedImage = InUncheckedImage; return *this; }
 
 	/* CheckBox appearance when the CheckBox is unchecked and hovered */
 	UPROPERTY( EditAnywhere, BlueprintReadWrite, Category = Appearance )
 	FSlateBrush UncheckedHoveredImage;
-	FCheckBoxStyle & SetUncheckedHoveredImage( const FSlateBrush& InUncheckedHoveredImage ){ UncheckedHoveredImage = InUncheckedHoveredImage; return *this; }
+	FCheckBoxStyle& SetUncheckedHoveredImage( const FSlateBrush& InUncheckedHoveredImage ){ UncheckedHoveredImage = InUncheckedHoveredImage; return *this; }
 
 	/* CheckBox appearance when the CheckBox is unchecked and hovered */
 	UPROPERTY( EditAnywhere, BlueprintReadWrite, Category = Appearance )
 	FSlateBrush UncheckedPressedImage;
-	FCheckBoxStyle & SetUncheckedPressedImage( const FSlateBrush& InUncheckedPressedImage ){ UncheckedPressedImage = InUncheckedPressedImage; return *this; }
+	FCheckBoxStyle& SetUncheckedPressedImage( const FSlateBrush& InUncheckedPressedImage ){ UncheckedPressedImage = InUncheckedPressedImage; return *this; }
 
 	/* CheckBox appearance when the CheckBox is checked */
 	UPROPERTY( EditAnywhere, BlueprintReadWrite, Category = Appearance )
 	FSlateBrush CheckedImage;
-	FCheckBoxStyle & SetCheckedImage( const FSlateBrush& InCheckedImage ){ CheckedImage = InCheckedImage; return *this; }
+	FCheckBoxStyle& SetCheckedImage( const FSlateBrush& InCheckedImage ){ CheckedImage = InCheckedImage; return *this; }
 
 	/* CheckBox appearance when checked and hovered */
 	UPROPERTY( EditAnywhere, BlueprintReadWrite, Category = Appearance )
 	FSlateBrush CheckedHoveredImage;
-	FCheckBoxStyle & SetCheckedHoveredImage( const FSlateBrush& InCheckedHoveredImage ){ CheckedHoveredImage = InCheckedHoveredImage; return *this; }
+	FCheckBoxStyle& SetCheckedHoveredImage( const FSlateBrush& InCheckedHoveredImage ){ CheckedHoveredImage = InCheckedHoveredImage; return *this; }
 
 	/* CheckBox appearance when checked and pressed */
 	UPROPERTY( EditAnywhere, BlueprintReadWrite, Category = Appearance )
 	FSlateBrush CheckedPressedImage;
-	FCheckBoxStyle & SetCheckedPressedImage( const FSlateBrush& InCheckedPressedImage ){ CheckedPressedImage = InCheckedPressedImage; return *this; }
+	FCheckBoxStyle& SetCheckedPressedImage( const FSlateBrush& InCheckedPressedImage ){ CheckedPressedImage = InCheckedPressedImage; return *this; }
 	
 	/* CheckBox appearance when the CheckBox is undetermined */
 	UPROPERTY( EditAnywhere, BlueprintReadWrite, Category = Appearance )
 	FSlateBrush UndeterminedImage;
-	FCheckBoxStyle & SetUndeterminedImage( const FSlateBrush& InUndeterminedImage ){ UndeterminedImage = InUndeterminedImage; return *this; }
+	FCheckBoxStyle& SetUndeterminedImage( const FSlateBrush& InUndeterminedImage ){ UndeterminedImage = InUndeterminedImage; return *this; }
 
 	/* CheckBox appearance when CheckBox is undetermined and hovered */
 	UPROPERTY( EditAnywhere, BlueprintReadWrite, Category = Appearance )
 	FSlateBrush UndeterminedHoveredImage;
-	FCheckBoxStyle & SetUndeterminedHoveredImage( const FSlateBrush& InUndeterminedHoveredImage ){ UndeterminedHoveredImage = InUndeterminedHoveredImage; return *this; }
+	FCheckBoxStyle& SetUndeterminedHoveredImage( const FSlateBrush& InUndeterminedHoveredImage ){ UndeterminedHoveredImage = InUndeterminedHoveredImage; return *this; }
 
 	/* CheckBox appearance when CheckBox is undetermined and pressed */
 	UPROPERTY( EditAnywhere, BlueprintReadWrite, Category = Appearance )
 	FSlateBrush UndeterminedPressedImage;
-	FCheckBoxStyle & SetUndeterminedPressedImage( const FSlateBrush& InUndeterminedPressedImage ){ UndeterminedPressedImage = InUndeterminedPressedImage; return *this; }
+	FCheckBoxStyle& SetUndeterminedPressedImage( const FSlateBrush& InUndeterminedPressedImage ){ UndeterminedPressedImage = InUndeterminedPressedImage; return *this; }
 
 	/** Padding */
 	UPROPERTY(EditAnywhere, BlueprintReadWrite, Category=Appearance)
@@ -160,17 +160,17 @@ struct SLATECORE_API FCheckBoxStyle : public FSlateWidgetStyle
 	/** Background appearance */
 	UPROPERTY( EditAnywhere, BlueprintReadWrite, Category = Appearance )
 	FSlateBrush BackgroundImage;
-	FCheckBoxStyle & SetBackgroundImage( const FSlateBrush& InBackgroundImage ){ BackgroundImage = InBackgroundImage; return *this; }
+	FCheckBoxStyle& SetBackgroundImage( const FSlateBrush& InBackgroundImage ){ BackgroundImage = InBackgroundImage; return *this; }
 
 	/** Background appearance when hovered */
 	UPROPERTY( EditAnywhere, BlueprintReadWrite, Category = Appearance )
 	FSlateBrush BackgroundHoveredImage;
-	FCheckBoxStyle & SetBackgroundHoveredImage( const FSlateBrush& InBackgroundHoveredImage ){ BackgroundHoveredImage = InBackgroundHoveredImage; return *this; }
+	FCheckBoxStyle& SetBackgroundHoveredImage( const FSlateBrush& InBackgroundHoveredImage ){ BackgroundHoveredImage = InBackgroundHoveredImage; return *this; }
 
 	/** Background appearance when pressed */
 	UPROPERTY( EditAnywhere, BlueprintReadWrite, Category = Appearance )
 	FSlateBrush BackgroundPressedImage;
-	FCheckBoxStyle & SetBackgroundPressedImage( const FSlateBrush& InBackgroundPressedImage ){ BackgroundPressedImage = InBackgroundPressedImage; return *this; }
+	FCheckBoxStyle& SetBackgroundPressedImage( const FSlateBrush& InBackgroundPressedImage ){ BackgroundPressedImage = InBackgroundPressedImage; return *this; }
 
 	/** The normal unchecked foreground color */
 	UPROPERTY(EditAnywhere, BlueprintReadWrite, Category=Appearance)
@@ -1375,6 +1375,10 @@ struct SLATECORE_API FSpinBoxStyle : public FSlateWidgetStyle
 	FSlateBrush BackgroundBrush;
 	FSpinBoxStyle& SetBackgroundBrush( const FSlateBrush& InBackgroundBrush ){ BackgroundBrush = InBackgroundBrush; return *this; }
 
+	UPROPERTY(EditAnywhere, BlueprintReadWrite, Category = Appearance)
+	FSlateBrush ActiveBackgroundBrush;
+	FSpinBoxStyle& SetActiveBackgroundBrush(const FSlateBrush& InBackgroundBrush) { ActiveBackgroundBrush = InBackgroundBrush; return *this; }
+
 	/** Brush used to draw the background of the spinbox when it's hovered over */
 	UPROPERTY(EditAnywhere, BlueprintReadWrite, Category=Appearance)
 	FSlateBrush HoveredBackgroundBrush;
@@ -1385,6 +1389,11 @@ struct SLATECORE_API FSpinBoxStyle : public FSlateWidgetStyle
 	FSlateBrush ActiveFillBrush;
 	FSpinBoxStyle& SetActiveFillBrush( const FSlateBrush& InActiveFillBrush ){ ActiveFillBrush = InActiveFillBrush; return *this; }
 
+	/** Brush used to fill the spinbox when it's hovered and not active */
+	UPROPERTY(EditAnywhere, BlueprintReadWrite, Category = Appearance)
+	FSlateBrush HoveredFillBrush;
+	FSpinBoxStyle& SetHoveredFillBrush(const FSlateBrush& InHoveredBrush) { HoveredFillBrush = InHoveredBrush; return *this; }
+
 	/** Brush used to fill the spinbox when it's inactive */
 	UPROPERTY(EditAnywhere, BlueprintReadWrite, Category=Appearance)
 	FSlateBrush InactiveFillBrush;
@@ -1405,6 +1414,11 @@ struct SLATECORE_API FSpinBoxStyle : public FSlateWidgetStyle
 	FMargin TextPadding;
 	FSpinBoxStyle& SetTextPadding( const FMargin& InTextPadding ){ TextPadding = InTextPadding; return *this; }
 
+	/** Padding between the background brush and the fill brush */
+	UPROPERTY(EditAnywhere, BlueprintReadWrite, Category = Appearance)
+	FMargin InsetPadding;
+	FSpinBoxStyle& SetInsetPadding(const FMargin& InInsetPadding) { InsetPadding = InInsetPadding; return *this; }
+
 	/**
 	 * Unlinks all colors in this style.
 	 * @see FSlateColor::Unlink
@@ -1413,7 +1427,9 @@ struct SLATECORE_API FSpinBoxStyle : public FSlateWidgetStyle
 	{
 		BackgroundBrush.UnlinkColors();
 		HoveredBackgroundBrush.UnlinkColors();
+		ActiveBackgroundBrush.UnlinkColors();
 		ActiveFillBrush.UnlinkColors();
+		HoveredFillBrush.UnlinkColors();
 		InactiveFillBrush.UnlinkColors();
 		ArrowsImage.UnlinkColors();
 		ForegroundColor.Unlink();
diff --git a/Engine/Source/Runtime/SlateCore/Public/Styling/ToolBarStyle.h b/Engine/Source/Runtime/SlateCore/Public/Styling/ToolBarStyle.h
--- a/Engine/Source/Runtime/SlateCore/Public/Styling/ToolBarStyle.h
+++ b/Engine/Source/Runtime/SlateCore/Public/Styling/ToolBarStyle.h
@@ -57,10 +57,18 @@ struct SLATECORE_API FToolBarStyle : public FSlateWidgetStyle
 	FComboButtonStyle ComboButtonStyle;
 	FToolBarStyle& SetComboButtonStyle(const FComboButtonStyle& InComboButtonStyle) { ComboButtonStyle = InComboButtonStyle; return *this; }
 
+	UPROPERTY(EditAnywhere, BlueprintReadWrite, Category = Appearance)
+	FButtonStyle SettingsButtonStyle;
+	FToolBarStyle& SetSettingsButtonStyle(const FButtonStyle& InSettingsButton) { SettingsButtonStyle = InSettingsButton; return *this; }
+
 	UPROPERTY(EditAnywhere, BlueprintReadWrite, Category = Appearance)
 	FComboButtonStyle SettingsComboButton;
 	FToolBarStyle& SetSettingsComboButtonStyle(const FComboButtonStyle& InSettingsComboButton) { SettingsComboButton = InSettingsComboButton; return *this; }
 
+	UPROPERTY(EditAnywhere, BlueprintReadWrite, Category = Appearance)
+	FCheckBoxStyle SettingsToggleButton;
+	FToolBarStyle& SetSettingsToggleButtonStyle(const FCheckBoxStyle& InSettingsToggleButton) { SettingsToggleButton = InSettingsToggleButton; return *this; }
+
 	UPROPERTY(EditAnywhere, BlueprintReadWrite, Category = Appearance)
 	FButtonStyle ButtonStyle;
 	FToolBarStyle& SetButtonStyle(const FButtonStyle& InButtonStyle) { ButtonStyle = InButtonStyle; return *this; }
