Updated segmented control style

diff --git a/Engine/Source/Runtime/Slate/Public/Widgets/Input/SSegmentedControl.h b/Engine/Source/Runtime/Slate/Public/Widgets/Input/SSegmentedControl.h
--- a/Engine/Source/Runtime/Slate/Public/Widgets/Input/SSegmentedControl.h
+++ b/Engine/Source/Runtime/Slate/Public/Widgets/Input/SSegmentedControl.h
@@ -168,7 +168,11 @@ public:
 
 	void RebuildChildren()
 	{
-		TSharedPtr<SUniformGridPanel> UniformBox = SNew(SUniformGridPanel);
+		// The right padding will be applied later at the end so we dont accumulate left+right padding between all buttons
+		FMargin SlotPadding = Style->UniformPadding;
+		SlotPadding.Right = 0.0f;
+
+		TSharedPtr<SUniformGridPanel> UniformBox = SNew(SUniformGridPanel).SlotPadding(SlotPadding);
 
 		const int32 NumSlots = Children.Num();
 		for ( int32 SlotIndex = 0; SlotIndex < NumSlots; ++SlotIndex )
@@ -240,7 +244,12 @@ public:
 
 		ChildSlot
 		[
-			UniformBox.ToSharedRef()
+			SNew(SBorder)
+			.BorderImage(&Style->BackgroundBrush)
+			.Padding(FMargin(0,0,Style->UniformPadding.Right,0))
+			[
+				UniformBox.ToSharedRef()
+			]
 		];
 
 	}
diff --git a/Engine/Source/Runtime/SlateCore/Private/Styling/SegmentedControlStyle.cpp b/Engine/Source/Runtime/SlateCore/Private/Styling/SegmentedControlStyle.cpp
--- a/Engine/Source/Runtime/SlateCore/Private/Styling/SegmentedControlStyle.cpp
+++ b/Engine/Source/Runtime/SlateCore/Private/Styling/SegmentedControlStyle.cpp
@@ -1,15 +1,19 @@
 // Copyright Epic Games, Inc. All Rights Reserved.
 
 #include "Styling/SegmentedControlStyle.h"
+#include "Brushes/SlateNoResource.h"
 
 const FName FSegmentedControlStyle::TypeName( TEXT("FSegmentedControlStyle") );
 
 FSegmentedControlStyle::FSegmentedControlStyle()
+	: BackgroundBrush(FSlateNoResource())
+	, UniformPadding(0)
 {
 }
 
 void FSegmentedControlStyle::GetResources( TArray< const FSlateBrush* >& OutBrushes ) const
 {
+	OutBrushes.Add(&BackgroundBrush);
 	ControlStyle.GetResources(OutBrushes);
 	FirstControlStyle.GetResources(OutBrushes);
 	LastControlStyle.GetResources(OutBrushes);
diff --git a/Engine/Source/Runtime/SlateCore/Private/Styling/StarshipCoreStyle.cpp b/Engine/Source/Runtime/SlateCore/Private/Styling/StarshipCoreStyle.cpp
--- a/Engine/Source/Runtime/SlateCore/Private/Styling/StarshipCoreStyle.cpp
+++ b/Engine/Source/Runtime/SlateCore/Private/Styling/StarshipCoreStyle.cpp
@@ -1392,47 +1392,12 @@ void FStarshipCoreStyle::SetupCheckboxStyles(TSharedRef<FStyle>& Style)
 	/* Style for a segmented box */
 	FCheckBoxStyle SegmentedBoxLeft = FCheckBoxStyle()
 		.SetCheckBoxType(ESlateCheckBoxType::ToggleButton)
-
-		.SetUncheckedImage(BOX_BRUSH("/Starship/CoreWidgets/SegmentedBox/left", FVector2D(16.f, 16.f), FMargin(4.0f / 16.0f), FStyleColors::Input))
-		.SetUncheckedHoveredImage(BOX_BRUSH("/Starship/CoreWidgets/SegmentedBox/left", FVector2D(16.f, 16.f), FMargin(4.0f / 16.0f), FStyleColors::Hover))
-		.SetUncheckedPressedImage(BOX_BRUSH("/Starship/CoreWidgets/SegmentedBox/left", FVector2D(16.f, 16.f), FMargin(4.0f / 16.0f), FStyleColors::Secondary))
-		.SetCheckedImage(BOX_BRUSH("/Starship/CoreWidgets/SegmentedBox/left", FVector2D(16.f, 16.f), FMargin(4.0f / 16.0f), FStyleColors::Dropdown))
-		.SetCheckedHoveredImage(BOX_BRUSH("/Starship/CoreWidgets/SegmentedBox/left", FVector2D(16.f, 16.f), FMargin(4.0f / 16.0f), FStyleColors::Hover))
-		.SetCheckedPressedImage(BOX_BRUSH("/Starship/CoreWidgets/SegmentedBox/left", FVector2D(16.f, 16.f), FMargin(4.0f / 16.0f), FStyleColors::Header))
-		.SetForegroundColor(FStyleColors::Foreground)
-		.SetHoveredForegroundColor(FStyleColors::ForegroundHover)
-		.SetPressedForegroundColor(FStyleColors::ForegroundHover)
-		.SetCheckedForegroundColor(FStyleColors::ForegroundHover)
-		.SetCheckedHoveredForegroundColor(FStyleColors::ForegroundHover)
-		.SetCheckedPressedForegroundColor(FStyleColors::ForegroundHover)
-		.SetPadding(SegmentedControlMargins);
-
-
-	const FCheckBoxStyle SegmentedBoxCenter = FCheckBoxStyle()
-		.SetCheckBoxType(ESlateCheckBoxType::ToggleButton)
-		.SetUncheckedImage(FSlateColorBrush(FStyleColors::Input))
-		.SetUncheckedHoveredImage(FSlateColorBrush(FStyleColors::Hover))
-		.SetUncheckedPressedImage(FSlateColorBrush(FStyleColors::Secondary))
-		.SetCheckedImage(FSlateColorBrush(FStyleColors::Dropdown))
-		.SetCheckedHoveredImage(FSlateColorBrush(FStyleColors::Hover))
-		.SetCheckedPressedImage(FSlateColorBrush(FStyleColors::Header))
-		.SetForegroundColor(FStyleColors::Foreground)
-		.SetHoveredForegroundColor(FStyleColors::ForegroundHover)
-		.SetPressedForegroundColor(FStyleColors::ForegroundHover)
-		.SetCheckedForegroundColor(FStyleColors::ForegroundHover)
-		.SetCheckedHoveredForegroundColor(FStyleColors::ForegroundHover)
-		.SetCheckedPressedForegroundColor(FStyleColors::ForegroundHover)
-		.SetPadding(SegmentedControlMargins);
-
-
-	const FCheckBoxStyle SegmentedBoxRight = FCheckBoxStyle()
-		.SetCheckBoxType(ESlateCheckBoxType::ToggleButton)
-		.SetUncheckedImage(BOX_BRUSH("/Starship/CoreWidgets/SegmentedBox/right", FVector2D(16.f, 16.f), FMargin(4.0f / 16.0f), FStyleColors::Input))
-		.SetUncheckedHoveredImage(BOX_BRUSH("/Starship/CoreWidgets/SegmentedBox/right", FVector2D(16.f, 16.f), FMargin(4.0f / 16.0f), FStyleColors::Hover))
-		.SetUncheckedPressedImage(BOX_BRUSH("/Starship/CoreWidgets/SegmentedBox/right", FVector2D(16.f, 16.f), FMargin(4.0f / 16.0f), FStyleColors::Header))
-		.SetCheckedImage(BOX_BRUSH("/Starship/CoreWidgets/SegmentedBox/right", FVector2D(16.f, 16.f), FMargin(4.0f / 16.0f), FStyleColors::Dropdown))
-		.SetCheckedHoveredImage(BOX_BRUSH("/Starship/CoreWidgets/SegmentedBox/right", FVector2D(16.f, 16.f), FMargin(4.0f / 16.0f), FStyleColors::Hover))
-		.SetCheckedPressedImage(BOX_BRUSH("/Starship/CoreWidgets/SegmentedBox/right", FVector2D(16.f, 16.f), FMargin(4.0f / 16.0f), FStyleColors::Header))
+		.SetUncheckedImage(FSlateNoResource())
+		.SetUncheckedHoveredImage(FSlateRoundedBoxBrush(FStyleColors::Hover, InputFocusRadius))
+		.SetUncheckedPressedImage(FSlateRoundedBoxBrush(FStyleColors::Header, InputFocusRadius))
+		.SetCheckedImage(FSlateRoundedBoxBrush(FStyleColors::Dropdown, InputFocusRadius))
+		.SetCheckedHoveredImage(FSlateRoundedBoxBrush(FStyleColors::Hover, InputFocusRadius))
+		.SetCheckedPressedImage(FSlateRoundedBoxBrush(FStyleColors::Header, InputFocusRadius))
 		.SetForegroundColor(FStyleColors::Foreground)
 		.SetHoveredForegroundColor(FStyleColors::ForegroundHover)
 		.SetPressedForegroundColor(FStyleColors::ForegroundHover)
@@ -1442,14 +1407,21 @@ void FStarshipCoreStyle::SetupCheckboxStyles(TSharedRef<FStyle>& Style)
 		.SetPadding(SegmentedControlMargins);
 
 	Style->Set("SegmentedControl", FSegmentedControlStyle()
-		.SetControlStyle(SegmentedBoxCenter)
+		.SetControlStyle(SegmentedBoxLeft)
 		.SetFirstControlStyle(SegmentedBoxLeft)
-		.SetLastControlStyle(SegmentedBoxRight)
+		.SetLastControlStyle(SegmentedBoxLeft)
+		.SetBackgroundBrush(FSlateRoundedBoxBrush(FStyleColors::Input, InputFocusRadius))
+		.SetUniformPadding(FMargin(1))
 	);
 
-	SegmentedBoxLeft.SetUncheckedImage(BOX_BRUSH("/Starship/CoreWidgets/SegmentedBox/left", FVector2D(16.f, 16.f), FMargin(4.0f / 16.0f), FStyleColors::Dropdown));
-	SegmentedBoxLeft.SetCheckedImage(BOX_BRUSH("/Starship/CoreWidgets/SegmentedBox/left", FVector2D(16.f, 16.f), FMargin(4.0f / 16.0f), FStyleColors::Primary));
-	SegmentedBoxLeft.SetCheckedHoveredImage(BOX_BRUSH("/Starship/CoreWidgets/SegmentedBox/right", FVector2D(16.f, 16.f), FMargin(4.0f / 16.0f), FStyleColors::PrimaryHover));
+
+	SegmentedBoxLeft
+		.SetUncheckedImage(BOX_BRUSH("/Starship/CoreWidgets/SegmentedBox/left", FVector2D(16.f, 16.f), FMargin(4.0f / 16.0f), FStyleColors::Dropdown))
+		.SetUncheckedHoveredImage(BOX_BRUSH("/Starship/CoreWidgets/SegmentedBox/left", FVector2D(16.f, 16.f), FMargin(4.0f / 16.0f), FStyleColors::Hover))
+		.SetUncheckedPressedImage(BOX_BRUSH("/Starship/CoreWidgets/SegmentedBox/left", FVector2D(16.f, 16.f), FMargin(4.0f / 16.0f), FStyleColors::Secondary))
+		.SetCheckedImage(BOX_BRUSH("/Starship/CoreWidgets/SegmentedBox/left", FVector2D(16.f, 16.f), FMargin(4.0f / 16.0f), FStyleColors::Primary))
+		.SetCheckedHoveredImage(BOX_BRUSH("/Starship/CoreWidgets/SegmentedBox/left", FVector2D(16.f, 16.f), FMargin(4.0f / 16.0f), FStyleColors::PrimaryHover))
+		.SetCheckedPressedImage(BOX_BRUSH("/Starship/CoreWidgets/SegmentedBox/left", FVector2D(16.f, 16.f), FMargin(4.0f / 16.0f), FStyleColors::PrimaryPress));
 
 	Style->Set("SegmentedCombo.Left", SegmentedBoxLeft);
 
diff --git a/Engine/Source/Runtime/SlateCore/Public/Styling/SegmentedControlStyle.h b/Engine/Source/Runtime/SlateCore/Public/Styling/SegmentedControlStyle.h
--- a/Engine/Source/Runtime/SlateCore/Public/Styling/SegmentedControlStyle.h
+++ b/Engine/Source/Runtime/SlateCore/Public/Styling/SegmentedControlStyle.h
@@ -48,6 +48,18 @@ struct SLATECORE_API FSegmentedControlStyle: public FSlateWidgetStyle
 	FCheckBoxStyle LastControlStyle;
 	FSegmentedControlStyle& SetLastControlStyle( const FCheckBoxStyle& InStyle ){ LastControlStyle = InStyle; return *this; }
 
+	/** Background of the segmented control */
+	UPROPERTY(EditAnywhere, BlueprintReadWrite, Category=Appearance)
+	FSlateBrush BackgroundBrush;
+	FSegmentedControlStyle& SetBackgroundBrush(const FSlateBrush& InBrush) { BackgroundBrush = InBrush; return *this; }
+
+	/**
+	 * Padding between each control
+	 */
+	UPROPERTY(EditAnywhere, BlueprintReadWrite, Category = Appearance)
+	FMargin UniformPadding;
+	FSegmentedControlStyle& SetUniformPadding(const FMargin& InMargin) { UniformPadding = InMargin; return *this; }
+
 	/**
 	 * Unlinks all colors in this style.
 	 * @see FSlateColor::Unlink
@@ -57,6 +69,7 @@ struct SLATECORE_API FSegmentedControlStyle: public FSlateWidgetStyle
 		ControlStyle.UnlinkColors();
 		FirstControlStyle.UnlinkColors();
 		LastControlStyle.UnlinkColors();
+		BackgroundBrush.UnlinkColors();
 	}
 
 };
\ No newline at end of file
