Added Primary Button to new "ToolWidgets" module. The ToolWidgets module is designed to house common widgets used in tool like applications (editors, insights, ufe,etc) that do not depend on UnrealEd

diff --git a/Engine/Source/Developer/ToolWidgets/Private/SPrimaryButton.cpp b/Engine/Source/Developer/ToolWidgets/Private/SPrimaryButton.cpp
new file mode 100644
--- /dev/null
+++ b/Engine/Source/Developer/ToolWidgets/Private/SPrimaryButton.cpp
@@ -0,0 +1,38 @@
+// Copyright Epic Games, Inc. All Rights Reserved.
+
+#include "SPrimaryButton.h"
+#include "Widgets/SBoxPanel.h"
+#include "Widgets/Images/SImage.h"
+#include "Widgets/Text/STextBlock.h"
+
+void SPrimaryButton::Construct(const FArguments& InArgs)
+{
+	SButton::Construct(SButton::FArguments()
+		.ButtonStyle(&FAppStyle::Get().GetWidgetStyle<FButtonStyle>(InArgs._Icon.IsSet() ? "PrimaryButtonLabelAndIcon" : "PrimaryButton"))
+		.OnClicked(InArgs._OnClicked)
+		.ForegroundColor(FSlateColor::UseStyle())
+		[
+			SNew(SHorizontalBox)
+			+SHorizontalBox::Slot()
+			.AutoWidth()
+			.VAlign(VAlign_Center)
+			.Padding(FMargin(0,0,3,0))
+			[
+				SNew(SImage)
+				.ColorAndOpacity(FSlateColor::UseForeground())
+				.Image(InArgs._Icon)
+				.Visibility(InArgs._Icon.IsSet() ? EVisibility::HitTestInvisible : EVisibility::Collapsed)
+			]
+			+ SHorizontalBox::Slot()
+			.VAlign(VAlign_Center)
+			.AutoWidth()
+			[
+				SNew(STextBlock)
+				.TextStyle(&FAppStyle::Get().GetWidgetStyle<FTextBlockStyle>("PrimaryButtonText"))
+				.Text(InArgs._Text)
+				.Visibility(InArgs._Text.IsSet() ? EVisibility::HitTestInvisible : EVisibility::Collapsed)
+			]
+		]
+	);
+}
+
diff --git a/Engine/Source/Developer/ToolWidgets/Private/ToolWidgetsModule.cpp b/Engine/Source/Developer/ToolWidgets/Private/ToolWidgetsModule.cpp
new file mode 100644
--- /dev/null
+++ b/Engine/Source/Developer/ToolWidgets/Private/ToolWidgetsModule.cpp
@@ -0,0 +1,18 @@
+// Copyright Epic Games, Inc. All Rights Reserved.
+
+#include "Modules/ModuleInterface.h"
+#include "Modules/ModuleManager.h"
+
+class FToolWidgetsModule
+	: public IModuleInterface
+{
+public:
+
+	// IModuleInterface interface
+	
+	virtual void StartupModule() override { }
+	virtual void ShutdownModule() override { }
+};
+
+
+IMPLEMENT_MODULE(FToolWidgetsModule, ToolWidgets);
diff --git a/Engine/Source/Developer/ToolWidgets/Public/SPrimaryButton.h b/Engine/Source/Developer/ToolWidgets/Public/SPrimaryButton.h
new file mode 100644
--- /dev/null
+++ b/Engine/Source/Developer/ToolWidgets/Public/SPrimaryButton.h
@@ -0,0 +1,29 @@
+// Copyright Epic Games, Inc. All Rights Reserved.
+
+#pragma once
+
+#include "Framework/SlateDelegates.h"
+#include "Widgets/DeclarativeSyntaxSupport.h"
+#include "Widgets/Input/SButton.h"
+
+class TOOLWIDGETS_API SPrimaryButton : public SButton
+{
+public:
+
+	SLATE_BEGIN_ARGS(SPrimaryButton)
+	{}
+		/** The text to display in the button. */
+		SLATE_ATTRIBUTE(FText, Text)
+
+		/** Optional icon to display in the button. */
+		SLATE_ATTRIBUTE(const FSlateBrush*, Icon)
+
+		/** 
+		 * Called when the button is clicked  
+		 */
+		SLATE_EVENT(FOnClicked, OnClicked)
+
+	SLATE_END_ARGS()
+
+	void Construct(const FArguments& InArgs);
+};
\ No newline at end of file
diff --git a/Engine/Source/Developer/ToolWidgets/ToolWidgets.Build.cs b/Engine/Source/Developer/ToolWidgets/ToolWidgets.Build.cs
new file mode 100644
--- /dev/null
+++ b/Engine/Source/Developer/ToolWidgets/ToolWidgets.Build.cs
@@ -0,0 +1,31 @@
+// Copyright Epic Games, Inc. All Rights Reserved.
+
+using UnrealBuildTool;
+
+public class ToolWidgets : ModuleRules
+{
+	public ToolWidgets(ReadOnlyTargetRules Target) : base(Target)
+	{
+		/** NOTE: THIS MODULE SHOULD NOT EVER DEPEND ON UNREALED. 
+		 * If you are adding a reusable widget that depends on UnrealEd, add it to EditorWidgets instead
+		 */
+		PublicDependencyModuleNames.AddRange(
+			new string[] {
+				"Core",
+				"Slate",
+				"SlateCore",
+			}
+		);
+
+		PrivateDependencyModuleNames.AddRange(
+			new string[] {
+			}
+		);
+
+		PrivateIncludePaths.AddRange(
+			new string[] {
+				"Developer/ToolWidgets/Private",
+			}
+		);
+	}
+}
diff --git a/Engine/Source/Runtime/SlateCore/Private/Styling/StarshipCoreStyle.cpp b/Engine/Source/Runtime/SlateCore/Private/Styling/StarshipCoreStyle.cpp
--- a/Engine/Source/Runtime/SlateCore/Private/Styling/StarshipCoreStyle.cpp
+++ b/Engine/Source/Runtime/SlateCore/Private/Styling/StarshipCoreStyle.cpp
@@ -1113,6 +1113,10 @@ void FStarshipCoreStyle::SetupButtonStyles(TSharedRef<FStyle>& Style)
 		.SetPressedPadding(PressedButtonMargins);
 
 
+	const FButtonStyle PrimaryButtonLabelAndIcon = FButtonStyle(PrimaryButton)
+		.SetNormalPadding(FMargin(4.f, 1.f, 4.f, 1.f))
+		.SetPressedPadding(FMargin(4.f, 2.0f, 4.f, 0.0f));
+
 	const FButtonStyle Button = FButtonStyle()
 		.SetNormal(FSlateRoundedBoxBrush(FStyleColors::Secondary, 4.0f, FStyleColors::Input, InputFocusThickness))
 		.SetHovered(FSlateRoundedBoxBrush(FStyleColors::Hover, 4.0f, FStyleColors::Input, InputFocusThickness))
@@ -1125,19 +1129,6 @@ void FStarshipCoreStyle::SetupButtonStyles(TSharedRef<FStyle>& Style)
 		.SetNormalPadding(ButtonMargins)
 		.SetPressedPadding(PressedButtonMargins);
 
-	// @todo Recolor - Get rid of this, it is a copy of button
-	const FButtonStyle SecondaryButton = FButtonStyle()
-		.SetNormal(FSlateRoundedBoxBrush(FStyleColors::Secondary, 4.0f, FStyleColors::InputOutline, InputFocusThickness))
-		.SetHovered(FSlateRoundedBoxBrush(FStyleColors::Hover, 4.0f, FStyleColors::InputOutline, InputFocusThickness))
-		.SetPressed(FSlateRoundedBoxBrush(FStyleColors::Header, 4.0f, FStyleColors::InputOutline, InputFocusThickness))
-		.SetDisabled(FSlateRoundedBoxBrush(FStyleColors::Dropdown, 4.0f, FStyleColors::InputOutline, InputFocusThickness))
-		.SetNormalForeground(FStyleColors::ForegroundHover)
-		.SetHoveredForeground(FStyleColors::ForegroundHover)
-		.SetPressedForeground(FStyleColors::ForegroundHover)
-		.SetDisabledForeground(FStyleColors::Foreground)
-		.SetNormalPadding(FMargin(8.f, 4.5f, 8.f, 3.5f))
-		.SetPressedPadding(FMargin(8.f, 5.f, 6.f, 3.f));
-
 
 	const FButtonStyle SimpleButton = FButtonStyle()
 		.SetNormal(FSlateNoResource())
@@ -1162,16 +1153,15 @@ void FStarshipCoreStyle::SetupButtonStyles(TSharedRef<FStyle>& Style)
 
 		const FTextBlockStyle& NormalText = Style->GetWidgetStyle<FTextBlockStyle>("NormalText");
 
-		Style->Set("ButtonText", FTextBlockStyle(NormalText).SetFont(FStyleFonts::Get().NormalBold));
+		Style->Set("ButtonText", NormalText);
 
 		Style->Set("PrimaryButton", PrimaryButton);
+		Style->Set("PrimaryButtonLabelAndIcon", PrimaryButtonLabelAndIcon);
 		Style->Set("Button", Button);
 		Style->Set("SimpleButton", SimpleButtonIconOnly);
 		Style->Set("SimpleButtonLabelAndIcon", SimpleButtonLabelAndIcon);
-		Style->Set("SecondaryButton", SecondaryButton);
 
-		//@todo Recolor - This is reundant with button text now.
-		Style->Set("DialogButtonText", FTextBlockStyle(NormalText).SetFont(FStyleFonts::Get().NormalBold));
+		Style->Set("PrimaryButtonText", FTextBlockStyle(NormalText).SetFont(FStyleFonts::Get().NormalBold));
 
 		FSlateFontInfo AddFont = FStyleFonts::Get().SmallBold;
 		AddFont.LetterSpacing = 200;
