Upade starshipcorestyle with some button styles for backwards compat

diff --git a/Engine/Source/Runtime/SlateCore/Private/Styling/StarshipCoreStyle.cpp b/Engine/Source/Runtime/SlateCore/Private/Styling/StarshipCoreStyle.cpp
--- a/Engine/Source/Runtime/SlateCore/Private/Styling/StarshipCoreStyle.cpp
+++ b/Engine/Source/Runtime/SlateCore/Private/Styling/StarshipCoreStyle.cpp
@@ -1166,6 +1166,12 @@ void FStarshipCoreStyle::SetupButtonStyles(TSharedRef<FStyle>& Style)
 		FSlateFontInfo AddFont = FStyleFonts::Get().SmallBold;
 		AddFont.LetterSpacing = 200;
 
+		Style->Set("SecondaryButton", SecondaryButton);
+
+		//@todo Recolor - This is redundant with button text now. Temporarily here for backwards compat
+		Style->Set("DialogButtonText", NormalText);
+
+
 		//@todo Recolor - Rename to simple button text
 		Style->Set("SmallButtonText", FTextBlockStyle(NormalText));
 
