Add back in SecondaryButton style

diff --git a/Engine/Source/Runtime/SlateCore/Private/Styling/StarshipCoreStyle.cpp b/Engine/Source/Runtime/SlateCore/Private/Styling/StarshipCoreStyle.cpp
--- a/Engine/Source/Runtime/SlateCore/Private/Styling/StarshipCoreStyle.cpp
+++ b/Engine/Source/Runtime/SlateCore/Private/Styling/StarshipCoreStyle.cpp
@@ -1166,7 +1166,7 @@ void FStarshipCoreStyle::SetupButtonStyles(TSharedRef<FStyle>& Style)
 		FSlateFontInfo AddFont = FStyleFonts::Get().SmallBold;
 		AddFont.LetterSpacing = 200;
 
-		Style->Set("SecondaryButton", SecondaryButton);
+		Style->Set("SecondaryButton", Button);
 
 		//@todo Recolor - This is redundant with button text now. Temporarily here for backwards compat
 		Style->Set("DialogButtonText", NormalText);
