Update default editor buttons padding and size

diff --git a/Engine/Source/Runtime/SlateCore/Private/Styling/StarshipCoreStyle.cpp b/Engine/Source/Runtime/SlateCore/Private/Styling/StarshipCoreStyle.cpp
--- a/Engine/Source/Runtime/SlateCore/Private/Styling/StarshipCoreStyle.cpp
+++ b/Engine/Source/Runtime/SlateCore/Private/Styling/StarshipCoreStyle.cpp
@@ -1112,10 +1112,9 @@ void FStarshipCoreStyle::SetupButtonStyles(TSharedRef<FStyle>& Style)
 		.SetNormalPadding(ButtonMargins)
 		.SetPressedPadding(PressedButtonMargins);
 
-
 	const FButtonStyle PrimaryButtonLabelAndIcon = FButtonStyle(PrimaryButton)
-		.SetNormalPadding(FMargin(4.f, 1.f, 4.f, 1.f))
-		.SetPressedPadding(FMargin(4.f, 2.0f, 4.f, 0.0f));
+		.SetNormalPadding(FMargin(4.f, 1.5f, 4.f, 1.5f))
+		.SetPressedPadding(FMargin(4.f, 2.5f, 4.f, 0.5f));
 
 	const FButtonStyle Button = FButtonStyle()
 		.SetNormal(FSlateRoundedBoxBrush(FStyleColors::Secondary, 4.0f, FStyleColors::Input, InputFocusThickness))
diff --git a/Engine/Source/Runtime/SlateCore/Public/Styling/CoreStyle.h b/Engine/Source/Runtime/SlateCore/Public/Styling/CoreStyle.h
--- a/Engine/Source/Runtime/SlateCore/Public/Styling/CoreStyle.h
+++ b/Engine/Source/Runtime/SlateCore/Public/Styling/CoreStyle.h
@@ -105,9 +105,9 @@ namespace CoreStyleConstants
 	// Common Margins
 	const FMargin DefaultMargins(8.f, 4.f);
 	// Buttons already have a built in (4., 2.) padding - adding to that a little
-	const FMargin ButtonMargins(12.f, 1.f, 12.f, 1.f);
+	const FMargin ButtonMargins(12.f, 1.5f, 12.f, 1.5f);
 
-	const FMargin PressedButtonMargins(12.f, 2.0f, 12.f, 0.0f);
+	const FMargin PressedButtonMargins(12.f, 2.5f, 12.f, 0.5f);
 	const FMargin ToggleButtonMargins(16.0f, 2.0f);
 	const FMargin ComboButtonMargin(8.f, 1.f, 8.f, 1.f);
 	const FMargin PressedComboButtonMargin(8.f, 2.f, 8.f, 0.f);
