Replace chevrons with triangles in tree views/details panels + tweak colors to match latest designs

diff --git a/Engine/Source/Editor/PropertyEditor/Private/SDetailExpanderArrow.cpp b/Engine/Source/Editor/PropertyEditor/Private/SDetailExpanderArrow.cpp
--- a/Engine/Source/Editor/PropertyEditor/Private/SDetailExpanderArrow.cpp
+++ b/Engine/Source/Editor/PropertyEditor/Private/SDetailExpanderArrow.cpp
@@ -27,7 +27,7 @@ void SDetailExpanderArrow::Construct(const FArguments& InArgs, TSharedRef<SDetai
 			[
 				SNew(SImage)
 				.Image(this, &SDetailExpanderArrow::GetExpanderImage)
-				.ColorAndOpacity(FSlateColor::UseForeground())
+				.ColorAndOpacity(FSlateColor::UseSubduedForeground())
 			]
 		]
 	];
diff --git a/Engine/Source/Runtime/Slate/Private/Widgets/Views/SExpanderArrow.cpp b/Engine/Source/Runtime/Slate/Private/Widgets/Views/SExpanderArrow.cpp
--- a/Engine/Source/Runtime/Slate/Private/Widgets/Views/SExpanderArrow.cpp
+++ b/Engine/Source/Runtime/Slate/Private/Widgets/Views/SExpanderArrow.cpp
@@ -31,7 +31,7 @@ void SExpanderArrow::Construct( const FArguments& InArgs, const TSharedPtr<class
 		[
 			SNew(SImage)
 			.Image( this, &SExpanderArrow::GetExpanderImage )
-			.ColorAndOpacity( FSlateColor::UseForeground() )
+			.ColorAndOpacity( FSlateColor::UseSubduedForeground() )
 		]
 	];
 }
diff --git a/Engine/Source/Runtime/SlateCore/Private/Styling/StarshipCoreStyle.cpp b/Engine/Source/Runtime/SlateCore/Private/Styling/StarshipCoreStyle.cpp
--- a/Engine/Source/Runtime/SlateCore/Private/Styling/StarshipCoreStyle.cpp
+++ b/Engine/Source/Runtime/SlateCore/Private/Styling/StarshipCoreStyle.cpp
@@ -493,8 +493,8 @@ TSharedRef<ISlateStyle> FStarshipCoreStyle::Create()
 	// SExpandableArea defaults...
 	{
 		Style->Set("ExpandableArea", FExpandableAreaStyle()
-			.SetCollapsedImage(IMAGE_BRUSH_SVG("Starship/Common/chevron-right", Icon16x16, DefaultForeground))
-			.SetExpandedImage(IMAGE_BRUSH_SVG("Starship/Common/chevron-down", Icon16x16, DefaultForeground))
+			.SetCollapsedImage(IMAGE_BRUSH_SVG("Starship/Common/caret-right", Icon16x16, DefaultForeground))
+			.SetExpandedImage(IMAGE_BRUSH_SVG("Starship/Common/caret-down", Icon16x16, DefaultForeground))
 		);
 		Style->Set("ExpandableArea.TitleFont", StyleFonts.SmallBold);
 		Style->Set("ExpandableArea.Border", new FSlateRoundedBoxBrush(FStyleColors::Panel, 4.f) );
@@ -1740,10 +1740,10 @@ void FStarshipCoreStyle::SetupTableViewStyles(TSharedRef<FStyle>& Style)
 		.SetOddRowBackgroundBrush(IMAGE_BRUSH("Common/Selection", Icon8x8, FLinearColor(0.0f, 0.0f, 0.0f, 0.1f)));
 	Style->Set("TableView.DarkRow", DarkTableRowStyle);
 
-	Style->Set("TreeArrow_Collapsed", new IMAGE_BRUSH_SVG("Starship/Common/chevron-right", Icon16x16, FStyleColors::Foreground));
-	Style->Set("TreeArrow_Collapsed_Hovered", new IMAGE_BRUSH_SVG("Starship/Common/chevron-right", Icon16x16, FStyleColors::ForegroundHover));
-	Style->Set("TreeArrow_Expanded", new IMAGE_BRUSH_SVG("Starship/Common/chevron-down", Icon16x16, FStyleColors::Foreground));
-	Style->Set("TreeArrow_Expanded_Hovered", new IMAGE_BRUSH_SVG("Starship/Common/chevron-down", Icon16x16, FStyleColors::ForegroundHover));
+	Style->Set("TreeArrow_Collapsed", new IMAGE_BRUSH_SVG("Starship/Common/caret-right", Icon16x16, FStyleColors::Foreground));
+	Style->Set("TreeArrow_Collapsed_Hovered", new IMAGE_BRUSH_SVG("Starship/Common/caret-right", Icon16x16, FStyleColors::ForegroundHover));
+	Style->Set("TreeArrow_Expanded", new IMAGE_BRUSH_SVG("Starship/Common/caret-down", Icon16x16, FStyleColors::Foreground));
+	Style->Set("TreeArrow_Expanded_Hovered", new IMAGE_BRUSH_SVG("Starship/Common/caret-down", Icon16x16, FStyleColors::ForegroundHover));
 
 
 	const FTableColumnHeaderStyle TableColumnHeaderStyle = FTableColumnHeaderStyle()
