PR #8406: [UE5] Slate line/spline rendering fixes (Contributed by [REDACTED_PERSON])

diff --git a/Engine/Shaders/Private/SlateElementPixelShader.usf b/Engine/Shaders/Private/SlateElementPixelShader.usf
--- a/Engine/Shaders/Private/SlateElementPixelShader.usf
+++ b/Engine/Shaders/Private/SlateElementPixelShader.usf
@@ -239,32 +239,30 @@ half4 GetRoundedBoxElementColor( VertexToPixelInterpolants InVertex )
 }
 #endif
 
-#if SHADER_TYPE == ST_LineSegment 
-/** 
+#if SHADER_TYPE == ST_LineSegment
+/**
  * Generates an anti-aliased line segment pixel
- * The idea for the anti-alising technique is based on the fast prefiltered lines technique published in GPU Gems 2 
+ * This is based on the fast prefiltered lines technique published in GPU Gems 2
+ *
+ * Instead of passing edge function coefficients as uniform parameters we use
+ * texture coordinates to define the line
+ *
+ * When rasterized, the submitted geometry must contain every pixel that
+ * could have coverage > 0. Failing to do so will create sharp, aliased edges
  */
 half4 GetLineSegmentElementColor( VertexToPixelInterpolants InVertex )
 {
-	float LineWidth = ShaderParams.x;
-	float FilterWidthScale = ShaderParams.y;
-	float Gradient = GetUV(InVertex, 0).x;
-	float2 GradientDerivative = float2(abs(ddx(Gradient)), abs(ddy(Gradient)));
-	float PixelSizeInUV = sqrt(dot(GradientDerivative, GradientDerivative));
-	float HalfLineWidthUV = 0.5f * PixelSizeInUV * LineWidth;
-	float HalfFilterWidthUV = FilterWidthScale * PixelSizeInUV;
-	float DistanceToLineCenter = abs(0.5f - Gradient);
-	float LineCoverage = smoothstep(HalfLineWidthUV + HalfFilterWidthUV, HalfLineWidthUV - HalfFilterWidthUV, DistanceToLineCenter);
-
-	if (LineCoverage <= 0.0f)
-	{
-		discard;
-	}
+	const float2 Gradient = GetUV(InVertex, 0);
+
+	// Get coverage based on distance from segment sides and ends
+	const float2 OutsideFilterUV = float2(1.0f, 1.0f);
+	const float2 InsideFilterUV = float2(ShaderParams.x, 0.0f);
+	const float2 LineCoverage = smoothstep(OutsideFilterUV, InsideFilterUV, abs(Gradient));
 
 	half4 Color = InVertex.Color;
-	Color.a *= LineCoverage;
+	Color.a *= LineCoverage.x * LineCoverage.y;
 	return Color;
-} 
+}
 #endif
 
 HALF4_TYPE Main( VertexToPixelInterpolants VIn ) : SV_Target0
diff --git a/Engine/Shaders/StandaloneRenderer/D3D/SlateElementPixelShader.hlsl b/Engine/Shaders/StandaloneRenderer/D3D/SlateElementPixelShader.hlsl
--- a/Engine/Shaders/StandaloneRenderer/D3D/SlateElementPixelShader.hlsl
+++ b/Engine/Shaders/StandaloneRenderer/D3D/SlateElementPixelShader.hlsl
@@ -183,27 +183,16 @@ float4 GetRoundedBoxElementColor( VertexOut InVertex )
 	return OutColor;
 }
 
-float4 GetSplineElementColor( VertexOut InVertex )
+float4 GetLineSegmentElementColor( VertexOut InVertex )
 {
-	const float LineWidth = ShaderParams.x;
-	const float FilterWidthScale = ShaderParams.y;
+	const float2 Gradient = InVertex.TextureCoordinates;
 
-	const float Gradient = InVertex.TextureCoordinates.x;
-	const float2 GradientDerivative = float2( abs(ddx(Gradient)), abs(ddy(Gradient)) );
-	const float PixelSizeInUV = sqrt(dot(GradientDerivative, GradientDerivative));
-	
-	const float HalfLineWidthUV = 0.5f * PixelSizeInUV * LineWidth;	
-	const float HalfFilterWidthUV = FilterWidthScale * PixelSizeInUV;
-	const float DistanceToLineCenter = abs(0.5f - Gradient);
-	const float LineCoverage = smoothstep(HalfLineWidthUV + HalfFilterWidthUV, HalfLineWidthUV-HalfFilterWidthUV, DistanceToLineCenter);
-
-	if (LineCoverage <= 0.0f)
-	{
-		discard;
-	}
+	const float2 OutsideFilterUV = float2(1.0f, 1.0f);
+	const float2 InsideFilterUV = float2(ShaderParams.x, 0.0f);
+	const float2 LineCoverage = smoothstep(OutsideFilterUV, InsideFilterUV, abs(Gradient));
 
 	float4 Color = InVertex.Color;
-	Color.a *= LineCoverage;
+	Color.a *= LineCoverage.x * LineCoverage.y;
 	return Color;
 }
 
@@ -233,7 +222,7 @@ float4 Main( VertexOut InVertex ) : SV_Target
 	}
 	else
 	{
-		OutColor = GetSplineElementColor( InVertex );
+		OutColor = GetLineSegmentElementColor( InVertex );
 	}
 
 	// gamma correct
diff --git a/Engine/Shaders/StandaloneRenderer/OpenGL/SlateElementPixelShader.glsl b/Engine/Shaders/StandaloneRenderer/OpenGL/SlateElementPixelShader.glsl
--- a/Engine/Shaders/StandaloneRenderer/OpenGL/SlateElementPixelShader.glsl
+++ b/Engine/Shaders/StandaloneRenderer/OpenGL/SlateElementPixelShader.glsl
@@ -286,28 +286,17 @@ vec4 GetRoundedBoxElementColor()
 	return OutColor;
 }
 
-vec4 GetSplineElementColor()
+vec4 GetLineSegmentElementColor()
 {
-	float LineWidth = ShaderParams.x;
-	float FilterWidthScale = ShaderParams.y;
+	vec2 Gradient = TexCoords.xy;
 
-	float Gradient = TexCoords.x;
-	vec2 GradientDerivative = vec2(abs(dFdx(Gradient)), abs(dFdy(Gradient)));
-	float PixelSizeInUV = sqrt(dot(GradientDerivative, GradientDerivative));
-	
-	float HalfLineWidthUV = 0.5f*PixelSizeInUV * LineWidth;	
-	float HalfFilterWidthUV = FilterWidthScale * PixelSizeInUV;
-	float DistanceToLineCenter = abs(0.5 - Gradient);
-	float LineCoverage = smoothstep(HalfLineWidthUV+HalfFilterWidthUV, HalfLineWidthUV-HalfFilterWidthUV, DistanceToLineCenter);
-
-	if (LineCoverage <= 0.0)
-	{
-		discard;
-	}
+	vec2 OutsideFilterUV = vec2(1.0, 1.0);
+	vec2 InsideFilterUV = vec2(ShaderParams.x, 0.0);
+	vec2 LineCoverage = smoothstep(OutsideFilterUV, InsideFilterUV, abs(Gradient));
 
-	vec4 InColor = Color;
-	InColor.a *= LineCoverage;
-	return InColor;
+	vec4 OutColor = Color;
+	OutColor.a *= LineCoverage.x * LineCoverage.y;
+	return OutColor;
 }
 
 void main()
@@ -336,7 +325,7 @@ void main()
 	}
 	else
 	{
-		OutColor = GetSplineElementColor();
+		OutColor = GetLineSegmentElementColor();
 	}
 	
 	// gamma correct
diff --git a/Engine/Source/Editor/CurveEditor/Private/SCurveEditorView.cpp b/Engine/Source/Editor/CurveEditor/Private/SCurveEditorView.cpp
--- a/Engine/Source/Editor/CurveEditor/Private/SCurveEditorView.cpp
+++ b/Engine/Source/Editor/CurveEditor/Private/SCurveEditorView.cpp
@@ -190,23 +190,6 @@ void SCurveEditorView::GetCurveDrawParams(TArray<FCurveDrawParams>& OutDrawParam
 			);
 		}
 
-		// In order to work around a bug in slate's line rendering, detect sharp corners and inject an additional 1 pixel segment
-		for (int Index = 1; Index < Params.InterpolatingPoints.Num() - 1; ++Index)
-		{
-			FVector2D Line2 = Params.InterpolatingPoints[Index + 1] - Params.InterpolatingPoints[Index];
-			if (Line2.SizeSquared() > 1.0f)
-			{
-				Line2.Normalize();
-				FVector2D Line1 = Params.InterpolatingPoints[Index] - Params.InterpolatingPoints[Index - 1];
-				Line1.Normalize();
-				if (FVector2D::DotProduct(Line1, Line2) <= 0.5f)
-				{
-					Params.InterpolatingPoints.Insert(Params.InterpolatingPoints[Index] + Line2, Index + 1);
-					++Index;
-				}
-			}
-		}
-
 		TArray<FKeyHandle> VisibleKeys;
 		CurveModel->GetKeys(*CurveEditor, InputMin, InputMax, TNumericLimits<double>::Lowest(), TNumericLimits<double>::Max(), VisibleKeys);
 
diff --git a/Engine/Source/Runtime/SlateCore/Private/Rendering/DrawElements.cpp b/Engine/Source/Runtime/SlateCore/Private/Rendering/DrawElements.cpp
--- a/Engine/Source/Runtime/SlateCore/Private/Rendering/DrawElements.cpp
+++ b/Engine/Source/Runtime/SlateCore/Private/Rendering/DrawElements.cpp
@@ -580,7 +580,15 @@ void FSlateDrawElement::MakeLines( FSlateWindowElementList& ElementList, uint32
 	DataPayload.SetThickness(Thickness);
 	DataPayload.SetLines(Points, bAntialias, &PointColors);
 
-	Element.Init(ElementList, EElementType::ET_Line, InLayer, PaintGeometry, InDrawEffects);
+	ESlateDrawEffect DrawEffects = InDrawEffects;
+	if (bAntialias)
+	{
+		// If the line is to be anti-aliased, we cannot reliably snap
+		// the generated vertices.
+		DrawEffects |= ESlateDrawEffect::NoPixelSnapping;
+	}
+
+	Element.Init(ElementList, EElementType::ET_Line, InLayer, PaintGeometry, DrawEffects);
 
 }
 
diff --git a/Engine/Source/Runtime/SlateCore/Private/Rendering/ElementBatcher.cpp b/Engine/Source/Runtime/SlateCore/Private/Rendering/ElementBatcher.cpp
--- a/Engine/Source/Runtime/SlateCore/Private/Rendering/ElementBatcher.cpp
+++ b/Engine/Source/Runtime/SlateCore/Private/Rendering/ElementBatcher.cpp
@@ -1523,99 +1523,11 @@ void FSlateElementBatcher::AddGradientElement( const FSlateDrawElement& DrawElem
 	}
 }
 
-/**
-* Calculates the intersection of two line segments P1->P2, P3->P4
-* The tolerance setting is used when the lines aren't currently intersecting but will intersect in the future
-* The higher the tolerance the greater the distance that the intersection point can be.
-*
-* @return true if the line intersects.  Populates Intersection
-*/
-static bool LineIntersect(const FVector2D& P1, const FVector2D& P2, const FVector2D& P3, const FVector2D& P4, FVector2D& Intersect, float Tolerance = .1f)
-{
-	float NumA = ((P4.X - P3.X)*(P1.Y - P3.Y) - (P4.Y - P3.Y)*(P1.X - P3.X));
-	float NumB = ((P2.X - P1.X)*(P1.Y - P3.Y) - (P2.Y - P1.Y)*(P1.X - P3.X));
-
-	float Denom = (P4.Y - P3.Y)*(P2.X - P1.X) - (P4.X - P3.X)*(P2.Y - P1.Y);
-
-	if (FMath::IsNearlyZero(NumA) && FMath::IsNearlyZero(NumB))
-	{
-		// Lines are the same
-		Intersect = (P1 + P2) / 2;
-		return true;
-	}
-
-	if (FMath::IsNearlyZero(Denom))
-	{
-		// Lines are parallel
-		return false;
-	}
-
-	float B = NumB / Denom;
-	float A = NumA / Denom;
-
-	// Note that this is a "tweaked" intersection test for the purpose of joining line segments.  We don't just want to know if the line segments
-	// Intersect, but where they would if they don't currently. Except that we don't care in the case that where the segments 
-	// intersection is so far away that its infeasible to use the intersection point later.
-	if (A >= -Tolerance && A <= (1.0f + Tolerance) && B >= -Tolerance && B <= (1.0f + Tolerance))
-	{
-		Intersect = P1 + A*(P2 - P1);
-		return true;
-	}
-
-	return false;
-}
-
 
-/** Utility class for building a strip of lines. */
-struct FLineBuilder
+/** Utility class for building a strip of triangles for a spline. */
+struct FSplineBuilder
 {
-	// Will append 5 vertexes to OutBatchVertices and 9 indexes to
-	// OutBatchIndices. Creates the following cap geometry:
-	//
-	// Cap Vertex Indexes              Cap Measurements
-	//
-	//     U == 0
-	//   2-4----                        2-------4-------....
-	//   |\| 							|       |     ^ 
-	//   | 0  <-- U==0.5				|<- d ->o    2h  
-	//   |/|							|       |     v
-	//   1-3----						1-------3-------....
-	//     U == 0
-	//                                 d is CapDirection
-	//                                 h is Up
-	//                                 o is CapOrigin
-
-	static void MakeCap(
-		FSlateRenderBatch& RenderBatch,
-		const FSlateRenderTransform& RenderTransform,
-		const FVector2D& CapOrigin,
-		const FVector2D& CapDirection,
-		const FVector2D& Up,
-		const FColor& Color
-	)
-	{
-		const uint32 FirstVertIndex = RenderBatch.GetNumVertices();
-
-		RenderBatch.AddVertex(FSlateVertex::Make<ESlateVertexRounding::Disabled>(RenderTransform, CapOrigin, FVector2D(0.5f, 0.0f), FVector2D::ZeroVector, Color));
-		RenderBatch.AddVertex(FSlateVertex::Make<ESlateVertexRounding::Disabled>(RenderTransform, CapOrigin + CapDirection + Up, FVector2D(0.0f, 0.0f), FVector2D::ZeroVector, Color));
-		RenderBatch.AddVertex(FSlateVertex::Make<ESlateVertexRounding::Disabled>(RenderTransform, CapOrigin + CapDirection - Up, FVector2D(0.0f, 0.0f), FVector2D::ZeroVector, Color));
-		RenderBatch.AddVertex(FSlateVertex::Make<ESlateVertexRounding::Disabled>(RenderTransform, CapOrigin + Up, FVector2D(0.0f, 0.0f), FVector2D::ZeroVector, Color));
-		RenderBatch.AddVertex(FSlateVertex::Make<ESlateVertexRounding::Disabled>(RenderTransform, CapOrigin - Up, FVector2D(0.0f, 0.0f), FVector2D::ZeroVector, Color));
-
-		RenderBatch.AddIndex(FirstVertIndex + 0);
-		RenderBatch.AddIndex(FirstVertIndex + 3);
-		RenderBatch.AddIndex(FirstVertIndex + 1);
-
-		RenderBatch.AddIndex(FirstVertIndex + 0);
-		RenderBatch.AddIndex(FirstVertIndex + 1);
-		RenderBatch.AddIndex(FirstVertIndex + 2);
-
-		RenderBatch.AddIndex(FirstVertIndex + 0);
-		RenderBatch.AddIndex(FirstVertIndex + 2);
-		RenderBatch.AddIndex(FirstVertIndex + 4);
-	}
-
-	FLineBuilder(FSlateRenderBatch& InRenderBatch, const FVector2D StartPoint, float HalfThickness, const FSlateRenderTransform& InRenderTransform, const FColor& InColor)
+	FSplineBuilder(FSlateRenderBatch& InRenderBatch, const FVector2D StartPoint, float HalfThickness, const FSlateRenderTransform& InRenderTransform, const FColor& InColor)
 		: RenderBatch(InRenderBatch)
 		, RenderTransform(InRenderTransform)
 		, LastPointAdded()
@@ -1671,7 +1583,7 @@ private:
 			const FVector2D LastUp = LastNormal*HalfLineThickness;
 
 			RenderBatch.AddVertex(FSlateVertex::Make<ESlateVertexRounding::Disabled>(RenderTransform, LastPointAdded[1] + LastUp, FVector2D(1.0f, 0.0f), FVector2D::ZeroVector, InColor));
-			RenderBatch.AddVertex(FSlateVertex::Make<ESlateVertexRounding::Disabled>(RenderTransform, LastPointAdded[1] - LastUp, FVector2D(0.0f, 0.0f), FVector2D::ZeroVector, InColor));
+			RenderBatch.AddVertex(FSlateVertex::Make<ESlateVertexRounding::Disabled>(RenderTransform, LastPointAdded[1] - LastUp, FVector2D(-1.0f, 0.0f), FVector2D::ZeroVector, InColor));
 		}
 
 		if (NumPointsAdded >= 2)
@@ -1679,7 +1591,7 @@ private:
 			const FVector2D AveragedUp = (0.5f*(NewNormal + LastNormal)).GetSafeNormal()*HalfLineThickness;
 
 			RenderBatch.AddVertex(FSlateVertex::Make<ESlateVertexRounding::Disabled>(RenderTransform, LastPointAdded[0] + AveragedUp, FVector2D(1.0f, 0.0f), FVector2D::ZeroVector, InColor));
-			RenderBatch.AddVertex(FSlateVertex::Make<ESlateVertexRounding::Disabled>(RenderTransform, LastPointAdded[0] - AveragedUp, FVector2D(0.0f, 0.0f), FVector2D::ZeroVector, InColor));
+			RenderBatch.AddVertex(FSlateVertex::Make<ESlateVertexRounding::Disabled>(RenderTransform, LastPointAdded[0] - AveragedUp, FVector2D(-1.0f, 0.0f), FVector2D::ZeroVector, InColor));
 
 			const int32 NumVerts = RenderBatch.GetNumVertices();
 
@@ -1717,7 +1629,7 @@ private:
 			const FVector2D LastUp = LastNormal*HalfLineThickness;
 
 			RenderBatch.AddVertex(FSlateVertex::Make<ESlateVertexRounding::Disabled>(RenderTransform, LastPointAdded[0] + LastUp, FVector2D(1.0f, 0.0f), FVector2D::ZeroVector, InColor));
-			RenderBatch.AddVertex(FSlateVertex::Make<ESlateVertexRounding::Disabled>(RenderTransform, LastPointAdded[0] - LastUp, FVector2D(0.0f, 0.0f), FVector2D::ZeroVector, InColor));
+			RenderBatch.AddVertex(FSlateVertex::Make<ESlateVertexRounding::Disabled>(RenderTransform, LastPointAdded[0] - LastUp, FVector2D(-1.0f, 0.0f), FVector2D::ZeroVector, InColor));
 
 			const int32 NumVerts = RenderBatch.GetNumVertices();
 
@@ -1821,7 +1733,7 @@ private:
 		OutCurveParams[6] = P3;
 	}
 
-	static void Subdivide(const FVector2D&  P0, const FVector2D&  P1, const FVector2D&  P2, const FVector2D&  P3, FLineBuilder& LineBuilder, float MaxBiasTimesTwo = 2.0f)
+	static void Subdivide(const FVector2D&  P0, const FVector2D&  P1, const FVector2D&  P2, const FVector2D&  P3, FSplineBuilder& SplineBuilder, float MaxBiasTimesTwo = 2.0f)
 	{
 		const float Curviness = ComputeCurviness(P0, P1, P2, P3);
 		if (Curviness > MaxBiasTimesTwo)
@@ -1830,16 +1742,16 @@ private:
 			FVector2D TwoCurves[7];
 			deCasteljauSplit(P0, P1, P2, P3, TwoCurves);
 			// Subdivide left, then right
-			Subdivide(TwoCurves[0], TwoCurves[1], TwoCurves[2], TwoCurves[3], LineBuilder, MaxBiasTimesTwo);
-			Subdivide(TwoCurves[3], TwoCurves[4], TwoCurves[5], TwoCurves[6], LineBuilder, MaxBiasTimesTwo);
+			Subdivide(TwoCurves[0], TwoCurves[1], TwoCurves[2], TwoCurves[3], SplineBuilder, MaxBiasTimesTwo);
+			Subdivide(TwoCurves[3], TwoCurves[4], TwoCurves[5], TwoCurves[6], SplineBuilder, MaxBiasTimesTwo);
 		}
 		else
 		{
-			LineBuilder.AppendPoint(P3, LineBuilder.SingleColor);
+			SplineBuilder.AppendPoint(P3, SplineBuilder.SingleColor);
 		}
 	}
 
-	static void Subdivide_WithColorGradient(const FLinearColor& StartColor, const FLinearColor& EndColor, const FSlateElementBatcher& InBatcher, const FVector2D&  P0, const FVector2D&  P1, const FVector2D&  P2, const FVector2D&  P3, FLineBuilder& LineBuilder, float MaxBiasTimesTwo = 2.0f)
+	static void Subdivide_WithColorGradient(const FLinearColor& StartColor, const FLinearColor& EndColor, const FSlateElementBatcher& InBatcher, const FVector2D&  P0, const FVector2D&  P1, const FVector2D&  P2, const FVector2D&  P3, FSplineBuilder& SplineBuilder, float MaxBiasTimesTwo = 2.0f)
 	{
 		const float Curviness = ComputeCurviness(P0, P1, P2, P3);
 		if (Curviness > MaxBiasTimesTwo)
@@ -1849,12 +1761,12 @@ private:
 			deCasteljauSplit(P0, P1, P2, P3, TwoCurves);
 			const FLinearColor MidpointColor = FLinearColor::LerpUsingHSV(StartColor, EndColor, 0.5f);
 			// Subdivide left, then right
-			Subdivide_WithColorGradient(StartColor, MidpointColor, InBatcher, TwoCurves[0], TwoCurves[1], TwoCurves[2], TwoCurves[3], LineBuilder, MaxBiasTimesTwo);
-			Subdivide_WithColorGradient(MidpointColor, EndColor, InBatcher, TwoCurves[3], TwoCurves[4], TwoCurves[5], TwoCurves[6], LineBuilder, MaxBiasTimesTwo);
+			Subdivide_WithColorGradient(StartColor, MidpointColor, InBatcher, TwoCurves[0], TwoCurves[1], TwoCurves[2], TwoCurves[3], SplineBuilder, MaxBiasTimesTwo);
+			Subdivide_WithColorGradient(MidpointColor, EndColor, InBatcher, TwoCurves[3], TwoCurves[4], TwoCurves[5], TwoCurves[6], SplineBuilder, MaxBiasTimesTwo);
 		}
 		else
 		{
-			LineBuilder.AppendPoint(P3, InBatcher.PackVertexColor(EndColor));
+			SplineBuilder.AppendPoint(P3, InBatcher.PackVertexColor(EndColor));
 		}
 	}
 	
@@ -1882,35 +1794,31 @@ void FSlateElementBatcher::AddSplineElement(const FSlateDrawElement& DrawElement
 
 	const FSlateRenderTransform& RenderTransform = DrawElement.GetRenderTransform();
 	const FSlateSplinePayload& InPayload = DrawElement.GetDataPayload<FSlateSplinePayload>();
-	ESlateDrawEffect InDrawEffects = DrawElement.GetDrawEffects();
+	const ESlateDrawEffect InDrawEffects = DrawElement.GetDrawEffects();
 	const int32 Layer = DrawElement.GetLayer();
 
-	// 1 is the minimum thickness we support for generating geometry.
-	// The shader takes care of sub-pixel line widths.
-	// Thickness is given in screenspace, so convert it to local space before proceeding.
-	float InThickness = FMath::Max(1.0f, DrawElement.GetInverseLayoutTransform().GetScale() * InPayload.GetThickness());
-
-	// Width of the filter size to use for anti-aliasing.
+	// Filter size to use for anti-aliasing.
 	// Increasing this value will increase the fuzziness of line edges.
-	const float FilterScale = 1.0f;
+	const float FilterRadius = 1.0f;
 
-	static const float TwoRootTwo = 2 * UE_SQRT_2;
-	// Compute the actual size of the line we need based on thickness.
-	// Each line segment will be a bit thicker than the line to account
-	// for the size of the filter.
-	const float LineThickness = (TwoRootTwo + InThickness);
+	const float HalfThickness = 0.5f * FMath::Max(0.0f, InPayload.GetThickness());
 
 	// The amount we increase each side of the line to generate enough pixels
-	const float HalfThickness = LineThickness * .5f + FilterScale;
+	const float FilteredHalfThickness = HalfThickness + FilterRadius;
 
-	const FColor SplineColor = (InPayload.GradientStops.Num()==1) ? PackVertexColor(InPayload.GradientStops[0].Color) : PackVertexColor(InPayload.GetTint());
+	const float InsideFilterU = (HalfThickness - FilterRadius) / FilteredHalfThickness;
+	const FVector4 ShaderParams = FVector4(InsideFilterU, 0.0f, 0.0f, 0.0f);
+	FSlateRenderBatch& RenderBatch = CreateRenderBatch(Layer, FShaderParams::MakePixelShaderParams(ShaderParams), nullptr, ESlateDrawPrimitive::TriangleList, ESlateShader::LineSegment, InDrawEffects, ESlateBatchDrawFlag::None, DrawElement);
 
-	FSlateRenderBatch& RenderBatch = CreateRenderBatch(Layer, FShaderParams::MakePixelShaderParams(FVector4(InPayload.GetThickness(), FilterScale, 0, 0)), nullptr, ESlateDrawPrimitive::TriangleList, ESlateShader::LineSegment, InDrawEffects, ESlateBatchDrawFlag::None, DrawElement);
+	// Thickness is given in screenspace, convert it to local space for geometry build
+	const float LocalHalfThickness = FilteredHalfThickness / DrawElement.GetScale();
+
+	const FColor SplineColor = (InPayload.GradientStops.Num()==1) ? PackVertexColor(InPayload.GradientStops[0].Color) : PackVertexColor(InPayload.GetTint());
 
-	FLineBuilder LineBuilder(
+	FSplineBuilder SplineBuilder(
 		RenderBatch,
 		InPayload.P0,
-		HalfThickness,
+		LocalHalfThickness,
 		RenderTransform,
 		SplineColor
 	);
@@ -1918,227 +1826,328 @@ void FSlateElementBatcher::AddSplineElement(const FSlateDrawElement& DrawElement
 	if (/*const bool bNoGradient = */InPayload.GradientStops.Num() <= 1)
 	{
 		// Normal scenario where there is no color gradient.
-		LineBuilder.BuildBezierGeometry(InPayload.P0, InPayload.P1, InPayload.P2, InPayload.P3);
+		SplineBuilder.BuildBezierGeometry(InPayload.P0, InPayload.P1, InPayload.P2, InPayload.P3);
 	}
 	else
 	{
 		// Deprecated scenario _WithColorGradient
-		LineBuilder.BuildBezierGeometry_WithColorGradient( InPayload.GradientStops, 1, InPayload.P0, InPayload.P1, InPayload.P2, InPayload.P3, *this);
+		SplineBuilder.BuildBezierGeometry_WithColorGradient( InPayload.GradientStops, 1, InPayload.P0, InPayload.P1, InPayload.P2, InPayload.P3, *this);
 	}
 }
 
-
-template<ESlateVertexRounding Rounding>
-void FSlateElementBatcher::AddLineElement( const FSlateDrawElement& DrawElement )
+/** Equivalent to V.GetRotated(90.0f) */
+static FVector2D GetRotated90(const FVector2D& V)
 {
-	const FSlateLinePayload& DrawElementPayload = DrawElement.GetDataPayload<FSlateLinePayload>();
-	const FSlateRenderTransform& RenderTransform = DrawElement.GetRenderTransform();
-	ESlateDrawEffect DrawEffects = DrawElement.GetDrawEffects();
-	const int32 Layer = DrawElement.GetLayer();
+	return FVector2D(-V.Y, V.X);
+}
 
-	const TArray<FVector2D>& Points = DrawElementPayload.GetPoints();
-	const TArray<FLinearColor>& PointColors = DrawElementPayload.GetPointColors();
 
-	const int32 NumPoints = Points.Num();
-	if (NumPoints < 2)
+/** Utility class that builds triangle strips for antialiased lines */
+struct FLineBuilder
+{
+	/** Constructor
+	*
+	* @param ElementScale		Element layout scale, used to convert screenspace thickness/radius.
+	* @param HalfThickness		Half thickness of the lines in screenspace pixels.
+	* @param FilterRadius		Antialiasing filter radius in screenspace pixels.
+	* @param MiterAngleLimit	Corners sharper than this will be split instead of mitered (0-180 degrees).
+	*/
+	FLineBuilder(FSlateRenderBatch& InRenderBatch, const FSlateRenderTransform& InRenderTransform, float ElementScale, float HalfThickness, float FilterRadius, float MiterAngleLimit) :
+		RenderBatch(InRenderBatch),
+		RenderTransform(InRenderTransform),
+		LocalHalfThickness((HalfThickness + FilterRadius) / ElementScale),
+		LocalFilterRadius(FilterRadius / ElementScale),
+		LocalCapLength((FilterRadius / ElementScale) * 2.0f),
+		AngleCosineLimit(FMath::Cos(FMath::DegreesToRadians((180.0f - MiterAngleLimit) * 0.5f)))
 	{
-		return;
 	}
-	
-	const FColor FinalTint = PackVertexColor(DrawElementPayload.GetTint());
-
-	if(DrawElementPayload.IsAntialiased() )
-	{
-		//
-		//  The true center of the line is represented by o---o---o
-		//
-		//
-		//           Two triangles make up each trapezoidal line segment
-		//                /        |  |   
-		//               v         |  |   
-		//    +-+---------------+  |  | 
-		//    |\|              / \ v  | 
-		//    | o-------------o   \   |  +--------- U==0
-		//    |/|            / \   \  |  | 
-		//    +-+-----------+   \   \ v  v  
-		//                   \   \   +------+-+
-		//     ^              \   \ /       |/| 
-		//     |               \   o--------o | <-- Endcap
-		//     Endcap           \ /         |\|
-		//                       +----------+-+
-		//                               ^
-		//                               |
-		//                               +--------- U==1
-		//
-		// Each trapezoidal section has a Vertex.U==1 on the bottom and Vertex.U==0 on top.
-		// Endcaps have Vertex.U==0.5 in the middle and Vertex.U==0 on the outside.
-		// This enables easy distance calculations to the "true center" of the line for
-		// anti-aliasing calculations performed in the pixels shader.
-
-
-
-
-		// Half of the width of the filter size to use for anti-aliasing.
-		// Increasing this value will increase the fuzziness of line edges.
-		const float FilterScale = 1.0f;
 
-		// Thickness is given in screen space, so convert it to local space before proceeding.
-		float RequestedThickness = DrawElementPayload.GetThickness();
-		
-		static const float TwoRootTwo = 2 * UE_SQRT_2;
-		// Compute the actual size of the line we need based on thickness.
-		// Each line segment will be a bit thicker than the line to account
-		// for the size of the filter.
-		const float LineThickness = (TwoRootTwo + RequestedThickness );
+	/**
+	* Build geometry for each line segment between the passed-in points. Segments
+	* are represented by trapezoids with a pair of edges parallel to the segment.
+	* Each trapezoid connects to those for neighbouring segments if it's possible to
+	* miter the join cleanly. Rectangular caps are added to antialias any open ends.
+	*
+	* All quads are built using two triangles. The UV coordinates are based on
+	* topology only, not modified by vertex position. This is why each antialiased
+	* edge is kept parallel to the opposite side. Not doing so would result in the
+	* UVs shearing apart at the diagonal where the two triangles meet.
+	*/
+	template<ESlateVertexRounding Rounding>
+	void BuildLineGeometry(const TArray<FVector2D>& Points, const TArray<FColor>& PackedColors, const FColor& PackedTint) const
+	{
+		FColor PointColor = PackedColors.Num() ? PackedColors[0] : PackedTint;
+		FVector2D Position = Points[0];
+		FVector2D NextPosition = Points[1];
 
-		// The amount we increase each side of the line to generate enough pixels
-		const float HalfThickness = LineThickness * .5f + FilterScale;
+		FVector2D Direction;
+		float Length;
+		(NextPosition - Position).ToDirectionAndLength(Direction, Length);
+		FVector2D Up = GetRotated90(Direction) * LocalHalfThickness;
 
-		// Find a batch for the element
-		FSlateRenderBatch& RenderBatch = CreateRenderBatch( Layer, FShaderParams::MakePixelShaderParams( FVector4(RequestedThickness, FilterScale,0,0) ), nullptr, ESlateDrawPrimitive::TriangleList, ESlateShader::LineSegment, DrawEffects, ESlateBatchDrawFlag::None, DrawElement);
+		// Build the start cap at the first point
+		MakeStartCap<Rounding>(Position, Direction, Length, Up, PointColor);
 
-		FVector2D StartPos = Points[0];
-		FVector2D EndPos = Points[1];
+		// Build the intermediate points and their incoming segments
+		const int32 LastPointIndex = Points.Num() - 1;
+		for (int32 Point = 1; Point < LastPointIndex; ++Point)
+		{
+			const FVector2D LastDirection = Direction;
+			const FVector2D LastUp = Up;
+			const float LastLength = Length;
+
+			PointColor = PackedColors.Num() ? PackedColors[Point] : PackedTint;
+			Position = NextPosition;
+			NextPosition = Points[Point + 1];
+
+			(NextPosition - Position).ToDirectionAndLength(Direction, Length);
+			Up = GetRotated90(Direction) * LocalHalfThickness;
+
+			// Project "up" onto the miter normal to get perpendicular distance to the miter line
+			const FVector2D MiterNormal = GetMiterNormal(LastDirection, Direction);
+			const float DistanceToMiterLine = FVector2D::DotProduct(Up, MiterNormal);
+
+			// Get the component of Direction perpendicular to the miter line
+			const float DirDotMiterNormal = FVector2D::DotProduct(Direction, MiterNormal);
+
+			// Break the strip at zero-length segments, if the miter angle is
+			// too tight or the miter would cross either segment's bisector
+			const float MinSegmentLength = FMath::Min(LastLength, Length);
+			if (MinSegmentLength > SMALL_NUMBER &&
+				DirDotMiterNormal >= AngleCosineLimit &&
+				(MinSegmentLength * 0.5f * DirDotMiterNormal) >= FMath::Abs(DistanceToMiterLine))
+			{
+				// Calculate the offset to put the vertices on the miter line while
+				// keeping the antialiased edges of each segment's quad parallel
+				const float ParallelDistance = DistanceToMiterLine / DirDotMiterNormal;
+				const FVector2D MiterUp = Up - (Direction * ParallelDistance);
+
+				RenderBatch.AddVertex(FSlateVertex::Make<Rounding>(RenderTransform, Position + MiterUp, FVector2D(1.0f, 0.0f), PointColor));
+				RenderBatch.AddVertex(FSlateVertex::Make<Rounding>(RenderTransform, Position - MiterUp, FVector2D(-1.0f, 0.0f), PointColor));
+				AddQuadIndices(RenderBatch);
+			}
+			else
+			{
+				// Can't miter, so end the current strip and start a new one
+				MakeEndCap<Rounding>(Position, LastDirection, LastLength, LastUp, PointColor);
+				MakeStartCap<Rounding>(Position, Direction, Length, Up, PointColor);
+			}
+		}
 
-		FVector2D Normal = FVector2D( StartPos.Y - EndPos.Y, EndPos.X - StartPos.X ).GetSafeNormal();
-		FVector2D Up = Normal * HalfThickness;
+		// Build the last point's incoming segment and end cap
+		PointColor = PackedColors.Num() ? PackedColors[LastPointIndex] : PackedTint;
+		MakeEndCap<Rounding>(NextPosition, Direction, Length, Up, PointColor);
+	}
 
-		FColor StartColor = PointColors.Num() ? PackVertexColor(PointColors[0] * DrawElementPayload.GetTint()) : FinalTint;
-		FColor EndColor = PointColors.Num() ? PackVertexColor(PointColors[1] * DrawElementPayload.GetTint()) : FinalTint;
-	
-		const FVector2D StartCapDirection = HalfThickness*((StartPos - EndPos).GetSafeNormal());
-		FLineBuilder::MakeCap(RenderBatch, RenderTransform, StartPos, StartCapDirection, Up, StartColor);
-		const uint32 IndexStart = RenderBatch.GetNumVertices();
+	/**
+	* Add a quad using the last four vertices added to the batch
+	*
+	*  Topology:
+	*    1--3
+	*    |\ |
+	*    | \|
+	*    0--2
+	*/
+	static void AddQuadIndices(FSlateRenderBatch& RenderBatch)
+	{
+		const int32 NumVerts = RenderBatch.GetNumVertices();
+		IndexQuad(RenderBatch, NumVerts - 3, NumVerts - 1, NumVerts - 2, NumVerts - 4);
+	}
 
-		// First two points in the line.
-		RenderBatch.AddVertex(FSlateVertex::Make<Rounding>( RenderTransform, StartPos + Up, FVector2D(1.0, 0.0f), FVector2D::ZeroVector, StartColor ) );
-		RenderBatch.AddVertex(FSlateVertex::Make<Rounding>( RenderTransform, StartPos - Up, FVector2D(0.0, 0.0f), FVector2D::ZeroVector, StartColor) );
+private:
+	/**
+	* Create a rectangular cap to antialias the start of a segment
+	*/
+	template<ESlateVertexRounding Rounding>
+	void MakeStartCap(
+		const FVector2D& Position,
+		const FVector2D& Direction,
+		float SegmentLength,
+		const FVector2D& Up,
+		const FColor& Color) const
+	{
+		// Don't build a cap for a zero-length segment
+		// Up would be zero-length anyway, so it would be invisible
+		if (SegmentLength > SMALL_NUMBER)
+		{
+			// Extend the segment to cover the filtered pixels at the start
+			// Center the cap over Position if possible, but never place
+			// vertices on the far side of this segment's midpoint
+			const float InwardDistance = FMath::Min(LocalFilterRadius, SegmentLength * 0.5f);
+			const FVector2D CapInward = Direction * InwardDistance;
+			const FVector2D CapOutward = Direction * (InwardDistance - LocalCapLength);
+
+			RenderBatch.AddVertex(FSlateVertex::Make<Rounding>(RenderTransform, Position + CapOutward + Up, FVector2D(1.0f, -1.0f), Color));
+			RenderBatch.AddVertex(FSlateVertex::Make<Rounding>(RenderTransform, Position + CapOutward - Up, FVector2D(-1.0f, -1.0f), Color));
+			RenderBatch.AddVertex(FSlateVertex::Make<Rounding>(RenderTransform, Position + CapInward + Up, FVector2D(1.0f, 0.0f), Color));
+			RenderBatch.AddVertex(FSlateVertex::Make<Rounding>(RenderTransform, Position + CapInward - Up, FVector2D(-1.0f, 0.0f), Color));
+			AddQuadIndices(RenderBatch);
+		}
+	}
 
-		// Generate the rest of the segments
-		for( int32 Point = 1; Point < NumPoints; ++Point )
+	/**
+	* Create a rectangular cap to antialias the end of a segment
+	*/
+	template<ESlateVertexRounding Rounding>
+	void MakeEndCap(
+		const FVector2D& Position,
+		const FVector2D& Direction,
+		float SegmentLength,
+		const FVector2D& Up,
+		const FColor& Color) const
+	{
+		// Don't build a cap for a zero-length segment
+		// Up would be zero-length anyway, so it would be invisible
+		if (SegmentLength > SMALL_NUMBER)
 		{
-			EndPos = Points[Point];
-			// Determine if we should check the intersection point with the next line segment.
-			// We will adjust were this line ends to the intersection
-			bool bCheckIntersection = (Point + 1) < NumPoints;
+			// Extend the segment to cover the filtered pixels at the end
+			// Center the cap over Position if possible, but never place
+			// vertices on the far side of this segment's midpoint
+			const float InwardDistance = FMath::Min(LocalFilterRadius, SegmentLength * 0.5f);
+			const FVector2D CapInward = Direction * -InwardDistance;
+			const FVector2D CapOutward = Direction * (LocalCapLength - InwardDistance);
+
+			RenderBatch.AddVertex(FSlateVertex::Make<Rounding>(RenderTransform, Position + CapInward + Up, FVector2D(1.0f, 0.0f), Color));
+			RenderBatch.AddVertex(FSlateVertex::Make<Rounding>(RenderTransform, Position + CapInward - Up, FVector2D(-1.0f, 0.0f), Color));
+			AddQuadIndices(RenderBatch);
+
+			RenderBatch.AddVertex(FSlateVertex::Make<Rounding>(RenderTransform, Position + CapOutward + Up, FVector2D(1.0f, 1.0f), Color));
+			RenderBatch.AddVertex(FSlateVertex::Make<Rounding>(RenderTransform, Position + CapOutward - Up, FVector2D(-1.0f, 1.0f), Color));
+			AddQuadIndices(RenderBatch);
+		}
+	}
 
-			// Compute the normal to the line
-			Normal = FVector2D( StartPos.Y - EndPos.Y, EndPos.X - StartPos.X ).GetSafeNormal();
+	/**
+	* Returns the normal to the miter line given the direction of two segments.
+	* The miter line between two segments bisects the angle between them.
+	*
+	* Inputs must be normalized vectors pointing from the start to the end of
+	* two consecutive segments.
+	*/
+	static FVector2D GetMiterNormal(const FVector2D& InboundSegmentDir, const FVector2D& OutboundSegmentDir)
+	{
+		const FVector2D DirSum = InboundSegmentDir + OutboundSegmentDir;
+		const float SizeSquared = DirSum.SizeSquared();
 
-			// Create the new vertices for the thick line segment
-			Up = Normal * HalfThickness;
+		if (SizeSquared > SMALL_NUMBER)
+		{
+			return DirSum * FMath::InvSqrt(SizeSquared);
+		}
 
-			FColor PointColor = PointColors.Num() ? PackVertexColor(PointColors[Point] * DrawElementPayload.GetTint()) : FinalTint;
+		// Inputs are antiparallel, so any perpendicular vector suffices
+		return GetRotated90(InboundSegmentDir);
+	}
 
-			FVector2D IntersectUpper = EndPos + Up;
-			FVector2D IntersectLower = EndPos - Up;
+	FSlateRenderBatch& RenderBatch;
+	const FSlateRenderTransform& RenderTransform;
 
-			if( bCheckIntersection )
-			{
-				// The end point of the next segment
-				const FVector2D NextEndPos = Points[Point+1];
+	const float LocalHalfThickness;
+	const float LocalFilterRadius;
+	const float LocalCapLength;
+	const float AngleCosineLimit;
+};
 
-				// The normal of the next segment
-				const FVector2D NextNormal = FVector2D( EndPos.Y - NextEndPos.Y, NextEndPos.X - EndPos.X ).GetSafeNormal();
 
-				// The next amount to adjust the vertices by 
-				FVector2D NextUp = NextNormal * HalfThickness;
+template<ESlateVertexRounding Rounding>
+void FSlateElementBatcher::AddLineElement( const FSlateDrawElement& DrawElement )
+{
+	const FSlateLinePayload& DrawElementPayload = DrawElement.GetDataPayload<FSlateLinePayload>();
+	const FSlateRenderTransform& RenderTransform = DrawElement.GetRenderTransform();
+	const ESlateDrawEffect DrawEffects = DrawElement.GetDrawEffects();
+	const int32 Layer = DrawElement.GetLayer();
 
-				FVector2D IntersectionPoint;
-				if( LineIntersect( StartPos + Up, EndPos + Up, EndPos + NextUp, NextEndPos + NextUp, IntersectionPoint ) )
-				{
-					// If the lines intersect adjust where the line starts
-					IntersectUpper = IntersectionPoint;
+	const TArray<FVector2D>& Points = DrawElementPayload.GetPoints();
+	const TArray<FLinearColor>& PointColors = DrawElementPayload.GetPointColors();
 
-					// visualizes the intersection
-					//AddQuadElement( IntersectUpper-FVector2D(1,1), FVector2D(2,2), 1, InClippingRect, Layer+1, FColor::Orange);
-				}
+	const int32 NumPoints = Points.Num();
+	if (NumPoints < 2)
+	{
+		return;
+	}
 
-				if( LineIntersect( StartPos - Up, EndPos - Up, EndPos - NextUp, NextEndPos - NextUp, IntersectionPoint ) )
-				{
-					// If the lines intersect adjust where the line starts
-					IntersectLower = IntersectionPoint;
+	// Pack element tint and per-point colors
+	const FColor PackedTint = PackVertexColor(DrawElementPayload.GetTint());
+	TArray<FColor> PackedColors;
+	PackedColors.Reserve(PointColors.Num());
+	for (const FLinearColor& PointColor : PointColors)
+	{
+		PackedColors.Add(PackVertexColor(PointColor * DrawElementPayload.GetTint()));
+	}
 
-					// visualizes the intersection
-					//AddQuadElement( IntersectLower-FVector2D(1,1), FVector2D(2,2), 1, InClippingRect, Layer+1, FColor::Yellow);
-				}
-			}
+	if (DrawElementPayload.IsAntialiased())
+	{
+		// Filter size to use for anti-aliasing.
+		// Increasing this value will increase the fuzziness of line edges.
+		const float FilterRadius = 1.0f;
 
-			RenderBatch.AddVertex(FSlateVertex::Make<Rounding>( RenderTransform, IntersectUpper, FVector2D(1.0, 0.0f), FVector2D::ZeroVector, PointColor ) );
-			RenderBatch.AddVertex(FSlateVertex::Make<Rounding>( RenderTransform, IntersectLower, FVector2D(0.0, 0.0f), FVector2D::ZeroVector, PointColor ) );
-			
-			// Counterclockwise winding on triangles
-			RenderBatch.AddIndex(IndexStart + 2 * Point - 1);
-			RenderBatch.AddIndex(IndexStart + 2 * Point - 2);
-			RenderBatch.AddIndex(IndexStart + 2 * Point + 0);
+		// Corners sharper than this will be split instead of mitered
+		// A split can also be forced by repeating points in the input
+		const float MiterAngleLimit = 90.0f - SMALL_NUMBER;
 
-			RenderBatch.AddIndex(IndexStart + 2 * Point - 1);
-			RenderBatch.AddIndex(IndexStart + 2 * Point + 0);
-			RenderBatch.AddIndex(IndexStart + 2 * Point + 1);
+		const float HalfThickness = 0.5f * FMath::Max(0.0f, DrawElementPayload.GetThickness());
 
-			StartPos = EndPos;
-		}
+		// The amount we increase each side of the line to generate enough pixels
+		const float FilteredHalfThickness = HalfThickness + FilterRadius;
 
-		EndPos = Points[NumPoints - 1];
-		StartPos = Points[NumPoints - 2];
-		const FVector2D EndCapDirection = HalfThickness*((EndPos-StartPos).GetSafeNormal());
-		FLineBuilder::MakeCap(RenderBatch, RenderTransform, EndPos, EndCapDirection, Up, EndColor);
+		const float InsideFilterU = (HalfThickness - FilterRadius) / FilteredHalfThickness;
+		const FVector4 ShaderParams = FVector4(InsideFilterU, 0.0f, 0.0f, 0.0f);
+		FSlateRenderBatch& RenderBatch = CreateRenderBatch(Layer, FShaderParams::MakePixelShaderParams(ShaderParams), nullptr, ESlateDrawPrimitive::TriangleList, ESlateShader::LineSegment, DrawEffects, ESlateBatchDrawFlag::None, DrawElement);
+
+		FLineBuilder LineBuilder(
+			RenderBatch,
+			RenderTransform,
+			DrawElement.GetScale(),
+			HalfThickness,
+			FilterRadius,
+			MiterAngleLimit);
+
+		LineBuilder.BuildLineGeometry<Rounding>(Points, PackedColors, PackedTint);
 	}
 	else
 	{
-		if (DrawElementPayload.GetThickness() == 1)
+		if (DrawElementPayload.GetThickness() == 1.0f)
 		{
-			// Find a batch for the element
-			FSlateRenderBatch& RenderBatch = CreateRenderBatch(Layer, FShaderParams(), nullptr, ESlateDrawPrimitive::LineList, ESlateShader::Default, DrawEffects, ESlateBatchDrawFlag::None, DrawElement);
-
 			// Generate the line segments using the native line rendering of the platform.
-			for (int32 Point = 0; Point < NumPoints - 1; ++Point)
-			{
-				const uint32 IndexStart = RenderBatch.GetNumVertices();
-				FVector2D StartPos = Points[Point];
-				FVector2D EndPos = Points[Point + 1];
+			FSlateRenderBatch& RenderBatch = CreateRenderBatch(Layer, FShaderParams(), nullptr, ESlateDrawPrimitive::LineList, ESlateShader::Default, DrawEffects, ESlateBatchDrawFlag::None, DrawElement);
 
-				FColor StartColor = PointColors.Num() ? PackVertexColor(PointColors[Point]  * DrawElementPayload.GetTint()) : FinalTint;
-				FColor EndColor = PointColors.Num() ? PackVertexColor(PointColors[Point+1] * DrawElementPayload.GetTint()) : FinalTint;
+			const FColor StartColor = PackedColors.Num() ? PackedColors[0] : PackedTint;
+			RenderBatch.AddVertex(FSlateVertex::Make<Rounding>(RenderTransform, Points[0], FVector2D::ZeroVector, StartColor));
 
-				RenderBatch.AddVertex(FSlateVertex::Make<Rounding>(RenderTransform, StartPos, FVector2D::ZeroVector, StartColor));
-				RenderBatch.AddVertex(FSlateVertex::Make<Rounding>(RenderTransform, EndPos, FVector2D::ZeroVector, EndColor));
+			for (int32 Point = 1; Point < NumPoints; ++Point)
+			{
+				const FColor PointColor = PackedColors.Num() ? PackedColors[Point] : PackedTint;
+				RenderBatch.AddVertex(FSlateVertex::Make<Rounding>(RenderTransform, Points[Point], FVector2D::ZeroVector, PointColor));
 
-				RenderBatch.AddIndex(IndexStart);
-				RenderBatch.AddIndex(IndexStart + 1);
+				const uint32 IndexStart = RenderBatch.GetNumVertices();
+				RenderBatch.AddIndex(IndexStart - 2);
+				RenderBatch.AddIndex(IndexStart - 1);
 			}
 		}
 		else
 		{
-			// Find a batch for the element
+			// Generate the line segments using non-aa'ed quads.
 			FSlateRenderBatch& RenderBatch = CreateRenderBatch(Layer, FShaderParams(), nullptr, ESlateDrawPrimitive::TriangleList, ESlateShader::Default, DrawEffects, ESlateBatchDrawFlag::None, DrawElement);
 
+			// Thickness is given in screenspace, convert it to local space for geometry build
+			const float LocalHalfThickness = (DrawElementPayload.GetThickness() * 0.5f) / DrawElement.GetScale();
 
-			// Generate the line segments using non-aa'ed polylines.
 			for (int32 Point = 0; Point < NumPoints - 1; ++Point)
 			{
-				const uint32 IndexStart = RenderBatch.GetNumVertices();
 				const FVector2D StartPos = Points[Point];
 				const FVector2D EndPos = Points[Point + 1];
 
-				FColor StartColor	= PointColors.Num() ? PackVertexColor(PointColors[Point]   * DrawElementPayload.GetTint()) : FinalTint;
-				FColor EndColor		= PointColors.Num() ? PackVertexColor(PointColors[Point+1] * DrawElementPayload.GetTint()) : FinalTint;
-	
-				const FVector2D SegmentNormal = (EndPos - StartPos).GetSafeNormal();
-				const FVector2D HalfThickNormal = SegmentNormal * (DrawElementPayload.GetThickness() * 0.5f);
+				const FColor StartColor = PackedColors.Num() ? PackedColors[Point] : PackedTint;
+				const FColor EndColor = PackedColors.Num() ? PackedColors[Point+1] : PackedTint;
 
-				RenderBatch.AddVertex(FSlateVertex::Make<Rounding>(RenderTransform, StartPos + FVector2D(HalfThickNormal.Y, -HalfThickNormal.X), FVector2D::ZeroVector, FVector2D::ZeroVector, StartColor));
-				RenderBatch.AddVertex(FSlateVertex::Make<Rounding>(RenderTransform, StartPos + FVector2D(-HalfThickNormal.Y, HalfThickNormal.X), FVector2D::ZeroVector, FVector2D::ZeroVector, StartColor));
-				RenderBatch.AddVertex(FSlateVertex::Make<Rounding>(RenderTransform, EndPos + FVector2D(HalfThickNormal.Y, -HalfThickNormal.X), FVector2D::ZeroVector, FVector2D::ZeroVector, EndColor));
-				RenderBatch.AddVertex(FSlateVertex::Make<Rounding>(RenderTransform, EndPos + FVector2D(-HalfThickNormal.Y, HalfThickNormal.X), FVector2D::ZeroVector, FVector2D::ZeroVector, EndColor));
+				const FVector2D SegmentDir = (EndPos - StartPos).GetSafeNormal();
+				const FVector2D Up = GetRotated90(SegmentDir) * LocalHalfThickness;
 
-				RenderBatch.AddIndex(IndexStart + 0);
-				RenderBatch.AddIndex(IndexStart + 1);
-				RenderBatch.AddIndex(IndexStart + 2);
+				RenderBatch.AddVertex(FSlateVertex::Make<Rounding>(RenderTransform, StartPos + Up, FVector2D::ZeroVector, StartColor));
+				RenderBatch.AddVertex(FSlateVertex::Make<Rounding>(RenderTransform, StartPos - Up, FVector2D::ZeroVector, StartColor));
+				RenderBatch.AddVertex(FSlateVertex::Make<Rounding>(RenderTransform, EndPos + Up, FVector2D::ZeroVector, EndColor));
+				RenderBatch.AddVertex(FSlateVertex::Make<Rounding>(RenderTransform, EndPos - Up, FVector2D::ZeroVector, EndColor));
 
-				RenderBatch.AddIndex(IndexStart + 2);
-				RenderBatch.AddIndex(IndexStart + 1);
-				RenderBatch.AddIndex(IndexStart + 3);
+				FLineBuilder::AddQuadIndices(RenderBatch);
 			}
 		}
 	}
diff --git a/Engine/Source/Runtime/SlateCore/Public/Rendering/ElementBatcher.h b/Engine/Source/Runtime/SlateCore/Public/Rendering/ElementBatcher.h
--- a/Engine/Source/Runtime/SlateCore/Public/Rendering/ElementBatcher.h
+++ b/Engine/Source/Runtime/SlateCore/Public/Rendering/ElementBatcher.h
@@ -242,8 +242,6 @@ private:
  */
 class FSlateElementBatcher
 {
-
-	friend struct FLineBuilder;
 public:
 
 	SLATECORE_API FSlateElementBatcher( TSharedRef<FSlateRenderingPolicy> InRenderingPolicy );
