Fix to comment refering to UE4 in Unicode block range definitions

diff --git a/Engine/Source/Runtime/SlateCore/Public/Fonts/UnicodeBlockRange.inl b/Engine/Source/Runtime/SlateCore/Public/Fonts/UnicodeBlockRange.inl
--- a/Engine/Source/Runtime/SlateCore/Public/Fonts/UnicodeBlockRange.inl
+++ b/Engine/Source/Runtime/SlateCore/Public/Fonts/UnicodeBlockRange.inl
@@ -166,7 +166,7 @@ REGISTER_UNICODE_BLOCK_RANGE(0xfe70, 0xfeff, ArabicPresentationFormsB,					NSLOC
 REGISTER_UNICODE_BLOCK_RANGE(0xff00, 0xffef, HalfwidthAndFullwidthForms,				NSLOCTEXT("UnicodeBlock", "HalfwidthAndFullwidthForms", "Halfwidth and Fullwidth Forms"))
 REGISTER_UNICODE_BLOCK_RANGE(0xfff0, 0xffff, Specials,									NSLOCTEXT("UnicodeBlock", "Specials", "Specials"))
 
-// The following are outside the Basic Multilingual Plane, so may not be supported by UE4 on all platforms
+// The following are outside the Basic Multilingual Plane, so may not be supported by UE on all platforms
 REGISTER_UNICODE_BLOCK_RANGE(0x10000, 0x1007f, LinearBSyllabary,						NSLOCTEXT("UnicodeBlock", "LinearBSyllabary", "Linear B Syllabary"))
 REGISTER_UNICODE_BLOCK_RANGE(0x10080, 0x100ff, LinearBIdeograms,						NSLOCTEXT("UnicodeBlock", "LinearBIdeograms", "Linear B Ideograms"))
 REGISTER_UNICODE_BLOCK_RANGE(0x10100, 0x1013f, AegeanNumbers,							NSLOCTEXT("UnicodeBlock", "AegeanNumbers", "Aegean Numbers"))
