Slate: Add a debug function to display the SlateAttribute name in the watch window.

diff --git a/Engine/Source/Runtime/SlateCore/Private/Types/SlateAttributeMetaData.cpp b/Engine/Source/Runtime/SlateCore/Private/Types/SlateAttributeMetaData.cpp
--- a/Engine/Source/Runtime/SlateCore/Private/Types/SlateAttributeMetaData.cpp
+++ b/Engine/Source/Runtime/SlateCore/Private/Types/SlateAttributeMetaData.cpp
@@ -293,6 +293,27 @@ TArray<FName> FSlateAttributeMetaData::GetAttributeNames(const SWidget& OwningWi
 }
 
 
+const TCHAR* DebugSlateAttribute(const SWidget* Widget, int32 Index)
+{
+	if (Widget)
+	{
+		if (const FSlateAttributeMetaData* MetaData = FSlateAttributeMetaData::FindMetaData(*Widget))
+		{
+			if (MetaData->Attributes.IsValidIndex(Index))
+			{
+				FName AttributeName = MetaData->Attributes[Index].GetAttributeName(*Widget);
+
+				// Hardcoded static array. This function is only used inside the debugger so it should be fine to return it.
+				static TCHAR TempName[FName::StringBufferSize];
+				FCString::Strcpy(TempName, *FName::SafeString(AttributeName.GetDisplayIndex(), AttributeName.GetNumber()));
+				return TempName;
+			}
+		}
+	}
+	return nullptr;
+}
+
+
 FName FSlateAttributeMetaData::FGetterItem::GetAttributeName(const SWidget& OwningWidget) const
 {
 	return CachedAttributeDescriptor ? CachedAttributeDescriptor->GetName() : FName();
diff --git a/Engine/Source/Runtime/SlateCore/Public/Types/SlateAttributeMetaData.h b/Engine/Source/Runtime/SlateCore/Public/Types/SlateAttributeMetaData.h
--- a/Engine/Source/Runtime/SlateCore/Public/Types/SlateAttributeMetaData.h
+++ b/Engine/Source/Runtime/SlateCore/Public/Types/SlateAttributeMetaData.h
@@ -75,13 +75,26 @@ public:
 		return IndexOfAttribute(Attribute) != INDEX_NONE;
 	}
 
-	int32 GetRegisteredAttributeCount() const { return Attributes.Num(); }
+	int32 GetRegisteredAttributeCount() const
+	{
+		return Attributes.Num();
+	}
 
-	int32 GetRegisteredAffectVisibilityAttributeCount() const { return AffectVisibilityCounter; }
+	int32 GetRegisteredAffectVisibilityAttributeCount() const
+	{
+		return AffectVisibilityCounter;
+	}
 
 	/** Get the name of all the attributes, if available. */
 	static TArray<FName> GetAttributeNames(const SWidget& OwningWidget);
 
+	/**
+	 * Helper function that can be used inside the debuggers watch window.
+	 * E.g. "DebugSlateAttribute(this, 0)"
+	 * E.g. "DebugSlateAttribute((SWidget*)(0x123456789), 0)"
+	 */ 
+	friend const TCHAR* DebugSlateAttribute(const SWidget* Widget, int32 Index);
+
 private:
 	using ISlateAttributeContainer = SlateAttributePrivate::ISlateAttributeContainer;
 	using ESlateAttributeType = SlateAttributePrivate::ESlateAttributeType;
