Slate: Fix bad merge. Assign GSlateIsOnFastProcessInvalidation on FSlateInvalidationRoot::ProcessInvalidation.

diff --git a/Engine/Source/Runtime/SlateCore/Private/FastUpdate/SlateInvalidationRoot.cpp b/Engine/Source/Runtime/SlateCore/Private/FastUpdate/SlateInvalidationRoot.cpp
--- a/Engine/Source/Runtime/SlateCore/Private/FastUpdate/SlateInvalidationRoot.cpp
+++ b/Engine/Source/Runtime/SlateCore/Private/FastUpdate/SlateInvalidationRoot.cpp
@@ -880,6 +880,8 @@ bool FSlateInvalidationRoot::ProcessInvalidation()
 	SCOPED_NAMED_EVENT(Slate_InvalidationProcessing, FColor::Blue);
 	CSV_SCOPED_TIMING_STAT(Slate, InvalidationProcessing);
 
+	TGuardValue<bool> OnFastPathGuard(GSlateIsOnFastProcessInvalidation, true);
+
 	bool bWidgetsNeedRepaint = false;
 
 	if (!bNeedsSlowPath)
