SlateInsight: Fix bad initialization.

diff --git a/Engine/Source/Runtime/SlateCore/Private/Trace/SlateTrace.cpp b/Engine/Source/Runtime/SlateCore/Private/Trace/SlateTrace.cpp
--- a/Engine/Source/Runtime/SlateCore/Private/Trace/SlateTrace.cpp
+++ b/Engine/Source/Runtime/SlateCore/Private/Trace/SlateTrace.cpp
@@ -163,14 +163,14 @@ FSlateTrace::FScopedWidgetPaintTrace::~FScopedWidgetPaintTrace()
 //-----------------------------------------------------------------------------------//
 
 FSlateTrace::FScopedWidgetUpdateTrace::FScopedWidgetUpdateTrace(const SWidget* InWidget)
+	: Widget(InWidget)
+	, UpdateFlags(EWidgetUpdateFlags::None)
 {
 	if (UE_TRACE_CHANNELEXPR_IS_ENABLED(SlateChannel))
 	{
-		UpdateFlags = EWidgetUpdateFlags::None;
 		if (EnumHasAnyFlags(InWidget->UpdateFlags, EWidgetUpdateFlags::NeedsActiveTimerUpdate | EWidgetUpdateFlags::NeedsTick))
 		{
 			StartCycle = FPlatformTime::Cycles64();
-			Widget = InWidget;
 			UpdateFlags = InWidget->UpdateFlags & (EWidgetUpdateFlags::NeedsTick | EWidgetUpdateFlags::NeedsActiveTimerUpdate);
 		}
 	}
