Increase threshold for miter angles to better support line animations.

diff --git a/Engine/Source/Runtime/SlateCore/Private/Rendering/ElementBatcher.cpp b/Engine/Source/Runtime/SlateCore/Private/Rendering/ElementBatcher.cpp
--- a/Engine/Source/Runtime/SlateCore/Private/Rendering/ElementBatcher.cpp
+++ b/Engine/Source/Runtime/SlateCore/Private/Rendering/ElementBatcher.cpp
@@ -2082,7 +2082,7 @@ void FSlateElementBatcher::AddLineElement( const FSlateDrawElement& DrawElement
 
 		// Corners sharper than this will be split instead of mitered
 		// A split can also be forced by repeating points in the input
-		const float MiterAngleLimit = 90.0f - SMALL_NUMBER;
+		const float MiterAngleLimit = 90.0f - KINDA_SMALL_NUMBER;
 
 		const float HalfThickness = 0.5f * FMath::Max(0.0f, DrawElementPayload.GetThickness());
 
