Fix UI widgets not updating to debug safezones

diff --git a/Engine/Source/Runtime/SlateCore/Private/Application/SlateApplicationBase.cpp b/Engine/Source/Runtime/SlateCore/Private/Application/SlateApplicationBase.cpp
--- a/Engine/Source/Runtime/SlateCore/Private/Application/SlateApplicationBase.cpp
+++ b/Engine/Source/Runtime/SlateCore/Private/Application/SlateApplicationBase.cpp
@@ -161,12 +161,55 @@ bool FSlateApplicationBase::IsSlateAsleep()
 	return bIsSlateAsleep;
 }
 
+void FSlateApplicationBase::UpdateCustomSafeZone(const FMargin& NewSafeZoneRatio, bool bShouldRecacheMetrics)
+{
+	if (bShouldRecacheMetrics)
+	{
+		FDisplayMetrics DisplayMetrics;
+		GetDisplayMetrics(DisplayMetrics);
+	}
+	CustomSafeZoneRatio = NewSafeZoneRatio;
+
+	// Allow for a custom margin of zero when explictly set
+	if (CustomSafeZoneState != ECustomSafeZoneState::Set)
+	{
+		CustomSafeZoneState = NewSafeZoneRatio == FMargin() ? ECustomSafeZoneState::Unset : ECustomSafeZoneState::Debug;
+	}
+}
+
+#if WITH_EDITOR
+void FSlateApplicationBase::SwapSafeZoneTypes()
+{
+	FDisplayMetrics DisplayMetrics;
+	GetDisplayMetrics(DisplayMetrics);
+
+	if (FDisplayMetrics::GetDebugTitleSafeZoneRatio() < 1.0f)
+	{
+		ResetCustomSafeZone();
+		CustomSafeZoneState = ECustomSafeZoneState::Debug;
+		OnDebugSafeZoneChanged.Broadcast(FMargin(), false);
+	}
+}
+#endif // WITH_EDITOR
+
 void FSlateApplicationBase::ResetCustomSafeZone()
 {
 	CustomSafeZoneRatio = FMargin();
 	CustomSafeZoneState = ECustomSafeZoneState::Unset;
 }
 
+bool FSlateApplicationBase::IsCustomSafeZoneSet() const
+{
+	return CustomSafeZoneState == ECustomSafeZoneState::Set 
+		|| CustomSafeZoneState == ECustomSafeZoneState::Debug && CustomSafeZoneRatio != FMargin();
+}
+
+void FSlateApplicationBase::SetCustomSafeZone(const FMargin& InSafeZone)
+{
+	CustomSafeZoneRatio = InSafeZone;
+	CustomSafeZoneState = ECustomSafeZoneState::Set;
+}
+
 void FSlateApplicationBase::ToggleGlobalInvalidation(bool bIsGlobalInvalidationEnabled)
 {
 	if (GSlateEnableGlobalInvalidation != bIsGlobalInvalidationEnabled)
diff --git a/Engine/Source/Runtime/SlateCore/Public/Application/SlateApplicationBase.h b/Engine/Source/Runtime/SlateCore/Public/Application/SlateApplicationBase.h
--- a/Engine/Source/Runtime/SlateCore/Public/Application/SlateApplicationBase.h
+++ b/Engine/Source/Runtime/SlateCore/Public/Application/SlateApplicationBase.h
@@ -591,38 +591,9 @@ protected:
 	/** Given a window, locate a widget under the cursor in it; returns an invalid path if cursor is not over this window. */
 	virtual FWidgetPath LocateWidgetInWindow(FVector2D ScreenspaceMouseCoordinate, const TSharedRef<SWindow>& Window, bool bIgnoreEnabledStatus, int32 UserIndex) const = 0;
 	
-	void UpdateCustomSafeZone(const FMargin& NewSafeZoneRatio, bool bShouldRecacheMetrics) 
-	{
-		if (bShouldRecacheMetrics)
-		{
-			FDisplayMetrics DisplayMetrics;
-			GetDisplayMetrics(DisplayMetrics);
-		}
-		CustomSafeZoneRatio = NewSafeZoneRatio; 
-
-		// Allow for a custom margin of zero, only unset via debug ratio or explicit reset function		
-		if (CustomSafeZoneState == ECustomSafeZoneState::Debug && CustomSafeZoneRatio == FMargin())
-		{
-			CustomSafeZoneState = ECustomSafeZoneState::Unset;
-		}
-		else
-		{
-			CustomSafeZoneState = ECustomSafeZoneState::Set;
-		}
-	}
+	void UpdateCustomSafeZone(const FMargin& NewSafeZoneRatio, bool bShouldRecacheMetrics);
 #if WITH_EDITOR
-	void SwapSafeZoneTypes()
-	{
-		FDisplayMetrics DisplayMetrics;
-		GetDisplayMetrics(DisplayMetrics);
-
-		if (FDisplayMetrics::GetDebugTitleSafeZoneRatio() < 1.0f)
-		{
-			ResetCustomSafeZone();
-			CustomSafeZoneState = ECustomSafeZoneState::Debug;
-			OnDebugSafeZoneChanged.Broadcast(FMargin(), false);
-		}
-	}
+	void SwapSafeZoneTypes();
 #endif
 
 protected:
@@ -666,7 +637,8 @@ public:
 	}
 
 	void ResetCustomSafeZone();
-	bool IsCustomSafeZoneSet() const { return CustomSafeZoneState == ECustomSafeZoneState::Set; }
+	bool IsCustomSafeZoneSet() const;
+	void SetCustomSafeZone(const FMargin& InSafeZone);
 	const FMargin& GetCustomSafeZone() const { return CustomSafeZoneRatio; }
 
 #if WITH_EDITORONLY_DATA
