Add parenthesis CIS.

diff --git a/Engine/Source/Runtime/SlateCore/Private/Application/SlateApplicationBase.cpp b/Engine/Source/Runtime/SlateCore/Private/Application/SlateApplicationBase.cpp
--- a/Engine/Source/Runtime/SlateCore/Private/Application/SlateApplicationBase.cpp
+++ b/Engine/Source/Runtime/SlateCore/Private/Application/SlateApplicationBase.cpp
@@ -201,7 +201,7 @@ void FSlateApplicationBase::ResetCustomSafeZone()
 bool FSlateApplicationBase::IsCustomSafeZoneSet() const
 {
 	return CustomSafeZoneState == ECustomSafeZoneState::Set 
-		|| CustomSafeZoneState == ECustomSafeZoneState::Debug && CustomSafeZoneRatio != FMargin();
+		|| (CustomSafeZoneState == ECustomSafeZoneState::Debug && CustomSafeZoneRatio != FMargin());
 }
 
 void FSlateApplicationBase::SetCustomSafeZone(const FMargin& InSafeZone)
