Add logging glyph crash.

diff --git a/Engine/Source/Runtime/SlateCore/Private/Fonts/FontCache.cpp b/Engine/Source/Runtime/SlateCore/Private/Fonts/FontCache.cpp
--- a/Engine/Source/Runtime/SlateCore/Private/Fonts/FontCache.cpp
+++ b/Engine/Source/Runtime/SlateCore/Private/Fonts/FontCache.cpp
@@ -135,6 +135,11 @@ FShapedGlyphSequence::FShapedGlyphSequence(TArray<FShapedGlyphEntry> InGlyphsToR
 
 		// Track reverse look-up data
 		FSourceIndexToGlyphData* SourceIndexToGlyphData = SourceIndicesToGlyphData.GetGlyphData(CurrentGlyph.SourceIndex);
+		UE_CLOG(!SourceIndexToGlyphData, LogSlate, Warning, TEXT("Glyph Index out of range! Index: %i. Valid: [%i, %i] "),
+			CurrentGlyph.SourceIndex,
+			SourceIndicesToGlyphData.GetSourceTextStartIndex(),
+			SourceIndicesToGlyphData.GetSourceTextEndIndex());
+
 		checkSlow(SourceIndexToGlyphData);
 		if (SourceIndexToGlyphData->IsValid())
 		{
@@ -382,7 +387,7 @@ FShapedGlyphSequencePtr FShapedGlyphSequence::GetSubSequence(const int32 InStart
 
 	if (EnumerateVisualGlyphsInSourceRange(InStartIndex, InEndIndex, GlyphCallback) == EEnumerateGlyphsResult::EnumerationComplete)
 	{
-		return MakeShareable(new FShapedGlyphSequence(MoveTemp(SubGlyphsToRender), TextBaseline, MaxTextHeight, FontMaterial, OutlineSettings, SubSequenceRange));
+		return MakeShared<FShapedGlyphSequence>(MoveTemp(SubGlyphsToRender), TextBaseline, MaxTextHeight, FontMaterial, OutlineSettings, SubSequenceRange);
 	}
 
 	return nullptr;
