Slate: Add a flag to help debug invalid resources on the render thread. It's not foolproof, ie the memory could be reused.

diff --git a/Engine/Source/Runtime/SlateCore/Private/Rendering/DrawElements.cpp b/Engine/Source/Runtime/SlateCore/Private/Rendering/DrawElements.cpp
--- a/Engine/Source/Runtime/SlateCore/Private/Rendering/DrawElements.cpp
+++ b/Engine/Source/Runtime/SlateCore/Private/Rendering/DrawElements.cpp
@@ -596,6 +596,7 @@ void FSlateDrawElement::MakeViewport( FSlateWindowElementList& ElementList, uint
 	FSlateViewportPayload& DataPayload = ElementList.CreatePayload<FSlateViewportPayload>(Element);
 
 	DataPayload.SetViewport(Viewport, InTint);
+	check(DataPayload.RenderTargetResource == nullptr || !DataPayload.RenderTargetResource->Debug_IsDestroyed());
 
 	Element.Init(ElementList, EElementType::ET_Viewport, InLayer, PaintGeometry, InDrawEffects);
 }
diff --git a/Engine/Source/Runtime/SlateCore/Private/Rendering/SlateRenderBatch.cpp b/Engine/Source/Runtime/SlateCore/Private/Rendering/SlateRenderBatch.cpp
--- a/Engine/Source/Runtime/SlateCore/Private/Rendering/SlateRenderBatch.cpp
+++ b/Engine/Source/Runtime/SlateCore/Private/Rendering/SlateRenderBatch.cpp
@@ -2,6 +2,7 @@
 
 #include "Rendering/SlateRenderBatch.h"
 #include "Rendering/ElementBatcher.h"
+#include "Textures/SlateShaderResource.h"
 
 FSlateRenderBatch::FSlateRenderBatch(
 	int32 InLayer,
@@ -40,6 +41,7 @@ FSlateRenderBatch::FSlateRenderBatch(
 	, bIsMergable(true)
 	, bIsMerged(false)
 {
+	check(ShaderResource == nullptr || !ShaderResource->Debug_IsDestroyed());
 }
 
 
diff --git a/Engine/Source/Runtime/SlateCore/Public/Textures/SlateShaderResource.h b/Engine/Source/Runtime/SlateCore/Public/Textures/SlateShaderResource.h
--- a/Engine/Source/Runtime/SlateCore/Public/Textures/SlateShaderResource.h
+++ b/Engine/Source/Runtime/SlateCore/Public/Textures/SlateShaderResource.h
@@ -65,14 +65,24 @@ public:
 
 #if SLATE_CHECK_UOBJECT_RENDER_RESOURCES
 	virtual void CheckForStaleResources() const { }
+	bool Debug_IsDestroyed() const { return DestroyState != 0x21; }
+
+	/** Virtual destructor. */
+	virtual ~FSlateShaderResource() { DestroyState = 0x84; }
 #else
 	FORCEINLINE void CheckForStaleResources() const { }
+	bool Debug_IsDestroyed() const { return false; }
+
+	/** Virtual destructor. */
+	virtual ~FSlateShaderResource() = default;
 #endif
 
-public:
 
-	/** Virtual destructor. */
-	virtual ~FSlateShaderResource() { }
+#if SLATE_CHECK_UOBJECT_RENDER_RESOURCES
+private:
+	// Use an int, instead of a bool, to make sure the value is "exactly" what it's set in the destructor.
+	uint8 DestroyState = 0x21;
+#endif
 };
 
 /**
diff --git a/Engine/Source/Runtime/SlateRHIRenderer/Private/SlateRHIRenderingPolicy.cpp b/Engine/Source/Runtime/SlateRHIRenderer/Private/SlateRHIRenderingPolicy.cpp
--- a/Engine/Source/Runtime/SlateRHIRenderer/Private/SlateRHIRenderingPolicy.cpp
+++ b/Engine/Source/Runtime/SlateRHIRenderer/Private/SlateRHIRenderingPolicy.cpp
@@ -784,7 +784,7 @@ void FSlateRHIRenderingPolicy::DrawElements(
 				false);
 
 			const uint32 PrimitiveCount = RenderBatch.DrawPrimitiveType == ESlateDrawPrimitive::LineList ? RenderBatch.NumIndices / 2 : RenderBatch.NumIndices / 3;
-
+			check(ShaderResource == nullptr || !ShaderResource->Debug_IsDestroyed());
 			ESlateShaderResource::Type ResourceType = ShaderResource ? ShaderResource->GetType() : ESlateShaderResource::Invalid;
 			if (ResourceType != ESlateShaderResource::Material && ShaderType != ESlateShader::PostProcess)
 			{
