Slate: Reduce the number of invalidation for the Hoverred state.

diff --git a/Engine/Source/Runtime/Slate/Private/Widgets/Input/SButton.cpp b/Engine/Source/Runtime/Slate/Private/Widgets/Input/SButton.cpp
--- a/Engine/Source/Runtime/Slate/Private/Widgets/Input/SButton.cpp
+++ b/Engine/Source/Runtime/Slate/Private/Widgets/Input/SButton.cpp
@@ -33,6 +33,7 @@ void SButton::PrivateRegisterAttributes(FSlateAttributeInitializer& AttributeIni
 			}));
 
 	AttributeInitializer.OverrideInvalidationReason("EnabledState", FSlateAttributeDescriptor::FInvalidateWidgetReasonAttribute{EInvalidateWidgetReason::Layout|EInvalidateWidgetReason::Paint});
+	AttributeInitializer.OverrideInvalidationReason("Hovered", FSlateAttributeDescriptor::FInvalidateWidgetReasonAttribute{EInvalidateWidgetReason::Layout|EInvalidateWidgetReason::Paint});
 
 	AttributeInitializer.OverrideOnValueChanged("EnabledState"
 		, FSlateAttributeDescriptor::ECallbackOverrideType::ExecuteAfterPrevious
diff --git a/Engine/Source/Runtime/SlateCore/Private/Widgets/SWidget.cpp b/Engine/Source/Runtime/SlateCore/Private/Widgets/SWidget.cpp
--- a/Engine/Source/Runtime/SlateCore/Private/Widgets/SWidget.cpp
+++ b/Engine/Source/Runtime/SlateCore/Private/Widgets/SWidget.cpp
@@ -173,8 +173,8 @@ void SWidget::PrivateRegisterAttributes(FSlateAttributeInitializer& AttributeIni
 	//The order in which SlateAttribute are declared in the .h dictates of the order.
 	SLATE_ADD_MEMBER_ATTRIBUTE_DEFINITION_WITH_NAME(AttributeInitializer, "Visibility", VisibilityAttribute, EInvalidateWidgetReason::Visibility)
 		.AffectVisibility();
-	SLATE_ADD_MEMBER_ATTRIBUTE_DEFINITION_WITH_NAME(AttributeInitializer, "EnabledState", EnabledStateAttribute, EInvalidateWidgetReason::Paint);
-	SLATE_ADD_MEMBER_ATTRIBUTE_DEFINITION_WITH_NAME(AttributeInitializer, "Hovered", HoveredAttribute, EInvalidateWidgetReason::Paint);
+	SLATE_ADD_MEMBER_ATTRIBUTE_DEFINITION_WITH_NAME(AttributeInitializer, "EnabledState", EnabledStateAttribute, EInvalidateWidgetReason::None);
+	SLATE_ADD_MEMBER_ATTRIBUTE_DEFINITION_WITH_NAME(AttributeInitializer, "Hovered", HoveredAttribute, EInvalidateWidgetReason::None);
 	SLATE_ADD_MEMBER_ATTRIBUTE_DEFINITION_WITH_NAME(AttributeInitializer, "RenderTransform", RenderTransformAttribute, EInvalidateWidgetReason::Layout | EInvalidateWidgetReason::RenderTransform);
 	SLATE_ADD_MEMBER_ATTRIBUTE_DEFINITION_WITH_NAME(AttributeInitializer, "RenderTransformPivot", RenderTransformPivotAttribute, EInvalidateWidgetReason::Layout | EInvalidateWidgetReason::RenderTransform);
 }
