Refactoring FGenericAccessibleMessageHandler to raise accessible events that include user Id and additional arguments. This is to stabilize the accessibility API and allow more information to be passed
around for modules that require more accessibility information such as User/Controller Id.
#rnx

diff --git a/Engine/Source/Runtime/ApplicationCore/Private/IOS/IOSApplication.cpp b/Engine/Source/Runtime/ApplicationCore/Private/IOS/IOSApplication.cpp
--- a/Engine/Source/Runtime/ApplicationCore/Private/IOS/IOSApplication.cpp
+++ b/Engine/Source/Runtime/ApplicationCore/Private/IOS/IOSApplication.cpp
@@ -267,18 +267,20 @@ FAutoConsoleVariableRef IOSAccessibleAnnouncementDealyRef(
 	TEXT("We need to introduce a small delay to avoid iOS system accessibility announcements from stomping on our requested user announcement. Delays <= 0.05f are too short and result in the announcement being dropped. Dellays ~0.075f result in unstable delivery")
 );
 
-void FIOSApplication::OnAccessibleEventRaised(TSharedRef<IAccessibleWidget> Widget, EAccessibleEvent Event, FVariant OldValue, FVariant NewValue)
+
+void FIOSApplication::OnAccessibleEventRaised(const FGenericAccessibleMessageHandler::FAccessibleEventArgs& Args)
 {
 	// This should only be triggered by the accessible message handler which initiates from the Slate thread.
 	check(IsInGameThread());
 
-	const AccessibleWidgetId Id = Widget->GetId();
-	switch (Event)
+	const AccessibleWidgetId Id = Args.Widget->GetId();
+	switch (Args.Event)
 	{
 	case EAccessibleEvent::ParentChanged:
 	{
+		const AccessibleWidgetId NewParentId = Args.NewValue.GetValue<AccessibleWidgetId>();
 		dispatch_async(dispatch_get_main_queue(), ^{
-			[[[FIOSAccessibilityCache AccessibilityElementCache] GetAccessibilityElement:Id] SetParent:NewValue.GetValue<AccessibleWidgetId>()];
+			[[[FIOSAccessibilityCache AccessibilityElementCache] GetAccessibilityElement:Id] SetParent:NewParentId];
 		});
 		// LayoutChanged is to indicate things like "a widget became visible or hidden" while
 		// ScreenChanged is for large-scale UI changes. It can potentially take an NSString to read
@@ -299,7 +301,7 @@ void FIOSApplication::OnAccessibleEventRaised(TSharedRef<IAccessibleWidget> Widg
 		break;
 	case EAccessibleEvent::Notification:
 		{
-			NSString* Announcement = [NSString stringWithFString: NewValue.GetValue<FString>()];
+			NSString* Announcement = [NSString stringWithFString: Args.NewValue.GetValue<FString>()];
 			dispatch_async(dispatch_get_main_queue(), ^{
 				// If we don't sleep for a small period of time, system accessibility
 				// announcements can stomp on this announcement and so it's never made
@@ -313,3 +315,4 @@ void FIOSApplication::OnAccessibleEventRaised(TSharedRef<IAccessibleWidget> Widg
 	}
 }
 #endif
+
diff --git a/Engine/Source/Runtime/ApplicationCore/Private/Mac/MacApplication.cpp b/Engine/Source/Runtime/ApplicationCore/Private/Mac/MacApplication.cpp
--- a/Engine/Source/Runtime/ApplicationCore/Private/Mac/MacApplication.cpp
+++ b/Engine/Source/Runtime/ApplicationCore/Private/Mac/MacApplication.cpp
@@ -2296,16 +2296,16 @@ float GMacAccessibleAnnouncementDelay = 0.1f;
 FAutoConsoleVariableRef MacAccessibleAnnouncementDealyRef(
 	TEXT("mac.AccessibleAnnouncementDelay"),
 	GMacAccessibleAnnouncementDelay,
-	TEXT("We need to introduce a small delay to avoid OSX system accessibility announcements from stomping on our requested user announcement. Delays <= 0.05f are too short and result in the announcement being dropped. Dellays ~0.075f result in unstable delivery")
+	TEXT("We need to introduce a small delay to avoid OSX system accessibility announcements from stomping on our requested user announcement. Delays <= 0.05f are too short and result in the announcement being dropped. Delays ~0.075f result in unstable delivery")
 );
 
-void FMacApplication::OnAccessibleEventRaised(TSharedRef<IAccessibleWidget> Widget, EAccessibleEvent Event, FVariant OldValue, FVariant NewValue)
+void FMacApplication::OnAccessibleEventRaised(const FGenericAccessibleMessageHandler::FAccessibleEventArgs& Args)
 {
 	// This should only be triggered by the accessible message handler which initiates from the Slate thread.
 	check(IsInGameThread());
 	
-	const AccessibleWidgetId Id = Widget->GetId();
-	switch (Event)
+	const AccessibleWidgetId Id = Args.Widget->GetId();
+	switch (Args.Event)
 	{
 		case EAccessibleEvent::FocusChange:
 		{
@@ -2314,11 +2314,12 @@ void FMacApplication::OnAccessibleEventRaised(TSharedRef<IAccessibleWidget> Widg
 		}
 		case EAccessibleEvent::ParentChanged:
 		{
+			FVariant NewValueCopy = Args.OldValue;
 			MainThreadCall(^{
 				FMacAccessibilityElement* Element = [[FMacAccessibilityManager AccessibilityManager] GetAccessibilityElement:Id];
 				if(Element)
 				{
-					const AccessibleWidgetId NewParentId = NewValue.GetValue<AccessibleWidgetId>();
+					const AccessibleWidgetId NewParentId = NewValueCopy.GetValue<AccessibleWidgetId>();
 					Element.ParentId = NewParentId;
 					FMacAccessibilityElement* NewParent = [[FMacAccessibilityManager AccessibilityManager]GetAccessibilityElement:NewParentId];
 					Element.accessibilityParent = NewParent;
@@ -2348,7 +2349,7 @@ void FMacApplication::OnAccessibleEventRaised(TSharedRef<IAccessibleWidget> Widg
 		}
 		case EAccessibleEvent::Notification:
 		{
-			NSString* Announcement = [NSString stringWithFString:NewValue.GetValue<FString>()];
+			NSString* Announcement = [NSString stringWithFString:Args.NewValue.GetValue<FString>()];
 			NSDictionary* AnnouncementInfo = @{NSAccessibilityAnnouncementKey: Announcement, NSAccessibilityPriorityKey: @(NSAccessibilityPriorityHigh)};
 			MainThreadCall(^{
 				// If we don't wait for a small period of time, system announcements can stomp on our announcement
diff --git a/Engine/Source/Runtime/ApplicationCore/Private/Windows/Accessibility/WindowsUIAManager.cpp b/Engine/Source/Runtime/ApplicationCore/Private/Windows/Accessibility/WindowsUIAManager.cpp
--- a/Engine/Source/Runtime/ApplicationCore/Private/Windows/Accessibility/WindowsUIAManager.cpp
+++ b/Engine/Source/Runtime/ApplicationCore/Private/Windows/Accessibility/WindowsUIAManager.cpp
@@ -173,29 +173,29 @@ void FWindowsUIAManager::OnAccessibilityDisabled()
 	}
 }
 
-void FWindowsUIAManager::OnEventRaised(TSharedRef<IAccessibleWidget> Widget, EAccessibleEvent Event, FVariant OldValue, FVariant NewValue)
+void FWindowsUIAManager::OnEventRaised(const FGenericAccessibleMessageHandler::FAccessibleEventArgs& Args)
 {
 	if (UiaClientsAreListening())
 	{
-		FScopedWidgetProvider ScopedProvider(GetWidgetProvider(Widget));
+		FScopedWidgetProvider ScopedProvider(GetWidgetProvider(Args.Widget));
 
-		switch (Event)
+		switch (Args.Event)
 		{
 		case EAccessibleEvent::FocusChange:
 		{
 			// On focus change, emit a generic FocusChanged event as well as a per-Provider PropertyChanged event
 			// todo: handle difference between any focus vs keyboard focus
-			if (Widget->HasFocus())
+			if (Args.Widget->HasFocus())
 			{
 				UiaRaiseAutomationEvent(&ScopedProvider.Provider, UIA_AutomationFocusChangedEventId);
 			}
-			EmitPropertyChangedEvent(&ScopedProvider.Provider, UIA_HasKeyboardFocusPropertyId, OldValue, NewValue);
+			EmitPropertyChangedEvent(&ScopedProvider.Provider, UIA_HasKeyboardFocusPropertyId, Args.OldValue, Args.NewValue);
 			break;
 		}
 		case EAccessibleEvent::Activate:
 			if (ScopedProvider.Provider.SupportsInterface(UIA_TogglePatternId))
 			{
-				EmitPropertyChangedEvent(&ScopedProvider.Provider, UIA_ToggleToggleStatePropertyId, OldValue, NewValue);
+				EmitPropertyChangedEvent(&ScopedProvider.Provider, UIA_ToggleToggleStatePropertyId, Args.OldValue, Args.NewValue);
 
 			}
 			else if (ScopedProvider.Provider.SupportsInterface(UIA_InvokePatternId))
@@ -232,7 +232,7 @@ void FWindowsUIAManager::OnEventRaised(TSharedRef<IAccessibleWidget> Widget, EAc
 			}
 			if (NotificationFunc)
 			{
-				NotificationFunc(&ScopedProvider.Provider, NotificationKindEnum, NotificationProcessingEnum, SysAllocString(*NewValue.GetValue<FString>()), SysAllocString(TEXT("")));
+				NotificationFunc(&ScopedProvider.Provider, NotificationKindEnum, NotificationProcessingEnum, SysAllocString(*Args.NewValue.GetValue<FString>()), SysAllocString(TEXT("")));
 			}
 			break;
 		}
@@ -241,8 +241,8 @@ void FWindowsUIAManager::OnEventRaised(TSharedRef<IAccessibleWidget> Widget, EAc
 		// For now, I'm disabling this until we figure out if it's absolutely necessary.
 		//case EAccessibleEvent::ParentChanged:
 		//{
-		//	const AccessibleWidgetId OldId = OldValue.GetValue<AccessibleWidgetId>();
-		//	const AccessibleWidgetId NewId = NewValue.GetValue<AccessibleWidgetId>();
+		//	const AccessibleWidgetId OldId = Args.OldValue.GetValue<AccessibleWidgetId>();
+		//	const AccessibleWidgetId NewId = Args.NewValue.GetValue<AccessibleWidgetId>();
 		//	if (OldId != IAccessibleWidget::InvalidAccessibleWidgetId)
 		//	{
 		//		GetWidgetProvider(WindowsApplication.GetAccessibleMessageHandler()->GetAccessibleWidgetFromId(OldId).ToSharedRef());
diff --git a/Engine/Source/Runtime/ApplicationCore/Public/GenericPlatform/Accessibility/GenericAccessibleInterfaces.h b/Engine/Source/Runtime/ApplicationCore/Public/GenericPlatform/Accessibility/GenericAccessibleInterfaces.h
--- a/Engine/Source/Runtime/ApplicationCore/Public/GenericPlatform/Accessibility/GenericAccessibleInterfaces.h
+++ b/Engine/Source/Runtime/ApplicationCore/Public/GenericPlatform/Accessibility/GenericAccessibleInterfaces.h
@@ -503,17 +503,38 @@ public:
 class FGenericAccessibleMessageHandler
 {
 public:
-	/** A widget raised an event to pass to the native OS implementation. */
-	DECLARE_DELEGATE_FourParams(
-		FAccessibleEvent,
+	/**
+	 * The arguments for an accessible event that is raised by accessible widgets to be
+	 * passed to an accessibility event handler such as a native OS.
+	 * It is up to the client to create an instance of this struct and fill in the appropriate data members.
+	 * @see RaiseEvent
+	 */
+	struct FAccessibleEventArgs
+	{
+		FAccessibleEventArgs(TSharedRef<IAccessibleWidget> InWidget, EAccessibleEvent InEvent, FVariant InOldValue = FVariant(), FVariant InNewValue = FVariant(), int32 InUserId = 0)
+			: Widget(InWidget)
+			, Event(InEvent)
+			, OldValue(InOldValue)
+			, NewValue(InNewValue)
+			, UserId(InUserId)
+		{}
+		
 		/** The accessible widget that generated the event */
-		TSharedRef<IAccessibleWidget>,
+		TSharedRef<IAccessibleWidget> Widget;
 		/** The type of event generated */
-		EAccessibleEvent,
+		EAccessibleEvent Event;
 		/** If this was a property changed event, the 'before' value */
-		FVariant,
+		FVariant OldValue;
 		/** If this was a property changed event, the 'after' value. This may also be set for other events such as Notification. */
-		FVariant);
+		FVariant NewValue;
+		/** The Id of the user this event is intended for. Think of a hardware device such as a controller or keyboard/mouse. */
+		int32 UserId;
+	};
+	/**
+	 * A delegate accessible event handlers such as platform accessibility APIs can
+	 * listen for accessibility events raised by widgets.  .
+	 */
+	DECLARE_DELEGATE_OneParam(FAccessibleEvent, const FAccessibleEventArgs&);
 
 	FGenericAccessibleMessageHandler() : bApplicationIsAccessible(false), bIsActive(false) {}
 
@@ -573,14 +594,12 @@ public:
 	/**
 	 * Push an event from an accessible widget back to the platform layer.
 	 *
-	 * @param Widget The widget raising the event
-	 * @param Event The type of event being raised
-	 * @param OldValue See EAccessibleEvent documentation for more details.
-	 * @param NewValue See EAccessibleEvent documentation for more details.
+	 * @param Args The arguments to be passed to the accessible event delegate.
+	 * @see FAccessibleEventArgs
 	 */
-	void RaiseEvent(TSharedRef<IAccessibleWidget> Widget, EAccessibleEvent Event, FVariant OldValue = FVariant(), FVariant NewValue = FVariant())
+	void RaiseEvent(const FAccessibleEventArgs& Args)
 	{
-		AccessibleEventDelegate.ExecuteIfBound(Widget, Event, OldValue, NewValue);
+		AccessibleEventDelegate.ExecuteIfBound(Args);
 	}
 
 	/**
diff --git a/Engine/Source/Runtime/ApplicationCore/Public/IOS/IOSApplication.h b/Engine/Source/Runtime/ApplicationCore/Public/IOS/IOSApplication.h
--- a/Engine/Source/Runtime/ApplicationCore/Public/IOS/IOSApplication.h
+++ b/Engine/Source/Runtime/ApplicationCore/Public/IOS/IOSApplication.h
@@ -49,7 +49,7 @@ private:
 
 	FIOSApplication();
 #if WITH_ACCESSIBILITY
-	void OnAccessibleEventRaised(TSharedRef<IAccessibleWidget> Widget, EAccessibleEvent Event, FVariant OldValue, FVariant NewValue);
+	void OnAccessibleEventRaised(const FGenericAccessibleMessageHandler::FAccessibleEventArgs& Args);
 #endif
 
 private:
diff --git a/Engine/Source/Runtime/ApplicationCore/Public/Mac/MacApplication.h b/Engine/Source/Runtime/ApplicationCore/Public/Mac/MacApplication.h
--- a/Engine/Source/Runtime/ApplicationCore/Public/Mac/MacApplication.h
+++ b/Engine/Source/Runtime/ApplicationCore/Public/Mac/MacApplication.h
@@ -322,7 +322,7 @@ private:
 #endif
 
 #if WITH_ACCESSIBILITY
-	void OnAccessibleEventRaised(TSharedRef<IAccessibleWidget> Widget, EAccessibleEvent Event, FVariant OldValue, FVariant NewValue);
+	void OnAccessibleEventRaised(const FGenericAccessibleMessageHandler::FAccessibleEventArgs& Args);
 	#endif
 
 private:
diff --git a/Engine/Source/Runtime/ApplicationCore/Public/Windows/Accessibility/WindowsUIAManager.h b/Engine/Source/Runtime/ApplicationCore/Public/Windows/Accessibility/WindowsUIAManager.h
--- a/Engine/Source/Runtime/ApplicationCore/Public/Windows/Accessibility/WindowsUIAManager.h
+++ b/Engine/Source/Runtime/ApplicationCore/Public/Windows/Accessibility/WindowsUIAManager.h
@@ -86,7 +86,7 @@ public:
 
 private:
 	/** Callback function for processing events raised from the AccessibleMessageHandler */
-	void OnEventRaised(TSharedRef<IAccessibleWidget> Widget, EAccessibleEvent Event, FVariant OldValue, FVariant NewValue);
+	void OnEventRaised(const FGenericAccessibleMessageHandler::FAccessibleEventArgs& Args);
 
 	/** Called when the first window widget provider is requested. */
 	void OnAccessibilityEnabled(); 
diff --git a/Engine/Source/Runtime/Slate/Private/Framework/Application/SlateApplication.cpp b/Engine/Source/Runtime/Slate/Private/Framework/Application/SlateApplication.cpp
--- a/Engine/Source/Runtime/Slate/Private/Framework/Application/SlateApplication.cpp
+++ b/Engine/Source/Runtime/Slate/Private/Framework/Application/SlateApplication.cpp
@@ -2823,7 +2823,7 @@ bool FSlateApplication::SetUserFocus(FSlateUser& User, const FWidgetPath& InFocu
 		// Let previously-focused widget know that it's losing focus
 		OldFocusedWidget->OnFocusLost(FocusEvent);
 #if WITH_ACCESSIBILITY
-		GetAccessibleMessageHandler()->OnWidgetEventRaised(OldFocusedWidget.ToSharedRef(), EAccessibleEvent::FocusChange, true, false);
+		GetAccessibleMessageHandler()->OnWidgetEventRaised(FSlateAccessibleMessageHandler::FSlateWidgetAccessibleEventArgs(OldFocusedWidget.ToSharedRef(), EAccessibleEvent::FocusChange, true, false, User.GetUserIndex()));
 #endif
 	}
 
@@ -2860,7 +2860,7 @@ bool FSlateApplication::SetUserFocus(FSlateUser& User, const FWidgetPath& InFocu
 		GetRelevantNavConfig(User.GetUserIndex())->OnNavigationChangedFocus(OldFocusedWidget, NewFocusedWidget, FocusEvent);
 
 #if WITH_ACCESSIBILITY
-		GetAccessibleMessageHandler()->OnWidgetEventRaised(NewFocusedWidget.ToSharedRef(), EAccessibleEvent::FocusChange, false, true);
+		GetAccessibleMessageHandler()->OnWidgetEventRaised(FSlateAccessibleMessageHandler::FSlateWidgetAccessibleEventArgs(NewFocusedWidget.ToSharedRef(), EAccessibleEvent::FocusChange, false, true, User.GetUserIndex()));
 #endif
 	}
 
diff --git a/Engine/Source/Runtime/Slate/Private/Widgets/Input/SButton.cpp b/Engine/Source/Runtime/Slate/Private/Widgets/Input/SButton.cpp
--- a/Engine/Source/Runtime/Slate/Private/Widgets/Input/SButton.cpp
+++ b/Engine/Source/Runtime/Slate/Private/Widgets/Input/SButton.cpp
@@ -464,7 +464,8 @@ FReply SButton::ExecuteOnClick()
 	{
 		FReply Reply = OnClicked.Execute();
 #if WITH_ACCESSIBILITY
-		FSlateApplicationBase::Get().GetAccessibleMessageHandler()->OnWidgetEventRaised(AsShared(), EAccessibleEvent::Activate);
+		// @TODOAccessibility: This should pass the Id of the user that clicked the button but we don't want to change the regular Slate API just yet
+		FSlateApplicationBase::Get().GetAccessibleMessageHandler()->OnWidgetEventRaised(FSlateAccessibleMessageHandler::FSlateWidgetAccessibleEventArgs(AsShared(), EAccessibleEvent::Activate));
 #endif
 		return Reply;
 	}
diff --git a/Engine/Source/Runtime/Slate/Private/Widgets/Input/SCheckBox.cpp b/Engine/Source/Runtime/Slate/Private/Widgets/Input/SCheckBox.cpp
--- a/Engine/Source/Runtime/Slate/Private/Widgets/Input/SCheckBox.cpp
+++ b/Engine/Source/Runtime/Slate/Private/Widgets/Input/SCheckBox.cpp
@@ -387,7 +387,8 @@ void SCheckBox::ToggleCheckedState()
 	}
 
 #if WITH_ACCESSIBILITY
-	FSlateApplicationBase::Get().GetAccessibleMessageHandler()->OnWidgetEventRaised(AsShared(), EAccessibleEvent::Activate, State == ECheckBoxState::Checked, IsCheckboxChecked.Get() == ECheckBoxState::Checked);
+	// @TODOAccessibility: Technically we should pass the Id of the user that toggled the checkbox, but we don't want to change the Slate API as much as possible
+	FSlateApplicationBase::Get().GetAccessibleMessageHandler()->OnWidgetEventRaised(FSlateAccessibleMessageHandler::FSlateWidgetAccessibleEventArgs(AsShared(), EAccessibleEvent::Activate, State == ECheckBoxState::Checked, IsCheckboxChecked.Get() == ECheckBoxState::Checked));
 #endif
 }
 
diff --git a/Engine/Source/Runtime/Slate/Public/Widgets/Views/SListView.h b/Engine/Source/Runtime/Slate/Public/Widgets/Views/SListView.h
--- a/Engine/Source/Runtime/Slate/Public/Widgets/Views/SListView.h
+++ b/Engine/Source/Runtime/Slate/Public/Widgets/Views/SListView.h
@@ -894,8 +894,10 @@ public:
 			{
 				TSharedRef<SWidget> TableRowWidget = TableRow->AsWidget();
 				// We don't need to worry about raising a focus change event for the 
-				// widget with accessibility focus  as FSlateAccessibleMessageHandler will take care of signalling a focus lost event 
-				FSlateApplicationBase::Get().GetAccessibleMessageHandler()->OnWidgetEventRaised(TableRowWidget, EAccessibleEvent::FocusChange, false, true);
+				// widget with accessibility focus  as FSlateAccessibleMessageHandler will take care of signalling a focus lost event
+				// @TODOAccessibility: Technically we need to pass in the user Id that selected the row so the event can be routed to the correct user.
+				// But we don't want to change the Slate API drastically right now
+				FSlateApplicationBase::Get().GetAccessibleMessageHandler()->OnWidgetEventRaised(FSlateAccessibleMessageHandler::FSlateWidgetAccessibleEventArgs(TableRowWidget, EAccessibleEvent::FocusChange, false, true));
 			}
 		}
 #endif
diff --git a/Engine/Source/Runtime/SlateCore/Private/Widgets/Accessibility/SlateAccessibleMessageHandler.cpp b/Engine/Source/Runtime/SlateCore/Private/Widgets/Accessibility/SlateAccessibleMessageHandler.cpp
--- a/Engine/Source/Runtime/SlateCore/Private/Widgets/Accessibility/SlateAccessibleMessageHandler.cpp
+++ b/Engine/Source/Runtime/SlateCore/Private/Widgets/Accessibility/SlateAccessibleMessageHandler.cpp
@@ -189,56 +189,56 @@ void FSlateAccessibleMessageHandler::OnWidgetRemoved(SWidget* Widget)
 		TSharedPtr<FSlateAccessibleWidget> RemovedWidget = FSlateAccessibleWidgetCache::RemoveWidget(Widget);
 		if (RemovedWidget.IsValid())
 		{
-			RaiseEvent(RemovedWidget.ToSharedRef(), EAccessibleEvent::WidgetRemoved);
+			RaiseEvent(FGenericAccessibleMessageHandler::FAccessibleEventArgs(RemovedWidget.ToSharedRef(), EAccessibleEvent::WidgetRemoved));
 			// If this ensure fails, bDirty = true must be called to ensure the tree is kept up to date.
 			ensureMsgf(!Widget->GetParentWidget().IsValid(), TEXT("A widget was unexpectedly deleted before detaching from its parent."));
 		}
 	}
 }
 
-void FSlateAccessibleMessageHandler::OnWidgetEventRaised(TSharedRef<SWidget> Widget, EAccessibleEvent Event, FVariant OldValue, FVariant NewValue)
+void FSlateAccessibleMessageHandler::OnWidgetEventRaised(const FSlateWidgetAccessibleEventArgs& Args)
 {
 	if (IsActive())
 	{
 		SCOPE_CYCLE_COUNTER(STAT_AccessibilitySlateEventRaised);
 		// todo: not sure what to do for a case like focus changed to not-accessible widget. maybe pass through a nullptr?
-		if (Widget->IsAccessible())
+		if (Args.Widget->IsAccessible())
 		{
-			TSharedRef<FSlateAccessibleWidget> AccessibleWidget = FSlateAccessibleWidgetCache::GetAccessibleWidget(Widget);
+			TSharedRef<FSlateAccessibleWidget> AccessibleWidget = FSlateAccessibleWidgetCache::GetAccessibleWidget(Args.Widget);
 			// Perform FSlateAccessibleMessageHandler preprocessing here  before platform processing
-			// Most of the time, accessibility focus is the same as Slate focus. We sync them here 
+			// Most of the time, accessibility focus is the same as Slate focus. We sync them here
 			// as a focus event is raised from SlateApplication for keyboard/gamepad focusable widgets.
 			// However, we can also raise a focus change event manually on non-Slate focusable widgets e.g STableRow
-			// That's where accessibility focus and Slate focus can diverge 
-			if (Event == EAccessibleEvent::FocusChange)
+			// That's where accessibility focus and Slate focus can diverge
+			if (Args.Event == EAccessibleEvent::FocusChange)
 			{
-				// If the Widget is gaining focus 
-				if (NewValue.GetValue<bool>())
+				// If the Widget is gaining focus
+				if (Args.NewValue.GetValue<bool>())
 				{
-					// For widgets that are non-keyboard focusable but support accessibility, 
-					// we ened to manually raise a focus change event 
-					// to signal the platform that the currently accessibility focused widget is losing focus 
+					// For widgets that are non-keyboard focusable but support accessibility,
+					// we ened to manually raise a focus change event
+					// to signal the platform that the currently accessibility focused widget is losing focus
 					// As regular navigation focus in FSlateApplication doesn't apply to these widgets
-					if (!Widget->SupportsKeyboardFocus())
+					if (!Args.Widget->SupportsKeyboardFocus())
 					{
 						TSharedPtr<IAccessibleWidget> CurrentAccessibilityFocusedWidget = GetAccessibilityFocusedWidget();
 						if (CurrentAccessibilityFocusedWidget.IsValid())
 						{
-							// We manually raise a focus change event here   to allow platforms to unfocus 
-							// from the currently focused widget 
-							FSlateAccessibleMessageHandler::RaiseEvent(CurrentAccessibilityFocusedWidget.ToSharedRef(), EAccessibleEvent::FocusChange, true, false);
+							// We manually raise a focus change event here   to allow platforms to unfocus
+							// from the currently focused widget
+							//FSlateAccessibleMessageHandler::RaiseEvent(CurrentAccessibilityFocusedWidget.ToSharedRef(), EAccessibleEvent::FocusChange, true, false);
+							FSlateAccessibleMessageHandler::RaiseEvent(FGenericAccessibleMessageHandler::FAccessibleEventArgs(CurrentAccessibilityFocusedWidget.ToSharedRef(), EAccessibleEvent::FocusChange, true, false, Args.SlateUserId));
 						}
 					}
-					// Update the widget that currently has accessibility focus 
+					// Update the widget that currently has accessibility focus
 					SetAccessibilityFocusedWidget(AccessibleWidget);
 				}
 			}
-			FSlateAccessibleMessageHandler::RaiseEvent(AccessibleWidget, Event, OldValue, NewValue);
+			FSlateAccessibleMessageHandler::RaiseEvent(FGenericAccessibleMessageHandler::FAccessibleEventArgs(AccessibleWidget, Args.Event, Args.OldValue, Args.NewValue, Args.SlateUserId));
 		}
 	}
 }
 
-
 int32 GAccessibleWidgetsProcessedPerTick = 100;
 FAutoConsoleVariableRef AccessibleWidgetsProcessedPerTickRef(
 	TEXT("Slate.AccessibleWidgetsProcessedPerTick"),
@@ -321,7 +321,7 @@ void FSlateAccessibleMessageHandler::MakeAccessibleAnnouncement(const FString& A
 		TSharedPtr<SWidget> ActiveTopLevelWindow = FSlateApplicationBase::Get().GetActiveTopLevelWindow();
 		if(ActiveTopLevelWindow.IsValid())
 		{
-			RaiseEvent(FSlateAccessibleWidgetCache::GetAccessibleWidget(ActiveTopLevelWindow.ToSharedRef()), EAccessibleEvent::Notification, FString(), AnnouncementString);
+			RaiseEvent(FGenericAccessibleMessageHandler::FAccessibleEventArgs(FSlateAccessibleWidgetCache::GetAccessibleWidget(ActiveTopLevelWindow.ToSharedRef()), EAccessibleEvent::Notification, FString(), AnnouncementString));
 		}
 	}
 }
diff --git a/Engine/Source/Runtime/SlateCore/Private/Widgets/Accessibility/SlateCoreAccessibleWidgets.cpp b/Engine/Source/Runtime/SlateCore/Private/Widgets/Accessibility/SlateCoreAccessibleWidgets.cpp
--- a/Engine/Source/Runtime/SlateCore/Private/Widgets/Accessibility/SlateCoreAccessibleWidgets.cpp
+++ b/Engine/Source/Runtime/SlateCore/Private/Widgets/Accessibility/SlateCoreAccessibleWidgets.cpp
@@ -192,9 +192,9 @@ void FSlateAccessibleWidget::UpdateParent(TSharedPtr<IAccessibleWidget> NewParen
 	if (Parent != NewParent)
 	{
 		FSlateApplicationBase::Get().GetAccessibleMessageHandler()->RaiseEvent(
-			AsShared(), EAccessibleEvent::ParentChanged,
+			FGenericAccessibleMessageHandler::FAccessibleEventArgs(AsShared(), EAccessibleEvent::ParentChanged,
 			Parent.IsValid() ? Parent.Pin()->GetId() : IAccessibleWidget::InvalidAccessibleWidgetId,
-			NewParent.IsValid() ? NewParent->GetId() : IAccessibleWidget::InvalidAccessibleWidgetId);
+			NewParent.IsValid() ? NewParent->GetId() : IAccessibleWidget::InvalidAccessibleWidgetId));
 		Parent = StaticCastSharedPtr<FSlateAccessibleWidget>(NewParent);
 	}
 }
diff --git a/Engine/Source/Runtime/SlateCore/Public/Widgets/Accessibility/SlateAccessibleMessageHandler.h b/Engine/Source/Runtime/SlateCore/Public/Widgets/Accessibility/SlateAccessibleMessageHandler.h
--- a/Engine/Source/Runtime/SlateCore/Public/Widgets/Accessibility/SlateAccessibleMessageHandler.h
+++ b/Engine/Source/Runtime/SlateCore/Public/Widgets/Accessibility/SlateAccessibleMessageHandler.h
@@ -40,16 +40,41 @@ public:
 	 */
 	void OnWidgetRemoved(SWidget* Widget);
 
+	/**
+	 * The arguments that make up an accessible event raised by an SWidget.
+	 * It is up to the user to create an instance of this struct when they want to raise an accessible event from an SWidget
+	 * to be passed to an accessible event handler.
+	 * @see OnWidgetEventRaised
+	 */
+	struct FSlateWidgetAccessibleEventArgs
+	{
+		FSlateWidgetAccessibleEventArgs(TSharedRef<SWidget> InWidget, EAccessibleEvent InEvent, FVariant InOldValue = FVariant(), FVariant InNewValue = FVariant(), int32 InSlateUserId = 0)
+			: Widget(InWidget)
+			, Event(InEvent)
+			, OldValue(InOldValue)
+			, NewValue(InNewValue)
+			, SlateUserId(InSlateUserId)
+		{}
+			
+		/** The widget that's raising the accessible event */
+		TSharedRef<SWidget> Widget;
+		/** The type of event being raised. */
+		EAccessibleEvent Event;
+		/** The value of the property being changed before the change took place. */
+		FVariant OldValue;
+		/** The value of the property being changed after the change took place. Alternatively, can contain any miscellaneous data for the accessible event. */
+				FVariant NewValue;
+		/** The index of the Slate user that feedback for the accessible event should be directed towards. */
+		int32 SlateUserId;
+	};
 	/**
 	 * Callback for a Slate widget indicating that a property change occurred. This may also be used by certain events
 	 * such as Notification which don't have an 'OldValue'. Only NewValue should be set for those types of events.
 	 *
-	 * @param Widget The widget raising the event.
-	 * @param Event The type of event being raised.
-	 * @param OldValue The value of the property being changed before the change occurred.
-	 * @param NewValue The value of the property being changed after the change occurred, or any miscellaneous data for the event.
+	 * @param Args The arguments for the accessible event being raised.
+	 * @see FSlateWidgetAccessibleEventArgs
 	 */
-	void OnWidgetEventRaised(TSharedRef<SWidget> Widget, EAccessibleEvent Event, FVariant OldValue = FVariant(), FVariant NewValue = FVariant());
+	void OnWidgetEventRaised(const FSlateWidgetAccessibleEventArgs& Args);
 
 	/**
 	 * Refresh the accessible widget tree next available tick. This should be called any time the Slate tree changes.
