Slate: Add W component to SVectorInputBox.

diff --git a/Engine/Source/Runtime/Slate/Public/SlateFwd.h b/Engine/Source/Runtime/Slate/Public/SlateFwd.h
--- a/Engine/Source/Runtime/Slate/Public/SlateFwd.h
+++ b/Engine/Source/Runtime/Slate/Public/SlateFwd.h
@@ -36,8 +36,8 @@ class SVirtualKeyboardEntry;
 class SSafeZone;
 struct FMarqueeRect;
 
-template <typename NumericType> class SNumericVectorInputBox;
-using SVectorInputBox = SNumericVectorInputBox<float>;
+template <typename NumericType, typename VectorType, int32 NumberOfComponent> class SNumericVectorInputBox;
+using SVectorInputBox = SNumericVectorInputBox<float, UE::Math::TVector<float>, 3>;
 template <typename NumericType> class SNumericRotatorInputBox;
 using SRotatorInputBox = SNumericRotatorInputBox<float>;
 class SVirtualJoystick;
diff --git a/Engine/Source/Runtime/Slate/Public/Widgets/Input/SVectorInputBox.h b/Engine/Source/Runtime/Slate/Public/Widgets/Input/SVectorInputBox.h
--- a/Engine/Source/Runtime/Slate/Public/Widgets/Input/SVectorInputBox.h
+++ b/Engine/Source/Runtime/Slate/Public/Widgets/Input/SVectorInputBox.h
@@ -13,13 +13,15 @@
 #include "Styling/CoreStyle.h"
 #include "Framework/SlateDelegates.h"
 
+#include <type_traits>
+
 class FArrangedChildren;
 class SHorizontalBox;
 
 /**
  * Vector Slate control
  */
-template<typename NumericType>
+template<typename NumericType, typename VectorType = UE::Math::TVector<NumericType>, int32 NumberOfComponents = 3>
 class SNumericVectorInputBox : public SCompoundWidget
 {
 public:
@@ -29,71 +31,125 @@ public:
 	/** Notification for float value committed */
 	DECLARE_DELEGATE_TwoParams(FOnNumericValueCommitted, NumericType, ETextCommit::Type);
 
-	SLATE_BEGIN_ARGS(SNumericVectorInputBox<NumericType>)
-		: _Font( FAppStyle::Get().GetFontStyle("NormalFont") )
-		, _AllowSpin( false )
-		, _SpinDelta( 1 )
-		, _bColorAxisLabels( false )
-		{}
-		
+	struct FArguments;
+
+private:
+	using ThisClass = SNumericVectorInputBox<NumericType, VectorType, NumberOfComponents>;
+
+	struct FVectorXArgumentsEmpty {};
+	template<typename ArgumentType>
+	struct FVectorXArguments : FVectorXArgumentsEmpty
+	{
+		using WidgetArgsType = ArgumentType;
+
 		/** X Component of the vector */
-		SLATE_ATTRIBUTE( TOptional<NumericType>, X )
+		SLATE_ATTRIBUTE(TOptional<NumericType>, X)
+
+		/** Called when the x value of the vector is changed */
+		SLATE_EVENT(FOnNumericValueChanged, OnXChanged)
+
+		/** Called when the x value of the vector is committed */
+		SLATE_EVENT(FOnNumericValueCommitted, OnXCommitted)
+
+		/** Menu extender delegate for the X value */
+		SLATE_EVENT(FMenuExtensionDelegate, ContextMenuExtenderX)
+	};
+
+	struct FVectorYArgumentsEmpty {};
+	template<typename ArgumentType>
+	struct FVectorYArguments : FVectorYArgumentsEmpty
+	{
+		using WidgetArgsType = ArgumentType;
 
 		/** Y Component of the vector */
-		SLATE_ATTRIBUTE( TOptional<NumericType>, Y )
+		SLATE_ATTRIBUTE(TOptional<NumericType>, Y)
+
+		/** Called when the Y value of the vector is changed */
+		SLATE_EVENT(FOnNumericValueChanged, OnYChanged)
+
+		/** Called when the Y value of the vector is committed */
+		SLATE_EVENT(FOnNumericValueCommitted, OnYCommitted)
+
+		/** Menu extender delegate for the Y value */
+		SLATE_EVENT(FMenuExtensionDelegate, ContextMenuExtenderY)
+	};
+
+	struct FVectorZArgumentsEmpty {};
+	template<typename ArgumentType>
+	struct FVectorZArguments : FVectorZArgumentsEmpty
+	{
+		using WidgetArgsType = ArgumentType;
 
 		/** Z Component of the vector */
-		SLATE_ATTRIBUTE( TOptional<NumericType>, Z )
+		SLATE_ATTRIBUTE(TOptional<NumericType>, Z)
 
-		/** Font to use for the text in this box */
-		SLATE_ATTRIBUTE( FSlateFontInfo, Font )
+		/** Called when the Z value of the vector is changed */
+		SLATE_EVENT(FOnNumericValueChanged, OnZChanged)
 
-		/** Whether or not values can be spun or if they should be typed in */
-		SLATE_ARGUMENT( bool, AllowSpin )
+		/** Called when the Z value of the vector is committed */
+		SLATE_EVENT(FOnNumericValueCommitted, OnZCommitted)
 
-		/** The delta amount to apply, per pixel, when the spinner is dragged. */
-		SLATE_ATTRIBUTE( NumericType, SpinDelta )
+		/** Menu extender delegate for the Z value */
+		SLATE_EVENT(FMenuExtensionDelegate, ContextMenuExtenderZ)
+	};
 
-		/** Should the axis labels be colored */
-		SLATE_ARGUMENT( bool, bColorAxisLabels )		
+	struct FVectorWArgumentsEmpty {};
+	template<typename ArgumentType>
+	struct FVectorWArguments : FVectorWArgumentsEmpty
+	{
+		using WidgetArgsType = ArgumentType;
 
-		/** Allow responsive layout to crush the label and margins when there is not a lot of room */
-		UE_DEPRECATED(5.0, "AllowResponsiveLayout unused as it is no longer necessary.")
-		FArguments& AllowResponsiveLayout(bool bAllow)
-		{
-			return TSlateBaseNamedArgs<SNumericVectorInputBox<NumericType>>::Me();
-		}
-					
-		/** Called when the x value of the vector is changed */
-		SLATE_EVENT( FOnNumericValueChanged, OnXChanged )
+		/** W Component of the vector */
+		SLATE_ATTRIBUTE(TOptional<NumericType>, W)
 
-		/** Called when the y value of the vector is changed */
-		SLATE_EVENT( FOnNumericValueChanged, OnYChanged )
+		/** Called when the W value of the vector is changed */
+		SLATE_EVENT(FOnNumericValueChanged, OnWChanged)
 
-		/** Called when the z value of the vector is changed */
-		SLATE_EVENT( FOnNumericValueChanged, OnZChanged )
+		/** Called when the W value of the vector is committed */
+		SLATE_EVENT(FOnNumericValueCommitted, OnWCommitted)
 
-		/** Called when the x value of the vector is committed */
-		SLATE_EVENT( FOnNumericValueCommitted, OnXCommitted )
+		/** Menu extender delegate for the W value */
+		SLATE_EVENT(FMenuExtensionDelegate, ContextMenuExtenderW)
+	};
 
-		/** Called when the y value of the vector is committed */
-		SLATE_EVENT( FOnNumericValueCommitted, OnYCommitted )
+public:
+	//SLATE_BEGIN_ARGS(SNumericVectorInputBox<NumericType>)
+	struct FArguments : public TSlateBaseNamedArgs<ThisClass>
+		, std::conditional<NumberOfComponents >= 1, FVectorXArguments<FArguments>, FVectorXArgumentsEmpty>::type
+		, std::conditional<NumberOfComponents >= 2, FVectorYArguments<FArguments>, FVectorYArgumentsEmpty>::type
+		, std::conditional<NumberOfComponents >= 3, FVectorZArguments<FArguments>, FVectorZArgumentsEmpty>::type
+		, std::conditional<NumberOfComponents >= 4, FVectorWArguments<FArguments>, FVectorWArgumentsEmpty>::type
+	{
+		typedef FArguments WidgetArgsType;
+		FORCENOINLINE FArguments()
+			: _Font(FAppStyle::Get().GetFontStyle("NormalFont"))
+			, _AllowSpin(false)
+			, _SpinDelta(1)
+			, _bColorAxisLabels(false)
+		{}
 
-		/** Called when the z value of the vector is committed */
-		SLATE_EVENT( FOnNumericValueCommitted, OnZCommitted )
+		/** Font to use for the text in this box */
+		SLATE_ATTRIBUTE( FSlateFontInfo, Font )
 
-		/** Menu extender delegate for the X value */
-		SLATE_EVENT( FMenuExtensionDelegate, ContextMenuExtenderX )
+		/** Whether or not values can be spun or if they should be typed in */
+		SLATE_ARGUMENT( bool, AllowSpin )
 
-		/** Menu extender delegate for the Y value */
-		SLATE_EVENT( FMenuExtensionDelegate, ContextMenuExtenderY )
+		/** The delta amount to apply, per pixel, when the spinner is dragged. */
+		SLATE_ATTRIBUTE( NumericType, SpinDelta )
 
-		/** Menu extender delegate for the Z value */
-		SLATE_EVENT( FMenuExtensionDelegate, ContextMenuExtenderZ )
+		/** Should the axis labels be colored */
+		SLATE_ARGUMENT( bool, bColorAxisLabels )		
+
+			/** Allow responsive layout to crush the label and margins when there is not a lot of room */
+			UE_DEPRECATED(5.0, "AllowResponsiveLayout unused as it is no longer necessary.")
+			FArguments& AllowResponsiveLayout(bool bAllow)
+			{
+				return TSlateBaseNamedArgs<ThisClass>::Me();
+			}
 
 		/** Called right before the slider begins to move for any of the vector components */
 		SLATE_EVENT( FSimpleDelegate, OnBeginSliderMovement )
-		
+
 		/** Called right after the slider handle is released by the user for any of the vector components */
 		SLATE_EVENT( FOnNumericValueChanged, OnEndSliderMovement )
 
@@ -107,122 +163,120 @@ public:
 	 *
 	 * @param	InArgs	The declaration data for this widget
 	 */
-	void Construct( const FArguments& InArgs )
-		{
-			TSharedRef<SHorizontalBox> HorizontalBox = SNew(SHorizontalBox);
+	void Construct(const FArguments& InArgs)
+	{
+		TSharedRef<SHorizontalBox> HorizontalBox = SNew(SHorizontalBox);
 
-			ChildSlot
-				[
-					HorizontalBox
-				];
+		ChildSlot
+		[
+				HorizontalBox
+		];
 
+		if constexpr (NumberOfComponents >= 1)
+		{
 			ConstructX(InArgs, HorizontalBox);
+		}
+		if constexpr (NumberOfComponents >= 2)
+		{
 			ConstructY(InArgs, HorizontalBox);
+		}
+		if constexpr (NumberOfComponents >= 3)
+		{
 			ConstructZ(InArgs, HorizontalBox);
 		}
-private:
-
+		if constexpr (NumberOfComponents >= 4)
+		{
+			ConstructW(InArgs, HorizontalBox);
+		}
+	}
 
+private:
 	/**
-	 * Construct widgets for the X Value
+	 * Construct the widget component
 	 */
-	void ConstructX( const FArguments& InArgs, TSharedRef<SHorizontalBox> HorizontalBox )
+	void ConstructComponent(const FArguments& InArgs,
+		const FLinearColor& LabelColor,
+		const FText& TooltipText,
+		TSharedRef<SHorizontalBox>& HorizontalBox,
+		const TAttribute<TOptional<NumericType>>& Component,
+		const FOnNumericValueChanged& OnComponentChanged,
+		const FOnNumericValueCommitted& OnComponentCommitted,
+		const FMenuExtensionDelegate& OnContextMenuExtenderComponent)
 	{
 		TSharedRef<SWidget> LabelWidget = SNullWidget::NullWidget;
 		if (InArgs._bColorAxisLabels)
 		{
-			const FLinearColor LabelColor = SNumericEntryBox<NumericType>::RedLabelBackgroundColor;
 			LabelWidget = SNumericEntryBox<NumericType>::BuildNarrowColorLabel(LabelColor);
 		}
 
-		TAttribute<TOptional<NumericType>> Value = InArgs._X;
+		TAttribute<TOptional<NumericType>> Value = Component;
 
 		HorizontalBox->AddSlot()
-			[
-				SNew(SNumericEntryBox<NumericType>)
-				.AllowSpin(InArgs._AllowSpin)
-				.Font(InArgs._Font)
-				.Value(InArgs._X)
-				.OnValueChanged(InArgs._OnXChanged)
-				.OnValueCommitted(InArgs._OnXCommitted)
-				.ToolTipText(MakeAttributeLambda([Value]
+		[
+			SNew(SNumericEntryBox<NumericType>)
+			.AllowSpin(InArgs._AllowSpin)
+			.Font(InArgs._Font)
+			.Value(Component)
+			.OnValueChanged(OnComponentChanged)
+			.OnValueCommitted(OnComponentCommitted)
+			.ToolTipText(MakeAttributeLambda([Component, TooltipText]
+			{
+				if (Component.Get().IsSet())
 				{
-					if (Value.Get().IsSet())
-					{
-						return FText::Format(NSLOCTEXT("SVectorInputBox", "X_ToolTip", "X: {0}"), Value.Get().GetValue());
-					}
-					return NSLOCTEXT("SVectorInputBox", "MultipleValues", "Multiple Values");
-				}))
-				.UndeterminedString(NSLOCTEXT("SVectorInputBox", "MultipleValues", "Multiple Values"))
-				.ContextMenuExtender(InArgs._ContextMenuExtenderX)
-				.TypeInterface(InArgs._TypeInterface)
-				.MinValue(TOptional<NumericType>())
-				.MaxValue(TOptional<NumericType>())
-				.MinSliderValue(TOptional<NumericType>())
-				.MaxSliderValue(TOptional<NumericType>())
-				.LinearDeltaSensitivity(1)
-				.Delta(InArgs._SpinDelta)
-				.OnBeginSliderMovement(InArgs._OnBeginSliderMovement)
-				.OnEndSliderMovement(InArgs._OnEndSliderMovement)
-				.LabelPadding(FMargin(3.f))
-				.LabelLocation(SNumericEntryBox<NumericType>::ELabelLocation::Inside)
-				.Label()
-				[
-					LabelWidget
-				]
-			];
+					return FText::Format(TooltipText, Component.Get().GetValue());
+				}
+				return NSLOCTEXT("SVectorInputBox", "MultipleValues", "Multiple Values");
+			}))
+			.UndeterminedString(NSLOCTEXT("SVectorInputBox", "MultipleValues", "Multiple Values"))
+			.ContextMenuExtender(OnContextMenuExtenderComponent)
+			.TypeInterface(InArgs._TypeInterface)
+			.MinValue(TOptional<NumericType>())
+			.MaxValue(TOptional<NumericType>())
+			.MinSliderValue(TOptional<NumericType>())
+			.MaxSliderValue(TOptional<NumericType>())
+			.LinearDeltaSensitivity(1)
+			.Delta(InArgs._SpinDelta)
+			.OnBeginSliderMovement(InArgs._OnBeginSliderMovement)
+			.OnEndSliderMovement(InArgs._OnEndSliderMovement)
+			.LabelPadding(FMargin(3.f))
+			.LabelLocation(SNumericEntryBox<NumericType>::ELabelLocation::Inside)
+			.Label()
+			[
+				LabelWidget
+			]
+	];
+}
 
+	/**
+	 * Construct widgets for the X Value
+	 */
+	void ConstructX(const FArguments& InArgs, TSharedRef<SHorizontalBox> HorizontalBox)
+	{
+		ConstructComponent(InArgs,
+			SNumericEntryBox<NumericType>::RedLabelBackgroundColor,
+			NSLOCTEXT("SVectorInputBox", "X_ToolTip", "X: {0}"),
+			HorizontalBox,
+			InArgs._X,
+			InArgs._OnXChanged,
+			InArgs._OnXCommitted,
+			InArgs._ContextMenuExtenderX
+		);
 	}
 
 	/**
 	 * Construct widgets for the Y Value
 	 */
-	void ConstructY( const FArguments& InArgs, TSharedRef<SHorizontalBox> HorizontalBox )
+	void ConstructY(const FArguments& InArgs, TSharedRef<SHorizontalBox> HorizontalBox)
 	{
-		TSharedRef<SWidget> LabelWidget = SNullWidget::NullWidget;
-		if (InArgs._bColorAxisLabels)
-		{
-			const FLinearColor LabelColor = SNumericEntryBox<NumericType>::GreenLabelBackgroundColor;
-			LabelWidget = SNumericEntryBox<NumericType>::BuildNarrowColorLabel(LabelColor);
-		}
-
-		TAttribute<TOptional<NumericType>> Value = InArgs._Y;
-
-		HorizontalBox->AddSlot()
-			[
-				SNew(SNumericEntryBox<NumericType>)
-				.AllowSpin(InArgs._AllowSpin)
-				.Font(InArgs._Font)
-				.Value(InArgs._Y)
-				.OnValueChanged(InArgs._OnYChanged)
-				.OnValueCommitted(InArgs._OnYCommitted)
-				.ToolTipText(MakeAttributeLambda([Value]
-				{
-					if (Value.Get().IsSet())
-					{
-						return FText::Format(NSLOCTEXT("SVectorInputBox", "Y_ToolTip", "Y: {0}"), Value.Get().GetValue());
-					}
-					return NSLOCTEXT("SVectorInputBox", "MultipleValues", "Multiple Values");
-				}))
-				.UndeterminedString(NSLOCTEXT("SVectorInputBox", "MultipleValues", "Multiple Values"))
-				.ContextMenuExtender(InArgs._ContextMenuExtenderY)
-				.TypeInterface(InArgs._TypeInterface)
-				.MinValue(TOptional<NumericType>())
-				.MaxValue(TOptional<NumericType>())
-				.MinSliderValue(TOptional<NumericType>())
-				.MaxSliderValue(TOptional<NumericType>())
-				.LinearDeltaSensitivity(1)
-				.Delta(InArgs._SpinDelta)
-				.OnBeginSliderMovement(InArgs._OnBeginSliderMovement)
-				.OnEndSliderMovement(InArgs._OnEndSliderMovement)
-				.LabelPadding(FMargin(3))
-				.LabelLocation(SNumericEntryBox<NumericType>::ELabelLocation::Inside)
-				.Label()
-				[
-					LabelWidget
-				]
-			];
-
+		ConstructComponent(InArgs,
+			SNumericEntryBox<NumericType>::GreenLabelBackgroundColor,
+			NSLOCTEXT("SVectorInputBox", "Y_ToolTip", "Y: {0}"),
+			HorizontalBox,
+			InArgs._Y,
+			InArgs._OnYChanged,
+			InArgs._OnYCommitted,
+			InArgs._ContextMenuExtenderY
+		);
 	}
 
 	/**
@@ -230,55 +284,36 @@ private:
 	 */
 	void ConstructZ(const FArguments& InArgs, TSharedRef<SHorizontalBox> HorizontalBox)
 	{
-		TSharedRef<SWidget> LabelWidget = SNullWidget::NullWidget;
-		if (InArgs._bColorAxisLabels)
-		{
-			const FLinearColor LabelColor = SNumericEntryBox<NumericType>::BlueLabelBackgroundColor;
-			LabelWidget = SNumericEntryBox<NumericType>::BuildNarrowColorLabel(LabelColor);
-		}
-
-		TAttribute<TOptional<NumericType>> Value = InArgs._Z;
-
-		HorizontalBox->AddSlot()
-			[
-				SNew(SNumericEntryBox<NumericType>)
-				.AllowSpin(InArgs._AllowSpin)
-				.Font(InArgs._Font)
-				.Value(InArgs._Z)
-				.OnValueChanged(InArgs._OnZChanged)
-				.OnValueCommitted(InArgs._OnZCommitted)
-				.ToolTipText(MakeAttributeLambda([Value]
-				{
-					if (Value.Get().IsSet())
-					{
-						return FText::Format(NSLOCTEXT("SVectorInputBox", "Z_ToolTip", "Z: {0}"), Value.Get().GetValue());
-					}
-					return NSLOCTEXT("SVectorInputBox", "MultipleValues", "Multiple Values");
-				}))
-				.UndeterminedString(NSLOCTEXT("SVectorInputBox", "MultipleValues", "Multiple Values"))
-				.ContextMenuExtender(InArgs._ContextMenuExtenderZ)
-				.TypeInterface(InArgs._TypeInterface)
-				.MinValue(TOptional<NumericType>())
-				.MaxValue(TOptional<NumericType>())
-				.MinSliderValue(TOptional<NumericType>())
-				.MaxSliderValue(TOptional<NumericType>())
-				.LinearDeltaSensitivity(1)
-				.Delta(InArgs._SpinDelta)
-				.OnBeginSliderMovement(InArgs._OnBeginSliderMovement)
-				.OnEndSliderMovement(InArgs._OnEndSliderMovement)
-				.LabelPadding(FMargin(3))
-				.LabelLocation(SNumericEntryBox<NumericType>::ELabelLocation::Inside)
-				.Label()
-				[
-					LabelWidget
-				]
-			];
+		ConstructComponent(InArgs,
+			SNumericEntryBox<NumericType>::BlueLabelBackgroundColor,
+			NSLOCTEXT("SVectorInputBox", "Z_ToolTip", "Z: {0}"),
+			HorizontalBox,
+			InArgs._Z,
+			InArgs._OnZChanged,
+			InArgs._OnZCommitted,
+			InArgs._ContextMenuExtenderZ
+		);
 	}
 
+	/**
+	 * Construct widgets for the W Value
+	 */
+	void ConstructW(const FArguments& InArgs, TSharedRef<SHorizontalBox> HorizontalBox)
+	{
+		ConstructComponent(InArgs,
+			FLinearColor::Yellow,
+			NSLOCTEXT("SVectorInputBox", "W_ToolTip", "W: {0}"),
+			HorizontalBox,
+			InArgs._W,
+			InArgs._OnWChanged,
+			InArgs._OnWCommitted,
+			InArgs._ContextMenuExtenderW
+		);
+	}
 };
 
 /**
  * For backward compatibility
  */
-using SVectorInputBox = SNumericVectorInputBox<float>;
+using SVectorInputBox = SNumericVectorInputBox<float, UE::Math::TVector<float>, 3>;
 
diff --git a/Engine/Source/Runtime/SlateCore/Public/Widgets/DeclarativeSyntaxSupport.h b/Engine/Source/Runtime/SlateCore/Public/Widgets/DeclarativeSyntaxSupport.h
--- a/Engine/Source/Runtime/SlateCore/Public/Widgets/DeclarativeSyntaxSupport.h
+++ b/Engine/Source/Runtime/SlateCore/Public/Widgets/DeclarativeSyntaxSupport.h
@@ -21,7 +21,7 @@ template<typename WidgetType> struct TSlateBaseNamedArgs;
 /**
  * Slate widgets are constructed through SNew and SAssignNew.
  * e.g.
- *      
+ *
  *     TSharedRef<SButton> MyButton = SNew(SButton);
  *        or
  *     TSharedPtr<SButton> MyButton;
@@ -119,7 +119,7 @@ template<typename WidgetType> struct TSlateBaseNamedArgs;
 		WidgetArgsType& AttrName( TAttribute< AttrType > InAttribute ) \
 		{ \
 			_##AttrName = MoveTemp(InAttribute); \
-			return this->Me(); \
+			return static_cast<WidgetArgsType*>(this)->Me(); \
 		} \
 		\
 		/* Bind attribute with delegate to a global function
@@ -128,59 +128,58 @@ template<typename WidgetType> struct TSlateBaseNamedArgs;
 		WidgetArgsType& AttrName##_Static( typename TAttribute< AttrType >::FGetter::template FStaticDelegate<VarTypes...>::FFuncPtr InFunc, VarTypes... Vars )	\
 		{ \
 			_##AttrName = TAttribute< AttrType >::Create( TAttribute< AttrType >::FGetter::CreateStatic( InFunc, Vars... ) ); \
-			return this->Me(); \
+			return static_cast<WidgetArgsType*>(this)->Me(); \
 		} \
-	\
-	\
+		\
 		/* Bind attribute with delegate to a lambda
 		 * technically this works for any functor types, but lambdas are the primary use case */ \
 		WidgetArgsType& AttrName##_Lambda(TFunction< AttrType(void) >&& InFunctor) \
 		{ \
 			_##AttrName = TAttribute< AttrType >::Create(Forward<TFunction< AttrType(void) >>(InFunctor)); \
-			return this->Me(); \
+			return static_cast<WidgetArgsType*>(this)->Me(); \
 		} \
-	\
+		\
 		/* Bind attribute with delegate to a raw C++ class method */ \
 		template< class UserClass, typename... VarTypes >	\
 		WidgetArgsType& AttrName##_Raw( UserClass* InUserObject, typename TAttribute< AttrType >::FGetter::template TConstMethodPtr< UserClass, VarTypes... > InFunc, VarTypes... Vars )	\
 		{ \
 			_##AttrName = TAttribute< AttrType >::Create( TAttribute< AttrType >::FGetter::CreateRaw( InUserObject, InFunc, Vars... ) ); \
-			return this->Me(); \
+			return static_cast<WidgetArgsType*>(this)->Me(); \
 		} \
-	\
+		\
 		/* Bind attribute with delegate to a shared pointer-based class method.  Slate mostly uses shared pointers so we use an overload for this type of binding. */ \
 		template< class UserClass, typename... VarTypes >	\
 		WidgetArgsType& AttrName( TSharedRef< UserClass > InUserObjectRef, typename TAttribute< AttrType >::FGetter::template TConstMethodPtr< UserClass, VarTypes... > InFunc, VarTypes... Vars )	\
 		{ \
 			_##AttrName = TAttribute< AttrType >::Create( TAttribute< AttrType >::FGetter::CreateSP( InUserObjectRef, InFunc, Vars... ) ); \
-			return this->Me(); \
+			return static_cast<WidgetArgsType*>(this)->Me(); \
 		} \
-	\
+		\
 		/* Bind attribute with delegate to a shared pointer-based class method.  Slate mostly uses shared pointers so we use an overload for this type of binding. */ \
 		template< class UserClass, typename... VarTypes >	\
 		WidgetArgsType& AttrName( UserClass* InUserObject, typename TAttribute< AttrType >::FGetter::template TConstMethodPtr< UserClass, VarTypes... > InFunc, VarTypes... Vars )	\
 		{ \
 			_##AttrName = TAttribute< AttrType >::Create( TAttribute< AttrType >::FGetter::CreateSP( InUserObject, InFunc, Vars... ) ); \
-			return this->Me(); \
+			return static_cast<WidgetArgsType*>(this)->Me(); \
 		} \
-	\
+		\
 		/* Bind attribute with delegate to a UObject-based class method */ \
 		template< class UserClass, typename... VarTypes >	\
 		WidgetArgsType& AttrName##_UObject( UserClass* InUserObject, typename TAttribute< AttrType >::FGetter::template TConstMethodPtr< UserClass, VarTypes... > InFunc, VarTypes... Vars )	\
 		{ \
 			_##AttrName = TAttribute< AttrType >::Create( TAttribute< AttrType >::FGetter::CreateUObject( InUserObject, InFunc, Vars... ) ); \
-			return this->Me(); \
+			return static_cast<WidgetArgsType*>(this)->Me(); \
 		} \
 
 #define SLATE_PRIVATE_ARGUMENT_VARIABLE( ArgType, ArgName ) \
 		ArgType _##ArgName
-		
+
 
 #define SLATE_PRIVATE_ARGUMENT_FUNCTION( ArgType, ArgName ) \
 		WidgetArgsType& ArgName( ArgType InArg ) \
 		{ \
 			_##ArgName = InArg; \
-			return this->Me(); \
+			return static_cast<WidgetArgsType*>(this)->Me(); \
 		}
 
 /**
@@ -217,44 +216,44 @@ template<typename WidgetType> struct TSlateBaseNamedArgs;
 		WidgetArgsType& ArgName( const ArgType* InArg ) \
 		{ \
 			_##ArgName = InArg; \
-			return this->Me(); \
+			return static_cast<WidgetArgsType*>(this)->Me(); \
 		}\
 		\
 		WidgetArgsType& ArgName( const class USlateWidgetStyleAsset* const InSlateStyleAsset ) \
 		{ \
 			_##ArgName = InSlateStyleAsset->GetStyleChecked< ArgType >(); \
-			return this->Me(); \
+			return static_cast<WidgetArgsType*>(this)->Me(); \
 		}\
 		\
 		WidgetArgsType& ArgName( const TWeakObjectPtr< const class USlateWidgetStyleAsset >& InSlateStyleAsset ) \
 		{ \
 			_##ArgName = InSlateStyleAsset.Get()->GetStyleChecked< ArgType >(); \
-			return this->Me(); \
+			return static_cast<WidgetArgsType*>(this)->Me(); \
 		}\
 		\
 		WidgetArgsType& ArgName( const class ISlateStyle* InSlateStyle, const FName& StyleName, const ANSICHAR* Specifier = nullptr ) \
 		{ \
 			check( InSlateStyle != nullptr ); \
 			_##ArgName = &InSlateStyle->GetWidgetStyle< ArgType >( StyleName, Specifier ); \
-			return this->Me(); \
+			return static_cast<WidgetArgsType*>(this)->Me(); \
 		}\
 		\
 		WidgetArgsType& ArgName( const class ISlateStyle& InSlateStyle, const FName& StyleName, const ANSICHAR* Specifier = nullptr ) \
 		{ \
 			_##ArgName = &InSlateStyle.GetWidgetStyle< ArgType >( StyleName, Specifier ); \
-			return this->Me(); \
+			return static_cast<WidgetArgsType*>(this)->Me(); \
 		}\
 		\
 		WidgetArgsType& ArgName( const TWeakObjectPtr< const class ISlateStyle >& InSlateStyle, const FName& StyleName, const ANSICHAR* Specifier = nullptr ) \
 		{ \
 			_##ArgName = &InSlateStyle.Get()->GetWidgetStyle< ArgType >( StyleName, Specifier ); \
-			return this->Me(); \
+			return static_cast<WidgetArgsType*>(this)->Me(); \
 		} \
 		\
 		WidgetArgsType& ArgName( const TSharedPtr< const class ISlateStyle >& InSlateStyle, const FName& StyleName, const ANSICHAR* Specifier = nullptr ) \
 		{ \
 			_##ArgName = &InSlateStyle->GetWidgetStyle< ArgType >( StyleName, Specifier ); \
-			return this->Me(); \
+			return static_cast<WidgetArgsType*>(this)->Me(); \
 		} 
 
 
@@ -268,7 +267,7 @@ template<typename WidgetType> struct TSlateBaseNamedArgs;
 		WidgetArgsType& operator + (SlotType& SlotToAdd) \
 		{ \
 			Slots.Add( &SlotToAdd ); \
-			return *this; \
+			return static_cast<WidgetArgsType*>(this)->Me(); \
 		}
 
 
@@ -281,7 +280,7 @@ template<typename WidgetType> struct TSlateBaseNamedArgs;
 	WidgetArgsType& operator + (const typename SlotType::FArguments& ArgumentsForNewSlot) \
 		{ \
 			Slots.Add( new SlotType( ArgumentsForNewSlot ) ); \
-			return *this; \
+			return static_cast<WidgetArgsType*>(this)->Me(); \
 		}
 
 
@@ -299,19 +298,19 @@ template<typename WidgetType> struct TSlateBaseNamedArgs;
 		WidgetArgsType& operator + (typename SlotType::FSlotArguments& SlotToAdd) \
 		{ \
 			_##SlotName.Add( MoveTemp(SlotToAdd) ); \
-			return *this; \
+			return static_cast<WidgetArgsType*>(this)->Me(); \
 		} \
 		WidgetArgsType& operator + (typename SlotType::FSlotArguments&& SlotToAdd) \
 		{ \
 			_##SlotName.Add( MoveTemp(SlotToAdd) ); \
-			return *this; \
+			return static_cast<WidgetArgsType*>(this)->Me(); \
 		}
 
 #define SLATE_SLOT_BEGIN_ARGS( SlotType, SlotParentType ) \
 	public: \
 	struct FSlotArguments : public SlotParentType::FSlotArguments \
 	SLATE_PRIVATE_SLOT_BEGIN_ARGS( SlotType, SlotParentType )
-		
+
 #define SLATE_SLOT_BEGIN_ARGS_OneMixin( SlotType, SlotParentType, Mixin1 ) \
 	public: \
 	struct FSlotArguments : public SlotParentType::FSlotArguments, public Mixin1::FSlotArgumentsMixin \
@@ -326,7 +325,7 @@ template<typename WidgetType> struct TSlateBaseNamedArgs;
 	public: \
 	struct FSlotArguments : public SlotParentType::FSlotArguments, public Mixin1::FSlotArgumentsMixin, public Mixin2::FSlotArgumentsMixin, public Mixin3::FSlotArgumentsMixin \
 	SLATE_PRIVATE_SLOT_BEGIN_ARGS( SlotType, SlotParentType )
-	
+
 #define SLATE_SLOT_BEGIN_ARGS_FourMixins( SlotType, SlotParentType, Mixin1, Mixin2, Mixin3, Mixin4 ) \
 	public: \
 	struct FSlotArguments : public SlotParentType::FSlotArguments, public Mixin1::FSlotArgumentsMixin, public Mixin2::FSlotArgumentsMixin, public Mixin3::FSlotArgumentsMixin, public Mixin4::FSlotArgumentsMixin \
@@ -395,7 +394,7 @@ struct NamedSlotProperty
 		DeclarationType & operator[]( const TSharedRef<SWidget> InChild ) \
 		{ \
 			_##SlotName.Widget = InChild; \
-			return *this; \
+			return static_cast<WidgetArgsType*>(this)->Me(); \
 		}
 
 
@@ -408,13 +407,13 @@ struct NamedSlotProperty
 		WidgetArgsType& EventName( const DelegateName& InDelegate ) \
 		{ \
 			_##EventName = InDelegate; \
-			return *this; \
+			return static_cast<WidgetArgsType*>(this)->Me(); \
 		} \
 		\
 		WidgetArgsType& EventName( DelegateName&& InDelegate ) \
 		{ \
 			_##EventName = MoveTemp(InDelegate); \
-			return *this; \
+			return static_cast<WidgetArgsType*>(this)->Me(); \
 		} \
 		\
 		/* Set event delegate to a global function */ \
@@ -423,7 +422,7 @@ struct NamedSlotProperty
 		WidgetArgsType& EventName##_Static( StaticFuncPtr InFunc, VarTypes... Vars )	\
 		{ \
 			_##EventName = DelegateName::CreateStatic( InFunc, Vars... ); \
-			return *this; \
+			return static_cast<WidgetArgsType*>(this)->Me(); \
 		} \
 		\
 		/* Set event delegate to a lambda
@@ -431,8 +430,8 @@ struct NamedSlotProperty
 		template<typename FunctorType, typename... VarTypes> \
 		WidgetArgsType& EventName##_Lambda(FunctorType&& InFunctor, VarTypes... Vars) \
 		{ \
-			_##EventName = DelegateName::CreateLambda(Forward<FunctorType>(InFunctor), Vars... ); \
-			return this->Me(); \
+			_##EventName = DelegateName::CreateLambda( Forward<FunctorType>(InFunctor), Vars... ); \
+			return static_cast<WidgetArgsType*>(this)->Me(); \
 		} \
 		\
 		/* Set event delegate to a raw C++ class method */ \
@@ -440,13 +439,13 @@ struct NamedSlotProperty
 		WidgetArgsType& EventName##_Raw( UserClass* InUserObject, typename DelegateName::template TMethodPtr< UserClass, VarTypes... > InFunc, VarTypes... Vars )	\
 		{ \
 			_##EventName = DelegateName::CreateRaw( InUserObject, InFunc, Vars... ); \
-			return *this; \
+			return static_cast<WidgetArgsType*>(this)->Me(); \
 		} \
 		template< class UserClass, typename... VarTypes >	\
 		WidgetArgsType& EventName##_Raw( UserClass* InUserObject, typename DelegateName::template TConstMethodPtr< UserClass, VarTypes... > InFunc, VarTypes... Vars )	\
 		{ \
 			_##EventName = DelegateName::CreateRaw( InUserObject, InFunc, Vars... ); \
-			return *this; \
+			return static_cast<WidgetArgsType*>(this)->Me(); \
 		} \
 		\
 		/* Set event delegate to a shared pointer-based class method.  Slate mostly uses shared pointers so we use an overload for this type of binding. */ \
@@ -454,13 +453,13 @@ struct NamedSlotProperty
 		WidgetArgsType& EventName( TSharedRef< UserClass > InUserObjectRef, typename DelegateName::template TMethodPtr< UserClass, VarTypes... > InFunc, VarTypes... Vars )	\
 		{ \
 			_##EventName = DelegateName::CreateSP( InUserObjectRef, InFunc, Vars... ); \
-			return *this; \
+			return static_cast<WidgetArgsType*>(this)->Me(); \
 		} \
 		template< class UserClass, typename... VarTypes >	\
 		WidgetArgsType& EventName( TSharedRef< UserClass > InUserObjectRef, typename DelegateName::template TConstMethodPtr< UserClass, VarTypes... > InFunc, VarTypes... Vars )	\
 		{ \
 			_##EventName = DelegateName::CreateSP( InUserObjectRef, InFunc, Vars... ); \
-			return *this; \
+			return static_cast<WidgetArgsType*>(this)->Me(); \
 		} \
 		\
 		/* Set event delegate to a shared pointer-based class method.  Slate mostly uses shared pointers so we use an overload for this type of binding. */ \
@@ -468,13 +467,13 @@ struct NamedSlotProperty
 		WidgetArgsType& EventName( UserClass* InUserObject, typename DelegateName::template TMethodPtr< UserClass, VarTypes... > InFunc, VarTypes... Vars )	\
 		{ \
 			_##EventName = DelegateName::CreateSP( InUserObject, InFunc, Vars... ); \
-			return *this; \
+			return static_cast<WidgetArgsType*>(this)->Me(); \
 		} \
 		template< class UserClass, typename... VarTypes >	\
 		WidgetArgsType& EventName( UserClass* InUserObject, typename DelegateName::template TConstMethodPtr< UserClass, VarTypes... > InFunc, VarTypes... Vars )	\
 		{ \
 			_##EventName = DelegateName::CreateSP( InUserObject, InFunc, Vars... ); \
-			return *this; \
+			return static_cast<WidgetArgsType*>(this)->Me(); \
 		} \
 		\
 		/* Set event delegate to a UObject-based class method */ \
@@ -488,7 +487,7 @@ struct NamedSlotProperty
 		WidgetArgsType& EventName##_UObject( UserClass* InUserObject, typename DelegateName::template TConstMethodPtr< UserClass, VarTypes... > InFunc, VarTypes... Vars )	\
 		{ \
 			_##EventName = DelegateName::CreateUObject( InUserObject, InFunc, Vars... ); \
-			return *this; \
+			return static_cast<WidgetArgsType*>(this)->Me(); \
 		} \
 		\
 		DelegateName _##EventName; \
@@ -501,7 +500,7 @@ struct NamedSlotProperty
 			{	\
 				_##UpgradedEventName = UpgradeFuncName(InDeprecatedDelegate);	\
 			}	\
-			return *this;	\
+			return static_cast<WidgetArgsType*>(this)->Me(); \
 		}	\
 		\
 		template <typename StaticFuncPtr, typename... VarTypes>	\
@@ -513,7 +512,7 @@ struct NamedSlotProperty
 				DelegateName DeprecatedDelegate = DelegateName::CreateStatic(InFunc, Vars...);	\
 				_##UpgradedEventName = UpgradeFuncName(DeprecatedDelegate);	\
 			}	\
-			return *this;	\
+			return static_cast<WidgetArgsType*>(this)->Me(); \
 		}	\
 		\
 		template <typename FunctorType, typename... VarTypes>	\
@@ -522,7 +521,7 @@ struct NamedSlotProperty
 		{	\
 			DelegateName DeprecatedDelegate = DelegateName::CreateLambda(Forward<FunctorType>(InFunctor), Vars...);	\
 			_##UpgradedEventName = UpgradeFuncName(DeprecatedDelegate);	\
-			return *this;	\
+			return static_cast<WidgetArgsType*>(this)->Me(); \
 		}	\
 		\
 		template <typename UserClass, typename... VarTypes>	\
@@ -531,7 +530,7 @@ struct NamedSlotProperty
 		{	\
 			DelegateName DeprecatedDelegate = DelegateName::CreateRaw(InUserObject, InFunc, Vars...);	\
 			_##UpgradedEventName = UpgradeFuncName(DeprecatedDelegate);	\
-			return *this;	\
+			return static_cast<WidgetArgsType*>(this)->Me(); \
 		}	\
 		\
 		template <typename UserClass, typename... VarTypes>	\
@@ -540,7 +539,7 @@ struct NamedSlotProperty
 		{	\
 			DelegateName DeprecatedDelegate = DelegateName::CreateRaw(InUserObject, InFunc, Vars...);	\
 			_##UpgradedEventName = UpgradeFuncName(DeprecatedDelegate);	\
-			return *this;	\
+			return static_cast<WidgetArgsType*>(this)->Me(); \
 		}	\
 		\
 		template <typename UserClass, typename... VarTypes>	\
@@ -549,7 +548,7 @@ struct NamedSlotProperty
 		{	\
 			DelegateName DeprecatedDelegate = DelegateName::CreateSP(InUserObjectRef, InFunc, Vars...);	\
 			_##UpgradedEventName = UpgradeFuncName(DeprecatedDelegate);	\
-			return *this;	\
+			return static_cast<WidgetArgsType*>(this)->Me(); \
 		}	\
 		\
 		template <typename UserClass, typename... VarTypes>	\
@@ -558,7 +557,7 @@ struct NamedSlotProperty
 		{	\
 			DelegateName DeprecatedDelegate = DelegateName::CreateSP(InUserObjectRef, InFunc, Vars...);	\
 			_##UpgradedEventName = UpgradeFuncName(DeprecatedDelegate);	\
-			return *this;	\
+			return static_cast<WidgetArgsType*>(this)->Me(); \
 		}	\
 		\
 		template <typename UserClass, typename... VarTypes>	\
@@ -567,7 +566,7 @@ struct NamedSlotProperty
 		{	\
 			DelegateName DeprecatedDelegate = DelegateName::CreateSP(InUserObject, InFunc, Vars...);	\
 			_##UpgradedEventName = UpgradeFuncName(DeprecatedDelegate);	\
-			return *this;	\
+			return static_cast<WidgetArgsType*>(this)->Me(); \
 		}	\
 		\
 		template <typename UserClass, typename... VarTypes>	\
@@ -576,7 +575,7 @@ struct NamedSlotProperty
 		{	\
 			DelegateName DeprecatedDelegate = DelegateName::CreateSP(InUserObject, InFunc, Vars...);	\
 			_##UpgradedEventName = UpgradeFuncName(DeprecatedDelegate);	\
-			return *this;	\
+			return static_cast<WidgetArgsType*>(this)->Me(); \
 		}	\
 		\
 		template <typename UserClass, typename... VarTypes>	\
@@ -585,7 +584,7 @@ struct NamedSlotProperty
 		{	\
 			DelegateName DeprecatedDelegate = DelegateName::CreateUObject(InUserObject, InFunc, Vars...);	\
 			_##UpgradedEventName = UpgradeFuncName(DeprecatedDelegate);	\
-			return *this;	\
+			return static_cast<WidgetArgsType*>(this)->Me(); \
 		}	\
 		\
 		template <typename UserClass, typename... VarTypes>	\
@@ -594,7 +593,7 @@ struct NamedSlotProperty
 		{	\
 			DelegateName DeprecatedDelegate = DelegateName::CreateUObject(InUserObject, InFunc, Vars...);	\
 			_##UpgradedEventName = UpgradeFuncName(DeprecatedDelegate);	\
-			return *this;	\
+			return static_cast<WidgetArgsType*>(this)->Me(); \
 		}	\
 
 
