Slate: Fix compilation warning when C4866 is enabled.

diff --git a/Engine/Source/Runtime/SlateCore/Public/Widgets/DeclarativeSyntaxSupport.h b/Engine/Source/Runtime/SlateCore/Public/Widgets/DeclarativeSyntaxSupport.h
--- a/Engine/Source/Runtime/SlateCore/Public/Widgets/DeclarativeSyntaxSupport.h
+++ b/Engine/Source/Runtime/SlateCore/Public/Widgets/DeclarativeSyntaxSupport.h
@@ -391,7 +391,7 @@ struct NamedSlotProperty
 
 #define SLATE_DEFAULT_SLOT( DeclarationType, SlotName ) \
 		SLATE_NAMED_SLOT(DeclarationType, SlotName) ; \
-		DeclarationType & operator[]( const TSharedRef<SWidget> InChild ) \
+		DeclarationType & operator[]( const TSharedRef<SWidget>& InChild ) \
 		{ \
 			_##SlotName.Widget = InChild; \
 			return static_cast<WidgetArgsType*>(this)->Me(); \
