Exposing FPointerEvent to be used in plugins

diff --git a/Engine/Source/Runtime/SlateCore/Public/Input/Events.h b/Engine/Source/Runtime/SlateCore/Public/Input/Events.h
--- a/Engine/Source/Runtime/SlateCore/Public/Input/Events.h
+++ b/Engine/Source/Runtime/SlateCore/Public/Input/Events.h
@@ -614,7 +614,7 @@ public:
  * It is passed to event handlers dealing with pointer-based input.
  */
 USTRUCT(BlueprintType)
-struct FPointerEvent
+struct SLATECORE_API FPointerEvent
 	: public FInputEvent
 {
 	GENERATED_USTRUCT_BODY()
@@ -899,9 +899,9 @@ public:
 		bIsTouchFirstMove = Other.bIsTouchFirstMove;
 	}
 
-	SLATECORE_API virtual FText ToText() const override;
+	virtual FText ToText() const override;
 
-	SLATECORE_API virtual bool IsPointerEvent() const override;
+	virtual bool IsPointerEvent() const override;
 
 	template<typename PointerEventType>
 	static PointerEventType MakeTranslatedEvent( const PointerEventType& InPointerEvent, const FVirtualPointerPosition& VirtualPosition )
