Tree Parent Row Styling (first level only)

diff --git a/Engine/Source/Runtime/Slate/Public/Widgets/Views/STableRow.h b/Engine/Source/Runtime/Slate/Public/Widgets/Views/STableRow.h
--- a/Engine/Source/Runtime/Slate/Public/Widgets/Views/STableRow.h
+++ b/Engine/Source/Runtime/Slate/Public/Widgets/Views/STableRow.h
@@ -925,6 +925,8 @@ public:
 
 		const bool bIsActive = OwnerTable->AsWidget()->HasKeyboardFocus();
 
+		const bool bItemHasChildren = OwnerTable->Private_DoesItemHaveChildren( IndexInList );
+
 		static FName GenericWhiteBoxBrush("GenericWhiteBox");
 
 		// @todo: Slate Style - make this part of the widget style
@@ -968,6 +970,12 @@ public:
 						: &Style->InactiveHighlightedBrush;
 				}
 			}
+			else if (bItemHasChildren && Style->bUseParentRowBrush && GetIndentLevel() == 0)
+			{
+				return IsHovered() 
+				? &Style->ParentRowBackgroundHoveredBrush	
+				: &Style->ParentRowBackgroundBrush;	
+			}
 			else
 			{
 				// Add a slightly lighter background for even rows
diff --git a/Engine/Source/Runtime/SlateCore/Private/Styling/SlateTypes.cpp b/Engine/Source/Runtime/SlateCore/Private/Styling/SlateTypes.cpp
--- a/Engine/Source/Runtime/SlateCore/Private/Styling/SlateTypes.cpp
+++ b/Engine/Source/Runtime/SlateCore/Private/Styling/SlateTypes.cpp
@@ -539,7 +539,8 @@ const FTableViewStyle& FTableViewStyle::GetDefault()
 }
 
 FTableRowStyle::FTableRowStyle()
-	: TextColor(FSlateColor::UseForeground())
+	: bUseParentRowBrush(false)
+	, TextColor(FSlateColor::UseForeground())
 	, SelectedTextColor(FLinearColor::White)
 {
 }
@@ -551,6 +552,8 @@ void FTableRowStyle::GetResources( TArray< const FSlateBrush* >& OutBrushes ) co
 	OutBrushes.Add( &ActiveBrush );
 	OutBrushes.Add( &InactiveHoveredBrush );
 	OutBrushes.Add( &InactiveBrush );
+	OutBrushes.Add( &ParentRowBackgroundHoveredBrush );
+	OutBrushes.Add( &ParentRowBackgroundBrush );
 	OutBrushes.Add( &EvenRowBackgroundHoveredBrush );
 	OutBrushes.Add( &EvenRowBackgroundBrush );
 	OutBrushes.Add( &OddRowBackgroundHoveredBrush );
diff --git a/Engine/Source/Runtime/SlateCore/Private/Styling/StarshipCoreStyle.cpp b/Engine/Source/Runtime/SlateCore/Private/Styling/StarshipCoreStyle.cpp
--- a/Engine/Source/Runtime/SlateCore/Private/Styling/StarshipCoreStyle.cpp
+++ b/Engine/Source/Runtime/SlateCore/Private/Styling/StarshipCoreStyle.cpp
@@ -1757,6 +1757,13 @@ void FStarshipCoreStyle::SetupTableViewStyles(TSharedRef<FStyle>& Style)
 
 	Style->Set("SimpleTableView.Row", SimpleTableRowStyle);
 
+	const FTableRowStyle ShowParentsTableRowStyle = FTableRowStyle(DefaultTableRowStyle)
+		.SetUseParentRowBrush(true)
+		.SetParentRowBackgroundBrush(FSlateRoundedBoxBrush(FStyleColors::Header, 2.f, FStyleColors::Transparent, 1.f))
+		.SetParentRowBackgroundHoveredBrush(FSlateRoundedBoxBrush(FStyleColors::SelectHover, 2.f));
+
+	Style->Set("ShowParentsTableView.Row", ShowParentsTableRowStyle);
+
 	const FTableRowStyle DarkTableRowStyle = FTableRowStyle(DefaultTableRowStyle)
 		.SetEvenRowBackgroundBrush(IMAGE_BRUSH("Common/Selection", Icon8x8, FLinearColor(0.0f, 0.0f, 0.0f, 0.1f)))
 		.SetOddRowBackgroundBrush(IMAGE_BRUSH("Common/Selection", Icon8x8, FLinearColor(0.0f, 0.0f, 0.0f, 0.1f)));
diff --git a/Engine/Source/Runtime/SlateCore/Public/Styling/SlateTypes.h b/Engine/Source/Runtime/SlateCore/Public/Styling/SlateTypes.h
--- a/Engine/Source/Runtime/SlateCore/Public/Styling/SlateTypes.h
+++ b/Engine/Source/Runtime/SlateCore/Public/Styling/SlateTypes.h
@@ -1557,6 +1557,21 @@ struct SLATECORE_API FTableRowStyle : public FSlateWidgetStyle
 	FSlateBrush InactiveBrush;
 	FTableRowStyle& SetInactiveBrush( const FSlateBrush& InInactiveBrush ){ InactiveBrush = InInactiveBrush; return *this; }
 
+	/** If using parent row brushes */
+	UPROPERTY(EditAnywhere, BlueprintReadWrite, Category = Appearance)
+	bool bUseParentRowBrush;
+	FTableRowStyle& SetUseParentRowBrush(bool InUseParentRowBrush) { bUseParentRowBrush = InUseParentRowBrush; return *this; }
+
+	/** Brush used for the top parent row  */
+	UPROPERTY(EditAnywhere, BlueprintReadWrite, Category=Appearance)
+	FSlateBrush ParentRowBackgroundBrush;
+	FTableRowStyle& SetParentRowBackgroundBrush( const FSlateBrush& InParentRowBackgroundBrush ){ ParentRowBackgroundBrush = InParentRowBackgroundBrush; return *this; }
+
+	/** Brush used for the top parent row and row is hovered */
+	UPROPERTY(EditAnywhere, BlueprintReadWrite, Category=Appearance)
+	FSlateBrush ParentRowBackgroundHoveredBrush;
+	FTableRowStyle& SetParentRowBackgroundHoveredBrush( const FSlateBrush& InParentRowBackgroundHoveredBrush ){ ParentRowBackgroundHoveredBrush = InParentRowBackgroundHoveredBrush; return *this; }
+
 	/** Brush used when an even row is hovered */
 	UPROPERTY(EditAnywhere, BlueprintReadWrite, Category=Appearance)
 	FSlateBrush EvenRowBackgroundHoveredBrush;
