Reduce UI memory regression caused by FVector2D conversion to double.
#rnx

diff --git a/Engine/Source/Developer/TraceInsights/Private/Insights/ViewModels/GraphTrack.cpp b/Engine/Source/Developer/TraceInsights/Private/Insights/ViewModels/GraphTrack.cpp
--- a/Engine/Source/Developer/TraceInsights/Private/Insights/ViewModels/GraphTrack.cpp
+++ b/Engine/Source/Developer/TraceInsights/Private/Insights/ViewModels/GraphTrack.cpp
@@ -249,8 +249,8 @@ void FGraphTrack::DrawSeries(const FGraphSeries& Series, FDrawContext& DrawConte
 		FSlateResourceHandle ResourceHandle = FSlateApplication::Get().GetRenderer()->GetResourceHandle(*WhiteBrush);
 		const FSlateShaderResourceProxy* ResourceProxy = ResourceHandle.GetResourceProxy();
 
-		FVector2D AtlasOffset = ResourceProxy ? ResourceProxy->StartUV : FVector2D(0.f, 0.f);
-		FVector2D AtlasUVSize = ResourceProxy ? ResourceProxy->SizeUV : FVector2D(1.f, 1.f);
+		FVector2f AtlasOffset = ResourceProxy ? ResourceProxy->StartUV : FVector2f(0.f, 0.f);
+		FVector2f AtlasUVSize = ResourceProxy ? ResourceProxy->SizeUV : FVector2f(1.f, 1.f);
 
 		const FVector2D Size = DrawContext.Geometry.GetLocalSize();
 
diff --git a/Engine/Source/Editor/Sequencer/Private/SSequencerSection.cpp b/Engine/Source/Editor/Sequencer/Private/SSequencerSection.cpp
--- a/Engine/Source/Editor/Sequencer/Private/SSequencerSection.cpp
+++ b/Engine/Source/Editor/Sequencer/Private/SSequencerSection.cpp
@@ -435,8 +435,8 @@ struct FSequencerSectionPainterImpl : FSequencerSectionPainter
 		FSlateResourceHandle ResourceHandle = FSlateApplication::Get().GetRenderer()->GetResourceHandle(*MyBrush);
 		const FSlateShaderResourceProxy* ResourceProxy = ResourceHandle.GetResourceProxy();
 
-		FVector2D AtlasOffset = ResourceProxy ? ResourceProxy->StartUV : FVector2D(0.f, 0.f);
-		FVector2D AtlasUVSize = ResourceProxy ? ResourceProxy->SizeUV : FVector2D(1.f, 1.f);
+		FVector2f AtlasOffset = ResourceProxy ? ResourceProxy->StartUV : FVector2f(0.f, 0.f);
+		FVector2f AtlasUVSize = ResourceProxy ? ResourceProxy->SizeUV : FVector2f(1.f, 1.f);
 
 		FSlateRenderTransform RenderTransform;
 
@@ -483,13 +483,13 @@ struct FSequencerSectionPainterImpl : FSequencerSectionPainter
 				float U = (Point.Location.X - SegmentStartTime) / ( FFrameNumber(UE::MovieScene::DiscreteSize(Segment.Range)) / TimeToPixelConverter.GetTickResolution() );
 
 				// Add verts top->bottom
-				FVector2D UV(U, 0.f);
+				FVector2f UV(U, 0.f);
 				Verts.Add(FSlateVertex::Make<ESlateVertexRounding::Disabled>(RenderTransform, (Pos + UV*Size*RangeGeometry.Scale), AtlasOffset + UV*AtlasUVSize, FillColor));
 
 				UV.Y = 1.f - Point.Location.Y;
 				BorderPoints.Add(UV*Size);
 				BorderPointColors.Add(Point.Color);
-				Verts.Add(FSlateVertex::Make<ESlateVertexRounding::Disabled>(RenderTransform, (Pos + UV*Size*RangeGeometry.Scale), AtlasOffset + FVector2D(UV.X, 0.5f)*AtlasUVSize, FillColor));
+				Verts.Add(FSlateVertex::Make<ESlateVertexRounding::Disabled>(RenderTransform, (Pos + UV*Size*RangeGeometry.Scale), AtlasOffset + FVector2f(UV.X, 0.5f)*AtlasUVSize, FillColor));
 
 				if (Verts.Num() >= 4)
 				{
diff --git a/Engine/Source/Runtime/Core/Public/Math/TransformCalculus2D.h b/Engine/Source/Runtime/Core/Public/Math/TransformCalculus2D.h
--- a/Engine/Source/Runtime/Core/Public/Math/TransformCalculus2D.h
+++ b/Engine/Source/Runtime/Core/Public/Math/TransformCalculus2D.h
@@ -42,7 +42,8 @@ namespace Math
 {
 
 /** Specialization for concatenating two 2D Translations. */
-inline FVector2D Concatenate(const FVector2D& LHS, const FVector2D& RHS)
+template<typename T>
+inline UE::Math::TVector2<T> Concatenate(const UE::Math::TVector2<T>& LHS, const UE::Math::TVector2<T>& RHS)
 {
 	return LHS + RHS;
 }
@@ -51,7 +52,8 @@ inline FVector2D Concatenate(const FVector2D& LHS, const FVector2D& RHS)
 }	// namespace UE
 
 /** Specialization for inverting a 2D translation. */
-inline FVector2D Inverse(const FVector2D& Transform)
+template<typename T>
+inline UE::Math::TVector2<T> Inverse(const UE::Math::TVector2<T>& Transform)
 {
 	return -Transform;
 }
@@ -592,7 +594,7 @@ public:
 	 */
 	FVector2D TransformPoint(const FVector2D& Point) const
 	{
-		return ::TransformPoint(Trans, ::TransformPoint(M, Point));
+		return ::TransformPoint(GetTranslation(), ::TransformPoint(M, Point));
 	}
 
 	/**
@@ -648,7 +650,7 @@ public:
 	FTransform2D Inverse() const
 	{
 		FMatrix2x2 InvM = ::Inverse(M);
-		FVector2D InvTrans = ::TransformPoint(InvM, ::Inverse(Trans));
+		FVector2D InvTrans = ::TransformPoint(InvM, ::Inverse(GetTranslation()));
 		return FTransform2D(InvM, InvTrans);
 	}
 	
@@ -667,7 +669,7 @@ public:
 	/** Access to the 2x2 transform */
 	const FMatrix2x2& GetMatrix() const { return M; }
 	/** Access to the translation */
-	const FVector2D& GetTranslation() const { return Trans; }
+	const FVector2D GetTranslation() const { return FVector2D(Trans); }
 
 	void SetTranslation(const FVector2D& InTrans) { Trans = InTrans; }
 
@@ -697,7 +699,7 @@ public:
 
 private:
 	FMatrix2x2 M;
-	FVector2D Trans;
+	FVector2f Trans;
 };
 
 template<> struct TIsPODType<FTransform2D> { enum { Value = true }; };
diff --git a/Engine/Source/Runtime/Slate/Private/Widgets/Text/STextBlock.cpp b/Engine/Source/Runtime/Slate/Private/Widgets/Text/STextBlock.cpp
--- a/Engine/Source/Runtime/Slate/Private/Widgets/Text/STextBlock.cpp
+++ b/Engine/Source/Runtime/Slate/Private/Widgets/Text/STextBlock.cpp
@@ -278,8 +278,8 @@ FVector2D STextBlock::ComputeDesiredSize(float LayoutScaleMultiplier) const
 		const FVector2D ComputedOutlineSize(LocalOutlineSize * 2.f, LocalOutlineSize);
 		const FVector2D TextSize = FSlateApplication::Get().GetRenderer()->GetFontMeasureService()->Measure(BoundText.Get(), GetFont()) + ComputedOutlineSize + LocalShadowOffset;
 
-		CachedSimpleDesiredSize = FVector2D(FMath::Max(MinDesiredWidth.Get(), TextSize.X), TextSize.Y);
-		return CachedSimpleDesiredSize.GetValue();
+		CachedSimpleDesiredSize = FVector2f(FMath::Max(MinDesiredWidth.Get(), TextSize.X), TextSize.Y);
+		return FVector2D(CachedSimpleDesiredSize.GetValue());
 	}
 	else
 	{
diff --git a/Engine/Source/Runtime/Slate/Public/Widgets/Text/ISlateEditableTextWidget.h b/Engine/Source/Runtime/Slate/Public/Widgets/Text/ISlateEditableTextWidget.h
--- a/Engine/Source/Runtime/Slate/Public/Widgets/Text/ISlateEditableTextWidget.h
+++ b/Engine/Source/Runtime/Slate/Public/Widgets/Text/ISlateEditableTextWidget.h
@@ -127,7 +127,7 @@ private:
 
 	ECursorMoveGranularity Granularity;
 	ECursorMoveMethod Method;
-	FVector2D DirectionOrPosition;
+	FVector2f DirectionOrPosition;
 	ECursorAction Action;
 	float GeometryScale;
 };
diff --git a/Engine/Source/Runtime/Slate/Public/Widgets/Text/STextBlock.h b/Engine/Source/Runtime/Slate/Public/Widgets/Text/STextBlock.h
--- a/Engine/Source/Runtime/Slate/Public/Widgets/Text/STextBlock.h
+++ b/Engine/Source/Runtime/Slate/Public/Widgets/Text/STextBlock.h
@@ -348,7 +348,7 @@ private:
 	TSlateAttribute<float> MinDesiredWidth;
 
 	/** If this is enabled, text shaping, wrapping, justification are disabled in favor of much faster text layout and measurement. */
-	mutable TOptional<FVector2D> CachedSimpleDesiredSize;
+	mutable TOptional<FVector2f> CachedSimpleDesiredSize;
 
 	/** Flags used to check if the SlateAttribute is set. */
 	union
diff --git a/Engine/Source/Runtime/Slate/Public/Widgets/Text/SlateEditableTextLayout.h b/Engine/Source/Runtime/Slate/Public/Widgets/Text/SlateEditableTextLayout.h
--- a/Engine/Source/Runtime/Slate/Public/Widgets/Text/SlateEditableTextLayout.h
+++ b/Engine/Source/Runtime/Slate/Public/Widgets/Text/SlateEditableTextLayout.h
@@ -589,7 +589,7 @@ private:
 	TArray<FTextLineHighlight> ActiveLineHighlights;
 
 	/** The scroll offset (in unscaled Slate units) for this text */
-	FVector2D ScrollOffset;
+	FVector2f ScrollOffset;
 
 	/** If set, the pending data containing a position that should be scrolled into view */
 	TOptional<SlateEditableTextTypes::FScrollInfo> PositionToScrollIntoView;
@@ -640,7 +640,7 @@ private:
 	TOptional<ETextOverflowPolicy> OverflowPolicyOverride;
 
 	/** The last known size of the widget from the previous OnPaint, used to recalculate wrapping */
-	FVector2D CachedSize;
+	FVector2f CachedSize;
 
 	/** A list commands to execute if a user presses the corresponding key-binding in the text box */
 	TSharedPtr<FUICommandList> UICommandList;
diff --git a/Engine/Source/Runtime/Slate/Public/Widgets/Text/SlateTextBlockLayout.h b/Engine/Source/Runtime/Slate/Public/Widgets/Text/SlateTextBlockLayout.h
--- a/Engine/Source/Runtime/Slate/Public/Widgets/Text/SlateTextBlockLayout.h
+++ b/Engine/Source/Runtime/Slate/Public/Widgets/Text/SlateTextBlockLayout.h
@@ -203,7 +203,7 @@ private:
 	TSharedPtr<ISlateRunRenderer> TextHighlighter;
 
 	/** The last known size of the layout from the previous OnPaint, used to guess at an auto-wrapping width in ComputeDesiredSize */
-	FVector2D CachedSize;
+	FVector2f CachedSize;
 
 	/** Cache where to wrap text at? */
 	float CachedWrapTextAt;
diff --git a/Engine/Source/Runtime/SlateCore/Private/Layout/Clipping.cpp b/Engine/Source/Runtime/SlateCore/Private/Layout/Clipping.cpp
--- a/Engine/Source/Runtime/SlateCore/Private/Layout/Clipping.cpp
+++ b/Engine/Source/Runtime/SlateCore/Private/Layout/Clipping.cpp
@@ -107,17 +107,17 @@ void FSlateClippingZone::InitializeFromArbitraryPoints(const FVector2D& InTopLef
 		const float Top = FMath::Min(RoundedAxisAlignedRect.Top, RoundedAxisAlignedRect.Bottom);
 		const float Bottom = FMath::Max(RoundedAxisAlignedRect.Top, RoundedAxisAlignedRect.Bottom);
 
-		TopLeft = FVector2D(Left, Top);
-		TopRight = FVector2D(Right, Top);
-		BottomLeft = FVector2D(Left, Bottom);
-		BottomRight = FVector2D(Right, Bottom);
+		TopLeft = FVector2f(Left, Top);
+		TopRight = FVector2f(Right, Top);
+		BottomLeft = FVector2f(Left, Bottom);
+		BottomRight = FVector2f(Right, Bottom);
 	}
 	else
 	{
-		TopLeft = InTopLeft;
-		TopRight = InTopRight;
-		BottomLeft = InBottomLeft;
-		BottomRight = InBottomRight;
+		TopLeft = FVector2f(InTopLeft);
+		TopRight = FVector2f(InTopRight);
+		BottomLeft = FVector2f(InBottomLeft);
+		BottomRight = FVector2f(InBottomRight);
 	}
 }
 
@@ -144,7 +144,7 @@ FSlateClippingZone FSlateClippingZone::Intersect(const FSlateClippingZone& Other
 
 FSlateRect FSlateClippingZone::GetBoundingBox() const
 {
-	FVector2D Points[4] =
+	FVector2f Points[4] =
 	{
 		TopLeft,
 		TopRight,
diff --git a/Engine/Source/Runtime/SlateCore/Public/Layout/Clipping.h b/Engine/Source/Runtime/SlateCore/Public/Layout/Clipping.h
--- a/Engine/Source/Runtime/SlateCore/Public/Layout/Clipping.h
+++ b/Engine/Source/Runtime/SlateCore/Public/Layout/Clipping.h
@@ -61,10 +61,10 @@ class SLATECORE_API FSlateClippingZone
 {
 public:
 
-	FVector2D TopLeft;
-	FVector2D TopRight;
-	FVector2D BottomLeft;
-	FVector2D BottomRight;
+	FVector2f TopLeft;
+	FVector2f TopRight;
+	FVector2f BottomLeft;
+	FVector2f BottomRight;
 
 	explicit FSlateClippingZone(const FShortRect& AxisAlignedRect);
 	explicit FSlateClippingZone(const FSlateRect& AxisAlignedRect);
diff --git a/Engine/Source/Runtime/SlateCore/Public/Layout/Geometry.h b/Engine/Source/Runtime/SlateCore/Public/Layout/Geometry.h
--- a/Engine/Source/Runtime/SlateCore/Public/Layout/Geometry.h
+++ b/Engine/Source/Runtime/SlateCore/Public/Layout/Geometry.h
@@ -67,9 +67,9 @@ public:
 		FSlateLayoutTransform AccumulatedLayoutTransform = Concatenate(LocalLayoutTransform, ParentAccumulatedLayoutTransform);
 		AccumulatedRenderTransform = TransformCast<FSlateRenderTransform>(AccumulatedLayoutTransform);
 		// HACK to allow us to make FGeometry public members immutable to catch misuse.
-		const_cast<FVector2D&>( AbsolutePosition ) = AccumulatedLayoutTransform.GetTranslation();
+		const_cast<FVector2f&>( AbsolutePosition ) = AccumulatedLayoutTransform.GetTranslation();
 		const_cast<float&>( Scale ) = AccumulatedLayoutTransform.GetScale();
-		const_cast<FVector2D&>( Position ) = LocalLayoutTransform.GetTranslation();
+		const_cast<FVector2f&>( Position ) = LocalLayoutTransform.GetTranslation();
 	}
 
 private:
@@ -110,9 +110,9 @@ private:
 	{
 		FSlateLayoutTransform AccumulatedLayoutTransform = Concatenate(InLocalLayoutTransform, ParentAccumulatedLayoutTransform);
 		// HACK to allow us to make FGeometry public members immutable to catch misuse.
-		const_cast<FVector2D&>( AbsolutePosition ) = AccumulatedLayoutTransform.GetTranslation();
+		const_cast<FVector2f&>( AbsolutePosition ) = AccumulatedLayoutTransform.GetTranslation();
 		const_cast<float&>( Scale ) = AccumulatedLayoutTransform.GetScale();
-		const_cast<FVector2D&>( Position ) = InLocalLayoutTransform.GetTranslation();
+		const_cast<FVector2f&>( Position ) = InLocalLayoutTransform.GetTranslation();
 	}
 
 	/**
@@ -137,9 +137,9 @@ private:
 	{
 		FSlateLayoutTransform AccumulatedLayoutTransform = Concatenate(InLocalLayoutTransform, ParentAccumulatedLayoutTransform);
 		// HACK to allow us to make FGeometry public members immutable to catch misuse.
-		const_cast<FVector2D&>( AbsolutePosition ) = AccumulatedLayoutTransform.GetTranslation();
+		const_cast<FVector2f&>( AbsolutePosition ) = AccumulatedLayoutTransform.GetTranslation();
 		const_cast<float&>( Scale ) = AccumulatedLayoutTransform.GetScale();
-		const_cast<FVector2D&>( Position ) = InLocalLayoutTransform.GetTranslation();
+		const_cast<FVector2f&>( Position ) = InLocalLayoutTransform.GetTranslation();
 	}
 
 public:
@@ -523,7 +523,7 @@ public:
 	{
 		FSlateLayoutTransform AccumulatedLayoutTransform = ::Concatenate(GetAccumulatedLayoutTransform(), LayoutTransform);
 		AccumulatedRenderTransform = ::Concatenate(AccumulatedRenderTransform, LayoutTransform);
-		const_cast<FVector2D&>( AbsolutePosition ) = AccumulatedLayoutTransform.GetTranslation();
+		const_cast<FVector2f&>( AbsolutePosition ) = AccumulatedLayoutTransform.GetTranslation();
 		const_cast<float&>( Scale ) = AccumulatedLayoutTransform.GetScale();
 	}
 
@@ -608,7 +608,7 @@ public:
 	 * 
 	 * Absolute coordinates could be either desktop or window space depending on what space the root of the widget hierarchy is in.
 	 */
-	const FVector2D AbsolutePosition;	
+	const FVector2f AbsolutePosition;	
 
 	/** 
 	 * !!! DEPRECATED !!! 
@@ -617,7 +617,7 @@ public:
 	 * If you know your children have no additional scale applied to them, you can use this as the Local->Parent layout transform. If your children
 	 * DO have additional scale applied, there is no way to determine the actual Local->Parent layout transform, since the scale is accumulated.
 	 */
-	const FVector2D /*Local*/Position;
+	const FVector2f /*Local*/Position;
 
 private:
 
diff --git a/Engine/Source/Runtime/SlateCore/Public/Layout/LayoutGeometry.h b/Engine/Source/Runtime/SlateCore/Public/Layout/LayoutGeometry.h
--- a/Engine/Source/Runtime/SlateCore/Public/Layout/LayoutGeometry.h
+++ b/Engine/Source/Runtime/SlateCore/Public/Layout/LayoutGeometry.h
@@ -16,7 +16,7 @@ public:
 
 	explicit FLayoutGeometry(const FSlateLayoutTransform& InLocalToParent, const FVector2D& SizeInLocalSpace)
 		: LocalToParent(InLocalToParent)
-		, LocalSize(SizeInLocalSpace)
+		, LocalSize(FVector2f(SizeInLocalSpace))
 	{
 	}
 
@@ -25,7 +25,7 @@ public:
 		return LocalToParent; 
 	}
 
-	const FVector2D& GetSizeInLocalSpace() const
+	FVector2D GetSizeInLocalSpace() const
 	{
 		return LocalSize;
 	}
@@ -73,6 +73,6 @@ public:
 
 private:
 	FSlateLayoutTransform LocalToParent;
-	FVector2D LocalSize;
+	FVector2f LocalSize;
 };
 
diff --git a/Engine/Source/Runtime/SlateCore/Public/Layout/PaintGeometry.h b/Engine/Source/Runtime/SlateCore/Public/Layout/PaintGeometry.h
--- a/Engine/Source/Runtime/SlateCore/Public/Layout/PaintGeometry.h
+++ b/Engine/Source/Runtime/SlateCore/Public/Layout/PaintGeometry.h
@@ -29,7 +29,7 @@ struct SLATECORE_API FPaintGeometry
 	 *
 	 * 
 	 */
-	FVector2D DrawPosition;
+	FVector2f DrawPosition;
 
 	/**
 	 * !!! DEPRECATED!!! Drawing should only happen in local space to ensure render transforms work.
@@ -63,7 +63,7 @@ struct SLATECORE_API FPaintGeometry
 
 private:
 	// Mutable to support legacy constructors. Doesn't account for render transforms.
-	mutable FVector2D DrawSize;
+	mutable FVector2f DrawSize;
 
 	// Mutable to support legacy constructors.
 	mutable FVector2D LocalSize;
diff --git a/Engine/Source/Runtime/SlateCore/Public/Layout/SlateRotatedRect.h b/Engine/Source/Runtime/SlateCore/Public/Layout/SlateRotatedRect.h
--- a/Engine/Source/Runtime/SlateCore/Public/Layout/SlateRotatedRect.h
+++ b/Engine/Source/Runtime/SlateCore/Public/Layout/SlateRotatedRect.h
@@ -37,11 +37,11 @@ public:
 public:
 
 	/** transformed Top-left corner. */
-	FVector2D TopLeft;
+	FVector2f TopLeft;
 	/** transformed X extent (right-left). */
-	FVector2D ExtentX;
+	FVector2f ExtentX;
 	/** transformed Y extent (bottom-top). */
-	FVector2D ExtentY;
+	FVector2f ExtentY;
 
 public:
 	bool operator == (const FSlateRotatedRect& Other) const
@@ -86,8 +86,8 @@ FSlateRotatedRect TransformRect(const TransformType& Transform, const FSlateRota
 {
 	return FSlateRotatedRect
 	(
-		TransformPoint(Transform, Rect.TopLeft),
-		TransformVector(Transform, Rect.ExtentX),
-		TransformVector(Transform, Rect.ExtentY)
+		TransformPoint(Transform, FVector2D(Rect.TopLeft)),
+		TransformVector(Transform, FVector2D(Rect.ExtentX)),
+		TransformVector(Transform, FVector2D(Rect.ExtentY))
 	);
 }
diff --git a/Engine/Source/Runtime/SlateCore/Public/Rendering/DrawElements.h b/Engine/Source/Runtime/SlateCore/Public/Rendering/DrawElements.h
--- a/Engine/Source/Runtime/SlateCore/Public/Rendering/DrawElements.h
+++ b/Engine/Source/Runtime/SlateCore/Public/Rendering/DrawElements.h
@@ -269,9 +269,9 @@ public:
 
 	FORCEINLINE const FSlateRenderTransform& GetRenderTransform() const { return RenderTransform; }
 	FORCEINLINE void SetRenderTransform(const FSlateRenderTransform& InRenderTransform) { RenderTransform = InRenderTransform; }
-	FORCEINLINE const FVector2D& GetPosition() const { return Position; }
+	FORCEINLINE FVector2D GetPosition() const { return FVector2D(Position); }
 	FORCEINLINE void SetPosition(const FVector2D& InPosition) { Position = InPosition; }
-	FORCEINLINE const FVector2D& GetLocalSize() const { return LocalSize; }
+	FORCEINLINE FVector2D GetLocalSize() const { return FVector2D(LocalSize); }
 	FORCEINLINE float GetScale() const { return Scale; }
 	FORCEINLINE ESlateDrawEffect GetDrawEffects() const { return DrawEffects; }
 	FORCEINLINE ESlateBatchDrawFlag GetBatchFlags() const { return BatchFlags; }
@@ -316,8 +316,8 @@ private:
 private:
 	FSlateDataPayload* DataPayload;
 	FSlateRenderTransform RenderTransform;
-	FVector2D Position;
-	FVector2D LocalSize;
+	FVector2f Position;
+	FVector2f LocalSize;
 	int32 LayerId;
 	FClipStateHandle ClipStateHandle;
 	float Scale;
diff --git a/Engine/Source/Runtime/SlateCore/Public/Rendering/RenderingCommon.h b/Engine/Source/Runtime/SlateCore/Public/Rendering/RenderingCommon.h
--- a/Engine/Source/Runtime/SlateCore/Public/Rendering/RenderingCommon.h
+++ b/Engine/Source/Runtime/SlateCore/Public/Rendering/RenderingCommon.h
@@ -223,7 +223,7 @@ struct SLATECORE_API FSlateVertex
 public:
 
 	template<ESlateVertexRounding Rounding>
-	static FSlateVertex Make(const FSlateRenderTransform& RenderTransform, const FVector2D InLocalPosition, const FVector2D InTexCoord, const FVector2D InTexCoord2, const FColor InColor, const FColor SecondaryColor = FColor())
+	static FSlateVertex Make(const FSlateRenderTransform& RenderTransform, const FVector2f InLocalPosition, const FVector2f InTexCoord, const FVector2f InTexCoord2, const FColor InColor, const FColor SecondaryColor = FColor())
 	{
 		FSlateVertex Vertex;
 		Vertex.TexCoords[0] = InTexCoord.X;
@@ -236,7 +236,7 @@ public:
 	}
 
 	template<ESlateVertexRounding Rounding>
-	static FSlateVertex Make(const FSlateRenderTransform& RenderTransform, const FVector2D InLocalPosition, const FVector2D InTexCoord, const FColor& InColor, const FColor SecondaryColor = FColor())
+	static FSlateVertex Make(const FSlateRenderTransform& RenderTransform, const FVector2f InLocalPosition, const FVector2f InTexCoord, const FColor& InColor, const FColor SecondaryColor = FColor())
 	{
 		FSlateVertex Vertex;
 		Vertex.TexCoords[0] = InTexCoord.X;
@@ -249,7 +249,7 @@ public:
 	}
 
 	template<ESlateVertexRounding Rounding>
-	static FSlateVertex Make(const FSlateRenderTransform& RenderTransform, const FVector2D InLocalPosition, const FVector4f InTexCoords, const FVector2D InMaterialTexCoords, const FColor InColor, const FColor SecondaryColor = FColor())
+	static FSlateVertex Make(const FSlateRenderTransform& RenderTransform, const FVector2f InLocalPosition, const FVector4f InTexCoords, const FVector2f InMaterialTexCoords, const FColor InColor, const FColor SecondaryColor = FColor())
 	{
 		FSlateVertex Vertex;
 		Vertex.TexCoords[0] = InTexCoords.X;
@@ -263,14 +263,14 @@ public:
 	}
 
 	template<ESlateVertexRounding Rounding>
-	static FSlateVertex Make(const FSlateRenderTransform& RenderTransform, const FVector2D InLocalPosition, const FVector2D InLocalSize, float Scale, const FVector4f InTexCoords, const FColor InColor, const FColor SecondaryColor = FColor())
+	static FSlateVertex Make(const FSlateRenderTransform& RenderTransform, const FVector2f InLocalPosition, const FVector2f InLocalSize, float Scale, const FVector4f InTexCoords, const FColor InColor, const FColor SecondaryColor = FColor())
 	{
 		FSlateVertex Vertex;
 		Vertex.TexCoords[0] = InTexCoords.X;
 		Vertex.TexCoords[1] = InTexCoords.Y;
 		Vertex.TexCoords[2] = InTexCoords.Z;
 		Vertex.TexCoords[3] = InTexCoords.W;
-		Vertex.MaterialTexCoords = FVector2D(InLocalPosition.X / InLocalSize.X, InLocalPosition.Y / InLocalSize.Y);
+		Vertex.MaterialTexCoords = FVector2f(InLocalPosition.X / InLocalSize.X, InLocalPosition.Y / InLocalSize.Y);
 		Vertex.InitCommon<Rounding>(RenderTransform, InLocalPosition, InColor, SecondaryColor);
 
 		const int32 PixelSizeX = FMath::RoundToInt(InLocalSize.X * Scale);
@@ -284,7 +284,7 @@ public:
 		ensureMsgf((int32)Vertex.PixelSize[1] == PixelSizeY, TEXT("Conversion of PixelSizeY is bigger than 16. Cast:%d, int16:%d, int32:%d")
 			, (int32)Vertex.PixelSize[1], Vertex.PixelSize[1], PixelSizeY);
 #endif
-		Vertex.MaterialTexCoords = FVector2D(InLocalPosition.X / InLocalSize.X, InLocalPosition.Y / InLocalSize.Y);
+		Vertex.MaterialTexCoords = FVector2f(InLocalPosition.X / InLocalSize.X, InLocalPosition.Y / InLocalSize.Y);
 
 		return Vertex;
 	}
@@ -297,7 +297,7 @@ public:
 		TexCoords[3] = InTexCoords.W;
 	}
 
-	void SetPosition(const FVector2D InPosition)
+	void SetPosition(const FVector2f InPosition)
 	{
 		Position = InPosition;
 	}
@@ -305,7 +305,7 @@ public:
 private:
 
 	template<ESlateVertexRounding Rounding>
-	FORCEINLINE void InitCommon(const FSlateRenderTransform& RenderTransform, const FVector2D InLocalPosition, const FColor InColor, const FColor InSecondaryColor)
+	FORCEINLINE void InitCommon(const FSlateRenderTransform& RenderTransform, const FVector2f InLocalPosition, const FColor InColor, const FColor InSecondaryColor)
 	{
 		Position = TransformPoint(RenderTransform, InLocalPosition);
 
@@ -370,8 +370,8 @@ struct FShortRect
 		return !bDoNotOverlap;
 	}
 
-	FVector2D GetTopLeft() const { return FVector2D(Left, Top); }
-	FVector2D GetBottomRight() const { return FVector2D(Right, Bottom); }
+	FVector2f GetTopLeft() const { return FVector2f(Left, Top); }
+	FVector2f GetBottomRight() const { return FVector2f(Right, Bottom); }
 
 	uint16 Left;
 	uint16 Top;
diff --git a/Engine/Source/Runtime/SlateCore/Public/Rendering/SlateLayoutTransform.h b/Engine/Source/Runtime/SlateCore/Public/Rendering/SlateLayoutTransform.h
--- a/Engine/Source/Runtime/SlateCore/Public/Rendering/SlateLayoutTransform.h
+++ b/Engine/Source/Runtime/SlateCore/Public/Rendering/SlateLayoutTransform.h
@@ -21,21 +21,21 @@ public:
 	/** Ctor from a scale followed by translate. Shortcut to Concatenate(InScale, InTranslation). */
 	explicit FSlateLayoutTransform(float InScale = 1.0f, const FVector2D& InTranslation = FVector2D(ForceInit))
 		:Scale(InScale)
-		,Translation(InTranslation)
+		,Translation(FVector2f(InTranslation))
 	{
 	}
 
 	/** Ctor from a 2D translation followed by a scale. Shortcut to Concatenate(InTranslation, InScale). While this is the opposite order we internally store them, we can represent this correctly. */
 	explicit FSlateLayoutTransform(const FVector2D& InTranslation/*, float InScale = 1.0f*/)
 		:Scale(1.0f/*InScale*/)
-		,Translation(InTranslation/**InScale*/)
+		,Translation(FVector2f(InTranslation)/**InScale*/)
 	{
 	}
 
 	/** Access to the 2D translation */
-	const FVector2D& GetTranslation() const
+	FVector2D GetTranslation() const
 	{
-		return Translation;
+		return FVector2D(Translation);
 	}
 
 	/** Access to the scale. */
@@ -55,13 +55,13 @@ public:
 	/** 2D transform support. */
 	FVector2D TransformPoint(const FVector2D& Point) const
 	{
-		return ::TransformPoint(Translation, ::TransformPoint(Scale, Point));
+		return ::TransformPoint(GetTranslation(), ::TransformPoint(Scale, Point));
 	}
 
 	/** 2D transform support. */
 	FVector2D TransformVector(const FVector2D& Vector) const
 	{
-		return ::TransformVector(Translation, ::TransformVector(Scale, Vector));
+		return ::TransformVector(GetTranslation(), ::TransformVector(Scale, Vector));
 	}
 
 	/**
@@ -75,7 +75,7 @@ public:
 	{
 		// New Translation is essentially: RHS.TransformPoint(TransformPoint(FVector2D::ZeroVector))
 		// Since Zero through LHS -> Translation we optimize it slightly to skip the zero multiplies.
-		return FSlateLayoutTransform(::Concatenate(Scale, RHS.Scale), RHS.TransformPoint(Translation));
+		return FSlateLayoutTransform(::Concatenate(Scale, RHS.Scale), RHS.TransformPoint(GetTranslation()));
 	}
 
 	/** Invert the transform/scale. */
@@ -98,7 +98,7 @@ public:
 
 private:
 	float Scale;
-	FVector2D Translation;
+	FVector2f Translation;
 };
 
 /** Specialization for concatenating a uniform scale and 2D Translation. */
diff --git a/Engine/Source/Runtime/SlateCore/Public/Styling/SlateBrush.h b/Engine/Source/Runtime/SlateCore/Public/Styling/SlateBrush.h
--- a/Engine/Source/Runtime/SlateCore/Public/Styling/SlateBrush.h
+++ b/Engine/Source/Runtime/SlateCore/Public/Styling/SlateBrush.h
@@ -324,9 +324,9 @@ public:
 	 *
 	 * @return UV region
 	 */
-	const FBox2D& GetUVRegion() const
+	FBox2D GetUVRegion() const
 	{
-		return UVRegion;
+		return FBox2D(UVRegion);
 	}
 
 	/**
@@ -336,7 +336,7 @@ public:
 	 */
 	void SetUVRegion(const FBox2D& InUVRegion)
 	{
-		UVRegion = InUVRegion;
+		UVRegion = FBox2f(InUVRegion);
 	}
 
 	/**
@@ -436,7 +436,7 @@ protected:
 	 *  When valid - overrides UV region specified in resource proxy
 	 */
 	UPROPERTY()
-	FBox2D UVRegion;
+	FBox2f UVRegion;
 
 public:
 	/** How to draw the image */
diff --git a/Engine/Source/Runtime/SlateCore/Public/Textures/SlateShaderResource.h b/Engine/Source/Runtime/SlateCore/Public/Textures/SlateShaderResource.h
--- a/Engine/Source/Runtime/SlateCore/Public/Textures/SlateShaderResource.h
+++ b/Engine/Source/Runtime/SlateCore/Public/Textures/SlateShaderResource.h
@@ -113,10 +113,10 @@ class SLATECORE_API FSlateShaderResourceProxy
 public:
 
 	/** The start uv of the texture.  If atlased this is some subUV of the atlas, 0,0 otherwise */
-	FVector2D StartUV;
+	FVector2f StartUV;
 
 	/** The size of the texture in UV space.  If atlas this some sub uv of the atlas.  1,1 otherwise */
-	FVector2D SizeUV;
+	FVector2f SizeUV;
 
 	/** The resource to be used for rendering */
 	FSlateShaderResource* Resource;
diff --git a/Engine/Source/Runtime/SlateCore/Public/Widgets/SWidget.h b/Engine/Source/Runtime/SlateCore/Public/Widgets/SWidget.h
--- a/Engine/Source/Runtime/SlateCore/Public/Widgets/SWidget.h
+++ b/Engine/Source/Runtime/SlateCore/Public/Widgets/SWidget.h
@@ -1831,7 +1831,7 @@ private:
 	TArray<TSharedRef<FActiveTimerHandle>> ActiveTimers;
 
 	/** Stores the ideal size this widget wants to be. */
-	TOptional<FVector2D> DesiredSize;
+	TOptional<FVector2f> DesiredSize;
 
 	/** Is this widget visible, hidden or collapsed */
 	TSlateAttribute<EVisibility> VisibilityAttribute;
diff --git a/Engine/Source/Runtime/SlateRHIRenderer/Private/SlateRHIRenderingPolicy.cpp b/Engine/Source/Runtime/SlateRHIRenderer/Private/SlateRHIRenderingPolicy.cpp
--- a/Engine/Source/Runtime/SlateRHIRenderer/Private/SlateRHIRenderingPolicy.cpp
+++ b/Engine/Source/Runtime/SlateRHIRenderer/Private/SlateRHIRenderingPolicy.cpp
@@ -315,8 +315,8 @@ static bool UpdateScissorRect(
 				const FVector2D ViewSize((float) SizeXY.X, (float) SizeXY.Y);
 
 				// Clamp scissor rect to BackBuffer size
-				const FVector2D TopLeft     = FMath::Min(FMath::Max(ScissorRect.TopLeft     + ViewTranslation2D, FVector2D(0.0f, 0.0f)), ViewSize);
-				const FVector2D BottomRight = FMath::Min(FMath::Max(ScissorRect.BottomRight + ViewTranslation2D, FVector2D(0.0f, 0.0f)), ViewSize);
+				const FVector2D TopLeft     = FMath::Min(FMath::Max(FVector2D(ScissorRect.TopLeft)     + ViewTranslation2D, FVector2D(0.0f, 0.0f)), ViewSize);
+				const FVector2D BottomRight = FMath::Min(FMath::Max(FVector2D(ScissorRect.BottomRight) + ViewTranslation2D, FVector2D(0.0f, 0.0f)), ViewSize);
 				
 				if (bSwitchVerticalAxis)
 				{
