Slate: Add debug functionnality to the InvalidationWidgetList. Use CVar instead of the reference variable when disabling an InvalidationRoot test.

diff --git a/Engine/Source/Runtime/SlateCore/Private/FastUpdate/SlateInvalidationRoot.cpp b/Engine/Source/Runtime/SlateCore/Private/FastUpdate/SlateInvalidationRoot.cpp
--- a/Engine/Source/Runtime/SlateCore/Private/FastUpdate/SlateInvalidationRoot.cpp
+++ b/Engine/Source/Runtime/SlateCore/Private/FastUpdate/SlateInvalidationRoot.cpp
@@ -72,7 +72,7 @@ bool GSlateInvalidationRootVerifyWidgetList = false;
 static FAutoConsoleVariableRef CVarSlateInvalidationRootVerifyWidgetList(
 	TEXT("Slate.InvalidationRoot.VerifyWidgetList"),
 	GSlateInvalidationRootVerifyWidgetList,
-	TEXT("Each frame, verify that the updated list doesn't match a newly created list.")
+	TEXT("Every tick, verify that the updated list does match a newly created list.")
 );
 void VerifyWidgetList(TSharedRef<SWidget> RootWidget, FSlateInvalidationRootHandle InvalidationRootHandle, FSlateInvalidationWidgetList& WidgetList);
 
@@ -80,29 +80,29 @@ bool GSlateInvalidationRootVerifyWidgetsIndex = false;
 static FAutoConsoleVariableRef CVarSlateInvalidationRootVerifyWidgetsIndex(
 	TEXT("Slate.InvalidationRoot.VerifyWidgetsIndex"),
 	GSlateInvalidationRootVerifyWidgetsIndex,
-	TEXT("Each frame, verify that every widgets has the correct index.")
+	TEXT("Every tick, verify that every widgets has the correct corresponding index.")
 );
 
 bool GSlateInvalidationRootVerifyValidWidgets = false;
 static FAutoConsoleVariableRef CVarSlateInvalidationRootVerifyValidWidgets(
 	TEXT("Slate.InvalidationRoot.VerifyValidWidgets"),
 	GSlateInvalidationRootVerifyValidWidgets,
-	TEXT("Each frame, verify that every WidgetProxy has a valid SWidget.")
+	TEXT("Every tick, verify that every WidgetProxy has a valid SWidget.")
 );
 
 bool GSlateInvalidationRootVerifyHittestGrid = false;
 static FAutoConsoleVariableRef CVarSlateInvalidationRootVerifyHittestGrid(
 	TEXT("Slate.InvalidationRoot.VerifyHittestGrid"),
 	GSlateInvalidationRootVerifyHittestGrid,
-	TEXT("Each frame, verify the hittest grid.")
+	TEXT("Every tick, verify the hittest grid.")
 );
 void VerifyHittest(SWidget* InvalidationRootWidget, FSlateInvalidationWidgetList& WidgetList, FHittestGrid* HittestGrid);
 
 bool GSlateInvalidationRootVerifyWidgetVisibility = false;
-static FAutoConsoleVariableRef CVarSlateInvalidationRootVerifyVisibility(
+static FAutoConsoleVariableRef CVarSlateInvalidationRootVerifyWidgetVisibility(
 	TEXT("Slate.InvalidationRoot.VerifyWidgetVisibility"),
 	GSlateInvalidationRootVerifyWidgetVisibility,
-	TEXT("Each frame, verify that the cached visibility of the widgets is properly set.")
+	TEXT("Every tick, verify that the cached visibility of the widgets is properly set.")
 );
 void VerifyWidgetVisibility(FSlateInvalidationWidgetList& WidgetList);
 
@@ -110,7 +110,7 @@ bool GSlateInvalidationRootVerifyWidgetVolatile = false;
 static FAutoConsoleVariableRef CVarSlateInvalidationRootVerifyWidgetVolatile(
 	TEXT("Slate.InvalidationRoot.VerifyWidgetVolatile"),
 	GSlateInvalidationRootVerifyWidgetVolatile,
-	TEXT("Each frame, verify that volatile widgets are mark properly and are in the correct list.")
+	TEXT("Every tick, verify that volatile widgets are mark properly and are in the correct list.")
 );
 void VerifyWidgetVolatile(FSlateInvalidationWidgetList& WidgetList, TArray<FSlateInvalidationWidgetIndex>& FinalUpdateList);
 
@@ -118,7 +118,7 @@ bool GSlateInvalidationRootVerifyWidgetsUpdateList = false;
 static FAutoConsoleVariableRef CVarSlateInvalidationRootVerifyWidgetsUpdateList(
 	TEXT("Slate.InvalidationRoot.VerifyWidgetUpdateList"),
 	GSlateInvalidationRootVerifyWidgetsUpdateList,
-	TEXT("Each frame, verify that pre and post update list contains the correct information and they are sorted.")
+	TEXT("Every tick, verify that pre and post update list contains the correct information and they are sorted.")
 );
 void VerifyWidgetsUpdateList_BeforeProcessPreUpdate(const TSharedRef<SWidget>&, FSlateInvalidationWidgetList*, FSlateInvalidationWidgetPreHeap*, FSlateInvalidationWidgetPostHeap*, TArray<FSlateInvalidationWidgetIndex>&);
 void VerifyWidgetsUpdateList_ProcessPrepassUpdate(FSlateInvalidationWidgetList*, FSlateInvalidationWidgetPrepassHeap*, FSlateInvalidationWidgetPostHeap*);
@@ -126,10 +126,10 @@ void VerifyWidgetsUpdateList_BeforeProcessPostUpdate(const TSharedRef<SWidget>&,
 void VerifyWidgetsUpdateList_AfterProcessPostUpdate(const TSharedRef<SWidget>&, FSlateInvalidationWidgetList*, FSlateInvalidationWidgetPreHeap*, FSlateInvalidationWidgetPostHeap*, TArray<FSlateInvalidationWidgetIndex>&);
 
 bool GSlateInvalidationRootVerifySlateAttribute = false;
-static FAutoConsoleVariableRef CVarSlateInvalidationRootVerifySlateAttributes(
+static FAutoConsoleVariableRef CVarSlateInvalidationRootVerifySlateAttribute(
 	TEXT("Slate.InvalidationRoot.VerifySlateAttribute"),
 	GSlateInvalidationRootVerifySlateAttribute,
-	TEXT("Each frame, verify that the widgets that have registered attribute are correctly updated once and the list contains all the widgets.")
+	TEXT("Every tick, verify that the widgets that have registered attribute are correctly updated once and the list contains all the widgets.")
 );
 void VerifySlateAttribute_BeforeUpdate(FSlateInvalidationWidgetList& FastWidgetPathList);
 void VerifySlateAttribute_AfterUpdate(const FSlateInvalidationWidgetList& FastWidgetPathList);
@@ -1185,7 +1185,7 @@ void LogPostInvalidationItem(const FSlateInvalidationWidgetList& FastWidgetPathL
 	if (!(Test)) \
 	{ \
 		UE_LOG(LogSlate, Error, Message, ##__VA_ARGS__); \
-		FlagToReset = false; \
+		FlagToReset->Set(false); \
 		return; \
 	}
 
@@ -1197,11 +1197,11 @@ void VerifyWidgetList(TSharedRef<SWidget> RootWidget, FSlateInvalidationRootHand
 	if (!bIsIdentical)
 	{
 		UE_LOG(LogSlate, Log, TEXT("**-- New Build List --**"));
-		List.LogWidgetsList();
+		List.LogWidgetsList(true);
 		UE_LOG(LogSlate, Log, TEXT("**-- Invaliation Root List --**"));
-		WidgetList.LogWidgetsList();
+		WidgetList.LogWidgetsList(true);
 
-		UE_SLATE_LOG_ERROR_IF_FALSE(false, GSlateInvalidationRootVerifyWidgetList, TEXT("The updated list doesn't match a newly created list."));
+		UE_SLATE_LOG_ERROR_IF_FALSE(false, CVarSlateInvalidationRootVerifyWidgetList, TEXT("The updated list doesn't match a newly created list."));
 	}
 }
 
@@ -1211,7 +1211,7 @@ void VerifyHittest(SWidget* InvalidationRootWidget, FSlateInvalidationWidgetList
 	check(HittestGrid);
 
 	UE_SLATE_LOG_ERROR_IF_FALSE(WidgetList.VerifySortOrder()
-		, GSlateInvalidationRootVerifyHittestGrid
+		, CVarSlateInvalidationRootVerifyHittestGrid
 		, TEXT("The array's sort order for InvalidationRoot '%s' is not respected.")
 		, *FReflectionMetaData::GetWidgetPath(InvalidationRootWidget));
 
@@ -1232,7 +1232,7 @@ void VerifyHittest(SWidget* InvalidationRootWidget, FSlateInvalidationWidgetList
 	for (const FHittestGrid::FWidgetSortData& Data : WeakHittestGridSortDatas)
 	{
 		TSharedPtr<SWidget> Widget = Data.WeakWidget.Pin();
-		UE_SLATE_LOG_ERROR_IF_FALSE(Widget, GSlateInvalidationRootVerifyHittestGrid, TEXT("A widget is invalid in the HittestGrid"));
+		UE_SLATE_LOG_ERROR_IF_FALSE(Widget, CVarSlateInvalidationRootVerifyHittestGrid, TEXT("A widget is invalid in the HittestGrid"));
 
 		FHittestWidgetSortData SortData = { Widget.Get(), Data.PrimarySort, Data.SecondarySort };
 		HittestGridSortDatas.Add(MoveTemp(SortData));
@@ -1258,7 +1258,7 @@ void VerifyHittest(SWidget* InvalidationRootWidget, FSlateInvalidationWidgetList
 				}
 
 				UE_SLATE_LOG_ERROR_IF_FALSE(Widget.GetProxyHandle().GetWidgetSortOrder() == HittestGridSortDatas[FoundHittestIndex].SecondarySort
-					, GSlateInvalidationRootVerifyHittestGrid
+					, CVarSlateInvalidationRootVerifyHittestGrid
 					, TEXT("The SecondarySort of widget '%s' doesn't match the SecondarySort inside the HittestGrid.")
 					, *FReflectionMetaData::GetWidgetPath(Widget));
 
@@ -1283,7 +1283,7 @@ void VerifyHittest(SWidget* InvalidationRootWidget, FSlateInvalidationWidgetList
 	for (const FHittestWidgetSortData& HittestGridSortData : HittestGridSortDatas)
 	{
 		UE_SLATE_LOG_ERROR_IF_FALSE(HittestGridSortData.Widget != nullptr
-			, GSlateInvalidationRootVerifyHittestGrid
+			, CVarSlateInvalidationRootVerifyHittestGrid
 			, TEXT("A widget in the HittestGrid is invalid and shouldn't be in the HittestGrid."));
 
 		if (!HittestGridSortData.Widget->GetProxyHandle().IsValid(HittestGridSortData.Widget))
@@ -1304,7 +1304,7 @@ void VerifyHittest(SWidget* InvalidationRootWidget, FSlateInvalidationWidgetList
 			if (ParentWidetIndex != FSlateInvalidationWidgetIndex::Invalid)
 			{
 				UE_SLATE_LOG_ERROR_IF_FALSE(WidgetList.IsValidIndex(ParentWidetIndex) && EnumHasAllFlags(WidgetList[ParentWidetIndex].CurrentInvalidateReason, EInvalidateWidgetReason::ChildOrder)
-					, GSlateInvalidationRootVerifyHittestGrid
+					, CVarSlateInvalidationRootVerifyHittestGrid
 					, TEXT("The parent widget doesn't has ChildOrder invaldiation. The hittest grid has an invalid widget."));
 			}
 		}
@@ -1327,13 +1327,13 @@ void VerifyWidgetVisibility(FSlateInvalidationWidgetList& WidgetList)
 					{
 						// Confirm that we have the correct parent
 						UE_SLATE_LOG_ERROR_IF_FALSE(WidgetList.IsValidIndex(InvalidationWidget.ParentIndex)
-							, GSlateInvalidationRootVerifyWidgetVisibility
+							, CVarSlateInvalidationRootVerifyWidgetVisibility
 							, TEXT("Widget '%s' Parent index is invalid.")
 							, *FReflectionMetaData::GetWidgetDebugInfo(Widget));
 
 						const FSlateInvalidationWidgetList::InvalidationWidgetType& ParentInvalidationWidget = WidgetList[InvalidationWidget.ParentIndex];
 						UE_SLATE_LOG_ERROR_IF_FALSE(ParentWidget.Get() == ParentInvalidationWidget.GetWidget()
-							, GSlateInvalidationRootVerifyWidgetVisibility
+							, CVarSlateInvalidationRootVerifyWidgetVisibility
 							, TEXT("Widget '%s' Parent is not '%s'.")
 							, *FReflectionMetaData::GetWidgetDebugInfo(Widget)
 							, *FReflectionMetaData::GetWidgetDebugInfo(ParentWidget.Get()));
@@ -1344,25 +1344,25 @@ void VerifyWidgetVisibility(FSlateInvalidationWidgetList& WidgetList)
 					else
 					{
 						UE_SLATE_LOG_ERROR_IF_FALSE(ParentWidget == nullptr || ParentWidget->Advanced_IsInvalidationRoot()
-							, GSlateInvalidationRootVerifyWidgetVisibility
+							, CVarSlateInvalidationRootVerifyWidgetVisibility
 							, TEXT("Widget '%s' Parent is valid and is not an invalidation root.")
 							, *FReflectionMetaData::GetWidgetDebugInfo(Widget));
 					}
 
 					UE_SLATE_LOG_ERROR_IF_FALSE(InvalidationWidget.Visibility.AreAncestorsVisible() == bParentIsVisible
-						, GSlateInvalidationRootVerifyWidgetVisibility
+						, CVarSlateInvalidationRootVerifyWidgetVisibility
 						, TEXT("Widget '%s' AreAncestorsVisible flag is wrong.")
 						, *FReflectionMetaData::GetWidgetDebugInfo(Widget));
 					UE_SLATE_LOG_ERROR_IF_FALSE(InvalidationWidget.Visibility.IsVisible() == (bParentIsVisible && WidgetVisibility.IsVisible())
-						, GSlateInvalidationRootVerifyWidgetVisibility
+						, CVarSlateInvalidationRootVerifyWidgetVisibility
 						, TEXT("Widget '%s' IsVisible flag is wrong.")
 						, *FReflectionMetaData::GetWidgetDebugInfo(Widget));
 					UE_SLATE_LOG_ERROR_IF_FALSE(InvalidationWidget.Visibility.IsCollapsed() == bParentIsCollapsed || WidgetVisibility == EVisibility::Collapsed
-						, GSlateInvalidationRootVerifyWidgetVisibility
+						, CVarSlateInvalidationRootVerifyWidgetVisibility
 						, TEXT("Widget '%s' IsCollapsed flag is wrong.")
 						, *FReflectionMetaData::GetWidgetDebugInfo(Widget));
 					UE_SLATE_LOG_ERROR_IF_FALSE(InvalidationWidget.Visibility.IsCollapseIndirectly() == bParentIsCollapsed
-					, GSlateInvalidationRootVerifyWidgetVisibility
+						, CVarSlateInvalidationRootVerifyWidgetVisibility
 						, TEXT("Widget '%s' IsCollapseIndirectly flag is wrong.")
 						, *FReflectionMetaData::GetWidgetDebugInfo(Widget));
 				}
@@ -1370,7 +1370,7 @@ void VerifyWidgetVisibility(FSlateInvalidationWidgetList& WidgetList)
 					if (!InvalidationWidget.Visibility.IsVisible())
 					{
 						UE_SLATE_LOG_ERROR_IF_FALSE(!Widget->GetPersistentState().CachedElementHandle.HasCachedElements()
-							, GSlateInvalidationRootVerifyWidgetVisibility
+							, CVarSlateInvalidationRootVerifyWidgetVisibility
 							, TEXT("Widget '%s' has cached element and is not visibled.")
 							, *FReflectionMetaData::GetWidgetDebugInfo(Widget));
 					}
@@ -1397,21 +1397,21 @@ void VerifyWidgetVolatile(FSlateInvalidationWidgetList& WidgetList, TArray<FSlat
 					Widget.CacheVolatility();
 					const bool bIsVolatile = Widget.IsVolatile();
 					UE_SLATE_LOG_ERROR_IF_FALSE(bWasVolatile == bIsVolatile
-						, GSlateInvalidationRootVerifyWidgetVolatile
+						, CVarSlateInvalidationRootVerifyWidgetVolatile
 						, TEXT("Widget '%s' volatily changed without an invalidation.")
 						, *FReflectionMetaData::GetWidgetDebugInfo(Widget));
 				}
 
 				const TSharedPtr<const SWidget> ParentWidget = Widget.GetParentWidget();
 				UE_SLATE_LOG_ERROR_IF_FALSE(ParentWidget
-					, GSlateInvalidationRootVerifyWidgetVolatile
+					, CVarSlateInvalidationRootVerifyWidgetVolatile
 					, TEXT("Parent widget of widget '%s' is invalid.")
 					, *FReflectionMetaData::GetWidgetDebugInfo(Widget));
 
 				{
 					const bool bShouldBeVolatileIndirectly = ParentWidget->IsVolatileIndirectly() || ParentWidget->IsVolatile();
 					UE_SLATE_LOG_ERROR_IF_FALSE(Widget.IsVolatileIndirectly() == bShouldBeVolatileIndirectly
-					, GSlateInvalidationRootVerifyWidgetVolatile
+					, CVarSlateInvalidationRootVerifyWidgetVolatile
 					, TEXT("Widget '%s' should be set as %s.")
 					, *FReflectionMetaData::GetWidgetDebugInfo(Widget)
 					, (bShouldBeVolatileIndirectly ? TEXT("volatile indirectly") : TEXT("not volatile indirectly")));
@@ -1420,7 +1420,7 @@ void VerifyWidgetVolatile(FSlateInvalidationWidgetList& WidgetList, TArray<FSlat
 				if (Widget.IsVolatile() && !Widget.IsVolatileIndirectly())
 				{
 					UE_SLATE_LOG_ERROR_IF_FALSE(Widget.HasAnyUpdateFlags(EWidgetUpdateFlags::NeedsVolatilePaint)
-						, GSlateInvalidationRootVerifyWidgetVolatile
+						, CVarSlateInvalidationRootVerifyWidgetVolatile
 						, TEXT("Widget '%s' is volatile but doesn't have the update flag NeedsVolatilePaint.")
 						, *FReflectionMetaData::GetWidgetDebugInfo(Widget));
 
@@ -1429,7 +1429,7 @@ void VerifyWidgetVolatile(FSlateInvalidationWidgetList& WidgetList, TArray<FSlat
 						const bool bIsVisible = Widget.GetProxyHandle().GetProxy().Visibility.IsVisible();
 						const bool bIsContains = FinalUpdateList.Contains(Widget.GetProxyHandle().GetWidgetIndex());
 						UE_SLATE_LOG_ERROR_IF_FALSE(bIsContains || !bIsVisible
-							, GSlateInvalidationRootVerifyWidgetVolatile
+							, CVarSlateInvalidationRootVerifyWidgetVolatile
 							, TEXT("Widget '%s' is volatile but is not in the update list.")
 							, *FReflectionMetaData::GetWidgetDebugInfo(Widget));
 					}
@@ -1452,28 +1452,28 @@ void VerifyWidgetsUpdateList_BeforeProcessPreUpdate(const TSharedRef<SWidget>& R
 	for (FSlateInvalidationWidgetIndex WidgetIndex : FinalUpdateList)
 	{
 		UE_SLATE_LOG_ERROR_IF_FALSE(FastWidgetPathList->IsValidIndex(WidgetIndex)
-			, GSlateInvalidationRootVerifyWidgetsUpdateList
+			, CVarSlateInvalidationRootVerifyWidgetsUpdateList
 			, TEXT("A WidgetIndex is invalid. The Widget can be invalid (because it's not been processed yet)."));
 	}
 
 	UE_SLATE_LOG_ERROR_IF_FALSE(WidgetsNeedingPreUpdate->IsValidHeap_Debug()
-		, GSlateInvalidationRootVerifyWidgetsUpdateList
+		, CVarSlateInvalidationRootVerifyWidgetsUpdateList
 		, TEXT("The PreUpdate list need to stay a valid heap"));
 
 	WidgetsNeedingPreUpdate->ForEachIndexes([FastWidgetPathList](const FSlateInvalidationWidgetPreHeap::FElement& Element)
 		{
 			UE_SLATE_LOG_ERROR_IF_FALSE(FastWidgetPathList->IsValidIndex(Element.GetWidgetIndex())
-				, GSlateInvalidationRootVerifyWidgetsUpdateList
+				, CVarSlateInvalidationRootVerifyWidgetsUpdateList
 				, TEXT("An element is not valid."));
 			const FSlateInvalidationWidgetList::InvalidationWidgetType& InvalidationWidget = (*FastWidgetPathList)[Element.GetWidgetIndex()];
 			if (SWidget* Widget = (*FastWidgetPathList)[Element.GetWidgetIndex()].GetWidget())
 			{
 				UE_SLATE_LOG_ERROR_IF_FALSE(Widget->GetProxyHandle().GetWidgetSortOrder() == Element.GetWidgetSortOrder()
-					, GSlateInvalidationRootVerifyWidgetsUpdateList
+					, CVarSlateInvalidationRootVerifyWidgetsUpdateList
 					, TEXT("The sort order of the widget '%s' do not matches what is in the heap.")
 					, *FReflectionMetaData::GetWidgetDebugInfo(Widget));
 				UE_SLATE_LOG_ERROR_IF_FALSE(Widget->GetProxyHandle().GetWidgetIndex() == Element.GetWidgetIndex()
-					, GSlateInvalidationRootVerifyWidgetsUpdateList
+					, CVarSlateInvalidationRootVerifyWidgetsUpdateList
 					, TEXT("The widget index of the widget '%s' do not matches what is in the heap.")
 					, *FReflectionMetaData::GetWidgetDebugInfo(Widget));
 			}
@@ -1482,7 +1482,7 @@ void VerifyWidgetsUpdateList_BeforeProcessPreUpdate(const TSharedRef<SWidget>& R
 	FastWidgetPathList->ForEachInvalidationWidget([WidgetsNeedingPreUpdate](FSlateInvalidationWidgetList::InvalidationWidgetType& InvalidationWidget)
 		{
 			UE_SLATE_LOG_ERROR_IF_FALSE(WidgetsNeedingPreUpdate->Contains_Debug(InvalidationWidget.Index) == InvalidationWidget.bContainedByWidgetPreHeap
-				, GSlateInvalidationRootVerifyWidgetsUpdateList
+				, CVarSlateInvalidationRootVerifyWidgetsUpdateList
 				, TEXT("Widget '%s' is or is not in the PreUpdate but the flag say otherwise.")
 				, *FReflectionMetaData::GetWidgetDebugInfo(InvalidationWidget.GetWidget()));
 		});
@@ -1496,7 +1496,7 @@ void VerifyWidgetsUpdateList_ProcessPrepassUpdate(FSlateInvalidationWidgetList*
 			const FSlateInvalidationWidgetList::InvalidationWidgetType& InvalidationWidget = (*FastWidgetPathList)[Element.GetWidgetIndex()];
 			const SWidget* Widget = InvalidationWidget.GetWidget();
 			UE_SLATE_LOG_ERROR_IF_FALSE(WidgetsNeedingPostUpdate->Contains_Debug(Element.GetWidgetIndex())
-				, GSlateInvalidationRootVerifyWidgetsUpdateList
+				, CVarSlateInvalidationRootVerifyWidgetsUpdateList
 				, TEXT("Widget '%s' is in the WidgetsNeedingPrepassUpdate list but not in the WidgetsNeedingPostUpdate.")
 				, *FReflectionMetaData::GetWidgetDebugInfo(InvalidationWidget.GetWidget()));
 		});
@@ -1514,23 +1514,23 @@ void VerifyWidgetsUpdateList_BeforeProcessPostUpdate(const TSharedRef<SWidget>&
 	}
 
 	UE_SLATE_LOG_ERROR_IF_FALSE(WidgetsNeedingPostUpdate->IsValidHeap_Debug()
-		, GSlateInvalidationRootVerifyWidgetsUpdateList
+		, CVarSlateInvalidationRootVerifyWidgetsUpdateList
 		, TEXT("The PostUpdate list need to stay a valid heap"));
 
 	// Widget in SWidget::Prepass may have added new widget.
 	//This is not desired, it will introduce a frame delay compare to slow path.
 	//UE_SLATE_LOG_ERROR_IF_FALSE(WidgetsNeedingPreUpdate->Num() == 0
-	//	, GSlateInvalidationRootVerifyWidgetsUpdateList
+	//	, CVarSlateInvalidationRootVerifyWidgetsUpdateList
 	//	, TEXT("The PreUpdate list should be empty"));
 
 	UE_SLATE_LOG_ERROR_IF_FALSE(FinalUpdateList.Num() == 0
-		, GSlateInvalidationRootVerifyWidgetsUpdateList
+		, CVarSlateInvalidationRootVerifyWidgetsUpdateList
 		, TEXT("The Final Update list should be empty."));
 
 	WidgetsNeedingPostUpdate->ForEachIndexes([FastWidgetPathList](const FSlateInvalidationWidgetPreHeap::FElement& Element)
 		{
 			UE_SLATE_LOG_ERROR_IF_FALSE(FastWidgetPathList->IsValidIndex(Element.GetWidgetIndex())
-				, GSlateInvalidationRootVerifyWidgetsUpdateList
+				, CVarSlateInvalidationRootVerifyWidgetsUpdateList
 				, TEXT("An element is not valid."));
 
 			const FSlateInvalidationWidgetList::InvalidationWidgetType& InvalidationWidget = (*FastWidgetPathList)[Element.GetWidgetIndex()];
@@ -1540,11 +1540,11 @@ void VerifyWidgetsUpdateList_BeforeProcessPostUpdate(const TSharedRef<SWidget>&
 			if (Widget)
 			{
 				UE_SLATE_LOG_ERROR_IF_FALSE(Widget->GetProxyHandle().GetWidgetSortOrder() == Element.GetWidgetSortOrder()
-					, GSlateInvalidationRootVerifyWidgetsUpdateList
+					, CVarSlateInvalidationRootVerifyWidgetsUpdateList
 					, TEXT("The sort order of the widget '%s' do not matches what is in the heap.")
 					, *FReflectionMetaData::GetWidgetDebugInfo(Widget));
 				UE_SLATE_LOG_ERROR_IF_FALSE(Widget->GetProxyHandle().GetWidgetIndex() == Element.GetWidgetIndex()
-					, GSlateInvalidationRootVerifyWidgetsUpdateList
+					, CVarSlateInvalidationRootVerifyWidgetsUpdateList
 					, TEXT("The widget index of the widget '%s' do not matches what is in the heap.")
 					, *FReflectionMetaData::GetWidgetDebugInfo(Widget));
 			}
@@ -1553,7 +1553,7 @@ void VerifyWidgetsUpdateList_BeforeProcessPostUpdate(const TSharedRef<SWidget>&
 	FastWidgetPathList->ForEachInvalidationWidget([WidgetsNeedingPostUpdate](FSlateInvalidationWidgetList::InvalidationWidgetType& InvalidationWidget)
 		{
 			UE_SLATE_LOG_ERROR_IF_FALSE(WidgetsNeedingPostUpdate->Contains_Debug(InvalidationWidget.Index) == InvalidationWidget.bContainedByWidgetPostHeap
-				, GSlateInvalidationRootVerifyWidgetsUpdateList
+				, CVarSlateInvalidationRootVerifyWidgetsUpdateList
 				, TEXT("Widget '%s' is or is not in the PostUpdate but the flag say otherwise.")
 				, *FReflectionMetaData::GetWidgetDebugInfo(InvalidationWidget.GetWidget()));
 		});
@@ -1571,10 +1571,10 @@ void VerifyWidgetsUpdateList_AfterProcessPostUpdate(const TSharedRef<SWidget>& R
 	}
 
 	UE_SLATE_LOG_ERROR_IF_FALSE(WidgetsNeedingPreUpdate->Num() == 0
-		, GSlateInvalidationRootVerifyWidgetsUpdateList
+		, CVarSlateInvalidationRootVerifyWidgetsUpdateList
 		, TEXT("The list of Pre Update should already been processed."));
 	UE_SLATE_LOG_ERROR_IF_FALSE(WidgetsNeedingPostUpdate->Num() == 0
-		, GSlateInvalidationRootVerifyWidgetsUpdateList
+		, CVarSlateInvalidationRootVerifyWidgetsUpdateList
 		, TEXT("The list of Post Update should already been processed."));
 
 	for(FSlateInvalidationWidgetIndex WidgetIndex : FinalUpdateList)
@@ -1586,11 +1586,11 @@ void VerifyWidgetsUpdateList_AfterProcessPostUpdate(const TSharedRef<SWidget>& R
 		if (Widget)
 		{
 			UE_SLATE_LOG_ERROR_IF_FALSE(InvalidationWidget.CurrentInvalidateReason == EInvalidateWidgetReason::None
-				, GSlateInvalidationRootVerifyWidgetsUpdateList
+				, CVarSlateInvalidationRootVerifyWidgetsUpdateList
 				, TEXT("The widget '%s' is in the update list and it still has a Invalidation Reason.")
 				, *FReflectionMetaData::GetWidgetDebugInfo(Widget));
 			UE_SLATE_LOG_ERROR_IF_FALSE(Widget->HasAnyUpdateFlags(EWidgetUpdateFlags::AnyUpdate)
-				, GSlateInvalidationRootVerifyWidgetsUpdateList
+				, CVarSlateInvalidationRootVerifyWidgetsUpdateList
 				, TEXT("The widget '%s' is in the update list but doesn't have an update flag set.")
 				, *FReflectionMetaData::GetWidgetDebugInfo(Widget));
 		}
@@ -1609,7 +1609,7 @@ void VerifySlateAttribute_AfterUpdate(const FSlateInvalidationWidgetList& FastWi
 {
 	const bool bElementIndexListValid = FastWidgetPathList.VerifyElementIndexList();
 	UE_SLATE_LOG_ERROR_IF_FALSE(bElementIndexListValid
-		, GSlateInvalidationRootVerifySlateAttribute
+		, CVarSlateInvalidationRootVerifySlateAttribute
 		, TEXT("The VerifySlateAttribute failed in post."));
 }
 
diff --git a/Engine/Source/Runtime/SlateCore/Private/FastUpdate/SlateInvalidationWidgetList.cpp b/Engine/Source/Runtime/SlateCore/Private/FastUpdate/SlateInvalidationWidgetList.cpp
--- a/Engine/Source/Runtime/SlateCore/Private/FastUpdate/SlateInvalidationWidgetList.cpp
+++ b/Engine/Source/Runtime/SlateCore/Private/FastUpdate/SlateInvalidationWidgetList.cpp
@@ -781,27 +781,6 @@ void FSlateInvalidationWidgetList::Internal_FindChildren(FSlateInvalidationWidge
 }
 
 
-//FSlateInvalidationWidgetIndex FSlateInvalidationWidgetList::_FindPreviousSibling(FSlateInvalidationWidgetIndex WidgetIndex) const
-//{
-//	FSlateInvalidationWidgetIndex Result = FSlateInvalidationWidgetIndex::Invalid;
-//	const InvalidationWidgetType& InvalidationWidget = (*this)[WidgetIndex];
-//	if (InvalidationWidget.ParentIndex != FSlateInvalidationWidgetIndex::Invalid)
-//	{
-//		const FSlateInvalidationWidgetIndex PreviousWidgetIndex = DecrementIndex(WidgetIndex);
-//		const InvalidationWidgetType& PreviousInvalidationWidget = (*this)[WidgetIndex];
-//		if (PreviousInvalidationWidget.ParentIndex == InvalidationWidget.ParentIndex)
-//		{
-//			Result = PreviousWidgetIndex;
-//		}
-//		else
-//		{
-//			Result = PreviousInvalidationWidget.ParentIndex;
-//		}
-//	}
-//	return  Result;
-//}
-
-
 FSlateInvalidationWidgetIndex FSlateInvalidationWidgetList::IncrementIndex(FSlateInvalidationWidgetIndex Index) const
 {
 	check(Data.IsValidIndex(Index.ArrayIndex));
@@ -842,6 +821,24 @@ FSlateInvalidationWidgetIndex FSlateInvalidationWidgetList::DecrementIndex(FSlat
 }
 
 
+
+FSlateInvalidationWidgetIndex FSlateInvalidationWidgetList::FindNextSibling(FSlateInvalidationWidgetIndex WidgetIndex) const
+{
+	FSlateInvalidationWidgetIndex Result = FSlateInvalidationWidgetIndex::Invalid;
+	const InvalidationWidgetType& InvalidationWidget = (*this)[WidgetIndex];
+	const FSlateInvalidationWidgetIndex NextSiblingIndex = IncrementIndex(InvalidationWidget.LeafMostChildIndex);
+	if (NextSiblingIndex != FSlateInvalidationWidgetIndex::Invalid)
+	{
+		const InvalidationWidgetType& NextSiblingInvalidationWidget = (*this)[NextSiblingIndex];
+		if (NextSiblingInvalidationWidget.ParentIndex == InvalidationWidget.ParentIndex)
+		{
+			Result = NextSiblingIndex;
+		}
+	}
+	return Result;
+}
+
+
 void FSlateInvalidationWidgetList::Empty()
 {
 	Data.Empty();
@@ -1577,56 +1574,80 @@ bool FSlateInvalidationWidgetList::DeapCompare(const FSlateInvalidationWidgetLis
 }
 
 
-void FSlateInvalidationWidgetList::LogWidgetsList()
+void FSlateInvalidationWidgetList::LogWidgetsList(bool bOnlyVisible) const
 {
-	TStringBuilder<256> Builder;
+	auto GetDepth = [this](FSlateInvalidationWidgetIndex Index)
+	{
+		int32 Depth = 0;
+		while(Index != FSlateInvalidationWidgetIndex::Invalid)
+		{
+			++Depth;
+			Index = (*this)[Index].ParentIndex;
+		}
+		return Depth-1;
+	};
+
+	TStringBuilder<512> Builder;
 	for (FSlateInvalidationWidgetIndex Index = FirstIndex(); Index != FSlateInvalidationWidgetIndex::Invalid; Index = IncrementIndex(Index))
 	{
 		Builder.Reset();
+
 		const InvalidationWidgetType& InvalidateWidget = (*this)[Index];
-		if (SWidget* Widget = InvalidateWidget.GetWidget())
+		if (!bOnlyVisible || InvalidateWidget.Visibility.IsVisible())
 		{
-			Builder << FReflectionMetaData::GetWidgetDebugInfo(Widget) << TEXT("(") << Widget->GetTag() << TEXT(")");
-		}
-		else
-		{
-			Builder << TEXT("[None]");
-		}
-		Builder << TEXT("\t");
-		if (InvalidateWidget.ParentIndex != FSlateInvalidationWidgetIndex::Invalid)
-		{
-			if (SWidget* Widget = (*this)[InvalidateWidget.ParentIndex].GetWidget())
+			for (int32 Depth = GetDepth(Index); Depth > 0; --Depth)
+			{
+				Builder << TEXT(' ');
+			}
+
+			if (SWidget* Widget = InvalidateWidget.GetWidget())
 			{
-				Builder << FReflectionMetaData::GetWidgetDebugInfo(Widget) << TEXT("(") << Widget->GetTag() << TEXT(")");
+				Builder << FReflectionMetaData::GetWidgetDebugInfo(Widget)
+					<< TEXT('(') << Widget->GetTag() << TEXT(')')
+					<< TEXT('[') << Widget->GetPersistentState().LayerId << TEXT(',') << Widget->GetPersistentState().OutgoingLayerId << TEXT(']');
 			}
 			else
 			{
 				Builder << TEXT("[None]");
 			}
-		}
-		else
-		{
-			Builder << TEXT("[---]");
-		}
-		Builder << TEXT("\t");
-		if (InvalidateWidget.LeafMostChildIndex != FSlateInvalidationWidgetIndex::Invalid)
-		{
-			if (SWidget* Widget = (*this)[InvalidateWidget.LeafMostChildIndex].GetWidget())
+			Builder << TEXT("\t");
+
+			if (InvalidateWidget.ParentIndex != FSlateInvalidationWidgetIndex::Invalid)
 			{
-				Builder << FReflectionMetaData::GetWidgetDebugInfo(Widget) << TEXT("(") << Widget->GetTag() << TEXT(")");
+				if (SWidget* Widget = (*this)[InvalidateWidget.ParentIndex].GetWidget())
+				{
+					Builder << FReflectionMetaData::GetWidgetDebugInfo(Widget) << TEXT("(") << Widget->GetTag() << TEXT(")");
+				}
+				else
+				{
+					Builder << TEXT("[None]");
+				}
 			}
 			else
 			{
-				Builder << TEXT("[None]");
+				Builder << TEXT("[---]");
 			}
-		}
-		else
-		{
-			Builder << TEXT("[---]");
-		}
-		Builder << TEXT("\t");
+			Builder << TEXT("\t");
 
-		UE_LOG(LogSlate, Log, TEXT("%s"), Builder.ToString());
+			if (InvalidateWidget.LeafMostChildIndex != FSlateInvalidationWidgetIndex::Invalid)
+			{
+				if (SWidget* Widget = (*this)[InvalidateWidget.LeafMostChildIndex].GetWidget())
+				{
+					Builder << FReflectionMetaData::GetWidgetDebugInfo(Widget) << TEXT("(") << Widget->GetTag() << TEXT(")");
+				}
+				else
+				{
+					Builder << TEXT("[None]");
+				}
+			}
+			else
+			{
+				Builder << TEXT("[---]");
+			}
+			Builder << TEXT("\t");
+
+			UE_LOG(LogSlate, Log, TEXT("%s"), Builder.ToString());
+		}
 	}
 }
 
diff --git a/Engine/Source/Runtime/SlateCore/Private/FastUpdate/SlateInvalidationWidgetList.h b/Engine/Source/Runtime/SlateCore/Private/FastUpdate/SlateInvalidationWidgetList.h
--- a/Engine/Source/Runtime/SlateCore/Private/FastUpdate/SlateInvalidationWidgetList.h
+++ b/Engine/Source/Runtime/SlateCore/Private/FastUpdate/SlateInvalidationWidgetList.h
@@ -47,18 +47,18 @@ public:
 			check(OrderMin <= OrderMax);
 #endif
 		}
-		bool Include(FSlateInvalidationWidgetSortOrder Other) const
+		UE_NODISCARD bool Include(FSlateInvalidationWidgetSortOrder Other) const
 		{
 			return OrderMin <= Other && Other <= OrderMax;
 		}
-		bool IsValid() const { return InclusiveMin != FSlateInvalidationWidgetIndex::Invalid; }
+		UE_NODISCARD bool IsValid() const { return InclusiveMin != FSlateInvalidationWidgetIndex::Invalid; }
 
 		FSlateInvalidationWidgetIndex GetInclusiveMinWidgetIndex() const { return InclusiveMin; }
 		FSlateInvalidationWidgetIndex GetInclusiveMaxWidgetIndex() const { return InclusiveMax; }
 		FSlateInvalidationWidgetSortOrder GetInclusiveMinWidgetSortOrder() const { return OrderMin; }
 		FSlateInvalidationWidgetSortOrder GetInclusiveMaxWidgetSortOrder() const { return OrderMax; }
 
-		bool operator==(const FIndexRange& Other) const { return Other.InclusiveMin == InclusiveMin && Other.InclusiveMax == InclusiveMax; }
+		UE_NODISCARD bool operator==(const FIndexRange& Other) const { return Other.InclusiveMin == InclusiveMin && Other.InclusiveMax == InclusiveMax; }
 	};
 
 public:
@@ -105,7 +105,7 @@ public:
 		struct FReIndexOperation
 		{
 			FReIndexOperation(const FIndexRange& InRange, FSlateInvalidationWidgetIndex InReIndexTarget) : Range(InRange), ReIndexTarget(InReIndexTarget) {}
-			const FIndexRange& GetRange() const { return Range; }
+			UE_NODISCARD const FIndexRange& GetRange() const { return Range; }
 			UE_NODISCARD FSlateInvalidationWidgetIndex ReIndex(FSlateInvalidationWidgetIndex Index) const;
 		private:
 			const FIndexRange& Range;
@@ -114,7 +114,7 @@ public:
 		struct FReSortOperation
 		{
 			FReSortOperation(const FIndexRange& InRange) : Range(InRange) {}
-			const FIndexRange& GetRange() const { return Range; }
+			UE_NODISCARD const FIndexRange& GetRange() const { return Range; }
 		private:
 			const FIndexRange& Range;
 		};
@@ -158,8 +158,28 @@ public:
 			const int32 ElementNum = ElementList.Num();
 			for (int32 ElementIndex = Data[ArrayIndex].StartIndex; ElementIndex < ElementNum; ++ElementIndex)
 			{
-				SWidget* Widget = ElementList[ElementIndex].GetWidget();
-				if (Widget)
+				if (SWidget* Widget = ElementList[ElementIndex].GetWidget())
+				{
+					Pred(*Widget);
+				}
+			}
+
+			ArrayIndex = Data[ArrayIndex].NextArrayIndex;
+		}
+	}
+
+	/** Performs an operation on all SWidget in the list. */
+	template<typename Predicate>
+	void ForEachWidget(Predicate Pred) const
+	{
+		int32 ArrayIndex = FirstArrayIndex;
+		while (ArrayIndex != INDEX_NONE)
+		{
+			const ElementListType& ElementList = Data[ArrayIndex].ElementList;
+			const int32 ElementNum = ElementList.Num();
+			for (int32 ElementIndex = Data[ArrayIndex].StartIndex; ElementIndex < ElementNum; ++ElementIndex)
+			{
+				if (const SWidget* Widget = ElementList[ElementIndex].GetWidget())
 				{
 					Pred(*Widget);
 				}
@@ -187,6 +207,75 @@ public:
 		}
 	}
 
+	/** Performs an operation on all InvalidationWidget in the list. */
+	template<typename Predicate>
+	void ForEachInvalidationWidget(Predicate Pred) const
+	{
+		int32 ArrayIndex = FirstArrayIndex;
+		while (ArrayIndex != INDEX_NONE)
+		{
+			const ElementListType& ElementList = Data[ArrayIndex].ElementList;
+			const int32 ElementNum = ElementList.Num();
+			for (int32 ElementIndex = Data[ArrayIndex].StartIndex; ElementIndex < ElementNum; ++ElementIndex)
+			{
+				Pred(ElementList[ElementIndex]);
+			}
+
+			ArrayIndex = Data[ArrayIndex].NextArrayIndex;
+		}
+	}
+
+	/** Performs an operation on all InvalidationWidget in the list bellow the provided (WidgetIndex not including it). */
+	template<typename Predicate>
+	void ForEachInvalidationWidget(FSlateInvalidationWidgetIndex BellowWidgetIndex, Predicate Pred)
+	{
+		check(IsValidIndex(BellowWidgetIndex));
+		const InvalidationWidgetType& BeginInvalidationWidget = (*this)[BellowWidgetIndex];
+		const FSlateInvalidationWidgetIndex EndWidgetIndex = BeginInvalidationWidget.LeafMostChildIndex;
+		const FSlateInvalidationWidgetIndex BeginWidgetIndex = IncrementIndex(BellowWidgetIndex);
+
+		if (BeginWidgetIndex != FSlateInvalidationWidgetIndex::Invalid)
+		{
+			const bool bSameWidgetArrayIndex = BeginWidgetIndex.ArrayIndex == EndWidgetIndex.ArrayIndex;
+			int32 ArrayIndex = BeginWidgetIndex.ArrayIndex;
+			{
+				ElementListType& ElementList = Data[ArrayIndex].ElementList;
+				const int32 ElementNum = bSameWidgetArrayIndex ? EndWidgetIndex.ElementIndex + 1 : ElementList.Num();
+				for (int32 ElementIndex = BeginWidgetIndex.ElementIndex; ElementIndex < ElementNum; ++ElementIndex)
+				{
+					Pred(ElementList[ElementIndex]);
+				}
+
+				ArrayIndex = Data[ArrayIndex].NextArrayIndex;
+			}
+
+			if (!bSameWidgetArrayIndex)
+			{
+				while (ArrayIndex != EndWidgetIndex.ArrayIndex)
+				{
+					ElementListType& ElementList = Data[ArrayIndex].ElementList;
+					const int32 ElementNum = ElementList.Num();
+					for (int32 ElementIndex = Data[ArrayIndex].StartIndex; ElementIndex < ElementNum; ++ElementIndex)
+					{
+						Pred(ElementList[ElementIndex]);
+					}
+
+					ArrayIndex = Data[ArrayIndex].NextArrayIndex;
+				}
+
+				{
+					check(ArrayIndex == EndWidgetIndex.ArrayIndex);
+					ElementListType& ElementList = Data[ArrayIndex].ElementList;
+					const int32 ElementNum = EndWidgetIndex.ElementIndex + 1;
+					for (int32 ElementIndex = Data[ArrayIndex].StartIndex; ElementIndex < ElementNum; ++ElementIndex)
+					{
+						Pred(ElementList[ElementIndex]);
+					}
+				}
+			}
+		}
+	}
+
 	/** Iterator that goes over all the widgets with registered attribute. */
 	struct FWidgetAttributeIterator
 	{
@@ -227,7 +316,7 @@ public:
 		void AdvanceToNextParent();
 
 		/** Is the iterator pointing to a valid widget index. */
-		bool IsValid() const { return CurrentWidgetIndex != FSlateInvalidationWidgetIndex::Invalid; }
+		UE_NODISCARD bool IsValid() const { return CurrentWidgetIndex != FSlateInvalidationWidgetIndex::Invalid; }
 
 	private:
 		void AdvanceArrayIndex(int32 ArrayIndex);
@@ -259,7 +348,7 @@ public:
 		void Advance();
 
 		/** Is the iterator pointing to a valid widget index. */
-		bool IsValid() const { return CurrentWidgetIndex != FSlateInvalidationWidgetIndex::Invalid; }
+		UE_NODISCARD bool IsValid() const { return CurrentWidgetIndex != FSlateInvalidationWidgetIndex::Invalid; }
 
 	private:
 		void Internal_Advance();
@@ -304,13 +393,13 @@ public:
 	}
 
 	/** Returns true if there is not element in the WidgetList. */
-	bool IsEmpty() const
+	UE_NODISCARD bool IsEmpty() const
 	{
 		return FirstArrayIndex == INDEX_NONE || Data[FirstArrayIndex].ElementList.Num() == 0;
 	}
 
 	/** Returns the first index from the WidgetList. */
-	FSlateInvalidationWidgetIndex FirstIndex() const
+	UE_NODISCARD FSlateInvalidationWidgetIndex FirstIndex() const
 	{
 		return FirstArrayIndex == INDEX_NONE
 			? FSlateInvalidationWidgetIndex::Invalid
@@ -318,7 +407,7 @@ public:
 	}
 
 	/** Returns the last index from the WidgetList. */
-	FSlateInvalidationWidgetIndex LastIndex() const
+	UE_NODISCARD FSlateInvalidationWidgetIndex LastIndex() const
 	{
 		return LastArrayIndex == INDEX_NONE ? FSlateInvalidationWidgetIndex::Invalid : FSlateInvalidationWidgetIndex{ (IndexType)LastArrayIndex, (IndexType)(Data[LastArrayIndex].ElementList.Num() - 1) };
 	}
@@ -329,6 +418,9 @@ public:
 	/** Decrement a widget index to the next entry in the WidgetList. */
 	UE_NODISCARD FSlateInvalidationWidgetIndex DecrementIndex(FSlateInvalidationWidgetIndex Index) const;
 
+	/** Find the next widget index that share the same parent. Return Invalid if there is no sibling. */
+	UE_NODISCARD FSlateInvalidationWidgetIndex FindNextSibling(FSlateInvalidationWidgetIndex WidgetIndex) const;
+
 	/** Empties the WidgetList. */
 	void Empty();
 
@@ -357,7 +449,7 @@ public:
 	bool DeapCompare(const FSlateInvalidationWidgetList& Other) const;
 
 	/** For testing purposes. Log the tree. */
-	void LogWidgetsList();
+	void LogWidgetsList(bool bOnlyVisible) const;
 
 	/** For testing purposes. Verify that every widgets has the correct index. */
 	bool VerifyWidgetsIndex() const;
