Fix editable text box default focused color black.

diff --git a/Engine/Source/Runtime/SlateCore/Public/Styling/SlateTypes.h b/Engine/Source/Runtime/SlateCore/Public/Styling/SlateTypes.h
--- a/Engine/Source/Runtime/SlateCore/Public/Styling/SlateTypes.h
+++ b/Engine/Source/Runtime/SlateCore/Public/Styling/SlateTypes.h
@@ -1016,6 +1016,7 @@ struct SLATECORE_API FEditableTextBoxStyle : public FSlateWidgetStyle
 		ForegroundColor.Unlink();
 		BackgroundColor.Unlink();
 		ReadOnlyForegroundColor.Unlink();
+		FocusedForegroundColor.Unlink();
 		ScrollBarStyle.UnlinkColors();
 	}
 };
@@ -1638,6 +1639,8 @@ struct SLATECORE_API FTableRowStyle : public FSlateWidgetStyle
 		ActiveBrush.UnlinkColors();
 		InactiveHoveredBrush.UnlinkColors();
 		InactiveBrush.UnlinkColors();
+		ParentRowBackgroundBrush.UnlinkColors();
+		ParentRowBackgroundHoveredBrush.UnlinkColors();
 		EvenRowBackgroundHoveredBrush.UnlinkColors();
 		EvenRowBackgroundBrush.UnlinkColors();
 		OddRowBackgroundHoveredBrush.UnlinkColors();
