Slate: Change ToggleButtonCheckboxAlt PressedForegroundColor to all white

diff --git a/Engine/Source/Runtime/SlateCore/Private/Styling/StarshipCoreStyle.cpp b/Engine/Source/Runtime/SlateCore/Private/Styling/StarshipCoreStyle.cpp
--- a/Engine/Source/Runtime/SlateCore/Private/Styling/StarshipCoreStyle.cpp
+++ b/Engine/Source/Runtime/SlateCore/Private/Styling/StarshipCoreStyle.cpp
@@ -1404,8 +1404,8 @@ void FStarshipCoreStyle::SetupCheckboxStyles(TSharedRef<FStyle>& Style)
 		.SetForegroundColor(FStyleColors::Foreground)
 		.SetHoveredForegroundColor(FStyleColors::ForegroundHover)
 		.SetPressedForegroundColor(FStyleColors::ForegroundHover)
-		.SetCheckedForegroundColor(FStyleColors::Primary)
-		.SetCheckedHoveredForegroundColor(FStyleColors::PrimaryHover)
+		.SetCheckedForegroundColor(FStyleColors::White)
+		.SetCheckedHoveredForegroundColor(FStyleColors::White)
 		.SetPadding(ToggleButtonMargins);
 
 	//@todo Recolor - Get rid of this. ToggleButtonCheckbox is what we should use
