Fix for no unity error
#rnx

diff --git a/Engine/Source/Runtime/SlateCore/Public/Layout/ChildrenBase.h b/Engine/Source/Runtime/SlateCore/Public/Layout/ChildrenBase.h
--- a/Engine/Source/Runtime/SlateCore/Public/Layout/ChildrenBase.h
+++ b/Engine/Source/Runtime/SlateCore/Public/Layout/ChildrenBase.h
@@ -4,6 +4,7 @@
 
 #include "CoreMinimal.h"
 #include "SlotBase.h"
+#include "SlateGlobals.h"
 
 class SWidget;
 class FSlotBase;
