Slate: Remove had used to detect rare thread issue.

diff --git a/Engine/Source/Runtime/SlateCore/Public/Layout/ChildrenBase.h b/Engine/Source/Runtime/SlateCore/Public/Layout/ChildrenBase.h
--- a/Engine/Source/Runtime/SlateCore/Public/Layout/ChildrenBase.h
+++ b/Engine/Source/Runtime/SlateCore/Public/Layout/ChildrenBase.h
@@ -5,10 +5,16 @@
 #include "CoreMinimal.h"
 #include "SlotBase.h"
 #include "SlateGlobals.h"
+#include "Debugging/SlateDebugging.h"
+#include "Types/ReflectionMetadata.h"
 
 class SWidget;
 class FSlotBase;
 
+#ifndef UE_SLATE_WITH_CHILD_SLOT_DEBUGGING
+	#define UE_SLATE_WITH_CHILD_SLOT_DEBUGGING 0
+#endif // !UE_SLATE_WITH_CHILD_SLOT_DEBUGGING
+
 /**
  * FChildren is an interface that must be implemented by all child containers.
  * It allows iteration over a list of any Widget's children regardless of how
@@ -54,10 +60,22 @@ public:
 	template<typename Predicate>
 	void ForEachWidget(Predicate Pred)
 	{
+#if WITH_SLATE_DEBUGGING 
+		TGuardValue<bool> IteratingGuard(bIsIteratingChildren, true);
+#endif
+
 		int32 WidgetCount = Num();
 		for (int32 Index = 0; Index < WidgetCount; ++Index)
 		{
-			Pred(GetChildRefAt(Index).GetWidget());
+			SWidget& ChildWidget = GetChildRefAt(Index).GetWidget();
+			Pred(ChildWidget);
+#if UE_SLATE_WITH_CHILD_SLOT_DEBUGGING 
+			UE_CLOG(WidgetCount != Num(), LogSlate, Error,
+				TEXT("Num Children changed during iteration! Child: %s, Owner: %s [%s]"),
+				*FReflectionMetaData::GetWidgetDebugInfo(ChildWidget),
+				*FReflectionMetaData::GetWidgetDebugInfo(Owner),
+				*Name.ToString());
+#endif // UE_SLATE_WITH_CHILD_SLOT_DEBUGGING
 		}
 	}
 
@@ -65,10 +83,22 @@ public:
 	template<typename Predicate>
 	void ForEachWidget(Predicate Pred) const
 	{
+#if WITH_SLATE_DEBUGGING 
+		TGuardValue<bool> IteratingGuard(bIsIteratingChildren, true);
+#endif
+
 		int32 WidgetCount = Num();
 		for (int32 Index = 0; Index < WidgetCount; ++Index)
 		{
-			Pred(GetChildRefAt(Index).GetWidget());
+			const SWidget& ChildWidget = GetChildRefAt(Index).GetWidget();
+			Pred(ChildWidget);
+#if UE_SLATE_WITH_CHILD_SLOT_DEBUGGING 
+			UE_CLOG(WidgetCount != Num(), LogSlate, Error,
+				TEXT("Num Children changed during iteration! Child: %s, Owner: %s [%s]"),
+				*FReflectionMetaData::GetWidgetDebugInfo(ChildWidget),
+				*FReflectionMetaData::GetWidgetDebugInfo(Owner),
+				*Name.ToString());
+#endif // UE_SLATE_WITH_CHILD_SLOT_DEBUGGINGs
 		}
 	}
 
@@ -170,7 +200,17 @@ protected:
 
 
 protected:
-	virtual ~FChildren(){}
+#if !WITH_SLATE_DEBUGGING 
+	virtual ~FChildren() = default;
+#else
+	virtual ~FChildren()
+	{
+		UE_CLOG(bIsIteratingChildren, LogSlate, Error,
+			TEXT("Destroying widget while iterating children! Owner: %s [%s]"),
+			*FReflectionMetaData::GetWidgetDebugInfo(Owner),
+			*Name.ToString());
+	}
+#endif // !WITH_SLATE_DEBUGGING
 
 protected:
 	UE_DEPRECATED(5.0, "Direct access to Owner is now deprecated. Use the getter.")
@@ -178,5 +218,9 @@ protected:
 
 private:
 	FName Name;
+
+#if WITH_SLATE_DEBUGGING
+	mutable bool bIsIteratingChildren = false;
+#endif // WITH_SLATE_DEBUGGING
 };
 PRAGMA_ENABLE_DEPRECATION_WARNINGS
