Slate: Fix reference to rvalue in FChildren

diff --git a/Engine/Source/Runtime/SlateCore/Public/Layout/ChildrenBase.h b/Engine/Source/Runtime/SlateCore/Public/Layout/ChildrenBase.h
--- a/Engine/Source/Runtime/SlateCore/Public/Layout/ChildrenBase.h
+++ b/Engine/Source/Runtime/SlateCore/Public/Layout/ChildrenBase.h
@@ -67,12 +67,12 @@ public:
 		int32 WidgetCount = Num();
 		for (int32 Index = 0; Index < WidgetCount; ++Index)
 		{
-			SWidget& ChildWidget = GetChildRefAt(Index).GetWidget();
-			Pred(ChildWidget);
+			FWidgetRef WidgetRef = GetChildRefAt(Index);
+			Pred(WidgetRef.GetWidget());
 #if UE_SLATE_WITH_CHILD_SLOT_DEBUGGING 
 			UE_CLOG(WidgetCount != Num(), LogSlate, Error,
 				TEXT("Num Children changed during iteration! Child: %s, Owner: %s [%s]"),
-				*FReflectionMetaData::GetWidgetDebugInfo(ChildWidget),
+				*FReflectionMetaData::GetWidgetDebugInfo(WidgetRef.GetWidget()),
 				*FReflectionMetaData::GetWidgetDebugInfo(Owner),
 				*Name.ToString());
 #endif // UE_SLATE_WITH_CHILD_SLOT_DEBUGGING
@@ -90,12 +90,12 @@ public:
 		int32 WidgetCount = Num();
 		for (int32 Index = 0; Index < WidgetCount; ++Index)
 		{
-			const SWidget& ChildWidget = GetChildRefAt(Index).GetWidget();
-			Pred(ChildWidget);
+			FConstWidgetRef WidgetRef = GetChildRefAt(Index);
+			Pred(WidgetRef.GetWidget());
 #if UE_SLATE_WITH_CHILD_SLOT_DEBUGGING 
 			UE_CLOG(WidgetCount != Num(), LogSlate, Error,
 				TEXT("Num Children changed during iteration! Child: %s, Owner: %s [%s]"),
-				*FReflectionMetaData::GetWidgetDebugInfo(ChildWidget),
+				*FReflectionMetaData::GetWidgetDebugInfo(WidgetRef.GetWidget()),
 				*FReflectionMetaData::GetWidgetDebugInfo(Owner),
 				*Name.ToString());
 #endif // UE_SLATE_WITH_CHILD_SLOT_DEBUGGINGs
