Fix FVector2D variant casts - Runtime
Submitted on behalf of fred.kimberley

diff --git a/Engine/Source/Runtime/AppFramework/Private/Widgets/Testing/STestSuite.cpp b/Engine/Source/Runtime/AppFramework/Private/Widgets/Testing/STestSuite.cpp
--- a/Engine/Source/Runtime/AppFramework/Private/Widgets/Testing/STestSuite.cpp
+++ b/Engine/Source/Runtime/AppFramework/Private/Widgets/Testing/STestSuite.cpp
@@ -1148,15 +1148,15 @@ private:
 	int32 TestCustomVerts(const FOnPaintHandlerParams& InParams)
 	{
 		const float Radius = FMath::Min(InParams.Geometry.GetLocalSize().X, InParams.Geometry.GetLocalSize().Y) * 0.5f;
-		const FVector2D Center = InParams.Geometry.AbsolutePosition + InParams.Geometry.GetLocalSize() * 0.5f;
+		const FVector2D Center = FVector2D(InParams.Geometry.AbsolutePosition) + InParams.Geometry.GetLocalSize() * 0.5f;
 
 		const FSlateBrush* MyBrush = FCoreStyle::Get().GetBrush("ColorWheel.HueValueCircle");
 
 		FSlateResourceHandle Handle = MyBrush->GetRenderingResource();
 		const FSlateShaderResourceProxy* ResourceProxy = Handle.GetResourceProxy();
 
-		FVector2D UVCenter = FVector2D::ZeroVector;
-		FVector2D UVRadius = FVector2D(1,1);
+		FVector2f UVCenter = FVector2f::ZeroVector;
+		FVector2f UVRadius = FVector2f(1,1);
 		if (ResourceProxy != nullptr)
 		{
 			UVRadius = 0.5f * ResourceProxy->SizeUV;
@@ -5601,7 +5601,7 @@ TOptional<FSlateRenderTransform> GetTestRenderTransform()
 
 FVector2D GetTestRenderTransformPivot()
 {
-	return TestSuiteNS::Offset;
+	return FVector2D(TestSuiteNS::Offset);
 }
 
 void SRenderTransformManipulatorWidgetImpl::Construct(const FArguments& InArgs)
diff --git a/Engine/Source/Runtime/Core/Public/Math/TransformCalculus.h b/Engine/Source/Runtime/Core/Public/Math/TransformCalculus.h
--- a/Engine/Source/Runtime/Core/Public/Math/TransformCalculus.h
+++ b/Engine/Source/Runtime/Core/Public/Math/TransformCalculus.h
@@ -401,8 +401,9 @@ inline VectorType TransformVector(const TransformType& Transform, const VectorTy
 {
 	return Transform.TransformVector(Vector);
 }
+
 /**
- * Generic implementation of TransformPoint for 2D vectors. Attempts to use a member function of the TransformType.
+ * Generic implementation of TransformPoint for 2D vectors with double precision. Attempts to use a member function of the TransformType.
  */
 template<typename TransformType>
 inline FVector2D TransformPoint(const TransformType& Transform, const FVector2D& Point)
@@ -410,6 +411,15 @@ inline FVector2D TransformPoint(const TransformType& Transform, const FVector2D&
 	return Transform.TransformPoint(Point);
 }
 
+/**
+ * Generic implementation of TransformPoint for 2D vectors with float precision. Attempts to use a member function of the TransformType.
+ */
+template<typename TransformType>
+inline FVector2f TransformPoint(const TransformType& Transform, const FVector2f& Point)
+{
+	return FVector2f(Transform.TransformPoint(FVector2D(Point)));
+}
+
 /**
  * Generic implementation of TransformVector for 2D vectors. Attempts to use a member function of the TransformType.
  */
diff --git a/Engine/Source/Runtime/Core/Public/Math/TransformCalculus2D.h b/Engine/Source/Runtime/Core/Public/Math/TransformCalculus2D.h
--- a/Engine/Source/Runtime/Core/Public/Math/TransformCalculus2D.h
+++ b/Engine/Source/Runtime/Core/Public/Math/TransformCalculus2D.h
@@ -106,11 +106,13 @@ public:
 	explicit FScale2D(float InScaleX, float InScaleY) :Scale(InScaleX, InScaleY) {}
 	/** Ctor. initialize from an FVector defining the 3D scale. */
 	explicit FScale2D(const FVector2D& InScale) :Scale((FVector2f)InScale) {}
+	/** Ctor. initialize from an FVector defining the 3D scale. */
+	explicit FScale2D(const FVector2f& InScale) :Scale(InScale) {}
 	
 	/** Transform 2D Point */
 	FVector2D TransformPoint(const FVector2D& Point) const
 	{
-		return Scale * Point;
+		return FVector2D(Scale) * Point;
 	}
 	/** Transform 2D Vector*/
 	FVector2D TransformVector(const FVector2D& Vector) const
@@ -126,7 +128,7 @@ public:
 	/** Invert the scale. */
 	FScale2D Inverse() const
 	{
-		return FScale2D(FVector2f(1.0f / Scale.X, 1.0f / Scale.Y));
+		return FScale2D(1.0f / Scale.X, 1.0f / Scale.Y);
 	}
 
 	/** Equality. */
@@ -193,7 +195,7 @@ public:
 	 */
 	FVector2D TransformPoint(const FVector2D& Point) const
 	{
-		return Point + FVector2D(Point.Y, Point.X) * Shear;
+		return Point + FVector2D(Point.Y, Point.X) * FVector2D(Shear);
 	}
 	/** Transform 2D Vector*/
 	FVector2D TransformVector(const FVector2D& Vector) const
@@ -295,7 +297,7 @@ public:
 	 */
 	FQuat2D Concatenate(const FQuat2D& RHS) const
 	{
-		return FQuat2D(TransformPoint(RHS.Rot));
+		return FQuat2D(TransformPoint(FVector2D(RHS.Rot)));
 	}
 	/**
 	 * Invert the rotation  defined by complex numbers:
@@ -616,7 +618,7 @@ public:
 	{
 		return FTransform2D(
 			::Concatenate(M, RHS.M),
-			::Concatenate(::TransformPoint(RHS.M, Trans), RHS.Trans));
+			::Concatenate(::TransformPoint(RHS.M, FVector2D(Trans)), FVector2D(RHS.Trans)));
 	}
 
 	/** 
@@ -671,14 +673,14 @@ public:
 	/** Access to the translation */
 	const FVector2D GetTranslation() const { return FVector2D(Trans); }
 
-	void SetTranslation(const FVector2D& InTrans) { Trans = InTrans; }
+	void SetTranslation(const FVector2D& InTrans) { Trans = FVector2f(InTrans); }
 
 	/**
 	 * Specialized function to determine if a transform is precisely the identity transform. Uses exact float comparison, so rounding error is not considered.
 	 */
 	bool IsIdentity() const
 	{
-		return M.IsIdentity() && Trans == FVector2D::ZeroVector;
+		return M.IsIdentity() && Trans == FVector2f::ZeroVector;
 	}
 
 	/**
diff --git a/Engine/Source/Runtime/Datasmith/DatasmithCore/Private/DatasmithUtils.cpp b/Engine/Source/Runtime/Datasmith/DatasmithCore/Private/DatasmithUtils.cpp
--- a/Engine/Source/Runtime/Datasmith/DatasmithCore/Private/DatasmithUtils.cpp
+++ b/Engine/Source/Runtime/Datasmith/DatasmithCore/Private/DatasmithUtils.cpp
@@ -270,9 +270,9 @@ bool FDatasmithMeshUtils::ToRawMesh(const FDatasmithMesh& Mesh, FRawMesh& RawMes
 
 			Mesh.GetFaceUV( FaceIndex, UVChannel, UVIndex1, UVIndex2, UVIndex3 );
 
-			RawMesh.WedgeTexCoords[ UVChannel ].Add( Mesh.GetUV( UVChannel, UVIndex1 ) );
-			RawMesh.WedgeTexCoords[ UVChannel ].Add( Mesh.GetUV( UVChannel, UVIndex2 ) );
-			RawMesh.WedgeTexCoords[ UVChannel ].Add( Mesh.GetUV( UVChannel, UVIndex3 ) );
+			RawMesh.WedgeTexCoords[ UVChannel ].Add( FVector2f(Mesh.GetUV( UVChannel, UVIndex1 )) );
+			RawMesh.WedgeTexCoords[ UVChannel ].Add( FVector2f(Mesh.GetUV( UVChannel, UVIndex2 )) );
+			RawMesh.WedgeTexCoords[ UVChannel ].Add( FVector2f(Mesh.GetUV( UVChannel, UVIndex3 )) );
 		}
 	}
 
@@ -384,7 +384,7 @@ bool FDatasmithMeshUtils::ToMeshDescription(FDatasmithMesh& DsMesh, FMeshDescrip
 				FVector2D UVVector = DsMesh.GetUV(UVChannelIndex, UV[CornerIndex]);
 				if (!UVVector.ContainsNaN())
 				{
-					VertexInstanceUVs.Set(CornerVertexInstanceIDs[CornerIndex], UVChannelIndex, UVVector);
+					VertexInstanceUVs.Set(CornerVertexInstanceIDs[CornerIndex], UVChannelIndex, FVector2f(UVVector));
 				}
 			}
 		}
diff --git a/Engine/Source/Runtime/Engine/Private/ActiveSound.cpp b/Engine/Source/Runtime/Engine/Private/ActiveSound.cpp
--- a/Engine/Source/Runtime/Engine/Private/ActiveSound.cpp
+++ b/Engine/Source/Runtime/Engine/Private/ActiveSound.cpp
@@ -1421,7 +1421,7 @@ float FActiveSound::GetAttenuationFrequency(const FSoundAttenuationSettings* Set
 	else
 	{
 		// In manual absorption mode, the frequency ranges are interpreted as a true "range"
-		FVector2f ActualFreqRange(FMath::Min(FrequencyRange.X, FrequencyRange.Y), FMath::Max(FrequencyRange.X, FrequencyRange.Y));
+		FVector2D ActualFreqRange(FMath::Min(FrequencyRange.X, FrequencyRange.Y), FMath::Max(FrequencyRange.X, FrequencyRange.Y));
 
 		// Normalize the distance values to a value between 0 and 1
 		FVector2f AbsorptionDistanceRange = { Settings->LPFRadiusMin, Settings->LPFRadiusMax };
@@ -1434,12 +1434,12 @@ float FActiveSound::GetAttenuationFrequency(const FSoundAttenuationSettings* Set
 		if (Settings->bEnableLogFrequencyScaling)
 		{
 			// Use the mapped value in the log scale mapping
-			OutputFrequency = Audio::GetLogFrequencyClamped(MappedFrequencyValue, FVector2f(0.0f, 1.0f), ActualFreqRange);
+			OutputFrequency = Audio::GetLogFrequencyClamped(MappedFrequencyValue, FVector2D(0.0f, 1.0f), ActualFreqRange);
 		}
 		else
 		{
 			// Do a straight linear interpolation between the absorption frequency ranges
-			OutputFrequency = FMath::GetMappedRangeValueClamped(FVector2f(0.0f, 1.0f), ActualFreqRange, MappedFrequencyValue);
+			OutputFrequency = FMath::GetMappedRangeValueClamped(FVector2f(0.0f, 1.0f), FVector2f(ActualFreqRange), MappedFrequencyValue);
 		}
 	}
 
diff --git a/Engine/Source/Runtime/Engine/Private/Components/DirectionalLightComponent.cpp b/Engine/Source/Runtime/Engine/Private/Components/DirectionalLightComponent.cpp
--- a/Engine/Source/Runtime/Engine/Private/Components/DirectionalLightComponent.cpp
+++ b/Engine/Source/Runtime/Engine/Private/Components/DirectionalLightComponent.cpp
@@ -379,7 +379,7 @@ public:
 		LightParameters.Direction = -GetDirection();
 		LightParameters.Tangent = -GetDirection();
 
-		LightParameters.SpotAngles = FVector2D(0, 0);
+		LightParameters.SpotAngles = FVector2f::ZeroVector;
 		LightParameters.SpecularScale = SpecularScale;
 		LightParameters.SourceRadius = FMath::Sin( 0.5f * FMath::DegreesToRadians( LightSourceAngle ) );
 		LightParameters.SoftSourceRadius = FMath::Sin( 0.5f * FMath::DegreesToRadians( LightSourceSoftAngle ) );
diff --git a/Engine/Source/Runtime/Engine/Private/Components/LineBatchComponent.cpp b/Engine/Source/Runtime/Engine/Private/Components/LineBatchComponent.cpp
--- a/Engine/Source/Runtime/Engine/Private/Components/LineBatchComponent.cpp
+++ b/Engine/Source/Runtime/Engine/Private/Components/LineBatchComponent.cpp
@@ -63,7 +63,7 @@ void FLineBatcherSceneProxy::GetDynamicMeshElements(const TArray<const FSceneVie
 				// set up geometry
 				for (int32 VertIdx=0; VertIdx < M.MeshVerts.Num(); ++VertIdx)
 				{
-					MeshBuilder.AddVertex( M.MeshVerts[VertIdx], FVector2D::ZeroVector, PosX, PosY, PosZ, FColor::White );
+					MeshBuilder.AddVertex( M.MeshVerts[VertIdx], FVector2f::ZeroVector, PosX, PosY, PosZ, FColor::White );
 				}
 				//MeshBuilder.AddTriangles(M.MeshIndices);
 				for (int32 Idx=0; Idx < M.MeshIndices.Num(); Idx+=3)
diff --git a/Engine/Source/Runtime/Engine/Private/Components/MaterialBillboardComponent.cpp b/Engine/Source/Runtime/Engine/Private/Components/MaterialBillboardComponent.cpp
--- a/Engine/Source/Runtime/Engine/Private/Components/MaterialBillboardComponent.cpp
+++ b/Engine/Source/Runtime/Engine/Private/Components/MaterialBillboardComponent.cpp
@@ -204,10 +204,10 @@ public:
 							StaticMeshVertexBuffers.PositionVertexBuffer.VertexPosition(WriteOffset + 2) = -WorldSizeX * TangentY + -WorldSizeY * TangentX;
 							StaticMeshVertexBuffers.PositionVertexBuffer.VertexPosition(WriteOffset + 3) = +WorldSizeX * TangentY + -WorldSizeY * TangentX;
 
-							StaticMeshVertexBuffers.StaticMeshVertexBuffer.SetVertexUV(WriteOffset + 0, 0, FVector2D(0, 0));
-							StaticMeshVertexBuffers.StaticMeshVertexBuffer.SetVertexUV(WriteOffset + 1, 0, FVector2D(0, 1));
-							StaticMeshVertexBuffers.StaticMeshVertexBuffer.SetVertexUV(WriteOffset + 2, 0, FVector2D(1, 0));
-							StaticMeshVertexBuffers.StaticMeshVertexBuffer.SetVertexUV(WriteOffset + 3, 0, FVector2D(1, 1));
+							StaticMeshVertexBuffers.StaticMeshVertexBuffer.SetVertexUV(WriteOffset + 0, 0, FVector2f(0, 0));
+							StaticMeshVertexBuffers.StaticMeshVertexBuffer.SetVertexUV(WriteOffset + 1, 0, FVector2f(0, 1));
+							StaticMeshVertexBuffers.StaticMeshVertexBuffer.SetVertexUV(WriteOffset + 2, 0, FVector2f(1, 0));
+							StaticMeshVertexBuffers.StaticMeshVertexBuffer.SetVertexUV(WriteOffset + 3, 0, FVector2f(1, 1));
 
 							// Set up the FMeshElement.
 							FMeshBatch& Mesh = Collector.AllocateMesh();
diff --git a/Engine/Source/Runtime/Engine/Private/Components/ModelComponent.cpp b/Engine/Source/Runtime/Engine/Private/Components/ModelComponent.cpp
--- a/Engine/Source/Runtime/Engine/Private/Components/ModelComponent.cpp
+++ b/Engine/Source/Runtime/Engine/Private/Components/ModelComponent.cpp
@@ -655,7 +655,7 @@ bool UModelComponent::GetPhysicsTriMeshData(struct FTriMeshCollisionData* Collis
 		CollisionData->Vertices[i] = Model->VertexBuffer.Vertices[i].Position;
 		if (bCopyUVs)
 		{
-			CollisionData->UVs[0][i] = Model->VertexBuffer.Vertices[i].TexCoord;
+			CollisionData->UVs[0][i] = FVector2D(Model->VertexBuffer.Vertices[i].TexCoord);
 		}
 	}
 
diff --git a/Engine/Source/Runtime/Engine/Private/Components/PointLightComponent.cpp b/Engine/Source/Runtime/Engine/Private/Components/PointLightComponent.cpp
--- a/Engine/Source/Runtime/Engine/Private/Components/PointLightComponent.cpp
+++ b/Engine/Source/Runtime/Engine/Private/Components/PointLightComponent.cpp
@@ -40,7 +40,7 @@ void FPointLightSceneProxy::GetLightShaderParameters(FLightRenderParameters& Lig
 	// TODO LWC - GetDirection() seems like it needs to be normalized, somehow accumulating error with large-scale position values
 	LightParameters.Direction = -GetDirection();
 	LightParameters.Tangent = FVector(WorldToLight.M[0][2], WorldToLight.M[1][2], WorldToLight.M[2][2]);
-	LightParameters.SpotAngles = FVector2D( -2.0f, 1.0f );
+	LightParameters.SpotAngles = FVector2f( -2.0f, 1.0f );
 	LightParameters.SpecularScale = SpecularScale;
 	LightParameters.SourceRadius = SourceRadius;
 	LightParameters.SoftSourceRadius = SoftSourceRadius;
diff --git a/Engine/Source/Runtime/Engine/Private/Components/RectLightComponent.cpp b/Engine/Source/Runtime/Engine/Private/Components/RectLightComponent.cpp
--- a/Engine/Source/Runtime/Engine/Private/Components/RectLightComponent.cpp
+++ b/Engine/Source/Runtime/Engine/Private/Components/RectLightComponent.cpp
@@ -230,7 +230,7 @@ void FRectLightSceneProxy::GetLightShaderParameters(FLightRenderParameters& Ligh
 
 	LightParameters.Direction = -GetDirection();
 	LightParameters.Tangent = FVector(WorldToLight.M[0][2], WorldToLight.M[1][2], WorldToLight.M[2][2]);
-	LightParameters.SpotAngles = FVector2D(-2.0f, 1.0f);
+	LightParameters.SpotAngles = FVector2f(-2.0f, 1.0f);
 	LightParameters.SpecularScale = SpecularScale;
 	LightParameters.SourceRadius = SourceWidth * 0.5f;
 	LightParameters.SoftSourceRadius = 0.0f;
diff --git a/Engine/Source/Runtime/Engine/Private/Components/SkinnedMeshComponent.cpp b/Engine/Source/Runtime/Engine/Private/Components/SkinnedMeshComponent.cpp
--- a/Engine/Source/Runtime/Engine/Private/Components/SkinnedMeshComponent.cpp
+++ b/Engine/Source/Runtime/Engine/Private/Components/SkinnedMeshComponent.cpp
@@ -2933,7 +2933,7 @@ FVector2D USkinnedMeshComponent::GetVertexUV(int32 VertexIndex, uint32 UVChannel
 	int32 VertexBase = Section.BaseVertexIndex;
 	uint32 ClampedUVChannel = FMath::Min(UVChannel, LODData.StaticVertexBuffers.StaticMeshVertexBuffer.GetNumTexCoords());
 
-	return LODData.StaticVertexBuffers.StaticMeshVertexBuffer.GetVertexUV(VertexBase + VertIndex, ClampedUVChannel);
+	return FVector2D(LODData.StaticVertexBuffers.StaticMeshVertexBuffer.GetVertexUV(VertexBase + VertIndex, ClampedUVChannel));
 }
 
 void USkinnedMeshComponent::HideBone( int32 BoneIndex, EPhysBodyOp PhysBodyOption)
diff --git a/Engine/Source/Runtime/Engine/Private/Components/SplineMeshComponent.cpp b/Engine/Source/Runtime/Engine/Private/Components/SplineMeshComponent.cpp
--- a/Engine/Source/Runtime/Engine/Private/Components/SplineMeshComponent.cpp
+++ b/Engine/Source/Runtime/Engine/Private/Components/SplineMeshComponent.cpp
@@ -82,10 +82,10 @@ void FSplineMeshVertexFactoryShaderParameters::GetElementShaderBindings(
 	// LWC_TODO: Precision loss
 	ParamData[0] = FVector4f(SplineParams.StartPos, SplineParams.StartRoll);
 	ParamData[1] = FVector4f(SplineParams.StartTangent, SplineParams.EndRoll);
-	ParamData[2] = FVector4f(SplineParams.StartScale, SplineParams.StartOffset);
+	ParamData[2] = FVector4f(FVector2f(SplineParams.StartScale), FVector2f(SplineParams.StartOffset));
 	ParamData[3] = FVector4f(SplineParams.EndPos, (float)(int32)SplineProxy->bSmoothInterpRollScale);
 	ParamData[4] = FVector4f(SplineParams.EndTangent, SplineProxy->SplineMeshMinZ);
-	ParamData[5] = FVector4f(SplineParams.EndScale, SplineParams.EndOffset);
+	ParamData[5] = FVector4f(FVector2f(SplineParams.EndScale), FVector2f(SplineParams.EndOffset));
 	ParamData[6] = FVector4f(SplineProxy->SplineUpDir, SplineProxy->SplineMeshScaleZ);
 
 		FVector3f DirMask(0, 0, 0);
diff --git a/Engine/Source/Runtime/Engine/Private/Components/SpotLightComponent.cpp b/Engine/Source/Runtime/Engine/Private/Components/SpotLightComponent.cpp
--- a/Engine/Source/Runtime/Engine/Private/Components/SpotLightComponent.cpp
+++ b/Engine/Source/Runtime/Engine/Private/Components/SpotLightComponent.cpp
@@ -60,7 +60,7 @@ public:
 		LightParameters.FalloffExponent = FalloffExponent;
 		LightParameters.Direction = -GetDirection();
 		LightParameters.Tangent = FVector(WorldToLight.M[0][2], WorldToLight.M[1][2], WorldToLight.M[2][2]);
-		LightParameters.SpotAngles = FVector2D(CosOuterCone, InvCosConeDifference);
+		LightParameters.SpotAngles = FVector2f(CosOuterCone, InvCosConeDifference);
 		LightParameters.SpecularScale = SpecularScale;
 		LightParameters.SourceRadius = SourceRadius;
 		LightParameters.SoftSourceRadius = SoftSourceRadius;
diff --git a/Engine/Source/Runtime/Engine/Private/GeomTools.cpp b/Engine/Source/Runtime/Engine/Private/GeomTools.cpp
--- a/Engine/Source/Runtime/Engine/Private/GeomTools.cpp
+++ b/Engine/Source/Runtime/Engine/Private/GeomTools.cpp
@@ -25,7 +25,7 @@ static FClipSMVertex GetVert(const UStaticMesh* StaticMesh, int32 VertIndex)
 	const int32 NumUVs = LODModel.VertexBuffers.StaticMeshVertexBuffer.GetNumTexCoords();
 	for(int32 UVIndex = 0;UVIndex < NumUVs;UVIndex++)
 	{
-		Result.UVs[UVIndex] = LODModel.VertexBuffers.StaticMeshVertexBuffer.GetVertexUV(VertIndex,UVIndex);
+		Result.UVs[UVIndex] = FVector2D(LODModel.VertexBuffers.StaticMeshVertexBuffer.GetVertexUV(VertIndex,UVIndex));
 	}
 	for(int32 UVIndex = NumUVs;UVIndex < UE_ARRAY_COUNT(Result.UVs);UVIndex++)
 	{
diff --git a/Engine/Source/Runtime/Engine/Private/HierarchicalInstancedStaticMesh.cpp b/Engine/Source/Runtime/Engine/Private/HierarchicalInstancedStaticMesh.cpp
--- a/Engine/Source/Runtime/Engine/Private/HierarchicalInstancedStaticMesh.cpp
+++ b/Engine/Source/Runtime/Engine/Private/HierarchicalInstancedStaticMesh.cpp
@@ -3058,8 +3058,8 @@ void UHierarchicalInstancedStaticMeshComponent::PropagateLightingScenarioChange(
 					const int32 RenderIndex = GetRenderIndex(InstanceIndex);
 					if (RenderIndex != INDEX_NONE)
 					{
-						InstanceUpdateCmdBuffer.SetLightMapData(RenderIndex, MeshMapBuildData->PerInstanceLightmapData[InstanceIndex].LightmapUVBias);
-						InstanceUpdateCmdBuffer.SetShadowMapData(RenderIndex, MeshMapBuildData->PerInstanceLightmapData[InstanceIndex].ShadowmapUVBias);
+						InstanceUpdateCmdBuffer.SetLightMapData(RenderIndex, FVector2D(MeshMapBuildData->PerInstanceLightmapData[InstanceIndex].LightmapUVBias));
+						InstanceUpdateCmdBuffer.SetShadowMapData(RenderIndex, FVector2D(MeshMapBuildData->PerInstanceLightmapData[InstanceIndex].ShadowmapUVBias));
 					}
 				}
 			}
@@ -3100,8 +3100,8 @@ void UHierarchicalInstancedStaticMeshComponent::SetPerInstanceLightMapAndEditorD
 
 			if (MeshMapBuildData != nullptr && MeshMapBuildData->PerInstanceLightmapData.IsValidIndex(Index))
 			{
-				LightmapUVBias = MeshMapBuildData->PerInstanceLightmapData[Index].LightmapUVBias;
-				ShadowmapUVBias = MeshMapBuildData->PerInstanceLightmapData[Index].ShadowmapUVBias;
+				LightmapUVBias = FVector2D(MeshMapBuildData->PerInstanceLightmapData[Index].LightmapUVBias);
+				ShadowmapUVBias = FVector2D(MeshMapBuildData->PerInstanceLightmapData[Index].ShadowmapUVBias);
 
 				PerInstanceData.SetInstanceLightMapData(RenderIndex, LightmapUVBias, ShadowmapUVBias);
 			}
diff --git a/Engine/Source/Runtime/Engine/Private/InstancedStaticMesh.cpp b/Engine/Source/Runtime/Engine/Private/InstancedStaticMesh.cpp
--- a/Engine/Source/Runtime/Engine/Private/InstancedStaticMesh.cpp
+++ b/Engine/Source/Runtime/Engine/Private/InstancedStaticMesh.cpp
@@ -2371,8 +2371,8 @@ void UInstancedStaticMeshComponent::BuildRenderData(FStaticMeshInstanceData& Out
 
 		if (MeshMapBuildData != nullptr && MeshMapBuildData->PerInstanceLightmapData.IsValidIndex(Index))
 		{
-			LightmapUVBias = MeshMapBuildData->PerInstanceLightmapData[Index].LightmapUVBias;
-			ShadowmapUVBias = MeshMapBuildData->PerInstanceLightmapData[Index].ShadowmapUVBias;
+			LightmapUVBias = FVector2D(MeshMapBuildData->PerInstanceLightmapData[Index].LightmapUVBias);
+			ShadowmapUVBias = FVector2D(MeshMapBuildData->PerInstanceLightmapData[Index].ShadowmapUVBias);
 		}
 	
 		OutData.SetInstance(RenderIndex, FMatrix44f(InstanceData.Transform), RandomStream.GetFraction(), LightmapUVBias, ShadowmapUVBias);		// LWC_TODO: Precision loss
diff --git a/Engine/Source/Runtime/Engine/Private/Level.cpp b/Engine/Source/Runtime/Engine/Private/Level.cpp
--- a/Engine/Source/Runtime/Engine/Private/Level.cpp
+++ b/Engine/Source/Runtime/Engine/Private/Level.cpp
@@ -1845,7 +1845,7 @@ void ULevel::CreateModelComponents()
 						DestVertex->WorldTangentZ = NodeTangentZ;
 
 						// TEMP - Will be overridden when lighting is build!
-						Vert.ShadowTexCoord = DestVertex->TextureCoordinates[1];
+						Vert.ShadowTexCoord = FVector2f(DestVertex->TextureCoordinates[1]);
 
 						// Include the vertex in the surface's bounding box.
 						NodeGroup->BoundingBox += VertexWorldPosition;
diff --git a/Engine/Source/Runtime/Engine/Private/LightMap.cpp b/Engine/Source/Runtime/Engine/Private/LightMap.cpp
--- a/Engine/Source/Runtime/Engine/Private/LightMap.cpp
+++ b/Engine/Source/Runtime/Engine/Private/LightMap.cpp
@@ -307,12 +307,12 @@ struct FLightMapAllocation
 			{
 				// TODO: We currently only support one LOD of static lighting in foliage
 				// Need to create per-LOD instance data to fix that
-				MeshBuildData->PerInstanceLightmapData[InstanceIndex].LightmapUVBias = LightMap->GetCoordinateBias();
+				MeshBuildData->PerInstanceLightmapData[InstanceIndex].LightmapUVBias = FVector2f(LightMap->GetCoordinateBias());
 
 				const int32 RenderIndex = Component->GetRenderIndex(InstanceIndex);
 				if (RenderIndex != INDEX_NONE)
 				{
-					Component->InstanceUpdateCmdBuffer.SetLightMapData(RenderIndex, MeshBuildData->PerInstanceLightmapData[InstanceIndex].LightmapUVBias);
+					Component->InstanceUpdateCmdBuffer.SetLightMapData(RenderIndex, FVector2D(MeshBuildData->PerInstanceLightmapData[InstanceIndex].LightmapUVBias));
 					Component->MarkRenderStateDirty();
 				}
 			}
diff --git a/Engine/Source/Runtime/Engine/Private/LightmapUniformShaderParameters.cpp b/Engine/Source/Runtime/Engine/Private/LightmapUniformShaderParameters.cpp
--- a/Engine/Source/Runtime/Engine/Private/LightmapUniformShaderParameters.cpp
+++ b/Engine/Source/Runtime/Engine/Private/LightmapUniformShaderParameters.cpp
@@ -72,7 +72,7 @@ void GetPrecomputedLightingParameters(
 	const FShadowMapInteraction ShadowMapInteraction = LCI ? LCI->GetShadowMapInteraction(FeatureLevel) : FShadowMapInteraction();
 	if (ShadowMapInteraction.GetType() == SMIT_Texture)
 	{
-		Parameters.ShadowMapCoordinateScaleBias = FVector4f(ShadowMapInteraction.GetCoordinateScale(), ShadowMapInteraction.GetCoordinateBias());
+		Parameters.ShadowMapCoordinateScaleBias = FVector4f(FVector2f(ShadowMapInteraction.GetCoordinateScale()), FVector2f(ShadowMapInteraction.GetCoordinateBias()));
 		Parameters.StaticShadowMapMasks = FVector4f(ShadowMapInteraction.GetChannelValid(0), ShadowMapInteraction.GetChannelValid(1), ShadowMapInteraction.GetChannelValid(2), ShadowMapInteraction.GetChannelValid(3));
 		Parameters.InvUniformPenumbraSizes = ShadowMapInteraction.GetInvUniformPenumbraSize();
 	}
diff --git a/Engine/Source/Runtime/Engine/Private/Particles/ParticleCurveTexture.cpp b/Engine/Source/Runtime/Engine/Private/Particles/ParticleCurveTexture.cpp
--- a/Engine/Source/Runtime/Engine/Private/Particles/ParticleCurveTexture.cpp
+++ b/Engine/Source/Runtime/Engine/Private/Particles/ParticleCurveTexture.cpp
@@ -80,7 +80,7 @@ public:
 		FParticleCurveInjectionParameters Parameters;
 		Parameters.PixelScale.X = 1.0f / GParticleCurveTextureSizeX;
 		Parameters.PixelScale.Y = 1.0f / GParticleCurveTextureSizeY;
-		Parameters.CurveOffset = CurveOffset;
+		Parameters.CurveOffset = FVector2f(CurveOffset);
 		FParticleCurveInjectionBufferRef UniformBuffer = FParticleCurveInjectionBufferRef::CreateUniformBufferImmediate(Parameters, UniformBuffer_SingleDraw);
 		FRHIVertexShader* VertexShader = RHICmdList.GetBoundVertexShader();
 		SetUniformBufferParameter(RHICmdList, VertexShader, GetUniformBufferParameter<FParticleCurveInjectionParameters>(), UniformBuffer);
diff --git a/Engine/Source/Runtime/Engine/Private/Particles/ParticleEmitterInstances.cpp b/Engine/Source/Runtime/Engine/Private/Particles/ParticleEmitterInstances.cpp
--- a/Engine/Source/Runtime/Engine/Private/Particles/ParticleEmitterInstances.cpp
+++ b/Engine/Source/Runtime/Engine/Private/Particles/ParticleEmitterInstances.cpp
@@ -2887,7 +2887,7 @@ bool FParticleEmitterInstance::FillReplayData( FDynamicEmitterReplayDataBase& Ou
 		NewReplayData->NormalsSphereCenter = LODLevel->RequiredModule->NormalsSphereCenter;
 		NewReplayData->NormalsCylinderDirection = LODLevel->RequiredModule->NormalsCylinderDirection;
 
-		NewReplayData->PivotOffset = PivotOffset;
+		NewReplayData->PivotOffset = FVector2f(PivotOffset);
 
 		NewReplayData->bRemoveHMDRoll = LODLevel->RequiredModule->bRemoveHMDRoll;
 		NewReplayData->MinFacingCameraBlendDistance = LODLevel->RequiredModule->MinFacingCameraBlendDistance;
diff --git a/Engine/Source/Runtime/Engine/Private/Particles/ParticleGpuSimulation.cpp b/Engine/Source/Runtime/Engine/Private/Particles/ParticleGpuSimulation.cpp
--- a/Engine/Source/Runtime/Engine/Private/Particles/ParticleGpuSimulation.cpp
+++ b/Engine/Source/Runtime/Engine/Private/Particles/ParticleGpuSimulation.cpp
@@ -4201,7 +4201,7 @@ private:
 			FVector2D UVFlipSizeOffset = FVector2D(BaseSize.X < 0.0f ? 0.0f : 0.5f, BaseSize.Y < 0.0f ? 0.0f : 0.5f);
 			NewParticle->Size.X = (FMath::Abs(BaseSize.X) * EmitterInfo.InvMaxSize.X * 0.5f);
 			NewParticle->Size.Y = bSquare ? (NewParticle->Size.X) : (FMath::Abs(BaseSize.Y) * EmitterInfo.InvMaxSize.Y * 0.5f);
-			NewParticle->Size += UVFlipSizeOffset;
+			NewParticle->Size += FVector2f(UVFlipSizeOffset);
 
 			NewParticle->Rotation = FMath::Fractional( TempParticle->Rotation * OneOverTwoPi );
 			NewParticle->RelativeRotationRate = TempParticle->BaseRotationRate * OneOverTwoPi * EmitterInfo.InvRotationRateScale / NewParticle->TimeScale;
@@ -5312,7 +5312,7 @@ static void SetGPUSpriteResourceData( FGPUSpriteResources* Resources, const FGPU
 	Resources->UniformParameters.RotationRateScale = InResourceData.RotationRateScale;
 	Resources->UniformParameters.CameraMotionBlurAmount = InResourceData.CameraMotionBlurAmount;
 
-	Resources->UniformParameters.PivotOffset = InResourceData.PivotOffset;
+	Resources->UniformParameters.PivotOffset = FVector2f(InResourceData.PivotOffset);
 
 	Resources->SimulationParameters.AttributeCurve = GParticleCurveTexture.ComputeCurveScaleBias(Resources->SimulationAttrTexelAllocation);
 	Resources->SimulationParameters.AttributeCurveScale = (FVector4f)InResourceData.SimulationAttrCurveScale; // LWC_TODO: change property to FVector4f
diff --git a/Engine/Source/Runtime/Engine/Private/Particles/ParticleSystemRender.cpp b/Engine/Source/Runtime/Engine/Private/Particles/ParticleSystemRender.cpp
--- a/Engine/Source/Runtime/Engine/Private/Particles/ParticleSystemRender.cpp
+++ b/Engine/Source/Runtime/Engine/Private/Particles/ParticleSystemRender.cpp
@@ -485,7 +485,7 @@ bool FDynamicSpriteEmitterData::GetVertexAndIndexData(void* VertexData, void* Dy
 			FillVertex->OldPosition = ParticleOldPosition;
 			// Create a floating point particle ID from the counter, map into approximately 0-1
 			FillVertex->ParticleId = (Particle.Flags & STATE_CounterMask) / 10000.0f;
-			FillVertex->Size = GetParticleSizeWithUVFlipInSign(Particle, Size);
+			FillVertex->Size = FVector2f(GetParticleSizeWithUVFlipInSign(Particle, Size));
 			FillVertex->Rotation = Particle.Rotation;
 			FillVertex->SubImageIndex = SubImageIndex;
 			FillVertex->Color = Particle.Color;
@@ -637,7 +637,7 @@ bool FDynamicSpriteEmitterData::GetVertexAndIndexDataNonInstanced(void* VertexDa
 			FillVertex[VertexIndex].OldPosition	= ParticleOldPosition;
 			// Create a floating point particle ID from the counter, map into approximately 0-1
 			FillVertex[VertexIndex].ParticleId = (Particle.Flags & STATE_CounterMask) / 10000.0f;
-			FillVertex[VertexIndex].Size = GetParticleSizeWithUVFlipInSign(Particle, Size);
+			FillVertex[VertexIndex].Size = FVector2f(GetParticleSizeWithUVFlipInSign(Particle, Size));
 			FillVertex[VertexIndex].Rotation	= Particle.Rotation;
 			FillVertex[VertexIndex].SubImageIndex = SubImageIndex;
 			FillVertex[VertexIndex].Color		= Particle.Color;
@@ -3232,7 +3232,7 @@ int32 FDynamicBeam2EmitterData::FillVertexData_NoNoise(FAsyncBufferFillData& Me)
 				Vertex->Position	= Location + Offset;
 				Vertex->OldPosition	= Location;
 				Vertex->ParticleId	= 0;
-				Vertex->Size		= Size;
+				Vertex->Size		= FVector2f(Size);
 				Vertex->Tex_U		= 0.0f;
 				Vertex->Tex_V		= 0.0f;
 				Vertex->Tex_U2		= 0.0f;
@@ -3245,7 +3245,7 @@ int32 FDynamicBeam2EmitterData::FillVertexData_NoNoise(FAsyncBufferFillData& Me)
 				Vertex->Position	= Location - Offset;
 				Vertex->OldPosition	= Location;
 				Vertex->ParticleId	= 0;
-				Vertex->Size		= Size;
+				Vertex->Size		= FVector2f(Size);
 				Vertex->Tex_U		= 0.0f;
 				Vertex->Tex_V		= 1.0f;
 				Vertex->Tex_U2		= 0.0f;
@@ -3270,7 +3270,7 @@ int32 FDynamicBeam2EmitterData::FillVertexData_NoNoise(FAsyncBufferFillData& Me)
 				Vertex->Position	= EndPoint + Offset;
 				Vertex->OldPosition	= Particle->OldLocation;
 				Vertex->ParticleId	= 0;
-				Vertex->Size		= Size;
+				Vertex->Size		= FVector2f(Size);
 				Vertex->Tex_U		= fUEnd;
 				Vertex->Tex_V		= 0.0f;
 				Vertex->Tex_U2		= 1.0f;
@@ -3283,7 +3283,7 @@ int32 FDynamicBeam2EmitterData::FillVertexData_NoNoise(FAsyncBufferFillData& Me)
 				Vertex->Position	= EndPoint - Offset;
 				Vertex->OldPosition	= Particle->OldLocation;
 				Vertex->ParticleId	= 0;
-				Vertex->Size		= Size;
+				Vertex->Size		= FVector2f(Size);
 				Vertex->Tex_U		= fUEnd;
 				Vertex->Tex_V		= 1.0f;
 				Vertex->Tex_U2		= 1.0f;
@@ -3412,7 +3412,7 @@ int32 FDynamicBeam2EmitterData::FillVertexData_NoNoise(FAsyncBufferFillData& Me)
 				Vertex->Position	= Location + Offset;
 				Vertex->OldPosition	= Location;
 				Vertex->ParticleId	= 0;
-				Vertex->Size		= Size;
+				Vertex->Size		= FVector2f(Size);
 				Vertex->Tex_U		= fU;
 				Vertex->Tex_V		= 0.0f;
 				Vertex->Tex_U2		= 0.0f;
@@ -3425,7 +3425,7 @@ int32 FDynamicBeam2EmitterData::FillVertexData_NoNoise(FAsyncBufferFillData& Me)
 				Vertex->Position	= Location - Offset;
 				Vertex->OldPosition	= Location;
 				Vertex->ParticleId	= 0;
-				Vertex->Size		= Size;
+				Vertex->Size		= FVector2f(Size);
 				Vertex->Tex_U		= fU;
 				Vertex->Tex_V		= 1.0f;
 				Vertex->Tex_U2		= 0.0f;
@@ -3474,7 +3474,7 @@ int32 FDynamicBeam2EmitterData::FillVertexData_NoNoise(FAsyncBufferFillData& Me)
 					Vertex->Position	= EndPoint + Offset;
 					Vertex->OldPosition	= EndPoint;
 					Vertex->ParticleId	= 0;
-					Vertex->Size		= Size;
+					Vertex->Size		= FVector2f(Size);
 					Vertex->Tex_U		= fU + fTextureIncrement;
 					Vertex->Tex_V		= 0.0f;
 					Vertex->Tex_U2		= Tex_U2;
@@ -3487,7 +3487,7 @@ int32 FDynamicBeam2EmitterData::FillVertexData_NoNoise(FAsyncBufferFillData& Me)
 					Vertex->Position	= EndPoint - Offset;
 					Vertex->OldPosition	= EndPoint;
 					Vertex->ParticleId	= 0;
-					Vertex->Size		= Size;
+					Vertex->Size		= FVector2f(Size);
 					Vertex->Tex_U		= fU + fTextureIncrement;
 					Vertex->Tex_V		= 1.0f;
 					Vertex->Tex_U2		= Tex_U2;
@@ -3752,7 +3752,7 @@ int32 FDynamicBeam2EmitterData::FillData_Noise(FAsyncBufferFillData& Me) const
 				Vertex->Position	= Location + LastOffset;
 				Vertex->OldPosition	= Location;
 				Vertex->ParticleId	= 0;
-				Vertex->Size		= Size;
+				Vertex->Size		= FVector2f(Size);
 				Vertex->Tex_U		= fU;
 				Vertex->Tex_V		= 0.0f;
 				Vertex->Rotation	= Particle->Rotation;
@@ -3763,7 +3763,7 @@ int32 FDynamicBeam2EmitterData::FillData_Noise(FAsyncBufferFillData& Me) const
 				Vertex->Position	= Location - LastOffset;
 				Vertex->OldPosition	= Location;
 				Vertex->ParticleId	= 0;
-				Vertex->Size		= Size;
+				Vertex->Size		= FVector2f(Size);
 				Vertex->Tex_U		= fU;
 				Vertex->Tex_V		= 1.0f;
 				Vertex->Rotation	= Particle->Rotation;
@@ -3909,7 +3909,7 @@ int32 FDynamicBeam2EmitterData::FillData_Noise(FAsyncBufferFillData& Me) const
 						Vertex->Position	= InterpDrawPos + Offset;
 						Vertex->OldPosition	= InterpDrawPos;
 						Vertex->ParticleId	= 0;
-						Vertex->Size		= Size;
+						Vertex->Size		= FVector2f(Size);
 						Vertex->Tex_U		= fU;
 						Vertex->Tex_V		= 0.0f;
 						Vertex->Rotation	= Particle->Rotation;
@@ -3920,7 +3920,7 @@ int32 FDynamicBeam2EmitterData::FillData_Noise(FAsyncBufferFillData& Me) const
 						Vertex->Position	= InterpDrawPos - Offset;
 						Vertex->OldPosition	= InterpDrawPos;
 						Vertex->ParticleId	= 0;
-						Vertex->Size		= Size;
+						Vertex->Size		= FVector2f(Size);
 						Vertex->Tex_U		= fU;
 						Vertex->Tex_V		= 1.0f;
 						Vertex->Rotation	= Particle->Rotation;
@@ -4028,7 +4028,7 @@ int32 FDynamicBeam2EmitterData::FillData_Noise(FAsyncBufferFillData& Me) const
 						Vertex->Position	= InterpDrawPos + Offset;
 						Vertex->OldPosition	= InterpDrawPos;
 						Vertex->ParticleId	= 0;
-						Vertex->Size		= Size;
+						Vertex->Size		= FVector2f(Size);
 						Vertex->Tex_U		= fU;
 						Vertex->Tex_V		= 0.0f;
 						Vertex->Rotation	= Particle->Rotation;
@@ -4039,7 +4039,7 @@ int32 FDynamicBeam2EmitterData::FillData_Noise(FAsyncBufferFillData& Me) const
 						Vertex->Position	= InterpDrawPos - Offset;
 						Vertex->OldPosition	= InterpDrawPos;
 						Vertex->ParticleId	= 0;
-						Vertex->Size		= Size;
+						Vertex->Size		= FVector2f(Size);
 						Vertex->Tex_U		= fU;
 						Vertex->Tex_V		= 1.0f;
 						Vertex->Rotation	= Particle->Rotation;
@@ -4164,7 +4164,7 @@ int32 FDynamicBeam2EmitterData::FillData_Noise(FAsyncBufferFillData& Me) const
 				Vertex->Position	= Location + LastOffset;
 				Vertex->OldPosition	= Location;
 				Vertex->ParticleId	= 0;
-				Vertex->Size		= Size;
+				Vertex->Size		= FVector2f(Size);
 				Vertex->Tex_U		= fU;
 				Vertex->Tex_V		= 0.0f;
 				Vertex->Rotation	= Particle->Rotation;
@@ -4175,7 +4175,7 @@ int32 FDynamicBeam2EmitterData::FillData_Noise(FAsyncBufferFillData& Me) const
 				Vertex->Position	= Location - LastOffset;
 				Vertex->OldPosition	= Location;
 				Vertex->ParticleId	= 0;
-				Vertex->Size		= Size;
+				Vertex->Size		= FVector2f(Size);
 				Vertex->Tex_U		= fU;
 				Vertex->Tex_V		= 1.0f;
 				Vertex->Rotation	= Particle->Rotation;
@@ -4323,7 +4323,7 @@ int32 FDynamicBeam2EmitterData::FillData_Noise(FAsyncBufferFillData& Me) const
 						Vertex->Position	= InterpDrawPos + Offset;
 						Vertex->OldPosition	= InterpDrawPos;
 						Vertex->ParticleId	= 0;
-						Vertex->Size		= Size;
+						Vertex->Size		= FVector2f(Size);
 						Vertex->Tex_U		= fU;
 						Vertex->Tex_V		= 0.0f;
 						Vertex->Rotation	= Particle->Rotation;
@@ -4334,7 +4334,7 @@ int32 FDynamicBeam2EmitterData::FillData_Noise(FAsyncBufferFillData& Me) const
 						Vertex->Position	= InterpDrawPos - Offset;
 						Vertex->OldPosition	= InterpDrawPos;
 						Vertex->ParticleId	= 0;
-						Vertex->Size		= Size;
+						Vertex->Size		= FVector2f(Size);
 						Vertex->Tex_U		= fU;
 						Vertex->Tex_V		= 1.0f;
 						Vertex->Rotation	= Particle->Rotation;
@@ -4442,7 +4442,7 @@ int32 FDynamicBeam2EmitterData::FillData_Noise(FAsyncBufferFillData& Me) const
 						Vertex->Position	= InterpDrawPos + Offset;
 						Vertex->OldPosition	= InterpDrawPos;
 						Vertex->ParticleId	= 0;
-						Vertex->Size		= Size;
+						Vertex->Size		= FVector2f(Size);
 						Vertex->Tex_U		= fU;
 						Vertex->Tex_V		= 0.0f;
 						Vertex->Rotation	= Particle->Rotation;
@@ -4453,7 +4453,7 @@ int32 FDynamicBeam2EmitterData::FillData_Noise(FAsyncBufferFillData& Me) const
 						Vertex->Position	= InterpDrawPos - Offset;
 						Vertex->OldPosition	= InterpDrawPos;
 						Vertex->ParticleId	= 0;
-						Vertex->Size		= Size;
+						Vertex->Size		= FVector2f(Size);
 						Vertex->Tex_U		= fU;
 						Vertex->Tex_V		= 1.0f;
 						Vertex->Rotation	= Particle->Rotation;
@@ -4738,7 +4738,7 @@ int32 FDynamicBeam2EmitterData::FillData_InterpolatedNoise(FAsyncBufferFillData&
 			Vertex->Position	= Location + LastOffset;
 			Vertex->OldPosition	= Location;
 			Vertex->ParticleId	= 0;
-			Vertex->Size		= Size;
+			Vertex->Size		= FVector2f(Size);
 			Vertex->Tex_U		= fU;
 			Vertex->Tex_V		= 0.0f;
 			Vertex->Rotation	= Particle->Rotation;
@@ -4749,7 +4749,7 @@ int32 FDynamicBeam2EmitterData::FillData_InterpolatedNoise(FAsyncBufferFillData&
 			Vertex->Position	= Location - LastOffset;
 			Vertex->OldPosition	= Location;
 			Vertex->ParticleId	= 0;
-			Vertex->Size		= Size;
+			Vertex->Size		= FVector2f(Size);
 			Vertex->Tex_U		= fU;
 			Vertex->Tex_V		= 1.0f;
 			Vertex->Rotation	= Particle->Rotation;
@@ -4945,7 +4945,7 @@ int32 FDynamicBeam2EmitterData::FillData_InterpolatedNoise(FAsyncBufferFillData&
 					Vertex->Position	= InterpDrawPos + Offset;
 					Vertex->OldPosition	= InterpDrawPos;
 					Vertex->ParticleId	= 0;
-					Vertex->Size		= Size;
+					Vertex->Size		= FVector2f(Size);
 					Vertex->Tex_U		= fU;
 					Vertex->Tex_V		= 0.0f;
 					Vertex->Rotation	= Particle->Rotation;
@@ -4956,7 +4956,7 @@ int32 FDynamicBeam2EmitterData::FillData_InterpolatedNoise(FAsyncBufferFillData&
 					Vertex->Position	= InterpDrawPos - Offset;
 					Vertex->OldPosition	= InterpDrawPos;
 					Vertex->ParticleId	= 0;
-					Vertex->Size		= Size;
+					Vertex->Size		= FVector2f(Size);
 					Vertex->Tex_U		= fU;
 					Vertex->Tex_V		= 1.0f;
 					Vertex->Rotation	= Particle->Rotation;
@@ -5064,7 +5064,7 @@ int32 FDynamicBeam2EmitterData::FillData_InterpolatedNoise(FAsyncBufferFillData&
 					Vertex->Position	= InterpDrawPos + Offset;
 					Vertex->OldPosition	= InterpDrawPos;
 					Vertex->ParticleId	= 0;
-					Vertex->Size		= Size;
+					Vertex->Size		= FVector2f(Size);
 					Vertex->Tex_U		= fU;
 					Vertex->Tex_V		= 0.0f;
 					Vertex->Rotation	= Particle->Rotation;
@@ -5075,7 +5075,7 @@ int32 FDynamicBeam2EmitterData::FillData_InterpolatedNoise(FAsyncBufferFillData&
 					Vertex->Position	= InterpDrawPos - Offset;
 					Vertex->OldPosition	= InterpDrawPos;
 					Vertex->ParticleId	= 0;
-					Vertex->Size		= Size;
+					Vertex->Size		= FVector2f(Size);
 					Vertex->Tex_U		= fU;
 					Vertex->Tex_V		= 1.0f;
 					Vertex->Rotation	= Particle->Rotation;
diff --git a/Engine/Source/Runtime/Engine/Private/PhysicsEngine/PhysDrawing.cpp b/Engine/Source/Runtime/Engine/Private/PhysicsEngine/PhysDrawing.cpp
--- a/Engine/Source/Runtime/Engine/Private/PhysicsEngine/PhysDrawing.cpp
+++ b/Engine/Source/Runtime/Engine/Private/PhysicsEngine/PhysDrawing.cpp
@@ -781,7 +781,7 @@ void FKConvexElem::DrawElemSolid(FPrimitiveDrawInterface* PDI, const FTransform&
 		const int32 NumVerts = VertexData.Num();
 		for (const FVector& Vert : VertexData)
 		{
-			MeshBuilder.AddVertex(ElemTM.TransformPosition(Vert), DummyUV, DummyTangentX, DummyTangentY, DummyTangentZ, FColor::White);
+			MeshBuilder.AddVertex(ElemTM.TransformPosition(Vert), FVector2f(DummyUV), DummyTangentX, DummyTangentY, DummyTangentZ, FColor::White);
 		}
 
 		for (int32 Base = 0; Base < NumIndices; Base += 3)
diff --git a/Engine/Source/Runtime/Engine/Private/PlayerController.cpp b/Engine/Source/Runtime/Engine/Private/PlayerController.cpp
--- a/Engine/Source/Runtime/Engine/Private/PlayerController.cpp
+++ b/Engine/Source/Runtime/Engine/Private/PlayerController.cpp
@@ -2038,7 +2038,7 @@ bool APlayerController::GetHitResultUnderFinger(ETouchIndex::Type FingerIndex, E
 		GetInputTouchState(FingerIndex, TouchPosition.X, TouchPosition.Y, bIsPressed);
 		if (bIsPressed)
 		{
-			bHit = GetHitResultAtScreenPosition(TouchPosition, TraceChannel, bTraceComplex, HitResult);
+			bHit = GetHitResultAtScreenPosition(FVector2D(TouchPosition), TraceChannel, bTraceComplex, HitResult);
 		}
 	}
 
@@ -2060,7 +2060,7 @@ bool APlayerController::GetHitResultUnderFingerByChannel(ETouchIndex::Type Finge
 		GetInputTouchState(FingerIndex, TouchPosition.X, TouchPosition.Y, bIsPressed);
 		if (bIsPressed)
 		{
-			bHit = GetHitResultAtScreenPosition(TouchPosition, TraceChannel, bTraceComplex, HitResult);
+			bHit = GetHitResultAtScreenPosition(FVector2D(TouchPosition), TraceChannel, bTraceComplex, HitResult);
 		}
 	}
 
@@ -2082,7 +2082,7 @@ bool APlayerController::GetHitResultUnderFingerForObjects(ETouchIndex::Type Fing
 		GetInputTouchState(FingerIndex, TouchPosition.X, TouchPosition.Y, bIsPressed);
 		if (bIsPressed)
 		{
-			bHit = GetHitResultAtScreenPosition(TouchPosition, ObjectTypes, bTraceComplex, HitResult);
+			bHit = GetHitResultAtScreenPosition(FVector2D(TouchPosition), ObjectTypes, bTraceComplex, HitResult);
 		}
 	}
 
diff --git a/Engine/Source/Runtime/Engine/Private/PrimitiveDrawingUtils.cpp b/Engine/Source/Runtime/Engine/Private/PrimitiveDrawingUtils.cpp
--- a/Engine/Source/Runtime/Engine/Private/PrimitiveDrawingUtils.cpp
+++ b/Engine/Source/Runtime/Engine/Private/PrimitiveDrawingUtils.cpp
@@ -76,10 +76,10 @@ void DrawPlane10x10(class FPrimitiveDrawInterface* PDI,const FMatrix& ObjectToWo
 			float U1 = FMath::Lerp(UVMin.X, UVMax.X, x1 * 0.5f + 0.5f);
 
 			// Calculate verts for a face pointing down Z
-			MeshBuilder.AddVertex(FVector(x0, y0, 0), FVector2D(U0, V0), FVector(1, 0, 0), FVector(0, 1, 0), FVector(0, 0, 1), FColor::White);
-			MeshBuilder.AddVertex(FVector(x0, y1, 0), FVector2D(U0, V1), FVector(1, 0, 0), FVector(0, 1, 0), FVector(0, 0, 1), FColor::White);
-			MeshBuilder.AddVertex(FVector(x1, y1, 0), FVector2D(U1, V1), FVector(1, 0, 0), FVector(0, 1, 0), FVector(0, 0, 1), FColor::White);
-			MeshBuilder.AddVertex(FVector(x1, y0, 0), FVector2D(U1, V0), FVector(1, 0, 0), FVector(0, 1, 0), FVector(0, 0, 1), FColor::White);
+			MeshBuilder.AddVertex(FVector(x0, y0, 0), FVector2f(U0, V0), FVector(1, 0, 0), FVector(0, 1, 0), FVector(0, 0, 1), FColor::White);
+			MeshBuilder.AddVertex(FVector(x0, y1, 0), FVector2f(U0, V1), FVector(1, 0, 0), FVector(0, 1, 0), FVector(0, 0, 1), FColor::White);
+			MeshBuilder.AddVertex(FVector(x1, y1, 0), FVector2f(U1, V1), FVector(1, 0, 0), FVector(0, 1, 0), FVector(0, 0, 1), FColor::White);
+			MeshBuilder.AddVertex(FVector(x1, y0, 0), FVector2f(U1, V0), FVector(1, 0, 0), FVector(0, 1, 0), FVector(0, 0, 1), FColor::White);
 
 			int Index = (x + y * TileCount) * 4;
 			MeshBuilder.AddTriangle(Index + 0, Index + 1, Index + 2);
@@ -92,12 +92,12 @@ void DrawPlane10x10(class FPrimitiveDrawInterface* PDI,const FMatrix& ObjectToWo
 
 void DrawTriangle(class FPrimitiveDrawInterface* PDI, const FVector& A, const FVector& B, const FVector& C, const FMaterialRenderProxy* MaterialRenderProxy, uint8 DepthPriorityGroup)
 {
-	FVector2D UVs[4] =
+	FVector2f UVs[4] =
 	{
-		FVector2D(0,0),
-		FVector2D(0,1),
-		FVector2D(1,1),
-		FVector2D(1,0),
+		FVector2f(0,0),
+		FVector2f(0,1),
+		FVector2f(1,1),
+		FVector2f(1,0),
 	};
 
 	FDynamicMeshBuilder MeshBuilder(PDI->View->GetFeatureLevel());
@@ -556,7 +556,7 @@ void BuildCylinderVerts(const FVector& Base, const FVector& XAxis, const FVector
 		FDynamicMeshVertex MeshVertex;
 
 		MeshVertex.Position = Vertex - TopOffset;
-		MeshVertex.TextureCoordinate[0] = TC;
+		MeshVertex.TextureCoordinate[0] = FVector2f(TC);
 
 		MeshVertex.SetTangents(
 			-ZAxis,
@@ -583,7 +583,7 @@ void BuildCylinderVerts(const FVector& Base, const FVector& XAxis, const FVector
 		FDynamicMeshVertex MeshVertex;
 
 		MeshVertex.Position = Vertex + TopOffset;
-		MeshVertex.TextureCoordinate[0] = TC;
+		MeshVertex.TextureCoordinate[0] = FVector2f(TC);
 
 		MeshVertex.SetTangents(
 			-ZAxis,
@@ -864,7 +864,7 @@ void DrawDisc(class FPrimitiveDrawInterface* PDI,const FVector& Base,const FVect
 		FDynamicMeshVertex MeshVertex;
 		MeshVertex.Position = Vertex;
 		MeshVertex.Color = Color;
-		MeshVertex.TextureCoordinate[0] = TC;
+		MeshVertex.TextureCoordinate[0] = FVector2f(TC);
 		MeshVertex.TextureCoordinate[0].X += TCStep * SideIndex;
 
 		MeshVertex.SetTangents(
@@ -934,7 +934,7 @@ void DrawFlatArrow(class FPrimitiveDrawInterface* PDI,const FVector& Base,const
 		FDynamicMeshVertex MeshVertex;
 		MeshVertex.Position = ArrowPoints[i];
 		MeshVertex.Color = Color;
-		MeshVertex.TextureCoordinate[0] = FVector2D(0.0f, 0.0f);;
+		MeshVertex.TextureCoordinate[0] = FVector2f(0.0f, 0.0f);;
 		MeshVertex.SetTangents(XAxis^YAxis, YAxis, XAxis);
 		MeshBuilder.AddVertex(MeshVertex); //Add bottom vertex
 	}
@@ -1836,7 +1836,7 @@ void DrawUVsInternal(FViewport* InViewport, FCanvas* InCanvas, int32 InTextYPos,
 
 				for (int32 Corner = 0; Corner < 3; Corner++)
 				{
-					UVs[Corner] = (VertexBuffer.GetVertexUV(Indices[i + Corner], UVChannel));
+					UVs[Corner] = FVector2D(VertexBuffer.GetVertexUV(Indices[i + Corner], UVChannel));
 					bOutOfBounds[Corner] = IsUVOutOfBounds(UVs[Corner]);
 				}
 
diff --git a/Engine/Source/Runtime/Engine/Private/Rendering/SkeletalMeshLODImporterData.cpp b/Engine/Source/Runtime/Engine/Private/Rendering/SkeletalMeshLODImporterData.cpp
--- a/Engine/Source/Runtime/Engine/Private/Rendering/SkeletalMeshLODImporterData.cpp
+++ b/Engine/Source/Runtime/Engine/Private/Rendering/SkeletalMeshLODImporterData.cpp
@@ -244,7 +244,7 @@ bool FSkeletalMeshImportData::ApplyRigToGeo(FSkeletalMeshImportData& Other)
 		{
 			SlowTask.EnterProgressFrame(1.0f);
 		}
-		const FVector2D& CurWedgeUV = Wedges[WedgeIndex].UVs[0];
+		const FVector2f& CurWedgeUV = Wedges[WedgeIndex].UVs[0];
 		int32 NewVertexIndex = (int32)(Wedges[WedgeIndex].VertexIndex);
 		FVector3f& NewPointA = Points[NewVertexIndex];
 		SkeletalMeshImportData::FTriangle& NewFace = Faces[(WedgeIndex / 3)];
@@ -330,9 +330,9 @@ bool FSkeletalMeshImportData::ApplyRigToGeo(FSkeletalMeshImportData& Other)
 					int32 OldVertexIndex = Other.Wedges[OldWedgeIndex].VertexIndex;
 					int32 OldFaceIndex = (OldWedgeIndex / 3);
 					int32 OldFaceCorner = (OldWedgeIndex % 3);
-					const FVector2D& OldUV = Other.Wedges[OldWedgeIndex].UVs[0];
+					const FVector2f& OldUV = Other.Wedges[OldWedgeIndex].UVs[0];
 					const FVector3f& OldNormal = Other.Faces[OldFaceIndex].TangentZ[OldFaceCorner];
-					float UVDelta = FVector2D::DistSquared(CurWedgeUV, OldUV);
+					float UVDelta = FVector2f::DistSquared(CurWedgeUV, OldUV);
 					float NormalDelta = FMath::Abs(FMath::Acos(FVector3f::DotProduct(NewNormal, OldNormal)));
 					if (UVDelta > MaxUVDistance)
 					{
diff --git a/Engine/Source/Runtime/Engine/Private/SceneView.cpp b/Engine/Source/Runtime/Engine/Private/SceneView.cpp
--- a/Engine/Source/Runtime/Engine/Private/SceneView.cpp
+++ b/Engine/Source/Runtime/Engine/Private/SceneView.cpp
@@ -2541,12 +2541,12 @@ void FSceneView::SetupCommonViewUniformBufferParameters(
 
 	ViewUniformShaderParameters.GlobalClippingPlane = FVector4f(TranslatedGlobalClippingPlane.X, TranslatedGlobalClippingPlane.Y, TranslatedGlobalClippingPlane.Z, -TranslatedGlobalClippingPlane.W);
 
-	ViewUniformShaderParameters.FieldOfViewWideAngles = 2.f * InViewMatrices.ComputeHalfFieldOfViewPerAxis();
-	ViewUniformShaderParameters.PrevFieldOfViewWideAngles = 2.f * InPrevViewMatrices.ComputeHalfFieldOfViewPerAxis();
+	ViewUniformShaderParameters.FieldOfViewWideAngles = FVector2f(2.f * InViewMatrices.ComputeHalfFieldOfViewPerAxis());	// LWC_TODO: Precision loss
+	ViewUniformShaderParameters.PrevFieldOfViewWideAngles = FVector2f(2.f * InPrevViewMatrices.ComputeHalfFieldOfViewPerAxis());	// LWC_TODO: Precision loss
 	ViewUniformShaderParameters.DiffuseOverrideParameter = LocalDiffuseOverrideParameter;
 	ViewUniformShaderParameters.SpecularOverrideParameter = SpecularOverrideParameter;
 	ViewUniformShaderParameters.NormalOverrideParameter = NormalOverrideParameter;
-	ViewUniformShaderParameters.RoughnessOverrideParameter = LocalRoughnessOverrideParameter;
+	ViewUniformShaderParameters.RoughnessOverrideParameter = FVector2f(LocalRoughnessOverrideParameter);	// LWC_TODO: Precision loss
 	ViewUniformShaderParameters.WorldCameraMovementSinceLastFrame = InViewMatrices.GetViewOrigin() - InPrevViewMatrices.GetViewOrigin();
 	ViewUniformShaderParameters.CullingSign = bReverseCulling ? -1.0f : 1.0f;
 	ViewUniformShaderParameters.NearPlane = InViewMatrices.ComputeNearPlane();
diff --git a/Engine/Source/Runtime/Engine/Private/ShadowMap.cpp b/Engine/Source/Runtime/Engine/Private/ShadowMap.cpp
--- a/Engine/Source/Runtime/Engine/Private/ShadowMap.cpp
+++ b/Engine/Source/Runtime/Engine/Private/ShadowMap.cpp
@@ -111,11 +111,11 @@ struct FShadowMapAllocation
 			{
 				// TODO: We currently only support one LOD of static lighting in foliage
 				// Need to create per-LOD instance data to fix that
-				MeshBuildData->PerInstanceLightmapData[InstanceIndex].ShadowmapUVBias = ShadowMap->GetCoordinateBias();
+				MeshBuildData->PerInstanceLightmapData[InstanceIndex].ShadowmapUVBias = FVector2f(ShadowMap->GetCoordinateBias());
 				const int32 RenderIndex = Component->GetRenderIndex(InstanceIndex);
 				if (RenderIndex != INDEX_NONE)
 				{
-					Component->InstanceUpdateCmdBuffer.SetShadowMapData(RenderIndex, MeshBuildData->PerInstanceLightmapData[InstanceIndex].ShadowmapUVBias);
+					Component->InstanceUpdateCmdBuffer.SetShadowMapData(RenderIndex, FVector2D(MeshBuildData->PerInstanceLightmapData[InstanceIndex].ShadowmapUVBias));
 					Component->MarkRenderStateDirty();
 				}
 			}
diff --git a/Engine/Source/Runtime/Engine/Private/SkeletalMesh.cpp b/Engine/Source/Runtime/Engine/Private/SkeletalMesh.cpp
--- a/Engine/Source/Runtime/Engine/Private/SkeletalMesh.cpp
+++ b/Engine/Source/Runtime/Engine/Private/SkeletalMesh.cpp
@@ -1111,9 +1111,9 @@ static void AccumulateUVDensities(float* OutWeightedUVDensities, float* OutWeigh
 			for (int32 CoordinateIndex = 0; CoordinateIndex < NumCoordinateIndex; ++CoordinateIndex)
 			{
 				const float UVAera = FUVDensityAccumulator::GetUVChannelAera(
-					LODData.StaticVertexBuffers.StaticMeshVertexBuffer.GetVertexUV(Index0, CoordinateIndex),
-					LODData.StaticVertexBuffers.StaticMeshVertexBuffer.GetVertexUV(Index1, CoordinateIndex),
-					LODData.StaticVertexBuffers.StaticMeshVertexBuffer.GetVertexUV(Index2, CoordinateIndex));
+					FVector2D(LODData.StaticVertexBuffers.StaticMeshVertexBuffer.GetVertexUV(Index0, CoordinateIndex)),
+					FVector2D(LODData.StaticVertexBuffers.StaticMeshVertexBuffer.GetVertexUV(Index1, CoordinateIndex)),
+					FVector2D(LODData.StaticVertexBuffers.StaticMeshVertexBuffer.GetVertexUV(Index2, CoordinateIndex)));
 
 				UVDensityAccs[CoordinateIndex].PushTriangle(Aera, UVAera);
 			}
diff --git a/Engine/Source/Runtime/Engine/Private/SkeletalMeshMerge.cpp b/Engine/Source/Runtime/Engine/Private/SkeletalMeshMerge.cpp
--- a/Engine/Source/Runtime/Engine/Private/SkeletalMeshMerge.cpp
+++ b/Engine/Source/Runtime/Engine/Private/SkeletalMeshMerge.cpp
@@ -432,19 +432,19 @@ void FSkeletalMeshMerge::CopyVertexFromSource(VertexDataType& DestVert, const FS
 	const uint32 ValidLoopCount = FMath::Min(VertexDataType::NumTexCoords, LODNumTexCoords);
 	for (uint32 UVIndex = 0; UVIndex < ValidLoopCount; ++UVIndex)
 	{
-		FVector2D UVs = SrcLODData.StaticVertexBuffers.StaticMeshVertexBuffer.GetVertexUV_Typed<VertexDataType::StaticMeshVertexUVType>(SourceVertIdx, UVIndex);
+		FVector2D UVs = FVector2D(SrcLODData.StaticVertexBuffers.StaticMeshVertexBuffer.GetVertexUV_Typed<VertexDataType::StaticMeshVertexUVType>(SourceVertIdx, UVIndex));
 		if (UVIndex < (uint32)MergeSectionInfo.UVTransforms.Num())
 		{
 			FVector Transformed = MergeSectionInfo.UVTransforms[UVIndex].TransformPosition(FVector(UVs, 1.f));
 			UVs = FVector2D(Transformed.X, Transformed.Y);
 		}
-		DestVert.UVs[UVIndex] = UVs;
+		DestVert.UVs[UVIndex] = FVector2f(UVs);	// LWC_TODO: Precision loss
 	}
 	
 	// now just fill up zero value if we didn't reach till end
 	for (uint32 UVIndex = ValidLoopCount; UVIndex < VertexDataType::NumTexCoords; ++UVIndex)
 	{
-		DestVert.UVs[UVIndex] = FVector2D::ZeroVector;
+		DestVert.UVs[UVIndex] = FVector2f::ZeroVector;
 	}
 }
 
diff --git a/Engine/Source/Runtime/Engine/Private/SkeletalRenderCPUSkin.cpp b/Engine/Source/Runtime/Engine/Private/SkeletalRenderCPUSkin.cpp
--- a/Engine/Source/Runtime/Engine/Private/SkeletalRenderCPUSkin.cpp
+++ b/Engine/Source/Runtime/Engine/Private/SkeletalRenderCPUSkin.cpp
@@ -288,7 +288,7 @@ void FSkeletalMeshObjectCPUSkin::CacheVertices(int32 LODIndex, bool bForce) cons
 
 			for (uint32 UVIndex = 0; UVIndex < LOD.StaticVertexBuffers.StaticMeshVertexBuffer.GetNumTexCoords(); ++UVIndex)
 			{
-				MeshLOD.StaticMeshVertexBuffer.SetVertexUV(i, UVIndex, FVector2D(CachedFinalVertices[i].TextureCoordinates[UVIndex].X, CachedFinalVertices[i].TextureCoordinates[UVIndex].Y));
+				MeshLOD.StaticMeshVertexBuffer.SetVertexUV(i, UVIndex, FVector2f(CachedFinalVertices[i].TextureCoordinates[UVIndex].X, CachedFinalVertices[i].TextureCoordinates[UVIndex].Y));
 			}
 		}
 
@@ -1037,7 +1037,7 @@ static void SkinVertexSection(
 			// Copy UVs.
 			for (int32 UVIndex = 0; UVIndex < NumberOfUVs; ++UVIndex)
 			{
-				DestVertex->TextureCoordinates[UVIndex] = LOD.StaticVertexBuffers.StaticMeshVertexBuffer.GetVertexUV(Section.GetVertexBufferIndex() + VertexIndex, UVIndex);
+				DestVertex->TextureCoordinates[UVIndex] = FVector2D(LOD.StaticVertexBuffers.StaticMeshVertexBuffer.GetVertexUV(Section.GetVertexBufferIndex() + VertexIndex, UVIndex));
 			}
 
 			CurBaseVertIdx++;
diff --git a/Engine/Source/Runtime/Engine/Private/StaticMesh.cpp b/Engine/Source/Runtime/Engine/Private/StaticMesh.cpp
--- a/Engine/Source/Runtime/Engine/Private/StaticMesh.cpp
+++ b/Engine/Source/Runtime/Engine/Private/StaticMesh.cpp
@@ -1096,8 +1096,8 @@ void FStaticMeshVertexBuffers::InitModelBuffers(TArray<FModelVertex>& Vertices)
 
 		PositionVertexBuffer.VertexPosition(0) = FVector3f(0, 0, 0);
 		StaticMeshVertexBuffer.SetVertexTangents(0, FVector3f(1, 0, 0), FVector3f(0, 1, 0), FVector3f(0, 0, 1));
-		StaticMeshVertexBuffer.SetVertexUV(0, 0, FVector2D(0, 0));
-		StaticMeshVertexBuffer.SetVertexUV(0, 1, FVector2D(0, 0));
+		StaticMeshVertexBuffer.SetVertexUV(0, 0, FVector2f(0, 0));
+		StaticMeshVertexBuffer.SetVertexUV(0, 1, FVector2f(0, 0));
 	}
 }
 
@@ -1184,7 +1184,7 @@ void FStaticMeshVertexBuffers::InitFromDynamicVertex(FLocalVertexFactory* Vertex
 
 		PositionVertexBuffer.VertexPosition(0) = FVector3f(0, 0, 0);
 		StaticMeshVertexBuffer.SetVertexTangents(0, FVector3f(1, 0, 0), FVector3f(0, 1, 0), FVector3f(0, 0, 1));
-		StaticMeshVertexBuffer.SetVertexUV(0, 0, FVector2D(0, 0));
+		StaticMeshVertexBuffer.SetVertexUV(0, 0, FVector2f(0, 0));
 		ColorVertexBuffer.VertexColor(0) = FColor(1,1,1,1);
 		NumTexCoords = 1;
 		LightMapIndex = 0;
@@ -2670,9 +2670,9 @@ void FStaticMeshRenderData::ComputeUVDensities()
 					for (int32 UVIndex = 0; UVIndex < NumTexCoords; ++UVIndex)
 					{
 						const float UVAera = FUVDensityAccumulator::GetUVChannelAera(
-												LODModel.VertexBuffers.StaticMeshVertexBuffer.GetVertexUV(Index0, UVIndex), 
-												LODModel.VertexBuffers.StaticMeshVertexBuffer.GetVertexUV(Index1, UVIndex), 
-												LODModel.VertexBuffers.StaticMeshVertexBuffer.GetVertexUV(Index2, UVIndex));
+												FVector2D(LODModel.VertexBuffers.StaticMeshVertexBuffer.GetVertexUV(Index0, UVIndex)), 
+												FVector2D(LODModel.VertexBuffers.StaticMeshVertexBuffer.GetVertexUV(Index1, UVIndex)), 
+												FVector2D(LODModel.VertexBuffers.StaticMeshVertexBuffer.GetVertexUV(Index2, UVIndex)));
 
 						UVDensityAccs[UVIndex].PushTriangle(Aera, UVAera);
 					}
@@ -6465,7 +6465,7 @@ static int32 GetCollisionVertIndexForMeshVertIndex(int32 MeshVertIndex, TMap<int
 		for (int32 ChannelIdx = 0; ChannelIdx < OutUVs.Num(); ChannelIdx++)
 		{
 			check(OutPositions.Num() == OutUVs[ChannelIdx].Num());
-			OutUVs[ChannelIdx].Add(InVertBuffer.GetVertexUV(MeshVertIndex, ChannelIdx));
+			OutUVs[ChannelIdx].Add(FVector2D(InVertBuffer.GetVertexUV(MeshVertIndex, ChannelIdx)));
 		}
 
 		// Copy position
@@ -7073,7 +7073,7 @@ void UStaticMesh::CheckLightMapUVs( UStaticMesh* InStaticMesh, TArray< FString >
 			const uint32 VertexIndices[] = {Indices[StartIndex + 0], Indices[StartIndex + 1], Indices[StartIndex + 2]};
 			for(int i = 0; i<3;i++)
 			{
-				TriangleUVsOUT[i] = MeshLOD.VertexBuffers.StaticMeshVertexBuffer.GetVertexUV(VertexIndices[i], UVChannel);		
+				TriangleUVsOUT[i] = FVector2D(MeshLOD.VertexBuffers.StaticMeshVertexBuffer.GetVertexUV(VertexIndices[i], UVChannel));		
 			}
 		}
 
diff --git a/Engine/Source/Runtime/Engine/Private/StaticMeshLight.cpp b/Engine/Source/Runtime/Engine/Private/StaticMeshLight.cpp
--- a/Engine/Source/Runtime/Engine/Private/StaticMeshLight.cpp
+++ b/Engine/Source/Runtime/Engine/Private/StaticMeshLight.cpp
@@ -44,7 +44,7 @@ static void GetStaticLightingVertex(
 	checkSlow(VertexBuffer.GetNumTexCoords() <= UE_ARRAY_COUNT(OutVertex.TextureCoordinates));
 	for(uint32 LightmapTextureCoordinateIndex = 0;LightmapTextureCoordinateIndex < VertexBuffer.GetNumTexCoords();LightmapTextureCoordinateIndex++)
 	{
-		OutVertex.TextureCoordinates[LightmapTextureCoordinateIndex] = VertexBuffer.GetVertexUV(VertexIndex,LightmapTextureCoordinateIndex);
+		OutVertex.TextureCoordinates[LightmapTextureCoordinateIndex] = FVector2D(VertexBuffer.GetVertexUV(VertexIndex,LightmapTextureCoordinateIndex));
 	}
 }
 
diff --git a/Engine/Source/Runtime/Engine/Private/TextureDerivedData.cpp b/Engine/Source/Runtime/Engine/Private/TextureDerivedData.cpp
--- a/Engine/Source/Runtime/Engine/Private/TextureDerivedData.cpp
+++ b/Engine/Source/Runtime/Engine/Private/TextureDerivedData.cpp
@@ -129,10 +129,10 @@ static void SerializeForKey(FArchive& Ar, const FTextureBuildSettings& Settings)
 		TempUint32 = UE::Color::COLORSPACE_VER; Ar << TempUint32;
 		TempColorSpace = UE::Color::FColorSpace::GetWorking(); Ar << TempColorSpace;
 
-		TempVector2f = Settings.RedChromaticityCoordinate; Ar << TempVector2f;
-		TempVector2f = Settings.GreenChromaticityCoordinate; Ar << TempVector2f;
-		TempVector2f = Settings.BlueChromaticityCoordinate; Ar << TempVector2f;
-		TempVector2f = Settings.WhiteChromaticityCoordinate; Ar << TempVector2f;
+		TempVector2f = FVector2f(Settings.RedChromaticityCoordinate); Ar << TempVector2f;
+		TempVector2f = FVector2f(Settings.GreenChromaticityCoordinate); Ar << TempVector2f;
+		TempVector2f = FVector2f(Settings.BlueChromaticityCoordinate); Ar << TempVector2f;
+		TempVector2f = FVector2f(Settings.WhiteChromaticityCoordinate); Ar << TempVector2f;
 		TempByte = Settings.ChromaticAdaptationMethod; Ar << TempByte;
 	}
 
@@ -703,10 +703,10 @@ static void GetTextureBuildSettings(
 	OutBuildSettings.bLongLatSource = false;
 	OutBuildSettings.SourceEncodingOverride = static_cast<uint8>(Texture.SourceColorSettings.EncodingOverride);
 	OutBuildSettings.bHasColorSpaceDefinition = Texture.SourceColorSettings.ColorSpace != ETextureColorSpace::TCS_None;
-	OutBuildSettings.RedChromaticityCoordinate = Texture.SourceColorSettings.RedChromaticityCoordinate;
-	OutBuildSettings.GreenChromaticityCoordinate = Texture.SourceColorSettings.GreenChromaticityCoordinate;
-	OutBuildSettings.BlueChromaticityCoordinate = Texture.SourceColorSettings.BlueChromaticityCoordinate;
-	OutBuildSettings.WhiteChromaticityCoordinate = Texture.SourceColorSettings.WhiteChromaticityCoordinate;
+	OutBuildSettings.RedChromaticityCoordinate = FVector2f(Texture.SourceColorSettings.RedChromaticityCoordinate);
+	OutBuildSettings.GreenChromaticityCoordinate = FVector2f(Texture.SourceColorSettings.GreenChromaticityCoordinate);
+	OutBuildSettings.BlueChromaticityCoordinate = FVector2f(Texture.SourceColorSettings.BlueChromaticityCoordinate);
+	OutBuildSettings.WhiteChromaticityCoordinate = FVector2f(Texture.SourceColorSettings.WhiteChromaticityCoordinate);
 	OutBuildSettings.ChromaticAdaptationMethod = static_cast<uint8>(Texture.SourceColorSettings.ChromaticAdaptationMethod);
 
 	if (Texture.MaxTextureSize > 0)
diff --git a/Engine/Source/Runtime/Engine/Private/TileRendering.cpp b/Engine/Source/Runtime/Engine/Private/TileRendering.cpp
--- a/Engine/Source/Runtime/Engine/Private/TileRendering.cpp
+++ b/Engine/Source/Runtime/Engine/Private/TileRendering.cpp
@@ -130,10 +130,10 @@ void FCanvasTileRendererItem::FRenderData::InitTileMesh(const FSceneView& View,
 			StaticMeshVertexBuffers.PositionVertexBuffer.VertexPosition(FirstVertex + 2) = FVector(X + SizeX, View.UnscaledViewRect.Height() - Y, 0.0f);
 			StaticMeshVertexBuffers.PositionVertexBuffer.VertexPosition(FirstVertex + 3) = FVector(X, View.UnscaledViewRect.Height() - Y, 0.0f);
 
-			StaticMeshVertexBuffers.StaticMeshVertexBuffer.SetVertexUV(FirstVertex + 0, 0, FVector2D(U + SizeU, V + SizeV));
-			StaticMeshVertexBuffers.StaticMeshVertexBuffer.SetVertexUV(FirstVertex + 1, 0, FVector2D(U, V + SizeV));
-			StaticMeshVertexBuffers.StaticMeshVertexBuffer.SetVertexUV(FirstVertex + 2, 0, FVector2D(U + SizeU, V));
-			StaticMeshVertexBuffers.StaticMeshVertexBuffer.SetVertexUV(FirstVertex + 3, 0, FVector2D(U, V));
+			StaticMeshVertexBuffers.StaticMeshVertexBuffer.SetVertexUV(FirstVertex + 0, 0, FVector2f(U + SizeU, V + SizeV));
+			StaticMeshVertexBuffers.StaticMeshVertexBuffer.SetVertexUV(FirstVertex + 1, 0, FVector2f(U, V + SizeV));
+			StaticMeshVertexBuffers.StaticMeshVertexBuffer.SetVertexUV(FirstVertex + 2, 0, FVector2f(U + SizeU, V));
+			StaticMeshVertexBuffers.StaticMeshVertexBuffer.SetVertexUV(FirstVertex + 3, 0, FVector2f(U, V));
 		}
 		else
 		{
@@ -142,10 +142,10 @@ void FCanvasTileRendererItem::FRenderData::InitTileMesh(const FSceneView& View,
 			StaticMeshVertexBuffers.PositionVertexBuffer.VertexPosition(FirstVertex + 2) = FVector(X + SizeX, Y + SizeY, 0.0f);
 			StaticMeshVertexBuffers.PositionVertexBuffer.VertexPosition(FirstVertex + 3) = FVector(X, Y + SizeY, 0.0f);
 
-			StaticMeshVertexBuffers.StaticMeshVertexBuffer.SetVertexUV(FirstVertex + 0, 0, FVector2D(U + SizeU, V));
-			StaticMeshVertexBuffers.StaticMeshVertexBuffer.SetVertexUV(FirstVertex + 1, 0, FVector2D(U, V));
-			StaticMeshVertexBuffers.StaticMeshVertexBuffer.SetVertexUV(FirstVertex + 2, 0, FVector2D(U + SizeU, V + SizeV));
-			StaticMeshVertexBuffers.StaticMeshVertexBuffer.SetVertexUV(FirstVertex + 3, 0, FVector2D(U, V + SizeV));
+			StaticMeshVertexBuffers.StaticMeshVertexBuffer.SetVertexUV(FirstVertex + 0, 0, FVector2f(U + SizeU, V));
+			StaticMeshVertexBuffers.StaticMeshVertexBuffer.SetVertexUV(FirstVertex + 1, 0, FVector2f(U, V));
+			StaticMeshVertexBuffers.StaticMeshVertexBuffer.SetVertexUV(FirstVertex + 2, 0, FVector2f(U + SizeU, V + SizeV));
+			StaticMeshVertexBuffers.StaticMeshVertexBuffer.SetVertexUV(FirstVertex + 3, 0, FVector2f(U, V + SizeV));
 		}
 
 		for (int j = 0; j < CanvasTileVertexCount; j++)
diff --git a/Engine/Source/Runtime/Engine/Private/TriangleRendering.cpp b/Engine/Source/Runtime/Engine/Private/TriangleRendering.cpp
--- a/Engine/Source/Runtime/Engine/Private/TriangleRendering.cpp
+++ b/Engine/Source/Runtime/Engine/Private/TriangleRendering.cpp
@@ -117,9 +117,9 @@ void FCanvasTriangleRendererItem::FRenderData::InitTriangleMesh(const FSceneView
 		StaticMeshVertexBuffers.StaticMeshVertexBuffer.SetVertexTangents(StartIndex + 1, FVector(1.0f, 0.0f, 0.0f), FVector(0.0f, 1.0f, 0.0f), FVector(0.0f, 0.0f, 1.0f));
 		StaticMeshVertexBuffers.StaticMeshVertexBuffer.SetVertexTangents(StartIndex + 2, FVector(1.0f, 0.0f, 0.0f), FVector(0.0f, 1.0f, 0.0f), FVector(0.0f, 0.0f, 1.0f));
 
-		StaticMeshVertexBuffers.StaticMeshVertexBuffer.SetVertexUV(StartIndex + 0, 0, FVector2D(Tri.V1_UV.X, Tri.V1_UV.Y));
-		StaticMeshVertexBuffers.StaticMeshVertexBuffer.SetVertexUV(StartIndex + 1, 0, FVector2D(Tri.V0_UV.X, Tri.V0_UV.Y));
-		StaticMeshVertexBuffers.StaticMeshVertexBuffer.SetVertexUV(StartIndex + 2, 0, FVector2D(Tri.V2_UV.X, Tri.V2_UV.Y));
+		StaticMeshVertexBuffers.StaticMeshVertexBuffer.SetVertexUV(StartIndex + 0, 0, FVector2f(Tri.V1_UV.X, Tri.V1_UV.Y));
+		StaticMeshVertexBuffers.StaticMeshVertexBuffer.SetVertexUV(StartIndex + 1, 0, FVector2f(Tri.V0_UV.X, Tri.V0_UV.Y));
+		StaticMeshVertexBuffers.StaticMeshVertexBuffer.SetVertexUV(StartIndex + 2, 0, FVector2f(Tri.V2_UV.X, Tri.V2_UV.Y));
 
 		StaticMeshVertexBuffers.ColorVertexBuffer.VertexColor(StartIndex + 0) = Tri.V1_Color.ToFColor(true);
 		StaticMeshVertexBuffers.ColorVertexBuffer.VertexColor(StartIndex + 1) = Tri.V0_Color.ToFColor(true);
diff --git a/Engine/Source/Runtime/Engine/Public/MeshBuild.h b/Engine/Source/Runtime/Engine/Public/MeshBuild.h
--- a/Engine/Source/Runtime/Engine/Public/MeshBuild.h
+++ b/Engine/Source/Runtime/Engine/Public/MeshBuild.h
@@ -65,6 +65,11 @@ inline bool UVsEqual(const FVector2D& V1, const FVector2D& V2)
 	return FMath::Abs(V1.X - V2.X) <= Epsilon && FMath::Abs(V1.Y - V2.Y) <= Epsilon;
 }
 
+inline bool UVsEqual(const FVector2f& V1, const FVector2f& V2)
+{
+	return UVsEqual(FVector2D(V1), FVector2D(V2));
+}
+
 inline bool NormalsEqual(const FVector& V1,const FVector& V2, const FOverlappingThresholds& OverlappingThreshold)
 {
 	const float Epsilon = OverlappingThreshold.ThresholdTangentNormal;
@@ -76,3 +81,8 @@ inline bool UVsEqual(const FVector2D& V1, const FVector2D& V2, const FOverlappin
 	const float Epsilon = OverlappingThreshold.ThresholdUV;
 	return FMath::Abs(V1.X - V2.X) <= Epsilon && FMath::Abs(V1.Y - V2.Y) <= Epsilon;
 }
+
+inline bool UVsEqual(const FVector2f& V1, const FVector2f& V2, const FOverlappingThresholds& OverlappingThreshold)
+{
+	return UVsEqual(FVector2D(V1), FVector2D(V2), OverlappingThreshold);
+}
diff --git a/Engine/Source/Runtime/Experimental/GeometryCollectionEngine/Private/GeometryCollection/GeometryCollectionObject.cpp b/Engine/Source/Runtime/Experimental/GeometryCollectionEngine/Private/GeometryCollection/GeometryCollectionObject.cpp
--- a/Engine/Source/Runtime/Experimental/GeometryCollectionEngine/Private/GeometryCollection/GeometryCollectionObject.cpp
+++ b/Engine/Source/Runtime/Experimental/GeometryCollectionEngine/Private/GeometryCollection/GeometryCollectionObject.cpp
@@ -918,7 +918,7 @@ TUniquePtr<FGeometryCollectionNaniteData> UGeometryCollection::CreateNaniteData(
 				Vertex.UVs[UVIdx] = UVsArray[VertexStart + VertexIndex][UVIdx];
 				if (Vertex.UVs[UVIdx].ContainsNaN())
 				{
-					Vertex.UVs[UVIdx] = FVector2D::ZeroVector;
+					Vertex.UVs[UVIdx] = FVector2f::ZeroVector;
 				}
 			}
 		}
diff --git a/Engine/Source/Runtime/GeometryFramework/Public/Components/BaseDynamicMeshSceneProxy.h b/Engine/Source/Runtime/GeometryFramework/Public/Components/BaseDynamicMeshSceneProxy.h
--- a/Engine/Source/Runtime/GeometryFramework/Public/Components/BaseDynamicMeshSceneProxy.h
+++ b/Engine/Source/Runtime/GeometryFramework/Public/Components/BaseDynamicMeshSceneProxy.h
@@ -634,7 +634,7 @@ public:
 				{
 					FVector2f UV = (UVTriangles[k][j] != FDynamicMesh3::InvalidID) ?
 						UVOverlays[k]->GetElement(UVTriangles[k][j]) : FVector2f::Zero();
-					RenderBuffers->StaticMeshVertexBuffer.SetVertexUV(VertIdx, k, (FVector2D)UV);
+					RenderBuffers->StaticMeshVertexBuffer.SetVertexUV(VertIdx, k, UV);
 				}
 
 				FColor VertexFColor = (bHaveColors && TriColor[j] != FDynamicMesh3::InvalidID) ?
@@ -880,7 +880,7 @@ public:
 				{
 					FVector2f UV = (UVTriangles[k][j] != FDynamicMesh3::InvalidID) ?
 						UVOverlays[k]->GetElement(UVTriangles[k][j]) : FVector2f::Zero();
-					RenderBuffers->StaticMeshVertexBuffer.SetVertexUV(VertIdx, k, (FVector2D)UV);
+					RenderBuffers->StaticMeshVertexBuffer.SetVertexUV(VertIdx, k, UV);
 				}
 
 				++VertIdx;
diff --git a/Engine/Source/Runtime/Landscape/Private/LandscapeEdit.cpp b/Engine/Source/Runtime/Landscape/Private/LandscapeEdit.cpp
--- a/Engine/Source/Runtime/Landscape/Private/LandscapeEdit.cpp
+++ b/Engine/Source/Runtime/Landscape/Private/LandscapeEdit.cpp
@@ -7259,23 +7259,23 @@ bool ALandscapeProxy::LandscapeExportHeightmapToRenderTarget(UTextureRenderTarge
 				float HeightmapOffsetV = Component->HeightmapScaleBias.W + HeightmapSubsectionOffsetV * (float)SubY;
 
 				FCanvasUVTri Tri1;
-				Tri1.V0_Pos = FVector2f(SubSectionSectionBase.X, SubSectionSectionBase.Y);
-				Tri1.V1_Pos = FVector2f(SubSectionSectionBase.X + SubsectionSizeVerts, SubSectionSectionBase.Y);
-				Tri1.V2_Pos = FVector2f(SubSectionSectionBase.X + SubsectionSizeVerts, SubSectionSectionBase.Y + SubsectionSizeVerts);
+				Tri1.V0_Pos = FVector2D(SubSectionSectionBase.X, SubSectionSectionBase.Y);
+				Tri1.V1_Pos = FVector2D(SubSectionSectionBase.X + SubsectionSizeVerts, SubSectionSectionBase.Y);
+				Tri1.V2_Pos = FVector2D(SubSectionSectionBase.X + SubsectionSizeVerts, SubSectionSectionBase.Y + SubsectionSizeVerts);
 
-				Tri1.V0_UV = FVector2f(HeightmapOffsetU, HeightmapOffsetV);
-				Tri1.V1_UV = FVector2f(HeightmapOffsetU + HeightmapSubsectionOffsetU, HeightmapOffsetV);
-				Tri1.V2_UV = FVector2f(HeightmapOffsetU + HeightmapSubsectionOffsetU, HeightmapOffsetV + HeightmapSubsectionOffsetV);
+				Tri1.V0_UV = FVector2D(HeightmapOffsetU, HeightmapOffsetV);
+				Tri1.V1_UV = FVector2D(HeightmapOffsetU + HeightmapSubsectionOffsetU, HeightmapOffsetV);
+				Tri1.V2_UV = FVector2D(HeightmapOffsetU + HeightmapSubsectionOffsetU, HeightmapOffsetV + HeightmapSubsectionOffsetV);
 				TrianglesPerMID->TriangleList.Add(Tri1);
 
 				FCanvasUVTri Tri2;
-				Tri2.V0_Pos = FVector2f(SubSectionSectionBase.X + SubsectionSizeVerts, SubSectionSectionBase.Y + SubsectionSizeVerts);
-				Tri2.V1_Pos = FVector2f(SubSectionSectionBase.X, SubSectionSectionBase.Y + SubsectionSizeVerts);
-				Tri2.V2_Pos = FVector2f(SubSectionSectionBase.X, SubSectionSectionBase.Y);
+				Tri2.V0_Pos = FVector2D(SubSectionSectionBase.X + SubsectionSizeVerts, SubSectionSectionBase.Y + SubsectionSizeVerts);
+				Tri2.V1_Pos = FVector2D(SubSectionSectionBase.X, SubSectionSectionBase.Y + SubsectionSizeVerts);
+				Tri2.V2_Pos = FVector2D(SubSectionSectionBase.X, SubSectionSectionBase.Y);
 
-				Tri2.V0_UV = FVector2f(HeightmapOffsetU + HeightmapSubsectionOffsetU, HeightmapOffsetV + HeightmapSubsectionOffsetV);
-				Tri2.V1_UV = FVector2f(HeightmapOffsetU, HeightmapOffsetV + HeightmapSubsectionOffsetV);
-				Tri2.V2_UV = FVector2f(HeightmapOffsetU, HeightmapOffsetV);
+				Tri2.V0_UV = FVector2D(HeightmapOffsetU + HeightmapSubsectionOffsetU, HeightmapOffsetV + HeightmapSubsectionOffsetV);
+				Tri2.V1_UV = FVector2D(HeightmapOffsetU, HeightmapOffsetV + HeightmapSubsectionOffsetV);
+				Tri2.V2_UV = FVector2D(HeightmapOffsetU, HeightmapOffsetV);
 
 				TrianglesPerMID->TriangleList.Add(Tri2);
 			}
diff --git a/Engine/Source/Runtime/Landscape/Private/LandscapeEditLayers.cpp b/Engine/Source/Runtime/Landscape/Private/LandscapeEditLayers.cpp
--- a/Engine/Source/Runtime/Landscape/Private/LandscapeEditLayers.cpp
+++ b/Engine/Source/Runtime/Landscape/Private/LandscapeEditLayers.cpp
@@ -2634,27 +2634,27 @@ void ALandscape::DrawWeightmapComponentsToRenderTarget(const FString& InDebugNam
 		{
 		case ERTDrawingType::RTAtlas:
 		{
-			GenerateLayersRenderQuadsAtlas(InSectionBaseList[i], WeightmapScaleBias, SubsectionSizeQuads, WeightmapReadTextureSize, WeightmapWriteTextureSize, TriangleList);
+			GenerateLayersRenderQuadsAtlas(InSectionBaseList[i], FVector2D(WeightmapScaleBias), SubsectionSizeQuads, WeightmapReadTextureSize, WeightmapWriteTextureSize, TriangleList);
 		} break;
 
 		case ERTDrawingType::RTAtlasToNonAtlas:
 		{
-			GenerateLayersRenderQuadsAtlasToNonAtlas(InSectionBaseList[i], WeightmapScaleBias, SubsectionSizeQuads, WeightmapReadTextureSize, WeightmapWriteTextureSize, TriangleList);
+			GenerateLayersRenderQuadsAtlasToNonAtlas(InSectionBaseList[i], FVector2D(WeightmapScaleBias), SubsectionSizeQuads, WeightmapReadTextureSize, WeightmapWriteTextureSize, TriangleList);
 		} break;
 
 		case ERTDrawingType::RTNonAtlas:
 		{
-			GenerateLayersRenderQuadsNonAtlas(InSectionBaseList[i], WeightmapScaleBias, SubsectionSizeQuads, WeightmapReadTextureSize, WeightmapWriteTextureSize, TriangleList);
+			GenerateLayersRenderQuadsNonAtlas(InSectionBaseList[i], FVector2D(WeightmapScaleBias), SubsectionSizeQuads, WeightmapReadTextureSize, WeightmapWriteTextureSize, TriangleList);
 		} break;
 
 		case ERTDrawingType::RTNonAtlasToAtlas:
 		{
-			GenerateLayersRenderQuadsNonAtlasToAtlas(InSectionBaseList[i], WeightmapScaleBias, SubsectionSizeQuads, WeightmapReadTextureSize, WeightmapWriteTextureSize, TriangleList);
+			GenerateLayersRenderQuadsNonAtlasToAtlas(InSectionBaseList[i], FVector2D(WeightmapScaleBias), SubsectionSizeQuads, WeightmapReadTextureSize, WeightmapWriteTextureSize, TriangleList);
 		} break;
 
 		case ERTDrawingType::RTMips:
 		{
-			GenerateLayersRenderQuadsMip(InSectionBaseList[i], WeightmapScaleBias, SubsectionSizeQuads, WeightmapReadTextureSize, WeightmapWriteTextureSize, InMipRender, TriangleList);
+			GenerateLayersRenderQuadsMip(InSectionBaseList[i], FVector2D(WeightmapScaleBias), SubsectionSizeQuads, WeightmapReadTextureSize, WeightmapWriteTextureSize, InMipRender, TriangleList);
 		} break;
 
 		default:
@@ -2709,7 +2709,7 @@ void ALandscape::DrawWeightmapComponentsToRenderTarget(const FString& InDebugNam
 		SectionBaseList.Add(ComponentSectionBase);
 	}
 
-	DrawWeightmapComponentsToRenderTarget(InDebugName, SectionBaseList, FVector2D::ZeroVector, &WeightmapScaleBiasList, InWeightmapRTRead, InOptionalWeightmapRTRead2, InWeightmapRTWrite, InDrawType, InClearRTWrite, InShaderParams, InMipRender);
+	DrawWeightmapComponentsToRenderTarget(InDebugName, SectionBaseList, FVector2f::ZeroVector, &WeightmapScaleBiasList, InWeightmapRTRead, InOptionalWeightmapRTRead2, InWeightmapRTWrite, InDrawType, InClearRTWrite, InShaderParams, InMipRender);
 
 	PrintLayersDebugRT(InDebugName, InWeightmapRTWrite, InMipRender, false);
 }
@@ -2808,30 +2808,30 @@ void ALandscape::DrawHeightmapComponentsToRenderTarget(const FString& InDebugNam
 
 		switch (InDrawType)
 		{
-		case ERTDrawingType::RTAtlas:
-		{
-			GenerateLayersRenderQuadsAtlas(ComponentSectionBase, HeightmapScaleBias, SubsectionSizeQuads, HeightmapReadTextureSize, HeightmapWriteTextureSize, TriangleList);
-		} break;
+			case ERTDrawingType::RTAtlas:
+			{
+				GenerateLayersRenderQuadsAtlas(ComponentSectionBase, FVector2D(HeightmapScaleBias), SubsectionSizeQuads, HeightmapReadTextureSize, HeightmapWriteTextureSize, TriangleList);
+			} break;
 
-		case ERTDrawingType::RTAtlasToNonAtlas:
-		{
-			GenerateLayersRenderQuadsAtlasToNonAtlas(ComponentSectionBase, HeightmapScaleBias, SubsectionSizeQuads, HeightmapReadTextureSize, HeightmapWriteTextureSize, TriangleList);
-		} break;
+			case ERTDrawingType::RTAtlasToNonAtlas:
+			{
+				GenerateLayersRenderQuadsAtlasToNonAtlas(ComponentSectionBase, FVector2D(HeightmapScaleBias), SubsectionSizeQuads, HeightmapReadTextureSize, HeightmapWriteTextureSize, TriangleList);
+			} break;
 
-		case ERTDrawingType::RTNonAtlas:
-		{
-			GenerateLayersRenderQuadsNonAtlas(ComponentSectionBase, HeightmapScaleBias, SubsectionSizeQuads, HeightmapReadTextureSize, HeightmapWriteTextureSize, TriangleList);
-		} break;
+			case ERTDrawingType::RTNonAtlas:
+			{
+				GenerateLayersRenderQuadsNonAtlas(ComponentSectionBase, FVector2D(HeightmapScaleBias), SubsectionSizeQuads, HeightmapReadTextureSize, HeightmapWriteTextureSize, TriangleList);
+			} break;
 
-		case ERTDrawingType::RTNonAtlasToAtlas:
-		{
-			GenerateLayersRenderQuadsNonAtlasToAtlas(ComponentSectionBase, HeightmapScaleBias, SubsectionSizeQuads, HeightmapReadTextureSize, HeightmapWriteTextureSize, TriangleList);
-		} break;
+			case ERTDrawingType::RTNonAtlasToAtlas:
+			{
+				GenerateLayersRenderQuadsNonAtlasToAtlas(ComponentSectionBase, FVector2D(HeightmapScaleBias), SubsectionSizeQuads, HeightmapReadTextureSize, HeightmapWriteTextureSize, TriangleList);
+			} break;
 
-		case ERTDrawingType::RTMips:
-		{
-			GenerateLayersRenderQuadsMip(ComponentSectionBase, HeightmapScaleBias, SubsectionSizeQuads, HeightmapReadTextureSize, HeightmapWriteTextureSize, InMipRender, TriangleList);
-		} break;
+			case ERTDrawingType::RTMips:
+			{
+				GenerateLayersRenderQuadsMip(ComponentSectionBase, FVector2D(HeightmapScaleBias), SubsectionSizeQuads, HeightmapReadTextureSize, HeightmapWriteTextureSize, InMipRender, TriangleList);
+			} break;
 
 		default:
 		{
diff --git a/Engine/Source/Runtime/Landscape/Private/LandscapeGizmoActor.cpp b/Engine/Source/Runtime/Landscape/Private/LandscapeGizmoActor.cpp
--- a/Engine/Source/Runtime/Landscape/Private/LandscapeGizmoActor.cpp
+++ b/Engine/Source/Runtime/Landscape/Private/LandscapeGizmoActor.cpp
@@ -259,35 +259,35 @@ public:
 						FDynamicMeshBuilder MeshBuilder(View->GetFeatureLevel());
 
 						const FColor GizmoColor = FColor::White;
-						MeshBuilder.AddVertex(FrustumVerts[0], FVector2D(0, 0), FVector(1,0,0), FVector(0,1,0), FVector(0,0,1), GizmoColor);
-						MeshBuilder.AddVertex(FrustumVerts[1], FVector2D(1, 0), FVector(1,0,0), FVector(0,1,0), FVector(0,0,1), GizmoColor);
-						MeshBuilder.AddVertex(FrustumVerts[2], FVector2D(1, 1), FVector(1,0,0), FVector(0,1,0), FVector(0,0,1), GizmoColor);
-						MeshBuilder.AddVertex(FrustumVerts[3], FVector2D(0, 1), FVector(1,0,0), FVector(0,1,0), FVector(0,0,1), GizmoColor);
-
-						MeshBuilder.AddVertex(FrustumVerts[4], FVector2D(0, 0), FVector(1,0,0), FVector(0,1,0), FVector(0,0,1), GizmoColor);
-						MeshBuilder.AddVertex(FrustumVerts[5], FVector2D(1, 0), FVector(1,0,0), FVector(0,1,0), FVector(0,0,1), GizmoColor);
-						MeshBuilder.AddVertex(FrustumVerts[6], FVector2D(1, 1), FVector(1,0,0), FVector(0,1,0), FVector(0,0,1), GizmoColor);
-						MeshBuilder.AddVertex(FrustumVerts[7], FVector2D(0, 1), FVector(1,0,0), FVector(0,1,0), FVector(0,0,1), GizmoColor);
-
-						MeshBuilder.AddVertex(FrustumVerts[1], FVector2D(0, 0), FVector(1,0,0), FVector(0,1,0), FVector(0,0,1), GizmoColor);
-						MeshBuilder.AddVertex(FrustumVerts[0], FVector2D(1, 0), FVector(1,0,0), FVector(0,1,0), FVector(0,0,1), GizmoColor);
-						MeshBuilder.AddVertex(FrustumVerts[4], FVector2D(1, 1), FVector(1,0,0), FVector(0,1,0), FVector(0,0,1), GizmoColor);
-						MeshBuilder.AddVertex(FrustumVerts[5], FVector2D(0, 1), FVector(1,0,0), FVector(0,1,0), FVector(0,0,1), GizmoColor);
-
-						MeshBuilder.AddVertex(FrustumVerts[3], FVector2D(0, 0), FVector(1,0,0), FVector(0,1,0), FVector(0,0,1), GizmoColor);
-						MeshBuilder.AddVertex(FrustumVerts[2], FVector2D(1, 0), FVector(1,0,0), FVector(0,1,0), FVector(0,0,1), GizmoColor);
-						MeshBuilder.AddVertex(FrustumVerts[6], FVector2D(1, 1), FVector(1,0,0), FVector(0,1,0), FVector(0,0,1), GizmoColor);
-						MeshBuilder.AddVertex(FrustumVerts[7], FVector2D(0, 1), FVector(1,0,0), FVector(0,1,0), FVector(0,0,1), GizmoColor);
-
-						MeshBuilder.AddVertex(FrustumVerts[2], FVector2D(0, 0), FVector(1,0,0), FVector(0,1,0), FVector(0,0,1), GizmoColor);
-						MeshBuilder.AddVertex(FrustumVerts[1], FVector2D(1, 0), FVector(1,0,0), FVector(0,1,0), FVector(0,0,1), GizmoColor);
-						MeshBuilder.AddVertex(FrustumVerts[5], FVector2D(1, 1), FVector(1,0,0), FVector(0,1,0), FVector(0,0,1), GizmoColor);
-						MeshBuilder.AddVertex(FrustumVerts[6], FVector2D(0, 1), FVector(1,0,0), FVector(0,1,0), FVector(0,0,1), GizmoColor);
-
-						MeshBuilder.AddVertex(FrustumVerts[0], FVector2D(0, 0), FVector(1,0,0), FVector(0,1,0), FVector(0,0,1), GizmoColor);
-						MeshBuilder.AddVertex(FrustumVerts[3], FVector2D(1, 0), FVector(1,0,0), FVector(0,1,0), FVector(0,0,1), GizmoColor);
-						MeshBuilder.AddVertex(FrustumVerts[7], FVector2D(1, 1), FVector(1,0,0), FVector(0,1,0), FVector(0,0,1), GizmoColor);
-						MeshBuilder.AddVertex(FrustumVerts[4], FVector2D(0, 1), FVector(1,0,0), FVector(0,1,0), FVector(0,0,1), GizmoColor);
+						MeshBuilder.AddVertex(FrustumVerts[0], FVector2f(0, 0), FVector(1,0,0), FVector(0,1,0), FVector(0,0,1), GizmoColor);
+						MeshBuilder.AddVertex(FrustumVerts[1], FVector2f(1, 0), FVector(1,0,0), FVector(0,1,0), FVector(0,0,1), GizmoColor);
+						MeshBuilder.AddVertex(FrustumVerts[2], FVector2f(1, 1), FVector(1,0,0), FVector(0,1,0), FVector(0,0,1), GizmoColor);
+						MeshBuilder.AddVertex(FrustumVerts[3], FVector2f(0, 1), FVector(1,0,0), FVector(0,1,0), FVector(0,0,1), GizmoColor);
+
+						MeshBuilder.AddVertex(FrustumVerts[4], FVector2f(0, 0), FVector(1,0,0), FVector(0,1,0), FVector(0,0,1), GizmoColor);
+						MeshBuilder.AddVertex(FrustumVerts[5], FVector2f(1, 0), FVector(1,0,0), FVector(0,1,0), FVector(0,0,1), GizmoColor);
+						MeshBuilder.AddVertex(FrustumVerts[6], FVector2f(1, 1), FVector(1,0,0), FVector(0,1,0), FVector(0,0,1), GizmoColor);
+						MeshBuilder.AddVertex(FrustumVerts[7], FVector2f(0, 1), FVector(1,0,0), FVector(0,1,0), FVector(0,0,1), GizmoColor);
+
+						MeshBuilder.AddVertex(FrustumVerts[1], FVector2f(0, 0), FVector(1,0,0), FVector(0,1,0), FVector(0,0,1), GizmoColor);
+						MeshBuilder.AddVertex(FrustumVerts[0], FVector2f(1, 0), FVector(1,0,0), FVector(0,1,0), FVector(0,0,1), GizmoColor);
+						MeshBuilder.AddVertex(FrustumVerts[4], FVector2f(1, 1), FVector(1,0,0), FVector(0,1,0), FVector(0,0,1), GizmoColor);
+						MeshBuilder.AddVertex(FrustumVerts[5], FVector2f(0, 1), FVector(1,0,0), FVector(0,1,0), FVector(0,0,1), GizmoColor);
+
+						MeshBuilder.AddVertex(FrustumVerts[3], FVector2f(0, 0), FVector(1,0,0), FVector(0,1,0), FVector(0,0,1), GizmoColor);
+						MeshBuilder.AddVertex(FrustumVerts[2], FVector2f(1, 0), FVector(1,0,0), FVector(0,1,0), FVector(0,0,1), GizmoColor);
+						MeshBuilder.AddVertex(FrustumVerts[6], FVector2f(1, 1), FVector(1,0,0), FVector(0,1,0), FVector(0,0,1), GizmoColor);
+						MeshBuilder.AddVertex(FrustumVerts[7], FVector2f(0, 1), FVector(1,0,0), FVector(0,1,0), FVector(0,0,1), GizmoColor);
+
+						MeshBuilder.AddVertex(FrustumVerts[2], FVector2f(0, 0), FVector(1,0,0), FVector(0,1,0), FVector(0,0,1), GizmoColor);
+						MeshBuilder.AddVertex(FrustumVerts[1], FVector2f(1, 0), FVector(1,0,0), FVector(0,1,0), FVector(0,0,1), GizmoColor);
+						MeshBuilder.AddVertex(FrustumVerts[5], FVector2f(1, 1), FVector(1,0,0), FVector(0,1,0), FVector(0,0,1), GizmoColor);
+						MeshBuilder.AddVertex(FrustumVerts[6], FVector2f(0, 1), FVector(1,0,0), FVector(0,1,0), FVector(0,0,1), GizmoColor);
+
+						MeshBuilder.AddVertex(FrustumVerts[0], FVector2f(0, 0), FVector(1,0,0), FVector(0,1,0), FVector(0,0,1), GizmoColor);
+						MeshBuilder.AddVertex(FrustumVerts[3], FVector2f(1, 0), FVector(1,0,0), FVector(0,1,0), FVector(0,0,1), GizmoColor);
+						MeshBuilder.AddVertex(FrustumVerts[7], FVector2f(1, 1), FVector(1,0,0), FVector(0,1,0), FVector(0,0,1), GizmoColor);
+						MeshBuilder.AddVertex(FrustumVerts[4], FVector2f(0, 1), FVector(1,0,0), FVector(0,1,0), FVector(0,0,1), GizmoColor);
 
 						for (int32 i = 0; i < 6; ++i)
 						{
@@ -311,7 +311,7 @@ public:
 								FVector TangentX(SampledNormal.Z, 0, -SampledNormal.X);
 								TangentX = TangentX.GetSafeNormal();
 
-								MeshBuilder.AddVertex(SampledPositions[X + Y * SampleSizeX], FVector2D((float)X / (SampleSizeX), (float)Y / (SampleSizeY)), TangentX, SampledNormal^TangentX, SampledNormal, FColor::White);
+								MeshBuilder.AddVertex(SampledPositions[X + Y * SampleSizeX], FVector2f((float)X / (SampleSizeX), (float)Y / (SampleSizeY)), TangentX, SampledNormal^TangentX, SampledNormal, FColor::White);
 							}
 						}
 
diff --git a/Engine/Source/Runtime/Landscape/Private/LandscapeGrass.cpp b/Engine/Source/Runtime/Landscape/Private/LandscapeGrass.cpp
--- a/Engine/Source/Runtime/Landscape/Private/LandscapeGrass.cpp
+++ b/Engine/Source/Runtime/Landscape/Private/LandscapeGrass.cpp
@@ -477,7 +477,7 @@ bool FLandscapeGrassWeightMeshProcessor::Process(
 	for (int32 PassIndex = 0; PassIndex < NumPasses; ++PassIndex)
 	{
 		ShaderElementData.OutputPass = (PassIndex >= FirstHeightMipsPassIndex) ? 0 : PassIndex;
-		ShaderElementData.RenderOffset = ViewOffset + FVector2f(PassOffsetX * PassIndex, 0);
+		ShaderElementData.RenderOffset = FVector2f(ViewOffset) + FVector2f(PassOffsetX * PassIndex, 0);	// LWC_TODO: Precision loss
 
 		uint64 Mask = (PassIndex >= FirstHeightMipsPassIndex) ? HeightMips[PassIndex - FirstHeightMipsPassIndex] : BatchElementMask;
 
diff --git a/Engine/Source/Runtime/MeshConversion/Private/MeshDescriptionBuilder.cpp b/Engine/Source/Runtime/MeshConversion/Private/MeshDescriptionBuilder.cpp
--- a/Engine/Source/Runtime/MeshConversion/Private/MeshDescriptionBuilder.cpp
+++ b/Engine/Source/Runtime/MeshConversion/Private/MeshDescriptionBuilder.cpp
@@ -139,7 +139,7 @@ void FMeshDescriptionBuilder::AppendUVTriangle(const FTriangleID& TriangleID, co
 		const FVertexInstanceID CornerInstanceID = TriVertInstances[j];
 		const FVector2f& UVvalue = UVCoordinates[TempUVBuffer[j]];
 		
-		SetInstanceUV(CornerInstanceID, UVvalue, UVLayerIndex);
+		SetInstanceUV(CornerInstanceID, FVector2D(UVvalue), UVLayerIndex);
 	}
 }
 
diff --git a/Engine/Source/Runtime/MeshConversion/Private/MeshDescriptionUVsToDynamicMesh.cpp b/Engine/Source/Runtime/MeshConversion/Private/MeshDescriptionUVsToDynamicMesh.cpp
--- a/Engine/Source/Runtime/MeshConversion/Private/MeshDescriptionUVsToDynamicMesh.cpp
+++ b/Engine/Source/Runtime/MeshConversion/Private/MeshDescriptionUVsToDynamicMesh.cpp
@@ -152,7 +152,7 @@ TSharedPtr<FDynamicMesh3> FMeshDescriptionUVsToDynamicMesh::GetUVMesh(const FMes
 			TArrayView<const FVertexInstanceID> SourceTriangleInstanceIds = MeshDescription->GetTriangleVertexInstances(TriangleIDIn);
 			for (int32 i = 0; i < 3; ++i)
 			{
-				FVector2D UV = InstanceUVs.Get(SourceTriangleInstanceIds[i], UVLayerIndex);
+				FVector2f UV = InstanceUVs.Get(SourceTriangleInstanceIds[i], UVLayerIndex);
 				TriangleToInsertOut[i] = UVMeshWelder.FindOrAddUnique(UV, SourceTriangleVids[i].GetValue());
 			}
 		};
diff --git a/Engine/Source/Runtime/MovieSceneTracks/Private/Evaluation/MovieSceneParameterTemplate.cpp b/Engine/Source/Runtime/MovieSceneTracks/Private/Evaluation/MovieSceneParameterTemplate.cpp
--- a/Engine/Source/Runtime/MovieSceneTracks/Private/Evaluation/MovieSceneParameterTemplate.cpp
+++ b/Engine/Source/Runtime/MovieSceneTracks/Private/Evaluation/MovieSceneParameterTemplate.cpp
@@ -49,7 +49,7 @@ void FMovieSceneParameterSectionTemplate::EvaluateCurves(const FMovieSceneContex
 
 		if (bAnyEvaluated)
 		{
-			Values.Vector2DValues.Emplace(Vector.ParameterName, Value);
+			Values.Vector2DValues.Emplace(Vector.ParameterName, FVector2D(Value));
 		}
 	}
 
diff --git a/Engine/Source/Runtime/NavigationSystem/Private/NavMesh/NavMeshRenderingComponent.cpp b/Engine/Source/Runtime/NavigationSystem/Private/NavMesh/NavMeshRenderingComponent.cpp
--- a/Engine/Source/Runtime/NavigationSystem/Private/NavMesh/NavMeshRenderingComponent.cpp
+++ b/Engine/Source/Runtime/NavigationSystem/Private/NavMesh/NavMeshRenderingComponent.cpp
@@ -212,7 +212,7 @@ namespace FNavMeshRenderingHelpers
 	{
 		FDynamicMeshVertex* Vertex = new(MeshData.Vertices) FDynamicMeshVertex;
 		Vertex->Position = Pos;
-		Vertex->TextureCoordinate[0] = FVector2D::ZeroVector;
+		Vertex->TextureCoordinate[0] = FVector2f::ZeroVector;
 		Vertex->TangentX = FVector(1.0f, 0.0f, 0.0f);
 		Vertex->TangentZ = FVector(0.0f, 1.0f, 0.0f);
 		// store the sign of the determinant in TangentZ.W (-1=-128,+1=127)
diff --git a/Engine/Source/Runtime/Slate/Private/Framework/Application/SlateApplication.cpp b/Engine/Source/Runtime/Slate/Private/Framework/Application/SlateApplication.cpp
--- a/Engine/Source/Runtime/Slate/Private/Framework/Application/SlateApplication.cpp
+++ b/Engine/Source/Runtime/Slate/Private/Framework/Application/SlateApplication.cpp
@@ -3931,7 +3931,7 @@ bool FSlateApplication::TakeScreenshot(const TSharedRef<SWidget>& Widget, const
 	FSlateApplication::Get().GeneratePathToWidgetChecked(Widget, WidgetPath);
 
 	FArrangedWidget ArrangedWidget = WidgetPath.FindArrangedWidget(Widget).Get(FArrangedWidget::GetNullWidget());
-	FVector2D Position = ArrangedWidget.Geometry.AbsolutePosition;
+	FVector2D Position = FVector2D(ArrangedWidget.Geometry.AbsolutePosition);
 	FVector2D Size = ArrangedWidget.Geometry.GetDrawSize();
 	FVector2D WindowPosition = WidgetWindow->GetPositionInScreen();
 
diff --git a/Engine/Source/Runtime/Slate/Private/Framework/MultiBox/SClippingHorizontalBox.cpp b/Engine/Source/Runtime/Slate/Private/Framework/MultiBox/SClippingHorizontalBox.cpp
--- a/Engine/Source/Runtime/Slate/Private/Framework/MultiBox/SClippingHorizontalBox.cpp
+++ b/Engine/Source/Runtime/Slate/Private/Framework/MultiBox/SClippingHorizontalBox.cpp
@@ -27,7 +27,7 @@ void SClippingHorizontalBox::OnArrangeChildren( const FGeometry& AllottedGeometr
 	{
 		const FArrangedWidget& CurWidget = ArrangedChildren[ChildIdx];
 		FVector2D AbsWidgetPos = CurWidget.Geometry.LocalToAbsolute(CurWidget.Geometry.GetLocalSize());
-		FVector2D AbsBoxPos = AllottedGeometry.AbsolutePosition + AllottedGeometry.GetLocalSize() * AllottedGeometry.Scale;
+		FVector2D AbsBoxPos = FVector2D(AllottedGeometry.AbsolutePosition) + AllottedGeometry.GetLocalSize() * AllottedGeometry.Scale;
 		if (FMath::TruncToInt(AbsWidgetPos.X) > FMath::TruncToInt(AbsBoxPos.X))
 		{
 			++NumClippedChildren;
diff --git a/Engine/Source/Runtime/Slate/Private/Widgets/Input/SMenuAnchor.cpp b/Engine/Source/Runtime/Slate/Private/Widgets/Input/SMenuAnchor.cpp
--- a/Engine/Source/Runtime/Slate/Private/Widgets/Input/SMenuAnchor.cpp
+++ b/Engine/Source/Runtime/Slate/Private/Widgets/Input/SMenuAnchor.cpp
@@ -345,7 +345,7 @@ void SMenuAnchor::SetIsOpen( bool InIsOpen, const bool bFocusMenu, const int32 F
 					const FVector2D DesiredContentSize = MenuContentRef->GetDesiredSize();  // @todo slate: This is ignoring any window border size!
 					const EMenuPlacement PlacementMode = Placement.Get();
 
-					const FVector2D NewPosition = MyGeometry.AbsolutePosition;
+					const FVector2D NewPosition = FVector2D(MyGeometry.AbsolutePosition);
 					FVector2D NewWindowSize = DesiredContentSize;
 
 					FPopupTransitionEffect TransitionEffect( FPopupTransitionEffect::None );
diff --git a/Engine/Source/Runtime/Slate/Private/Widgets/Text/ISlateEditableTextWidget.cpp b/Engine/Source/Runtime/Slate/Private/Widgets/Text/ISlateEditableTextWidget.cpp
--- a/Engine/Source/Runtime/Slate/Private/Widgets/Text/ISlateEditableTextWidget.cpp
+++ b/Engine/Source/Runtime/Slate/Private/Widgets/Text/ISlateEditableTextWidget.cpp
@@ -39,7 +39,7 @@ ECursorAction FMoveCursor::GetAction() const
 
 FVector2D FMoveCursor::GetLocalPosition() const
 {
-	return DirectionOrPosition;
+	return FVector2D(DirectionOrPosition);
 }
 
 ECursorMoveGranularity FMoveCursor::GetGranularity() const
diff --git a/Engine/Source/Runtime/Slate/Private/Widgets/Text/SlateEditableTextLayout.cpp b/Engine/Source/Runtime/Slate/Private/Widgets/Text/SlateEditableTextLayout.cpp
--- a/Engine/Source/Runtime/Slate/Private/Widgets/Text/SlateEditableTextLayout.cpp
+++ b/Engine/Source/Runtime/Slate/Private/Widgets/Text/SlateEditableTextLayout.cpp
@@ -130,7 +130,7 @@ FSlateEditableTextLayout::FSlateEditableTextLayout(ISlateEditableTextWidget& InO
 	TextSelectionHighlighter = SlateEditableTextTypes::FTextSelectionHighlighter::Create();
 	SearchSelectionHighlighter = SlateEditableTextTypes::FTextSearchHighlighter::Create();
 
-	ScrollOffset = FVector2D::ZeroVector;
+	ScrollOffset = FVector2f::ZeroVector;
 	PreferredCursorScreenOffsetInLine = 0.0f;
 	SelectionStart = TOptional<FTextLocation>();
 	CurrentUndoLevel = INDEX_NONE;
@@ -143,7 +143,7 @@ FSlateEditableTextLayout::FSlateEditableTextLayout(ISlateEditableTextWidget& InO
 	bSelectionChangedExternally = false;
 	VirtualKeyboardTextCommitType = ETextCommit::Default;
 
-	CachedSize = FVector2D::ZeroVector;
+	CachedSize = FVector2f::ZeroVector;
 
 	auto ExecuteDeleteAction = [this]()
 	{
@@ -661,13 +661,13 @@ void FSlateEditableTextLayout::AdvanceSearch(const bool InReverse)
 FVector2D FSlateEditableTextLayout::SetHorizontalScrollFraction(const float InScrollOffsetFraction)
 {
 	ScrollOffset.X = FMath::Clamp<float>(InScrollOffsetFraction, 0.0, 1.0) * TextLayout->GetSize().X;
-	return ScrollOffset;
+	return FVector2D(ScrollOffset);
 }
 
 FVector2D FSlateEditableTextLayout::SetVerticalScrollFraction(const float InScrollOffsetFraction)
 {
 	ScrollOffset.Y = FMath::Clamp<float>(InScrollOffsetFraction, 0.0, 1.0) * TextLayout->GetSize().Y;
-	return ScrollOffset;
+	return FVector2D(ScrollOffset);
 }
 
 FVector2D FSlateEditableTextLayout::SetScrollOffset(const FVector2D& InScrollOffset, const FGeometry& InGeometry)
@@ -675,12 +675,12 @@ FVector2D FSlateEditableTextLayout::SetScrollOffset(const FVector2D& InScrollOff
 	const FVector2D ContentSize = TextLayout->GetSize();
 	ScrollOffset.X = FMath::Clamp(InScrollOffset.X, 0.0f, ContentSize.X - InGeometry.GetLocalSize().X);
 	ScrollOffset.Y = FMath::Clamp(InScrollOffset.Y, 0.0f, ContentSize.Y - InGeometry.GetLocalSize().Y);
-	return ScrollOffset;
+	return FVector2D(ScrollOffset);
 }
 
 FVector2D FSlateEditableTextLayout::GetScrollOffset() const
 {
-	return ScrollOffset;
+	return FVector2D(ScrollOffset);
 }
 
 
@@ -3380,14 +3380,14 @@ void FSlateEditableTextLayout::Tick(const FGeometry& AllottedGeometry, const dou
 		ScrollOffset.Y = OwnerWidget->UpdateAndClampVerticalScrollBar(ViewOffset, ViewFraction, ScrollBarVisiblityOverride);
 	}
 
-	TextLayout->SetVisibleRegion(AllottedGeometry.Size, ScrollOffset * TextLayout->GetScale());
+	TextLayout->SetVisibleRegion(AllottedGeometry.Size, FVector2D(ScrollOffset) * TextLayout->GetScale());
 }
 
 int32 FSlateEditableTextLayout::OnPaint(const FPaintArgs& Args, const FGeometry& AllottedGeometry, const FSlateRect& MyCullingRect, FSlateWindowElementList& OutDrawElements, int32 LayerId, const FWidgetStyle& InWidgetStyle, bool bParentEnabled)
 {
 	// Update the auto-wrap size now that we have computed paint geometry; won't take affect until text frame
 	// Note: This is done here rather than in Tick(), because Tick() doesn't get called while resizing windows, but OnPaint() does
-	CachedSize = AllottedGeometry.GetLocalSize();
+	CachedSize = FVector2f(AllottedGeometry.GetLocalSize());
 
 	// Only paint the hint text layout if we don't have any text set
 	if (TextLayout->IsEmpty() && HintTextLayout.IsValid())
@@ -3433,7 +3433,7 @@ void FSlateEditableTextLayout::CacheDesiredSize(float LayoutScaleMultiplier)
 	TextLayout->SetMargin(MarginValue);
 	TextLayout->SetLineHeightPercentage(LineHeightPercentage.Get());
 	TextLayout->SetJustification(Justification.Get());
-	TextLayout->SetVisibleRegion(CachedSize, ScrollOffset * TextLayout->GetScale());
+	TextLayout->SetVisibleRegion(FVector2D(CachedSize), FVector2D(ScrollOffset) * TextLayout->GetScale());
 	TextLayout->UpdateIfNeeded();
 }
 
@@ -3894,7 +3894,7 @@ bool FSlateEditableTextLayout::FTextInputMethodContext::GetTextBounds(const uint
 
 	// Translate the position (which is in local space) into screen (absolute) space
 	// Note: The local positions are pre-scaled, so we don't scale them again here
-	Position += CachedGeometry.AbsolutePosition;
+	Position += FVector2D(CachedGeometry.AbsolutePosition);
 	
 	return false; // false means "not clipped"
 }
@@ -3908,7 +3908,7 @@ void FSlateEditableTextLayout::FTextInputMethodContext::GetScreenBounds(FVector2
 		return;
 	}
 
-	Position = CachedGeometry.AbsolutePosition;
+	Position = FVector2D(CachedGeometry.AbsolutePosition);
 	Size = CachedGeometry.GetDrawSize();
 }
 
diff --git a/Engine/Source/Runtime/Slate/Private/Widgets/Text/SlateTextBlockLayout.cpp b/Engine/Source/Runtime/Slate/Private/Widgets/Text/SlateTextBlockLayout.cpp
--- a/Engine/Source/Runtime/Slate/Private/Widgets/Text/SlateTextBlockLayout.cpp
+++ b/Engine/Source/Runtime/Slate/Private/Widgets/Text/SlateTextBlockLayout.cpp
@@ -186,7 +186,7 @@ float FSlateTextBlockLayout::GetLayoutScale() const
 int32 FSlateTextBlockLayout::OnPaint(const FPaintArgs& InPaintArgs, const FGeometry& InAllottedGeometry, const FSlateRect& InClippingRect, FSlateWindowElementList& OutDrawElements, int32 LayerId, const FWidgetStyle& InWidgetStyle, bool bParentEnabled)
 {
 	// Store a new cached size with the scale
-	CachedSize = InAllottedGeometry.GetLocalSize();
+	CachedSize = FVector2f(InAllottedGeometry.GetLocalSize());
 
 	// Recompute wrapping in case the cached size changed.
 	TextLayout->SetWrappingWidth(CalculateWrappingWidth());
@@ -217,7 +217,7 @@ int32 FSlateTextBlockLayout::OnPaint(const FPaintArgs& InPaintArgs, const FGeome
 		}
 	}
 
-	TextLayout->SetVisibleRegion(CachedSize, AutoScrollValue * TextLayout->GetScale());
+	TextLayout->SetVisibleRegion(FVector2D(CachedSize), AutoScrollValue * TextLayout->GetScale());
 	TextLayout->UpdateIfNeeded();
 
 	return TextLayout->OnPaint(InPaintArgs, InAllottedGeometry, InClippingRect, OutDrawElements, LayerId, InWidgetStyle, bParentEnabled);
diff --git a/Engine/Source/Runtime/SlateCore/Private/Input/HittestGrid.cpp b/Engine/Source/Runtime/SlateCore/Private/Input/HittestGrid.cpp
--- a/Engine/Source/Runtime/SlateCore/Private/Input/HittestGrid.cpp
+++ b/Engine/Source/Runtime/SlateCore/Private/Input/HittestGrid.cpp
@@ -50,10 +50,10 @@ FVector2D ClosestPointOnSlateRotatedRect(const FVector2D &Point, const FSlateRot
 
 	const static int32 NumOfCorners = 4;
 	FVector2D Corners[NumOfCorners];
-	Corners[0] = RotatedRect.TopLeft;
-	Corners[1] = Corners[0] + RotatedRect.ExtentX;
-	Corners[2] = Corners[1] + RotatedRect.ExtentY;
-	Corners[3] = Corners[0] + RotatedRect.ExtentY;
+	Corners[0] = FVector2D(RotatedRect.TopLeft);
+	Corners[1] = FVector2D(Corners[0]) + FVector2D(RotatedRect.ExtentX);
+	Corners[2] = FVector2D(Corners[1]) + FVector2D(RotatedRect.ExtentY);
+	Corners[3] = FVector2D(Corners[0]) + FVector2D(RotatedRect.ExtentY);
 
 	FVector2D RetPoint;
 	float ClosestDistSq = FLT_MAX;
diff --git a/Engine/Source/Runtime/SlateCore/Private/Layout/Clipping.cpp b/Engine/Source/Runtime/SlateCore/Private/Layout/Clipping.cpp
--- a/Engine/Source/Runtime/SlateCore/Private/Layout/Clipping.cpp
+++ b/Engine/Source/Runtime/SlateCore/Private/Layout/Clipping.cpp
@@ -15,10 +15,10 @@ FSlateClippingZone::FSlateClippingZone(const FShortRect& AxisAlignedRect)
 	int16 Top = FMath::Min(AxisAlignedRect.Top, AxisAlignedRect.Bottom);
 	int16 Bottom = FMath::Max(AxisAlignedRect.Top, AxisAlignedRect.Bottom);
 
-	TopLeft = FVector2D(Left, Top);
-	TopRight = FVector2D(Right, Top);
-	BottomLeft = FVector2D(Left, Bottom);
-	BottomRight = FVector2D(Right, Bottom);
+	TopLeft = FVector2f(Left, Top);
+	TopRight = FVector2f(Right, Top);
+	BottomLeft = FVector2f(Left, Bottom);
+	BottomRight = FVector2f(Right, Bottom);
 }
 
 FSlateClippingZone::FSlateClippingZone(const FSlateRect& AxisAlignedRect)
@@ -32,10 +32,10 @@ FSlateClippingZone::FSlateClippingZone(const FSlateRect& AxisAlignedRect)
 	float Top = FMath::Min(RoundedAxisAlignedRect.Top, RoundedAxisAlignedRect.Bottom);
 	float Bottom = FMath::Max(RoundedAxisAlignedRect.Top, RoundedAxisAlignedRect.Bottom);
 
-	TopLeft = FVector2D(Left, Top);
-	TopRight = FVector2D(Right, Top);
-	BottomLeft = FVector2D(Left, Bottom);
-	BottomRight = FVector2D(Right, Bottom);
+	TopLeft = FVector2f(Left, Top);
+	TopRight = FVector2f(Right, Top);
+	BottomLeft = FVector2f(Left, Bottom);
+	BottomRight = FVector2f(Right, Bottom);
 }
 
 FSlateClippingZone::FSlateClippingZone(const FGeometry& BooundingGeometry)
@@ -186,7 +186,7 @@ bool FSlateClippingZone::IsPointInside(const FVector2D& Point) const
 	}
 	else
 	{
-		if (IsPointInTriangle(Point, TopLeft, TopRight, BottomLeft) || IsPointInTriangle(Point, BottomLeft, TopRight, BottomRight))
+		if (IsPointInTriangle(Point, FVector2D(TopLeft), FVector2D(TopRight), FVector2D(BottomLeft)) || IsPointInTriangle(Point, FVector2D(BottomLeft), FVector2D(TopRight), FVector2D(BottomRight)))
 		{
 			return true;
 		}
diff --git a/Engine/Source/Runtime/SlateCore/Private/Layout/SlateRotatedRect.cpp b/Engine/Source/Runtime/SlateCore/Private/Layout/SlateRotatedRect.cpp
--- a/Engine/Source/Runtime/SlateCore/Private/Layout/SlateRotatedRect.cpp
+++ b/Engine/Source/Runtime/SlateCore/Private/Layout/SlateRotatedRect.cpp
@@ -6,7 +6,7 @@
 // !!! WRH 2014/08/25 - this is a brute-force, not efficient implementation, uses a bunch of extra conditionals.
 FSlateRect FSlateRotatedRect::ToBoundingRect() const
 {
-	FVector2D Points[4] = 
+	FVector2f Points[4] = 
 	{
 		TopLeft,
 		TopLeft + ExtentX,
@@ -23,14 +23,14 @@ FSlateRect FSlateRotatedRect::ToBoundingRect() const
 
 bool FSlateRotatedRect::IsUnderLocation(const FVector2D& Location) const
 {
-	const FVector2D Offset = Location - TopLeft;
-	const float Det = FVector2D::CrossProduct(ExtentX, ExtentY);
+	const FVector2D Offset = Location - FVector2D(TopLeft);
+	const float Det = FVector2D::CrossProduct(FVector2D(ExtentX), FVector2D(ExtentY));
 
 	// Not exhaustively efficient. Could optimize the checks for [0..1] to short circuit faster.
-	const float S = -FVector2D::CrossProduct(Offset, ExtentX) / Det;
+	const float S = -FVector2D::CrossProduct(Offset, FVector2D(ExtentX)) / Det;
 	if (FMath::IsWithinInclusive(S, 0.0f, 1.0f))
 	{
-		const float T = FVector2D::CrossProduct(Offset, ExtentY) / Det;
+		const float T = FVector2D::CrossProduct(Offset, FVector2D(ExtentY)) / Det;
 		return FMath::IsWithinInclusive(T, 0.0f, 1.0f);
 	}
 	return false;
@@ -40,13 +40,13 @@ FSlateRotatedRect FSlateRotatedRect::MakeRotatedRect(const FSlateRect& ClipRectI
 {
 	const FSlateRotatedRect RotatedRect = TransformRect(LayoutToRenderTransform, FSlateRotatedRect(ClipRectInLayoutWindowSpace));
 
-	const FVector2D TopRight = RotatedRect.TopLeft + RotatedRect.ExtentX;
-	const FVector2D BottomLeft = RotatedRect.TopLeft + RotatedRect.ExtentY;
+	const FVector2D TopRight = FVector2D(RotatedRect.TopLeft) + FVector2D(RotatedRect.ExtentX);
+	const FVector2D BottomLeft = FVector2D(RotatedRect.TopLeft) + FVector2D(RotatedRect.ExtentY);
 
 	return FSlateRotatedRect(
-		RotatedRect.TopLeft,
-		TopRight - RotatedRect.TopLeft,
-		BottomLeft - RotatedRect.TopLeft);
+		FVector2D(RotatedRect.TopLeft),
+		TopRight - FVector2D(RotatedRect.TopLeft),
+		BottomLeft - FVector2D(RotatedRect.TopLeft));
 }
 
 FSlateRotatedRect FSlateRotatedRect::MakeSnappedRotatedRect(const FSlateRect& ClipRectInLayoutWindowSpace, const FTransform2D& LayoutToRenderTransform)
@@ -56,9 +56,9 @@ FSlateRotatedRect FSlateRotatedRect::MakeSnappedRotatedRect(const FSlateRect& Cl
 	// Pixel snapping is done here by rounding the resulting floats to ints, we do this before
 	// calculating the final extents of the clip box otherwise we'll get a smaller clip rect than a visual
 	// rect where each point is individually snapped.
-	const FVector2D SnappedTopLeft = ( RotatedRect.TopLeft ).RoundToVector();
-	const FVector2D SnappedTopRight = ( RotatedRect.TopLeft + RotatedRect.ExtentX ).RoundToVector();
-	const FVector2D SnappedBottomLeft = ( RotatedRect.TopLeft + RotatedRect.ExtentY ).RoundToVector();
+	const FVector2D SnappedTopLeft = ( FVector2D(RotatedRect.TopLeft) ).RoundToVector();
+	const FVector2D SnappedTopRight = ( FVector2D(RotatedRect.TopLeft) + FVector2D(RotatedRect.ExtentX) ).RoundToVector();
+	const FVector2D SnappedBottomLeft = ( FVector2D(RotatedRect.TopLeft) + FVector2D(RotatedRect.ExtentY) ).RoundToVector();
 
 	//NOTE: We explicitly do not re-snap the extent x/y, it wouldn't be correct to snap again in distance space
 	// even if two points are snapped, their distance wont necessarily be a whole number if those points are not
diff --git a/Engine/Source/Runtime/SlateCore/Private/Rendering/DrawElements.cpp b/Engine/Source/Runtime/SlateCore/Private/Rendering/DrawElements.cpp
--- a/Engine/Source/Runtime/SlateCore/Private/Rendering/DrawElements.cpp
+++ b/Engine/Source/Runtime/SlateCore/Private/Rendering/DrawElements.cpp
@@ -152,7 +152,7 @@ void FSlateDrawElement::Init(FSlateWindowElementList& ElementList, EElementType
 	RenderTransform = PaintGeometry.GetAccumulatedRenderTransform();
 	Position = PaintGeometry.DrawPosition;
 	Scale = PaintGeometry.DrawScale;
-	LocalSize = PaintGeometry.GetLocalSize();
+	LocalSize = FVector2f(PaintGeometry.GetLocalSize());
 	ClipStateHandle.SetPreCachedClipIndex(ElementList.GetClippingIndex());
 
 #if UE_SLATE_VERIFY_PIXELSIZE
@@ -167,7 +167,7 @@ void FSlateDrawElement::Init(FSlateWindowElementList& ElementList, EElementType
 	DrawEffects = InDrawEffects;
 	
 	// Calculate the layout to render transform as this is needed by several calculations downstream.
-	const FSlateLayoutTransform InverseLayoutTransform(Inverse(FSlateLayoutTransform(Scale, Position)));
+	const FSlateLayoutTransform InverseLayoutTransform(Inverse(FSlateLayoutTransform(Scale, FVector2D(Position))));
 
 	{
 		// This is a workaround because we want to keep track of the various Scenes 
@@ -199,7 +199,7 @@ void FSlateDrawElement::ApplyPositionOffset(const FVector2D& InOffset)
 	RenderTransform = Concatenate(RenderTransform, InOffset);
 
 	// Recompute cached layout to render transform
-	const FSlateLayoutTransform InverseLayoutTransform(Inverse(FSlateLayoutTransform(Scale, Position)));
+	const FSlateLayoutTransform InverseLayoutTransform(Inverse(FSlateLayoutTransform(Scale, FVector2D(Position))));
 }
 
 void FSlateDrawElement::MakeDebugQuad( FSlateWindowElementList& ElementList, uint32 InLayer, const FPaintGeometry& PaintGeometry, FLinearColor Tint)
diff --git a/Engine/Source/Runtime/SlateCore/Private/Rendering/ElementBatcher.cpp b/Engine/Source/Runtime/SlateCore/Private/Rendering/ElementBatcher.cpp
--- a/Engine/Source/Runtime/SlateCore/Private/Rendering/ElementBatcher.cpp
+++ b/Engine/Source/Runtime/SlateCore/Private/Rendering/ElementBatcher.cpp
@@ -512,10 +512,10 @@ void FSlateElementBatcher::AddDebugQuadElement(const FSlateDrawElement& DrawElem
 	const uint32 IndexStart = 0;
 
 	// Add four vertices to the list of verts to be added to the vertex buffer
-	RenderBatch.AddVertex(FSlateVertex::Make<Rounding>(RenderTransform, TopLeft, FVector2D(0.0f,0.0f),  Tint));
-	RenderBatch.AddVertex(FSlateVertex::Make<Rounding>(RenderTransform, TopRight, FVector2D(1.0f,0.0f), Tint));
-	RenderBatch.AddVertex(FSlateVertex::Make<Rounding>(RenderTransform, BotLeft, FVector2D(0.0f,1.0f),  Tint));
-	RenderBatch.AddVertex(FSlateVertex::Make<Rounding>(RenderTransform, BotRight, FVector2D(1.0f,1.0f), Tint));
+	RenderBatch.AddVertex(FSlateVertex::Make<Rounding>(RenderTransform, FVector2f(TopLeft), FVector2f(0.0f,0.0f),  Tint));
+	RenderBatch.AddVertex(FSlateVertex::Make<Rounding>(RenderTransform, FVector2f(TopRight), FVector2f(1.0f,0.0f), Tint));
+	RenderBatch.AddVertex(FSlateVertex::Make<Rounding>(RenderTransform, FVector2f(BotLeft), FVector2f(0.0f,1.0f),  Tint));
+	RenderBatch.AddVertex(FSlateVertex::Make<Rounding>(RenderTransform, FVector2f(BotRight), FVector2f(1.0f,1.0f), Tint));
 
 	// Add 6 indices to the vertex buffer.  (2 tri's per quad, 3 indices per tri)
 	RenderBatch.AddIndex(IndexStart + 0);
@@ -590,9 +590,9 @@ void FSlateElementBatcher::AddBoxElement(const FSlateDrawElement& DrawElement)
 		}
 		else
 		{
-			SizeUV = ResourceProxy->SizeUV;
-			StartUV = ResourceProxy->StartUV + HalfTexel;
-			EndUV = StartUV + ResourceProxy->SizeUV;
+			SizeUV = FVector2D(ResourceProxy->SizeUV);
+			StartUV = FVector2D(ResourceProxy->StartUV) + HalfTexel;
+			EndUV = StartUV + FVector2D(ResourceProxy->SizeUV);
 		}
 	}
 	else
@@ -725,23 +725,23 @@ void FSlateElementBatcher::AddBoxElement(const FSlateDrawElement& DrawElement)
 		FVector2D Position = TopLeft;
 		FVector2D EndPos = BotRight;
 
-		RenderBatch.AddVertex( FSlateVertex::Make<Rounding>( RenderTransform, FVector2D( Position.X, Position.Y ),		LocalSize, DrawScale, FVector4f(StartUV,										Tiling),	Tint, SecondaryColor) ); //0
-		RenderBatch.AddVertex( FSlateVertex::Make<Rounding>( RenderTransform, FVector2D( Position.X, TopMarginY ),		LocalSize, DrawScale, FVector4f(FVector2D( StartUV.X, TopMarginV ),			Tiling),	Tint, SecondaryColor) ); //1
-		RenderBatch.AddVertex( FSlateVertex::Make<Rounding>( RenderTransform, FVector2D( LeftMarginX, Position.Y ),		LocalSize, DrawScale, FVector4f(FVector2D( LeftMarginU, StartUV.Y ),			Tiling),	Tint, SecondaryColor) ); //2
-		RenderBatch.AddVertex( FSlateVertex::Make<Rounding>( RenderTransform, FVector2D( LeftMarginX, TopMarginY ),		LocalSize, DrawScale, FVector4f(FVector2D( LeftMarginU, TopMarginV ),		Tiling),	Tint, SecondaryColor) ); //3
-		RenderBatch.AddVertex( FSlateVertex::Make<Rounding>( RenderTransform, FVector2D( RightMarginX, Position.Y ),	LocalSize, DrawScale, FVector4f(FVector2D( RightMarginU, StartUV.Y ),		Tiling),	Tint, SecondaryColor) ); //4
-		RenderBatch.AddVertex( FSlateVertex::Make<Rounding>( RenderTransform, FVector2D( RightMarginX, TopMarginY ),	LocalSize, DrawScale, FVector4f(FVector2D( RightMarginU,TopMarginV),			Tiling),	Tint, SecondaryColor) ); //5
-		RenderBatch.AddVertex( FSlateVertex::Make<Rounding>( RenderTransform, FVector2D( EndPos.X, Position.Y ),		LocalSize, DrawScale, FVector4f(FVector2D( EndUV.X, StartUV.Y ),				Tiling),	Tint, SecondaryColor) ); //6
-		RenderBatch.AddVertex( FSlateVertex::Make<Rounding>( RenderTransform, FVector2D( EndPos.X, TopMarginY ),		LocalSize, DrawScale, FVector4f(FVector2D( EndUV.X, TopMarginV),				Tiling),	Tint, SecondaryColor) ); //7
-
-		RenderBatch.AddVertex( FSlateVertex::Make<Rounding>( RenderTransform, FVector2D( Position.X, BottomMarginY ),	LocalSize, DrawScale, FVector4f(FVector2D( StartUV.X, BottomMarginV ),		Tiling),	Tint, SecondaryColor) ); //8
-		RenderBatch.AddVertex( FSlateVertex::Make<Rounding>( RenderTransform, FVector2D( LeftMarginX, BottomMarginY ),	LocalSize, DrawScale, FVector4f(FVector2D( LeftMarginU, BottomMarginV ),		Tiling),	Tint, SecondaryColor) ); //9
-		RenderBatch.AddVertex( FSlateVertex::Make<Rounding>( RenderTransform, FVector2D( RightMarginX, BottomMarginY ),	LocalSize, DrawScale, FVector4f(FVector2D( RightMarginU, BottomMarginV ),	Tiling),	Tint, SecondaryColor) ); //10
-		RenderBatch.AddVertex( FSlateVertex::Make<Rounding>( RenderTransform, FVector2D( EndPos.X, BottomMarginY ),		LocalSize, DrawScale, FVector4f(FVector2D( EndUV.X, BottomMarginV ),			Tiling),	Tint, SecondaryColor) ); //11
-		RenderBatch.AddVertex( FSlateVertex::Make<Rounding>( RenderTransform, FVector2D( Position.X, EndPos.Y ),		LocalSize, DrawScale, FVector4f(FVector2D( StartUV.X, EndUV.Y ),				Tiling),	Tint, SecondaryColor) ); //12
-		RenderBatch.AddVertex( FSlateVertex::Make<Rounding>( RenderTransform, FVector2D( LeftMarginX, EndPos.Y ),		LocalSize, DrawScale, FVector4f(FVector2D( LeftMarginU, EndUV.Y ),			Tiling),	Tint, SecondaryColor) ); //13
-		RenderBatch.AddVertex( FSlateVertex::Make<Rounding>( RenderTransform, FVector2D( RightMarginX, EndPos.Y ),		LocalSize, DrawScale, FVector4f(FVector2D( RightMarginU, EndUV.Y ),			Tiling),	Tint, SecondaryColor) ); //14
-		RenderBatch.AddVertex( FSlateVertex::Make<Rounding>( RenderTransform, FVector2D( EndPos.X, EndPos.Y ),			LocalSize, DrawScale, FVector4f(EndUV,										Tiling),	Tint, SecondaryColor) ); //15
+		RenderBatch.AddVertex( FSlateVertex::Make<Rounding>( RenderTransform, FVector2f( Position.X, Position.Y ),		FVector2f(LocalSize), DrawScale, FVector4f(FVector2f( StartUV ),						FVector2f(Tiling)),	Tint, SecondaryColor) ); //0
+		RenderBatch.AddVertex( FSlateVertex::Make<Rounding>( RenderTransform, FVector2f( Position.X, TopMarginY ),		FVector2f(LocalSize), DrawScale, FVector4f(FVector2f( StartUV.X, TopMarginV ),			FVector2f(Tiling)),	Tint, SecondaryColor) ); //1
+		RenderBatch.AddVertex( FSlateVertex::Make<Rounding>( RenderTransform, FVector2f( LeftMarginX, Position.Y ),		FVector2f(LocalSize), DrawScale, FVector4f(FVector2f( LeftMarginU, StartUV.Y ),			FVector2f(Tiling)),	Tint, SecondaryColor) ); //2
+		RenderBatch.AddVertex( FSlateVertex::Make<Rounding>( RenderTransform, FVector2f( LeftMarginX, TopMarginY ),		FVector2f(LocalSize), DrawScale, FVector4f(FVector2f( LeftMarginU, TopMarginV ),		FVector2f(Tiling)),	Tint, SecondaryColor) ); //3
+		RenderBatch.AddVertex( FSlateVertex::Make<Rounding>( RenderTransform, FVector2f( RightMarginX, Position.Y ),	FVector2f(LocalSize), DrawScale, FVector4f(FVector2f( RightMarginU, StartUV.Y ),		FVector2f(Tiling)),	Tint, SecondaryColor) ); //4
+		RenderBatch.AddVertex( FSlateVertex::Make<Rounding>( RenderTransform, FVector2f( RightMarginX, TopMarginY ),	FVector2f(LocalSize), DrawScale, FVector4f(FVector2f( RightMarginU,TopMarginV),			FVector2f(Tiling)),	Tint, SecondaryColor) ); //5
+		RenderBatch.AddVertex( FSlateVertex::Make<Rounding>( RenderTransform, FVector2f( EndPos.X, Position.Y ),		FVector2f(LocalSize), DrawScale, FVector4f(FVector2f( EndUV.X, StartUV.Y ),				FVector2f(Tiling)),	Tint, SecondaryColor) ); //6
+		RenderBatch.AddVertex( FSlateVertex::Make<Rounding>( RenderTransform, FVector2f( EndPos.X, TopMarginY ),		FVector2f(LocalSize), DrawScale, FVector4f(FVector2f( EndUV.X, TopMarginV),				FVector2f(Tiling)),	Tint, SecondaryColor) ); //7
+
+		RenderBatch.AddVertex( FSlateVertex::Make<Rounding>( RenderTransform, FVector2f( Position.X, BottomMarginY ),	FVector2f(LocalSize), DrawScale, FVector4f(FVector2f( StartUV.X, BottomMarginV ),		FVector2f(Tiling)),	Tint, SecondaryColor) ); //8
+		RenderBatch.AddVertex( FSlateVertex::Make<Rounding>( RenderTransform, FVector2f( LeftMarginX, BottomMarginY ),	FVector2f(LocalSize), DrawScale, FVector4f(FVector2f( LeftMarginU, BottomMarginV ),		FVector2f(Tiling)),	Tint, SecondaryColor) ); //9
+		RenderBatch.AddVertex( FSlateVertex::Make<Rounding>( RenderTransform, FVector2f( RightMarginX, BottomMarginY ), FVector2f(LocalSize), DrawScale, FVector4f(FVector2f( RightMarginU, BottomMarginV ),	FVector2f(Tiling)),	Tint, SecondaryColor) ); //10
+		RenderBatch.AddVertex( FSlateVertex::Make<Rounding>( RenderTransform, FVector2f( EndPos.X, BottomMarginY ),		FVector2f(LocalSize), DrawScale, FVector4f(FVector2f( EndUV.X, BottomMarginV ),			FVector2f(Tiling)),	Tint, SecondaryColor) ); //11
+		RenderBatch.AddVertex( FSlateVertex::Make<Rounding>( RenderTransform, FVector2f( Position.X, EndPos.Y ),		FVector2f(LocalSize), DrawScale, FVector4f(FVector2f( StartUV.X, EndUV.Y ),				FVector2f(Tiling)),	Tint, SecondaryColor) ); //12
+		RenderBatch.AddVertex( FSlateVertex::Make<Rounding>( RenderTransform, FVector2f( LeftMarginX, EndPos.Y ),		FVector2f(LocalSize), DrawScale, FVector4f(FVector2f( LeftMarginU, EndUV.Y ),			FVector2f(Tiling)),	Tint, SecondaryColor) ); //13
+		RenderBatch.AddVertex( FSlateVertex::Make<Rounding>( RenderTransform, FVector2f( RightMarginX, EndPos.Y ),		FVector2f(LocalSize), DrawScale, FVector4f(FVector2f( RightMarginU, EndUV.Y ),			FVector2f(Tiling)),	Tint, SecondaryColor) ); //14
+		RenderBatch.AddVertex( FSlateVertex::Make<Rounding>( RenderTransform, FVector2f( EndPos.X, EndPos.Y ),			FVector2f(LocalSize), DrawScale, FVector4f(FVector2f( EndUV ),							FVector2f(Tiling)),	Tint, SecondaryColor) ); //15
 
 		// Top
 		RenderBatch.AddIndex( IndexStart + 0 );
@@ -814,24 +814,24 @@ void FSlateElementBatcher::AddBoxElement(const FSlateDrawElement& DrawElement)
 			const int32 FeatherStart = RenderBatch.GetNumVertices();
 
 			// Top
-			RenderBatch.AddVertex(FSlateVertex::Make<Rounding>(RenderTransform, FVector2D(Position.X, Position.Y) + FVector2D(-1, -1) / DrawScale, LocalSize, DrawScale, FVector4f(StartUV, Tiling), FeatherColor, SecondaryColor)); //0
-			RenderBatch.AddVertex(FSlateVertex::Make<Rounding>(RenderTransform, FVector2D(LeftMarginX, Position.Y) + FVector2D(0, -1) / DrawScale, LocalSize, DrawScale, FVector4f(FVector2D(LeftMarginU, StartUV.Y), Tiling), FeatherColor, SecondaryColor)); //1
-			RenderBatch.AddVertex(FSlateVertex::Make<Rounding>(RenderTransform, FVector2D(RightMarginX, Position.Y) + FVector2D(0, -1) / DrawScale, LocalSize, DrawScale, FVector4f(FVector2D(RightMarginU, StartUV.Y), Tiling), FeatherColor, SecondaryColor)); //2
-			RenderBatch.AddVertex(FSlateVertex::Make<Rounding>(RenderTransform, FVector2D(EndPos.X, Position.Y) + FVector2D(1, -1) / DrawScale, LocalSize, DrawScale, FVector4f(FVector2D(EndUV.X, StartUV.Y), Tiling), FeatherColor, SecondaryColor)); //3
+			RenderBatch.AddVertex(FSlateVertex::Make<Rounding>(RenderTransform, FVector2f(Position.X, Position.Y) + FVector2f(-1, -1) / DrawScale, FVector2f(LocalSize), DrawScale, FVector4f(FVector2f(StartUV), FVector2f(Tiling)), FeatherColor, SecondaryColor)); //0
+			RenderBatch.AddVertex(FSlateVertex::Make<Rounding>(RenderTransform, FVector2f(LeftMarginX, Position.Y) + FVector2f(0, -1) / DrawScale, FVector2f(LocalSize), DrawScale, FVector4f(FVector2f(LeftMarginU, StartUV.Y), FVector2f(Tiling)), FeatherColor, SecondaryColor)); //1
+			RenderBatch.AddVertex(FSlateVertex::Make<Rounding>(RenderTransform, FVector2f(RightMarginX, Position.Y) + FVector2f(0, -1) / DrawScale, FVector2f(LocalSize), DrawScale, FVector4f(FVector2f(RightMarginU, StartUV.Y), FVector2f(Tiling)), FeatherColor, SecondaryColor)); //2
+			RenderBatch.AddVertex(FSlateVertex::Make<Rounding>(RenderTransform, FVector2f(EndPos.X, Position.Y) + FVector2f(1, -1) / DrawScale, FVector2f(LocalSize), DrawScale, FVector4f(FVector2f(EndUV.X, StartUV.Y), FVector2f(Tiling)), FeatherColor, SecondaryColor)); //3
 
 			// Left
-			RenderBatch.AddVertex(FSlateVertex::Make<Rounding>(RenderTransform, FVector2D(Position.X, TopMarginY) + FVector2D(-1, 0) / DrawScale, LocalSize, DrawScale, FVector4f(FVector2D(StartUV.X, TopMarginV), Tiling), FeatherColor, SecondaryColor)); //4
-			RenderBatch.AddVertex(FSlateVertex::Make<Rounding>(RenderTransform, FVector2D(Position.X, BottomMarginY) + FVector2D(-1, 0) / DrawScale, LocalSize, DrawScale, FVector4f(FVector2D(StartUV.X, BottomMarginV), Tiling), FeatherColor, SecondaryColor)); //5
+			RenderBatch.AddVertex(FSlateVertex::Make<Rounding>(RenderTransform, FVector2f(Position.X, TopMarginY) + FVector2f(-1, 0) / DrawScale, FVector2f(LocalSize), DrawScale, FVector4f(FVector2f(StartUV.X, TopMarginV), FVector2f(Tiling)), FeatherColor, SecondaryColor)); //4
+			RenderBatch.AddVertex(FSlateVertex::Make<Rounding>(RenderTransform, FVector2f(Position.X, BottomMarginY) + FVector2f(-1, 0) / DrawScale, FVector2f(LocalSize), DrawScale, FVector4f(FVector2f(StartUV.X, BottomMarginV), FVector2f(Tiling)), FeatherColor, SecondaryColor)); //5
 
 			// Right
-			RenderBatch.AddVertex(FSlateVertex::Make<Rounding>(RenderTransform, FVector2D(EndPos.X, TopMarginY) + FVector2D(1, 0) / DrawScale, LocalSize, DrawScale, FVector4f(FVector2D(EndUV.X, TopMarginV), Tiling), FeatherColor, SecondaryColor)); //6
-			RenderBatch.AddVertex(FSlateVertex::Make<Rounding>(RenderTransform, FVector2D(EndPos.X, BottomMarginY) + FVector2D(1, 0) / DrawScale, LocalSize, DrawScale, FVector4f(FVector2D(EndUV.X, BottomMarginV), Tiling), FeatherColor, SecondaryColor)); //7
+			RenderBatch.AddVertex(FSlateVertex::Make<Rounding>(RenderTransform, FVector2f(EndPos.X, TopMarginY) + FVector2f(1, 0) / DrawScale, FVector2f(LocalSize), DrawScale, FVector4f(FVector2f(EndUV.X, TopMarginV), FVector2f(Tiling)), FeatherColor, SecondaryColor)); //6
+			RenderBatch.AddVertex(FSlateVertex::Make<Rounding>(RenderTransform, FVector2f(EndPos.X, BottomMarginY) + FVector2f(1, 0) / DrawScale, FVector2f(LocalSize), DrawScale, FVector4f(FVector2f(EndUV.X, BottomMarginV), FVector2f(Tiling)), FeatherColor, SecondaryColor)); //7
 
 			// Bottom
-			RenderBatch.AddVertex(FSlateVertex::Make<Rounding>(RenderTransform, FVector2D(Position.X, EndPos.Y) + FVector2D(-1, 1) / DrawScale, LocalSize, DrawScale, FVector4f(FVector2D(StartUV.X, EndUV.Y), Tiling), FeatherColor, SecondaryColor)); //8
-			RenderBatch.AddVertex(FSlateVertex::Make<Rounding>(RenderTransform, FVector2D(LeftMarginX, EndPos.Y) + FVector2D(0, 1) / DrawScale, LocalSize, DrawScale, FVector4f(FVector2D(LeftMarginU, EndUV.Y), Tiling), FeatherColor, SecondaryColor)); //9
-			RenderBatch.AddVertex(FSlateVertex::Make<Rounding>(RenderTransform, FVector2D(RightMarginX, EndPos.Y) + FVector2D(0, 1) / DrawScale, LocalSize, DrawScale, FVector4f(FVector2D(RightMarginU, EndUV.Y), Tiling), FeatherColor, SecondaryColor)); //10
-			RenderBatch.AddVertex(FSlateVertex::Make<Rounding>(RenderTransform, FVector2D(EndPos.X, EndPos.Y) + FVector2D(1, 1) / DrawScale, LocalSize, DrawScale, FVector4f(EndUV, Tiling), FeatherColor, SecondaryColor)); //11
+			RenderBatch.AddVertex(FSlateVertex::Make<Rounding>(RenderTransform, FVector2f(Position.X, EndPos.Y) + FVector2f(-1, 1) / DrawScale, FVector2f(LocalSize), DrawScale, FVector4f(FVector2f(StartUV.X, EndUV.Y), FVector2f(Tiling)), FeatherColor, SecondaryColor)); //8
+			RenderBatch.AddVertex(FSlateVertex::Make<Rounding>(RenderTransform, FVector2f(LeftMarginX, EndPos.Y) + FVector2f(0, 1) / DrawScale, FVector2f(LocalSize), DrawScale, FVector4f(FVector2f(LeftMarginU, EndUV.Y), FVector2f(Tiling)), FeatherColor, SecondaryColor)); //9
+			RenderBatch.AddVertex(FSlateVertex::Make<Rounding>(RenderTransform, FVector2f(RightMarginX, EndPos.Y) + FVector2f(0, 1) / DrawScale, FVector2f(LocalSize), DrawScale, FVector4f(FVector2f(RightMarginU, EndUV.Y), FVector2f(Tiling)), FeatherColor, SecondaryColor)); //10
+			RenderBatch.AddVertex(FSlateVertex::Make<Rounding>(RenderTransform, FVector2f(EndPos.X, EndPos.Y) + FVector2f(1, 1) / DrawScale, FVector2f(LocalSize), DrawScale, FVector4f(FVector2f(EndUV), FVector2f(Tiling)), FeatherColor, SecondaryColor)); //11
 
 			// Top Left
 			IndexQuad(RenderBatch, FeatherStart + 0, FeatherStart + 1, IndexStart + 2, IndexStart + 0);
@@ -888,10 +888,10 @@ void FSlateElementBatcher::AddBoxElement(const FSlateDrawElement& DrawElement)
 		}
 
 		// Add four vertices to the list of verts to be added to the vertex buffer
-		RenderBatch.AddVertex(FSlateVertex::Make<Rounding>(RenderTransform, TopLeft, LocalSize, DrawScale, FVector4f(StartUV, Tiling), Tint, SecondaryColor));
-		RenderBatch.AddVertex(FSlateVertex::Make<Rounding>(RenderTransform, TopRight, LocalSize, DrawScale, FVector4f(FVector2D(EndUV.X, StartUV.Y), Tiling), Tint, SecondaryColor));
-		RenderBatch.AddVertex(FSlateVertex::Make<Rounding>(RenderTransform, BotLeft, LocalSize, DrawScale, FVector4f(FVector2D(StartUV.X, EndUV.Y), Tiling), Tint, SecondaryColor));
-		RenderBatch.AddVertex(FSlateVertex::Make<Rounding>(RenderTransform, BotRight, LocalSize, DrawScale, FVector4f(EndUV, Tiling), Tint, SecondaryColor));
+		RenderBatch.AddVertex(FSlateVertex::Make<Rounding>(RenderTransform, FVector2f(TopLeft), FVector2f(LocalSize), DrawScale, FVector4f(FVector2f(StartUV), FVector2f(Tiling)), Tint, SecondaryColor));
+		RenderBatch.AddVertex(FSlateVertex::Make<Rounding>(RenderTransform, FVector2f(TopRight), FVector2f(LocalSize), DrawScale, FVector4f(FVector2f(EndUV.X, StartUV.Y), FVector2f(Tiling)), Tint, SecondaryColor));
+		RenderBatch.AddVertex(FSlateVertex::Make<Rounding>(RenderTransform, FVector2f(BotLeft), FVector2f(LocalSize), DrawScale, FVector4f(FVector2f(StartUV.X, EndUV.Y), FVector2f(Tiling)), Tint, SecondaryColor));
+		RenderBatch.AddVertex(FSlateVertex::Make<Rounding>(RenderTransform, FVector2f(BotRight), FVector2f(LocalSize), DrawScale, FVector4f(FVector2f(EndUV), FVector2f(Tiling)), Tint, SecondaryColor));
 
 		RenderBatch.AddIndex( IndexStart + 0 );
 		RenderBatch.AddIndex( IndexStart + 1 );
@@ -910,10 +910,10 @@ void FSlateElementBatcher::AddBoxElement(const FSlateDrawElement& DrawElement)
 		{
 			const int32 FeatherStart = RenderBatch.GetNumVertices();
 
-			RenderBatch.AddVertex(FSlateVertex::Make<Rounding>(RenderTransform, TopLeft + FVector2D(-1, -1) / DrawScale, LocalSize, DrawScale, FVector4f(StartUV, Tiling), FeatherColor, SecondaryColor));
-			RenderBatch.AddVertex(FSlateVertex::Make<Rounding>(RenderTransform, TopRight + FVector2D(1, -1) / DrawScale, LocalSize, DrawScale, FVector4f(FVector2D(EndUV.X, StartUV.Y), Tiling), FeatherColor, SecondaryColor));
-			RenderBatch.AddVertex(FSlateVertex::Make<Rounding>(RenderTransform, BotLeft + FVector2D(-1, 1) / DrawScale, LocalSize, DrawScale, FVector4f(FVector2D(StartUV.X, EndUV.Y), Tiling), FeatherColor, SecondaryColor));
-			RenderBatch.AddVertex(FSlateVertex::Make<Rounding>(RenderTransform, BotRight + FVector2D(1, 1) / DrawScale, LocalSize, DrawScale, FVector4f(EndUV, Tiling), FeatherColor, SecondaryColor));
+			RenderBatch.AddVertex(FSlateVertex::Make<Rounding>(RenderTransform, FVector2f(TopLeft) + FVector2f(-1, -1) / DrawScale, FVector2f(LocalSize), DrawScale, FVector4f(FVector2f(StartUV), FVector2f(Tiling)), FeatherColor, SecondaryColor));
+			RenderBatch.AddVertex(FSlateVertex::Make<Rounding>(RenderTransform, FVector2f(TopRight) + FVector2f(1, -1) / DrawScale, FVector2f(LocalSize), DrawScale, FVector4f(FVector2f(EndUV.X, StartUV.Y), FVector2f(Tiling)), FeatherColor, SecondaryColor));
+			RenderBatch.AddVertex(FSlateVertex::Make<Rounding>(RenderTransform, FVector2f(BotLeft) + FVector2f(-1, 1) / DrawScale, FVector2f(LocalSize), DrawScale, FVector4f(FVector2f(StartUV.X, EndUV.Y), FVector2f(Tiling)), FeatherColor, SecondaryColor));
+			RenderBatch.AddVertex(FSlateVertex::Make<Rounding>(RenderTransform, FVector2f(BotRight) + FVector2f(1, 1) / DrawScale, FVector2f(LocalSize), DrawScale, FVector4f(FVector2f(EndUV), FVector2f(Tiling)), FeatherColor, SecondaryColor));
 
 			// Top-Top
 			RenderBatch.AddIndex(FeatherStart + 0);
@@ -1171,10 +1171,10 @@ void FSlateElementBatcher::AddTextElement(const FSlateDrawElement& DrawElement)
 						}
 
 						// Add four vertices to the list of verts to be added to the vertex buffer
-						RenderBatch->AddVertex(FSlateVertex::Make<Rounding>( RenderTransform, UpperLeft,								FVector4f(U,V,Ut,Vt),						FVector2D(0.0f,0.0f), FontTint ));
-						RenderBatch->AddVertex(FSlateVertex::Make<Rounding>( RenderTransform, FVector2D(LowerRight.X,UpperLeft.Y),	FVector4f(U+SizeU, V, UtMax,Vt),				FVector2D(1.0f,0.0f), FontTint ));
-						RenderBatch->AddVertex(FSlateVertex::Make<Rounding>( RenderTransform, FVector2D(UpperLeft.X,LowerRight.Y),	FVector4f(U, V+SizeV, Ut,VtMax),				FVector2D(0.0f,1.0f), FontTint ));
-						RenderBatch->AddVertex(FSlateVertex::Make<Rounding>( RenderTransform, LowerRight,							FVector4f(U+SizeU, V+SizeV, UtMax,VtMax),	FVector2D(1.0f,1.0f), FontTint ));
+						RenderBatch->AddVertex(FSlateVertex::Make<Rounding>( RenderTransform, FVector2f(UpperLeft),					FVector4f(U,V,Ut,Vt),						FVector2f(0.0f,0.0f), FontTint ));
+						RenderBatch->AddVertex(FSlateVertex::Make<Rounding>( RenderTransform, FVector2f(LowerRight.X,UpperLeft.Y),	FVector4f(U+SizeU, V, UtMax,Vt),			FVector2f(1.0f,0.0f), FontTint ));
+						RenderBatch->AddVertex(FSlateVertex::Make<Rounding>( RenderTransform, FVector2f(UpperLeft.X,LowerRight.Y),	FVector4f(U, V+SizeV, Ut,VtMax),			FVector2f(0.0f,1.0f), FontTint ));
+						RenderBatch->AddVertex(FSlateVertex::Make<Rounding>( RenderTransform, FVector2f(LowerRight),				FVector4f(U+SizeU, V+SizeV, UtMax,VtMax),	FVector2f(1.0f,1.0f), FontTint ));
 
 						RenderBatch->AddIndex(IndexStart + 0);
 						RenderBatch->AddIndex(IndexStart + 1);
@@ -1502,14 +1502,14 @@ void FSlateElementBatcher::AddGradientElement( const FSlateDrawElement& DrawElem
 		if( StopIndex == 0 )
 		{
 			// First stop does not have a full quad yet so do not create indices
-			RenderBatch.AddVertex(FSlateVertex::Make<Rounding>(RenderTransform, StartPt, LocalSize, DrawScale, FVector4f(StartUV.X, StartUV.Y, 0,0), PackVertexColor(CurStop.Color), FColor::Transparent));
-			RenderBatch.AddVertex(FSlateVertex::Make<Rounding>(RenderTransform, EndPt, LocalSize, DrawScale, FVector4f(EndUV.X, EndUV.Y, 0, 0), PackVertexColor(CurStop.Color), FColor::Transparent));
+			RenderBatch.AddVertex(FSlateVertex::Make<Rounding>(RenderTransform, FVector2f(StartPt), FVector2f(LocalSize), DrawScale, FVector4f(StartUV.X, StartUV.Y, 0,0), PackVertexColor(CurStop.Color), FColor::Transparent));
+			RenderBatch.AddVertex(FSlateVertex::Make<Rounding>(RenderTransform, FVector2f(EndPt), FVector2f(LocalSize), DrawScale, FVector4f(EndUV.X, EndUV.Y, 0, 0), PackVertexColor(CurStop.Color), FColor::Transparent));
 		}
 		else
 		{
 			// All stops after the first have indices and generate quads
-			RenderBatch.AddVertex(FSlateVertex::Make<Rounding>(RenderTransform, StartPt, LocalSize, DrawScale, FVector4f(StartUV.X, StartUV.Y, 0, 0), PackVertexColor(CurStop.Color), FColor::Transparent));
-			RenderBatch.AddVertex(FSlateVertex::Make<Rounding>(RenderTransform, EndPt, LocalSize, DrawScale, FVector4f(EndUV.X, EndUV.Y, 0, 0), PackVertexColor(CurStop.Color), FColor::Transparent));
+			RenderBatch.AddVertex(FSlateVertex::Make<Rounding>(RenderTransform, FVector2f(StartPt), FVector2f(LocalSize), DrawScale, FVector4f(StartUV.X, StartUV.Y, 0, 0), PackVertexColor(CurStop.Color), FColor::Transparent));
+			RenderBatch.AddVertex(FSlateVertex::Make<Rounding>(RenderTransform, FVector2f(EndPt), FVector2f(LocalSize), DrawScale, FVector4f(EndUV.X, EndUV.Y, 0, 0), PackVertexColor(CurStop.Color), FColor::Transparent));
 
 			// Connect the indices to the previous vertices
 			RenderBatch.AddIndex(IndexStart - 2);
@@ -2142,10 +2142,10 @@ void FSlateElementBatcher::AddLineElement( const FSlateDrawElement& DrawElement
 				const FVector2D SegmentDir = (EndPos - StartPos).GetSafeNormal();
 				const FVector2D Up = GetRotated90(SegmentDir) * LocalHalfThickness;
 
-				RenderBatch.AddVertex(FSlateVertex::Make<Rounding>(RenderTransform, StartPos + Up, FVector2D::ZeroVector, StartColor));
-				RenderBatch.AddVertex(FSlateVertex::Make<Rounding>(RenderTransform, StartPos - Up, FVector2D::ZeroVector, StartColor));
-				RenderBatch.AddVertex(FSlateVertex::Make<Rounding>(RenderTransform, EndPos + Up, FVector2D::ZeroVector, EndColor));
-				RenderBatch.AddVertex(FSlateVertex::Make<Rounding>(RenderTransform, EndPos - Up, FVector2D::ZeroVector, EndColor));
+				RenderBatch.AddVertex(FSlateVertex::Make<Rounding>(RenderTransform, FVector2f(StartPos + Up), FVector2f::ZeroVector, StartColor));
+				RenderBatch.AddVertex(FSlateVertex::Make<Rounding>(RenderTransform, FVector2f(StartPos - Up), FVector2f::ZeroVector, StartColor));
+				RenderBatch.AddVertex(FSlateVertex::Make<Rounding>(RenderTransform, FVector2f(EndPos + Up), FVector2f::ZeroVector, EndColor));
+				RenderBatch.AddVertex(FSlateVertex::Make<Rounding>(RenderTransform, FVector2f(EndPos - Up), FVector2f::ZeroVector, EndColor));
 
 				FLineBuilder::AddQuadIndices(RenderBatch);
 			}
@@ -2201,10 +2201,10 @@ void FSlateElementBatcher::AddViewportElement( const FSlateDrawElement& DrawElem
 	const uint32 IndexStart = 0;
 
 	// Add four vertices to the list of verts to be added to the vertex buffer
-	RenderBatch.AddVertex( FSlateVertex::Make<Rounding>( RenderTransform, TopLeft,	FVector2D(0.0f,0.0f),	FinalColor ) );
-	RenderBatch.AddVertex( FSlateVertex::Make<Rounding>( RenderTransform, TopRight,	FVector2D(1.0f,0.0f),	FinalColor ) );
-	RenderBatch.AddVertex( FSlateVertex::Make<Rounding>( RenderTransform, BotLeft,	FVector2D(0.0f,1.0f),	FinalColor ) );
-	RenderBatch.AddVertex( FSlateVertex::Make<Rounding>( RenderTransform, BotRight,	FVector2D(1.0f,1.0f),	FinalColor ) );
+	RenderBatch.AddVertex( FSlateVertex::Make<Rounding>( RenderTransform, FVector2f(TopLeft),	FVector2f(0.0f,0.0f),	FinalColor ) );
+	RenderBatch.AddVertex( FSlateVertex::Make<Rounding>( RenderTransform, FVector2f(TopRight),	FVector2f(1.0f,0.0f),	FinalColor ) );
+	RenderBatch.AddVertex( FSlateVertex::Make<Rounding>( RenderTransform, FVector2f(BotLeft),	FVector2f(0.0f,1.0f),	FinalColor ) );
+	RenderBatch.AddVertex( FSlateVertex::Make<Rounding>( RenderTransform, FVector2f(BotRight),	FVector2f(1.0f,1.0f),	FinalColor ) );
 
 	// Add 6 indices to the vertex buffer.  (2 tri's per quad, 3 indices per tri)
 	RenderBatch.AddIndex( IndexStart + 0 );
@@ -2335,45 +2335,45 @@ void FSlateElementBatcher::AddBorderElement( const FSlateDrawElement& DrawElemen
 	FVector2D Zero(0,0);
 
 	// Add all the vertices needed for this element.  Vertices are duplicated so that we can have some sections with no tiling and some with tiling.
-	RenderBatch.AddVertex( FSlateVertex::Make<Rounding>( RenderTransform, Position,									LocalSize, DrawScale, FVector4f( StartUV.X, StartUV.Y, 0.0f, 0.0f),				Tint ) ); //0
-	RenderBatch.AddVertex( FSlateVertex::Make<Rounding>( RenderTransform, FVector2D( Position.X, TopMarginY ),		LocalSize, DrawScale, FVector4f( StartUV.X, TopMarginV, 0.0f, 0.0f),				Tint ) ); //1
-	RenderBatch.AddVertex( FSlateVertex::Make<Rounding>( RenderTransform, FVector2D( LeftMarginX, Position.Y ),		LocalSize, DrawScale, FVector4f( LeftMarginU, StartUV.Y, 0.0f, 0.0f),			Tint ) ); //2
-	RenderBatch.AddVertex( FSlateVertex::Make<Rounding>( RenderTransform, FVector2D( LeftMarginX, TopMarginY ),		LocalSize, DrawScale, FVector4f( LeftMarginU, TopMarginV, 0.0f, 0.0f),			Tint ) ); //3
-
-	RenderBatch.AddVertex( FSlateVertex::Make<Rounding>( RenderTransform, FVector2D( LeftMarginX, Position.Y ),		LocalSize, DrawScale, FVector4f( StartUV.X, StartUV.Y, TopTiling, 0.0f),			Tint ) ); //4
-	RenderBatch.AddVertex( FSlateVertex::Make<Rounding>( RenderTransform, FVector2D( LeftMarginX, TopMarginY ),		LocalSize, DrawScale, FVector4f( StartUV.X, TopMarginV, TopTiling, 0.0f),		Tint ) ); //5
-	RenderBatch.AddVertex( FSlateVertex::Make<Rounding>( RenderTransform, FVector2D( RightMarginX, Position.Y ),	LocalSize, DrawScale, FVector4f( EndUV.X, StartUV.Y, TopTiling, 0.0f),			Tint ) ); //6
-	RenderBatch.AddVertex( FSlateVertex::Make<Rounding>( RenderTransform, FVector2D( RightMarginX, TopMarginY ),	LocalSize, DrawScale, FVector4f( EndUV.X, TopMarginV, TopTiling, 0.0f),			Tint ) ); //7
-
-	RenderBatch.AddVertex( FSlateVertex::Make<Rounding>( RenderTransform, FVector2D( RightMarginX, Position.Y ),	LocalSize, DrawScale, FVector4f( RightMarginU, StartUV.Y, 0.0f, 0.0f),			Tint ) ); //8
-	RenderBatch.AddVertex( FSlateVertex::Make<Rounding>( RenderTransform, FVector2D( RightMarginX, TopMarginY ),	LocalSize, DrawScale, FVector4f( RightMarginU, TopMarginV, 0.0f, 0.0f),			Tint ) ); //9
-	RenderBatch.AddVertex( FSlateVertex::Make<Rounding>( RenderTransform, FVector2D( EndPos.X, Position.Y ),		LocalSize, DrawScale, FVector4f( EndUV.X, StartUV.Y, 0.0f, 0.0f),				Tint ) ); //10
-	RenderBatch.AddVertex( FSlateVertex::Make<Rounding>( RenderTransform, FVector2D( EndPos.X, TopMarginY ),		LocalSize, DrawScale, FVector4f( EndUV.X, TopMarginV, 0.0f, 0.0f),				Tint ) ); //11
-
-	RenderBatch.AddVertex( FSlateVertex::Make<Rounding>( RenderTransform, FVector2D( Position.X, TopMarginY ),		LocalSize, DrawScale, FVector4f( StartUV.X, StartUV.Y, 0.0f, LeftTiling),		Tint ) ); //12
-	RenderBatch.AddVertex( FSlateVertex::Make<Rounding>( RenderTransform, FVector2D( Position.X, BottomMarginY ),	LocalSize, DrawScale, FVector4f( StartUV.X, EndUV.Y, 0.0f, LeftTiling),			Tint ) ); //13
-	RenderBatch.AddVertex( FSlateVertex::Make<Rounding>( RenderTransform, FVector2D( LeftMarginX, TopMarginY ),		LocalSize, DrawScale, FVector4f( LeftMarginU, StartUV.Y, 0.0f, LeftTiling),		Tint ) ); //14
-	RenderBatch.AddVertex( FSlateVertex::Make<Rounding>( RenderTransform, FVector2D( LeftMarginX, BottomMarginY ),	LocalSize, DrawScale, FVector4f( LeftMarginU, EndUV.Y, 0.0f, LeftTiling),		Tint ) ); //15
-
-	RenderBatch.AddVertex( FSlateVertex::Make<Rounding>( RenderTransform, FVector2D( RightMarginX, TopMarginY ),	LocalSize, DrawScale, FVector4f( RightMarginU, StartUV.Y, 0.0f, RightTiling),	Tint ) ); //16
-	RenderBatch.AddVertex( FSlateVertex::Make<Rounding>( RenderTransform, FVector2D( RightMarginX, BottomMarginY ), LocalSize, DrawScale, FVector4f( RightMarginU, EndUV.Y, 0.0f, RightTiling),		Tint ) ); //17
-	RenderBatch.AddVertex( FSlateVertex::Make<Rounding>( RenderTransform, FVector2D( EndPos.X, TopMarginY ),		LocalSize, DrawScale, FVector4f( EndUV.X, StartUV.Y, 0.0f, RightTiling),			Tint ) ); //18
-	RenderBatch.AddVertex( FSlateVertex::Make<Rounding>( RenderTransform, FVector2D( EndPos.X, BottomMarginY ),		LocalSize, DrawScale, FVector4f( EndUV.X, EndUV.Y, 0.0f, RightTiling),			Tint ) ); //19
-
-	RenderBatch.AddVertex( FSlateVertex::Make<Rounding>( RenderTransform, FVector2D( Position.X, BottomMarginY ),	LocalSize, DrawScale, FVector4f( StartUV.X, BottomMarginV, 0.0f, 0.0f),			Tint ) ); //20
-	RenderBatch.AddVertex( FSlateVertex::Make<Rounding>( RenderTransform, FVector2D( Position.X, EndPos.Y ),		LocalSize, DrawScale, FVector4f( StartUV.X, EndUV.Y, 0.0f, 0.0f),				Tint ) ); //21
-	RenderBatch.AddVertex( FSlateVertex::Make<Rounding>( RenderTransform, FVector2D( LeftMarginX, BottomMarginY ),	LocalSize, DrawScale, FVector4f( LeftMarginU, BottomMarginV, 0.0f, 0.0f),		Tint ) ); //22
-	RenderBatch.AddVertex( FSlateVertex::Make<Rounding>( RenderTransform, FVector2D( LeftMarginX, EndPos.Y ),		LocalSize, DrawScale, FVector4f( LeftMarginU, EndUV.Y, 0.0f, 0.0f),				Tint ) ); //23
-
-	RenderBatch.AddVertex( FSlateVertex::Make<Rounding>( RenderTransform, FVector2D( LeftMarginX, BottomMarginY ),	LocalSize, DrawScale, FVector4f( StartUV.X, BottomMarginV, BottomTiling, 0.0f),	Tint ) ); //24
-	RenderBatch.AddVertex( FSlateVertex::Make<Rounding>( RenderTransform, FVector2D( LeftMarginX, EndPos.Y ),		LocalSize, DrawScale, FVector4f( StartUV.X, EndUV.Y, BottomTiling, 0.0f),		Tint ) ); //25
-	RenderBatch.AddVertex( FSlateVertex::Make<Rounding>( RenderTransform, FVector2D( RightMarginX,BottomMarginY ),	LocalSize, DrawScale, FVector4f( EndUV.X, BottomMarginV, BottomTiling, 0.0f),	Tint ) ); //26
-	RenderBatch.AddVertex( FSlateVertex::Make<Rounding>( RenderTransform, FVector2D( RightMarginX, EndPos.Y ),		LocalSize, DrawScale, FVector4f( EndUV.X, EndUV.Y, BottomTiling, 0.0f),			Tint ) ); //27
-
-	RenderBatch.AddVertex( FSlateVertex::Make<Rounding>( RenderTransform, FVector2D( RightMarginX, BottomMarginY ), LocalSize, DrawScale, FVector4f( RightMarginU, BottomMarginV, 0.0f, 0.0f),		Tint ) ); //29
-	RenderBatch.AddVertex( FSlateVertex::Make<Rounding>( RenderTransform, FVector2D( RightMarginX, EndPos.Y ),		LocalSize, DrawScale, FVector4f( RightMarginU, EndUV.Y, 0.0f, 0.0f),				Tint ) ); //30
-	RenderBatch.AddVertex( FSlateVertex::Make<Rounding>( RenderTransform, FVector2D( EndPos.X, BottomMarginY ),		LocalSize, DrawScale, FVector4f( EndUV.X, BottomMarginV, 0.0f, 0.0f),			Tint ) ); //31
-	RenderBatch.AddVertex( FSlateVertex::Make<Rounding>( RenderTransform, FVector2D( EndPos.X, EndPos.Y ),			LocalSize, DrawScale, FVector4f( EndUV.X, EndUV.Y, 0.0f, 0.0f),					Tint ) ); //32
+	RenderBatch.AddVertex( FSlateVertex::Make<Rounding>( RenderTransform, FVector2f( Position ),					FVector2f(LocalSize), DrawScale, FVector4f( StartUV.X, StartUV.Y, 0.0f, 0.0f),				Tint ) ); //0
+	RenderBatch.AddVertex( FSlateVertex::Make<Rounding>( RenderTransform, FVector2f( Position.X, TopMarginY ),		FVector2f(LocalSize), DrawScale, FVector4f( StartUV.X, TopMarginV, 0.0f, 0.0f),				Tint ) ); //1
+	RenderBatch.AddVertex( FSlateVertex::Make<Rounding>( RenderTransform, FVector2f( LeftMarginX, Position.Y ),		FVector2f(LocalSize), DrawScale, FVector4f( LeftMarginU, StartUV.Y, 0.0f, 0.0f),			Tint ) ); //2
+	RenderBatch.AddVertex( FSlateVertex::Make<Rounding>( RenderTransform, FVector2f( LeftMarginX, TopMarginY ),		FVector2f(LocalSize), DrawScale, FVector4f( LeftMarginU, TopMarginV, 0.0f, 0.0f),			Tint ) ); //3
+
+	RenderBatch.AddVertex( FSlateVertex::Make<Rounding>( RenderTransform, FVector2f( LeftMarginX, Position.Y ),		FVector2f(LocalSize), DrawScale, FVector4f( StartUV.X, StartUV.Y, TopTiling, 0.0f),			Tint ) ); //4
+	RenderBatch.AddVertex( FSlateVertex::Make<Rounding>( RenderTransform, FVector2f( LeftMarginX, TopMarginY ),		FVector2f(LocalSize), DrawScale, FVector4f( StartUV.X, TopMarginV, TopTiling, 0.0f),		Tint ) ); //5
+	RenderBatch.AddVertex( FSlateVertex::Make<Rounding>( RenderTransform, FVector2f( RightMarginX, Position.Y ),	FVector2f(LocalSize), DrawScale, FVector4f( EndUV.X, StartUV.Y, TopTiling, 0.0f),			Tint ) ); //6
+	RenderBatch.AddVertex( FSlateVertex::Make<Rounding>( RenderTransform, FVector2f( RightMarginX, TopMarginY ),	FVector2f(LocalSize), DrawScale, FVector4f( EndUV.X, TopMarginV, TopTiling, 0.0f),			Tint ) ); //7
+
+	RenderBatch.AddVertex( FSlateVertex::Make<Rounding>( RenderTransform, FVector2f( RightMarginX, Position.Y ),	FVector2f(LocalSize), DrawScale, FVector4f( RightMarginU, StartUV.Y, 0.0f, 0.0f),			Tint ) ); //8
+	RenderBatch.AddVertex( FSlateVertex::Make<Rounding>( RenderTransform, FVector2f( RightMarginX, TopMarginY ),	FVector2f(LocalSize), DrawScale, FVector4f( RightMarginU, TopMarginV, 0.0f, 0.0f),			Tint ) ); //9
+	RenderBatch.AddVertex( FSlateVertex::Make<Rounding>( RenderTransform, FVector2f( EndPos.X, Position.Y ),		FVector2f(LocalSize), DrawScale, FVector4f( EndUV.X, StartUV.Y, 0.0f, 0.0f),				Tint ) ); //10
+	RenderBatch.AddVertex( FSlateVertex::Make<Rounding>( RenderTransform, FVector2f( EndPos.X, TopMarginY ),		FVector2f(LocalSize), DrawScale, FVector4f( EndUV.X, TopMarginV, 0.0f, 0.0f),				Tint ) ); //11
+
+	RenderBatch.AddVertex( FSlateVertex::Make<Rounding>( RenderTransform, FVector2f( Position.X, TopMarginY ),		FVector2f(LocalSize), DrawScale, FVector4f( StartUV.X, StartUV.Y, 0.0f, LeftTiling),		Tint ) ); //12
+	RenderBatch.AddVertex( FSlateVertex::Make<Rounding>( RenderTransform, FVector2f( Position.X, BottomMarginY ),	FVector2f(LocalSize), DrawScale, FVector4f( StartUV.X, EndUV.Y, 0.0f, LeftTiling),			Tint ) ); //13
+	RenderBatch.AddVertex( FSlateVertex::Make<Rounding>( RenderTransform, FVector2f( LeftMarginX, TopMarginY ),		FVector2f(LocalSize), DrawScale, FVector4f( LeftMarginU, StartUV.Y, 0.0f, LeftTiling),		Tint ) ); //14
+	RenderBatch.AddVertex( FSlateVertex::Make<Rounding>( RenderTransform, FVector2f( LeftMarginX, BottomMarginY ),	FVector2f(LocalSize), DrawScale, FVector4f( LeftMarginU, EndUV.Y, 0.0f, LeftTiling),		Tint ) ); //15
+
+	RenderBatch.AddVertex( FSlateVertex::Make<Rounding>( RenderTransform, FVector2f( RightMarginX, TopMarginY ),	FVector2f(LocalSize), DrawScale, FVector4f( RightMarginU, StartUV.Y, 0.0f, RightTiling),	Tint ) ); //16
+	RenderBatch.AddVertex( FSlateVertex::Make<Rounding>( RenderTransform, FVector2f( RightMarginX, BottomMarginY ), FVector2f(LocalSize), DrawScale, FVector4f( RightMarginU, EndUV.Y, 0.0f, RightTiling),		Tint ) ); //17
+	RenderBatch.AddVertex( FSlateVertex::Make<Rounding>( RenderTransform, FVector2f( EndPos.X, TopMarginY ),		FVector2f(LocalSize), DrawScale, FVector4f( EndUV.X, StartUV.Y, 0.0f, RightTiling),			Tint ) ); //18
+	RenderBatch.AddVertex( FSlateVertex::Make<Rounding>( RenderTransform, FVector2f( EndPos.X, BottomMarginY ),		FVector2f(LocalSize), DrawScale, FVector4f( EndUV.X, EndUV.Y, 0.0f, RightTiling),			Tint ) ); //19
+
+	RenderBatch.AddVertex( FSlateVertex::Make<Rounding>( RenderTransform, FVector2f( Position.X, BottomMarginY ),	FVector2f(LocalSize), DrawScale, FVector4f( StartUV.X, BottomMarginV, 0.0f, 0.0f),			Tint ) ); //20
+	RenderBatch.AddVertex( FSlateVertex::Make<Rounding>( RenderTransform, FVector2f( Position.X, EndPos.Y ),		FVector2f(LocalSize), DrawScale, FVector4f( StartUV.X, EndUV.Y, 0.0f, 0.0f),				Tint ) ); //21
+	RenderBatch.AddVertex( FSlateVertex::Make<Rounding>( RenderTransform, FVector2f( LeftMarginX, BottomMarginY ),	FVector2f(LocalSize), DrawScale, FVector4f( LeftMarginU, BottomMarginV, 0.0f, 0.0f),		Tint ) ); //22
+	RenderBatch.AddVertex( FSlateVertex::Make<Rounding>( RenderTransform, FVector2f( LeftMarginX, EndPos.Y ),		FVector2f(LocalSize), DrawScale, FVector4f( LeftMarginU, EndUV.Y, 0.0f, 0.0f),				Tint ) ); //23
+
+	RenderBatch.AddVertex( FSlateVertex::Make<Rounding>( RenderTransform, FVector2f( LeftMarginX, BottomMarginY ),	FVector2f(LocalSize), DrawScale, FVector4f( StartUV.X, BottomMarginV, BottomTiling, 0.0f),	Tint ) ); //24
+	RenderBatch.AddVertex( FSlateVertex::Make<Rounding>( RenderTransform, FVector2f( LeftMarginX, EndPos.Y ),		FVector2f(LocalSize), DrawScale, FVector4f( StartUV.X, EndUV.Y, BottomTiling, 0.0f),		Tint ) ); //25
+	RenderBatch.AddVertex( FSlateVertex::Make<Rounding>( RenderTransform, FVector2f( RightMarginX,BottomMarginY ),	FVector2f(LocalSize), DrawScale, FVector4f( EndUV.X, BottomMarginV, BottomTiling, 0.0f),	Tint ) ); //26
+	RenderBatch.AddVertex( FSlateVertex::Make<Rounding>( RenderTransform, FVector2f( RightMarginX, EndPos.Y ),		FVector2f(LocalSize), DrawScale, FVector4f( EndUV.X, EndUV.Y, BottomTiling, 0.0f),			Tint ) ); //27
+
+	RenderBatch.AddVertex( FSlateVertex::Make<Rounding>( RenderTransform, FVector2f( RightMarginX, BottomMarginY ), FVector2f(LocalSize), DrawScale, FVector4f( RightMarginU, BottomMarginV, 0.0f, 0.0f),		Tint ) ); //29
+	RenderBatch.AddVertex( FSlateVertex::Make<Rounding>( RenderTransform, FVector2f( RightMarginX, EndPos.Y ),		FVector2f(LocalSize), DrawScale, FVector4f( RightMarginU, EndUV.Y, 0.0f, 0.0f),				Tint ) ); //30
+	RenderBatch.AddVertex( FSlateVertex::Make<Rounding>( RenderTransform, FVector2f( EndPos.X, BottomMarginY ),		FVector2f(LocalSize), DrawScale, FVector4f( EndUV.X, BottomMarginV, 0.0f, 0.0f),			Tint ) ); //31
+	RenderBatch.AddVertex( FSlateVertex::Make<Rounding>( RenderTransform, FVector2f( EndPos.X, EndPos.Y ),			FVector2f(LocalSize), DrawScale, FVector4f( EndUV.X, EndUV.Y, 0.0f, 0.0f),					Tint ) ); //32
 
 	// Top
 	RenderBatch.AddIndex( IndexStart + 0 );
@@ -2550,7 +2550,7 @@ void FSlateElementBatcher::AddPostProcessPass(const FSlateDrawElement& DrawEleme
 	if(SizeUV.X > 0 && SizeUV.Y > 0)
 	{
 		FShaderParams Params = FShaderParams::MakePixelShaderParams(
-			FVector4f(WorldTopLeft, WorldBotRight),
+			FVector4f(FVector2f(WorldTopLeft), FVector2f(WorldBotRight)),
 			FVector4f(Payload.PostProcessData.X, Payload.PostProcessData.Y, (float)Payload.DownsampleAmount, 0.f),
 			FVector4f(Payload.CornerRadius));
 
@@ -2823,10 +2823,10 @@ void FSlateElementBatcher::BuildShapedTextSequence(const FShapedTextBuildContext
 			}
 
 			// Add four vertices to the list of verts to be added to the vertex buffer
-			RenderBatch->AddVertex(FSlateVertex::Make<Rounding>(RenderTransform, UpperLeft, FVector4f(U, V, Ut, Vt), FVector2D(0.0f, 0.0f), Tint));
-			RenderBatch->AddVertex(FSlateVertex::Make<Rounding>(RenderTransform, FVector2D(LowerRight.X, UpperLeft.Y), FVector4f(U + SizeU, V, UtMax, Vt), FVector2D(1.0f, 0.0f), Tint));
-			RenderBatch->AddVertex(FSlateVertex::Make<Rounding>(RenderTransform, FVector2D(UpperLeft.X, LowerRight.Y), FVector4f(U, V + SizeV, Ut, VtMax), FVector2D(0.0f, 1.0f), Tint));
-			RenderBatch->AddVertex(FSlateVertex::Make<Rounding>(RenderTransform, LowerRight, FVector4f(U + SizeU, V + SizeV, UtMax, VtMax), FVector2D(1.0f, 1.0f), Tint));
+			RenderBatch->AddVertex(FSlateVertex::Make<Rounding>(RenderTransform, FVector2f(UpperLeft), FVector4f(U, V, Ut, Vt), FVector2f(0.0f, 0.0f), Tint));
+			RenderBatch->AddVertex(FSlateVertex::Make<Rounding>(RenderTransform, FVector2f(LowerRight.X, UpperLeft.Y), FVector4f(U + SizeU, V, UtMax, Vt), FVector2f(1.0f, 0.0f), Tint));
+			RenderBatch->AddVertex(FSlateVertex::Make<Rounding>(RenderTransform, FVector2f(UpperLeft.X, LowerRight.Y), FVector4f(U, V + SizeV, Ut, VtMax), FVector2f(0.0f, 1.0f), Tint));
+			RenderBatch->AddVertex(FSlateVertex::Make<Rounding>(RenderTransform, FVector2f(LowerRight), FVector4f(U + SizeU, V + SizeV, UtMax, VtMax), FVector2f(1.0f, 1.0f), Tint));
 
 			RenderBatch->AddIndex(IndexStart + 0);
 			RenderBatch->AddIndex(IndexStart + 1);
diff --git a/Engine/Source/Runtime/SlateCore/Private/Rendering/SlateVectorGraphicsCache.cpp b/Engine/Source/Runtime/SlateCore/Private/Rendering/SlateVectorGraphicsCache.cpp
--- a/Engine/Source/Runtime/SlateCore/Private/Rendering/SlateVectorGraphicsCache.cpp
+++ b/Engine/Source/Runtime/SlateCore/Private/Rendering/SlateVectorGraphicsCache.cpp
@@ -99,8 +99,8 @@ void FSlateVectorGraphicsCache::UpdateCache()
 					TUniquePtr<FSlateShaderResourceProxy>& NewProxy = ResourceMap.FindChecked(Request.Key);
 					NewProxy->Resource = NewResource.Get();
 
-					NewProxy->StartUV = FVector2D(0, 0);
-					NewProxy->SizeUV = FVector2D(1, 1);
+					NewProxy->StartUV = FVector2f(0, 0);
+					NewProxy->SizeUV = FVector2f(1, 1);
 					NewProxy->ActualSize = FIntPoint(PixelSize.X, PixelSize.Y);
 
 					NonAtlasedTextures.Add(MoveTemp(NewResource));
@@ -139,8 +139,8 @@ void FSlateVectorGraphicsCache::UpdateCache()
 
 						NewProxy->Resource = FoundAtlas->GetAtlasTexture();
 						// Compute the sub-uvs for the location of this texture in the atlas, accounting for padding
-						NewProxy->StartUV = FVector2D((float)(NewSlot->X + Padding) / FoundAtlas->GetWidth(), (float)(NewSlot->Y + Padding) / FoundAtlas->GetHeight());
-						NewProxy->SizeUV = FVector2D((float)(NewSlot->Width - Padding * 2) / FoundAtlas->GetWidth(), (float)(NewSlot->Height - Padding * 2) / FoundAtlas->GetHeight());
+						NewProxy->StartUV = FVector2f((float)(NewSlot->X + Padding) / FoundAtlas->GetWidth(), (float)(NewSlot->Y + Padding) / FoundAtlas->GetHeight());
+						NewProxy->SizeUV = FVector2f((float)(NewSlot->Width - Padding * 2) / FoundAtlas->GetWidth(), (float)(NewSlot->Height - Padding * 2) / FoundAtlas->GetHeight());
 						NewProxy->ActualSize = FIntPoint(PixelSize.X, PixelSize.Y);
 					}
 				}
diff --git a/Engine/Source/Runtime/SlateCore/Private/Widgets/SWidget.cpp b/Engine/Source/Runtime/SlateCore/Private/Widgets/SWidget.cpp
--- a/Engine/Source/Runtime/SlateCore/Private/Widgets/SWidget.cpp
+++ b/Engine/Source/Runtime/SlateCore/Private/Widgets/SWidget.cpp
@@ -705,7 +705,7 @@ void SWidget::InvalidateChildRemovedFromTree(SWidget& Child)
 
 FVector2D SWidget::GetDesiredSize() const
 {
-	return DesiredSize.Get(FVector2D::ZeroVector);
+	return FVector2D(DesiredSize.Get(FVector2f::ZeroVector));
 }
 
 void SWidget::AssignParentWidget(TSharedPtr<SWidget> InParent)
@@ -1573,11 +1573,11 @@ int32 SWidget::Paint(const FPaintArgs& Args, const FGeometry& AllottedGeometry,
 		FSlateClippingZone ClippingZone(AllottedGeometry);
 
 		TArray<FVector2D> Points;
-		Points.Add(ClippingZone.TopLeft);
-		Points.Add(ClippingZone.TopRight);
-		Points.Add(ClippingZone.BottomRight);
-		Points.Add(ClippingZone.BottomLeft);
-		Points.Add(ClippingZone.TopLeft);
+		Points.Add(FVector2D(ClippingZone.TopLeft));
+		Points.Add(FVector2D(ClippingZone.TopRight));
+		Points.Add(FVector2D(ClippingZone.BottomRight));
+		Points.Add(FVector2D(ClippingZone.BottomLeft));
+		Points.Add(FVector2D(ClippingZone.TopLeft));
 
 		const bool bAntiAlias = true;
 		FSlateDrawElement::MakeLines(
diff --git a/Engine/Source/Runtime/SlateCore/Public/Layout/Clipping.h b/Engine/Source/Runtime/SlateCore/Public/Layout/Clipping.h
--- a/Engine/Source/Runtime/SlateCore/Public/Layout/Clipping.h
+++ b/Engine/Source/Runtime/SlateCore/Public/Layout/Clipping.h
@@ -111,7 +111,7 @@ public:
 	{
 		if (bIsAxisAligned)
 		{
-			FVector2D Difference = TopLeft - BottomRight;
+			FVector2f Difference = TopLeft - BottomRight;
 			return FMath::IsNearlyZero(Difference.X) || FMath::IsNearlyZero(Difference.Y);
 		}
 
@@ -153,7 +153,7 @@ public:
 
 	FSlateClippingZone ConvertRelativeToAbsolute(const FVector2D& WindowOffset) const
 	{
-		FSlateClippingZone Absolute(TopLeft + WindowOffset, TopRight + WindowOffset, BottomLeft + WindowOffset, BottomRight + WindowOffset);
+		FSlateClippingZone Absolute(FVector2D(TopLeft) + WindowOffset, FVector2D(TopRight) + WindowOffset, FVector2D(BottomLeft) + WindowOffset, FVector2D(BottomRight) + WindowOffset);
 		Absolute.bIsAxisAligned = bIsAxisAligned;
 		Absolute.bIntersect = bIntersect;
 		Absolute.bAlwaysClip = bAlwaysClip;
diff --git a/Engine/Source/Runtime/SlateCore/Public/Layout/Geometry.h b/Engine/Source/Runtime/SlateCore/Public/Layout/Geometry.h
--- a/Engine/Source/Runtime/SlateCore/Public/Layout/Geometry.h
+++ b/Engine/Source/Runtime/SlateCore/Public/Layout/Geometry.h
@@ -67,9 +67,9 @@ public:
 		FSlateLayoutTransform AccumulatedLayoutTransform = Concatenate(LocalLayoutTransform, ParentAccumulatedLayoutTransform);
 		AccumulatedRenderTransform = TransformCast<FSlateRenderTransform>(AccumulatedLayoutTransform);
 		// HACK to allow us to make FGeometry public members immutable to catch misuse.
-		const_cast<FVector2f&>( AbsolutePosition ) = AccumulatedLayoutTransform.GetTranslation();
+		const_cast<FVector2f&>( AbsolutePosition ) = FVector2f(AccumulatedLayoutTransform.GetTranslation());
 		const_cast<float&>( Scale ) = AccumulatedLayoutTransform.GetScale();
-		const_cast<FVector2f&>( Position ) = LocalLayoutTransform.GetTranslation();
+		const_cast<FVector2f&>( Position ) = FVector2f(LocalLayoutTransform.GetTranslation());
 	}
 
 private:
@@ -110,9 +110,9 @@ private:
 	{
 		FSlateLayoutTransform AccumulatedLayoutTransform = Concatenate(InLocalLayoutTransform, ParentAccumulatedLayoutTransform);
 		// HACK to allow us to make FGeometry public members immutable to catch misuse.
-		const_cast<FVector2f&>( AbsolutePosition ) = AccumulatedLayoutTransform.GetTranslation();
+		const_cast<FVector2f&>( AbsolutePosition ) = FVector2f(AccumulatedLayoutTransform.GetTranslation());
 		const_cast<float&>( Scale ) = AccumulatedLayoutTransform.GetScale();
-		const_cast<FVector2f&>( Position ) = InLocalLayoutTransform.GetTranslation();
+		const_cast<FVector2f&>( Position ) = FVector2f(InLocalLayoutTransform.GetTranslation());
 	}
 
 	/**
@@ -137,9 +137,9 @@ private:
 	{
 		FSlateLayoutTransform AccumulatedLayoutTransform = Concatenate(InLocalLayoutTransform, ParentAccumulatedLayoutTransform);
 		// HACK to allow us to make FGeometry public members immutable to catch misuse.
-		const_cast<FVector2f&>( AbsolutePosition ) = AccumulatedLayoutTransform.GetTranslation();
+		const_cast<FVector2f&>( AbsolutePosition ) = FVector2f(AccumulatedLayoutTransform.GetTranslation());
 		const_cast<float&>( Scale ) = AccumulatedLayoutTransform.GetScale();
-		const_cast<FVector2f&>( Position ) = InLocalLayoutTransform.GetTranslation();
+		const_cast<FVector2f&>( Position ) = FVector2f(InLocalLayoutTransform.GetTranslation());
 	}
 
 public:
@@ -491,7 +491,7 @@ public:
 	/** @return the accumulated layout transform. Shouldn't be needed in general. */
 	FORCEINLINE FSlateLayoutTransform GetAccumulatedLayoutTransform() const
 	{
-		return FSlateLayoutTransform(Scale, AbsolutePosition);
+		return FSlateLayoutTransform(Scale, FVector2D(AbsolutePosition));
 	}
 
 	/**
@@ -506,7 +506,7 @@ public:
 	{
 		FSlateLayoutTransform AccumulatedLayoutTransform = ::Concatenate(GetAccumulatedLayoutTransform(), LayoutTransform);
 		AccumulatedRenderTransform = ::Concatenate(AccumulatedRenderTransform, LayoutTransform);
-		const_cast<FVector2f&>( AbsolutePosition ) = AccumulatedLayoutTransform.GetTranslation();
+		const_cast<FVector2f&>( AbsolutePosition ) = FVector2f(AccumulatedLayoutTransform.GetTranslation());
 		const_cast<float&>( Scale ) = AccumulatedLayoutTransform.GetScale();
 	}
 
@@ -547,7 +547,7 @@ public:
 	 */
 	FORCEINLINE FVector2D GetLocalPositionAtCoordinates(const FVector2D& NormalCoordinates) const
 	{
-		return Position + (NormalCoordinates * GetLocalSize());
+		return FVector2D(Position) + (NormalCoordinates * GetLocalSize());
 	}
 
 	bool HasRenderTransform() const { return bHasRenderTransform; }
diff --git a/Engine/Source/Runtime/SlateCore/Public/Layout/LayoutGeometry.h b/Engine/Source/Runtime/SlateCore/Public/Layout/LayoutGeometry.h
--- a/Engine/Source/Runtime/SlateCore/Public/Layout/LayoutGeometry.h
+++ b/Engine/Source/Runtime/SlateCore/Public/Layout/LayoutGeometry.h
@@ -27,12 +27,12 @@ public:
 
 	FVector2D GetSizeInLocalSpace() const
 	{
-		return LocalSize;
+		return FVector2D(LocalSize);
 	}
 
 	FVector2D GetSizeInParentSpace() const
 	{
-		return TransformVector(LocalToParent, LocalSize);
+		return TransformVector(LocalToParent, FVector2D(LocalSize));
 	}
 
 	FVector2D GetOffsetInParentSpace() const
@@ -42,7 +42,7 @@ public:
 
 	FSlateRect GetRectInLocalSpace() const
 	{
-		return FSlateRect(FVector2D(0.0f, 0.0f), LocalSize);
+		return FSlateRect(FVector2D(0.0f, 0.0f), FVector2D(LocalSize));
 	}
 
 	FSlateRect GetRectInParentSpace() const
diff --git a/Engine/Source/Runtime/SlateCore/Public/Layout/PaintGeometry.h b/Engine/Source/Runtime/SlateCore/Public/Layout/PaintGeometry.h
--- a/Engine/Source/Runtime/SlateCore/Public/Layout/PaintGeometry.h
+++ b/Engine/Source/Runtime/SlateCore/Public/Layout/PaintGeometry.h
@@ -54,9 +54,9 @@ struct SLATECORE_API FPaintGeometry
 	{ 
 		if (!bUsingLegacyConstructor) return;
 		
-		AccumulatedRenderTransform = FSlateRenderTransform(DrawScale, DrawPosition);
-		FSlateLayoutTransform AccumulatedLayoutTransform = FSlateLayoutTransform(DrawScale, DrawPosition);
-		LocalSize = TransformVector(Inverse(AccumulatedLayoutTransform), DrawSize);
+		AccumulatedRenderTransform = FSlateRenderTransform(DrawScale, FVector2D(DrawPosition));
+		FSlateLayoutTransform AccumulatedLayoutTransform = FSlateLayoutTransform(DrawScale, FVector2D(DrawPosition));
+		LocalSize = TransformVector(Inverse(AccumulatedLayoutTransform), FVector2D(DrawSize));
 	}
 
 	bool HasRenderTransform() const { return bHasRenderTransform; }
diff --git a/Engine/Source/Runtime/SlateCore/Public/Rendering/DrawElements.h b/Engine/Source/Runtime/SlateCore/Public/Rendering/DrawElements.h
--- a/Engine/Source/Runtime/SlateCore/Public/Rendering/DrawElements.h
+++ b/Engine/Source/Runtime/SlateCore/Public/Rendering/DrawElements.h
@@ -270,7 +270,7 @@ public:
 	FORCEINLINE const FSlateRenderTransform& GetRenderTransform() const { return RenderTransform; }
 	FORCEINLINE void SetRenderTransform(const FSlateRenderTransform& InRenderTransform) { RenderTransform = InRenderTransform; }
 	FORCEINLINE FVector2D GetPosition() const { return FVector2D(Position); }
-	FORCEINLINE void SetPosition(const FVector2D& InPosition) { Position = InPosition; }
+	FORCEINLINE void SetPosition(const FVector2D& InPosition) { Position = FVector2f(InPosition); }
 	FORCEINLINE FVector2D GetLocalSize() const { return FVector2D(LocalSize); }
 	FORCEINLINE float GetScale() const { return Scale; }
 	FORCEINLINE ESlateDrawEffect GetDrawEffects() const { return DrawEffects; }
@@ -289,7 +289,7 @@ public:
 
 	FORCEINLINE FSlateLayoutTransform GetInverseLayoutTransform() const
 	{
-		return Inverse(FSlateLayoutTransform(Scale, Position));
+		return Inverse(FSlateLayoutTransform(Scale, FVector2D(Position)));
 	}
 
 	void AddReferencedObjects(FReferenceCollector& Collector);
diff --git a/Engine/Source/Runtime/SlateCore/Public/Rendering/SlateLayoutTransform.h b/Engine/Source/Runtime/SlateCore/Public/Rendering/SlateLayoutTransform.h
--- a/Engine/Source/Runtime/SlateCore/Public/Rendering/SlateLayoutTransform.h
+++ b/Engine/Source/Runtime/SlateCore/Public/Rendering/SlateLayoutTransform.h
@@ -81,7 +81,7 @@ public:
 	/** Invert the transform/scale. */
 	FSlateLayoutTransform Inverse() const
 	{
-		return FSlateLayoutTransform(::Inverse(Scale), ::Inverse(Translation) * ::Inverse(Scale));
+		return FSlateLayoutTransform(::Inverse(Scale), ::Inverse(FVector2D(Translation)) * ::Inverse(Scale));
 	}
 
 	/** Equality. */
diff --git a/Engine/Source/Runtime/SlateCore/Public/Widgets/SWidget.h b/Engine/Source/Runtime/SlateCore/Public/Widgets/SWidget.h
--- a/Engine/Source/Runtime/SlateCore/Public/Widgets/SWidget.h
+++ b/Engine/Source/Runtime/SlateCore/Public/Widgets/SWidget.h
@@ -780,7 +780,7 @@ private:
 	 */
 	void SetDesiredSize(const FVector2D& InDesiredSize)
 	{
-		DesiredSize = InDesiredSize;
+		DesiredSize = FVector2f(InDesiredSize);
 	}
 
 #if STATS || ENABLE_STATNAMEDEVENTS
diff --git a/Engine/Source/Runtime/SlateRHIRenderer/Private/SlateAtlasedTextureResource.cpp b/Engine/Source/Runtime/SlateRHIRenderer/Private/SlateAtlasedTextureResource.cpp
--- a/Engine/Source/Runtime/SlateRHIRenderer/Private/SlateAtlasedTextureResource.cpp
+++ b/Engine/Source/Runtime/SlateRHIRenderer/Private/SlateAtlasedTextureResource.cpp
@@ -30,8 +30,8 @@ FSlateShaderResourceProxy* FSlateAtlasedTextureResource::FindOrCreateAtlasedProx
 		Proxy = new FSlateShaderResourceProxy();
 		Proxy->Resource = this;
 		Proxy->ActualSize = ActualSize.IntPoint();
-		Proxy->StartUV = AtlasData.StartUV;
-		Proxy->SizeUV = AtlasData.SizeUV;
+		Proxy->StartUV = FVector2f(AtlasData.StartUV);	// LWC_TODO: Precision loss
+		Proxy->SizeUV = FVector2f(AtlasData.SizeUV);	// LWC_TODO: Precision loss
 
 		ProxyMap.Add(InAtlasedObject, Proxy);
 	}
diff --git a/Engine/Source/Runtime/SlateRHIRenderer/Private/SlatePostProcessor.cpp b/Engine/Source/Runtime/SlateRHIRenderer/Private/SlatePostProcessor.cpp
--- a/Engine/Source/Runtime/SlateRHIRenderer/Private/SlatePostProcessor.cpp
+++ b/Engine/Source/Runtime/SlateRHIRenderer/Private/SlatePostProcessor.cpp
@@ -117,7 +117,7 @@ void FSlatePostProcessor::BlurRect(FRHICommandListImmediate& RHICmdList, IRender
 
 				if (bDownsample)
 				{
-					PixelShader->SetUVBounds(RHICmdList, FVector4f(FVector2D::ZeroVector, FVector2D((float)DownsampleSize.X / DestTextureWidth, (float)DownsampleSize.Y / DestTextureHeight) - HalfTexelOffset));
+					PixelShader->SetUVBounds(RHICmdList, FVector4f(FVector2f::ZeroVector, FVector2f((float)DownsampleSize.X / DestTextureWidth, (float)DownsampleSize.Y / DestTextureHeight) - FVector2f(HalfTexelOffset)));
 					PixelShader->SetBufferSizeAndDirection(RHICmdList, InvBufferSize, FVector2D(1, 0));
 
 					RendererModule.DrawRectangle(
@@ -133,14 +133,14 @@ void FSlatePostProcessor::BlurRect(FRHICommandListImmediate& RHICmdList, IRender
 				}
 				else
 				{
-					const FVector2D InvSrcTetureSize(1.f / SrcTextureWidth, 1.f / SrcTextureHeight);
+					const FVector2f InvSrcTextureSize(1.f / SrcTextureWidth, 1.f / SrcTextureHeight);
 
-					const FVector2D UVStart = FVector2D(DestRect.Left, DestRect.Top) * InvSrcTetureSize;
-					const FVector2D UVEnd = FVector2D(DestRect.Right, DestRect.Bottom) * InvSrcTetureSize;
-					const FVector2D SizeUV = UVEnd - UVStart;
+					const FVector2f UVStart = FVector2f(DestRect.Left, DestRect.Top) * InvSrcTextureSize;
+					const FVector2f UVEnd = FVector2f(DestRect.Right, DestRect.Bottom) * InvSrcTextureSize;
+					const FVector2f SizeUV = UVEnd - UVStart;
 
 					PixelShader->SetUVBounds(RHICmdList, FVector4f(UVStart, UVEnd));
-					PixelShader->SetBufferSizeAndDirection(RHICmdList, InvSrcTetureSize, FVector2D(1, 0));
+					PixelShader->SetBufferSizeAndDirection(RHICmdList, FVector2D(InvSrcTextureSize), FVector2D(1, 0));
 
 					RendererModule.DrawRectangle(
 						RHICmdList,
@@ -176,7 +176,7 @@ void FSlatePostProcessor::BlurRect(FRHICommandListImmediate& RHICmdList, IRender
 				SetGraphicsPipelineState(RHICmdList, GraphicsPSOInit, 0);
 
 				PixelShader->SetWeightsAndOffsets(RHICmdList, WeightsAndOffsets, SampleCount);
-				PixelShader->SetUVBounds(RHICmdList, FVector4f(FVector2D::ZeroVector, FVector2D((float)DownsampleSize.X / DestTextureWidth, (float)DownsampleSize.Y / DestTextureHeight) - HalfTexelOffset));
+				PixelShader->SetUVBounds(RHICmdList, FVector4f(FVector2f::ZeroVector, FVector2f((float)DownsampleSize.X / DestTextureWidth, (float)DownsampleSize.Y / DestTextureHeight) - FVector2f(HalfTexelOffset)));
 				PixelShader->SetTexture(RHICmdList, SourceTexture, BilinearClamp);
 				PixelShader->SetBufferSizeAndDirection(RHICmdList, InvBufferSize, FVector2D(0, 1));
 
@@ -320,11 +320,11 @@ void FSlatePostProcessor::DownsampleRect(FRHICommandListImmediate& RHICmdList, I
 		RHICmdList.Transition(FRHITransitionInfo(Params.SourceTexture, ERHIAccess::Unknown, ERHIAccess::SRVGraphics));
 		RHICmdList.Transition(FRHITransitionInfo(DestTexture, ERHIAccess::Unknown, ERHIAccess::RTV));
 
-		const FVector2D InvSrcTetureSize(1.f/SrcTextureWidth, 1.f/SrcTextureHeight);
+		const FVector2f InvSrcTextureSize(1.f/SrcTextureWidth, 1.f/SrcTextureHeight);
 
-		const FVector2D UVStart = FVector2D(DestRect.Left, DestRect.Top) * InvSrcTetureSize;
-		const FVector2D UVEnd = FVector2D(DestRect.Right, DestRect.Bottom) * InvSrcTetureSize;
-		const FVector2D SizeUV = UVEnd - UVStart;
+		const FVector2f UVStart = FVector2f(DestRect.Left, DestRect.Top) * InvSrcTextureSize;
+		const FVector2f UVEnd = FVector2f(DestRect.Right, DestRect.Bottom) * InvSrcTextureSize;
+		const FVector2f SizeUV = UVEnd - UVStart;
 		
 		RHICmdList.SetViewport(0, 0, 0, DestTextureWidth, DestTextureHeight, 0.0f);
 		RHICmdList.SetScissorRect(false, 0, 0, 0, 0);
@@ -344,7 +344,7 @@ void FSlatePostProcessor::DownsampleRect(FRHICommandListImmediate& RHICmdList, I
 			GraphicsPSOInit.PrimitiveType = PT_TriangleList;
 			SetGraphicsPipelineState(RHICmdList, GraphicsPSOInit, 0);
 
-			PixelShader->SetShaderParams(RHICmdList, FShaderParams::MakePixelShaderParams(FVector4f(InvSrcTetureSize.X, InvSrcTetureSize.Y, 0, 0)));
+			PixelShader->SetShaderParams(RHICmdList, FShaderParams::MakePixelShaderParams(FVector4f(InvSrcTextureSize.X, InvSrcTextureSize.Y, 0, 0)));
 			PixelShader->SetUVBounds(RHICmdList, FVector4f(UVStart, UVEnd));
 			PixelShader->SetTexture(RHICmdList, Params.SourceTexture, BilinearClamp);
 
@@ -425,12 +425,12 @@ void FSlatePostProcessor::UpsampleRect(FRHICommandListImmediate& RHICmdList, IRe
 		GraphicsPSOInit.PrimitiveType = PT_TriangleList;
 		SetGraphicsPipelineState(RHICmdList, GraphicsPSOInit, Params.StencilRef);
 
-		const FVector2D SizeUV(
+		const FVector2f SizeUV(
 			DownsampledWidth == SrcTextureWidth ? 1.0f : (DownsampledWidth / (float)SrcTextureWidth) - (1.0f / (float)SrcTextureWidth),
 			DownsampledHeight == SrcTextureHeight ? 1.0f : (DownsampledHeight / (float)SrcTextureHeight) - (1.0f / (float)SrcTextureHeight)
 			);
 
-		const FVector2D Size(DestRect.Right - DestRect.Left, DestRect.Bottom - DestRect.Top);
+		const FVector2f Size(DestRect.Right - DestRect.Left, DestRect.Bottom - DestRect.Top);
 		FShaderParams ShaderParams = FShaderParams::MakePixelShaderParams(FVector4f(Size, SizeUV), (FVector4f)Params.CornerRadius); // LWC_TODO: precision loss
 
 
@@ -516,11 +516,11 @@ static int32 ComputeWeights(int32 KernelSize, float Sigma, TArray<FVector4f>& Ou
 
 	OutWeightsAndOffsets.AddUninitialized(NumSamples%2 == 0 ? NumSamples / 2 : NumSamples/2+1);
 
-	OutWeightsAndOffsets[0] = FVector4f(FVector2D(GetWeight(0,Sigma), 0), GetWeightAndOffset(1, Sigma) );
+	OutWeightsAndOffsets[0] = FVector4f(FVector2f(GetWeight(0,Sigma), 0), FVector2f(GetWeightAndOffset(1, Sigma)) );
 	int32 SampleIndex = 1;
 	for (int32 X = 3; X < KernelSize; X += 4)
 	{
-		OutWeightsAndOffsets[SampleIndex] = FVector4f(GetWeightAndOffset(X, Sigma), GetWeightAndOffset(X + 2, Sigma));
+		OutWeightsAndOffsets[SampleIndex] = FVector4f(FVector2f(GetWeightAndOffset(X, Sigma)), FVector2f(GetWeightAndOffset(X + 2, Sigma)));
 
 		++SampleIndex;
 	}
diff --git a/Engine/Source/Runtime/SlateRHIRenderer/Private/SlateRHIRenderingPolicy.cpp b/Engine/Source/Runtime/SlateRHIRenderer/Private/SlateRHIRenderingPolicy.cpp
--- a/Engine/Source/Runtime/SlateRHIRenderer/Private/SlateRHIRenderingPolicy.cpp
+++ b/Engine/Source/Runtime/SlateRHIRenderer/Private/SlateRHIRenderingPolicy.cpp
@@ -465,7 +465,7 @@ static bool UpdateScissorRect(
 
 						SCOPE_CYCLE_COUNTER(STAT_SlateRTStencilDrawCall);
 
-						VertexShader->SetMaskRect(RHICmdList, MaskQuad.TopLeft, MaskQuad.TopRight, MaskQuad.BottomLeft, MaskQuad.BottomRight);
+						VertexShader->SetMaskRect(RHICmdList, FVector2D(MaskQuad.TopLeft), FVector2D(MaskQuad.TopRight), FVector2D(MaskQuad.BottomLeft), FVector2D(MaskQuad.BottomRight));
 
 						RHICmdList.SetStreamSource(0, StencilVertexBuffer.VertexBufferRHI, 0);
 						RHICmdList.DrawPrimitive(0, 2, 1);
@@ -510,7 +510,7 @@ static bool UpdateScissorRect(
 
 					SCOPE_CYCLE_COUNTER(STAT_SlateRTStencilDrawCall);
 					
-					VertexShader->SetMaskRect(RHICmdList, MaskQuad.TopLeft, MaskQuad.TopRight, MaskQuad.BottomLeft, MaskQuad.BottomRight);
+					VertexShader->SetMaskRect(RHICmdList, FVector2D(MaskQuad.TopLeft), FVector2D(MaskQuad.TopRight), FVector2D(MaskQuad.BottomLeft), FVector2D(MaskQuad.BottomRight));
 
 					RHICmdList.SetStreamSource(0, StencilVertexBuffer.VertexBufferRHI, 0);
 					RHICmdList.DrawPrimitive(0, 2, 1);
diff --git a/Engine/Source/Runtime/SlateRHIRenderer/Private/SlateRHIResourceManager.cpp b/Engine/Source/Runtime/SlateRHIRenderer/Private/SlateRHIResourceManager.cpp
--- a/Engine/Source/Runtime/SlateRHIRenderer/Private/SlateRHIResourceManager.cpp
+++ b/Engine/Source/Runtime/SlateRHIRenderer/Private/SlateRHIResourceManager.cpp
@@ -591,8 +591,8 @@ FSlateShaderResourceProxy* FSlateRHIResourceManager::GenerateTextureResource( co
 		NewProxy = new FSlateShaderResourceProxy;
 		NewProxy->Resource = Atlas->GetAtlasTexture();
 		const uint32 Padding = NewSlot->Padding;
-		NewProxy->StartUV = FVector2D((float)(NewSlot->X + Padding) / Atlas->GetWidth(), (float)(NewSlot->Y + Padding) / Atlas->GetHeight());
-		NewProxy->SizeUV = FVector2D( (float)(NewSlot->Width-Padding*2) / Atlas->GetWidth(), (float)(NewSlot->Height-Padding*2) / Atlas->GetHeight() );
+		NewProxy->StartUV = FVector2f((float)(NewSlot->X + Padding) / Atlas->GetWidth(), (float)(NewSlot->Y + Padding) / Atlas->GetHeight());
+		NewProxy->SizeUV = FVector2f( (float)(NewSlot->Width-Padding*2) / Atlas->GetWidth(), (float)(NewSlot->Height-Padding*2) / Atlas->GetHeight() );
 		NewProxy->ActualSize = FIntPoint( Width, Height );
 	}
 	else
@@ -610,8 +610,8 @@ FSlateShaderResourceProxy* FSlateRHIResourceManager::GenerateTextureResource( co
 
 		// The texture proxy only contains a single texture
 		NewProxy->Resource = Texture;
-		NewProxy->StartUV = FVector2D(0.0f, 0.0f);
-		NewProxy->SizeUV = FVector2D(1.0f, 1.0f);
+		NewProxy->StartUV = FVector2f::ZeroVector;
+		NewProxy->SizeUV = FVector2f::UnitVector;
 		NewProxy->ActualSize = FIntPoint( Width, Height );
 	}
 
diff --git a/Engine/Source/Runtime/SlateRHIRenderer/Private/SlateShaders.cpp b/Engine/Source/Runtime/SlateRHIRenderer/Private/SlateShaders.cpp
--- a/Engine/Source/Runtime/SlateRHIRenderer/Private/SlateShaders.cpp
+++ b/Engine/Source/Runtime/SlateRHIRenderer/Private/SlateShaders.cpp
@@ -197,7 +197,7 @@ void FSlateMaskingVS::SetVerticalAxisMultiplier(FRHICommandList& RHICmdList, flo
 void FSlateMaskingVS::SetMaskRect(FRHICommandList& RHICmdList, const FVector2D& TopLeft, const FVector2D& TopRight, const FVector2D& BotLeft, const FVector2D& BotRight)
 {
 	//FVector4f MaskRectVal[4] = { FVector4f(TopLeft, FVector2D::ZeroVector), FVector4f(TopRight, FVector2D::ZeroVector), FVector4f(BotLeft, FVector2D::ZeroVector), FVector4f(BotRight, FVector2D::ZeroVector) };
-	FVector4f MaskRectVal[2] = { FVector4f(TopLeft, TopRight), FVector4f(BotLeft, BotRight) };
+	FVector4f MaskRectVal[2] = { FVector4f(FVector2f(TopLeft), FVector2f(TopRight)), FVector4f(FVector2f(BotLeft), FVector2f(BotRight)) };	// LWC_TODO: Precision loss
 
 	SetShaderValue(RHICmdList, RHICmdList.GetBoundVertexShader(), MaskRect, MaskRectVal);
 }
diff --git a/Engine/Source/Runtime/SlateRHIRenderer/Private/SlateShaders.h b/Engine/Source/Runtime/SlateRHIRenderer/Private/SlateShaders.h
--- a/Engine/Source/Runtime/SlateRHIRenderer/Private/SlateShaders.h
+++ b/Engine/Source/Runtime/SlateRHIRenderer/Private/SlateShaders.h
@@ -348,7 +348,7 @@ public:
 
 	void SetBufferSizeAndDirection(FRHICommandList& RHICmdList, const FVector2D& InBufferSize, const FVector2D& InDir)
 	{
-		SetShaderValue(RHICmdList, RHICmdList.GetBoundPixelShader(), BufferSizeAndDirection, FVector4f(InBufferSize, InDir));
+		SetShaderValue(RHICmdList, RHICmdList.GetBoundPixelShader(), BufferSizeAndDirection, FVector4f(FVector2f(InBufferSize), FVector2f(InDir)));	// LWC_TODO: Precision loss
 	}
 
 	void SetWeightsAndOffsets(FRHICommandList& RHICmdList, const TArray<FVector4f>& InWeightsAndOffsets, int32 NumSamples )
diff --git a/Engine/Source/Runtime/StaticMeshDescription/Private/StaticMeshOperations.cpp b/Engine/Source/Runtime/StaticMeshDescription/Private/StaticMeshOperations.cpp
--- a/Engine/Source/Runtime/StaticMeshDescription/Private/StaticMeshOperations.cpp
+++ b/Engine/Source/Runtime/StaticMeshDescription/Private/StaticMeshOperations.cpp
@@ -188,9 +188,9 @@ static TTuple<FVector3f, FVector3f, FVector3f> GetTriangleTangentsAndNormals(flo
 	const FVector3f DPosition1 = VertexPositions[1] - Position0;
 	const FVector3f DPosition2 = VertexPositions[2] - Position0;
 
-	const FVector2f UV0 = VertexUVs[0];
-	const FVector2f DUV1 = VertexUVs[1] - UV0;
-	const FVector2f DUV2 = VertexUVs[2] - UV0;
+	const FVector2f UV0 = FVector2f(VertexUVs[0]);
+	const FVector2f DUV1 = FVector2f(VertexUVs[1]) - UV0;
+	const FVector2f DUV2 = FVector2f(VertexUVs[2]) - UV0;
 
 	// We have a left-handed coordinate system, but a counter-clockwise winding order
 	// Hence normal calculation has to take the triangle vectors cross product in reverse.
@@ -291,9 +291,9 @@ void FStaticMeshOperations::ComputeTriangleTangentsAndNormals(FMeshDescription&
 
 				FVector2D TriangleUVs[3] =
 				{
-					VertexUVs[TriangleVertexInstanceIDs[TriIndex]],
-					VertexUVs[TriangleVertexInstanceIDs[TriIndex + 1]],
-					VertexUVs[TriangleVertexInstanceIDs[TriIndex + 2]]
+					FVector2D(VertexUVs[TriangleVertexInstanceIDs[TriIndex]]),
+					FVector2D(VertexUVs[TriangleVertexInstanceIDs[TriIndex + 1]]),
+					FVector2D(VertexUVs[TriangleVertexInstanceIDs[TriIndex + 2]])
 				};
 
 				TTuple<FVector3f, FVector3f, FVector3f> Result = GetTriangleTangentsAndNormals(ComparisonThreshold, TriangleVertexPositions, TriangleUVs);
@@ -379,7 +379,7 @@ struct FVertexInfo
 	{
 		TriangleID = INDEX_NONE;
 		VertexInstanceID = INDEX_NONE;
-		UVs = FVector2D(0.0f, 0.0f);
+		UVs = FVector2f(0.0f, 0.0f);
 	}
 
 	FTriangleID TriangleID;
diff --git a/Engine/Source/Runtime/UMG/Private/Components/CanvasPanelSlot.cpp b/Engine/Source/Runtime/UMG/Private/Components/CanvasPanelSlot.cpp
--- a/Engine/Source/Runtime/UMG/Private/Components/CanvasPanelSlot.cpp
+++ b/Engine/Source/Runtime/UMG/Private/Components/CanvasPanelSlot.cpp
@@ -386,10 +386,10 @@ void UCanvasPanelSlot::RebaseLayout(bool PreserveSize)
 			// Determine the amount that would be offset from the anchor position if alignment was applied.
 			FVector2D AlignmentOffset = LayoutData.Alignment * PreEditGeometry.Size;
 
-			FVector2D MoveDelta = Geometry.Position - PreEditGeometry.Position;
+			FVector2D MoveDelta = FVector2D(Geometry.Position) - FVector2D(PreEditGeometry.Position);
 
 			// Determine where the widget's new position needs to be to maintain a stable location when the anchors change.
-			FVector2D LeftTopDelta = PreEditGeometry.Position - DefaultAnchorPosition;
+			FVector2D LeftTopDelta = FVector2D(PreEditGeometry.Position) - FVector2D(DefaultAnchorPosition);
 
 			const bool bAnchorsMoved = PreEditLayoutData.Anchors.Minimum != LayoutData.Anchors.Minimum || PreEditLayoutData.Anchors.Maximum != LayoutData.Anchors.Maximum;
 			const bool bMoved = PreEditLayoutData.Offsets.Left != LayoutData.Offsets.Left || PreEditLayoutData.Offsets.Top != LayoutData.Offsets.Top;
diff --git a/Engine/Source/Runtime/UMG/Private/Components/WidgetComponent.cpp b/Engine/Source/Runtime/UMG/Private/Components/WidgetComponent.cpp
--- a/Engine/Source/Runtime/UMG/Private/Components/WidgetComponent.cpp
+++ b/Engine/Source/Runtime/UMG/Private/Components/WidgetComponent.cpp
@@ -376,10 +376,10 @@ public:
 
 						if ( VisibilityMap & ( 1 << ViewIndex ) )
 						{
-							VertexIndices[0] = MeshBuilder.AddVertex(-FVector(0, U, V ),  FVector2D(0, 0), FVector(0, -1, 0), FVector(0, 0, -1), FVector(1, 0, 0), FColor::White);
-							VertexIndices[1] = MeshBuilder.AddVertex(-FVector(0, U, VL),  FVector2D(0, 1), FVector(0, -1, 0), FVector(0, 0, -1), FVector(1, 0, 0), FColor::White);
-							VertexIndices[2] = MeshBuilder.AddVertex(-FVector(0, UL, VL), FVector2D(1, 1), FVector(0, -1, 0), FVector(0, 0, -1), FVector(1, 0, 0), FColor::White);
-							VertexIndices[3] = MeshBuilder.AddVertex(-FVector(0, UL, V),  FVector2D(1, 0), FVector(0, -1, 0), FVector(0, 0, -1), FVector(1, 0, 0), FColor::White);
+							VertexIndices[0] = MeshBuilder.AddVertex(-FVector(0, U, V ),  FVector2f(0, 0), FVector(0, -1, 0), FVector(0, 0, -1), FVector(1, 0, 0), FColor::White);
+							VertexIndices[1] = MeshBuilder.AddVertex(-FVector(0, U, VL),  FVector2f(0, 1), FVector(0, -1, 0), FVector(0, 0, -1), FVector(1, 0, 0), FColor::White);
+							VertexIndices[2] = MeshBuilder.AddVertex(-FVector(0, UL, VL), FVector2f(1, 1), FVector(0, -1, 0), FVector(0, 0, -1), FVector(1, 0, 0), FColor::White);
+							VertexIndices[3] = MeshBuilder.AddVertex(-FVector(0, UL, V),  FVector2f(1, 0), FVector(0, -1, 0), FVector(0, 0, -1), FVector(1, 0, 0), FColor::White);
 
 							MeshBuilder.AddTriangle(VertexIndices[0], VertexIndices[1], VertexIndices[2]);
 							MeshBuilder.AddTriangle(VertexIndices[0], VertexIndices[2], VertexIndices[3]);
@@ -453,10 +453,10 @@ public:
 									LastTangentZ = TangentZ;
 								}
 
-								VertexIndices[0] = MeshBuilder.AddVertex(Vertex0, FVector2D(U0, 0), LastTangentX, LastTangentY, LastTangentZ, FColor::White);
-								VertexIndices[1] = MeshBuilder.AddVertex(Vertex1, FVector2D(U0, 1), LastTangentX, LastTangentY, LastTangentZ, FColor::White);
-								VertexIndices[2] = MeshBuilder.AddVertex(Vertex2, FVector2D(U1, 1), TangentX, TangentY, TangentZ, FColor::White);
-								VertexIndices[3] = MeshBuilder.AddVertex(Vertex3, FVector2D(U1, 0), TangentX, TangentY, TangentZ, FColor::White);
+								VertexIndices[0] = MeshBuilder.AddVertex(Vertex0, FVector2f(U0, 0), LastTangentX, LastTangentY, LastTangentZ, FColor::White);
+								VertexIndices[1] = MeshBuilder.AddVertex(Vertex1, FVector2f(U0, 1), LastTangentX, LastTangentY, LastTangentZ, FColor::White);
+								VertexIndices[2] = MeshBuilder.AddVertex(Vertex2, FVector2f(U1, 1), TangentX, TangentY, TangentZ, FColor::White);
+								VertexIndices[3] = MeshBuilder.AddVertex(Vertex3, FVector2f(U1, 0), TangentX, TangentY, TangentZ, FColor::White);
 
 								MeshBuilder.AddTriangle(VertexIndices[0], VertexIndices[1], VertexIndices[2]);
 								MeshBuilder.AddTriangle(VertexIndices[0], VertexIndices[2], VertexIndices[3]);
diff --git a/Engine/Source/Runtime/UMG/Private/Slate/SRetainerWidget.cpp b/Engine/Source/Runtime/UMG/Private/Slate/SRetainerWidget.cpp
--- a/Engine/Source/Runtime/UMG/Private/Slate/SRetainerWidget.cpp
+++ b/Engine/Source/Runtime/UMG/Private/Slate/SRetainerWidget.cpp
@@ -381,7 +381,7 @@ SRetainerWidget::EPaintRetainedContentResult SRetainerWidget::PaintRetainedConte
 
 	const FPaintGeometry PaintGeometry = AllottedGeometry.ToPaintGeometry();
 	FSlateRenderTransform AccumulatedRenderTransform = PaintGeometry.GetAccumulatedRenderTransform();
-	const FVector2D RenderSize = PaintGeometry.GetLocalSize() * AccumulatedRenderTransform.GetMatrix().GetScale().GetVector();
+	const FVector2f RenderSize = FVector2f(PaintGeometry.GetLocalSize()) * AccumulatedRenderTransform.GetMatrix().GetScale().GetVector();
 	const FIntPoint RoundedRenderSize = RenderSize.IntPoint();
 
 	if (RenderOnInvalidation)
