Fix ElementBatcher.cpp merge

diff --git a/Engine/Source/Runtime/SlateCore/Private/Rendering/ElementBatcher.cpp b/Engine/Source/Runtime/SlateCore/Private/Rendering/ElementBatcher.cpp
--- a/Engine/Source/Runtime/SlateCore/Private/Rendering/ElementBatcher.cpp
+++ b/Engine/Source/Runtime/SlateCore/Private/Rendering/ElementBatcher.cpp
@@ -1582,16 +1582,16 @@ private:
 			// Once we have two points, we have a normal, so we can generate the first bit of geometry.
 			const FVector2D LastUp = LastNormal*HalfLineThickness;
 
-			RenderBatch.AddVertex(FSlateVertex::Make<ESlateVertexRounding::Disabled>(RenderTransform, LastPointAdded[1] + LastUp, FVector2D(1.0f, 0.0f), FVector2D::ZeroVector, InColor));
-			RenderBatch.AddVertex(FSlateVertex::Make<ESlateVertexRounding::Disabled>(RenderTransform, LastPointAdded[1] - LastUp, FVector2D(-1.0f, 0.0f), FVector2D::ZeroVector, InColor));
+			RenderBatch.AddVertex(FSlateVertex::Make<ESlateVertexRounding::Disabled>(RenderTransform, FVector2f(LastPointAdded[1] + LastUp), FVector2f(1.0f, 0.0f), FVector2f::ZeroVector, InColor));
+			RenderBatch.AddVertex(FSlateVertex::Make<ESlateVertexRounding::Disabled>(RenderTransform, FVector2f(LastPointAdded[1] - LastUp), FVector2f(-1.0f, 0.0f), FVector2f::ZeroVector, InColor));
 		}
 
 		if (NumPointsAdded >= 2)
 		{
 			const FVector2D AveragedUp = (0.5f*(NewNormal + LastNormal)).GetSafeNormal()*HalfLineThickness;
 
-			RenderBatch.AddVertex(FSlateVertex::Make<ESlateVertexRounding::Disabled>(RenderTransform, LastPointAdded[0] + AveragedUp, FVector2D(1.0f, 0.0f), FVector2D::ZeroVector, InColor));
-			RenderBatch.AddVertex(FSlateVertex::Make<ESlateVertexRounding::Disabled>(RenderTransform, LastPointAdded[0] - AveragedUp, FVector2D(-1.0f, 0.0f), FVector2D::ZeroVector, InColor));
+			RenderBatch.AddVertex(FSlateVertex::Make<ESlateVertexRounding::Disabled>(RenderTransform, FVector2f(LastPointAdded[0] + AveragedUp), FVector2f(1.0f, 0.0f), FVector2f::ZeroVector, InColor));
+			RenderBatch.AddVertex(FSlateVertex::Make<ESlateVertexRounding::Disabled>(RenderTransform, FVector2f(LastPointAdded[0] - AveragedUp), FVector2f(-1.0f, 0.0f), FVector2f::ZeroVector, InColor));
 
 			const int32 NumVerts = RenderBatch.GetNumVertices();
 
@@ -1628,8 +1628,8 @@ private:
 			// last line segment.
 			const FVector2D LastUp = LastNormal*HalfLineThickness;
 
-			RenderBatch.AddVertex(FSlateVertex::Make<ESlateVertexRounding::Disabled>(RenderTransform, LastPointAdded[0] + LastUp, FVector2D(1.0f, 0.0f), FVector2D::ZeroVector, InColor));
-			RenderBatch.AddVertex(FSlateVertex::Make<ESlateVertexRounding::Disabled>(RenderTransform, LastPointAdded[0] - LastUp, FVector2D(-1.0f, 0.0f), FVector2D::ZeroVector, InColor));
+			RenderBatch.AddVertex(FSlateVertex::Make<ESlateVertexRounding::Disabled>(RenderTransform, FVector2f(LastPointAdded[0] + LastUp), FVector2f(1.0f, 0.0f), FVector2f::ZeroVector, InColor));
+			RenderBatch.AddVertex(FSlateVertex::Make<ESlateVertexRounding::Disabled>(RenderTransform, FVector2f(LastPointAdded[0] - LastUp), FVector2f(-1.0f, 0.0f), FVector2f::ZeroVector, InColor));
 
 			const int32 NumVerts = RenderBatch.GetNumVertices();
 
@@ -1922,8 +1922,8 @@ struct FLineBuilder
 				const float ParallelDistance = DistanceToMiterLine / DirDotMiterNormal;
 				const FVector2D MiterUp = Up - (Direction * ParallelDistance);
 
-				RenderBatch.AddVertex(FSlateVertex::Make<Rounding>(RenderTransform, Position + MiterUp, FVector2D(1.0f, 0.0f), PointColor));
-				RenderBatch.AddVertex(FSlateVertex::Make<Rounding>(RenderTransform, Position - MiterUp, FVector2D(-1.0f, 0.0f), PointColor));
+				RenderBatch.AddVertex(FSlateVertex::Make<Rounding>(RenderTransform, FVector2f(Position + MiterUp), FVector2f(1.0f, 0.0f), PointColor));
+				RenderBatch.AddVertex(FSlateVertex::Make<Rounding>(RenderTransform, FVector2f(Position - MiterUp), FVector2f(-1.0f, 0.0f), PointColor));
 				AddQuadIndices(RenderBatch);
 			}
 			else
@@ -1977,10 +1977,10 @@ private:
 			const FVector2D CapInward = Direction * InwardDistance;
 			const FVector2D CapOutward = Direction * (InwardDistance - LocalCapLength);
 
-			RenderBatch.AddVertex(FSlateVertex::Make<Rounding>(RenderTransform, Position + CapOutward + Up, FVector2D(1.0f, -1.0f), Color));
-			RenderBatch.AddVertex(FSlateVertex::Make<Rounding>(RenderTransform, Position + CapOutward - Up, FVector2D(-1.0f, -1.0f), Color));
-			RenderBatch.AddVertex(FSlateVertex::Make<Rounding>(RenderTransform, Position + CapInward + Up, FVector2D(1.0f, 0.0f), Color));
-			RenderBatch.AddVertex(FSlateVertex::Make<Rounding>(RenderTransform, Position + CapInward - Up, FVector2D(-1.0f, 0.0f), Color));
+			RenderBatch.AddVertex(FSlateVertex::Make<Rounding>(RenderTransform, FVector2f(Position + CapOutward + Up), FVector2f(1.0f, -1.0f), Color));
+			RenderBatch.AddVertex(FSlateVertex::Make<Rounding>(RenderTransform, FVector2f(Position + CapOutward - Up), FVector2f(-1.0f, -1.0f), Color));
+			RenderBatch.AddVertex(FSlateVertex::Make<Rounding>(RenderTransform, FVector2f(Position + CapInward + Up), FVector2f(1.0f, 0.0f), Color));
+			RenderBatch.AddVertex(FSlateVertex::Make<Rounding>(RenderTransform, FVector2f(Position + CapInward - Up), FVector2f(-1.0f, 0.0f), Color));
 			AddQuadIndices(RenderBatch);
 		}
 	}
@@ -2007,12 +2007,12 @@ private:
 			const FVector2D CapInward = Direction * -InwardDistance;
 			const FVector2D CapOutward = Direction * (LocalCapLength - InwardDistance);
 
-			RenderBatch.AddVertex(FSlateVertex::Make<Rounding>(RenderTransform, Position + CapInward + Up, FVector2D(1.0f, 0.0f), Color));
-			RenderBatch.AddVertex(FSlateVertex::Make<Rounding>(RenderTransform, Position + CapInward - Up, FVector2D(-1.0f, 0.0f), Color));
+			RenderBatch.AddVertex(FSlateVertex::Make<Rounding>(RenderTransform, FVector2f(Position + CapInward + Up), FVector2f(1.0f, 0.0f), Color));
+			RenderBatch.AddVertex(FSlateVertex::Make<Rounding>(RenderTransform, FVector2f(Position + CapInward - Up), FVector2f(-1.0f, 0.0f), Color));
 			AddQuadIndices(RenderBatch);
 
-			RenderBatch.AddVertex(FSlateVertex::Make<Rounding>(RenderTransform, Position + CapOutward + Up, FVector2D(1.0f, 1.0f), Color));
-			RenderBatch.AddVertex(FSlateVertex::Make<Rounding>(RenderTransform, Position + CapOutward - Up, FVector2D(-1.0f, 1.0f), Color));
+			RenderBatch.AddVertex(FSlateVertex::Make<Rounding>(RenderTransform, FVector2f(Position + CapOutward + Up), FVector2f(1.0f, 1.0f), Color));
+			RenderBatch.AddVertex(FSlateVertex::Make<Rounding>(RenderTransform, FVector2f(Position + CapOutward - Up), FVector2f(-1.0f, 1.0f), Color));
 			AddQuadIndices(RenderBatch);
 		}
 	}
@@ -2111,12 +2111,12 @@ void FSlateElementBatcher::AddLineElement( const FSlateDrawElement& DrawElement
 			FSlateRenderBatch& RenderBatch = CreateRenderBatch(Layer, FShaderParams(), nullptr, ESlateDrawPrimitive::LineList, ESlateShader::Default, DrawEffects, ESlateBatchDrawFlag::None, DrawElement);
 
 			const FColor StartColor = PackedColors.Num() ? PackedColors[0] : PackedTint;
-			RenderBatch.AddVertex(FSlateVertex::Make<Rounding>(RenderTransform, Points[0], FVector2D::ZeroVector, StartColor));
+			RenderBatch.AddVertex(FSlateVertex::Make<Rounding>(RenderTransform, FVector2f(Points[0]), FVector2f::ZeroVector, StartColor));
 
 			for (int32 Point = 1; Point < NumPoints; ++Point)
 			{
 				const FColor PointColor = PackedColors.Num() ? PackedColors[Point] : PackedTint;
-				RenderBatch.AddVertex(FSlateVertex::Make<Rounding>(RenderTransform, Points[Point], FVector2D::ZeroVector, PointColor));
+				RenderBatch.AddVertex(FSlateVertex::Make<Rounding>(RenderTransform, FVector2f(Points[Point]), FVector2f::ZeroVector, PointColor));
 
 				const uint32 IndexStart = RenderBatch.GetNumVertices();
 				RenderBatch.AddIndex(IndexStart - 2);
