Slate: Fix Num() that is not called every loop like previoulsy. Some widgets add/remove children while they are beeing iterating on.

diff --git a/Engine/Source/Runtime/SlateCore/Public/Layout/ChildrenBase.h b/Engine/Source/Runtime/SlateCore/Public/Layout/ChildrenBase.h
--- a/Engine/Source/Runtime/SlateCore/Public/Layout/ChildrenBase.h
+++ b/Engine/Source/Runtime/SlateCore/Public/Layout/ChildrenBase.h
@@ -11,9 +11,6 @@
 class SWidget;
 class FSlotBase;
 
-#ifndef UE_SLATE_WITH_CHILD_SLOT_DEBUGGING
-	#define UE_SLATE_WITH_CHILD_SLOT_DEBUGGING 0
-#endif // !UE_SLATE_WITH_CHILD_SLOT_DEBUGGING
 
 /**
  * FChildren is an interface that must be implemented by all child containers.
@@ -63,19 +60,10 @@ public:
 #if WITH_SLATE_DEBUGGING 
 		TGuardValue<bool> IteratingGuard(bIsIteratingChildren, true);
 #endif
-
-		int32 WidgetCount = Num();
-		for (int32 Index = 0; Index < WidgetCount; ++Index)
+		for (int32 Index = 0; Index < Num(); ++Index)
 		{
 			FWidgetRef WidgetRef = GetChildRefAt(Index);
 			Pred(WidgetRef.GetWidget());
-#if UE_SLATE_WITH_CHILD_SLOT_DEBUGGING 
-			UE_CLOG(WidgetCount != Num(), LogSlate, Error,
-				TEXT("Num Children changed during iteration! Child: %s, Owner: %s [%s]"),
-				*FReflectionMetaData::GetWidgetDebugInfo(WidgetRef.GetWidget()),
-				*FReflectionMetaData::GetWidgetDebugInfo(Owner),
-				*Name.ToString());
-#endif // UE_SLATE_WITH_CHILD_SLOT_DEBUGGING
 		}
 	}
 
@@ -87,18 +75,10 @@ public:
 		TGuardValue<bool> IteratingGuard(bIsIteratingChildren, true);
 #endif
 
-		int32 WidgetCount = Num();
-		for (int32 Index = 0; Index < WidgetCount; ++Index)
+		for (int32 Index = 0; Index < Num(); ++Index)
 		{
 			FConstWidgetRef WidgetRef = GetChildRefAt(Index);
 			Pred(WidgetRef.GetWidget());
-#if UE_SLATE_WITH_CHILD_SLOT_DEBUGGING 
-			UE_CLOG(WidgetCount != Num(), LogSlate, Error,
-				TEXT("Num Children changed during iteration! Child: %s, Owner: %s [%s]"),
-				*FReflectionMetaData::GetWidgetDebugInfo(WidgetRef.GetWidget()),
-				*FReflectionMetaData::GetWidgetDebugInfo(Owner),
-				*Name.ToString());
-#endif // UE_SLATE_WITH_CHILD_SLOT_DEBUGGINGs
 		}
 	}
 
