Fix PS4 minimal incremental

diff --git a/Engine/Source/Runtime/SlateCore/Private/Rendering/DrawElements.cpp b/Engine/Source/Runtime/SlateCore/Private/Rendering/DrawElements.cpp
--- a/Engine/Source/Runtime/SlateCore/Private/Rendering/DrawElements.cpp
+++ b/Engine/Source/Runtime/SlateCore/Private/Rendering/DrawElements.cpp
@@ -156,7 +156,7 @@ void FSlateDrawElement::Init(FSlateWindowElementList& ElementList, EElementType
 	ClipStateHandle.SetPreCachedClipIndex(ElementList.GetClippingIndex());
 
 #if UE_SLATE_VERIFY_PIXELSIZE
-	const FVector2D PixelSize = (LocalSize * Scale);
+	const FVector2f PixelSize = (LocalSize * Scale);
 	ensureMsgf(PixelSize.X >= 0.f && PixelSize.X <= (float)std::numeric_limits<uint16>::max(), TEXT("The size X '%f' is too small or big to fit in the SlateVertex buffer."), PixelSize.X);
 	ensureMsgf(PixelSize.Y >= 0.f && PixelSize.Y <= (float)std::numeric_limits<uint16>::max(), TEXT("The size Y '%f' is too small or big to fit in the SlateVertex buffer."), PixelSize.Y);
 #endif
