Slate: Update the LayerId in invalidation when the widget painted more elements than the previous paint. Disable by default. You can enable it with the cvar `Slate.InvalidationList.EnableReindexLayerId. This occurs when a child-widget was collapsed and is now visible. The sibling of the widget may have the same LayerId as the child-widget. That would affect how the elements are rendered.
#preglight 61f94bb6a6632a34f372935a

diff --git a/Engine/Source/Runtime/Slate/Private/Framework/Application/SlateApplication.cpp b/Engine/Source/Runtime/Slate/Private/Framework/Application/SlateApplication.cpp
--- a/Engine/Source/Runtime/Slate/Private/Framework/Application/SlateApplication.cpp
+++ b/Engine/Source/Runtime/Slate/Private/Framework/Application/SlateApplication.cpp
@@ -79,15 +79,12 @@ static FAutoConsoleVariableRef CVarSlateVerifyWidgetLayerId(
 	TEXT("Every tick, verify that widgets have a LayerId range that fits with their siblings and their parent.")
 );
 
-namespace UE
-{
-namespace Slate
+namespace UE::Slate::Private
 {
 	void VerifyParentChildrenRelationship(const TSharedRef<SWindow>& WindowToDraw);
 	void VerifyWidgetLayerId(const TSharedRef<SWindow>& WindowToDraw);
 }
-}
-#endif
+#endif //WITH_SLATE_DEBUGGING
 
 //////////////////////////////////////////////////////////////////////////
 
@@ -1080,7 +1077,7 @@ void FSlateApplication::DrawWindowAndChildren( const TSharedRef<SWindow>& Window
 #if WITH_SLATE_DEBUGGING
 	if (GSlateVerifyParentChildrenRelationship)
 	{
-		UE::Slate::VerifyParentChildrenRelationship(WindowToDraw);
+		UE::Slate::Private::VerifyParentChildrenRelationship(WindowToDraw);
 	}
 #endif
 
@@ -1146,7 +1143,7 @@ void FSlateApplication::DrawWindowAndChildren( const TSharedRef<SWindow>& Window
 #if WITH_SLATE_DEBUGGING
 		if (GSlateVerifyWidgetLayerId)
 		{
-			UE::Slate::VerifyWidgetLayerId(WindowToDraw);
+			UE::Slate::Private::VerifyWidgetLayerId(WindowToDraw);
 		}
 #endif
 
@@ -7107,74 +7104,82 @@ bool FSlateApplication::InputPreProcessorsHelper::PreProcessInput(ESlateDebuggin
 }
 
 
-namespace UE
-{
-namespace Slate
-{
 #if WITH_SLATE_DEBUGGING
-namespace Private
+namespace UE::Slate::Private
 {
-	bool VerifyParentChildrenRelationship_Recursive(SWidget* Parent, TMap<const SWidget*, const SWidget*>& AllWidgets)
-	{
-		bool bResult = true;
-		Parent->GetAllChildren()->ForEachWidget([&bResult, Parent, &AllWidgets](SWidget& ChildWidget)
+
+bool VerifyParentChildrenRelationship_Recursive(SWidget* Parent, TMap<const SWidget*, const SWidget*>& AllWidgets)
+{
+	bool bResult = true;
+	Parent->GetAllChildren()->ForEachWidget([&bResult, Parent, &AllWidgets](SWidget& ChildWidget)
+		{
+			if (&ChildWidget != &SNullWidget::NullWidget.Get())
 			{
-				if (&ChildWidget != &SNullWidget::NullWidget.Get())
+				if (AllWidgets.Find(&ChildWidget))
 				{
-					if (AllWidgets.Find(&ChildWidget))
-					{
-						bResult = false;
-						UE_LOG(LogSlate, Warning, TEXT("The widget '%s' is owned by more than one parent. 1:'%s' 2:'%s'.")
-							, *FReflectionMetaData::GetWidgetDebugInfo(ChildWidget)
-							, *FReflectionMetaData::GetWidgetDebugInfo(AllWidgets[&ChildWidget])
-							, *FReflectionMetaData::GetWidgetDebugInfo(Parent));
-					}
-					else
-					{
-						AllWidgets.Add(&ChildWidget, Parent);
-					}
-
-					if (ChildWidget.GetParentWidget().Get() != Parent)
-					{
-						bResult = false;
-						UE_LOG(LogSlate, Warning, TEXT("The widget '%s' has the wrong parent."), *FReflectionMetaData::GetWidgetDebugInfo(ChildWidget));
-					}
+					bResult = false;
+					UE_LOG(LogSlate, Warning, TEXT("The widget '%s' is owned by more than one parent. 1:'%s' 2:'%s'.")
+						, *FReflectionMetaData::GetWidgetDebugInfo(ChildWidget)
+						, *FReflectionMetaData::GetWidgetDebugInfo(AllWidgets[&ChildWidget])
+						, *FReflectionMetaData::GetWidgetDebugInfo(Parent));
+				}
+				else
+				{
+					AllWidgets.Add(&ChildWidget, Parent);
+				}
 
-					bResult = bResult && VerifyParentChildrenRelationship_Recursive(&ChildWidget, AllWidgets);
+				if (ChildWidget.GetParentWidget().Get() != Parent)
+				{
+					bResult = false;
+					UE_LOG(LogSlate, Warning, TEXT("The widget '%s' has the wrong parent."), *FReflectionMetaData::GetWidgetDebugInfo(ChildWidget));
 				}
-			});
 
-		return bResult;
-	}
+				bResult = bResult && VerifyParentChildrenRelationship_Recursive(&ChildWidget, AllWidgets);
+			}
+		});
 
-	bool VerifyWidgetLayerId_Recursive(SWidget& Widget)
-	{
-		bool bResult = true;
+	return bResult;
+}
 
-		const uint32 LastPaintFrame = Widget.Debug_GetLastPaintFrame();
-		const bool bIsDeferredPaint = Widget.GetPersistentState().bDeferredPainting;
-		const int32 InLayerId = Widget.GetPersistentState().LayerId;
-		const int32 OutLayerId = Widget.GetPersistentState().OutgoingLayerId;
+bool VerifyWidgetLayerId_Recursive(SWidget& Widget, bool bInsideInvalidationRoot)
+{
+	bool bResult = true;
+
+	bInsideInvalidationRoot = bInsideInvalidationRoot || Widget.Advanced_IsInvalidationRoot();
+	const uint32 LastPaintFrame = Widget.Debug_GetLastPaintFrame();
+	const bool bIsDeferredPaint = Widget.GetPersistentState().bDeferredPainting;
+	const int32 InLayerId = Widget.GetPersistentState().LayerId;
+	const int32 OutLayerId = Widget.GetPersistentState().OutgoingLayerId;
 
-		int32 PreviousInLayerId = InLayerId;
-		int32 PreviousOutLayerId = InLayerId;
+	int32 PreviousInLayerId = InLayerId;
+	int32 PreviousOutLayerId = InLayerId;
 
-		Widget.GetAllChildren()->ForEachWidget([&bResult, LastPaintFrame, bIsDeferredPaint, InLayerId, OutLayerId](SWidget& ChildWidget)
+	Widget.GetAllChildren()->ForEachWidget([&bResult, bInsideInvalidationRoot, LastPaintFrame, bIsDeferredPaint, InLayerId, OutLayerId](SWidget& ChildWidget)
+		{
+			if (&ChildWidget != &SNullWidget::NullWidget.Get())
 			{
-				if (&ChildWidget != &SNullWidget::NullWidget.Get())
+
+				if (ChildWidget.GetVisibility().IsVisible() && (bInsideInvalidationRoot || ChildWidget.Debug_GetLastPaintFrame() == LastPaintFrame))
 				{
-					if (ChildWidget.GetVisibility().IsVisible() && ChildWidget.Debug_GetLastPaintFrame() == LastPaintFrame)
+					const bool bIsChildDeferredPaint = ChildWidget.GetPersistentState().bDeferredPainting;
+					if (bIsChildDeferredPaint == bIsDeferredPaint)
 					{
-						const bool bIsChildDeferredPaint = ChildWidget.GetPersistentState().bDeferredPainting;
-						if (bIsChildDeferredPaint == bIsDeferredPaint)
-						{
-							const int32 ChildInLayerId = ChildWidget.GetPersistentState().LayerId;
-							const int32 ChildOutLayerId = ChildWidget.GetPersistentState().OutgoingLayerId;
+						const int32 ChildInLayerId = ChildWidget.GetPersistentState().LayerId;
+						const int32 ChildOutLayerId = ChildWidget.GetPersistentState().OutgoingLayerId;
 
-							const bool bLayerInIsValid = InLayerId <= ChildInLayerId && InLayerId <= ChildOutLayerId;
-							const bool bLayerOutIsValid = OutLayerId >= ChildInLayerId && OutLayerId >= ChildOutLayerId;
+						if (ChildInLayerId== 0)
+						{
+							return;
+						}
 
-							if (!bLayerInIsValid || !bLayerOutIsValid)
+						const bool bLayerInIsValid = InLayerId <= ChildInLayerId && InLayerId <= ChildOutLayerId;
+						const bool bLayerOutIsValid = OutLayerId >= ChildInLayerId && OutLayerId >= ChildOutLayerId;
+						if (!bLayerInIsValid || !bLayerOutIsValid)
+						{
+							SWidget* ParentWidget = ChildWidget.GetParentWidget().Get();
+							check(ParentWidget);
+							// The parent may just have tick and will be painted on the next frame (this is not desired but possible.
+							if (!ParentWidget->GetProxyHandle().HasAnyInvalidationReason(ParentWidget, EInvalidateWidgetReason::Paint))
 							{
 								bResult = false;
 								UE_LOG(LogSlate, Warning, TEXT("The widget '%s' LayerId is invalid. Parent: [%d,%d] Child: [%d,%d].")
@@ -7184,17 +7189,18 @@ namespace Private
 									, ChildInLayerId
 									, ChildOutLayerId);
 							}
-						}
 
-						bResult = bResult && VerifyWidgetLayerId_Recursive(ChildWidget);
+							return;
+						}
 					}
-				}
-			});
 
-		return bResult;
-	}
-}//Private
+					bResult = bResult && VerifyWidgetLayerId_Recursive(ChildWidget, bInsideInvalidationRoot);
+				}
+			}
+		});
 
+	return bResult;
+}
 
 void VerifyParentChildrenRelationship(const TSharedRef<SWindow>& WindowToDraw)
 {
@@ -7204,8 +7210,8 @@ void VerifyParentChildrenRelationship(const TSharedRef<SWindow>& WindowToDraw)
 		AllWidgets.Add(&WindowToDraw.Get(), nullptr);
 		if (!Private::VerifyParentChildrenRelationship_Recursive(&WindowToDraw.Get(), AllWidgets))
 		{
-			CVarSlateVerifyParentChildrenRelationship->Set(false);
-			ensureMsgf(false, TEXT("VerifyParentChildrenRelationship failed. See log for more info."));
+			CVarSlateVerifyParentChildrenRelationship->Set(false, CVarSlateVerifyParentChildrenRelationship->GetFlags());
+			ensureAlwaysMsgf(false, TEXT("VerifyParentChildrenRelationship failed. See log for more info."));
 		}
 	}
 }
@@ -7215,16 +7221,15 @@ void VerifyWidgetLayerId(const TSharedRef<SWindow>& WindowToDraw)
 {
 	if (WindowToDraw != SNullWidget::NullWidget)
 	{
-		if (!Private::VerifyWidgetLayerId_Recursive(WindowToDraw.Get()))
+		if (!Private::VerifyWidgetLayerId_Recursive(WindowToDraw.Get(), false))
 		{
-			CVarSlateVerifyWidgetLayerId->Set(false);
-			ensureMsgf(false, TEXT("VerifyWidgetLayerId failed. See log for more info."));
+			CVarSlateVerifyWidgetLayerId->Set(false, CVarSlateVerifyWidgetLayerId->GetFlags());
+			ensureAlwaysMsgf(false, TEXT("VerifyWidgetLayerId failed. See log for more info."));
 		}
-	}
-}
 
 
-#endif
-}// Slate
-}//UE
+	}
+}
 
+} //namespace
+#endif //WITH_SLATE_DEBUGGING
diff --git a/Engine/Source/Runtime/SlateCore/Private/FastUpdate/SlateInvalidationRoot.cpp b/Engine/Source/Runtime/SlateCore/Private/FastUpdate/SlateInvalidationRoot.cpp
--- a/Engine/Source/Runtime/SlateCore/Private/FastUpdate/SlateInvalidationRoot.cpp
+++ b/Engine/Source/Runtime/SlateCore/Private/FastUpdate/SlateInvalidationRoot.cpp
@@ -5,7 +5,6 @@
 #include "FastUpdate/SlateInvalidationRootList.h"
 #include "FastUpdate/SlateInvalidationWidgetHeap.h"
 #include "FastUpdate/SlateInvalidationWidgetList.h"
-#include "FastUpdate/SlateInvalidationWidgetSortOrder.h"
 #include "Async/TaskGraphInterfaces.h"
 #include "Application/SlateApplicationBase.h"
 #include "Widgets/SWidget.h"
@@ -27,7 +26,7 @@ static FAutoConsoleVariableRef CVarSlateInvalidationRootDumpUpdateList(
 	GSlateInvalidationRootDumpUpdateList,
 	TEXT("Each frame, log the widgets that will be updated.")
 );
-void DumpUpdateList(const FSlateInvalidationWidgetList& FastWidgetPathList, const TArray<FSlateInvalidationWidgetIndex>&);
+void DumpUpdateList(const FSlateInvalidationWidgetList& FastWidgetPathList, const TArray<FSlateInvalidationWidgetHeapElement>&);
 
 bool GSlateInvalidationRootDumpUpdateListOnce = false;
 static FAutoConsoleVariableRef CVarSlateInvalidationRootDumpUpdateListOnce(
@@ -112,7 +111,7 @@ static FAutoConsoleVariableRef CVarSlateInvalidationRootVerifyWidgetVolatile(
 	GSlateInvalidationRootVerifyWidgetVolatile,
 	TEXT("Every tick, verify that volatile widgets are mark properly and are in the correct list.")
 );
-void VerifyWidgetVolatile(FSlateInvalidationWidgetList& WidgetList, TArray<FSlateInvalidationWidgetIndex>& FinalUpdateList);
+void VerifyWidgetVolatile(FSlateInvalidationWidgetList& WidgetList, TArray<FSlateInvalidationWidgetHeapElement>& FinalUpdateList);
 
 bool GSlateInvalidationRootVerifyWidgetsUpdateList = false;
 static FAutoConsoleVariableRef CVarSlateInvalidationRootVerifyWidgetsUpdateList(
@@ -120,10 +119,10 @@ static FAutoConsoleVariableRef CVarSlateInvalidationRootVerifyWidgetsUpdateList(
 	GSlateInvalidationRootVerifyWidgetsUpdateList,
 	TEXT("Every tick, verify that pre and post update list contains the correct information and they are sorted.")
 );
-void VerifyWidgetsUpdateList_BeforeProcessPreUpdate(const TSharedRef<SWidget>&, FSlateInvalidationWidgetList*, FSlateInvalidationWidgetPreHeap*, FSlateInvalidationWidgetPostHeap*, TArray<FSlateInvalidationWidgetIndex>&);
+void VerifyWidgetsUpdateList_BeforeProcessPreUpdate(const TSharedRef<SWidget>&, FSlateInvalidationWidgetList*, FSlateInvalidationWidgetPreHeap*, FSlateInvalidationWidgetPostHeap*, TArray<FSlateInvalidationWidgetHeapElement>&);
 void VerifyWidgetsUpdateList_ProcessPrepassUpdate(FSlateInvalidationWidgetList*, FSlateInvalidationWidgetPrepassHeap*, FSlateInvalidationWidgetPostHeap*);
-void VerifyWidgetsUpdateList_BeforeProcessPostUpdate(const TSharedRef<SWidget>&, FSlateInvalidationWidgetList*, FSlateInvalidationWidgetPreHeap*, FSlateInvalidationWidgetPostHeap*, TArray<FSlateInvalidationWidgetIndex>&);
-void VerifyWidgetsUpdateList_AfterProcessPostUpdate(const TSharedRef<SWidget>&, FSlateInvalidationWidgetList*, FSlateInvalidationWidgetPreHeap*, FSlateInvalidationWidgetPostHeap*, TArray<FSlateInvalidationWidgetIndex>&);
+void VerifyWidgetsUpdateList_BeforeProcessPostUpdate(const TSharedRef<SWidget>&, FSlateInvalidationWidgetList*, FSlateInvalidationWidgetPreHeap*, FSlateInvalidationWidgetPostHeap*, TArray<FSlateInvalidationWidgetHeapElement>&);
+void VerifyWidgetsUpdateList_AfterProcessPostUpdate(const TSharedRef<SWidget>&, FSlateInvalidationWidgetList*, FSlateInvalidationWidgetPreHeap*, FSlateInvalidationWidgetPostHeap*, TArray<FSlateInvalidationWidgetHeapElement>&);
 
 bool GSlateInvalidationRootVerifySlateAttribute = false;
 static FAutoConsoleVariableRef CVarSlateInvalidationRootVerifySlateAttribute(
@@ -134,6 +133,13 @@ static FAutoConsoleVariableRef CVarSlateInvalidationRootVerifySlateAttribute(
 void VerifySlateAttribute_BeforeUpdate(FSlateInvalidationWidgetList& FastWidgetPathList);
 void VerifySlateAttribute_AfterUpdate(const FSlateInvalidationWidgetList& FastWidgetPathList);
 
+bool GSlateInvalidationRootVerifyWidgetsAreUpdatedOnce = false;
+static FAutoConsoleVariableRef CVarSlateInvalidationRootVerifyWidgetsAreUpdatedOnce(
+	TEXT("Slate.InvalidationRoot.VerifyWidgetsAreUpdatedOnce"),
+	GSlateInvalidationRootVerifyWidgetsAreUpdatedOnce,
+	TEXT("Verify that the widgets are painted only once per tick.")
+);
+
 #endif //UE_SLATE_WITH_INVALIDATIONWIDGETLIST_DEBUGGING
 
 
@@ -142,13 +148,19 @@ int32 GSlateInvalidationWidgetListMaxArrayElements = 64;
 FAutoConsoleVariableRef CVarSlateInvalidationWidgetListMaxArrayElements(
 	TEXT("Slate.InvalidationList.MaxArrayElements"),
 	GSlateInvalidationWidgetListMaxArrayElements,
-	TEXT("With Global Invalidation, the preferred size of the elements array."));
+	TEXT("With the invalidation system, the preferred size of the elements array."));
 
 int32 GSlateInvalidationWidgetListNumberElementLeftBeforeSplitting = 40;
 FAutoConsoleVariableRef CVarSlateInvalidationWidgetListNumElementLeftBeforeSplitting(
 	TEXT("Slate.InvalidationList.NumberElementLeftBeforeSplitting"),
 	GSlateInvalidationWidgetListNumberElementLeftBeforeSplitting,
-	TEXT("With Global Invalidation, when splitting, only split the array when the number of element left is under X."));
+	TEXT("With the invalidation system, when splitting, only split the array when the number of element left is under X."));
+
+bool GSlateInvalidationEnableReindexLayerId = false;
+FAutoConsoleVariableRef CVarSlateInvalidationEnableReindexLayerId(
+	TEXT("Slate.InvalidationList.EnableReindexLayerId"),
+	GSlateInvalidationEnableReindexLayerId,
+	TEXT("With invalidation system, when a painted widget returns a bigger LayerId that it used to, re-index the other widgets."));
 
 /**
  *
@@ -298,9 +310,9 @@ void FSlateInvalidationRoot::InvalidateWidget(FWidgetProxy& Proxy, EInvalidateWi
 			WidgetsNeedingPreUpdate->HeapPushUnique(Proxy);
 		}
 
-		if (!bProcessingPrepassUpdate && EnumHasAnyFlags(InvalidateReason, EInvalidateWidgetReason::Prepass))
+		if (!bProcessingPrepassUpdate && !bProcessingPostUpdate && EnumHasAnyFlags(InvalidateReason, EInvalidateWidgetReason::Prepass))
 		{
-			ensureMsgf(Proxy.GetWidget() && Proxy.GetWidget()->NeedsPrepass(), TEXT("A Prepass invalidation occurs and the widget doesn't hae the NeedsPrepass flag."));
+			ensureMsgf(Proxy.GetWidget() && Proxy.GetWidget()->NeedsPrepass(), TEXT("A Prepass invalidation occurs and the widget doesn't have the NeedsPrepass flag."));
 			ensureMsgf(EnumHasAnyFlags(InvalidateReason, EInvalidateWidgetReason::Layout), TEXT("The Prepass invalidation should include a Layout invalidation."));
 			WidgetsNeedingPrepassUpdate->PushBackUnique(Proxy);
 		}
@@ -433,61 +445,335 @@ void FSlateInvalidationRoot::OnWidgetDestroyed(const SWidget* Widget)
 	}
 }
 
-bool FSlateInvalidationRoot::PaintFastPath(const FSlateInvalidationContext& Context)
+void FSlateInvalidationRoot::PaintFastPath_AddUniqueSortedToFinalUpdateList(const FSlateInvalidationWidgetIndex InvalidationWidgetIndex)
 {
-	SCOPED_NAMED_EVENT(SWidget_FastPathUpdate, FColor::Green);
-	CSV_SCOPED_TIMING_STAT(Slate, FastPathUpdate);
+	const FSlateInvalidationWidgetSortOrder SortIndex{ *FastWidgetPathList, InvalidationWidgetIndex };
+	int32 BackwardIndex = FinalUpdateList.Num() - 1;
+	for (; BackwardIndex >= 0; --BackwardIndex)
+	{
+		const FSlateInvalidationWidgetHeapElement NextWidgetElement = FinalUpdateList[BackwardIndex];
+		if (NextWidgetElement.GetWidgetIndex() == InvalidationWidgetIndex)
+		{
+			// No need to insert
+			return;
+		}
 
-	check(!bNeedsSlowPath);
+		if (NextWidgetElement.GetWidgetSortOrder() > SortIndex)
+		{
+			break;
+		}
+	}
 
-	bool bWidgetsNeededRepaint = false;
+	FinalUpdateList.EmplaceAt(BackwardIndex + 1, InvalidationWidgetIndex, SortIndex);
+}
+
+namespace UE::Slate::Private
+{
+	struct FSlateInvalidationReindexHeap
 	{
-#if WITH_SLATE_DEBUGGING
-		if (GSlateInvalidationRootDumpUpdateList || GSlateInvalidationRootDumpUpdateListOnce)
+		FSlateInvalidationReindexHeap() = default;
+
+		struct FWidgetOrderGreater
 		{
-			DumpUpdateList(*FastWidgetPathList, FinalUpdateList);
+			FORCEINLINE bool operator()(const FSlateInvalidationWidgetHeapElement& A, const FSlateInvalidationWidgetHeapElement& B) const
+			{
+				return A.GetWidgetSortOrder() < B.GetWidgetSortOrder();
+			}
+		};
+
+		void HeapPush(FSlateInvalidationWidgetHeapElement Element)
+		{
+			Heap.HeapPush(Element, FWidgetOrderGreater());
 		}
-#endif
 
-		TGuardValue<bool> OnFastPathGuard(GSlateIsOnFastUpdatePath, true);
-		FSlateInvalidationWidgetList::FIndexRange PreviousPaintedWidgetRange;
+		UE_NODISCARD FSlateInvalidationWidgetHeapElement HeapPeek()
+		{
+			return Heap.HeapTop();
+		}
 
-		// The update list is put in reverse order by ProcessInvalidation
-		for (int32 ListIndex = FinalUpdateList.Num() - 1; ListIndex >= 0; --ListIndex)
+		void HeapPopDiscard()
+		{
+			Heap.HeapPopDiscard(FWidgetOrderGreater(), false);
+		}
+
+		UE_NODISCARD bool IsEmpty() const
+		{
+			return Heap.Num() == 0;
+		}
+
+	private:
+		TArray<FSlateInvalidationWidgetHeapElement, TInlineAllocator<16>> Heap;
+	};
+
+	struct FSlateInvalidationPaintFastPathContext
+	{
+		FSlateInvalidationReindexHeap ReindexUpdateList;
+	};
+}
+
+void FSlateInvalidationRoot::PaintFastPath_FixupParentLayerId(UE::Slate::Private::FSlateInvalidationPaintFastPathContext& FastPathContext, const FWidgetProxy& InvalidationWidget, const int32 NewOutgoingLayerId)
+{
+	// Update the parent outgoinglayerid and rerun the algo for it... we don't seems to need it maybe we can remove it for good.
+	const FSlateInvalidationWidgetIndex ParentInvalidationWidgetIndex = InvalidationWidget.ParentIndex;
+	if (ParentInvalidationWidgetIndex != FSlateInvalidationWidgetIndex::Invalid)
+	{
+		const FSlateInvalidationWidgetList::InvalidationWidgetType& ParentInvalidationWidget = (*FastWidgetPathList)[ParentInvalidationWidgetIndex];
+		SWidget* ParentWidgetPtr = ParentInvalidationWidget.GetWidget();
+
+		check(ParentWidgetPtr); // if the child is valid, then the parent must be valid
+		if (ParentWidgetPtr->GetPersistentState().OutgoingLayerId < NewOutgoingLayerId)
+		{
+			const_cast<FSlateWidgetPersistentState&>(ParentWidgetPtr->GetPersistentState()).OutgoingLayerId = NewOutgoingLayerId;
+			FastPathContext.ReindexUpdateList.HeapPush(FSlateInvalidationWidgetHeapElement{ ParentInvalidationWidgetIndex, FSlateInvalidationWidgetSortOrder{*FastWidgetPathList, ParentInvalidationWidgetIndex} });
+		}
+	}
+}
+
+void FSlateInvalidationRoot::PaintFastPath_FixupLayerId(UE::Slate::Private::FSlateInvalidationPaintFastPathContext& FastPaintContext, const FWidgetProxy& InvalidationWidget, const int32 NewOutgoingLayerId)
+{
+	// We ignore all deferred painting for now.
+	if (InvalidationWidget.GetWidget()->GetPersistentState().bDeferredPainting)
+	{
+		return;
+	}
+
+	// For all sibling update the LayerId.
+	//Option 1: The next sibling already needs a paint update. If the next sibling incremented the LayerId (most Widget), then just make sure the next LayerId is big enough. It will update the parent on its turn.
+	//Option 2: The next sibling already needs a paint update. If the next sibling has the same LayerId (SVerticalBox/SConstraintCanvas), then we need to update the parent outgoing.
+	//Option 3: The next sibling does not need a paint update. If the next sibling incremented the LayerId (most Widget), then re index the sibling (repaint it). It will update the parent on its turn.
+	//Option 4: The next sibling does not need a paint update. If the next sibling has the same LayerId (SVerticalBox/SConstraintCanvas), then we need to update the parent outgoing and continue to the next sibling until one needs a paint update.
+
+	bool bFoundAtLeastOneSibling = false;
+	bool bFixupParentLayerId = false;
+	const int32 LayerId = InvalidationWidget.GetWidget()->GetPersistentState().LayerId;
+
+	FSlateInvalidationWidgetList::FIndexRange ParentWidgetRange;
+	{
+		if (InvalidationWidget.ParentIndex != FSlateInvalidationWidgetIndex::Invalid)
+		{
+			const FSlateInvalidationWidgetList::InvalidationWidgetType& ParentInvalidationWidget = (*FastWidgetPathList)[InvalidationWidget.ParentIndex];
+			ParentWidgetRange = { *FastWidgetPathList, ParentInvalidationWidget.Index, ParentInvalidationWidget.LeafMostChildIndex };
+		}
+	}
+
+	FSlateInvalidationWidgetIndex NextSiblingIndex = FastWidgetPathList->FindNextSibling(InvalidationWidget.Index);
+	while (NextSiblingIndex != FSlateInvalidationWidgetIndex::Invalid)
+	{
+		const FSlateInvalidationWidgetList::InvalidationWidgetType& NextSiblingInvalidationWidget = (*FastWidgetPathList)[NextSiblingIndex];
+		SWidget* NextSiblingWidgetPtr = NextSiblingInvalidationWidget.GetWidget();
+		if (NextSiblingInvalidationWidget.Visibility.IsVisible() && NextSiblingWidgetPtr)
 		{
-			const FSlateInvalidationWidgetIndex MyIndex = FinalUpdateList[ListIndex];
-			if (PreviousPaintedWidgetRange.IsValid())
+			// We ignore all deferred painting for now.
+			if (!NextSiblingWidgetPtr->GetPersistentState().bDeferredPainting)
 			{
-				// It's already been processed by the previous draw
-				if (PreviousPaintedWidgetRange.Include(FSlateInvalidationWidgetSortOrder{*FastWidgetPathList , MyIndex}))
+				bFoundAtLeastOneSibling = true;
+
+				// Is this widget going to be painted anyway (already in the list). If so, just update it's internal LayerId.
+				//Because of the Visible and GetWidget are test, we cannot only test the last item from the FinaUpdateList.
+				bool bNextSiblingWillBePainted = false;
+				for (int32 BackwardIndex = FinalUpdateList.Num() - 1; BackwardIndex >= 0; --BackwardIndex)
+				{
+					const FSlateInvalidationWidgetHeapElement WidgetElement = FinalUpdateList[BackwardIndex];
+					if (WidgetElement.GetWidgetIndex() == NextSiblingIndex)
+					{
+						bNextSiblingWillBePainted = NextSiblingWidgetPtr->HasAnyUpdateFlags(EWidgetUpdateFlags::NeedsRepaint | EWidgetUpdateFlags::NeedsVolatilePaint);
+						break;
+					}
+					// We can stop if they are not in the same hierarchy.
+					if (ParentWidgetRange.IsValid() && !ParentWidgetRange.Include(WidgetElement.GetWidgetSortOrder()))
+					{
+						break;
+					}
+				}
+
+				if (bNextSiblingWillBePainted)
 				{
-					continue;
+					// Need to update the parent OuterLayerId in case we are using SVerticalBox/SConstraintCanvas (that reuse the same LayerId to improve rendering performance).
+					if (NextSiblingWidgetPtr->GetPersistentState().LayerId == LayerId)
+					{
+						// Push the parent update, the new LayerId. The algo will paint the next sibling. Then confirm (on the next loop that the parent is ok).
+					}
+					else if (NextSiblingWidgetPtr->GetPersistentState().LayerId < NewOutgoingLayerId + 1)
+					{
+						// The next sibling should also update the parent but in case the parent do a FMath::Max() logic, we do not take any chances.
+						// Make sure the LayerId is correct for the next sibling.
+						const_cast<FSlateWidgetPersistentState&>(NextSiblingWidgetPtr->GetPersistentState()).LayerId = NewOutgoingLayerId + 1;
+					}
+					break;
+				}
+				else if (NextSiblingWidgetPtr->GetPersistentState().LayerId == LayerId)
+				{
+					// We cannot assume that every other siblings will also have the same LayerId. We need to loop and update for every next sibling.
+					// Keep looping for the next sibling.
+				}
+				else if (NextSiblingWidgetPtr->GetPersistentState().LayerId < NewOutgoingLayerId + 1)
+				{
+					// We are re-indexing this sibling. It should also update the parent but in case the parent do a FMath::Max() logic, we do not take any chances.
+					// Re-indexing the this sibling (for now we are just painting it).
+					//NB, if we do a property re-indexing algo, we need to update this sibling LayerId and all the children (recursively) of this sibling.
+					const_cast<FSlateWidgetPersistentState&>(NextSiblingWidgetPtr->GetPersistentState()).LayerId = NewOutgoingLayerId + 1;
+					NextSiblingWidgetPtr->UpdateFlags |= EWidgetUpdateFlags::NeedsRepaint;
+
+					PaintFastPath_AddUniqueSortedToFinalUpdateList(NextSiblingIndex);
+					break;
+				}
+				else
+				{
+					// The Parent put gabs in the LayerId (like SOverlay) we can stop looking for issues.
+					break;
 				}
 			}
+		}
+		NextSiblingIndex = FastWidgetPathList->FindNextSibling(NextSiblingIndex);
+	}
+
+	// Update or no update. Check if the parent needs an update.
+	PaintFastPath_FixupParentLayerId(FastPaintContext, InvalidationWidget, NewOutgoingLayerId);
+}
+
+bool FSlateInvalidationRoot::PaintFastPath_UpdateNextWidget(const FSlateInvalidationContext& Context, UE::Slate::Private::FSlateInvalidationPaintFastPathContext& FastPaintContext)
+{
+	bool bNeedsPaint = false;
+	const bool bAllowShrinking = false;
+	const FSlateInvalidationWidgetIndex MyIndex = FinalUpdateList.Pop(bAllowShrinking).GetWidgetIndex();
+
+	FSlateInvalidationWidgetList::InvalidationWidgetType& InvalidationWidget = (*FastWidgetPathList)[MyIndex];
+	SWidget* WidgetPtr = InvalidationWidget.GetWidget();
+
+#if UE_SLATE_WITH_INVALIDATIONWIDGETLIST_DEBUGGING
+	if (GSlateInvalidationRootVerifyWidgetsAreUpdatedOnce)
+	{
+		ensureAlwaysMsgf(!InvalidationWidget.bDebug_Updated, TEXT("VerifyWidgetsAreUpdatedOnce failed. Widget '%s' is going to be updated more than once"), *FReflectionMetaData::GetWidgetDebugInfo(WidgetPtr));
+	}
+#endif
 
-			FWidgetProxy& WidgetProxy = (*FastWidgetPathList)[MyIndex];
-			SWidget* WidgetPtr = WidgetProxy.GetWidget();
+	// Check visibility, it was tested before adding it to the list but another widget may have change while updating.
+	if (InvalidationWidget.Visibility.IsVisible() && WidgetPtr)
+	{
+		const FWidgetProxy::FUpdateResult UpdateResult = InvalidationWidget.Update(*Context.PaintArgs, *Context.WindowElementList);
 
-			// Check visibility, it may have been in the update list but a parent who was also in the update list already updated it.
-			if (WidgetProxy.Visibility.IsVisible() && WidgetPtr)
+		// NB WidgetPtr can now be invalid. The widget can now be collapsed.
+
+		if (UpdateResult.bPainted)
+		{
 			{
-				const bool bNeedPaint = WidgetPtr->HasAnyUpdateFlags(EWidgetUpdateFlags::NeedsRepaint | EWidgetUpdateFlags::NeedsVolatilePaint);
-				bWidgetsNeededRepaint = bWidgetsNeededRepaint || bNeedPaint;
+				// Remove from the update list elements that are processed by this paint
+				FSlateInvalidationWidgetList::FIndexRange PaintedWidgetRange{ *FastWidgetPathList, InvalidationWidget.Index, InvalidationWidget.LeafMostChildIndex };
+				while (FinalUpdateList.Num() > 0)
+				{
+					const int32 LastIndex = FinalUpdateList.Num() - 1;
+					// It's already been processed by the previous draw
+					if (!PaintedWidgetRange.Include(FinalUpdateList[LastIndex].GetWidgetSortOrder()))
+					{
+						break;
+					}
+
+					// It's already been processed by the previous draw
+					FinalUpdateList.RemoveAt(LastIndex, 1, bAllowShrinking);
+				}
+			}
 
-				if (bNeedPaint)
+			// Did it painted more elements than it previously had
+			if (UpdateResult.NewOutgoingLayerId > UpdateResult.PreviousOutgoingLayerId && GSlateInvalidationEnableReindexLayerId)
+			{
+				if (InvalidationWidget.Visibility.IsVisible() && WidgetPtr)
 				{
-					PreviousPaintedWidgetRange = FSlateInvalidationWidgetList::FIndexRange(*FastWidgetPathList, WidgetProxy.Index, WidgetProxy.LeafMostChildIndex);
+					PaintFastPath_FixupLayerId(FastPaintContext, InvalidationWidget, UpdateResult.NewOutgoingLayerId);
 				}
+			}
+
+			CachedMaxLayerId = FMath::Max(UpdateResult.NewOutgoingLayerId, CachedMaxLayerId);
+		}
+	}
+
+	return bNeedsPaint;
+}
+
+bool FSlateInvalidationRoot::PaintFastPath(const FSlateInvalidationContext& Context)
+{
+	SCOPED_NAMED_EVENT(SWidget_FastPathUpdate, FColor::Green);
+	CSV_SCOPED_TIMING_STAT(Slate, FastPathUpdate);
 
-				const int32 NewLayerId = WidgetProxy.Update(*Context.PaintArgs, *Context.WindowElementList);
-				CachedMaxLayerId = FMath::Max(NewLayerId, CachedMaxLayerId);
+	check(!bNeedsSlowPath);
 
+#if WITH_SLATE_DEBUGGING
+	if (GSlateInvalidationRootDumpUpdateList || GSlateInvalidationRootDumpUpdateListOnce)
+	{
+		DumpUpdateList(*FastWidgetPathList, FinalUpdateList);
+	}
+#endif
+#if UE_SLATE_WITH_INVALIDATIONWIDGETLIST_DEBUGGING
+	if (GSlateInvalidationRootVerifyWidgetsAreUpdatedOnce)
+	{
+		FastWidgetPathList->ForEachInvalidationWidget([](FSlateInvalidationWidgetList::InvalidationWidgetType& InvalidationWidget)
+			{
+				InvalidationWidget.bDebug_Updated = false;
+			});
+	}
+#endif
+
+	bool bWidgetsNeededRepaint = false;
+	{
+		TGuardValue<bool> OnFastPathGuard(GSlateIsOnFastUpdatePath, true);
+
+
+		// Widgets that needs reindexing while updating the FinalUpdateList.
+		UE::Slate::Private::FSlateInvalidationPaintFastPathContext FastPaintContext;
+
+		// The update list is put in reverse order by ProcessInvalidation
+		while (FinalUpdateList.Num() > 0 || !FastPaintContext.ReindexUpdateList.IsEmpty())
+		{
+			if (FinalUpdateList.Num() > 0 && !FastPaintContext.ReindexUpdateList.IsEmpty())
+			{
+				const FSlateInvalidationWidgetIndex MyReindex = FastPaintContext.ReindexUpdateList.HeapPeek().GetWidgetIndex();
+				const FSlateInvalidationWidgetList::InvalidationWidgetType& ReindexInvalidationWidget = (*FastWidgetPathList)[MyReindex];
+
+				const FSlateInvalidationWidgetList::FIndexRange ReindexWidgetRange{ *FastWidgetPathList, ReindexInvalidationWidget.Index, ReindexInvalidationWidget.LeafMostChildIndex };
+				if (ReindexWidgetRange.Include(FinalUpdateList[FinalUpdateList.Num() - 1].GetWidgetSortOrder()))
+				{
+					// If the next widget to update is inside the widget range of the re-index widget, then update the child widget first.
+					const bool bPainted = PaintFastPath_UpdateNextWidget(Context, FastPaintContext);
+					bWidgetsNeededRepaint = bWidgetsNeededRepaint || bPainted;
+					if (bNeedsSlowPath)
+					{
+						break;
+					}
+				}
+				else
+				{
+					// The next widget to update is not a children/grand children of the re-index. Reindex it.
+					FastPaintContext.ReindexUpdateList.HeapPopDiscard();
+					SWidget* ReindexWidgetPtr = ReindexInvalidationWidget.GetWidget();
+					if (ReindexInvalidationWidget.Visibility.IsVisible() && ReindexWidgetPtr)
+					{
+						PaintFastPath_FixupLayerId(FastPaintContext, ReindexInvalidationWidget, ReindexWidgetPtr->GetPersistentState().OutgoingLayerId);
+					}
+				}
+			}
+			else if (FinalUpdateList.Num() > 0)
+			{
+				const bool bPainted = PaintFastPath_UpdateNextWidget(Context, FastPaintContext);
+				bWidgetsNeededRepaint = bWidgetsNeededRepaint || bPainted;
 				if (bNeedsSlowPath)
 				{
 					break;
 				}
 			}
+			else
+			{
+				const FSlateInvalidationWidgetIndex MyReindex = FastPaintContext.ReindexUpdateList.HeapPeek().GetWidgetIndex();
+				const FSlateInvalidationWidgetList::InvalidationWidgetType& ReindexInvalidationWidget = (*FastWidgetPathList)[MyReindex];
+
+				FastPaintContext.ReindexUpdateList.HeapPopDiscard();
+				SWidget* ReindexWidgetPtr = ReindexInvalidationWidget.GetWidget();
+				if (ReindexInvalidationWidget.Visibility.IsVisible() && ReindexWidgetPtr)
+				{
+					PaintFastPath_FixupLayerId(FastPaintContext, ReindexInvalidationWidget, ReindexWidgetPtr->GetPersistentState().OutgoingLayerId);
+				}
+			}
 		}
+
+		FinalUpdateList.Reset();
 	}
 
 	bool bExecuteSlowPath = bNeedsSlowPath;
@@ -601,8 +887,8 @@ void FSlateInvalidationRoot::ProcessPreUpdate()
 					{
 						if (Operation.GetRange().Include(Element.GetWidgetSortOrder()))
 						{
-							Element.GetWidgetIndexRef() = Operation.ReIndex(Element.GetWidgetIndex());
-							Element.GetWidgetSortOrderRef() = FSlateInvalidationWidgetSortOrder(Self->WidgetList, Element.GetWidgetIndex());
+							FSlateInvalidationWidgetIndex NewIndex = Operation.ReIndex(Element.GetWidgetIndex());
+							Element = FSlateInvalidationWidgetPreHeap::FElement{ NewIndex, FSlateInvalidationWidgetSortOrder{ Self->WidgetList, NewIndex } };
 						}
 					};
 					PreUpdate.ForEachIndexes(ReIndexIfNeeded);
@@ -629,7 +915,8 @@ void FSlateInvalidationRoot::ProcessPreUpdate()
 				{
 					for (FSlateInvalidationWidgetPreHeap::FElement* Element : WidgetToResort)
 					{
-						Element->GetWidgetSortOrderRef() = FSlateInvalidationWidgetSortOrder{ WidgetList, Element->GetWidgetIndex() };
+						FSlateInvalidationWidgetIndex Index = Element->GetWidgetIndex();
+						(*Element) = FSlateInvalidationWidgetPreHeap::FElement{ Index, FSlateInvalidationWidgetSortOrder{ WidgetList, Index } };
 					}
 					WidgetToResort.Reset();
 				}
@@ -874,8 +1161,8 @@ bool FSlateInvalidationRoot::ProcessPostUpdate()
 	// It update backward (biggest index to smallest)
 	while (WidgetsNeedingPostUpdate->Num() && !bNeedsSlowPath)
 	{
-		const FSlateInvalidationWidgetIndex WidgetIndex = WidgetsNeedingPostUpdate->HeapPop();
-		FWidgetProxy& WidgetProxy = (*FastWidgetPathList)[WidgetIndex];
+		const FSlateInvalidationWidgetPostHeap::FElement WidgetElement = WidgetsNeedingPostUpdate->HeapPop();
+		FWidgetProxy& WidgetProxy = (*FastWidgetPathList)[WidgetElement.GetWidgetIndex()];
 
 		// Widget could be null if it was removed and we are on the slow path
 		if (SWidget* WidgetPtr = WidgetProxy.GetWidget())
@@ -883,7 +1170,7 @@ bool FSlateInvalidationRoot::ProcessPostUpdate()
 #if WITH_SLATE_DEBUGGING
 			if (GSlateInvalidationRootDumpPostInvalidationList)
 			{
-				LogPostInvalidationItem(*FastWidgetPathList, WidgetIndex);
+				LogPostInvalidationItem(*FastWidgetPathList, WidgetElement.GetWidgetIndex());
 			}
 #endif
 
@@ -892,7 +1179,7 @@ bool FSlateInvalidationRoot::ProcessPostUpdate()
 
 			if (WidgetPtr->HasAnyUpdateFlags(EWidgetUpdateFlags::AnyUpdate) && WidgetProxy.Visibility.IsVisible())
 			{
-				FinalUpdateList.Add(WidgetIndex);
+				FinalUpdateList.Emplace(WidgetElement.GetWidgetIndex(), WidgetElement.GetWidgetSortOrder());
 			}
 		}
 	}
@@ -1082,14 +1369,14 @@ void FSlateInvalidationRoot::Advanced_ResetInvalidation(bool bClearResourcesImme
 }
 
 #if WITH_SLATE_DEBUGGING
-void DumpUpdateList(const FSlateInvalidationWidgetList& FastWidgetPathList, const TArray<FSlateInvalidationWidgetIndex>& FinalUpdateList)
+void DumpUpdateList(const FSlateInvalidationWidgetList& FastWidgetPathList, const TArray<FSlateInvalidationWidgetHeapElement>& FinalUpdateList)
 {
 	UE_LOG(LogSlate, Log, TEXT("Dumping Update List"));
 	UE_LOG(LogSlate, Log, TEXT("-------------------"));
 	// The update list is put in reverse order 
 	for (int32 ListIndex = FinalUpdateList.Num() - 1; ListIndex >= 0; --ListIndex)
 	{
-		const FSlateInvalidationWidgetIndex MyIndex = FinalUpdateList[ListIndex];
+		const FSlateInvalidationWidgetIndex MyIndex = FinalUpdateList[ListIndex].GetWidgetIndex();
 
 		const FWidgetProxy& WidgetProxy = FastWidgetPathList[MyIndex];
 		const SWidget* WidgetPtr = WidgetProxy.GetWidget();
@@ -1185,7 +1472,7 @@ void LogPostInvalidationItem(const FSlateInvalidationWidgetList& FastWidgetPathL
 	if (!(Test)) \
 	{ \
 		UE_LOG(LogSlate, Error, Message, ##__VA_ARGS__); \
-		FlagToReset->Set(false); \
+		FlagToReset->Set(false, FlagToReset->GetFlags()); \
 		return; \
 	}
 
@@ -1384,7 +1671,7 @@ void VerifyWidgetVisibility(FSlateInvalidationWidgetList& WidgetList)
 		});
 }
 
-void VerifyWidgetVolatile(FSlateInvalidationWidgetList& WidgetList, TArray<FSlateInvalidationWidgetIndex>& FinalUpdateList)
+void VerifyWidgetVolatile(FSlateInvalidationWidgetList& WidgetList, TArray<FSlateInvalidationWidgetHeapElement>& FinalUpdateList)
 {
 	SWidget* Root = WidgetList.GetRoot().Pin().Get();
 	check(Root);
@@ -1427,7 +1714,11 @@ void VerifyWidgetVolatile(FSlateInvalidationWidgetList& WidgetList, TArray<FSlat
 					if (Widget.GetProxyHandle().IsValid(Widget))
 					{
 						const bool bIsVisible = Widget.GetProxyHandle().GetProxy().Visibility.IsVisible();
-						const bool bIsContains = FinalUpdateList.Contains(Widget.GetProxyHandle().GetWidgetIndex());
+						const FSlateInvalidationWidgetIndex WidgetIndex = Widget.GetProxyHandle().GetWidgetIndex();
+						const bool bIsContains = FinalUpdateList.ContainsByPredicate([WidgetIndex](const FSlateInvalidationWidgetHeapElement& Other)
+							{
+								return Other.GetWidgetIndex() == WidgetIndex;
+							});
 						UE_SLATE_LOG_ERROR_IF_FALSE(bIsContains || !bIsVisible
 							, CVarSlateInvalidationRootVerifyWidgetVolatile
 							, TEXT("Widget '%s' is volatile but is not in the update list.")
@@ -1442,16 +1733,16 @@ void VerifyWidgetsUpdateList_BeforeProcessPreUpdate(const TSharedRef<SWidget>& R
 	FSlateInvalidationWidgetList* FastWidgetPathList,
 	FSlateInvalidationWidgetPreHeap* WidgetsNeedingPreUpdate,
 	FSlateInvalidationWidgetPostHeap* WidgetsNeedingPostUpdate,
-	TArray<FSlateInvalidationWidgetIndex>& FinalUpdateList)
+	TArray<FSlateInvalidationWidgetHeapElement>& FinalUpdateList)
 {
 	if (FastWidgetPathList->GetRoot().Pin() != RootWidget)
 	{
 		return;
 	}
 
-	for (FSlateInvalidationWidgetIndex WidgetIndex : FinalUpdateList)
+	for (const FSlateInvalidationWidgetHeapElement& WidgetElement : FinalUpdateList)
 	{
-		UE_SLATE_LOG_ERROR_IF_FALSE(FastWidgetPathList->IsValidIndex(WidgetIndex)
+		UE_SLATE_LOG_ERROR_IF_FALSE(FastWidgetPathList->IsValidIndex(WidgetElement.GetWidgetIndex())
 			, CVarSlateInvalidationRootVerifyWidgetsUpdateList
 			, TEXT("A WidgetIndex is invalid. The Widget can be invalid (because it's not been processed yet)."));
 	}
@@ -1506,7 +1797,7 @@ void VerifyWidgetsUpdateList_BeforeProcessPostUpdate(const TSharedRef<SWidget>&
 	FSlateInvalidationWidgetList* FastWidgetPathList,
 	FSlateInvalidationWidgetPreHeap* WidgetsNeedingPreUpdate,
 	FSlateInvalidationWidgetPostHeap* WidgetsNeedingPostUpdate,
-	TArray<FSlateInvalidationWidgetIndex>& FinalUpdateList)
+	TArray<FSlateInvalidationWidgetHeapElement>& FinalUpdateList)
 {
 	if (FastWidgetPathList->GetRoot().Pin() != RootWidget)
 	{
@@ -1563,7 +1854,7 @@ void VerifyWidgetsUpdateList_AfterProcessPostUpdate(const TSharedRef<SWidget>& R
 	FSlateInvalidationWidgetList* FastWidgetPathList,
 	FSlateInvalidationWidgetPreHeap* WidgetsNeedingPreUpdate,
 	FSlateInvalidationWidgetPostHeap* WidgetsNeedingPostUpdate,
-	TArray<FSlateInvalidationWidgetIndex>& FinalUpdateList)
+	TArray<FSlateInvalidationWidgetHeapElement>& FinalUpdateList)
 {
 	if (FastWidgetPathList->GetRoot().Pin() != RootWidget)
 	{
@@ -1577,8 +1868,9 @@ void VerifyWidgetsUpdateList_AfterProcessPostUpdate(const TSharedRef<SWidget>& R
 		, CVarSlateInvalidationRootVerifyWidgetsUpdateList
 		, TEXT("The list of Post Update should already been processed."));
 
-	for(FSlateInvalidationWidgetIndex WidgetIndex : FinalUpdateList)
+	for(const FSlateInvalidationWidgetHeapElement& WidgetElement : FinalUpdateList)
 	{
+		const FSlateInvalidationWidgetIndex WidgetIndex = WidgetElement.GetWidgetIndex();
 		const FSlateInvalidationWidgetList::InvalidationWidgetType& InvalidationWidget = (*FastWidgetPathList)[WidgetIndex];
 		const SWidget* Widget = InvalidationWidget.GetWidget();
 
diff --git a/Engine/Source/Runtime/SlateCore/Private/FastUpdate/SlateInvalidationWidgetHeap.h b/Engine/Source/Runtime/SlateCore/Private/FastUpdate/SlateInvalidationWidgetHeap.h
--- a/Engine/Source/Runtime/SlateCore/Private/FastUpdate/SlateInvalidationWidgetHeap.h
+++ b/Engine/Source/Runtime/SlateCore/Private/FastUpdate/SlateInvalidationWidgetHeap.h
@@ -12,37 +12,12 @@
 	#define UE_SLATE_WITH_INVALIDATIONWIDGETHEAP_DEBUGGING DO_CHECK
 #endif
 
-namespace UE
-{
-namespace Slate
-{
-namespace Private
+namespace UE::Slate::Private
 {
 #if UE_SLATE_WITH_INVALIDATIONWIDGETHEAP_DEBUGGING
 	extern bool GSlateInvalidationWidgetHeapVerifyWidgetContains;
 #endif
-
-	/**
-	 * Ordered list of WidgetIndex. The order is based on the WidgetSortIndex.
-	 */
-	struct FSlateInvalidationWidgetHeapElement
-	{
-	public:
-		FSlateInvalidationWidgetHeapElement(FSlateInvalidationWidgetIndex InIndex, FSlateInvalidationWidgetSortOrder InSortOrder)
-			: WidgetIndex(InIndex), WidgetSortOrder(InSortOrder)
-		{}
-		inline FSlateInvalidationWidgetIndex GetWidgetIndex() const { return WidgetIndex; }
-		inline FSlateInvalidationWidgetIndex& GetWidgetIndexRef() { return WidgetIndex; }
-		inline FSlateInvalidationWidgetSortOrder GetWidgetSortOrder() const { return WidgetSortOrder; }
-		inline FSlateInvalidationWidgetSortOrder& GetWidgetSortOrderRef() { return WidgetSortOrder; }
-
-	private:
-		FSlateInvalidationWidgetIndex WidgetIndex;
-		FSlateInvalidationWidgetSortOrder WidgetSortOrder;
-	};
-} // Private
-} // Slate
-} // UE
+} // namespace
 
 /**
  * Heap of widget that is ordered by increasing sort order. The list need to always stay ordered.
@@ -50,7 +25,7 @@ namespace Private
 class FSlateInvalidationWidgetPreHeap
 {
 public:
-	using FElement = UE::Slate::Private::FSlateInvalidationWidgetHeapElement;
+	using FElement = FSlateInvalidationWidgetHeapElement;
 	struct FWidgetOrderLess
 	{
 		FORCEINLINE bool operator()(const FElement& A, const FElement& B) const
@@ -203,7 +178,7 @@ private:
 class FSlateInvalidationWidgetPostHeap
 {
 public:
-	using FElement = UE::Slate::Private::FSlateInvalidationWidgetHeapElement;
+	using FElement = FSlateInvalidationWidgetHeapElement;
 	struct FWidgetOrderGreater
 	{
 		FORCEINLINE bool operator()(const FElement& A, const FElement& B) const
@@ -265,12 +240,12 @@ public:
 	}
 
 	/** Returns and removes the biggest WidgetIndex from the list. */
-	UE_NODISCARD FSlateInvalidationWidgetIndex HeapPop()
+	UE_NODISCARD FElement HeapPop()
 	{
 		check(bIsHeap == true);
-		FSlateInvalidationWidgetIndex Result = Heap.HeapTop().GetWidgetIndex();
+		FElement Result = Heap.HeapTop();
 		Heap.HeapPopDiscard(SortPredicate(), false);
-		OwnerList[Result].bContainedByWidgetPostHeap = false;
+		OwnerList[Result.GetWidgetIndex()].bContainedByWidgetPostHeap = false;
 		return Result;
 	}
 
@@ -408,7 +383,7 @@ private:
 class FSlateInvalidationWidgetPrepassHeap
 {
 public:
-	using FElement = UE::Slate::Private::FSlateInvalidationWidgetHeapElement;
+	using FElement = FSlateInvalidationWidgetHeapElement;
 	struct FWidgetOrderGreater
 	{
 		FORCEINLINE bool operator()(const FElement& A, const FElement& B) const
diff --git a/Engine/Source/Runtime/SlateCore/Private/FastUpdate/WidgetProxy.cpp b/Engine/Source/Runtime/SlateCore/Private/FastUpdate/WidgetProxy.cpp
--- a/Engine/Source/Runtime/SlateCore/Private/FastUpdate/WidgetProxy.cpp
+++ b/Engine/Source/Runtime/SlateCore/Private/FastUpdate/WidgetProxy.cpp
@@ -44,16 +44,16 @@ TSharedPtr<SWidget> FWidgetProxy::GetWidgetAsShared() const
 #endif
 }
 
-int32 FWidgetProxy::Update(const FPaintArgs& PaintArgs, FSlateWindowElementList& OutDrawElements)
+FWidgetProxy::FUpdateResult FWidgetProxy::Update(const FPaintArgs& PaintArgs, FSlateWindowElementList& OutDrawElements)
 {
 	TSharedPtr<SWidget> CurrentWidget = GetWidgetAsShared();
 	check(Visibility.IsVisible());
 
 	// If Outgoing layer id remains index none, there was no change
-	int32 OutgoingLayerId = INDEX_NONE;
+	FUpdateResult Result;
 	if (CurrentWidget->HasAnyUpdateFlags(EWidgetUpdateFlags::NeedsRepaint|EWidgetUpdateFlags::NeedsVolatilePaint))
 	{
-		OutgoingLayerId = Repaint(PaintArgs, OutDrawElements);
+		Result = Repaint(PaintArgs, OutDrawElements);
 	}
 	else
 	{
@@ -63,6 +63,10 @@ int32 FWidgetProxy::Update(const FPaintArgs& PaintArgs, FSlateWindowElementList&
 		{
 			SCOPE_CYCLE_COUNTER(STAT_SlateExecuteActiveTimers);
 			CurrentWidget->ExecuteActiveTimers(PaintArgs.GetCurrentTime(), PaintArgs.GetDeltaTime());
+
+#if UE_SLATE_WITH_INVALIDATIONWIDGETLIST_DEBUGGING
+			bDebug_Updated = true;
+#endif
 		}
 
 		if (CurrentWidget->HasAnyUpdateFlags(EWidgetUpdateFlags::NeedsTick))
@@ -73,14 +77,29 @@ int32 FWidgetProxy::Update(const FPaintArgs& PaintArgs, FSlateWindowElementList&
 			SCOPE_CYCLE_COUNTER(STAT_SlateTickWidgets);
 
 			CurrentWidget->Tick(MyState.DesktopGeometry, PaintArgs.GetCurrentTime(), PaintArgs.GetDeltaTime());
+
+#if UE_SLATE_WITH_INVALIDATIONWIDGETLIST_DEBUGGING
+			bDebug_Updated = true;
+#endif
 		}
 
+		// If the widet was invalidated while ticking. In slow mode the widget would be painted right away.
+		//For now postpone for the next frame. Enabling this code would cause more issues with the list management.
+		//if (EnumHasAnyFlags(CurrentInvalidateReason, EInvalidateWidgetReason::Paint))
+		//{
+		//	EnumRemoveFlags(CurrentInvalidateReason, EInvalidateWidgetReason::Paint);
+		//	const EWidgetUpdateFlags FlagToAddAfterRepaint = CurrentWidget->UpdateFlags & (EWidgetUpdateFlags::NeedsActiveTimerUpdate | EWidgetUpdateFlags::NeedsTick);
+		//	EnumRemoveFlags(CurrentWidget->UpdateFlags, EWidgetUpdateFlags::NeedsActiveTimerUpdate | EWidgetUpdateFlags::NeedsTick);
+		//	Result = Repaint(PaintArgs, OutDrawElements);
+		//	EnumAddFlags(CurrentWidget->UpdateFlags, FlagToAddAfterRepaint);
+		//}
+
 #if WITH_SLATE_DEBUGGING
 		FSlateDebugging::BroadcastWidgetUpdated(CurrentWidget.Get(), PreviousUpdateFlag);
 #endif
 	}
 
-	return OutgoingLayerId;
+	return Result;
 }
 
 void FWidgetProxy::ProcessLayoutInvalidation(FSlateInvalidationWidgetPostHeap& UpdateList, FSlateInvalidationWidgetList& FastWidgetPathList, FSlateInvalidationRoot& Root)
@@ -189,7 +208,7 @@ void FWidgetProxy::MarkProxyUpdatedThisFrame(FSlateInvalidationWidgetPostHeap& P
 	}
 }
 
-int32 FWidgetProxy::Repaint(const FPaintArgs& PaintArgs, FSlateWindowElementList& OutDrawElements) const
+FWidgetProxy::FUpdateResult FWidgetProxy::Repaint(const FPaintArgs& PaintArgs, FSlateWindowElementList& OutDrawElements) const
 {
 	SWidget* WidgetPtr = GetWidget();
 	check(WidgetPtr);
@@ -239,7 +258,7 @@ int32 FWidgetProxy::Repaint(const FPaintArgs& PaintArgs, FSlateWindowElementList
 		check(StartingClipIndex == OutDrawElements.GetClippingIndex());
 	}
 
-	return NewLayerId;
+	return FUpdateResult{ PrevLayerId, NewLayerId };
 }
 
 FWidgetProxyHandle::FWidgetProxyHandle(const FSlateInvalidationRootHandle& InInvalidationRoot, FSlateInvalidationWidgetIndex InIndex, FSlateInvalidationWidgetSortOrder InSortIndex)
@@ -279,6 +298,24 @@ FSlateInvalidationWidgetVisibility FWidgetProxyHandle::GetWidgetVisibility(const
 	return FSlateInvalidationWidgetVisibility();
 }
 
+bool FWidgetProxyHandle::HasAllInvalidationReason(const SWidget* Widget, EInvalidateWidgetReason Reason) const
+{
+	if (IsValid(Widget))
+	{
+		return EnumHasAllFlags(GetProxy().CurrentInvalidateReason, Reason);
+	}
+	return false;
+}
+
+bool FWidgetProxyHandle::HasAnyInvalidationReason(const SWidget* Widget, EInvalidateWidgetReason Reason) const
+{
+	if (IsValid(Widget))
+	{
+		return EnumHasAnyFlags(GetProxy().CurrentInvalidateReason, Reason);
+	}
+	return false;
+}
+
 FWidgetProxy& FWidgetProxyHandle::GetProxy()
 {
 	return GetInvalidationRoot_NoCheck()->GetFastPathWidgetList()[WidgetIndex];
diff --git a/Engine/Source/Runtime/SlateCore/Private/Widgets/SWidget.cpp b/Engine/Source/Runtime/SlateCore/Private/Widgets/SWidget.cpp
--- a/Engine/Source/Runtime/SlateCore/Private/Widgets/SWidget.cpp
+++ b/Engine/Source/Runtime/SlateCore/Private/Widgets/SWidget.cpp
@@ -127,6 +127,9 @@ static FAutoConsoleVariableRef CVarSlateDebugCulling(TEXT("Slate.DebugCulling"),
 bool GSlateEnsureAllVisibleWidgetsPaint = false;
 static FAutoConsoleVariableRef CVarSlateEnsureAllVisibleWidgetsPaint(TEXT("Slate.EnsureAllVisibleWidgetsPaint"), GSlateEnsureAllVisibleWidgetsPaint, TEXT("Ensures that if a child widget is visible before OnPaint, that it was painted this frame after OnPaint, if still marked as visible.  Only works if we're on the FastPaintPath."), ECVF_Default);
 
+bool GSlateEnsureOutgoingLayerId = false;
+static FAutoConsoleVariableRef CVarSlateEnsureOutgoingLayerId(TEXT("Slate.EnsureOutgoingLayerId"), GSlateEnsureOutgoingLayerId, TEXT("Ensures that child widget returns the correct layer id with OnPaint."), ECVF_Default);
+
 #endif
 
 #if STATS || ENABLE_STATNAMEDEVENTS
@@ -780,7 +783,10 @@ void SWidget::UpdateWidgetProxy(int32 NewLayerId, FSlateCachedElementsHandle& Ca
 #if WITH_SLATE_DEBUGGING
 	if (FastPathProxyHandle.IsValid(this))
 	{
-		const FWidgetProxy& MyProxy = FastPathProxyHandle.GetProxy();
+		FWidgetProxy& MyProxy = FastPathProxyHandle.GetProxy();
+#if UE_SLATE_WITH_INVALIDATIONWIDGETLIST_DEBUGGING
+		MyProxy.bDebug_Updated = true;
+#endif
 		ensureMsgf(MyProxy.Visibility.IsVisibleDirectly() == GetVisibility().IsVisible()
 			, TEXT("The visibility of the widget '%s' changed during Paint")
 			, *FReflectionMetaData::GetWidgetPath(this));
@@ -1517,7 +1523,7 @@ int32 SWidget::Paint(const FPaintArgs& Args, const FGeometry& AllottedGeometry,
 #if WITH_SLATE_DEBUGGING
 	TArray<TWeakPtr<const SWidget>, TInlineAllocator<16>> DebugChildWidgetsToPaint;
 
-	if (GSlateIsOnFastUpdatePath && GSlateEnsureAllVisibleWidgetsPaint)
+	if ((GSlateIsOnFastUpdatePath && GSlateEnsureAllVisibleWidgetsPaint) || GSlateEnsureOutgoingLayerId)
 	{
 		// Don't check for invalidation roots a completely different set of rules apply to those widgets.
 		if (!Advanced_IsInvalidationRoot())
@@ -1545,6 +1551,7 @@ int32 SWidget::Paint(const FPaintArgs& Args, const FGeometry& AllottedGeometry,
 	// Detect children that should have been painted, but were skipped during the paint process.
 	// this will result in geometry being left on screen and not cleared, because it's visible, yet wasn't painted.
 #if WITH_SLATE_DEBUGGING
+	if (GSlateIsOnFastUpdatePath && GSlateEnsureAllVisibleWidgetsPaint)
 	{
 		for (TWeakPtr<const SWidget>& DebugChildThatShouldHaveBeenPaintedPtr : DebugChildWidgetsToPaint)
 		{
@@ -1552,16 +1559,39 @@ int32 SWidget::Paint(const FPaintArgs& Args, const FGeometry& AllottedGeometry,
 			{
 				if (DebugChild->GetVisibility().IsVisible())
 				{
-					ensureAlwaysMsgf(DebugChild->Debug_GetLastPaintFrame() == GFrameNumber, TEXT("The Widget '%s' was visible, but never painted. This means it was skipped during painting, without alerting the fast path."), *FReflectionMetaData::GetWidgetPath(DebugChild.Get()));
+					if (DebugChild->Debug_GetLastPaintFrame() != GFrameNumber)
+					{
+						ensureAlwaysMsgf(false, TEXT("The widget '%s' was visible, but never painted. This means it was skipped during painting, without alerting the fast path."), *FReflectionMetaData::GetWidgetPath(DebugChild.Get()));
+						CVarSlateEnsureAllVisibleWidgetsPaint->Set(false, CVarSlateEnsureAllVisibleWidgetsPaint->GetFlags());
+					}
 				}
 			}
 			else
 			{
-				CVarSlateEnsureAllVisibleWidgetsPaint->Set(false);
+				CVarSlateEnsureAllVisibleWidgetsPaint->Set(false, CVarSlateEnsureAllVisibleWidgetsPaint->GetFlags());
 				ensureAlwaysMsgf(false, TEXT("A widget was destroyed while painting it's parent. This is not supported by the Invalidation system."));
 			}
 		}
 	}
+
+	if (GSlateEnsureOutgoingLayerId)
+	{
+		for (TWeakPtr<const SWidget>& DebugChildWeakPtr : DebugChildWidgetsToPaint)
+		{
+			if (TSharedPtr<const SWidget> DebugChild = DebugChildWeakPtr.Pin())
+			{
+				const bool bIsChildDeferredPaint = DebugChild->GetPersistentState().bDeferredPainting;
+				if (!bIsChildDeferredPaint && DebugChild->GetVisibility().IsVisible() && DebugChild->Debug_GetLastPaintFrame() == GFrameNumber)
+				{
+					if (NewLayerId < DebugChild->GetPersistentState().OutgoingLayerId)
+					{
+						ensureAlwaysMsgf(false, TEXT("The widget '%s' Outgoing Layer Id is bigger than its parent."), *FReflectionMetaData::GetWidgetPath(DebugChild.Get()));
+						CVarSlateEnsureOutgoingLayerId->Set(false, CVarSlateEnsureOutgoingLayerId->GetFlags());
+					}
+				}
+			}
+		}
+	}
 #endif
 
 	// Draw the clipping zone if we've got clipping enabled
diff --git a/Engine/Source/Runtime/SlateCore/Public/FastUpdate/SlateInvalidationRoot.h b/Engine/Source/Runtime/SlateCore/Public/FastUpdate/SlateInvalidationRoot.h
--- a/Engine/Source/Runtime/SlateCore/Public/FastUpdate/SlateInvalidationRoot.h
+++ b/Engine/Source/Runtime/SlateCore/Public/FastUpdate/SlateInvalidationRoot.h
@@ -7,6 +7,7 @@
 #include "WidgetProxy.h"
 #include "FastUpdate/SlateInvalidationRootHandle.h"
 #include "FastUpdate/SlateInvalidationWidgetIndex.h"
+#include "FastUpdate/SlateInvalidationWidgetSortOrder.h"
 #include "Rendering/DrawElements.h"
 
 struct FSlateCachedElementData;
@@ -17,6 +18,11 @@ class FSlateInvalidationWidgetPostHeap;
 class FSlateWindowElementList;
 class FWidgetStyle;
 
+namespace UE::Slate::Private
+{
+	struct FSlateInvalidationPaintFastPathContext;
+}
+
 #define UE_SLATE_DEBUGGING_CLEAR_ALL_FAST_PATH_DATA 0
 
 struct FSlateInvalidationContext
@@ -152,6 +158,10 @@ private:
 	void HandleInvalidateAllWidgets(bool bClearResourcesImmediately);
 
 	bool PaintFastPath(const FSlateInvalidationContext& Context);
+	bool PaintFastPath_UpdateNextWidget(const FSlateInvalidationContext& Context, UE::Slate::Private::FSlateInvalidationPaintFastPathContext& FastPathContext);
+	void PaintFastPath_FixupLayerId(UE::Slate::Private::FSlateInvalidationPaintFastPathContext& FastPathContext, const FWidgetProxy& InvalidationWidget, const int32 NewOutgoingLayerId);
+	void PaintFastPath_FixupParentLayerId(UE::Slate::Private::FSlateInvalidationPaintFastPathContext& FastPathContext, const FWidgetProxy& InvalidationWidget, const int32 NewOutgoingLayerId);
+	void PaintFastPath_AddUniqueSortedToFinalUpdateList(const FSlateInvalidationWidgetIndex InvalidationWidgetIndex);
 
 	/** Call when an invalidation occurred. */
 	void InvalidateWidget(FWidgetProxy& Proxy, EInvalidateWidgetReason InvalidateReason);
@@ -187,8 +197,8 @@ private:
 	 */
 	TUniquePtr<FSlateInvalidationWidgetPostHeap> WidgetsNeedingPostUpdate;
 
-	/** Index to widgets that will be updated. */
-	TArray<FSlateInvalidationWidgetIndex> FinalUpdateList;
+	/** Widgets that will be updated. */
+	TArray<FSlateInvalidationWidgetHeapElement> FinalUpdateList;
 
 	FSlateCachedElementData* CachedElementData;
 
diff --git a/Engine/Source/Runtime/SlateCore/Public/FastUpdate/SlateInvalidationWidgetSortOrder.h b/Engine/Source/Runtime/SlateCore/Public/FastUpdate/SlateInvalidationWidgetSortOrder.h
--- a/Engine/Source/Runtime/SlateCore/Public/FastUpdate/SlateInvalidationWidgetSortOrder.h
+++ b/Engine/Source/Runtime/SlateCore/Public/FastUpdate/SlateInvalidationWidgetSortOrder.h
@@ -32,4 +32,31 @@ public:
 	bool operator>= (const FSlateInvalidationWidgetSortOrder Other) const { return Order >= Other.Order; }
 	bool operator== (const FSlateInvalidationWidgetSortOrder Other) const { return Order == Other.Order; }
 	bool operator!= (const FSlateInvalidationWidgetSortOrder Other) const { return Order != Other.Order; }
+};
+
+
+/**
+ * Pair of WidgetIndex and WidgetSortIndex. Can be used to sort the widget.
+ */
+struct FSlateInvalidationWidgetHeapElement
+{
+public:
+	FSlateInvalidationWidgetHeapElement(FSlateInvalidationWidgetIndex InIndex, FSlateInvalidationWidgetSortOrder InSortOrder)
+		: WidgetIndex(InIndex), WidgetSortOrder(InSortOrder)
+	{
+	}
+
+	inline FSlateInvalidationWidgetIndex GetWidgetIndex() const
+	{
+		return WidgetIndex;
+	}
+
+	inline FSlateInvalidationWidgetSortOrder GetWidgetSortOrder() const
+	{
+		return WidgetSortOrder;
+	}
+
+private:
+	FSlateInvalidationWidgetIndex WidgetIndex;
+	FSlateInvalidationWidgetSortOrder WidgetSortOrder;
 };
\ No newline at end of file
diff --git a/Engine/Source/Runtime/SlateCore/Public/FastUpdate/WidgetProxy.h b/Engine/Source/Runtime/SlateCore/Public/FastUpdate/WidgetProxy.h
--- a/Engine/Source/Runtime/SlateCore/Public/FastUpdate/WidgetProxy.h
+++ b/Engine/Source/Runtime/SlateCore/Public/FastUpdate/WidgetProxy.h
@@ -113,7 +113,25 @@ class FWidgetProxy
 public:
 	FWidgetProxy(SWidget& InWidget);
 
-	int32 Update(const FPaintArgs& PaintArgs, FSlateWindowElementList& OutDrawElements);
+	struct FUpdateResult
+	{
+		FUpdateResult() = default;
+		FUpdateResult(int32 InPreviousOutgoingLayerId, int32 InNewOutgoingLayerId)
+			: PreviousOutgoingLayerId(InPreviousOutgoingLayerId)
+			, NewOutgoingLayerId(InNewOutgoingLayerId)
+			, bPainted(true)
+		{
+		}
+		int32 PreviousOutgoingLayerId = INDEX_NONE;
+		int32 NewOutgoingLayerId = INDEX_NONE;
+		bool bPainted = false;
+	};
+
+	/**
+	 * Similar to SWidget::Paint but use the saved PersistenState. Only paint/tick if needed.
+	 * @return the new and previous LayerId is the widget was painted.
+	 */
+	FUpdateResult Update(const FPaintArgs& PaintArgs, FSlateWindowElementList& OutDrawElements);
 
 	void ProcessLayoutInvalidation(FSlateInvalidationWidgetPostHeap& UpdateList, FSlateInvalidationWidgetList& FastPathWidgetList, FSlateInvalidationRoot& Root);
 	bool ProcessPostInvalidation(FSlateInvalidationWidgetPostHeap& UpdateList, FSlateInvalidationWidgetList& FastPathWidgetList, FSlateInvalidationRoot& Root);
@@ -146,7 +164,7 @@ public:
 #endif
 
 private:
-	int32 Repaint(const FPaintArgs& PaintArgs, FSlateWindowElementList& OutDrawElements) const;
+	FUpdateResult Repaint(const FPaintArgs& PaintArgs, FSlateWindowElementList& OutDrawElements) const;
 
 private:
 #if UE_SLATE_WITH_WIDGETPROXY_WEAKPTR
@@ -196,6 +214,8 @@ public:
 			uint8 bDebug_LastFrameVisibleSet : 1;
 			/** Use with "Slate.InvalidationRoot.VerifyWidgetAttribute". */
 			uint8 bDebug_AttributeUpdated : 1;
+			/** The widget was updated (paint or ticked/activetimer). */
+			uint8 bDebug_Updated : 1;
 		};
 		uint8 PrivateDebugFlags;
 	};
@@ -268,6 +288,8 @@ public:
 	FSlateInvalidationWidgetSortOrder GetWidgetSortOrder() const { return WidgetSortOrder; }
 
 	SLATECORE_API FSlateInvalidationWidgetVisibility GetWidgetVisibility(const SWidget* Widget) const;
+	SLATECORE_API bool HasAllInvalidationReason(const SWidget* Widget, EInvalidateWidgetReason Reason) const;
+	SLATECORE_API bool HasAnyInvalidationReason(const SWidget* Widget, EInvalidateWidgetReason Reason) const;
 
 	FWidgetProxy& GetProxy();
 	const FWidgetProxy& GetProxy() const;
