Slate: Try to fix random crash while iterating on SWidget.

diff --git a/Engine/Source/Runtime/SlateCore/Public/Layout/ChildrenBase.h b/Engine/Source/Runtime/SlateCore/Public/Layout/ChildrenBase.h
--- a/Engine/Source/Runtime/SlateCore/Public/Layout/ChildrenBase.h
+++ b/Engine/Source/Runtime/SlateCore/Public/Layout/ChildrenBase.h
@@ -60,10 +60,11 @@ public:
 #if WITH_SLATE_DEBUGGING 
 		TGuardValue<bool> IteratingGuard(bIsIteratingChildren, true);
 #endif
+
 		for (int32 Index = 0; Index < Num(); ++Index)
 		{
-			FWidgetRef WidgetRef = GetChildRefAt(Index);
-			Pred(WidgetRef.GetWidget());
+			TSharedRef<SWidget> Ref = GetChildAt(Index);
+			Pred(Ref.Get());
 		}
 	}
 
@@ -77,8 +78,8 @@ public:
 
 		for (int32 Index = 0; Index < Num(); ++Index)
 		{
-			FConstWidgetRef WidgetRef = GetChildRefAt(Index);
-			Pred(WidgetRef.GetWidget());
+			TSharedRef<const SWidget> Ref = GetChildAt(Index);
+			Pred(Ref.Get());
 		}
 	}
 
