Slate: GC FShapedGlyphSequence::FontMaterial. There were GCed by the Payload previously but nothing prevented a user to keep an instance of the FShapredGlyphSequence and not using it.

diff --git a/Engine/Source/Runtime/SlateCore/Private/Fonts/FontCache.cpp b/Engine/Source/Runtime/SlateCore/Private/Fonts/FontCache.cpp
--- a/Engine/Source/Runtime/SlateCore/Private/Fonts/FontCache.cpp
+++ b/Engine/Source/Runtime/SlateCore/Private/Fonts/FontCache.cpp
@@ -11,6 +11,7 @@
 #include "Fonts/SlateTextShaper.h"
 #include "Fonts/LegacySlateFontInfoCache.h"
 #include "Fonts/FontCacheUtils.h"
+#include "UObject/GCObject.h"
 
 #include <limits>
 
@@ -109,6 +110,41 @@ FShapedGlyphEntryKey::FShapedGlyphEntryKey(const FShapedGlyphFaceData& InFontFac
 	KeyHash = HashCombine(KeyHash, GetTypeHash(GlyphIndex));
 }
 
+
+namespace UE::Slate::Private
+{
+	/** Garbage collection */
+	struct FShapedGlyphSequenceMaterialReferenceCollector : public FGCObject
+	{
+	private:
+		FShapedGlyphSequenceMaterialReferenceCollector() = default;
+
+	public:
+		static FShapedGlyphSequenceMaterialReferenceCollector& Get()
+		{
+			static FShapedGlyphSequenceMaterialReferenceCollector Instance;
+			return Instance;
+		}
+
+		virtual void AddReferencedObjects(FReferenceCollector& Collector) override
+		{
+			for(FShapedGlyphSequence* Element : AllSequences)
+			{
+				Collector.AddReferencedObject(Element->FontMaterial);
+				Collector.AddReferencedObject(Element->OutlineSettings.OutlineMaterial);
+			}
+		}
+
+		virtual FString GetReferencerName() const override
+		{
+			return TEXT("FShapedGlyphSequence");
+		}
+
+		TSet<FShapedGlyphSequence*> AllSequences;
+	};
+}
+
+
 FShapedGlyphSequence::FShapedGlyphSequence(TArray<FShapedGlyphEntry> InGlyphsToRender, const int16 InTextBaseline, const uint16 InMaxTextHeight, const UObject* InFontMaterial, const FFontOutlineSettings& InOutlineSettings, const FSourceTextRange& InSourceTextRange)
 	: GlyphsToRender(MoveTemp(InGlyphsToRender))
 	, TextBaseline(InTextBaseline)
@@ -118,6 +154,7 @@ FShapedGlyphSequence::FShapedGlyphSequence(TArray<FShapedGlyphEntry> InGlyphsToR
 	, SequenceWidth(0)
 	, GlyphFontFaces()
 	, SourceIndicesToGlyphData(InSourceTextRange)
+	, bAddedToGC(false)
 {
 	const int32 NumGlyphsToRender = GlyphsToRender.Num();
 	for (int32 CurrentGlyphIndex = 0; CurrentGlyphIndex < NumGlyphsToRender; ++CurrentGlyphIndex)
@@ -181,12 +218,23 @@ FShapedGlyphSequence::FShapedGlyphSequence(TArray<FShapedGlyphEntry> InGlyphsToR
 		}
 	}
 
+	if (FontMaterial || OutlineSettings.OutlineMaterial)
+	{
+		UE::Slate::Private::FShapedGlyphSequenceMaterialReferenceCollector::Get().AllSequences.Add(this);
+		bAddedToGC = true;
+	}
+
 	// Track memory usage
 	INC_MEMORY_STAT_BY(STAT_SlateShapedGlyphSequenceMemory, GetAllocatedSize());
 }
 
 FShapedGlyphSequence::~FShapedGlyphSequence()
 {
+	if (bAddedToGC)
+	{
+		UE::Slate::Private::FShapedGlyphSequenceMaterialReferenceCollector::Get().AllSequences.Remove(this);
+	}
+
 	// Untrack memory usage
 	DEC_MEMORY_STAT_BY(STAT_SlateShapedGlyphSequenceMemory, GetAllocatedSize());
 }
@@ -422,12 +470,6 @@ FShapedGlyphSequencePtr FShapedGlyphSequence::GetSubSequence(const int32 InStart
 	return nullptr;
 }
 
-void FShapedGlyphSequence::AddReferencedObjects(FReferenceCollector& Collector)
-{
-	Collector.AddReferencedObject(FontMaterial);
-	Collector.AddReferencedObject(OutlineSettings.OutlineMaterial);
-}
-
 FShapedGlyphSequence::EEnumerateGlyphsResult FShapedGlyphSequence::EnumerateLogicalGlyphsInSourceRange(const int32 InStartIndex, const int32 InEndIndex, const FForEachShapedGlyphEntryCallback& InGlyphCallback) const
 {
 	if (InStartIndex == InEndIndex)
diff --git a/Engine/Source/Runtime/SlateCore/Public/Fonts/FontCache.h b/Engine/Source/Runtime/SlateCore/Public/Fonts/FontCache.h
--- a/Engine/Source/Runtime/SlateCore/Public/Fonts/FontCache.h
+++ b/Engine/Source/Runtime/SlateCore/Public/Fonts/FontCache.h
@@ -23,6 +23,11 @@ class FSlateFontRenderer;
 class FSlateShaderResource;
 class FSlateTextShaper;
 
+namespace UE::Slate::Private
+{
+	struct FShapedGlyphSequenceMaterialReferenceCollector;
+}
+
 enum class EFontCacheAtlasDataType : uint8
 {
 	/** Data was cached for a regular non-outline font */
@@ -331,12 +336,12 @@ public:
 	 */
 	FShapedGlyphSequencePtr GetSubSequence(const int32 InStartIndex, const int32 InEndIndex) const;
 
-	void AddReferencedObjects(FReferenceCollector& Collector);
-
 private:
+	friend UE::Slate::Private::FShapedGlyphSequenceMaterialReferenceCollector;
+
 	/** Non-copyable */
-	FShapedGlyphSequence(const FShapedGlyphSequence&);
-	FShapedGlyphSequence& operator=(const FShapedGlyphSequence&);
+	FShapedGlyphSequence(const FShapedGlyphSequence&) = delete;
+	FShapedGlyphSequence& operator=(const FShapedGlyphSequence&) = delete;
 
 	/** Helper function to share some common logic between the bound and unbound GetGlyphAtOffset functions */
 	bool HasFoundGlyphAtOffset(FSlateFontCache& InFontCache, const int32 InHorizontalOffset, const FShapedGlyphEntry& InCurrentGlyph, const int32 InCurrentGlyphIndex, int32& InOutCurrentOffset, const FShapedGlyphEntry*& OutMatchedGlyph) const;
@@ -435,7 +440,7 @@ private:
 	/** The maximum height of any glyph in the font we're using */
 	uint16 MaxTextHeight;
 	/** The material to use when rendering these glyphs */
-	const UObject* FontMaterial;
+	TObjectPtr<const UObject> FontMaterial;
 	/** Outline settings to use when rendering these glyphs */
 	FFontOutlineSettings OutlineSettings;
 	/** The cached width of the entire sequence */
@@ -444,6 +449,8 @@ private:
 	TArray<TWeakPtr<FFreeTypeFace>> GlyphFontFaces;
 	/** A map of source indices to their shaped glyph data indices - used to perform efficient reverse look-up */
 	FSourceIndicesToGlyphData SourceIndicesToGlyphData;
+	/** This sequence as UObject reference and was added to the reference collector */
+	bool bAddedToGC;
 };
 
 /** Information for rendering one non-shaped character */
diff --git a/Engine/Source/Runtime/SlateCore/Public/Rendering/DrawElementPayloads.h b/Engine/Source/Runtime/SlateCore/Public/Rendering/DrawElementPayloads.h
--- a/Engine/Source/Runtime/SlateCore/Public/Rendering/DrawElementPayloads.h
+++ b/Engine/Source/Runtime/SlateCore/Public/Rendering/DrawElementPayloads.h
@@ -187,19 +187,6 @@ struct FSlateShapedTextPayload : public FSlateDataPayload, public FSlateTintable
 		OverflowArgs = InArgs;
 		check(InArgs.OverflowDirection == ETextOverflowDirection::NoOverflow || InArgs.OverflowTextPtr.IsValid());
 	}
-
-	virtual void AddReferencedObjects(FReferenceCollector& Collector)
-	{
-		if (ShapedGlyphSequence.IsValid())
-		{
-			const_cast<FShapedGlyphSequence*>(ShapedGlyphSequence.Get())->AddReferencedObjects(Collector);
-		}
-
-		if (OverflowArgs.OverflowTextPtr.IsValid())
-		{
-			const_cast<FShapedGlyphSequence*>(OverflowArgs.OverflowTextPtr.Get())->AddReferencedObjects(Collector);
-		}
-	}
 };
 
 struct FSlateGradientPayload : public FSlateDataPayload
