Slate: Try to fix random crash that corrupt the memory in SlatePrepass.

diff --git a/Engine/Source/Runtime/Engine/Private/GameEngine.cpp b/Engine/Source/Runtime/Engine/Private/GameEngine.cpp
--- a/Engine/Source/Runtime/Engine/Private/GameEngine.cpp
+++ b/Engine/Source/Runtime/Engine/Private/GameEngine.cpp
@@ -649,31 +649,41 @@ TSharedRef<SWindow> UGameEngine::CreateGameWindow()
 
 void UGameEngine::SwitchGameWindowToUseGameViewport()
 {
-	if (GameViewportWindow.IsValid() && GameViewportWindow.Pin()->GetContent() != GameViewportWidget)
+	if (TSharedPtr<SWindow> GameViewportWindowPtr = GameViewportWindow.Pin())
 	{
-		if( !GameViewportWidget.IsValid() )
+		if (GameViewportWindowPtr->GetContent() != GameViewportWidget)
 		{
-			CreateGameViewport( GameViewport );
-		}
-		
-		TSharedRef<SViewport> GameViewportWidgetRef = GameViewportWidget.ToSharedRef();
-		TSharedPtr<SWindow> GameViewportWindowPtr = GameViewportWindow.Pin();
-		
-		GameViewportWindowPtr->SetContent(GameViewportWidgetRef);
-		GameViewportWindowPtr->SlatePrepass(FSlateApplication::Get().GetApplicationScale() * GameViewportWindowPtr->GetNativeWindow()->GetDPIScaleFactor());
-		
-		if ( SceneViewport.IsValid() )
-		{
-			SceneViewport->ResizeFrame((uint32)GSystemResolution.ResX, (uint32)GSystemResolution.ResY, GSystemResolution.WindowMode);
-		}
+			if (!GameViewportWidget.IsValid())
+			{
+				CreateGameViewport(GameViewport);
+			}
 
-		// Registration of the game viewport to that messages are correctly received.
-		// Could be a re-register, however it's necessary after the window is set.
-		FSlateApplication::Get().RegisterGameViewport(GameViewportWidgetRef);
+			if (GameViewportWidget.IsValid())
+			{
+				GameViewportWindowPtr->SetContent(GameViewportWidget.ToSharedRef());
+				GameViewportWindowPtr->SlatePrepass(FSlateApplication::Get().GetApplicationScale() * GameViewportWindowPtr->GetNativeWindow()->GetDPIScaleFactor());
+			}
+			else
+			{
+				UE_LOG(LogEngine, Error, TEXT("The Game Viewport Widget is invalid."));
+			}
 
-		if (FSlateApplication::IsInitialized())
-		{
-			FSlateApplication::Get().SetAllUserFocusToGameViewport(EFocusCause::SetDirectly);
+			if (SceneViewport.IsValid())
+			{
+				SceneViewport->ResizeFrame((uint32)GSystemResolution.ResX, (uint32)GSystemResolution.ResY, GSystemResolution.WindowMode);
+			}
+
+			// Registration of the game viewport to that messages are correctly received.
+			// Could be a re-register, however it's necessary after the window is set.
+			if (GameViewportWidget.IsValid())
+			{
+				FSlateApplication::Get().RegisterGameViewport(GameViewportWidget.ToSharedRef());
+			}
+
+			if (FSlateApplication::IsInitialized())
+			{
+				FSlateApplication::Get().SetAllUserFocusToGameViewport(EFocusCause::SetDirectly);
+			}
 		}
 	}
 }
diff --git a/Engine/Source/Runtime/SlateCore/Private/Widgets/SWidget.cpp b/Engine/Source/Runtime/SlateCore/Private/Widgets/SWidget.cpp
--- a/Engine/Source/Runtime/SlateCore/Private/Widgets/SWidget.cpp
+++ b/Engine/Source/Runtime/SlateCore/Private/Widgets/SWidget.cpp
@@ -1718,50 +1718,53 @@ void SWidget::Prepass_Internal(float InLayoutScaleMultiplier)
 void SWidget::Prepass_ChildLoop(float InLayoutScaleMultiplier, FChildren* MyChildren)
 {
 	int32 ChildIndex = 0;
-	MyChildren->ForEachWidget([this, &ChildIndex, InLayoutScaleMultiplier](SWidget& Child)
+	SWidget* Self = this;
+	auto ForEachPred = [Self, &ChildIndex, InLayoutScaleMultiplier](SWidget& Child)
+	{
+		const bool bUpdateAttributes = Child.HasRegisteredSlateAttribute() && Child.IsAttributesUpdatesEnabled() && !GSlateIsOnFastProcessInvalidation;
+		if (bUpdateAttributes)
 		{
-			const bool bUpdateAttributes = Child.HasRegisteredSlateAttribute() && Child.IsAttributesUpdatesEnabled() && !GSlateIsOnFastProcessInvalidation;
-			if (bUpdateAttributes)
-			{
-				FSlateAttributeMetaData::UpdateOnlyVisibilityAttributes(Child, FSlateAttributeMetaData::EInvalidationPermission::AllowInvalidationIfConstructed);
-			}
+			FSlateAttributeMetaData::UpdateOnlyVisibilityAttributes(Child, FSlateAttributeMetaData::EInvalidationPermission::AllowInvalidationIfConstructed);
+		}
 
-			if (Child.GetVisibility() != EVisibility::Collapsed)
+		if (Child.GetVisibility() != EVisibility::Collapsed)
+		{
+			if (bUpdateAttributes)
 			{
-				if (bUpdateAttributes)
-				{
 #if WITH_SLATE_DEBUGGING
-					EVisibility PreviousVisibility = GetVisibility();
-					int32 PreviousAllChildrenNum = Child.GetAllChildren()->Num();
+				EVisibility PreviousVisibility = Self->GetVisibility();
+				int32 PreviousAllChildrenNum = Child.GetAllChildren()->Num();
 #endif
 
-					FSlateAttributeMetaData::UpdateExceptVisibilityAttributes(Child, FSlateAttributeMetaData::EInvalidationPermission::AllowInvalidationIfConstructed);
+				FSlateAttributeMetaData::UpdateExceptVisibilityAttributes(Child, FSlateAttributeMetaData::EInvalidationPermission::AllowInvalidationIfConstructed);
 
 #if WITH_SLATE_DEBUGGING
-					ensureMsgf(PreviousVisibility == GetVisibility(), TEXT("The visibility of widget '%s' doens't match the previous visibility after the attribute update."), *FReflectionMetaData::GetWidgetDebugInfo(this));
-					ensureMsgf(PreviousAllChildrenNum == Child.GetAllChildren()->Num(), TEXT("The number of child of widget '%s' doens't match the previous count after the attribute update."), *FReflectionMetaData::GetWidgetDebugInfo(this));
+				ensureMsgf(PreviousVisibility == Self->GetVisibility(), TEXT("The visibility of widget '%s' doesn't match the previous visibility after the attribute update."), *FReflectionMetaData::GetWidgetDebugInfo(Self));
+				ensureMsgf(PreviousAllChildrenNum == Child.GetAllChildren()->Num(), TEXT("The number of child of widget '%s' doesn't match the previous count after the attribute update."), *FReflectionMetaData::GetWidgetDebugInfo(Self));
 #endif
-				}
+			}
 
-				const float ChildLayoutScaleMultiplier = bHasRelativeLayoutScale
-					? InLayoutScaleMultiplier * GetRelativeLayoutScale(ChildIndex, InLayoutScaleMultiplier)
-					: InLayoutScaleMultiplier;
+			const float ChildLayoutScaleMultiplier = Self->bHasRelativeLayoutScale
+				? InLayoutScaleMultiplier * Self->GetRelativeLayoutScale(ChildIndex, InLayoutScaleMultiplier)
+				: InLayoutScaleMultiplier;
 
-				// Recur: Descend down the widget tree.
-				Child.Prepass_Internal(ChildLayoutScaleMultiplier);
-			}
-			else
-			{
-				// If the child widget is collapsed, we need to store the new layout scale it will have when 
-				// it is finally visible and invalidate it's prepass so that it gets that when its visibility
-				// is finally invalidated.
-				Child.MarkPrepassAsDirty();
-				Child.PrepassLayoutScaleMultiplier = bHasRelativeLayoutScale
-					? InLayoutScaleMultiplier * GetRelativeLayoutScale(ChildIndex, InLayoutScaleMultiplier)
-					: InLayoutScaleMultiplier;
-			}
-			++ChildIndex;
-		});
+			// Recur: Descend down the widget tree.
+			Child.Prepass_Internal(ChildLayoutScaleMultiplier);
+		}
+		else
+		{
+			// If the child widget is collapsed, we need to store the new layout scale it will have when 
+			// it is finally visible and invalidate it's prepass so that it gets that when its visibility
+			// is finally invalidated.
+			Child.MarkPrepassAsDirty();
+			Child.PrepassLayoutScaleMultiplier = Self->bHasRelativeLayoutScale
+				? InLayoutScaleMultiplier * Self->GetRelativeLayoutScale(ChildIndex, InLayoutScaleMultiplier)
+				: InLayoutScaleMultiplier;
+		}
+		++ChildIndex;
+	};
+
+	MyChildren->ForEachWidget(ForEachPred);
 }
 
 TSharedRef<FActiveTimerHandle> SWidget::RegisterActiveTimer(float TickPeriod, FWidgetActiveTimerDelegate TickFunction)
