Fix editor utility widget style changes not saving due to UMG style CDO comparisions reverting changes.
Also fixes issue where any 32x32 button texture would get shrunk to zero on editor utility widgets.

diff --git a/Engine/Source/Runtime/SlateCore/Private/Styling/StarshipCoreStyle.cpp b/Engine/Source/Runtime/SlateCore/Private/Styling/StarshipCoreStyle.cpp
--- a/Engine/Source/Runtime/SlateCore/Private/Styling/StarshipCoreStyle.cpp
+++ b/Engine/Source/Runtime/SlateCore/Private/Styling/StarshipCoreStyle.cpp
@@ -1157,6 +1157,10 @@ void FStarshipCoreStyle::SetupButtonStyles(TSharedRef<FStyle>& Style)
 		.SetNormalPadding(ButtonMargins)
 		.SetPressedPadding(PressedButtonMargins);
 
+	const FButtonStyle EditorUtilityButton = FButtonStyle(Button)
+		.SetNormal(FSlateRoundedBoxBrush(FStyleColors::Secondary, 4.0f, FStyleColors::Input, InputFocusThickness, Icon32x32))
+		.SetHovered(FSlateRoundedBoxBrush(FStyleColors::Hover, 4.0f, FStyleColors::Input, InputFocusThickness, Icon32x32))
+		.SetPressed(FSlateRoundedBoxBrush(FStyleColors::Header, 4.0f, FStyleColors::Input, InputFocusThickness, Icon32x32));
 
 	const FButtonStyle SimpleButton = FButtonStyle()
 		.SetNormal(FSlateNoResource())
@@ -1186,6 +1190,7 @@ void FStarshipCoreStyle::SetupButtonStyles(TSharedRef<FStyle>& Style)
 		Style->Set("PrimaryButton", PrimaryButton);
 		Style->Set("PrimaryButtonLabelAndIcon", PrimaryButtonLabelAndIcon);
 		Style->Set("Button", Button);
+		Style->Set("EditorUtilityButton", EditorUtilityButton);
 		Style->Set("SimpleButton", SimpleButtonIconOnly);
 		Style->Set("SimpleButtonLabelAndIcon", SimpleButtonLabelAndIcon);
 
@@ -1243,6 +1248,10 @@ void FStarshipCoreStyle::SetupComboButtonStyles(TSharedRef<FStyle>& Style)
 		.SetNormalPadding(ComboButtonMargin)
 		.SetPressedPadding(PressedComboButtonMargin);
 
+	const FButtonStyle EditorUtilityComboButtonButton = FButtonStyle(ComboButtonButton)
+		.SetNormal(FSlateRoundedBoxBrush(FStyleColors::Input, InputFocusRadius, FStyleColors::InputOutline, InputFocusThickness, Icon32x32))
+		.SetHovered(FSlateRoundedBoxBrush(FStyleColors::Input, InputFocusRadius, FStyleColors::Hover, InputFocusThickness, Icon32x32))
+		.SetPressed(FSlateRoundedBoxBrush(FStyleColors::Input, InputFocusRadius, FStyleColors::Hover, InputFocusThickness, Icon32x32));
 
 	// SComboBox 
 	FComboButtonStyle ComboButton = FComboButtonStyle()
@@ -1253,6 +1262,10 @@ void FStarshipCoreStyle::SetupComboButtonStyles(TSharedRef<FStyle>& Style)
 		.SetMenuBorderPadding(0.f);
 	Style->Set("ComboButton", ComboButton);
 
+	// Note: This style is currently not referenced by name in UMG, just used for EditorUtilityComboBox
+	FComboButtonStyle EditorUtilityComboButton = FComboButtonStyle(ComboButton)
+		.SetButtonStyle(EditorUtilityComboButtonButton);
+	Style->Set("EditorUtilityComboButton", ComboButton); 
 
 	FComboBoxStyle ComboBox = FComboBoxStyle()
 		.SetContentPadding(0.f)
@@ -1260,6 +1273,10 @@ void FStarshipCoreStyle::SetupComboButtonStyles(TSharedRef<FStyle>& Style)
 		.SetComboButtonStyle(ComboButton);
 	Style->Set("ComboBox", ComboBox);
 
+	FComboBoxStyle EditorUtilityComboBox = FComboBoxStyle(ComboBox)
+		.SetComboButtonStyle(EditorUtilityComboButton);
+	Style->Set("EditorUtilityComboBox", ComboBox);
+
 	const FButtonStyle& SimpleButton = Style->GetWidgetStyle<FButtonStyle>("SimpleButtonLabelAndIcon");
 
 	// Simple Combo Box  (borderless)
diff --git a/Engine/Source/Runtime/UMG/Private/Components/Button.cpp b/Engine/Source/Runtime/UMG/Private/Components/Button.cpp
--- a/Engine/Source/Runtime/UMG/Private/Components/Button.cpp
+++ b/Engine/Source/Runtime/UMG/Private/Components/Button.cpp
@@ -35,7 +35,7 @@ UButton::UButton(const FObjectInitializer& ObjectInitializer)
 #if WITH_EDITOR 
 	if (EditorButtonStyle == nullptr)
 	{
-		EditorButtonStyle = new FButtonStyle(FCoreStyle::Get().GetWidgetStyle<FButtonStyle>("Button"));
+		EditorButtonStyle = new FButtonStyle(FCoreStyle::Get().GetWidgetStyle<FButtonStyle>("EditorUtilityButton"));
 
 		// Unlink UMG Editor colors from the editor settings colors.
 		EditorButtonStyle->UnlinkColors();
diff --git a/Engine/Source/Runtime/UMG/Private/Components/ComboBoxKey.cpp b/Engine/Source/Runtime/UMG/Private/Components/ComboBoxKey.cpp
--- a/Engine/Source/Runtime/UMG/Private/Components/ComboBoxKey.cpp
+++ b/Engine/Source/Runtime/UMG/Private/Components/ComboBoxKey.cpp
@@ -20,7 +20,7 @@ UComboBoxKey::UComboBoxKey()
 	{
 #if WITH_EDITOR 
 		const ISlateStyle& SlateStyle = IsEditorWidget() ? FCoreStyle::Get() : FUMGCoreStyle::Get();
-		WidgetStyle = SlateStyle.GetWidgetStyle<FComboBoxStyle>("ComboBox");
+		WidgetStyle = SlateStyle.GetWidgetStyle<FComboBoxStyle>(IsEditorWidget() ? "EditorUtilityComboBox" : "ComboBox");
 		ItemStyle = SlateStyle.GetWidgetStyle<FTableRowStyle>("TableView.Row");
 
 		if (IsEditorWidget())
diff --git a/Engine/Source/Runtime/UMG/Private/Components/ComboBoxString.cpp b/Engine/Source/Runtime/UMG/Private/Components/ComboBoxString.cpp
--- a/Engine/Source/Runtime/UMG/Private/Components/ComboBoxString.cpp
+++ b/Engine/Source/Runtime/UMG/Private/Components/ComboBoxString.cpp
@@ -45,7 +45,7 @@ UComboBoxString::UComboBoxString(const FObjectInitializer& ObjectInitializer)
 #if WITH_EDITOR 
 	if (EditorComboBoxStyle == nullptr)
 	{
-		EditorComboBoxStyle = new FComboBoxStyle(FCoreStyle::Get().GetWidgetStyle<FComboBoxStyle>("ComboBox"));
+		EditorComboBoxStyle = new FComboBoxStyle(FCoreStyle::Get().GetWidgetStyle<FComboBoxStyle>("EditorUtilityComboBox"));
 
 		// Unlink UMG Editor colors from the editor settings colors.
 		EditorComboBoxStyle->UnlinkColors();
diff --git a/Engine/Source/Runtime/UMG/Private/Components/ListView.cpp b/Engine/Source/Runtime/UMG/Private/Components/ListView.cpp
--- a/Engine/Source/Runtime/UMG/Private/Components/ListView.cpp
+++ b/Engine/Source/Runtime/UMG/Private/Components/ListView.cpp
@@ -35,7 +35,7 @@ UListView::UListView(const FObjectInitializer& ObjectInitializer)
 
 	if (DefaultListViewScrollBarStyle == nullptr)
 	{
-		DefaultListViewScrollBarStyle = new FScrollBarStyle(FUMGCoreStyle::Get().GetWidgetStyle<FScrollBarStyle>("ScrollBar"));
+		DefaultListViewScrollBarStyle = new FScrollBarStyle(FUMGCoreStyle::Get().GetWidgetStyle<FScrollBarStyle>("Scrollbar"));
 
 		// Unlink UMG default colors.
 		DefaultListViewScrollBarStyle->UnlinkColors();
@@ -55,7 +55,7 @@ UListView::UListView(const FObjectInitializer& ObjectInitializer)
 
 	if (EditorListViewScrollBarStyle == nullptr)
 	{
-		EditorListViewScrollBarStyle = new FScrollBarStyle(FCoreStyle::Get().GetWidgetStyle<FScrollBarStyle>("ScrollBar"));
+		EditorListViewScrollBarStyle = new FScrollBarStyle(FCoreStyle::Get().GetWidgetStyle<FScrollBarStyle>("Scrollbar"));
 
 		// Unlink UMG default colors.
 		EditorListViewScrollBarStyle->UnlinkColors();
