Slate: Add a CrashHandler for Slate code.

diff --git a/Engine/Source/Runtime/Core/Private/GenericPlatform/GenericPlatformCrashContext.cpp b/Engine/Source/Runtime/Core/Private/GenericPlatform/GenericPlatformCrashContext.cpp
--- a/Engine/Source/Runtime/Core/Private/GenericPlatform/GenericPlatformCrashContext.cpp
+++ b/Engine/Source/Runtime/Core/Private/GenericPlatform/GenericPlatformCrashContext.cpp
@@ -84,6 +84,10 @@ int32 FGenericCrashContext::StaticCrashContextIndex = 0;
 
 const FGuid FGenericCrashContext::ExecutionGuid = FGuid::NewGuid();
 
+#if WITH_ADDITIONAL_CRASH_CONTEXTS
+FAdditionalCrashContextDelegate FGenericCrashContext::AdditionalCrashContextDelegate;
+#endif //WITH_ADDITIONAL_CRASH_CONTEXTS
+
 namespace NCached
 {
 	static FSessionContext Session;
@@ -1403,6 +1407,7 @@ void FGenericCrashContext::DumpAdditionalContext(const TCHAR* CrashFolderAbsolut
 {
 #if WITH_ADDITIONAL_CRASH_CONTEXTS 
 	FCrashContextExtendedWriterImpl Writer(CrashFolderAbsolute);
+	AdditionalCrashContextDelegate.Broadcast(Writer);
 	FAdditionalCrashContextStack::ExecuteProviders(Writer);
 #endif
 }
diff --git a/Engine/Source/Runtime/Core/Public/GenericPlatform/GenericPlatformCrashContext.h b/Engine/Source/Runtime/Core/Public/GenericPlatform/GenericPlatformCrashContext.h
--- a/Engine/Source/Runtime/Core/Public/GenericPlatform/GenericPlatformCrashContext.h
+++ b/Engine/Source/Runtime/Core/Public/GenericPlatform/GenericPlatformCrashContext.h
@@ -225,6 +225,25 @@ struct FSharedCrashContext
 	void*					ExceptionProgramCounter;
 };
 
+#if WITH_ADDITIONAL_CRASH_CONTEXTS
+
+/**
+ * Interface for callbacks to add context to the crash report.
+ */
+struct FCrashContextExtendedWriter
+{
+	/** Adds a named buffer to the report. Intended for larger payloads. */
+	CORE_API virtual void AddBuffer(const TCHAR* Identifier, const uint8* Data, uint32 DataSize) = 0;
+
+	/** Add a named buffer containing a string to the report. */
+	CORE_API virtual void AddString(const TCHAR* Identifier, const TCHAR* DataStr) = 0;
+};
+
+/** Simple Delegate for additional crash context. */
+DECLARE_MULTICAST_DELEGATE_OneParam(FAdditionalCrashContextDelegate, FCrashContextExtendedWriter&);
+
+#endif //WITH_ADDITIONAL_CRASH_CONTEXTS
+
 /**
  *	Contains a runtime crash's properties that are common for all platforms.
  *	This may change in the future.
@@ -428,11 +447,25 @@ public:
 	/** Notify the crash context exit has been requested. */
 	static void SetEngineExit(bool bIsRequestExit);
 
+#if WITH_ADDITIONAL_CRASH_CONTEXTS
+	/** Delegate for additional crash context. */
+	static FAdditionalCrashContextDelegate& OnAdditionalCrashContextDelegate()
+	{
+		return AdditionalCrashContextDelegate;
+	}
+#endif //WITH_ADDITIONAL_CRASH_CONTEXTS
+
 	/** Sets the process id to that has crashed. On supported platforms this will analyze the given process rather than current. Default is current process. */
-	void SetCrashedProcess(const FProcHandle& Process) { ProcessHandle = Process; }
+	void SetCrashedProcess(const FProcHandle& Process)
+	{
+		ProcessHandle = Process;
+	}
 
 	/** Stores crashing thread id. */
-	void SetCrashedThreadId(uint32 InId) { CrashedThreadId = InId; }
+	void SetCrashedThreadId(uint32 InId)
+	{
+		CrashedThreadId = InId;
+	}
 
 	/** Sets the number of stack frames to ignore when symbolicating from a minidump */
 	void SetNumMinidumpFramesToIgnore(int32 InNumMinidumpFramesToIgnore);
@@ -571,6 +604,11 @@ private:
 	/**	Static counter records how many crash contexts have been constructed */
 	static int32 StaticCrashContextIndex;
 
+#if WITH_ADDITIONAL_CRASH_CONTEXTS
+	/** Delegate for additional crash context. */
+	static FAdditionalCrashContextDelegate AdditionalCrashContextDelegate;
+#endif //WITH_ADDITIONAL_CRASH_CONTEXTS
+
 	/** The buffer used to store the crash's properties. */
 	mutable FString CommonBuffer;
 
@@ -599,18 +637,6 @@ namespace RecoveryService
 
 #if WITH_ADDITIONAL_CRASH_CONTEXTS
 
-/**
- * Interface for callbacks to add context to the crash report.
- */
-struct FCrashContextExtendedWriter
-{
-	/** Adds a named buffer to the report. Intended for larger payloads. */
-	CORE_API virtual void AddBuffer(const TCHAR* Identifier, const uint8* Data, uint32 DataSize) = 0;
-		
-	/** Add a named buffer containing a string to the report. */
-	CORE_API virtual void AddString(const TCHAR* Identifier, const TCHAR* DataStr) = 0;
-};
-
 /**
  * A thread local stack of callbacks that can be issued at time of the crash.
  */
diff --git a/Engine/Source/Runtime/SlateCore/Private/Debugging/SlateCrashReporterHandler.cpp b/Engine/Source/Runtime/SlateCore/Private/Debugging/SlateCrashReporterHandler.cpp
new file mode 100644
--- /dev/null
+++ b/Engine/Source/Runtime/SlateCore/Private/Debugging/SlateCrashReporterHandler.cpp
@@ -0,0 +1,194 @@
+// Copyright Epic Games, Inc. All Rights Reserved.
+
+#include "Debugging/SlateCrashReporterHandler.h"
+
+#include "HAL/ThreadManager.h"
+#include "Misc/ScopeLock.h"
+
+#include "Widgets/SWidget.h"
+
+#if UE_WITH_SLATE_CRASHREPORTER
+
+namespace UE::Slate::Private
+{
+
+static bool GSlateEnableCrashHandler = true;
+static FAutoConsoleVariableRef CVarEnableNiagaraCRHandler(
+	TEXT("Slate.EnableCrashHandler"),
+	GSlateEnableCrashHandler,
+	TEXT("Slate will add states into the crash reporter."),
+	ECVF_Default);
+
+
+struct FSlateCrashReporterHandler
+{
+public:
+	enum EThread
+	{
+		Thread_Slate = 0,	// The Slate thread when preloading occurs
+		Thread_Game = 1,	// The Main Game thread
+		Thread_Other = 2,	// All other threads. Slate should never runs on other threads.
+	};
+	static const uint8 NumberOfThread = 3;
+
+private:
+	struct FContext
+	{
+		FContext(const SWidget& Widget)
+			: Type(Widget.GetType())
+			, Tag(Widget.GetTag())
+			, CreatedInLocation(Widget.GetCreatedInLocation())
+			, ContextName(NAME_None)
+		{
+		}
+
+		FContext(const SWidget& Widget, FName InContextName)
+			: FContext(Widget)
+		{
+			ContextName = InContextName;
+		}
+
+		FName Type;
+		FName Tag;
+		FName CreatedInLocation;
+		FName ContextName;
+	};
+
+	FSlateCrashReporterHandler() = default;
+	FSlateCrashReporterHandler(const FSlateCrashReporterHandler&) = delete;
+	FSlateCrashReporterHandler& operator=(const FSlateCrashReporterHandler&) = delete;
+
+public:
+	static EThread GetThreadId()
+	{
+		return IsInGameThread() ? EThread::Thread_Game : (IsInSlateThread() ? EThread::Thread_Slate : EThread::Thread_Other);
+	}
+
+	static const TCHAR* LexToString(EThread Thread)
+	{
+		switch(Thread)
+		{
+		case EThread::Thread_Game:
+			return TEXT("Game");
+		case EThread::Thread_Slate:
+			return TEXT("Slate");
+		}
+		return TEXT("Other");
+	}
+
+	static FSlateCrashReporterHandler& Get()
+	{
+		struct FLocal
+		{
+			FSlateCrashReporterHandler* Instance;
+			FLocal()
+			{
+				Instance = new FSlateCrashReporterHandler();
+				FGenericCrashContext::OnAdditionalCrashContextDelegate().AddRaw(Instance, &FSlateCrashReporterHandler::DumpInfo);
+			}
+			~FLocal()
+			{
+				FGenericCrashContext::OnAdditionalCrashContextDelegate().RemoveAll(Instance);
+				delete Instance;
+			}
+		};
+		static FLocal Local;
+		return *Local.Instance;
+	}
+
+	EThread PushInfo(const SWidget& Widget, FName ContextName)
+	{
+		const EThread Thread = GetThreadId();
+		if (Thread == EThread::Thread_Other)
+		{
+			FScopeLock LockGuard(&CriticalSection);
+			Contexts[EThread::Thread_Other].Emplace(Widget, ContextName);
+		}
+		else
+		{
+			Contexts[Thread].Emplace(Widget, ContextName);
+		}
+		return Thread;
+	}
+
+	void PopInfo(EThread InThread)
+	{
+		if (InThread == EThread::Thread_Other)
+		{
+			FScopeLock LockGuard(&CriticalSection);
+			Contexts[EThread::Thread_Other].Pop();
+		}
+		else
+		{
+			Contexts[InThread].Pop();
+		}
+	}
+
+	void DumpInfo(FCrashContextExtendedWriter& Writer)
+	{
+		FScopeLock LockGuard(&CriticalSection);
+
+		Builder.Reset();
+		for (int32 ThreadIndex = 0; ThreadIndex < NumberOfThread; ++ThreadIndex)
+		{
+			if (Contexts[ThreadIndex].Num())
+			{
+				Builder.Append(TEXT("Thread: "));
+				Builder.Append(LexToString((EThread)ThreadIndex));
+				Builder.AppendChar(TEXT('\n'));
+				for (const FContext& Context : Contexts[ThreadIndex])
+				{
+					Builder.Append(TEXT("Context:'"));
+					Context.ContextName.AppendString(Builder);
+					Builder.Append(TEXT("' Type:'"));
+					Context.Type.AppendString(Builder);
+					Builder.Append(TEXT("' Tag:'"));
+					Context.Tag.AppendString(Builder);
+					Builder.Append(TEXT("' Location:'"));
+					Context.CreatedInLocation.AppendString(Builder);
+					Builder.Append(TEXT("'\n"));
+				}
+			}
+		}
+		
+		if (Builder.Len())
+		{
+			Writer.AddString(TEXT("Slate"), Builder.ToString());
+		}
+	}
+
+private:
+	FCriticalSection CriticalSection;
+	TArray<FContext> Contexts[NumberOfThread];
+	TStringBuilder<4096> Builder;
+};
+
+} // namespace
+
+namespace UE::Slate
+{
+
+FName FSlateCrashReporterScope::NAME_Paint = TEXT("Paint");
+FName FSlateCrashReporterScope::NAME_Prepass = TEXT("Prepass");
+
+FSlateCrashReporterScope::FSlateCrashReporterScope(const SWidget& Widget, FName Tag)
+{
+	bWasEnabled = Private::GSlateEnableCrashHandler != 0;
+	if (bWasEnabled)
+	{
+		ThreadId = Private::FSlateCrashReporterHandler::Get().PushInfo(Widget, Tag);
+	}
+}
+
+
+FSlateCrashReporterScope::~FSlateCrashReporterScope()
+{
+	if (bWasEnabled)
+	{
+		Private::FSlateCrashReporterHandler::Get().PopInfo((Private::FSlateCrashReporterHandler::EThread)ThreadId);
+	}
+}
+
+} //namespace
+
+#endif
\ No newline at end of file
diff --git a/Engine/Source/Runtime/SlateCore/Private/Widgets/SWidget.cpp b/Engine/Source/Runtime/SlateCore/Private/Widgets/SWidget.cpp
--- a/Engine/Source/Runtime/SlateCore/Private/Widgets/SWidget.cpp
+++ b/Engine/Source/Runtime/SlateCore/Private/Widgets/SWidget.cpp
@@ -15,6 +15,7 @@
 #include "Application/ActiveTimerHandle.h"
 #include "Input/HittestGrid.h"
 #include "Debugging/SlateDebugging.h"
+#include "Debugging/SlateCrashReporterHandler.h"
 #include "Debugging/WidgetList.h"
 #include "Widgets/SWindow.h"
 #include "Trace/SlateTrace.h"
@@ -679,6 +680,7 @@ void SWidget::SlatePrepass()
 
 void SWidget::SlatePrepass(float InLayoutScaleMultiplier)
 {
+	UE_SLATE_CRASH_REPORTER_PREPASS_SCOPE(*this);
 	SCOPE_CYCLE_COUNTER(STAT_SlatePrepass);
 
 	if (!GSlateIsOnFastUpdatePath || bNeedsPrepass)
diff --git a/Engine/Source/Runtime/SlateCore/Private/Widgets/SWindow.cpp b/Engine/Source/Runtime/SlateCore/Private/Widgets/SWindow.cpp
--- a/Engine/Source/Runtime/SlateCore/Private/Widgets/SWindow.cpp
+++ b/Engine/Source/Runtime/SlateCore/Private/Widgets/SWindow.cpp
@@ -3,6 +3,7 @@
 #include "Widgets/SWindow.h"
 #include "Application/SlateWindowHelper.h"
 #include "Application/SlateApplicationBase.h"
+#include "Debugging/SlateCrashReporterHandler.h"
 #include "Layout/WidgetPath.h"
 #include "Input/HittestGrid.h"
 #include "HAL/PlatformApplicationMisc.h"
@@ -2066,6 +2067,8 @@ int32 SWindow::PaintSlowPath(const FSlateInvalidationContext& Context)
 
 int32 SWindow::PaintWindow( double CurrentTime, float DeltaTime, FSlateWindowElementList& OutDrawElements, const FWidgetStyle& InWidgetStyle, bool bParentEnabled )
 {
+	UE_SLATE_CRASH_REPORTER_PAINT_SCOPE(*this);
+
 	OutDrawElements.BeginDeferredGroup();
 
 	const bool HittestCleared = HittestGrid->SetHittestArea(GetPositionInScreen(), GetViewportSize());
diff --git a/Engine/Source/Runtime/SlateCore/Public/Debugging/SlateCrashReporterHandler.h b/Engine/Source/Runtime/SlateCore/Public/Debugging/SlateCrashReporterHandler.h
new file mode 100644
--- /dev/null
+++ b/Engine/Source/Runtime/SlateCore/Public/Debugging/SlateCrashReporterHandler.h
@@ -0,0 +1,52 @@
+// Copyright Epic Games, Inc. All Rights Reserved.
+
+#pragma once
+
+#include "CoreMinimal.h"
+#include "GenericPlatform/GenericPlatformCrashContext.h"
+
+class SWidget;
+
+#ifndef UE_WITH_SLATE_CRASHREPORTER
+	#define UE_WITH_SLATE_CRASHREPORTER !(UE_BUILD_SHIPPING || UE_BUILD_TEST) && WITH_ADDITIONAL_CRASH_CONTEXTS
+#endif
+
+namespace UE::Slate
+{
+	
+#if UE_WITH_SLATE_CRASHREPORTER
+
+/** Helper object allowing easy tracking of Slate code in crash reporter. */
+class SLATECORE_API FSlateCrashReporterScope
+{
+public:
+	static FName NAME_Paint;
+	static FName NAME_Prepass;
+
+	explicit FSlateCrashReporterScope(const SWidget& Widget, FName Context);
+	~FSlateCrashReporterScope();
+
+private:
+	uint8 ThreadId = 0;
+	bool bWasEnabled = false;
+};
+
+#endif //UE_WITH_SLATE_CRASHREPORTER
+
+} //namespace
+
+#if UE_WITH_SLATE_CRASHREPORTER
+
+#define UE_SLATE_CRASH_REPORTER_PAINT_SCOPE(Widget) UE::Slate::FSlateCrashReporterScope ANONYMOUS_VARIABLE(SlateAddCrashContext) {Widget, UE::Slate::FSlateCrashReporterScope::NAME_Paint}
+#define UE_SLATE_CRASH_REPORTER_PREPASS_SCOPE(Widget) UE::Slate::FSlateCrashReporterScope ANONYMOUS_VARIABLE(SlateAddCrashContext) {Widget, UE::Slate::FSlateCrashReporterScope::NAME_Prepass}
+
+#define UE_SLATE_CRASH_REPORTER_CONTEXT_SCOPE(Widget, Context) UE::Slate::FSlateCrashReporterScope ANONYMOUS_VARIABLE(SlateAddCrashContext) {Widget, Context}
+
+#else //UE_WITH_SLATE_CRASHREPORTER
+
+#define UE_SLATE_CRASH_REPORTER_PAINT_SCOPE(Widget)
+#define UE_SLATE_CRASH_REPORTER_PREPASS_SCOPE(Widget)
+
+#define UE_SLATE_CRASH_REPORTER_CONTEXT_SCOPE(Widget, Context)
+
+#endif //UE_WITH_SLATE_CRASHREPORTER
\ No newline at end of file
