Fix menu's randomly jitter based on aspect ratio and screen resolution

diff --git a/Engine/Source/Runtime/SlateCore/Private/Widgets/SWindow.cpp b/Engine/Source/Runtime/SlateCore/Private/Widgets/SWindow.cpp
--- a/Engine/Source/Runtime/SlateCore/Private/Widgets/SWindow.cpp
+++ b/Engine/Source/Runtime/SlateCore/Private/Widgets/SWindow.cpp
@@ -947,7 +947,7 @@ void SWindow::ReshapeWindow( FVector2D NewPosition, FVector2D NewSize )
 	const FVector2D CurrentSize = GetSizeInScreen();
 
 	const FVector2D NewPositionTruncated = FVector2D(FMath::TruncToInt(NewPosition.X), FMath::TruncToInt(NewPosition.Y));
-	const FVector2D NewSizeRounded = FVector2D(FMath::CeilToInt(NewSize.X), FMath::CeilToInt(NewSize.Y));
+	const FVector2D NewSizeRounded = FVector2D(FMath::RoundToInt(NewSize.X), FMath::RoundToInt(NewSize.Y));
 
 	if ( CurrentPosition != NewPositionTruncated || CurrentSize != NewSizeRounded )
 	{
@@ -992,8 +992,8 @@ void SWindow::ResizeWindowSize( FVector2D NewWindowSize )
 	NewWindowSize.Y = FMath::Min(SizeLimits.GetMaxHeight().Get(NewWindowSize.Y), NewWindowSize.Y);
 
 	// ReshapeWindow W/H takes an int, so lets move our new W/H to int before checking if they are the same size
-	FIntPoint CurrentIntSize = FIntPoint(FMath::CeilToInt(Size.X), FMath::CeilToInt(Size.Y));
-	FIntPoint NewIntSize     = FIntPoint(FMath::CeilToInt(NewWindowSize.X), FMath::CeilToInt(NewWindowSize.Y));
+	FIntPoint CurrentIntSize = FIntPoint(FMath::RoundToInt(Size.X), FMath::RoundToInt(Size.Y));
+	FIntPoint NewIntSize     = FIntPoint(FMath::RoundToInt(NewWindowSize.X), FMath::RoundToInt(NewWindowSize.Y));
 
 	if (CurrentIntSize != NewIntSize)
 	{
