Slate: Reverted changes made in CL 19089835. It tried to clear the SlateBuffer as soon as we are over with it, but it may introduce racing conditions with the GC and the rendering thread.

diff --git a/Engine/Source/Runtime/SlateCore/Private/Rendering/SlateDrawBuffer.cpp b/Engine/Source/Runtime/SlateCore/Private/Rendering/SlateDrawBuffer.cpp
--- a/Engine/Source/Runtime/SlateCore/Private/Rendering/SlateDrawBuffer.cpp
+++ b/Engine/Source/Runtime/SlateCore/Private/Rendering/SlateDrawBuffer.cpp
@@ -27,7 +27,7 @@ FSlateWindowElementList& FSlateDrawBuffer::AddWindowElementList(TSharedRef<SWind
 			WindowElementLists.Add(ExistingElementList);
 			WindowElementListsPool.RemoveAtSwap(WindowIndex);
 
-			ensureMsgf(ExistingElementList->GetBatchData().GetNumFinalBatches() == 0, TEXT("The Buffer should have been clear when it was unlocked."));
+			ensureMsgf(ExistingElementList->GetBatchData().GetNumFinalBatches() == 0, TEXT("The Buffer should have been clear when it was locked."));
 			ExistingElementList->ResetElementList();
 
 			return *ExistingElementList;
@@ -64,6 +64,12 @@ bool FSlateDrawBuffer::Lock()
 	bool bIsLock = bLocked.compare_exchange_strong(ExpectedValue, true);
 	if (bIsLock)
 	{
+		// Rendering doesn't need the batch data anymore
+		for (TSharedRef<FSlateWindowElementList>& ExistingElementList : WindowElementLists)
+		{
+			ExistingElementList->ResetElementList();
+		}
+
 		bIsLockedBySlateThread = IsInSlateThread();
 	}
 	return bIsLock;
@@ -71,12 +77,6 @@ bool FSlateDrawBuffer::Lock()
 
 void FSlateDrawBuffer::Unlock()
 {
-	// Rendering doesn't need the batch data anymore
-	for (TSharedRef<FSlateWindowElementList>& ExistingElementList : WindowElementLists)
-	{
-		ExistingElementList->ResetElementList();
-	}
-
 	bLocked = false;
 }
 
