Slate: Some conversion for 2D to 2f in the Slate Renderer.
#prelfight 621f7023257fd6e0996086a8

diff --git a/Engine/Source/Runtime/SlateCore/Private/Rendering/DrawElements.cpp b/Engine/Source/Runtime/SlateCore/Private/Rendering/DrawElements.cpp
--- a/Engine/Source/Runtime/SlateCore/Private/Rendering/DrawElements.cpp
+++ b/Engine/Source/Runtime/SlateCore/Private/Rendering/DrawElements.cpp
@@ -49,7 +49,7 @@ static bool ShouldCull(const FSlateWindowElementList& ElementList, const FPaintG
 	const FVector2D& LocalSize = PaintGeometry.GetLocalSize();
 	const float DrawScale = PaintGeometry.DrawScale;
 
-	const FVector2D PixelSize = (LocalSize * DrawScale);
+	const FVector2f PixelSize = (LocalSize * DrawScale);
 	if (PixelSize.X <= 0.f || PixelSize.Y <= 0.f)
 	{
 		return true;
@@ -123,7 +123,7 @@ FSlateWindowElementList::FSlateWindowElementList(const TSharedPtr<SWindow>& InPa
 	, RenderTargetWindow(nullptr)
 	, bNeedsDeferredResolve(false)
 	, ResolveToDeferredIndex()
-	, WindowSize(FVector2D(0.0f, 0.0f))
+	, WindowSize(FVector2f(0.0f, 0.0f))
 	//, bReportReferences(true)
 {
 	if (InPaintWindow.IsValid())
@@ -193,7 +193,7 @@ void FSlateDrawElement::Init(FSlateWindowElementList& ElementList, EElementType
 	}
 }
 
-void FSlateDrawElement::ApplyPositionOffset(const FVector2D& InOffset)
+void FSlateDrawElement::ApplyPositionOffset(FVector2f InOffset)
 {
 	SetPosition(GetPosition() + InOffset);
 	RenderTransform = Concatenate(RenderTransform, InOffset);
@@ -288,18 +288,6 @@ void FSlateDrawElement::MakeBox(
 	MakeBoxInternal(ElementList, InLayer, PaintGeometry, InBrush, InDrawEffects, InTint);
 }
 
-void FSlateDrawElement::MakeBox( 
-	FSlateWindowElementList& ElementList,
-	uint32 InLayer, 
-	const FPaintGeometry& PaintGeometry, 
-	const FSlateBrush* InBrush, 
-	const FSlateResourceHandle& InRenderingHandle,
-	ESlateDrawEffect InDrawEffects, 
-	const FLinearColor& InTint )
-{
-	MakeBox(ElementList, InLayer, PaintGeometry, InBrush, InDrawEffects, InTint);
-}
-
 void FSlateDrawElement::MakeRotatedBox(
 	FSlateWindowElementList& ElementList,
 	uint32 InLayer,
@@ -471,7 +459,7 @@ void FSlateDrawElement::MakeGradient( FSlateWindowElementList& ElementList, uint
 	Element.Init(ElementList, EElementType::ET_Gradient, InLayer, PaintGeometry, InDrawEffects);
 }
 
-void FSlateDrawElement::MakeSpline( FSlateWindowElementList& ElementList, uint32 InLayer, const FPaintGeometry& PaintGeometry, const FVector2D& InStart, const FVector2D& InStartDir, const FVector2D& InEnd, const FVector2D& InEndDir, float InThickness, ESlateDrawEffect InDrawEffects, const FLinearColor& InTint )
+void FSlateDrawElement::MakeSpline( FSlateWindowElementList& ElementList, uint32 InLayer, const FPaintGeometry& PaintGeometry, FVector2f InStart, FVector2f InStartDir, FVector2f InEnd, FVector2f InEndDir, float InThickness, ESlateDrawEffect InDrawEffects, const FLinearColor& InTint )
 {
 	PaintGeometry.CommitTransformsIfUsingLegacyConstructor();
 
@@ -488,7 +476,7 @@ void FSlateDrawElement::MakeSpline( FSlateWindowElementList& ElementList, uint32
 	Element.Init(ElementList, EElementType::ET_Spline, InLayer, PaintGeometry, InDrawEffects);
 }
 
-void FSlateDrawElement::MakeCubicBezierSpline(FSlateWindowElementList & ElementList, uint32 InLayer, const FPaintGeometry & PaintGeometry, const FVector2D & P0, const FVector2D & P1, const FVector2D & P2, const FVector2D & P3, float InThickness, ESlateDrawEffect InDrawEffects, const FLinearColor & InTint)
+void FSlateDrawElement::MakeCubicBezierSpline(FSlateWindowElementList & ElementList, uint32 InLayer, const FPaintGeometry & PaintGeometry, FVector2f P0, FVector2f P1, FVector2f P2, FVector2f P3, float InThickness, ESlateDrawEffect InDrawEffects, const FLinearColor & InTint)
 {
 	PaintGeometry.CommitTransformsIfUsingLegacyConstructor();
 
@@ -505,47 +493,11 @@ void FSlateDrawElement::MakeCubicBezierSpline(FSlateWindowElementList & ElementL
 	Element.Init(ElementList, EElementType::ET_Spline, InLayer, PaintGeometry, InDrawEffects);
 }
 
-void FSlateDrawElement::MakeDrawSpaceSpline( FSlateWindowElementList& ElementList, uint32 InLayer, const FVector2D& InStart, const FVector2D& InStartDir, const FVector2D& InEnd, const FVector2D& InEndDir, float InThickness, ESlateDrawEffect InDrawEffects, const FLinearColor& InTint )
+void FSlateDrawElement::MakeDrawSpaceSpline( FSlateWindowElementList& ElementList, uint32 InLayer, FVector2f InStart, FVector2f InStartDir, FVector2f InEnd, FVector2f InEndDir, float InThickness, ESlateDrawEffect InDrawEffects, const FLinearColor& InTint )
 {
 	MakeSpline( ElementList, InLayer, FPaintGeometry(), InStart, InStartDir, InEnd, InEndDir, InThickness, InDrawEffects, InTint );
 }
 
-void FSlateDrawElement::MakeDrawSpaceGradientSpline(FSlateWindowElementList& ElementList, uint32 InLayer, const FVector2D& InStart, const FVector2D& InStartDir, const FVector2D& InEnd, const FVector2D& InEndDir, const TArray<FSlateGradientStop>& InGradientStops, float InThickness, ESlateDrawEffect InDrawEffects)
-{
-	const FPaintGeometry PaintGeometry;
-	PaintGeometry.CommitTransformsIfUsingLegacyConstructor();
-
-	if (ShouldCull(ElementList))
-	{
-		return;
-	}
-
-	FSlateDrawElement& Element = ElementList.AddUninitialized();
-
-	FSlateSplinePayload& DataPayload = ElementList.CreatePayload<FSlateSplinePayload>(Element);
-	DataPayload.SetGradientHermiteSpline(InStart, InStartDir, InEnd, InEndDir, InThickness, InGradientStops);
-
-	Element.Init(ElementList, EElementType::ET_Spline, InLayer, PaintGeometry, InDrawEffects);
-}
-
-void FSlateDrawElement::MakeDrawSpaceGradientSpline(FSlateWindowElementList& ElementList, uint32 InLayer, const FVector2D& InStart, const FVector2D& InStartDir, const FVector2D& InEnd, const FVector2D& InEndDir, const FSlateRect InClippingRect, const TArray<FSlateGradientStop>& InGradientStops, float InThickness, ESlateDrawEffect InDrawEffects)
-{
-	const FPaintGeometry PaintGeometry;
-	PaintGeometry.CommitTransformsIfUsingLegacyConstructor();
-
-	if (ShouldCull(ElementList))
-	{
-		return;
-	}
-
-	FSlateDrawElement& Element = ElementList.AddUninitialized();
-
-	FSlateSplinePayload& DataPayload = ElementList.CreatePayload<FSlateSplinePayload>(Element);
-	DataPayload.SetGradientHermiteSpline(InStart, InStartDir, InEnd, InEndDir, InThickness, InGradientStops);
-
-	Element.Init(ElementList, EElementType::ET_Spline, InLayer, PaintGeometry, InDrawEffects);
-}
-
 void FSlateDrawElement::MakeLines(FSlateWindowElementList& ElementList, uint32 InLayer, const FPaintGeometry& PaintGeometry, const TArray<FVector2D>& Points, ESlateDrawEffect InDrawEffects, const FLinearColor& InTint, bool bAntialias, float Thickness)
 {
 	PaintGeometry.CommitTransformsIfUsingLegacyConstructor();
diff --git a/Engine/Source/Runtime/SlateCore/Public/Rendering/DrawElementPayloads.h b/Engine/Source/Runtime/SlateCore/Public/Rendering/DrawElementPayloads.h
--- a/Engine/Source/Runtime/SlateCore/Public/Rendering/DrawElementPayloads.h
+++ b/Engine/Source/Runtime/SlateCore/Public/Rendering/DrawElementPayloads.h
@@ -6,7 +6,7 @@
 
 struct FSlateGradientStop
 {
-	FVector2D Position;
+	FVector2f Position;
 	FLinearColor Color;
 
 	/**
@@ -15,7 +15,7 @@ struct FSlateGradientStop
 						  A two stop gradient should go from (0,0), to (Width,Height).
 	 * @param InColor	- The color to lerp towards at this stop.
 	 */
-	FSlateGradientStop(const FVector2D& InPosition, const FLinearColor& InColor)
+	FSlateGradientStop(FVector2f InPosition, const FLinearColor& InColor)
 		: Position(InPosition)
 		, Color(InColor)
 	{
@@ -64,7 +64,7 @@ struct FSlateBoxPayload : public FSlateDataPayload, public FSlateTintableElement
 	ESlateBrushDrawType::Type GetBrushDrawType() const { return DrawType; }
 	const FSlateShaderResourceProxy* GetResourceProxy() const { return ResourceProxy; }
 
-	void SetBrush(const FSlateBrush* InBrush, FVector2D LocalSize, float DrawScale)
+	void SetBrush(const FSlateBrush* InBrush, FVector2f LocalSize, float DrawScale)
 	{
 		check(InBrush);
 		ensureMsgf(InBrush->GetDrawType() != ESlateBrushDrawType::NoDrawType, TEXT("This should have been filtered out earlier in the Make... call."));
@@ -226,10 +226,10 @@ struct FSlateSplinePayload : public FSlateDataPayload, public FSlateTintableElem
 	//     P0 *             * P3            P0 *   \   * P3
 	//                                              \ /
 	//                                               + P2	
-	FVector2D P0;
-	FVector2D P1;
-	FVector2D P2;
-	FVector2D P3;
+	FVector2f P0;
+	FVector2f P1;
+	FVector2f P2;
+	FVector2f P3;
 
 	float Thickness;
 
@@ -237,7 +237,7 @@ struct FSlateSplinePayload : public FSlateDataPayload, public FSlateTintableElem
 	void SetThickness(float InThickness) { Thickness = InThickness; }
 	float GetThickness() const { return Thickness; }
 
-	void SetCubicBezier(const FVector2D& InP0, const FVector2D& InP1, const FVector2D& InP2, const FVector2D& InP3, float InThickness, const FLinearColor& InTint)
+	void SetCubicBezier(FVector2f InP0, FVector2f InP1, FVector2f InP2, FVector2f InP3, float InThickness, const FLinearColor& InTint)
 	{
 		Tint = InTint;
 		P0 = InP0;
@@ -247,7 +247,7 @@ struct FSlateSplinePayload : public FSlateDataPayload, public FSlateTintableElem
 		Thickness = InThickness;
 	}
 
-	void SetHermiteSpline(const FVector2D& InStart, const FVector2D& InStartDir, const FVector2D& InEnd, const FVector2D& InEndDir, float InThickness, const FLinearColor& InTint)
+	void SetHermiteSpline(FVector2f InStart, FVector2f InStartDir, FVector2f InEnd, FVector2f InEndDir, float InThickness, const FLinearColor& InTint)
 	{
 		Tint = InTint;
 		P0 = InStart;
@@ -257,7 +257,7 @@ struct FSlateSplinePayload : public FSlateDataPayload, public FSlateTintableElem
 		Thickness = InThickness;
 	}
 
-	void SetGradientHermiteSpline(const FVector2D& InStart, const FVector2D& InStartDir, const FVector2D& InEnd, const FVector2D& InEndDir, float InThickness, const TArray<FSlateGradientStop>& InGradientStops)
+	void SetGradientHermiteSpline(FVector2f InStart, FVector2f InStartDir, FVector2f InEnd, FVector2f InEndDir, float InThickness, const TArray<FSlateGradientStop>& InGradientStops)
 	{
 		P0 = InStart;
 		P1 = InStart + InStartDir / 3.0f;
@@ -342,10 +342,10 @@ struct FSlateCachedBufferPayload : public FSlateDataPayload
 {
 	// Cached render data
 	class FSlateRenderDataHandle* CachedRenderData;
-	FVector2D CachedRenderDataOffset;
+	FVector2f CachedRenderDataOffset;
 
 	// Cached Buffers
-	void SetCachedBuffer(FSlateRenderDataHandle* InRenderDataHandle, const FVector2D& Offset)
+	void SetCachedBuffer(FSlateRenderDataHandle* InRenderDataHandle, FVector2f Offset)
 	{
 		check(InRenderDataHandle);
 
diff --git a/Engine/Source/Runtime/SlateCore/Public/Rendering/DrawElements.h b/Engine/Source/Runtime/SlateCore/Public/Rendering/DrawElements.h
--- a/Engine/Source/Runtime/SlateCore/Public/Rendering/DrawElements.h
+++ b/Engine/Source/Runtime/SlateCore/Public/Rendering/DrawElements.h
@@ -104,16 +104,6 @@ public:
 		const FLinearColor& InTint = FLinearColor::White );
 
 
-	UE_DEPRECATED(4.20, "Storing and passing in a FSlateResourceHandle to MakeBox is no longer necessary.")
-	SLATECORE_API static void MakeBox(
-		FSlateWindowElementList& ElementList,
-		uint32 InLayer, 
-		const FPaintGeometry& PaintGeometry, 
-		const FSlateBrush* InBrush, 
-		const FSlateResourceHandle& InRenderingHandle, 
-		ESlateDrawEffect InDrawEffects = ESlateDrawEffect::None, 
-		const FLinearColor& InTint = FLinearColor::White );
-
 	SLATECORE_API static void MakeRotatedBox(
 		FSlateWindowElementList& ElementList,
 		uint32 InLayer, 
@@ -185,7 +175,7 @@ public:
 	 * @param InDrawEffects         Optional draw effects to apply
 	 * @param InTint                Color to tint the element
 	 */
-	SLATECORE_API static void MakeSpline( FSlateWindowElementList& ElementList, uint32 InLayer, const FPaintGeometry& PaintGeometry, const FVector2D& InStart, const FVector2D& InStartDir, const FVector2D& InEnd, const FVector2D& InEndDir, float InThickness = 0.0f, ESlateDrawEffect InDrawEffects = ESlateDrawEffect::None, const FLinearColor& InTint=FLinearColor::White );
+	SLATECORE_API static void MakeSpline( FSlateWindowElementList& ElementList, uint32 InLayer, const FPaintGeometry& PaintGeometry, FVector2f InStart, FVector2f InStartDir, FVector2f InEnd, FVector2f InEndDir, float InThickness = 0.0f, ESlateDrawEffect InDrawEffects = ESlateDrawEffect::None, const FLinearColor& InTint=FLinearColor::White );
 
 	/**
 	 * Creates a Bezier Spline element
@@ -200,17 +190,10 @@ public:
 	 * @param InDrawEffects         Optional draw effects to apply
 	 * @param InTint                Color to tint the element
 	 */
-	SLATECORE_API static void MakeCubicBezierSpline(FSlateWindowElementList& ElementList, uint32 InLayer, const FPaintGeometry& PaintGeometry, const FVector2D& P0, const FVector2D& P1, const FVector2D& P2, const FVector2D& P3, float InThickness = 0.0f, ESlateDrawEffect InDrawEffects = ESlateDrawEffect::None, const FLinearColor& InTint = FLinearColor::White);
+	SLATECORE_API static void MakeCubicBezierSpline(FSlateWindowElementList& ElementList, uint32 InLayer, const FPaintGeometry& PaintGeometry, FVector2f P0, FVector2f P1, FVector2f P2, FVector2f P3, float InThickness = 0.0f, ESlateDrawEffect InDrawEffects = ESlateDrawEffect::None, const FLinearColor& InTint = FLinearColor::White);
 
 	/** Just like MakeSpline but in draw-space coordinates. This is useful for connecting already-transformed widgets together. */
-	SLATECORE_API static void MakeDrawSpaceSpline(FSlateWindowElementList& ElementList, uint32 InLayer, const FVector2D& InStart, const FVector2D& InStartDir, const FVector2D& InEnd, const FVector2D& InEndDir, float InThickness = 0.0f, ESlateDrawEffect InDrawEffects = ESlateDrawEffect::None, const FLinearColor& InTint=FLinearColor::White);
-
-	/** Just like MakeSpline but in draw-space coordinates. This is useful for connecting already-transformed widgets together. */
-	UE_DEPRECATED(4.20, "Splines with color gradients will not be supported in the future.")
-	SLATECORE_API static void MakeDrawSpaceGradientSpline( FSlateWindowElementList& ElementList, uint32 InLayer, const FVector2D& InStart, const FVector2D& InStartDir, const FVector2D& InEnd, const FVector2D& InEndDir, const TArray<FSlateGradientStop>& InGradientStops, float InThickness = 0.0f, ESlateDrawEffect InDrawEffects = ESlateDrawEffect::None );
-
-	UE_DEPRECATED(4.20, "Splines with color gradients will not be supported in the future.")
-	static void MakeDrawSpaceGradientSpline(FSlateWindowElementList& ElementList, uint32 InLayer, const FVector2D& InStart, const FVector2D& InStartDir, const FVector2D& InEnd, const FVector2D& InEndDir, const FSlateRect InClippingRect, const TArray<FSlateGradientStop>& InGradientStops, float InThickness = 0.0f, ESlateDrawEffect InDrawEffects = ESlateDrawEffect::None);
+	SLATECORE_API static void MakeDrawSpaceSpline(FSlateWindowElementList& ElementList, uint32 InLayer, FVector2f InStart, FVector2f InStartDir, FVector2f InEnd, FVector2f InEndDir, float InThickness = 0.0f, ESlateDrawEffect InDrawEffects = ESlateDrawEffect::None, const FLinearColor& InTint=FLinearColor::White);
 
 	/**
 	 * Creates a line defined by the provided points
@@ -267,29 +250,87 @@ public:
 	template<typename PayloadType>
 	FORCEINLINE PayloadType& GetDataPayload() { return *(PayloadType*)DataPayload; }
 
-	FORCEINLINE const FSlateRenderTransform& GetRenderTransform() const { return RenderTransform; }
-	FORCEINLINE void SetRenderTransform(const FSlateRenderTransform& InRenderTransform) { RenderTransform = InRenderTransform; }
-	FORCEINLINE FVector2D GetPosition() const { return FVector2D(Position); }
-	FORCEINLINE void SetPosition(const FVector2D& InPosition) { Position = FVector2f(InPosition); }
-	FORCEINLINE FVector2D GetLocalSize() const { return FVector2D(LocalSize); }
-	FORCEINLINE float GetScale() const { return Scale; }
-	FORCEINLINE ESlateDrawEffect GetDrawEffects() const { return DrawEffects; }
-	FORCEINLINE ESlateBatchDrawFlag GetBatchFlags() const { return BatchFlags; }
-	FORCEINLINE bool IsPixelSnapped() const { return !EnumHasAllFlags(DrawEffects, ESlateDrawEffect::NoPixelSnapping); }
+	FORCEINLINE const FSlateRenderTransform& GetRenderTransform() const
+	{
+		return RenderTransform;
+	}
 
-	FORCEINLINE int32 GetPrecachedClippingIndex() const { return ClipStateHandle.GetPrecachedClipIndex(); }
-	FORCEINLINE void SetPrecachedClippingIndex(int32 InClippingIndex) { ClipStateHandle.SetPreCachedClipIndex(InClippingIndex); }
+	FORCEINLINE void SetRenderTransform(const FSlateRenderTransform& InRenderTransform)
+	{
+		RenderTransform = InRenderTransform;
+	}
 
-	FORCEINLINE void SetCachedClippingState(const FSlateClippingState* CachedState) { ClipStateHandle.SetCachedClipState(CachedState); }
-	FORCEINLINE const FClipStateHandle& GetClippingHandle() const { return ClipStateHandle; }
-	FORCEINLINE const int8 GetSceneIndex() const { return SceneIndex; }
+	FORCEINLINE FVector2f GetPosition() const
+	{
+		return Position;
+	}
+
+	FORCEINLINE void SetPosition(FVector2f InPosition)
+	{
+		Position = InPosition;
+	}
 
-	FORCEINLINE void SetIsCached(bool bInIsCached) { bIsCached = bInIsCached; }
-	FORCEINLINE bool IsCached() const { return bIsCached; }
+	FORCEINLINE FVector2f GetLocalSize() const
+	{
+		return LocalSize;
+	}
+
+	FORCEINLINE float GetScale() const
+	{
+		return Scale;
+	}
+	FORCEINLINE ESlateDrawEffect GetDrawEffects() const
+	{
+		return DrawEffects;
+	}
+
+	FORCEINLINE ESlateBatchDrawFlag GetBatchFlags() const
+	{
+		return BatchFlags;
+	}
+
+	FORCEINLINE bool IsPixelSnapped() const
+	{
+		return !EnumHasAllFlags(DrawEffects, ESlateDrawEffect::NoPixelSnapping);
+	}
+
+	FORCEINLINE int32 GetPrecachedClippingIndex() const
+	{
+		return ClipStateHandle.GetPrecachedClipIndex();
+	}
+	FORCEINLINE void SetPrecachedClippingIndex(int32 InClippingIndex)
+	{
+		ClipStateHandle.SetPreCachedClipIndex(InClippingIndex);
+	}
+
+	FORCEINLINE void SetCachedClippingState(const FSlateClippingState* CachedState)
+	{
+		ClipStateHandle.SetCachedClipState(CachedState);
+	}
+
+	FORCEINLINE const FClipStateHandle& GetClippingHandle() const
+	{
+		return ClipStateHandle;
+	}
+
+	FORCEINLINE const int8 GetSceneIndex() const
+	{
+		return SceneIndex;
+	}
+
+	FORCEINLINE void SetIsCached(bool bInIsCached)
+	{
+		bIsCached = bInIsCached;
+	}
+
+	FORCEINLINE bool IsCached() const
+	{
+		return bIsCached;
+	}
 
 	FORCEINLINE FSlateLayoutTransform GetInverseLayoutTransform() const
 	{
-		return Inverse(FSlateLayoutTransform(Scale, FVector2D(Position)));
+		return Inverse(FSlateLayoutTransform(Scale, Position));
 	}
 
 	void AddReferencedObjects(FReferenceCollector& Collector);
@@ -300,13 +341,7 @@ public:
 	* @param Element		   Element to update
 	* @param InOffset         Absolute translation delta
 	*/
-	void ApplyPositionOffset(const FVector2D& InOffset);
-
-	UE_DEPRECATED(4.23, "GetClippingIndex has been deprecated.  If you were using this please use GetPrecachedClippingIndex instead.")
-	FORCEINLINE const int32 GetClippingIndex() const { return GetPrecachedClippingIndex(); }
-
-	UE_DEPRECATED(4.23, "SetClippingIndex has been deprecated.  If you were using this please use SetPrecachedClippingIndex instead.")
-	FORCEINLINE void SetClippingIndex(const int32 InClippingIndex) { SetPrecachedClippingIndex(InClippingIndex); }
+	void ApplyPositionOffset(FVector2f InOffset);
 
 private:
 	void Init(FSlateWindowElementList& ElementList, EElementType InElementType, uint32 InLayer, const FPaintGeometry& PaintGeometry, ESlateDrawEffect InDrawEffects);
@@ -467,14 +502,6 @@ public:
 
 	SLATECORE_API ~FSlateWindowElementList();
 
-	/** @return Get the window that we will be painting */
-	UE_DEPRECATED(4.21, "FSlateWindowElementList::GetWindow is not thread safe but window element lists are accessed on multiple threads.  Please call GetPaintWindow instead")
-	FORCEINLINE TSharedPtr<SWindow> GetWindow() const
-	{
-		// check that we are in game thread or are in slate/movie loading thread
-		check(IsInGameThread() || IsInSlateThread());
-		return WeakPaintWindow.Pin();
-	}
 
 	/** @return Get the window that we will be painting */
 	SWindow* GetPaintWindow() const
@@ -497,19 +524,8 @@ public:
 		return UncachedDrawElements;
 	}
 
-	/**
-	 * Add a draw element to the list
-	 *
-	 * @param InDrawElement  The draw element to add
-	 */
-	UE_DEPRECATED(4.23, "AddItem is deprecated, use AddUninitialized instead")
-	void AddItem(const FSlateDrawElement& InDrawElement)
-	{
-		UncachedDrawElements.Add(InDrawElement);
-	}
-
 	/** @return Get the window size that we will be painting */
-	FVector2D GetWindowSize() const
+	FVector2f GetWindowSize() const
 	{
 		return WindowSize;
 	}
@@ -586,12 +602,18 @@ public:
 
 	int32 PaintDeferred(int32 LayerId, const FSlateRect& MyCullingRect);
 
-	bool ShouldResolveDeferred() const { return bNeedsDeferredResolve; }
+	bool ShouldResolveDeferred() const
+	{
+		return bNeedsDeferredResolve;
+	}
 
 	SLATECORE_API void BeginDeferredGroup();
 	SLATECORE_API void EndDeferredGroup();
 
-	TArray< TSharedPtr<FDeferredPaint> > GetDeferredPaintList() const { return DeferredPaintList; }
+	TArray< TSharedPtr<FDeferredPaint> > GetDeferredPaintList() const
+	{
+		return DeferredPaintList;
+	}
 
 
 	//--------------------------------------------------------------------------
@@ -634,9 +656,16 @@ public:
 private:
 	FSlateDrawElement& AddCachedElement();
 
-	TArrayView<FSlateCachedElementData* const> GetCachedElementDataList() const { return MakeArrayView(CachedElementDataList.GetData(), CachedElementDataList.Num()); }
+	TArrayView<FSlateCachedElementData* const> GetCachedElementDataList() const
+	{
+		return MakeArrayView(CachedElementDataList.GetData(), CachedElementDataList.Num());
+	}
+
+	FSlateCachedElementData* GetCurrentCachedElementData() const
+	{
+		return CachedElementDataListStack.Num() ? CachedElementDataList[CachedElementDataListStack.Top()] : nullptr;
+	}
 
-	FSlateCachedElementData* GetCurrentCachedElementData() const { return CachedElementDataListStack.Num() ? CachedElementDataList[CachedElementDataListStack.Top()] : nullptr; }
 private:
 	/**
 	* Window which owns the widgets that are being painted but not necessarily rendered to
@@ -695,5 +724,5 @@ private:
 	// End Fast Path
 
 	/** Store the size of the window being used to paint */
-	FVector2D WindowSize;
+	FVector2f WindowSize;
 };
diff --git a/Engine/Source/Runtime/SlateCore/Public/Rendering/ElementBatcher.h b/Engine/Source/Runtime/SlateCore/Public/Rendering/ElementBatcher.h
--- a/Engine/Source/Runtime/SlateCore/Public/Rendering/ElementBatcher.h
+++ b/Engine/Source/Runtime/SlateCore/Public/Rendering/ElementBatcher.h
@@ -275,8 +275,8 @@ public:
 	}
 
 private:
-	void AddElementsInternal(const FSlateDrawElementArray& DrawElements, const FVector2D& ViewportSize);
-	void AddCachedElements(FSlateCachedElementData& CachedElementData, const FVector2D& ViewportSize);
+	void AddElementsInternal(const FSlateDrawElementArray& DrawElements, FVector2f ViewportSize);
+	void AddCachedElements(FSlateCachedElementData& CachedElementData, FVector2f ViewportSize);
 
 	/** 
 	 * Creates vertices necessary to draw a Quad element 
@@ -335,7 +335,7 @@ private:
 
 	void AddCustomVerts( const FSlateDrawElement& DrawElement );
 
-	void AddPostProcessPass(const FSlateDrawElement& DrawElement, const FVector2D& WindowSize);
+	void AddPostProcessPass(const FSlateDrawElement& DrawElement, FVector2f WindowSize);
 
 	FSlateRenderBatch& CreateRenderBatch(
 		int32 Layer,
