Fix additional menu jitters by saving rounded window sizes & restoring ceil usage to prevent unnecessary scroll bars.

diff --git a/Engine/Source/Runtime/SlateCore/Private/Widgets/SWindow.cpp b/Engine/Source/Runtime/SlateCore/Private/Widgets/SWindow.cpp
--- a/Engine/Source/Runtime/SlateCore/Private/Widgets/SWindow.cpp
+++ b/Engine/Source/Runtime/SlateCore/Private/Widgets/SWindow.cpp
@@ -951,10 +951,13 @@ void SWindow::ReshapeWindow( FVector2D NewPosition, FVector2D NewSize )
 	const FVector2D CurrentPosition = GetPositionInScreen();
 	const FVector2D CurrentSize = GetSizeInScreen();
 
+	// Ceil (Minus a tad for float precision) to ensure contents are not a sub-pixel larger than the window, which will create unnecessary scroll bars 
+	const FVector2D OldPositionTruncated = FVector2D(FMath::TruncToInt(CurrentPosition.X), FMath::TruncToInt(CurrentPosition.Y));
+	const FVector2D OldSizeRounded = FVector2D(FMath::CeilToInt(CurrentSize.X - KINDA_SMALL_NUMBER), FMath::CeilToInt(CurrentSize.Y - KINDA_SMALL_NUMBER));
 	const FVector2D NewPositionTruncated = FVector2D(FMath::TruncToInt(NewPosition.X), FMath::TruncToInt(NewPosition.Y));
-	const FVector2D NewSizeRounded = FVector2D(FMath::RoundToInt(NewSize.X), FMath::RoundToInt(NewSize.Y));
+	const FVector2D NewSizeRounded = FVector2D(FMath::CeilToInt(NewSize.X - KINDA_SMALL_NUMBER), FMath::CeilToInt(NewSize.Y - KINDA_SMALL_NUMBER));
 
-	if ( CurrentPosition != NewPositionTruncated || CurrentSize != NewSizeRounded )
+	if (OldPositionTruncated != NewPositionTruncated || OldSizeRounded != NewSizeRounded )
 	{
 		if ( NativeWindow.IsValid() )
 		{
@@ -968,11 +971,11 @@ void SWindow::ReshapeWindow( FVector2D NewPosition, FVector2D NewSize )
 		}
 		else
 		{
-			InitialDesiredScreenPosition = NewPosition;
-			InitialDesiredSize = NewSize;
+			InitialDesiredScreenPosition = NewPositionTruncated;
+			InitialDesiredSize = NewSizeRounded;
 		}
 
-		SetCachedSize(NewSize);
+		SetCachedSize(NewSizeRounded);
 	}
 }
 
@@ -996,9 +999,9 @@ void SWindow::ResizeWindowSize( FVector2D NewWindowSize )
 	NewWindowSize.Y = FMath::Max(SizeLimits.GetMinHeight().Get(NewWindowSize.Y), NewWindowSize.Y);
 	NewWindowSize.Y = FMath::Min(SizeLimits.GetMaxHeight().Get(NewWindowSize.Y), NewWindowSize.Y);
 
-	// ReshapeWindow W/H takes an int, so lets move our new W/H to int before checking if they are the same size
-	FIntPoint CurrentIntSize = FIntPoint(FMath::RoundToInt(Size.X), FMath::RoundToInt(Size.Y));
-	FIntPoint NewIntSize     = FIntPoint(FMath::RoundToInt(NewWindowSize.X), FMath::RoundToInt(NewWindowSize.Y));
+	// Ceil (Minus a tad for float precision) to ensure contents are not a sub-pixel larger than the window, which will create unnecessary scroll bars 
+	FIntPoint CurrentIntSize = FIntPoint(FMath::CeilToInt(Size.X - KINDA_SMALL_NUMBER), FMath::CeilToInt(Size.Y - KINDA_SMALL_NUMBER));
+	FIntPoint NewIntSize     = FIntPoint(FMath::CeilToInt(NewWindowSize.X - KINDA_SMALL_NUMBER), FMath::CeilToInt(NewWindowSize.Y - KINDA_SMALL_NUMBER));
 
 	if (CurrentIntSize != NewIntSize)
 	{
@@ -1008,10 +1011,10 @@ void SWindow::ResizeWindowSize( FVector2D NewWindowSize )
 		}
 		else
 		{
-			InitialDesiredSize = NewWindowSize;
+			InitialDesiredSize = NewIntSize;
 		}
 	}
-	SetCachedSize(NewWindowSize);
+	SetCachedSize(NewIntSize);
 }
 
 FSlateRect SWindow::GetFullScreenInfo() const
