Slate: Clear the DrawElementList when the render has finished with the elements.
#prefligh 6222030ba00412627d3c733e

diff --git a/Engine/Source/Runtime/SlateCore/Private/Rendering/DrawElements.cpp b/Engine/Source/Runtime/SlateCore/Private/Rendering/DrawElements.cpp
--- a/Engine/Source/Runtime/SlateCore/Private/Rendering/DrawElements.cpp
+++ b/Engine/Source/Runtime/SlateCore/Private/Rendering/DrawElements.cpp
@@ -994,6 +994,7 @@ void FSlateWindowElementList::ResetElementList()
 	CachedElementDataListStack.Empty();
 
 	check(WidgetDrawStack.Num() == 0);
+	check(ResolveToDeferredIndex.Num() == 0);
 
 	RenderTargetWindow = nullptr;
 }
diff --git a/Engine/Source/Runtime/SlateCore/Private/Rendering/SlateDrawBuffer.cpp b/Engine/Source/Runtime/SlateCore/Private/Rendering/SlateDrawBuffer.cpp
--- a/Engine/Source/Runtime/SlateCore/Private/Rendering/SlateDrawBuffer.cpp
+++ b/Engine/Source/Runtime/SlateCore/Private/Rendering/SlateDrawBuffer.cpp
@@ -10,13 +10,16 @@
 /* FSlateDrawBuffer interface
  *****************************************************************************/
 
-FSlateDrawBuffer::~FSlateDrawBuffer()
-{
-}
+FSlateDrawBuffer::FSlateDrawBuffer()
+	: bIsLocked(false)
+	, bIsLockedBySlateThread(false)
+	, ResourceVersion(0)
+{ }
 
 FSlateWindowElementList& FSlateDrawBuffer::AddWindowElementList(TSharedRef<SWindow> ForWindow)
 {
 	ensureMsgf(IsLocked(), TEXT("The SlateDrawBuffer should be lock before modifying it."));
+	FScopeLock ScopeLock(&GCLock);
 
 	for ( int32 WindowIndex = 0; WindowIndex < WindowElementListsPool.Num(); ++WindowIndex )
 	{
@@ -44,6 +47,7 @@ FSlateWindowElementList& FSlateDrawBuffer::AddWindowElementList(TSharedRef<SWind
 void FSlateDrawBuffer::RemoveUnusedWindowElement(const TArray<SWindow*>& AllWindows)
 {
 	ensureMsgf(IsLocked(), TEXT("The SlateDrawBuffer should be lock before modifying it."));
+	FScopeLock ScopeLock(&GCLock);
 
 	// Remove any window elements that are no longer valid.
 	for (int32 WindowIndex = 0; WindowIndex < WindowElementLists.Num(); ++WindowIndex)
@@ -61,15 +65,9 @@ void FSlateDrawBuffer::RemoveUnusedWindowElement(const TArray<SWindow*>& AllWind
 bool FSlateDrawBuffer::Lock()
 {
 	bool ExpectedValue = false;
-	bool bIsLock = bLocked.compare_exchange_strong(ExpectedValue, true);
+	bool bIsLock = bIsLocked.compare_exchange_strong(ExpectedValue, true);
 	if (bIsLock)
 	{
-		// Rendering doesn't need the batch data anymore
-		for (TSharedRef<FSlateWindowElementList>& ExistingElementList : WindowElementLists)
-		{
-			ExistingElementList->ResetElementList();
-		}
-
 		bIsLockedBySlateThread = IsInSlateThread();
 	}
 	return bIsLock;
@@ -77,15 +75,25 @@ bool FSlateDrawBuffer::Lock()
 
 void FSlateDrawBuffer::Unlock()
 {
-	bLocked = false;
+	ensureMsgf(IsLocked(), TEXT("The SlateDrawBuffer should be lock before modifying it."));
+	FScopeLock ScopeLock(&GCLock);
+
+	// Rendering doesn't need the batch data anymore
+	for (TSharedRef<FSlateWindowElementList>& ExistingElementList : WindowElementLists)
+	{
+		ExistingElementList->ResetElementList();
+	}
+
+	bIsLocked = false;
 }
 
 void FSlateDrawBuffer::AddReferencedObjects(FReferenceCollector& Collector)
 {
 	// Locked buffers are the only ones that are currently referencing objects.
 	//If unlocked, the element list is not in-use and contains "to-be-cleared" data.
-	if(bLocked && !bIsLockedBySlateThread)
+	if(bIsLocked && !bIsLockedBySlateThread)
 	{
+		FScopeLock ScopeLock(&GCLock);
 		for (TSharedRef<FSlateWindowElementList>& ElementList : WindowElementLists)
 		{
 			ElementList->AddReferencedObjects(Collector);
@@ -96,6 +104,7 @@ void FSlateDrawBuffer::AddReferencedObjects(FReferenceCollector& Collector)
 void FSlateDrawBuffer::ClearBuffer()
 {
 	ensureMsgf(IsLocked(), TEXT("The SlateDrawBuffer should be lock before modifying it."));
+	FScopeLock ScopeLock(&GCLock);
 
 	// Remove any window elements that are no longer valid.
 	for (int32 WindowIndex = 0; WindowIndex < WindowElementListsPool.Num(); ++WindowIndex)
diff --git a/Engine/Source/Runtime/SlateCore/Public/Rendering/SlateDrawBuffer.h b/Engine/Source/Runtime/SlateCore/Public/Rendering/SlateDrawBuffer.h
--- a/Engine/Source/Runtime/SlateCore/Public/Rendering/SlateDrawBuffer.h
+++ b/Engine/Source/Runtime/SlateCore/Public/Rendering/SlateDrawBuffer.h
@@ -4,6 +4,7 @@
 
 #include "CoreMinimal.h"
 #include "UObject/GCObject.h"
+#include "HAL/CriticalSection.h"
 
 #include <atomic>
 
@@ -16,16 +17,8 @@ class SWindow;
 class SLATECORE_API FSlateDrawBuffer : public FGCObject
 {
 public:
-
 	/** Default constructor. */
-	explicit FSlateDrawBuffer()
-		: bLocked(false)
-		, bIsLockedBySlateThread(false)
-		, ResourceVersion(0)
-	{ }
-
-public:
-	~FSlateDrawBuffer();
+	FSlateDrawBuffer();
 
 	/** Removes all data from the buffer. */
 	void ClearBuffer();
@@ -69,7 +62,7 @@ public:
 	/** @return true if the buffer is locked. */
 	bool IsLocked() const
 	{
-		return bLocked;
+		return bIsLocked;
 	}
 
 	/** FGCObject Interface */
@@ -79,15 +72,16 @@ public:
 		return TEXT("FSlateDrawBuffer for Uncached Elements");
 	}
 
-protected:
-	// List of window element lists.
+private:
+	// List of window element lists
 	TArray< TSharedRef<FSlateWindowElementList> > WindowElementLists;
 
 	// List of window element lists that we store from the previous frame 
 	// that we restore if they're requested again.
 	TArray< TSharedRef<FSlateWindowElementList> > WindowElementListsPool;
 
-	std::atomic<bool> bLocked;
+	FCriticalSection GCLock;
+	std::atomic<bool> bIsLocked;
 	bool bIsLockedBySlateThread;
 
 	// Last recorded version from the render. The WindowElementListsPool is emptied when this changes.
