Get 16 bytes back on FSlateBrush
#rnx

diff --git a/Engine/Source/Runtime/SlateCore/Private/Styling/SlateBrush.cpp b/Engine/Source/Runtime/SlateCore/Private/Styling/SlateBrush.cpp
--- a/Engine/Source/Runtime/SlateCore/Private/Styling/SlateBrush.cpp
+++ b/Engine/Source/Runtime/SlateCore/Private/Styling/SlateBrush.cpp
@@ -11,9 +11,6 @@ FSlateBrush::FSlateBrush()
 	, Tint_DEPRECATED(FLinearColor::White)
 #endif
 	, TintColor(FLinearColor::White)
-	, ResourceObject(nullptr)
-	, ResourceName(NAME_None)
-	, UVRegion(ForceInit)
 	, DrawAs(ESlateBrushDrawType::Image)
 	, Tiling(ESlateBrushTileType::NoTile)
 	, Mirroring(ESlateBrushMirrorType::NoMirror)
@@ -21,6 +18,9 @@ FSlateBrush::FSlateBrush()
 	, bIsDynamicallyLoaded(false)
 	, bHasUObject_DEPRECATED(false)
 	, bIsSet(true)
+	, ResourceObject(nullptr)
+	, ResourceName(NAME_None)
+	, UVRegion(ForceInit)
 {
 }
 
@@ -41,15 +41,15 @@ FSlateBrush::FSlateBrush( ESlateBrushDrawType::Type InDrawType,
 	, Tint_DEPRECATED(FLinearColor::White)
 #endif
 	, TintColor( InTint )
+	, DrawAs(InDrawType)
+	, Tiling(InTiling)
+	, Mirroring(ESlateBrushMirrorType::NoMirror)
+	, ImageType(InImageType)
+	, bIsDynamicallyLoaded(bInDynamicallyLoaded)
+	, bIsSet(true)
 	, ResourceObject( InObjectResource )
 	, ResourceName( InResourceName )
 	, UVRegion( ForceInit )
-	, DrawAs( InDrawType )
-	, Tiling( InTiling )
-	, Mirroring( ESlateBrushMirrorType::NoMirror )
-	, ImageType( InImageType )
-	, bIsDynamicallyLoaded( bInDynamicallyLoaded )
-	, bIsSet(true)
 {
 	bHasUObject_DEPRECATED = (InObjectResource != nullptr) || InResourceName.ToString().StartsWith(FSlateBrush::UTextureIdentifier());
 
@@ -77,15 +77,15 @@ FSlateBrush::FSlateBrush( ESlateBrushDrawType::Type InDrawType,
 	, Tint_DEPRECATED(FLinearColor::White)
 #endif
 	, TintColor( InTint )
+	, DrawAs(InDrawType)
+	, Tiling(InTiling)
+	, Mirroring(ESlateBrushMirrorType::NoMirror)
+	, ImageType(InImageType)
+	, bIsDynamicallyLoaded(bInDynamicallyLoaded)
+	, bIsSet(true)
 	, ResourceObject( InObjectResource )
 	, ResourceName( InResourceName )
 	, UVRegion( ForceInit )
-	, DrawAs( InDrawType )
-	, Tiling( InTiling )
-	, Mirroring( ESlateBrushMirrorType::NoMirror )
-	, ImageType( InImageType )
-	, bIsDynamicallyLoaded( bInDynamicallyLoaded )
-	, bIsSet(true)
 {
 	bHasUObject_DEPRECATED = (InObjectResource != nullptr) || InResourceName.ToString().StartsWith(FSlateBrush::UTextureIdentifier());
 
@@ -112,15 +112,15 @@ FSlateBrush::FSlateBrush( ESlateBrushDrawType::Type InDrawType,
 	, Tint_DEPRECATED(FLinearColor::White)
 #endif
 	, TintColor(InTint)
-	, ResourceObject(InObjectResource)
-	, ResourceName(InResourceName)
-	, UVRegion(ForceInit)
 	, DrawAs(InDrawType)
 	, Tiling(InTiling)
-	, Mirroring( ESlateBrushMirrorType::NoMirror )
+	, Mirroring(ESlateBrushMirrorType::NoMirror)
 	, ImageType(InImageType)
 	, bIsDynamicallyLoaded(bInDynamicallyLoaded)
 	, bIsSet(true)
+	, ResourceObject(InObjectResource)
+	, ResourceName(InResourceName)
+	, UVRegion(ForceInit)
 {
 	bHasUObject_DEPRECATED = (InObjectResource != nullptr) || InResourceName.ToString().StartsWith(FSlateBrush::UTextureIdentifier());
 
diff --git a/Engine/Source/Runtime/SlateCore/Public/Styling/SlateBrush.h b/Engine/Source/Runtime/SlateCore/Public/Styling/SlateBrush.h
--- a/Engine/Source/Runtime/SlateCore/Public/Styling/SlateBrush.h
+++ b/Engine/Source/Runtime/SlateCore/Public/Styling/SlateBrush.h
@@ -213,7 +213,38 @@
 	GENERATED_USTRUCT_BODY()
 
 	friend class FSlateShaderResourceManager;
+
+protected:
+	/** Whether or not the brush path is a path to a UObject */
+	UPROPERTY()
+	uint8 bIsDynamicallyLoaded:1;
+
+	/** Whether or not the brush has a UTexture resource */
+	UPROPERTY()
+	uint8 bHasUObject_DEPRECATED:1;
+
+	/** This is true for all constructed brushes except for optional brushes */
+	uint8 bIsSet : 1;
+
 public:
+
+	/** How to draw the image */
+	UPROPERTY(EditAnywhere, BlueprintReadWrite, Category=Brush)
+	TEnumAsByte<enum ESlateBrushDrawType::Type > DrawAs;
+
+	/** How to tile the image in Image mode */
+	UPROPERTY(EditAnywhere, BlueprintReadWrite, Category=Brush)
+	TEnumAsByte<enum ESlateBrushTileType::Type> Tiling;
+
+	/** How to mirror the image in Image mode.  This is normally only used for dynamic image brushes where the source texture
+	    comes from a hardware device such as a web camera. */
+	UPROPERTY(EditAnywhere, BlueprintReadWrite, Category=Brush)
+	TEnumAsByte<enum ESlateBrushMirrorType::Type> Mirroring;
+
+	/** The type of image */
+	UPROPERTY()
+	TEnumAsByte<enum ESlateBrushImageType::Type> ImageType;
+
 	/** Size of the resource in Slate Units */
 	UPROPERTY(EditAnywhere, BlueprintReadWrite, Category=Brush)
 	FVector2D ImageSize;
@@ -232,6 +263,7 @@
 	UPROPERTY(EditAnywhere, BlueprintReadWrite, Category=Brush, meta=( DisplayName="Tint", sRGB="true" ))
 	FSlateColor TintColor;
 
+public:
 	/** How to draw the outline.  Currently only used for RoundedBox type brushes. */
 	UPROPERTY(EditAnywhere, BlueprintReadWrite, Category=Brush)
 	FSlateBrushOutlineSettings OutlineSettings;
@@ -449,38 +481,11 @@
 	FBox2f UVRegion;
 
 public:
-	/** How to draw the image */
-	UPROPERTY(EditAnywhere, BlueprintReadWrite, Category=Brush)
-	TEnumAsByte<enum ESlateBrushDrawType::Type > DrawAs;
 
-	/** How to tile the image in Image mode */
-	UPROPERTY(EditAnywhere, BlueprintReadWrite, Category=Brush)
-	TEnumAsByte<enum ESlateBrushTileType::Type> Tiling;
-
-	/** How to mirror the image in Image mode.  This is normally only used for dynamic image brushes where the source texture
-	    comes from a hardware device such as a web camera. */
-	UPROPERTY(EditAnywhere, BlueprintReadWrite, Category=Brush)
-	TEnumAsByte<enum ESlateBrushMirrorType::Type> Mirroring;
-
-	/** The type of image */
-	UPROPERTY()
-	TEnumAsByte<enum ESlateBrushImageType::Type> ImageType;
-
 	/** Rendering resource for this brush */
 	mutable FSlateResourceHandle ResourceHandle;
 protected:
 
-	/** Whether or not the brush path is a path to a UObject */
-	UPROPERTY()
-	uint8 bIsDynamicallyLoaded:1;
-
-	/** Whether or not the brush has a UTexture resource */
-	UPROPERTY()
-	uint8 bHasUObject_DEPRECATED:1;
-
-	/** This is true for all constructed brushes except for optional brushes */
-	uint8 bIsSet : 1;
-
 	/** 
 	 * This constructor is protected; use one of the deriving classes instead.
 	 *
