Fix initialization order
#rnx

diff --git a/Engine/Source/Runtime/SlateCore/Private/Styling/SlateBrush.cpp b/Engine/Source/Runtime/SlateCore/Private/Styling/SlateBrush.cpp
--- a/Engine/Source/Runtime/SlateCore/Private/Styling/SlateBrush.cpp
+++ b/Engine/Source/Runtime/SlateCore/Private/Styling/SlateBrush.cpp
@@ -5,19 +5,19 @@
 #include "Application/SlateApplicationBase.h"
 
 FSlateBrush::FSlateBrush()
-	: ImageSize(SlateBrushDefs::DefaultImageSize, SlateBrushDefs::DefaultImageSize)
+	: bIsDynamicallyLoaded(false)
+	, bHasUObject_DEPRECATED(false)
+	, bIsSet(true)
+	, DrawAs(ESlateBrushDrawType::Image)
+	, Tiling(ESlateBrushTileType::NoTile)
+	, Mirroring(ESlateBrushMirrorType::NoMirror)
+	, ImageType(ESlateBrushImageType::NoImage)
+	, ImageSize(SlateBrushDefs::DefaultImageSize, SlateBrushDefs::DefaultImageSize)
 	, Margin(0.0f)
 #if WITH_EDITORONLY_DATA
 	, Tint_DEPRECATED(FLinearColor::White)
 #endif
 	, TintColor(FLinearColor::White)
-	, DrawAs(ESlateBrushDrawType::Image)
-	, Tiling(ESlateBrushTileType::NoTile)
-	, Mirroring(ESlateBrushMirrorType::NoMirror)
-	, ImageType(ESlateBrushImageType::NoImage)
-	, bIsDynamicallyLoaded(false)
-	, bHasUObject_DEPRECATED(false)
-	, bIsSet(true)
 	, ResourceObject(nullptr)
 	, ResourceName(NAME_None)
 	, UVRegion(ForceInit)
@@ -35,18 +35,18 @@
 						  bool bInDynamicallyLoaded
 						)
 
-	: ImageSize( InImageSize )
+	: bIsDynamicallyLoaded(bInDynamicallyLoaded)
+	, bIsSet(true)
+	, DrawAs(InDrawType)
+	, Tiling(InTiling)
+	, Mirroring(ESlateBrushMirrorType::NoMirror)
+	, ImageType(InImageType)
+	, ImageSize( InImageSize )
 	, Margin( InMargin )
 #if WITH_EDITORONLY_DATA
 	, Tint_DEPRECATED(FLinearColor::White)
 #endif
 	, TintColor( InTint )
-	, DrawAs(InDrawType)
-	, Tiling(InTiling)
-	, Mirroring(ESlateBrushMirrorType::NoMirror)
-	, ImageType(InImageType)
-	, bIsDynamicallyLoaded(bInDynamicallyLoaded)
-	, bIsSet(true)
 	, ResourceObject( InObjectResource )
 	, ResourceName( InResourceName )
 	, UVRegion( ForceInit )
@@ -71,18 +71,18 @@
  						  bool bInDynamicallyLoaded
  						)
 
-	: ImageSize( InImageSize )
+	: bIsDynamicallyLoaded(bInDynamicallyLoaded)
+	, bIsSet(true)
+	, DrawAs(InDrawType)
+	, Tiling(InTiling)
+	, Mirroring(ESlateBrushMirrorType::NoMirror)
+	, ImageType(InImageType)
+	, ImageSize(InImageSize)
 	, Margin( InMargin )
 #if WITH_EDITORONLY_DATA
 	, Tint_DEPRECATED(FLinearColor::White)
 #endif
 	, TintColor( InTint )
-	, DrawAs(InDrawType)
-	, Tiling(InTiling)
-	, Mirroring(ESlateBrushMirrorType::NoMirror)
-	, ImageType(InImageType)
-	, bIsDynamicallyLoaded(bInDynamicallyLoaded)
-	, bIsSet(true)
 	, ResourceObject( InObjectResource )
 	, ResourceName( InResourceName )
 	, UVRegion( ForceInit )
@@ -106,18 +106,18 @@
 						  bool bInDynamicallyLoaded
  						)
 
-	: ImageSize(InImageSize)
+	: bIsDynamicallyLoaded(bInDynamicallyLoaded)
+	, bIsSet(true)
+	, DrawAs(InDrawType)
+	, Tiling(InTiling)
+	, Mirroring(ESlateBrushMirrorType::NoMirror)
+	, ImageType(InImageType)
+	, ImageSize(InImageSize)
 	, Margin(InMargin)
 #if WITH_EDITORONLY_DATA
 	, Tint_DEPRECATED(FLinearColor::White)
 #endif
 	, TintColor(InTint)
-	, DrawAs(InDrawType)
-	, Tiling(InTiling)
-	, Mirroring(ESlateBrushMirrorType::NoMirror)
-	, ImageType(InImageType)
-	, bIsDynamicallyLoaded(bInDynamicallyLoaded)
-	, bIsSet(true)
 	, ResourceObject(InObjectResource)
 	, ResourceName(InResourceName)
 	, UVRegion(ForceInit)
