Made it so that yo can't interacte with the main frame until you choose a project.
The window is only closable if a project is open.

The goal here is to avoid situations where a user start modifying the empty level while not in a project.

diff --git a/Engine/Source/Runtime/SlateCore/Private/Widgets/SWindow.cpp b/Engine/Source/Runtime/SlateCore/Private/Widgets/SWindow.cpp
--- a/Engine/Source/Runtime/SlateCore/Private/Widgets/SWindow.cpp
+++ b/Engine/Source/Runtime/SlateCore/Private/Widgets/SWindow.cpp
@@ -1270,7 +1270,7 @@ void SWindow::SetContent( TSharedRef<SWidget> InContent )
 	Invalidate(EInvalidateWidgetReason::ChildOrder);
 }
 
-TSharedRef<const SWidget> SWindow::GetContent() const
+TSharedRef<SWidget> SWindow::GetContent()
 {
 	if (ContentSlot)
 	{
diff --git a/Engine/Source/Runtime/SlateCore/Public/Widgets/SWindow.h b/Engine/Source/Runtime/SlateCore/Public/Widgets/SWindow.h
--- a/Engine/Source/Runtime/SlateCore/Public/Widgets/SWindow.h
+++ b/Engine/Source/Runtime/SlateCore/Public/Widgets/SWindow.h
@@ -620,7 +620,7 @@ public:
 	 *
 	 * @return	The widget content for this window
 	 */
-	TSharedRef<const SWidget> GetContent() const;
+	TSharedRef<SWidget> GetContent();
 
 	/**
 	 * Check whether we have a full window overlay, used to draw content over the entire window.
