Updated Check styles & Differentiating Menu Check Styles

diff --git a/Engine/Source/Runtime/SlateCore/Private/Styling/StarshipCoreStyle.cpp b/Engine/Source/Runtime/SlateCore/Private/Styling/StarshipCoreStyle.cpp
--- a/Engine/Source/Runtime/SlateCore/Private/Styling/StarshipCoreStyle.cpp
+++ b/Engine/Source/Runtime/SlateCore/Private/Styling/StarshipCoreStyle.cpp
@@ -2048,10 +2048,13 @@ void FStarshipCoreStyle::SetupMultiboxStyles(TSharedRef<FStyle>& Style)
 		/* Set images for various SCheckBox states associated with menu check box items... */
 		FLinearColor Transparent20 = FLinearColor(1.0, 1.0, 1.0, 0.2);
 		FLinearColor Transparent01 = FLinearColor(1.0, 1.0, 1.0, 0.01);
-		const FCheckBoxStyle BasicMenuCheckBoxStyle = FCheckBoxStyle()
-			.SetUncheckedImage(           IMAGE_BRUSH_SVG("Starship/Common/check", Icon16x16, Transparent01))
-			.SetUncheckedHoveredImage(    IMAGE_BRUSH_SVG("Starship/Common/check", Icon16x16, Transparent20))
-			.SetUncheckedPressedImage(    IMAGE_BRUSH_SVG("Starship/Common/check", Icon16x16, Transparent20))
+
+
+
+		const FCheckBoxStyle MenuCheckIndicator = FCheckBoxStyle()
+			.SetUncheckedImage(           FSlateNoResource())
+			.SetUncheckedHoveredImage(    FSlateNoResource())
+			.SetUncheckedPressedImage(    FSlateNoResource())
 
 			.SetCheckedImage(             IMAGE_BRUSH_SVG("Starship/Common/check", Icon16x16))
 			.SetCheckedHoveredImage(      IMAGE_BRUSH_SVG("Starship/Common/check", Icon16x16))
@@ -2061,9 +2064,37 @@ void FStarshipCoreStyle::SetupMultiboxStyles(TSharedRef<FStyle>& Style)
 			.SetUndeterminedHoveredImage( IMAGE_BRUSH_SVG("Starship/Common/check", Icon16x16, Transparent20))
 			.SetUndeterminedPressedImage( IMAGE_BRUSH_SVG("Starship/Common/check", Icon16x16, Transparent20));
 
+		// SCheckBox defaults...
+		const float CheckboxCornerRadius = 3.f;
+		const float CheckboxOutlineThickness = 1.0f;
+
+		const FCheckBoxStyle MenuCheckBox = FCheckBoxStyle()
+			.SetCheckBoxType(ESlateCheckBoxType::CheckBox)
+
+			.SetForegroundColor(FStyleColors::White)
+			.SetHoveredForegroundColor(FStyleColors::White)
+			.SetPressedForegroundColor(FStyleColors::White)
+			.SetCheckedForegroundColor(FStyleColors::White)
+			.SetCheckedHoveredForegroundColor(FStyleColors::White)
+			.SetCheckedPressedForegroundColor(FStyleColors::White)
+			.SetUndeterminedForegroundColor(FStyleColors::White)
+
+			.SetUncheckedImage(FSlateNoResource())
+			.SetUncheckedHoveredImage(FSlateNoResource())
+			.SetUncheckedPressedImage(FSlateNoResource())
+			.SetCheckedImage(IMAGE_BRUSH_SVG("Starship/CoreWidgets/CheckBox/check", Icon16x16, FStyleColors::White))
+			.SetCheckedHoveredImage(IMAGE_BRUSH_SVG("Starship/CoreWidgets/CheckBox/check", Icon16x16, FStyleColors::White))
+			.SetCheckedPressedImage(IMAGE_BRUSH_SVG("Starship/CoreWidgets/CheckBox/check", Icon16x16, FStyleColors::White))
+			.SetUndeterminedImage(IMAGE_BRUSH_SVG("Starship/CoreWidgets/CheckBox/indeterminate", Icon16x16, FStyleColors::White))
+			.SetUndeterminedHoveredImage(IMAGE_BRUSH_SVG("Starship/CoreWidgets/CheckBox/indeterminate", Icon16x16, FStyleColors::White))
+			.SetUndeterminedPressedImage(IMAGE_BRUSH_SVG("Starship/CoreWidgets/CheckBox/indeterminate", Icon16x16, FStyleColors::White))
+			.SetBackgroundImage(FSlateRoundedBoxBrush(FStyleColors::Panel, CheckboxCornerRadius, FStyleColors::Hover, CheckboxOutlineThickness, Icon18x18))
+			.SetBackgroundHoveredImage(FSlateRoundedBoxBrush(FStyleColors::Panel, CheckboxCornerRadius, FStyleColors::Hover, CheckboxOutlineThickness, Icon18x18))
+			.SetBackgroundPressedImage(FSlateRoundedBoxBrush(FStyleColors::Panel, CheckboxCornerRadius, FStyleColors::Hover, CheckboxOutlineThickness, Icon18x18));
+
 		/* ...and add the new style */
-		Style->Set("Menu.CheckBox", BasicMenuCheckBoxStyle);
-		Style->Set("Menu.Check", BasicMenuCheckBoxStyle);
+		Style->Set("Menu.CheckBox", MenuCheckBox);
+		Style->Set("Menu.Check", MenuCheckIndicator);
 
 		/* This radio button is actually just a check box with different images */
 		/* Set images for various Menu radio button (SCheckBox) states... */
