Fix a bug where a volatile widget redraws at an offset within its retainer box, by caching the view offset it was painted with on the slow path and reusing it on the fast path.

diff --git a/Engine/Source/Runtime/SlateCore/Private/FastUpdate/SlateInvalidationRoot.cpp b/Engine/Source/Runtime/SlateCore/Private/FastUpdate/SlateInvalidationRoot.cpp
--- a/Engine/Source/Runtime/SlateCore/Private/FastUpdate/SlateInvalidationRoot.cpp
+++ b/Engine/Source/Runtime/SlateCore/Private/FastUpdate/SlateInvalidationRoot.cpp
@@ -194,7 +194,8 @@ FAutoConsoleVariableRef CVarSlateInvalidationEnableReindexLayerId(
 FSlateInvalidationRootList GSlateInvalidationRootListInstance;
 
 FSlateInvalidationRoot::FSlateInvalidationRoot()
-	: CachedElementData(new FSlateCachedElementData)
+	: CachedViewOffset(0.0f, 0.0f)
+	, CachedElementData(new FSlateCachedElementData)
 	, InvalidationRootWidget(nullptr)
 	, RootHittestGrid(nullptr)
 	, CachedMaxLayerId(0)
@@ -388,6 +389,8 @@ FSlateInvalidationResult FSlateInvalidationRoot::PaintInvalidationRoot(const FSl
 		GSlateIsOnFastUpdatePath = false;
 		bNeedsSlowPath = false;
 
+		CachedViewOffset = Context.ViewOffset;
+
 		{
 			if (Context.bAllowFastPathUpdate)
 			{
@@ -427,6 +430,7 @@ FSlateInvalidationResult FSlateInvalidationRoot::PaintInvalidationRoot(const FSl
 	}
 #endif
 
+	Result.ViewOffset = CachedViewOffset;
 	Result.MaxLayerIdPainted = CachedMaxLayerId;
 	return Result;
 }
diff --git a/Engine/Source/Runtime/SlateCore/Public/FastUpdate/SlateInvalidationRoot.h b/Engine/Source/Runtime/SlateCore/Public/FastUpdate/SlateInvalidationRoot.h
--- a/Engine/Source/Runtime/SlateCore/Public/FastUpdate/SlateInvalidationRoot.h
+++ b/Engine/Source/Runtime/SlateCore/Public/FastUpdate/SlateInvalidationRoot.h
@@ -28,7 +28,8 @@ namespace UE::Slate::Private
 struct FSlateInvalidationContext
 {
 	FSlateInvalidationContext(FSlateWindowElementList& InWindowElementList, const FWidgetStyle& InWidgetStyle)
-		: PaintArgs(nullptr)
+		: ViewOffset(0.0f, 0.0f)
+		, PaintArgs(nullptr)
 		, WidgetStyle(InWidgetStyle)
 		, WindowElementList(&InWindowElementList)
 		, LayoutScaleMultiplier(1.0f)
@@ -39,6 +40,7 @@ struct FSlateInvalidationContext
 	} 
 
 	FSlateRect CullingRect;
+	FVector2f ViewOffset;
 	const FPaintArgs* PaintArgs;
 	const FWidgetStyle& WidgetStyle;
 	FSlateWindowElementList* WindowElementList;
@@ -58,10 +60,13 @@ enum class ESlateInvalidationPaintType
 struct FSlateInvalidationResult
 {
 	FSlateInvalidationResult()
-		: MaxLayerIdPainted(0)
+		: ViewOffset(0.0f, 0.0f)
+		, MaxLayerIdPainted(0)
 		, bRepaintedWidgets(false)
 	{}
 
+	/** The view offset to use with the draw buffer */
+	FVector2f ViewOffset;
 	/** The max layer id painted or cached */
 	int32 MaxLayerIdPainted;
 	/** If we had to repaint any widget */
@@ -200,6 +205,8 @@ private:
 	/** Widgets that will be updated. */
 	TArray<FSlateInvalidationWidgetHeapElement> FinalUpdateList;
 
+	FVector2f CachedViewOffset;
+
 	FSlateCachedElementData* CachedElementData;
 
 	SWidget* InvalidationRootWidget;
diff --git a/Engine/Source/Runtime/UMG/Private/Slate/WidgetRenderer.cpp b/Engine/Source/Runtime/UMG/Private/Slate/WidgetRenderer.cpp
--- a/Engine/Source/Runtime/UMG/Private/Slate/WidgetRenderer.cpp
+++ b/Engine/Source/Runtime/UMG/Private/Slate/WidgetRenderer.cpp
@@ -297,6 +297,7 @@ bool FWidgetRenderer::DrawInvalidationRoot(TSharedRef<SVirtualWindow>& VirtualWi
 	{
 		// Need to set a new window element list so make a copy
 		FSlateInvalidationContext ContextCopy = Context;
+		ContextCopy.ViewOffset = ViewOffset;
 
 		{
 			// Get the free buffer & add our virtual window
@@ -315,7 +316,7 @@ bool FWidgetRenderer::DrawInvalidationRoot(TSharedRef<SVirtualWindow>& VirtualWi
 
 				Renderer->DrawWindow_GameThread(ScopedDrawBuffer.GetDrawBuffer());
 
-				ScopedDrawBuffer.GetDrawBuffer().ViewOffset = ViewOffset;
+				ScopedDrawBuffer.GetDrawBuffer().ViewOffset = Result.ViewOffset;
 
 				FRenderThreadUpdateContext RenderThreadUpdateContext =
 				{
