Fix a bug where a widget disappears for a frame when it changes from volatile to non-volatile, by adding the NeedsRepaint update flag in SWidget::UpdateFastPathVolatility when the widget is not volatile.

diff --git a/Engine/Source/Runtime/SlateCore/Private/Widgets/SWidget.cpp b/Engine/Source/Runtime/SlateCore/Private/Widgets/SWidget.cpp
--- a/Engine/Source/Runtime/SlateCore/Private/Widgets/SWidget.cpp
+++ b/Engine/Source/Runtime/SlateCore/Private/Widgets/SWidget.cpp
@@ -918,6 +918,7 @@ void SWidget::UpdateFastPathVolatility(bool bParentVolatile)
 	}
 	else
 	{
+		AddUpdateFlags(EWidgetUpdateFlags::NeedsRepaint);
 		RemoveUpdateFlags(EWidgetUpdateFlags::NeedsVolatilePaint);
 	}
 
