Merging CL#19518642

Replace uses of FVector2f with FVector2D in fix for a bug where a volatile widget redraws at an offset within its retainer box, to fix a compile error in contexts where FVector2f contains doubles.

diff --git a/Engine/Source/Runtime/SlateCore/Public/FastUpdate/SlateInvalidationRoot.h b/Engine/Source/Runtime/SlateCore/Public/FastUpdate/SlateInvalidationRoot.h
--- a/Engine/Source/Runtime/SlateCore/Public/FastUpdate/SlateInvalidationRoot.h
+++ b/Engine/Source/Runtime/SlateCore/Public/FastUpdate/SlateInvalidationRoot.h
@@ -40,7 +40,7 @@ struct FSlateInvalidationContext
 	} 
 
 	FSlateRect CullingRect;
-	FVector2f ViewOffset;
+	FVector2D ViewOffset;
 	const FPaintArgs* PaintArgs;
 	const FWidgetStyle& WidgetStyle;
 	FSlateWindowElementList* WindowElementList;
@@ -66,7 +66,7 @@ struct FSlateInvalidationResult
 	{}
 
 	/** The view offset to use with the draw buffer */
-	FVector2f ViewOffset;
+	FVector2D ViewOffset;
 	/** The max layer id painted or cached */
 	int32 MaxLayerIdPainted;
 	/** If we had to repaint any widget */
@@ -205,7 +205,7 @@ private:
 	/** Widgets that will be updated. */
 	TArray<FSlateInvalidationWidgetHeapElement> FinalUpdateList;
 
-	FVector2f CachedViewOffset;
+	FVector2D CachedViewOffset;
 
 	FSlateCachedElementData* CachedElementData;
 
