Fix an FVector2D conversion cast in debug builds.
#rnx

diff --git a/Engine/Source/Runtime/SlateCore/Private/Rendering/DrawElements.cpp b/Engine/Source/Runtime/SlateCore/Private/Rendering/DrawElements.cpp
--- a/Engine/Source/Runtime/SlateCore/Private/Rendering/DrawElements.cpp
+++ b/Engine/Source/Runtime/SlateCore/Private/Rendering/DrawElements.cpp
@@ -156,7 +156,7 @@ void FSlateDrawElement::Init(FSlateWindowElementList& ElementList, EElementType
 	ClipStateHandle.SetPreCachedClipIndex(ElementList.GetClippingIndex());
 
 #if UE_SLATE_VERIFY_PIXELSIZE
-	const FVector2f PixelSize = (LocalSize * Scale);
+	const FVector2D PixelSize = FVector2D(LocalSize * Scale);
 	ensureMsgf(PixelSize.X >= 0.f && PixelSize.X <= (float)std::numeric_limits<uint16>::max(), TEXT("The size X '%f' is too small or big to fit in the SlateVertex buffer."), PixelSize.X);
 	ensureMsgf(PixelSize.Y >= 0.f && PixelSize.Y <= (float)std::numeric_limits<uint16>::max(), TEXT("The size Y '%f' is too small or big to fit in the SlateVertex buffer."), PixelSize.Y);
 #endif
@@ -590,15 +590,7 @@ void FSlateDrawElement::MakeLines( FSlateWindowElementList& ElementList, uint32
 	DataPayload.SetThickness(Thickness);
 	DataPayload.SetLines(Points, bAntialias, &PointColors);
 
-	ESlateDrawEffect DrawEffects = InDrawEffects;
-	if (bAntialias)
-	{
-		// If the line is to be anti-aliased, we cannot reliably snap
-		// the generated vertices.
-		DrawEffects |= ESlateDrawEffect::NoPixelSnapping;
-	}
-
-	Element.Init(ElementList, EElementType::ET_Line, InLayer, PaintGeometry, DrawEffects);
+	Element.Init(ElementList, EElementType::ET_Line, InLayer, PaintGeometry, InDrawEffects);
 
 }
 
@@ -672,7 +664,7 @@ void FSlateDrawElement::MakePostProcessPass(FSlateWindowElementList& ElementList
 	FSlatePostProcessPayload& DataPayload = ElementList.CreatePayload<FSlatePostProcessPayload>(Element);
 	DataPayload.DownsampleAmount = DownsampleAmount;
 	DataPayload.PostProcessData = Params;
-	DataPayload.CornerRadius = CornerRadius * PaintGeometry.DrawScale;
+	DataPayload.CornerRadius = CornerRadius;
 
 	Element.Init(ElementList, EElementType::ET_PostProcessPass, InLayer, PaintGeometry, ESlateDrawEffect::None);
 }
