Slate: Use reference when iterating on SWidget children.

diff --git a/Engine/Source/Runtime/SlateCore/Public/Layout/ChildrenBase.h b/Engine/Source/Runtime/SlateCore/Public/Layout/ChildrenBase.h
--- a/Engine/Source/Runtime/SlateCore/Public/Layout/ChildrenBase.h
+++ b/Engine/Source/Runtime/SlateCore/Public/Layout/ChildrenBase.h
@@ -61,10 +61,11 @@ public:
 		TGuardValue<bool> IteratingGuard(bIsIteratingChildren, true);
 #endif
 
-		for (int32 Index = 0; Index < Num(); ++Index)
+		int32 WidgetCount = Num();
+		for (int32 Index = 0; Index < WidgetCount; ++Index)
 		{
-			TSharedRef<SWidget> Ref = GetChildAt(Index);
-			Pred(Ref.Get());
+			FWidgetRef WidgetRef = GetChildRefAt(Index);
+			Pred(WidgetRef.GetWidget());
 		}
 	}
 
@@ -76,10 +77,11 @@ public:
 		TGuardValue<bool> IteratingGuard(bIsIteratingChildren, true);
 #endif
 
-		for (int32 Index = 0; Index < Num(); ++Index)
+		int32 WidgetCount = Num();
+		for (int32 Index = 0; Index < WidgetCount; ++Index)
 		{
-			TSharedRef<const SWidget> Ref = GetChildAt(Index);
-			Pred(Ref.Get());
+			FConstWidgetRef WidgetRef = GetChildRefAt(Index);
+			Pred(WidgetRef.GetWidget());
 		}
 	}
 
