FSlateInvalidationRoot::PaintFastPath_UpdateNextWidget returns true when the widget it is updating needs painting.

diff --git a/Engine/Source/Runtime/SlateCore/Private/FastUpdate/SlateInvalidationRoot.cpp b/Engine/Source/Runtime/SlateCore/Private/FastUpdate/SlateInvalidationRoot.cpp
--- a/Engine/Source/Runtime/SlateCore/Private/FastUpdate/SlateInvalidationRoot.cpp
+++ b/Engine/Source/Runtime/SlateCore/Private/FastUpdate/SlateInvalidationRoot.cpp
@@ -660,6 +660,8 @@ bool FSlateInvalidationRoot::PaintFastPath_UpdateNextWidget(const FSlateInvalida
 
 		if (UpdateResult.bPainted)
 		{
+			bNeedsPaint = true;
+
 			{
 				// Remove from the update list elements that are processed by this paint
 				FSlateInvalidationWidgetList::FIndexRange PaintedWidgetRange{ *FastWidgetPathList, InvalidationWidget.Index, InvalidationWidget.LeafMostChildIndex };
