Slate: Add test to find FShapedGlyphSequence that are not GCed correctly. It can be activated with the cvar `Slate.CheckUObjectShapedGlyphSequence`

diff --git a/Engine/Source/Runtime/SlateCore/Private/Fonts/FontCache.cpp b/Engine/Source/Runtime/SlateCore/Private/Fonts/FontCache.cpp
--- a/Engine/Source/Runtime/SlateCore/Private/Fonts/FontCache.cpp
+++ b/Engine/Source/Runtime/SlateCore/Private/Fonts/FontCache.cpp
@@ -181,6 +181,19 @@ FShapedGlyphSequence::FShapedGlyphSequence(TArray<FShapedGlyphEntry> InGlyphsToR
 		}
 	}
 
+#if SLATE_CHECK_UOBJECT_SHAPED_GLYPH_SEQUENCE
+	if (FontMaterial)
+	{
+		FontMaterialWeakPtr = FontMaterial;
+		DebugFontMaterialName = FontMaterial->GetFName();
+	}
+	if (OutlineSettings.OutlineMaterial)
+	{
+		OutlineMaterialWeakPtr = OutlineSettings.OutlineMaterial;
+		DebugOutlineMaterialName = OutlineSettings.OutlineMaterial->GetFName();
+	}
+#endif
+
 	// Track memory usage
 	INC_MEMORY_STAT_BY(STAT_SlateShapedGlyphSequenceMemory, GetAllocatedSize());
 }
@@ -424,6 +437,24 @@ FShapedGlyphSequencePtr FShapedGlyphSequence::GetSubSequence(const int32 InStart
 
 void FShapedGlyphSequence::AddReferencedObjects(FReferenceCollector& Collector)
 {
+#if SLATE_CHECK_UOBJECT_SHAPED_GLYPH_SEQUENCE
+	if (GSlateCheckUObjectShapedGlyphSequence)
+	{
+		// pending kill objects may still be rendered for a frame so it is valid for the check to pass
+		const bool bEvenIfPendingKill = true;
+		// This test needs to be thread safe. It doesn't give us as many chances to trap bugs here but it is still useful
+		const bool bThreadSafe = true;
+		if (!DebugFontMaterialName.IsNone())
+		{
+			checkf(FontMaterialWeakPtr.IsValid(bEvenIfPendingKill, bThreadSafe), TEXT("Material %s has become invalid. This means the FShapedGlyphSequence::FontMaterial was garbage collected while slate was using it"), *DebugFontMaterialName.ToString());
+		}
+		if (!DebugOutlineMaterialName.IsNone())
+		{
+			checkf(OutlineMaterialWeakPtr.IsValid(bEvenIfPendingKill, bThreadSafe), TEXT("Material %s has become invalid. This means the FShapedGlyphSequence::OutlineSettings::OutlineMaterial was garbage collected while slate was using it"), *DebugOutlineMaterialName.ToString());
+		}
+	}
+#endif
+
 	Collector.AddReferencedObject(FontMaterial);
 	Collector.AddReferencedObject(OutlineSettings.OutlineMaterial);
 }
diff --git a/Engine/Source/Runtime/SlateCore/Private/SlateCoreClasses.cpp b/Engine/Source/Runtime/SlateCore/Private/SlateCoreClasses.cpp
--- a/Engine/Source/Runtime/SlateCore/Private/SlateCoreClasses.cpp
+++ b/Engine/Source/Runtime/SlateCore/Private/SlateCoreClasses.cpp
@@ -65,6 +65,15 @@ static FAutoConsoleVariableRef CVarSlateCheckUObjectRenderResources(
 bool GSlateCheckUObjectRenderResourcesShouldLogFatal = false;
 #endif
 
+#if SLATE_CHECK_UOBJECT_SHAPED_GLYPH_SEQUENCE
+bool GSlateCheckUObjectShapedGlyphSequence = true;
+static FAutoConsoleVariableRef CVarSlateCheckUObjectShapedGlyphSequence(
+	TEXT("Slate.CheckUObjectShapedGlyphSequence"),
+	GSlateCheckUObjectShapedGlyphSequence,
+	TEXT("")
+);
+#endif
+
 #if WITH_SLATE_DEBUGGING
 FAutoConsoleVariable CVarInvalidationDebugging(
 	TEXT("Slate.InvalidationDebugging"),
diff --git a/Engine/Source/Runtime/SlateCore/Public/Fonts/FontCache.h b/Engine/Source/Runtime/SlateCore/Public/Fonts/FontCache.h
--- a/Engine/Source/Runtime/SlateCore/Public/Fonts/FontCache.h
+++ b/Engine/Source/Runtime/SlateCore/Public/Fonts/FontCache.h
@@ -3,6 +3,7 @@
 #pragma once
 
 #include "CoreMinimal.h"
+#include "SlateGlobals.h"
 #include "Fonts/ShapedTextFwd.h"
 #include "UObject/ObjectMacros.h"
 #include "Fonts/SlateFontInfo.h"
@@ -444,6 +445,13 @@ private:
 	TArray<TWeakPtr<FFreeTypeFace>> GlyphFontFaces;
 	/** A map of source indices to their shaped glyph data indices - used to perform efficient reverse look-up */
 	FSourceIndicesToGlyphData SourceIndicesToGlyphData;
+#if SLATE_CHECK_UOBJECT_SHAPED_GLYPH_SEQUENCE
+	// Used to guard against crashes when the material object is deleted. This is expensive so we do not do it in shipping
+	TWeakObjectPtr<const UObject> FontMaterialWeakPtr;
+	TWeakObjectPtr<const UObject> OutlineMaterialWeakPtr;
+	FName DebugFontMaterialName;
+	FName DebugOutlineMaterialName;
+#endif
 };
 
 /** Information for rendering one non-shaped character */
diff --git a/Engine/Source/Runtime/SlateCore/Public/SlateGlobals.h b/Engine/Source/Runtime/SlateCore/Public/SlateGlobals.h
--- a/Engine/Source/Runtime/SlateCore/Public/SlateGlobals.h
+++ b/Engine/Source/Runtime/SlateCore/Public/SlateGlobals.h
@@ -10,6 +10,9 @@
 // Enabled cvar GSlateCheckUObjectRenderResources that will check for invalid reference in the slate resources manager
 #define SLATE_CHECK_UOBJECT_RENDER_RESOURCES !UE_BUILD_SHIPPING
 
+// Enabled cvar GSlateCheckUObjectShapedGlyphSequence that will check for invalid reference before using them
+#define SLATE_CHECK_UOBJECT_SHAPED_GLYPH_SEQUENCE !UE_BUILD_SHIPPING
+
 #ifndef SLATE_CULL_WIDGETS
 	#define SLATE_CULL_WIDGETS 1
 #endif
@@ -79,6 +82,10 @@ extern SLATECORE_API bool GSlateCheckUObjectRenderResources;
 extern SLATECORE_API bool GSlateCheckUObjectRenderResourcesShouldLogFatal;
 #endif
 
+#if SLATE_CHECK_UOBJECT_SHAPED_GLYPH_SEQUENCE
+extern SLATECORE_API bool GSlateCheckUObjectShapedGlyphSequence;
+#endif
+
 #if WITH_SLATE_DEBUGGING
 extern SLATECORE_API bool GSlateHitTestGridDebugging;
 #endif
