Title bar height is now correctly set to 34 from the previous 24.

diff --git a/Engine/Source/Runtime/SlateCore/Public/Widgets/SWindow.h b/Engine/Source/Runtime/SlateCore/Public/Widgets/SWindow.h
--- a/Engine/Source/Runtime/SlateCore/Public/Widgets/SWindow.h
+++ b/Engine/Source/Runtime/SlateCore/Public/Widgets/SWindow.h
@@ -94,7 +94,7 @@ enum class ESizingRule : uint8
 namespace SWindowDefs
 {
 	/** Height of a Slate window title bar, in pixels */
-	static const float DefaultTitleBarSize = 24.0f;
+	static const float DefaultTitleBarSize = 34.0f;
 }
 
 /** Proxy structure to handle deprecated construction from bool */
