Slate: Transform some ElementPayload from Vector2D to Vector2f. Use FDeprecateVector2D to create a deprecation path for function that returns a FVector2D.

diff --git a/Engine/Source/Developer/SlateReflector/Private/VisualTreeCapture.cpp b/Engine/Source/Developer/SlateReflector/Private/VisualTreeCapture.cpp
--- a/Engine/Source/Developer/SlateReflector/Private/VisualTreeCapture.cpp
+++ b/Engine/Source/Developer/SlateReflector/Private/VisualTreeCapture.cpp
@@ -37,7 +37,7 @@ FVisualEntry::FVisualEntry(const TWeakPtr<const SWidget>& InWidget, int32 InElem
 FVisualEntry::FVisualEntry(const TSharedRef<const SWidget>& InWidget, const FSlateDrawElement& InElement)
 {
 	const FSlateRenderTransform& Transform = InElement.GetRenderTransform();
-	const FVector2D& LocalSize = InElement.GetLocalSize();
+	const FVector2D LocalSize = InElement.GetLocalSize();
 
 	TopLeft = Transform.TransformPoint(FVector2D(0, 0));
 	TopRight = Transform.TransformPoint(FVector2D(LocalSize.X, 0));
@@ -60,7 +60,7 @@ void FVisualEntry::Resolve(const FSlateWindowElementList& ElementList)
 
 	const FSlateDrawElement& Element = ElementList.GetUncachedDrawElements()[ElementIndex];
 	const FSlateRenderTransform& Transform = Element.GetRenderTransform();
-	const FVector2D& LocalSize = Element.GetLocalSize();
+	const FVector2D LocalSize = Element.GetLocalSize();
 
 	TopLeft = Transform.TransformPoint(FVector2D(0, 0));
 	TopRight = Transform.TransformPoint(FVector2D(LocalSize.X, 0));
diff --git a/Engine/Source/Editor/WorldPartitionEditor/Private/WorldPartition/SWorldPartitionEditorGridSpatialHash.cpp b/Engine/Source/Editor/WorldPartitionEditor/Private/WorldPartition/SWorldPartitionEditorGridSpatialHash.cpp
--- a/Engine/Source/Editor/WorldPartitionEditor/Private/WorldPartition/SWorldPartitionEditorGridSpatialHash.cpp
+++ b/Engine/Source/Editor/WorldPartitionEditor/Private/WorldPartition/SWorldPartitionEditorGridSpatialHash.cpp
@@ -133,8 +133,9 @@ int32 SWorldPartitionEditorGridSpatialHash::PaintGrid(const FGeometry& AllottedG
 			const FVector2D ScreenSpaceSize = WorldImageGeometry.GetLocalSize();
 			const FVector2D ViewportPositon = -WorldImageGeometry.GetAccumulatedRenderTransform().GetTranslation() + AllottedGeometry.GetAbsolutePosition();
 
-			const FVector2D UV0 = WorldMiniMapBrush.GetUVRegion().Min;
-			const FVector2D UV1 = WorldMiniMapBrush.GetUVRegion().Max;
+			FBox2D UVRegion = WorldMiniMapBrush.GetUVRegion();
+			const FVector2D UV0 = UVRegion.Min;
+			const FVector2D UV1 = UVRegion.Max;
 
 			const ERHIFeatureLevel::Type InFeatureLevel = GMaxRHIFeatureLevel;
 			const int32 MipLevel = -1;
diff --git a/Engine/Source/Runtime/AppFramework/Private/Widgets/Testing/STestSuite.cpp b/Engine/Source/Runtime/AppFramework/Private/Widgets/Testing/STestSuite.cpp
--- a/Engine/Source/Runtime/AppFramework/Private/Widgets/Testing/STestSuite.cpp
+++ b/Engine/Source/Runtime/AppFramework/Private/Widgets/Testing/STestSuite.cpp
@@ -752,7 +752,7 @@ private:
 		for (int i = 0; i < NewNumGradientStops; ++i)
 		{
 			// Note that the position is ignored by the spline
-			GradientStops.Add( FSlateGradientStop(FVector2D(static_cast<float>(i)/FMath::Max(1,NewNumGradientStops-1), 0.0f), i%2==0 ? ColorOne : ColorTwo) );
+			GradientStops.Add( FSlateGradientStop(FVector2f(static_cast<float>(i)/FMath::Max(1,NewNumGradientStops-1), 0.0f), i%2==0 ? ColorOne : ColorTwo) );
 		}
 		SplineWidget->SetGradient(GradientStops);
 	}
@@ -940,10 +940,10 @@ private:
 	{
 		TArray<FSlateGradientStop> GradientStops;
 
-		GradientStops.Add( FSlateGradientStop(FVector2D(InParams.Geometry.GetLocalSize().X*.1f, 0), FColor::Yellow) );
-		GradientStops.Add( FSlateGradientStop( FVector2D(InParams.Geometry.GetLocalSize().X*.25f,0), FColor::Magenta ) );
-		GradientStops.Add( FSlateGradientStop( FVector2D(InParams.Geometry.GetLocalSize().X*.75f,0), FColor::Blue ) );
-		GradientStops.Add( FSlateGradientStop( FVector2D(InParams.Geometry.GetLocalSize().X*0.9f,0), FColor::Green) );
+		GradientStops.Add( FSlateGradientStop(FVector2f(InParams.Geometry.GetLocalSize().X*.1f, 0), FColor::Yellow) );
+		GradientStops.Add( FSlateGradientStop( FVector2f(InParams.Geometry.GetLocalSize().X*.25f,0), FColor::Magenta ) );
+		GradientStops.Add( FSlateGradientStop( FVector2f(InParams.Geometry.GetLocalSize().X*.75f,0), FColor::Blue ) );
+		GradientStops.Add( FSlateGradientStop( FVector2f(InParams.Geometry.GetLocalSize().X*0.9f,0), FColor::Green) );
 
 		FSlateDrawElement::MakeGradient(
 			InParams.OutDrawElements,
diff --git a/Engine/Source/Runtime/SlateCore/Private/Rendering/DrawElements.cpp b/Engine/Source/Runtime/SlateCore/Private/Rendering/DrawElements.cpp
--- a/Engine/Source/Runtime/SlateCore/Private/Rendering/DrawElements.cpp
+++ b/Engine/Source/Runtime/SlateCore/Private/Rendering/DrawElements.cpp
@@ -128,7 +128,7 @@ FSlateWindowElementList::FSlateWindowElementList(const TSharedPtr<SWindow>& InPa
 {
 	if (InPaintWindow.IsValid())
 	{
-		WindowSize = InPaintWindow->GetSizeInScreen();
+		WindowSize = UE::Slate::CastToVector2f(InPaintWindow->GetSizeInScreen());
 	}
 
 	// Only keep UObject resources alive if this window element list is born on the game thread.
@@ -193,7 +193,12 @@ void FSlateDrawElement::Init(FSlateWindowElementList& ElementList, EElementType
 	}
 }
 
-void FSlateDrawElement::ApplyPositionOffset(const FVector2D& InOffset)
+void FSlateDrawElement::ApplyPositionOffset(FVector2d InOffset)
+{
+	ApplyPositionOffset(UE::Slate::CastToVector2f(InOffset));
+}
+
+void FSlateDrawElement::ApplyPositionOffset(FVector2f InOffset)
 {
 	SetPosition(GetPosition() + InOffset);
 	RenderTransform = Concatenate(RenderTransform, InOffset);
diff --git a/Engine/Source/Runtime/SlateCore/Private/Rendering/ElementBatcher.cpp b/Engine/Source/Runtime/SlateCore/Private/Rendering/ElementBatcher.cpp
--- a/Engine/Source/Runtime/SlateCore/Private/Rendering/ElementBatcher.cpp
+++ b/Engine/Source/Runtime/SlateCore/Private/Rendering/ElementBatcher.cpp
@@ -293,7 +293,7 @@ void FSlateElementBatcher::AddElements(FSlateWindowElementList& WindowElementLis
 	check(BatchDataHDR->GetRenderBatches().Num() == 0);
 
 
-	FVector2D ViewportSize = WindowElementList.GetPaintWindow()->GetViewportSize();
+	FVector2f ViewportSize = UE::Slate::CastToVector2f(WindowElementList.GetPaintWindow()->GetViewportSize());
 
 	PrecachedClippingStates = &WindowElementList.ClippingManager.GetClippingStates();
 
@@ -341,7 +341,7 @@ void FSlateElementBatcher::AddElements(FSlateWindowElementList& WindowElementLis
 }
 
 
-void FSlateElementBatcher::AddElementsInternal(const FSlateDrawElementArray& DrawElements, const FVector2D& ViewportSize)
+void FSlateElementBatcher::AddElementsInternal(const FSlateDrawElementArray& DrawElements, FVector2f ViewportSize)
 {
 	for (const FSlateDrawElement& DrawElement : DrawElements)
 	{
@@ -442,7 +442,7 @@ void FSlateElementBatcher::AddElementsInternal(const FSlateDrawElementArray& Dra
 	}
 }
 
-void FSlateElementBatcher::AddCachedElements(FSlateCachedElementData& CachedElementData, const FVector2D& ViewportSize)
+void FSlateElementBatcher::AddCachedElements(FSlateCachedElementData& CachedElementData, FVector2f ViewportSize)
 {
 	CSV_SCOPED_TIMING_STAT(Slate, AddCachedElements);
 	SCOPED_NAMED_EVENT_TEXT("Slate::AddCachedElements", FColor::Magenta);
@@ -497,7 +497,7 @@ void FSlateElementBatcher::AddDebugQuadElement(const FSlateDrawElement& DrawElem
 
 	const FColor Tint = PackVertexColor(DrawElementPayload.GetTint());
 	const FSlateRenderTransform& RenderTransform = DrawElement.GetRenderTransform();
-	const FVector2D& LocalSize = DrawElement.GetLocalSize();
+	const FVector2f LocalSize = DrawElement.GetLocalSize();
 	ESlateDrawEffect InDrawEffects = DrawElement.GetDrawEffects();
 	const int32 Layer = DrawElement.GetLayer();
 
@@ -506,10 +506,10 @@ void FSlateElementBatcher::AddDebugQuadElement(const FSlateDrawElement& DrawElem
 	const FColor Color = PackVertexColor(DrawElementPayload.GetTint());
 
 	// Determine the four corners of the quad
-	FVector2D TopLeft = FVector2D::ZeroVector;
-	FVector2D TopRight = FVector2D(LocalSize.X, 0);
-	FVector2D BotLeft = FVector2D(0, LocalSize.Y);
-	FVector2D BotRight = FVector2D(LocalSize.X, LocalSize.Y);
+	FVector2f TopLeft = FVector2f::ZeroVector;
+	FVector2f TopRight = FVector2f(LocalSize.X, 0);
+	FVector2f BotLeft = FVector2f(0.f, LocalSize.Y);
+	FVector2f BotRight = FVector2f(LocalSize.X, LocalSize.Y);
 
 	// The start index of these vertices in the index buffer
 	//const uint32 IndexStart = BatchVertices.Num();
@@ -550,7 +550,7 @@ void FSlateElementBatcher::AddBoxElement(const FSlateDrawElement& DrawElement)
 	const FColor Tint = PackVertexColor(DrawElementPayload.GetTint());
 	const FSlateRenderTransform& ElementRenderTransform = DrawElement.GetRenderTransform();
 	const FSlateRenderTransform RenderTransform = DrawElement.GetRenderTransform();// GetBoxRenderTransform(DrawElement);
-	const FVector2D LocalSize = DrawElement.GetLocalSize();
+	const FVector2f LocalSize = DrawElement.GetLocalSize();
 
 	const ESlateDrawEffect InDrawEffects = DrawElement.GetDrawEffects();
 	const int32 Layer = DrawElement.GetLayer();
@@ -558,18 +558,18 @@ void FSlateElementBatcher::AddBoxElement(const FSlateDrawElement& DrawElement)
 	const float DrawScale = DrawElement.GetScale();
 
 	// Do pixel snapping
-	FVector2D TopLeft(0,0);
-	FVector2D BotRight(LocalSize);
+	FVector2f TopLeft(0.f, 0.f);
+	FVector2f BotRight(LocalSize);
 
 	uint32 TextureWidth = 1;
 	uint32 TextureHeight = 1;
 
 	// Get the default start and end UV.  If the texture is atlased this value will be a subset of this
-	FVector2D StartUV = FVector2D(0.0f,0.0f);
-	FVector2D EndUV = FVector2D(1.0f,1.0f);
-	FVector2D SizeUV;
+	FVector2f StartUV = FVector2f(0.0f, 0.0f);
+	FVector2f EndUV = FVector2f(1.0f, 1.0f);
+	FVector2f SizeUV;
 
-	FVector2D HalfTexel;
+	FVector2f HalfTexel;
 
 	const FSlateShaderResourceProxy* ResourceProxy = DrawElementPayload.GetResourceProxy();
 	FSlateShaderResource* Resource = nullptr;
@@ -582,9 +582,9 @@ void FSlateElementBatcher::AddBoxElement(const FSlateDrawElement& DrawElement)
 		TextureHeight = ResourceProxy->ActualSize.Y != 0 ? ResourceProxy->ActualSize.Y : 1;
 
 		// Texel offset
-		HalfTexel = FVector2D( PixelCenterOffset/TextureWidth, PixelCenterOffset/TextureHeight );
+		HalfTexel = FVector2f( PixelCenterOffset/TextureWidth, PixelCenterOffset/TextureHeight );
 
-		const FBox2D& BrushUV = DrawElementPayload.GetBrushUVRegion();
+		const FBox2f& BrushUV = DrawElementPayload.GetBrushUVRegion();
 		//In case brush has valid UV region - use it instead of proxy UV
 		if (BrushUV.bIsValid)
 		{
@@ -594,16 +594,16 @@ void FSlateElementBatcher::AddBoxElement(const FSlateDrawElement& DrawElement)
 		}
 		else
 		{
-			SizeUV = FVector2D(ResourceProxy->SizeUV);
-			StartUV = FVector2D(ResourceProxy->StartUV) + HalfTexel;
-			EndUV = StartUV + FVector2D(ResourceProxy->SizeUV);
+			SizeUV = FVector2f(ResourceProxy->SizeUV);
+			StartUV = FVector2f(ResourceProxy->StartUV) + HalfTexel;
+			EndUV = StartUV + FVector2f(ResourceProxy->SizeUV);
 		}
 	}
 	else
 	{
 		// no texture
-		SizeUV = FVector2D(1.0f,1.0f);
-		HalfTexel = FVector2D( PixelCenterOffset, PixelCenterOffset );
+		SizeUV = FVector2f(1.0f,1.0f);
+		HalfTexel = FVector2f( PixelCenterOffset, PixelCenterOffset );
 	}
 
 
@@ -640,7 +640,7 @@ void FSlateElementBatcher::AddBoxElement(const FSlateDrawElement& DrawElement)
 	float HorizontalTiling = bTileHorizontal ? LocalSize.X / TextureWidth : 1.0f;
 	float VerticalTiling = bTileVertical ? LocalSize.Y / TextureHeight : 1.0f;
 
-	const FVector2D Tiling( HorizontalTiling, VerticalTiling );
+	const FVector2f Tiling( HorizontalTiling, VerticalTiling );
 
 	// The start index of these vertices in the index buffer
 	const uint32 IndexStart = 0;// BatchVertices.Num();
@@ -648,8 +648,8 @@ void FSlateElementBatcher::AddBoxElement(const FSlateDrawElement& DrawElement)
 
 	const FMargin& Margin = DrawElementPayload.GetBrushMargin();
 
-	const FVector2D TopRight = FVector2D(BotRight.X, TopLeft.Y);
-	const FVector2D BotLeft = FVector2D(TopLeft.X, BotRight.Y);
+	const FVector2f TopRight = FVector2f(BotRight.X, TopLeft.Y);
+	const FVector2f BotLeft = FVector2f(TopLeft.X, BotRight.Y);
 
 	const FColor FeatherColor(0, 0, 0, 0);
 
@@ -686,8 +686,8 @@ void FSlateElementBatcher::AddBoxElement(const FSlateDrawElement& DrawElement)
 
 		if( bMirrorHorizontal || bMirrorVertical )
 		{
-			const FVector2D UVMin = StartUV;
-			const FVector2D UVMax = EndUV;
+			const FVector2f UVMin = StartUV;
+			const FVector2f UVMax = EndUV;
 
 			if( bMirrorHorizontal )
 			{
@@ -726,26 +726,26 @@ void FSlateElementBatcher::AddBoxElement(const FSlateDrawElement& DrawElement)
 			BottomMarginY = TopMarginY;
 		}
 
-		FVector2D Position = TopLeft;
-		FVector2D EndPos = BotRight;
-
-		RenderBatch.AddVertex( FSlateVertex::Make<Rounding>( RenderTransform, FVector2f( Position.X, Position.Y ),		FVector2f(LocalSize), DrawScale, FVector4f(FVector2f( StartUV ),						FVector2f(Tiling)),	Tint, SecondaryColor) ); //0
-		RenderBatch.AddVertex( FSlateVertex::Make<Rounding>( RenderTransform, FVector2f( Position.X, TopMarginY ),		FVector2f(LocalSize), DrawScale, FVector4f(FVector2f( StartUV.X, TopMarginV ),			FVector2f(Tiling)),	Tint, SecondaryColor) ); //1
-		RenderBatch.AddVertex( FSlateVertex::Make<Rounding>( RenderTransform, FVector2f( LeftMarginX, Position.Y ),		FVector2f(LocalSize), DrawScale, FVector4f(FVector2f( LeftMarginU, StartUV.Y ),			FVector2f(Tiling)),	Tint, SecondaryColor) ); //2
-		RenderBatch.AddVertex( FSlateVertex::Make<Rounding>( RenderTransform, FVector2f( LeftMarginX, TopMarginY ),		FVector2f(LocalSize), DrawScale, FVector4f(FVector2f( LeftMarginU, TopMarginV ),		FVector2f(Tiling)),	Tint, SecondaryColor) ); //3
-		RenderBatch.AddVertex( FSlateVertex::Make<Rounding>( RenderTransform, FVector2f( RightMarginX, Position.Y ),	FVector2f(LocalSize), DrawScale, FVector4f(FVector2f( RightMarginU, StartUV.Y ),		FVector2f(Tiling)),	Tint, SecondaryColor) ); //4
-		RenderBatch.AddVertex( FSlateVertex::Make<Rounding>( RenderTransform, FVector2f( RightMarginX, TopMarginY ),	FVector2f(LocalSize), DrawScale, FVector4f(FVector2f( RightMarginU,TopMarginV),			FVector2f(Tiling)),	Tint, SecondaryColor) ); //5
-		RenderBatch.AddVertex( FSlateVertex::Make<Rounding>( RenderTransform, FVector2f( EndPos.X, Position.Y ),		FVector2f(LocalSize), DrawScale, FVector4f(FVector2f( EndUV.X, StartUV.Y ),				FVector2f(Tiling)),	Tint, SecondaryColor) ); //6
-		RenderBatch.AddVertex( FSlateVertex::Make<Rounding>( RenderTransform, FVector2f( EndPos.X, TopMarginY ),		FVector2f(LocalSize), DrawScale, FVector4f(FVector2f( EndUV.X, TopMarginV),				FVector2f(Tiling)),	Tint, SecondaryColor) ); //7
-
-		RenderBatch.AddVertex( FSlateVertex::Make<Rounding>( RenderTransform, FVector2f( Position.X, BottomMarginY ),	FVector2f(LocalSize), DrawScale, FVector4f(FVector2f( StartUV.X, BottomMarginV ),		FVector2f(Tiling)),	Tint, SecondaryColor) ); //8
-		RenderBatch.AddVertex( FSlateVertex::Make<Rounding>( RenderTransform, FVector2f( LeftMarginX, BottomMarginY ),	FVector2f(LocalSize), DrawScale, FVector4f(FVector2f( LeftMarginU, BottomMarginV ),		FVector2f(Tiling)),	Tint, SecondaryColor) ); //9
-		RenderBatch.AddVertex( FSlateVertex::Make<Rounding>( RenderTransform, FVector2f( RightMarginX, BottomMarginY ), FVector2f(LocalSize), DrawScale, FVector4f(FVector2f( RightMarginU, BottomMarginV ),	FVector2f(Tiling)),	Tint, SecondaryColor) ); //10
-		RenderBatch.AddVertex( FSlateVertex::Make<Rounding>( RenderTransform, FVector2f( EndPos.X, BottomMarginY ),		FVector2f(LocalSize), DrawScale, FVector4f(FVector2f( EndUV.X, BottomMarginV ),			FVector2f(Tiling)),	Tint, SecondaryColor) ); //11
-		RenderBatch.AddVertex( FSlateVertex::Make<Rounding>( RenderTransform, FVector2f( Position.X, EndPos.Y ),		FVector2f(LocalSize), DrawScale, FVector4f(FVector2f( StartUV.X, EndUV.Y ),				FVector2f(Tiling)),	Tint, SecondaryColor) ); //12
-		RenderBatch.AddVertex( FSlateVertex::Make<Rounding>( RenderTransform, FVector2f( LeftMarginX, EndPos.Y ),		FVector2f(LocalSize), DrawScale, FVector4f(FVector2f( LeftMarginU, EndUV.Y ),			FVector2f(Tiling)),	Tint, SecondaryColor) ); //13
-		RenderBatch.AddVertex( FSlateVertex::Make<Rounding>( RenderTransform, FVector2f( RightMarginX, EndPos.Y ),		FVector2f(LocalSize), DrawScale, FVector4f(FVector2f( RightMarginU, EndUV.Y ),			FVector2f(Tiling)),	Tint, SecondaryColor) ); //14
-		RenderBatch.AddVertex( FSlateVertex::Make<Rounding>( RenderTransform, FVector2f( EndPos.X, EndPos.Y ),			FVector2f(LocalSize), DrawScale, FVector4f(FVector2f( EndUV ),							FVector2f(Tiling)),	Tint, SecondaryColor) ); //15
+		FVector2f Position = TopLeft;
+		FVector2f EndPos = BotRight;
+
+		RenderBatch.AddVertex( FSlateVertex::Make<Rounding>( RenderTransform, FVector2f( Position.X, Position.Y ),		LocalSize, DrawScale, FVector4f(FVector2f( StartUV ),						Tiling),	Tint, SecondaryColor) ); //0
+		RenderBatch.AddVertex( FSlateVertex::Make<Rounding>( RenderTransform, FVector2f( Position.X, TopMarginY ),		LocalSize, DrawScale, FVector4f(FVector2f( StartUV.X, TopMarginV ),			Tiling),	Tint, SecondaryColor) ); //1
+		RenderBatch.AddVertex( FSlateVertex::Make<Rounding>( RenderTransform, FVector2f( LeftMarginX, Position.Y ),		LocalSize, DrawScale, FVector4f(FVector2f( LeftMarginU, StartUV.Y ),		Tiling),	Tint, SecondaryColor) ); //2
+		RenderBatch.AddVertex( FSlateVertex::Make<Rounding>( RenderTransform, FVector2f( LeftMarginX, TopMarginY ),		LocalSize, DrawScale, FVector4f(FVector2f( LeftMarginU, TopMarginV ),		Tiling),	Tint, SecondaryColor) ); //3
+		RenderBatch.AddVertex( FSlateVertex::Make<Rounding>( RenderTransform, FVector2f( RightMarginX, Position.Y ),	LocalSize, DrawScale, FVector4f(FVector2f( RightMarginU, StartUV.Y ),		Tiling),	Tint, SecondaryColor) ); //4
+		RenderBatch.AddVertex( FSlateVertex::Make<Rounding>( RenderTransform, FVector2f( RightMarginX, TopMarginY ),	LocalSize, DrawScale, FVector4f(FVector2f( RightMarginU,TopMarginV),		Tiling),	Tint, SecondaryColor) ); //5
+		RenderBatch.AddVertex( FSlateVertex::Make<Rounding>( RenderTransform, FVector2f( EndPos.X, Position.Y ),		LocalSize, DrawScale, FVector4f(FVector2f( EndUV.X, StartUV.Y ),			Tiling),	Tint, SecondaryColor) ); //6
+		RenderBatch.AddVertex( FSlateVertex::Make<Rounding>( RenderTransform, FVector2f( EndPos.X, TopMarginY ),		LocalSize, DrawScale, FVector4f(FVector2f( EndUV.X, TopMarginV),			Tiling),	Tint, SecondaryColor) ); //7
+
+		RenderBatch.AddVertex( FSlateVertex::Make<Rounding>( RenderTransform, FVector2f( Position.X, BottomMarginY ),	LocalSize, DrawScale, FVector4f(FVector2f( StartUV.X, BottomMarginV ),		Tiling),	Tint, SecondaryColor) ); //8
+		RenderBatch.AddVertex( FSlateVertex::Make<Rounding>( RenderTransform, FVector2f( LeftMarginX, BottomMarginY ),	LocalSize, DrawScale, FVector4f(FVector2f( LeftMarginU, BottomMarginV ),	Tiling),	Tint, SecondaryColor) ); //9
+		RenderBatch.AddVertex( FSlateVertex::Make<Rounding>( RenderTransform, FVector2f( RightMarginX, BottomMarginY ), LocalSize, DrawScale, FVector4f(FVector2f( RightMarginU, BottomMarginV ),	Tiling),	Tint, SecondaryColor) ); //10
+		RenderBatch.AddVertex( FSlateVertex::Make<Rounding>( RenderTransform, FVector2f( EndPos.X, BottomMarginY ),		LocalSize, DrawScale, FVector4f(FVector2f( EndUV.X, BottomMarginV ),		Tiling),	Tint, SecondaryColor) ); //11
+		RenderBatch.AddVertex( FSlateVertex::Make<Rounding>( RenderTransform, FVector2f( Position.X, EndPos.Y ),		LocalSize, DrawScale, FVector4f(FVector2f( StartUV.X, EndUV.Y ),			Tiling),	Tint, SecondaryColor) ); //12
+		RenderBatch.AddVertex( FSlateVertex::Make<Rounding>( RenderTransform, FVector2f( LeftMarginX, EndPos.Y ),		LocalSize, DrawScale, FVector4f(FVector2f( LeftMarginU, EndUV.Y ),			Tiling),	Tint, SecondaryColor) ); //13
+		RenderBatch.AddVertex( FSlateVertex::Make<Rounding>( RenderTransform, FVector2f( RightMarginX, EndPos.Y ),		LocalSize, DrawScale, FVector4f(FVector2f( RightMarginU, EndUV.Y ),			Tiling),	Tint, SecondaryColor) ); //14
+		RenderBatch.AddVertex( FSlateVertex::Make<Rounding>( RenderTransform, FVector2f( EndPos.X, EndPos.Y ),			LocalSize, DrawScale, FVector4f(FVector2f( EndUV ),							Tiling),	Tint, SecondaryColor) ); //15
 
 		// Top
 		RenderBatch.AddIndex( IndexStart + 0 );
@@ -818,24 +818,24 @@ void FSlateElementBatcher::AddBoxElement(const FSlateDrawElement& DrawElement)
 			const int32 FeatherStart = RenderBatch.GetNumVertices();
 
 			// Top
-			RenderBatch.AddVertex(FSlateVertex::Make<Rounding>(RenderTransform, FVector2f(Position.X, Position.Y) + FVector2f(-1, -1) / DrawScale, FVector2f(LocalSize), DrawScale, FVector4f(FVector2f(StartUV), FVector2f(Tiling)), FeatherColor, SecondaryColor)); //0
-			RenderBatch.AddVertex(FSlateVertex::Make<Rounding>(RenderTransform, FVector2f(LeftMarginX, Position.Y) + FVector2f(0, -1) / DrawScale, FVector2f(LocalSize), DrawScale, FVector4f(FVector2f(LeftMarginU, StartUV.Y), FVector2f(Tiling)), FeatherColor, SecondaryColor)); //1
-			RenderBatch.AddVertex(FSlateVertex::Make<Rounding>(RenderTransform, FVector2f(RightMarginX, Position.Y) + FVector2f(0, -1) / DrawScale, FVector2f(LocalSize), DrawScale, FVector4f(FVector2f(RightMarginU, StartUV.Y), FVector2f(Tiling)), FeatherColor, SecondaryColor)); //2
-			RenderBatch.AddVertex(FSlateVertex::Make<Rounding>(RenderTransform, FVector2f(EndPos.X, Position.Y) + FVector2f(1, -1) / DrawScale, FVector2f(LocalSize), DrawScale, FVector4f(FVector2f(EndUV.X, StartUV.Y), FVector2f(Tiling)), FeatherColor, SecondaryColor)); //3
+			RenderBatch.AddVertex(FSlateVertex::Make<Rounding>(RenderTransform, FVector2f(Position.X, Position.Y) + FVector2f(-1.f, -1.f) / DrawScale,		LocalSize, DrawScale, FVector4f(FVector2f(StartUV), Tiling), FeatherColor, SecondaryColor)); //0
+			RenderBatch.AddVertex(FSlateVertex::Make<Rounding>(RenderTransform, FVector2f(LeftMarginX, Position.Y) + FVector2f(0.f, -1.f) / DrawScale,		LocalSize, DrawScale, FVector4f(FVector2f(LeftMarginU, StartUV.Y), Tiling), FeatherColor, SecondaryColor)); //1
+			RenderBatch.AddVertex(FSlateVertex::Make<Rounding>(RenderTransform, FVector2f(RightMarginX, Position.Y) + FVector2f(0.f, -1.f) / DrawScale,		LocalSize, DrawScale, FVector4f(FVector2f(RightMarginU, StartUV.Y), Tiling), FeatherColor, SecondaryColor)); //2
+			RenderBatch.AddVertex(FSlateVertex::Make<Rounding>(RenderTransform, FVector2f(EndPos.X, Position.Y) + FVector2f(1.f, -1.f) / DrawScale,			LocalSize, DrawScale, FVector4f(FVector2f(EndUV.X, StartUV.Y), Tiling), FeatherColor, SecondaryColor)); //3
 
 			// Left
-			RenderBatch.AddVertex(FSlateVertex::Make<Rounding>(RenderTransform, FVector2f(Position.X, TopMarginY) + FVector2f(-1, 0) / DrawScale, FVector2f(LocalSize), DrawScale, FVector4f(FVector2f(StartUV.X, TopMarginV), FVector2f(Tiling)), FeatherColor, SecondaryColor)); //4
-			RenderBatch.AddVertex(FSlateVertex::Make<Rounding>(RenderTransform, FVector2f(Position.X, BottomMarginY) + FVector2f(-1, 0) / DrawScale, FVector2f(LocalSize), DrawScale, FVector4f(FVector2f(StartUV.X, BottomMarginV), FVector2f(Tiling)), FeatherColor, SecondaryColor)); //5
+			RenderBatch.AddVertex(FSlateVertex::Make<Rounding>(RenderTransform, FVector2f(Position.X, TopMarginY) + FVector2f(-1.f, 0.f) / DrawScale,		LocalSize, DrawScale, FVector4f(FVector2f(StartUV.X, TopMarginV), Tiling), FeatherColor, SecondaryColor)); //4
+			RenderBatch.AddVertex(FSlateVertex::Make<Rounding>(RenderTransform, FVector2f(Position.X, BottomMarginY) + FVector2f(-1.f, 0.f) / DrawScale,	LocalSize, DrawScale, FVector4f(FVector2f(StartUV.X, BottomMarginV), Tiling), FeatherColor, SecondaryColor)); //5
 
 			// Right
-			RenderBatch.AddVertex(FSlateVertex::Make<Rounding>(RenderTransform, FVector2f(EndPos.X, TopMarginY) + FVector2f(1, 0) / DrawScale, FVector2f(LocalSize), DrawScale, FVector4f(FVector2f(EndUV.X, TopMarginV), FVector2f(Tiling)), FeatherColor, SecondaryColor)); //6
-			RenderBatch.AddVertex(FSlateVertex::Make<Rounding>(RenderTransform, FVector2f(EndPos.X, BottomMarginY) + FVector2f(1, 0) / DrawScale, FVector2f(LocalSize), DrawScale, FVector4f(FVector2f(EndUV.X, BottomMarginV), FVector2f(Tiling)), FeatherColor, SecondaryColor)); //7
+			RenderBatch.AddVertex(FSlateVertex::Make<Rounding>(RenderTransform, FVector2f(EndPos.X, TopMarginY) + FVector2f(1.f, 0.f) / DrawScale,			LocalSize, DrawScale, FVector4f(FVector2f(EndUV.X, TopMarginV), Tiling), FeatherColor, SecondaryColor)); //6
+			RenderBatch.AddVertex(FSlateVertex::Make<Rounding>(RenderTransform, FVector2f(EndPos.X, BottomMarginY) + FVector2f(1.f, 0.f) / DrawScale,		LocalSize, DrawScale, FVector4f(FVector2f(EndUV.X, BottomMarginV), Tiling), FeatherColor, SecondaryColor)); //7
 
 			// Bottom
-			RenderBatch.AddVertex(FSlateVertex::Make<Rounding>(RenderTransform, FVector2f(Position.X, EndPos.Y) + FVector2f(-1, 1) / DrawScale, FVector2f(LocalSize), DrawScale, FVector4f(FVector2f(StartUV.X, EndUV.Y), FVector2f(Tiling)), FeatherColor, SecondaryColor)); //8
-			RenderBatch.AddVertex(FSlateVertex::Make<Rounding>(RenderTransform, FVector2f(LeftMarginX, EndPos.Y) + FVector2f(0, 1) / DrawScale, FVector2f(LocalSize), DrawScale, FVector4f(FVector2f(LeftMarginU, EndUV.Y), FVector2f(Tiling)), FeatherColor, SecondaryColor)); //9
-			RenderBatch.AddVertex(FSlateVertex::Make<Rounding>(RenderTransform, FVector2f(RightMarginX, EndPos.Y) + FVector2f(0, 1) / DrawScale, FVector2f(LocalSize), DrawScale, FVector4f(FVector2f(RightMarginU, EndUV.Y), FVector2f(Tiling)), FeatherColor, SecondaryColor)); //10
-			RenderBatch.AddVertex(FSlateVertex::Make<Rounding>(RenderTransform, FVector2f(EndPos.X, EndPos.Y) + FVector2f(1, 1) / DrawScale, FVector2f(LocalSize), DrawScale, FVector4f(FVector2f(EndUV), FVector2f(Tiling)), FeatherColor, SecondaryColor)); //11
+			RenderBatch.AddVertex(FSlateVertex::Make<Rounding>(RenderTransform, FVector2f(Position.X, EndPos.Y) + FVector2f(-1.f, 1.f) / DrawScale,			LocalSize, DrawScale, FVector4f(FVector2f(StartUV.X, EndUV.Y), Tiling), FeatherColor, SecondaryColor)); //8
+			RenderBatch.AddVertex(FSlateVertex::Make<Rounding>(RenderTransform, FVector2f(LeftMarginX, EndPos.Y) + FVector2f(0.f, 1.f) / DrawScale,			LocalSize, DrawScale, FVector4f(FVector2f(LeftMarginU, EndUV.Y), Tiling), FeatherColor, SecondaryColor)); //9
+			RenderBatch.AddVertex(FSlateVertex::Make<Rounding>(RenderTransform, FVector2f(RightMarginX, EndPos.Y) + FVector2f(0.f, 1.f) / DrawScale,		LocalSize, DrawScale, FVector4f(FVector2f(RightMarginU, EndUV.Y), Tiling), FeatherColor, SecondaryColor)); //10
+			RenderBatch.AddVertex(FSlateVertex::Make<Rounding>(RenderTransform, FVector2f(EndPos.X, EndPos.Y) + FVector2f(1.f, 1.f) / DrawScale,			LocalSize, DrawScale, FVector4f(FVector2f(EndUV), Tiling), FeatherColor, SecondaryColor)); //11
 
 			// Top Left
 			IndexQuad(RenderBatch, FeatherStart + 0, FeatherStart + 1, IndexStart + 2, IndexStart + 0);
@@ -876,8 +876,8 @@ void FSlateElementBatcher::AddBoxElement(const FSlateDrawElement& DrawElement)
 	{
 		if( bMirrorHorizontal || bMirrorVertical )
 		{
-			const FVector2D UVMin = StartUV;
-			const FVector2D UVMax = EndUV;
+			const FVector2f UVMin = StartUV;
+			const FVector2f UVMax = EndUV;
 
 			if( bMirrorHorizontal )
 			{
@@ -892,10 +892,10 @@ void FSlateElementBatcher::AddBoxElement(const FSlateDrawElement& DrawElement)
 		}
 
 		// Add four vertices to the list of verts to be added to the vertex buffer
-		RenderBatch.AddVertex(FSlateVertex::Make<Rounding>(RenderTransform, FVector2f(TopLeft), FVector2f(LocalSize), DrawScale, FVector4f(FVector2f(StartUV), FVector2f(Tiling)), Tint, SecondaryColor));
-		RenderBatch.AddVertex(FSlateVertex::Make<Rounding>(RenderTransform, FVector2f(TopRight), FVector2f(LocalSize), DrawScale, FVector4f(FVector2f(EndUV.X, StartUV.Y), FVector2f(Tiling)), Tint, SecondaryColor));
-		RenderBatch.AddVertex(FSlateVertex::Make<Rounding>(RenderTransform, FVector2f(BotLeft), FVector2f(LocalSize), DrawScale, FVector4f(FVector2f(StartUV.X, EndUV.Y), FVector2f(Tiling)), Tint, SecondaryColor));
-		RenderBatch.AddVertex(FSlateVertex::Make<Rounding>(RenderTransform, FVector2f(BotRight), FVector2f(LocalSize), DrawScale, FVector4f(FVector2f(EndUV), FVector2f(Tiling)), Tint, SecondaryColor));
+		RenderBatch.AddVertex(FSlateVertex::Make<Rounding>(RenderTransform, TopLeft, LocalSize, DrawScale, FVector4f(FVector2f(StartUV), Tiling), Tint, SecondaryColor));
+		RenderBatch.AddVertex(FSlateVertex::Make<Rounding>(RenderTransform, TopRight, LocalSize, DrawScale, FVector4f(FVector2f(EndUV.X, StartUV.Y), Tiling), Tint, SecondaryColor));
+		RenderBatch.AddVertex(FSlateVertex::Make<Rounding>(RenderTransform, BotLeft, LocalSize, DrawScale, FVector4f(FVector2f(StartUV.X, EndUV.Y), Tiling), Tint, SecondaryColor));
+		RenderBatch.AddVertex(FSlateVertex::Make<Rounding>(RenderTransform, BotRight, LocalSize, DrawScale, FVector4f(FVector2f(EndUV), Tiling), Tint, SecondaryColor));
 
 		RenderBatch.AddIndex( IndexStart + 0 );
 		RenderBatch.AddIndex( IndexStart + 1 );
@@ -914,10 +914,10 @@ void FSlateElementBatcher::AddBoxElement(const FSlateDrawElement& DrawElement)
 		{
 			const int32 FeatherStart = RenderBatch.GetNumVertices();
 
-			RenderBatch.AddVertex(FSlateVertex::Make<Rounding>(RenderTransform, FVector2f(TopLeft) + FVector2f(-1, -1) / DrawScale, FVector2f(LocalSize), DrawScale, FVector4f(FVector2f(StartUV), FVector2f(Tiling)), FeatherColor, SecondaryColor));
-			RenderBatch.AddVertex(FSlateVertex::Make<Rounding>(RenderTransform, FVector2f(TopRight) + FVector2f(1, -1) / DrawScale, FVector2f(LocalSize), DrawScale, FVector4f(FVector2f(EndUV.X, StartUV.Y), FVector2f(Tiling)), FeatherColor, SecondaryColor));
-			RenderBatch.AddVertex(FSlateVertex::Make<Rounding>(RenderTransform, FVector2f(BotLeft) + FVector2f(-1, 1) / DrawScale, FVector2f(LocalSize), DrawScale, FVector4f(FVector2f(StartUV.X, EndUV.Y), FVector2f(Tiling)), FeatherColor, SecondaryColor));
-			RenderBatch.AddVertex(FSlateVertex::Make<Rounding>(RenderTransform, FVector2f(BotRight) + FVector2f(1, 1) / DrawScale, FVector2f(LocalSize), DrawScale, FVector4f(FVector2f(EndUV), FVector2f(Tiling)), FeatherColor, SecondaryColor));
+			RenderBatch.AddVertex(FSlateVertex::Make<Rounding>(RenderTransform, TopLeft + FVector2f(-1.f, -1.f) / DrawScale, LocalSize, DrawScale, FVector4f(FVector2f(StartUV), Tiling), FeatherColor, SecondaryColor));
+			RenderBatch.AddVertex(FSlateVertex::Make<Rounding>(RenderTransform, TopRight + FVector2f(1.f, -1.f) / DrawScale, LocalSize, DrawScale, FVector4f(FVector2f(EndUV.X, StartUV.Y), Tiling), FeatherColor, SecondaryColor));
+			RenderBatch.AddVertex(FSlateVertex::Make<Rounding>(RenderTransform, BotLeft + FVector2f(-1.f, 1.f) / DrawScale, LocalSize, DrawScale, FVector4f(FVector2f(StartUV.X, EndUV.Y), Tiling), FeatherColor, SecondaryColor));
+			RenderBatch.AddVertex(FSlateVertex::Make<Rounding>(RenderTransform, BotRight + FVector2f(1.f, 1.f) / DrawScale, LocalSize, DrawScale, FVector4f(FVector2f(EndUV), Tiling), FeatherColor, SecondaryColor));
 
 			// Top-Top
 			RenderBatch.AddIndex(FeatherStart + 0);
@@ -1003,7 +1003,7 @@ void FSlateElementBatcher::AddTextElement(const FSlateDrawElement& DrawElement)
 	const int32 Layer = DrawElement.GetLayer();
 
 	// extract the layout transform from the draw element
-	FSlateLayoutTransform LayoutTransform(DrawElement.GetScale(), DrawElement.GetPosition());
+	FSlateLayoutTransform LayoutTransform(DrawElement.GetScale(), FVector2f(DrawElement.GetPosition()));
 
 	// We don't just scale up fonts, we draw them in local space pre-scaled so we don't get scaling artifacts.
 	// So we need to pull the layout scale out of the layout and render transform so we can apply them
@@ -1059,7 +1059,7 @@ void FSlateElementBatcher::AddTextElement(const FSlateDrawElement& DrawElement)
 
 		int32 Kerning = 0;
 
-		FVector2D TopLeft(0, 0);
+		FVector2f TopLeft(0.f, 0.f);
 
 		const float PosX = TopLeft.X;
 		float PosY = TopLeft.Y;
@@ -1152,10 +1152,10 @@ void FSlateElementBatcher::AddTextElement(const FSlateDrawElement& DrawElement)
 					const float SizeV = Entry.VSize * InvTextureSizeY;
 
 					{
-						FVector2D UpperLeft( X, Y );
-						FVector2D UpperRight( X+SizeX, Y );
-						FVector2D LowerLeft( X, Y+SizeY );
-						FVector2D LowerRight( X+SizeX, Y+SizeY );
+						FVector2f UpperLeft(X, Y);
+						FVector2f UpperRight(X + SizeX, Y);
+						FVector2f LowerLeft(X, Y + SizeY);
+						FVector2f LowerRight(X + SizeX, Y + SizeY);
 
 						// The start index of these vertices in the index buffer
 						const uint32 IndexStart = RenderBatch->GetNumVertices();
@@ -1284,7 +1284,7 @@ void FSlateElementBatcher::AddShapedTextElement( const FSlateDrawElement& DrawEl
 	const int32 Layer = DrawElement.GetLayer();
 
 	// extract the layout transform from the draw element
-	FSlateLayoutTransform LayoutTransform(DrawElement.GetScale(), DrawElement.GetPosition());
+	FSlateLayoutTransform LayoutTransform(DrawElement.GetScale(), FVector2f(DrawElement.GetPosition()));
 
 	// We don't just scale up fonts, we draw them in local space pre-scaled so we don't get scaling artifacts.
 	// So we need to pull the layout scale out of the layout and render transform so we can apply them
@@ -1302,7 +1302,7 @@ void FSlateElementBatcher::AddShapedTextElement( const FSlateDrawElement& DrawEl
 
 	auto BuildFontGeometry = [&](const FFontOutlineSettings& InOutlineSettings, const FColor& InTint, const UObject* FontMaterial, int32 InLayer, float InHorizontalOffset)
 	{
-		FVector2D TopLeft(0, 0);
+		FVector2f TopLeft(0, 0);
 
 		const float PosX = TopLeft.X+InHorizontalOffset;
 		float PosY = TopLeft.Y;
@@ -1382,7 +1382,7 @@ template<ESlateVertexRounding Rounding>
 void FSlateElementBatcher::AddGradientElement( const FSlateDrawElement& DrawElement )
 {
 	const FSlateRenderTransform& RenderTransform = DrawElement.GetRenderTransform();
-	const FVector2D LocalSize = DrawElement.GetLocalSize();
+	const FVector2f LocalSize = DrawElement.GetLocalSize();
 	const FSlateGradientPayload& InPayload = DrawElement.GetDataPayload<FSlateGradientPayload>();
 	const ESlateDrawEffect InDrawEffects = DrawElement.GetDrawEffects();
 	const int32 Layer = DrawElement.GetLayer();
@@ -1394,7 +1394,7 @@ void FSlateElementBatcher::AddGradientElement( const FSlateDrawElement& DrawElem
 	FShaderParams ShaderParams;
 
 	ESlateShader ShaderType = ESlateShader::Default;
-	if (InPayload.CornerRadius != FVector4f(0))
+	if (InPayload.CornerRadius != FVector4f(0.f))
 	{
 		ShaderType = ESlateShader::RoundedBox;
 		ShaderParams.PixelParams = FVector4f(0.0f, 0.0f, LocalSize.X, LocalSize.Y);
@@ -1413,10 +1413,10 @@ void FSlateElementBatcher::AddGradientElement( const FSlateDrawElement& DrawElem
 			DrawElement);
 
 	// Determine the four corners of the quad containing the gradient
-	FVector2D TopLeft = FVector2D::ZeroVector;
-	FVector2D TopRight = FVector2D(LocalSize.X, 0);
-	FVector2D BotLeft = FVector2D(0, LocalSize.Y);
-	FVector2D BotRight = FVector2D(LocalSize.X, LocalSize.Y);
+	FVector2f TopLeft = FVector2f::ZeroVector;
+	FVector2f TopRight = FVector2f(LocalSize.X, 0);
+	FVector2f BotLeft = FVector2f(0, LocalSize.Y);
+	FVector2f BotRight = FVector2f(LocalSize.X, LocalSize.Y);
 
 	// Copy the gradient stops.. We may need to add more
 	TArray<FSlateGradientStop> GradientStops = InPayload.GradientStops;
@@ -1431,7 +1431,7 @@ void FSlateElementBatcher::AddGradientElement( const FSlateDrawElement& DrawElem
 		if( 0.0f < FirstStop.Position.X )
 		{
 			// The first stop is after the left side of the quad.  Add a stop at the left side of the quad using the same color as the first stop
-			GradientStops.Insert( FSlateGradientStop( FVector2D(0.0, 0.0), FirstStop.Color ), 0 );
+			GradientStops.Insert( FSlateGradientStop( FVector2f(0.0f, 0.0f), FirstStop.Color ), 0 );
 		}
 
 		if( LocalSize.X > LastStop.Position.X )
@@ -1445,7 +1445,7 @@ void FSlateElementBatcher::AddGradientElement( const FSlateDrawElement& DrawElem
 		if( 0.0f < FirstStop.Position.Y )
 		{
 			// The first stop is after the top side of the quad.  Add a stop at the top side of the quad using the same color as the first stop
-			GradientStops.Insert( FSlateGradientStop( FVector2D(0.0, 0.0), FirstStop.Color ), 0 );
+			GradientStops.Insert( FSlateGradientStop( FVector2f(0.0f, 0.0f), FirstStop.Color ), 0 );
 		}
 
 		if( LocalSize.Y > LastStop.Position.Y )
@@ -1464,12 +1464,12 @@ void FSlateElementBatcher::AddGradientElement( const FSlateDrawElement& DrawElem
 		const FSlateGradientStop& CurStop = GradientStops[StopIndex];
 
 		// The start vertex at this stop
-		FVector2D StartPt;
+		FVector2f StartPt;
 		// The end vertex at this stop
-		FVector2D EndPt;
+		FVector2f EndPt;
 
-		FVector2D StartUV;
-		FVector2D EndUV;
+		FVector2f StartUV;
+		FVector2f EndUV;
 
 		if( InPayload.GradientType == Orient_Vertical )
 		{
@@ -1481,10 +1481,10 @@ void FSlateElementBatcher::AddGradientElement( const FSlateDrawElement& DrawElem
 			EndPt.X += CurStop.Position.X;
 
 			StartUV.X = StartPt.X / TopRight.X;
-			StartUV.Y = 0;
+			StartUV.Y = 0.f;
 
 			EndUV.X = EndPt.X / TopRight.X;
-			EndUV.Y = 1;
+			EndUV.Y = 1.f;
 		}
 		else
 		{
@@ -1496,24 +1496,24 @@ void FSlateElementBatcher::AddGradientElement( const FSlateDrawElement& DrawElem
 			StartPt.Y += CurStop.Position.Y;
 			EndPt.Y += CurStop.Position.Y;
 
-			StartUV.X = 0;
+			StartUV.X = 0.f;
 			StartUV.Y = StartPt.Y / BotLeft.Y;
 
-			EndUV.X = 1;
+			EndUV.X = 1.f;
 			EndUV.Y = StartPt.Y / BotLeft.Y;
 		}
 
 		if( StopIndex == 0 )
 		{
 			// First stop does not have a full quad yet so do not create indices
-			RenderBatch.AddVertex(FSlateVertex::Make<Rounding>(RenderTransform, FVector2f(StartPt), FVector2f(LocalSize), DrawScale, FVector4f(StartUV.X, StartUV.Y, 0,0), PackVertexColor(CurStop.Color), FColor::Transparent));
-			RenderBatch.AddVertex(FSlateVertex::Make<Rounding>(RenderTransform, FVector2f(EndPt), FVector2f(LocalSize), DrawScale, FVector4f(EndUV.X, EndUV.Y, 0, 0), PackVertexColor(CurStop.Color), FColor::Transparent));
+			RenderBatch.AddVertex(FSlateVertex::Make<Rounding>(RenderTransform, StartPt, LocalSize, DrawScale, FVector4f(StartUV.X, StartUV.Y, 0.f, 0.f), PackVertexColor(CurStop.Color), FColor::Transparent));
+			RenderBatch.AddVertex(FSlateVertex::Make<Rounding>(RenderTransform, EndPt, LocalSize, DrawScale, FVector4f(EndUV.X, EndUV.Y, 0.f, 0.f), PackVertexColor(CurStop.Color), FColor::Transparent));
 		}
 		else
 		{
 			// All stops after the first have indices and generate quads
-			RenderBatch.AddVertex(FSlateVertex::Make<Rounding>(RenderTransform, FVector2f(StartPt), FVector2f(LocalSize), DrawScale, FVector4f(StartUV.X, StartUV.Y, 0, 0), PackVertexColor(CurStop.Color), FColor::Transparent));
-			RenderBatch.AddVertex(FSlateVertex::Make<Rounding>(RenderTransform, FVector2f(EndPt), FVector2f(LocalSize), DrawScale, FVector4f(EndUV.X, EndUV.Y, 0, 0), PackVertexColor(CurStop.Color), FColor::Transparent));
+			RenderBatch.AddVertex(FSlateVertex::Make<Rounding>(RenderTransform, StartPt, LocalSize, DrawScale, FVector4f(StartUV.X, StartUV.Y, 0.f, 0.f), PackVertexColor(CurStop.Color), FColor::Transparent));
+			RenderBatch.AddVertex(FSlateVertex::Make<Rounding>(RenderTransform, EndPt, LocalSize, DrawScale, FVector4f(EndUV.X, EndUV.Y, 0.f, 0.f), PackVertexColor(CurStop.Color), FColor::Transparent));
 
 			// Connect the indices to the previous vertices
 			RenderBatch.AddIndex(IndexStart - 2);
@@ -1531,11 +1531,11 @@ void FSlateElementBatcher::AddGradientElement( const FSlateDrawElement& DrawElem
 /** Utility class for building a strip of triangles for a spline. */
 struct FSplineBuilder
 {
-	FSplineBuilder(FSlateRenderBatch& InRenderBatch, const FVector2D StartPoint, float HalfThickness, const FSlateRenderTransform& InRenderTransform, const FColor& InColor)
+	FSplineBuilder(FSlateRenderBatch& InRenderBatch, const FVector2f StartPoint, float HalfThickness, const FSlateRenderTransform& InRenderTransform, const FColor& InColor)
 		: RenderBatch(InRenderBatch)
 		, RenderTransform(InRenderTransform)
 		, LastPointAdded()
-		, LastNormal(FVector2D::ZeroVector)
+		, LastNormal(FVector2f::ZeroVector)
 		, HalfLineThickness(HalfThickness)
 		, NumPointsAdded(1)
 		, SingleColor(InColor)
@@ -1544,14 +1544,14 @@ struct FSplineBuilder
 	}
 
 	
-	void BuildBezierGeometry_WithColorGradient(const TArray<FSlateGradientStop>& GradientStops, int32 GradientStopIndex, const FVector2D& P0, const FVector2D& P1, const FVector2D& P2, const FVector2D& P3, const FSlateElementBatcher& InBatcher)
+	void BuildBezierGeometry_WithColorGradient(const TArray<FSlateGradientStop>& GradientStops, int32 GradientStopIndex, const FVector2f P0, const FVector2f P1, const FVector2f P2, const FVector2f P3, const FSlateElementBatcher& InBatcher)
 	{
 		const int32 NumGradientStops = GradientStops.Num();
 		const float SubdivisionPoint = 1.0f / (NumGradientStops - GradientStopIndex);
 		
 		if (GradientStopIndex < NumGradientStops - 1)
 		{
-			FVector2D TwoCurves[7];
+			FVector2f TwoCurves[7];
 			deCasteljauSplit_WithColorGradient(P0, P1, P2, P3, TwoCurves, SubdivisionPoint);
 			Subdivide_WithColorGradient(GradientStops[GradientStopIndex - 1].Color, GradientStops[GradientStopIndex].Color, InBatcher, TwoCurves[0], TwoCurves[1], TwoCurves[2], TwoCurves[3], *this, 1.0f);
 			BuildBezierGeometry_WithColorGradient(GradientStops, GradientStopIndex + 1, TwoCurves[3], TwoCurves[4], TwoCurves[5], TwoCurves[6], InBatcher);
@@ -1565,26 +1565,26 @@ struct FSplineBuilder
 		
 	}
 
-	void BuildBezierGeometry(const FVector2D& P0, const FVector2D& P1, const FVector2D& P2, const FVector2D& P3)
+	void BuildBezierGeometry(const FVector2f P0, const FVector2f P1, const FVector2f P2, const FVector2f P3)
 	{
 		Subdivide(P0, P1, P2, P3, *this, 1.0f);
 		Finish(P3, SingleColor);
 	}
 	
 private:
-	void AppendPoint(const FVector2D NewPoint, const FColor& InColor)
+	void AppendPoint(const FVector2f NewPoint, const FColor& InColor)
 	{
 		// We only add vertexes for the previous line segment.
 		// This is because we want to average the previous and new normals
 		// In order to prevent overlapping line segments on the spline.
 		// These occur especially when curvature is high.
 
-		const FVector2D NewNormal = FVector2D(LastPointAdded[0].Y - NewPoint.Y, NewPoint.X - LastPointAdded[0].X).GetSafeNormal();
+		const FVector2f NewNormal = FVector2f(LastPointAdded[0].Y - NewPoint.Y, NewPoint.X - LastPointAdded[0].X).GetSafeNormal();
 
 		if (NumPointsAdded == 2)
 		{
 			// Once we have two points, we have a normal, so we can generate the first bit of geometry.
-			const FVector2D LastUp = LastNormal*HalfLineThickness;
+			const FVector2f LastUp = LastNormal*HalfLineThickness;
 
 			RenderBatch.AddVertex(FSlateVertex::Make<ESlateVertexRounding::Disabled>(RenderTransform, FVector2f(LastPointAdded[1] + LastUp), FVector2f(1.0f, 0.0f), FVector2f::ZeroVector, InColor));
 			RenderBatch.AddVertex(FSlateVertex::Make<ESlateVertexRounding::Disabled>(RenderTransform, FVector2f(LastPointAdded[1] - LastUp), FVector2f(-1.0f, 0.0f), FVector2f::ZeroVector, InColor));
@@ -1592,7 +1592,7 @@ private:
 
 		if (NumPointsAdded >= 2)
 		{
-			const FVector2D AveragedUp = (0.5f*(NewNormal + LastNormal)).GetSafeNormal()*HalfLineThickness;
+			const FVector2f AveragedUp = (0.5f*(NewNormal + LastNormal)).GetSafeNormal()*HalfLineThickness;
 
 			RenderBatch.AddVertex(FSlateVertex::Make<ESlateVertexRounding::Disabled>(RenderTransform, FVector2f(LastPointAdded[0] + AveragedUp), FVector2f(1.0f, 0.0f), FVector2f::ZeroVector, InColor));
 			RenderBatch.AddVertex(FSlateVertex::Make<ESlateVertexRounding::Disabled>(RenderTransform, FVector2f(LastPointAdded[0] - AveragedUp), FVector2f(-1.0f, 0.0f), FVector2f::ZeroVector, InColor));
@@ -1616,7 +1616,7 @@ private:
 		++NumPointsAdded;
 	}
 
-	void Finish(const FVector2D& LastPoint, const FColor& InColor)
+	void Finish(const FVector2f LastPoint, const FColor& InColor)
 	{
 		if (NumPointsAdded < 3)
 		{
@@ -1630,7 +1630,7 @@ private:
 			// We have added the last point, but the line builder only builds
 			// geometry for the previous line segment. Build geometry for the
 			// last line segment.
-			const FVector2D LastUp = LastNormal*HalfLineThickness;
+			const FVector2f LastUp = LastNormal*HalfLineThickness;
 
 			RenderBatch.AddVertex(FSlateVertex::Make<ESlateVertexRounding::Disabled>(RenderTransform, FVector2f(LastPointAdded[0] + LastUp), FVector2f(1.0f, 0.0f), FVector2f::ZeroVector, InColor));
 			RenderBatch.AddVertex(FSlateVertex::Make<ESlateVertexRounding::Disabled>(RenderTransform, FVector2f(LastPointAdded[0] - LastUp), FVector2f(-1.0f, 0.0f), FVector2f::ZeroVector, InColor));
@@ -1669,10 +1669,10 @@ private:
 	* Use manhattan distance: 2*Deviation = |P1Deviation.x| + |P1Deviation.y| + |P2Deviation.x| + |P2Deviation.y|
 	*
 	*/
-	static float ComputeCurviness(const FVector2D&  P0, const FVector2D&  P1, const FVector2D&  P2, const FVector2D&  P3)
+	static float ComputeCurviness(const FVector2f P0, const FVector2f P1, const FVector2f P2, const FVector2f P3)
 	{
-		FVector2D TwoP1Deviations = P0 + P2 - 2 * P1;
-		FVector2D TwoP2Deviations = P1 + P3 - 2 * P2;
+		FVector2f TwoP1Deviations = P0 + P2 - 2 * P1;
+		FVector2f TwoP2Deviations = P1 + P3 - 2 * P2;
 		float TwoDeviations = FMath::Abs(TwoP1Deviations.X) + FMath::Abs(TwoP1Deviations.Y) + FMath::Abs(TwoP2Deviations.X) + FMath::Abs(TwoP2Deviations.Y);
 		return TwoDeviations;
 	}
@@ -1696,16 +1696,16 @@ private:
 	* The final points L3 and R0 are both the midpoint of (L2,R1)
 	*
 	*/
-	static void deCasteljauSplit(const FVector2D&  P0, const FVector2D&  P1, const FVector2D&  P2, const FVector2D& P3, FVector2D OutCurveParams[7])
+	static void deCasteljauSplit(const FVector2f P0, const FVector2f P1, const FVector2f P2, const FVector2f P3, FVector2f OutCurveParams[7])
 	{
-		FVector2D L1 = (P0 + P1) * 0.5f;
-		FVector2D M = (P1 + P2) * 0.5f;
-		FVector2D R2 = (P2 + P3) * 0.5f;
+		FVector2f L1 = (P0 + P1) * 0.5f;
+		FVector2f M = (P1 + P2) * 0.5f;
+		FVector2f R2 = (P2 + P3) * 0.5f;
 
-		FVector2D L2 = (L1 + M) * 0.5f;
-		FVector2D R1 = (M + R2) * 0.5f;
+		FVector2f L2 = (L1 + M) * 0.5f;
+		FVector2f R1 = (M + R2) * 0.5f;
 
-		FVector2D L3R0 = (L2 + R1) * 0.5f;
+		FVector2f L3R0 = (L2 + R1) * 0.5f;
 
 		OutCurveParams[0] = P0;
 		OutCurveParams[1] = L1;
@@ -1717,16 +1717,16 @@ private:
 	}
 
 	/** More general form of the deCasteljauSplit splits the curve into two parts at a point between 0 and 1 along the curve's length. */
-	static void deCasteljauSplit_WithColorGradient(const FVector2D&  P0, const FVector2D&  P1, const FVector2D&  P2, const FVector2D& P3, FVector2D OutCurveParams[7], float SplitPoint = 0.5f)
+	static void deCasteljauSplit_WithColorGradient(const FVector2f P0, const FVector2f P1, const FVector2f P2, const FVector2f P3, FVector2f OutCurveParams[7], float SplitPoint = 0.5f)
 	{
-		FVector2D L1 = FMath::Lerp(P0,P1,SplitPoint);
-		FVector2D M = FMath::Lerp(P1,P2,SplitPoint);
-		FVector2D R2 = FMath::Lerp(P2,P3,SplitPoint);
+		FVector2f L1 = FMath::Lerp(P0,P1,SplitPoint);
+		FVector2f M = FMath::Lerp(P1, P2, SplitPoint);
+		FVector2f R2 = FMath::Lerp(P2, P3, SplitPoint);
 
-		FVector2D L2 = FMath::Lerp(L1,M,SplitPoint);
-		FVector2D R1 = FMath::Lerp(M,R2,SplitPoint);
+		FVector2f L2 = FMath::Lerp(L1, M, SplitPoint);
+		FVector2f R1 = FMath::Lerp(M, R2, SplitPoint);
 
-		FVector2D L3R0 = FMath::Lerp(L2,R1,SplitPoint);
+		FVector2f L3R0 = FMath::Lerp(L2,R1,SplitPoint);
 
 		OutCurveParams[0] = P0;
 		OutCurveParams[1] = L1;
@@ -1737,13 +1737,13 @@ private:
 		OutCurveParams[6] = P3;
 	}
 
-	static void Subdivide(const FVector2D&  P0, const FVector2D&  P1, const FVector2D&  P2, const FVector2D&  P3, FSplineBuilder& SplineBuilder, float MaxBiasTimesTwo = 2.0f)
+	static void Subdivide(const FVector2f P0, const FVector2f P1, const FVector2f P2, const FVector2f P3, FSplineBuilder& SplineBuilder, float MaxBiasTimesTwo = 2.0f)
 	{
 		const float Curviness = ComputeCurviness(P0, P1, P2, P3);
 		if (Curviness > MaxBiasTimesTwo)
 		{
 			// Split the Bezier into two curves.
-			FVector2D TwoCurves[7];
+			FVector2f TwoCurves[7];
 			deCasteljauSplit(P0, P1, P2, P3, TwoCurves);
 			// Subdivide left, then right
 			Subdivide(TwoCurves[0], TwoCurves[1], TwoCurves[2], TwoCurves[3], SplineBuilder, MaxBiasTimesTwo);
@@ -1755,13 +1755,13 @@ private:
 		}
 	}
 
-	static void Subdivide_WithColorGradient(const FLinearColor& StartColor, const FLinearColor& EndColor, const FSlateElementBatcher& InBatcher, const FVector2D&  P0, const FVector2D&  P1, const FVector2D&  P2, const FVector2D&  P3, FSplineBuilder& SplineBuilder, float MaxBiasTimesTwo = 2.0f)
+	static void Subdivide_WithColorGradient(const FLinearColor& StartColor, const FLinearColor& EndColor, const FSlateElementBatcher& InBatcher, const FVector2f P0, const FVector2f P1, const FVector2f P2, const FVector2f P3, FSplineBuilder& SplineBuilder, float MaxBiasTimesTwo = 2.0f)
 	{
 		const float Curviness = ComputeCurviness(P0, P1, P2, P3);
 		if (Curviness > MaxBiasTimesTwo)
 		{
 			// Split the Bezier into two curves.
-			FVector2D TwoCurves[7];
+			FVector2f TwoCurves[7];
 			deCasteljauSplit(P0, P1, P2, P3, TwoCurves);
 			const FLinearColor MidpointColor = FLinearColor::LerpUsingHSV(StartColor, EndColor, 0.5f);
 			// Subdivide left, then right
@@ -1777,8 +1777,8 @@ private:
 private:
 	FSlateRenderBatch& RenderBatch;
 	const FSlateRenderTransform& RenderTransform;
-	FVector2D LastPointAdded[2];
-	FVector2D LastNormal;
+	FVector2f LastPointAdded[2];
+	FVector2f LastNormal;
 	float HalfLineThickness;
 	int32 NumPointsAdded;
 	FColor SingleColor;
@@ -1840,9 +1840,9 @@ void FSlateElementBatcher::AddSplineElement(const FSlateDrawElement& DrawElement
 }
 
 /** Equivalent to V.GetRotated(90.0f) */
-static FVector2D GetRotated90(const FVector2D& V)
+static FVector2f GetRotated90(const FVector2f V)
 {
-	return FVector2D(-V.Y, V.X);
+	return FVector2f(-V.Y, V.X);
 }
 
 
@@ -1878,16 +1878,16 @@ struct FLineBuilder
 	* UVs shearing apart at the diagonal where the two triangles meet.
 	*/
 	template<ESlateVertexRounding Rounding>
-	void BuildLineGeometry(const TArray<FVector2D>& Points, const TArray<FColor>& PackedColors, const FColor& PackedTint) const
+	void BuildLineGeometry(const TArray<FVector2f>& Points, const TArray<FColor>& PackedColors, const FColor& PackedTint) const
 	{
 		FColor PointColor = PackedColors.Num() ? PackedColors[0] : PackedTint;
-		FVector2D Position = Points[0];
-		FVector2D NextPosition = Points[1];
+		FVector2f Position = Points[0];
+		FVector2f NextPosition = Points[1];
 
-		FVector2D Direction;
+		FVector2f Direction;
 		float Length;
 		(NextPosition - Position).ToDirectionAndLength(Direction, Length);
-		FVector2D Up = GetRotated90(Direction) * LocalHalfThickness;
+		FVector2f Up = GetRotated90(Direction) * LocalHalfThickness;
 
 		// Build the start cap at the first point
 		MakeStartCap<Rounding>(Position, Direction, Length, Up, PointColor);
@@ -1896,8 +1896,8 @@ struct FLineBuilder
 		const int32 LastPointIndex = Points.Num() - 1;
 		for (int32 Point = 1; Point < LastPointIndex; ++Point)
 		{
-			const FVector2D LastDirection = Direction;
-			const FVector2D LastUp = Up;
+			const FVector2f LastDirection = Direction;
+			const FVector2f LastUp = Up;
 			const float LastLength = Length;
 
 			PointColor = PackedColors.Num() ? PackedColors[Point] : PackedTint;
@@ -1908,11 +1908,11 @@ struct FLineBuilder
 			Up = GetRotated90(Direction) * LocalHalfThickness;
 
 			// Project "up" onto the miter normal to get perpendicular distance to the miter line
-			const FVector2D MiterNormal = GetMiterNormal(LastDirection, Direction);
-			const float DistanceToMiterLine = FVector2D::DotProduct(Up, MiterNormal);
+			const FVector2f MiterNormal = GetMiterNormal(LastDirection, Direction);
+			const float DistanceToMiterLine = FVector2f::DotProduct(Up, MiterNormal);
 
 			// Get the component of Direction perpendicular to the miter line
-			const float DirDotMiterNormal = FVector2D::DotProduct(Direction, MiterNormal);
+			const float DirDotMiterNormal = FVector2f::DotProduct(Direction, MiterNormal);
 
 			// Break the strip at zero-length segments, if the miter angle is
 			// too tight or the miter would cross either segment's bisector
@@ -1924,7 +1924,7 @@ struct FLineBuilder
 				// Calculate the offset to put the vertices on the miter line while
 				// keeping the antialiased edges of each segment's quad parallel
 				const float ParallelDistance = DistanceToMiterLine / DirDotMiterNormal;
-				const FVector2D MiterUp = Up - (Direction * ParallelDistance);
+				const FVector2f MiterUp = Up - (Direction * ParallelDistance);
 
 				RenderBatch.AddVertex(FSlateVertex::Make<Rounding>(RenderTransform, FVector2f(Position + MiterUp), FVector2f(1.0f, 0.0f), PointColor));
 				RenderBatch.AddVertex(FSlateVertex::Make<Rounding>(RenderTransform, FVector2f(Position - MiterUp), FVector2f(-1.0f, 0.0f), PointColor));
@@ -1964,10 +1964,10 @@ private:
 	*/
 	template<ESlateVertexRounding Rounding>
 	void MakeStartCap(
-		const FVector2D& Position,
-		const FVector2D& Direction,
+		const FVector2f Position,
+		const FVector2f Direction,
 		float SegmentLength,
-		const FVector2D& Up,
+		const FVector2f Up,
 		const FColor& Color) const
 	{
 		// Don't build a cap for a zero-length segment
@@ -1978,8 +1978,8 @@ private:
 			// Center the cap over Position if possible, but never place
 			// vertices on the far side of this segment's midpoint
 			const float InwardDistance = FMath::Min(LocalFilterRadius, SegmentLength * 0.5f);
-			const FVector2D CapInward = Direction * InwardDistance;
-			const FVector2D CapOutward = Direction * (InwardDistance - LocalCapLength);
+			const FVector2f CapInward = Direction * InwardDistance;
+			const FVector2f CapOutward = Direction * (InwardDistance - LocalCapLength);
 
 			RenderBatch.AddVertex(FSlateVertex::Make<Rounding>(RenderTransform, FVector2f(Position + CapOutward + Up), FVector2f(1.0f, -1.0f), Color));
 			RenderBatch.AddVertex(FSlateVertex::Make<Rounding>(RenderTransform, FVector2f(Position + CapOutward - Up), FVector2f(-1.0f, -1.0f), Color));
@@ -1994,10 +1994,10 @@ private:
 	*/
 	template<ESlateVertexRounding Rounding>
 	void MakeEndCap(
-		const FVector2D& Position,
-		const FVector2D& Direction,
+		const FVector2f Position,
+		const FVector2f Direction,
 		float SegmentLength,
-		const FVector2D& Up,
+		const FVector2f Up,
 		const FColor& Color) const
 	{
 		// Don't build a cap for a zero-length segment
@@ -2008,8 +2008,8 @@ private:
 			// Center the cap over Position if possible, but never place
 			// vertices on the far side of this segment's midpoint
 			const float InwardDistance = FMath::Min(LocalFilterRadius, SegmentLength * 0.5f);
-			const FVector2D CapInward = Direction * -InwardDistance;
-			const FVector2D CapOutward = Direction * (LocalCapLength - InwardDistance);
+			const FVector2f CapInward = Direction * -InwardDistance;
+			const FVector2f CapOutward = Direction * (LocalCapLength - InwardDistance);
 
 			RenderBatch.AddVertex(FSlateVertex::Make<Rounding>(RenderTransform, FVector2f(Position + CapInward + Up), FVector2f(1.0f, 0.0f), Color));
 			RenderBatch.AddVertex(FSlateVertex::Make<Rounding>(RenderTransform, FVector2f(Position + CapInward - Up), FVector2f(-1.0f, 0.0f), Color));
@@ -2028,9 +2028,9 @@ private:
 	* Inputs must be normalized vectors pointing from the start to the end of
 	* two consecutive segments.
 	*/
-	static FVector2D GetMiterNormal(const FVector2D& InboundSegmentDir, const FVector2D& OutboundSegmentDir)
+	static FVector2f GetMiterNormal(const FVector2f InboundSegmentDir, const FVector2f OutboundSegmentDir)
 	{
-		const FVector2D DirSum = InboundSegmentDir + OutboundSegmentDir;
+		const FVector2f DirSum = InboundSegmentDir + OutboundSegmentDir;
 		const float SizeSquared = DirSum.SizeSquared();
 
 		if (SizeSquared > SMALL_NUMBER)
@@ -2060,7 +2060,7 @@ void FSlateElementBatcher::AddLineElement( const FSlateDrawElement& DrawElement
 	const ESlateDrawEffect DrawEffects = DrawElement.GetDrawEffects();
 	const int32 Layer = DrawElement.GetLayer();
 
-	const TArray<FVector2D>& Points = DrawElementPayload.GetPoints();
+	const TArray<FVector2f>& Points = DrawElementPayload.GetPoints();
 	const TArray<FLinearColor>& PointColors = DrawElementPayload.GetPointColors();
 
 	const int32 NumPoints = Points.Num();
@@ -2115,12 +2115,12 @@ void FSlateElementBatcher::AddLineElement( const FSlateDrawElement& DrawElement
 			FSlateRenderBatch& RenderBatch = CreateRenderBatch(Layer, FShaderParams(), nullptr, ESlateDrawPrimitive::LineList, ESlateShader::Default, DrawEffects, ESlateBatchDrawFlag::None, DrawElement);
 
 			const FColor StartColor = PackedColors.Num() ? PackedColors[0] : PackedTint;
-			RenderBatch.AddVertex(FSlateVertex::Make<Rounding>(RenderTransform, FVector2f(Points[0]), FVector2f::ZeroVector, StartColor));
+			RenderBatch.AddVertex(FSlateVertex::Make<Rounding>(RenderTransform, Points[0], FVector2f::ZeroVector, StartColor));
 
 			for (int32 Point = 1; Point < NumPoints; ++Point)
 			{
 				const FColor PointColor = PackedColors.Num() ? PackedColors[Point] : PackedTint;
-				RenderBatch.AddVertex(FSlateVertex::Make<Rounding>(RenderTransform, FVector2f(Points[Point]), FVector2f::ZeroVector, PointColor));
+				RenderBatch.AddVertex(FSlateVertex::Make<Rounding>(RenderTransform, Points[Point], FVector2f::ZeroVector, PointColor));
 
 				const uint32 IndexStart = RenderBatch.GetNumVertices();
 				RenderBatch.AddIndex(IndexStart - 2);
@@ -2137,14 +2137,14 @@ void FSlateElementBatcher::AddLineElement( const FSlateDrawElement& DrawElement
 
 			for (int32 Point = 0; Point < NumPoints - 1; ++Point)
 			{
-				const FVector2D StartPos = Points[Point];
-				const FVector2D EndPos = Points[Point + 1];
+				const FVector2f StartPos = Points[Point];
+				const FVector2f EndPos = Points[Point + 1];
 
 				const FColor StartColor = PackedColors.Num() ? PackedColors[Point] : PackedTint;
 				const FColor EndColor = PackedColors.Num() ? PackedColors[Point+1] : PackedTint;
 
-				const FVector2D SegmentDir = (EndPos - StartPos).GetSafeNormal();
-				const FVector2D Up = GetRotated90(SegmentDir) * LocalHalfThickness;
+				const FVector2f SegmentDir = (EndPos - StartPos).GetSafeNormal();
+				const FVector2f Up = GetRotated90(SegmentDir) * LocalHalfThickness;
 
 				RenderBatch.AddVertex(FSlateVertex::Make<Rounding>(RenderTransform, FVector2f(StartPos + Up), FVector2f::ZeroVector, StartColor));
 				RenderBatch.AddVertex(FSlateVertex::Make<Rounding>(RenderTransform, FVector2f(StartPos - Up), FVector2f::ZeroVector, StartColor));
@@ -2161,7 +2161,7 @@ template<ESlateVertexRounding Rounding>
 void FSlateElementBatcher::AddViewportElement( const FSlateDrawElement& DrawElement )
 {
 	const FSlateRenderTransform& RenderTransform = DrawElement.GetRenderTransform();
-	const FVector2D& LocalSize = DrawElement.GetLocalSize();
+	const FVector2f LocalSize = DrawElement.GetLocalSize();
 	const FSlateViewportPayload& DrawElementPayload = DrawElement.GetDataPayload<FSlateViewportPayload>();
 	ESlateDrawEffect InDrawEffects = DrawElement.GetDrawEffects();
 	const int32 Layer = DrawElement.GetLayer();
@@ -2193,27 +2193,27 @@ void FSlateElementBatcher::AddViewportElement( const FSlateDrawElement& DrawElem
 	}
 
 	// Do pixel snapping
-	FVector2D TopLeft(0,0);
-	FVector2D BotRight(LocalSize);
+	FVector2f TopLeft(0.f, 0.f);
+	FVector2f BotRight(LocalSize);
 
 	// If the viewport disallows scaling, force size to current texture size.
 	if (ViewportResource != nullptr && !DrawElementPayload.bAllowViewportScaling)
 	{
 		const float ElementScale = DrawElement.GetScale();
-		BotRight = FVector2D(ViewportResource->GetWidth() / ElementScale, ViewportResource->GetHeight() / ElementScale);
+		BotRight = FVector2f(ViewportResource->GetWidth() / ElementScale, ViewportResource->GetHeight() / ElementScale);
 	}
 
-	FVector2D TopRight = FVector2D( BotRight.X, TopLeft.Y);
-	FVector2D BotLeft =	 FVector2D( TopLeft.X, BotRight.Y);
+	FVector2f TopRight = FVector2f(BotRight.X, TopLeft.Y);
+	FVector2f BotLeft = FVector2f(TopLeft.X, BotRight.Y);
 
 	// The start index of these vertices in the index buffer
 	const uint32 IndexStart = 0;
 
 	// Add four vertices to the list of verts to be added to the vertex buffer
-	RenderBatch.AddVertex( FSlateVertex::Make<Rounding>( RenderTransform, FVector2f(TopLeft),	FVector2f(0.0f,0.0f),	FinalColor ) );
-	RenderBatch.AddVertex( FSlateVertex::Make<Rounding>( RenderTransform, FVector2f(TopRight),	FVector2f(1.0f,0.0f),	FinalColor ) );
-	RenderBatch.AddVertex( FSlateVertex::Make<Rounding>( RenderTransform, FVector2f(BotLeft),	FVector2f(0.0f,1.0f),	FinalColor ) );
-	RenderBatch.AddVertex( FSlateVertex::Make<Rounding>( RenderTransform, FVector2f(BotRight),	FVector2f(1.0f,1.0f),	FinalColor ) );
+	RenderBatch.AddVertex( FSlateVertex::Make<Rounding>( RenderTransform, TopLeft,	FVector2f(0.0f,0.0f),	FinalColor ) );
+	RenderBatch.AddVertex( FSlateVertex::Make<Rounding>( RenderTransform, TopRight,	FVector2f(1.0f,0.0f),	FinalColor ) );
+	RenderBatch.AddVertex( FSlateVertex::Make<Rounding>( RenderTransform, BotLeft,	FVector2f(0.0f,1.0f),	FinalColor ) );
+	RenderBatch.AddVertex( FSlateVertex::Make<Rounding>( RenderTransform, BotRight,	FVector2f(1.0f,1.0f),	FinalColor ) );
 
 	// Add 6 indices to the vertex buffer.  (2 tri's per quad, 3 indices per tri)
 	RenderBatch.AddIndex( IndexStart + 0 );
@@ -2233,10 +2233,10 @@ void FSlateElementBatcher::AddViewportElement( const FSlateDrawElement& DrawElem
 		DrawFlags = ESlateBatchDrawFlag::NoBlending;
 		const FColor TransparentBlackColor(0, 0, 0, 0);
 		FSlateRenderBatch& RenderBatch2 = CreateRenderBatch(Layer, FShaderParams(), nullptr, ESlateDrawPrimitive::TriangleList, ShaderType, InDrawEffects, DrawFlags, DrawElement);
-		RenderBatch2.AddVertex(FSlateVertex::Make<Rounding>(RenderTransform, FVector2f(TopLeft), FVector2f(0.0f, 0.0f), TransparentBlackColor));
-		RenderBatch2.AddVertex(FSlateVertex::Make<Rounding>(RenderTransform, FVector2f(TopRight), FVector2f(1.0f, 0.0f), TransparentBlackColor));
-		RenderBatch2.AddVertex(FSlateVertex::Make<Rounding>(RenderTransform, FVector2f(BotLeft), FVector2f(0.0f, 1.0f), TransparentBlackColor));
-		RenderBatch2.AddVertex(FSlateVertex::Make<Rounding>(RenderTransform, FVector2f(BotRight), FVector2f(1.0f, 1.0f), TransparentBlackColor));
+		RenderBatch2.AddVertex(FSlateVertex::Make<Rounding>(RenderTransform, TopLeft, FVector2f(0.0f, 0.0f), TransparentBlackColor));
+		RenderBatch2.AddVertex(FSlateVertex::Make<Rounding>(RenderTransform, TopRight, FVector2f(1.0f, 0.0f), TransparentBlackColor));
+		RenderBatch2.AddVertex(FSlateVertex::Make<Rounding>(RenderTransform, BotLeft, FVector2f(0.0f, 1.0f), TransparentBlackColor));
+		RenderBatch2.AddVertex(FSlateVertex::Make<Rounding>(RenderTransform, BotRight, FVector2f(1.0f, 1.0f), TransparentBlackColor));
 
 		// Add 6 indices to the vertex buffer.  (2 tri's per quad, 3 indices per tri)
 		RenderBatch2.AddIndex(IndexStart + 0);
@@ -2255,7 +2255,7 @@ void FSlateElementBatcher::AddBorderElement( const FSlateDrawElement& DrawElemen
 {
 	const FSlateBoxPayload& DrawElementPayload = DrawElement.GetDataPayload<FSlateBoxPayload>();
 	const FSlateRenderTransform& RenderTransform = DrawElement.GetRenderTransform();
-	const FVector2D& LocalSize = DrawElement.GetLocalSize();
+	const FVector2f LocalSize = DrawElement.GetLocalSize();
 	const ESlateDrawEffect InDrawEffects = DrawElement.GetDrawEffects();
 
 	const int32 Layer = DrawElement.GetLayer();
@@ -2273,19 +2273,19 @@ void FSlateElementBatcher::AddBorderElement( const FSlateDrawElement& DrawElemen
 		TextureWidth = Resource->GetWidth();
 		TextureHeight = Resource->GetHeight();
 	}
-	FVector2D TextureSizeLocalSpace = TransformVector(DrawElement.GetInverseLayoutTransform(), FVector2D((float)TextureWidth, (float)TextureHeight));
+	FVector2f TextureSizeLocalSpace = TransformVector(DrawElement.GetInverseLayoutTransform(), FVector2f((float)TextureWidth, (float)TextureHeight));
  
 	// Texel offset
-	const FVector2D HalfTexel( PixelCenterOffset/TextureWidth, PixelCenterOffset/TextureHeight );
+	const FVector2f HalfTexel( PixelCenterOffset/TextureWidth, PixelCenterOffset/TextureHeight );
 
-	const FVector2D StartUV = HalfTexel;
-	const FVector2D EndUV = FVector2D( 1.0f, 1.0f ) + HalfTexel;
+	const FVector2f StartUV = HalfTexel;
+	const FVector2f EndUV = FVector2f( 1.0f, 1.0f ) + HalfTexel;
 
 	FMargin Margin = DrawElementPayload.GetBrushMargin();
 
 	// Do pixel snapping
-	FVector2D TopLeft(0,0);
-	FVector2D BotRight(LocalSize);
+	FVector2f TopLeft(0.f, 0.f);
+	FVector2f BotRight(LocalSize);
 
 	// Account for negative sizes
 	bool bIsFlippedX = TopLeft.X > BotRight.X;
@@ -2296,8 +2296,8 @@ void FSlateElementBatcher::AddBorderElement( const FSlateDrawElement& DrawElemen
 	Margin.Bottom = bIsFlippedY ? -Margin.Bottom : Margin.Bottom;
 
 	// Determine the margins for each quad
-	FVector2D TopLeftMargin(TextureSizeLocalSpace * FVector2D(Margin.Left, Margin.Top));
-	FVector2D BotRightMargin(LocalSize - TextureSizeLocalSpace * FVector2D(Margin.Right, Margin.Bottom));
+	FVector2f TopLeftMargin(TextureSizeLocalSpace * FVector2f(Margin.Left, Margin.Top));
+	FVector2f BotRightMargin(LocalSize - TextureSizeLocalSpace * FVector2f(Margin.Right, Margin.Bottom));
 
 	float LeftMarginX = TopLeftMargin.X;
 	float TopMarginY = TopLeftMargin.Y;
@@ -2359,55 +2359,55 @@ void FSlateElementBatcher::AddBorderElement( const FSlateDrawElement& DrawElemen
 	float RightTiling = LeftTiling;
 	float BottomTiling = TopTiling;
 
-	FVector2D Position = TopLeft;
-	FVector2D EndPos = BotRight;
+	FVector2f Position = TopLeft;
+	FVector2f EndPos = BotRight;
 
 	// The start index of these vertices in the index buffer
 	const uint32 IndexStart = RenderBatch.GetNumVertices();
 
 	// Zero for second UV indicates no tiling and to just pass the UV though (for the corner sections)
-	FVector2D Zero(0,0);
+	FVector2f Zero(0.f,0.f);
 
 	// Add all the vertices needed for this element.  Vertices are duplicated so that we can have some sections with no tiling and some with tiling.
-	RenderBatch.AddVertex( FSlateVertex::Make<Rounding>( RenderTransform, FVector2f( Position ),					FVector2f(LocalSize), DrawScale, FVector4f( StartUV.X, StartUV.Y, 0.0f, 0.0f),				Tint ) ); //0
-	RenderBatch.AddVertex( FSlateVertex::Make<Rounding>( RenderTransform, FVector2f( Position.X, TopMarginY ),		FVector2f(LocalSize), DrawScale, FVector4f( StartUV.X, TopMarginV, 0.0f, 0.0f),				Tint ) ); //1
-	RenderBatch.AddVertex( FSlateVertex::Make<Rounding>( RenderTransform, FVector2f( LeftMarginX, Position.Y ),		FVector2f(LocalSize), DrawScale, FVector4f( LeftMarginU, StartUV.Y, 0.0f, 0.0f),			Tint ) ); //2
-	RenderBatch.AddVertex( FSlateVertex::Make<Rounding>( RenderTransform, FVector2f( LeftMarginX, TopMarginY ),		FVector2f(LocalSize), DrawScale, FVector4f( LeftMarginU, TopMarginV, 0.0f, 0.0f),			Tint ) ); //3
-
-	RenderBatch.AddVertex( FSlateVertex::Make<Rounding>( RenderTransform, FVector2f( LeftMarginX, Position.Y ),		FVector2f(LocalSize), DrawScale, FVector4f( StartUV.X, StartUV.Y, TopTiling, 0.0f),			Tint ) ); //4
-	RenderBatch.AddVertex( FSlateVertex::Make<Rounding>( RenderTransform, FVector2f( LeftMarginX, TopMarginY ),		FVector2f(LocalSize), DrawScale, FVector4f( StartUV.X, TopMarginV, TopTiling, 0.0f),		Tint ) ); //5
-	RenderBatch.AddVertex( FSlateVertex::Make<Rounding>( RenderTransform, FVector2f( RightMarginX, Position.Y ),	FVector2f(LocalSize), DrawScale, FVector4f( EndUV.X, StartUV.Y, TopTiling, 0.0f),			Tint ) ); //6
-	RenderBatch.AddVertex( FSlateVertex::Make<Rounding>( RenderTransform, FVector2f( RightMarginX, TopMarginY ),	FVector2f(LocalSize), DrawScale, FVector4f( EndUV.X, TopMarginV, TopTiling, 0.0f),			Tint ) ); //7
-
-	RenderBatch.AddVertex( FSlateVertex::Make<Rounding>( RenderTransform, FVector2f( RightMarginX, Position.Y ),	FVector2f(LocalSize), DrawScale, FVector4f( RightMarginU, StartUV.Y, 0.0f, 0.0f),			Tint ) ); //8
-	RenderBatch.AddVertex( FSlateVertex::Make<Rounding>( RenderTransform, FVector2f( RightMarginX, TopMarginY ),	FVector2f(LocalSize), DrawScale, FVector4f( RightMarginU, TopMarginV, 0.0f, 0.0f),			Tint ) ); //9
-	RenderBatch.AddVertex( FSlateVertex::Make<Rounding>( RenderTransform, FVector2f( EndPos.X, Position.Y ),		FVector2f(LocalSize), DrawScale, FVector4f( EndUV.X, StartUV.Y, 0.0f, 0.0f),				Tint ) ); //10
-	RenderBatch.AddVertex( FSlateVertex::Make<Rounding>( RenderTransform, FVector2f( EndPos.X, TopMarginY ),		FVector2f(LocalSize), DrawScale, FVector4f( EndUV.X, TopMarginV, 0.0f, 0.0f),				Tint ) ); //11
-
-	RenderBatch.AddVertex( FSlateVertex::Make<Rounding>( RenderTransform, FVector2f( Position.X, TopMarginY ),		FVector2f(LocalSize), DrawScale, FVector4f( StartUV.X, StartUV.Y, 0.0f, LeftTiling),		Tint ) ); //12
-	RenderBatch.AddVertex( FSlateVertex::Make<Rounding>( RenderTransform, FVector2f( Position.X, BottomMarginY ),	FVector2f(LocalSize), DrawScale, FVector4f( StartUV.X, EndUV.Y, 0.0f, LeftTiling),			Tint ) ); //13
-	RenderBatch.AddVertex( FSlateVertex::Make<Rounding>( RenderTransform, FVector2f( LeftMarginX, TopMarginY ),		FVector2f(LocalSize), DrawScale, FVector4f( LeftMarginU, StartUV.Y, 0.0f, LeftTiling),		Tint ) ); //14
-	RenderBatch.AddVertex( FSlateVertex::Make<Rounding>( RenderTransform, FVector2f( LeftMarginX, BottomMarginY ),	FVector2f(LocalSize), DrawScale, FVector4f( LeftMarginU, EndUV.Y, 0.0f, LeftTiling),		Tint ) ); //15
-
-	RenderBatch.AddVertex( FSlateVertex::Make<Rounding>( RenderTransform, FVector2f( RightMarginX, TopMarginY ),	FVector2f(LocalSize), DrawScale, FVector4f( RightMarginU, StartUV.Y, 0.0f, RightTiling),	Tint ) ); //16
-	RenderBatch.AddVertex( FSlateVertex::Make<Rounding>( RenderTransform, FVector2f( RightMarginX, BottomMarginY ), FVector2f(LocalSize), DrawScale, FVector4f( RightMarginU, EndUV.Y, 0.0f, RightTiling),		Tint ) ); //17
-	RenderBatch.AddVertex( FSlateVertex::Make<Rounding>( RenderTransform, FVector2f( EndPos.X, TopMarginY ),		FVector2f(LocalSize), DrawScale, FVector4f( EndUV.X, StartUV.Y, 0.0f, RightTiling),			Tint ) ); //18
-	RenderBatch.AddVertex( FSlateVertex::Make<Rounding>( RenderTransform, FVector2f( EndPos.X, BottomMarginY ),		FVector2f(LocalSize), DrawScale, FVector4f( EndUV.X, EndUV.Y, 0.0f, RightTiling),			Tint ) ); //19
-
-	RenderBatch.AddVertex( FSlateVertex::Make<Rounding>( RenderTransform, FVector2f( Position.X, BottomMarginY ),	FVector2f(LocalSize), DrawScale, FVector4f( StartUV.X, BottomMarginV, 0.0f, 0.0f),			Tint ) ); //20
-	RenderBatch.AddVertex( FSlateVertex::Make<Rounding>( RenderTransform, FVector2f( Position.X, EndPos.Y ),		FVector2f(LocalSize), DrawScale, FVector4f( StartUV.X, EndUV.Y, 0.0f, 0.0f),				Tint ) ); //21
-	RenderBatch.AddVertex( FSlateVertex::Make<Rounding>( RenderTransform, FVector2f( LeftMarginX, BottomMarginY ),	FVector2f(LocalSize), DrawScale, FVector4f( LeftMarginU, BottomMarginV, 0.0f, 0.0f),		Tint ) ); //22
-	RenderBatch.AddVertex( FSlateVertex::Make<Rounding>( RenderTransform, FVector2f( LeftMarginX, EndPos.Y ),		FVector2f(LocalSize), DrawScale, FVector4f( LeftMarginU, EndUV.Y, 0.0f, 0.0f),				Tint ) ); //23
-
-	RenderBatch.AddVertex( FSlateVertex::Make<Rounding>( RenderTransform, FVector2f( LeftMarginX, BottomMarginY ),	FVector2f(LocalSize), DrawScale, FVector4f( StartUV.X, BottomMarginV, BottomTiling, 0.0f),	Tint ) ); //24
-	RenderBatch.AddVertex( FSlateVertex::Make<Rounding>( RenderTransform, FVector2f( LeftMarginX, EndPos.Y ),		FVector2f(LocalSize), DrawScale, FVector4f( StartUV.X, EndUV.Y, BottomTiling, 0.0f),		Tint ) ); //25
-	RenderBatch.AddVertex( FSlateVertex::Make<Rounding>( RenderTransform, FVector2f( RightMarginX,BottomMarginY ),	FVector2f(LocalSize), DrawScale, FVector4f( EndUV.X, BottomMarginV, BottomTiling, 0.0f),	Tint ) ); //26
-	RenderBatch.AddVertex( FSlateVertex::Make<Rounding>( RenderTransform, FVector2f( RightMarginX, EndPos.Y ),		FVector2f(LocalSize), DrawScale, FVector4f( EndUV.X, EndUV.Y, BottomTiling, 0.0f),			Tint ) ); //27
-
-	RenderBatch.AddVertex( FSlateVertex::Make<Rounding>( RenderTransform, FVector2f( RightMarginX, BottomMarginY ), FVector2f(LocalSize), DrawScale, FVector4f( RightMarginU, BottomMarginV, 0.0f, 0.0f),		Tint ) ); //29
-	RenderBatch.AddVertex( FSlateVertex::Make<Rounding>( RenderTransform, FVector2f( RightMarginX, EndPos.Y ),		FVector2f(LocalSize), DrawScale, FVector4f( RightMarginU, EndUV.Y, 0.0f, 0.0f),				Tint ) ); //30
-	RenderBatch.AddVertex( FSlateVertex::Make<Rounding>( RenderTransform, FVector2f( EndPos.X, BottomMarginY ),		FVector2f(LocalSize), DrawScale, FVector4f( EndUV.X, BottomMarginV, 0.0f, 0.0f),			Tint ) ); //31
-	RenderBatch.AddVertex( FSlateVertex::Make<Rounding>( RenderTransform, FVector2f( EndPos.X, EndPos.Y ),			FVector2f(LocalSize), DrawScale, FVector4f( EndUV.X, EndUV.Y, 0.0f, 0.0f),					Tint ) ); //32
+	RenderBatch.AddVertex( FSlateVertex::Make<Rounding>( RenderTransform, FVector2f( Position ),					LocalSize, DrawScale, FVector4f( StartUV.X, StartUV.Y, 0.0f, 0.0f),				Tint ) ); //0
+	RenderBatch.AddVertex( FSlateVertex::Make<Rounding>( RenderTransform, FVector2f( Position.X, TopMarginY ),		LocalSize, DrawScale, FVector4f( StartUV.X, TopMarginV, 0.0f, 0.0f),				Tint ) ); //1
+	RenderBatch.AddVertex( FSlateVertex::Make<Rounding>( RenderTransform, FVector2f( LeftMarginX, Position.Y ),		LocalSize, DrawScale, FVector4f( LeftMarginU, StartUV.Y, 0.0f, 0.0f),			Tint ) ); //2
+	RenderBatch.AddVertex( FSlateVertex::Make<Rounding>( RenderTransform, FVector2f( LeftMarginX, TopMarginY ),		LocalSize, DrawScale, FVector4f( LeftMarginU, TopMarginV, 0.0f, 0.0f),			Tint ) ); //3
+
+	RenderBatch.AddVertex( FSlateVertex::Make<Rounding>( RenderTransform, FVector2f( LeftMarginX, Position.Y ),		LocalSize, DrawScale, FVector4f( StartUV.X, StartUV.Y, TopTiling, 0.0f),			Tint ) ); //4
+	RenderBatch.AddVertex( FSlateVertex::Make<Rounding>( RenderTransform, FVector2f( LeftMarginX, TopMarginY ),		LocalSize, DrawScale, FVector4f( StartUV.X, TopMarginV, TopTiling, 0.0f),		Tint ) ); //5
+	RenderBatch.AddVertex( FSlateVertex::Make<Rounding>( RenderTransform, FVector2f( RightMarginX, Position.Y ),	LocalSize, DrawScale, FVector4f( EndUV.X, StartUV.Y, TopTiling, 0.0f),			Tint ) ); //6
+	RenderBatch.AddVertex( FSlateVertex::Make<Rounding>( RenderTransform, FVector2f( RightMarginX, TopMarginY ),	LocalSize, DrawScale, FVector4f( EndUV.X, TopMarginV, TopTiling, 0.0f),			Tint ) ); //7
+
+	RenderBatch.AddVertex( FSlateVertex::Make<Rounding>( RenderTransform, FVector2f( RightMarginX, Position.Y ),	LocalSize, DrawScale, FVector4f( RightMarginU, StartUV.Y, 0.0f, 0.0f),			Tint ) ); //8
+	RenderBatch.AddVertex( FSlateVertex::Make<Rounding>( RenderTransform, FVector2f( RightMarginX, TopMarginY ),	LocalSize, DrawScale, FVector4f( RightMarginU, TopMarginV, 0.0f, 0.0f),			Tint ) ); //9
+	RenderBatch.AddVertex( FSlateVertex::Make<Rounding>( RenderTransform, FVector2f( EndPos.X, Position.Y ),		LocalSize, DrawScale, FVector4f( EndUV.X, StartUV.Y, 0.0f, 0.0f),				Tint ) ); //10
+	RenderBatch.AddVertex( FSlateVertex::Make<Rounding>( RenderTransform, FVector2f( EndPos.X, TopMarginY ),		LocalSize, DrawScale, FVector4f( EndUV.X, TopMarginV, 0.0f, 0.0f),				Tint ) ); //11
+
+	RenderBatch.AddVertex( FSlateVertex::Make<Rounding>( RenderTransform, FVector2f( Position.X, TopMarginY ),		LocalSize, DrawScale, FVector4f( StartUV.X, StartUV.Y, 0.0f, LeftTiling),		Tint ) ); //12
+	RenderBatch.AddVertex( FSlateVertex::Make<Rounding>( RenderTransform, FVector2f( Position.X, BottomMarginY ),	LocalSize, DrawScale, FVector4f( StartUV.X, EndUV.Y, 0.0f, LeftTiling),			Tint ) ); //13
+	RenderBatch.AddVertex( FSlateVertex::Make<Rounding>( RenderTransform, FVector2f( LeftMarginX, TopMarginY ),		LocalSize, DrawScale, FVector4f( LeftMarginU, StartUV.Y, 0.0f, LeftTiling),		Tint ) ); //14
+	RenderBatch.AddVertex( FSlateVertex::Make<Rounding>( RenderTransform, FVector2f( LeftMarginX, BottomMarginY ),	LocalSize, DrawScale, FVector4f( LeftMarginU, EndUV.Y, 0.0f, LeftTiling),		Tint ) ); //15
+
+	RenderBatch.AddVertex( FSlateVertex::Make<Rounding>( RenderTransform, FVector2f( RightMarginX, TopMarginY ),	LocalSize, DrawScale, FVector4f( RightMarginU, StartUV.Y, 0.0f, RightTiling),	Tint ) ); //16
+	RenderBatch.AddVertex( FSlateVertex::Make<Rounding>( RenderTransform, FVector2f( RightMarginX, BottomMarginY ), LocalSize, DrawScale, FVector4f( RightMarginU, EndUV.Y, 0.0f, RightTiling),		Tint ) ); //17
+	RenderBatch.AddVertex( FSlateVertex::Make<Rounding>( RenderTransform, FVector2f( EndPos.X, TopMarginY ),		LocalSize, DrawScale, FVector4f( EndUV.X, StartUV.Y, 0.0f, RightTiling),			Tint ) ); //18
+	RenderBatch.AddVertex( FSlateVertex::Make<Rounding>( RenderTransform, FVector2f( EndPos.X, BottomMarginY ),		LocalSize, DrawScale, FVector4f( EndUV.X, EndUV.Y, 0.0f, RightTiling),			Tint ) ); //19
+
+	RenderBatch.AddVertex( FSlateVertex::Make<Rounding>( RenderTransform, FVector2f( Position.X, BottomMarginY ),	LocalSize, DrawScale, FVector4f( StartUV.X, BottomMarginV, 0.0f, 0.0f),			Tint ) ); //20
+	RenderBatch.AddVertex( FSlateVertex::Make<Rounding>( RenderTransform, FVector2f( Position.X, EndPos.Y ),		LocalSize, DrawScale, FVector4f( StartUV.X, EndUV.Y, 0.0f, 0.0f),				Tint ) ); //21
+	RenderBatch.AddVertex( FSlateVertex::Make<Rounding>( RenderTransform, FVector2f( LeftMarginX, BottomMarginY ),	LocalSize, DrawScale, FVector4f( LeftMarginU, BottomMarginV, 0.0f, 0.0f),		Tint ) ); //22
+	RenderBatch.AddVertex( FSlateVertex::Make<Rounding>( RenderTransform, FVector2f( LeftMarginX, EndPos.Y ),		LocalSize, DrawScale, FVector4f( LeftMarginU, EndUV.Y, 0.0f, 0.0f),				Tint ) ); //23
+
+	RenderBatch.AddVertex( FSlateVertex::Make<Rounding>( RenderTransform, FVector2f( LeftMarginX, BottomMarginY ),	LocalSize, DrawScale, FVector4f( StartUV.X, BottomMarginV, BottomTiling, 0.0f),	Tint ) ); //24
+	RenderBatch.AddVertex( FSlateVertex::Make<Rounding>( RenderTransform, FVector2f( LeftMarginX, EndPos.Y ),		LocalSize, DrawScale, FVector4f( StartUV.X, EndUV.Y, BottomTiling, 0.0f),		Tint ) ); //25
+	RenderBatch.AddVertex( FSlateVertex::Make<Rounding>( RenderTransform, FVector2f( RightMarginX,BottomMarginY ),	LocalSize, DrawScale, FVector4f( EndUV.X, BottomMarginV, BottomTiling, 0.0f),	Tint ) ); //26
+	RenderBatch.AddVertex( FSlateVertex::Make<Rounding>( RenderTransform, FVector2f( RightMarginX, EndPos.Y ),		LocalSize, DrawScale, FVector4f( EndUV.X, EndUV.Y, BottomTiling, 0.0f),			Tint ) ); //27
+
+	RenderBatch.AddVertex( FSlateVertex::Make<Rounding>( RenderTransform, FVector2f( RightMarginX, BottomMarginY ), LocalSize, DrawScale, FVector4f( RightMarginU, BottomMarginV, 0.0f, 0.0f),		Tint ) ); //29
+	RenderBatch.AddVertex( FSlateVertex::Make<Rounding>( RenderTransform, FVector2f( RightMarginX, EndPos.Y ),		LocalSize, DrawScale, FVector4f( RightMarginU, EndUV.Y, 0.0f, 0.0f),				Tint ) ); //30
+	RenderBatch.AddVertex( FSlateVertex::Make<Rounding>( RenderTransform, FVector2f( EndPos.X, BottomMarginY ),		LocalSize, DrawScale, FVector4f( EndUV.X, BottomMarginV, 0.0f, 0.0f),			Tint ) ); //31
+	RenderBatch.AddVertex( FSlateVertex::Make<Rounding>( RenderTransform, FVector2f( EndPos.X, EndPos.Y ),			LocalSize, DrawScale, FVector4f( EndUV.X, EndUV.Y, 0.0f, 0.0f),					Tint ) ); //32
 
 	// Top
 	RenderBatch.AddIndex( IndexStart + 0 );
@@ -2550,41 +2550,40 @@ void FSlateElementBatcher::AddCustomVerts(const FSlateDrawElement& DrawElement)
 	}*/
 }
 
-
-void FSlateElementBatcher::AddPostProcessPass(const FSlateDrawElement& DrawElement, const FVector2D& WindowSize)
+void FSlateElementBatcher::AddPostProcessPass(const FSlateDrawElement& DrawElement, FVector2f WindowSize)
 {
 	++NumPostProcessPasses;
 
 	const FSlateRenderTransform& RenderTransform = DrawElement.GetRenderTransform();
-	const FVector2D& LocalSize = DrawElement.GetLocalSize();
+	const FVector2f LocalSize = DrawElement.GetLocalSize();
 	
 	const FSlatePostProcessPayload& Payload = DrawElement.GetDataPayload<FSlatePostProcessPayload>();
 
 	//@todo doesn't work with rotated or skewed objects yet
-	const FVector2D& Position = DrawElement.GetPosition();
+	const FVector2f Position = DrawElement.GetPosition();
 
 	const int32 Layer = DrawElement.GetLayer();
 
 	// Determine the four corners of the quad
-	FVector2D TopLeft = FVector2D::ZeroVector;
-	FVector2D TopRight = FVector2D(LocalSize.X, 0);
-	FVector2D BotLeft = FVector2D(0, LocalSize.Y);
-	FVector2D BotRight = FVector2D(LocalSize.X, LocalSize.Y);
+	FVector2f TopLeft = FVector2f::ZeroVector;
+	FVector2f TopRight = FVector2f(LocalSize.X, 0.f);
+	FVector2f BotLeft = FVector2f(0.f, LocalSize.Y);
+	FVector2f BotRight = FVector2f(LocalSize.X, LocalSize.Y);
 
 
 	// Offset by half a texel if the platform requires it for pixel perfect sampling
-	//FVector2D HalfTexel = FVector2D(PixelCenterOffset / WindowSize.X, PixelCenterOffset / WindowSize.Y);
+	FVector2f HalfTexel = FVector2f(PixelCenterOffset / WindowSize.X, PixelCenterOffset / WindowSize.Y);
 
-	FVector2D WorldTopLeft = TransformPoint(RenderTransform, TopLeft).RoundToVector();
-	FVector2D WorldBotRight = TransformPoint(RenderTransform, BotRight).RoundToVector();
+	FVector2f WorldTopLeft = TransformPoint(RenderTransform, TopLeft).RoundToVector();
+	FVector2f WorldBotRight = TransformPoint(RenderTransform, BotRight).RoundToVector();
 
-	FVector2D SizeUV = (WorldBotRight - WorldTopLeft) / WindowSize;
+	FVector2f SizeUV = (WorldBotRight - WorldTopLeft) / WindowSize;
 
 	// These could be negative with rotation or negative scales.  This is not supported yet
 	if(SizeUV.X > 0 && SizeUV.Y > 0)
 	{
 		FShaderParams Params = FShaderParams::MakePixelShaderParams(
-			FVector4f(FVector2f(WorldTopLeft), FVector2f(WorldBotRight)),
+			FVector4f(WorldTopLeft, WorldBotRight),
 			FVector4f(Payload.PostProcessData.X, Payload.PostProcessData.Y, (float)Payload.DownsampleAmount, 0.f),
 			FVector4f(Payload.CornerRadius));
 
@@ -2835,10 +2834,10 @@ void FSlateElementBatcher::BuildShapedTextSequence(const FShapedTextBuildContext
 
 		if(bCanRenderGlyph && RenderBatch)
 		{
-			const FVector2D UpperLeft(X, Y);
-			const FVector2D UpperRight(X + SizeX, Y);
-			const FVector2D LowerLeft(X, Y + SizeY);
-			const FVector2D LowerRight(X + SizeX, Y + SizeY);
+			const FVector2f UpperLeft(X, Y);
+			const FVector2f UpperRight(X + SizeX, Y);
+			const FVector2f LowerLeft(X, Y + SizeY);
+			const FVector2f LowerRight(X + SizeX, Y + SizeY);
 
 			// The start index of these vertices in the index buffer
 			const uint32 IndexStart = RenderBatch->GetNumVertices();
diff --git a/Engine/Source/Runtime/SlateCore/Public/Rendering/DrawElementPayloads.h b/Engine/Source/Runtime/SlateCore/Public/Rendering/DrawElementPayloads.h
--- a/Engine/Source/Runtime/SlateCore/Public/Rendering/DrawElementPayloads.h
+++ b/Engine/Source/Runtime/SlateCore/Public/Rendering/DrawElementPayloads.h
@@ -3,10 +3,11 @@
 #pragma once
 
 #include "SlateRenderBatch.h"
+#include "Types/SlateVector2.h"
 
 struct FSlateGradientStop
 {
-	FVector2D Position;
+	FVector2f Position;
 	FLinearColor Color;
 
 	/**
@@ -15,7 +16,12 @@ struct FSlateGradientStop
 						  A two stop gradient should go from (0,0), to (Width,Height).
 	 * @param InColor	- The color to lerp towards at this stop.
 	 */
-	FSlateGradientStop(const FVector2D& InPosition, const FLinearColor& InColor)
+	FSlateGradientStop(const FVector2d InPosition, const FLinearColor InColor)
+		: Position(UE::Slate::CastToVector2f(InPosition))
+		, Color(InColor)
+	{
+	}
+	FSlateGradientStop(const FVector2f InPosition, const FLinearColor InColor)
 		: Position(InPosition)
 		, Color(InColor)
 	{
@@ -51,20 +57,25 @@ struct FSlateTintableElement
 struct FSlateBoxPayload : public FSlateDataPayload, public FSlateTintableElement
 {
 	FMargin Margin;
-	FBox2D UVRegion;
+	FBox2f UVRegion;
 	const FSlateShaderResourceProxy* ResourceProxy;
 	ESlateBrushTileType::Type Tiling;
 	ESlateBrushMirrorType::Type Mirroring;
 	ESlateBrushDrawType::Type DrawType;
 
 	const FMargin& GetBrushMargin() const { return Margin; }
-	const FBox2D& GetBrushUVRegion() const { return UVRegion; }
+	const FBox2f& GetBrushUVRegion() const { return UVRegion; }
 	ESlateBrushTileType::Type GetBrushTiling() const { return Tiling; }
 	ESlateBrushMirrorType::Type GetBrushMirroring() const { return Mirroring; }
 	ESlateBrushDrawType::Type GetBrushDrawType() const { return DrawType; }
 	const FSlateShaderResourceProxy* GetResourceProxy() const { return ResourceProxy; }
 
 	void SetBrush(const FSlateBrush* InBrush, FVector2D LocalSize, float DrawScale)
+	{
+		SetBrush(InBrush, UE::Slate::CastToVector2f(LocalSize), DrawScale);
+	}
+
+	void SetBrush(const FSlateBrush* InBrush, FVector2f LocalSize, float DrawScale)
 	{
 		check(InBrush);
 		ensureMsgf(InBrush->GetDrawType() != ESlateBrushDrawType::NoDrawType, TEXT("This should have been filtered out earlier in the Make... call."));
@@ -75,7 +86,7 @@ struct FSlateBoxPayload : public FSlateDataPayload, public FSlateTintableElement
 		Tiling = InBrush->GetTiling();
 		Mirroring = InBrush->GetMirroring();
 		DrawType = InBrush->GetDrawType();
-		FSlateResourceHandle Handle = InBrush->GetRenderingResource(LocalSize, DrawScale);
+		FSlateResourceHandle Handle = InBrush->GetRenderingResource(FVector2D(LocalSize), DrawScale);
 		if (Handle.IsValid())
 		{
 			ResourceProxy = Handle.GetResourceProxy();
@@ -208,9 +219,9 @@ struct FSlateGradientPayload : public FSlateDataPayload
 	EOrientation GradientType;
 	FVector4f CornerRadius;
 
-	void SetGradient(const TArray<FSlateGradientStop>& InGradientStops, EOrientation InGradientType, FVector4f InCornerRadius)
+	void SetGradient(TArray<FSlateGradientStop> InGradientStops, EOrientation InGradientType, FVector4f InCornerRadius)
 	{
-		GradientStops = InGradientStops;
+		GradientStops = MoveTemp(InGradientStops);
 		GradientType = InGradientType;
 		CornerRadius = InCornerRadius;
 	}
@@ -226,10 +237,10 @@ struct FSlateSplinePayload : public FSlateDataPayload, public FSlateTintableElem
 	//     P0 *             * P3            P0 *   \   * P3
 	//                                              \ /
 	//                                               + P2	
-	FVector2D P0;
-	FVector2D P1;
-	FVector2D P2;
-	FVector2D P3;
+	FVector2f P0;
+	FVector2f P1;
+	FVector2f P2;
+	FVector2f P3;
 
 	float Thickness;
 
@@ -237,7 +248,18 @@ struct FSlateSplinePayload : public FSlateDataPayload, public FSlateTintableElem
 	void SetThickness(float InThickness) { Thickness = InThickness; }
 	float GetThickness() const { return Thickness; }
 
-	void SetCubicBezier(const FVector2D& InP0, const FVector2D& InP1, const FVector2D& InP2, const FVector2D& InP3, float InThickness, const FLinearColor& InTint)
+	void SetCubicBezier(const FVector2d InP0, const FVector2d InP1, const FVector2d InP2, const FVector2d InP3, float InThickness, const FLinearColor InTint)
+	{
+		SetCubicBezier(
+			UE::Slate::CastToVector2f(InP0),
+			UE::Slate::CastToVector2f(InP1),
+			UE::Slate::CastToVector2f(InP2),
+			UE::Slate::CastToVector2f(InP3),
+			InThickness,
+			InTint);
+	}
+
+	void SetCubicBezier(const FVector2f InP0, const FVector2f InP1, const FVector2f InP2, const FVector2f InP3, float InThickness, const FLinearColor InTint)
 	{
 		Tint = InTint;
 		P0 = InP0;
@@ -247,7 +269,18 @@ struct FSlateSplinePayload : public FSlateDataPayload, public FSlateTintableElem
 		Thickness = InThickness;
 	}
 
-	void SetHermiteSpline(const FVector2D& InStart, const FVector2D& InStartDir, const FVector2D& InEnd, const FVector2D& InEndDir, float InThickness, const FLinearColor& InTint)
+	void SetHermiteSpline(const FVector2d InStart, const FVector2d InStartDir, const FVector2d InEnd, const FVector2d InEndDir, float InThickness, const FLinearColor InTint)
+	{
+		SetHermiteSpline(
+			UE::Slate::CastToVector2f(InStart),
+			UE::Slate::CastToVector2f(InStartDir),
+			UE::Slate::CastToVector2f(InEnd),
+			UE::Slate::CastToVector2f(InEndDir),
+			InThickness,
+			InTint);
+	}
+
+	void SetHermiteSpline(const FVector2f InStart, const FVector2f InStartDir, const FVector2f InEnd, const FVector2f InEndDir, float InThickness, const FLinearColor InTint)
 	{
 		Tint = InTint;
 		P0 = InStart;
@@ -257,28 +290,39 @@ struct FSlateSplinePayload : public FSlateDataPayload, public FSlateTintableElem
 		Thickness = InThickness;
 	}
 
-	void SetGradientHermiteSpline(const FVector2D& InStart, const FVector2D& InStartDir, const FVector2D& InEnd, const FVector2D& InEndDir, float InThickness, const TArray<FSlateGradientStop>& InGradientStops)
+	void SetGradientHermiteSpline(const FVector2d InStart, const FVector2d InStartDir, const FVector2d InEnd, const FVector2d InEndDir, float InThickness, TArray<FSlateGradientStop> InGradientStops)
+	{
+		SetGradientHermiteSpline(
+			UE::Slate::CastToVector2f(InStart),
+			UE::Slate::CastToVector2f(InStartDir),
+			UE::Slate::CastToVector2f(InEnd),
+			UE::Slate::CastToVector2f(InEndDir),
+			InThickness,
+			MoveTemp(InGradientStops));
+	}
+
+	void SetGradientHermiteSpline(const FVector2f InStart, const FVector2f InStartDir, const FVector2f InEnd, const FVector2f InEndDir, float InThickness, TArray<FSlateGradientStop> InGradientStops)
 	{
 		P0 = InStart;
 		P1 = InStart + InStartDir / 3.0f;
 		P2 = InEnd - InEndDir / 3.0f;
 		P3 = InEnd;
 		Thickness = InThickness;
-		GradientStops = InGradientStops;
+		GradientStops = MoveTemp(InGradientStops);
 	}
 };
 
 
 struct FSlateLinePayload : public FSlateDataPayload, public FSlateTintableElement
 { 
-	TArray<FVector2D> Points;
+	TArray<FVector2f> Points;
 	TArray<FLinearColor> PointColors;
 	float Thickness;
 
 	bool bAntialias;
 
 	bool IsAntialiased() const { return bAntialias; }
-	const TArray<FVector2D>& GetPoints() const { return Points; }
+	const TArray<FVector2f>& GetPoints() const { return Points; }
 	const TArray<FLinearColor>& GetPointColors() const { return PointColors; }
 	float GetThickness() const { return Thickness; }
 
@@ -288,9 +332,21 @@ struct FSlateLinePayload : public FSlateDataPayload, public FSlateTintableElemen
 	}
 
 	void SetLines(const TArray<FVector2D>& InPoints, bool bInAntialias, const TArray<FLinearColor>* InPointColors = nullptr)
+	{
+		TArray<FVector2f> NewPoints;
+		NewPoints.Reserve(InPoints.Num());
+		for (FVector2D Vect : InPoints)
+		{
+			NewPoints.Add(UE::Slate::CastToVector2f(Vect));
+		}
+		SetLines(MoveTemp(NewPoints), bInAntialias, InPointColors);
+	}
+
+
+	void SetLines(TArray<FVector2f> InPoints, bool bInAntialias, const TArray<FLinearColor>* InPointColors = nullptr)
 	{
 		bAntialias = bInAntialias;
-		Points = InPoints;
+		Points = MoveTemp(InPoints);
 		if (InPointColors)
 		{
 			PointColors = *InPointColors;
@@ -342,10 +398,15 @@ struct FSlateCachedBufferPayload : public FSlateDataPayload
 {
 	// Cached render data
 	class FSlateRenderDataHandle* CachedRenderData;
-	FVector2D CachedRenderDataOffset;
+	FVector2f CachedRenderDataOffset;
 
 	// Cached Buffers
-	void SetCachedBuffer(FSlateRenderDataHandle* InRenderDataHandle, const FVector2D& Offset)
+	void SetCachedBuffer(FSlateRenderDataHandle* InRenderDataHandle, const FVector2d Offset)
+	{
+		SetCachedBuffer(InRenderDataHandle, UE::Slate::CastToVector2f(Offset));
+	}
+
+	void SetCachedBuffer(FSlateRenderDataHandle* InRenderDataHandle, const FVector2f Offset)
 	{
 		check(InRenderDataHandle);
 
@@ -366,12 +427,12 @@ struct FSlateCustomVertsPayload : public FSlateDataPayload
 	uint32 InstanceOffset;
 	uint32 NumInstances;
 
-	void SetCustomVerts(const FSlateShaderResourceProxy* InRenderProxy, const TArray<FSlateVertex>& InVerts, const TArray<SlateIndex>& InIndices, ISlateUpdatableInstanceBufferRenderProxy* InInstanceData, uint32 InInstanceOffset, uint32 InNumInstances)
+	void SetCustomVerts(const FSlateShaderResourceProxy* InRenderProxy, TArray<FSlateVertex> InVerts, TArray<SlateIndex> InIndices, ISlateUpdatableInstanceBufferRenderProxy* InInstanceData, uint32 InInstanceOffset, uint32 InNumInstances)
 	{
 		ResourceProxy = InRenderProxy;
 
-		Vertices = InVerts;
-		Indices = InIndices;
+		Vertices = MoveTemp(InVerts);
+		Indices = MoveTemp(InIndices);
 
 		InstanceData = InInstanceData;
 		InstanceOffset = InInstanceOffset;
diff --git a/Engine/Source/Runtime/SlateCore/Public/Rendering/DrawElements.h b/Engine/Source/Runtime/SlateCore/Public/Rendering/DrawElements.h
--- a/Engine/Source/Runtime/SlateCore/Public/Rendering/DrawElements.h
+++ b/Engine/Source/Runtime/SlateCore/Public/Rendering/DrawElements.h
@@ -13,6 +13,7 @@
 #include "Layout/SlateRect.h"
 #include "Layout/Clipping.h"
 #include "Types/PaintArgs.h"
+#include "Types/SlateVector2.h"
 #include "Layout/Geometry.h"
 #include "Rendering/ShaderResourceManager.h"
 #include "Rendering/RenderingCommon.h"
@@ -269,9 +270,27 @@ public:
 
 	FORCEINLINE const FSlateRenderTransform& GetRenderTransform() const { return RenderTransform; }
 	FORCEINLINE void SetRenderTransform(const FSlateRenderTransform& InRenderTransform) { RenderTransform = InRenderTransform; }
-	FORCEINLINE FVector2D GetPosition() const { return FVector2D(Position); }
-	FORCEINLINE void SetPosition(const FVector2D& InPosition) { Position = FVector2f(InPosition); }
-	FORCEINLINE FVector2D GetLocalSize() const { return FVector2D(LocalSize); }
+
+	FORCEINLINE UE::Slate::FDeprecateVector2D GetPosition() const
+	{
+		return Position;
+	}
+
+	UE_DEPRECATED(5.1, "Slate rendering now use float instead of double. Use SetPosition(FVector2f)")
+	FORCEINLINE void SetPosition(FVector2d InPosition)
+	{
+		Position = UE::Slate::CastToVector2f(InPosition);
+	}
+	FORCEINLINE void SetPosition(FVector2f InPosition)
+	{
+		Position = InPosition;
+	}
+
+	FORCEINLINE UE::Slate::FDeprecateVector2D GetLocalSize() const
+	{
+		return LocalSize;
+	}
+
 	FORCEINLINE float GetScale() const { return Scale; }
 	FORCEINLINE ESlateDrawEffect GetDrawEffects() const { return DrawEffects; }
 	FORCEINLINE ESlateBatchDrawFlag GetBatchFlags() const { return BatchFlags; }
@@ -289,18 +308,21 @@ public:
 
 	FORCEINLINE FSlateLayoutTransform GetInverseLayoutTransform() const
 	{
-		return Inverse(FSlateLayoutTransform(Scale, FVector2D(Position)));
+		return Inverse(FSlateLayoutTransform(Scale, Position));
 	}
 
 	void AddReferencedObjects(FReferenceCollector& Collector);
 
+	UE_DEPRECATED(5.1, "Slate rendering now use float instead of double. Use ApplyPositionOffset(FVector2f)")
+	void ApplyPositionOffset(FVector2d InOffset);
+
 	/**
-	* Update element cached position with an arbitrary offset
-	*
-	* @param Element		   Element to update
-	* @param InOffset         Absolute translation delta
-	*/
-	void ApplyPositionOffset(const FVector2D& InOffset);
+	 * Update element cached position with an arbitrary offset
+	 *
+	 * @param Element		   Element to update
+	 * @param InOffset         Absolute translation delta
+	 */
+	void ApplyPositionOffset(FVector2f InOffset);
 
 	UE_DEPRECATED(4.23, "GetClippingIndex has been deprecated.  If you were using this please use GetPrecachedClippingIndex instead.")
 	FORCEINLINE const int32 GetClippingIndex() const { return GetPrecachedClippingIndex(); }
@@ -509,7 +531,7 @@ public:
 	}
 
 	/** @return Get the window size that we will be painting */
-	FVector2D GetWindowSize() const
+	UE::Slate::FDeprecateVector2D GetWindowSize() const
 	{
 		return WindowSize;
 	}
@@ -699,5 +721,5 @@ private:
 	// End Fast Path
 
 	/** Store the size of the window being used to paint */
-	FVector2D WindowSize;
+	FVector2f WindowSize;
 };
diff --git a/Engine/Source/Runtime/SlateCore/Public/Rendering/ElementBatcher.h b/Engine/Source/Runtime/SlateCore/Public/Rendering/ElementBatcher.h
--- a/Engine/Source/Runtime/SlateCore/Public/Rendering/ElementBatcher.h
+++ b/Engine/Source/Runtime/SlateCore/Public/Rendering/ElementBatcher.h
@@ -275,8 +275,8 @@ public:
 	}
 
 private:
-	void AddElementsInternal(const FSlateDrawElementArray& DrawElements, const FVector2D& ViewportSize);
-	void AddCachedElements(FSlateCachedElementData& CachedElementData, const FVector2D& ViewportSize);
+	void AddElementsInternal(const FSlateDrawElementArray& DrawElements, FVector2f ViewportSize);
+	void AddCachedElements(FSlateCachedElementData& CachedElementData, FVector2f ViewportSize);
 
 	/** 
 	 * Creates vertices necessary to draw a Quad element 
@@ -335,7 +335,7 @@ private:
 
 	void AddCustomVerts( const FSlateDrawElement& DrawElement );
 
-	void AddPostProcessPass(const FSlateDrawElement& DrawElement, const FVector2D& WindowSize);
+	void AddPostProcessPass(const FSlateDrawElement& DrawElement, FVector2f WindowSize);
 
 	FSlateRenderBatch& CreateRenderBatch(
 		int32 Layer,
diff --git a/Engine/Source/Runtime/SlateCore/Public/Styling/SlateBrush.h b/Engine/Source/Runtime/SlateCore/Public/Styling/SlateBrush.h
--- a/Engine/Source/Runtime/SlateCore/Public/Styling/SlateBrush.h
+++ b/Engine/Source/Runtime/SlateCore/Public/Styling/SlateBrush.h
@@ -7,6 +7,8 @@
 #include "Styling/SlateColor.h"
 #include "Layout/Margin.h"
 #include "Textures/SlateShaderResource.h"
+#include "Types/SlateBox2.h"
+#include "Types/SlateVector2.h"
 #include "SlateBrush.generated.h"
 
 /**
@@ -366,9 +368,9 @@ public:
 	 *
 	 * @return UV region
 	 */
-	FBox2D GetUVRegion() const
+	UE::Slate::FDeprecateBox2D GetUVRegion() const
 	{
-		return FBox2D(UVRegion);
+		return UVRegion;
 	}
 
 	/**
@@ -376,11 +378,16 @@ public:
 	 *
 	 * @param InUVRegion When valid - overrides UV region specified in resource proxy
 	 */
-	void SetUVRegion(const FBox2D& InUVRegion)
+	void SetUVRegion(const FBox2d& InUVRegion)
 	{
 		UVRegion = FBox2f(InUVRegion);
 	}
 
+	void SetUVRegion(const FBox2f& InUVRegion)
+	{
+		UVRegion = InUVRegion;
+	}
+
 	/**
 	 * Compares this brush with another for equality.
 	 *
diff --git a/Engine/Source/Runtime/SlateCore/Public/Types/SlateBox2.h b/Engine/Source/Runtime/SlateCore/Public/Types/SlateBox2.h
new file mode 100644
--- /dev/null
+++ b/Engine/Source/Runtime/SlateCore/Public/Types/SlateBox2.h
@@ -0,0 +1,62 @@
+// Copyright Epic Games, Inc. All Rights Reserved.
+
+#pragma once
+
+#include "Types/SlateVector2.h"
+
+namespace UE::Slate
+{
+	
+FORCEINLINE static FBox2f CastToBox2f(const FBox2d& InValue)
+{
+	return FBox2f(InValue);
+}
+
+/**
+ * Structure for deprecating FBox2D to FBox2f
+ */
+struct FDeprecateBox2D
+{
+	FDeprecateBox2D() = default;
+	FDeprecateBox2D(const FBox2f& InValue)
+		: Data(InValue)
+	{
+	}
+
+public:
+	operator FBox2d() const
+	{
+		return FBox2d(Data);
+	}
+
+	operator FBox2f() const
+	{
+		return Data;
+	}
+	
+public:
+	bool operator==(const FBox2d& Other) const
+	{
+		return CastToBox2f(Other) == Data;
+	}
+	
+	bool operator==(const FBox2f& Other) const
+	{
+		return Data == Other;
+	}
+
+	bool operator!=(const FBox2d& Other) const
+	{
+		return CastToBox2f(Other) != Data;
+	}
+	
+	bool operator!=(const FBox2f& Other) const
+	{
+		return Data != Other;
+	}
+
+private:
+	FBox2f Data;
+};
+
+} // namespace
diff --git a/Engine/Source/Runtime/SlateCore/Public/Types/SlateVector2.h b/Engine/Source/Runtime/SlateCore/Public/Types/SlateVector2.h
new file mode 100644
--- /dev/null
+++ b/Engine/Source/Runtime/SlateCore/Public/Types/SlateVector2.h
@@ -0,0 +1,110 @@
+// Copyright Epic Games, Inc. All Rights Reserved.
+
+#pragma once
+
+#include "CoreMinimal.h"
+#include "Math/UnrealMathUtility.h"
+
+namespace UE::Slate
+{
+
+FORCEINLINE static FVector2f CastToVector2f(FVector2d InValue)
+{
+	const float X = UE_REAL_TO_FLOAT(InValue.X);
+	const float Y = UE_REAL_TO_FLOAT(InValue.Y);
+#if 0
+	ensureAlways(FMath::IsNearlyEqual((double)X, InValue.X));
+	ensureAlways(FMath::IsNearlyEqual((double)Y, InValue.Y));
+#endif
+	return FVector2f(X, Y);
+}
+
+/**
+ * Structure for deprecating FVector2D to FVector2f
+ */
+struct FDeprecateVector2D
+{
+	FDeprecateVector2D() = default;
+	FDeprecateVector2D(FVector2f InValue)
+		: Data(InValue)
+	{
+	}
+
+public:
+	operator FVector2d() const
+	{
+		return FVector2d(Data);
+	}
+
+	operator FVector2f() const
+	{
+		return Data;
+	}
+
+public:
+	FORCEINLINE FVector2f operator+(const FVector2d V) const
+	{
+		return Data + CastToVector2f(V);
+	}
+
+	FORCEINLINE FVector2f operator+(const FVector2f V) const
+	{
+		return Data + V;
+	}
+
+	FORCEINLINE FVector2f operator-(const FVector2d V) const
+	{
+		return Data - CastToVector2f(V);
+	}
+
+	FORCEINLINE FVector2f operator-(const FVector2f V) const
+	{
+		return Data - V;
+	}
+
+	FORCEINLINE FVector2f operator+(float Scale) const
+	{
+		return Data + Scale;
+	}
+
+	FORCEINLINE FVector2f operator-(float Scale) const
+	{
+		return Data - Scale;
+	}
+
+	FORCEINLINE FVector2f operator*(float Scale) const
+	{
+		return Data * Scale;
+	}
+
+	FORCEINLINE FVector2f operator/(float Scale) const
+	{
+		return Data / Scale;
+	}
+
+public:
+	bool operator==(FVector2d V) const
+	{
+		return CastToVector2f(V) == Data;
+	}
+
+	bool operator==(FVector2f V) const
+	{
+		return V == Data;
+	}
+
+	bool operator!=(FVector2d V) const
+	{
+		return CastToVector2f(V) != Data;
+	}
+
+	bool operator!=(FVector2f V) const
+	{
+		return V != Data;
+	}
+
+private:
+	FVector2f Data;
+};
+
+} // namespace
