DrawElements
Slate: Transform DrawElements from Vector2D to Vector2f. Use FDeprecateVector2D to create a deprecation path for function that returns a FVector2D.

diff --git a/Engine/Source/Runtime/AppFramework/Private/Widgets/Testing/STestSuite.cpp b/Engine/Source/Runtime/AppFramework/Private/Widgets/Testing/STestSuite.cpp
--- a/Engine/Source/Runtime/AppFramework/Private/Widgets/Testing/STestSuite.cpp
+++ b/Engine/Source/Runtime/AppFramework/Private/Widgets/Testing/STestSuite.cpp
@@ -466,16 +466,16 @@ class SSplineWithHandles : public SCompoundWidget
 {
 public:
 	SLATE_BEGIN_ARGS(SSplineWithHandles)
-	: _P0(FVector2D(0,32))
-	, _P1(FVector2D(100, 32))
-	, _P2(FVector2D(0,132))
-	, _P3(FVector2D(100,132))
+	: _P0(FVector2f(0.f,32.f))
+	, _P1(FVector2f(100.f, 32.f))
+	, _P2(FVector2f(0.f,132.f))
+	, _P3(FVector2f(100.f,132.f))
 	, _SplineThickness(1.0f)
 	{}
-		SLATE_ARGUMENT(FVector2D, P0)
-		SLATE_ARGUMENT(FVector2D, P1)
-		SLATE_ARGUMENT(FVector2D, P2)
-		SLATE_ARGUMENT(FVector2D, P3)
+		SLATE_ARGUMENT(FVector2f, P0)
+		SLATE_ARGUMENT(FVector2f, P1)
+		SLATE_ARGUMENT(FVector2f, P2)
+		SLATE_ARGUMENT(FVector2f, P3)
 		SLATE_ATTRIBUTE(float, SplineThickness)
 	SLATE_END_ARGS()
 
@@ -485,7 +485,7 @@ public:
 		SplineThickness = InArgs._SplineThickness;
 	}
 
-	void SetPoints(FVector2D P0, FVector2D P1, FVector2D P2, FVector2D P3)
+	void SetPoints(FVector2f P0, FVector2f P1, FVector2f P2, FVector2f P3)
 	{
 		BezierPoints[0] = P0;
 		BezierPoints[1] = P1;
@@ -499,7 +499,7 @@ public:
 	}
 
 private:
-	const FVector2D BezierPointRadius = FVector2D(8,8);
+	const FVector2f BezierPointRadius = FVector2f(8.f,8.f);
 
 	TSharedPtr<SSpinBox<float>> ThicknessSpinBox;
 	
@@ -509,18 +509,18 @@ private:
 		const bool bEnabled = ShouldBeEnabled(bParentEnabled);
 		
 		// Draw control points.
-		for (int i = 0; i < 4; ++i)
+		for (int32 i = 0; i < 4; ++i)
 		{
-			const FVector2D RadiusToUse = (i == 0 || i == 3) ? BezierPointRadius : 0.5f*BezierPointRadius;
+			const FVector2f RadiusToUse = (i == 0 || i == 3) ? BezierPointRadius : 0.5f*BezierPointRadius;
 			const FLinearColor ColorToUse = (i == 0 || i == 1) ? FLinearColor::Green : FLinearColor::Blue;
 
 			FSlateDrawElement::MakeBox(
 				OutDrawElements,
 				LayerId,
-				AllottedGeometry.ToPaintGeometry(2* RadiusToUse, FSlateLayoutTransform(BezierPoints[i]- RadiusToUse)),
-				&WhiteBox,				
+				AllottedGeometry.ToPaintGeometry(FVector2D(2.f* RadiusToUse), FSlateLayoutTransform(BezierPoints[i]- RadiusToUse)),
+				&WhiteBox,
 				bEnabled ? ESlateDrawEffect::None : ESlateDrawEffect::DisabledEffect,
-				ColorToUse * FLinearColor(1, 1, 1, 0.25f)
+				ColorToUse * FLinearColor(1.f, 1.f, 1.f, 0.25f)
 			);
 
 		}
@@ -584,7 +584,7 @@ private:
 	{
 		if (MouseEvent.GetEffectingButton()==EKeys::LeftMouseButton)
 		{
-			const FVector2D LocalCursorPos = MyGeometry.AbsoluteToLocal(MouseEvent.GetScreenSpacePosition());
+			const FVector2f LocalCursorPos = UE::Slate::CastToVector2f(MyGeometry.AbsoluteToLocal(MouseEvent.GetScreenSpacePosition()));
 			PointBeingDragged = PointIndexFromCursorPos(BezierPoints, LocalCursorPos, 2*BezierPointRadius.X*2*BezierPointRadius.X);
 
 			if (PointBeingDragged != INDEX_NONE)
@@ -613,7 +613,7 @@ private:
 	{
 		if (PointBeingDragged != INDEX_NONE)
 		{
-			const FVector2D LocalCursorPos = MyGeometry.AbsoluteToLocal(MouseEvent.GetScreenSpacePosition());
+			const FVector2f LocalCursorPos = UE::Slate::CastToVector2f(MyGeometry.AbsoluteToLocal(MouseEvent.GetScreenSpacePosition()));
 			BezierPoints[PointBeingDragged] = LocalCursorPos;
 			return FReply::Handled();
 		}
@@ -623,7 +623,7 @@ private:
 	//~ SWidget
 
 
-	static int32 PointIndexFromCursorPos(const FVector2D* Points, const FVector2D LocalCursorPos, const float RadiusSquared)
+	static int32 PointIndexFromCursorPos(const FVector2f* Points, const FVector2f LocalCursorPos, const float RadiusSquared)
 	{
 		for (int i = 0; i < 4; ++i)
 		{			
@@ -637,7 +637,7 @@ private:
 	}
 
 	FSlateColorBrush WhiteBox = FSlateColorBrush(FColor::White);
-	FVector2D BezierPoints[4];
+	FVector2f BezierPoints[4];
 	int32 PointBeingDragged = INDEX_NONE;
 	TAttribute<float> SplineThickness;
 	TArray<FSlateGradientStop> GradientStops;
@@ -647,15 +647,15 @@ class SSplineTest : public SCompoundWidget
 {
 public:
 	SLATE_BEGIN_ARGS(SSplineTest)
-		: _P0(FVector2D(0, 32))
-		, _P1(FVector2D(100, 32))
-		, _P2(FVector2D(0, 132))
-		, _P3(FVector2D(100, 132))
+		: _P0(FVector2f(0.f, 32.f))
+		, _P1(FVector2f(100.f, 32.f))
+		, _P2(FVector2f(0.f, 132.f))
+		, _P3(FVector2f(100.f, 132.f))
 	{}
-		SLATE_ARGUMENT(FVector2D, P0)
-		SLATE_ARGUMENT(FVector2D, P1)
-		SLATE_ARGUMENT(FVector2D, P2)
-		SLATE_ARGUMENT(FVector2D, P3)
+		SLATE_ARGUMENT(FVector2f, P0)
+		SLATE_ARGUMENT(FVector2f, P1)
+		SLATE_ARGUMENT(FVector2f, P2)
+		SLATE_ARGUMENT(FVector2f, P3)
 	SLATE_END_ARGS()
 
 	void Construct(const FArguments& InArgs)
@@ -679,7 +679,7 @@ public:
 				+ SHorizontalBox::Slot()
 				.FillWidth(1)
 				[
-					SNew(SSpinBox<int>)
+					SNew(SSpinBox<int32>)
 					.MinValue(0)
 					.MaxValue(16)
 					.OnValueChanged(this, &SSplineTest::OnNumGradientStopsChanged)
@@ -690,7 +690,7 @@ public:
 					.Text(LOCTEXT("HorizontalHermite", "Horizontal"))
 					.OnClicked_Lambda([this]()
 					{
-						SplineWidget->SetPoints(FVector2D(0, 64), FVector2D(66, 64), FVector2D(133, 64), FVector2D(200, 64));
+						SplineWidget->SetPoints(FVector2f(0.f, 64.f), FVector2f(66.f, 64.f), FVector2f(133.f, 64.f), FVector2f(200.f, 64.f));
 						return FReply::Handled();
 					})
 				]
@@ -700,7 +700,7 @@ public:
 					.Text(LOCTEXT("VerticalHermite", "Vertical"))
 					.OnClicked_Lambda([this]()
 					{
-						SplineWidget->SetPoints(FVector2D(64,0), FVector2D(64, 66), FVector2D(64, 133), FVector2D(64, 200));
+						SplineWidget->SetPoints(FVector2f(64.f,0.f), FVector2f(64.f, 66.f), FVector2f(64.f, 133.f), FVector2f(64.f, 200.f));
 						return FReply::Handled();
 					})
 				]
@@ -710,7 +710,7 @@ public:
 					.Text(LOCTEXT("HorizontalBendy", "Bendy"))
 					.OnClicked_Lambda([this]()
 					{
-						SplineWidget->SetPoints(FVector2D(64, 64), FVector2D(128, 64), FVector2D(64, 128), FVector2D(128, 128));
+						SplineWidget->SetPoints(FVector2f(64.f, 64.f), FVector2f(128.f, 64.f), FVector2f(64.f, 128.f), FVector2f(128.f, 128.f));
 						return FReply::Handled();
 					})
 				]
@@ -720,7 +720,7 @@ public:
 					.Text(LOCTEXT("HorizontalDoubleBendy", "DoubleBendy"))
 					.OnClicked_Lambda([this]()
 					{
-						SplineWidget->SetPoints(FVector2D(128, 128), FVector2D(128-96, 128), FVector2D(192+96, 192), FVector2D(192, 192));
+						SplineWidget->SetPoints(FVector2f(128.f, 128.f), FVector2f(128.f -96.f, 128.f), FVector2f(192.f +96.f, 192.f), FVector2f(192.f, 192.f));
 						return FReply::Handled();
 					})
 				]
@@ -742,7 +742,7 @@ public:
 private:
 	float GetSplineThickness() const { return SplineThickness; }
 	void OnSplineThicknessChanged(float InNewThickness) { SplineThickness = InNewThickness; }
-	void OnNumGradientStopsChanged(int NewNumGradientStops)
+	void OnNumGradientStopsChanged(int32 NewNumGradientStops)
 	{
 		const auto ColorOne = FLinearColor::MakeRandomColor();
 		const auto ColorTwo = FLinearColor::MakeRandomColor();
@@ -959,10 +959,10 @@ private:
 
 	int32 TestSplineElement( const FOnPaintHandlerParams& InParams )
 	{
-		const FVector2D Start(10,10);
-		const FVector2D StartDir(InParams.Geometry.GetLocalSize().X * 1000 / 600,0);
-		const FVector2D End(InParams.Geometry.GetLocalSize().X/4, InParams.Geometry.GetLocalSize().Y-10);
-		const FVector2D EndDir(InParams.Geometry.GetLocalSize().X * 1000 / 600,0);
+		const FVector2f Start(10,10);
+		const FVector2f StartDir(InParams.Geometry.GetLocalSize().X * 1000 / 600,0);
+		const FVector2f End(InParams.Geometry.GetLocalSize().X/4, InParams.Geometry.GetLocalSize().Y-10);
+		const FVector2f EndDir(InParams.Geometry.GetLocalSize().X * 1000 / 600,0);
 
 		FSlateDrawElement::MakeSpline(
 			InParams.OutDrawElements,
@@ -975,15 +975,15 @@ private:
 			FColor::White
 		);
 	
-		FVector2D LineStart =  FVector2D( InParams.Geometry.GetLocalSize().X/4, 10.0f );
+		FVector2f LineStart =  FVector2f( InParams.Geometry.GetLocalSize().X/4, 10.0f );
 
-		TArray<FVector2D> LinePoints;
+		TArray<FVector2f> LinePoints;
 		TArray<FLinearColor> LineColors;
 		LinePoints.Add(LineStart); LineColors.Add(FLinearColor::Red);
-		LinePoints.Add( LineStart + FVector2D( 100.0f, 50.0f ) );
-		LinePoints.Add( LineStart + FVector2D( 200.0f, 10.0f ) );
-		LinePoints.Add( LineStart + FVector2D( 300.0f, 50.0f ) );
-		LinePoints.Add( LineStart + FVector2D( 400.0f, 10.0f ) );
+		LinePoints.Add( LineStart + FVector2f( 100.0f, 50.0f ) );
+		LinePoints.Add( LineStart + FVector2f( 200.0f, 10.0f ) );
+		LinePoints.Add( LineStart + FVector2f( 300.0f, 50.0f ) );
+		LinePoints.Add( LineStart + FVector2f( 400.0f, 10.0f ) );
 
 
 		FSlateDrawElement::MakeLines( 
@@ -993,45 +993,45 @@ private:
 			LinePoints,
 			InParams.bEnabled ? ESlateDrawEffect::None : ESlateDrawEffect::DisabledEffect,
 			FColor::White,
-			true,	
+			true,
 			InParams.Geometry.Scale
 		);
 
 		FSlateDrawElement::MakeLines(
 			InParams.OutDrawElements,
 			InParams.Layer,
-			InParams.Geometry.ToOffsetPaintGeometry(FVector2D(0, 10)),
+			InParams.Geometry.ToOffsetPaintGeometry(FVector2D(0.f, 10.f)),
 			LinePoints,
 			InParams.bEnabled ? ESlateDrawEffect::None : ESlateDrawEffect::DisabledEffect,
 			FColor::White,
 			true,
-			2*InParams.Geometry.Scale
+			2.f * InParams.Geometry.Scale
 		);
 
 		FSlateDrawElement::MakeLines(
 			InParams.OutDrawElements,
 			InParams.Layer,
-			InParams.Geometry.ToOffsetPaintGeometry(FVector2D(0, 20)),
+			InParams.Geometry.ToOffsetPaintGeometry(FVector2D(0.f, 20.f)),
 			LinePoints,
 			InParams.bEnabled ? ESlateDrawEffect::None : ESlateDrawEffect::DisabledEffect,
 			FColor::White,
 			true,
-			3*InParams.Geometry.Scale
+			3.f * InParams.Geometry.Scale
 		);
 
 		FSlateDrawElement::MakeLines(
 			InParams.OutDrawElements,
 			InParams.Layer,
-			InParams.Geometry.ToOffsetPaintGeometry(FVector2D(0, 34)),
+			InParams.Geometry.ToOffsetPaintGeometry(FVector2D(0.f, 34.f)),
 			LinePoints,
 			InParams.bEnabled ? ESlateDrawEffect::None : ESlateDrawEffect::DisabledEffect,
 			FColor::White,
 			true,
-			4 * InParams.Geometry.Scale
+			4.f * InParams.Geometry.Scale
 		);
 			
 
-		LineStart =  LinePoints[ LinePoints.Num() - 1 ] + FVector2D(50,10);
+		LineStart =  LinePoints[ LinePoints.Num() - 1 ] + FVector2f(50.f,10.f);
 		LinePoints.Empty();
 
 		static float CurTime = 0; 
@@ -1040,7 +1040,7 @@ private:
 
 		for( float I = 0; I < 10*PI; I+=.1f)
 		{
-			LinePoints.Add( LineStart + FVector2D( I*15 , 15*FMath::Sin( I + CurTime) ) );
+			LinePoints.Add( LineStart + FVector2f( I*15 , 15*FMath::Sin( I + CurTime) ) );
 		}
 
 		static FColor Color = FColor::MakeRandomColor();
@@ -1057,32 +1057,32 @@ private:
 		FSlateDrawElement::MakeLines(
 			InParams.OutDrawElements,
 			InParams.Layer,
-			InParams.Geometry.ToOffsetPaintGeometry(FVector2D(0, 10)),
+			InParams.Geometry.ToOffsetPaintGeometry(FVector2D(0.f, 10.f)),
 			LinePoints,
 			InParams.bEnabled ? ESlateDrawEffect::NoPixelSnapping : ESlateDrawEffect::NoPixelSnapping | ESlateDrawEffect::DisabledEffect,
 			Color,
 			true,
-			2*InParams.Geometry.Scale);
+			2.f * InParams.Geometry.Scale);
 
 		FSlateDrawElement::MakeLines(
 			InParams.OutDrawElements,
 			InParams.Layer,
-			InParams.Geometry.ToOffsetPaintGeometry(FVector2D(0, 20)),
+			InParams.Geometry.ToOffsetPaintGeometry(FVector2D(0.f, 20.f)),
 			LinePoints,
 			InParams.bEnabled ? ESlateDrawEffect::NoPixelSnapping : ESlateDrawEffect::NoPixelSnapping | ESlateDrawEffect::DisabledEffect,
 			Color,
 			true,
-			3 * InParams.Geometry.Scale);
+			3.f * InParams.Geometry.Scale);
 
 		FSlateDrawElement::MakeLines(
 			InParams.OutDrawElements,
 			InParams.Layer,
-			InParams.Geometry.ToOffsetPaintGeometry(FVector2D(0, 34)),
+			InParams.Geometry.ToOffsetPaintGeometry(FVector2D(0.f, 34.f)),
 			LinePoints,
 			InParams.bEnabled ? ESlateDrawEffect::NoPixelSnapping : ESlateDrawEffect::NoPixelSnapping | ESlateDrawEffect::DisabledEffect,
 			Color,
 			true,
-			4 * InParams.Geometry.Scale);
+			4.f * InParams.Geometry.Scale);
 
 		return InParams.Layer;
 
@@ -5901,10 +5901,10 @@ TSharedRef<SDockTab> SpawnTab(const FSpawnTabArgs& Args, FName TabIdentifier)
 					.Clipping(EWidgetClipping::ClipToBounds)
 					[
 						SNew(SSplineTest)
-						.P0(FVector2D(64, 64))
-						.P1(FVector2D(256, 64))
-						.P2(FVector2D(64, 256))
-						.P3(FVector2D(256, 256))
+						.P0(FVector2f(64, 64))
+						.P1(FVector2f(256, 64))
+						.P2(FVector2f(64, 256))
+						.P3(FVector2f(256, 256))
 					]
 				]
 				+ SUniformGridPanel::Slot(1, 0)
@@ -5913,10 +5913,10 @@ TSharedRef<SDockTab> SpawnTab(const FSpawnTabArgs& Args, FName TabIdentifier)
 					.Clipping(EWidgetClipping::ClipToBounds)
 					[
 						SNew(SSplineTest)
-						.P0(FVector2D(128 - 32, 32))
-						.P1(FVector2D(128 + 32, 128))
-						.P2(FVector2D(128 - 32, 128))
-						.P3(FVector2D(128 + 32, 32))
+						.P0(FVector2f(128 - 32, 32))
+						.P1(FVector2f(128 + 32, 128))
+						.P2(FVector2f(128 - 32, 128))
+						.P3(FVector2f(128 + 32, 32))
 					]
 				]
 				+ SUniformGridPanel::Slot(0, 1)
@@ -5925,10 +5925,10 @@ TSharedRef<SDockTab> SpawnTab(const FSpawnTabArgs& Args, FName TabIdentifier)
 					.Clipping(EWidgetClipping::ClipToBounds)
 					[
 						SNew(SSplineTest)
-						.P0(FVector2D(64, 64))
-						.P1(FVector2D(256, 64))
-						.P2(FVector2D(256, 256))
-						.P3(FVector2D(64, 256))
+						.P0(FVector2f(64, 64))
+						.P1(FVector2f(256, 64))
+						.P2(FVector2f(256, 256))
+						.P3(FVector2f(64, 256))
 					]
 				]
 				+ SUniformGridPanel::Slot(1, 1)
@@ -5937,10 +5937,10 @@ TSharedRef<SDockTab> SpawnTab(const FSpawnTabArgs& Args, FName TabIdentifier)
 					.Clipping(EWidgetClipping::ClipToBounds)
 					[
 						SNew(SSplineTest)
-						.P0(FVector2D(64, 64))
-						.P1(FVector2D(128, 128))
-						.P2(FVector2D(196, 196))
-						.P3(FVector2D(256, 256))
+						.P0(FVector2f(64, 64))
+						.P1(FVector2f(128, 128))
+						.P2(FVector2f(196, 196))
+						.P3(FVector2f(256, 256))
 					]
 				]
 				
diff --git a/Engine/Source/Runtime/Slate/Private/Framework/Docking/SDockingCross.cpp b/Engine/Source/Runtime/Slate/Private/Framework/Docking/SDockingCross.cpp
--- a/Engine/Source/Runtime/Slate/Private/Framework/Docking/SDockingCross.cpp
+++ b/Engine/Source/Runtime/Slate/Private/Framework/Docking/SDockingCross.cpp
@@ -39,19 +39,19 @@ int32 SDockingCross::OnPaint( const FPaintArgs& Args, const FGeometry& AllottedG
 	//  |/           \|
 	//  +-------------+
 	
-	const FVector2D P0( DockZoneSizeX, DockZoneSizeY );
-	const FVector2D P1( AllottedGeometry.GetLocalSize().X - DockZoneSizeX, DockZoneSizeY );
-	const FVector2D P2( AllottedGeometry.GetLocalSize().X - DockZoneSizeX, AllottedGeometry.GetLocalSize().Y - DockZoneSizeY );
-	const FVector2D P3( DockZoneSizeX, AllottedGeometry.GetLocalSize().Y - DockZoneSizeY );
+	const FVector2f P0( DockZoneSizeX, DockZoneSizeY );
+	const FVector2f P1( AllottedGeometry.GetLocalSize().X - DockZoneSizeX, DockZoneSizeY );
+	const FVector2f P2( AllottedGeometry.GetLocalSize().X - DockZoneSizeX, AllottedGeometry.GetLocalSize().Y - DockZoneSizeY );
+	const FVector2f P3( DockZoneSizeX, AllottedGeometry.GetLocalSize().Y - DockZoneSizeY );
 
-	const FVector2D P0_Outer( 0,0 );
-	const FVector2D P1_Outer( AllottedGeometry.GetLocalSize().X,0 );
-	const FVector2D P2_Outer( AllottedGeometry.GetLocalSize().X,AllottedGeometry.GetLocalSize().Y );
-	const FVector2D P3_Outer( 0,AllottedGeometry.GetLocalSize().Y );
+	const FVector2f P0_Outer( 0.f,0.f );
+	const FVector2f P1_Outer( AllottedGeometry.GetLocalSize().X,0.f );
+	const FVector2f P2_Outer( AllottedGeometry.GetLocalSize().X,AllottedGeometry.GetLocalSize().Y );
+	const FVector2f P3_Outer( 0.f,AllottedGeometry.GetLocalSize().Y );
 
 	// Inner Box
 	{
-		TArray<FVector2D> InnerBox;
+		TArray<FVector2f> InnerBox;
 
 		InnerBox.Add( P0 );
 		InnerBox.Add( P1 );
@@ -63,13 +63,13 @@ int32 SDockingCross::OnPaint( const FPaintArgs& Args, const FGeometry& AllottedG
 	      OutDrawElements,
 			LayerId,
 			AllottedGeometry.ToPaintGeometry(),
-			InnerBox
+			MoveTemp(InnerBox)
 		);
 	}
 
 	// Outer Box
 	{
-		TArray<FVector2D> OuterBox;
+		TArray<FVector2f> OuterBox;
 
 		OuterBox.Add( P0_Outer );
 		OuterBox.Add( P1_Outer );
@@ -81,60 +81,60 @@ int32 SDockingCross::OnPaint( const FPaintArgs& Args, const FGeometry& AllottedG
 	      OutDrawElements,
 			LayerId,
 			AllottedGeometry.ToPaintGeometry(),
-			OuterBox
+			MoveTemp(OuterBox)
 		);
 	}
 	
 	// Diagonals
 	{
-		TArray<FVector2D> Points;
+		TArray<FVector2f> Points;
 
-		Points.Add( FVector2D(P0_Outer) );
-		Points.Add( FVector2D(P0) );
+		Points.Add( P0_Outer );
+		Points.Add( P0 );
 
 		FSlateDrawElement::MakeLines(
 	      OutDrawElements,
 			LayerId,
 			AllottedGeometry.ToPaintGeometry(),
-			Points
+			MoveTemp(Points)
 		);
 	}
 	{
-		TArray<FVector2D> Points;
+		TArray<FVector2f> Points;
 
-		Points.Add( FVector2D(P1_Outer) );
-		Points.Add( FVector2D(P1) );
+		Points.Add(FVector2f(P1_Outer));
+		Points.Add(FVector2f(P1));
 
 		FSlateDrawElement::MakeLines(
 	      OutDrawElements,
 			LayerId,
 			AllottedGeometry.ToPaintGeometry(),
-			Points
+			MoveTemp(Points)
 		);
 	}
 	{
-		TArray<FVector2D> Points;
+		TArray<FVector2f> Points;
 
-		Points.Add( FVector2D(P2_Outer) );
-		Points.Add( FVector2D(P2) );
+		Points.Add(FVector2f(P2_Outer));
+		Points.Add(FVector2f(P2));
 
 		FSlateDrawElement::MakeLines(
 	      OutDrawElements,
 			LayerId,
 			AllottedGeometry.ToPaintGeometry(),
-			Points
+			MoveTemp(Points)
 		);
 	}{
-		TArray<FVector2D> Points;
+		TArray<FVector2f> Points;
 
-		Points.Add( FVector2D(P3_Outer) );
-		Points.Add( FVector2D(P3) );
+		Points.Add(FVector2f(P3_Outer));
+		Points.Add(FVector2f(P3));
 
 		FSlateDrawElement::MakeLines(
 	      OutDrawElements,
 			LayerId,
 			AllottedGeometry.ToPaintGeometry(),
-			Points
+			MoveTemp(Points)
 		);
 	}
 	
diff --git a/Engine/Source/Runtime/Slate/Private/Widgets/Images/SSpinningImage.cpp b/Engine/Source/Runtime/Slate/Private/Widgets/Images/SSpinningImage.cpp
--- a/Engine/Source/Runtime/Slate/Private/Widgets/Images/SSpinningImage.cpp
+++ b/Engine/Source/Runtime/Slate/Private/Widgets/Images/SSpinningImage.cpp
@@ -39,7 +39,7 @@ int32 SSpinningImage::OnPaint( const FPaintArgs& Args, const FGeometry& Allotted
 			ImageBrush,
 			DrawEffects,
 			Angle,
-			TOptional<FVector2D>(), // Will auto rotate about center
+			TOptional<FVector2f>(), // Will auto rotate about center
 			FSlateDrawElement::RelativeToElement,
 			FinalColorAndOpacity
 			);
diff --git a/Engine/Source/Runtime/SlateCore/Private/Rendering/DrawElements.cpp b/Engine/Source/Runtime/SlateCore/Private/Rendering/DrawElements.cpp
--- a/Engine/Source/Runtime/SlateCore/Private/Rendering/DrawElements.cpp
+++ b/Engine/Source/Runtime/SlateCore/Private/Rendering/DrawElements.cpp
@@ -46,10 +46,10 @@ static bool ShouldCull(const FSlateWindowElementList& ElementList)
 
 static bool ShouldCull(const FSlateWindowElementList& ElementList, const FPaintGeometry& PaintGeometry)
 {
-	const FVector2D& LocalSize = PaintGeometry.GetLocalSize();
+	const FVector2f LocalSize = UE::Slate::CastToVector2f(PaintGeometry.GetLocalSize());
 	const float DrawScale = PaintGeometry.DrawScale;
 
-	const FVector2D PixelSize = (LocalSize * DrawScale);
+	const FVector2f PixelSize = (LocalSize * DrawScale);
 	if (PixelSize.X <= 0.f || PixelSize.Y <= 0.f)
 	{
 		return true;
@@ -123,7 +123,7 @@ FSlateWindowElementList::FSlateWindowElementList(const TSharedPtr<SWindow>& InPa
 	, RenderTargetWindow(nullptr)
 	, bNeedsDeferredResolve(false)
 	, ResolveToDeferredIndex()
-	, WindowSize(FVector2D(0.0f, 0.0f))
+	, WindowSize(FVector2f(0.0f, 0.0f))
 	//, bReportReferences(true)
 {
 	if (InPaintWindow.IsValid())
@@ -156,7 +156,7 @@ void FSlateDrawElement::Init(FSlateWindowElementList& ElementList, EElementType
 	ClipStateHandle.SetPreCachedClipIndex(ElementList.GetClippingIndex());
 
 #if UE_SLATE_VERIFY_PIXELSIZE
-	const FVector2D PixelSize = FVector2D(LocalSize * Scale);
+	const FVector2f PixelSize = FVector2f(LocalSize * Scale);
 	ensureMsgf(PixelSize.X >= 0.f && PixelSize.X <= (float)std::numeric_limits<uint16>::max(), TEXT("The size X '%f' is too small or big to fit in the SlateVertex buffer."), PixelSize.X);
 	ensureMsgf(PixelSize.Y >= 0.f && PixelSize.Y <= (float)std::numeric_limits<uint16>::max(), TEXT("The size Y '%f' is too small or big to fit in the SlateVertex buffer."), PixelSize.Y);
 #endif
@@ -167,7 +167,7 @@ void FSlateDrawElement::Init(FSlateWindowElementList& ElementList, EElementType
 	DrawEffects = InDrawEffects;
 	
 	// Calculate the layout to render transform as this is needed by several calculations downstream.
-	const FSlateLayoutTransform InverseLayoutTransform(Inverse(FSlateLayoutTransform(Scale, FVector2D(Position))));
+	const FSlateLayoutTransform InverseLayoutTransform(Inverse(FSlateLayoutTransform(Scale, Position)));
 
 	{
 		// This is a workaround because we want to keep track of the various Scenes 
@@ -204,7 +204,7 @@ void FSlateDrawElement::ApplyPositionOffset(FVector2f InOffset)
 	RenderTransform = Concatenate(RenderTransform, InOffset);
 
 	// Recompute cached layout to render transform
-	const FSlateLayoutTransform InverseLayoutTransform(Inverse(FSlateLayoutTransform(Scale, FVector2D(Position))));
+	const FSlateLayoutTransform InverseLayoutTransform(Inverse(FSlateLayoutTransform(Scale, Position)));
 }
 
 void FSlateDrawElement::MakeDebugQuad( FSlateWindowElementList& ElementList, uint32 InLayer, const FPaintGeometry& PaintGeometry, FLinearColor Tint)
@@ -312,7 +312,22 @@ void FSlateDrawElement::MakeRotatedBox(
 	const FSlateBrush* InBrush,
 	ESlateDrawEffect InDrawEffects,
 	float Angle2D,
-	TOptional<FVector2D> InRotationPoint,
+	TOptional<FVector2d> InRotationPoint,
+	ERotationSpace RotationSpace,
+	const FLinearColor& InTint)
+{
+	TOptional<FVector2f> RotationPoint = InRotationPoint.IsSet() ? UE::Slate::CastToVector2f(InRotationPoint.GetValue()) : TOptional<FVector2f>();
+	MakeRotatedBox(ElementList, InLayer, PaintGeometry, InBrush, InDrawEffects, Angle2D, RotationPoint, RotationSpace, InTint);
+}
+
+void FSlateDrawElement::MakeRotatedBox(
+	FSlateWindowElementList& ElementList,
+	uint32 InLayer,
+	const FPaintGeometry& PaintGeometry,
+	const FSlateBrush* InBrush,
+	ESlateDrawEffect InDrawEffects,
+	float Angle2D,
+	TOptional<FVector2f> InRotationPoint,
 	ERotationSpace RotationSpace,
 	const FLinearColor& InTint)
 {
@@ -327,8 +342,8 @@ void FSlateDrawElement::MakeRotatedBox(
 	
 	if (Angle2D != 0.0f)
 	{
-		const FVector2D RotationPoint = GetRotationPoint(PaintGeometry, InRotationPoint, RotationSpace);
-		const FSlateRenderTransform RotationTransform = Concatenate(Inverse(RotationPoint), FQuat2D(Angle2D), RotationPoint);
+		const FVector2f RotationPoint = GetRotationPoint(PaintGeometry, InRotationPoint, RotationSpace);
+		const FSlateRenderTransform RotationTransform = Concatenate(Inverse(RotationPoint), FQuat2f(Angle2D), RotationPoint);
 		DrawElement.SetRenderTransform(Concatenate(RotationTransform, DrawElement.GetRenderTransform()));
 	}
 }
@@ -476,7 +491,12 @@ void FSlateDrawElement::MakeGradient( FSlateWindowElementList& ElementList, uint
 	Element.Init(ElementList, EElementType::ET_Gradient, InLayer, PaintGeometry, InDrawEffects);
 }
 
-void FSlateDrawElement::MakeSpline( FSlateWindowElementList& ElementList, uint32 InLayer, const FPaintGeometry& PaintGeometry, const FVector2D& InStart, const FVector2D& InStartDir, const FVector2D& InEnd, const FVector2D& InEndDir, float InThickness, ESlateDrawEffect InDrawEffects, const FLinearColor& InTint )
+void FSlateDrawElement::MakeSpline(FSlateWindowElementList& ElementList, uint32 InLayer, const FPaintGeometry& PaintGeometry, const FVector2d InStart, const FVector2d InStartDir, const FVector2d InEnd, const FVector2d InEndDir, float InThickness, ESlateDrawEffect InDrawEffects, const FLinearColor& InTint)
+{
+	MakeSpline(ElementList, InLayer, PaintGeometry, UE::Slate::CastToVector2f(InStart), UE::Slate::CastToVector2f(InStartDir), UE::Slate::CastToVector2f(InEnd), UE::Slate::CastToVector2f(InEndDir), InThickness, InDrawEffects, InTint);
+}
+
+void FSlateDrawElement::MakeSpline(FSlateWindowElementList& ElementList, uint32 InLayer, const FPaintGeometry& PaintGeometry, const FVector2f InStart, const FVector2f InStartDir, const FVector2f InEnd, const FVector2f InEndDir, float InThickness, ESlateDrawEffect InDrawEffects, const FLinearColor& InTint)
 {
 	PaintGeometry.CommitTransformsIfUsingLegacyConstructor();
 
@@ -493,7 +513,12 @@ void FSlateDrawElement::MakeSpline( FSlateWindowElementList& ElementList, uint32
 	Element.Init(ElementList, EElementType::ET_Spline, InLayer, PaintGeometry, InDrawEffects);
 }
 
-void FSlateDrawElement::MakeCubicBezierSpline(FSlateWindowElementList & ElementList, uint32 InLayer, const FPaintGeometry & PaintGeometry, const FVector2D & P0, const FVector2D & P1, const FVector2D & P2, const FVector2D & P3, float InThickness, ESlateDrawEffect InDrawEffects, const FLinearColor & InTint)
+void FSlateDrawElement::MakeCubicBezierSpline(FSlateWindowElementList& ElementList, uint32 InLayer, const FPaintGeometry& PaintGeometry, const FVector2d P0, const FVector2d P1, const FVector2d P2, const FVector2d P3, float InThickness, ESlateDrawEffect InDrawEffects, const FLinearColor& InTint)
+{
+	MakeCubicBezierSpline(ElementList, InLayer, PaintGeometry, UE::Slate::CastToVector2f(P0), UE::Slate::CastToVector2f(P1), UE::Slate::CastToVector2f(P2), UE::Slate::CastToVector2f(P3), InThickness, InDrawEffects, InTint);
+}
+
+void FSlateDrawElement::MakeCubicBezierSpline(FSlateWindowElementList& ElementList, uint32 InLayer, const FPaintGeometry& PaintGeometry, const FVector2f P0, const FVector2f P1, const FVector2f P2, const FVector2f P3, float InThickness, ESlateDrawEffect InDrawEffects, const FLinearColor& InTint)
 {
 	PaintGeometry.CommitTransformsIfUsingLegacyConstructor();
 
@@ -510,9 +535,14 @@ void FSlateDrawElement::MakeCubicBezierSpline(FSlateWindowElementList & ElementL
 	Element.Init(ElementList, EElementType::ET_Spline, InLayer, PaintGeometry, InDrawEffects);
 }
 
-void FSlateDrawElement::MakeDrawSpaceSpline( FSlateWindowElementList& ElementList, uint32 InLayer, const FVector2D& InStart, const FVector2D& InStartDir, const FVector2D& InEnd, const FVector2D& InEndDir, float InThickness, ESlateDrawEffect InDrawEffects, const FLinearColor& InTint )
+void FSlateDrawElement::MakeDrawSpaceSpline(FSlateWindowElementList& ElementList, uint32 InLayer, const FVector2d InStart, const FVector2d InStartDir, const FVector2d InEnd, const FVector2d InEndDir, float InThickness, ESlateDrawEffect InDrawEffects, const FLinearColor& InTint)
+{
+	MakeSpline(ElementList, InLayer, FPaintGeometry(), InStart, InStartDir, InEnd, InEndDir, InThickness, InDrawEffects, InTint);
+}
+
+void FSlateDrawElement::MakeDrawSpaceSpline(FSlateWindowElementList& ElementList, uint32 InLayer, const FVector2f InStart, const FVector2f InStartDir, const FVector2f InEnd, const FVector2f InEndDir, float InThickness, ESlateDrawEffect InDrawEffects, const FLinearColor& InTint)
 {
-	MakeSpline( ElementList, InLayer, FPaintGeometry(), InStart, InStartDir, InEnd, InEndDir, InThickness, InDrawEffects, InTint );
+	MakeSpline(ElementList, InLayer, FPaintGeometry(), InStart, InStartDir, InEnd, InEndDir, InThickness, InDrawEffects, InTint);
 }
 
 void FSlateDrawElement::MakeDrawSpaceGradientSpline(FSlateWindowElementList& ElementList, uint32 InLayer, const FVector2D& InStart, const FVector2D& InStartDir, const FVector2D& InEnd, const FVector2D& InEndDir, const TArray<FSlateGradientStop>& InGradientStops, float InThickness, ESlateDrawEffect InDrawEffects)
@@ -551,7 +581,23 @@ void FSlateDrawElement::MakeDrawSpaceGradientSpline(FSlateWindowElementList& Ele
 	Element.Init(ElementList, EElementType::ET_Spline, InLayer, PaintGeometry, InDrawEffects);
 }
 
-void FSlateDrawElement::MakeLines(FSlateWindowElementList& ElementList, uint32 InLayer, const FPaintGeometry& PaintGeometry, const TArray<FVector2D>& Points, ESlateDrawEffect InDrawEffects, const FLinearColor& InTint, bool bAntialias, float Thickness)
+void FSlateDrawElement::MakeLines(FSlateWindowElementList& ElementList, uint32 InLayer, const FPaintGeometry& PaintGeometry, const TArray<FVector2d>& Points, ESlateDrawEffect InDrawEffects, const FLinearColor& InTint, bool bAntialias, float Thickness)
+{
+	if (ShouldCull(ElementList) || Points.Num() < 2)
+	{
+		return;
+	}
+
+	TArray<FVector2f> NewVector;
+	NewVector.Reserve(Points.Num());
+	for (FVector2d Vect : Points)
+	{
+		NewVector.Add(UE::Slate::CastToVector2f(Vect));
+	}
+	MakeLines(ElementList, InLayer, PaintGeometry, MoveTemp(NewVector), InDrawEffects, InTint, bAntialias, Thickness);
+}
+
+void FSlateDrawElement::MakeLines(FSlateWindowElementList& ElementList, uint32 InLayer, const FPaintGeometry& PaintGeometry, TArray<FVector2f> Points, ESlateDrawEffect InDrawEffects, const FLinearColor& InTint, bool bAntialias, float Thickness)
 {
 	PaintGeometry.CommitTransformsIfUsingLegacyConstructor();
 
@@ -566,20 +612,36 @@ void FSlateDrawElement::MakeLines(FSlateWindowElementList& ElementList, uint32 I
 
 	DataPayload.SetTint(InTint);
 	DataPayload.SetThickness(Thickness);
-	DataPayload.SetLines(Points, bAntialias, nullptr);
+	DataPayload.SetLines(MoveTemp(Points), bAntialias);
 
 	ESlateDrawEffect DrawEffects = InDrawEffects;
 	if (bAntialias)
 	{
 		// If the line is to be anti-aliased, we cannot reliably snap
-		// the generated vertices.
+		// the generated vertices
 		DrawEffects |= ESlateDrawEffect::NoPixelSnapping;
 	}
 
 	Element.Init(ElementList, EElementType::ET_Line, InLayer, PaintGeometry, DrawEffects);
 }
 
-void FSlateDrawElement::MakeLines( FSlateWindowElementList& ElementList, uint32 InLayer, const FPaintGeometry& PaintGeometry, const TArray<FVector2D>& Points, const TArray<FLinearColor>& PointColors, ESlateDrawEffect InDrawEffects, const FLinearColor& InTint, bool bAntialias, float Thickness )
+void FSlateDrawElement::MakeLines(FSlateWindowElementList& ElementList, uint32 InLayer, const FPaintGeometry& PaintGeometry, const TArray<FVector2d>& Points, const TArray<FLinearColor>& PointColors, ESlateDrawEffect InDrawEffects, const FLinearColor& InTint, bool bAntialias, float Thickness)
+{
+	if (ShouldCull(ElementList) || Points.Num() < 2)
+	{
+		return;
+	}
+
+	TArray<FVector2f> NewVector;
+	NewVector.Reserve(Points.Num());
+	for (FVector2d Vect : Points)
+	{
+		NewVector.Add(UE::Slate::CastToVector2f(Vect));
+	}
+	MakeLines(ElementList, InLayer, PaintGeometry, MoveTemp(NewVector), PointColors, InDrawEffects, InTint, bAntialias, Thickness);
+}
+
+void FSlateDrawElement::MakeLines( FSlateWindowElementList& ElementList, uint32 InLayer, const FPaintGeometry& PaintGeometry, TArray<FVector2f> Points, TArray<FLinearColor> PointColors, ESlateDrawEffect InDrawEffects, const FLinearColor& InTint, bool bAntialias, float Thickness )
 {
 	PaintGeometry.CommitTransformsIfUsingLegacyConstructor();
 
@@ -593,10 +655,9 @@ void FSlateDrawElement::MakeLines( FSlateWindowElementList& ElementList, uint32
 	FSlateLinePayload& DataPayload = ElementList.CreatePayload<FSlateLinePayload>(Element);
 	DataPayload.SetTint(InTint);
 	DataPayload.SetThickness(Thickness);
-	DataPayload.SetLines(Points, bAntialias, &PointColors);
+	DataPayload.SetLines(MoveTemp(Points), bAntialias, MoveTemp(PointColors));
 
 	Element.Init(ElementList, EElementType::ET_Line, InLayer, PaintGeometry, InDrawEffects);
-
 }
 
 void FSlateDrawElement::MakeViewport( FSlateWindowElementList& ElementList, uint32 InLayer, const FPaintGeometry& PaintGeometry, TSharedPtr<const ISlateViewport> Viewport, ESlateDrawEffect InDrawEffects, const FLinearColor& InTint )
@@ -694,11 +755,11 @@ FSlateDrawElement::~FSlateDrawElement()
 	}
 }
 
-FVector2D FSlateDrawElement::GetRotationPoint(const FPaintGeometry& PaintGeometry, const TOptional<FVector2D>& UserRotationPoint, ERotationSpace RotationSpace)
+FVector2f FSlateDrawElement::GetRotationPoint(const FPaintGeometry& PaintGeometry, const TOptional<FVector2f>& UserRotationPoint, ERotationSpace RotationSpace)
 {
-	FVector2D RotationPoint(0, 0);
+	FVector2f RotationPoint(0, 0);
 
-	const FVector2D& LocalSize = PaintGeometry.GetLocalSize();
+	const FVector2f LocalSize = UE::Slate::CastToVector2f(PaintGeometry.GetLocalSize());
 
 	switch (RotationSpace)
 	{
@@ -711,7 +772,7 @@ FVector2D FSlateDrawElement::GetRotationPoint(const FPaintGeometry& PaintGeometr
 		case RelativeToWorld:
 		{
 			// its in world space, must convert the point to local space.
-			RotationPoint = TransformPoint(Inverse(PaintGeometry.GetAccumulatedRenderTransform()), UserRotationPoint.Get(FVector2D::ZeroVector));
+			RotationPoint = TransformPoint(Inverse(PaintGeometry.GetAccumulatedRenderTransform()), UserRotationPoint.Get(FVector2f::ZeroVector));
 		}
 		break;
 	default:
diff --git a/Engine/Source/Runtime/SlateCore/Private/Widgets/SWidget.cpp b/Engine/Source/Runtime/SlateCore/Private/Widgets/SWidget.cpp
--- a/Engine/Source/Runtime/SlateCore/Private/Widgets/SWidget.cpp
+++ b/Engine/Source/Runtime/SlateCore/Private/Widgets/SWidget.cpp
@@ -1605,19 +1605,19 @@ int32 SWidget::Paint(const FPaintArgs& Args, const FGeometry& AllottedGeometry,
 	{
 		FSlateClippingZone ClippingZone(AllottedGeometry);
 
-		TArray<FVector2D> Points;
-		Points.Add(FVector2D(ClippingZone.TopLeft));
-		Points.Add(FVector2D(ClippingZone.TopRight));
-		Points.Add(FVector2D(ClippingZone.BottomRight));
-		Points.Add(FVector2D(ClippingZone.BottomLeft));
-		Points.Add(FVector2D(ClippingZone.TopLeft));
+		TArray<FVector2f> Points;
+		Points.Add(FVector2f(ClippingZone.TopLeft));
+		Points.Add(FVector2f(ClippingZone.TopRight));
+		Points.Add(FVector2f(ClippingZone.BottomRight));
+		Points.Add(FVector2f(ClippingZone.BottomLeft));
+		Points.Add(FVector2f(ClippingZone.TopLeft));
 
 		const bool bAntiAlias = true;
 		FSlateDrawElement::MakeLines(
 			OutDrawElements,
 			NewLayerId,
 			FPaintGeometry(),
-			Points,
+			MoveTemp(Points),
 			ESlateDrawEffect::None,
 			ClippingZone.IsAxisAligned() ? FLinearColor::Yellow : FLinearColor::Red,
 			bAntiAlias,
diff --git a/Engine/Source/Runtime/SlateCore/Public/Rendering/DrawElementPayloads.h b/Engine/Source/Runtime/SlateCore/Public/Rendering/DrawElementPayloads.h
--- a/Engine/Source/Runtime/SlateCore/Public/Rendering/DrawElementPayloads.h
+++ b/Engine/Source/Runtime/SlateCore/Public/Rendering/DrawElementPayloads.h
@@ -339,18 +339,29 @@ struct FSlateLinePayload : public FSlateDataPayload, public FSlateTintableElemen
 		{
 			NewPoints.Add(UE::Slate::CastToVector2f(Vect));
 		}
-		SetLines(MoveTemp(NewPoints), bInAntialias, InPointColors);
+		if (InPointColors)
+		{
+			SetLines(MoveTemp(NewPoints), bInAntialias, *InPointColors);
+		}
+		else
+		{
+			SetLines(MoveTemp(NewPoints), bInAntialias);
+		}
 	}
 
 
-	void SetLines(TArray<FVector2f> InPoints, bool bInAntialias, const TArray<FLinearColor>* InPointColors = nullptr)
+	void SetLines(TArray<FVector2f> InPoints, bool bInAntialias)
 	{
 		bAntialias = bInAntialias;
 		Points = MoveTemp(InPoints);
-		if (InPointColors)
-		{
-			PointColors = *InPointColors;
-		}
+		PointColors.Reset();
+	}
+
+	void SetLines(TArray<FVector2f> InPoints, bool bInAntialias, TArray<FLinearColor> InPointColors)
+	{
+		bAntialias = bInAntialias;
+		Points = MoveTemp(InPoints);
+		PointColors = MoveTemp(InPointColors);
 	}
 };
 
diff --git a/Engine/Source/Runtime/SlateCore/Public/Rendering/DrawElements.h b/Engine/Source/Runtime/SlateCore/Public/Rendering/DrawElements.h
--- a/Engine/Source/Runtime/SlateCore/Public/Rendering/DrawElements.h
+++ b/Engine/Source/Runtime/SlateCore/Public/Rendering/DrawElements.h
@@ -120,12 +120,23 @@ public:
 		uint32 InLayer, 
 		const FPaintGeometry& PaintGeometry, 
 		const FSlateBrush* InBrush, 
-		ESlateDrawEffect InDrawEffects = ESlateDrawEffect::None, 
-		float Angle = 0.0f,
-		TOptional<FVector2D> InRotationPoint = TOptional<FVector2D>(),
+		ESlateDrawEffect,
+		float Angle,
+		TOptional<FVector2d> InRotationPoint,
 		ERotationSpace RotationSpace = RelativeToElement,
 		const FLinearColor& InTint = FLinearColor::White );
 
+	SLATECORE_API static void MakeRotatedBox(
+		FSlateWindowElementList& ElementList,
+		uint32 InLayer,
+		const FPaintGeometry& PaintGeometry,
+		const FSlateBrush* InBrush,
+		ESlateDrawEffect InDrawEffects = ESlateDrawEffect::None,
+		float Angle = 0.0f,
+		TOptional<FVector2f> InRotationPoint = TOptional<FVector2f>(),
+		ERotationSpace RotationSpace = RelativeToElement,
+		const FLinearColor& InTint = FLinearColor::White);
+
 	/**
 	 * Creates a text element which displays a string of a rendered in a certain font on the screen
 	 *
@@ -186,7 +197,8 @@ public:
 	 * @param InDrawEffects         Optional draw effects to apply
 	 * @param InTint                Color to tint the element
 	 */
-	SLATECORE_API static void MakeSpline( FSlateWindowElementList& ElementList, uint32 InLayer, const FPaintGeometry& PaintGeometry, const FVector2D& InStart, const FVector2D& InStartDir, const FVector2D& InEnd, const FVector2D& InEndDir, float InThickness = 0.0f, ESlateDrawEffect InDrawEffects = ESlateDrawEffect::None, const FLinearColor& InTint=FLinearColor::White );
+	SLATECORE_API static void MakeSpline(FSlateWindowElementList& ElementList, uint32 InLayer, const FPaintGeometry& PaintGeometry, const FVector2d InStart, const FVector2d InStartDir, const FVector2d InEnd, const FVector2d InEndDir, float InThickness = 0.0f, ESlateDrawEffect InDrawEffects = ESlateDrawEffect::None, const FLinearColor& InTint=FLinearColor::White);
+	SLATECORE_API static void MakeSpline(FSlateWindowElementList& ElementList, uint32 InLayer, const FPaintGeometry& PaintGeometry, const FVector2f InStart, const FVector2f InStartDir, const FVector2f InEnd, const FVector2f InEndDir, float InThickness = 0.0f, ESlateDrawEffect InDrawEffects = ESlateDrawEffect::None, const FLinearColor& InTint=FLinearColor::White);
 
 	/**
 	 * Creates a Bezier Spline element
@@ -201,10 +213,12 @@ public:
 	 * @param InDrawEffects         Optional draw effects to apply
 	 * @param InTint                Color to tint the element
 	 */
-	SLATECORE_API static void MakeCubicBezierSpline(FSlateWindowElementList& ElementList, uint32 InLayer, const FPaintGeometry& PaintGeometry, const FVector2D& P0, const FVector2D& P1, const FVector2D& P2, const FVector2D& P3, float InThickness = 0.0f, ESlateDrawEffect InDrawEffects = ESlateDrawEffect::None, const FLinearColor& InTint = FLinearColor::White);
+	SLATECORE_API static void MakeCubicBezierSpline(FSlateWindowElementList& ElementList, uint32 InLayer, const FPaintGeometry& PaintGeometry, const FVector2d P0, const FVector2d P1, const FVector2d P2, const FVector2d P3, float InThickness = 0.0f, ESlateDrawEffect InDrawEffects = ESlateDrawEffect::None, const FLinearColor& InTint = FLinearColor::White);
+	SLATECORE_API static void MakeCubicBezierSpline(FSlateWindowElementList& ElementList, uint32 InLayer, const FPaintGeometry& PaintGeometry, const FVector2f P0, const FVector2f P1, const FVector2f P2, const FVector2f P3, float InThickness = 0.0f, ESlateDrawEffect InDrawEffects = ESlateDrawEffect::None, const FLinearColor& InTint = FLinearColor::White);
 
 	/** Just like MakeSpline but in draw-space coordinates. This is useful for connecting already-transformed widgets together. */
-	SLATECORE_API static void MakeDrawSpaceSpline(FSlateWindowElementList& ElementList, uint32 InLayer, const FVector2D& InStart, const FVector2D& InStartDir, const FVector2D& InEnd, const FVector2D& InEndDir, float InThickness = 0.0f, ESlateDrawEffect InDrawEffects = ESlateDrawEffect::None, const FLinearColor& InTint=FLinearColor::White);
+	SLATECORE_API static void MakeDrawSpaceSpline(FSlateWindowElementList& ElementList, uint32 InLayer, const FVector2d InStart, const FVector2d InStartDir, const FVector2d InEnd, const FVector2d InEndDir, float InThickness = 0.0f, ESlateDrawEffect InDrawEffects = ESlateDrawEffect::None, const FLinearColor& InTint=FLinearColor::White);
+	SLATECORE_API static void MakeDrawSpaceSpline(FSlateWindowElementList& ElementList, uint32 InLayer, const FVector2f InStart, const FVector2f InStartDir, const FVector2f InEnd, const FVector2f InEndDir, float InThickness = 0.0f, ESlateDrawEffect InDrawEffects = ESlateDrawEffect::None, const FLinearColor& InTint=FLinearColor::White);
 
 	/** Just like MakeSpline but in draw-space coordinates. This is useful for connecting already-transformed widgets together. */
 	UE_DEPRECATED(4.20, "Splines with color gradients will not be supported in the future.")
@@ -216,17 +230,34 @@ public:
 	/**
 	 * Creates a line defined by the provided points
 	 *
-	 * @param ElementList			   The list in which to add elements
+	 * @param ElementList              The list in which to add elements
+	 * @param InLayer                  The layer to draw the element on
+	 * @param PaintGeometry            DrawSpace position and dimensions; see FPaintGeometry
+	 * @param Points                   Points that make up the lines.  The points are joined together. I.E if Points has A,B,C there the line is A-B-C.  To draw non-joining line segments call MakeLines multiple times
+	 * @param InDrawEffects            Optional draw effects to apply
+	 * @param InTint                   Color to tint the element
+	 * @param bAntialias               Should antialiasing be applied to the line?
+	 * @param Thickness                The thickness of the line
+	 */
+	SLATECORE_API static void MakeLines(FSlateWindowElementList& ElementList, uint32 InLayer, const FPaintGeometry& PaintGeometry, const TArray<FVector2d>& Points, ESlateDrawEffect InDrawEffects = ESlateDrawEffect::None, const FLinearColor& InTint=FLinearColor::White, bool bAntialias = true, float Thickness = 1.0f);
+	SLATECORE_API static void MakeLines(FSlateWindowElementList& ElementList, uint32 InLayer, const FPaintGeometry& PaintGeometry, TArray<FVector2f> Points, ESlateDrawEffect InDrawEffects = ESlateDrawEffect::None, const FLinearColor& InTint = FLinearColor::White, bool bAntialias = true, float Thickness = 1.0f);
+
+
+	/**
+	 * Creates a line defined by the provided points
+	 *
+	 * @param ElementList              The list in which to add elements
 	 * @param InLayer                  The layer to draw the element on
 	 * @param PaintGeometry            DrawSpace position and dimensions; see FPaintGeometry
 	 * @param Points                   Points that make up the lines.  The points are joined together. I.E if Points has A,B,C there the line is A-B-C.  To draw non-joining line segments call MakeLines multiple times
+	 * @param PointColors              Vertex Color for each defined points
 	 * @param InDrawEffects            Optional draw effects to apply
 	 * @param InTint                   Color to tint the element
 	 * @param bAntialias               Should antialiasing be applied to the line?
 	 * @param Thickness                The thickness of the line
 	 */
-	SLATECORE_API static void MakeLines( FSlateWindowElementList& ElementList, uint32 InLayer, const FPaintGeometry& PaintGeometry, const TArray<FVector2D>& Points, ESlateDrawEffect InDrawEffects = ESlateDrawEffect::None, const FLinearColor& InTint=FLinearColor::White, bool bAntialias = true, float Thickness = 1.0f );
-	SLATECORE_API static void MakeLines( FSlateWindowElementList& ElementList, uint32 InLayer, const FPaintGeometry& PaintGeometry, const TArray<FVector2D>& Points, const TArray<FLinearColor>& PointColors, ESlateDrawEffect InDrawEffects = ESlateDrawEffect::None, const FLinearColor& InTint=FLinearColor::White, bool bAntialias = true, float Thickness = 1.0f );
+	SLATECORE_API static void MakeLines(FSlateWindowElementList& ElementList, uint32 InLayer, const FPaintGeometry& PaintGeometry, const TArray<FVector2d>& Points, const TArray<FLinearColor>& PointColors, ESlateDrawEffect InDrawEffects = ESlateDrawEffect::None, const FLinearColor& InTint=FLinearColor::White, bool bAntialias = true, float Thickness = 1.0f);
+	SLATECORE_API static void MakeLines(FSlateWindowElementList& ElementList, uint32 InLayer, const FPaintGeometry& PaintGeometry, TArray<FVector2f> Points, TArray<FLinearColor> PointColors, ESlateDrawEffect InDrawEffects = ESlateDrawEffect::None, const FLinearColor& InTint=FLinearColor::White, bool bAntialias = true, float Thickness = 1.0f);
 
 	/**
 	 * Creates a viewport element which is useful for rendering custom data in a texture into Slate
@@ -276,7 +307,7 @@ public:
 		return Position;
 	}
 
-	UE_DEPRECATED(5.1, "Slate rendering now use float instead of double. Use SetPosition(FVector2f)")
+	UE_DEPRECATED(5.1, "Slate rendering uses float instead of double. Use SetPosition(FVector2f)")
 	FORCEINLINE void SetPosition(FVector2d InPosition)
 	{
 		Position = UE::Slate::CastToVector2f(InPosition);
@@ -313,7 +344,7 @@ public:
 
 	void AddReferencedObjects(FReferenceCollector& Collector);
 
-	UE_DEPRECATED(5.1, "Slate rendering now use float instead of double. Use ApplyPositionOffset(FVector2f)")
+	UE_DEPRECATED(5.1, "Slate rendering uses float instead of double. Use ApplyPositionOffset(FVector2f)")
 	void ApplyPositionOffset(FVector2d InOffset);
 
 	/**
@@ -333,7 +364,7 @@ public:
 private:
 	void Init(FSlateWindowElementList& ElementList, EElementType InElementType, uint32 InLayer, const FPaintGeometry& PaintGeometry, ESlateDrawEffect InDrawEffects);
 
-	static FVector2D GetRotationPoint( const FPaintGeometry& PaintGeometry, const TOptional<FVector2D>& UserRotationPoint, ERotationSpace RotationSpace );
+	static FVector2f GetRotationPoint( const FPaintGeometry& PaintGeometry, const TOptional<FVector2f>& UserRotationPoint, ERotationSpace RotationSpace );
 	static FSlateDrawElement& MakeBoxInternal(FSlateWindowElementList& ElementList, uint32 InLayer, const FPaintGeometry& PaintGeometry, const FSlateBrush* InBrush, ESlateDrawEffect InDrawEffects, const FLinearColor& InTint);
 private:
 	FSlateDataPayload* DataPayload;
