Slate: Transform ResourceManager & SlateRenderer from Vector2D to Vector2f.

diff --git a/Engine/Source/Developer/StandaloneRenderer/Private/OpenGL/SlateOpenGLRenderer.cpp b/Engine/Source/Developer/StandaloneRenderer/Private/OpenGL/SlateOpenGLRenderer.cpp
--- a/Engine/Source/Developer/StandaloneRenderer/Private/OpenGL/SlateOpenGLRenderer.cpp
+++ b/Engine/Source/Developer/StandaloneRenderer/Private/OpenGL/SlateOpenGLRenderer.cpp
@@ -3,6 +3,7 @@
 #include "OpenGL/SlateOpenGLRenderer.h"
 #include "Fonts/FontTypes.h"
 #include "Fonts/FontCache.h"
+#include "Types/SlateVector2.h"
 #include "Widgets/SWindow.h"
 #include "OpenGL/SlateOpenGLTextures.h"
 
@@ -152,7 +153,7 @@ void FSlateOpenGLRenderer::DrawWindows( FSlateDrawBuffer& InWindowDrawBuffer )
 				{
 					//@todo implement fullscreen
 					const bool bFullscreen = false;
-					Private_ResizeViewport( WindowSize, *Viewport, bFullscreen );
+					Private_ResizeViewport( UE::Slate::CastToVector2f(WindowSize), *Viewport, bFullscreen );
 				}
 
 				Viewport->MakeCurrent();
@@ -235,7 +236,7 @@ void FSlateOpenGLRenderer::RequestResize( const TSharedPtr<SWindow>& InWindow, u
 	// @todo implement.  Viewports are currently resized in DrawWindows
 }
 
-void FSlateOpenGLRenderer::Private_ResizeViewport( const FVector2D& WindowSize, FSlateOpenGLViewport& InViewport, bool bFullscreen )
+void FSlateOpenGLRenderer::Private_ResizeViewport( FVector2f WindowSize, FSlateOpenGLViewport& InViewport, bool bFullscreen )
 {
 	uint32 Width = FMath::TruncToInt(WindowSize.X);
 	uint32 Height = FMath::TruncToInt(WindowSize.Y);
@@ -255,7 +256,7 @@ void FSlateOpenGLRenderer::UpdateFullscreenState( const TSharedRef<SWindow> InWi
 //		uint32 ResX = OverrideResX ? OverrideResX : GSystemResolution.ResX;
 //		uint32 ResY = OverrideResY ? OverrideResY : GSystemResolution.ResY;
 
-		Private_ResizeViewport( FVector2D( Viewport->ViewportRect.Right, Viewport->ViewportRect.Bottom ), *Viewport, bFullscreen );
+		Private_ResizeViewport( FVector2f( Viewport->ViewportRect.Right, Viewport->ViewportRect.Bottom ), *Viewport, bFullscreen );
 	}
 }
 
@@ -271,7 +272,7 @@ bool FSlateOpenGLRenderer::GenerateDynamicImageResource(FName ResourceName, uint
 	return TextureManager->CreateDynamicTextureResource(ResourceName, Width, Height, Bytes) != NULL;
 }
 
-FSlateResourceHandle FSlateOpenGLRenderer::GetResourceHandle(const FSlateBrush& Brush, FVector2D LocalSize, float DrawScale)
+FSlateResourceHandle FSlateOpenGLRenderer::GetResourceHandle(const FSlateBrush& Brush, FVector2f LocalSize, float DrawScale)
 {
 	return TextureManager->GetResourceHandle(Brush, LocalSize, DrawScale);
 }
diff --git a/Engine/Source/Developer/StandaloneRenderer/Private/OpenGL/SlateOpenGLRenderer.h b/Engine/Source/Developer/StandaloneRenderer/Private/OpenGL/SlateOpenGLRenderer.h
--- a/Engine/Source/Developer/StandaloneRenderer/Private/OpenGL/SlateOpenGLRenderer.h
+++ b/Engine/Source/Developer/StandaloneRenderer/Private/OpenGL/SlateOpenGLRenderer.h
@@ -159,7 +159,7 @@ public:
 	virtual void RestoreSystemResolution(const TSharedRef<SWindow> InWindow) override {}
 	virtual void ReleaseDynamicResource( const FSlateBrush& Brush ) override;
 	virtual bool GenerateDynamicImageResource(FName ResourceName, uint32 Width, uint32 Height, const TArray< uint8 >& Bytes) override;
-	virtual FSlateResourceHandle GetResourceHandle(const FSlateBrush& Brush, FVector2D LocalSize, float DrawScale) override;
+	virtual FSlateResourceHandle GetResourceHandle(const FSlateBrush& Brush, FVector2f LocalSize, float DrawScale) override;
 	virtual void RemoveDynamicBrushResource( TSharedPtr<FSlateDynamicImageBrush> BrushToRemove ) override;
 	virtual void LoadStyleResources(const ISlateStyle& Style) override;
 	virtual FSlateUpdatableTexture* CreateUpdatableTexture(uint32 Width, uint32 Height) override;
@@ -178,7 +178,7 @@ private:
 	 * @param WindowSize    The Size of the window in pixels
 	 * @param InViewport    The viewport to resize
 	 */
-	void Private_ResizeViewport( const FVector2D& WindowSize, FSlateOpenGLViewport& InViewport, bool bFullscreen );
+	void Private_ResizeViewport( FVector2f WindowSize, FSlateOpenGLViewport& InViewport, bool bFullscreen );
 private:
 	/** View matrix to use when rendering */
 	FMatrix ViewMatrix;
diff --git a/Engine/Source/Developer/StandaloneRenderer/Private/OpenGL/SlateOpenGLTextureManager.cpp b/Engine/Source/Developer/StandaloneRenderer/Private/OpenGL/SlateOpenGLTextureManager.cpp
--- a/Engine/Source/Developer/StandaloneRenderer/Private/OpenGL/SlateOpenGLTextureManager.cpp
+++ b/Engine/Source/Developer/StandaloneRenderer/Private/OpenGL/SlateOpenGLTextureManager.cpp
@@ -416,7 +416,7 @@ void FSlateOpenGLTextureManager::ConditionalFlushCache()
 /**
  * Returns a texture with the passed in name or NULL if it cannot be found.
  */
-FSlateShaderResourceProxy* FSlateOpenGLTextureManager::GetShaderResource(const FSlateBrush& Brush, FVector2D LocalSize, float DrawScale)
+FSlateShaderResourceProxy* FSlateOpenGLTextureManager::GetShaderResource(const FSlateBrush& Brush, FVector2f LocalSize, float DrawScale)
 {
 	FSlateShaderResourceProxy* Texture = NULL;
 	if (Brush.GetImageType() == ESlateBrushImageType::Vector)
@@ -435,7 +435,7 @@ FSlateShaderResourceProxy* FSlateOpenGLTextureManager::GetShaderResource(const F
 	return Texture;
 }
 
-FSlateShaderResourceProxy* FSlateOpenGLTextureManager::GetVectorResource(const FSlateBrush& Brush, FVector2D LocalSize, float DrawScale)
+FSlateShaderResourceProxy* FSlateOpenGLTextureManager::GetVectorResource(const FSlateBrush& Brush, FVector2f LocalSize, float DrawScale)
 {
 	return VectorGraphicsCache->GetShaderResource(Brush, LocalSize, DrawScale);
 }
@@ -484,7 +484,7 @@ FAtlasSlotInfo FSlateOpenGLTextureManager::GetAtlasSlotInfoAtPosition(FIntPoint
 		const FAtlasedTextureSlot* Slot = Atlas->GetSlotAtPosition(InPosition);
 		if (Slot)
 		{
-			NewInfo.AtlasSlotRect = FSlateRect(FVector2D(Slot->X, Slot->Y), FVector2D(Slot->X + Slot->Width, Slot->Y + Slot->Height));
+			NewInfo.AtlasSlotRect = FSlateRect(FVector2f(Slot->X, Slot->Y), FVector2f(Slot->X + Slot->Width, Slot->Y + Slot->Height));
 
 			NewInfo.TextureName = AtlasDebugData.FindRef(Slot);
 
diff --git a/Engine/Source/Developer/StandaloneRenderer/Private/OpenGL/SlateOpenGLTextureManager.h b/Engine/Source/Developer/StandaloneRenderer/Private/OpenGL/SlateOpenGLTextureManager.h
--- a/Engine/Source/Developer/StandaloneRenderer/Private/OpenGL/SlateOpenGLTextureManager.h
+++ b/Engine/Source/Developer/StandaloneRenderer/Private/OpenGL/SlateOpenGLTextureManager.h
@@ -29,7 +29,7 @@ public:
 	/**
 	 * Returns a texture with the passed in name or NULL if it cannot be found.
 	 */
-	virtual FSlateShaderResourceProxy* GetShaderResource(const FSlateBrush& Brush, FVector2D LocalSize, float DrawScale) override;
+	virtual FSlateShaderResourceProxy* GetShaderResource(const FSlateBrush& Brush, FVector2f LocalSize, float DrawScale) override;
 
 
 	virtual ISlateAtlasProvider* GetTextureAtlasProvider() override;
@@ -70,7 +70,7 @@ private:
 
 	FSlateShaderResourceProxy* GetDynamicTextureResource(const FSlateBrush& InBrush);
 
-	FSlateShaderResourceProxy* GetVectorResource(const FSlateBrush& Brush, FVector2D LocalSize, float DrawScale);
+	FSlateShaderResourceProxy* GetVectorResource(const FSlateBrush& Brush, FVector2f LocalSize, float DrawScale);
 private:
 	/** Represents a dynamic resource for rendering */
 	struct FDynamicTextureResource
diff --git a/Engine/Source/Developer/StandaloneRenderer/Private/Windows/D3D/SlateD3DRenderer.cpp b/Engine/Source/Developer/StandaloneRenderer/Private/Windows/D3D/SlateD3DRenderer.cpp
--- a/Engine/Source/Developer/StandaloneRenderer/Private/Windows/D3D/SlateD3DRenderer.cpp
+++ b/Engine/Source/Developer/StandaloneRenderer/Private/Windows/D3D/SlateD3DRenderer.cpp
@@ -413,7 +413,7 @@ bool FSlateD3DRenderer::GenerateDynamicImageResource(FName ResourceName, uint32
 	return Result != nullptr;
 }
 
-FSlateResourceHandle FSlateD3DRenderer::GetResourceHandle(const FSlateBrush& Brush, FVector2D LocalSize, float DrawScale)
+FSlateResourceHandle FSlateD3DRenderer::GetResourceHandle(const FSlateBrush& Brush, FVector2f LocalSize, float DrawScale)
 {
 	if (!TextureManager.IsValid())
 	{
diff --git a/Engine/Source/Developer/StandaloneRenderer/Private/Windows/D3D/SlateD3DRenderer.h b/Engine/Source/Developer/StandaloneRenderer/Private/Windows/D3D/SlateD3DRenderer.h
--- a/Engine/Source/Developer/StandaloneRenderer/Private/Windows/D3D/SlateD3DRenderer.h
+++ b/Engine/Source/Developer/StandaloneRenderer/Private/Windows/D3D/SlateD3DRenderer.h
@@ -63,7 +63,7 @@ public:
 	virtual void RestoreSystemResolution(const TSharedRef<SWindow> InWindow) override {}
 	virtual void ReleaseDynamicResource( const FSlateBrush& Brush ) override;
 	virtual bool GenerateDynamicImageResource(FName ResourceName, uint32 Width, uint32 Height, const TArray< uint8 >& Bytes) override;
-	virtual FSlateResourceHandle GetResourceHandle(const FSlateBrush& Brush, FVector2D LocalSize, float DrawScale) override;
+	virtual FSlateResourceHandle GetResourceHandle(const FSlateBrush& Brush, FVector2f LocalSize, float DrawScale) override;
 	virtual void RemoveDynamicBrushResource( TSharedPtr<FSlateDynamicImageBrush> BrushToRemove ) override;
 	virtual void LoadStyleResources( const ISlateStyle& Style ) override;
 	virtual FSlateUpdatableTexture* CreateUpdatableTexture(uint32 Width, uint32 Height) override;
diff --git a/Engine/Source/Developer/StandaloneRenderer/Private/Windows/D3D/SlateD3DTextureManager.cpp b/Engine/Source/Developer/StandaloneRenderer/Private/Windows/D3D/SlateD3DTextureManager.cpp
--- a/Engine/Source/Developer/StandaloneRenderer/Private/Windows/D3D/SlateD3DTextureManager.cpp
+++ b/Engine/Source/Developer/StandaloneRenderer/Private/Windows/D3D/SlateD3DTextureManager.cpp
@@ -114,7 +114,7 @@ FAtlasSlotInfo FSlateD3DTextureManager::GetAtlasSlotInfoAtPosition(FIntPoint InP
 		const FAtlasedTextureSlot* Slot = Atlas->GetSlotAtPosition(InPosition);
 		if (Slot)
 		{
-			NewInfo.AtlasSlotRect = FSlateRect(FVector2D(Slot->X, Slot->Y), FVector2D(Slot->X + Slot->Width, Slot->Y + Slot->Height));
+			NewInfo.AtlasSlotRect = FSlateRect(FVector2f(Slot->X, Slot->Y), FVector2f(Slot->X + Slot->Width, Slot->Y + Slot->Height));
 
 			NewInfo.TextureName = AtlasDebugData.FindRef(Slot);
 
@@ -203,7 +203,7 @@ void FSlateD3DTextureManager::CreateTextures( const TArray< const FSlateBrush* >
 /**
  * Returns a texture with the passed in name or NULL if it cannot be found.
  */
-FSlateShaderResourceProxy* FSlateD3DTextureManager::GetShaderResource(const FSlateBrush& Brush, FVector2D LocalSize, float DrawScale)
+FSlateShaderResourceProxy* FSlateD3DTextureManager::GetShaderResource(const FSlateBrush& Brush, FVector2f LocalSize, float DrawScale)
 {
 	FSlateShaderResourceProxy* Texture = NULL;
 	if (Brush.GetImageType() == ESlateBrushImageType::Vector)
@@ -338,7 +338,7 @@ FSlateShaderResourceProxy* FSlateD3DTextureManager::GetDynamicTextureResource( c
 	return  NULL;
 }
 
-FSlateShaderResourceProxy* FSlateD3DTextureManager::GetVectorResource(const FSlateBrush& Brush, FVector2D LocalSize, float DrawScale)
+FSlateShaderResourceProxy* FSlateD3DTextureManager::GetVectorResource(const FSlateBrush& Brush, FVector2f LocalSize, float DrawScale)
 {
 	return VectorGraphicsCache->GetShaderResource(Brush, LocalSize, DrawScale);
 }
diff --git a/Engine/Source/Developer/StandaloneRenderer/Private/Windows/D3D/SlateD3DTextureManager.h b/Engine/Source/Developer/StandaloneRenderer/Private/Windows/D3D/SlateD3DTextureManager.h
--- a/Engine/Source/Developer/StandaloneRenderer/Private/Windows/D3D/SlateD3DTextureManager.h
+++ b/Engine/Source/Developer/StandaloneRenderer/Private/Windows/D3D/SlateD3DTextureManager.h
@@ -38,7 +38,7 @@ public:
 	/**
 	 * Returns a texture with the passed in name or NULL if it cannot be found.
 	 */
-	virtual FSlateShaderResourceProxy* GetShaderResource(const FSlateBrush& Brush, FVector2D LocalSize, float DrawScale) override;
+	virtual FSlateShaderResourceProxy* GetShaderResource(const FSlateBrush& Brush, FVector2f LocalSize, float DrawScale) override;
 
 	virtual ISlateAtlasProvider* GetTextureAtlasProvider() override;
 
@@ -76,7 +76,7 @@ private:
 	 * @param LocalSize	The unscaled local size of the final image
 	 * @param DrawScale	Any scaling applied to the final image
 	 */
-	FSlateShaderResourceProxy* GetVectorResource(const FSlateBrush& Brush, FVector2D LocalSize, float DrawScale);
+	FSlateShaderResourceProxy* GetVectorResource(const FSlateBrush& Brush, FVector2f LocalSize, float DrawScale);
 
 	/** 
 	 * Creates textures from files on disk and atlases them if possible
diff --git a/Engine/Source/Runtime/SlateCore/Private/Rendering/ShaderResourceManager.cpp b/Engine/Source/Runtime/SlateCore/Private/Rendering/ShaderResourceManager.cpp
--- a/Engine/Source/Runtime/SlateCore/Private/Rendering/ShaderResourceManager.cpp
+++ b/Engine/Source/Runtime/SlateCore/Private/Rendering/ShaderResourceManager.cpp
@@ -4,7 +4,7 @@
 
 DECLARE_CYCLE_STAT(TEXT("GetResourceHandle Time"), STAT_SlateGetResourceHandle, STATGROUP_Slate);
 
-FSlateResourceHandle FSlateShaderResourceManager::GetResourceHandle(const FSlateBrush& Brush, FVector2D LocalSize, float DrawScale)
+FSlateResourceHandle FSlateShaderResourceManager::GetResourceHandle(const FSlateBrush& Brush, FVector2f LocalSize, float DrawScale)
 {
 	SCOPE_CYCLE_COUNTER(STAT_SlateGetResourceHandle);
 
@@ -38,5 +38,5 @@ FSlateResourceHandle FSlateShaderResourceManager::GetResourceHandle(const FSlate
 
 FSlateResourceHandle FSlateShaderResourceManager::GetResourceHandle(const FSlateBrush& Brush)
 {
-	return GetResourceHandle(Brush, FVector2D::ZeroVector, 1.0f);
+	return GetResourceHandle(Brush, FVector2f::ZeroVector, 1.0f);
 }
diff --git a/Engine/Source/Runtime/SlateCore/Private/Rendering/SlateVectorGraphicsCache.cpp b/Engine/Source/Runtime/SlateCore/Private/Rendering/SlateVectorGraphicsCache.cpp
--- a/Engine/Source/Runtime/SlateCore/Private/Rendering/SlateVectorGraphicsCache.cpp
+++ b/Engine/Source/Runtime/SlateCore/Private/Rendering/SlateVectorGraphicsCache.cpp
@@ -3,6 +3,7 @@
 #include "Rendering/SlateVectorGraphicsCache.h"
 #include "Rendering/SlateSVGRasterizer.h"
 #include "Textures/TextureAtlas.h"
+#include "Types/SlateVector2.h"
 #include "ProfilingDebugging/ScopedTimers.h"
 #include "Async/ParallelFor.h"
 
@@ -34,7 +35,12 @@ FSlateVectorGraphicsCache::FSlateVectorGraphicsCache(TSharedPtr<ISlateTextureAtl
 {
 }
 
-FSlateShaderResourceProxy* FSlateVectorGraphicsCache::GetShaderResource(const FSlateBrush& Brush, FVector2D LocalSize, float DrawScale)
+//FSlateShaderResourceProxy* FSlateVectorGraphicsCache::GetShaderResource(const FSlateBrush& Brush, FVector2d LocalSize, float DrawScale)
+//{
+//	return GetShaderResource(Brush, UE::Slate::CastToVector2f(LocalSize), DrawScale);
+//}
+
+FSlateShaderResourceProxy* FSlateVectorGraphicsCache::GetShaderResource(const FSlateBrush& Brush, FVector2f LocalSize, float DrawScale)
 {
 	FSlateShaderResourceProxy* Proxy = nullptr;
 
diff --git a/Engine/Source/Runtime/SlateCore/Private/Styling/SlateBrush.cpp b/Engine/Source/Runtime/SlateCore/Private/Styling/SlateBrush.cpp
--- a/Engine/Source/Runtime/SlateCore/Private/Styling/SlateBrush.cpp
+++ b/Engine/Source/Runtime/SlateCore/Private/Styling/SlateBrush.cpp
@@ -3,6 +3,7 @@
 #include "Styling/SlateBrush.h"
 #include "SlateGlobals.h"
 #include "Application/SlateApplicationBase.h"
+#include "Types/SlateVector2.h"
 
 FSlateBrush::FSlateBrush()
 	: bIsDynamicallyLoaded(false)
@@ -144,7 +145,7 @@ void FSlateBrush::UpdateRenderingResource(FVector2D LocalSize, float DrawScale)
 		// For vector graphics we will rebuild the handle only if the shape needs to be rasterized again and the new size and scale
 		if (!ResourceHandle.IsValid() || ImageType == ESlateBrushImageType::Vector)
 		{
-			ResourceHandle = FSlateApplicationBase::Get().GetRenderer()->GetResourceHandle(*this, LocalSize, DrawScale);
+			ResourceHandle = FSlateApplicationBase::Get().GetRenderer()->GetResourceHandle(*this, UE::Slate::CastToVector2f(LocalSize), DrawScale);
 		}
 	}
 
diff --git a/Engine/Source/Runtime/SlateCore/Public/Layout/SlateRect.h b/Engine/Source/Runtime/SlateCore/Public/Layout/SlateRect.h
--- a/Engine/Source/Runtime/SlateCore/Public/Layout/SlateRect.h
+++ b/Engine/Source/Runtime/SlateCore/Public/Layout/SlateRect.h
@@ -44,6 +44,13 @@ public:
 		, Bottom(UE_REAL_TO_FLOAT(InEndPos.Y))
 	{ }
 
+	FSlateRect(FVector2f InStartPos, FVector2f InEndPos)
+		: Left(InStartPos.X)
+		, Top(InStartPos.Y)
+		, Right(InEndPos.X)
+		, Bottom(InEndPos.Y)
+	{ }
+
 	/**
 	 * Creates a rect from a top left point and extent. Provided as a factory function to not conflict
 	 * with the TopLeft + BottomRight ctor.
diff --git a/Engine/Source/Runtime/SlateCore/Public/Rendering/ShaderResourceManager.h b/Engine/Source/Runtime/SlateCore/Public/Rendering/ShaderResourceManager.h
--- a/Engine/Source/Runtime/SlateCore/Public/Rendering/ShaderResourceManager.h
+++ b/Engine/Source/Runtime/SlateCore/Public/Rendering/ShaderResourceManager.h
@@ -7,6 +7,7 @@
 #include "Textures/SlateShaderResource.h"
 #include "Textures/SlateTextureData.h"
 #include "Textures/TextureAtlas.h"
+#include "Types/SlateVector2.h"
 
 /**
  * Represents a new texture that has been loaded but no resource created for it
@@ -52,7 +53,12 @@ public:
 	/** 
 	 * Returns a texture associated with the passed in name.  Should return nullptr if not found 
 	 */
-	virtual FSlateShaderResourceProxy* GetShaderResource(const FSlateBrush& Brush, FVector2D LocalSize, float DrawScale) = 0;
+	UE_DEPRECATED(5.1, "Slate rendering uses float instead of double. Use GetShaderResource(FName,FVector2f, float)")
+	virtual FSlateShaderResourceProxy* GetShaderResource(const FSlateBrush& Brush, FVector2d LocalSize, float DrawScale)
+	{
+		return GetShaderResource(Brush, UE::Slate::CastToVector2f(LocalSize), DrawScale);
+	}
+	virtual FSlateShaderResourceProxy* GetShaderResource(const FSlateBrush& Brush, FVector2f LocalSize, float DrawScale) = 0;
 
 	/**
 	 * Creates a handle to a Slate resource
@@ -65,7 +71,12 @@ public:
 	 * @param	DrawScale	The scale of the draw element using this brush.  This scale is used to rescale vector graphics only
 	 * @return	The created resource handle.  
 	 */
-	virtual FSlateResourceHandle GetResourceHandle(const FSlateBrush& Brush, FVector2D LocalSize, float DrawScale);
+	UE_DEPRECATED(5.1, "Slate rendering uses float instead of double. Use GetResourceHandle(FName,FVector2f, float)")
+	virtual FSlateResourceHandle GetResourceHandle(const FSlateBrush& Brush, FVector2d LocalSize, float DrawScale)
+	{
+		return GetResourceHandle(Brush, UE::Slate::CastToVector2f(LocalSize), DrawScale);
+	}
+	virtual FSlateResourceHandle GetResourceHandle(const FSlateBrush& Brush, FVector2f LocalSize, float DrawScale);
 	virtual FSlateResourceHandle GetResourceHandle(const FSlateBrush& Brush);
 
 	virtual FSlateShaderResource* GetFontShaderResource( int32 InTextureAtlasIndex, FSlateShaderResource* FontTextureAtlas, const class UObject* FontMaterial ) { return FontTextureAtlas; }
diff --git a/Engine/Source/Runtime/SlateCore/Public/Rendering/SlateRenderBatch.h b/Engine/Source/Runtime/SlateCore/Public/Rendering/SlateRenderBatch.h
--- a/Engine/Source/Runtime/SlateCore/Public/Rendering/SlateRenderBatch.h
+++ b/Engine/Source/Runtime/SlateCore/Public/Rendering/SlateRenderBatch.h
@@ -132,7 +132,7 @@ public:
 	FShaderParams ShaderParams;
 
 	/** Dynamically modified offset that occurs when we have relative position stored render batches. */
-	FVector2D DynamicOffset;
+	FVector2f DynamicOffset;
 
 	/** The Stored clipping state for the corresponding clipping state index.  The indices are not directly comparable later, so we need to expand it to the full state to be compared. */
 	const FSlateClippingState* ClippingState;
diff --git a/Engine/Source/Runtime/SlateCore/Public/Rendering/SlateRenderer.h b/Engine/Source/Runtime/SlateCore/Public/Rendering/SlateRenderer.h
--- a/Engine/Source/Runtime/SlateCore/Public/Rendering/SlateRenderer.h
+++ b/Engine/Source/Runtime/SlateCore/Public/Rendering/SlateRenderer.h
@@ -8,6 +8,7 @@
 #include "Rendering/DrawElements.h"
 #include "Templates/RefCounting.h"
 #include "Fonts/FontTypes.h"
+#include "Types/SlateVector2.h"
 #include "PixelFormat.h"
 
 class FRHITexture;
@@ -335,7 +336,12 @@ public:
 	 * @param	DrawScale	The draw scale of the element using the brush (Vector graphics only)
 	 * @return	The created resource handle.
 	 */
-	virtual FSlateResourceHandle GetResourceHandle(const FSlateBrush& Brush, FVector2D LocalSize, float DrawScale) = 0;
+	UE_DEPRECATED(5.1, "Slate rendering uses float instead of double. Use GetResourceHandle(FName,FVector2f, float)")
+	virtual FSlateResourceHandle GetResourceHandle(const FSlateBrush& Brush, FVector2d LocalSize, float DrawScale)
+	{
+		return GetResourceHandle(Brush, UE::Slate::CastToVector2f(LocalSize), DrawScale);
+	}
+	virtual FSlateResourceHandle GetResourceHandle(const FSlateBrush& Brush, FVector2f LocalSize, float DrawScale) = 0;
 
 	/**
 	 * Creates a handle to a Slate resource
@@ -348,7 +354,7 @@ public:
 	 */
 	virtual FSlateResourceHandle GetResourceHandle(const FSlateBrush& Brush)
 	{
-		return GetResourceHandle(Brush, FVector2D::ZeroVector, 1.0f);
+		return GetResourceHandle(Brush, FVector2f::ZeroVector, 1.0f);
 	}
 
 
diff --git a/Engine/Source/Runtime/SlateCore/Public/Rendering/SlateVectorGraphicsCache.h b/Engine/Source/Runtime/SlateCore/Public/Rendering/SlateVectorGraphicsCache.h
--- a/Engine/Source/Runtime/SlateCore/Public/Rendering/SlateVectorGraphicsCache.h
+++ b/Engine/Source/Runtime/SlateCore/Public/Rendering/SlateVectorGraphicsCache.h
@@ -4,6 +4,7 @@
 
 #include "Styling/SlateBrush.h"
 #include "Textures/TextureAtlas.h"
+#include "Types/SlateVector2.h"
 
 class ISlateTextureAtlasFactory;
 class FSlateTextureAtlas;
@@ -15,7 +16,13 @@ class FSlateVectorGraphicsCache : public FSlateFlushableAtlasCache
 public:
 	SLATECORE_API FSlateVectorGraphicsCache(TSharedPtr<ISlateTextureAtlasFactory> InAtlasFactory);
 
-	SLATECORE_API FSlateShaderResourceProxy* GetShaderResource(const FSlateBrush& Brush, FVector2D LocalSize, float DrawScale);
+	UE_DEPRECATED(5.1, "Slate rendering uses float instead of double. Use GetShaderResource(FName,FVector2f, float)")
+	FSlateShaderResourceProxy* GetShaderResource(const FSlateBrush& Brush, FVector2d LocalSize, float DrawScale)
+	{
+		return GetShaderResource(Brush, UE::Slate::CastToVector2f(LocalSize), DrawScale);
+	}
+
+	SLATECORE_API FSlateShaderResourceProxy* GetShaderResource(const FSlateBrush& Brush, FVector2f LocalSize, float DrawScale);
 
 	SLATECORE_API void UpdateCache();
 	SLATECORE_API void ConditionalFlushCache();
@@ -40,7 +47,13 @@ private:
 		FName BrushName;
 		FIntPoint PixelSize;
 
-		FVectorCacheKey(FName InBrushName, FVector2D LocalSize, float DrawScale)
+		UE_DEPRECATED(5.1, "Slate rendering uses float instead of double. Use FVectorCacheKey(FName,FVector2f, float)")
+		FVectorCacheKey(FName InBrushName, FVector2d LocalSize, float DrawScale)
+			: FVectorCacheKey(InBrushName, UE::Slate::CastToVector2f(LocalSize), DrawScale)
+		{
+		}
+
+		FVectorCacheKey(FName InBrushName, FVector2f LocalSize, float DrawScale)
 			: BrushName(InBrushName)
 			, PixelSize((LocalSize*DrawScale).IntPoint())
 		{
@@ -67,7 +80,12 @@ private:
 		FVectorCacheKey Key;
 		TArray<uint8> PixelData;
 
-		FRasterRequest(FName InBrushName, FVector2D InLocalSize, float InDrawScale)
+		UE_DEPRECATED(5.1, "Slate rendering uses float instead of double. Use FRasterRequest(FName,FVector2f, float)")
+		FRasterRequest(FName InBrushName, FVector2d InLocalSize, float InDrawScale)
+			: FRasterRequest(InBrushName, UE::Slate::CastToVector2f(InLocalSize), InDrawScale)
+		{}
+
+		FRasterRequest(FName InBrushName, FVector2f InLocalSize, float InDrawScale)
 			: Key(InBrushName, InLocalSize, InDrawScale)
 		{}
 	};
diff --git a/Engine/Source/Runtime/SlateNullRenderer/Private/SlateNullRenderer.cpp b/Engine/Source/Runtime/SlateNullRenderer/Private/SlateNullRenderer.cpp
--- a/Engine/Source/Runtime/SlateNullRenderer/Private/SlateNullRenderer.cpp
+++ b/Engine/Source/Runtime/SlateNullRenderer/Private/SlateNullRenderer.cpp
@@ -69,7 +69,7 @@ bool FSlateNullRenderer::GenerateDynamicImageResource( FName ResourceName, uint3
 	return false;
 }
 
-FSlateResourceHandle FSlateNullRenderer::GetResourceHandle(const FSlateBrush& Brush, FVector2D LocalSize, float DrawScale)
+FSlateResourceHandle FSlateNullRenderer::GetResourceHandle(const FSlateBrush& Brush, FVector2f LocalSize, float DrawScale)
 {
 	return ResourceManager.IsValid() ? ResourceManager->GetResourceHandle(Brush) : FSlateResourceHandle();
 }
diff --git a/Engine/Source/Runtime/SlateNullRenderer/Private/SlateNullRenderer.h b/Engine/Source/Runtime/SlateNullRenderer/Private/SlateNullRenderer.h
--- a/Engine/Source/Runtime/SlateNullRenderer/Private/SlateNullRenderer.h
+++ b/Engine/Source/Runtime/SlateNullRenderer/Private/SlateNullRenderer.h
@@ -34,7 +34,7 @@ public:
 	virtual FIntPoint GenerateDynamicImageResource(const FName InTextureName) override;
 	virtual bool GenerateDynamicImageResource( FName ResourceName, uint32 Width, uint32 Height, const TArray< uint8 >& Bytes ) override;
 	virtual void PrepareToTakeScreenshot(const FIntRect& Rect, TArray<FColor>* OutColorData, SWindow* InScreenshotWindow) override;
-	virtual FSlateResourceHandle GetResourceHandle(const FSlateBrush& Brush, FVector2D LocalSize, float DrawScale) override;
+	virtual FSlateResourceHandle GetResourceHandle(const FSlateBrush& Brush, FVector2f LocalSize, float DrawScale) override;
 	virtual FSlateUpdatableTexture* CreateUpdatableTexture(uint32 Width, uint32 Height) override;
 	virtual FSlateUpdatableTexture* CreateSharedHandleTexture(void* SharedHandle) override;
 	virtual void ReleaseUpdatableTexture(FSlateUpdatableTexture* Texture) override;
diff --git a/Engine/Source/Runtime/SlateNullRenderer/Private/SlateNullRendererModule.cpp b/Engine/Source/Runtime/SlateNullRenderer/Private/SlateNullRendererModule.cpp
--- a/Engine/Source/Runtime/SlateNullRenderer/Private/SlateNullRendererModule.cpp
+++ b/Engine/Source/Runtime/SlateNullRenderer/Private/SlateNullRendererModule.cpp
@@ -28,7 +28,7 @@ public:
 #endif
 
 	// FSlateShaderResourceManager interface
-	virtual FSlateShaderResourceProxy* GetShaderResource(const FSlateBrush& Brush, FVector2D LocalSize, float DrawScale) override { return nullptr; }
+	virtual FSlateShaderResourceProxy* GetShaderResource(const FSlateBrush& Brush, FVector2f LocalSize, float DrawScale) override { return nullptr; }
 	virtual ISlateAtlasProvider* GetTextureAtlasProvider() { return this; }
 	virtual FSlateResourceHandle GetResourceHandle( const FSlateBrush& InBrush ) override 
 	{
diff --git a/Engine/Source/Runtime/SlateRHIRenderer/Private/SlateRHIRenderer.cpp b/Engine/Source/Runtime/SlateRHIRenderer/Private/SlateRHIRenderer.cpp
--- a/Engine/Source/Runtime/SlateRHIRenderer/Private/SlateRHIRenderer.cpp
+++ b/Engine/Source/Runtime/SlateRHIRenderer/Private/SlateRHIRenderer.cpp
@@ -339,7 +339,7 @@ void FSlateRHIRenderer::CreateViewport(const TSharedRef<SWindow> Window)
 
 	if (!WindowToViewportInfo.Contains(&Window.Get()))
 	{
-		const FVector2D WindowSize = Window->GetViewportSize();
+		const FVector2f WindowSize = UE::Slate::CastToVector2f(Window->GetViewportSize());
 
 		// Clamp the window size to a reasonable default anything below 8 is a d3d warning and 8 is used anyway.
 		// @todo Slate: This is a hack to work around menus being summoned with 0,0 for window size until they are ticked.
@@ -1310,7 +1310,7 @@ void FSlateRHIRenderer::DrawWindows_Private(FSlateDrawBuffer& WindowDrawBuffer)
 
 		if (Window)
 		{
-			const FVector2D WindowSize = Window->GetViewportSize();
+			const FVector2f WindowSize = UE::Slate::CastToVector2f(Window->GetViewportSize());
 			if (WindowSize.X > 0 && WindowSize.Y > 0)
 			{
 				// The viewport need to be created at this point  
@@ -1494,7 +1494,7 @@ bool FSlateRHIRenderer::GenerateDynamicImageResource(FName ResourceName, FSlateT
 	return TextureResource.IsValid();
 }
 
-FSlateResourceHandle FSlateRHIRenderer::GetResourceHandle(const FSlateBrush& Brush, FVector2D LocalSize, float DrawScale)
+FSlateResourceHandle FSlateRHIRenderer::GetResourceHandle(const FSlateBrush& Brush, FVector2f LocalSize, float DrawScale)
 {
 	return ResourceManager->GetResourceHandle(Brush, LocalSize, DrawScale);
 }
diff --git a/Engine/Source/Runtime/SlateRHIRenderer/Private/SlateRHIRenderer.h b/Engine/Source/Runtime/SlateRHIRenderer/Private/SlateRHIRenderer.h
--- a/Engine/Source/Runtime/SlateRHIRenderer/Private/SlateRHIRenderer.h
+++ b/Engine/Source/Runtime/SlateRHIRenderer/Private/SlateRHIRenderer.h
@@ -196,7 +196,7 @@ public:
 	virtual FIntPoint GenerateDynamicImageResource(const FName InTextureName) override;
 	virtual bool GenerateDynamicImageResource( FName ResourceName, uint32 Width, uint32 Height, const TArray< uint8 >& Bytes ) override;
 	virtual bool GenerateDynamicImageResource( FName ResourceName, FSlateTextureDataRef TextureData ) override;
-	virtual FSlateResourceHandle GetResourceHandle(const FSlateBrush& Brush, FVector2D LocalSize, float DrawScale) override;
+	virtual FSlateResourceHandle GetResourceHandle(const FSlateBrush& Brush, FVector2f LocalSize, float DrawScale) override;
 	virtual bool CanRenderResource(UObject& InResourceObject) const override;
 	virtual void* GetViewportResource( const SWindow& Window ) override;
 	virtual void SetColorVisionDeficiencyType(EColorVisionDeficiency Type, int32 Severity, bool bCorrectDeficiency, bool bShowCorrectionWithDeficiency) override;
diff --git a/Engine/Source/Runtime/SlateRHIRenderer/Private/SlateRHIResourceManager.cpp b/Engine/Source/Runtime/SlateRHIRenderer/Private/SlateRHIResourceManager.cpp
--- a/Engine/Source/Runtime/SlateRHIRenderer/Private/SlateRHIResourceManager.cpp
+++ b/Engine/Source/Runtime/SlateRHIRenderer/Private/SlateRHIResourceManager.cpp
@@ -648,7 +648,7 @@ static void LoadUObjectForBrush( const FSlateBrush& InBrush )
 	}
 }
 
-FSlateShaderResourceProxy* FSlateRHIResourceManager::GetShaderResource(const FSlateBrush& InBrush, FVector2D LocalSize, float DrawScale)
+FSlateShaderResourceProxy* FSlateRHIResourceManager::GetShaderResource(const FSlateBrush& InBrush, FVector2f LocalSize, float DrawScale)
 {
 	SCOPE_CYCLE_COUNTER( STAT_SlateGetResourceTime );
 
@@ -937,7 +937,7 @@ FSlateMaterialResource* FSlateRHIResourceManager::GetMaterialResource(const UObj
 	return MaterialResource.Get();
 }
 
-FSlateShaderResourceProxy* FSlateRHIResourceManager::GetVectorResource(const FSlateBrush& Brush, FVector2D LocalSize, float DrawScale)
+FSlateShaderResourceProxy* FSlateRHIResourceManager::GetVectorResource(const FSlateBrush& Brush, FVector2f LocalSize, float DrawScale)
 {
 	return VectorGraphicsCache->GetShaderResource(Brush, LocalSize, DrawScale);
 }
diff --git a/Engine/Source/Runtime/SlateRHIRenderer/Private/SlateRHIResourceManager.h b/Engine/Source/Runtime/SlateRHIRenderer/Private/SlateRHIResourceManager.h
--- a/Engine/Source/Runtime/SlateRHIRenderer/Private/SlateRHIResourceManager.h
+++ b/Engine/Source/Runtime/SlateRHIRenderer/Private/SlateRHIResourceManager.h
@@ -157,7 +157,7 @@ public:
 	void ConditionalFlushAtlases();
 
 	/** FSlateShaderResourceManager interface */
-	virtual FSlateShaderResourceProxy* GetShaderResource(const FSlateBrush& InBrush, FVector2D LocalSize, float DrawScale) override;
+	virtual FSlateShaderResourceProxy* GetShaderResource(const FSlateBrush& InBrush, FVector2f LocalSize, float DrawScale) override;
 	virtual FSlateShaderResource* GetFontShaderResource( int32 InTextureAtlasIndex, FSlateShaderResource* FontTextureAtlas, const class UObject* FontMaterial ) override;
 	virtual ISlateAtlasProvider* GetTextureAtlasProvider() override;
 
@@ -283,7 +283,7 @@ private:
 	 * @param LocalSize	The unscaled local size of the final image
 	 * @param DrawScale	Any scaling applied to the final image
 	 */
-	FSlateShaderResourceProxy* GetVectorResource(const FSlateBrush& Brush, FVector2D LocalSize, float DrawScale);
+	FSlateShaderResourceProxy* GetVectorResource(const FSlateBrush& Brush, FVector2f LocalSize, float DrawScale);
 
 	/**
 	 * Called when the application exists before the UObject system shuts down so we can free object resources
