Slate: Add more info to the GSlateCheckUObjectShapedGlyphSequence test.

diff --git a/Engine/Source/Runtime/SlateCore/Private/Fonts/FontCache.cpp b/Engine/Source/Runtime/SlateCore/Private/Fonts/FontCache.cpp
--- a/Engine/Source/Runtime/SlateCore/Private/Fonts/FontCache.cpp
+++ b/Engine/Source/Runtime/SlateCore/Private/Fonts/FontCache.cpp
@@ -446,11 +446,19 @@ void FShapedGlyphSequence::AddReferencedObjects(FReferenceCollector& Collector)
 		const bool bThreadSafe = true;
 		if (!DebugFontMaterialName.IsNone())
 		{
-			checkf(FontMaterialWeakPtr.IsValid(bEvenIfPendingKill, bThreadSafe), TEXT("Material %s has become invalid. This means the FShapedGlyphSequence::FontMaterial was garbage collected while slate was using it"), *DebugFontMaterialName.ToString());
+			const UObject* FontMaterialPin = FontMaterialWeakPtr.GetEvenIfUnreachable();
+			if (FontMaterial != FontMaterialPin)
+			{
+				UE_LOG(LogSlate, Fatal, TEXT("Material % s has become invalid. This means the FShapedGlyphSequence::FontMaterial was garbage collected while slate was using it"), *DebugFontMaterialName.ToString());
+			}
 		}
 		if (!DebugOutlineMaterialName.IsNone())
 		{
-			checkf(OutlineMaterialWeakPtr.IsValid(bEvenIfPendingKill, bThreadSafe), TEXT("Material %s has become invalid. This means the FShapedGlyphSequence::OutlineSettings::OutlineMaterial was garbage collected while slate was using it"), *DebugOutlineMaterialName.ToString());
+			const UObject* OutlineMaterialPin = OutlineMaterialWeakPtr.GetEvenIfUnreachable();
+			if (OutlineSettings.OutlineMaterial != OutlineMaterialPin)
+			{
+				UE_LOG(LogSlate, Fatal, TEXT("Material %s has become invalid. This means the FShapedGlyphSequence::OutlineSettings::OutlineMaterial was garbage collected while slate was using it"), *DebugOutlineMaterialName.ToString());
+			}
 		}
 	}
 #endif
