Added Tamil to the list of character ranges that require complex text shaping

diff --git a/Engine/Source/Runtime/SlateCore/Private/Fonts/SlateTextShaper.cpp b/Engine/Source/Runtime/SlateCore/Private/Fonts/SlateTextShaper.cpp
--- a/Engine/Source/Runtime/SlateCore/Private/Fonts/SlateTextShaper.cpp
+++ b/Engine/Source/Runtime/SlateCore/Private/Fonts/SlateTextShaper.cpp
@@ -173,6 +173,9 @@ void FSlateTextShaper::PerformTextShaping(const TCHAR* InText, const int32 InTex
 					RETURN_TRUE_IF_CHAR_WITHIN_RANGE(TEXT('\uA8E0'), TEXT('\uA8FF'));
 					RETURN_TRUE_IF_CHAR_WITHIN_RANGE(TEXT('\u1CD0'), TEXT('\u1CFF'));
 
+					// Tamil
+					RETURN_TRUE_IF_CHAR_WITHIN_RANGE(TEXT('\u0B80'), TEXT('\u0BFF'));
+
 					// Telugu
 					RETURN_TRUE_IF_CHAR_WITHIN_RANGE(TEXT('\u0C00'), TEXT('\u0C7F'));
 
