Added Bengali, Gujarati, Odia, Kannada, and Malayalam to the list of character ranges that require complex text shaping

diff --git a/Engine/Source/Runtime/SlateCore/Private/Fonts/SlateTextShaper.cpp b/Engine/Source/Runtime/SlateCore/Private/Fonts/SlateTextShaper.cpp
--- a/Engine/Source/Runtime/SlateCore/Private/Fonts/SlateTextShaper.cpp
+++ b/Engine/Source/Runtime/SlateCore/Private/Fonts/SlateTextShaper.cpp
@@ -173,12 +173,27 @@ void FSlateTextShaper::PerformTextShaping(const TCHAR* InText, const int32 InTex
 					RETURN_TRUE_IF_CHAR_WITHIN_RANGE(TEXT('\uA8E0'), TEXT('\uA8FF'));
 					RETURN_TRUE_IF_CHAR_WITHIN_RANGE(TEXT('\u1CD0'), TEXT('\u1CFF'));
 
+					// Bengali
+					RETURN_TRUE_IF_CHAR_WITHIN_RANGE(TEXT('\u0980'), TEXT('\u09FF'));
+
+					// Gujarati
+					RETURN_TRUE_IF_CHAR_WITHIN_RANGE(TEXT('\u0A80'), TEXT('\u0AFF'));
+
+					// Odia
+					RETURN_TRUE_IF_CHAR_WITHIN_RANGE(TEXT('\u0B00'), TEXT('\u0B7F'));
+
 					// Tamil
 					RETURN_TRUE_IF_CHAR_WITHIN_RANGE(TEXT('\u0B80'), TEXT('\u0BFF'));
 
 					// Telugu
 					RETURN_TRUE_IF_CHAR_WITHIN_RANGE(TEXT('\u0C00'), TEXT('\u0C7F'));
 
+					// Kannada
+					RETURN_TRUE_IF_CHAR_WITHIN_RANGE(TEXT('\u0C80'), TEXT('\u0CFF'));
+
+					// Malayalam
+					RETURN_TRUE_IF_CHAR_WITHIN_RANGE(TEXT('\u0D00'), TEXT('\u0D7F'));
+
 					// Thai
 					RETURN_TRUE_IF_CHAR_WITHIN_RANGE(TEXT('\u0E00'), TEXT('\u0E7F'));
 					
