Add an FInputDeviceId property to slate input events. This will allow us to derive the owning Platform User from each input event, letting you map inputs to specific local platform users or input devices later down the line. This is a small part of the foundation for making this behavior work.

diff --git a/Engine/Source/Runtime/SlateCore/Public/Input/Events.h b/Engine/Source/Runtime/SlateCore/Public/Input/Events.h
--- a/Engine/Source/Runtime/SlateCore/Public/Input/Events.h
+++ b/Engine/Source/Runtime/SlateCore/Public/Input/Events.h
@@ -8,6 +8,7 @@
 #include "InputCoreTypes.h"
 #include "Types/SlateEnums.h"
 #include "GenericPlatform/GenericApplication.h"
+#include "GenericPlatform/GenericPlatformInputDeviceMapper.h"
 #include "Layout/Geometry.h"
 #include "Events.generated.h"
 
@@ -162,6 +163,7 @@
 		: ModifierKeys(FModifierKeysState())
 		, bIsRepeat(false)
 		, UserIndex(0)
+		, InputDeviceId(INPUTDEVICEID_NONE)
 		, EventPath(nullptr)
 	{ }
 
@@ -178,6 +180,16 @@
 		, EventPath(nullptr)
 	{ }
 
+	FInputEvent(const FModifierKeysState& InModifierKeys, const FInputDeviceId InDeviceId, const bool bInIsRepeat)
+		: ModifierKeys(InModifierKeys)
+		, bIsRepeat(bInIsRepeat)
+		, InputDeviceId(InDeviceId)
+		, EventPath(nullptr)
+	{
+		// Set the User Index to the PlatformUser ID by default for backwards compatibility
+		UserIndex = GetPlatformUserId().GetInternalId();
+	}
+
 	/**
 	 * Virtual destructor.
 	 */
@@ -317,6 +329,22 @@
 		return UserIndex;
 	}
 
+	/**
+	 * Returns the input device that caused this event.
+	 */
+	FInputDeviceId GetInputDeviceId() const
+	{
+		return InputDeviceId;
+	}
+
+	/**
+	 * Returns the associated platform user that caused this event
+	 */
+	FPlatformUserId GetPlatformUserId() const
+	{
+		return IPlatformInputDeviceMapper::Get().GetUserForInputDevice(InputDeviceId);
+	}
+
 	/** The event path provides additional context for handling */
 	SLATECORE_API FGeometry FindGeometry(const TSharedRef<SWidget>& WidgetToFind) const;
 
@@ -351,6 +379,9 @@
 
 	// The index of the user that caused the event.
 	uint32 UserIndex;
+	
+	// The ID of the input device that caused this event.
+	FInputDeviceId InputDeviceId;
 
 	// Events are sent along paths. See (GetEventPath).
 	const FWidgetPath* EventPath;
@@ -405,6 +436,19 @@
 		, CharacterCode(InCharacterCode)
 		, KeyCode(InKeyCode)
 	{ }
+	
+	FKeyEvent(const FKey InKey,
+			const FModifierKeysState& InModifierKeys, 
+			const FInputDeviceId InDeviceId,
+			const bool bInIsRepeat,
+			const uint32 InCharacterCode,
+			const uint32 InKeyCode
+	)
+		: FInputEvent(InModifierKeys, InDeviceId, bInIsRepeat)
+		, Key(InKey)
+		, CharacterCode(InCharacterCode)
+		, KeyCode(InKeyCode)
+	{ }
 
 	/**
 	 * Returns the name of the key for this event
@@ -500,6 +544,18 @@
 		, AnalogValue(InAnalogValue)
 	{ }
 
+	FAnalogInputEvent(const FKey InKey,
+		const FModifierKeysState& InModifierKeys,
+		const FInputDeviceId InDeviceId,
+		const bool bInIsRepeat,
+		const uint32 InCharacterCode,
+		const uint32 InKeyCode,
+		const float InAnalogValue
+	)
+		: FKeyEvent(InKey, InModifierKeys, InDeviceId, bInIsRepeat, InCharacterCode, InKeyCode)
+		, AnalogValue(InAnalogValue)
+	{ }
+
 	/**
 	 * Returns the analog value between 0 and 1.
 	 * 0 being not pressed at all, 1 being fully pressed.
@@ -549,6 +605,11 @@
 		, Character(InCharacter)
 	{ }
 
+	FCharacterEvent(const TCHAR InCharacter, const FModifierKeysState& InModifierKeys, const FInputDeviceId InDeviceId, const bool bInIsRepeat)
+		: FInputEvent(InModifierKeys, InDeviceId, bInIsRepeat)
+		, Character(InCharacter)
+	{ }
+
 	/**
 	 * Returns the character for this event
 	 *
@@ -692,6 +753,33 @@
 	{ }
 
 	FPointerEvent(
+		FInputDeviceId InDeviceId,
+		uint32 InPointerIndex,
+		const FVector2D& InScreenSpacePosition,
+		const FVector2D& InLastScreenSpacePosition,
+		const TSet<FKey>& InPressedButtons,
+		FKey InEffectingButton,
+		float InWheelDelta,
+		const FModifierKeysState& InModifierKeys
+	)
+		: FInputEvent(InModifierKeys, InDeviceId, false)
+		, ScreenSpacePosition(InScreenSpacePosition)
+		, LastScreenSpacePosition(InLastScreenSpacePosition)
+		, CursorDelta(InScreenSpacePosition - InLastScreenSpacePosition)
+		, PressedButtons(&InPressedButtons)
+		, EffectingButton(InEffectingButton)
+		, PointerIndex(InPointerIndex)
+		, TouchpadIndex(0)
+		, Force(1.0f)
+		, bIsTouchEvent(false)
+		, GestureType(EGestureEvent::None)
+		, WheelOrGestureDelta(0.0f, InWheelDelta)
+		, bIsDirectionInvertedFromDevice(false)
+		, bIsTouchForceChanged(false)
+		, bIsTouchFirstMove(false)
+	{ }
+
+	FPointerEvent(
 		uint32 InUserIndex,
 		uint32 InPointerIndex,
 		const FVector2D& InScreenSpacePosition,
@@ -975,6 +1063,20 @@
 		, Acceleration(InAcceleration)
 	{ }
 
+	FMotionEvent(
+		const FInputDeviceId InDeviceId,
+		const FVector& InTilt, 
+		const FVector& InRotationRate, 
+		const FVector& InGravity, 
+		const FVector& InAcceleration
+	)
+		: FInputEvent(FModifierKeysState(), InDeviceId, false)
+		, Tilt(InTilt)
+		, RotationRate(InRotationRate)
+		, Gravity(InGravity)
+		, Acceleration(InAcceleration)
+	{ }
+
 public:
 
 	/** Returns the index of the user that caused the event */
@@ -1042,6 +1144,12 @@
 		, NavigationGenesis(InNavigationGenesis)
 	{ }
 
+	FNavigationEvent(const FModifierKeysState& InModifierKeys, const FInputDeviceId InDeviceId, EUINavigation InNavigationType, ENavigationGenesis InNavigationGenesis)
+		: FInputEvent(InModifierKeys, InDeviceId, false)
+		, NavigationType(InNavigationType)
+		, NavigationGenesis(InNavigationGenesis)
+	{ }
+
 public:
 
 	/** Returns the type of navigation request (Left, Right, Up, Down) */
