Slate: Deprecate FSlateFontCache::GetOverflowEllipsisText. The glyph was cached forever. If the glyph contains a material, that material can be invalid/GCed. Reusing that same font later (with the same material but with a different instance) would reuse the same cached glyph with a dangling material pointer.

diff --git a/Engine/Source/Runtime/Slate/Private/Framework/Text/ShapedTextCache.cpp b/Engine/Source/Runtime/Slate/Private/Framework/Text/ShapedTextCache.cpp
--- a/Engine/Source/Runtime/Slate/Private/Framework/Text/ShapedTextCache.cpp
+++ b/Engine/Source/Runtime/Slate/Private/Framework/Text/ShapedTextCache.cpp
@@ -91,6 +91,23 @@ FShapedGlyphSequenceRef FShapedTextCache::FindOrAddShapedText(const FCachedShape
 	return ShapedText.ToSharedRef();
 }
 
+FShapedGlyphSequenceRef FShapedTextCache::FindOrAddOverflowEllipsisText(const float InScale, const FShapedTextContext& InTextContext, const FSlateFontInfo& InFontInfo)
+{
+	FCachedShapedTextKey Key = FCachedShapedTextKey(FTextRange(), InScale, InTextContext, InFontInfo);
+	FShapedGlyphSequencePtr ShapedText = FindShapedText(Key);
+	if (ShapedText)
+	{
+		return ShapedText.ToSharedRef();
+	}
+
+	const TSharedPtr<FSlateFontCache> FontCache = FontCachePtr.Pin();
+	if (FontCache.IsValid())
+	{
+		return FontCache->ShapeOverflowEllipsisText(InFontInfo, InScale);
+	}
+	return MakeShareable(new FShapedGlyphSequence());
+}
+
 void FShapedTextCache::Clear()
 {
 	CachedShapedText.Reset();
diff --git a/Engine/Source/Runtime/Slate/Private/Framework/Text/SlateTextRun.cpp b/Engine/Source/Runtime/Slate/Private/Framework/Text/SlateTextRun.cpp
--- a/Engine/Source/Runtime/Slate/Private/Framework/Text/SlateTextRun.cpp
+++ b/Engine/Source/Runtime/Slate/Private/Framework/Text/SlateTextRun.cpp
@@ -115,8 +115,7 @@ int32 FSlateTextRun::OnPaint(const FPaintArgs& PaintArgs, const FTextArgs& TextA
 	FTextOverflowArgs OverflowArgs;
 	if (TextArgs.OverflowPolicy == ETextOverflowPolicy::Ellipsis && TextArgs.OverflowDirection != ETextOverflowDirection::NoOverflow)
 	{
-		TSharedRef<FSlateFontCache> FontCache = FSlateApplication::Get().GetRenderer()->GetFontCache();
-		OverflowArgs.OverflowTextPtr = FontCache->GetOverflowEllipsisText(Style.Font, AllottedGeometry.GetAccumulatedLayoutTransform().GetScale());
+		OverflowArgs.OverflowTextPtr = BlockTextContext.ShapedTextCache->FindOrAddOverflowEllipsisText(AllottedGeometry.GetAccumulatedLayoutTransform().GetScale(), BlockTextContext, Style.Font);
 		OverflowArgs.OverflowDirection = TextArgs.OverflowDirection;
 		OverflowArgs.bForceEllipsisDueToClippedLine = TextArgs.bForceEllipsisDueToClippedLine;
 	}
diff --git a/Engine/Source/Runtime/Slate/Public/Framework/Text/ShapedTextCache.h b/Engine/Source/Runtime/Slate/Public/Framework/Text/ShapedTextCache.h
--- a/Engine/Source/Runtime/Slate/Public/Framework/Text/ShapedTextCache.h
+++ b/Engine/Source/Runtime/Slate/Public/Framework/Text/ShapedTextCache.h
@@ -95,6 +95,14 @@ public:
 	FShapedGlyphSequenceRef FindOrAddShapedText(const FCachedShapedTextKey& InKey, const TCHAR* InText);
 	FShapedGlyphSequenceRef FindOrAddShapedText(const FCachedShapedTextKey& InKey, const TCHAR* InText, const TextBiDi::ETextDirection InTextDirection);
 
+	/**
+	 * Try and find an existing overflow shaped text instance, or add a new entry to the cache if one cannot be found.
+	 * The overflow sequence is used to replace characters that are clipped.
+	 *
+	 * @return The shaped text instance
+	 */
+	FShapedGlyphSequenceRef FindOrAddOverflowEllipsisText(const float InScale, const FShapedTextContext& InTextContext, const FSlateFontInfo& InFontInfo);
+
 	/**
 	 * Clear this cache
 	 */
diff --git a/Engine/Source/Runtime/SlateCore/Private/Fonts/FontCache.cpp b/Engine/Source/Runtime/SlateCore/Private/Fonts/FontCache.cpp
--- a/Engine/Source/Runtime/SlateCore/Private/Fonts/FontCache.cpp
+++ b/Engine/Source/Runtime/SlateCore/Private/Fonts/FontCache.cpp
@@ -902,7 +902,6 @@ FSlateFontCache::~FSlateFontCache()
 
 	// Make sure things get destroyed in the correct order
 	FontToCharacterListCache.Empty();
-	FontToOverflowGlyphSequence.Empty();
 	ShapedGlyphToAtlasData.Empty();
 	TextShaper.Reset();
 	FontRenderer.Reset();
@@ -1142,18 +1141,12 @@ FShapedGlyphFontAtlasData FSlateFontCache::GetShapedGlyphFontAtlasData( const FS
 
 FShapedGlyphSequenceRef FSlateFontCache::GetOverflowEllipsisText(const FSlateFontInfo& InFontInfo, const float InFontScale)
 {
-	FSlateFontKey Key(InFontInfo, InFontInfo.OutlineSettings, InFontScale);
-
-	FShapedGlyphSequenceRef* FoundSequence = FontToOverflowGlyphSequence.Find(Key);
-
-	if (FoundSequence)
-	{
-		return *FoundSequence;
-	}
-
-	FShapedGlyphSequenceRef NewSequence = ShapeBidirectionalText(EllipsisText.ToString(), InFontInfo, InFontScale, TextBiDi::ETextDirection::LeftToRight, GetDefaultTextShapingMethod());
+	return ShapeOverflowEllipsisText(InFontInfo, InFontScale);
+}
 
-	return FontToOverflowGlyphSequence.Add(Key, NewSequence);
+FShapedGlyphSequenceRef FSlateFontCache::ShapeOverflowEllipsisText(const FSlateFontInfo& InFontInfo, const float InFontScale)
+{
+	return ShapeBidirectionalText(EllipsisText.ToString(), InFontInfo, InFontScale, TextBiDi::ETextDirection::LeftToRight, GetDefaultTextShapingMethod());
 }
 
 const FFontData& FSlateFontCache::GetDefaultFontData( const FSlateFontInfo& InFontInfo ) const
@@ -1337,7 +1330,6 @@ void FSlateFontCache::FlushData()
 	}
 
 	FontToCharacterListCache.Empty();
-	FontToOverflowGlyphSequence.Empty();
 
 	ShapedGlyphToAtlasData.Empty();
 }
diff --git a/Engine/Source/Runtime/SlateCore/Public/Fonts/FontCache.h b/Engine/Source/Runtime/SlateCore/Public/Fonts/FontCache.h
--- a/Engine/Source/Runtime/SlateCore/Public/Fonts/FontCache.h
+++ b/Engine/Source/Runtime/SlateCore/Public/Fonts/FontCache.h
@@ -676,6 +676,11 @@ public:
 	FShapedGlyphSequenceRef ShapeUnidirectionalText( const FString& InText, const FSlateFontInfo &InFontInfo, const float InFontScale, const TextBiDi::ETextDirection InTextDirection, const ETextShapingMethod InTextShapingMethod ) const;
 	FShapedGlyphSequenceRef ShapeUnidirectionalText( const TCHAR* InText, const int32 InTextStart, const int32 InTextLen, const FSlateFontInfo &InFontInfo, const float InFontScale, const TextBiDi::ETextDirection InTextDirection, const ETextShapingMethod InTextShapingMethod ) const;
 
+	/**
+	 * Performs text shaping on the overflow glyph sequence for a given font. The overflow sequence is used to replace characters that are clipped
+	 */
+	FShapedGlyphSequenceRef ShapeOverflowEllipsisText(const FSlateFontInfo& InFontInfo, const float InFontScale);
+
 	/** 
 	 * Gets information for how to draw all non-shaped characters in the specified string. Caches characters as they are found
 	 * 
@@ -693,6 +698,7 @@ public:
 	/**
 	 * Gets the overflow glyph sequence for a given font. The overflow sequence is used to replace characters that are clipped
 	 */
+	UE_DEPRECATED(5.1, "GetOverflowEllipsisText is known to create dangling pointer. Use FShapedTextCache::FindOrAddOverflowEllipsisText.")
 	FShapedGlyphSequenceRef GetOverflowEllipsisText(const FSlateFontInfo& InFontInfo, const float InFontScale);
 
 public:
@@ -901,9 +907,6 @@ private:
 	/** Mapping Font keys to cached data */
 	TMap<FSlateFontKey, TUniquePtr<FCharacterList>, FDefaultSetAllocator, FSlateFontKeyFuncs<TUniquePtr<FCharacterList>>> FontToCharacterListCache;
 
-	/** Caches overflow text to display usually an ellipsis character for text elements that are clipped and request replacing clipped text with an ellipsis */
-	TMap<FSlateFontKey, FShapedGlyphSequenceRef, FDefaultSetAllocator, FSlateFontKeyFuncs<FShapedGlyphSequenceRef>> FontToOverflowGlyphSequence;
-
 	/** Mapping shaped glyphs to their cached atlas data */
 	TMap<FShapedGlyphEntryKey, TSharedRef<FShapedGlyphFontAtlasData>> ShapedGlyphToAtlasData;
 
