Added skew transform to slate fonts to mimic italic. The font atlas is loading the extra glyphs.
#rnx
#tests: Ran a preflight build

diff --git a/Engine/Source/Editor/DetailCustomizations/Private/SlateFontInfoCustomization.cpp b/Engine/Source/Editor/DetailCustomizations/Private/SlateFontInfoCustomization.cpp
--- a/Engine/Source/Editor/DetailCustomizations/Private/SlateFontInfoCustomization.cpp
+++ b/Engine/Source/Editor/DetailCustomizations/Private/SlateFontInfoCustomization.cpp
@@ -92,6 +92,8 @@ void FSlateFontInfoStructCustomization::CustomizeChildren(TSharedRef<IPropertyHa
 	InStructBuilder.AddProperty(FontSizeProperty.ToSharedRef());
 
 	InStructBuilder.AddProperty(InStructPropertyHandle->GetChildHandle(GET_MEMBER_NAME_CHECKED(FSlateFontInfo, LetterSpacing)).ToSharedRef());
+	
+	InStructBuilder.AddProperty(InStructPropertyHandle->GetChildHandle(GET_MEMBER_NAME_CHECKED(FSlateFontInfo, SkewAmount)).ToSharedRef());
 
 	InStructBuilder.AddProperty(InStructPropertyHandle->GetChildHandle(GET_MEMBER_NAME_CHECKED(FSlateFontInfo, FontMaterial)).ToSharedRef());
 
diff --git a/Engine/Source/Runtime/SlateCore/Private/Fonts/FontCache.cpp b/Engine/Source/Runtime/SlateCore/Private/Fonts/FontCache.cpp
--- a/Engine/Source/Runtime/SlateCore/Private/Fonts/FontCache.cpp
+++ b/Engine/Source/Runtime/SlateCore/Private/Fonts/FontCache.cpp
@@ -100,6 +100,7 @@ FShapedGlyphEntryKey::FShapedGlyphEntryKey(const FShapedGlyphFaceData& InFontFac
 	, FontScale(InFontFaceData.FontScale)
 	, GlyphIndex(InGlyphIndex)
 	, KeyHash(0)
+	, FontSkew(InFontFaceData.FontSkew)
 {
 	KeyHash = HashCombine(KeyHash, GetTypeHash(FontFace));
 	KeyHash = HashCombine(KeyHash, GetTypeHash(FontSize));
@@ -107,6 +108,7 @@ FShapedGlyphEntryKey::FShapedGlyphEntryKey(const FShapedGlyphFaceData& InFontFac
 	KeyHash = HashCombine(KeyHash, GetTypeHash(OutlineSeparateFillAlpha));
 	KeyHash = HashCombine(KeyHash, GetTypeHash(FontScale));
 	KeyHash = HashCombine(KeyHash, GetTypeHash(GlyphIndex));
+	KeyHash = HashCombine(KeyHash, GetTypeHash(FontSkew));
 }
 
 FShapedGlyphSequence::FShapedGlyphSequence(TArray<FShapedGlyphEntry> InGlyphsToRender, const int16 InTextBaseline, const uint16 InMaxTextHeight, const UObject* InFontMaterial, const FFontOutlineSettings& InOutlineSettings, const FSourceTextRange& InSourceTextRange)
@@ -810,7 +812,7 @@ FCharacterList::FCharacterListEntry* FCharacterList::CacheCharacter(TCHAR Charac
 			}
 
 			FCharacterListEntry NewInternalEntry;
-			NewInternalEntry.ShapedGlyphEntry.FontFaceData = MakeShared<FShapedGlyphFaceData>(FaceGlyphData.FaceAndMemory, GlyphFlags, FontInfo.Size, FinalFontScale);
+			NewInternalEntry.ShapedGlyphEntry.FontFaceData = MakeShared<FShapedGlyphFaceData>(FaceGlyphData.FaceAndMemory, GlyphFlags, FontInfo.Size, FinalFontScale, FontInfo.SkewAmount);
 			NewInternalEntry.ShapedGlyphEntry.GlyphIndex = GlyphIndex;
 			NewInternalEntry.ShapedGlyphEntry.XAdvance = XAdvance;
 			NewInternalEntry.ShapedGlyphEntry.bIsVisible = !bIsWhitespace;
diff --git a/Engine/Source/Runtime/SlateCore/Private/Fonts/SlateTextShaper.cpp b/Engine/Source/Runtime/SlateCore/Private/Fonts/SlateTextShaper.cpp
--- a/Engine/Source/Runtime/SlateCore/Private/Fonts/SlateTextShaper.cpp
+++ b/Engine/Source/Runtime/SlateCore/Private/Fonts/SlateTextShaper.cpp
@@ -281,6 +281,22 @@ FShapedGlyphSequenceRef FSlateTextShaper::FinalizeTextShaping(TArray<FShapedGlyp
 
 		if (FaceGlyphData.FaceAndMemory.IsValid() && FaceGlyphData.FaceAndMemory->IsFaceValid())
 		{
+
+			if (FMath::IsNearlyEqual(InFontInfo.SkewAmount, 0.f)) 
+			{
+				FT_Set_Transform(FaceGlyphData.FaceAndMemory->GetFace(), nullptr, nullptr);
+			}
+			else
+			{
+				// Skewing / Fake Italics (could do character rotation in future?).
+				FT_Matrix TransformMatrix;
+				TransformMatrix.xx = 0x10000L;
+				TransformMatrix.xy = InFontInfo.SkewAmount * 0x10000L;
+				TransformMatrix.yx = 0;
+				TransformMatrix.yy = 0x10000L;
+				FT_Set_Transform(FaceGlyphData.FaceAndMemory->GetFace(), &TransformMatrix, nullptr);
+			}
+
 			FreeTypeUtils::ApplySizeAndScale(FaceGlyphData.FaceAndMemory->GetFace(), InFontInfo.Size, InFontScale);
 			
 			TextBaseline = FreeTypeUtils::Convert26Dot6ToRoundedPixel<int16>(FaceGlyphData.FaceAndMemory->GetDescender());
@@ -389,7 +405,7 @@ void FSlateTextShaper::PerformKerningOnlyTextShaping(const TCHAR* InText, const
 			const int16 LetterSpacingScaled = (int16)LetterSpacingScaledAsFloat;
 
 			FreeTypeUtils::ApplySizeAndScale(KerningOnlyTextSequenceEntry.FaceAndMemory->GetFace(), InFontInfo.Size, FinalFontScale);
-			TSharedRef<FShapedGlyphFaceData> ShapedGlyphFaceData = MakeShared<FShapedGlyphFaceData>(KerningOnlyTextSequenceEntry.FaceAndMemory, GlyphFlags, InFontInfo.Size, FinalFontScale);
+			TSharedRef<FShapedGlyphFaceData> ShapedGlyphFaceData = MakeShared<FShapedGlyphFaceData>(KerningOnlyTextSequenceEntry.FaceAndMemory, GlyphFlags, InFontInfo.Size, FinalFontScale, InFontInfo.SkewAmount);
 			TSharedPtr<FFreeTypeKerningCache> KerningCache = FTCacheDirectory->GetKerningCache(KerningOnlyTextSequenceEntry.FaceAndMemory->GetFace(), FT_KERNING_DEFAULT, InFontInfo.Size, FinalFontScale);
 			TSharedRef<FFreeTypeAdvanceCache> AdvanceCache = FTCacheDirectory->GetAdvanceCache(KerningOnlyTextSequenceEntry.FaceAndMemory->GetFace(), GlyphFlags, InFontInfo.Size, FinalFontScale);
 
@@ -668,7 +684,7 @@ void FSlateTextShaper::PerformHarfBuzzTextShaping(const TCHAR* InText, const int
 			const float FinalFontScale = InFontScale * HarfBuzzTextSequenceEntry.SubFontScalingFactor;
 
 			hb_font_t* HarfBuzzFont = HarfBuzzFontFactory.CreateFont(*HarfBuzzTextSequenceEntry.FaceAndMemory, GlyphFlags, InFontInfo.Size, FinalFontScale);
-			TSharedRef<FShapedGlyphFaceData> ShapedGlyphFaceData = MakeShared<FShapedGlyphFaceData>(HarfBuzzTextSequenceEntry.FaceAndMemory, GlyphFlags, InFontInfo.Size, FinalFontScale);
+			TSharedRef<FShapedGlyphFaceData> ShapedGlyphFaceData = MakeShared<FShapedGlyphFaceData>(HarfBuzzTextSequenceEntry.FaceAndMemory, GlyphFlags, InFontInfo.Size, FinalFontScale, InFontInfo.SkewAmount);
 			TSharedPtr<FFreeTypeKerningCache> KerningCache = FTCacheDirectory->GetKerningCache(HarfBuzzTextSequenceEntry.FaceAndMemory->GetFace(), FT_KERNING_DEFAULT, InFontInfo.Size, FinalFontScale);
 			TSharedRef<FFreeTypeAdvanceCache> AdvanceCache = FTCacheDirectory->GetAdvanceCache(HarfBuzzTextSequenceEntry.FaceAndMemory->GetFace(), ShapedGlyphFaceData->GlyphFlags, InFontInfo.Size, FinalFontScale);
 
diff --git a/Engine/Source/Runtime/SlateCore/Private/Fonts/SlateTextShaper.h b/Engine/Source/Runtime/SlateCore/Private/Fonts/SlateTextShaper.h
--- a/Engine/Source/Runtime/SlateCore/Private/Fonts/SlateTextShaper.h
+++ b/Engine/Source/Runtime/SlateCore/Private/Fonts/SlateTextShaper.h
@@ -17,12 +17,13 @@ class IBreakIterator;
 class FShapedGlyphFaceData
 {
 public:
-	FShapedGlyphFaceData(TWeakPtr<FFreeTypeFace> InFontFace, const uint32 InGlyphFlags, const int32 InFontSize, const float InFontScale)
+	FShapedGlyphFaceData(TWeakPtr<FFreeTypeFace> InFontFace, const uint32 InGlyphFlags, const int32 InFontSize, const float InFontScale, const float InFontSkew)
 		: FontFace(MoveTemp(InFontFace))
 		, GlyphFlags(InGlyphFlags)
 		, FontSize(InFontSize)
 		, FontScale(InFontScale)
 		, BitmapRenderScale(1.0f)
+		, FontSkew(InFontSkew)
 	{
 #if WITH_FREETYPE
 		if (TSharedPtr<FFreeTypeFace> FontFacePin = FontFace.Pin())
@@ -42,6 +43,8 @@ public:
 	float FontScale;
 	/** Any additional scale that should be applied when rendering this glyph */
 	float BitmapRenderScale;
+	/** The skew transform amount for the rendered font */
+	float FontSkew;
 };
 
 
diff --git a/Engine/Source/Runtime/SlateCore/Public/Fonts/FontCache.h b/Engine/Source/Runtime/SlateCore/Public/Fonts/FontCache.h
--- a/Engine/Source/Runtime/SlateCore/Public/Fonts/FontCache.h
+++ b/Engine/Source/Runtime/SlateCore/Public/Fonts/FontCache.h
@@ -164,7 +164,8 @@ public:
 			&& OutlineSize == Other.OutlineSize
 			&& OutlineSeparateFillAlpha == Other.OutlineSeparateFillAlpha
 			&& FontScale == Other.FontScale
-			&& GlyphIndex == Other.GlyphIndex;
+			&& GlyphIndex == Other.GlyphIndex
+			&& FontSkew == Other.FontSkew;
 	}
 
 	FORCEINLINE bool operator!=(const FShapedGlyphEntryKey& Other) const
@@ -192,6 +193,8 @@ private:
 	uint32 GlyphIndex;
 	/** Cached hash value used for map lookups */
 	uint32 KeyHash;
+	/** The skew transform amount for the rendered font */
+	float FontSkew;
 };
 
 /** Information for rendering a shaped text sequence */
diff --git a/Engine/Source/Runtime/SlateCore/Public/Fonts/SlateFontInfo.h b/Engine/Source/Runtime/SlateCore/Public/Fonts/SlateFontInfo.h
--- a/Engine/Source/Runtime/SlateCore/Public/Fonts/SlateFontInfo.h
+++ b/Engine/Source/Runtime/SlateCore/Public/Fonts/SlateFontInfo.h
@@ -164,6 +164,10 @@ struct SLATECORE_API FSlateFontInfo
 	UPROPERTY(EditAnywhere, BlueprintReadWrite, Category=SlateStyleRules, meta=(ClampMin=-1000, ClampMax=10000))
 	int32 LetterSpacing = 0;
 
+	/** A skew amount to apply to the text. */
+	UPROPERTY(EditAnywhere, BlueprintReadWrite, Category=SlateStyleRules)
+	float SkewAmount = 0.0f;
+
 	/** The font fallback level. Runtime only, don't set on shared FSlateFontInfo, as it may change the font elsewhere (make a copy). */
 	EFontFallback FontFallback;
 
@@ -246,7 +250,8 @@ public:
 			&& OutlineSettings.IsIdenticalToForCaching(Other.OutlineSettings)
 			&& CompositeFont == Other.CompositeFont
 			&& TypefaceFontName == Other.TypefaceFontName
-			&& Size == Other.Size;
+			&& Size == Other.Size
+			&& SkewAmount == Other.SkewAmount;
 	}
 
 	inline bool IsIdenticalTo(const FSlateFontInfo& Other) const
@@ -257,7 +262,8 @@ public:
 			&& CompositeFont == Other.CompositeFont
 			&& TypefaceFontName == Other.TypefaceFontName
 			&& Size == Other.Size
-			&& LetterSpacing == Other.LetterSpacing;
+			&& LetterSpacing == Other.LetterSpacing
+			&& SkewAmount == Other.SkewAmount;
 	}
 
 	inline bool operator==(const FSlateFontInfo& Other) const
@@ -296,6 +302,7 @@ public:
 		Hash = HashCombine(Hash, GetTypeHash(FontInfo.CompositeFont));
 		Hash = HashCombine(Hash, GetTypeHash(FontInfo.TypefaceFontName));
 		Hash = HashCombine(Hash, GetTypeHash(FontInfo.Size));
+		Hash = HashCombine(Hash, GetTypeHash(FontInfo.SkewAmount));
 		return Hash;
 	}
 
