Clamp font skew values
#rnx

diff --git a/Engine/Source/Runtime/SlateCore/Private/Fonts/FontCache.cpp b/Engine/Source/Runtime/SlateCore/Private/Fonts/FontCache.cpp
--- a/Engine/Source/Runtime/SlateCore/Private/Fonts/FontCache.cpp
+++ b/Engine/Source/Runtime/SlateCore/Private/Fonts/FontCache.cpp
@@ -812,7 +812,7 @@ FCharacterList::FCharacterListEntry* FCharacterList::CacheCharacter(TCHAR Charac
 			}
 
 			FCharacterListEntry NewInternalEntry;
-			NewInternalEntry.ShapedGlyphEntry.FontFaceData = MakeShared<FShapedGlyphFaceData>(FaceGlyphData.FaceAndMemory, GlyphFlags, FontInfo.Size, FinalFontScale, FontInfo.SkewAmount);
+			NewInternalEntry.ShapedGlyphEntry.FontFaceData = MakeShared<FShapedGlyphFaceData>(FaceGlyphData.FaceAndMemory, GlyphFlags, FontInfo.Size, FinalFontScale, FontInfo.GetClampSkew());
 			NewInternalEntry.ShapedGlyphEntry.GlyphIndex = GlyphIndex;
 			NewInternalEntry.ShapedGlyphEntry.XAdvance = XAdvance;
 			NewInternalEntry.ShapedGlyphEntry.bIsVisible = !bIsWhitespace;
diff --git a/Engine/Source/Runtime/SlateCore/Private/Fonts/SlateFontInfo.cpp b/Engine/Source/Runtime/SlateCore/Private/Fonts/SlateFontInfo.cpp
--- a/Engine/Source/Runtime/SlateCore/Private/Fonts/SlateFontInfo.cpp
+++ b/Engine/Source/Runtime/SlateCore/Private/Fonts/SlateFontInfo.cpp
@@ -228,3 +228,8 @@ uint16 FSlateFontInfo::GetClampSize() const
 {
 	return (uint16)FMath::Clamp<int32>(Size, 0, std::numeric_limits<uint16>::max());
 }
+
+float FSlateFontInfo::GetClampSkew() const
+{
+	return FMath::Clamp(SkewAmount, -5.f, 5.f);
+}
\ No newline at end of file
diff --git a/Engine/Source/Runtime/SlateCore/Private/Fonts/SlateTextShaper.cpp b/Engine/Source/Runtime/SlateCore/Private/Fonts/SlateTextShaper.cpp
--- a/Engine/Source/Runtime/SlateCore/Private/Fonts/SlateTextShaper.cpp
+++ b/Engine/Source/Runtime/SlateCore/Private/Fonts/SlateTextShaper.cpp
@@ -282,7 +282,7 @@ FShapedGlyphSequenceRef FSlateTextShaper::FinalizeTextShaping(TArray<FShapedGlyp
 		if (FaceGlyphData.FaceAndMemory.IsValid() && FaceGlyphData.FaceAndMemory->IsFaceValid())
 		{
 
-			if (FMath::IsNearlyEqual(InFontInfo.SkewAmount, 0.f)) 
+			if (FMath::IsNearlyEqual(InFontInfo.GetClampSkew(), 0.f))
 			{
 				FT_Set_Transform(FaceGlyphData.FaceAndMemory->GetFace(), nullptr, nullptr);
 			}
@@ -291,7 +291,7 @@ FShapedGlyphSequenceRef FSlateTextShaper::FinalizeTextShaping(TArray<FShapedGlyp
 				// Skewing / Fake Italics (could do character rotation in future?).
 				FT_Matrix TransformMatrix;
 				TransformMatrix.xx = 0x10000L;
-				TransformMatrix.xy = InFontInfo.SkewAmount * 0x10000L;
+				TransformMatrix.xy = InFontInfo.GetClampSkew() * 0x10000L;
 				TransformMatrix.yx = 0;
 				TransformMatrix.yy = 0x10000L;
 				FT_Set_Transform(FaceGlyphData.FaceAndMemory->GetFace(), &TransformMatrix, nullptr);
@@ -405,7 +405,7 @@ void FSlateTextShaper::PerformKerningOnlyTextShaping(const TCHAR* InText, const
 			const int16 LetterSpacingScaled = (int16)LetterSpacingScaledAsFloat;
 
 			FreeTypeUtils::ApplySizeAndScale(KerningOnlyTextSequenceEntry.FaceAndMemory->GetFace(), InFontInfo.Size, FinalFontScale);
-			TSharedRef<FShapedGlyphFaceData> ShapedGlyphFaceData = MakeShared<FShapedGlyphFaceData>(KerningOnlyTextSequenceEntry.FaceAndMemory, GlyphFlags, InFontInfo.Size, FinalFontScale, InFontInfo.SkewAmount);
+			TSharedRef<FShapedGlyphFaceData> ShapedGlyphFaceData = MakeShared<FShapedGlyphFaceData>(KerningOnlyTextSequenceEntry.FaceAndMemory, GlyphFlags, InFontInfo.Size, FinalFontScale, InFontInfo.GetClampSkew());
 			TSharedPtr<FFreeTypeKerningCache> KerningCache = FTCacheDirectory->GetKerningCache(KerningOnlyTextSequenceEntry.FaceAndMemory->GetFace(), FT_KERNING_DEFAULT, InFontInfo.Size, FinalFontScale);
 			TSharedRef<FFreeTypeAdvanceCache> AdvanceCache = FTCacheDirectory->GetAdvanceCache(KerningOnlyTextSequenceEntry.FaceAndMemory->GetFace(), GlyphFlags, InFontInfo.Size, FinalFontScale);
 
@@ -684,7 +684,7 @@ void FSlateTextShaper::PerformHarfBuzzTextShaping(const TCHAR* InText, const int
 			const float FinalFontScale = InFontScale * HarfBuzzTextSequenceEntry.SubFontScalingFactor;
 
 			hb_font_t* HarfBuzzFont = HarfBuzzFontFactory.CreateFont(*HarfBuzzTextSequenceEntry.FaceAndMemory, GlyphFlags, InFontInfo.Size, FinalFontScale);
-			TSharedRef<FShapedGlyphFaceData> ShapedGlyphFaceData = MakeShared<FShapedGlyphFaceData>(HarfBuzzTextSequenceEntry.FaceAndMemory, GlyphFlags, InFontInfo.Size, FinalFontScale, InFontInfo.SkewAmount);
+			TSharedRef<FShapedGlyphFaceData> ShapedGlyphFaceData = MakeShared<FShapedGlyphFaceData>(HarfBuzzTextSequenceEntry.FaceAndMemory, GlyphFlags, InFontInfo.Size, FinalFontScale, InFontInfo.GetClampSkew());
 			TSharedPtr<FFreeTypeKerningCache> KerningCache = FTCacheDirectory->GetKerningCache(HarfBuzzTextSequenceEntry.FaceAndMemory->GetFace(), FT_KERNING_DEFAULT, InFontInfo.Size, FinalFontScale);
 			TSharedRef<FFreeTypeAdvanceCache> AdvanceCache = FTCacheDirectory->GetAdvanceCache(HarfBuzzTextSequenceEntry.FaceAndMemory->GetFace(), ShapedGlyphFaceData->GlyphFlags, InFontInfo.Size, FinalFontScale);
 
diff --git a/Engine/Source/Runtime/SlateCore/Public/Fonts/SlateFontInfo.h b/Engine/Source/Runtime/SlateCore/Public/Fonts/SlateFontInfo.h
--- a/Engine/Source/Runtime/SlateCore/Public/Fonts/SlateFontInfo.h
+++ b/Engine/Source/Runtime/SlateCore/Public/Fonts/SlateFontInfo.h
@@ -165,7 +165,7 @@ struct SLATECORE_API FSlateFontInfo
 	int32 LetterSpacing = 0;
 
 	/** A skew amount to apply to the text. */
-	UPROPERTY(EditAnywhere, BlueprintReadWrite, Category=SlateStyleRules)
+	UPROPERTY(EditAnywhere, BlueprintReadWrite, Category=SlateStyleRules, meta=(ClampMin=-5, ClampMax=5))
 	float SkewAmount = 0.0f;
 
 	/** The font fallback level. Runtime only, don't set on shared FSlateFontInfo, as it may change the font elsewhere (make a copy). */
@@ -285,6 +285,9 @@ public:
 	/** Get the font size clamp for the font renderer (on 16bits) */
 	uint16 GetClampSize() const;
 
+	/** Get the skew amount clamp for the text shaper */
+	float GetClampSkew() const;
+
 	/**
 	 * Calculates a type hash value for a font info.
 	 *
