UnrealFrontend: Fix missing launch icon and status color

diff --git a/Engine/Source/Developer/SessionFrontend/Private/Widgets/Browser/SSessionBrowserTreeInstanceRow.h b/Engine/Source/Developer/SessionFrontend/Private/Widgets/Browser/SSessionBrowserTreeInstanceRow.h
--- a/Engine/Source/Developer/SessionFrontend/Private/Widgets/Browser/SSessionBrowserTreeInstanceRow.h
+++ b/Engine/Source/Developer/SessionFrontend/Private/Widgets/Browser/SSessionBrowserTreeInstanceRow.h
@@ -277,7 +277,7 @@ private:
 				return FStyleColors::AccentGreen;
 			}
 
-			return FStyleColors::Warning;
+			return FStyleColors::AccentGray;
 		}
 
 		return FStyleColors::Error;
diff --git a/Engine/Source/Editor/EditorStyle/Private/StarshipStyle.cpp b/Engine/Source/Editor/EditorStyle/Private/StarshipStyle.cpp
--- a/Engine/Source/Editor/EditorStyle/Private/StarshipStyle.cpp
+++ b/Engine/Source/Editor/EditorStyle/Private/StarshipStyle.cpp
@@ -395,7 +395,6 @@ void FStarshipEditorStyle::FStyle::SetupGeneralStyles()
 		Set("Icons.SelectInViewport", new IMAGE_BRUSH_SVG("Starship/Common/SelectInViewport", Icon16x16));
 		Set("Icons.BrowseContent", new IMAGE_BRUSH_SVG("Starship/Common/BrowseContent", Icon16x16));
 		Set("Icons.Use", new IMAGE_BRUSH_SVG("Starship/Common/use-circle", Icon16x16));
-		Set("Icons.Launch", new CORE_IMAGE_BRUSH_SVG("Starship/Common/ProjectLauncher", Icon16x16));
 		Set("Icons.Next", new IMAGE_BRUSH_SVG("Starship/Common/NextArrow", Icon16x16));
 		Set("Icons.Previous", new IMAGE_BRUSH_SVG("Starship/Common/PreviousArrow", Icon16x16));
 		Set("Icons.Visibility", new IMAGE_BRUSH_SVG("Starship/Common/Visibility", Icon20x20));
diff --git a/Engine/Source/Runtime/SlateCore/Private/Styling/StarshipCoreStyle.cpp b/Engine/Source/Runtime/SlateCore/Private/Styling/StarshipCoreStyle.cpp
--- a/Engine/Source/Runtime/SlateCore/Private/Styling/StarshipCoreStyle.cpp
+++ b/Engine/Source/Runtime/SlateCore/Private/Styling/StarshipCoreStyle.cpp
@@ -394,6 +394,7 @@ TSharedRef<ISlateStyle> FStarshipCoreStyle::Create()
 		Style->Set("Icons.C++", new IMAGE_BRUSH_SVG("Starship/Common/CPP", Icon16x16));
 
 		Style->Set("Icons.Advanced", new IMAGE_BRUSH_SVG("Starship/Common/Advanced", Icon16x16));
+		Style->Set("Icons.Launch", new IMAGE_BRUSH_SVG("Starship/Common/ProjectLauncher", Icon16x16));
 		Style->Set("Icons.Rotate90Clockwise", new IMAGE_BRUSH_SVG("Starship/Common/Rotate90Clockwise", Icon16x16));
 		Style->Set("Icons.Rotate90CounterClockwise", new IMAGE_BRUSH_SVG("Starship/Common/Rotate90CounterClockwise", Icon16x16));
 		Style->Set("Icons.Rotate180", new IMAGE_BRUSH_SVG("Starship/Common/Rotate180", Icon16x16));
