Cull when stencil AABB & backbuffer intersection results in a zero sized scissor

diff --git a/Engine/Source/Runtime/SlateCore/Private/Rendering/DrawElements.cpp b/Engine/Source/Runtime/SlateCore/Private/Rendering/DrawElements.cpp
--- a/Engine/Source/Runtime/SlateCore/Private/Rendering/DrawElements.cpp
+++ b/Engine/Source/Runtime/SlateCore/Private/Rendering/DrawElements.cpp
@@ -38,7 +38,28 @@ static bool ShouldCull(const FSlateWindowElementList& ElementList)
 	if (CurrentIndex != INDEX_NONE)
 	{
 		const FSlateClippingState& ClippingState = ClippingManager.GetClippingStates()[CurrentIndex];
-		return ClippingState.HasZeroArea();
+		if (ClippingState.GetClippingMethod() == EClippingMethod::Scissor)
+		{
+			return ClippingState.HasZeroArea();
+		}
+		else if (ClippingState.GetClippingMethod() == EClippingMethod::Stencil)
+		{
+			FSlateRect WindowRect = FSlateRect(FVector2f(0, 0), ElementList.GetWindowSize());
+			if (WindowRect.GetArea() > 0)
+			{
+				for (const FSlateClippingZone& Stencil : ClippingState.StencilQuads)
+				{
+					bool bOverlapping = false;
+					FSlateRect ClippedStencil = Stencil.GetBoundingBox().IntersectionWith(WindowRect, bOverlapping);
+					
+					if (!bOverlapping || ClippedStencil.GetArea() <= 0)
+					{
+						return true;
+					}
+				}
+			}
+		}
+
 	}
 
 	return false;
