Slate: Convert SlateRect, Margin and SlateRotatedRect from double to float. Without the deprecation path, it creates a lot of errors. I will deprecate slowly the functionnality that use FVector2D once I convert the other basic structures.

diff --git a/Engine/Source/Editor/CurveEditor/Private/DragOperations/CurveEditorDragOperation_Marquee.cpp b/Engine/Source/Editor/CurveEditor/Private/DragOperations/CurveEditorDragOperation_Marquee.cpp
--- a/Engine/Source/Editor/CurveEditor/Private/DragOperations/CurveEditorDragOperation_Marquee.cpp
+++ b/Engine/Source/Editor/CurveEditor/Private/DragOperations/CurveEditorDragOperation_Marquee.cpp
@@ -68,7 +68,7 @@ void FCurveEditorDragOperation_Marquee::OnEndDrag(FVector2D InitialPosition, FVe
 			const FGeometry& LocalGeometry = View->GetCachedGeometry();
 			FSlateLayoutTransform ContainerToView = InverseContainerTransform.Concatenate(LocalGeometry.GetAccumulatedLayoutTransform()).Inverse();
 
-			FSlateRect UnclippedLocalMarquee = FSlateRect(ContainerToView.TransformPoint(Marquee.GetTopLeft()), ContainerToView.TransformPoint(Marquee.GetBottomRight()));
+			FSlateRect UnclippedLocalMarquee = FSlateRect(ContainerToView.TransformPoint(Marquee.GetTopLeft2f()), ContainerToView.TransformPoint(Marquee.GetBottomRight2f()));
 			FSlateRect ClippedLocalMarquee = UnclippedLocalMarquee.IntersectionWith(FSlateRect(FVector2D(0.f,0.f), LocalGeometry.GetLocalSize()));
 
 			if (ClippedLocalMarquee.IsValid() && !ClippedLocalMarquee.IsEmpty())
diff --git a/Engine/Source/Editor/Sequencer/Private/SSequencerSection.cpp b/Engine/Source/Editor/Sequencer/Private/SSequencerSection.cpp
--- a/Engine/Source/Editor/Sequencer/Private/SSequencerSection.cpp
+++ b/Engine/Source/Editor/Sequencer/Private/SSequencerSection.cpp
@@ -643,7 +643,7 @@
 
 			if (bExistingEmptySpace)
 			{
-				FPaintGeometry PaintGeom = SectionGeometry.MakeChild(CurrentArea->GetSize(), FSlateLayoutTransform(CurrentArea->GetTopLeft())).ToPaintGeometry();
+				FPaintGeometry PaintGeom = SectionGeometry.MakeChild(CurrentArea->GetSize(), FSlateLayoutTransform(CurrentArea->GetTopLeft2f())).ToPaintGeometry();
 				FSlateDrawElement::MakeBox(DrawElements, LayerId, PaintGeom, EmptySpaceBrush, DrawEffects);
 				CurrentArea.Reset();
 			}
@@ -656,7 +656,7 @@
 
 		if (CurrentArea.IsSet())
 		{
-			FPaintGeometry PaintGeom = SectionGeometry.MakeChild(CurrentArea->GetSize(), FSlateLayoutTransform(CurrentArea->GetTopLeft())).ToPaintGeometry();
+			FPaintGeometry PaintGeom = SectionGeometry.MakeChild(CurrentArea->GetSize(), FSlateLayoutTransform(CurrentArea->GetTopLeft2f())).ToPaintGeometry();
 			FSlateDrawElement::MakeBox(DrawElements, LayerId, PaintGeom, EmptySpaceBrush, DrawEffects);
 		}
 
diff --git a/Engine/Source/Runtime/SlateCore/Private/Layout/LayoutUtils.cpp b/Engine/Source/Runtime/SlateCore/Private/Layout/LayoutUtils.cpp
--- a/Engine/Source/Runtime/SlateCore/Private/Layout/LayoutUtils.cpp
+++ b/Engine/Source/Runtime/SlateCore/Private/Layout/LayoutUtils.cpp
@@ -7,11 +7,11 @@ FVector2D ComputePopupFitInRect(const FSlateRect& InAnchor, const FSlateRect& Po
 	const bool bAdjustmentNeeded = PopupRect.IntersectionWith(RectToFit) != PopupRect;
 	if (bAdjustmentNeeded)
 	{
-		const FVector2D PopupPosition = PopupRect.GetTopLeft();
-		const FVector2D& PopupSize = PopupRect.GetSize();
+		const FVector2f PopupPosition = PopupRect.GetTopLeft2f();
+		const FVector2f PopupSize = PopupRect.GetSize2f();
 
 		// In the direction we are opening, see if there is enough room. If there is not, flip the opening direction along the same axis.
-		FVector2D NewPosition = FVector2D::ZeroVector;
+		FVector2f NewPosition = FVector2f::ZeroVector;
 		if (Orientation == Orient_Horizontal)
 		{
 			const bool bFitsRight = InAnchor.Right + PopupSize.X < RectToFit.Right;
@@ -20,12 +20,12 @@ FVector2D ComputePopupFitInRect(const FSlateRect& InAnchor, const FSlateRect& Po
 			if (bFitsRight || !bFitsLeft)
 			{
 				// The menu fits to the right of the anchor or it does not fit to the left, display to the right
-				NewPosition = FVector2D(InAnchor.Right, InAnchor.Top);
+				NewPosition = FVector2f(InAnchor.Right, InAnchor.Top);
 			}
 			else
 			{
 				// The menu does not fit to the right of the anchor but it does fit to the left, display to the left
-				NewPosition = FVector2D(InAnchor.Left - PopupSize.X, InAnchor.Top);
+				NewPosition = FVector2f(InAnchor.Left - PopupSize.X, InAnchor.Top);
 			}
 		}
 		else
@@ -36,12 +36,12 @@ FVector2D ComputePopupFitInRect(const FSlateRect& InAnchor, const FSlateRect& Po
 			if (bFitsDown || !bFitsUp)
 			{
 				// The menu fits below the anchor or it does not fit above, display below
-				NewPosition = FVector2D(InAnchor.Left, InAnchor.Bottom);
+				NewPosition = FVector2f(InAnchor.Left, InAnchor.Bottom);
 			}
 			else
 			{
 				// The menu does not fit below the anchor but it does fit above, display above
-				NewPosition = FVector2D(InAnchor.Left, InAnchor.Top - PopupSize.Y);
+				NewPosition = FVector2f(InAnchor.Left, InAnchor.Top - PopupSize.Y);
 			}
 
 			if (!bFitsDown && !bFitsUp)
@@ -54,9 +54,9 @@ FVector2D ComputePopupFitInRect(const FSlateRect& InAnchor, const FSlateRect& Po
 		// This can happen along the opposite axis that we are opening with
 		// Assumes this window has a valid size
 		// Adjust any menus that my not fit on the screen where they are opened
-		FVector2D StartPos = NewPosition;
-		FVector2D EndPos = NewPosition + PopupSize;
-		FVector2D Adjust = FVector2D::ZeroVector;
+		FVector2f StartPos = NewPosition;
+		FVector2f EndPos = NewPosition + PopupSize;
+		FVector2f Adjust = FVector2f::ZeroVector;
 		if (StartPos.X < RectToFit.Left)
 		{
 			// Window is clipped by the left side of the work area
@@ -83,10 +83,10 @@ FVector2D ComputePopupFitInRect(const FSlateRect& InAnchor, const FSlateRect& Po
 
 		NewPosition += Adjust;
 
-		return NewPosition;
+		return FVector2D(NewPosition);
 	}
 	else
 	{
-		return PopupRect.GetTopLeft();
+		return FVector2D(PopupRect.GetTopLeft());
 	}
 }
diff --git a/Engine/Source/Runtime/SlateCore/Private/Layout/SlateRotatedRect.cpp b/Engine/Source/Runtime/SlateCore/Private/Layout/SlateRotatedRect.cpp
--- a/Engine/Source/Runtime/SlateCore/Private/Layout/SlateRotatedRect.cpp
+++ b/Engine/Source/Runtime/SlateCore/Private/Layout/SlateRotatedRect.cpp
@@ -21,44 +21,49 @@ FSlateRect FSlateRotatedRect::ToBoundingRect() const
 		);
 }
 
-bool FSlateRotatedRect::IsUnderLocation(const FVector2D& Location) const
+bool FSlateRotatedRect::IsUnderLocation(const FVector2f Location) const
 {
-	const FVector2D Offset = Location - FVector2D(TopLeft);
-	const float Det = FVector2D::CrossProduct(FVector2D(ExtentX), FVector2D(ExtentY));
+	const FVector2f Offset = Location - TopLeft;
+	const float Det = FVector2f::CrossProduct(ExtentX, ExtentY);
 
 	// Not exhaustively efficient. Could optimize the checks for [0..1] to short circuit faster.
-	const float S = -FVector2D::CrossProduct(Offset, FVector2D(ExtentX)) / Det;
+	const float S = -FVector2f::CrossProduct(Offset, ExtentX) / Det;
 	if (FMath::IsWithinInclusive(S, 0.0f, 1.0f))
 	{
-		const float T = FVector2D::CrossProduct(Offset, FVector2D(ExtentY)) / Det;
+		const float T = FVector2f::CrossProduct(Offset, ExtentY) / Det;
 		return FMath::IsWithinInclusive(T, 0.0f, 1.0f);
 	}
 	return false;
 }
 
-FSlateRotatedRect FSlateRotatedRect::MakeRotatedRect(const FSlateRect& ClipRectInLayoutWindowSpace, const FTransform2D& LayoutToRenderTransform)
+bool FSlateRotatedRect::IsUnderLocation(const FVector2D Location) const
+{
+	return IsUnderLocation(UE::Slate::CastToVector2f(Location));
+}
+
+FSlateRotatedRect FSlateRotatedRect::MakeRotatedRect(const FSlateRect& ClipRectInLayoutWindowSpace, const FTransform2f& LayoutToRenderTransform)
 {
 	const FSlateRotatedRect RotatedRect = TransformRect(LayoutToRenderTransform, FSlateRotatedRect(ClipRectInLayoutWindowSpace));
 
-	const FVector2D TopRight = FVector2D(RotatedRect.TopLeft) + FVector2D(RotatedRect.ExtentX);
-	const FVector2D BottomLeft = FVector2D(RotatedRect.TopLeft) + FVector2D(RotatedRect.ExtentY);
+	const FVector2f TopRight = RotatedRect.TopLeft + RotatedRect.ExtentX;
+	const FVector2f BottomLeft = RotatedRect.TopLeft + RotatedRect.ExtentY;
 
 	return FSlateRotatedRect(
-		FVector2D(RotatedRect.TopLeft),
-		TopRight - FVector2D(RotatedRect.TopLeft),
-		BottomLeft - FVector2D(RotatedRect.TopLeft));
+		RotatedRect.TopLeft,
+		TopRight - RotatedRect.TopLeft,
+		BottomLeft - RotatedRect.TopLeft);
 }
 
-FSlateRotatedRect FSlateRotatedRect::MakeSnappedRotatedRect(const FSlateRect& ClipRectInLayoutWindowSpace, const FTransform2D& LayoutToRenderTransform)
+FSlateRotatedRect FSlateRotatedRect::MakeSnappedRotatedRect(const FSlateRect& ClipRectInLayoutWindowSpace, const FTransform2f& LayoutToRenderTransform)
 {
 	const FSlateRotatedRect RotatedRect = TransformRect(LayoutToRenderTransform, FSlateRotatedRect(ClipRectInLayoutWindowSpace));
 
 	// Pixel snapping is done here by rounding the resulting floats to ints, we do this before
 	// calculating the final extents of the clip box otherwise we'll get a smaller clip rect than a visual
 	// rect where each point is individually snapped.
-	const FVector2D SnappedTopLeft = ( FVector2D(RotatedRect.TopLeft) ).RoundToVector();
-	const FVector2D SnappedTopRight = ( FVector2D(RotatedRect.TopLeft) + FVector2D(RotatedRect.ExtentX) ).RoundToVector();
-	const FVector2D SnappedBottomLeft = ( FVector2D(RotatedRect.TopLeft) + FVector2D(RotatedRect.ExtentY) ).RoundToVector();
+	const FVector2f SnappedTopLeft = ( (RotatedRect.TopLeft) ).RoundToVector();
+	const FVector2f SnappedTopRight = ( (RotatedRect.TopLeft) + (RotatedRect.ExtentX) ).RoundToVector();
+	const FVector2f SnappedBottomLeft = ( (RotatedRect.TopLeft) + (RotatedRect.ExtentY) ).RoundToVector();
 
 	//NOTE: We explicitly do not re-snap the extent x/y, it wouldn't be correct to snap again in distance space
 	// even if two points are snapped, their distance wont necessarily be a whole number if those points are not
diff --git a/Engine/Source/Runtime/SlateCore/Private/Rendering/DrawElements.cpp b/Engine/Source/Runtime/SlateCore/Private/Rendering/DrawElements.cpp
--- a/Engine/Source/Runtime/SlateCore/Private/Rendering/DrawElements.cpp
+++ b/Engine/Source/Runtime/SlateCore/Private/Rendering/DrawElements.cpp
@@ -221,7 +221,7 @@ void FSlateDrawElement::ApplyPositionOffset(FVector2d InOffset)
 
 void FSlateDrawElement::ApplyPositionOffset(FVector2f InOffset)
 {
-	SetPosition(GetPosition() + InOffset);
+	SetPosition(Position + InOffset);
 	RenderTransform = Concatenate(RenderTransform, InOffset);
 
 	// Recompute cached layout to render transform
diff --git a/Engine/Source/Runtime/SlateCore/Public/Layout/Margin.h b/Engine/Source/Runtime/SlateCore/Public/Layout/Margin.h
--- a/Engine/Source/Runtime/SlateCore/Public/Layout/Margin.h
+++ b/Engine/Source/Runtime/SlateCore/Public/Layout/Margin.h
@@ -5,6 +5,7 @@
 #include "CoreMinimal.h"
 #include "UObject/ObjectMacros.h"
 #include "Types/SlateEnums.h"
+#include "Types/SlateVector2.h"
 #include "Margin.generated.h"
 
 
@@ -47,7 +48,7 @@ public:
 	{ }
 
 	/** Construct a Margin with uniform space on all sides */
-	FMargin( float UniformMargin )
+	FMargin(float UniformMargin)
 		: Left(UniformMargin)
 		, Top(UniformMargin)
 		, Right(UniformMargin)
@@ -55,7 +56,7 @@ public:
 	{ }
 	
 	/** Construct a Margin where Horizontal describes Left and Right spacing while Vertical describes Top and Bottom spacing */
-	FMargin( float Horizontal, float Vertical )
+	FMargin(float Horizontal, float Vertical)
 		: Left(Horizontal)
 		, Top(Vertical)
 		, Right(Horizontal)
@@ -63,7 +64,15 @@ public:
 	{ }
 
 	/** Construct a Margin where Horizontal describes Left and Right spacing while Vertical describes Top and Bottom spacing */
-	FMargin(const FVector2D& InVector)
+	FMargin(const FVector2f InVector)
+		: Left(InVector.X)
+		, Top(InVector.Y)
+		, Right(InVector.X)
+		, Bottom(InVector.Y)
+	{ }
+
+	/** Construct a Margin where Horizontal describes Left and Right spacing while Vertical describes Top and Bottom spacing */
+	FMargin(const FVector2d InVector)
 		: Left(UE_REAL_TO_FLOAT(InVector.X))
 		, Top(UE_REAL_TO_FLOAT(InVector.Y))
 		, Right(UE_REAL_TO_FLOAT(InVector.X))
@@ -71,7 +80,7 @@ public:
 	{ }
 	
 	/** Construct a Margin where the spacing on each side is individually specified. */
-	FMargin( float InLeft, float InTop, float InRight, float InBottom )
+	FMargin(float InLeft, float InTop, float InRight, float InBottom)
 		: Left(InLeft)
 		, Top(InTop)
 		, Right(InRight)
@@ -79,7 +88,15 @@ public:
 	{ }
 
 	/** Construct a Margin where the margins are coming from a FVector4 */
-	FMargin(const FVector4& InVector)
+	FMargin(const FVector4f InVector)
+		: Left(InVector.X)
+		, Top(InVector.Y)
+		, Right(InVector.Z)
+		, Bottom(InVector.W)
+	{ }
+
+	/** Construct a Margin where the margins are coming from a FVector4 */
+	FMargin(const FVector4d InVector)
 		: Left(UE_REAL_TO_FLOAT(InVector.X))
 		, Top(UE_REAL_TO_FLOAT(InVector.Y))
 		, Right(UE_REAL_TO_FLOAT(InVector.Z))
@@ -161,9 +178,13 @@ public:
 	 *
 	 * @return The top left offset.
 	 */
-	FVector2D GetTopLeft() const
+	UE::Slate::FDeprecateVector2D GetTopLeft() const
+	{
+		return UE::Slate::FDeprecateVector2D(Left, Top);
+	}
+	FVector2f GetTopLeft2f() const
 	{
-		return FVector2D(Left, Top);
+		return FVector2f(Left, Top);
 	}
 
 	/**
@@ -171,9 +192,13 @@ public:
 	 *
 	 * @return Cumulative margin size.
 	 */
-	FVector2D GetDesiredSize( ) const
+	UE::Slate::FDeprecateVector2D GetDesiredSize() const
+	{
+		return UE::Slate::FDeprecateVector2D(GetDesiredSize2f());
+	}
+	FVector2f GetDesiredSize2f() const
 	{
-		return FVector2D(Left + Right, Top + Bottom);
+		return FVector2f(Left + Right, Top + Bottom);
 	}
 
 	/**
diff --git a/Engine/Source/Runtime/SlateCore/Public/Layout/SlateRect.h b/Engine/Source/Runtime/SlateCore/Public/Layout/SlateRect.h
--- a/Engine/Source/Runtime/SlateCore/Public/Layout/SlateRect.h
+++ b/Engine/Source/Runtime/SlateCore/Public/Layout/SlateRect.h
@@ -4,6 +4,7 @@
 
 #include "CoreMinimal.h"
 #include "Layout/Margin.h"
+#include "Types/SlateVector2.h"
 
 /** 
  * A rectangle defined by upper-left and lower-right corners.
@@ -37,13 +38,6 @@ public:
 		, Bottom(InBottom)
 	{ }
 
-	FSlateRect( const FVector2D& InStartPos, const FVector2D& InEndPos )
-		: Left(UE_REAL_TO_FLOAT(InStartPos.X))
-		, Top(UE_REAL_TO_FLOAT(InStartPos.Y))
-		, Right(UE_REAL_TO_FLOAT(InEndPos.X))
-		, Bottom(UE_REAL_TO_FLOAT(InEndPos.Y))
-	{ }
-
 	FSlateRect(FVector2f InStartPos, FVector2f InEndPos)
 		: Left(InStartPos.X)
 		, Top(InStartPos.Y)
@@ -51,14 +45,25 @@ public:
 		, Bottom(InEndPos.Y)
 	{ }
 
+	FSlateRect( const FVector2d InStartPos, const FVector2d InEndPos )
+		: Left(UE_REAL_TO_FLOAT(InStartPos.X))
+		, Top(UE_REAL_TO_FLOAT(InStartPos.Y))
+		, Right(UE_REAL_TO_FLOAT(InEndPos.X))
+		, Bottom(UE_REAL_TO_FLOAT(InEndPos.Y))
+	{ }
+
 	/**
 	 * Creates a rect from a top left point and extent. Provided as a factory function to not conflict
 	 * with the TopLeft + BottomRight ctor.
 	 */
-	static FSlateRect FromPointAndExtent(const FVector2D& TopLeft, const FVector2D& Size)
+	static FSlateRect FromPointAndExtent(const FVector2f TopLeft, const FVector2f Size)
 	{
 		return FSlateRect(TopLeft, TopLeft + Size);
 	}
+	static FSlateRect FromPointAndExtent(const FVector2d TopLeft, const FVector2d Size)
+	{
+		return FromPointAndExtent(UE::Slate::CastToVector2f(TopLeft), UE::Slate::CastToVector2f(Size));
+	}
 	
 public:
 
@@ -83,9 +88,13 @@ public:
 	 *
 	 * @return The size as a vector.
 	 */
-	FORCEINLINE FVector2D GetSize() const
+	FORCEINLINE UE::Slate::FDeprecateVector2D GetSize() const
+	{
+		return UE::Slate::FDeprecateVector2D(GetSize2f());
+	}
+	FORCEINLINE FVector2f GetSize2f() const
 	{
-		return FVector2D( Right - Left, Bottom - Top );
+		return FVector2f(Right - Left, Bottom - Top);
 	}
 
 	/**
@@ -101,9 +110,13 @@ public:
 	 * 
 	 * @return The center point.
 	 */
-	FORCEINLINE FVector2D GetCenter() const
+	FORCEINLINE UE::Slate::FDeprecateVector2D GetCenter() const
+	{
+		return UE::Slate::FDeprecateVector2D(GetCenter2f());
+	}
+	FORCEINLINE FVector2f GetCenter2f() const
 	{
-		return FVector2D( Left, Top ) + GetSize() * 0.5f;
+		return FVector2f(Left, Top) + GetSize() * 0.5f;
 	}
 
 	/**
@@ -111,9 +124,13 @@ public:
 	 * 
 	 * @return The top-left position.
 	 */
-	FORCEINLINE FVector2D GetTopLeft() const
+	FORCEINLINE UE::Slate::FDeprecateVector2D GetTopLeft() const
+	{
+		return UE::Slate::FDeprecateVector2D(GetTopLeft2f());
+	}
+	FORCEINLINE FVector2f GetTopLeft2f() const
 	{
-		return FVector2D( Left, Top );
+		return FVector2f(Left, Top);
 	}
 
 	/**
@@ -121,9 +138,13 @@ public:
 	 *
 	 * @return The top-right position.
 	 */
-	FORCEINLINE FVector2D GetTopRight() const
+	FORCEINLINE UE::Slate::FDeprecateVector2D GetTopRight() const
+	{
+		return UE::Slate::FDeprecateVector2D(GetTopRight2f());
+	}
+	FORCEINLINE FVector2f GetTopRight2f() const
 	{
-		return FVector2D(Right, Top);
+		return FVector2f(Right, Top);
 	}
 
 	/**
@@ -131,9 +152,13 @@ public:
 	 * 
 	 * @return The bottom-right position.
 	 */
-	FORCEINLINE FVector2D GetBottomRight() const
+	FORCEINLINE UE::Slate::FDeprecateVector2D GetBottomRight() const
 	{
-		return FVector2D( Right, Bottom );
+		return UE::Slate::FDeprecateVector2D(GetBottomRight2f());
+	}
+	FORCEINLINE FVector2f GetBottomRight2f() const
+	{
+		return FVector2f(Right, Bottom);
 	}
 
 	/**
@@ -141,9 +166,13 @@ public:
 	 * 
 	 * @return The bottom-left position.
 	 */
-	FORCEINLINE FVector2D GetBottomLeft() const
+	FORCEINLINE UE::Slate::FDeprecateVector2D GetBottomLeft() const
+	{
+		return UE::Slate::FDeprecateVector2D(GetBottomLeft2f());
+	}
+	FORCEINLINE FVector2f GetBottomLeft2f() const
 	{
-		return FVector2D( Left, Bottom );
+		return FVector2f(Left, Bottom);
 	}
 
 	/**
@@ -177,9 +206,13 @@ public:
 	 *
 	 * @return An offset rectangle.
 	 */
-	FORCEINLINE FSlateRect OffsetBy( const FVector2D& OffsetAmount ) const
+	FORCEINLINE FSlateRect OffsetBy( const FVector2f OffsetAmount ) const
 	{
-		return FSlateRect(GetTopLeft() + OffsetAmount, GetBottomRight() + OffsetAmount);
+		return FSlateRect(GetTopLeft2f() + OffsetAmount, GetBottomRight2f() + OffsetAmount);
+	}
+	FORCEINLINE FSlateRect OffsetBy( const FVector2d OffsetAmount ) const
+	{
+		return OffsetBy(UE::Slate::CastToVector2f(OffsetAmount));
 	}
 
 	/**
@@ -191,7 +224,7 @@ public:
 	 */
 	FORCEINLINE FSlateRect ScaleBy(float ScaleBy) const
 	{
-		const FVector2D Delta = GetSize() * 0.5f * ScaleBy;
+		const FVector2f Delta = GetSize() * 0.5f * ScaleBy;
 		return ExtendBy(FMargin(Delta));
 	}
 
@@ -262,10 +295,14 @@ public:
 	 * @param Point	The point to check
 	 * @return True if the point is inside the rectangle
 	 */
-	FORCEINLINE bool ContainsPoint( const FVector2D& Point ) const
+	FORCEINLINE bool ContainsPoint( const FVector2f Point ) const
 	{
 		return Point.X >= Left && Point.X <= Right && Point.Y >= Top && Point.Y <= Bottom;
 	}
+	FORCEINLINE bool ContainsPoint( const FVector2d Point ) const
+	{
+		return ContainsPoint(UE::Slate::CastToVector2f(Point));
+	}
 
 	bool operator==( const FSlateRect& Other ) const
 	{
@@ -348,8 +385,8 @@ public:
 template <typename TransformType>
 FSlateRect TransformRect(const TransformType& Transform, const FSlateRect& Rect)
 {
-	FVector2D TopLeftTransformed = TransformPoint(Transform, FVector2D(Rect.Left, Rect.Top));
-	FVector2D BottomRightTransformed = TransformPoint(Transform, FVector2D(Rect.Right, Rect.Bottom));
+	FVector2f TopLeftTransformed = TransformPoint(Transform, FVector2f(Rect.Left, Rect.Top));
+	FVector2f BottomRightTransformed = TransformPoint(Transform, FVector2f(Rect.Right, Rect.Bottom));
 
 	if (TopLeftTransformed.X > BottomRightTransformed.X)
 	{
diff --git a/Engine/Source/Runtime/SlateCore/Public/Layout/SlateRotatedRect.h b/Engine/Source/Runtime/SlateCore/Public/Layout/SlateRotatedRect.h
--- a/Engine/Source/Runtime/SlateCore/Public/Layout/SlateRotatedRect.h
+++ b/Engine/Source/Runtime/SlateCore/Public/Layout/SlateRotatedRect.h
@@ -6,6 +6,7 @@
 #include "Layout/SlateRect.h"
 #include "Rendering/SlateLayoutTransform.h"
 #include "Rendering/SlateRenderTransform.h"
+#include "Types/SlateVector2.h"
 
 /**
  * Stores a rectangle that has been transformed by an arbitrary render transform.
@@ -27,13 +28,21 @@ public:
 	}
 
 	/** Per-element constructor. */
-	FSlateRotatedRect(const FVector2D& InTopLeft, const FVector2D& InExtentX, const FVector2D& InExtentY)
+	FSlateRotatedRect(const FVector2f InTopLeft, const FVector2f InExtentX, const FVector2f InExtentY)
 		: TopLeft(InTopLeft)
 		, ExtentX(InExtentX)
 		, ExtentY(InExtentY)
 	{
 	}
 
+	/** Per-element constructor. */
+	FSlateRotatedRect(const FVector2d InTopLeft, const FVector2d InExtentX, const FVector2d InExtentY)
+		: TopLeft(UE::Slate::CastToVector2f(InTopLeft))
+		, ExtentX(UE::Slate::CastToVector2f(InExtentX))
+		, ExtentY(UE::Slate::CastToVector2f(InExtentY))
+	{
+	}
+
 public:
 
 	/** transformed Top-left corner. */
@@ -58,14 +67,15 @@ public:
 	FSlateRect ToBoundingRect() const;
 
 	/** Point-in-rect test. */
-	bool IsUnderLocation(const FVector2D& Location) const;
+	bool IsUnderLocation(const FVector2f Location) const;
+	bool IsUnderLocation(const FVector2d Location) const;
 
 	static FSlateRotatedRect MakeRotatedRect(const FSlateRect& ClipRectInLayoutWindowSpace, const FSlateLayoutTransform& InverseLayoutTransform, const FSlateRenderTransform& RenderTransform)
 	{
 		return MakeRotatedRect(ClipRectInLayoutWindowSpace, Concatenate(InverseLayoutTransform, RenderTransform));
 	}
 
-	static FSlateRotatedRect MakeRotatedRect(const FSlateRect& ClipRectInLayoutWindowSpace, const FTransform2D& LayoutToRenderTransform);
+	static FSlateRotatedRect MakeRotatedRect(const FSlateRect& ClipRectInLayoutWindowSpace, const FTransform2f& LayoutToRenderTransform);
 
 	static FSlateRotatedRect MakeSnappedRotatedRect(const FSlateRect& ClipRectInLayoutWindowSpace, const FSlateLayoutTransform& InverseLayoutTransform, const FSlateRenderTransform& RenderTransform)
 	{
@@ -75,7 +85,7 @@ public:
 	/**
 	* Used to construct a rotated rect from an aligned clip rect and a set of layout and render transforms from the geometry, snapped to pixel boundaries. Returns a float or float16 version of the rect based on the typedef.
 	*/
-	static FSlateRotatedRect MakeSnappedRotatedRect(const FSlateRect& ClipRectInLayoutWindowSpace, const FTransform2D& LayoutToRenderTransform);
+	static FSlateRotatedRect MakeSnappedRotatedRect(const FSlateRect& ClipRectInLayoutWindowSpace, const FTransform2f& LayoutToRenderTransform);
 };
 
 /**
@@ -86,8 +96,8 @@ FSlateRotatedRect TransformRect(const TransformType& Transform, const FSlateRota
 {
 	return FSlateRotatedRect
 	(
-		TransformPoint(Transform, FVector2D(Rect.TopLeft)),
-		TransformVector(Transform, FVector2D(Rect.ExtentX)),
-		TransformVector(Transform, FVector2D(Rect.ExtentY))
+		TransformPoint(Transform, Rect.TopLeft),
+		TransformVector(Transform, Rect.ExtentX),
+		TransformVector(Transform, Rect.ExtentY)
 	);
 }
diff --git a/Engine/Source/Runtime/SlateCore/Public/Rendering/DrawElements.h b/Engine/Source/Runtime/SlateCore/Public/Rendering/DrawElements.h
--- a/Engine/Source/Runtime/SlateCore/Public/Rendering/DrawElements.h
+++ b/Engine/Source/Runtime/SlateCore/Public/Rendering/DrawElements.h
@@ -304,7 +304,7 @@ public:
 
 	FORCEINLINE UE::Slate::FDeprecateVector2D GetPosition() const
 	{
-		return Position;
+		return UE::Slate::FDeprecateVector2D(Position);
 	}
 
 	UE_DEPRECATED(5.1, "Slate rendering uses float instead of double. Use SetPosition(FVector2f)")
@@ -319,7 +319,7 @@ public:
 
 	FORCEINLINE UE::Slate::FDeprecateVector2D GetLocalSize() const
 	{
-		return LocalSize;
+		return UE::Slate::FDeprecateVector2D(LocalSize);
 	}
 
 	FORCEINLINE float GetScale() const { return Scale; }
@@ -564,7 +564,7 @@ public:
 	/** @return Get the window size that we will be painting */
 	UE::Slate::FDeprecateVector2D GetWindowSize() const
 	{
-		return WindowSize;
+		return UE::Slate::FDeprecateVector2D(WindowSize);
 	}
 
 	template<typename PayloadType>
diff --git a/Engine/Source/Runtime/SlateCore/Public/Types/SlateVector2.h b/Engine/Source/Runtime/SlateCore/Public/Types/SlateVector2.h
--- a/Engine/Source/Runtime/SlateCore/Public/Types/SlateVector2.h
+++ b/Engine/Source/Runtime/SlateCore/Public/Types/SlateVector2.h
@@ -24,87 +24,197 @@ FORCEINLINE static FVector2f CastToVector2f(FVector2d InValue)
  */
 struct FDeprecateVector2D
 {
+public:
 	FDeprecateVector2D() = default;
-	FDeprecateVector2D(FVector2f InValue)
-		: Data(InValue)
+	explicit FDeprecateVector2D(FVector2f InValue)
+		: X(InValue.X)
+		, Y(InValue.Y)
+	{
+	}
+	explicit FDeprecateVector2D(float InX, float InY)
+		: X(InX)
+		, Y(InY)
 	{
 	}
 
+public:
+	union
+	{
+		struct
+		{
+			float X;
+			float Y;
+		};
+
+		UE_DEPRECATED(all, "For internal use only")
+		float XY[2];
+	};
+
 public:
 	operator FVector2d() const
 	{
-		return FVector2d(Data);
+		return FVector2d(X, Y);
 	}
 
 	operator FVector2f() const
 	{
-		return Data;
+		return AsVector2f();
 	}
 
 public:
-	FORCEINLINE FVector2f operator+(const FVector2d V) const
+	/**  Get the length (magnitude) of this vector. */
+	float Size() const
 	{
-		return Data + CastToVector2f(V);
+		return AsVector2f().Size();
 	}
 
-	FORCEINLINE FVector2f operator+(const FVector2f V) const
+	/** Get the length (magnitude) of this vector. */
+	float Length() const
 	{
-		return Data + V;
+		return AsVector2f().Length();
 	}
 
-	FORCEINLINE FVector2f operator-(const FVector2d V) const
+	/** Get the squared length of this vector. */
+	float SizeSquared() const
 	{
-		return Data - CastToVector2f(V);
+		return AsVector2f().SizeSquared();
 	}
 
-	FORCEINLINE FVector2f operator-(const FVector2f V) const
+	/**  Get the squared length of this vector. */
+	float SquaredLength() const
 	{
-		return Data - V;
+		return AsVector2f().SquaredLength();
 	}
 
-	FORCEINLINE FVector2f operator+(float Scale) const
+	/** Checks whether all components of the vector are exactly zero. */
+	bool IsZero() const
 	{
-		return Data + Scale;
+		return AsVector2f().IsZero();
 	}
 
-	FORCEINLINE FVector2f operator-(float Scale) const
+	/** Get this vector as an Int Point. */
+	FIntPoint IntPoint() const
 	{
-		return Data - Scale;
+		return AsVector2f().IntPoint();
 	}
 
-	FORCEINLINE FVector2f operator*(float Scale) const
+	/** Get this vector as a vector where each component has been rounded to the nearest int. */
+	FDeprecateVector2D RoundToVector() const
 	{
-		return Data * Scale;
+		return FDeprecateVector2D(AsVector2f().RoundToVector());
 	}
 
-	FORCEINLINE FVector2f operator/(float Scale) const
+	/** Gets a specific component of the vector. */
+	float Component(int32 Index) const
 	{
-		return Data / Scale;
+		return AsVector2f().Component(Index);
 	}
 
 public:
-	bool operator==(FVector2d V) const
+	FORCEINLINE FDeprecateVector2D operator+(const FVector2f V) const
 	{
-		return CastToVector2f(V) == Data;
+		return FDeprecateVector2D(AsVector2f() + V);
+	}
+	FORCEINLINE FDeprecateVector2D operator+(const FVector2d V) const
+	{
+		return FDeprecateVector2D(AsVector2f() + CastToVector2f(V));
+	}
+	FORCEINLINE FDeprecateVector2D operator+(const FDeprecateVector2D V) const
+	{
+		return FDeprecateVector2D(AsVector2f() + V.AsVector2f());
 	}
 
-	bool operator==(FVector2f V) const
+	FORCEINLINE FDeprecateVector2D operator-(const FVector2f V) const
+	{
+		return FDeprecateVector2D(AsVector2f() - V);
+	}
+	FORCEINLINE FDeprecateVector2D operator-(const FVector2d V) const
 	{
-		return V == Data;
+		return FDeprecateVector2D(AsVector2f() - CastToVector2f(V));
+	}
+	FORCEINLINE FDeprecateVector2D operator-(const FDeprecateVector2D V) const
+	{
+		return FDeprecateVector2D(AsVector2f() - V.AsVector2f());
 	}
 
-	bool operator!=(FVector2d V) const
+	FORCEINLINE FDeprecateVector2D operator*(const FVector2f V) const
+	{
+		return FDeprecateVector2D(AsVector2f() * V);
+	}
+	FORCEINLINE FDeprecateVector2D operator*(const FVector2d V) const
+	{
+		return FDeprecateVector2D(AsVector2f() * CastToVector2f(V));
+	}
+	FORCEINLINE FDeprecateVector2D operator*(const FDeprecateVector2D V) const
+	{
+		return FDeprecateVector2D(AsVector2f() * V.AsVector2f());
+	}
+
+	FORCEINLINE FDeprecateVector2D operator/(const FVector2f V) const
+	{
+		return FDeprecateVector2D(AsVector2f() / V);
+	}
+	FORCEINLINE FDeprecateVector2D operator/(const FVector2d V) const
+	{
+		return FDeprecateVector2D(AsVector2f() / CastToVector2f(V));
+	}
+	FORCEINLINE FDeprecateVector2D operator/(const FDeprecateVector2D V) const
 	{
-		return CastToVector2f(V) != Data;
+		return FDeprecateVector2D(AsVector2f() / V.AsVector2f());
+	}
+
+	FORCEINLINE FDeprecateVector2D operator+(float Scale) const
+	{
+		return FDeprecateVector2D(AsVector2f() + Scale);
+	}
+
+	FORCEINLINE FDeprecateVector2D operator-(float Scale) const
+	{
+		return FDeprecateVector2D(AsVector2f() - Scale);
+	}
+
+	FORCEINLINE FDeprecateVector2D operator*(float Scale) const
+	{
+		return FDeprecateVector2D(AsVector2f() * Scale);
+	}
+
+	FORCEINLINE FDeprecateVector2D operator/(float Scale) const
+	{
+		return FDeprecateVector2D(AsVector2f() / Scale);
+	}
+
+public:
+	bool operator==(FVector2f V) const
+	{
+		return V == AsVector2f();
+	}
+	bool operator==(FVector2d V) const
+	{
+		return CastToVector2f(V) == AsVector2f();
+	}
+	bool operator==(FDeprecateVector2D V) const
+	{
+		return X == V.X && Y == V.Y;
 	}
 
 	bool operator!=(FVector2f V) const
 	{
-		return V != Data;
+		return V != AsVector2f();
+	}
+	bool operator!=(FVector2d V) const
+	{
+		return CastToVector2f(V) != AsVector2f();
+	}
+	bool operator!=(FDeprecateVector2D V) const
+	{
+		return X != V.X || Y != V.Y;
 	}
 
 private:
-	FVector2f Data;
+	FVector2f AsVector2f() const
+	{
+		return FVector2f(X, Y);
+	}
 };
 
 } // namespace
