AssetEditorToolkit: Fix toolbar being the wrong color if there were no custom commands

diff --git a/Engine/Source/Editor/UnrealEd/Private/Toolkits/AssetEditorToolkit.cpp b/Engine/Source/Editor/UnrealEd/Private/Toolkits/AssetEditorToolkit.cpp
--- a/Engine/Source/Editor/UnrealEd/Private/Toolkits/AssetEditorToolkit.cpp
+++ b/Engine/Source/Editor/UnrealEd/Private/Toolkits/AssetEditorToolkit.cpp
@@ -1118,7 +1118,13 @@ void FAssetEditorToolkit::GenerateToolbar()
 		]
 		+SHorizontalBox::Slot()
 		[
-			ToolBarWidget
+			SNew(SBorder)
+			.VAlign(VAlign_Center)
+			.BorderImage(FAppStyle::Get().GetBrush("AssetEditorToolbar.Background"))
+			.Padding(FMargin(0.0f))
+			[
+				ToolBarWidget
+			]
 		]
 		+SHorizontalBox::Slot()
 		.HAlign(HAlign_Right)
@@ -1126,7 +1132,7 @@ void FAssetEditorToolkit::GenerateToolbar()
 		[
 			SNew(SBorder)
 			.VAlign(VAlign_Center)
-			.BorderImage(FAppStyle::Get().GetBrush("Brushes.Panel"))
+			.BorderImage(FAppStyle::Get().GetBrush("AssetEditorToolbar.Background"))
 			.Padding(FMargin(0.0f))
 			[
 				MiscWidgets
diff --git a/Engine/Source/Runtime/SlateCore/Private/Styling/StarshipCoreStyle.cpp b/Engine/Source/Runtime/SlateCore/Private/Styling/StarshipCoreStyle.cpp
--- a/Engine/Source/Runtime/SlateCore/Private/Styling/StarshipCoreStyle.cpp
+++ b/Engine/Source/Runtime/SlateCore/Private/Styling/StarshipCoreStyle.cpp
@@ -2066,9 +2066,11 @@ void FStarshipCoreStyle::SetupMultiboxStyles(TSharedRef<FStyle>& Style)
 		FComboButtonStyle SlimToolBarComboButton = FComboButtonStyle(Style->GetWidgetStyle<FComboButtonStyle>("SimpleComboButton"))
 			.SetDownArrowPadding(FMargin(1.f, 0.f, 2.f, 0.f));
 
+		FSlateBrush* SlimToolbarBackground = new FSlateColorBrush(FStyleColors::Panel);
+		
 		FToolBarStyle SlimToolbarStyle =
 			FToolBarStyle()
-			.SetBackground(FSlateColorBrush(FStyleColors::Panel))
+			.SetBackground(*SlimToolbarBackground)
 			.SetBackgroundPadding(FMargin(4.f, 6.f, 0.f, 6.f))
 			.SetExpandBrush(IMAGE_BRUSH("Icons/toolbar_expand_16x", Icon16x16))
 			.SetComboButtonPadding(FMargin(0.0f, 0.0f))
@@ -2133,6 +2135,9 @@ void FStarshipCoreStyle::SetupMultiboxStyles(TSharedRef<FStyle>& Style)
 		
 		Style->Set("AssetEditorToolbar", SlimToolbarStyle);
 
+		// Used by empty toolbars and misc widgets to have the same background color
+		Style->Set("AssetEditorToolbar.Background", SlimToolbarBackground);
+
 		// Callout Toolbar - Used to "call out" the toolbar button with text
 		{
 			SlimToolbarStyle.SetShowLabels(true);
