Add a "Icons.AlertCircle" and "Icons.AlertCircleWithColor" to the StarshipCoreStyle. This uses the "Starship/Common/alert-circle" svg, which is currently utilized by the SourceControl plugin and a few other plugins as a generic "Warning" symbol. I wanted to add it here to the starship style so that for my use case I don't have to add a dependency on a source control icon, and instead can use a more generic one.

The source control icon also has a yellow accent color on the brush, which I don't want.

#rnx

diff --git a/Engine/Source/Runtime/SlateCore/Private/Styling/StarshipCoreStyle.cpp b/Engine/Source/Runtime/SlateCore/Private/Styling/StarshipCoreStyle.cpp
--- a/Engine/Source/Runtime/SlateCore/Private/Styling/StarshipCoreStyle.cpp
+++ b/Engine/Source/Runtime/SlateCore/Private/Styling/StarshipCoreStyle.cpp
@@ -319,6 +319,9 @@ TSharedRef<ISlateStyle> FStarshipCoreStyle::Create()
 		Style->Set("Icons.Warning.Large", new IMAGE_BRUSH_SVG("Starship/Common/alert-triangle-large", Icon32x32));
 		Style->Set("Icons.WarningWithColor.Large", new IMAGE_BRUSH_SVG("Starship/Common/alert-triangle-large", Icon32x32, FStyleColors::Warning));
 
+		Style->Set("Icons.AlertCircle", new IMAGE_BRUSH_SVG("Starship/Common/alert-circle", Icon16x16));
+		Style->Set("Icons.AlertCircleWithColor", new IMAGE_BRUSH_SVG("Starship/Common/alert-circle", Icon16x16, FStyleColors::Warning));
+		
 		Style->Set("Icons.Success.Large", new IMAGE_BRUSH_SVG("Starship/Common/check-circle-large", Icon32x32));
 		Style->Set("Icons.SuccessWithColor.Large", new IMAGE_BRUSH_SVG("Starship/Common/check-circle-large", Icon32x32, FStyleColors::Success));
 
