Slate: Fix InvalidationWidgetReason unitest bad name.
#rnx

diff --git a/Engine/Source/Runtime/SlateCore/Private/Test/SlateInvalidationWidgetReasonText.cpp b/Engine/Source/Runtime/SlateCore/Private/Test/SlateInvalidationWidgetReasonText.cpp
--- a/Engine/Source/Runtime/SlateCore/Private/Test/SlateInvalidationWidgetReasonText.cpp
+++ b/Engine/Source/Runtime/SlateCore/Private/Test/SlateInvalidationWidgetReasonText.cpp
@@ -9,7 +9,7 @@
 
 #define LOCTEXT_NAMESPACE "Slate.InvalidationWidgetReason.Lex"
 
-IMPLEMENT_SIMPLE_AUTOMATION_TEST(FSlateInvalidationWidgetReasonLexTest, "Slate.InvalidationWidgetReason.x", EAutomationTestFlags::ApplicationContextMask | EAutomationTestFlags::EngineFilter)
+IMPLEMENT_SIMPLE_AUTOMATION_TEST(FSlateInvalidationWidgetReasonLexTest, "Slate.ParsingInvalidationWidgetReason", EAutomationTestFlags::ApplicationContextMask | EAutomationTestFlags::EngineFilter)
 
 bool FSlateInvalidationWidgetReasonLexTest::RunTest(const FString& Parameters)
 {
