Add a constructor for Pointer Input events that takes the new InputDeviceId instead of the old int32 controller id

#rnx

diff --git a/Engine/Source/Runtime/SlateCore/Public/Input/Events.h b/Engine/Source/Runtime/SlateCore/Public/Input/Events.h
--- a/Engine/Source/Runtime/SlateCore/Public/Input/Events.h
+++ b/Engine/Source/Runtime/SlateCore/Public/Input/Events.h
@@ -872,6 +872,35 @@ public:
 		, bIsTouchForceChanged(bInIsForceChanged)
 		, bIsTouchFirstMove(bInIsFirstMove)
 	{ }
+	
+	FPointerEvent(
+		FInputDeviceId InDeviceId,
+		uint32 InPointerIndex,
+		const FVector2D& InScreenSpacePosition,
+		const FVector2D& InLastScreenSpacePosition,
+		float InForce,
+		bool bPressLeftMouseButton,
+		bool bInIsForceChanged = false,
+		bool bInIsFirstMove = false,
+		const FModifierKeysState& InModifierKeys = FModifierKeysState(),
+		uint32 InTouchpadIndex=0
+	)
+	: FInputEvent(InModifierKeys, InDeviceId, false)
+		, ScreenSpacePosition(InScreenSpacePosition)
+		, LastScreenSpacePosition(InLastScreenSpacePosition)
+		, CursorDelta(InScreenSpacePosition - InLastScreenSpacePosition)
+		, PressedButtons(bPressLeftMouseButton ? &FTouchKeySet::StandardSet : &FTouchKeySet::EmptySet)
+		, EffectingButton(EKeys::LeftMouseButton)
+		, PointerIndex(InPointerIndex)
+		, TouchpadIndex(InTouchpadIndex)
+		, Force(InForce)
+		, bIsTouchEvent(true)
+		, GestureType(EGestureEvent::None)
+		, WheelOrGestureDelta(0.0f, 0.0f)
+		, bIsDirectionInvertedFromDevice(false)
+		, bIsTouchForceChanged(bInIsForceChanged)
+		, bIsTouchFirstMove(bInIsFirstMove)
+	{ }
 
 	/** A constructor for gesture events */
 	FPointerEvent(
