Slate: Fix bad memory access after the invalidation widget list got resized.

diff --git a/Engine/Source/Runtime/SlateCore/Private/FastUpdate/SlateInvalidationWidgetList.cpp b/Engine/Source/Runtime/SlateCore/Private/FastUpdate/SlateInvalidationWidgetList.cpp
--- a/Engine/Source/Runtime/SlateCore/Private/FastUpdate/SlateInvalidationWidgetList.cpp
+++ b/Engine/Source/Runtime/SlateCore/Private/FastUpdate/SlateInvalidationWidgetList.cpp
@@ -549,23 +549,24 @@ bool FSlateInvalidationWidgetList::ProcessChildOrderInvalidation(FSlateInvalidat
 	bool bIsInvalidationWidgetStillValid = true;
 	TGuardValue<IProcessChildOrderInvalidationCallback*> TmpGuard(CurrentInvalidationCallback, &Callback);
 
-	const InvalidationWidgetType& InvalidationWidget = (*this)[WidgetIndex];
-	SWidget* WidgetPtr = InvalidationWidget.GetWidget();
-	check(WidgetPtr);
-	{
-		const bool bWidgetIndexMatches = WidgetPtr->GetProxyHandle().GetWidgetIndex() == InvalidationWidget.Index;
-		if (!bWidgetIndexMatches)
-		{
-			ensureMsgf(false, TEXT("The widget index doesn't match the index in the InvalidationWidgetList"));
-			return false;
-		}
-	}
-
+	SWidget* WidgetPtr = nullptr;
 	int32 StartChildIndex = INDEX_NONE;
 	enum class EOperation : uint8 { Remove, Cut, None } CurrentOperation = EOperation::None;
 
 	// Find all the InvalidationWidget children
 	{
+		const InvalidationWidgetType& InvalidationWidget = (*this)[WidgetIndex];
+		WidgetPtr = InvalidationWidget.GetWidget();
+		check(WidgetPtr);
+		{
+			const bool bWidgetIndexMatches = WidgetPtr->GetProxyHandle().GetWidgetIndex() == InvalidationWidget.Index;
+			if (!bWidgetIndexMatches)
+			{
+				ensureMsgf(false, TEXT("The widget index doesn't match the index in the InvalidationWidgetList"));
+				return false;
+			}
+		}
+
 		FIndexRange ChildToRemoveRange;
 
 		// Was added, but it should not be there anymore
@@ -674,31 +675,30 @@ bool FSlateInvalidationWidgetList::ProcessChildOrderInvalidation(FSlateInvalidat
 				CutArray(InvalidationWidget.LeafMostChildIndex);
 			}
 		}
+	}
 
-		// Rebuild
-		if (StartChildIndex != INDEX_NONE)
+	// Rebuild
+	if (StartChildIndex != INDEX_NONE)
+	{
+		check(WidgetPtr);
+		FChildren* InvalidatedChildren = WidgetPtr->GetAllChildren();
+		check(InvalidatedChildren);
+		if (InvalidatedChildren->Num() > 0)
 		{
-			FChildren* InvalidatedChildren = WidgetPtr->GetAllChildren();
-			check(InvalidatedChildren);
-			if (InvalidatedChildren->Num() > 0)
-			{
-				SCOPED_NAMED_EVENT(Slate_InvalidationList_ProcessRebuild, FColorList::Blue);
-				FSlateInvalidationWidgetIndex PreviousWidgetIndex = WidgetPtr->GetProxyHandle().GetWidgetIndex();
-				FSlateInvalidationWidgetIndex PreviousLeafMostChildIndex = InvalidationWidget.LeafMostChildIndex;
-				Internal_RebuildWidgetListTree(*WidgetPtr, StartChildIndex);
-				// The WidgetInvalidation is now invalid.
+			SCOPED_NAMED_EVENT(Slate_InvalidationList_ProcessRebuild, FColorList::Blue);
+			Internal_RebuildWidgetListTree(*WidgetPtr, StartChildIndex);
 
 #if UE_SLATE_WITH_INVALIDATIONWIDGETLIST_CHILDORDERCHECK
-				const InvalidationWidgetType& NewInvalidationWidget = (*this)[PreviousWidgetIndex];
-				const FIndexRange BuiltRange = (NewInvalidationWidget.LeafMostChildIndex == PreviousLeafMostChildIndex)
-					? FIndexRange{ *this, PreviousLeafMostChildIndex, NewInvalidationWidget.LeafMostChildIndex }
-					: FIndexRange{ *this, IncrementIndex(PreviousLeafMostChildIndex), NewInvalidationWidget.LeafMostChildIndex };
-				if (BuiltRange.IsValid())
-				{
-					Callback.ProxiesBuilt(BuiltRange);
-				}
-#endif
+			FSlateInvalidationWidgetIndex NewWidgetIndex = WidgetPtr->GetProxyHandle().GetWidgetIndex();
+			const InvalidationWidgetType& NewInvalidationWidget = (*this)[NewWidgetIndex];
+			const FIndexRange BuiltRange = (NewInvalidationWidget.LeafMostChildIndex == NewWidgetIndex)
+				? FIndexRange{ *this, NewWidgetIndex, NewInvalidationWidget.LeafMostChildIndex }
+				: FIndexRange{ *this, IncrementIndex(NewWidgetIndex), NewInvalidationWidget.LeafMostChildIndex };
+			if (BuiltRange.IsValid())
+			{
+				Callback.ProxiesBuilt(BuiltRange);
 			}
+#endif
 		}
 	}
 
diff --git a/Engine/Source/Runtime/SlateCore/Private/Test/SlateInvalidationWidgetListTest.cpp b/Engine/Source/Runtime/SlateCore/Private/Test/SlateInvalidationWidgetListTest.cpp
--- a/Engine/Source/Runtime/SlateCore/Private/Test/SlateInvalidationWidgetListTest.cpp
+++ b/Engine/Source/Runtime/SlateCore/Private/Test/SlateInvalidationWidgetListTest.cpp
@@ -491,7 +491,7 @@ bool FSlateInvalidationWidgetListTest::RunTest(const FString& Parameters)
 							bool bIsIncluded = Range.Include(Element.Get<1>());
 							if (bIsIncluded != Element.Get<2>())
 							{
-								Self->AddError(FString::Printf(TEXT("'%s' should be %s in the ProxiesBuilt range"), *Element.Get<0>(), (bIsIncluded ? TEXT("included") : TEXT("not be included"))));
+								Self->AddError(FString::Printf(TEXT("'%s' should be %s in the ProxiesBuilt range"), *Element.Get<0>(), (!bIsIncluded ? TEXT("included") : TEXT("not be included"))));
 							}
 						}
 					}
