Apply skew before the rendering of glyphs when added to the atlas
#rnx

diff --git a/Engine/Source/Runtime/SlateCore/Private/Fonts/SlateFontRenderer.cpp b/Engine/Source/Runtime/SlateCore/Private/Fonts/SlateFontRenderer.cpp
--- a/Engine/Source/Runtime/SlateCore/Private/Fonts/SlateFontRenderer.cpp
+++ b/Engine/Source/Runtime/SlateCore/Private/Fonts/SlateFontRenderer.cpp
@@ -415,6 +415,20 @@ bool FSlateFontRenderer::GetRenderData(const FShapedGlyphEntry& InShapedGlyph, c
 
 		if (FaceGlyphData.FaceAndMemory->IsFaceValid())
 		{
+			if (FMath::IsNearlyEqual(InShapedGlyph.FontFaceData->FontSkew, 0.f)) 
+			{
+				FT_Set_Transform(FaceGlyphData.FaceAndMemory->GetFace(), nullptr, nullptr);
+			}
+			else 
+			{
+				// Skewing / Fake Italics (could do character rotation in future?).
+				FT_Matrix TransformMatrix;
+				TransformMatrix.xx = 0x10000L;
+				TransformMatrix.xy = InShapedGlyph.FontFaceData->FontSkew * 0x10000L;
+				TransformMatrix.yx = 0;
+				TransformMatrix.yy = 0x10000L;
+				FT_Set_Transform(FaceGlyphData.FaceAndMemory->GetFace(), &TransformMatrix, nullptr);
+			}
 			FT_Error Error = FreeTypeUtils::LoadGlyph(FaceGlyphData.FaceAndMemory->GetFace(), FaceGlyphData.GlyphIndex, FaceGlyphData.GlyphFlags, InShapedGlyph.FontFaceData->FontSize, InShapedGlyph.FontFaceData->FontScale);
 			if (Error == 0)
 			{
