Editable Text Blocks:  add way to set font and font materials in blueprint

diff --git a/Engine/Source/Runtime/Slate/Private/Widgets/Input/SEditableText.cpp b/Engine/Source/Runtime/Slate/Private/Widgets/Input/SEditableText.cpp
--- a/Engine/Source/Runtime/Slate/Private/Widgets/Input/SEditableText.cpp
+++ b/Engine/Source/Runtime/Slate/Private/Widgets/Input/SEditableText.cpp
@@ -95,6 +95,17 @@ FText SEditableText::GetText() const
 void SEditableText::SetFont( const TAttribute< FSlateFontInfo >& InNewFont )
 {
 	Font = InNewFont;
+
+	Invalidate(EInvalidateWidgetReason::Layout);
+}
+
+void SEditableText::SetTextStyle( const FEditableTextStyle& InNewTextStyle )
+{
+	Font = InNewTextStyle.Font;
+	ColorAndOpacity = InNewTextStyle.ColorAndOpacity;
+	BackgroundImageSelected = &InNewTextStyle.BackgroundImageSelected;
+
+	Invalidate(EInvalidateWidgetReason::Layout);
 }
 
 void SEditableText::Tick( const FGeometry& AllottedGeometry, const double InCurrentTime, const float InDeltaTime )
diff --git a/Engine/Source/Runtime/Slate/Public/Widgets/Input/SEditableText.h b/Engine/Source/Runtime/Slate/Public/Widgets/Input/SEditableText.h
--- a/Engine/Source/Runtime/Slate/Public/Widgets/Input/SEditableText.h
+++ b/Engine/Source/Runtime/Slate/Public/Widgets/Input/SEditableText.h
@@ -238,6 +238,13 @@ public:
 	 */
 	void SetFont( const TAttribute< FSlateFontInfo >& InNewFont );
 
+	/**
+	 * Sets the text style used to draw the text
+	 *
+	 * @param  NewTextStyle	The new text style to use
+	 */
+	void SetTextStyle( const FEditableTextStyle& InNewTextStyle );
+
 	/**
 	 * Sets the minimum width that a text block should be.
 	 *
diff --git a/Engine/Source/Runtime/SlateCore/Public/Styling/SlateTypes.h b/Engine/Source/Runtime/SlateCore/Public/Styling/SlateTypes.h
--- a/Engine/Source/Runtime/SlateCore/Public/Styling/SlateTypes.h
+++ b/Engine/Source/Runtime/SlateCore/Public/Styling/SlateTypes.h
@@ -339,6 +339,8 @@ struct SLATECORE_API FTextBlockStyle : public FSlateWidgetStyle
 	FTextBlockStyle& SetFontName(const ANSICHAR* InFontName) { Font = FSlateFontInfo(InFontName, Font.GetClampSize()); return *this; }
 	FTextBlockStyle& SetFontSize(uint16 InSize) { Font.Size = InSize; return *this; }
 	FTextBlockStyle& SetTypefaceFontName(const FName& InTypefaceFontName) { Font.TypefaceFontName = InTypefaceFontName; return *this; }
+	FTextBlockStyle& SetFontMaterial(UObject* InMaterial) { Font.FontMaterial = InMaterial; return *this; }
+	FTextBlockStyle& SetFontOutlineMaterial(UObject* InMaterial) { Font.OutlineSettings.OutlineMaterial = InMaterial; return *this; }
 
 	/** The color and opacity of this text */
 	UPROPERTY(EditAnywhere, BlueprintReadWrite, Category=Appearance, meta=(DisplayName="Color"))
@@ -796,6 +798,8 @@ struct SLATECORE_API FEditableTextStyle : public FSlateWidgetStyle
 	FSlateFontInfo Font;
 	FEditableTextStyle& SetFont(const FSlateFontInfo& InFont) { Font = InFont; return *this; }
 	FEditableTextStyle& SetFont(const FName& InFontName, uint16 InSize) { Font = FSlateFontInfo(InFontName, InSize); return *this; }
+	FEditableTextStyle& SetFontMaterial(UObject* InMaterial) { Font.FontMaterial = InMaterial; return *this; }
+	FEditableTextStyle& SetFontOutlineMaterial(UObject* InMaterial) { Font.OutlineSettings.OutlineMaterial = InMaterial; return *this; }
 
 	/** The color and opacity of this text */
 	UPROPERTY(EditAnywhere, BlueprintReadWrite, Category=Appearance)
diff --git a/Engine/Source/Runtime/UMG/Private/Components/EditableText.cpp b/Engine/Source/Runtime/UMG/Private/Components/EditableText.cpp
--- a/Engine/Source/Runtime/UMG/Private/Components/EditableText.cpp
+++ b/Engine/Source/Runtime/UMG/Private/Components/EditableText.cpp
@@ -327,6 +327,41 @@ void UEditableText::SetKeyboardType(EVirtualKeyboardType::Type Type)
 	KeyboardType = Type;
 }
 
+const FSlateFontInfo& UEditableText::GetFont() const
+{
+	return WidgetStyle.Font;
+}
+
+void UEditableText::SetFont(FSlateFontInfo InFontInfo)
+{
+	WidgetStyle.SetFont(InFontInfo);
+
+	if (MyEditableText.IsValid())
+	{
+		MyEditableText->SetTextStyle(WidgetStyle);
+	}
+}
+
+void UEditableText::SetFontMaterial(UMaterialInterface* InMaterial)
+{
+	WidgetStyle.SetFontMaterial(InMaterial);
+
+	if (MyEditableText.IsValid())
+	{
+		MyEditableText->SetTextStyle(WidgetStyle);
+	}
+}
+
+void UEditableText::SetFontOutlineMaterial(UMaterialInterface* InMaterial)
+{
+	WidgetStyle.SetFontOutlineMaterial(InMaterial);
+
+	if (MyEditableText.IsValid())
+	{
+		MyEditableText->SetTextStyle(WidgetStyle);
+	}
+}
+
 void UEditableText::HandleOnTextChanged(const FText& InText)
 {
 	OnTextChanged.Broadcast(InText);
diff --git a/Engine/Source/Runtime/UMG/Private/Components/MultiLineEditableText.cpp b/Engine/Source/Runtime/UMG/Private/Components/MultiLineEditableText.cpp
--- a/Engine/Source/Runtime/UMG/Private/Components/MultiLineEditableText.cpp
+++ b/Engine/Source/Runtime/UMG/Private/Components/MultiLineEditableText.cpp
@@ -270,6 +270,41 @@ void UMultiLineEditableText::SetWidgetStyle(const FTextBlockStyle& InWidgetStyle
 	}
 }
 
+const FSlateFontInfo& UMultiLineEditableText::GetFont() const
+{
+	return WidgetStyle.Font;
+}
+
+void UMultiLineEditableText::SetFont(FSlateFontInfo InFontInfo)
+{
+	WidgetStyle.SetFont(InFontInfo);
+
+	if (MyMultiLineEditableText.IsValid())
+	{
+		MyMultiLineEditableText->SetTextStyle(&WidgetStyle);
+	}
+}
+
+void UMultiLineEditableText::SetFontMaterial(UMaterialInterface* InMaterial)
+{
+	WidgetStyle.SetFontMaterial(InMaterial);
+
+	if (MyMultiLineEditableText.IsValid())
+	{
+		MyMultiLineEditableText->SetTextStyle(&WidgetStyle);
+	}
+}
+
+void UMultiLineEditableText::SetFontOutlineMaterial(UMaterialInterface* InMaterial)
+{
+	WidgetStyle.SetFontOutlineMaterial(InMaterial);
+
+	if (MyMultiLineEditableText.IsValid())
+	{
+		MyMultiLineEditableText->SetTextStyle(&WidgetStyle);
+	}
+}
+
 void UMultiLineEditableText::HandleOnTextChanged(const FText& InText)
 {
 	OnTextChanged.Broadcast(InText);
diff --git a/Engine/Source/Runtime/UMG/Public/Components/EditableText.h b/Engine/Source/Runtime/UMG/Public/Components/EditableText.h
--- a/Engine/Source/Runtime/UMG/Public/Components/EditableText.h
+++ b/Engine/Source/Runtime/UMG/Public/Components/EditableText.h
@@ -225,6 +225,18 @@ public:
 
 	void SetKeyboardType(EVirtualKeyboardType::Type Type);
 
+	UFUNCTION(BlueprintCallable, Category = "Appearance")
+	const FSlateFontInfo& GetFont() const;
+
+	UFUNCTION(BlueprintCallable, Category = "Appearance")
+	void SetFont(FSlateFontInfo InFontInfo);
+
+	UFUNCTION(BlueprintCallable, Category = "Appearance")
+	void SetFontMaterial(UMaterialInterface* InMaterial);
+
+	UFUNCTION(BlueprintCallable, Category = "Appearance")
+	void SetFontOutlineMaterial(UMaterialInterface* InMaterial);
+
 public:
 	
 	//~ Begin UWidget Interface
diff --git a/Engine/Source/Runtime/UMG/Public/Components/MultiLineEditableText.h b/Engine/Source/Runtime/UMG/Public/Components/MultiLineEditableText.h
--- a/Engine/Source/Runtime/UMG/Public/Components/MultiLineEditableText.h
+++ b/Engine/Source/Runtime/UMG/Public/Components/MultiLineEditableText.h
@@ -152,6 +152,18 @@ public:
 	UFUNCTION(BlueprintSetter)
 	void SetWidgetStyle(const FTextBlockStyle& InWidgetStyle);
 
+	UFUNCTION(BlueprintCallable, Category = "Appearance")
+	const FSlateFontInfo& GetFont() const;
+
+	UFUNCTION(BlueprintCallable, Category = "Appearance")
+	void SetFont(FSlateFontInfo InFontInfo);
+
+	UFUNCTION(BlueprintCallable, Category = "Appearance")
+	void SetFontMaterial(UMaterialInterface* InMaterial);
+
+	UFUNCTION(BlueprintCallable, Category = "Appearance")
+	void SetFontOutlineMaterial(UMaterialInterface* InMaterial);
+
 	//~ Begin UTextLayoutWidget Interface
 	virtual void SetJustification(ETextJustify::Type InJustification) override;
 	//~ End UTextLayoutWidget Interface
