Fixed copy paste error for LargeBold font being assigned to a small font

diff --git a/Engine/Source/Runtime/SlateCore/Private/Styling/StarshipCoreStyle.cpp b/Engine/Source/Runtime/SlateCore/Private/Styling/StarshipCoreStyle.cpp
--- a/Engine/Source/Runtime/SlateCore/Private/Styling/StarshipCoreStyle.cpp
+++ b/Engine/Source/Runtime/SlateCore/Private/Styling/StarshipCoreStyle.cpp
@@ -1088,7 +1088,7 @@ void FStarshipCoreStyle::SetupTextStyles(TSharedRef<FStyle>& Style)
 	Style->Set("Font.Large", StyleFonts.Large);
 	Style->Set("NormalFontBold", StyleFonts.NormalBold);
 	Style->Set("SmallFontBold", StyleFonts.SmallBold);
-	Style->Set("Font.Large.Bold", StyleFonts.Small);
+	Style->Set("Font.Large.Bold", StyleFonts.LargeBold);
 	Style->Set("NormalFontItalic", StyleFonts.NormalItalic);
 	Style->Set("NormalFontBoldItalic", StyleFonts.NormalBoldItalic);
 
