Only set SImage MouseDown and DoubleClick event if bound.

This can save a fair amount of memory when there are many widgets.

diff --git a/Engine/Source/Runtime/Slate/Private/Widgets/Text/STextBlock.cpp b/Engine/Source/Runtime/Slate/Private/Widgets/Text/STextBlock.cpp
--- a/Engine/Source/Runtime/Slate/Private/Widgets/Text/STextBlock.cpp
+++ b/Engine/Source/Runtime/Slate/Private/Widgets/Text/STextBlock.cpp
@@ -112,7 +112,10 @@ void STextBlock::Construct( const FArguments& InArgs )
 
 	bSimpleTextMode = InArgs._SimpleTextMode;
 
-	SetOnMouseDoubleClick(InArgs._OnDoubleClicked);
+	if (InArgs._OnDoubleClicked.IsBound())
+	{
+		SetOnMouseDoubleClick(InArgs._OnDoubleClicked);
+	}
 
 	SetText(InArgs._Text);
 
diff --git a/Engine/Source/Runtime/SlateCore/Private/Widgets/Images/SImage.cpp b/Engine/Source/Runtime/SlateCore/Private/Widgets/Images/SImage.cpp
--- a/Engine/Source/Runtime/SlateCore/Private/Widgets/Images/SImage.cpp
+++ b/Engine/Source/Runtime/SlateCore/Private/Widgets/Images/SImage.cpp
@@ -33,7 +33,10 @@ void SImage::Construct( const FArguments& InArgs )
 
 	DesiredSizeOverrideAttribute.Assign(*this, InArgs._DesiredSizeOverride);
 
-	SetOnMouseButtonDown(InArgs._OnMouseButtonDown);
+	if (InArgs._OnMouseButtonDown.IsBound())
+	{
+		SetOnMouseButtonDown(InArgs._OnMouseButtonDown);
+	}
 }
 
 int32 SImage::OnPaint( const FPaintArgs& Args, const FGeometry& AllottedGeometry, const FSlateRect& MyCullingRect, FSlateWindowElementList& OutDrawElements, int32 LayerId, const FWidgetStyle& InWidgetStyle, bool bParentEnabled ) const
