Duplicating a theme will now start from the current theme rather than the default dark theme.

diff --git a/Engine/Source/Runtime/SlateCore/Private/Styling/StyleColors.cpp b/Engine/Source/Runtime/SlateCore/Private/Styling/StyleColors.cpp
--- a/Engine/Source/Runtime/SlateCore/Private/Styling/StyleColors.cpp
+++ b/Engine/Source/Runtime/SlateCore/Private/Styling/StyleColors.cpp
@@ -399,7 +399,10 @@ void USlateThemeManager::LoadThemeColors(FStyleTheme& Theme)
 	FString ThemeData;
 
 	// Start with the hard coded default colors. They are a fallback if the theme is incomplete 
-	Theme.LoadedDefaultColors = MakeArrayView<FLinearColor>(DefaultColors, (int32)EStyleColor::MAX);
+	if (Theme.LoadedDefaultColors.IsEmpty()) 
+	{ 
+		Theme.LoadedDefaultColors = MakeArrayView<FLinearColor>(DefaultColors, (int32)EStyleColor::MAX);
+	}
 
 	if (FFileHelper::LoadFileToString(ThemeData, *Theme.Filename))
 	{
