PR #9552: Update SlateBrush.h (Contributed by [REDACTED_PERSON])

diff --git a/Engine/Source/Runtime/SlateCore/Public/Styling/SlateBrush.h b/Engine/Source/Runtime/SlateCore/Public/Styling/SlateBrush.h
--- a/Engine/Source/Runtime/SlateCore/Public/Styling/SlateBrush.h
+++ b/Engine/Source/Runtime/SlateCore/Public/Styling/SlateBrush.h
@@ -207,7 +207,7 @@ namespace SlateBrushDefs
 }
 
 /**
- * An brush which contains information about how to draw a Slate element
+ * A brush which contains information about how to draw a Slate element
  */
 USTRUCT(BlueprintType) //, meta = (HasNativeMake = ""))
 struct SLATECORE_API FSlateBrush
