-Text font size changes in Text Box (Multy-Line) widget applies only after moving the text widget
It was the initial bug, but while looking at it, I noticed that the MultilineEditableTextBox was containing both a Font via FEditableTextBoxStyle, and another one via FTextBlockStyle, thus being error prone / inconsistent.
In order to fix the underlying issue (in addition to fix the initial bug), I removed the Font from FEditableTextBoxStyle, and moved the FTextBlockStyle from MultilineEditableTextBox to FEditableTextBoxStyle.
It solves the duplication issue and so make it clear where the Font should be set/read from.
However, as the text block style is now embedded in the editable text box style, it cannot be initialized the exact same way, and I had to do some changes to ensure there was no regression, by configuring various FEditableTextBoxStyle in some style files. I also change the default value for TextBlockStyle to better match our default theme.

-Default font is not set for text widgets.
EditableWidget: ensure to have a default font, and to set the style when calling SynchronizeProperties to ensure it reacts directly without having to force a refresh (by moving the widget for instance)

Bonus:
-Move to cpp some private methods that where 'forced' inline (and we were using function pointer on them). It will avoid some noise in public interface and speed up iteration / compile time when playing with them.

[RN] MultilineEditableTextBox was containing both a Font via FEditableTextBoxStyle, and another one via FTextBlockStyle, thus being error prone / inconsistent.The Font from FEditableTextBoxStyle has been removed, and the FTextBlockStyle moved from MultilineEditableTextBox to FEditableTextBoxStyle. It solves the duplication issue and so make it clear where the Font should be set/read from.
However, as the FTextBlockStyle is now embedded in the FEditableTextBoxStyle, it cannot be initialized the exact same way, and you can now configure the FTextBlockStyle of FEditableTextBoxStyle when creating one from scratch, by calling SetTextStyle on it.

Test
- created a Widget blueprint with different editable types combination: multiline or single line, box or no box.
    -Validated that everything was reacting live as expected now.
    -Created a blue print to set the text style and validated it was working.
    -Create data with old version, then open it with updated version to validate that the visual was still the same and deprecation of style working as expected.
-checked different places in the editor using variation of editable text to ensure they were behaving as before (detail view, console command entry, comment on blueprint node).

diff --git a/Engine/Plugins/Animation/DeformerGraph/Source/OptimusEditor/Private/OptimusEditorStyle.cpp b/Engine/Plugins/Animation/DeformerGraph/Source/OptimusEditor/Private/OptimusEditorStyle.cpp
--- a/Engine/Plugins/Animation/DeformerGraph/Source/OptimusEditor/Private/OptimusEditorStyle.cpp
+++ b/Engine/Plugins/Animation/DeformerGraph/Source/OptimusEditor/Private/OptimusEditorStyle.cpp
@@ -58,6 +58,7 @@ FOptimusEditorStyle::FOptimusEditorStyle() :
 		Set("TextEditor.Border", new BOX_BRUSH("Images/TextEditorBorder", FMargin(4.0f/16.0f), FLinearColor(0.02f,0.02f,0.02f,1)));
 
 		const FEditableTextBoxStyle EditableTextBoxStyle = FEditableTextBoxStyle()
+			.SetTextStyle(NormalText)
 			.SetBackgroundImageNormal( FSlateNoResource() )
 			.SetBackgroundImageHovered( FSlateNoResource() )
 			.SetBackgroundImageFocused( FSlateNoResource() )
diff --git a/Engine/Plugins/Animation/LiveLink/Source/LiveLinkEditor/Private/LiveLinkEditorModule.cpp b/Engine/Plugins/Animation/LiveLink/Source/LiveLinkEditor/Private/LiveLinkEditorModule.cpp
--- a/Engine/Plugins/Animation/LiveLink/Source/LiveLinkEditor/Private/LiveLinkEditorModule.cpp
+++ b/Engine/Plugins/Animation/LiveLink/Source/LiveLinkEditor/Private/LiveLinkEditorModule.cpp
@@ -135,6 +135,7 @@ public:
 		StyleSet->Set("ComboButton", ComboButton);
 
 		FEditableTextBoxStyle EditableTextStyle = FEditableTextBoxStyle()
+			.SetTextStyle(FAppStyle::Get().GetWidgetStyle<FTextBlockStyle>("NormalText"))
 			.SetFont(FCoreStyle::GetDefaultFontStyle("Regular", 9))
 			.SetBackgroundImageNormal(FSlateNoResource())
 			.SetBackgroundImageHovered(FSlateNoResource())
diff --git a/Engine/Plugins/Enterprise/DataprepEditor/Source/DataprepEditor/Private/DataprepEditorStyle.cpp b/Engine/Plugins/Enterprise/DataprepEditor/Source/DataprepEditor/Private/DataprepEditorStyle.cpp
--- a/Engine/Plugins/Enterprise/DataprepEditor/Source/DataprepEditor/Private/DataprepEditorStyle.cpp
+++ b/Engine/Plugins/Enterprise/DataprepEditor/Source/DataprepEditor/Private/DataprepEditorStyle.cpp
@@ -171,6 +171,7 @@ void FDataprepEditorStyle::Initialize()
 			StyleSet->Set( "Graph.ActionNode.Title", GraphActionNodeTitle );
 
 			FEditableTextBoxStyle GraphActionNodeTitleEditableText = FEditableTextBoxStyle()
+			.SetTextStyle(FAppStyle::Get().GetWidgetStyle<FTextBlockStyle>("NormalText"))
 			.SetFont(GraphActionNodeTitle.Font);
 			StyleSet->Set( "Graph.ActionNode.NodeTitleEditableText", GraphActionNodeTitleEditableText );
 
diff --git a/Engine/Plugins/Experimental/Animation/CurveExpression/Source/Editor/Private/CurveExpressionEditorStyle.cpp b/Engine/Plugins/Experimental/Animation/CurveExpression/Source/Editor/Private/CurveExpressionEditorStyle.cpp
--- a/Engine/Plugins/Experimental/Animation/CurveExpression/Source/Editor/Private/CurveExpressionEditorStyle.cpp
+++ b/Engine/Plugins/Experimental/Animation/CurveExpression/Source/Editor/Private/CurveExpressionEditorStyle.cpp
@@ -57,6 +57,7 @@ FCurveExpressionEditorStyle::FCurveExpressionEditorStyle() :
 		Set("TextEditor.Border", new BOX_BRUSH("Images/TextEditorBorder", FMargin(4.0f/16.0f), FLinearColor(0.02f,0.02f,0.02f,1)));
 
 		const FEditableTextBoxStyle EditableTextBoxStyle = FEditableTextBoxStyle()
+			.SetTextStyle(NormalText)
 			.SetBackgroundImageNormal( FSlateNoResource() )
 			.SetBackgroundImageHovered( FSlateNoResource() )
 			.SetBackgroundImageFocused( FSlateNoResource() )
diff --git a/Engine/Plugins/Experimental/CodeEditor/Source/CodeEditor/Private/CodeEditorStyle.cpp b/Engine/Plugins/Experimental/CodeEditor/Source/CodeEditor/Private/CodeEditorStyle.cpp
--- a/Engine/Plugins/Experimental/CodeEditor/Source/CodeEditor/Private/CodeEditorStyle.cpp
+++ b/Engine/Plugins/Experimental/CodeEditor/Source/CodeEditor/Private/CodeEditorStyle.cpp
@@ -78,6 +78,7 @@ void FCodeEditorStyle::Initialize()
 		StyleSet->Set("TextEditor.Border", new BOX_BRUSH("UI/TextEditorBorder", FMargin(4.0f/16.0f), FLinearColor(0.02f,0.02f,0.02f,1)));
 
 		const FEditableTextBoxStyle EditableTextBoxStyle = FEditableTextBoxStyle()
+			.SetTextStyle(NormalText)
 			.SetBackgroundImageNormal( FSlateNoResource() )
 			.SetBackgroundImageHovered( FSlateNoResource() )
 			.SetBackgroundImageFocused( FSlateNoResource() )
diff --git a/Engine/Plugins/MovieScene/MovieRenderPipeline/Source/MovieRenderPipelineEditor/Private/MovieRenderPipelineStyle.cpp b/Engine/Plugins/MovieScene/MovieRenderPipeline/Source/MovieRenderPipelineEditor/Private/MovieRenderPipelineStyle.cpp
--- a/Engine/Plugins/MovieScene/MovieRenderPipeline/Source/MovieRenderPipelineEditor/Private/MovieRenderPipelineStyle.cpp
+++ b/Engine/Plugins/MovieScene/MovieRenderPipeline/Source/MovieRenderPipelineEditor/Private/MovieRenderPipelineStyle.cpp
@@ -63,6 +63,7 @@ FMovieRenderPipelineStyle::FMovieRenderPipelineStyle()
 		.SetColorAndOpacity(AlmostWhite);
 
 	FEditableTextBoxStyle EditableTextStyle = FEditableTextBoxStyle()
+		.SetTextStyle(FCoreStyle::Get().GetWidgetStyle<FTextBlockStyle>("NormalText"))
 		.SetFont(FCoreStyle::GetDefaultFontStyle("Regular", 9))
 		.SetBackgroundImageNormal(FSlateNoResource())
 		.SetBackgroundImageHovered(FSlateNoResource())
diff --git a/Engine/Plugins/Runtime/MassEntity/Source/MassEntityEditor/Private/MassEditorStyle.cpp b/Engine/Plugins/Runtime/MassEntity/Source/MassEntityEditor/Private/MassEditorStyle.cpp
--- a/Engine/Plugins/Runtime/MassEntity/Source/MassEntityEditor/Private/MassEditorStyle.cpp
+++ b/Engine/Plugins/Runtime/MassEntity/Source/MassEntityEditor/Private/MassEditorStyle.cpp
@@ -52,6 +52,7 @@ void FMassEntityEditorStyle::Initialize()
 		StyleSet->Set("Mass.State.Title", StateTitle);
 
 		FEditableTextBoxStyle StateTitleEditableText = FEditableTextBoxStyle()
+			.SetTextStyle(NormalText)
 			.SetFont(TTF_CORE_FONT("Fonts/Roboto-Bold", 10))
 			.SetBackgroundImageNormal(BOX_BRUSH("Common/TextBox", FMargin(4.0f / 16.0f)))
 			.SetBackgroundImageHovered(BOX_BRUSH("Common/TextBox_Hovered", FMargin(4.0f / 16.0f)))
@@ -73,6 +74,7 @@ void FMassEntityEditorStyle::Initialize()
 		StyleSet->Set("Mass.Task.Title", TaskTitle);
 
 		FEditableTextBoxStyle TaskTitleEditableText = FEditableTextBoxStyle()
+			.SetTextStyle(NormalText)
 			.SetFont(TTF_CORE_FONT("Fonts/Roboto-Regular", 9))
 			.SetBackgroundImageNormal(BOX_BRUSH("Common/TextBox", FMargin(4.0f / 16.0f)))
 			.SetBackgroundImageHovered(BOX_BRUSH("Common/TextBox_Hovered", FMargin(4.0f / 16.0f)))
diff --git a/Engine/Plugins/Runtime/StateTree/Source/StateTreeEditorModule/Private/StateTreeEditorStyle.cpp b/Engine/Plugins/Runtime/StateTree/Source/StateTreeEditorModule/Private/StateTreeEditorStyle.cpp
--- a/Engine/Plugins/Runtime/StateTree/Source/StateTreeEditorModule/Private/StateTreeEditorStyle.cpp
+++ b/Engine/Plugins/Runtime/StateTree/Source/StateTreeEditorModule/Private/StateTreeEditorStyle.cpp
@@ -60,6 +60,7 @@ void FStateTreeEditorStyle::Initialize()
 		StyleSet->Set("StateTree.State.Title", StateTitle);
 
 		FEditableTextBoxStyle StateTitleEditableText = FEditableTextBoxStyle()
+			.SetTextStyle(NormalText)
 			.SetFont(TTF_CORE_FONT("Fonts/Roboto-Bold", 10))
 			.SetBackgroundImageNormal(BOX_BRUSH("Common/TextBox", FMargin(4.0f / 16.0f)))
 			.SetBackgroundImageHovered(BOX_BRUSH("Common/TextBox_Hovered", FMargin(4.0f / 16.0f)))
@@ -144,7 +145,7 @@ void FStateTreeEditorStyle::Initialize()
 
 	const FEditableTextBoxStyle& NormalEditableTextBox = FCoreStyle::Get().GetWidgetStyle<FEditableTextBoxStyle>("NormalEditableTextBox");
 	FEditableTextBoxStyle NameEditStyle(NormalEditableTextBox);
-	NameEditStyle.Font.Size = 10;
+	NameEditStyle.TextStyle.Font.Size = 10;
 	StyleSet->Set("StateTree.Node.Name", NameEditStyle);
 	
 	FSlateStyleRegistry::RegisterSlateStyle(*StyleSet.Get());
diff --git a/Engine/Plugins/VirtualProduction/SequencerPlaylists/Source/SequencerPlaylists/Private/SequencerPlaylistsStyle.cpp b/Engine/Plugins/VirtualProduction/SequencerPlaylists/Source/SequencerPlaylists/Private/SequencerPlaylistsStyle.cpp
--- a/Engine/Plugins/VirtualProduction/SequencerPlaylists/Source/SequencerPlaylists/Private/SequencerPlaylistsStyle.cpp
+++ b/Engine/Plugins/VirtualProduction/SequencerPlaylists/Source/SequencerPlaylists/Private/SequencerPlaylistsStyle.cpp
@@ -220,6 +220,7 @@ TSharedRef<FSlateStyleSet> FSequencerPlaylistsStyle::Create()
 	Style->Set("SequencerPlaylists.HoverTransport.Reset", ResetHoverTransportButton);
 
 	FEditableTextBoxStyle EditableTextStyle = FEditableTextBoxStyle()
+		.SetTextStyle(FAppStyle::Get().GetWidgetStyle<FTextBlockStyle>("NormalText"))
 		.SetFont(FCoreStyle::GetDefaultFontStyle("Regular", 9))
 		.SetBackgroundImageNormal(FSlateNoResource())
 		.SetBackgroundImageHovered(FSlateNoResource())
diff --git a/Engine/Plugins/VirtualProduction/Takes/Source/TakeRecorder/Private/TakeRecorderStyle.cpp b/Engine/Plugins/VirtualProduction/Takes/Source/TakeRecorder/Private/TakeRecorderStyle.cpp
--- a/Engine/Plugins/VirtualProduction/Takes/Source/TakeRecorder/Private/TakeRecorderStyle.cpp
+++ b/Engine/Plugins/VirtualProduction/Takes/Source/TakeRecorder/Private/TakeRecorderStyle.cpp
@@ -92,6 +92,7 @@ FTakeRecorderStyle::FTakeRecorderStyle()
 		.SetColorAndOpacity(AlmostWhite);
 
 	FEditableTextBoxStyle EditableTextStyle = FEditableTextBoxStyle()
+		.SetTextStyle(FCoreStyle::Get().GetWidgetStyle<FTextBlockStyle>("NormalText"))
 		.SetFont(FCoreStyle::GetDefaultFontStyle("Regular", 9))
 		.SetBackgroundImageNormal(FSlateNoResource())
 		.SetBackgroundImageHovered(FSlateNoResource())
diff --git a/Engine/Source/Developer/OutputLog/Private/SDeviceOutputLog.cpp b/Engine/Source/Developer/OutputLog/Private/SDeviceOutputLog.cpp
--- a/Engine/Source/Developer/OutputLog/Private/SDeviceOutputLog.cpp
+++ b/Engine/Source/Developer/OutputLog/Private/SDeviceOutputLog.cpp
@@ -30,7 +30,6 @@ void SDeviceOutputLog::Construct( const FArguments& InArgs )
 
 	MessagesTextBox = SNew(SMultiLineEditableTextBox)
 		.Style(FOutputLogStyle::Get(), "Log.TextBox")
-		.TextStyle(FOutputLogStyle::Get(), "Log.Normal")
 		.ForegroundColor(FLinearColor::Gray)
 		.Marshaller(MessagesTextMarshaller)
 		.IsReadOnly(true)
diff --git a/Engine/Source/Developer/OutputLog/Private/SOutputLog.cpp b/Engine/Source/Developer/OutputLog/Private/SOutputLog.cpp
--- a/Engine/Source/Developer/OutputLog/Private/SOutputLog.cpp
+++ b/Engine/Source/Developer/OutputLog/Private/SOutputLog.cpp
@@ -1080,7 +1080,6 @@ void SOutputLog::Construct( const FArguments& InArgs, bool bCreateDrawerDockButt
 
 	MessagesTextBox = SNew(SMultiLineEditableTextBox)
 		.Style(FOutputLogStyle::Get(), "Log.TextBox")
-		.TextStyle(FOutputLogStyle::Get(), "Log.Normal")
 		.Marshaller(MessagesTextMarshaller)
 		.IsReadOnly(true)
 		.AlwaysShowScrollbars(true)
diff --git a/Engine/Source/Editor/EditorStyle/Private/SlateEditorStyle.cpp b/Engine/Source/Editor/EditorStyle/Private/SlateEditorStyle.cpp
--- a/Engine/Source/Editor/EditorStyle/Private/SlateEditorStyle.cpp
+++ b/Engine/Source/Editor/EditorStyle/Private/SlateEditorStyle.cpp
@@ -663,6 +663,7 @@ void FSlateEditorStyle::FStyle::SetupGeneralStyles()
 
 	// EditableTextBox
 	NormalEditableTextBoxStyle = FEditableTextBoxStyle()
+		.SetTextStyle(NormalText)
 		.SetBackgroundImageNormal( BOX_BRUSH( "Common/TextBox", FMargin(4.0f/16.0f) ) )
 		.SetBackgroundImageHovered( BOX_BRUSH( "Common/TextBox_Hovered", FMargin(4.0f/16.0f) ) )
 		.SetBackgroundImageFocused( BOX_BRUSH( "Common/TextBox_Hovered", FMargin(4.0f/16.0f) ) )
@@ -682,6 +683,7 @@ void FSlateEditorStyle::FStyle::SetupGeneralStyles()
 		Set( "SpecialEditableTextImageNormal", SpecialEditableTextImageNormal );
 
 		const FEditableTextBoxStyle SpecialEditableTextBoxStyle = FEditableTextBoxStyle()
+			.SetTextStyle(NormalText)
 			.SetBackgroundImageNormal( *SpecialEditableTextImageNormal )
 			.SetBackgroundImageHovered( BOX_BRUSH( "Common/TextBox_Special_Hovered", FMargin(8.0f/32.0f) ) )
 			.SetBackgroundImageFocused( BOX_BRUSH( "Common/TextBox_Special_Hovered", FMargin(8.0f/32.0f) ) )
@@ -741,6 +743,7 @@ void FSlateEditorStyle::FStyle::SetupGeneralStyles()
 			.SetShadowColorAndOpacity( FLinearColor::Black );
 
 		FEditableTextBoxStyle InlineEditableTextBlockEditable = FEditableTextBoxStyle()
+			.SetTextStyle(NormalText)
 			.SetFont(NormalText.Font)
 			.SetBackgroundImageNormal( BOX_BRUSH( "Common/TextBox", FMargin(4.0f/16.0f) ) )
 			.SetBackgroundImageHovered( BOX_BRUSH( "Common/TextBox_Hovered", FMargin(4.0f/16.0f) ) )
@@ -2586,6 +2589,7 @@ void FSlateEditorStyle::FStyle::SetupSequencerStyles()
 	);
 
 	Set("Sequencer.HyperlinkTextBox", FEditableTextBoxStyle()
+		.SetTextStyle(NormalText)
 		.SetFont(DEFAULT_FONT("Regular", 9))
 		.SetBackgroundImageNormal(FSlateNoResource())
 		.SetBackgroundImageHovered(FSlateNoResource())
@@ -4007,6 +4011,7 @@ void FSlateEditorStyle::FStyle::SetupGraphEditorStyles()
 				Set( "Graph.StateNode.NodeTitle", GraphStateNodeTitle );
 
 				FEditableTextBoxStyle GraphStateNodeTitleEditableText = FEditableTextBoxStyle()
+					.SetTextStyle(NormalText)
 					.SetFont(NormalText.Font)
 					.SetBackgroundImageNormal( BOX_BRUSH( "Common/TextBox", FMargin(4.0f/16.0f) ) )
 					.SetBackgroundImageHovered( BOX_BRUSH( "Common/TextBox_Hovered", FMargin(4.0f/16.0f) ) )
@@ -4213,6 +4218,7 @@ void FSlateEditorStyle::FStyle::SetupGraphEditorStyles()
 			Set( "Graph.Node.NodeTitle", GraphNodeTitle );
 
 			FEditableTextBoxStyle GraphNodeTitleEditableText = FEditableTextBoxStyle()
+				.SetTextStyle(NormalText)
 				.SetFont(NormalText.Font)
 				.SetBackgroundImageNormal( BOX_BRUSH( "Common/TextBox", FMargin(4.0f/16.0f) ) )
 				.SetBackgroundImageHovered( BOX_BRUSH( "Common/TextBox_Hovered", FMargin(4.0f/16.0f) ) )
@@ -4241,6 +4247,7 @@ void FSlateEditorStyle::FStyle::SetupGraphEditorStyles()
 			Set( "Graph.CommentBlock.Title", GraphCommentBlockTitle );
 
 			FEditableTextBoxStyle GraphCommentBlockTitleEditableText = FEditableTextBoxStyle()
+				.SetTextStyle(NormalText)
 				.SetFont(GraphCommentBlockTitle.Font)
 				.SetBackgroundImageNormal( BOX_BRUSH( "Common/TextBox", FMargin(4.0f/16.0f) ) )
 				.SetBackgroundImageHovered( BOX_BRUSH( "Common/TextBox_Hovered", FMargin(4.0f/16.0f) ) )
@@ -4284,6 +4291,7 @@ void FSlateEditorStyle::FStyle::SetupGraphEditorStyles()
 					.SetShadowColorAndOpacity(FLinearColor(0.8f, 0.8f, 0.8f, 0.5));
 
 				FEditableTextBoxStyle InlineEditableTextBlockEditable = FEditableTextBoxStyle()
+					.SetTextStyle(NormalText)
 					.SetFont(DEFAULT_FONT("Regular", 9))
 					.SetBackgroundImageNormal(BOX_BRUSH("Common/TextBox", FMargin(4.0f / 16.0f)))
 					.SetBackgroundImageHovered(BOX_BRUSH("Common/TextBox_Hovered", FMargin(4.0f / 16.0f)))
@@ -4553,6 +4561,7 @@ void FSlateEditorStyle::FStyle::SetupGraphEditorStyles()
 			// EditableTextBox
 			{
 				Set( "Graph.EditableTextBox", FEditableTextBoxStyle()
+					.SetTextStyle(NormalText)
 					.SetBackgroundImageNormal( BOX_BRUSH( "Graph/CommonWidgets/TextBox", FMargin(4.0f/16.0f) ) )
 					.SetBackgroundImageHovered( BOX_BRUSH( "Graph/CommonWidgets/TextBox_Hovered", FMargin(4.0f/16.0f) ) )
 					.SetBackgroundImageFocused( BOX_BRUSH( "Graph/CommonWidgets/TextBox_Hovered", FMargin(4.0f/16.0f) ) )
@@ -4564,6 +4573,7 @@ void FSlateEditorStyle::FStyle::SetupGraphEditorStyles()
 			// VectorEditableTextBox
 			{
 				Set( "Graph.VectorEditableTextBox", FEditableTextBoxStyle()
+					.SetTextStyle(NormalText)
 					.SetBackgroundImageNormal( BOX_BRUSH( "Graph/CommonWidgets/TextBox", FMargin(4.0f/16.0f) ) )
 					.SetBackgroundImageHovered( BOX_BRUSH( "Graph/CommonWidgets/TextBox_Hovered", FMargin(4.0f/16.0f) ) )
 					.SetBackgroundImageFocused( BOX_BRUSH( "Graph/CommonWidgets/TextBox_Hovered", FMargin(4.0f/16.0f) ) )
@@ -7622,6 +7632,7 @@ void FSlateEditorStyle::FStyle::SetupToolkitStyles()
 		Set("PhysicsAssetEditor.Graph.Node.Shadow", new BOX_BRUSH( "Graph/RegularNode_shadow", FMargin(18.0f/64.0f) ) );
 
 		FEditableTextBoxStyle EditableTextBlock = FEditableTextBoxStyle()
+			.SetTextStyle(NormalText)
 			.SetFont(NormalText.Font)
 			.SetBackgroundImageNormal(BOX_BRUSH("Common/TextBox", FMargin(4.0f / 16.0f)))
 			.SetBackgroundImageHovered(BOX_BRUSH("Common/TextBox_Hovered", FMargin(4.0f / 16.0f)))
diff --git a/Engine/Source/Editor/EditorStyle/Private/StarshipStyle.cpp b/Engine/Source/Editor/EditorStyle/Private/StarshipStyle.cpp
--- a/Engine/Source/Editor/EditorStyle/Private/StarshipStyle.cpp
+++ b/Engine/Source/Editor/EditorStyle/Private/StarshipStyle.cpp
@@ -614,6 +614,7 @@ void FStarshipEditorStyle::FStyle::SetupGeneralStyles()
 		Set( "SpecialEditableTextImageNormal", SpecialEditableTextImageNormal );
 
 		const FEditableTextBoxStyle SpecialEditableTextBoxStyle = FEditableTextBoxStyle()
+			.SetTextStyle(NormalText)
 			.SetBackgroundImageNormal( *SpecialEditableTextImageNormal )
 			.SetBackgroundImageHovered( BOX_BRUSH( "Common/TextBox_Special_Hovered", FMargin(8.0f/32.0f) ) )
 			.SetBackgroundImageFocused( BOX_BRUSH( "Common/TextBox_Special_Hovered", FMargin(8.0f/32.0f) ) )
@@ -2169,6 +2170,7 @@ void FStarshipEditorStyle::FStyle::SetupSequencerStyles()
 		);
 
 		Set("Sequencer.HyperlinkTextBox", FEditableTextBoxStyle()
+			.SetTextStyle(NormalText)
 			.SetFont(DEFAULT_FONT("Regular", 9))
 			.SetBackgroundImageNormal(FSlateNoResource())
 			.SetBackgroundImageHovered(FSlateNoResource())
@@ -3265,6 +3267,7 @@ void FStarshipEditorStyle::FStyle::SetupGraphEditorStyles()
 				Set( "Graph.StateNode.NodeTitle", GraphStateNodeTitle );
 
 				FEditableTextBoxStyle GraphStateNodeTitleEditableText = FEditableTextBoxStyle()
+					.SetTextStyle(NormalText)
 					.SetFont(NormalText.Font)
 					.SetBackgroundImageNormal(FSlateRoundedBoxBrush(FStyleColors::Input, CoreStyleConstants::InputFocusRadius, FStyleColors::InputOutline, CoreStyleConstants::InputFocusThickness))
 					.SetBackgroundImageHovered(FSlateRoundedBoxBrush(FStyleColors::Input, CoreStyleConstants::InputFocusRadius, FStyleColors::Hover, CoreStyleConstants::InputFocusThickness))
@@ -3554,6 +3557,7 @@ void FStarshipEditorStyle::FStyle::SetupGraphEditorStyles()
 					.SetShadowColorAndOpacity(FLinearColor(0.8f, 0.8f, 0.8f, 0.5));
 
 				FEditableTextBoxStyle InlineEditableTextBlockEditable = FEditableTextBoxStyle()
+					.SetTextStyle(NormalText)
 					.SetFont(DEFAULT_FONT("Regular", 9))
 					.SetBackgroundImageNormal(FSlateRoundedBoxBrush(FStyleColors::Input, CoreStyleConstants::InputFocusRadius, FStyleColors::InputOutline, CoreStyleConstants::InputFocusThickness))
 					.SetBackgroundImageHovered(FSlateRoundedBoxBrush(FStyleColors::Input, CoreStyleConstants::InputFocusRadius, FStyleColors::Hover, CoreStyleConstants::InputFocusThickness))
@@ -3807,6 +3811,7 @@ void FStarshipEditorStyle::FStyle::SetupGraphEditorStyles()
 			// EditableTextBox
 			{
 				Set( "Graph.EditableTextBox", FEditableTextBoxStyle()
+					.SetTextStyle(NormalText)
 					.SetBackgroundImageNormal( BOX_BRUSH( "Graph/CommonWidgets/TextBox", FMargin(4.0f/16.0f) ) )
 					.SetBackgroundImageHovered( BOX_BRUSH( "Graph/CommonWidgets/TextBox_Hovered", FMargin(4.0f/16.0f) ) )
 					.SetBackgroundImageFocused( BOX_BRUSH( "Graph/CommonWidgets/TextBox_Hovered", FMargin(4.0f/16.0f) ) )
@@ -3818,6 +3823,7 @@ void FStarshipEditorStyle::FStyle::SetupGraphEditorStyles()
 			// VectorEditableTextBox
 			{
 				Set( "Graph.VectorEditableTextBox", FEditableTextBoxStyle()
+					.SetTextStyle(NormalText)
 					.SetBackgroundImageNormal( BOX_BRUSH( "Graph/CommonWidgets/TextBox", FMargin(4.0f/16.0f) ) )
 					.SetBackgroundImageHovered( BOX_BRUSH( "Graph/CommonWidgets/TextBox_Hovered", FMargin(4.0f/16.0f) ) )
 					.SetBackgroundImageFocused( BOX_BRUSH( "Graph/CommonWidgets/TextBox_Hovered", FMargin(4.0f/16.0f) ) )
@@ -6406,6 +6412,7 @@ void FStarshipEditorStyle::FStyle::SetupToolkitStyles()
 
 		// SEditableTextBox defaults...
 		Set("CurveTableEditor.Cell.Text", FEditableTextBoxStyle()
+			.SetTextStyle(NormalText)
 			.SetBackgroundImageNormal(FSlateNoResource())
 			.SetBackgroundImageHovered(FSlateRoundedBoxBrush(FStyleColors::Input, 0.0, FStyleColors::Hover, 1.0))
 			.SetBackgroundImageFocused(FSlateRoundedBoxBrush(FStyleColors::Input, 0.0, FStyleColors::Primary, 1.0))
diff --git a/Engine/Source/Editor/EditorWidgets/Private/EditorWidgetsStyle.cpp b/Engine/Source/Editor/EditorWidgets/Private/EditorWidgetsStyle.cpp
--- a/Engine/Source/Editor/EditorWidgets/Private/EditorWidgetsStyle.cpp
+++ b/Engine/Source/Editor/EditorWidgets/Private/EditorWidgetsStyle.cpp
@@ -1,4 +1,4 @@
-﻿// Copyright Epic Games, Inc. All Rights Reserved.
+// Copyright Epic Games, Inc. All Rights Reserved.
 
 #include "EditorWidgetsStyle.h"
 
@@ -36,11 +36,15 @@ void FEditorWidgetsStyle::Initialize()
 	const FTextBlockStyle NormalText = FTextBlockStyle()
 		.SetColorAndOpacity(NormalEditableTextBoxStyle.ForegroundColor)
 		.SetHighlightColor(NormalEditableTextBoxStyle.FocusedForegroundColor)
-		.SetFont(NormalEditableTextBoxStyle.Font)
-		.SetFontSize(NormalEditableTextBoxStyle.Font.Size);
+		.SetFont(NormalEditableTextBoxStyle.TextStyle.Font)
+		.SetFontSize(NormalEditableTextBoxStyle.TextStyle.Font.Size);
+
+	const FEditableTextBoxStyle TextBoxStyle = FEditableTextBoxStyle(NormalEditableTextBoxStyle).SetTextStyle(NormalText);
+
 
 	// Text editor
 	{
+		StyleSet->Set("NormalEditableTextBox", TextBoxStyle);
 		StyleSet->Set("TextEditor.NormalText", NormalText);
 
 		StyleSet->Set("SyntaxHighlight.Template.Normal", NormalText);
diff --git a/Engine/Source/Editor/EditorWidgets/Private/STemplateStringEditableTextBox.cpp b/Engine/Source/Editor/EditorWidgets/Private/STemplateStringEditableTextBox.cpp
--- a/Engine/Source/Editor/EditorWidgets/Private/STemplateStringEditableTextBox.cpp
+++ b/Engine/Source/Editor/EditorWidgets/Private/STemplateStringEditableTextBox.cpp
@@ -1,4 +1,4 @@
-﻿// Copyright Epic Games, Inc. All Rights Reserved.
+// Copyright Epic Games, Inc. All Rights Reserved.
 
 #include "STemplateStringEditableTextBox.h"
 
@@ -16,7 +16,7 @@ BEGIN_SLATE_FUNCTION_BUILD_OPTIMIZATION
 void STemplateStringEditableTextBox::Construct(const FArguments& InArgs)
 {
 	SMultiLineEditableTextBox::Construct(SMultiLineEditableTextBox::FArguments()
-		.TextStyle(&FEditorWidgetsStyle::Get().GetWidgetStyle<FTextBlockStyle>("TextEditor.NormalText"))
+		.Style(&FEditorWidgetsStyle::Get().GetWidgetStyle<FEditableTextBoxStyle>("NormalEditableTextBox"))
 		.Text(InArgs._Text)
 		.Marshaller(FTemplateStringSyntaxHighlighterMarshaller::Create(FTemplateStringSyntaxHighlighterMarshaller::FSyntaxTextStyle()))
 		.AllowMultiLine(false)
diff --git a/Engine/Source/Editor/GraphEditor/Private/SGraphNodeComment.cpp b/Engine/Source/Editor/GraphEditor/Private/SGraphNodeComment.cpp
--- a/Engine/Source/Editor/GraphEditor/Private/SGraphNodeComment.cpp
+++ b/Engine/Source/Editor/GraphEditor/Private/SGraphNodeComment.cpp
@@ -179,7 +179,7 @@ void SGraphNodeComment::UpdateGraphNode()
 	CachedFontSize = CommentNode->GetFontSize();
 
 	CommentStyle = FAppStyle::Get().GetWidgetStyle<FInlineEditableTextBlockStyle>("Graph.CommentBlock.TitleInlineEditableText");
-	CommentStyle.EditableTextBoxStyle.Font.Size = CachedFontSize;
+	CommentStyle.EditableTextBoxStyle.TextStyle.Font.Size = CachedFontSize;
 	CommentStyle.TextStyle.Font.Size = CachedFontSize;
 
 	this->ContentScale.Bind( this, &SGraphNode::GetContentScale );
diff --git a/Engine/Source/Editor/UnrealEd/Private/Dialogs/SOutputLogDialog.cpp b/Engine/Source/Editor/UnrealEd/Private/Dialogs/SOutputLogDialog.cpp
--- a/Engine/Source/Editor/UnrealEd/Private/Dialogs/SOutputLogDialog.cpp
+++ b/Engine/Source/Editor/UnrealEd/Private/Dialogs/SOutputLogDialog.cpp
@@ -91,7 +91,6 @@ void SOutputLogDialog::Construct( const FArguments& InArgs )
 						[
 							SNew(SMultiLineEditableTextBox)
 								.Style(FAppStyle::Get(), "Log.TextBox")
-								.TextStyle(FAppStyle::Get(), "Log.Normal")
 								.ForegroundColor(FLinearColor::Gray)
 								.Text(FText::TrimTrailing(Log))
 								.IsReadOnly(true)
diff --git a/Engine/Source/Programs/CrashReportClient/Private/CrashReportClientStyle.cpp b/Engine/Source/Programs/CrashReportClient/Private/CrashReportClientStyle.cpp
--- a/Engine/Source/Programs/CrashReportClient/Private/CrashReportClientStyle.cpp
+++ b/Engine/Source/Programs/CrashReportClient/Private/CrashReportClientStyle.cpp
@@ -86,7 +86,9 @@ TSharedRef< FSlateStyleSet > FCrashReportClientStyle::Create()
 		.SetHoveredThumbImage(BOX_BRUSH("Common/Scrollbar_Thumb", FMargin(4.f / 16.f)));
 
 	// SEditableTextBox defaults...
+	const FTextBlockStyle& NormalText = FAppStyle::Get().GetWidgetStyle<FTextBlockStyle>("NormalText");
 	const FEditableTextBoxStyle NormalEditableTextBoxStyle = FEditableTextBoxStyle()
+		.SetTextStyle(NormalText)
 		.SetBackgroundImageNormal(*GenericWhiteBox)
 		.SetBackgroundImageHovered(*GenericWhiteBox)
 		.SetBackgroundImageFocused(*GenericWhiteBox)
diff --git a/Engine/Source/Runtime/Core/Public/UObject/UE5ReleaseStreamObjectVersion.h b/Engine/Source/Runtime/Core/Public/UObject/UE5ReleaseStreamObjectVersion.h
--- a/Engine/Source/Runtime/Core/Public/UObject/UE5ReleaseStreamObjectVersion.h
+++ b/Engine/Source/Runtime/Core/Public/UObject/UE5ReleaseStreamObjectVersion.h
@@ -130,6 +130,9 @@ struct CORE_API FUE5ReleaseStreamObjectVersion
 		// Add WithSerializer type trait and implementation for FRawAnimSequenceTrack
 		RawAnimSequenceTrackSerializer,
 
+		// Removed font from FEditableTextBoxStyle, and added FTextBlockStyle instead.
+		RemoveDuplicatedStyleInfo,
+
 		// -----<new versions can be added above this line>-------------------------------------------------
 		VersionPlusOne,
 		LatestVersion = VersionPlusOne - 1
diff --git a/Engine/Source/Runtime/Slate/Private/Widgets/Input/SEditableText.cpp b/Engine/Source/Runtime/Slate/Private/Widgets/Input/SEditableText.cpp
--- a/Engine/Source/Runtime/Slate/Private/Widgets/Input/SEditableText.cpp
+++ b/Engine/Source/Runtime/Slate/Private/Widgets/Input/SEditableText.cpp
@@ -53,7 +53,7 @@ void SEditableText::Construct( const FArguments& InArgs )
 	BackgroundImageSelected = InArgs._BackgroundImageSelected;
 
 	// We use the given style when creating the text layout as it may not be safe to call the override delegates until we've finished being constructed
-	// The first call to SyncronizeTextStyle will apply the correct overrides, and that will happen before the first paint
+	// The first call to SynchronizeTextStyle will apply the correct overrides, and that will happen before the first paint
 	check(InArgs._Style);
 	FTextBlockStyle TextStyle = FCoreStyle::Get().GetWidgetStyle<FTextBlockStyle>("NormalText");
 	TextStyle.Font = InArgs._Style->Font;
@@ -99,7 +99,12 @@ void SEditableText::SetFont( const TAttribute< FSlateFontInfo >& InNewFont )
 	Invalidate(EInvalidateWidgetReason::Layout);
 }
 
-void SEditableText::SetTextStyle( const FEditableTextStyle& InNewTextStyle )
+FSlateFontInfo SEditableText::GetFont() const
+{
+	return Font.Get();
+}
+
+void SEditableText::SetEditableTextStyle( const FEditableTextStyle& InNewTextStyle )
 {
 	Font = InNewTextStyle.Font;
 	ColorAndOpacity = InNewTextStyle.ColorAndOpacity;
@@ -108,6 +113,15 @@ void SEditableText::SetTextStyle( const FEditableTextStyle& InNewTextStyle )
 	Invalidate(EInvalidateWidgetReason::Layout);
 }
 
+void SEditableText::SetTextBlockStyle(const FTextBlockStyle* InTextStyle)
+{
+	if (InTextStyle)
+	{
+		EditableTextLayout->SetTextStyle(*InTextStyle);
+		Invalidate(EInvalidateWidgetReason::Layout); //Using Layout as changing text block size can affect the size.
+	}
+}
+
 void SEditableText::Tick( const FGeometry& AllottedGeometry, const double InCurrentTime, const float InDeltaTime )
 {
 	EditableTextLayout->Tick(AllottedGeometry, InCurrentTime, InDeltaTime);
diff --git a/Engine/Source/Runtime/Slate/Private/Widgets/Input/SEditableTextBox.cpp b/Engine/Source/Runtime/Slate/Private/Widgets/Input/SEditableTextBox.cpp
--- a/Engine/Source/Runtime/Slate/Private/Widgets/Input/SEditableTextBox.cpp
+++ b/Engine/Source/Runtime/Slate/Private/Widgets/Input/SEditableTextBox.cpp
@@ -125,8 +125,19 @@ void SEditableTextBox::SetStyle(const FEditableTextBoxStyle* InStyle)
 	BorderImageHovered = &Style->BackgroundImageHovered;
 	BorderImageFocused = &Style->BackgroundImageFocused;
 	BorderImageReadOnly = &Style->BackgroundImageReadOnly;
+
+	SetTextBlockStyle(&Style->TextStyle);
 }
 
+void SEditableTextBox::SetTextBlockStyle(const FTextBlockStyle* InTextStyle)
+{
+	// The Construct() function will call this before EditableText exists,
+	// so we need a guard here to ignore that function call.
+	if (EditableText.IsValid())
+	{
+		EditableText->SetTextBlockStyle(InTextStyle);
+	}
+}
 
 void SEditableTextBox::SetText( const TAttribute< FText >& InNewText )
 {
@@ -302,6 +313,24 @@ FReply SEditableTextBox::OnKeyDown(const FGeometry& MyGeometry, const FKeyEvent&
 	return FReply::Unhandled();
 }
 
+FMargin SEditableTextBox::DeterminePadding() const
+{
+	check(Style);
+	return PaddingOverride.IsSet() ? PaddingOverride.Get() : Style->Padding;
+}
+
+FSlateFontInfo SEditableTextBox::DetermineFont() const
+{
+	check(Style);
+	return FontOverride.IsSet() ? FontOverride.Get() : Style->TextStyle.Font;
+}
+
+FSlateColor SEditableTextBox::DetermineBackgroundColor() const
+{
+	check(Style);
+	return BackgroundColorOverride.IsSet() ? BackgroundColorOverride.Get() : Style->BackgroundColor;
+}
+
 FSlateColor SEditableTextBox::DetermineForegroundColor() const
 {
 	check(Style);  
diff --git a/Engine/Source/Runtime/Slate/Private/Widgets/Input/SMultiLineEditableTextBox.cpp b/Engine/Source/Runtime/Slate/Private/Widgets/Input/SMultiLineEditableTextBox.cpp
--- a/Engine/Source/Runtime/Slate/Private/Widgets/Input/SMultiLineEditableTextBox.cpp
+++ b/Engine/Source/Runtime/Slate/Private/Widgets/Input/SMultiLineEditableTextBox.cpp
@@ -33,12 +33,7 @@ namespace
 void SMultiLineEditableTextBox::Construct( const FArguments& InArgs )
 {
 	check (InArgs._Style);
-	Style = InArgs._Style;
-
-	BorderImageNormal = &InArgs._Style->BackgroundImageNormal;
-	BorderImageHovered = &InArgs._Style->BackgroundImageHovered;
-	BorderImageFocused = &InArgs._Style->BackgroundImageFocused;
-	BorderImageReadOnly = &InArgs._Style->BackgroundImageReadOnly;
+	SetStyle(InArgs._Style);
 
 	PaddingOverride = InArgs._Padding;
 	HScrollBarPaddingOverride = InArgs._HScrollBarPadding;
@@ -97,7 +92,7 @@ void SMultiLineEditableTextBox::Construct( const FArguments& InArgs )
 					.Text( InArgs._Text )
 					.HintText( InArgs._HintText )
 					.SearchText( InArgs._SearchText )
-					.TextStyle( InArgs._TextStyle )
+					.TextStyle( &InArgs._Style->TextStyle )
 					.Marshaller( InArgs._Marshaller )
 					.Font( this, &SMultiLineEditableTextBox::DetermineFont )
 					.IsReadOnly( InArgs._IsReadOnly )
@@ -205,11 +200,47 @@ void SMultiLineEditableTextBox::SetStyle(const FEditableTextBoxStyle* InStyle)
 	BorderImageHovered = &Style->BackgroundImageHovered;
 	BorderImageFocused = &Style->BackgroundImageFocused;
 	BorderImageReadOnly = &Style->BackgroundImageReadOnly;
+
+	SetTextStyle(&Style->TextStyle);
 }
 
 void SMultiLineEditableTextBox::SetTextStyle(const FTextBlockStyle* InTextStyle)
 {
-	EditableText->SetTextStyle(InTextStyle);
+	// The Construct() function will call this before EditableText exists,
+	// so we need a guard here to ignore that function call.
+	if (EditableText.IsValid())
+	{
+		EditableText->SetTextStyle(InTextStyle);
+	}
+}
+
+FMargin SMultiLineEditableTextBox::DeterminePadding() const
+{
+	check(Style);
+	return PaddingOverride.IsSet() ? PaddingOverride.Get() : Style->Padding;
+}
+
+FMargin SMultiLineEditableTextBox::DetermineHScrollBarPadding() const
+{
+	check(Style);
+	return HScrollBarPaddingOverride.IsSet() ? HScrollBarPaddingOverride.Get() : Style->HScrollBarPadding;
+}
+
+FMargin SMultiLineEditableTextBox::DetermineVScrollBarPadding() const
+{
+	check(Style);
+	return VScrollBarPaddingOverride.IsSet() ? VScrollBarPaddingOverride.Get() : Style->VScrollBarPadding;
+}
+
+FSlateFontInfo SMultiLineEditableTextBox::DetermineFont() const
+{
+	return FontOverride.IsSet() ? FontOverride.Get() : Style->TextStyle.Font;
+}
+
+FSlateColor SMultiLineEditableTextBox::DetermineBackgroundColor() const
+{
+	check(Style);  
+	return BackgroundColorOverride.IsSet() ? BackgroundColorOverride.Get() : Style->BackgroundColor;
 }
 
 FSlateColor SMultiLineEditableTextBox::DetermineForegroundColor() const
diff --git a/Engine/Source/Runtime/Slate/Private/Widgets/Input/SSearchBox.cpp b/Engine/Source/Runtime/Slate/Private/Widgets/Input/SSearchBox.cpp
--- a/Engine/Source/Runtime/Slate/Private/Widgets/Input/SSearchBox.cpp
+++ b/Engine/Source/Runtime/Slate/Private/Widgets/Input/SSearchBox.cpp
@@ -20,7 +20,7 @@ void SSearchBox::Construct( const FArguments& InArgs )
 	DelayChangeNotificationsWhileTyping = InArgs._DelayChangeNotificationsWhileTyping;
 	DelayChangeNotificationsWhileTypingSeconds = InArgs._DelayChangeNotificationsWhileTypingSeconds;
 
-	InactiveFont = InArgs._Style->TextBoxStyle.Font;
+	InactiveFont = InArgs._Style->TextBoxStyle.TextStyle.Font;
 	ActiveFont = InArgs._Style->ActiveFontInfo;
 
 	SEditableTextBox::Construct( SEditableTextBox::FArguments()
diff --git a/Engine/Source/Runtime/Slate/Private/Widgets/Text/SMultiLineEditableText.cpp b/Engine/Source/Runtime/Slate/Private/Widgets/Text/SMultiLineEditableText.cpp
--- a/Engine/Source/Runtime/Slate/Private/Widgets/Text/SMultiLineEditableText.cpp
+++ b/Engine/Source/Runtime/Slate/Private/Widgets/Text/SMultiLineEditableText.cpp
@@ -175,6 +175,12 @@ void SMultiLineEditableText::SetFont(const TAttribute< FSlateFontInfo >& InNewFo
 	EditableTextLayout->SetTextStyle(TextStyle);
 }
 
+FSlateFontInfo SMultiLineEditableText::GetFont() const
+{
+	FTextBlockStyle TextStyle = EditableTextLayout->GetTextStyle();
+	return TextStyle.Font;
+}
+
 void SMultiLineEditableText::SetTextShapingMethod(const TOptional<ETextShapingMethod>& InTextShapingMethod)
 {
 	EditableTextLayout->SetTextShapingMethod(InTextShapingMethod);
diff --git a/Engine/Source/Runtime/Slate/Public/Widgets/Input/SEditableText.h b/Engine/Source/Runtime/Slate/Public/Widgets/Input/SEditableText.h
--- a/Engine/Source/Runtime/Slate/Public/Widgets/Input/SEditableText.h
+++ b/Engine/Source/Runtime/Slate/Public/Widgets/Input/SEditableText.h
@@ -236,14 +236,25 @@ public:
 	 *
 	 * @param  InNewFont	The new font to use
 	 */
-	void SetFont( const TAttribute< FSlateFontInfo >& InNewFont );
+	void SetFont(const TAttribute< FSlateFontInfo >& InNewFont);
+
+	/** Gets the font used to draw the text. */
+	FSlateFontInfo GetFont() const;
 
 	/**
 	 * Sets the text style used to draw the text
 	 *
 	 * @param  NewTextStyle	The new text style to use
 	 */
-	void SetTextStyle( const FEditableTextStyle& InNewTextStyle );
+	void SetEditableTextStyle( const FEditableTextStyle& InNewTextStyle );
+
+	/** @See TextStyle */
+	/**
+	 * Sets the text block style used to draw the text
+	 *
+	 * @param  InTextStyle	The new text block style to use
+	 */
+	void SetTextBlockStyle(const FTextBlockStyle* InTextStyle);
 
 	/**
 	 * Sets the minimum width that a text block should be.
diff --git a/Engine/Source/Runtime/Slate/Public/Widgets/Input/SEditableTextBox.h b/Engine/Source/Runtime/Slate/Public/Widgets/Input/SEditableTextBox.h
--- a/Engine/Source/Runtime/Slate/Public/Widgets/Input/SEditableTextBox.h
+++ b/Engine/Source/Runtime/Slate/Public/Widgets/Input/SEditableTextBox.h
@@ -191,6 +191,13 @@ public:
 	/** See attribute Style */
 	void SetStyle(const FEditableTextBoxStyle* InStyle);
 
+	/**
+	 * Sets the text block style currently used
+	 *
+	 * @param  InTextStyle  The new text string
+	 */
+	void SetTextBlockStyle(const FTextBlockStyle* InTextStyle);
+
 	/**
 	 * Sets the text string currently being edited 
 	 *
@@ -448,10 +455,9 @@ protected:
 
 private:
 
-	FMargin FORCEINLINE DeterminePadding() const { check(Style);  return PaddingOverride.IsSet() ? PaddingOverride.Get() : Style->Padding; }
-	FSlateFontInfo FORCEINLINE DetermineFont() const { check(Style);  return FontOverride.IsSet() ? FontOverride.Get() : Style->Font;  }
-	FSlateColor FORCEINLINE DetermineBackgroundColor() const { check(Style);  return BackgroundColorOverride.IsSet() ? BackgroundColorOverride.Get() : Style->BackgroundColor; }
-	
+	FMargin DeterminePadding() const;
+	FSlateFontInfo DetermineFont() const;
+	FSlateColor DetermineBackgroundColor() const;
 	FSlateColor DetermineForegroundColor() const;
 
 	/** Styling: border image to draw when not hovered or focused */
diff --git a/Engine/Source/Runtime/Slate/Public/Widgets/Input/SMultiLineEditableTextBox.h b/Engine/Source/Runtime/Slate/Public/Widgets/Input/SMultiLineEditableTextBox.h
--- a/Engine/Source/Runtime/Slate/Public/Widgets/Input/SMultiLineEditableTextBox.h
+++ b/Engine/Source/Runtime/Slate/Public/Widgets/Input/SMultiLineEditableTextBox.h
@@ -36,7 +36,6 @@ public:
 
 	SLATE_BEGIN_ARGS( SMultiLineEditableTextBox )
 		: _Style(&FCoreStyle::Get().GetWidgetStyle<FEditableTextBoxStyle>("NormalEditableTextBox"))
-		, _TextStyle(&FCoreStyle::Get().GetWidgetStyle<FTextBlockStyle>("NormalText"))
 		, _Marshaller()
 		, _Text()
 		, _HintText()
@@ -77,10 +76,10 @@ public:
 		{}
 
 		/** The styling of the textbox */
-		SLATE_STYLE_ARGUMENT( FEditableTextBoxStyle, Style )
+		SLATE_STYLE_ARGUMENT(FEditableTextBoxStyle, Style)
 
 		/** Pointer to a style of the text block, which dictates the font, color, and shadow options. */
-		SLATE_STYLE_ARGUMENT( FTextBlockStyle, TextStyle )
+		SLATE_STYLE_ARGUMENT_DEPRECATED(FTextBlockStyle, TextStyle, 5.2, "TextStyle is deprecated and will be ignored. Please use the TextStyle embedded in FEditableTextBoxStyle Style.")
 
 		/** The marshaller used to get/set the raw text to/from the text layout. */
 		SLATE_ARGUMENT(TSharedPtr< ITextLayoutMarshaller >, Marshaller)
@@ -537,12 +536,11 @@ protected:
 
 private:
 
-	FMargin FORCEINLINE DeterminePadding() const { check(Style);  return PaddingOverride.IsSet() ? PaddingOverride.Get() : Style->Padding; }
-	FMargin FORCEINLINE DetermineHScrollBarPadding() const { check(Style);  return HScrollBarPaddingOverride.IsSet() ? HScrollBarPaddingOverride.Get() : Style->HScrollBarPadding; }
-	FMargin FORCEINLINE DetermineVScrollBarPadding() const { check(Style);  return VScrollBarPaddingOverride.IsSet() ? VScrollBarPaddingOverride.Get() : Style->VScrollBarPadding; }
-	FSlateFontInfo FORCEINLINE DetermineFont() const { check(Style);  return FontOverride.IsSet() ? FontOverride.Get() : Style->Font; }
-	FSlateColor FORCEINLINE DetermineBackgroundColor() const { check(Style);  return BackgroundColorOverride.IsSet() ? BackgroundColorOverride.Get() : Style->BackgroundColor; }
-
+	FMargin DeterminePadding() const;
+	FMargin DetermineHScrollBarPadding() const;
+	FMargin DetermineVScrollBarPadding() const;
+	FSlateFontInfo DetermineFont() const;
+	FSlateColor DetermineBackgroundColor() const;
 	FSlateColor DetermineForegroundColor() const;
 
 	/** Styling: border image to draw when not hovered or focused */
diff --git a/Engine/Source/Runtime/Slate/Public/Widgets/Input/SNumericEntryBox.h b/Engine/Source/Runtime/Slate/Public/Widgets/Input/SNumericEntryBox.h
--- a/Engine/Source/Runtime/Slate/Public/Widgets/Input/SNumericEntryBox.h
+++ b/Engine/Source/Runtime/Slate/Public/Widgets/Input/SNumericEntryBox.h
@@ -226,7 +226,7 @@ public:
 		{
 			SAssignNew(SpinBox, SSpinBox<NumericType>)
 				.Style(InArgs._SpinBoxStyle)
-				.Font(InArgs._Font.IsSet() ? InArgs._Font : InArgs._EditableTextBoxStyle->Font)
+				.Font(InArgs._Font.IsSet() ? InArgs._Font : InArgs._EditableTextBoxStyle->TextStyle.Font)
 				.Value(this, &SNumericEntryBox<NumericType>::OnGetValueForSpinBox)
 				.Delta(InArgs._Delta)
 				.ShiftMouseMovePixelPerDelta(InArgs._ShiftMouseMovePixelPerDelta)
@@ -255,7 +255,7 @@ public:
 		SAssignNew(EditableText, SEditableText)
 			.Text(this, &SNumericEntryBox<NumericType>::OnGetValueForTextBox)
 			.Visibility(bAllowSpin ? EVisibility::Collapsed : EVisibility::Visible)
-			.Font(InArgs._Font.IsSet() ? InArgs._Font : InArgs._EditableTextBoxStyle->Font)
+			.Font(InArgs._Font.IsSet() ? InArgs._Font : InArgs._EditableTextBoxStyle->TextStyle.Font)
 			.SelectAllTextWhenFocused(true)
 			.ClearKeyboardFocusOnCommit(false)
 			.OnTextChanged(this, &SNumericEntryBox<NumericType>::OnTextChanged)
diff --git a/Engine/Source/Runtime/Slate/Public/Widgets/Text/SMultiLineEditableText.h b/Engine/Source/Runtime/Slate/Public/Widgets/Text/SMultiLineEditableText.h
--- a/Engine/Source/Runtime/Slate/Public/Widgets/Text/SMultiLineEditableText.h
+++ b/Engine/Source/Runtime/Slate/Public/Widgets/Text/SMultiLineEditableText.h
@@ -292,6 +292,7 @@ public:
 
 	/** See attribute Font */
 	void SetFont(const TAttribute< FSlateFontInfo >& InNewFont);
+	FSlateFontInfo GetFont() const;
 
 	/** See TextShapingMethod attribute */
 	void SetTextShapingMethod(const TOptional<ETextShapingMethod>& InTextShapingMethod);
diff --git a/Engine/Source/Runtime/SlateCore/Private/Styling/SlateTypes.cpp b/Engine/Source/Runtime/SlateCore/Private/Styling/SlateTypes.cpp
--- a/Engine/Source/Runtime/SlateCore/Private/Styling/SlateTypes.cpp
+++ b/Engine/Source/Runtime/SlateCore/Private/Styling/SlateTypes.cpp
@@ -1,8 +1,10 @@
 // Copyright Epic Games, Inc. All Rights Reserved.
 
 #include "Styling/SlateTypes.h"
+#include "Brushes/SlateColorBrush.h"
 #include "Brushes/SlateNoResource.h"
 #include "Styling/StyleDefaults.h"
+#include "Styling/StyleColors.h"
 #include "Widgets/InvalidateWidgetReason.h"
 #include "Widgets/SWidget.h"
 
@@ -251,7 +253,11 @@ FEditableTextBoxStyle::FEditableTextBoxStyle()
 	, BackgroundImageFocused()
 	, BackgroundImageReadOnly()
 	, Padding(FMargin(4.0f, 2.0f))
-	, Font(FStyleDefaults::GetFontInfo(9))
+	, TextStyle(FTextBlockStyle().SetFont(FStyleDefaults::GetFontInfo(9))
+				.SetColorAndOpacity(FSlateColor::UseForeground())
+				.SetSelectedBackgroundColor(FStyleColors::Highlight)
+				.SetHighlightColor(FStyleColors::Black)
+				.SetHighlightShape(FSlateColorBrush(FStyleColors::AccentGreen)))
 	, ForegroundColor(FSlateColor::UseForeground())
 	, BackgroundColor(FLinearColor::White)
 	, ReadOnlyForegroundColor(FSlateColor::UseForeground())
@@ -259,10 +265,6 @@ FEditableTextBoxStyle::FEditableTextBoxStyle()
 {
 }
 
-FEditableTextBoxStyle::~FEditableTextBoxStyle()
-{
-}
-
 void FEditableTextBoxStyle::GetResources( TArray< const FSlateBrush* >& OutBrushes ) const
 {
 	OutBrushes.Add( &BackgroundImageNormal );
@@ -390,7 +392,7 @@ FSearchBoxStyle& FSearchBoxStyle::SetTextBoxStyle( const FEditableTextBoxStyle&
 	TextBoxStyle = InTextBoxStyle;
 	if (!ActiveFontInfo.HasValidFont())
 	{
-		ActiveFontInfo = TextBoxStyle.Font;
+		ActiveFontInfo = TextBoxStyle.TextStyle.Font;
 	}
 	return *this;
 }
diff --git a/Engine/Source/Runtime/SlateCore/Private/Styling/StarshipCoreStyle.cpp b/Engine/Source/Runtime/SlateCore/Private/Styling/StarshipCoreStyle.cpp
--- a/Engine/Source/Runtime/SlateCore/Private/Styling/StarshipCoreStyle.cpp
+++ b/Engine/Source/Runtime/SlateCore/Private/Styling/StarshipCoreStyle.cpp
@@ -834,6 +834,7 @@ TSharedRef<ISlateStyle> FStarshipCoreStyle::Create()
 		);
 
 		Style->Set("Log.TextBox", FEditableTextBoxStyle(NormalEditableTextBoxStyle)
+			.SetTextStyle(NormalLogText)
 			.SetBackgroundImageNormal(BOX_BRUSH("Common/WhiteGroupBorder", FMargin(4.0f / 16.0f)))
 			.SetBackgroundImageHovered(BOX_BRUSH("Common/WhiteGroupBorder", FMargin(4.0f / 16.0f)))
 			.SetBackgroundImageFocused(BOX_BRUSH("Common/WhiteGroupBorder", FMargin(4.0f / 16.0f)))
@@ -1188,6 +1189,7 @@ void FStarshipCoreStyle::SetupTextStyles(TSharedRef<FStyle>& Style)
 	);
 
 	const FEditableTextBoxStyle DarkEditableTextBoxStyle = FEditableTextBoxStyle()
+		.SetTextStyle(NormalText)
 		.SetBackgroundImageNormal(BOX_BRUSH("Common/TextBox_Dark", FMargin(4.0f / 16.0f)))
 		.SetBackgroundImageHovered(BOX_BRUSH("Common/TextBox_Hovered_Dark", FMargin(4.0f / 16.0f)))
 		.SetBackgroundImageFocused(BOX_BRUSH("Common/TextBox_Hovered_Dark", FMargin(4.0f / 16.0f)))
@@ -1233,6 +1235,7 @@ void FStarshipCoreStyle::SetupTextStyles(TSharedRef<FStyle>& Style)
 
 	// SEditableTextBox defaults...
 	const FEditableTextBoxStyle NormalEditableTextBoxStyle = FEditableTextBoxStyle()
+		.SetTextStyle(NormalText)
 		.SetBackgroundImageNormal(FSlateRoundedBoxBrush(FStyleColors::Input, InputFocusRadius, FStyleColors::InputOutline, InputFocusThickness))
 		.SetBackgroundImageHovered(FSlateRoundedBoxBrush(FStyleColors::Input, InputFocusRadius, FStyleColors::Hover, InputFocusThickness))
 		.SetBackgroundImageFocused(FSlateRoundedBoxBrush(FStyleColors::Input, InputFocusRadius, FStyleColors::Primary, InputFocusThickness))
diff --git a/Engine/Source/Runtime/SlateCore/Private/Styling/UMGCoreStyle.cpp b/Engine/Source/Runtime/SlateCore/Private/Styling/UMGCoreStyle.cpp
--- a/Engine/Source/Runtime/SlateCore/Private/Styling/UMGCoreStyle.cpp
+++ b/Engine/Source/Runtime/SlateCore/Private/Styling/UMGCoreStyle.cpp
@@ -530,6 +530,7 @@ TSharedRef<ISlateStyle> FUMGCoreStyle::Create()
 
 	// SEditableTextBox defaults...
 	static const FEditableTextBoxStyle NormalEditableTextBoxStyle = FEditableTextBoxStyle()
+		.SetTextStyle(NormalText)
 		.SetBackgroundImageNormal( BOX_BRUSH( "Common/TextBox", FMargin(4.0f/16.0f) ) )
 		.SetBackgroundImageHovered( BOX_BRUSH( "Common/TextBox_Hovered", FMargin(4.0f/16.0f) ) )
 		.SetBackgroundImageFocused( BOX_BRUSH( "Common/TextBox_Hovered", FMargin(4.0f/16.0f) ) )
@@ -541,6 +542,7 @@ TSharedRef<ISlateStyle> FUMGCoreStyle::Create()
 	}
 
 	static const FEditableTextBoxStyle DarkEditableTextBoxStyle = FEditableTextBoxStyle()
+		.SetTextStyle(NormalText)
 		.SetBackgroundImageNormal(BOX_BRUSH("Common/TextBox_Dark", FMargin(4.0f / 16.0f)))
 		.SetBackgroundImageHovered(BOX_BRUSH("Common/TextBox_Hovered_Dark", FMargin(4.0f / 16.0f)))
 		.SetBackgroundImageFocused(BOX_BRUSH("Common/TextBox_Hovered_Dark", FMargin(4.0f / 16.0f)))
@@ -585,6 +587,7 @@ TSharedRef<ISlateStyle> FUMGCoreStyle::Create()
 			.SetFont(SmallText.Font);
 
 		FEditableTextBoxStyle InlineEditableTextBlockEditable = FEditableTextBoxStyle()
+			.SetTextStyle(NormalText)
 			.SetFont(NormalText.Font)
 			.SetBackgroundImageNormal( BOX_BRUSH( "Common/TextBox", FMargin(4.0f/16.0f) ) )
 			.SetBackgroundImageHovered( BOX_BRUSH( "Common/TextBox_Hovered", FMargin(4.0f/16.0f) ) )
@@ -673,6 +676,7 @@ TSharedRef<ISlateStyle> FUMGCoreStyle::Create()
 	// SSearchBox defaults...
 	{
 		static const FEditableTextBoxStyle SpecialEditableTextBoxStyle = FEditableTextBoxStyle()
+			.SetTextStyle(NormalText)
 			.SetBackgroundImageNormal( BOX_BRUSH( "Common/TextBox_Special", FMargin(8.0f/32.0f) ) )
 			.SetBackgroundImageHovered( BOX_BRUSH( "Common/TextBox_Special_Hovered", FMargin(8.0f/32.0f) ) )
 			.SetBackgroundImageFocused( BOX_BRUSH( "Common/TextBox_Special_Hovered", FMargin(8.0f/32.0f) ) )
diff --git a/Engine/Source/Runtime/SlateCore/Public/Styling/SlateTypes.h b/Engine/Source/Runtime/SlateCore/Public/Styling/SlateTypes.h
--- a/Engine/Source/Runtime/SlateCore/Public/Styling/SlateTypes.h
+++ b/Engine/Source/Runtime/SlateCore/Public/Styling/SlateTypes.h
@@ -930,9 +930,14 @@ struct SLATECORE_API FEditableTextBoxStyle : public FSlateWidgetStyle
 {
 	GENERATED_USTRUCT_BODY()
 
+	PRAGMA_DISABLE_DEPRECATION_WARNINGS
 	FEditableTextBoxStyle();
+	FEditableTextBoxStyle(const FEditableTextBoxStyle&) = default;
 
-	virtual ~FEditableTextBoxStyle();
+	FEditableTextBoxStyle& operator=(const FEditableTextBoxStyle&) = default;
+
+	virtual ~FEditableTextBoxStyle() = default;
+	PRAGMA_ENABLE_DEPRECATION_WARNINGS
 
 	virtual void GetResources( TArray< const FSlateBrush* >& OutBrushes ) const override;
 
@@ -966,11 +971,19 @@ struct SLATECORE_API FEditableTextBoxStyle : public FSlateWidgetStyle
 	FMargin Padding;
 	FEditableTextBoxStyle& SetPadding( const FMargin& InPadding ){ Padding = InPadding; return *this; }
 
+#if WITH_EDITORONLY_DATA
 	/** Font family and size to be used when displaying this text. */
-	UPROPERTY(EditAnywhere, BlueprintReadWrite, Category=Appearance)
-	FSlateFontInfo Font;
-	FEditableTextBoxStyle& SetFont(const FSlateFontInfo& InFont) { Font = InFont; return *this; }
-	FEditableTextBoxStyle& SetFont(const FName& InFontName, uint16 InSize) { Font = FSlateFontInfo(InFontName, InSize); return *this; }
+	UE_DEPRECATED(5.2, "Font has been deprecated as it was duplicated information already available elsewhere. Please use TextStyle.Font instead.")
+	UPROPERTY(Category = Appearance)
+	FSlateFontInfo Font_DEPRECATED;
+#endif
+	FEditableTextBoxStyle& SetFont(const FSlateFontInfo& InFont) { TextStyle.Font = InFont; return *this; }
+	FEditableTextBoxStyle& SetFont(const FName& InFontName, uint16 InSize) { return SetFont(FSlateFontInfo(InFontName, InSize)); }
+
+	/** The style of the text block, which dictates the font, color, and shadow options. Style overrides all other properties! */
+	UPROPERTY(EditAnywhere, BlueprintReadWrite, Category = Appearance)
+	FTextBlockStyle TextStyle;
+	FEditableTextBoxStyle& SetTextStyle(const FTextBlockStyle& InTextStyle) { TextStyle = InTextStyle; return *this; }
 
 	/** The foreground color of text. */
 	UPROPERTY(EditAnywhere, BlueprintReadWrite, Category=Appearance)
diff --git a/Engine/Source/Runtime/SlateCore/Public/Widgets/DeclarativeSyntaxSupport.h b/Engine/Source/Runtime/SlateCore/Public/Widgets/DeclarativeSyntaxSupport.h
--- a/Engine/Source/Runtime/SlateCore/Public/Widgets/DeclarativeSyntaxSupport.h
+++ b/Engine/Source/Runtime/SlateCore/Public/Widgets/DeclarativeSyntaxSupport.h
@@ -217,55 +217,88 @@ template<typename WidgetType> struct TSlateBaseNamedArgs;
 	SLATE_PRIVATE_ARGUMENT_FUNCTION(ArgType, ArgName) \
 	SLATE_PRIVATE_ARGUMENT_VARIABLE(ArgType, ArgName)
 
-/**
- * Use this macro to declare a slate argument.
- * Arguments differ from attributes in that they can only be values
- */
-#define SLATE_STYLE_ARGUMENT( ArgType, ArgName ) \
+#define SLATE_PRIVATE_STYLE_ARGUMENT_VARIABLE( ArgType, ArgName ) \
 		const ArgType* _##ArgName; \
+
+#define SLATE_PRIVATE_STYLE_ARGUMENT_FUNCTION_ARGTYPE_PTR( ArgType, ArgName ) \
 		WidgetArgsType& ArgName( const ArgType* InArg ) \
 		{ \
 			_##ArgName = InArg; \
 			return static_cast<WidgetArgsType*>(this)->Me(); \
-		}\
-		\
+		}
+#define SLATE_PRIVATE_STYLE_ARGUMENT_FUNCTION_WIDGETSTYLE_PTR( ArgType, ArgName ) \
 		WidgetArgsType& ArgName( const class USlateWidgetStyleAsset* const InSlateStyleAsset ) \
 		{ \
 			_##ArgName = InSlateStyleAsset->GetStyleChecked< ArgType >(); \
 			return static_cast<WidgetArgsType*>(this)->Me(); \
-		}\
-		\
+		}
+#define SLATE_PRIVATE_STYLE_ARGUMENT_FUNCTION_WIDGETSTYLE_WEAKPTR( ArgType, ArgName ) \
 		WidgetArgsType& ArgName( const TWeakObjectPtr< const class USlateWidgetStyleAsset >& InSlateStyleAsset ) \
 		{ \
 			_##ArgName = InSlateStyleAsset.Get()->GetStyleChecked< ArgType >(); \
 			return static_cast<WidgetArgsType*>(this)->Me(); \
-		}\
-		\
+		}
+#define SLATE_PRIVATE_STYLE_ARGUMENT_FUNCTION_SLATESTYLE_PTR( ArgType, ArgName ) \
 		WidgetArgsType& ArgName( const class ISlateStyle* InSlateStyle, const FName& StyleName, const ANSICHAR* Specifier = nullptr ) \
 		{ \
 			check( InSlateStyle != nullptr ); \
 			_##ArgName = &InSlateStyle->GetWidgetStyle< ArgType >( StyleName, Specifier ); \
 			return static_cast<WidgetArgsType*>(this)->Me(); \
-		}\
-		\
+		}
+#define SLATE_PRIVATE_STYLE_ARGUMENT_FUNCTION_SLATESTYLE_REF( ArgType, ArgName ) \
 		WidgetArgsType& ArgName( const class ISlateStyle& InSlateStyle, const FName& StyleName, const ANSICHAR* Specifier = nullptr ) \
 		{ \
 			_##ArgName = &InSlateStyle.GetWidgetStyle< ArgType >( StyleName, Specifier ); \
 			return static_cast<WidgetArgsType*>(this)->Me(); \
-		}\
-		\
+		}
+#define SLATE_PRIVATE_STYLE_ARGUMENT_FUNCTION_SLATESTYLE_WEAKPTR( ArgType, ArgName ) \
 		WidgetArgsType& ArgName( const TWeakObjectPtr< const class ISlateStyle >& InSlateStyle, const FName& StyleName, const ANSICHAR* Specifier = nullptr ) \
 		{ \
 			_##ArgName = &InSlateStyle.Get()->GetWidgetStyle< ArgType >( StyleName, Specifier ); \
 			return static_cast<WidgetArgsType*>(this)->Me(); \
-		} \
-		\
+		}
+#define SLATE_PRIVATE_STYLE_ARGUMENT_FUNCTION_SLATESTYLE_SHAREDPTR( ArgType, ArgName ) \
 		WidgetArgsType& ArgName( const TSharedPtr< const class ISlateStyle >& InSlateStyle, const FName& StyleName, const ANSICHAR* Specifier = nullptr ) \
 		{ \
 			_##ArgName = &InSlateStyle->GetWidgetStyle< ArgType >( StyleName, Specifier ); \
 			return static_cast<WidgetArgsType*>(this)->Me(); \
 		} 
 
+ /**
+  * Use this macro to declare a slate style argument.
+  * Arguments differ from attributes in that they can only be values
+  */
+#define SLATE_STYLE_ARGUMENT( ArgType, ArgName ) \
+	SLATE_PRIVATE_STYLE_ARGUMENT_VARIABLE( ArgType, ArgName ); \
+	SLATE_PRIVATE_STYLE_ARGUMENT_FUNCTION_ARGTYPE_PTR( ArgType, ArgName ) \
+	SLATE_PRIVATE_STYLE_ARGUMENT_FUNCTION_WIDGETSTYLE_PTR( ArgType, ArgName ) \
+	SLATE_PRIVATE_STYLE_ARGUMENT_FUNCTION_WIDGETSTYLE_WEAKPTR( ArgType, ArgName ) \
+	SLATE_PRIVATE_STYLE_ARGUMENT_FUNCTION_SLATESTYLE_PTR( ArgType, ArgName ) \
+	SLATE_PRIVATE_STYLE_ARGUMENT_FUNCTION_SLATESTYLE_REF( ArgType, ArgName ) \
+	SLATE_PRIVATE_STYLE_ARGUMENT_FUNCTION_SLATESTYLE_WEAKPTR( ArgType, ArgName ) \
+	SLATE_PRIVATE_STYLE_ARGUMENT_FUNCTION_SLATESTYLE_SHAREDPTR( ArgType, ArgName )
+
+ /**
+  * Use when a style argument used to be declared SLATE_STYLE_ARGUMENT and it should no longer be used
+  */
+#define SLATE_STYLE_ARGUMENT_DEPRECATED( ArgType, ArgName, DeprecationVersion, DeprecationMessage) \
+	SLATE_PRIVATE_STYLE_ARGUMENT_VARIABLE( ArgType, ArgName ); \
+	PRAGMA_DISABLE_DEPRECATION_WARNINGS\
+	UE_DEPRECATED(DeprecationVersion, DeprecationMessage) \
+	SLATE_PRIVATE_STYLE_ARGUMENT_FUNCTION_ARGTYPE_PTR( ArgType, ArgName ) \
+	UE_DEPRECATED(DeprecationVersion, DeprecationMessage) \
+	SLATE_PRIVATE_STYLE_ARGUMENT_FUNCTION_WIDGETSTYLE_PTR( ArgType, ArgName ) \
+	UE_DEPRECATED(DeprecationVersion, DeprecationMessage) \
+	SLATE_PRIVATE_STYLE_ARGUMENT_FUNCTION_WIDGETSTYLE_WEAKPTR( ArgType, ArgName ) \
+	UE_DEPRECATED(DeprecationVersion, DeprecationMessage) \
+	SLATE_PRIVATE_STYLE_ARGUMENT_FUNCTION_SLATESTYLE_PTR( ArgType, ArgName ) \
+	UE_DEPRECATED(DeprecationVersion, DeprecationMessage) \
+	SLATE_PRIVATE_STYLE_ARGUMENT_FUNCTION_SLATESTYLE_REF( ArgType, ArgName ) \
+	UE_DEPRECATED(DeprecationVersion, DeprecationMessage) \
+	SLATE_PRIVATE_STYLE_ARGUMENT_FUNCTION_SLATESTYLE_WEAKPTR( ArgType, ArgName ) \
+	UE_DEPRECATED(DeprecationVersion, DeprecationMessage) \
+	SLATE_PRIVATE_STYLE_ARGUMENT_FUNCTION_SLATESTYLE_SHAREDPTR( ArgType, ArgName )\
+	PRAGMA_ENABLE_DEPRECATION_WARNINGS
 
 /**
  * Use this macro between SLATE_BEGIN_ARGS and SLATE_END_ARGS
diff --git a/Engine/Source/Runtime/UMG/Private/Components/EditableText.cpp b/Engine/Source/Runtime/UMG/Private/Components/EditableText.cpp
--- a/Engine/Source/Runtime/UMG/Private/Components/EditableText.cpp
+++ b/Engine/Source/Runtime/UMG/Private/Components/EditableText.cpp
@@ -34,6 +34,12 @@ UEditableText::UEditableText(const FObjectInitializer& ObjectInitializer)
 	}
 
 	WidgetStyle = *DefaultEditableTextStyle;
+	if (!IsRunningDedicatedServer())
+	{
+		static ConstructorHelpers::FObjectFinder<UFont> defaultFontObj(*UWidget::GetDefaultFontName());
+		WidgetStyle.SetFont(FSlateFontInfo(defaultFontObj.Object, 24, FName("Regular")));
+	}
+
 
 #if WITH_EDITOR 
 	if (EditorEditableTextStyle == nullptr)
@@ -117,6 +123,7 @@ void UEditableText::SynchronizeProperties()
 	TAttribute<FText> HintTextBinding = PROPERTY_BINDING(FText, HintText);
 
 	MyEditableText->SetText(TextBinding);
+	MyEditableText->SetEditableTextStyle(WidgetStyle);
 	MyEditableText->SetHintText(HintTextBinding);
 	MyEditableText->SetIsReadOnly(IsReadOnly);
 	MyEditableText->SetIsPassword(IsPassword);
@@ -341,7 +348,7 @@ void UEditableText::SetFont(FSlateFontInfo InFontInfo)
 
 	if (MyEditableText.IsValid())
 	{
-		MyEditableText->SetTextStyle(WidgetStyle);
+		MyEditableText->SetEditableTextStyle(WidgetStyle);
 	}
 }
 
@@ -351,7 +358,7 @@ void UEditableText::SetFontMaterial(UMaterialInterface* InMaterial)
 
 	if (MyEditableText.IsValid())
 	{
-		MyEditableText->SetTextStyle(WidgetStyle);
+		MyEditableText->SetEditableTextStyle(WidgetStyle);
 	}
 }
 
@@ -361,7 +368,7 @@ void UEditableText::SetFontOutlineMaterial(UMaterialInterface* InMaterial)
 
 	if (MyEditableText.IsValid())
 	{
-		MyEditableText->SetTextStyle(WidgetStyle);
+		MyEditableText->SetEditableTextStyle(WidgetStyle);
 	}
 }
 
diff --git a/Engine/Source/Runtime/UMG/Private/Components/EditableTextBox.cpp b/Engine/Source/Runtime/UMG/Private/Components/EditableTextBox.cpp
--- a/Engine/Source/Runtime/UMG/Private/Components/EditableTextBox.cpp
+++ b/Engine/Source/Runtime/UMG/Private/Components/EditableTextBox.cpp
@@ -2,6 +2,8 @@
 
 #include "Components/EditableTextBox.h"
 #include "UObject/ConstructorHelpers.h"
+#include "UObject/UE5MainStreamObjectVersion.h"
+#include "UObject/UE5ReleaseStreamObjectVersion.h"
 #include "Engine/Font.h"
 #include "Widgets/DeclarativeSyntaxSupport.h"
 #include "Widgets/Input/SEditableTextBox.h"
@@ -72,6 +74,28 @@ UEditableTextBox::UEditableTextBox(const FObjectInitializer& ObjectInitializer)
 #endif
 }
 
+#if WITH_EDITORONLY_DATA
+void UEditableTextBox::PostLoad()
+{
+	Super::PostLoad();
+
+PRAGMA_DISABLE_DEPRECATION_WARNINGS
+	if (GetLinkerCustomVersion(FUE5ReleaseStreamObjectVersion::GUID) < FUE5ReleaseStreamObjectVersion::RemoveDuplicatedStyleInfo)
+	{
+		FTextBlockStyle& TextStyle = WidgetStyle.TextStyle;
+		TextStyle.SetFont(WidgetStyle.Font_DEPRECATED);
+	}
+PRAGMA_ENABLE_DEPRECATION_WARNINGS
+}
+
+void UEditableTextBox::Serialize(FArchive& Ar)
+{
+	Ar.UsingCustomVersion(FUE5ReleaseStreamObjectVersion::GUID);
+
+	Super::Serialize(Ar);
+}
+#endif
+
 void UEditableTextBox::ReleaseSlateResources(bool bReleaseChildren)
 {
 	Super::ReleaseSlateResources(bReleaseChildren);
diff --git a/Engine/Source/Runtime/UMG/Private/Components/MultiLineEditableText.cpp b/Engine/Source/Runtime/UMG/Private/Components/MultiLineEditableText.cpp
--- a/Engine/Source/Runtime/UMG/Private/Components/MultiLineEditableText.cpp
+++ b/Engine/Source/Runtime/UMG/Private/Components/MultiLineEditableText.cpp
@@ -120,8 +120,6 @@ void UMultiLineEditableText::SynchronizeProperties()
 	MyMultiLineEditableText->SetRevertTextOnEscape(RevertTextOnEscape);
 	MyMultiLineEditableText->SetClearKeyboardFocusOnCommit(ClearKeyboardFocusOnCommit);
 
-//	MyMultiLineEditableText->SetColorAndOpacity(ColorAndOpacity);
-
 	// TODO UMG Complete making all properties settable on SMultiLineEditableText
 
 	Super::SynchronizeTextLayoutProperties(*MyMultiLineEditableText);
diff --git a/Engine/Source/Runtime/UMG/Private/Components/MultiLineEditableTextBox.cpp b/Engine/Source/Runtime/UMG/Private/Components/MultiLineEditableTextBox.cpp
--- a/Engine/Source/Runtime/UMG/Private/Components/MultiLineEditableTextBox.cpp
+++ b/Engine/Source/Runtime/UMG/Private/Components/MultiLineEditableTextBox.cpp
@@ -2,6 +2,8 @@
 
 #include "Components/MultiLineEditableTextBox.h"
 #include "UObject/ConstructorHelpers.h"
+#include "UObject/UE5MainStreamObjectVersion.h"
+#include "UObject/UE5ReleaseStreamObjectVersion.h"
 #include "Engine/Font.h"
 #include "Widgets/DeclarativeSyntaxSupport.h"
 #include "Widgets/Input/SMultiLineEditableTextBox.h"
@@ -15,11 +17,9 @@
 // UMultiLineEditableTextBox
 
 static FEditableTextBoxStyle* DefaultMultiLineEditableTextBoxStyle = nullptr;
-static FTextBlockStyle* DefaultMultiLineEditableTextBoxTextStyle = nullptr;
 
 #if WITH_EDITOR
 static FEditableTextBoxStyle* EditorMultiLineEditableTextBoxStyle = nullptr;
-static FTextBlockStyle* EditorMultiLineEditableTextBoxTextStyle = nullptr;
 #endif 
 
 PRAGMA_DISABLE_DEPRECATION_WARNINGS
@@ -34,16 +34,16 @@ UMultiLineEditableTextBox::UMultiLineEditableTextBox(const FObjectInitializer& O
 		DefaultMultiLineEditableTextBoxStyle->UnlinkColors();
 	}
 
-	if (DefaultMultiLineEditableTextBoxTextStyle == nullptr)
-	{
-		DefaultMultiLineEditableTextBoxTextStyle = new FTextBlockStyle(FUMGCoreStyle::Get().GetWidgetStyle<FTextBlockStyle>("NormalText"));
+	WidgetStyle = *DefaultMultiLineEditableTextBoxStyle;
+#if WITH_EDITOR
+	TextStyle_DEPRECATED = DefaultMultiLineEditableTextBoxStyle->TextStyle;
+#endif
 
-		// Unlink UMG default colors.
-		DefaultMultiLineEditableTextBoxTextStyle->UnlinkColors();
+	if (!IsRunningDedicatedServer())
+	{
+		static ConstructorHelpers::FObjectFinder<UFont> defaultFontObj(*UWidget::GetDefaultFontName());
+		WidgetStyle.SetFont(FSlateFontInfo(defaultFontObj.Object, 24, FName("Regular")));
 	}
-	
-	WidgetStyle = *DefaultMultiLineEditableTextBoxStyle;
-	TextStyle = *DefaultMultiLineEditableTextBoxTextStyle;
 
 #if WITH_EDITOR 
 	if (EditorMultiLineEditableTextBoxStyle == nullptr)
@@ -53,19 +53,11 @@ UMultiLineEditableTextBox::UMultiLineEditableTextBox(const FObjectInitializer& O
 		// Unlink UMG Editor colors from the editor settings colors.
 		EditorMultiLineEditableTextBoxStyle->UnlinkColors();
 	}
-
-	if (EditorMultiLineEditableTextBoxTextStyle == nullptr)
-	{
-		EditorMultiLineEditableTextBoxTextStyle = new FTextBlockStyle(FCoreStyle::Get().GetWidgetStyle<FTextBlockStyle>("NormalText"));
-
-		// Unlink UMG Editor colors from the editor settings colors.
-		EditorMultiLineEditableTextBoxTextStyle->UnlinkColors();
-	}
 	
 	if (IsEditorWidget())
 	{
 		WidgetStyle = *EditorMultiLineEditableTextBoxStyle;
-		TextStyle = *EditorMultiLineEditableTextBoxTextStyle;
+		TextStyle_DEPRECATED = DefaultMultiLineEditableTextBoxStyle->TextStyle;
 
 		// The CDO isn't an editor widget and thus won't use the editor style, call post edit change to mark difference from CDO
 		PostEditChange();
@@ -78,6 +70,28 @@ UMultiLineEditableTextBox::UMultiLineEditableTextBox(const FObjectInitializer& O
 	AutoWrapText = true;
 }
 
+#if WITH_EDITOR
+void UMultiLineEditableTextBox::PostLoad()
+{
+	Super::PostLoad();
+
+	if (GetLinkerCustomVersion(FUE5ReleaseStreamObjectVersion::GUID) < FUE5ReleaseStreamObjectVersion::RemoveDuplicatedStyleInfo)
+	{
+	PRAGMA_DISABLE_DEPRECATION_WARNINGS
+		TextStyle_DEPRECATED.SetFont(WidgetStyle.Font_DEPRECATED);
+		WidgetStyle.SetTextStyle(TextStyle_DEPRECATED);
+	PRAGMA_ENABLE_DEPRECATION_WARNINGS
+	}
+}
+#endif
+
+void UMultiLineEditableTextBox::Serialize(FArchive& Ar)
+{
+	Super::Serialize(Ar);
+
+	Ar.UsingCustomVersion(FUE5ReleaseStreamObjectVersion::GUID);
+}
+
 void UMultiLineEditableTextBox::ReleaseSlateResources(bool bReleaseChildren)
 {
 	Super::ReleaseSlateResources(bReleaseChildren);
@@ -89,7 +103,6 @@ TSharedRef<SWidget> UMultiLineEditableTextBox::RebuildWidget()
 {
 	MyEditableTextBlock = SNew(SMultiLineEditableTextBox)
 		.Style(&WidgetStyle)
-		.TextStyle(&TextStyle)
 		.AllowContextMenu(AllowContextMenu)
 		.IsReadOnly(bIsReadOnly)
 //		.MinDesiredWidth(MinimumDesiredWidth)
@@ -206,11 +219,11 @@ PRAGMA_ENABLE_DEPRECATION_WARNINGS
 
 void UMultiLineEditableTextBox::SetTextStyle(const FTextBlockStyle& InTextStyle)
 {
-	TextStyle = InTextStyle;
+	WidgetStyle.SetTextStyle(InTextStyle);
 
 	if (MyEditableTextBlock.IsValid())
 	{
-		MyEditableTextBlock->SetTextStyle(&TextStyle);
+		MyEditableTextBlock->SetTextStyle(&InTextStyle);
 	}
 }
 
diff --git a/Engine/Source/Runtime/UMG/Public/Components/EditableTextBox.h b/Engine/Source/Runtime/UMG/Public/Components/EditableTextBox.h
--- a/Engine/Source/Runtime/UMG/Public/Components/EditableTextBox.h
+++ b/Engine/Source/Runtime/UMG/Public/Components/EditableTextBox.h
@@ -264,6 +264,11 @@ public:
 	virtual const FText GetPaletteCategory() override;
 #endif
 
+#if WITH_EDITORONLY_DATA
+	virtual void Serialize(FArchive& Ar) override;
+	virtual void PostLoad() override;
+#endif
+
 protected:
 	//~ Begin UWidget Interface
 	virtual TSharedRef<SWidget> RebuildWidget() override;
diff --git a/Engine/Source/Runtime/UMG/Public/Components/MultiLineEditableTextBox.h b/Engine/Source/Runtime/UMG/Public/Components/MultiLineEditableTextBox.h
--- a/Engine/Source/Runtime/UMG/Public/Components/MultiLineEditableTextBox.h
+++ b/Engine/Source/Runtime/UMG/Public/Components/MultiLineEditableTextBox.h
@@ -48,9 +48,12 @@ public:
 	UPROPERTY(EditAnywhere, BlueprintReadWrite, Category="Style", meta=(DisplayName="Style"))
 	FEditableTextBoxStyle WidgetStyle;
 
+#if WITH_EDITORONLY_DATA
 	/** The text style */
-	UPROPERTY(EditAnywhere, BlueprintReadWrite, BlueprintSetter = SetTextStyle, Category="Style", meta=(DisplayName="Text Style"))
-	FTextBlockStyle TextStyle;
+	UE_DEPRECATED(5.2, "TextStyle has been deprecated as it was mainly duplicated information already available inside WidgetStyle. Please use the WidgetStyle.TextStyle instead.")
+	UPROPERTY(Category="Style", meta=(DisplayName="Text Style"))
+	FTextBlockStyle TextStyle_DEPRECATED;
+#endif
 
 	/** Sets the Text as Readonly to prevent it from being modified interactively by the user */
 	UE_DEPRECATED(5.1, "Direct access to IsReadOnly is deprecated. Please use the getter or setter.")
@@ -118,7 +121,7 @@ public:
 	UFUNCTION(BlueprintCallable, Category = "Widget", meta = (DisplayName = "SetIsReadOnly (Multi-Line Text Box)"))
 	void SetIsReadOnly(UPARAM(DisplayName = "ReadyOnly") bool bReadOnly);
 
-	UFUNCTION(BlueprintSetter)
+	UFUNCTION(BlueprintCallable, Category="Widget", meta=(DisplayName="SetTextStyle (Multi-Line Text Box)"))
 	void SetTextStyle(const FTextBlockStyle& InTextStyle);
 
 	UFUNCTION(BlueprintCallable, Category="Widget", meta=(DisplayName="SetForegroundColor (Multi-Line Text Box)"))
@@ -143,8 +146,11 @@ public:
 
 #if WITH_EDITOR
 	virtual const FText GetPaletteCategory() override;
+	virtual void PostLoad() override;
 #endif
 
+	virtual void Serialize(FArchive& Ar) override;
+
 protected:
 	//~ Begin UWidget Interface
 	virtual TSharedRef<SWidget> RebuildWidget() override;
