Removed category from two deprecated properties.

diff --git a/Engine/Source/Runtime/SlateCore/Public/Styling/SlateTypes.h b/Engine/Source/Runtime/SlateCore/Public/Styling/SlateTypes.h
--- a/Engine/Source/Runtime/SlateCore/Public/Styling/SlateTypes.h
+++ b/Engine/Source/Runtime/SlateCore/Public/Styling/SlateTypes.h
@@ -974,7 +974,7 @@ struct SLATECORE_API FEditableTextBoxStyle : public FSlateWidgetStyle
 #if WITH_EDITORONLY_DATA
 	/** Font family and size to be used when displaying this text. */
 	UE_DEPRECATED(5.2, "Font has been deprecated as it was duplicated information already available elsewhere. Please use TextStyle.Font instead.")
-	UPROPERTY(Category = Appearance)
+	UPROPERTY()
 	FSlateFontInfo Font_DEPRECATED;
 #endif
 	FEditableTextBoxStyle& SetFont(const FSlateFontInfo& InFont) { TextStyle.Font = InFont; return *this; }
diff --git a/Engine/Source/Runtime/UMG/Public/Components/MultiLineEditableTextBox.h b/Engine/Source/Runtime/UMG/Public/Components/MultiLineEditableTextBox.h
--- a/Engine/Source/Runtime/UMG/Public/Components/MultiLineEditableTextBox.h
+++ b/Engine/Source/Runtime/UMG/Public/Components/MultiLineEditableTextBox.h
@@ -51,7 +51,7 @@ public:
 #if WITH_EDITORONLY_DATA
 	/** The text style */
 	UE_DEPRECATED(5.2, "TextStyle has been deprecated as it was mainly duplicated information already available inside WidgetStyle. Please use the WidgetStyle.TextStyle instead.")
-	UPROPERTY(Category="Style", meta=(DisplayName="Text Style"))
+	UPROPERTY(meta=(DisplayName="Text Style"))
 	FTextBlockStyle TextStyle_DEPRECATED;
 #endif
 
