Fix heap corruption due to multiple async calls to symbol load

diff --git a/Engine/Source/Runtime/SlateCore/Private/Trace/SlateTrace.cpp b/Engine/Source/Runtime/SlateCore/Private/Trace/SlateTrace.cpp
--- a/Engine/Source/Runtime/SlateCore/Private/Trace/SlateTrace.cpp
+++ b/Engine/Source/Runtime/SlateCore/Private/Trace/SlateTrace.cpp
@@ -14,6 +14,7 @@
 #include "Types/ReflectionMetadata.h"
 #include "Widgets/SWidget.h"
 #include "FastUpdate/WidgetProxy.h"
+#include "Tasks/Task.h"
 
 #if !(UE_BUILD_SHIPPING)
 
@@ -23,6 +24,11 @@ static FAutoConsoleVariableRef CVarCaptureRootInvalidationCallstacks(
 	bCaptureRootInvalidationCallstacks,
 	TEXT("Whenever a widget is the root cause of an invalidation, capture the callstack for slate insights."));
 
+namespace UE::Slate::Private
+{
+	static FCriticalSection ModuleLoadLock;
+}
+
 #endif // !(UE_BUILD_SHIPPING)
 
 //-----------------------------------------------------------------------------------//
@@ -396,10 +402,14 @@ void FSlateTrace::WidgetInvalidated(const SWidget* Widget, const SWidget* Invest
 			<< WidgetInvalidated.InvalidateWidgetReason(static_cast<uint8>(Reason));
 
 #if !(UE_BUILD_SHIPPING)
-		if (!Investigator && bCaptureRootInvalidationCallstacks)
+		// Don't create an async task that locks when the engine is shutting down
+		if (!Investigator && bCaptureRootInvalidationCallstacks && !IsEngineExitRequested())
 		{
+			// Use Low since Normal always crashes
+			LowLevelTasks::ETaskPriority SymbolResolvePriority = LowLevelTasks::ETaskPriority::BackgroundLow;
+
 			// Note: Done in seperate thread as symbol resolution is very slow, possibly 1~80ms based on widget complexity.
-			AsyncTask(ENamedThreads::AnyBackgroundThreadNormalTask, [Cycle, StackTrace, StackTraceDepth]()
+			UE::Tasks::Launch(UE_SOURCE_LOCATION, [Cycle, StackTrace, StackTraceDepth]()
 			{
 				const SIZE_T CallStackSize = 65535;
 				ANSICHAR* CallStack = (ANSICHAR*)FMemory::SystemMalloc(CallStackSize);
@@ -407,19 +417,25 @@ void FSlateTrace::WidgetInvalidated(const SWidget* Widget, const SWidget* Invest
 				{
 					CallStack[0] = 0;
 				}
-
-				uint32 StackDepth = 0;
-				while (StackDepth < StackTraceDepth)
+				
+				// Need to lock on certain platforms when loading symbols, due to single threaded loading
+				// See: https://docs.microsoft.com/en-us/windows/win32/api/dbghelp/nf-dbghelp-symloadmoduleexw#remarks
 				{
-					// Skip the first two backraces, that's from us.
-					if (StackDepth >= 2 && StackTrace[StackDepth])
+					FScopeLock Lock(&UE::Slate::Private::ModuleLoadLock);
+
+					uint32 StackDepth = 0;
+					while (StackDepth < StackTraceDepth)
 					{
-						FProgramCounterSymbolInfo SymbolInfo;
-						FPlatformStackWalk::ProgramCounterToSymbolInfo(StackTrace[StackDepth], SymbolInfo);
-						FPlatformStackWalk::SymbolInfoToHumanReadableString(SymbolInfo, CallStack, CallStackSize);
-						FCStringAnsi::Strncat(CallStack, LINE_TERMINATOR_ANSI, CallStackSize);
+						// Skip the first two backraces, that's from us.
+						if (StackDepth >= 2 && StackTrace[StackDepth])
+						{
+							FProgramCounterSymbolInfo SymbolInfo;
+							FPlatformStackWalk::ProgramCounterToSymbolInfo(StackTrace[StackDepth], SymbolInfo);
+							FPlatformStackWalk::SymbolInfoToHumanReadableString(SymbolInfo, CallStack, CallStackSize);
+							FCStringAnsi::Strncat(CallStack, LINE_TERMINATOR_ANSI, CallStackSize);
+						}
+						StackDepth++;
 					}
-					StackDepth++;
 				}
 
 				UE_TRACE_LOG(SlateTrace, InvalidationCallstack, SlateChannel)
@@ -427,7 +443,7 @@ void FSlateTrace::WidgetInvalidated(const SWidget* Widget, const SWidget* Invest
 					<< InvalidationCallstack.CallstackText(CallStack);
 
 				FMemory::SystemFree(CallStack);
-			});
+			}, SymbolResolvePriority);
 		}
 #endif // !(UE_BUILD_SHIPPING)
 	}
