After recent editable text box and style dev, in the case where no font was configured on an editable text box, and when editing in the designer, the 'no font' symbols were displayed instead of the actual glyph.
While before the change, a font was diplayed... temporarily (it would disappear  if you started to play with the font settings and moved the object in the viewport).

A font will now be displayed as before, but will still disappear in the same condition, but it's not valid anyway to have no font, and soon (hopefully with 5.1) we will also assign a default font to avoid having this issue.

Tests: in the editor, opened in the UMG designer problematic assets (old misconfigured data, on which the deprecation code would be applyied, and validated that a font was used as before). Opened other assets without issue.

diff --git a/Engine/Source/Runtime/SlateCore/Private/Styling/SlateTypes.cpp b/Engine/Source/Runtime/SlateCore/Private/Styling/SlateTypes.cpp
--- a/Engine/Source/Runtime/SlateCore/Private/Styling/SlateTypes.cpp
+++ b/Engine/Source/Runtime/SlateCore/Private/Styling/SlateTypes.cpp
@@ -253,6 +253,9 @@ FEditableTextBoxStyle::FEditableTextBoxStyle()
 	, BackgroundImageFocused()
 	, BackgroundImageReadOnly()
 	, Padding(FMargin(4.0f, 2.0f))
+#if WITH_EDITORONLY_DATA
+	, Font_DEPRECATED(FStyleDefaults::GetFontInfo(9))
+#endif
 	, TextStyle(FTextBlockStyle().SetFont(FStyleDefaults::GetFontInfo(9))
 				.SetColorAndOpacity(FSlateColor::UseForeground())
 				.SetSelectedBackgroundColor(FStyleColors::Highlight)
