Slate: Added interface to explicitly disable Slate throttling through code.

diff --git a/Engine/Source/Runtime/SlateCore/Private/Application/ThrottleManager.cpp b/Engine/Source/Runtime/SlateCore/Private/Application/ThrottleManager.cpp
--- a/Engine/Source/Runtime/SlateCore/Private/Application/ThrottleManager.cpp
+++ b/Engine/Source/Runtime/SlateCore/Private/Application/ThrottleManager.cpp
@@ -10,6 +10,7 @@ FSlateThrottleManager::FSlateThrottleManager( )
 	: bShouldThrottle(1)
 	, CVarAllowThrottle(TEXT("Slate.bAllowThrottling"), bShouldThrottle, TEXT("Allow Slate to throttle parts of the engine to ensure the UI is responsive") )
 	, ThrottleCount(0)
+	, DisableThrottleCount(0)
 { }
 
 
@@ -33,7 +34,7 @@ bool FSlateThrottleManager::IsAllowingExpensiveTasks( ) const
 {
 	// Expensive tasks are allowed if the number of active throttle requests is zero
 	// or we always always allow it due to a CVar
-	return (ThrottleCount == 0) || !bShouldThrottle;
+	return (ThrottleCount == 0) || !bShouldThrottle || DisableThrottleCount > 0;
 }
 
 
@@ -49,6 +50,13 @@ void FSlateThrottleManager::LeaveResponsiveMode( FThrottleRequest& InHandle )
 	}
 }
 
+void FSlateThrottleManager::DisableThrottle(bool bState)
+{
+	DisableThrottleCount += bState ? 1 : -1;
+	ensure(DisableThrottleCount >= 0);
+}
+
+
 
 /* FSlateThrottleManager static functions
  *****************************************************************************/
diff --git a/Engine/Source/Runtime/SlateCore/Public/Application/ThrottleManager.h b/Engine/Source/Runtime/SlateCore/Public/Application/ThrottleManager.h
--- a/Engine/Source/Runtime/SlateCore/Public/Application/ThrottleManager.h
+++ b/Engine/Source/Runtime/SlateCore/Public/Application/ThrottleManager.h
@@ -75,6 +75,15 @@ public:
 	 */
 	bool IsAllowingExpensiveTasks() const;
 
+	/**
+	 * Explicitly disable Slate throttling. This is intended to be used for code-driven exemptions
+	 * to Slate throttling such as interactive actions that require multiple refreshes & world ticks.
+	 * For single refresh scenarios, consider FEditorViewportClient.Invalidate(false, false).
+	 *
+	 * @param bState true if throttling should be disabled.
+	 */
+	void DisableThrottle(bool bState);
+
 public:
 
 	/**
@@ -91,4 +100,7 @@ private:
 
 	/** Number of active throttle requests */
 	uint32 ThrottleCount;
+
+	/** If false, then throttling is explicitly disabled. */
+	int32 DisableThrottleCount;
 };
