Increase background threading priority of callstack capture in slate insights.

diff --git a/Engine/Source/Runtime/SlateCore/Private/Trace/SlateTrace.cpp b/Engine/Source/Runtime/SlateCore/Private/Trace/SlateTrace.cpp
--- a/Engine/Source/Runtime/SlateCore/Private/Trace/SlateTrace.cpp
+++ b/Engine/Source/Runtime/SlateCore/Private/Trace/SlateTrace.cpp
@@ -405,8 +405,7 @@ void FSlateTrace::WidgetInvalidated(const SWidget* Widget, const SWidget* Invest
 		// Don't create an async task that locks when the engine is shutting down
 		if (!Investigator && bCaptureRootInvalidationCallstacks && !IsEngineExitRequested())
 		{
-			// Use Low since Normal always crashes
-			LowLevelTasks::ETaskPriority SymbolResolvePriority = LowLevelTasks::ETaskPriority::BackgroundLow;
+			LowLevelTasks::ETaskPriority SymbolResolvePriority = LowLevelTasks::ETaskPriority::BackgroundNormal;
 
 			// Note: Done in seperate thread as symbol resolution is very slow, possibly 1~80ms based on widget complexity.
 			UE::Tasks::Launch(UE_SOURCE_LOCATION, [Cycle, StackTrace, StackTraceDepth]()
