[SlateCore]
* DrawElementPayloads - Added include to fix non-unity error

diff --git a/Engine/Source/Runtime/SlateCore/Public/Rendering/DrawElementPayloads.h b/Engine/Source/Runtime/SlateCore/Public/Rendering/DrawElementPayloads.h
--- a/Engine/Source/Runtime/SlateCore/Public/Rendering/DrawElementPayloads.h
+++ b/Engine/Source/Runtime/SlateCore/Public/Rendering/DrawElementPayloads.h
@@ -4,6 +4,7 @@
 
 #include "DrawElementTextOverflowArgs.h"
 #include "Fonts/FontCache.h"
+#include "Fonts/ShapedTextFwd.h"
 #include "SlateRenderBatch.h"
 #include "Styling/SlateBrush.h"
 #include "Types/SlateVector2.h"
