Improve Scrollbox slot to support the same settings than the vertical/horizontal/stack box:
- size rule auto and strech.
- max size.
Expose the size rule also in the ScrollBoxSlot for UMG.

tests:
-reverted some workaround in the code and used the new FillSize slot parameter with success.
-tested a lot of sub editors, as the code also impact vertical/horizontal box.
-created new UMG data in QAGame using the new attributes (will submit them soon), and used blueprint graph to change the size settings of a scrollbox slot.

diff --git a/Engine/Source/Runtime/Slate/Private/Widgets/Layout/SScrollBox.cpp b/Engine/Source/Runtime/Slate/Private/Widgets/Layout/SScrollBox.cpp
--- a/Engine/Source/Runtime/Slate/Private/Widgets/Layout/SScrollBox.cpp
+++ b/Engine/Source/Runtime/Slate/Private/Widgets/Layout/SScrollBox.cpp
@@ -9,6 +9,27 @@
 #include "Widgets/Images/SImage.h"
 
 
+void SScrollBox::FSlot::Construct(const FChildren& SlotOwner, FSlotArguments&& InArgs)
+{
+	TBasicLayoutWidgetSlot<FSlot>::Construct(SlotOwner, MoveTemp(InArgs));
+	if (InArgs._MaxSize.IsSet())
+	{
+		SetMaxSize(MoveTemp(InArgs._MaxSize));
+	}
+	if (InArgs._SizeParam.IsSet())
+	{
+		SetSizeParam(MoveTemp(InArgs._SizeParam.GetValue()));
+	}
+}
+
+void SScrollBox::FSlot::RegisterAttributes(FSlateWidgetSlotAttributeInitializer& AttributeInitializer)
+{
+	TBasicLayoutWidgetSlot<FSlot>::RegisterAttributes(AttributeInitializer);
+	SLATE_ADD_SLOT_ATTRIBUTE_DEFINITION_WITH_NAME(FSlot, AttributeInitializer, "Slot.MaxSize", MaxSize, EInvalidateWidgetReason::Layout);
+	SLATE_ADD_SLOT_ATTRIBUTE_DEFINITION_WITH_NAME(FSlot, AttributeInitializer, "Slot.SizeValue", SizeValue, EInvalidateWidgetReason::Layout)
+		.UpdatePrerequisite("Slot.MaxSize");
+}
+
 SScrollBox::FSlot::FSlotArguments SScrollBox::Slot()
 {
 	return FSlot::FSlotArguments(MakeUnique<FSlot>());
@@ -41,26 +62,17 @@ void SScrollPanel::Construct(const FArguments& InArgs, TArray<SScrollBox::FSlot:
 
 void SScrollPanel::OnArrangeChildren(const FGeometry& AllottedGeometry, FArrangedChildren& ArrangedChildren) const
 {
-	float ScrollPadding = Orientation == Orient_Vertical ? AllottedGeometry.GetLocalSize().Y : AllottedGeometry.GetLocalSize().X;
-	float CurChildOffset = -PhysicalOffset; 
-	CurChildOffset += BackPadScrolling ? ScrollPadding : 0;
+	const float ScrollPadding = Orientation == Orient_Vertical ? AllottedGeometry.GetLocalSize().Y : AllottedGeometry.GetLocalSize().X;
+	const float ChildrenOffset = -PhysicalOffset + (BackPadScrolling ? ScrollPadding : 0);
+	const bool AllowShrink = false;
 
-	for (int32 SlotIndex = 0; SlotIndex < Children.Num(); ++SlotIndex)
+	if (Orientation == EOrientation::Orient_Horizontal)
 	{
-		const SScrollBox::FSlot& ThisSlot = Children[SlotIndex];
-		const EVisibility ChildVisibility = ThisSlot.GetWidget()->GetVisibility();
-
-		if (ChildVisibility != EVisibility::Collapsed)
-		{
-			if (Orientation == Orient_Vertical)
-			{
-				CurChildOffset = ArrangeChildVerticalAndReturnOffset(AllottedGeometry, ArrangedChildren, ThisSlot, CurChildOffset);
-			}
-			else
-			{
-				CurChildOffset = ArrangeChildHorizontalAndReturnOffset(AllottedGeometry, ArrangedChildren, ThisSlot, CurChildOffset);
-			}
-		}
+		ArrangeChildrenInStack<EOrientation::Orient_Horizontal>(GSlateFlowDirection, this->Children, AllottedGeometry, ArrangedChildren, ChildrenOffset, AllowShrink);
+	}
+	else
+	{
+		ArrangeChildrenInStack<EOrientation::Orient_Vertical>(GSlateFlowDirection, this->Children, AllottedGeometry, ArrangedChildren, ChildrenOffset, AllowShrink);
 	}
 }
 
@@ -75,13 +87,13 @@ FVector2D SScrollPanel::ComputeDesiredSize(float) const
 			const FVector2D ChildDesiredSize = ThisSlot.GetWidget()->GetDesiredSize();
 			if (Orientation == Orient_Vertical)
 			{
-				ThisDesiredSize.X = FMath::Max(ChildDesiredSize.X, ThisDesiredSize.X);
+				ThisDesiredSize.X = FMath::Max(ChildDesiredSize.X + ThisSlot.GetPadding().GetTotalSpaceAlong<Orient_Horizontal>(), ThisDesiredSize.X);
 				ThisDesiredSize.Y += ChildDesiredSize.Y + ThisSlot.GetPadding().GetTotalSpaceAlong<Orient_Vertical>();
 			}
 			else
 			{
 				ThisDesiredSize.X += ChildDesiredSize.X + ThisSlot.GetPadding().GetTotalSpaceAlong<Orient_Horizontal>();
-				ThisDesiredSize.Y = FMath::Max(ChildDesiredSize.Y, ThisDesiredSize.Y);
+				ThisDesiredSize.Y = FMath::Max(ChildDesiredSize.Y + ThisSlot.GetPadding().GetTotalSpaceAlong<Orient_Vertical>(), ThisDesiredSize.Y);
 			}
 		}
 	}
@@ -94,34 +106,6 @@ FVector2D SScrollPanel::ComputeDesiredSize(float) const
 	return ThisDesiredSize;
 }
 
-float SScrollPanel::ArrangeChildVerticalAndReturnOffset(const FGeometry& AllottedGeometry, FArrangedChildren& ArrangedChildren, const SScrollBox::FSlot& ThisSlot, float CurChildOffset) const
-{
-	const FMargin& ThisPadding = ThisSlot.GetPadding();
-	const FVector2D& WidgetDesiredSize = ThisSlot.GetWidget()->GetDesiredSize();
-	const float ThisSlotDesiredHeight = WidgetDesiredSize.Y + ThisPadding.GetTotalSpaceAlong<Orient_Vertical>();
-
-	// Figure out the size and local position of the child within the slot.  There is no vertical alignment, because 
-	// it does not make sense in a panel where items are stacked vertically end-to-end.
-	AlignmentArrangeResult XAlignmentResult = AlignChild<Orient_Horizontal>(AllottedGeometry.GetLocalSize().X, ThisSlot, ThisPadding);
-
-	ArrangedChildren.AddWidget(AllottedGeometry.MakeChild(ThisSlot.GetWidget(), FVector2D(XAlignmentResult.Offset, CurChildOffset + ThisPadding.Top), FVector2D(XAlignmentResult.Size, WidgetDesiredSize.Y)));
-	return CurChildOffset + ThisSlotDesiredHeight;
-}
-
-float SScrollPanel::ArrangeChildHorizontalAndReturnOffset(const FGeometry& AllottedGeometry, FArrangedChildren& ArrangedChildren, const SScrollBox::FSlot& ThisSlot, float CurChildOffset) const
-{
-	const FMargin& ThisPadding = ThisSlot.GetPadding();
-	const FVector2D& WidgetDesiredSize = ThisSlot.GetWidget()->GetDesiredSize();
-	const float ThisSlotDesiredWidth = WidgetDesiredSize.X + ThisPadding.GetTotalSpaceAlong<Orient_Horizontal>();
-
-	// Figure out the size and local position of the child within the slot.  There is no horizontal alignment, because
-	// it doesn't make sense in a panel where items are stacked horizontally end-to-end.
-	AlignmentArrangeResult YAlignmentResult = AlignChild<Orient_Vertical>(AllottedGeometry.GetLocalSize().Y, ThisSlot, ThisPadding);
-
-	ArrangedChildren.AddWidget(AllottedGeometry.MakeChild(ThisSlot.GetWidget(), FVector2D(CurChildOffset + ThisPadding.Left, YAlignmentResult.Offset), FVector2D(WidgetDesiredSize.X, YAlignmentResult.Size)));
-	return CurChildOffset + ThisSlotDesiredWidth;
-}
-
 SScrollBox::SScrollBox()
 {
 	VerticalScrollBarSlot = nullptr;
diff --git a/Engine/Source/Runtime/Slate/Public/Widgets/Layout/SScrollBox.h b/Engine/Source/Runtime/Slate/Public/Widgets/Layout/SScrollBox.h
--- a/Engine/Source/Runtime/Slate/Public/Widgets/Layout/SScrollBox.h
+++ b/Engine/Source/Runtime/Slate/Public/Widgets/Layout/SScrollBox.h
@@ -70,7 +70,106 @@ class SLATE_API SScrollBox : public SCompoundWidget
 {
 public:
 	/** A Slot that provides layout options for the contents of a scrollable box. */
-	using FSlot = FBasicLayoutWidgetSlot;
+	class SLATE_API FSlot : public TBasicLayoutWidgetSlot<FSlot>
+	{
+	public:
+		SLATE_SLOT_BEGIN_ARGS(FSlot, TBasicLayoutWidgetSlot<FSlot>)
+			SLATE_ARGUMENT(TOptional<FSizeParam>, SizeParam)
+			TAttribute<float> _MaxSize;
+
+			/** The widget's DesiredSize will be used as the space required. */
+			FSlotArguments& AutoSize()
+			{
+				_SizeParam = FAuto();
+				return Me();
+			}
+			/** The available space will be distributed proportionately. */
+			FSlotArguments& FillSize(TAttribute<float> InStretchCoefficient)
+			{
+				_SizeParam = FStretch(MoveTemp(InStretchCoefficient));
+				return Me();
+			}
+			/** Set the max size in SlateUnit this slot can be. */
+			FSlotArguments& MaxSize(TAttribute<float> InMaxHeight)
+			{
+				_MaxSize = MoveTemp(InMaxHeight);
+				return Me();
+			}
+		SLATE_SLOT_END_ARGS()
+
+		/** Default values for a slot. */
+		FSlot()
+			: TBasicLayoutWidgetSlot<FSlot>(HAlign_Fill, VAlign_Fill)
+			, SizeRule(FSizeParam::SizeRule_Auto)
+			, SizeValue(*this, 1.f)
+			, MaxSize(*this, 0.0f)
+		{ }
+
+		void Construct(const FChildren& SlotOwner, FSlotArguments&& InArgs);
+		static void RegisterAttributes(FSlateWidgetSlotAttributeInitializer& AttributeInitializer);
+
+		/** Get the space rule this slot should occupy along scrollbox's direction. */
+		FSizeParam::ESizeRule GetSizeRule() const
+		{
+			return SizeRule;
+		}
+
+		/** Get the space rule value this slot should occupy along scrollbox's direction. */
+		float GetSizeValue() const
+		{
+			return SizeValue.Get();
+		}
+
+		/** Get the max size the slot can be.*/
+		float GetMaxSize() const
+		{
+			return MaxSize.Get();
+		}
+
+		/** Set the size Param of the slot, It could be a FStretch or a FAuto. */
+		void SetSizeParam(FSizeParam InSizeParam)
+		{
+			SizeRule = InSizeParam.SizeRule;
+			SizeValue.Assign(*this, MoveTemp(InSizeParam.Value));
+		}
+
+		/** The widget's DesiredSize will be used as the space required. */
+		void SetSizeToAuto()
+		{
+			SetSizeParam(FAuto());
+		}
+
+		/** The available space will be distributed proportionately. */
+		void SetSizeToStretch(TAttribute<float> StretchCoefficient)
+		{
+			SetSizeParam(FStretch(MoveTemp(StretchCoefficient)));
+		}
+
+		/** Set the max size in SlateUnit this slot can be. */
+		void SetMaxSize(TAttribute<float> InMaxSize)
+		{
+			MaxSize.Assign(*this, MoveTemp(InMaxSize));
+		}
+
+	private:
+		/**
+		 * How much space this slot should occupy along scrollbox's direction.
+		 * When SizeRule is SizeRule_Auto, the widget's DesiredSize will be used as the space required.
+		 * When SizeRule is SizeRule_Stretch, the available space will be distributed proportionately between
+		 * peer Widgets depending on the Value property. Available space is space remaining after all the
+		 * peers' SizeRule_Auto requirements have been satisfied.
+		 */
+
+		 /** The sizing rule to use. */
+		FSizeParam::ESizeRule SizeRule;
+
+		/** The actual value this size parameter stores. */
+		typename TBasicLayoutWidgetSlot<FSlot>::template TSlateSlotAttribute<float> SizeValue;
+
+		/** The max size that this slot can be (0 if no max) */
+		typename TBasicLayoutWidgetSlot<FSlot>::template TSlateSlotAttribute<float> MaxSize;
+
+	};
 
 	SLATE_BEGIN_ARGS(SScrollBox)
 		: _Style( &FAppStyle::Get().GetWidgetStyle<FScrollBoxStyle>("ScrollBox") )
@@ -447,10 +546,10 @@ public:
 	SLATE_ARGUMENT(bool, BackPadScrolling)
 	SLATE_ARGUMENT(bool, FrontPadScrolling)
 
-		SLATE_END_ARGS()
+	SLATE_END_ARGS()
 
-		SScrollPanel()
-		: Children(this)
+	SScrollPanel()
+	: Children(this)
 	{
 	}
 
@@ -486,11 +585,6 @@ protected:
 	virtual FVector2D ComputeDesiredSize(float) const override;
 	// End SWidget overrides.
 
-private:
-
-	float ArrangeChildVerticalAndReturnOffset(const FGeometry& AllottedGeometry, FArrangedChildren& ArrangedChildren, const SScrollBox::FSlot& ThisSlot, float CurChildOffset) const;
-	float ArrangeChildHorizontalAndReturnOffset(const FGeometry& AllottedGeometry, FArrangedChildren& ArrangedChildren, const SScrollBox::FSlot& ThisSlot, float CurChildOffset) const;
-
 private:
 
 	EOrientation Orientation;
diff --git a/Engine/Source/Runtime/SlateCore/Private/Widgets/SBoxPanel.cpp b/Engine/Source/Runtime/SlateCore/Private/Widgets/SBoxPanel.cpp
--- a/Engine/Source/Runtime/SlateCore/Private/Widgets/SBoxPanel.cpp
+++ b/Engine/Source/Runtime/SlateCore/Private/Widgets/SBoxPanel.cpp
@@ -105,138 +105,6 @@ const SStackBox::FSlot& SStackBox::GetSlot(int32 SlotIndex) const
 	return static_cast<const SStackBox::FSlot&>(BaseSlot);
 }
 
-template<EOrientation Orientation, typename SlotType>
-static void ArrangeChildrenAlong(EFlowDirection InLayoutFlow, const TPanelChildren<SlotType>& Children, const FGeometry& AllottedGeometry, FArrangedChildren& ArrangedChildren )
-{
-	// Allotted space will be given to fixed-size children first.
-	// Remaining space will be proportionately divided between stretch children (SizeRule_Stretch)
-	// based on their stretch coefficient
-
-	if (Children.Num() > 0)
-	{
-		float StretchCoefficientTotal = 0;
-		float FixedTotal = 0;
-
-		bool bAnyChildVisible = false;
-		// Compute the sum of stretch coefficients (SizeRule_Stretch) and space required by fixed-size widgets
-		// (SizeRule_Auto).
-		for( int32 ChildIndex=0; ChildIndex < Children.Num(); ++ChildIndex )
-		{
-			const SlotType& CurChild = Children[ChildIndex];
-
-			if ( CurChild.GetWidget()->GetVisibility() != EVisibility::Collapsed )
-			{
-				bAnyChildVisible = true;
-				// All widgets contribute their margin to the fixed space requirement
-				FixedTotal += CurChild.GetPadding().template GetTotalSpaceAlong<Orientation>();
-
-				if ( CurChild.GetSizeRule() == FSizeParam::SizeRule_Stretch )
-				{
-					// for stretch children we save sum up the stretch coefficients
-					StretchCoefficientTotal += CurChild.GetSizeValue();
-				}
-				else
-				{
-					FVector2D ChildDesiredSize = CurChild.GetWidget()->GetDesiredSize();
-
-					// Auto-sized children contribute their desired size to the fixed space requirement
-					const float ChildSize = (Orientation == Orient_Vertical)
-						? ChildDesiredSize.Y
-						: ChildDesiredSize.X;
-
-					// Clamp to the max size if it was specified
-					float MaxSize = CurChild.GetMaxSize();
-					FixedTotal += MaxSize > 0 ? FMath::Min( MaxSize, ChildSize ) : ChildSize;
-				}
-			}
-		}
-
-		if (!bAnyChildVisible)
-		{
-			return;
-		}
-
-		// The space available for SizeRule_Stretch widgets is any space that wasn't taken up by fixed-sized widgets.
-		const float NonFixedSpace = FMath::Max( 0.0f, (Orientation == Orient_Vertical)
-			? AllottedGeometry.GetLocalSize().Y - FixedTotal
-			: AllottedGeometry.GetLocalSize().X - FixedTotal );
-
-		float PositionSoFar = 0;
-
-		// Now that we have the total fixed-space requirement and the total stretch coefficients we can
-		// arrange widgets top-to-bottom or left-to-right (depending on the orientation).
-		for (TPanelChildrenConstIterator<SlotType> It(Children, Orientation, InLayoutFlow); It; ++It)
-		{
-			const SlotType& CurChild = *It;
-			const EVisibility ChildVisibility = CurChild.GetWidget()->GetVisibility();
-
-			// Figure out the area allocated to the child in the direction of BoxPanel
-			// The area allocated to the slot is ChildSize + the associated margin.
-			float ChildSize = 0;
-			if (ChildVisibility != EVisibility::Collapsed)
-			{
-				// The size of the widget depends on its size type
-				if (CurChild.GetSizeRule() == FSizeParam::SizeRule_Stretch)
-				{
-					if (StretchCoefficientTotal > 0)
-					{
-						// Stretch widgets get a fraction of the space remaining after all the fixed-space requirements are met
-						ChildSize = NonFixedSpace * CurChild.GetSizeValue() / StretchCoefficientTotal;
-					}
-				}
-				else
-				{
-					const FVector2D ChildDesiredSize = CurChild.GetWidget()->GetDesiredSize();
-
-					// Auto-sized widgets get their desired-size value
-					ChildSize = (Orientation == Orient_Vertical)
-						? ChildDesiredSize.Y
-						: ChildDesiredSize.X;
-				}
-
-				// Clamp to the max size if it was specified
-				float MaxSize = CurChild.GetMaxSize();
-				if (MaxSize > 0)
-				{
-					ChildSize = FMath::Min(MaxSize, ChildSize);
-				}
-			}
-
-			const FMargin SlotPadding(LayoutPaddingWithFlow(InLayoutFlow, CurChild.GetPadding()));
-
-			FVector2D SlotSize = (Orientation == Orient_Vertical)
-				? FVector2D(AllottedGeometry.GetLocalSize().X, ChildSize + SlotPadding.template GetTotalSpaceAlong<Orient_Vertical>())
-				: FVector2D(ChildSize + SlotPadding.template GetTotalSpaceAlong<Orient_Horizontal>(), AllottedGeometry.GetLocalSize().Y);
-
-			// Figure out the size and local position of the child within the slot			
-			AlignmentArrangeResult XAlignmentResult = AlignChild<Orient_Horizontal>(InLayoutFlow, SlotSize.X, CurChild, SlotPadding);
-			AlignmentArrangeResult YAlignmentResult = AlignChild<Orient_Vertical>(SlotSize.Y, CurChild, SlotPadding);
-
-			const FVector2D LocalPosition = (Orientation == Orient_Vertical)
-				? FVector2D(XAlignmentResult.Offset, PositionSoFar + YAlignmentResult.Offset)
-				: FVector2D(PositionSoFar + XAlignmentResult.Offset, YAlignmentResult.Offset);
-
-			const FVector2D LocalSize = FVector2D(XAlignmentResult.Size, YAlignmentResult.Size);
-
-			// Add the information about this child to the output list (ArrangedChildren)
-			ArrangedChildren.AddWidget(ChildVisibility, AllottedGeometry.MakeChild(
-				// The child widget being arranged
-				CurChild.GetWidget(),
-				// Child's local position (i.e. position within parent)
-				LocalPosition,
-				// Child's size
-				LocalSize
-			));
-
-			if (ChildVisibility != EVisibility::Collapsed)
-			{
-				// Offset the next child by the size of the current child and any post-child (bottom/right) margin
-				PositionSoFar += (Orientation == Orient_Vertical) ? SlotSize.Y : SlotSize.X;
-			}
-		}
-	}
-}
-
 /**
  * Panels arrange their children in a space described by the AllottedGeometry parameter. The results of the arrangement
  * should be returned by appending a FArrangedWidget pair for every child widget. See StackPanel for an example
@@ -246,13 +114,16 @@ static void ArrangeChildrenAlong(EFlowDirection InLayoutFlow, const TPanelChildr
  */
 void SBoxPanel::OnArrangeChildren( const FGeometry& AllottedGeometry, FArrangedChildren& ArrangedChildren ) const
 {
+	const float Offset = 0.0f;
+	const bool AllowShrink = true;
+
 	if ( this->Orientation == EOrientation::Orient_Horizontal )
 	{
-		ArrangeChildrenAlong<EOrientation::Orient_Horizontal>(GSlateFlowDirection, this->Children, AllottedGeometry, ArrangedChildren );
+		ArrangeChildrenInStack<EOrientation::Orient_Horizontal>(GSlateFlowDirection, this->Children, AllottedGeometry, ArrangedChildren, Offset, AllowShrink );
 	}
 	else
 	{
-		ArrangeChildrenAlong<EOrientation::Orient_Vertical>(GSlateFlowDirection, this->Children, AllottedGeometry, ArrangedChildren );
+		ArrangeChildrenInStack<EOrientation::Orient_Vertical>(GSlateFlowDirection, this->Children, AllottedGeometry, ArrangedChildren, Offset, AllowShrink );
 	}
 }
 
diff --git a/Engine/Source/Runtime/SlateCore/Public/Layout/LayoutUtils.h b/Engine/Source/Runtime/SlateCore/Public/Layout/LayoutUtils.h
--- a/Engine/Source/Runtime/SlateCore/Public/Layout/LayoutUtils.h
+++ b/Engine/Source/Runtime/SlateCore/Public/Layout/LayoutUtils.h
@@ -2,6 +2,8 @@
 
 #pragma once
 
+#include "Layout/Children.h"
+#include "Types/SlateStructs.h"
 #include "CoreMinimal.h"
 #include "Margin.h"
 #include "Visibility.h"
@@ -216,6 +218,144 @@ static void ArrangeSingleChild(EFlowDirection InFlowDirection, const FGeometry&
 	}
 }
 
+template<EOrientation Orientation, typename SlotType>
+static void ArrangeChildrenInStack(EFlowDirection InLayoutFlow, const TPanelChildren<SlotType>& Children, const FGeometry& AllottedGeometry, FArrangedChildren& ArrangedChildren, float InOffset, bool InAllowShrink)
+{
+	// Allotted space will be given to fixed-size children first.
+	// Remaining space will be proportionately divided between stretch children (SizeRule_Stretch)
+	// based on their stretch coefficient
+
+	if (Children.Num() > 0)
+	{
+		float StretchCoefficientTotal = 0.0f;
+		float FixedTotal = 0.0f;
+		float StretchSizeTotal = 0.0f;
+
+		bool bAnyChildVisible = false;
+		// Compute the sum of stretch coefficients (SizeRule_Stretch) and space required by fixed-size widgets (SizeRule_Auto),
+		// as well as the total desired size.
+		for (int32 ChildIndex = 0; ChildIndex < Children.Num(); ++ChildIndex)
+		{
+			const SlotType& CurChild = Children[ChildIndex];
+
+			if (CurChild.GetWidget()->GetVisibility() != EVisibility::Collapsed)
+			{
+				bAnyChildVisible = true;
+				// All widgets contribute their margin to the fixed space requirement
+				FixedTotal += CurChild.GetPadding().template GetTotalSpaceAlong<Orientation>();
+
+				FVector2D ChildDesiredSize = CurChild.GetWidget()->GetDesiredSize();
+
+				// Auto-sized children contribute their desired size to the fixed space requirement
+				const float ChildSize = (Orientation == Orient_Vertical)
+					? ChildDesiredSize.Y
+					: ChildDesiredSize.X;
+
+				if (CurChild.GetSizeRule() == FSizeParam::SizeRule_Stretch)
+				{
+					// for stretch children we save sum up the stretch coefficients
+					StretchCoefficientTotal += CurChild.GetSizeValue();
+					StretchSizeTotal += ChildSize;
+				}
+				else
+				{
+
+					// Clamp to the max size if it was specified
+					float MaxSize = CurChild.GetMaxSize();
+					FixedTotal += MaxSize > 0.0f ? FMath::Min(MaxSize, ChildSize) : ChildSize;
+				}
+			}
+		}
+
+		if (!bAnyChildVisible)
+		{
+			return;
+		}
+
+		//When shrink is not allowed, we'll ensure to use all the space desired by the stretchable widgets.
+		const float MinSize = InAllowShrink ? 0.0f : StretchSizeTotal;
+
+		// The space available for SizeRule_Stretch widgets is any space that wasn't taken up by fixed-sized widgets.
+		const float NonFixedSpace = FMath::Max(MinSize, (Orientation == Orient_Vertical)
+			? AllottedGeometry.GetLocalSize().Y - FixedTotal
+			: AllottedGeometry.GetLocalSize().X - FixedTotal);
+
+		float PositionSoFar = 0.0f;
+
+		// Now that we have the total fixed-space requirement and the total stretch coefficients we can
+		// arrange widgets top-to-bottom or left-to-right (depending on the orientation).
+		for (TPanelChildrenConstIterator<SlotType> It(Children, Orientation, InLayoutFlow); It; ++It)
+		{
+			const SlotType& CurChild = *It;
+			const EVisibility ChildVisibility = CurChild.GetWidget()->GetVisibility();
+
+			// Figure out the area allocated to the child in the direction of BoxPanel
+			// The area allocated to the slot is ChildSize + the associated margin.
+			float ChildSize = 0.0f;
+			if (ChildVisibility != EVisibility::Collapsed)
+			{
+				// The size of the widget depends on its size type
+				if (CurChild.GetSizeRule() == FSizeParam::SizeRule_Stretch)
+				{
+					if (StretchCoefficientTotal > 0.0f)
+					{
+						// Stretch widgets get a fraction of the space remaining after all the fixed-space requirements are met
+						ChildSize = NonFixedSpace * CurChild.GetSizeValue() / StretchCoefficientTotal;
+					}
+				}
+				else
+				{
+					const FVector2D ChildDesiredSize = CurChild.GetWidget()->GetDesiredSize();
+
+					// Auto-sized widgets get their desired-size value
+					ChildSize = (Orientation == Orient_Vertical)
+						? ChildDesiredSize.Y
+						: ChildDesiredSize.X;
+				}
+
+				// Clamp to the max size if it was specified
+				float MaxSize = CurChild.GetMaxSize();
+				if (MaxSize > 0.0f)
+				{
+					ChildSize = FMath::Min(MaxSize, ChildSize);
+				}
+			}
+
+			const FMargin SlotPadding(LayoutPaddingWithFlow(InLayoutFlow, CurChild.GetPadding()));
+
+			FVector2D SlotSize = (Orientation == Orient_Vertical)
+				? FVector2D(AllottedGeometry.GetLocalSize().X, ChildSize + SlotPadding.template GetTotalSpaceAlong<Orient_Vertical>())
+				: FVector2D(ChildSize + SlotPadding.template GetTotalSpaceAlong<Orient_Horizontal>(), AllottedGeometry.GetLocalSize().Y);
+
+			// Figure out the size and local position of the child within the slot			
+			AlignmentArrangeResult XAlignmentResult = AlignChild<Orient_Horizontal>(InLayoutFlow, SlotSize.X, CurChild, SlotPadding);
+			AlignmentArrangeResult YAlignmentResult = AlignChild<Orient_Vertical>(SlotSize.Y, CurChild, SlotPadding);
+
+			const FVector2D LocalPosition = (Orientation == Orient_Vertical)
+				? FVector2D(XAlignmentResult.Offset, PositionSoFar + YAlignmentResult.Offset + InOffset)
+				: FVector2D(PositionSoFar + XAlignmentResult.Offset + InOffset, YAlignmentResult.Offset);
+
+			const FVector2D LocalSize = FVector2D(XAlignmentResult.Size, YAlignmentResult.Size);
+
+			// Add the information about this child to the output list (ArrangedChildren)
+			ArrangedChildren.AddWidget(ChildVisibility, AllottedGeometry.MakeChild(
+				// The child widget being arranged
+				CurChild.GetWidget(),
+				// Child's local position (i.e. position within parent)
+				LocalPosition,
+				// Child's size
+				LocalSize
+			));
+
+			if (ChildVisibility != EVisibility::Collapsed)
+			{
+				// Offset the next child by the size of the current child and any post-child (bottom/right) margin
+				PositionSoFar += (Orientation == Orient_Vertical) ? SlotSize.Y : SlotSize.X;
+			}
+		}
+	}
+}
+
 static FMargin LayoutPaddingWithFlow(EFlowDirection InLayoutFlow, const FMargin& InPadding)
 {
 	FMargin ReturnPadding(InPadding);
diff --git a/Engine/Source/Runtime/UMG/Private/Components/ScrollBoxSlot.cpp b/Engine/Source/Runtime/UMG/Private/Components/ScrollBoxSlot.cpp
--- a/Engine/Source/Runtime/UMG/Private/Components/ScrollBoxSlot.cpp
+++ b/Engine/Source/Runtime/UMG/Private/Components/ScrollBoxSlot.cpp
@@ -16,6 +16,7 @@ PRAGMA_DISABLE_DEPRECATION_WARNINGS
 	HorizontalAlignment = HAlign_Fill;
 	VerticalAlignment = VAlign_Fill;
 PRAGMA_ENABLE_DEPRECATION_WARNINGS
+	Size = FSlateChildSize(ESlateSizeRule::Automatic);
 }
 
 void UScrollBoxSlot::BuildSlot(TSharedRef<SScrollBox> ScrollBox)
@@ -26,7 +27,8 @@ PRAGMA_DISABLE_DEPRECATION_WARNINGS
 		.HAlign(HorizontalAlignment)
 		.VAlign(VerticalAlignment)
 		.Expose(Slot)
-		[
+		.SizeParam(UWidget::ConvertSerializedSizeParamToRuntime(Size))
+	[
 			Content == nullptr ? SNullWidget::NullWidget : Content->TakeWidget()
 		];
 PRAGMA_ENABLE_DEPRECATION_WARNINGS
@@ -47,6 +49,20 @@ void UScrollBoxSlot::SetPadding(FMargin InPadding)
 	}
 }
 
+FSlateChildSize UScrollBoxSlot::GetSize() const
+{
+	return Size;
+}
+
+void UScrollBoxSlot::SetSize(FSlateChildSize InSize)
+{
+	Size = InSize;
+	if (Slot)
+	{
+		Slot->SetSizeParam(UWidget::ConvertSerializedSizeParamToRuntime(InSize));
+	}
+}
+
 EHorizontalAlignment UScrollBoxSlot::GetHorizontalAlignment() const
 {
 	return Slot ? Slot->GetHorizontalAlignment() : HorizontalAlignment.GetValue();
@@ -83,6 +99,7 @@ PRAGMA_DISABLE_DEPRECATION_WARNINGS
 	SetHorizontalAlignment(HorizontalAlignment);
 	SetVerticalAlignment(VerticalAlignment);
 PRAGMA_ENABLE_DEPRECATION_WARNINGS
+	SetSize(Size);
 }
 
 void UScrollBoxSlot::ReleaseSlateResources(bool bReleaseChildren)
diff --git a/Engine/Source/Runtime/UMG/Public/Components/ScrollBoxSlot.h b/Engine/Source/Runtime/UMG/Public/Components/ScrollBoxSlot.h
--- a/Engine/Source/Runtime/UMG/Public/Components/ScrollBoxSlot.h
+++ b/Engine/Source/Runtime/UMG/Public/Components/ScrollBoxSlot.h
@@ -6,6 +6,7 @@
 #include "UObject/ObjectMacros.h"
 #include "UObject/ScriptMacros.h"
 #include "Layout/Margin.h"
+#include "Components/SlateWrapperTypes.h"
 #include "Widgets/Layout/SScrollBox.h"
 #include "Components/PanelSlot.h"
 
@@ -17,6 +18,12 @@ class UMG_API UScrollBoxSlot : public UPanelSlot
 {
 	GENERATED_UCLASS_BODY()
 
+private:
+
+	/** How much space this slot should occupy in the direction of the panel. */
+	UPROPERTY(EditAnywhere, BlueprintReadWrite, Getter, Setter, Category = "Layout|ScrollBox Slot", meta = (AllowPrivateAccess = "true", DisplayAfter = "Padding"))
+	FSlateChildSize Size;
+
 public:
 	
 	/** The padding area between the slot and the content it contains. */
@@ -41,6 +48,10 @@ public:
 	UFUNCTION(BlueprintCallable, Category="Layout|ScrollBox Slot")
 	void SetPadding(FMargin InPadding);
 
+	FSlateChildSize GetSize() const;
+
+	void SetSize(FSlateChildSize InSize);
+
 	EHorizontalAlignment GetHorizontalAlignment() const;
 
 	UFUNCTION(BlueprintCallable, Category="Layout|ScrollBox Slot")
