Set user index on FPointerEvent in the WidgetInteractionComponent. This fixes an issue where there could be large input deltas on mobile when using a virtual slate user because the index of 8 would not be cleared, since it is not registered with the device mapper.

#rnx

diff --git a/Engine/Source/Runtime/SlateCore/Public/Input/Events.h b/Engine/Source/Runtime/SlateCore/Public/Input/Events.h
--- a/Engine/Source/Runtime/SlateCore/Public/Input/Events.h
+++ b/Engine/Source/Runtime/SlateCore/Public/Input/Events.h
@@ -896,7 +896,8 @@ public:
 		bool bInIsForceChanged = false,
 		bool bInIsFirstMove = false,
 		const FModifierKeysState& InModifierKeys = FModifierKeysState(),
-		uint32 InTouchpadIndex=0
+		uint32 InTouchpadIndex=0,
+		const TOptional<int32> InOptionalSlateUserIndex = TOptional<int32>()
 	)
 	: FInputEvent(InModifierKeys, InDeviceId, false)
 		, ScreenSpacePosition(InScreenSpacePosition)
@@ -913,7 +914,12 @@ public:
 		, bIsDirectionInvertedFromDevice(false)
 		, bIsTouchForceChanged(bInIsForceChanged)
 		, bIsTouchFirstMove(bInIsFirstMove)
-	{ }
+	{
+		if (InOptionalSlateUserIndex.IsSet())
+		{
+			UserIndex = InOptionalSlateUserIndex.GetValue();
+		}
+	}
 
 	/** A constructor for gesture events */
 	FPointerEvent(
diff --git a/Engine/Source/Runtime/UMG/Private/Components/WidgetInteractionComponent.cpp b/Engine/Source/Runtime/UMG/Private/Components/WidgetInteractionComponent.cpp
--- a/Engine/Source/Runtime/UMG/Private/Components/WidgetInteractionComponent.cpp
+++ b/Engine/Source/Runtime/UMG/Private/Components/WidgetInteractionComponent.cpp
@@ -475,8 +475,12 @@ void UWidgetInteractionComponent::PressPointerKey(FKey Key)
 			LocalHitLocation,
 			LastLocalHitLocation,
 			1.0f,
-			false);
-		
+			false,
+			false,
+			false,
+			FModifierKeysState(),
+			0,
+			VirtualUser->GetUserIndex());		
 	}
 	else
 	{
@@ -488,7 +492,8 @@ void UWidgetInteractionComponent::PressPointerKey(FKey Key)
 			PressedKeys,
 			Key,
 			0.0f,
-			ModifierKeys);
+			ModifierKeys,
+			VirtualUser->GetUserIndex());
 	}
 	
 		
