Fixes for the currently set theme not saving over restarts, possibility of deleting or altering default theme.

diff --git a/Engine/Source/Editor/UnrealEd/Private/Settings/EditorStyleSettingsCustomization.cpp b/Engine/Source/Editor/UnrealEd/Private/Settings/EditorStyleSettingsCustomization.cpp
--- a/Engine/Source/Editor/UnrealEd/Private/Settings/EditorStyleSettingsCustomization.cpp
+++ b/Engine/Source/Editor/UnrealEd/Private/Settings/EditorStyleSettingsCustomization.cpp
@@ -691,6 +691,7 @@ void FEditorStyleSettingsCustomization::OnThemePicked(TSharedPtr<FString> NewSel
 	// If set directly in code, the theme was already applied
 	if(SelectInfo != ESelectInfo::Direct)
 	{
+		StyleSetting->SaveConfig();
 		USlateThemeManager::Get().ApplyTheme(StyleSetting->CurrentAppliedTheme);
 	}
 }
diff --git a/Engine/Source/Runtime/SlateCore/Private/Styling/StyleColors.cpp b/Engine/Source/Runtime/SlateCore/Private/Styling/StyleColors.cpp
--- a/Engine/Source/Runtime/SlateCore/Private/Styling/StyleColors.cpp
+++ b/Engine/Source/Runtime/SlateCore/Private/Styling/StyleColors.cpp
@@ -7,6 +7,7 @@
 #include "Serialization/JsonReader.h"
 #include "Serialization/JsonSerializer.h"
 #include "HAL/PlatformFileManager.h"
+#include "Misc/App.h"
 
 #include UE_INLINE_GENERATED_CPP_BY_NAME(StyleColors)
 
@@ -246,8 +247,12 @@ bool USlateThemeManager::IsEngineTheme() const
 	const FString EnginePath = GetEngineThemeDir() / GetCurrentTheme().DisplayName.ToString() + TEXT(".json");
 
 	IPlatformFile& FileManager = FPlatformFileManager::Get().GetPlatformFile();
-
-	if (FileManager.FileExists(*EnginePath))
+	// todo: move default theme check to a standalone function in 5.2
+	if (GetCurrentTheme() == DefaultDarkTheme)
+	{
+		return true;
+	}
+	else if (FileManager.FileExists(*EnginePath))
 	{
 		return true;
 	}
@@ -322,7 +327,7 @@ FString USlateThemeManager::GetProjectThemeDir() const
 
 FString USlateThemeManager::GetUserThemeDir() const
 {
-	return FPaths::EngineVersionAgnosticUserDir() / ThemesSubDir;
+	return FPaths::Combine(FPlatformProcess::UserSettingsDir(), *FApp::GetEpicProductIdentifier()) / ThemesSubDir;
 }
 
 void USlateThemeManager::LoadThemesFromDirectory(const FString& Directory)
