Code cleanup: Fixed typo in spelling of volatility in ensure messages.

#UE5 #RNX #NoReleaseNotes

diff --git a/Engine/Source/Runtime/SlateCore/Private/Widgets/SWidget.cpp b/Engine/Source/Runtime/SlateCore/Private/Widgets/SWidget.cpp
--- a/Engine/Source/Runtime/SlateCore/Private/Widgets/SWidget.cpp
+++ b/Engine/Source/Runtime/SlateCore/Private/Widgets/SWidget.cpp
@@ -795,13 +795,13 @@ void SWidget::UpdateWidgetProxy(int32 NewLayerId, FSlateCachedElementsHandle& Ca
 		if (IsVolatile() && !IsVolatileIndirectly())
 		{
 			ensureMsgf(HasAnyUpdateFlags(EWidgetUpdateFlags::NeedsVolatilePaint)
-				, TEXT("The votality of the widget '%s' changed during Paint")
+				, TEXT("The volatility of the widget '%s' changed during Paint")
 				, *FReflectionMetaData::GetWidgetPath(this));
 		}
 		else
 		{
 			ensureMsgf(!HasAnyUpdateFlags(EWidgetUpdateFlags::NeedsVolatilePaint)
-				, TEXT("The votality of the widget '%s' changed during Paint")
+				, TEXT("The volatility of the widget '%s' changed during Paint")
 				, *FReflectionMetaData::GetWidgetPath(this));
 		}
 	}
