Fix SlateInvalidationRoot flow direction behavior that was not handling the Inherit behavior properly => it was overwriting RightToLeft setting with LeftToRight.

test:
-repro data from UDN no longer fail.
-no regression seen when testing with Lyra (english and arabic, global invalidation on/off).

diff --git a/Engine/Source/Runtime/SlateCore/Private/FastUpdate/SlateInvalidationRoot.cpp b/Engine/Source/Runtime/SlateCore/Private/FastUpdate/SlateInvalidationRoot.cpp
--- a/Engine/Source/Runtime/SlateCore/Private/FastUpdate/SlateInvalidationRoot.cpp
+++ b/Engine/Source/Runtime/SlateCore/Private/FastUpdate/SlateInvalidationRoot.cpp
@@ -374,12 +374,7 @@ FSlateInvalidationResult FSlateInvalidationRoot::PaintInvalidationRoot(const FSl
 		bNeedScreenPositionShift = false;
 	}
 
-	EFlowDirection NewFlowDirection = GSlateFlowDirection;
-	if (RootWidget->GetFlowDirectionPreference() == EFlowDirectionPreference::Inherit)
-	{
-		NewFlowDirection = GSlateFlowDirectionShouldFollowCultureByDefault ? FLayoutLocalization::GetLocalizedLayoutDirection() : EFlowDirection::LeftToRight;
-	}
-	TGuardValue<EFlowDirection> FlowGuard(GSlateFlowDirection, NewFlowDirection);
+	TGuardValue<EFlowDirection> FlowGuard(GSlateFlowDirection, RootWidget->ComputeFlowDirection());
 	if (!Context.bAllowFastPathUpdate || bNeedsSlowPath || GSlateIsInInvalidationSlowPath)
 	{
 		SCOPED_NAMED_EVENT(Slate_PaintSlowPath, FColor::Red);
@@ -857,9 +852,9 @@ void FSlateInvalidationRoot::ProcessPreUpdate()
 #endif
 
 			/** */
-			struct FChildOriderInvalidationCallbackImpl : FSlateInvalidationWidgetList::IProcessChildOrderInvalidationCallback
+			struct FChildOrderInvalidationCallbackImpl : FSlateInvalidationWidgetList::IProcessChildOrderInvalidationCallback
 			{
-				FChildOriderInvalidationCallbackImpl(
+				FChildOrderInvalidationCallbackImpl(
 					const FSlateInvalidationWidgetList& InWidgetList
 					, FSlateInvalidationWidgetPreHeap& InPreUpdate
 					, FSlateInvalidationWidgetPrepassHeap& InPrepassUpdate
@@ -869,7 +864,7 @@ void FSlateInvalidationRoot::ProcessPreUpdate()
 					, PrepassUpdate(InPrepassUpdate)
 					, PostUpdate(InPostUpdate)
 				{}
-				virtual ~FChildOriderInvalidationCallbackImpl() = default;
+				virtual ~FChildOrderInvalidationCallbackImpl() = default;
 				const FSlateInvalidationWidgetList& WidgetList;
 				FSlateInvalidationWidgetPreHeap& PreUpdate;
 				FSlateInvalidationWidgetPrepassHeap& PrepassUpdate;
@@ -888,7 +883,7 @@ void FSlateInvalidationRoot::ProcessPreUpdate()
 				virtual void ProxiesReIndexed(const FReIndexOperation& Operation) override
 				{
 					// Re-index in Pre and Post list (modify the index and the sort value)
-					FChildOriderInvalidationCallbackImpl const* Self = this;
+					FChildOrderInvalidationCallbackImpl const* Self = this;
 					auto ReIndexIfNeeded = [&Operation, Self](FSlateInvalidationWidgetPreHeap::FElement& Element)
 					{
 						if (Operation.GetRange().Include(Element.GetWidgetSortOrder()))
@@ -905,7 +900,7 @@ void FSlateInvalidationRoot::ProcessPreUpdate()
 				virtual void ProxiesPreResort(const FReSortOperation& Operation) override
 				{
 					// The sort order value will change but the order (operator<) is still valid.
-					FChildOriderInvalidationCallbackImpl* Self = this;
+					FChildOrderInvalidationCallbackImpl* Self = this;
 					auto ReSortIfNeeded = [&Operation, Self](FSlateInvalidationWidgetPreHeap::FElement& Element)
 					{
 						if (Operation.GetRange().Include(Element.GetWidgetSortOrder()))
@@ -928,7 +923,7 @@ void FSlateInvalidationRoot::ProcessPreUpdate()
 				}
 			};
 
-			FChildOriderInvalidationCallbackImpl ChildOrderInvalidationCallback{ *FastWidgetPathList, *WidgetsNeedingPreUpdate, *WidgetsNeedingPrepassUpdate, *WidgetsNeedingPostUpdate };
+			FChildOrderInvalidationCallbackImpl ChildOrderInvalidationCallback{ *FastWidgetPathList, *WidgetsNeedingPreUpdate, *WidgetsNeedingPrepassUpdate, *WidgetsNeedingPostUpdate };
 
 			while(WidgetsNeedingPreUpdate->Num() > 0 && !bNeedsSlowPath)
 			{
