Misc trivial cleanup (comments errors, removed dead code, clean unsafe cast warning)

diff --git a/Engine/Source/Runtime/Slate/Public/Widgets/Text/SlateEditableTextLayout.h b/Engine/Source/Runtime/Slate/Public/Widgets/Text/SlateEditableTextLayout.h
--- a/Engine/Source/Runtime/Slate/Public/Widgets/Text/SlateEditableTextLayout.h
+++ b/Engine/Source/Runtime/Slate/Public/Widgets/Text/SlateEditableTextLayout.h
@@ -665,7 +665,7 @@ private:
 	/** How the text was committed by the virtual keyboard */
 	ETextCommit::Type VirtualKeyboardTextCommitType;
 
-	/** Override for for the overflow policy. If this is not set the text style setting is used */
+	/** Override for the overflow policy. If this is not set the text style setting is used */
 	TOptional<ETextOverflowPolicy> OverflowPolicyOverride;
 
 	/** The last known size of the widget from the previous OnPaint, used to recalculate wrapping */
diff --git a/Engine/Source/Runtime/SlateCore/Private/Widgets/SWidget.cpp b/Engine/Source/Runtime/SlateCore/Private/Widgets/SWidget.cpp
--- a/Engine/Source/Runtime/SlateCore/Private/Widgets/SWidget.cpp
+++ b/Engine/Source/Runtime/SlateCore/Private/Widgets/SWidget.cpp
@@ -1391,13 +1391,6 @@ int32 SWidget::Paint(const FPaintArgs& Args, const FGeometry& AllottedGeometry,
 	INC_DWORD_STAT(STAT_SlateNumPaintedWidgets);
 	UE_TRACE_SCOPED_SLATE_WIDGET_PAINT(this);
 
-	const SWidget* PaintParent = Args.GetPaintParent();
-	//if (GSlateEnableGlobalInvalidation)
-	//{
-	//	bInheritedVolatility = PaintParent ? (PaintParent->IsVolatileIndirectly() || PaintParent->IsVolatile()) : false;
-	//}
-
-
 	// If this widget clips to its bounds, then generate a new clipping rect representing the intersection of the bounding
 	// rectangle of the widget's geometry, and the current clipping rectangle.
 	bool bClipToBounds, bAlwaysClip, bIntersectClipBounds;
@@ -2044,7 +2037,7 @@ bool SWidget::CanChildrenBeAccessible() const
 bool SWidget::IsChildWidgetCulled(const FSlateRect& MyCullingRect, const FArrangedWidget& ArrangedChild) const
 {
 	// If we've enabled global invalidation it's safe to run the culling logic and just 'stop' drawing
-	// a widget, that widget has to be given an opportunity to paint, as wlel as all its children, the
+	// a widget, that widget has to be given an opportunity to paint, as well as all its children, the
 	// only correct way is to remove the widget from the tree, or to change the visibility of it.
 	if (GSlateIsOnFastUpdatePath)
 	{
diff --git a/Engine/Source/Runtime/UMG/Private/UserWidget.cpp b/Engine/Source/Runtime/UMG/Private/UserWidget.cpp
--- a/Engine/Source/Runtime/UMG/Private/UserWidget.cpp
+++ b/Engine/Source/Runtime/UMG/Private/UserWidget.cpp
@@ -537,7 +537,7 @@ UUMGSequencePlayer* UUserWidget::PlayAnimationForward(UWidgetAnimation* InAnimat
 		return Player;
 	}
 
-	return PlayAnimation(InAnimation, 0.0f, 1.0f, EUMGSequencePlayMode::Forward, PlaybackSpeed, bRestoreState);
+	return PlayAnimation(InAnimation, 0.0f, 1, EUMGSequencePlayMode::Forward, PlaybackSpeed, bRestoreState);
 }
 
 UUMGSequencePlayer* UUserWidget::PlayAnimationReverse(UWidgetAnimation* InAnimation, float PlaybackSpeed, bool bRestoreState)
@@ -555,7 +555,7 @@ UUMGSequencePlayer* UUserWidget::PlayAnimationReverse(UWidgetAnimation* InAnimat
 		return Player;
 	}
 
-	return PlayAnimation(InAnimation, 0.0f, 1.0f, EUMGSequencePlayMode::Reverse, PlaybackSpeed, bRestoreState);
+	return PlayAnimation(InAnimation, 0.0f, 1, EUMGSequencePlayMode::Reverse, PlaybackSpeed, bRestoreState);
 }
 
 void UUserWidget::StopAnimation(const UWidgetAnimation* InAnimation)
diff --git a/Engine/Source/Runtime/UMG/Public/Components/Widget.h b/Engine/Source/Runtime/UMG/Public/Components/Widget.h
--- a/Engine/Source/Runtime/UMG/Public/Components/Widget.h
+++ b/Engine/Source/Runtime/UMG/Public/Components/Widget.h
@@ -773,7 +773,7 @@ public:
 	/**
 	 * Gets the underlying slate widget or constructs it if it doesn't exist.  If you're looking to replace
 	 * what slate widget gets constructed look for RebuildWidget.  For extremely special cases where you actually
-	 * need to change the the GC Root widget of the constructed User Widget - you need to use TakeDerivedWidget
+	 * need to change the GC Root widget of the constructed User Widget - you need to use TakeDerivedWidget
 	 * you must also take care to not call TakeWidget before calling TakeDerivedWidget, as that would put the wrong
 	 * expected wrapper around the resulting widget being constructed.
 	 */
